Copying //UE4/Dev-Physics to Dev-Main (//UE4/Dev-Main) [at] 7703071

diff --git a/Engine/Source/Runtime/Engine/Private/PhysicsEngine/ImmediatePhysics/ImmediatePhysicsPhysX/ImmediatePhysicsSimulation_PhysX.cpp b/Engine/Source/Runtime/Engine/Private/PhysicsEngine/ImmediatePhysics/ImmediatePhysicsPhysX/ImmediatePhysicsSimulation_PhysX.cpp
new file mode 100644
--- /dev/null
+++ b/Engine/Source/Runtime/Engine/Private/PhysicsEngine/ImmediatePhysics/ImmediatePhysicsPhysX/ImmediatePhysicsSimulation_PhysX.cpp
@@ -0,0 +1,1171 @@
+// Copyright 1998-2019 Epic Games, Inc. All Rights Reserved.
+
+#include "Physics/ImmediatePhysics/ImmediatePhysicsPhysX/ImmediatePhysicsSimulation_PhysX.h"
+
+#if WITH_PHYSX && PHYSICS_INTERFACE_PHYSX
+
+#include "Physics/ImmediatePhysics/ImmediatePhysicsPhysX/ImmediatePhysicsActorHandle_PhysX.h"
+#include "Physics/ImmediatePhysics/ImmediatePhysicsPhysX/ImmediatePhysicsJointHandle_PhysX.h"
+#include "Physics/ImmediatePhysics/ImmediatePhysicsPhysX/ImmediatePhysicsContactPointRecorder_PhysX.h"
+#include "Physics/ImmediatePhysics/ImmediatePhysicsPhysX/ResourceManager_PhysX.h"
+
+#include "Physics/ImmediatePhysics/ImmediatePhysicsStats.h"
+
+#include "PhysicsEngine/BodyInstance.h"
+#include "PhysicsEngine/ConstraintInstance.h"
+
+#include "PhysicalMaterials/PhysicalMaterial.h"
+
+namespace ImmediatePhysics_PhysX
+{
+
+FSimulation::FSimulation()
+{
+	NumActiveSimulatedBodies = 0;
+	NumSimulatedBodies = 0;
+	NumKinematicBodies = 0;
+	NumSimulatedShapesWithCollision = 0;
+	NumPositionIterations = 1;
+	NumVelocityIterations = 1;
+	bRecreateIterationCache = false;
+	SimCount = 0;
+
+	NumContactHeaders = 0;
+	NumJointHeaders = 0;
+	NumActiveJoints = 0;
+	bDirtyJointData = false;
+}
+
+FSimulation::~FSimulation()
+{
+	for(FActorHandle* ActorHandle : ActorHandles)
+	{
+		delete ActorHandle;
+	}
+
+	ActorHandles.Empty();
+}
+
+void FSimulation::SetPositionIterationCount(int32 InIterationCount)
+{
+	NumVelocityIterations = InIterationCount;
+}
+
+void FSimulation::SetVelocityIterationCount(int32 InIterationCount)
+{
+	NumPositionIterations = InIterationCount;
+}
+
+void FSimulation::SetNumActiveBodies(int32 InNumActiveBodies)
+{
+	NumActiveSimulatedBodies = InNumActiveBodies;
+
+	bDirtyJointData = true;
+}
+
+FJointHandle* FSimulation::CreateJoint(FConstraintInstance* ConstraintInstance, FActorHandle* Actor1, FActorHandle* Actor2)
+{
+	//TODO: make this robust
+	if (Actor1 == nullptr) return nullptr;
+
+	check(&Actor1->OwningSimulation == this);
+	check(!Actor2 || &Actor2->OwningSimulation == this);	//Actor2 can be null
+
+	FJointHandle* NewHandle = nullptr;
+
+#if PHYSICS_INTERFACE_PHYSX
+	PxD6Joint* Joint = ConstraintInstance->GetPhysicsConstraintRef().ConstraintData;
+	int32 JointDataIdx = JointData.Emplace(Joint);
+	PrepareJointData(JointData[JointDataIdx]);
+
+	if(!JointData[JointDataIdx].HasConstraints())
+	{
+		JointData.RemoveAt(JointDataIdx, /*Count=*/1, /*bAllowShrinking=*/false);
+	}
+	else
+	{
+		Joints.Emplace();
+		FJoint& NewJoint = Joints[JointDataIdx];
+		//TODO: make dynamic vs static API more robust
+		NewJoint.DynamicActor = Actor2;
+		NewJoint.OtherActor = Actor1;
+
+		NewHandle = new FJointHandle(*this, Joints.Num() - 1);
+		JointHandles.Add(NewHandle);
+		bDirtyJointData = true;
+	}
+
+	ValidateArrays();
+#endif
+
+	return NewHandle;
+}
+
+void FSimulation::SetIgnoreCollisionPairTable(const TArray<FIgnorePair>& InIgnoreTable)
+{
+	IgnoreCollisionPairTable.Reset();
+	for(const FIgnorePair& IgnorePair : InIgnoreTable)
+	{
+		TSet<FActorHandle*>& IgnoreActorsA = IgnoreCollisionPairTable.FindOrAdd(IgnorePair.A);
+		IgnoreActorsA.Add(IgnorePair.B);
+
+		TSet<FActorHandle*>& IgnoreActorsB = IgnoreCollisionPairTable.FindOrAdd(IgnorePair.B);
+		IgnoreActorsB.Add(IgnorePair.A);
+	}
+
+	bRecreateIterationCache = true;
+}
+
+void FSimulation::SetIgnoreCollisionActors(const TArray<FActorHandle*>& InIgnoreCollisionActors)
+{
+	IgnoreCollisionActors.Reset();
+	IgnoreCollisionActors.Append(InIgnoreCollisionActors);
+	bRecreateIterationCache = true;
+}
+
+bool FSimulation::GetIsKinematic(int32 ActorDataIndex)
+{
+	const int32 NumMobileBodies = NumSimulatedBodies + NumKinematicBodies;
+	const bool bIsSimulated = ActorDataIndex < NumSimulatedBodies;
+	const bool bIsKinematic = !bIsSimulated && ActorDataIndex < NumMobileBodies;
+
+	return bIsKinematic;
+}
+
+void FSimulation::SetIsKinematic(int32 ActorDataIndex, bool bKinematic)
+{
+	const int32 NumMobileBodies = NumSimulatedBodies + NumKinematicBodies;
+	const bool bIsSimulated = ActorDataIndex < NumSimulatedBodies;
+	const bool bIsKinematic = !bIsSimulated && ActorDataIndex < NumKinematicBodies;
+	const bool bStatic = !bIsSimulated && !bIsKinematic;
+
+	if(bStatic)
+	{
+		// Should we support this case? Currently we build static solver data on actor creation
+		// so likely can't just force this to be kinematic here.
+		return;
+	}
+
+	if(bIsSimulated && bKinematic)
+	{
+		// A to B
+
+		// Need to clear out some stuff first, no velocities when kinematic
+		immediate::PxRigidBodyData& RBData = RigidBodiesData[ActorDataIndex];
+		RBData.angularVelocity = PxVec3(0.0f);
+		RBData.linearVelocity = PxVec3(0.0f);
+
+		int32 LastSimIndex = NumSimulatedBodies - 1;
+		if(ActorDataIndex != LastSimIndex)
+		{
+			SwapActorData(ActorDataIndex, LastSimIndex);
+		}
+
+		// Rearrange groups so kinematic now contains the body
+		--NumSimulatedBodies;
+		--NumActiveSimulatedBodies;
+		++NumKinematicBodies;
+	}
+	else if(bIsKinematic && !bKinematic)
+	{
+		// B to A
+		int32 FirstKinematicIndex = NumSimulatedBodies;
+
+		if(ActorDataIndex != FirstKinematicIndex)
+		{
+			SwapActorData(ActorDataIndex, FirstKinematicIndex);
+		}
+
+		++NumSimulatedBodies;
+		++NumActiveSimulatedBodies;
+		--NumKinematicBodies;
+	}
+}
+
+template <EActorType ActorType>
+int32 FSimulation::CreateActor(PxRigidActor* RigidActor, const FTransform& WorldTM)
+{
+	bDirtyJointData = true;	//new entity potentially re-orders bodies so joint data becomes stale. TODO: can this be optimized for adding static cases that don't re-order?
+	bRecreateIterationCache = true;	//new entity potentially re-orders bodies so our iteration cache becomes stale. TODO: can this be optimized for adding static cases that don't re-order?
+
+	FActor* NewEntity = new (Actors) FActor();
+	const int32 ActorDataIndex = Actors.Num() - 1;
+
+	FActorHandle* NewActorHandle = new FActorHandle(*this, ActorDataIndex);
+	ActorHandles.Add(NewActorHandle);
+
+	PxSolverBodyData* NewSolverBodyData = new (SolverBodiesData) PxSolverBodyData();
+	immediate::PxRigidBodyData* NewRigidBodyData = new (RigidBodiesData) immediate::PxRigidBodyData();
+
+	//Fill rigid body data based on the actor we've been given
+	PxRigidBody* const RigidBody = ActorType == EActorType::StaticActor ? nullptr : static_cast<PxRigidBody*>(RigidActor);	//Note the const pointer - this is for templated code
+	PxRigidDynamic* const RigidDynamic = ActorType == EActorType::DynamicActor ? static_cast<PxRigidDynamic*>(RigidActor) : nullptr;	//Note the const pointer - this is for templated code
+
+	NewRigidBodyData->invMass = RigidDynamic ? RigidDynamic->getInvMass() : 0.f;
+	NewRigidBodyData->angularVelocity = PxVec3(PxZero);
+	NewRigidBodyData->maxDepenetrationVelocity = RigidDynamic ? RigidDynamic->getMaxDepenetrationVelocity() : PX_MAX_F32;
+	NewRigidBodyData->invInertia = RigidDynamic ? RigidDynamic->getMassSpaceInvInertiaTensor() : PxVec3(PxZero);
+	NewRigidBodyData->maxContactImpulse = RigidBody ? RigidBody->getMaxContactImpulse() : PX_MAX_F32;
+	
+	NewActorHandle->ActorToBody = RigidBody ? P2UTransform(RigidBody->getCMassLocalPose()) : FTransform::Identity;	//TODO: Why is this an unreal transform?
+	NewRigidBodyData->body2World = U2PTransform(NewActorHandle->ActorToBody * WorldTM);
+
+	NewRigidBodyData->linearDamping = RigidDynamic ? RigidDynamic->getLinearDamping() : 0.f;
+	NewRigidBodyData->angularDamping = RigidDynamic ? RigidDynamic->getAngularDamping() : 0.f;
+	NewRigidBodyData->maxLinearVelocitySq = PX_MAX_F32;	//NOTE: rigid dynamics don't seem to have this max property.
+
+	PendingAcceleration.AddZeroed();
+	KinematicTargets.AddZeroed();
+
+	if(RigidDynamic)
+	{
+		NewRigidBodyData->maxAngularVelocitySq = RigidDynamic->getMaxAngularVelocity();
+		NewRigidBodyData->maxAngularVelocitySq *= NewRigidBodyData->maxAngularVelocitySq;
+
+		PxU32 PositionIterations, VelocityIterations;
+		RigidDynamic->getSolverIterationCounts(PositionIterations, VelocityIterations);
+
+		NumPositionIterations = FMath::Max((int32)PositionIterations, NumPositionIterations);
+		NumVelocityIterations = FMath::Max((int32)VelocityIterations, NumVelocityIterations);
+
+	}else
+	{
+		NewRigidBodyData->maxAngularVelocitySq = PX_MAX_F32;
+	}
+
+	if (ActorType == EActorType::StaticActor)
+	{
+		//If static we fill this struct here, for dynamic we fill during every call to simulate
+		immediate::PxConstructStaticSolverBody(NewRigidBodyData->body2World, *NewSolverBodyData);
+	}
+
+	NewEntity->CreateGeometry(RigidActor, U2PTransform(NewActorHandle->ActorToBody));
+
+	ValidateArrays();
+
+	//We sort entities so that dynamics are first, kinematics second, and static third
+	if (ActorType == EActorType::DynamicActor)	//dynamic first
+	{
+		if (NumSimulatedBodies < ActorDataIndex)	//There's at least one kinematic/static entity so swap
+		{
+			const int32 FirstNonSimulatedIndex = NumSimulatedBodies;
+			SwapActorData(FirstNonSimulatedIndex, ActorDataIndex);
+		}
+
+		++NumSimulatedBodies;
+	}
+	else if(ActorType == EActorType::KinematicActor)  //then kinematics
+	{
+		const int32 NumMobileBodies = NumKinematicBodies + NumSimulatedBodies;
+		if (NumMobileBodies < ActorDataIndex)	//There's at least one static entity so swap
+		{
+			const int32 FirstStaticIndex = NumMobileBodies;
+			SwapActorData(FirstStaticIndex, ActorDataIndex);
+		}
+
+		++NumKinematicBodies;
+	}
+
+	return NewActorHandle->ActorDataIndex;
+}
+
+FActorHandle* FSimulation::InsertActorData(const FActor& InActor, const FActorData& InData)
+{
+	// @todo(ccaulfield): does not seem to be called anywhere...
+
+	FActor* NewActor = new(Actors) FActor(InActor);
+	const int32 ActorDataIndex = Actors.Num() - 1;
+
+	FActorHandle* NewActorHandle = new FActorHandle(*this, ActorDataIndex);
+	ActorHandles.Add(NewActorHandle);
+
+	PxSolverBodyData* NewSolverBodyData = new (SolverBodiesData) PxSolverBodyData();
+	RigidBodiesData.Add(InData.RigidBodyData);
+
+	immediate::PxRigidBodyData& InternalData = RigidBodiesData.Last();
+
+	NewActorHandle->ActorToBody = FTransform::Identity; // #PHYS2 should be the centre of mass of the actor
+	InternalData.body2World = U2PTransform(NewActorHandle->ActorToBody * InData.InitialTransform);
+
+	PendingAcceleration.AddZeroed();
+	KinematicTargets.AddZeroed();
+
+	const bool bDynamic = !InData.bStatic && !InData.bKinematic;
+	if(bDynamic)
+	{
+		// Dynamic object
+
+		// #PHYS2 Setup max angular velocity
+	}
+	else if(InData.bKinematic)
+	{
+		// Kinematic object
+	}
+	else
+	{
+		// Static object
+		immediate::PxConstructStaticSolverBody(InternalData.body2World, *NewSolverBodyData);
+	}
+
+	ValidateArrays();
+
+	//We sort entities so that dynamics are first, kinematics second, and static third
+	if(bDynamic)	//dynamic first
+	{
+		if(NumSimulatedBodies < ActorDataIndex)	//There's at least one kinematic/static entity so swap
+		{
+			const int32 FirstNonSimulatedIndex = NumSimulatedBodies;
+			SwapActorData(FirstNonSimulatedIndex, ActorDataIndex);
+		}
+
+		++NumSimulatedBodies;
+	}
+	else if(InData.bKinematic)  //then kinematics
+	{
+		const int32 NumMobileBodies = NumKinematicBodies + NumSimulatedBodies;
+		if(NumMobileBodies < ActorDataIndex)	//There's at least one static entity so swap
+		{
+			const int32 FirstStaticIndex = NumMobileBodies;
+			SwapActorData(FirstStaticIndex, ActorDataIndex);
+		}
+
+		++NumKinematicBodies;
+	}
+
+	NumActiveSimulatedBodies = NumSimulatedBodies;
+
+	return NewActorHandle;
+}
+
+void FSimulation::RemoveActor(FActorHandle* InHandle)
+{
+	// @todo(ccaulfield): this is never called. Also I think it should delete the handle?
+
+	int32 DataIndex = InHandle->ActorDataIndex;
+
+	check(Actors.IsValidIndex(DataIndex));
+	check(DataIndex >= 0);
+
+	bool bDynamic = DataIndex < NumSimulatedBodies;
+	bool bKinematic = !bDynamic && DataIndex < NumKinematicBodies;
+	bool bStatic = !bDynamic && !bKinematic;
+
+	Actors.RemoveAt(DataIndex);
+	ActorHandles.RemoveAt(DataIndex);
+	RigidBodiesData.RemoveAt(DataIndex);
+	SolverBodiesData.RemoveAt(DataIndex);
+	PendingAcceleration.RemoveAt(DataIndex);
+	KinematicTargets.RemoveAt(DataIndex);
+
+	// Starting where we removed the actor we need to update any handles beyond it
+	const int32 NumTotalHandles = ActorHandles.Num();
+	for(int32 CurrIndex = DataIndex; CurrIndex < NumTotalHandles; ++CurrIndex)
+	{
+		ActorHandles[CurrIndex]->ActorDataIndex--;
+	}
+
+	bDirtyJointData = true;
+
+	if(bDynamic)
+	{
+		NumSimulatedBodies--;
+
+		// This is used for RBAN LOD optimizations, in normal simulations this should be identical to NumSimulatedBodies
+		NumActiveSimulatedBodies--;
+	}
+	else if(bKinematic)
+	{
+		NumKinematicBodies--;
+	}
+
+	bRecreateIterationCache = true;	//reordering of bodies so we need to change iteration order potentially
+}
+
+FActorHandle* FSimulation::CreateDynamicActor(FBodyInstance* BodyInstance, const FTransform& TM)
+{
+	PxRigidActor* PRigidActor = FPhysicsInterface::GetPxRigidActor_AssumesLocked(BodyInstance->GetPhysicsActorHandle());
+	if (PxRigidDynamic* PDynamic = PRigidActor->is<PxRigidDynamic>())
+	{
+		const int32 ActorDataIndex = CreateActor<EActorType::DynamicActor>(PDynamic, TM);
+		NumActiveSimulatedBodies = NumSimulatedBodies;
+		bDirtyJointData = true;
+		return ActorHandles[ActorDataIndex];
+	}
+	return nullptr;
+}
+
+FActorHandle* FSimulation::CreateKinematicActor(FBodyInstance* BodyInstance, const FTransform& TM)
+{
+	PxRigidActor* PRigidActor = FPhysicsInterface::GetPxRigidActor_AssumesLocked(BodyInstance->GetPhysicsActorHandle());
+	if (PxRigidBody* PRigidBody = PRigidActor->is<PxRigidBody>())
+	{
+		const int32 ActorDataIndex = CreateActor<EActorType::KinematicActor>(PRigidBody, TM);
+		return ActorHandles[ActorDataIndex];
+	}
+	return nullptr;
+}
+
+FActorHandle* FSimulation::CreateStaticActor(FBodyInstance* BodyInstance)
+{
+	if (PxRigidActor* RigidActor = FPhysicsInterface_PhysX::GetPxRigidActor_AssumesLocked(BodyInstance->ActorHandle))
+	{
+		const int32 ActorDataIndex = CreateActor<EActorType::StaticActor>(RigidActor, P2UTransform(RigidActor->getGlobalPose()));
+		return ActorHandles[ActorDataIndex];
+	}
+	return nullptr;
+}
+
+FActorHandle* FSimulation::CreateActor(EActorType ActorType, FBodyInstance* BodyInstance, const FTransform& TM)
+{
+	switch (ActorType)
+	{
+	case EActorType::StaticActor:
+		return CreateStaticActor(BodyInstance);
+	case EActorType::KinematicActor:
+		return CreateKinematicActor(BodyInstance, TM);
+	case EActorType::DynamicActor:
+		return CreateDynamicActor(BodyInstance, TM);
+	}
+	return nullptr;
+}
+
+void FSimulation::SwapActorData(int32 Actor1DataIdx, int32 Actor2DataIdx)
+{
+	check(Actors.IsValidIndex(Actor1DataIdx));
+	check(Actors.IsValidIndex(Actor2DataIdx));
+
+	Swap(Actors[Actor1DataIdx], Actors[Actor2DataIdx]);
+	Swap(ActorHandles[Actor1DataIdx], ActorHandles[Actor2DataIdx]);
+	Swap(RigidBodiesData[Actor1DataIdx], RigidBodiesData[Actor2DataIdx]);
+	Swap(SolverBodiesData[Actor1DataIdx], SolverBodiesData[Actor2DataIdx]);
+	Swap(PendingAcceleration[Actor1DataIdx], PendingAcceleration[Actor2DataIdx]);
+	Swap(KinematicTargets[Actor1DataIdx], KinematicTargets[Actor2DataIdx]);
+
+	//Update entity index on the handle
+	ActorHandles[Actor1DataIdx]->ActorDataIndex = Actor1DataIdx;
+	ActorHandles[Actor2DataIdx]->ActorDataIndex = Actor2DataIdx;
+	bDirtyJointData = true;	//reordering of bodies could lead to stale joint data
+
+	bRecreateIterationCache = true;	//reordering of bodies so we need to change iteration order potentially
+}
+
+void FSimulation::SwapJointData(int32 Joint1Idx, int32 Joint2Idx)
+{
+	check(JointData.IsValidIndex(Joint1Idx));
+	check(JointData.IsValidIndex(Joint2Idx));
+
+	Swap(JointData[Joint1Idx], JointData[Joint2Idx]);
+	Swap(Joints[Joint1Idx], Joints[Joint2Idx]);
+	Swap(JointHandles[Joint1Idx], JointHandles[Joint2Idx]);
+
+	//Update joint index on the handle
+	JointHandles[Joint1Idx]->JointDataIndex = Joint1Idx;
+	JointHandles[Joint2Idx]->JointDataIndex = Joint2Idx;
+	bDirtyJointData = true;	//reordering of joints could lead to stale joint data
+}
+
+void FSimulation::ValidateArrays() const
+{
+	check(Actors.Num() == ActorHandles.Num());
+	check(Actors.Num() == RigidBodiesData.Num());
+	check(Actors.Num() == SolverBodiesData.Num());
+	check(Actors.Num() == PendingAcceleration.Num());
+	check(Actors.Num() == KinematicTargets.Num());
+	check(Joints.Num() == JointData.Num());
+	check(Joints.Num() == JointHandles.Num());
+}
+
+/*static char*  CacheHack = nullptr;
+int32 GCacheSumHack = 0;
+
+TAutoConsoleVariable<int32> GHackEvictCache(TEXT("p.ForceCacheEvict"), 1, TEXT("blah"));
+
+void FSimulation::EvictCache()
+{
+	if(GHackEvictCache.GetValueOnAnyThread() == 0)
+	{
+		return;
+	}
+
+	if(CacheHack == nullptr)
+	{
+		CacheHack = (char*) FMemory::Malloc(1024*1024*2);
+	}
+
+	const int32 NumCacheLines = 1024*1024*2/64;
+	for(int i=0; i<NumCacheLines; ++i)
+	{
+		GCacheSumHack += CacheHack[i*64];
+	}
+	
+}*/
+
+DECLARE_DWORD_COUNTER_STAT(TEXT("Simulated Bodies"), STAT_IPNumSimulatedBodies, STATGROUP_ImmediatePhysics);
+DECLARE_DWORD_COUNTER_STAT(TEXT("Active Simulated Bodies"), STAT_IPNumActiveSimulatedBodies, STATGROUP_ImmediatePhysics);
+DECLARE_DWORD_COUNTER_STAT(TEXT("Kinematic Bodies"), STAT_IPNumKinematicBodies, STATGROUP_ImmediatePhysics);
+
+DECLARE_CYCLE_STAT(TEXT("FSimulation::Simulate"), STAT_ImmediateSimulate, STATGROUP_ImmediatePhysics);
+
+void FSimulation::Simulate(float DeltaTime, const FVector& Gravity)
+{
+	// Lock the resource manager before we simulate - any writes need to wait until we're done
+	FScopedSharedResourceLock<LockMode::Read> ScopeLock;
+	Simulate_AssumesLocked(DeltaTime, Gravity);
+}
+
+void FSimulation::Simulate_AssumesLocked(float DeltaTime, const FVector& Gravity)
+{
+	SET_DWORD_STAT(STAT_IPNumSimulatedBodies, NumSimulatedBodies);
+	SET_DWORD_STAT(STAT_IPNumActiveSimulatedBodies, NumActiveSimulatedBodies);
+	SET_DWORD_STAT(STAT_IPNumKinematicBodies, NumKinematicBodies);
+
+	DeltaTime = FMath::Min(DeltaTime, 0.033f);
+	//Create dynamic bodies and integrate their unconstrained velocities
+	if(DeltaTime > 0)
+	{
+		SCOPE_CYCLE_COUNTER(STAT_ImmediateSimulate);
+
+		++SimCount;
+
+		ConstructSolverBodies(DeltaTime, Gravity);
+	
+		if(bRecreateIterationCache)
+		{
+			PrepareIterationCache();
+		}
+		
+		GenerateContacts();
+		
+		BatchConstraints();
+
+		PrepareConstraints(DeltaTime);
+
+		SolveAndIntegrate(DeltaTime);
+
+		Workspace.Reset();
+	}
+
+	//EvictCache();
+}
+
+DECLARE_CYCLE_STAT(TEXT("ConstructSolverBodies"), STAT_ImmediateConstructSolverBodies, STATGROUP_ImmediatePhysics);
+
+void FSimulation::ConstructSolverBodies(float DeltaTime, const FVector& Gravity)
+{
+	SCOPE_CYCLE_COUNTER(STAT_ImmediateConstructSolverBodies);
+	const int32 NumBytes = Actors.Num() * sizeof(PxSolverBody);
+	SolverBodies = (PxSolverBody*)Workspace.Alloc(NumBytes);
+
+	for(int32 BodyIdx = 0; BodyIdx < NumActiveSimulatedBodies; ++BodyIdx)
+	{
+		RigidBodiesData[BodyIdx].linearVelocity += PendingAcceleration[BodyIdx] * DeltaTime;
+	}
+
+	FMemory::Memzero(PendingAcceleration.GetData(), sizeof(PendingAcceleration[0]) * NumSimulatedBodies);
+
+	FMemory::Memzero(SolverBodies, NumBytes);
+	immediate::PxConstructSolverBodies(RigidBodiesData.GetData(), SolverBodiesData.GetData(), NumActiveSimulatedBodies, U2PVector(Gravity), DeltaTime);
+	
+	const float DeltaTimeInv = 1.f / DeltaTime;
+
+	//compute velocity for kinematic targets
+	const int32 NumActors = Actors.Num();
+	for (int32 KinematicBodyIdx = NumSimulatedBodies; KinematicBodyIdx < NumActors; ++KinematicBodyIdx)
+	{
+		const FImmediateKinematicTarget& KinematicTarget = KinematicTargets[KinematicBodyIdx];
+		if(KinematicTarget.bTargetSet)
+		{
+			immediate::PxRigidBodyData& RigidBodyData = RigidBodiesData[KinematicBodyIdx];
+			const PxTransform& CurrentTM = RigidBodyData.body2World;
+			PxVec3 LinearDelta = KinematicTarget.BodyToWorld.p - CurrentTM.p;
+			PxQuat AngularDelta = KinematicTarget.BodyToWorld.q * CurrentTM.q.getConjugate();
+
+			if (AngularDelta.w < 0)	//shortest angle.
+				AngularDelta = -AngularDelta;
+
+			PxReal Angle;
+			PxVec3 Axis;
+			AngularDelta.toRadiansAndUnitAxis(Angle, Axis);
+			
+			RigidBodyData.linearVelocity = LinearDelta * DeltaTimeInv;
+			RigidBodyData.angularVelocity = Axis * Angle * DeltaTimeInv;
+		}
+	}
+
+	immediate::PxConstructSolverBodies(RigidBodiesData.GetData() + NumActiveSimulatedBodies, SolverBodiesData.GetData() + NumActiveSimulatedBodies, NumKinematicBodies + (NumSimulatedBodies - NumActiveSimulatedBodies), PxVec3(PxZero), DeltaTime);
+}
+
+DECLARE_CYCLE_STAT(TEXT("PrepareIterationCache"), STAT_ImmediatePrepareIterationCache, STATGROUP_ImmediatePhysics);
+
+void FSimulation::PrepareIterationCache()
+{
+	SCOPE_CYCLE_COUNTER(STAT_ImmediatePrepareIterationCache);
+	// Iteration cache is not yet available, so we need to iterate over bodies and prepare the data as needed
+	const int32 NumActors = Actors.Num();
+	ShapeSOA.LocalTMs.Empty(NumActors);
+	ShapeSOA.Materials.Empty(NumActors);
+	ShapeSOA.Geometries.Empty(NumActors);
+	ShapeSOA.Bounds.Empty(NumActors);
+	ShapeSOA.OwningActors.Empty(NumActors);
+	ShapeSOA.BoundsOffsets.Empty(NumActors);
+
+	int32 NumShapes = 0;
+	NumSimulatedShapesWithCollision = 0;
+
+	//flatten out all shapes of actors that have collision
+	for(int32 ActorIdx = 0; ActorIdx < NumActors; ++ActorIdx)
+	{
+		const FActor& Actor = Actors[ActorIdx];
+		if(IgnoreCollisionActors.Contains(ActorHandles[ActorIdx]))
+		{
+			if (ActorIdx+1 == NumSimulatedBodies)
+			{
+				NumSimulatedShapesWithCollision = NumShapes;
+			}
+
+			continue;
+		}
+
+		for(const FShape& Shape : Actor.Shapes)
+		{
+			ShapeSOA.LocalTMs.Add(Shape.LocalTM);
+			ShapeSOA.Materials.Add(Shape.Material ? *Shape.Material : FMaterial());
+			ShapeSOA.Geometries.Add(Shape.Geometry);
+			ShapeSOA.Bounds.Add(Shape.BoundsMagnitude);
+			ShapeSOA.BoundsOffsets.Add(Shape.BoundsOffset);
+			ShapeSOA.OwningActors.Add(ActorIdx);
+			++NumShapes;
+		}
+
+		if (ActorIdx+1 == NumSimulatedBodies)
+		{
+			NumSimulatedShapesWithCollision = NumShapes;
+		}
+	}
+
+	//build cache of which shapes to skip
+	int32 IterationCount = 0;
+	SkipCollisionCache.Empty(NumShapes);
+
+	for(int32 SimShapeIdx = 0; SimShapeIdx < NumSimulatedShapesWithCollision; ++SimShapeIdx)
+	{
+		const int32 SimActorIdx = ShapeSOA.OwningActors[SimShapeIdx];
+		const TSet<FActorHandle*>* IgnoreActorsForSimulated = IgnoreCollisionPairTable.Find(ActorHandles[SimActorIdx]);
+
+		for(int32 OtherShapeIdx = SimShapeIdx + 1; OtherShapeIdx < NumShapes; ++OtherShapeIdx)
+		{
+			const int32 OtherActorIdx = ShapeSOA.OwningActors[OtherShapeIdx];
+
+			if(SimActorIdx == OtherActorIdx || (IgnoreActorsForSimulated && IgnoreActorsForSimulated->Find(ActorHandles[OtherActorIdx])))
+			{
+				SkipCollisionCache.Add(IterationCount++);
+			}
+			else
+			{
+				++IterationCount;
+			}
+		}
+	}
+
+#if PERSISTENT_CONTACT_PAIRS
+	const int32 NumPairs = (NumShapes) * (NumShapes - 1) / 2;
+	ShapeSOA.ContactPairData.Empty(NumPairs);
+	ShapeSOA.ContactPairData.AddZeroed(NumPairs);
+#endif
+
+	bRecreateIterationCache = false;
+}
+
+DECLARE_CYCLE_STAT(TEXT("GenerateContacts"), STAT_ImmediateGenerateContacts, STATGROUP_ImmediatePhysics);
+
+void FSimulation::GenerateContacts()
+{
+	SCOPE_CYCLE_COUNTER(STAT_ImmediateGenerateContacts);
+
+	//narrow phase, find actual contact points
+	CacheAllocator.Reset();
+
+	ContactPairs.Reset();
+	ContactPoints.Reset();
+	
+	const int32 NumShapes = ShapeSOA.LocalTMs.Num();
+	PxTransform* ShapeWorldTMs = (PxTransform*)Workspace.Alloc(sizeof(PxTransform) * NumShapes);
+	for(int32 ShapeIdx = 0; ShapeIdx < NumShapes; ++ShapeIdx)
+	{
+		const immediate::PxRigidBodyData& Body = RigidBodiesData[ShapeSOA.OwningActors[ShapeIdx]];
+		ShapeWorldTMs[ShapeIdx] = Body.body2World * ShapeSOA.LocalTMs[ShapeIdx];
+	}
+
+	int32 IterationCount = 0;
+	int32 PotentialPairCount = 0;
+	const int32 IterationCacheSize = SkipCollisionCache.Num();
+	int32 CurrentIterationCacheIdx = 0;
+	const int32* SkipCollisionCacheRaw = SkipCollisionCache.GetData();
+
+	for(int32 SimShapeIdx = 0; SimShapeIdx < NumSimulatedShapesWithCollision; ++SimShapeIdx)
+	{
+		const float SimulatedRadius = ShapeSOA.Bounds[SimShapeIdx];
+		const PxTransform& SimulatedShapeTM = ShapeWorldTMs[SimShapeIdx];
+		const PxVec3 SimulatedBoundsOffset = ShapeSOA.BoundsOffsets[SimShapeIdx];
+		const PxVec3 SimulatedShapeBoundsOrigin = SimulatedShapeTM.transform(SimulatedBoundsOffset);
+		
+		const int32 SimulatedActorIdx = ShapeSOA.OwningActors[SimShapeIdx];
+		const FMaterial& SimulatedShapeMaterial = ShapeSOA.Materials[SimShapeIdx];
+		const PxGeometry* SimulatedGeometry = ShapeSOA.Geometries[SimShapeIdx];
+
+		if(SimulatedActorIdx >= NumActiveSimulatedBodies)
+		{
+			break;
+		}
+
+		for(int32 OtherShapeIdx = SimShapeIdx + 1; OtherShapeIdx < NumShapes; ++OtherShapeIdx)
+		{
+			if(CurrentIterationCacheIdx < IterationCacheSize && SkipCollisionCache[CurrentIterationCacheIdx] == IterationCount++)
+			{
+				++CurrentIterationCacheIdx;
+				continue;
+			}
+
+			const int32 OtherActorIdx = ShapeSOA.OwningActors[OtherShapeIdx];
+			if(OtherActorIdx >= NumActiveSimulatedBodies && OtherActorIdx < NumSimulatedBodies)
+			{
+				continue;
+			}
+
+			PxCache* InCache;
+			const int32 PotentialPairIdx = PotentialPairCount++;
+#if PERSISTENT_CONTACT_PAIRS
+			FPersistentContactPairData& PersistentPairData = ShapeSOA.ContactPairData[PotentialPairIdx];
+			InCache = &PersistentPairData.Cache;
+
+			if(PersistentPairData.SimCount + 1 != SimCount)	//Note: in the case where we roll over to 0 does this matter? Could use % but seems unneeded for now
+			{
+				PersistentPairData.Clear();
+			}
+#else
+			PxCache Cache;
+			InCache = &Cache;
+#endif
+			const FMaterial& OtherShapeMaterial = ShapeSOA.Materials[OtherShapeIdx];
+			const float OtherRadius = ShapeSOA.Bounds[OtherShapeIdx];
+			const float TotalRadius = (SimulatedRadius + OtherRadius);
+			const float TotalRadius2 = TotalRadius * TotalRadius;
+			const PxTransform& OtherShapeTM = ShapeWorldTMs[OtherShapeIdx];
+			const PxVec3 OtherShapeBoundsOrigin = OtherShapeTM.transform(ShapeSOA.BoundsOffsets[OtherShapeIdx]);
+
+			const float Distance2 = (SimulatedShapeBoundsOrigin  - OtherShapeBoundsOrigin).magnitudeSquared();
+
+			if ( Distance2 > TotalRadius2)
+			{
+#if PERSISTENT_CONTACT_PAIRS
+				PersistentPairData.Clear();
+#endif
+				continue;	//no intersection so skip narrow phase
+			}
+
+			FContactPointRecorder ContactRecorder(*this, SimulatedActorIdx, OtherActorIdx, PotentialPairIdx, SimulatedShapeMaterial, OtherShapeMaterial);
+			if (!immediate::PxGenerateContacts(&SimulatedGeometry, &ShapeSOA.Geometries[OtherShapeIdx], &SimulatedShapeTM, &OtherShapeTM, InCache, 1, ContactRecorder, 4.f, 1.f, 100.f, CacheAllocator))
+			{
+#if PERSISTENT_CONTACT_PAIRS
+				//no contacts found so clear friction
+				PersistentPairData.Frictions = nullptr;
+				PersistentPairData.NumFrictions = 0;
+#endif
+			}
+		}
+	}
+}
+
+bool GBatchJoints = false;	//TODO: temp hack to fix crash in immediate mode due to joint batching
+
+DECLARE_CYCLE_STAT(TEXT("BatchConstraints"), STAT_ImmediateBatchConstraints, STATGROUP_ImmediatePhysics);
+
+void FSimulation::BatchConstraints()
+{
+	SCOPE_CYCLE_COUNTER(STAT_ImmediateBatchConstraints);
+	
+	const int32 NumContactPairs = ContactPairs.Num();
+	const int32 NumJoints = JointData.Num();
+	const int32 NumConstraints = NumJoints + NumContactPairs;
+	OrderedDescriptors.SetNum(NumConstraints); //Ordered descriptors are semi-persistent. We only reorder joint constraints when needed, but we always update contact constraints which come at the end of the array.
+	BatchHeaders.SetNum(NumConstraints);	   //batch headers follow the same approach
+	
+	// Joint constraints if needed
+	if (bDirtyJointData)
+	{
+		NumActiveJoints = 0;
+		if(NumJoints > 0)
+		{
+			PxSolverConstraintDesc* JointDescriptors = (PxSolverConstraintDesc*)Workspace.Alloc(sizeof(PxSolverConstraintDesc) * NumJoints);
+
+			for (int32 JointIdx = 0; JointIdx < NumJoints; ++JointIdx)
+			{
+				PxSolverConstraintDesc& JointDescriptor = JointDescriptors[NumActiveJoints];
+
+				const FJoint& Joint = Joints[JointIdx];
+				const int32 DynamicActorIdx = Joint.DynamicActor->ActorDataIndex;
+				const int32 OtherActorIdx = Joint.OtherActor != nullptr ? Joint.OtherActor->ActorDataIndex : INDEX_NONE;
+
+				if( (DynamicActorIdx >= NumActiveSimulatedBodies && DynamicActorIdx < NumSimulatedBodies) || (OtherActorIdx >= NumActiveSimulatedBodies && OtherActorIdx < NumSimulatedBodies) )
+				{
+					continue;
+				}
+
+				JointDescriptor.bodyA = &SolverBodies[DynamicActorIdx];
+				JointDescriptor.bodyB = OtherActorIdx != INDEX_NONE ? &SolverBodies[OtherActorIdx] : nullptr;
+
+				JointDescriptor.bodyADataIndex = PxU16(DynamicActorIdx);
+				JointDescriptor.bodyBDataIndex = PxU16(OtherActorIdx);
+				JointDescriptor.linkIndexA = PxSolverConstraintDesc::NO_LINK;
+				JointDescriptor.linkIndexB = PxSolverConstraintDesc::NO_LINK;
+				JointDescriptor.writeBack = nullptr;
+
+				JointDescriptor.constraint = (PxU8*)JointHandles[JointIdx];
+				JointDescriptor.constraintLengthOver16 = PxSolverConstraintDesc::eJOINT_CONSTRAINT;
+
+				++NumActiveJoints;
+			}
+
+			if(GBatchJoints)
+			{
+				NumJointHeaders = physx::immediate::PxBatchConstraints(JointDescriptors, NumActiveJoints, SolverBodies, Actors.Num(), BatchHeaders.GetData(), OrderedDescriptors.GetData());
+			}
+			else
+			{
+				NumJointHeaders = NumActiveJoints;
+				OrderedDescriptors.AddUninitialized(NumJointHeaders);
+				FMemory::Memcpy(OrderedDescriptors.GetData(), JointDescriptors, sizeof(OrderedDescriptors[0]) * NumJointHeaders);
+
+				for(int32 Offset = 0; Offset < NumActiveJoints; ++Offset)
+				{
+					PxConstraintBatchHeader& BatchHeader = BatchHeaders[Offset];
+					BatchHeader.mStartIndex = Offset;
+					BatchHeader.mStride = 1;
+				}
+			}
+			
+
+			//ensure D6 joint data and low level joint data are matching the Ordered descriptors
+			for(int32 DescriptorIdx = 0; DescriptorIdx < NumActiveJoints; ++DescriptorIdx)
+			{
+				FJointHandle* JointHandle = (FJointHandle*) OrderedDescriptors[DescriptorIdx].constraint;
+				if(JointHandle->JointDataIndex != DescriptorIdx)
+				{
+					//The joint or the descriptor have moved - either way fix it
+					SwapJointData(DescriptorIdx, JointHandle->JointDataIndex);
+				}
+			}
+		}
+		else
+		{
+			NumJointHeaders = 0;
+		}
+	}
+
+	if(NumContactPairs > 0)
+	{
+		//Contact constraints
+		PxSolverConstraintDesc* ContactDescriptors = (PxSolverConstraintDesc*)Workspace.Alloc(sizeof(PxSolverConstraintDesc) * NumContactPairs);
+		for (int32 ContactPairIdx = 0; ContactPairIdx < NumContactPairs; ++ContactPairIdx)
+		{
+			const FContactPair& ContactPair = ContactPairs[ContactPairIdx];
+			PxSolverConstraintDesc& ContactDescriptor = ContactDescriptors[ContactPairIdx];
+
+			//Set body pointers and bodyData idxs
+			ContactDescriptor.bodyA = &SolverBodies[ContactPair.DynamicActorDataIndex];
+			ContactDescriptor.bodyB = &SolverBodies[ContactPair.OtherActorDataIndex];
+
+			ContactDescriptor.bodyADataIndex = PxU16(ContactPair.DynamicActorDataIndex);
+			ContactDescriptor.bodyBDataIndex = PxU16(ContactPair.OtherActorDataIndex);
+			ContactDescriptor.linkIndexA = PxSolverConstraintDesc::NO_LINK;
+			ContactDescriptor.linkIndexB = PxSolverConstraintDesc::NO_LINK;
+			ContactDescriptor.writeBack = nullptr;
+
+			ContactDescriptor.constraint = (PxU8*)(&ContactPair);	//Need to save this off for later when the constraint descriptors get re-ordered
+			ContactDescriptor.constraintLengthOver16 = PxSolverConstraintDesc::eCONTACT_CONSTRAINT;
+		}
+
+		NumContactHeaders = physx::immediate::PxBatchConstraints(ContactDescriptors, NumContactPairs, SolverBodies, Actors.Num(), BatchHeaders.GetData() + NumJointHeaders, OrderedDescriptors.GetData() + NumActiveJoints);
+	}
+	else
+	{
+		NumContactHeaders = 0;
+	}
+}
+
+//float GReadCacheHack = 0.f;
+
+DECLARE_CYCLE_STAT(TEXT("PrepareConstraints"), STAT_ImmediatePrepareConstraints, STATGROUP_ImmediatePhysics);
+
+void FSimulation::PrepareConstraints(float DeltaTime)
+{
+	SCOPE_CYCLE_COUNTER(STAT_ImmediatePrepareConstraints);
+
+	//Prepares constraint data based on latest position data and contact pairs
+	const int32 NumBodies = Actors.Num();
+	const float InvDeltaTime = 1.f / DeltaTime;
+	ConstraintAllocator.Reset();
+
+	PxSolverConstraintPrepDesc JointConstraints[4];
+	Px1DConstraint Rows[24];	//24 rows because we have at most 6 rows per joint and we could be solving at most 4 joints in one batch
+
+	/*FPlatformMisc::Prefetch(OrderedDescriptors.GetData());
+	FPlatformMisc::Prefetch(BatchHeaders.GetData());
+	FPlatformMisc::Prefetch(SolverBodiesData.GetData());
+	FPlatformMisc::Prefetch(JointData.GetData());*/
+
+	const int32 NumMobileBodies = NumSimulatedBodies + NumKinematicBodies;
+	
+	//Joint constraints
+	for (PxU32 HeaderIdx = 0; HeaderIdx < NumJointHeaders; ++HeaderIdx)
+	{
+		PxConstraintBatchHeader& Header = BatchHeaders[HeaderIdx];
+		Header.mConstraintType = PxSolverConstraintDesc::eJOINT_CONSTRAINT;	//Since joint headers are cached and solver re-uses this we have to reset it every time
+
+		FMemory::Memzero(Rows, sizeof(Rows));	//TODO: do we actually need to do this?
+		for (Px1DConstraint& Constraint : Rows)
+		{
+			Constraint.minImpulse = -FLT_MAX;
+			Constraint.maxImpulse = FLT_MAX;
+		}
+
+		int32 CurRows = 0;
+
+		for (PxU32 BatchInnerIdx = 0; BatchInnerIdx < Header.mStride; ++BatchInnerIdx)
+		{
+			const PxU32 DescriptorIdx = Header.mStartIndex + BatchInnerIdx;
+			//FPlatformMisc::PrefetchBlock(&JointData[DescriptorIdx], 128);
+			//GReadCacheHack = JointData[DescriptorIdx].invMassScale.angular0;
+
+			PxSolverConstraintDesc& OrderedDescriptor = OrderedDescriptors[DescriptorIdx];
+			PxSolverConstraintPrepDesc& JointDescriptor = JointConstraints[BatchInnerIdx];
+
+			//TODO: is it worth caching some of this? Most of the data doesn't change each frame
+			JointDescriptor.body0 = OrderedDescriptor.bodyA;
+			JointDescriptor.body1 = OrderedDescriptor.bodyB;
+			JointDescriptor.data0 = &SolverBodiesData[OrderedDescriptor.bodyADataIndex];
+			JointDescriptor.data1 = &SolverBodiesData[OrderedDescriptor.bodyBDataIndex];
+
+			JointDescriptor.bodyFrame0 = JointDescriptor.data0->body2World;
+			JointDescriptor.bodyFrame1 = JointDescriptor.data1->body2World;
+
+			JointDescriptor.bodyState0 = PxSolverConstraintPrepDescBase::eDYNAMIC_BODY;
+			JointDescriptor.bodyState1 = JointDescriptor.body1 == nullptr ? PxSolverConstraintPrepDescBase::eSTATIC_BODY : (OrderedDescriptor.bodyBDataIndex < (uint32)NumSimulatedBodies ? PxSolverConstraintPrepDescBase::eDYNAMIC_BODY : PxSolverConstraintPrepDescBase::eSTATIC_BODY);
+			JointDescriptor.desc = &OrderedDescriptor;
+			JointDescriptor.mInvMassScales.angular0 = JointDescriptor.mInvMassScales.angular1 = JointDescriptor.mInvMassScales.linear0 = JointDescriptor.mInvMassScales.linear1 = 1.f;
+			JointDescriptor.writeback = NULL;
+			//Constraint.getBreakForce(JointDescriptor.linBreakForce, JointDescriptor.angBreakForce);
+			JointDescriptor.linBreakForce = FLT_MAX;
+			JointDescriptor.angBreakForce = FLT_MAX;
+
+			JointDescriptor.minResponseThreshold = 0;//Constraint.getMinResponseThreshold();
+			JointDescriptor.disablePreprocessing = false;//!!(Constraint.getFlags() & PxConstraintFlag::eDISABLE_PREPROCESSING);
+			JointDescriptor.improvedSlerp = true;//!!(Constraint.getFlags() & PxConstraintFlag::eIMPROVED_SLERP);
+			JointDescriptor.driveLimitsAreForces = false;//!!(Constraint.getFlags() & PxConstraintFlag::eDRIVE_LIMITS_ARE_FORCES);
+
+			PxVec3 Body0WorldOffset;
+			PxConstraintInvMassScale InvMassScale;
+			const D6JointData& Joint = JointData[DescriptorIdx];	
+			int32 NumRows = D6JointSolverPrep(&Rows[CurRows], Body0WorldOffset, 0, InvMassScale, (const void*)&Joint, JointDescriptor.bodyFrame0, JointDescriptor.bodyFrame1);
+
+			JointDescriptor.mInvMassScales = InvMassScale;
+
+			JointDescriptor.numRows = NumRows;
+			JointDescriptor.rows = &Rows[CurRows];
+			CurRows += NumRows;
+		}
+
+		immediate::PxCreateJointConstraints(&Header, 1, JointConstraints, ConstraintAllocator, DeltaTime, InvDeltaTime);
+	}
+
+	bDirtyJointData = false;
+
+	//Contact constraints
+	const PxU32 NumTotalHeaders = NumContactHeaders + NumJointHeaders;
+	PxSolverContactDesc ContactDescriptors[4];
+	const FContactPair* InnerContactPairs[4];
+	const PxU32 NumJoints = JointData.Num();
+
+	for (PxU32 HeaderIdx = NumJointHeaders; HeaderIdx < NumTotalHeaders; ++HeaderIdx)
+	{
+		PxConstraintBatchHeader& Header = BatchHeaders[HeaderIdx];
+		check(Header.mConstraintType == PxSolverConstraintDesc::eCONTACT_CONSTRAINT);
+
+		Header.mStartIndex += NumActiveJoints;
+		
+		for (PxU32 BatchInnerIdx = 0; BatchInnerIdx < Header.mStride; ++BatchInnerIdx)
+		{
+			PxSolverConstraintDesc& OrderedDescriptor = OrderedDescriptors[Header.mStartIndex + BatchInnerIdx];
+			PxSolverContactDesc& ContactDescriptor = ContactDescriptors[BatchInnerIdx];
+			//Extract the contact pair that we saved in this structure earlier.
+			const FContactPair& ContactPair = *reinterpret_cast<const FContactPair*>(OrderedDescriptor.constraint);
+			InnerContactPairs[BatchInnerIdx] = &ContactPair;
+
+			ContactDescriptor.body0 = OrderedDescriptor.bodyA;
+			ContactDescriptor.body1 = OrderedDescriptor.bodyB;
+			ContactDescriptor.data0 = &SolverBodiesData[OrderedDescriptor.bodyADataIndex];
+			ContactDescriptor.data1 = &SolverBodiesData[OrderedDescriptor.bodyBDataIndex];
+
+			ContactDescriptor.bodyFrame0 = ContactDescriptor.data0->body2World;
+			ContactDescriptor.bodyFrame1 = ContactDescriptor.data1->body2World;
+
+			ContactDescriptor.contactForces = nullptr;//&ContactForces[ContactPair.StartContactIndex];
+			ContactDescriptor.contacts = &ContactPoints[ContactPair.StartContactIndex];
+			ContactDescriptor.numContacts = ContactPair.NumContacts;
+
+#if  PERSISTENT_CONTACT_PAIRS
+			FPersistentContactPairData& PersistentPairData = ShapeSOA.ContactPairData[ContactPair.PairIdx];
+
+			ContactDescriptor.frictionPtr = PersistentPairData.Frictions;
+			ContactDescriptor.frictionCount = PersistentPairData.NumFrictions;
+#else
+			ContactDescriptor.frictionPtr = nullptr;
+			ContactDescriptor.frictionCount = 0;
+#endif
+
+			ContactDescriptor.disableStrongFriction = false;
+			ContactDescriptor.hasMaxImpulse = false;
+			ContactDescriptor.hasForceThresholds = false;
+			ContactDescriptor.shapeInteraction = NULL;
+			ContactDescriptor.restDistance = 0.f;
+			ContactDescriptor.maxCCDSeparation = PX_MAX_F32;
+
+			ContactDescriptor.bodyState0 = PxSolverConstraintPrepDescBase::eDYNAMIC_BODY;
+			ContactDescriptor.bodyState1 = ContactPair.OtherActorDataIndex < (uint32)NumSimulatedBodies ? PxSolverConstraintPrepDescBase::eDYNAMIC_BODY : PxSolverConstraintPrepDescBase::eSTATIC_BODY;
+			ContactDescriptor.desc = &OrderedDescriptor;
+			ContactDescriptor.mInvMassScales.angular0 = ContactDescriptor.mInvMassScales.angular1 = ContactDescriptor.mInvMassScales.linear0 = ContactDescriptor.mInvMassScales.linear1 = 1.f;
+		}
+
+		immediate::PxCreateContactConstraints(&Header, 1, &ContactDescriptors[0], ConstraintAllocator, InvDeltaTime, -200.f, 4.f, 1.f);
+
+#if PERSISTENT_CONTACT_PAIRS
+		// Now cache friction results
+		for (PxU32 BatchInnerIdx = 0; BatchInnerIdx < Header.mStride; ++BatchInnerIdx)
+		{
+			PxSolverContactDesc& ContactDescriptor = ContactDescriptors[BatchInnerIdx];
+
+			//Extract the contact pair that we saved in this structure earlier.
+			const FContactPair& ContactPair = *InnerContactPairs[BatchInnerIdx];
+			FPersistentContactPairData& PersistentPairData = ShapeSOA.ContactPairData[ContactPair.PairIdx];
+
+			PersistentPairData.Frictions = ContactDescriptor.frictionPtr;
+			PersistentPairData.NumFrictions = ContactDescriptor.frictionCount;
+		}
+#endif
+	}
+}
+
+DECLARE_CYCLE_STAT(TEXT("SolveAndIntegrate"), STAT_ImmediateSolveAndIntegrate, STATGROUP_ImmediatePhysics);
+
+
+void FSimulation::SolveAndIntegrate(float DeltaTime)
+{
+	SCOPE_CYCLE_COUNTER(STAT_ImmediateSolveAndIntegrate);
+
+	PxVec3* LinearMotionVelocity = (PxVec3*) Workspace.Alloc(sizeof(PxVec3) * NumActiveSimulatedBodies * 2);
+	PxVec3* AngularMotionVelocity = &LinearMotionVelocity[NumActiveSimulatedBodies];
+
+	//Solve all constraints
+	immediate::PxSolveConstraints(BatchHeaders.GetData(), NumContactHeaders + NumJointHeaders, OrderedDescriptors.GetData(), SolverBodies, LinearMotionVelocity, AngularMotionVelocity, NumActiveSimulatedBodies, NumPositionIterations, NumVelocityIterations);
+
+	//Integrate velocities
+	immediate::PxIntegrateSolverBodies(SolverBodiesData.GetData(), SolverBodies, LinearMotionVelocity, AngularMotionVelocity, NumActiveSimulatedBodies, DeltaTime);
+
+	//Copy positions of dynamic bodies back into rigid body
+	for (int32 DynamicsBodyIdx = 0; DynamicsBodyIdx < NumActiveSimulatedBodies; ++DynamicsBodyIdx)
+	{
+		RigidBodiesData[DynamicsBodyIdx].linearVelocity = SolverBodiesData[DynamicsBodyIdx].linearVelocity;
+		RigidBodiesData[DynamicsBodyIdx].angularVelocity = SolverBodiesData[DynamicsBodyIdx].angularVelocity;
+		RigidBodiesData[DynamicsBodyIdx].body2World = SolverBodiesData[DynamicsBodyIdx].body2World;
+	}
+
+	//Copy kinematic targets
+	const int32 NumActors = Actors.Num();
+	for(int32 KinematicBodyIdx = NumSimulatedBodies; KinematicBodyIdx < NumActors; ++KinematicBodyIdx)
+	{
+		if(KinematicTargets[KinematicBodyIdx].bTargetSet)
+		{
+			RigidBodiesData[KinematicBodyIdx].body2World = KinematicTargets[KinematicBodyIdx].BodyToWorld;
+			RigidBodiesData[KinematicBodyIdx].linearVelocity = PxVec3(PxZero);
+			RigidBodiesData[KinematicBodyIdx].angularVelocity = PxVec3(PxZero);
+		}
+	}
+
+	FMemory::Memzero(KinematicTargets.GetData(), NumActors * sizeof(KinematicTargets[0]));
+}
+
+void FSimulation::AddRadialForce(int32 ActorDataIndex, const FVector& InOrigin, float Strength, float Radius, ERadialImpulseFalloff Falloff, EForceType ForceType)
+{
+	if(IsSimulated(ActorDataIndex))
+	{
+		PxVec3 Origin = U2PVector(InOrigin);
+		immediate::PxRigidBodyData& LowLevelBody = GetLowLevelBody(ActorDataIndex);
+		const PxTransform& PCOMTransform = LowLevelBody.body2World;
+		PxVec3 Delta = LowLevelBody.body2World.p - Origin;
+
+		float Mag = Delta.magnitude(); // Distance from COM to origin
+		if (Mag > Radius)
+		{
+			return;
+		}
+
+		Delta.normalize();
+
+		// Scale by U2PScale here, because units are velocity * mass. 
+		float ImpulseMag = Strength;
+		if (Falloff == RIF_Linear)
+		{
+			ImpulseMag *= (1.0f - (Mag / Radius));
+		}
+
+		const PxVec3 PImpulse = Delta * ImpulseMag;
+		const PxVec3 ApplyDelta = (ForceType == EForceType::AddAcceleration || ForceType == EForceType::AddVelocity) ? PImpulse : PImpulse * LowLevelBody.invMass;
+		
+		if(ForceType == EForceType::AddImpulse || ForceType == EForceType::AddVelocity)
+		{
+			LowLevelBody.linearVelocity += ApplyDelta;
+		}
+		else
+		{
+			PendingAcceleration[ActorDataIndex] += ApplyDelta;
+		}
+	}
+}
+
+void FSimulation::AddForce(int32 ActorDataIndex, const FVector& Force)
+{
+	if (IsSimulated(ActorDataIndex))
+	{
+		immediate::PxRigidBodyData& LowLevelBody = GetLowLevelBody(ActorDataIndex);
+		PendingAcceleration[ActorDataIndex] += U2PVector(Force) * LowLevelBody.invMass;
+	}
+}
+
+FActorData CreateActorData(const FActorCreationParams& InParams)
+{
+	FActorData OutData;
+
+	immediate::PxRigidBodyData& RigidData = OutData.RigidBodyData;
+
+	RigidData.linearVelocity = PxVec3(PxZero);
+	RigidData.invMass = 0.0f;
+	RigidData.angularVelocity = PxVec3(PxZero);
+	RigidData.maxDepenetrationVelocity = PX_MAX_F32;
+	RigidData.invInertia = PxVec3(PxZero);
+	RigidData.maxContactImpulse = PX_MAX_F32;
+	RigidData.linearDamping = 0.0f;
+	RigidData.angularDamping = 0.0f;
+	RigidData.maxLinearVelocitySq = PX_MAX_F32;
+	RigidData.maxAngularVelocitySq = PX_MAX_F32;
+	RigidData.body2World = U2PTransform(InParams.InitialTM);
+
+	OutData.InitialTransform = InParams.InitialTM;
+	OutData.bStatic = InParams.bStatic;
+	OutData.bKinematic = false; //#PHYS2 Add kinematic support
+
+	return OutData;
+}
+
+}
+
+#endif // WITH_PHYSX
\ No newline at end of file
diff --git a/Engine/Source/Runtime/Engine/Private/PhysicsEngine/ImmediatePhysics/ImmediatePhysicsPhysX/ResourceManager_PhysX.cpp b/Engine/Source/Runtime/Engine/Private/PhysicsEngine/ImmediatePhysics/ImmediatePhysicsPhysX/ResourceManager_PhysX.cpp
new file mode 100644
--- /dev/null
+++ b/Engine/Source/Runtime/Engine/Private/PhysicsEngine/ImmediatePhysics/ImmediatePhysicsPhysX/ResourceManager_PhysX.cpp
@@ -0,0 +1,52 @@
+// Copyright 1998-2019 Epic Games, Inc. All Rights Reserved.
+
+#include "Physics/ImmediatePhysics/ImmediatePhysicsPhysX/ResourceManager_PhysX.h"
+
+#if WITH_PHYSX && PHYSICS_INTERFACE_PHYSX
+
+namespace ImmediatePhysics_PhysX
+{
+
+	FSharedResourceManager FSharedResourceManager::Instance;
+
+	FResourceHandle FSharedResourceManager::CreateMaterial()
+	{
+		int32 Index = Materials.Add(TResourceWithId<FMaterial>());
+		int32 Id = Materials[Index].Id;
+
+		return FResourceHandle(EResourceType::Material, Index, Id);
+	}
+
+	void FSharedResourceManager::ReleaseMaterial(int32 InIndex)
+	{
+		Materials.RemoveAt(InIndex);
+	}
+
+	int32 FSharedResourceManager::GetMaterialId(int32 InIndex)
+	{
+		if (InIndex != INDEX_NONE && Materials.IsAllocated(InIndex))
+		{
+			return Materials[InIndex].Id;
+		}
+
+		return INDEX_NONE;
+	}
+
+	FMaterial* FSharedResourceManager::GetMaterial(int32 InIndex)
+	{
+		if (InIndex != INDEX_NONE && Materials.IsAllocated(InIndex))
+		{
+			return &Materials[InIndex].Resource;
+		}
+
+		return nullptr;
+	}
+
+	FRWLock& FSharedResourceManager::GetLockObject()
+	{
+		return ResourceLock;
+	}
+
+}
+
+#endif // WITH_PHYSX
\ No newline at end of file
diff --git a/Engine/Source/Runtime/Engine/Private/PhysicsEngine/PhysAnim.cpp b/Engine/Source/Runtime/Engine/Private/PhysicsEngine/PhysAnim.cpp
--- a/Engine/Source/Runtime/Engine/Private/PhysicsEngine/PhysAnim.cpp
+++ b/Engine/Source/Runtime/Engine/Private/PhysicsEngine/PhysAnim.cpp
@@ -588,7 +588,7 @@ void USkeletalMeshComponent::UpdateKinematicBonesToAnim(const TArray<FTransform>
 					FPhysicsActorHandle& ActorHandle = BodyInst->ActorHandle;
 					const bool bIsRigidBody = FPhysicsInterface::IsRigidBody(ActorHandle);
 
-					if(ActorHandle.IsValid() && bIsRigidBody && (bTeleport || !BodyInst->IsInstanceSimulatingPhysics()))	//If we have a body and it's kinematic, or we are teleporting a simulated body
+					if (FPhysicsInterface::IsValid(ActorHandle) && bIsRigidBody && (bTeleport || !BodyInst->IsInstanceSimulatingPhysics()))	//If we have a body and it's kinematic, or we are teleporting a simulated body
 					{
 						const int32 BoneIndex = BodyInst->InstanceBoneIndex;
 
diff --git a/Engine/Source/Runtime/Engine/Private/PhysicsEngine/PhysLevel.cpp b/Engine/Source/Runtime/Engine/Private/PhysicsEngine/PhysLevel.cpp
--- a/Engine/Source/Runtime/Engine/Private/PhysicsEngine/PhysLevel.cpp
+++ b/Engine/Source/Runtime/Engine/Private/PhysicsEngine/PhysLevel.cpp
@@ -266,7 +266,7 @@ void PostEngineInitialize()
 		// If the solver module is available, pass along our settings provider
 		// #BG - Collect all chaos modules settings into one provider?
 		ChaosModule->SetSettingsProvider(&GEngineChaosSettingsProvider);
-		ChaosModule->Initialize();
+		ChaosModule->OnSettingsChanged();
 	}
 #endif
 }
diff --git a/Engine/Source/Runtime/Engine/Private/PhysicsEngine/PhysScene_PhysX.cpp b/Engine/Source/Runtime/Engine/Private/PhysicsEngine/PhysScene_PhysX.cpp
--- a/Engine/Source/Runtime/Engine/Private/PhysicsEngine/PhysScene_PhysX.cpp
+++ b/Engine/Source/Runtime/Engine/Private/PhysicsEngine/PhysScene_PhysX.cpp
@@ -1,6 +1,6 @@
 // Copyright 1998-2019 Epic Games, Inc. All Rights Reserved.
 
-#if !WITH_CHAOS && !WITH_IMMEDIATE_PHYSX && !PHYSICS_INTERFACE_LLIMMEDIATE
+#if !WITH_CHAOS && !WITH_IMMEDIATE_PHYSX
 
 #include "Physics/PhysScene_PhysX.h"
 #include "Misc/CommandLine.h"
@@ -23,7 +23,10 @@
 
 #include "PhysicsInterfaceDeclaresCore.h"
 #if !WITH_CHAOS_NEEDS_TO_BE_FIXED
+
+#if INCLUDE_CHAOS
 #include "SQAccelerator.h"
+#endif
 
 #if WITH_PHYSX
 #include "PhysXPublic.h"
@@ -2298,4 +2301,4 @@ int32 FPhysScene::GetNumAwakeBodies()
 
 #endif // !WITH_CHAOS_NEEDS_TO_BE_FIXED
 
-#endif //  !WITH_CHAOS && !WITH_IMMEDIATE_PHYSX && !PHYSICS_INTERFACE_LLIMMEDIATE
+#endif //  !WITH_CHAOS && !WITH_IMMEDIATE_PHYSX
diff --git a/Engine/Source/Runtime/Engine/Private/PhysicsEngine/PhysSubstepTasks.cpp b/Engine/Source/Runtime/Engine/Private/PhysicsEngine/PhysSubstepTasks.cpp
--- a/Engine/Source/Runtime/Engine/Private/PhysicsEngine/PhysSubstepTasks.cpp
+++ b/Engine/Source/Runtime/Engine/Private/PhysicsEngine/PhysSubstepTasks.cpp
@@ -243,7 +243,7 @@ bool IsKinematicHelper(const PxRigidBody* PRigidBody)
 void FPhysSubstepTask::ApplyForces_AssumesLocked(const FPhysTarget& PhysTarget, FBodyInstance* BodyInstance)
 {
 #if WITH_PHYSX
-#if WITH_CHAOS || WITH_IMMEDIATE_PHYSX || PHYSICS_INTERFACE_LLIMMEDIATE
+#if WITH_CHAOS || WITH_IMMEDIATE_PHYSX
     check(false);
 #else
 	/** Apply Forces */
@@ -277,7 +277,7 @@ void FPhysSubstepTask::ApplyForces_AssumesLocked(const FPhysTarget& PhysTarget,
 void FPhysSubstepTask::ApplyTorques_AssumesLocked(const FPhysTarget& PhysTarget, FBodyInstance* BodyInstance)
 {
 #if WITH_PHYSX
-#if WITH_CHAOS || WITH_IMMEDIATE_PHYSX || PHYSICS_INTERFACE_LLIMMEDIATE
+#if WITH_CHAOS || WITH_IMMEDIATE_PHYSX
     check(false);
 #else
 	/** Apply Torques */
@@ -296,7 +296,7 @@ void FPhysSubstepTask::ApplyTorques_AssumesLocked(const FPhysTarget& PhysTarget,
 void FPhysSubstepTask::ApplyRadialForces_AssumesLocked(const FPhysTarget& PhysTarget, FBodyInstance* BodyInstance)
 {
 #if WITH_PHYSX
-#if WITH_CHAOS || WITH_IMMEDIATE_PHYSX || PHYSICS_INTERFACE_LLIMMEDIATE
+#if WITH_CHAOS || WITH_IMMEDIATE_PHYSX
     check(false);
 #else
 	/** Apply Torques */
@@ -316,7 +316,7 @@ void FPhysSubstepTask::ApplyRadialForces_AssumesLocked(const FPhysTarget& PhysTa
 void FPhysSubstepTask::InterpolateKinematicActor_AssumesLocked(const FPhysTarget& PhysTarget, FBodyInstance* BodyInstance, float InAlpha)
 {
 #if WITH_PHYSX
-#if WITH_CHAOS || WITH_IMMEDIATE_PHYSX || PHYSICS_INTERFACE_LLIMMEDIATE
+#if WITH_CHAOS || WITH_IMMEDIATE_PHYSX
     check(false);
 #else
 	PxRigidDynamic * PRigidDynamic = FPhysicsInterface_PhysX::GetPxRigidDynamic_AssumesLocked(BodyInstance->GetPhysicsActorHandle());
@@ -362,7 +362,7 @@ void FPhysSubstepTask::SubstepInterpolation(float InAlpha, float DeltaTime)
 	PxScene * PScene = PAScene;
 	SCOPED_SCENE_WRITE_LOCK(PScene);
 #endif
-#if WITH_CHAOS || WITH_IMMEDIATE_PHYSX || PHYSICS_INTERFACE_LLIMMEDIATE
+#if WITH_CHAOS || WITH_IMMEDIATE_PHYSX
     check(false);
 #else
 
diff --git a/Engine/Source/Runtime/Engine/Private/PhysicsEngine/PhysUtils.cpp b/Engine/Source/Runtime/Engine/Private/PhysicsEngine/PhysUtils.cpp
--- a/Engine/Source/Runtime/Engine/Private/PhysicsEngine/PhysUtils.cpp
+++ b/Engine/Source/Runtime/Engine/Private/PhysicsEngine/PhysUtils.cpp
@@ -261,7 +261,7 @@ FCollisionResponseContainer::FCollisionResponseContainer(ECollisionResponse Defa
 	SetAllChannels(DefaultResponse);
 }
 
-#if WITH_CHAOS || PHYSICS_INTERFACE_LLIMMEDIATE
+#if WITH_CHAOS
 bool FPhysScene::ExecPxVis(uint32 SceneType, const TCHAR* Cmd, FOutputDevice* Ar)
 {
     return false;
@@ -359,7 +359,7 @@ bool FPhysScene::ExecPxVis(const TCHAR* Cmd, FOutputDevice* Ar)
 }
 #endif
 
-#if WITH_CHAOS || PHYSICS_INTERFACE_LLIMMEDIATE
+#if WITH_CHAOS
 bool FPhysScene::ExecApexVis(uint32 SceneType, const TCHAR* Cmd, FOutputDevice* Ar)
 {
     return false;
@@ -467,7 +467,7 @@ bool FPhysScene::ExecApexVis(const TCHAR* Cmd, FOutputDevice* Ar)
 #endif
 
 
-#if WITH_CHAOS || PHYSICS_INTERFACE_LLIMMEDIATE
+#if WITH_CHAOS
 bool FPhysicsInterface::ExecPhysCommands(const TCHAR* Cmd, FOutputDevice* Ar, UWorld* InWorld)
 {
     return false;
diff --git a/Engine/Source/Runtime/Engine/Private/PhysicsEngine/PhysXSupport.cpp b/Engine/Source/Runtime/Engine/Private/PhysicsEngine/PhysXSupport.cpp
--- a/Engine/Source/Runtime/Engine/Private/PhysicsEngine/PhysXSupport.cpp
+++ b/Engine/Source/Runtime/Engine/Private/PhysicsEngine/PhysXSupport.cpp
@@ -163,7 +163,7 @@ PxFilterFlags PhysXSimFilterShader(	PxFilterObjectAttributes attributes0, PxFilt
 	}
 	
 	// if these bodies are from the same component, use the disable table to see if we should disable collision. This case should only happen for things like skeletalmesh and destruction. The table is only created for skeletal mesh components at the moment
-#if !WITH_CHAOS && !PHYSICS_INTERFACE_LLIMMEDIATE
+#if !WITH_CHAOS
 	if(filterData0.word2 == filterData1.word2)
 	{
 		check(constantBlockSize == sizeof(FPhysSceneShaderInfo));
@@ -228,7 +228,7 @@ PxFilterFlags PhysXSimFilterShader(	PxFilterObjectAttributes attributes0, PxFilt
 	return PxFilterFlags();
 }
 
-#if !WITH_CHAOS && !PHYSICS_INTERFACE_LLIMMEDIATE
+#if !WITH_CHAOS
 
 /** Figures out the new FCollisionNotifyInfo needed for pending notification. It adds it, and then returns an array that goes from pair index to notify collision index */
 TArray<int32> AddCollisionNotifyInfo(const FBodyInstance* Body0, const FBodyInstance* Body1, const physx::PxContactPair * Pairs, uint32 NumPairs, TArray<FCollisionNotifyInfo> & PendingNotifyInfos)
@@ -569,7 +569,7 @@ void AddToCollection(PxCollection* PCollection, PxBase* PBase)
 
 PxCollection* MakePhysXCollection(const TArray<UPhysicalMaterial*>& PhysicalMaterials, const TArray<UBodySetup*>& BodySetups, uint64 BaseId)
 {
-#if WITH_CHAOS || WITH_IMMEDIATE_PHYSX || PHYSICS_INTERFACE_LLIMMEDIATE
+#if WITH_CHAOS || WITH_IMMEDIATE_PHYSX
     ensure(false);
     return nullptr;
 #else
diff --git a/Engine/Source/Runtime/Engine/Private/PhysicsEngine/PhysXSupport.h b/Engine/Source/Runtime/Engine/Private/PhysicsEngine/PhysXSupport.h
--- a/Engine/Source/Runtime/Engine/Private/PhysicsEngine/PhysXSupport.h
+++ b/Engine/Source/Runtime/Engine/Private/PhysicsEngine/PhysXSupport.h
@@ -132,7 +132,7 @@ PxFilterFlags PhysXSimFilterShader(	PxFilterObjectAttributes attributes0, PxFilt
 									PxFilterObjectAttributes attributes1, PxFilterData filterData1,
 									PxPairFlags& pairFlags, const void* constantBlock, PxU32 constantBlockSize );
 
-#if !WITH_CHAOS && !PHYSICS_INTERFACE_LLIMMEDIATE
+#if !WITH_CHAOS
 /** Event callback used to notify engine about various collision events */
 class ENGINE_API FPhysXSimEventCallback : public PxSimulationEventCallback
 {
diff --git a/Engine/Source/Runtime/Engine/Private/PhysicsEngine/PhysicalAnimationComponent.cpp b/Engine/Source/Runtime/Engine/Private/PhysicsEngine/PhysicalAnimationComponent.cpp
--- a/Engine/Source/Runtime/Engine/Private/PhysicsEngine/PhysicalAnimationComponent.cpp
+++ b/Engine/Source/Runtime/Engine/Private/PhysicsEngine/PhysicalAnimationComponent.cpp
@@ -375,7 +375,7 @@ void UPhysicalAnimationComponent::UpdatePhysicsEngineImp()
 					int32 ChildBodyIdx = PhysAsset->FindBodyIndex(PhysAnimData.BodyName);
 					if (FBodyInstance* ChildBody = (ChildBodyIdx == INDEX_NONE ? nullptr : SkeletalMeshComponent->Bodies[ChildBodyIdx]))
 					{
-#if WITH_CHAOS || WITH_IMMEDIATE_PHYSX || PHYSICS_INTERFACE_LLIMMEDIATE
+#if WITH_CHAOS || WITH_IMMEDIATE_PHYSX
                         ensure(false);
 #else
 						if (PxRigidActor* PRigidActor = FPhysicsInterface_PhysX::GetPxRigidActor_AssumesLocked(ChildBody->ActorHandle))
diff --git a/Engine/Source/Runtime/Engine/Private/PhysicsEngine/PhysicsGeometry.cpp b/Engine/Source/Runtime/Engine/Private/PhysicsEngine/PhysicsGeometry.cpp
new file mode 100644
--- /dev/null
+++ b/Engine/Source/Runtime/Engine/Private/PhysicsEngine/PhysicsGeometry.cpp
@@ -0,0 +1,307 @@
+// Copyright 1998-2019 Epic Games, Inc. All Rights Reserved.
+
+#if WITH_PHYSX && PHYSICS_INTERFACE_PHYSX
+
+#include "Physics/PhysicsGeometry.h"
+
+#include "PhysicsEngine/AggregateGeom.h"
+
+//#include "PhysXPublic.h"
+#include "PhysicsEngine/PhysXSupport.h"
+#include "PhysicsEngine/PhysicsSettings.h"
+#include "Engine/Engine.h"
+#include "PhysicalMaterials/PhysicalMaterial.h"
+#include "PhysicsEngine/BodySetup.h"
+#include "PxSphereGeometry.h"
+
+using namespace physx;
+
+extern TAutoConsoleVariable<float> CVarContactOffsetFactor;
+extern TAutoConsoleVariable<float> CVarMaxContactOffset;
+
+void FBodySetupShapeIterator::GetContactOffsetParams(float& InOutContactOffsetFactor, float& InOutMinContactOffset, float& InOutMaxContactOffset)
+{
+	// Get contact offset params
+	InOutContactOffsetFactor = CVarContactOffsetFactor.GetValueOnAnyThread();
+	InOutMaxContactOffset = CVarMaxContactOffset.GetValueOnAnyThread();
+
+	InOutContactOffsetFactor = InOutContactOffsetFactor < 0.f ? UPhysicsSettings::Get()->ContactOffsetMultiplier : InOutContactOffsetFactor;
+	InOutMaxContactOffset = InOutMaxContactOffset < 0.f ? UPhysicsSettings::Get()->MaxContactOffset : InOutMaxContactOffset;
+
+	InOutMinContactOffset = UPhysicsSettings::Get()->MinContactOffset;
+}
+
+FBodySetupShapeIterator::FBodySetupShapeIterator(const FVector& InScale3D, const FTransform& InRelativeTM, bool bDoubleSidedTrimeshes)
+	: Scale3D(InScale3D)
+	, RelativeTM(InRelativeTM)
+	, bDoubleSidedTriMeshGeo(bDoubleSidedTrimeshes)
+{
+	SetupNonUniformHelper(Scale3D, MinScale, MinScaleAbs, ShapeScale3DAbs);
+	{
+		float MinScaleRelative;
+		float MinScaleAbsRelative;
+		FVector Scale3DAbsRelative;
+		FVector Scale3DRelative = RelativeTM.GetScale3D();
+
+		SetupNonUniformHelper(Scale3DRelative, MinScaleRelative, MinScaleAbsRelative, Scale3DAbsRelative);
+
+		MinScaleAbs *= MinScaleAbsRelative;
+		ShapeScale3DAbs.X *= Scale3DAbsRelative.X;
+		ShapeScale3DAbs.Y *= Scale3DAbsRelative.Y;
+		ShapeScale3DAbs.Z *= Scale3DAbsRelative.Z;
+
+		ShapeScale3D = Scale3D;
+		ShapeScale3D.X *= Scale3DAbsRelative.X;
+		ShapeScale3D.Y *= Scale3DAbsRelative.Y;
+		ShapeScale3D.Z *= Scale3DAbsRelative.Z;
+	}
+
+	GetContactOffsetParams(ContactOffsetFactor, MinContactOffset, MaxContactOffset);
+}
+
+FBodySetupShapeIterator::FBodySetupShapeIterator()
+	: Scale3D(FVector::OneVector)
+	, RelativeTM(FTransform::Identity)
+	, bDoubleSidedTriMeshGeo(false)
+{
+
+}
+
+template <typename ElemType, typename GeomType>
+void FBodySetupShapeIterator::ForEachShape(const TArrayView<ElemType>& Elements, TFunctionRef<void(const ElemType& Elem, const GeomType& Geom, const FBodySetupShapeIterator::TTransform& LocalPose, float ContactOffset, float RestOffset)> VisitorFunc) const
+{
+	for(int32 ElemIdx = 0; ElemIdx < Elements.Num(); ElemIdx++)
+	{
+		const ElemType& Elem = Elements[ElemIdx];
+		GeomType Geom;
+		PxTransform PLocalPose;
+
+		if(PopulateGeometryAndTransform(Elem, Geom, PLocalPose))
+		{
+			const float RestOffset = ComputeRestOffset(Elem);
+			const float ContactOffset = ComputeContactOffset(Geom) + RestOffset;
+			VisitorFunc(Elem, Geom, PLocalPose, ContactOffset, RestOffset);
+		}
+		else
+		{
+			UE_LOG(LogPhysics, Warning, TEXT("ForeachShape(%s): ScaledElem[%d] invalid"), *GetDebugName<ElemType>(), ElemIdx);
+		}
+	}
+}
+
+
+//////////////////////// Sphere elements ////////////////////////
+template <> bool FBodySetupShapeIterator::PopulateGeometryAndTransform(const FKSphereElem& SphereElem, PxSphereGeometry& OutGeometry, FBodySetupShapeIterator::TTransform& OutTM) const
+{
+	const FKSphereElem ScaledSphereElem = SphereElem.GetFinalScaled(Scale3D, RelativeTM);
+	OutGeometry.radius = FMath::Max(ScaledSphereElem.Radius, KINDA_SMALL_NUMBER);
+
+	if(ensure(OutGeometry.isValid()))
+	{
+		OutTM = PxTransform(U2PVector(ScaledSphereElem.Center));
+		return true;
+	}
+	else
+	{
+		return false;
+	}
+}
+
+template <> float FBodySetupShapeIterator::ComputeContactOffset(const PxSphereGeometry& PSphereGeom) const
+{
+	return FMath::Clamp(ContactOffsetFactor * PSphereGeom.radius, MinContactOffset, MaxContactOffset);
+}
+
+template <> FString FBodySetupShapeIterator::GetDebugName<FKSphereElem>()  const
+{
+	return TEXT("Sphere");
+}
+
+template <typename ElemType> float FBodySetupShapeIterator::ComputeRestOffset(const ElemType& Elem) const
+{
+	return Elem.RestOffset;
+}
+
+template <> float FBodySetupShapeIterator::ComputeRestOffset(physx::PxTriangleMesh* const&) const
+{
+	return 0.f;
+}
+
+/////////////////// Box elements //////////////////////////////
+template <> bool FBodySetupShapeIterator::PopulateGeometryAndTransform(const FKBoxElem& BoxElem, PxBoxGeometry& OutGeometry, FBodySetupShapeIterator::TTransform& OutTM) const
+{
+	const FKBoxElem ScaledBoxElem = BoxElem.GetFinalScaled(Scale3D, RelativeTM);
+	const FTransform& BoxTransform = ScaledBoxElem.GetTransform();
+
+	OutGeometry.halfExtents.x = FMath::Max(ScaledBoxElem.X * 0.5f, KINDA_SMALL_NUMBER);
+	OutGeometry.halfExtents.y = FMath::Max(ScaledBoxElem.Y * 0.5f, KINDA_SMALL_NUMBER);
+	OutGeometry.halfExtents.z = FMath::Max(ScaledBoxElem.Z * 0.5f, KINDA_SMALL_NUMBER);
+
+	if(OutGeometry.isValid() && BoxTransform.IsValid())
+	{
+		OutTM = U2PTransform(BoxTransform);
+		if(ensure(OutTM.isValid()))
+		{
+			return true;
+		}
+	}
+
+	return false;
+}
+
+template <> float FBodySetupShapeIterator::ComputeContactOffset(const PxBoxGeometry& PBoxGeom) const
+{
+	return FMath::Clamp(ContactOffsetFactor * PBoxGeom.halfExtents.minElement(), MinContactOffset, MaxContactOffset);
+}
+
+template <> FString FBodySetupShapeIterator::GetDebugName<FKBoxElem>()  const
+{
+	return TEXT("Box");
+}
+
+/////////////////////// Capsule elements /////////////////////////////
+template <> bool FBodySetupShapeIterator::PopulateGeometryAndTransform(const FKSphylElem& SphylElem, PxCapsuleGeometry& OutGeometry, FBodySetupShapeIterator::TTransform& OutTM) const
+{
+	const FKSphylElem ScaledSphylElem = SphylElem.GetFinalScaled(Scale3D, RelativeTM);
+
+	OutGeometry.halfHeight = FMath::Max(ScaledSphylElem.Length * 0.5f, KINDA_SMALL_NUMBER);
+	OutGeometry.radius = FMath::Max(ScaledSphylElem.Radius, KINDA_SMALL_NUMBER);
+
+	if(OutGeometry.isValid())
+	{
+		// The stored capsule transform assumes the capsule axis is down Z. In PhysX, it points down X, so we twiddle the matrix a bit here (swap X and Z and negate Y).
+		OutTM = PxTransform(U2PVector(ScaledSphylElem.Center), U2PQuat(ScaledSphylElem.Rotation.Quaternion()) * U2PSphylBasis);
+
+		if(ensure(OutTM.isValid()))
+		{
+			return true;
+		}
+	}
+
+	return false;
+}
+
+template <> float FBodySetupShapeIterator::ComputeContactOffset(const PxCapsuleGeometry& PCapsuleGeom) const
+{
+	return FMath::Clamp(ContactOffsetFactor * PCapsuleGeom.radius, MinContactOffset, MaxContactOffset);
+}
+
+template <> FString FBodySetupShapeIterator::GetDebugName<FKSphylElem>() const
+{
+	return TEXT("Capsule");
+}
+
+////////////////////////////// Convex elements ////////////////////////////
+template <> bool FBodySetupShapeIterator::PopulateGeometryAndTransform(const FKConvexElem& ConvexElem, PxConvexMeshGeometry& OutGeometry, FBodySetupShapeIterator::TTransform& OutTM) const
+{
+	FTransform NegativeScaleCompensation;
+	const bool bUseNegX = CalcMeshNegScaleCompensation(Scale3D * RelativeTM.GetScale3D(), NegativeScaleCompensation);
+	OutTM = U2PTransform(NegativeScaleCompensation);
+
+	PxConvexMesh* UseConvexMesh = bUseNegX ? ConvexElem.GetMirroredConvexMesh() : ConvexElem.GetConvexMesh();
+	if(UseConvexMesh)
+	{
+		OutGeometry.convexMesh = UseConvexMesh;
+		OutGeometry.scale.scale = U2PVector(ShapeScale3DAbs);	//scale shape about the origin
+
+																													//Scale the position independent of shape scale. This is because physx transforms have no concept of scale
+		PxTransform PElementTransform = U2PTransform(RelativeTM);
+		OutTM.q *= PElementTransform.q;
+		OutTM.p = PElementTransform.p;
+		OutTM.p.x *= Scale3D.X;
+		OutTM.p.y *= Scale3D.Y;
+		OutTM.p.z *= Scale3D.Z;
+
+		if(OutGeometry.isValid())
+		{
+			PxVec3 PBoundsExtents = OutGeometry.convexMesh->getLocalBounds().getExtents();
+
+			if(ensure(OutTM.isValid()))
+			{
+				return true;
+			}
+			else
+			{
+				UE_LOG(LogPhysics, Warning, TEXT("PopulatePhysXGeometryAndTransform(Convex): ConvexElem invalid"));
+			}
+		}
+		else
+		{
+			UE_LOG(LogPhysics, Warning, TEXT("PopulatePhysXGeometryAndTransform(Convex): ConvexElem has invalid transform"));
+		}
+	}
+	else
+	{
+		UE_LOG(LogPhysics, Warning, TEXT("PopulatePhysXGeometryAndTransform(Convex): ConvexElem is missing ConvexMesh"));
+	}
+
+	return false;
+}
+
+template <> float FBodySetupShapeIterator::ComputeContactOffset(const PxConvexMeshGeometry& PConvexGeom) const
+{
+	PxVec3 PBoundsExtents = PConvexGeom.convexMesh->getLocalBounds().getExtents();
+	return FMath::Clamp(ContactOffsetFactor * PBoundsExtents.minElement(), MinContactOffset, MaxContactOffset);
+}
+
+template <> FString FBodySetupShapeIterator::GetDebugName<FKConvexElem>() const
+{
+	return TEXT("Convex");
+}
+
+
+///////////////////////////////////Trimesh elements ////////////////////////////
+template <> bool FBodySetupShapeIterator::PopulateGeometryAndTransform(PxTriangleMesh* const & TriMesh, PxTriangleMeshGeometry& OutGeometry, FBodySetupShapeIterator::TTransform& OutTM) const
+{
+	OutGeometry.triangleMesh = TriMesh;
+	OutGeometry.scale.scale = U2PVector(ShapeScale3D); //scale shape about the origin
+
+	auto ClampScale = [](float& Val) -> void
+	{
+		Val = Val <= 0.f ? FMath::Min(Val, -KINDA_SMALL_NUMBER) : FMath::Max(KINDA_SMALL_NUMBER, Val);
+	};
+
+	ClampScale(OutGeometry.scale.scale.x);
+	ClampScale(OutGeometry.scale.scale.y);
+	ClampScale(OutGeometry.scale.scale.z);
+
+	if(bDoubleSidedTriMeshGeo)
+	{
+		OutGeometry.meshFlags |= PxMeshGeometryFlag::eDOUBLE_SIDED;
+	}
+
+	if(OutGeometry.isValid())
+	{
+		//Scale the position independent of shape scale. This is because physx transforms have no concept of scale
+		OutTM = U2PTransform(RelativeTM);
+		OutTM.p.x *= Scale3D.X;
+		OutTM.p.y *= Scale3D.Y;
+		OutTM.p.z *= Scale3D.Z;
+
+		return true;
+	}
+	else
+	{
+		UE_LOG(LogPhysics, Log, TEXT("PopulatePhysXGeometryAndTransform(TriMesh): TriMesh invalid"));
+	}
+
+	return false;
+}
+
+template <> float FBodySetupShapeIterator::ComputeContactOffset(const PxTriangleMeshGeometry& PTriMeshGeom) const
+{
+	return MaxContactOffset;
+}
+
+template <> FString FBodySetupShapeIterator::GetDebugName<PxTriangleMesh*>() const
+{
+	return TEXT("Trimesh");
+}
+
+template void FBodySetupShapeIterator::ForEachShape(const TArrayView<FKSphereElem>&, TFunctionRef<void(const FKSphereElem&, const physx::PxSphereGeometry&, const physx::PxTransform&, float, float)>) const;
+template void FBodySetupShapeIterator::ForEachShape(const TArrayView<FKBoxElem>&, TFunctionRef<void(const FKBoxElem&, const physx::PxBoxGeometry&, const physx::PxTransform&, float, float)>) const;
+template void FBodySetupShapeIterator::ForEachShape(const TArrayView<FKSphylElem>&, TFunctionRef<void(const FKSphylElem&, const physx::PxCapsuleGeometry&, const physx::PxTransform&, float, float)>) const;
+template void FBodySetupShapeIterator::ForEachShape(const TArrayView<FKConvexElem>&, TFunctionRef<void(const FKConvexElem&, const physx::PxConvexMeshGeometry&, const physx::PxTransform&, float, float)>) const;
+template void FBodySetupShapeIterator::ForEachShape(const TArrayView<physx::PxTriangleMesh*>&, TFunctionRef<void(physx::PxTriangleMesh* const &, const physx::PxTriangleMeshGeometry&, const physx::PxTransform&, float, float)>) const;
+
+#endif
\ No newline at end of file
diff --git a/Engine/Source/Runtime/Engine/Private/PhysicsEngine/PhysicsHandleComponent.cpp b/Engine/Source/Runtime/Engine/Private/PhysicsEngine/PhysicsHandleComponent.cpp
--- a/Engine/Source/Runtime/Engine/Private/PhysicsEngine/PhysicsHandleComponent.cpp
+++ b/Engine/Source/Runtime/Engine/Private/PhysicsEngine/PhysicsHandleComponent.cpp
@@ -80,7 +80,7 @@ void UPhysicsHandleComponent::GrabComponent(class UPrimitiveComponent* InCompone
 	FRotator GrabbedRotation = FRotator::ZeroRotator;
 
 #if WITH_PHYSX
-	if(BodyInstance->ActorHandle.IsValid())
+	if(FPhysicsInterface::IsValid(BodyInstance->ActorHandle))
 	{
 		FPhysicsCommand::ExecuteRead(BodyInstance->ActorHandle, [&](const FPhysicsActorHandle& Actor)
 		{
diff --git a/Engine/Source/Runtime/Engine/Private/PhysicsEngine/PhysicsInterfacePhysX.cpp b/Engine/Source/Runtime/Engine/Private/PhysicsEngine/PhysicsInterfacePhysX.cpp
--- a/Engine/Source/Runtime/Engine/Private/PhysicsEngine/PhysicsInterfacePhysX.cpp
+++ b/Engine/Source/Runtime/Engine/Private/PhysicsEngine/PhysicsInterfacePhysX.cpp
@@ -1,6 +1,6 @@
 // Copyright 1998-2019 Epic Games, Inc. All Rights Reserved.
 
-#if !WITH_CHAOS && !WITH_IMMEDIATE_PHYSX && !PHYSICS_INTERFACE_LLIMMEDIATE
+#if !WITH_CHAOS && !WITH_IMMEDIATE_PHYSX
 
 #include "Physics/PhysicsInterfacePhysX.h"
 #include "Physics/PhysicsInterfaceUtils.h"
@@ -13,7 +13,10 @@
 #include "PhysicsInterfaceDeclaresCore.h"
 
 #if !WITH_CHAOS_NEEDS_TO_BE_FIXED
+
+#if INCLUDE_CHAOS
 #include "SQAccelerator.h"
+#endif
 
 #if WITH_PHYSX
 #include "PhysXPublic.h"
@@ -27,7 +30,7 @@
 
 #include "PhysicsEngine/ConstraintDrives.h"
 #include "PhysicsEngine/AggregateGeom.h"
-#include "Physics/PhysicsGeometryPhysX.h"
+#include "Physics/PhysicsGeometry.h"
 #include "Engine/Engine.h"
 #include "PhysicalMaterials/PhysicalMaterial.h"
 #include "PhysicsEngine/PhysicsSettings.h"
@@ -76,6 +79,16 @@ private:
 	EPhysicsInterfaceScopedLockType LockType;
 };
 
+FPhysicsActorHandle_PhysX DefaultPhysicsActorHandle()
+{
+	return FPhysicsActorHandle_PhysX();
+}
+
+bool IsValidHandle(const FPhysicsActorHandle_PhysX& Handle)
+{
+	return Handle.IsValid();
+}
+
 PxRigidActor* FPhysicsInterface_PhysX::GetPxRigidActor_AssumesLocked(const FPhysicsActorHandle_PhysX& InRef)
 {
 	return InRef.SyncActor;
@@ -2521,7 +2534,7 @@ bool FPhysicsInterface_PhysX::Sweep_Geom(FHitResult& OutHit, const FBodyInstance
 
 			if((RigidBody != nullptr) && (RigidBody->getNbShapes() != 0) && (InInstance->OwnerComponent != nullptr))
 			{
-				FPhysXShapeAdaptor ShapeAdaptor(InShapeRotation, InShape);
+				FPhysXShapeAdapter ShapeAdapter(InShapeRotation, InShape);
 				
 				const FVector Delta = InEnd - InStart;
 				const float DeltaMag = Delta.Size();
@@ -2530,7 +2543,7 @@ bool FPhysicsInterface_PhysX::Sweep_Geom(FHitResult& OutHit, const FBodyInstance
 					PxHitFlags POutputFlags = PxHitFlag::ePOSITION | PxHitFlag::eNORMAL | PxHitFlag::eDISTANCE | PxHitFlag::eFACE_INDEX | PxHitFlag::eMTD;
 
 					UPrimitiveComponent* OwnerComponentInst = InInstance->OwnerComponent.Get();
-					PxTransform PStartTM(U2PVector(InStart), U2PQuat(ShapeAdaptor.GetGeomOrientation()));
+					PxTransform PStartTM(U2PVector(InStart), U2PQuat(ShapeAdapter.GetGeomOrientation()));
 					PxTransform PCompTM(U2PTransform(OwnerComponentInst->GetComponentTransform()));
 
 					PxVec3 PDir = U2PVector(Delta / DeltaMag);
@@ -2562,7 +2575,7 @@ bool FPhysicsInterface_PhysX::Sweep_Geom(FHitResult& OutHit, const FBodyInstance
 						if((bSweepComplex && bShapeIsComplex) || (!bSweepComplex && bShapeIsSimple))
 						{
 							PxTransform PGlobalPose = PCompTM.transform(PShape->getLocalPose());
-							const PxGeometry& Geometry = ShapeAdaptor.GetGeometry();
+							const PxGeometry& Geometry = ShapeAdapter.GetGeometry();
 							if(PxGeometryQuery::sweep(PDir, DeltaMag, Geometry, PStartTM, PShape->getGeometry().any(), PGlobalPose, PHit, POutputFlags))
 							{
 								// we just like to make sure if the hit is made
@@ -2651,14 +2664,14 @@ bool Overlap_GeomInternal(const FBodyInstance* InInstance, const PxGeometry& InP
 
 bool FPhysicsInterface_PhysX::Overlap_Geom(const FBodyInstance* InInstance, const FPhysicsGeometryCollection& InGeometry, const FTransform& InShapeTransform, FMTDResult* OutOptResult)
 {
-	PxGeometry&  PGeom = InGeometry.GetGeometry();
+	const PxGeometry&  PGeom = InGeometry.GetGeometry();
 
 	return Overlap_GeomInternal(InInstance, PGeom, InShapeTransform, OutOptResult);
 }
 
 bool FPhysicsInterface_PhysX::Overlap_Geom(const FBodyInstance* InInstance, const FCollisionShape& InCollisionShape, const FQuat& InShapeRotation, const FTransform& InShapeTransform, FMTDResult* OutOptResult /*= nullptr*/)
 {
-	FPhysXShapeAdaptor Adaptor(InShapeRotation, InCollisionShape);
+	FPhysXShapeAdapter Adaptor(InShapeRotation, InCollisionShape);
 
 	return Overlap_GeomInternal(InInstance, Adaptor.GetGeometry(), Adaptor.GetGeomPose(InShapeTransform.GetTranslation()), OutOptResult);
 }
@@ -2777,39 +2790,64 @@ ECollisionShapeType FPhysicsGeometryCollection_PhysX::GetType() const
 	return P2UCollisionShapeType(GeomHolder->getType());
 }
 
-physx::PxGeometry& FPhysicsGeometryCollection_PhysX::GetGeometry() const
+const physx::PxGeometry& FPhysicsGeometryCollection_PhysX::GetGeometry() const
 {
 	return GeomHolder->any();
 }
 
-bool FPhysicsGeometryCollection_PhysX::GetBoxGeometry(physx::PxBoxGeometry& OutGeom) const
+physx::PxGeometry& FPhysicsGeometryCollection_PhysX::GetGeometry()
 {
-	OutGeom = GeomHolder->box();
-	return true;
+	return GeomHolder->any();
 }
 
-bool FPhysicsGeometryCollection_PhysX::GetSphereGeometry(physx::PxSphereGeometry& OutGeom) const
+const physx::PxBoxGeometry& FPhysicsGeometryCollection_PhysX::GetBoxGeometry() const
 {
-	OutGeom = GeomHolder->sphere();
-	return true;
+	return GeomHolder->box();
 }
 
-bool FPhysicsGeometryCollection_PhysX::GetCapsuleGeometry(physx::PxCapsuleGeometry& OutGeom) const
+physx::PxBoxGeometry& FPhysicsGeometryCollection_PhysX::GetBoxGeometry()
 {
-	OutGeom = GeomHolder->capsule();
-	return true;
+	return GeomHolder->box();
 }
 
-bool FPhysicsGeometryCollection_PhysX::GetConvexGeometry(physx::PxConvexMeshGeometry& OutGeom) const
+const physx::PxSphereGeometry& FPhysicsGeometryCollection_PhysX::GetSphereGeometry() const
 {
-	OutGeom = GeomHolder->convexMesh();
-	return true;
+	return GeomHolder->sphere();
 }
 
-bool FPhysicsGeometryCollection_PhysX::GetTriMeshGeometry(physx::PxTriangleMeshGeometry& OutGeom) const
+physx::PxSphereGeometry& FPhysicsGeometryCollection_PhysX::GetSphereGeometry()
 {
-	OutGeom = GeomHolder->triangleMesh();
-	return true;
+	return GeomHolder->sphere();
+}
+
+const physx::PxCapsuleGeometry& FPhysicsGeometryCollection_PhysX::GetCapsuleGeometry() const
+{
+	return GeomHolder->capsule();
+}
+
+physx::PxCapsuleGeometry& FPhysicsGeometryCollection_PhysX::GetCapsuleGeometry()
+{
+	return GeomHolder->capsule();
+}
+
+const physx::PxConvexMeshGeometry& FPhysicsGeometryCollection_PhysX::GetConvexGeometry() const
+{
+	return GeomHolder->convexMesh();
+}
+
+physx::PxConvexMeshGeometry& FPhysicsGeometryCollection_PhysX::GetConvexGeometry()
+{
+	return GeomHolder->convexMesh();
+}
+
+const physx::PxTriangleMeshGeometry& FPhysicsGeometryCollection_PhysX::GetTriMeshGeometry() const
+{
+	return GeomHolder->triangleMesh();
+}
+
+physx::PxTriangleMeshGeometry& FPhysicsGeometryCollection_PhysX::GetTriMeshGeometry()
+{
+	return GeomHolder->triangleMesh();
 }
 
 FPhysicsGeometryCollection_PhysX::FPhysicsGeometryCollection_PhysX(const FPhysicsShapeHandle_PhysX& ShapeRef)
diff --git a/Engine/Source/Runtime/Engine/Private/PhysicsEngine/PhysicsInterfaceUtils.cpp b/Engine/Source/Runtime/Engine/Private/PhysicsEngine/PhysicsInterfaceUtils.cpp
--- a/Engine/Source/Runtime/Engine/Private/PhysicsEngine/PhysicsInterfaceUtils.cpp
+++ b/Engine/Source/Runtime/Engine/Private/PhysicsEngine/PhysicsInterfaceUtils.cpp
@@ -93,10 +93,13 @@ FCollisionFilterData CreateTraceQueryFilterData(const uint8 MyChannel, const boo
 #define TRACE_MULTI		1
 #define TRACE_SINGLE	0
 
+#endif
+
 /** Utility for creating a PhysX PxFilterData for performing a query (trace) against the scene */
 FCollisionFilterData CreateQueryFilterData(const uint8 MyChannel, const bool bTraceComplex, const FCollisionResponseContainer& InCollisionResponseContainer, const struct FCollisionQueryParams& QueryParam, const struct FCollisionObjectQueryParams & ObjectParam, const bool bMultitrace)
 {
-	if (ObjectParam.IsValid())
+#if WITH_PHYSX
+	if(ObjectParam.IsValid())
 	{
 		return CreateObjectQueryFilterData(bTraceComplex, (bMultitrace ? TRACE_MULTI : TRACE_SINGLE), ObjectParam);
 	}
@@ -104,6 +107,9 @@ FCollisionFilterData CreateQueryFilterData(const uint8 MyChannel, const bool bTr
 	{
 		return CreateTraceQueryFilterData(MyChannel, bTraceComplex, InCollisionResponseContainer, QueryParam);
 	}
+#else 
+	//#TODO implement chaos
+	ensure(false);
+	return FCollisionFilterData();
+#endif
 }
-
-#endif
\ No newline at end of file
diff --git a/Engine/Source/Runtime/Engine/Private/PhysicsEngine/ScopedSQHitchRepeater.h b/Engine/Source/Runtime/Engine/Private/PhysicsEngine/ScopedSQHitchRepeater.h
--- a/Engine/Source/Runtime/Engine/Private/PhysicsEngine/ScopedSQHitchRepeater.h
+++ b/Engine/Source/Runtime/Engine/Private/PhysicsEngine/ScopedSQHitchRepeater.h
@@ -13,8 +13,6 @@
 
 #if PHYSICS_INTERFACE_PHYSX
 #include "PhysXInterfaceWrapper.h"
-#elif PHYSICS_INTERFACE_LLIMMEDIATE
-#include "Physics/Experimental/LLImmediateInterfaceWrapper.h"
 #endif
 
 #if DETECT_SQ_HITCHES
@@ -31,8 +29,6 @@ struct FSQHitchRepeaterCVars
 };
 #endif 
 
-#if WITH_PHYSX
-
 /** Various info we want to capture for hitch detection reporting */
 struct FHitchDetectionInfo
 {
@@ -177,5 +173,3 @@ struct FScopedSQHitchRepeater
 	bool RepeatOnHitch() const { return false; }
 #endif
 };
-
-#endif // WITH_PHYSX
\ No newline at end of file
diff --git a/Engine/Source/Runtime/Engine/Private/SkeletalMeshComponentPhysics.cpp b/Engine/Source/Runtime/Engine/Private/SkeletalMeshComponentPhysics.cpp
--- a/Engine/Source/Runtime/Engine/Private/SkeletalMeshComponentPhysics.cpp
+++ b/Engine/Source/Runtime/Engine/Private/SkeletalMeshComponentPhysics.cpp
@@ -24,6 +24,7 @@
 #include "Modules/ModuleManager.h"
 #include "Rendering/SkeletalMeshRenderData.h"
 #include "Physics/PhysicsInterfaceCore.h"
+#include "AnimationRuntime.h"
 
 #include "Logging/MessageLog.h"
 #include "CollisionDebugDrawingPublic.h"
@@ -636,18 +637,16 @@ void USkeletalMeshComponent::InstantiatePhysicsAssetRefPose(const UPhysicsAsset&
 {
 	if(SkeletalMesh)
 	{
-		auto BoneTMCallable = [this](int32 BoneIndex)
-		{
-			const FReferenceSkeleton& RefSkel = SkeletalMesh->RefSkeleton;
+		// @todo(ccaulfield): check perf here against uncached version: FAnimationRuntime::GetComponentSpaceTransform for RBAN scenarios
+		const FReferenceSkeleton& RefSkeleton = SkeletalMesh->RefSkeleton;
+		TArray<FTransform> CachedTransforms;
+		TArray<bool> CachedTransformReady;
+		CachedTransforms.SetNumUninitialized(RefSkeleton.GetRefBonePose().Num());
+		CachedTransformReady.SetNumZeroed(RefSkeleton.GetRefBonePose().Num());
 
-			if(RefSkel.IsValidIndex(BoneIndex))
-			{
-				return RefSkel.GetRefBonePose()[BoneIndex];
-			}
-			else
-			{
-				return FTransform::Identity;
-			}
+		auto BoneTMCallable = [&](int32 BoneIndex)
+		{
+			return FAnimationRuntime::GetComponentSpaceTransformWithCache(RefSkeleton, RefSkeleton.GetRefBonePose(), BoneIndex, CachedTransforms, CachedTransformReady);
 		};
 
 		InstantiatePhysicsAsset_Internal(PhysAsset, Scale3D, OutBodies, OutConstraints, BoneTMCallable, PhysScene, OwningComponent, UseRootBodyIndex, UseAggregate);
@@ -915,14 +914,14 @@ void USkeletalMeshComponent::TermBodiesBelow(FName ParentBoneName)
 
 float USkeletalMeshComponent::GetTotalMassBelowBone(FName InBoneName)
 {
-	float TotalMass = 0.f;
+	float TotalBodyMass = 0.f;
 
-	ForEachBodyBelow(InBoneName, /*bIncludeSelf=*/true, /*bSkipCustomPhysics=*/false, [&TotalMass](FBodyInstance* BI)
+	ForEachBodyBelow(InBoneName, /*bIncludeSelf=*/true, /*bSkipCustomPhysics=*/false, [&TotalBodyMass](FBodyInstance* BI)
 	{
-		TotalMass += BI->GetBodyMass();
+		TotalBodyMass += BI->GetBodyMass();
 	});
 
-	return TotalMass;
+	return TotalBodyMass;
 }
 
 void USkeletalMeshComponent::SetAllBodiesSimulatePhysics(bool bNewSimulate)
@@ -2460,13 +2459,21 @@ void USkeletalMeshComponent::FindClothCollisions(FClothCollisionData& OutCollisi
 	}
 }
 
-void USkeletalMeshComponent::ExtractCollisionsForCloth(USkeletalMeshComponent* SourceComponent, UPhysicsAsset* PhysicsAsset, USkeletalMeshComponent* DestClothComponent, FClothCollisionData& OutCollisions, FClothCollisionSource& ClothCollisionSource)
+void USkeletalMeshComponent::ExtractCollisionsForCloth(
+	USkeletalMeshComponent* SourceComponent, 
+	UPhysicsAsset* PhysicsAsset, 
+	USkeletalMeshComponent* DestClothComponent, 
+	FClothCollisionData& OutCollisions, 
+	FClothCollisionSource& ClothCollisionSource)
 {
-	// Extract collisions from this mesh 'raw', as this isnt a mesh that has cloth simulation
+	// Extract collisions from this mesh 'raw', as this isn't a mesh that has cloth simulation
 	// (but we want it to affect other meshes with cloth simulation)
 	if(SourceComponent->SkeletalMesh && PhysicsAsset)
 	{
-		const FTransform& ComponentToComponentTransform = SourceComponent->GetComponentTransform() * DestClothComponent->GetComponentTransform().Inverse();
+		const FTransform& ComponentToComponentTransform = 
+			SourceComponent != DestClothComponent ? 
+				SourceComponent->GetComponentTransform() * DestClothComponent->GetComponentTransform().Inverse() : 
+				FTransform::Identity;
 
 		// Init cache on first copy
 		if(!ClothCollisionSource.bCached || ClothCollisionSource.CachedSkeletalMesh.Get() != SourceComponent->SkeletalMesh)
@@ -2698,8 +2705,7 @@ void USkeletalMeshComponent::ProcessClothCollisionWithEnvironment()
 
 							case ECollisionShapeType::Sphere:
 							{
-								PxSphereGeometry SphereGeo;
-								GeoCollection.GetSphereGeometry(SphereGeo);
+								PxSphereGeometry& SphereGeo = GeoCollection.GetSphereGeometry();
 
 								NewCollisionData.Spheres.AddDefaulted();
 								FClothCollisionPrim_Sphere& NewSphere = NewCollisionData.Spheres.Last();
@@ -2712,8 +2718,7 @@ void USkeletalMeshComponent::ProcessClothCollisionWithEnvironment()
 
 							case ECollisionShapeType::Capsule:
 							{
-								PxCapsuleGeometry CapGeo;
-								GeoCollection.GetCapsuleGeometry(CapGeo);
+								PxCapsuleGeometry& CapGeo = GeoCollection.GetCapsuleGeometry();
 
 								const int32 BaseSphereIndex = NewCollisionData.Spheres.Num();
 
@@ -2741,8 +2746,7 @@ void USkeletalMeshComponent::ProcessClothCollisionWithEnvironment()
 
 							case ECollisionShapeType::Box:
 							{
-								PxBoxGeometry BoxGeo;
-								GeoCollection.GetBoxGeometry(BoxGeo);
+								PxBoxGeometry& BoxGeo = GeoCollection.GetBoxGeometry();
 
 								// We're building the box in local space, so to get to the cloth transform
 								// we need to go through local -> actor -> world -> cloth
@@ -2786,8 +2790,7 @@ void USkeletalMeshComponent::ProcessClothCollisionWithEnvironment()
 
 							case ECollisionShapeType::Convex:
 							{
-								PxConvexMeshGeometry MeshGeo;
-								GeoCollection.GetConvexGeometry(MeshGeo);
+								PxConvexMeshGeometry& MeshGeo = GeoCollection.GetConvexGeometry();
 
 								// we need to inflate the hull to get nicer collisions (only particles collide)
 								const static float Inflate = 2.0f;
@@ -3286,19 +3289,19 @@ float USkeletalMeshComponent::GetBoneMass(FName BoneName, bool bScaleMass) const
 FVector USkeletalMeshComponent::GetSkeletalCenterOfMass() const
 {
 	FVector Location = FVector::ZeroVector;
-	float TotalMass = 0.0f;
+	float Mass = 0.0f;
 	for (int32 i = 0; i < Bodies.Num(); i++)
 	{
 		FBodyInstance* BI = Bodies[i];
 		if (BI->IsValidBodyInstance())
 		{
-			float Mass = BI->MassScale*BI->GetBodyMass();
-			Location += Mass*BI->GetCOMPosition();
-			TotalMass += Mass;
+			float BodyMass = BI->MassScale*BI->GetBodyMass();
+			Location += BodyMass*BI->GetCOMPosition();
+			Mass += BodyMass;
 		}
 	}
 
-	return Location / TotalMass;
+	return Location / Mass;
 }
 
 void USkeletalMeshComponent::SetAllUseCCD(bool InUseCCD)
diff --git a/Engine/Source/Runtime/Engine/Public/AnimationRuntime.h b/Engine/Source/Runtime/Engine/Public/AnimationRuntime.h
--- a/Engine/Source/Runtime/Engine/Public/AnimationRuntime.h
+++ b/Engine/Source/Runtime/Engine/Public/AnimationRuntime.h
@@ -431,6 +431,18 @@ public:
 	static void FillUpComponentSpaceTransforms(const FReferenceSkeleton& RefSkeleton, const TArray<FTransform> &BoneSpaceTransforms, TArray<FTransform> &ComponentSpaceTransforms);
 	static void MakeSkeletonRefPoseFromMesh(const USkeletalMesh* InMesh, const USkeleton* InSkeleton, TArray<FTransform>& OutBoneBuffer);
 
+	/**
+	 * Calculate the component-space bone transform for the specified bone. This is similar to GetComponentSpaceTransform(), but uses a cached array of 
+	 * transforms to prevent re-calculations when requesting transforms for multiple bones with shared parents. This is only likely to be more efficient than
+	 * FillUpComponentSpaceTransforms if you only need a relatively small number of bone transforms.
+	 * @param	InRefSkeleton			The bone hierarchy
+	 * @param	InBoneSpaceTransforms	Bone-space transforms for all bones in the hierarchy
+	 * @param	BoneIndex				We calculate the component-space transform of this bone
+	 * @param	CachedTransforms		An array of transforms which holds the transforms of any bones whose transforms have been previously calculated. Should be initialized with CachedTransforms.SetNumUninitialized(BoneSpaceTransforms.Num()) before first use.
+	 * @param	CachedTransformReady	An array of flags indicating which bone transforms have been cached. Should be initialized with CachedTransformReady.SetNumZeroed(BoneSpaceTransforms.Num()) before first use.
+	 */
+	static const FTransform& GetComponentSpaceTransformWithCache(const FReferenceSkeleton& InRefSkeleton, const TArray<FTransform> &InBoneSpaceTransforms, int32 BoneIndex, TArray<FTransform>& CachedTransforms, TArray<bool>& CachedTransformReady);
+
 #if WITH_EDITOR
 	static void FillUpComponentSpaceTransformsRefPose(const USkeleton* Skeleton, TArray<FTransform> &ComponentSpaceTransforms);
 	static void FillUpComponentSpaceTransformsRetargetBasePose(const USkeleton* Skeleton, TArray<FTransform> &ComponentSpaceTransforms);
diff --git a/Engine/Source/Runtime/Engine/Public/Collision.h b/Engine/Source/Runtime/Engine/Public/Collision.h
--- a/Engine/Source/Runtime/Engine/Public/Collision.h
+++ b/Engine/Source/Runtime/Engine/Public/Collision.h
@@ -30,7 +30,7 @@ DECLARE_CYCLE_STAT_EXTERN(TEXT("BodyInstanceLineTrace"), STAT_Collision_FBodyIns
 DECLARE_CYCLE_STAT_EXTERN(TEXT("PreFilter"), STAT_Collision_PreFilter, STATGROUP_Collision, );
 
 /** Enable collision analyzer support */
-#if (1 && !(UE_BUILD_SHIPPING || UE_BUILD_TEST) && WITH_EDITOR && WITH_UNREAL_DEVELOPER_TOOLS && WITH_PHYSX && PHYSICS_INTERFACE_PHYSX)
+#if (1 && !(UE_BUILD_SHIPPING || UE_BUILD_TEST) && WITH_EDITOR && WITH_UNREAL_DEVELOPER_TOOLS)
 	#define ENABLE_COLLISION_ANALYZER 1
 #else
 	#define ENABLE_COLLISION_ANALYZER 0
diff --git a/Engine/Source/Runtime/Engine/Public/PhysXInterfaceWrapper.h b/Engine/Source/Runtime/Engine/Public/PhysXInterfaceWrapper.h
--- a/Engine/Source/Runtime/Engine/Public/PhysXInterfaceWrapper.h
+++ b/Engine/Source/Runtime/Engine/Public/PhysXInterfaceWrapper.h
@@ -1,13 +1,16 @@
 // Copyright 1998-2019 Epic Games, Inc. All Rights Reserved.
 #pragma once
 
-#if WITH_PHYSX && PHYSICS_INTERFACE_PHYSX
+#if WITH_PHYSX
 #include "PhysXInterfaceWrapperCore.h"
 #include "PhysXPublic.h"
 #include "Physics/PhysicsInterfaceUtils.h"
 #include "CustomPhysXPayload.h"
 #include "PhysicsInterfaceWrapperShared.h"
 
+namespace PhysXInterface
+{
+#if PHYSICS_INTERFACE_PHYSX
 /**
 * Helper to lock/unlock a scene that also makes sure to unlock everything when it goes out of scope.
 * Multiple locks on the same scene are NOT SAFE. You can't call LockRead() if already locked.
@@ -28,6 +31,7 @@ struct FScopedSceneReadLock
 
 	PxScene* SceneLock;
 };
+#endif
 
 inline PxQueryFlags StaticDynamicQueryFlags(const FCollisionQueryParams& Params)
 {
@@ -63,5 +67,6 @@ T* GetUserData(const PxShape& Shape)
 {
 	return FPhysxUserData::Get<T>(Shape.userData);
 }
+}
 
 #endif
\ No newline at end of file
diff --git a/Engine/Source/Runtime/Engine/Public/Physics/Experimental/ChaosDerivedDataReader.h b/Engine/Source/Runtime/Engine/Public/Physics/Experimental/ChaosDerivedDataReader.h
new file mode 100644
--- /dev/null
+++ b/Engine/Source/Runtime/Engine/Public/Physics/Experimental/ChaosDerivedDataReader.h
@@ -0,0 +1,45 @@
+// Copyright 1998-2019 Epic Games, Inc. All Rights Reserved.
+
+#pragma once
+
+#if INCLUDE_CHAOS
+
+#include "Templates/UniquePtr.h"
+#include "PhysicsEngine/BodySetup.h"
+namespace Chaos
+{
+	template<typename T, int d>
+	class TImplicitObject;
+
+	template <typename T>
+	class TTriangleMeshImplicitObject;
+}
+
+struct FUntypedBulkData;
+
+template<typename T, int d>
+class FChaosDerivedDataReader
+{
+public:
+
+	// Only valid use is to explicitly read chaos bulk data
+	explicit FChaosDerivedDataReader(FUntypedBulkData* InBulkData);
+
+	TArray<TUniquePtr<Chaos::TImplicitObject<T, d>>> ConvexImplicitObjects;
+	TArray<TUniquePtr<Chaos::TTriangleMeshImplicitObject<T>>> TrimeshImplicitObjects;
+	FBodySetupUVInfo UVInfo;
+
+private:
+	FChaosDerivedDataReader() = delete;
+	FChaosDerivedDataReader(const FChaosDerivedDataReader& Other) = delete;
+	FChaosDerivedDataReader(FChaosDerivedDataReader&& Other) = delete;
+	FChaosDerivedDataReader& operator =(const FChaosDerivedDataReader& Other) = delete;
+	FChaosDerivedDataReader& operator =(FChaosDerivedDataReader&& Other) = delete;
+
+	bool bReadSuccessful;
+
+};
+
+extern template class FChaosDerivedDataReader<float, 3>;
+
+#endif
\ No newline at end of file
diff --git a/Engine/Source/Runtime/Engine/Public/Physics/Experimental/ChaosInterfaceWrapper.h b/Engine/Source/Runtime/Engine/Public/Physics/Experimental/ChaosInterfaceWrapper.h
--- a/Engine/Source/Runtime/Engine/Public/Physics/Experimental/ChaosInterfaceWrapper.h
+++ b/Engine/Source/Runtime/Engine/Public/Physics/Experimental/ChaosInterfaceWrapper.h
@@ -1,11 +1,20 @@
 // Copyright 1998-2019 Epic Games, Inc. All Rights Reserved.
 #pragma once
 
+#if INCLUDE_CHAOS
 #include "ChaosInterfaceWrapperCore.h"
 #include "CollisionQueryFilterCallbackCore.h"
+#include "Chaos/ParticleHandle.h"
+#include "PhysxUserData.h"
+//todo: move this include into an impl header
 
 class UPhysicalMaterial;
 struct FBodyInstance;
+struct FCollisionFilterData;
+struct FCollisionQueryParams;
+
+namespace ChaosInterface
+{
 
 // Needed by low level SQ calls. Right now there's no specific locking for LLI
 // #PHYS2 update as locking becomes necessary
@@ -17,7 +26,11 @@ struct FScopedSceneReadLock
 
 inline FQueryFilterData MakeQueryFilterData(const FCollisionFilterData& FilterData, EQueryFlags QueryFlags, const FCollisionQueryParams& Params)
 {
+#if WITH_PHYSX
+	return PxQueryFilterData(U2PFilterData(FilterData), U2PQueryFlags(QueryFlags));
+#else
 	return FQueryFilterData();
+#endif
 }
 
 inline UPhysicalMaterial* GetUserData(const FPhysTypeDummy& Material)
@@ -25,12 +38,15 @@ inline UPhysicalMaterial* GetUserData(const FPhysTypeDummy& Material)
 	return nullptr;
 }
 
-inline FBodyInstance* GetUserData(const FPhysActorDummy& Actor)
+inline FBodyInstance* GetUserData(const Chaos::TGeometryParticle<float,3>& Actor)
 {
-	return nullptr;
+	void* UserData = Actor.UserData();
+	return UserData ? FPhysxUserData::Get<FBodyInstance>(Actor.UserData()) : nullptr;
 }
 
+}
+void LowLevelRaycast(FPhysScene& Scene, const FVector& Start, const FVector& Dir, float DeltaMag, FPhysicsHitCallback<FHitRaycast>& HitBuffer, EHitFlags OutputFlags, FQueryFlags QueryFlags, const FCollisionFilterData& Filter, const FQueryFilterData& QueryFilterData, ICollisionQueryFilterCallbackBase* QueryCallback);
+void LowLevelSweep(FPhysScene& Scene, const FPhysicsGeometry& Geom, const FTransform& StartTM, const FVector& Dir, float DeltaMag, FPhysicsHitCallback<FHitSweep>& HitBuffer, EHitFlags OutputFlags, FQueryFlags QueryFlags, const FCollisionFilterData& Filter, const FQueryFilterData& QueryFilterData, ICollisionQueryFilterCallbackBase* QueryCallback);
+void LowLevelOverlap(FPhysScene& Scene, const FPhysicsGeometry& Geom, const FTransform& GeomPose, FPhysicsHitCallback<FHitOverlap>& HitBuffer, FQueryFlags QueryFlags, const FCollisionFilterData& Filter, const FQueryFilterData& QueryFilterData, ICollisionQueryFilterCallbackBase* QueryCallback);
 
-inline void LowLevelRaycast(FPhysScene& Scene, const FVector& Start, const FVector& Dir, float DeltaMag, FPhysicsHitCallback<FHitRaycast>& HitBuffer, EHitFlags OutputFlags, FQueryFlags QueryFlags, const FCollisionFilterData& Filter, const FQueryFilterData& QueryFilterData, ICollisionQueryFilterCallbackBase* QueryCallback) {}
-inline void LowLevelSweep(FPhysScene& Scene, const FPhysicsGeometry& Geom, const FTransform& StartTM, const FVector& Dir, float DeltaMag, FPhysicsHitCallback<FHitSweep>& HitBuffer, EHitFlags OutputFlags, FQueryFlags QueryFlags, const FCollisionFilterData& Filter, const FQueryFilterData& QueryFilterData, ICollisionQueryFilterCallbackBase* QueryCallback) {}
-inline void LowLevelOverlap(FPhysScene& Scene, const FPhysicsGeometry& Geom, const FTransform& GeomPose, FPhysicsHitCallback<FHitOverlap>& HitBuffer, FQueryFlags QueryFlags, const FCollisionFilterData& Filter, const FQueryFilterData& QueryFilterData, ICollisionQueryFilterCallbackBase* QueryCallback) {}
\ No newline at end of file
+#endif // INCLUDE_CHAOS
\ No newline at end of file
