Copying //UE4/Dev-Physics to Dev-Main (//UE4/Dev-Main) [at] 7703071

diff --git a/Engine/Source/Runtime/Engine/Public/Physics/Experimental/PhysInterface_Chaos.h b/Engine/Source/Runtime/Engine/Public/Physics/Experimental/PhysInterface_Chaos.h
--- a/Engine/Source/Runtime/Engine/Public/Physics/Experimental/PhysInterface_Chaos.h
+++ b/Engine/Source/Runtime/Engine/Public/Physics/Experimental/PhysInterface_Chaos.h
@@ -4,22 +4,24 @@
 
 #if WITH_CHAOS
 
+#include "ChaosInterfaceWrapper.h"
 #include "Engine/Engine.h"
-#include "Physics/PhysicsInterfaceDeclares.h"
-#include "Physics/PhysicsInterfaceTypes.h"
-#include "PhysicsEngine/ConstraintDrives.h"
-#include "PhysicsEngine/ConstraintTypes.h"
-#include "PhysicsPublic.h"
-#include "PhysicsReplication.h"
 #include "Chaos/ArrayCollection.h"
 #include "Chaos/ArrayCollectionArray.h"
+#include "Chaos/Declares.h"
 #include "Chaos/Capsule.h"
 #include "Chaos/ImplicitObject.h"
 #include "Chaos/Pair.h"
 #include "Chaos/Transform.h"
-#include "Physics/GenericPhysicsInterface.h"
+#include "PhysicsEngine/ConstraintDrives.h"
+#include "PhysicsEngine/ConstraintTypes.h"
+#include "PhysicsPublic.h"
+#include "PhysicsReplication.h"
 #include "PhysicsInterfaceWrapperShared.h"
-#include "ChaosInterfaceWrapper.h"
+#include "Physics/PhysicsInterfaceDeclares.h"
+#include "Physics/PhysicsInterfaceTypes.h"
+#include "Physics/GenericPhysicsInterface.h"
+#include "Physics/Experimental/PhysicsUserData_Chaos.h"
 
 //NOTE: Do not include Chaos headers directly as it means recompiling all of engine. This should be reworked to avoid allocations
 
@@ -31,12 +33,12 @@ struct FBodyInstance;
 struct FPhysxUserData;
 class IPhysicsReplicationFactory;
 
-class FBodyInstancePhysicsObject;
-
 class AWorldSettings;
 
 namespace Chaos
 {
+	class IDispatcher;
+
 	template <typename T, int>
 	class TBVHParticles;
 
@@ -51,40 +53,10 @@ namespace Chaos
 
 	template <typename T, int>
 	class TPBDSpringConstraints;
-}
-
-class FPhysicsActorReference_Chaos 
-{
-public:
-
-	FPhysicsActorReference_Chaos( 
-		FBodyInstance * InBodyInstance = nullptr,
-		FPhysScene_ChaosInterface* InScene = nullptr, 
-		FBodyInstancePhysicsObject* InPhysicsObject = nullptr)
-		: BodyInstance(InBodyInstance)
-		, Scene(InScene)
-		, PhysicsObject(InPhysicsObject){}
-
-	bool IsValid() const { return PhysicsObject != nullptr; }
-	bool Equals(const FPhysicsActorReference_Chaos& Handle) const { return PhysicsObject == Handle.PhysicsObject; }
 
-	FPhysScene_ChaosInterface* GetScene() { return Scene; }
-	const FPhysScene_ChaosInterface* GetScene() const { return Scene; }
-	void SetScene(FPhysScene_ChaosInterface* InScene) { Scene = InScene; }
-
-	FBodyInstancePhysicsObject* GetPhysicsObject() { return PhysicsObject; }
-	/*const*/ FBodyInstancePhysicsObject* GetPhysicsObject() const { return PhysicsObject; }
-	void SetPhysicsObject(FBodyInstancePhysicsObject* InPhysicsObject) { PhysicsObject = InPhysicsObject; }
-
-	FBodyInstance* GetBodyInstance() { return BodyInstance; }
-	const FBodyInstance* GetBodyInstance() const { return BodyInstance; }
-	void SetBodyInstance(FBodyInstance * InBodyInstance) { BodyInstance = InBodyInstance; }
-
-private:
-	FBodyInstance* BodyInstance;
-	FPhysScene_ChaosInterface* Scene;
-	FBodyInstancePhysicsObject* PhysicsObject;
-};
+	template <typename T, int>
+	class TConvex;
+}
 
 class FPhysicsConstraintReference_Chaos
 {
@@ -101,35 +73,81 @@ public:
 class FPhysicsShapeReference_Chaos
 {
 public:
-	typedef Chaos::TImplicitObject<float, 3> Internal;
-
-
+	
 	FPhysicsShapeReference_Chaos()
-		: Object(nullptr), bSimulation(false), bQuery(false), ActorRef() {}
-	FPhysicsShapeReference_Chaos(Internal* ObjectIn, bool bSimulationIn, bool bQueryIn, FPhysicsActorReference_Chaos ActorRefIn)
-		: Object(ObjectIn),bSimulation(bSimulationIn),bQuery(bQueryIn),ActorRef(ActorRefIn){}
+		: Shape(nullptr), bSimulation(false), bQuery(false), ActorRef() {}
+	FPhysicsShapeReference_Chaos(Chaos::TPerShapeData<float, 3>* ShapeIn, bool bSimulationIn, bool bQueryIn, FPhysicsActorHandle ActorRefIn)
+		: Shape(ShapeIn),bSimulation(bSimulationIn),bQuery(bQueryIn),ActorRef(ActorRefIn){}
 	FPhysicsShapeReference_Chaos(const FPhysicsShapeReference_Chaos& Other)
-		: Object(Other.Object)
+		: Shape(Other.Shape)
 		, bSimulation(Other.bSimulation)
 		, bQuery(Other.bQuery)
 		, ActorRef(Other.ActorRef){}
 
 
-	bool IsValid() const { return (Object != nullptr); }
-	bool Equals(const FPhysicsShapeReference_Chaos& Other) const { return Object == Other.Object; }
+	bool IsValid() const { return (Shape != nullptr); }
+	bool Equals(const FPhysicsShapeReference_Chaos& Other) const { return Shape == Other.Shape; }
     bool operator==(const FPhysicsShapeReference_Chaos& Other) const { return Equals(Other); }
+	
+	const Chaos::TImplicitObject<float, 3>& GetGeometry() const { check(IsValid()); return *Shape->Geometry; }
+
 
-	Internal* Object;
+	Chaos::TPerShapeData<float, 3>* Shape;
 	bool bSimulation;
 	bool bQuery;
-    FPhysicsActorReference_Chaos ActorRef;
+    FPhysicsActorHandle ActorRef;
+};
+
+class FPhysicsShapeAdapter_Chaos
+{
+public:
+	FPhysicsShapeAdapter_Chaos(const FQuat& Rot, const FCollisionShape& CollisionShape);
+
+	const FPhysicsGeometry& GetGeometry() const;
+	FTransform GetGeometryPose(const FVector& Pos) const;
+	const FQuat& GetGeomOrientation() const;
+
+private:
+	TUniquePtr<FPhysicsGeometry> Geometry;
+	FQuat GeometryRotation;
 };
 
 FORCEINLINE uint32 GetTypeHash(const FPhysicsShapeReference_Chaos& InShapeReference)
 {
-    return PointerHash(InShapeReference.Object);
+    return PointerHash(InShapeReference.Shape);
 }
 
+/**
+ Wrapper around geometry. This is really just needed to make the physx chaos abstraction easier
+ */
+struct ENGINE_API FPhysicsGeometryCollection_Chaos
+{
+	// Delete default constructor, want only construction by interface (private constructor below)
+	FPhysicsGeometryCollection_Chaos() = delete;
+	// No copying or assignment, move construction only, these are defaulted in the source file as they need
+	// to be able to delete physx::PxGeometryHolder which is incomplete here
+	FPhysicsGeometryCollection_Chaos(const FPhysicsGeometryCollection_Chaos& Copy) = delete;
+	FPhysicsGeometryCollection_Chaos& operator=(const FPhysicsGeometryCollection_Chaos& Copy) = delete;
+	FPhysicsGeometryCollection_Chaos(FPhysicsGeometryCollection_Chaos&& Steal);
+	FPhysicsGeometryCollection_Chaos& operator=(FPhysicsGeometryCollection_Chaos&& Steal) = delete;
+	~FPhysicsGeometryCollection_Chaos();
+
+	ECollisionShapeType GetType() const;
+	const Chaos::TImplicitObject<float, 3>& GetGeometry() const;
+	const Chaos::TBox<float, 3>& GetBoxGeometry() const;
+	const Chaos::TSphere<float, 3>&  GetSphereGeometry() const;
+	const Chaos::TCapsule<float>&  GetCapsuleGeometry() const;
+	const Chaos::TConvex<float, 3>& GetConvexGeometry() const;
+	const Chaos::TTriangleMeshImplicitObject<float>& GetTriMeshGeometry() const;
+
+private:
+	friend class FPhysInterface_Chaos;
+	explicit FPhysicsGeometryCollection_Chaos(const FPhysicsShapeReference_Chaos& InShape);
+
+	const Chaos::TImplicitObject<float, 3>& Geom;
+};
+
+
 // Temp interface
 namespace physx
 {
@@ -144,20 +162,24 @@ namespace physx
 struct FContactModifyCallback;
 class ULineBatchComponent;
 
-class FPhysInterface_Chaos : public FGenericPhysicsInterface
+class ENGINE_API FPhysInterface_Chaos : public FGenericPhysicsInterface
 {
 public:
-    ENGINE_API FPhysInterface_Chaos(const AWorldSettings* Settings=nullptr);
-    ENGINE_API ~FPhysInterface_Chaos();
+    FPhysInterface_Chaos(const AWorldSettings* Settings=nullptr);
+    ~FPhysInterface_Chaos();
 
     // Interface needed for interface
-	static ENGINE_API void CreateActor(const FActorCreationParams& InParams, FPhysicsActorHandle& Handle);
-	static ENGINE_API void ReleaseActor(FPhysicsActorHandle& InActorReference, FPhysScene* InScene = nullptr, bool bNeverDeferRelease=false);
-
-	static ENGINE_API FPhysicsAggregateReference_Chaos CreateAggregate(int32 MaxBodies);
-	static ENGINE_API void ReleaseAggregate(FPhysicsAggregateReference_Chaos& InAggregate);
-	static ENGINE_API int32 GetNumActorsInAggregate(const FPhysicsAggregateReference_Chaos& InAggregate);
-	static ENGINE_API void AddActorToAggregate_AssumesLocked(const FPhysicsAggregateReference_Chaos& InAggregate, const FPhysicsActorReference_Chaos& InActor);
+	static void CreateActor(const FActorCreationParams& InParams, FPhysicsActorHandle& Handle);
+	static void ReleaseActor(FPhysicsActorHandle& InActorReference, FPhysScene* InScene = nullptr, bool bNeverDeferRelease=false);
+	static bool IsValid(const FPhysicsActorHandle& Handle) { return Handle != nullptr; }
+	static void AddActorToSolver(FPhysicsActorHandle& Handle, Chaos::FPhysicsSolver* Solver, Chaos::IDispatcher* Dispatcher);
+	static void RemoveActorFromSolver(FPhysicsActorHandle& Handle, Chaos::FPhysicsSolver* Solver, Chaos::IDispatcher* Dispatcher);
+	static const FBodyInstance* ShapeToOriginalBodyInstance(const FBodyInstance* InCurrentInstance, const Chaos::TPerShapeData<float, 3>* InShape);
+
+	static FPhysicsAggregateReference_Chaos CreateAggregate(int32 MaxBodies);
+	static void ReleaseAggregate(FPhysicsAggregateReference_Chaos& InAggregate);
+	static int32 GetNumActorsInAggregate(const FPhysicsAggregateReference_Chaos& InAggregate);
+	static void AddActorToAggregate_AssumesLocked(const FPhysicsAggregateReference_Chaos& InAggregate, const FPhysicsActorHandle& InActor);
 
 	// Material interface functions
     // @todo(mlentine): How do we set material on the solver?
@@ -168,7 +190,7 @@ public:
 
 	// Actor interface functions
 	template<typename AllocatorType>
-	static ENGINE_API int32 GetAllShapes_AssumedLocked(const FPhysicsActorReference_Chaos& InActorReference, TArray<FPhysicsShapeHandle, AllocatorType>& OutShapes);
+	static int32 GetAllShapes_AssumedLocked(const FPhysicsActorHandle& InActorReference, TArray<FPhysicsShapeHandle, AllocatorType>& OutShapes);
 	static int32 GetNumShapes(const FPhysicsActorHandle& InHandle);
 
 	static void ReleaseShape(const FPhysicsShapeHandle& InShape);
@@ -176,167 +198,159 @@ public:
 	static void AttachShape(const FPhysicsActorHandle& InActor, const FPhysicsShapeHandle& InNewShape);
 	static void DetachShape(const FPhysicsActorHandle& InActor, FPhysicsShapeHandle& InShape, bool bWakeTouching = true);
 
-	static ENGINE_API void SetActorUserData_AssumesLocked(FPhysicsActorReference_Chaos& InActorReference, FPhysxUserData* InUserData);
+	static void SetActorUserData_AssumesLocked(FPhysicsActorHandle& InActorReference, FPhysicsUserData* InUserData);
 
-	static ENGINE_API bool IsRigidBody(const FPhysicsActorReference_Chaos& InActorReference);
-	static ENGINE_API bool IsDynamic(const FPhysicsActorReference_Chaos& InActorReference)
+	static bool IsRigidBody(const FPhysicsActorHandle& InActorReference);
+	static bool IsDynamic(const FPhysicsActorHandle& InActorReference)
     {
         return !IsStatic(InActorReference);
     }
-    static ENGINE_API bool IsStatic(const FPhysicsActorReference_Chaos& InActorReference);
-	static ENGINE_API bool IsKinematic_AssumesLocked(const FPhysicsActorReference_Chaos& InActorReference);
-	static ENGINE_API bool IsSleeping(const FPhysicsActorReference_Chaos& InActorReference);
-	static ENGINE_API bool IsCcdEnabled(const FPhysicsActorReference_Chaos& InActorReference);
+    static bool IsStatic(const FPhysicsActorHandle& InActorReference);
+    static bool IsKinematic(const FPhysicsActorHandle& InActorReference);
+	static bool IsKinematic_AssumesLocked(const FPhysicsActorHandle& InActorReference);
+	static bool IsSleeping(const FPhysicsActorHandle& InActorReference);
+	static bool IsCcdEnabled(const FPhysicsActorHandle& InActorReference);
     // @todo(mlentine): We don't have a notion of sync vs async and are a bit of both. Does this work?
-    static bool HasSyncSceneData(const FPhysicsActorReference_Chaos& InHandle) { return true; }
-    static bool HasAsyncSceneData(const FPhysicsActorReference_Chaos& InHandle) { return false; }
-	static ENGINE_API bool IsInScene(const FPhysicsActorReference_Chaos& InActorReference);
-	static ENGINE_API bool CanSimulate_AssumesLocked(const FPhysicsActorReference_Chaos& InActorReference);
-	static ENGINE_API float GetMass_AssumesLocked(const FPhysicsActorReference_Chaos& InActorReference);
+    static bool HasSyncSceneData(const FPhysicsActorHandle& InHandle) { return true; }
+    static bool HasAsyncSceneData(const FPhysicsActorHandle& InHandle) { return false; }
+	static bool IsInScene(const FPhysicsActorHandle& InActorReference);
+	static FPhysScene* GetCurrentScene(const FPhysicsActorHandle& InHandle);
+	static bool CanSimulate_AssumesLocked(const FPhysicsActorHandle& InActorReference);
+	static float GetMass_AssumesLocked(const FPhysicsActorHandle& InActorReference);
 
-	static ENGINE_API void SetSendsSleepNotifies_AssumesLocked(const FPhysicsActorReference_Chaos& InActorReference, bool bSendSleepNotifies);
-	static ENGINE_API void PutToSleep_AssumesLocked(const FPhysicsActorReference_Chaos& InActorReference);
-	static ENGINE_API void WakeUp_AssumesLocked(const FPhysicsActorReference_Chaos& InActorReference);
+	static void SetSendsSleepNotifies_AssumesLocked(const FPhysicsActorHandle& InActorReference, bool bSendSleepNotifies);
+	static void PutToSleep_AssumesLocked(const FPhysicsActorHandle& InActorReference);
+	static void WakeUp_AssumesLocked(const FPhysicsActorHandle& InActorReference);
 
-	static ENGINE_API void SetIsKinematic_AssumesLocked(const FPhysicsActorReference_Chaos& InActorReference, bool bIsKinematic);
-	static ENGINE_API void SetCcdEnabled_AssumesLocked(const FPhysicsActorReference_Chaos& InActorReference, bool bIsCcdEnabled);
+	static void SetIsKinematic_AssumesLocked(const FPhysicsActorHandle& InActorReference, bool bIsKinematic);
+	static void SetCcdEnabled_AssumesLocked(const FPhysicsActorHandle& InActorReference, bool bIsCcdEnabled);
 
-	static ENGINE_API FTransform GetGlobalPose_AssumesLocked(const FPhysicsActorReference_Chaos& InActorReference);
-	static ENGINE_API void SetGlobalPose_AssumesLocked(const FPhysicsActorReference_Chaos& InActorReference, const FTransform& InNewPose, bool bAutoWake = true);
+	static FTransform GetGlobalPose_AssumesLocked(const FPhysicsActorHandle& InActorReference);
+	static void SetGlobalPose_AssumesLocked(const FPhysicsActorHandle& InActorReference, const FTransform& InNewPose, bool bAutoWake = true);
 
-    static ENGINE_API FTransform GetTransform_AssumesLocked(const FPhysicsActorHandle& InRef, bool bForceGlobalPose = false);
+    static FTransform GetTransform_AssumesLocked(const FPhysicsActorHandle& InRef, bool bForceGlobalPose = false);
 
-	static ENGINE_API bool HasKinematicTarget_AssumesLocked(const FPhysicsActorReference_Chaos& InActorReference);
-	static ENGINE_API FTransform GetKinematicTarget_AssumesLocked(const FPhysicsActorReference_Chaos& InActorReference);
-	static ENGINE_API void SetKinematicTarget_AssumesLocked(const FPhysicsActorReference_Chaos& InActorReference, const FTransform& InNewTarget);
+	static bool HasKinematicTarget_AssumesLocked(const FPhysicsActorHandle& InActorReference);
+	static FTransform GetKinematicTarget_AssumesLocked(const FPhysicsActorHandle& InActorReference);
+	static void SetKinematicTarget_AssumesLocked(const FPhysicsActorHandle& InActorReference, const FTransform& InNewTarget);
 
-	static ENGINE_API FVector GetLinearVelocity_AssumesLocked(const FPhysicsActorReference_Chaos& InActorReference);
-	static ENGINE_API void SetLinearVelocity_AssumesLocked(const FPhysicsActorReference_Chaos& InActorReference, const FVector& InNewVelocity, bool bAutoWake = true);
+	static FVector GetLinearVelocity_AssumesLocked(const FPhysicsActorHandle& InActorReference);
+	static void SetLinearVelocity_AssumesLocked(const FPhysicsActorHandle& InActorReference, const FVector& InNewVelocity, bool bAutoWake = true);
 
-	static ENGINE_API FVector GetAngularVelocity_AssumesLocked(const FPhysicsActorReference_Chaos& InActorReference);
-	static ENGINE_API void SetAngularVelocity_AssumesLocked(const FPhysicsActorReference_Chaos& InActorReference, const FVector& InNewVelocity, bool bAutoWake = true);
-	static ENGINE_API float GetMaxAngularVelocity_AssumesLocked(const FPhysicsActorReference_Chaos& InActorReference);
-	static ENGINE_API void SetMaxAngularVelocity_AssumesLocked(const FPhysicsActorReference_Chaos& InActorReference, float InMaxAngularVelocity);
+	static FVector GetAngularVelocity_AssumesLocked(const FPhysicsActorHandle& InActorReference);
+	static void SetAngularVelocity_AssumesLocked(const FPhysicsActorHandle& InActorReference, const FVector& InNewVelocity, bool bAutoWake = true);
+	static float GetMaxAngularVelocity_AssumesLocked(const FPhysicsActorHandle& InActorReference);
+	static void SetMaxAngularVelocity_AssumesLocked(const FPhysicsActorHandle& InActorReference, float InMaxAngularVelocity);
 
-	static ENGINE_API float GetMaxDepenetrationVelocity_AssumesLocked(const FPhysicsActorReference_Chaos& InActorReference);
-	static ENGINE_API void SetMaxDepenetrationVelocity_AssumesLocked(const FPhysicsActorReference_Chaos& InActorReference, float InMaxDepenetrationVelocity);
+	static float GetMaxDepenetrationVelocity_AssumesLocked(const FPhysicsActorHandle& InActorReference);
+	static void SetMaxDepenetrationVelocity_AssumesLocked(const FPhysicsActorHandle& InActorReference, float InMaxDepenetrationVelocity);
 
-	static ENGINE_API FVector GetWorldVelocityAtPoint_AssumesLocked(const FPhysicsActorReference_Chaos& InActorReference, const FVector& InPoint);
+	static FVector GetWorldVelocityAtPoint_AssumesLocked(const FPhysicsActorHandle& InActorReference, const FVector& InPoint);
 
-	static ENGINE_API FTransform GetComTransform_AssumesLocked(const FPhysicsActorReference_Chaos& InActorReference);
-	static ENGINE_API FTransform GetComTransformLocal_AssumesLocked(const FPhysicsActorReference_Chaos& InActorReference);
+	static FTransform GetComTransform_AssumesLocked(const FPhysicsActorHandle& InActorReference);
+	static FTransform GetComTransformLocal_AssumesLocked(const FPhysicsActorHandle& InActorReference);
 
-	static ENGINE_API FVector GetLocalInertiaTensor_AssumesLocked(const FPhysicsActorReference_Chaos& InActorReference);
-	static ENGINE_API FBox GetBounds_AssumesLocked(const FPhysicsActorReference_Chaos& InActorReference);
+	static FVector GetLocalInertiaTensor_AssumesLocked(const FPhysicsActorHandle& InActorReference);
+	static FBox GetBounds_AssumesLocked(const FPhysicsActorHandle& InActorReference);
 
-	static ENGINE_API void SetLinearDamping_AssumesLocked(const FPhysicsActorReference_Chaos& InActorReference, float InDamping);
-	static ENGINE_API void SetAngularDamping_AssumesLocked(const FPhysicsActorReference_Chaos& InActorReference, float InDamping);
+	static void SetLinearDamping_AssumesLocked(const FPhysicsActorHandle& InActorReference, float InDamping);
+	static void SetAngularDamping_AssumesLocked(const FPhysicsActorHandle& InActorReference, float InDamping);
 
-	static ENGINE_API void AddImpulse_AssumesLocked(const FPhysicsActorReference_Chaos& InActorReference, const FVector& InForce);
-	static ENGINE_API void AddAngularImpulseInRadians_AssumesLocked(const FPhysicsActorReference_Chaos& InActorReference, const FVector& InTorque);
-	static ENGINE_API void AddVelocity_AssumesLocked(const FPhysicsActorReference_Chaos& InActorReference, const FVector& InForce);
-	static ENGINE_API void AddAngularVelocityInRadians_AssumesLocked(const FPhysicsActorReference_Chaos& InActorReference, const FVector& InTorque);
-	static ENGINE_API void AddImpulseAtLocation_AssumesLocked(const FPhysicsActorReference_Chaos& InActorReference, const FVector& InImpulse, const FVector& InLocation);
-	static ENGINE_API void AddRadialImpulse_AssumesLocked(const FPhysicsActorReference_Chaos& InActorReference, const FVector& InOrigin, float InRadius, float InStrength, ERadialImpulseFalloff InFalloff, bool bInVelChange);
+	static void AddImpulse_AssumesLocked(const FPhysicsActorHandle& InActorReference, const FVector& InForce);
+	static void AddAngularImpulseInRadians_AssumesLocked(const FPhysicsActorHandle& InActorReference, const FVector& InTorque);
+	static void AddVelocity_AssumesLocked(const FPhysicsActorHandle& InActorReference, const FVector& InForce);
+	static void AddAngularVelocityInRadians_AssumesLocked(const FPhysicsActorHandle& InActorReference, const FVector& InTorque);
+	static void AddImpulseAtLocation_AssumesLocked(const FPhysicsActorHandle& InActorReference, const FVector& InImpulse, const FVector& InLocation);
+	static void AddRadialImpulse_AssumesLocked(const FPhysicsActorHandle& InActorReference, const FVector& InOrigin, float InRadius, float InStrength, ERadialImpulseFalloff InFalloff, bool bInVelChange);
 
-	static ENGINE_API bool IsGravityEnabled_AssumesLocked(const FPhysicsActorReference_Chaos& InActorReference);
-	static ENGINE_API void SetGravityEnabled_AssumesLocked(const FPhysicsActorReference_Chaos& InActorReference, bool bEnabled);
+	static bool IsGravityEnabled_AssumesLocked(const FPhysicsActorHandle& InActorReference);
+	static void SetGravityEnabled_AssumesLocked(const FPhysicsActorHandle& InActorReference, bool bEnabled);
 
-	static ENGINE_API float GetSleepEnergyThreshold_AssumesLocked(const FPhysicsActorReference_Chaos& InActorReference);
-	static ENGINE_API void SetSleepEnergyThreshold_AssumesLocked(const FPhysicsActorReference_Chaos& InActorReference, float InEnergyThreshold);
+	static float GetSleepEnergyThreshold_AssumesLocked(const FPhysicsActorHandle& InActorReference);
+	static void SetSleepEnergyThreshold_AssumesLocked(const FPhysicsActorHandle& InActorReference, float InEnergyThreshold);
 
-	static void SetMass_AssumesLocked(const FPhysicsActorReference_Chaos& InHandle, float InMass);
-	static void SetMassSpaceInertiaTensor_AssumesLocked(const FPhysicsActorReference_Chaos& InHandle, const FVector& InTensor);
-	static void SetComLocalPose_AssumesLocked(const FPhysicsActorReference_Chaos& InHandle, const FTransform& InComLocalPose);
+	static void SetMass_AssumesLocked(const FPhysicsActorHandle& InHandle, float InMass);
+	static void SetMassSpaceInertiaTensor_AssumesLocked(const FPhysicsActorHandle& InHandle, const FVector& InTensor);
+	static void SetComLocalPose_AssumesLocked(const FPhysicsActorHandle& InHandle, const FTransform& InComLocalPose);
 
-	static ENGINE_API float GetStabilizationEnergyThreshold_AssumesLocked(const FPhysicsActorReference_Chaos& InHandle);
-	static ENGINE_API void SetStabilizationEnergyThreshold_AssumesLocked(const FPhysicsActorReference_Chaos& InHandle, float InThreshold);
-	static ENGINE_API uint32 GetSolverPositionIterationCount_AssumesLocked(const FPhysicsActorReference_Chaos& InHandle);
-	static ENGINE_API void SetSolverPositionIterationCount_AssumesLocked(const FPhysicsActorReference_Chaos& InHandle, uint32 InSolverIterationCount);
-	static ENGINE_API uint32 GetSolverVelocityIterationCount_AssumesLocked(const FPhysicsActorReference_Chaos& InHandle);
-	static ENGINE_API void SetSolverVelocityIterationCount_AssumesLocked(const FPhysicsActorReference_Chaos& InHandle, uint32 InSolverIterationCount);
-	static ENGINE_API float GetWakeCounter_AssumesLocked(const FPhysicsActorReference_Chaos& InHandle);
-	static ENGINE_API void SetWakeCounter_AssumesLocked(const FPhysicsActorReference_Chaos& InHandle, float InWakeCounter);
+	static float GetStabilizationEnergyThreshold_AssumesLocked(const FPhysicsActorHandle& InHandle);
+	static void SetStabilizationEnergyThreshold_AssumesLocked(const FPhysicsActorHandle& InHandle, float InThreshold);
+	static uint32 GetSolverPositionIterationCount_AssumesLocked(const FPhysicsActorHandle& InHandle);
+	static void SetSolverPositionIterationCount_AssumesLocked(const FPhysicsActorHandle& InHandle, uint32 InSolverIterationCount);
+	static uint32 GetSolverVelocityIterationCount_AssumesLocked(const FPhysicsActorHandle& InHandle);
+	static void SetSolverVelocityIterationCount_AssumesLocked(const FPhysicsActorHandle& InHandle, uint32 InSolverIterationCount);
+	static float GetWakeCounter_AssumesLocked(const FPhysicsActorHandle& InHandle);
+	static void SetWakeCounter_AssumesLocked(const FPhysicsActorHandle& InHandle, float InWakeCounter);
 
-	static ENGINE_API SIZE_T GetResourceSizeEx(const FPhysicsActorReference_Chaos& InActorRef);
+	static SIZE_T GetResourceSizeEx(const FPhysicsActorHandle& InActorRef);
 	
-    static ENGINE_API FPhysicsConstraintReference_Chaos CreateConstraint(const FPhysicsActorReference_Chaos& InActorRef1, const FPhysicsActorReference_Chaos& InActorRef2, const FTransform& InLocalFrame1, const FTransform& InLocalFrame2);
-	static ENGINE_API void SetConstraintUserData(const FPhysicsConstraintReference_Chaos& InConstraintRef, void* InUserData);
-	static ENGINE_API void ReleaseConstraint(FPhysicsConstraintReference_Chaos& InConstraintRef);
-
-	static ENGINE_API FTransform GetLocalPose(const FPhysicsConstraintReference_Chaos& InConstraintRef, EConstraintFrame::Type InFrame);
-	static ENGINE_API FTransform GetGlobalPose(const FPhysicsConstraintReference_Chaos& InConstraintRef, EConstraintFrame::Type InFrame);
-	static ENGINE_API FVector GetLocation(const FPhysicsConstraintReference_Chaos& InConstraintRef);
-	static ENGINE_API void GetForce(const FPhysicsConstraintReference_Chaos& InConstraintRef, FVector& OutLinForce, FVector& OutAngForce);
-	static ENGINE_API void GetDriveLinearVelocity(const FPhysicsConstraintReference_Chaos& InConstraintRef, FVector& OutLinVelocity);
-	static ENGINE_API void GetDriveAngularVelocity(const FPhysicsConstraintReference_Chaos& InConstraintRef, FVector& OutAngVelocity);
-
-	static ENGINE_API float GetCurrentSwing1(const FPhysicsConstraintReference_Chaos& InConstraintRef);
-	static ENGINE_API float GetCurrentSwing2(const FPhysicsConstraintReference_Chaos& InConstraintRef);
-	static ENGINE_API float GetCurrentTwist(const FPhysicsConstraintReference_Chaos& InConstraintRef);
-
-	static ENGINE_API void SetCanVisualize(const FPhysicsConstraintReference_Chaos& InConstraintRef, bool bInCanVisualize);
-	static ENGINE_API void SetCollisionEnabled(const FPhysicsConstraintReference_Chaos& InConstraintRef, bool bInCollisionEnabled);
-	static ENGINE_API void SetProjectionEnabled_AssumesLocked(const FPhysicsConstraintReference_Chaos& InConstraintRef, bool bInProjectionEnabled, float InLinearTolerance = 0.0f, float InAngularToleranceDegrees = 0.0f);
-	static ENGINE_API void SetParentDominates_AssumesLocked(const FPhysicsConstraintReference_Chaos& InConstraintRef, bool bInParentDominates);
-	static ENGINE_API void SetBreakForces_AssumesLocked(const FPhysicsConstraintReference_Chaos& InConstraintRef, float InLinearBreakForce, float InAngularBreakForce);
-	static ENGINE_API void SetLocalPose(const FPhysicsConstraintReference_Chaos& InConstraintRef, const FTransform& InPose, EConstraintFrame::Type InFrame);
-
-	static ENGINE_API void SetLinearMotionLimitType_AssumesLocked(const FPhysicsConstraintReference_Chaos& InConstraintRef, PhysicsInterfaceTypes::ELimitAxis InAxis, ELinearConstraintMotion InMotion);
-	static ENGINE_API void SetAngularMotionLimitType_AssumesLocked(const FPhysicsConstraintReference_Chaos& InConstraintRef, PhysicsInterfaceTypes::ELimitAxis InAxis, EAngularConstraintMotion InMotion);
-
-	static ENGINE_API void UpdateLinearLimitParams_AssumesLocked(const FPhysicsConstraintReference_Chaos& InConstraintRef, float InLimit, float InAverageMass, const FLinearConstraint& InParams);
-	static ENGINE_API void UpdateConeLimitParams_AssumesLocked(const FPhysicsConstraintReference_Chaos& InConstraintRef, float InAverageMass, const FConeConstraint& InParams);
-	static ENGINE_API void UpdateTwistLimitParams_AssumesLocked(const FPhysicsConstraintReference_Chaos& InConstraintRef, float InAverageMass, const FTwistConstraint& InParams);
-	static ENGINE_API void UpdateLinearDrive_AssumesLocked(const FPhysicsConstraintReference_Chaos& InConstraintRef, const FLinearDriveConstraint& InDriveParams);
-	static ENGINE_API void UpdateAngularDrive_AssumesLocked(const FPhysicsConstraintReference_Chaos& InConstraintRef, const FAngularDriveConstraint& InDriveParams);
-	static ENGINE_API void UpdateDriveTarget_AssumesLocked(const FPhysicsConstraintReference_Chaos& InConstraintRef, const FLinearDriveConstraint& InLinDrive, const FAngularDriveConstraint& InAngDrive);
-	static ENGINE_API void SetDrivePosition(const FPhysicsConstraintReference_Chaos& InConstraintRef, const FVector& InPosition);
-	static ENGINE_API void SetDriveOrientation(const FPhysicsConstraintReference_Chaos& InConstraintRef, const FQuat& InOrientation);
-	static ENGINE_API void SetDriveLinearVelocity(const FPhysicsConstraintReference_Chaos& InConstraintRef, const FVector& InLinVelocity);
-	static ENGINE_API void SetDriveAngularVelocity(const FPhysicsConstraintReference_Chaos& InConstraintRef, const FVector& InAngVelocity);
-
-	static ENGINE_API void SetTwistLimit(const FPhysicsConstraintReference_Chaos& InConstraintRef, float InLowerLimit, float InUpperLimit, float InContactDistance);
-	static ENGINE_API void SetSwingLimit(const FPhysicsConstraintReference_Chaos& InConstraintRef, float InYLimit, float InZLimit, float InContactDistance);
-	static ENGINE_API void SetLinearLimit(const FPhysicsConstraintReference_Chaos& InConstraintRef, float InLimit);
-
-	static ENGINE_API bool IsBroken(const FPhysicsConstraintReference_Chaos& InConstraintRef);
-
-	static ENGINE_API bool ExecuteOnUnbrokenConstraintReadOnly(const FPhysicsConstraintReference_Chaos& InConstraintRef, TFunctionRef<void(const FPhysicsConstraintReference_Chaos&)> Func);
-	static ENGINE_API bool ExecuteOnUnbrokenConstraintReadWrite(const FPhysicsConstraintReference_Chaos& InConstraintRef, TFunctionRef<void(const FPhysicsConstraintReference_Chaos&)> Func);
+    static FPhysicsConstraintReference_Chaos CreateConstraint(const FPhysicsActorHandle& InActorRef1, const FPhysicsActorHandle& InActorRef2, const FTransform& InLocalFrame1, const FTransform& InLocalFrame2);
+	static void SetConstraintUserData(const FPhysicsConstraintReference_Chaos& InConstraintRef, void* InUserData);
+	static void ReleaseConstraint(FPhysicsConstraintReference_Chaos& InConstraintRef);
+
+	static FTransform GetLocalPose(const FPhysicsConstraintReference_Chaos& InConstraintRef, EConstraintFrame::Type InFrame);
+	static FTransform GetGlobalPose(const FPhysicsConstraintReference_Chaos& InConstraintRef, EConstraintFrame::Type InFrame);
+	static FVector GetLocation(const FPhysicsConstraintReference_Chaos& InConstraintRef);
+	static void GetForce(const FPhysicsConstraintReference_Chaos& InConstraintRef, FVector& OutLinForce, FVector& OutAngForce);
+	static void GetDriveLinearVelocity(const FPhysicsConstraintReference_Chaos& InConstraintRef, FVector& OutLinVelocity);
+	static void GetDriveAngularVelocity(const FPhysicsConstraintReference_Chaos& InConstraintRef, FVector& OutAngVelocity);
+
+	static float GetCurrentSwing1(const FPhysicsConstraintReference_Chaos& InConstraintRef);
+	static float GetCurrentSwing2(const FPhysicsConstraintReference_Chaos& InConstraintRef);
+	static float GetCurrentTwist(const FPhysicsConstraintReference_Chaos& InConstraintRef);
+
+	static void SetCanVisualize(const FPhysicsConstraintReference_Chaos& InConstraintRef, bool bInCanVisualize);
+	static void SetCollisionEnabled(const FPhysicsConstraintReference_Chaos& InConstraintRef, bool bInCollisionEnabled);
+	static void SetProjectionEnabled_AssumesLocked(const FPhysicsConstraintReference_Chaos& InConstraintRef, bool bInProjectionEnabled, float InLinearTolerance = 0.0f, float InAngularToleranceDegrees = 0.0f);
+	static void SetParentDominates_AssumesLocked(const FPhysicsConstraintReference_Chaos& InConstraintRef, bool bInParentDominates);
+	static void SetBreakForces_AssumesLocked(const FPhysicsConstraintReference_Chaos& InConstraintRef, float InLinearBreakForce, float InAngularBreakForce);
+	static void SetLocalPose(const FPhysicsConstraintReference_Chaos& InConstraintRef, const FTransform& InPose, EConstraintFrame::Type InFrame);
+
+	static void SetLinearMotionLimitType_AssumesLocked(const FPhysicsConstraintReference_Chaos& InConstraintRef, PhysicsInterfaceTypes::ELimitAxis InAxis, ELinearConstraintMotion InMotion);
+	static void SetAngularMotionLimitType_AssumesLocked(const FPhysicsConstraintReference_Chaos& InConstraintRef, PhysicsInterfaceTypes::ELimitAxis InAxis, EAngularConstraintMotion InMotion);
+
+	static void UpdateLinearLimitParams_AssumesLocked(const FPhysicsConstraintReference_Chaos& InConstraintRef, float InLimit, float InAverageMass, const FLinearConstraint& InParams);
+	static void UpdateConeLimitParams_AssumesLocked(const FPhysicsConstraintReference_Chaos& InConstraintRef, float InAverageMass, const FConeConstraint& InParams);
+	static void UpdateTwistLimitParams_AssumesLocked(const FPhysicsConstraintReference_Chaos& InConstraintRef, float InAverageMass, const FTwistConstraint& InParams);
+	static void UpdateLinearDrive_AssumesLocked(const FPhysicsConstraintReference_Chaos& InConstraintRef, const FLinearDriveConstraint& InDriveParams);
+	static void UpdateAngularDrive_AssumesLocked(const FPhysicsConstraintReference_Chaos& InConstraintRef, const FAngularDriveConstraint& InDriveParams);
+	static void UpdateDriveTarget_AssumesLocked(const FPhysicsConstraintReference_Chaos& InConstraintRef, const FLinearDriveConstraint& InLinDrive, const FAngularDriveConstraint& InAngDrive);
+	static void SetDrivePosition(const FPhysicsConstraintReference_Chaos& InConstraintRef, const FVector& InPosition);
+	static void SetDriveOrientation(const FPhysicsConstraintReference_Chaos& InConstraintRef, const FQuat& InOrientation);
+	static void SetDriveLinearVelocity(const FPhysicsConstraintReference_Chaos& InConstraintRef, const FVector& InLinVelocity);
+	static void SetDriveAngularVelocity(const FPhysicsConstraintReference_Chaos& InConstraintRef, const FVector& InAngVelocity);
+
+	static void SetTwistLimit(const FPhysicsConstraintReference_Chaos& InConstraintRef, float InLowerLimit, float InUpperLimit, float InContactDistance);
+	static void SetSwingLimit(const FPhysicsConstraintReference_Chaos& InConstraintRef, float InYLimit, float InZLimit, float InContactDistance);
+	static void SetLinearLimit(const FPhysicsConstraintReference_Chaos& InConstraintRef, float InLimit);
+
+	static bool IsBroken(const FPhysicsConstraintReference_Chaos& InConstraintRef);
+
+	static bool ExecuteOnUnbrokenConstraintReadOnly(const FPhysicsConstraintReference_Chaos& InConstraintRef, TFunctionRef<void(const FPhysicsConstraintReference_Chaos&)> Func);
+	static bool ExecuteOnUnbrokenConstraintReadWrite(const FPhysicsConstraintReference_Chaos& InConstraintRef, TFunctionRef<void(const FPhysicsConstraintReference_Chaos&)> Func);
 
     /////////////////////////////////////////////
 
     // Interface needed for cmd
-    static ENGINE_API bool ExecuteRead(const FPhysicsActorReference_Chaos& InActorReference, TFunctionRef<void(const FPhysicsActorReference_Chaos& Actor)> InCallable)
+    static bool ExecuteRead(const FPhysicsActorHandle& InActorReference, TFunctionRef<void(const FPhysicsActorHandle& Actor)> InCallable)
     {
-		if(InActorReference.IsValid())
-		{
-			InCallable(InActorReference);
-			return true;
-		}
-
-		return false;
+		InCallable(InActorReference);
+		return true;
     }
 
-    static ENGINE_API bool ExecuteRead(USkeletalMeshComponent* InMeshComponent, TFunctionRef<void()> InCallable)
+    static bool ExecuteRead(USkeletalMeshComponent* InMeshComponent, TFunctionRef<void()> InCallable)
     {
         InCallable();
         return true;
     }
 
-    static ENGINE_API bool ExecuteRead(const FPhysicsActorReference_Chaos& InActorReferenceA, const FPhysicsActorReference_Chaos& InActorReferenceB, TFunctionRef<void(const FPhysicsActorReference_Chaos& ActorA, const FPhysicsActorReference_Chaos& ActorB)> InCallable)
+    static bool ExecuteRead(const FPhysicsActorHandle& InActorReferenceA, const FPhysicsActorHandle& InActorReferenceB, TFunctionRef<void(const FPhysicsActorHandle& ActorA, const FPhysicsActorHandle& ActorB)> InCallable)
     {
-		if(InActorReferenceA.IsValid() || InActorReferenceB.IsValid())
-		{
-			InCallable(InActorReferenceA, InActorReferenceB);
-			return true;
-		}
-
-		return false;
+		InCallable(InActorReferenceA, InActorReferenceB);
+		return true;
     }
 
-    static ENGINE_API bool ExecuteRead(const FPhysicsConstraintReference_Chaos& InConstraintRef, TFunctionRef<void(const FPhysicsConstraintReference_Chaos& Constraint)> InCallable)
+    static bool ExecuteRead(const FPhysicsConstraintReference_Chaos& InConstraintRef, TFunctionRef<void(const FPhysicsConstraintReference_Chaos& Constraint)> InCallable)
     {
 		if(InConstraintRef.IsValid())
 		{
@@ -347,7 +361,7 @@ public:
 		return false;
     }
 
-    static ENGINE_API bool ExecuteRead(FPhysScene* InScene, TFunctionRef<void()> InCallable)
+    static bool ExecuteRead(FPhysScene* InScene, TFunctionRef<void()> InCallable)
     {
 		if(InScene)
 		{
@@ -358,35 +372,25 @@ public:
 		return false;
     }
 
-    static ENGINE_API bool ExecuteWrite(const FPhysicsActorReference_Chaos& InActorReference, TFunctionRef<void(const FPhysicsActorReference_Chaos& Actor)> InCallable)
+    static bool ExecuteWrite(const FPhysicsActorHandle& InActorReference, TFunctionRef<void(const FPhysicsActorHandle& Actor)> InCallable)
     {
-		if(InActorReference.IsValid())
-		{
-			InCallable(InActorReference);
-			return true;
-		}
-
-		return false;
+		InCallable(InActorReference);
+		return true;
     }
 
-    static ENGINE_API bool ExecuteWrite(USkeletalMeshComponent* InMeshComponent, TFunctionRef<void()> InCallable)
+    static bool ExecuteWrite(USkeletalMeshComponent* InMeshComponent, TFunctionRef<void()> InCallable)
     {
         InCallable();
         return true;
     }
 
-    static ENGINE_API bool ExecuteWrite(const FPhysicsActorReference_Chaos& InActorReferenceA, const FPhysicsActorReference_Chaos& InActorReferenceB, TFunctionRef<void(const FPhysicsActorReference_Chaos& ActorA, const FPhysicsActorReference_Chaos& ActorB)> InCallable)
+    static bool ExecuteWrite(const FPhysicsActorHandle& InActorReferenceA, const FPhysicsActorHandle& InActorReferenceB, TFunctionRef<void(const FPhysicsActorHandle& ActorA, const FPhysicsActorHandle& ActorB)> InCallable)
     {
-		if(InActorReferenceA.IsValid() || InActorReferenceB.IsValid())
-		{
-			InCallable(InActorReferenceA, InActorReferenceB);
-			return true;
-		}
-
-		return false;
+		InCallable(InActorReferenceA, InActorReferenceB);
+		return true;
     }
 
-    static ENGINE_API bool ExecuteWrite(const FPhysicsConstraintReference_Chaos& InConstraintRef, TFunctionRef<void(const FPhysicsConstraintReference_Chaos& Constraint)> InCallable)
+    static bool ExecuteWrite(const FPhysicsConstraintReference_Chaos& InConstraintRef, TFunctionRef<void(const FPhysicsConstraintReference_Chaos& Constraint)> InCallable)
     {
 		if(InConstraintRef.IsValid())
 		{
@@ -397,7 +401,7 @@ public:
 		return false;
     }
 	
-    static ENGINE_API bool ExecuteWrite(FPhysScene* InScene, TFunctionRef<void()> InCallable)
+    static bool ExecuteWrite(FPhysScene* InScene, TFunctionRef<void()> InCallable)
     {
 		if(InScene)
 		{
@@ -416,73 +420,9 @@ public:
 		}
     }
 
-	// Scene query interface functions
-
-	static ENGINE_API bool RaycastTest(const UWorld* World, const FVector Start, const FVector End, ECollisionChannel TraceChannel, const FCollisionQueryParams& Params, const FCollisionResponseParams& ResponseParams, const FCollisionObjectQueryParams& ObjectParams = FCollisionObjectQueryParams::DefaultObjectQueryParam)
-    {
-        return false;
-    }
-    static ENGINE_API bool RaycastSingle(const UWorld* World, struct FHitResult& OutHit, const FVector Start, const FVector End, ECollisionChannel TraceChannel, const FCollisionQueryParams& Params, const FCollisionResponseParams& ResponseParams, const FCollisionObjectQueryParams& ObjectParams = FCollisionObjectQueryParams::DefaultObjectQueryParam)
-    {
-        return false;
-    }
-    static ENGINE_API bool RaycastMulti(const UWorld* World, TArray<struct FHitResult>& OutHits, const FVector& Start, const FVector& End, ECollisionChannel TraceChannel, const FCollisionQueryParams& Params, const FCollisionResponseParams& ResponseParams, const FCollisionObjectQueryParams& ObjectParams = FCollisionObjectQueryParams::DefaultObjectQueryParam)
-    {
-        return false;
-    }
-
-    static ENGINE_API bool GeomOverlapBlockingTest(const UWorld* World, const FCollisionShape& CollisionShape, const FVector& Pos, const FQuat& Rot, ECollisionChannel TraceChannel, const FCollisionQueryParams& Params, const FCollisionResponseParams& ResponseParams, const FCollisionObjectQueryParams& ObjectParams = FCollisionObjectQueryParams::DefaultObjectQueryParam)
-    {
-        return false;
-    }
-    static ENGINE_API bool GeomOverlapAnyTest(const UWorld* World, const FCollisionShape& CollisionShape, const FVector& Pos, const FQuat& Rot, ECollisionChannel TraceChannel, const FCollisionQueryParams& Params, const FCollisionResponseParams& ResponseParams, const FCollisionObjectQueryParams& ObjectParams = FCollisionObjectQueryParams::DefaultObjectQueryParam)
-    {
-        return false;
-    }
-    static ENGINE_API bool GeomOverlapMulti(const UWorld* World, const FCollisionShape& CollisionShape, const FVector& Pos, const FQuat& Rot, TArray<FOverlapResult>& OutOverlaps, ECollisionChannel TraceChannel, const FCollisionQueryParams& Params, const FCollisionResponseParams& ResponseParams, const FCollisionObjectQueryParams& ObjectParams = FCollisionObjectQueryParams::DefaultObjectQueryParam)
-    {
-        return false;
-    }
-
-	// GEOM SWEEP
-
-    static ENGINE_API bool GeomSweepTest(const UWorld* World, const FCollisionShape& CollisionShape, const FQuat& Rot, FVector Start, FVector End, ECollisionChannel TraceChannel, const FCollisionQueryParams& Params, const FCollisionResponseParams& ResponseParams, const FCollisionObjectQueryParams& ObjectParams = FCollisionObjectQueryParams::DefaultObjectQueryParam)
-    {
-        return false;
-    }
-    static ENGINE_API bool GeomSweepSingle(const UWorld* World, const FCollisionShape& CollisionShape, const FQuat& Rot, FHitResult& OutHit, FVector Start, FVector End, ECollisionChannel TraceChannel, const FCollisionQueryParams& Params, const FCollisionResponseParams& ResponseParams, const FCollisionObjectQueryParams& ObjectParams = FCollisionObjectQueryParams::DefaultObjectQueryParam)
-    {
-        return false;
-    }
-    static ENGINE_API bool GeomSweepMulti(const UWorld* World, const FCollisionShape& CollisionShape, const FQuat& Rot, TArray<FHitResult>& OutHits, FVector Start, FVector End, ECollisionChannel TraceChannel, const FCollisionQueryParams& Params, const FCollisionResponseParams& ResponseParams, const FCollisionObjectQueryParams& ObjectParams = FCollisionObjectQueryParams::DefaultObjectQueryParam)
-    {
-        return false;
-    }
-
-	template<typename GeomType>
-    static bool GeomSweepMulti(const UWorld* World, const GeomType& InGeom, const FQuat& InGeomRot, TArray<FHitResult>& OutHits, FVector Start, FVector End, ECollisionChannel TraceChannel, const FCollisionQueryParams& Params, const FCollisionResponseParams& ResponseParams, const FCollisionObjectQueryParams& ObjectParams = FCollisionObjectQueryParams::DefaultObjectQueryParam)
-    {
-        return false;
-    }
-	template<typename GeomType>
-    static bool GeomOverlapMulti(const UWorld* World, const GeomType& InGeom, const FVector& InPosition, const FQuat& InRotation, TArray<FOverlapResult>& OutOverlaps, ECollisionChannel TraceChannel, const FCollisionQueryParams& Params, const FCollisionResponseParams& ResponseParams, const FCollisionObjectQueryParams& ObjectParams)
-    {
-        return false;
-    }
-
 	// Misc
 
-	static ENGINE_API bool ExecPhysCommands(const TCHAR* Cmd, FOutputDevice* Ar, UWorld* InWorld);
-
-	static ENGINE_API FPhysScene* GetCurrentScene(FPhysicsActorHandle& InActorReference)
-	{
-		return InActorReference.GetScene();
-	}
-
-	static ENGINE_API const FPhysScene* GetCurrentScene(const FPhysicsActorHandle& InActorReference)
-	{
-		return InActorReference.GetScene();
-	}
+	static bool ExecPhysCommands(const TCHAR* Cmd, FOutputDevice* Ar, UWorld* InWorld);
 
 #if WITH_PHYSX
     static void CalculateMassPropertiesFromShapeCollection(physx::PxMassProperties& OutProperties, const TArray<FPhysicsShapeHandle>& InShapes, float InDensityKGPerCM);
@@ -491,15 +431,16 @@ public:
 	// Shape interface functions
 	static FPhysicsShapeHandle CreateShape(physx::PxGeometry* InGeom, bool bSimulation = true, bool bQuery = true, UPhysicalMaterial* InSimpleMaterial = nullptr, TArray<UPhysicalMaterial*>* InComplexMaterials = nullptr);
 	
-	static void ENGINE_API AddGeometry(FPhysicsActorHandle& InActor, const FGeometryAddParams& InParams, TArray<FPhysicsShapeHandle>* OutOptShapes = nullptr);
+	static void AddGeometry(FPhysicsActorHandle& InActor, const FGeometryAddParams& InParams, TArray<FPhysicsShapeHandle>* OutOptShapes = nullptr);
 	static FPhysicsShapeHandle CloneShape(const FPhysicsShapeHandle& InShape);
-
-	static ENGINE_API FCollisionFilterData GetSimulationFilter(const FPhysicsShapeHandle& InShape);
-	static ENGINE_API FCollisionFilterData GetQueryFilter(const FPhysicsShapeHandle& InShape);
+	static FPhysicsGeometryCollection_Chaos GetGeometryCollection(const FPhysicsShapeHandle& InShape);
+	
+	static FCollisionFilterData GetSimulationFilter(const FPhysicsShapeHandle& InShape);
+	static FCollisionFilterData GetQueryFilter(const FPhysicsShapeHandle& InShape);
 	static bool IsSimulationShape(const FPhysicsShapeHandle& InShape);
 	static bool IsQueryShape(const FPhysicsShapeHandle& InShape);
 	static bool IsShapeType(const FPhysicsShapeHandle& InShape, ECollisionShapeType InType);
-	static ENGINE_API ECollisionShapeType GetShapeType(const FPhysicsShapeHandle& InShape);
+	static ECollisionShapeType GetShapeType(const FPhysicsShapeHandle& InShape);
 	static FTransform GetLocalTransform(const FPhysicsShapeHandle& InShape);
     static void* GetUserData(const FPhysicsShapeHandle& InShape) { return nullptr; }
 
@@ -513,8 +454,8 @@ public:
     // @todo(mlentine): Which of these do we need to support?
 	// Set the mask filter of a shape, which is an extra level of filtering during collision detection / query for extra channels like "Blue Team" and "Red Team"
     static void SetMaskFilter(const FPhysicsShapeHandle& InShape, FMaskFilter InFilter) {}
-    static void SetSimulationFilter(const FPhysicsShapeHandle& InShape, const FCollisionFilterData& InFilter) {}
-    static void SetQueryFilter(const FPhysicsShapeHandle& InShape, const FCollisionFilterData& InFilter) {}
+	static void SetSimulationFilter(const FPhysicsShapeHandle& InShape, const FCollisionFilterData& InFilter);
+	static void SetQueryFilter(const FPhysicsShapeHandle& InShape, const FCollisionFilterData& InFilter);
     static void SetIsSimulationShape(const FPhysicsShapeHandle& InShape, bool bIsSimShape) { const_cast<FPhysicsShapeHandle&>(InShape).bSimulation = bIsSimShape; }
     static void SetIsQueryShape(const FPhysicsShapeHandle& InShape, bool bIsQueryShape) { const_cast<FPhysicsShapeHandle&>(InShape).bSimulation = bIsQueryShape; }
     static void SetUserData(const FPhysicsShapeHandle& InShape, void* InUserData) {}
@@ -541,9 +482,9 @@ FORCEINLINE ECollisionShapeType GetType(const Chaos::TImplicitObject<float, 3>&
 	return ECollisionShapeType::None;
 }
 */
-FORCEINLINE ECollisionShapeType GetGeometryType(const Chaos::TImplicitObject<float, 3>& Geom)
+FORCEINLINE ECollisionShapeType GetGeometryType(const Chaos::TPerShapeData<float, 3>& Shape)
 {
-	return GetType(Geom);
+	return GetType(*Shape.Geometry);
 }
 
 /*
@@ -558,34 +499,19 @@ FORCEINLINE float GetHalfHeight(const Chaos::TCapsule<float>& Capsule)
 }
 */
 
-FORCEINLINE FVector FindBoxOpposingNormal(const FPhysTypeDummy& PHit, const FVector& TraceDirectionDenorm, const FVector InNormal)
-{
-	return FVector(0.0f, 0.0f, 1.0f);
-}
+FVector FindBoxOpposingNormal(const FLocationHit& PHit, const FVector& TraceDirectionDenorm, const FVector& InNormal);
+FVector FindHeightFieldOpposingNormal(const FLocationHit& PHit, const FVector& TraceDirectionDenorm, const FVector& InNormal);
+FVector FindConvexMeshOpposingNormal(const FLocationHit& PHit, const FVector& TraceDirectionDenorm, const FVector& InNormal);
+FVector FindTriMeshOpposingNormal(const FLocationHit& PHit, const FVector& TraceDirectionDenorm, const FVector& InNormal);
 
-FORCEINLINE FVector FindHeightFieldOpposingNormal(const FPhysTypeDummy& PHit, const FVector& TraceDirectionDenorm, const FVector InNormal)
+FORCEINLINE void DrawOverlappingTris(const UWorld* World, const FLocationHit& Hit, const Chaos::TImplicitObject<float, 3>& Geom, const FTransform& QueryTM)
 {
-	return FVector(0.0f, 0.0f, 1.0f);
+	//TODO_SQ_IMPLEMENTATION
 }
 
-FORCEINLINE FVector FindConvexMeshOpposingNormal(const FPhysTypeDummy& PHit, const FVector& TraceDirectionDenorm, const FVector InNormal)
+FORCEINLINE void ComputeZeroDistanceImpactNormalAndPenetration(const UWorld* World, const FLocationHit& Hit, const Chaos::TImplicitObject<float, 3>& Geom, const FTransform& QueryTM, FHitResult& OutResult)
 {
-	return FVector(0.0f, 0.0f, 1.0f);
-}
-
-FORCEINLINE FVector FindTriMeshOpposingNormal(const FPhysTypeDummy& PHit, const FVector& TraceDirectionDenorm, const FVector InNormal)
-{
-	return FVector(0.0f, 0.0f, 1.0f);
-}
-
-FORCEINLINE void DrawOverlappingTris(const UWorld* World, const FPhysTypeDummy& Hit, const Chaos::TImplicitObject<float, 3>& Geom, const FTransform& QueryTM)
-{
-
-}
-
-FORCEINLINE void ComputeZeroDistanceImpactNormalAndPenetration(const UWorld* World, const FPhysTypeDummy& Hit, const Chaos::TImplicitObject<float, 3>& Geom, const FTransform& QueryTM, FHitResult& OutResult)
-{
-
+	//TODO_SQ_IMPLEMENTATION
 }
 
 inline FPhysTypeDummy* GetMaterialFromInternalFaceIndex(const FPhysicsShape& Shape, uint32 InternalFaceIndex)
@@ -593,11 +519,6 @@ inline FPhysTypeDummy* GetMaterialFromInternalFaceIndex(const FPhysicsShape& Sha
 	return nullptr;
 }
 
-inline FCollisionFilterData GetSimulationFilterData(const FPhysicsShape& Shape)
-{
-	return FCollisionFilterData();
-}
-
 inline uint32 GetTriangleMeshExternalFaceIndex(const FPhysicsShape& Shape, uint32 InternalFaceIndex)
 {
 	return GetInvalidPhysicsFaceIndex();
@@ -616,8 +537,14 @@ inline void SetShape(FPhysTypeDummy& Hit, FPhysTypeDummy* Shape)
 bool IsBlocking(const FPhysicsShape& PShape, const FCollisionFilterData& QueryFilter);
 
 template <>
-ENGINE_API int32 FPhysInterface_Chaos::GetAllShapes_AssumedLocked(const FPhysicsActorReference_Chaos& InActorHandle, TArray<FPhysicsShapeReference_Chaos, FDefaultAllocator>& OutShapes);
+ENGINE_API int32 FPhysInterface_Chaos::GetAllShapes_AssumedLocked(const FPhysicsActorHandle& InActorHandle, TArray<FPhysicsShapeReference_Chaos, FDefaultAllocator>& OutShapes);
 template <>
-ENGINE_API int32 FPhysInterface_Chaos::GetAllShapes_AssumedLocked(const FPhysicsActorReference_Chaos& InActorHandle, PhysicsInterfaceTypes::FInlineShapeArray& OutShapes);
+ENGINE_API int32 FPhysInterface_Chaos::GetAllShapes_AssumedLocked(const FPhysicsActorHandle& InActorHandle, PhysicsInterfaceTypes::FInlineShapeArray& OutShapes);
+
+template<>
+bool FGenericPhysicsInterface::GeomSweepMulti(const UWorld* World, const FPhysicsGeometryCollection& InGeom, const FQuat& InGeomRot, TArray<FHitResult>& OutHits, FVector Start, FVector End, ECollisionChannel TraceChannel, const FCollisionQueryParams& Params, const FCollisionResponseParams& ResponseParams, const FCollisionObjectQueryParams& ObjectParams /*= FCollisionObjectQueryParams::DefaultObjectQueryParam*/);
+
+template<>
+bool FGenericPhysicsInterface::GeomOverlapMulti(const UWorld* World, const FPhysicsGeometryCollection& InGeom, const FVector& InPosition, const FQuat& InRotation, TArray<FOverlapResult>& OutOverlaps, ECollisionChannel TraceChannel, const FCollisionQueryParams& Params, const FCollisionResponseParams& ResponseParams, const FCollisionObjectQueryParams& ObjectParams);
 
 #endif
diff --git a/Engine/Source/Runtime/Engine/Public/Physics/Experimental/PhysScene_Chaos.h b/Engine/Source/Runtime/Engine/Public/Physics/Experimental/PhysScene_Chaos.h
--- a/Engine/Source/Runtime/Engine/Public/Physics/Experimental/PhysScene_Chaos.h
+++ b/Engine/Source/Runtime/Engine/Public/Physics/Experimental/PhysScene_Chaos.h
@@ -9,13 +9,12 @@
 #include "Physics/PhysScene.h"
 #include "UObject/UObjectGlobals.h"
 #include "UObject/GCObject.h"
-#include "SolverObjects/SolverObject.h"
 #include "GameFramework/Actor.h"
 #include "PhysicsPublic.h"
 #include "PhysInterface_Chaos.h"
 #include "Chaos/Transform.h"
-#include "Chaos/PBDSpringConstraints.h"
-#include "Chaos/PerParticleGravity.h"
+#include "Chaos/Declares.h"
+#include "Chaos/Framework/SingleParticlePhysicsProxyFwd.h"
 
 #ifndef CHAOS_WITH_PAUSABLE_SOLVER
 #define CHAOS_WITH_PAUSABLE_SOLVER 1
@@ -29,23 +28,22 @@ class FChaosSolversModule;
 struct FForceFieldProxy;
 struct FSolverStateStorage;
 
-class FSkeletalMeshPhysicsObject;
-class FStaticMeshPhysicsObject;
-class FBodyInstancePhysicsObject;
-class FGeometryCollectionPhysicsObject;
-class FFieldSystemPhysicsObject;
+class FSkeletalMeshPhysicsProxy;
+class FStaticMeshPhysicsProxy;
+class FGeometryCollectionPhysicsProxy;
+class FFieldSystemPhysicsProxy;
+
+class IPhysicsProxyBase;
 
 namespace Chaos
 {
-	class FPBDRigidsSolver;
 	class FPhysicsProxy;
 	class IDispatcher;
 
-	template <typename T, int>
-	class PerParticleGravity;
+	enum EEventType;
 
-	template <typename T, int>
-	class TPBDSpringConstraints;
+	template<typename PayloadType, typename HandlerType>
+	class TRawEventHandler;
 }
 
 /**
@@ -75,7 +73,7 @@ public:
 	/**
 	 * Get the internal Chaos solver object
 	 */
-	Chaos::FPBDRigidsSolver* GetSolver() const;
+	Chaos::FPhysicsSolver* GetSolver() const;
 
 	/** Returns the actor that owns this solver. */
 	AActor* GetSolverActor() const;
@@ -89,21 +87,29 @@ public:
 	/**
 	 * Called during creation of the physics state for gamethread objects to pass off an object to the physics thread
 	 */
-	void AddObject(UPrimitiveComponent* Component, FSkeletalMeshPhysicsObject* InObject);
-	void AddObject(UPrimitiveComponent* Component, FStaticMeshPhysicsObject* InObject);
-	void AddObject(UPrimitiveComponent* Component, FBodyInstancePhysicsObject* InObject);
-	void AddObject(UPrimitiveComponent* Component, FGeometryCollectionPhysicsObject* InObject);
-	void AddObject(UPrimitiveComponent* Component, FFieldSystemPhysicsObject* InObject);
+	void AddObject(UPrimitiveComponent* Component, FSkeletalMeshPhysicsProxy* InObject);
+	void AddObject(UPrimitiveComponent* Component, FStaticMeshPhysicsProxy* InObject);
+	void AddObject(UPrimitiveComponent* Component, FGeometryParticlePhysicsProxy* InObject);
+	void AddObject(UPrimitiveComponent* Component, FGeometryCollectionPhysicsProxy* InObject);
+	void AddObject(UPrimitiveComponent* Component, FFieldSystemPhysicsProxy* InObject);
 
 	/**
 	 * Called during physics state destruction for the game thread to remove objects from the simulation
 	 * #BG TODO - Doesn't actually remove from the evolution at the moment
 	 */
-	void RemoveObject(FSkeletalMeshPhysicsObject* InObject);
-	void RemoveObject(FStaticMeshPhysicsObject* InObject);
-	void RemoveObject(FBodyInstancePhysicsObject* InObject);
-	void RemoveObject(FGeometryCollectionPhysicsObject* InObject);
-	void RemoveObject(FFieldSystemPhysicsObject* InObject);	
+	void RemoveObject(FSkeletalMeshPhysicsProxy* InObject);
+	void RemoveObject(FStaticMeshPhysicsProxy* InObject);
+	void RemoveObject(FGeometryParticlePhysicsProxy* InObject);
+	void RemoveObject(FGeometryCollectionPhysicsProxy* InObject);
+	void RemoveObject(FFieldSystemPhysicsProxy* InObject);	
+
+	template<typename PayloadType>
+	void RegisterEvent(const Chaos::EEventType& EventID, TFunction<void(const Chaos::FPBDRigidsSolver* Solver, PayloadType& EventData)> InLambda);
+	void UnregisterEvent(const Chaos::EEventType& EventID);
+
+	template<typename PayloadType, typename HandlerType>
+	void RegisterEventHandler(const Chaos::EEventType& EventID, HandlerType* Handler, typename Chaos::TRawEventHandler<PayloadType, HandlerType>::FHandlerFunction Func);
+	void UnregisterEventHandler(const Chaos::EEventType& EventID, const void* Handler);
 
 	void Shutdown();
 
@@ -117,17 +123,17 @@ public:
 	
 	/** Given a solver object, returns its associated component. */
 	template<class OwnerType>
-	OwnerType* GetOwningComponent(ISolverObjectBase* SolverObject) const
+	OwnerType* GetOwningComponent(IPhysicsProxyBase* PhysicsProxy) const
 	{ 
-		UPrimitiveComponent* const* CompPtr = SolverObjectToComponentMap.Find(SolverObject);
+		UPrimitiveComponent* const* CompPtr = PhysicsProxyToComponentMap.Find(PhysicsProxy);
 		return CompPtr ? Cast<OwnerType>(*CompPtr) : nullptr;
 	}
 
 	/** Given a component, returns its associated solver object. */
-	ISolverObjectBase* GetOwnedSolverObject(UPrimitiveComponent* Comp) const
+	IPhysicsProxyBase* GetOwnedPhysicsProxy(UPrimitiveComponent* Comp) const
 	{
-		ISolverObjectBase* const* SolverObjectPtr = ComponentToSolverObjectMap.Find(Comp);
-		return SolverObjectPtr ? *SolverObjectPtr : nullptr;
+		IPhysicsProxyBase* const* PhysicsProxyPtr = ComponentToPhysicsProxyMap.Find(Comp);
+		return PhysicsProxyPtr ? *PhysicsProxyPtr : nullptr;
 	}
 
 private:
@@ -136,8 +142,8 @@ private:
 	void OnUpdateWorldPause();
 #endif
 
-	void AddToComponentMaps(UPrimitiveComponent* Component, ISolverObjectBase* InObject);
-	void RemoveFromComponentMaps(ISolverObjectBase* InObject);
+	void AddToComponentMaps(UPrimitiveComponent* Component, IPhysicsProxyBase* InObject);
+	void RemoveFromComponentMaps(IPhysicsProxyBase* InObject);
 
 #if WITH_EDITOR
 	/**
@@ -154,13 +160,13 @@ private:
 	FChaosSolversModule* ChaosModule;
 
 	// Solver representing this scene
-	Chaos::FPBDRigidsSolver* SceneSolver;
+	Chaos::FPhysicsSolver* SceneSolver;
 
-	// Maps SolverObject to Component that created the SolverObject
-	TMap<ISolverObjectBase*, UPrimitiveComponent*> SolverObjectToComponentMap;
+	// Maps PhysicsProxy to Component that created the PhysicsProxy
+	TMap<IPhysicsProxyBase*, UPrimitiveComponent*> PhysicsProxyToComponentMap;
 
-	// Maps Component to SolverObject that is created
-	TMap<UPrimitiveComponent*, ISolverObjectBase*> ComponentToSolverObjectMap;
+	// Maps Component to PhysicsProxy that is created
+	TMap<UPrimitiveComponent*, IPhysicsProxyBase*> ComponentToPhysicsProxyMap;
 
 	/** The SolverActor that spawned and owns this scene */
 	TWeakObjectPtr<AActor> SolverActor;
@@ -176,6 +182,40 @@ private:
 #endif
 };
 
+template<typename PayloadType>
+void FPhysScene_Chaos::RegisterEvent(const Chaos::EEventType& EventID, TFunction<void(const Chaos::FPBDRigidsSolver* Solver, PayloadType& EventData)> InLambda)
+{
+	check(IsInGameThread());
+
+	Chaos::IDispatcher* Dispatcher = GetDispatcher();
+	Chaos::FPBDRigidsSolver* Solver = GetSolver();
+
+	if (Dispatcher)
+	{
+		Dispatcher->EnqueueCommandImmediate([EventID, InLambda, InSolver = Solver](Chaos::FPersistentPhysicsTask* PhysThread)
+		{
+			InSolver->GetEventManager()->RegisterEvent<PayloadType>(InSolver, InLambda);
+		});
+	}
+}
+
+template<typename PayloadType, typename HandlerType>
+void FPhysScene_Chaos::RegisterEventHandler(const Chaos::EEventType& EventID, HandlerType* Handler, typename Chaos::TRawEventHandler<PayloadType, HandlerType>::FHandlerFunction Func)
+{
+	check(IsInGameThread());
+
+	Chaos::IDispatcher* Dispatcher = GetDispatcher();
+	Chaos::FPBDRigidsSolver* Solver = GetSolver();
+
+	if (Dispatcher)
+	{
+		Dispatcher->EnqueueCommandImmediate([EventID, Handler, Func, InSolver = Solver](Chaos::FPersistentPhysicsTask* PhysThread)
+		{
+			InSolver->GetEventManager()->RegisterHandler<PayloadType>(EventID, Handler, Func);
+		});
+	}
+}
+
 class UWorld;
 class AWorldSettings;
 class FPhysicsReplication;
@@ -196,7 +236,9 @@ public:
 	void OnWorldBeginPlay();
 	void OnWorldEndPlay();
 
-	void AddActorsToScene_AssumesLocked(const TArray<FPhysicsActorHandle>& InActors);
+	// In Chaos, this function will update the pointers from actor handles to their proxies.
+	// So the array of handles must be non-const.
+	void ENGINE_API AddActorsToScene_AssumesLocked(TArray<FPhysicsActorHandle>& InActors);
 	void AddAggregateToScene(const FPhysicsAggregateHandle& InAggregate);
 
 	void SetOwningWorld(UWorld* InOwningWorld);
@@ -204,8 +246,8 @@ public:
 	UWorld* GetOwningWorld();
 	const UWorld* GetOwningWorld() const;
 
-	Chaos::FPBDRigidsSolver* GetSolver();
-	const Chaos::FPBDRigidsSolver* GetSolver() const;
+	Chaos::FPhysicsSolver* GetSolver();
+	const Chaos::FPhysicsSolver* GetSolver() const;
 
 	FPhysicsReplication* GetPhysicsReplication();
 	void RemoveBodyInstanceFromPendingLists_AssumesLocked(FBodyInstance* BodyInstance, int32 SceneType);
@@ -260,15 +302,36 @@ public:
 	ENGINE_API bool ExecPxVis(uint32 SceneType, const TCHAR* Cmd, FOutputDevice* Ar);
 	ENGINE_API bool ExecApexVis(uint32 SceneType, const TCHAR* Cmd, FOutputDevice* Ar);
 
+	template<typename PayloadType>
+	void RegisterEvent(const Chaos::EEventType& EventID, TFunction<void(const Chaos::FPBDRigidsSolver* Solver, PayloadType& EventData)> InLambda)
+	{
+		Scene.RegisterEvent(EventID, InLambda);
+	}
+	void UnregisterEvent(const Chaos::EEventType& EventID)
+	{
+		Scene.UnregisterEvent(EventID);
+	}
+
+	template<typename PayloadType, typename HandlerType>
+	void RegisterEventHandler(const Chaos::EEventType& EventID, HandlerType* Handler, typename Chaos::TRawEventHandler<PayloadType, HandlerType>::FHandlerFunction Func)
+	{
+		Scene.RegisterEventHandler<PayloadType, HandlerType>(EventID, Handler, Func);
+	}
+	void UnregisterEventHandler(const Chaos::EEventType& EventID, const void* Handler)
+	{
+		Scene.UnregisterEventHandler(EventID, Handler);
+	}
+
 private:
 
 	void SyncBodies();
+	void SyncBodies(Chaos::FPhysicsSolver* Solver);
 
-	void SetKinematicTransform(FPhysicsActorReference_Chaos& InActorReference, const Chaos::TRigidTransform<float, 3>& NewTransform)
+	void SetKinematicTransform(FPhysicsActorHandle& InActorReference, const Chaos::TRigidTransform<float, 3>& NewTransform)
 	{
 		// #todo : Initialize
 		// Set the buffered kinematic data on the game and render thread
-		// InActorReference.GetPhysicsObject()->SetKinematicData(...)
+		// InActorReference.GetPhysicsProxy()->SetKinematicData(...)
 	}
 
 	void Lock()
@@ -296,15 +359,15 @@ private:
 		// #todo : Implement
 	}
 
-	FPhysicsConstraintReference_Chaos AddSpringConstraint(const TArray< TPair<FPhysicsActorReference_Chaos, FPhysicsActorReference_Chaos> >& Constraint);
+	FPhysicsConstraintReference_Chaos AddSpringConstraint(const TArray< TPair<FPhysicsActorHandle, FPhysicsActorHandle> >& Constraint);
 	void RemoveSpringConstraint(const FPhysicsConstraintReference_Chaos& Constraint);
 
-	void AddForce(const Chaos::TVector<float, 3>& Force, FPhysicsActorReference_Chaos& Handle)
+	void AddForce(const Chaos::TVector<float, 3>& Force, FPhysicsActorHandle& Handle)
 	{
 		// #todo : Implement
 	}
 
-	void AddTorque(const Chaos::TVector<float, 3>& Torque, FPhysicsActorReference_Chaos& Handle)
+	void AddTorque(const Chaos::TVector<float, 3>& Torque, FPhysicsActorHandle& Handle)
 	{
 		// #todo : Implement
 	}
@@ -317,9 +380,6 @@ private:
 	// @todo(mlentine): Locking is very heavy handed right now; need to make less so.
 	FCriticalSection MCriticalSection;
 	float MDeltaTime;
-	TUniquePtr<Chaos::PerParticleGravity<float, 3>> MGravity;
-	//Springs
-	TUniquePtr<Chaos::TPBDSpringConstraints<float, 3>> MSpringConstraints;
 	//Body Instances
 	Chaos::TArrayCollectionArray<FBodyInstance*> BodyInstances;
 	// Temp Interface
diff --git a/Engine/Source/Runtime/Engine/Public/Physics/Experimental/PhysicsInterfaceImmediatePhysX.h b/Engine/Source/Runtime/Engine/Public/Physics/Experimental/PhysicsInterfaceImmediatePhysX.h
--- a/Engine/Source/Runtime/Engine/Public/Physics/Experimental/PhysicsInterfaceImmediatePhysX.h
+++ b/Engine/Source/Runtime/Engine/Public/Physics/Experimental/PhysicsInterfaceImmediatePhysX.h
@@ -44,6 +44,7 @@ struct ENGINE_API FPhysicsActorReference_ImmediatePhysX
 
 	bool IsValid() const;
 	bool Equals(const FPhysicsActorReference_ImmediatePhysX& Other) const;
+	bool operator==(const FPhysicsActorReference_ImmediatePhysX& Other) const { return Equals(Other); }
 
 	FPhysScene_ImmediatePhysX* Scene;
 	uint32 Index;
@@ -130,6 +131,7 @@ struct ENGINE_API FPhysicsGeometryCollection_ImmediatePhysX
         }
         return ECollisionShapeType::None;
     }
+#if 0	//Not sure if this is needed, refactoring this code for Chaos usage so leaving as is for now
     physx::PxGeometry& GetGeometry() const { return *Geometry; }
     bool GetBoxGeometry(physx::PxBoxGeometry& OutGeom) const
     {
@@ -151,6 +153,7 @@ struct ENGINE_API FPhysicsGeometryCollection_ImmediatePhysX
     {
         return GPhysXSDK->createShape(GetGeometry(), nullptr, 0)->getTriangleMeshGeometry(OutGeom);
     }
+#endif
 };
 
 // @todo(mlentine): Once we finalize the format of this we need to redo this
diff --git a/Engine/Source/Runtime/Engine/Public/Physics/PhysScene_PhysX.h b/Engine/Source/Runtime/Engine/Public/Physics/PhysScene_PhysX.h
--- a/Engine/Source/Runtime/Engine/Public/Physics/PhysScene_PhysX.h
+++ b/Engine/Source/Runtime/Engine/Public/Physics/PhysScene_PhysX.h
@@ -28,7 +28,7 @@ struct FSimulationScratchBuffer
 	int32 BufferSize;
 };
 
-#if !WITH_CHAOS && !WITH_IMMEDIATE_PHYSX && !PHYSICS_INTERFACE_LLIMMEDIATE
+#if !WITH_CHAOS && !WITH_IMMEDIATE_PHYSX
 
 class FPhysicsReplication;
 class IPhysicsReplicationFactory;
diff --git a/Engine/Source/Runtime/Engine/Public/Physics/PhysicsGeometry.h b/Engine/Source/Runtime/Engine/Public/Physics/PhysicsGeometry.h
new file mode 100644
--- /dev/null
+++ b/Engine/Source/Runtime/Engine/Public/Physics/PhysicsGeometry.h
@@ -0,0 +1,79 @@
+// Copyright 1998-2019 Epic Games, Inc. All Rights Reserved.
+
+#pragma once
+
+#include "PhysicsEngine/AggregateGeom.h"
+#include "Containers/ArrayView.h"
+
+struct FKSphereElem;
+
+namespace physx
+{
+	class PxSphereGeometry;
+	class PxTransform;
+
+	class PxSphereGeometry;
+	class PxBoxGeometry;
+	class PxCapsuleGeometry;
+	class PxConvexMeshGeometry;
+	class PxTriangleMesh;
+	class PxTriangleMeshGeometry;
+}
+
+class UBodySetup;
+
+/** Helper struct for iterating over shapes in a body setup.*/
+struct ENGINE_API FBodySetupShapeIterator
+{
+	// Only valid for PhysX and Chaos - will fail to compile when used otherwise
+	using TTransform = TChooseClass<WITH_PHYSX && PHYSICS_INTERFACE_PHYSX, physx::PxTransform, TChooseClass<WITH_CHAOS, FTransform, void>::Result>::Result;
+
+	FBodySetupShapeIterator();
+
+#if WITH_PHYSX && PHYSICS_INTERFACE_PHYSX
+	FBodySetupShapeIterator(const FVector& InScale3D, const FTransform& InRelativeTM, bool bDoubleSidedTrimeshes);
+	/** Iterates over the elements array and creates the needed geometry and local pose. Note that this memory is on the stack so it's illegal to use it by reference outside the lambda */
+	template <typename ElemType, typename GeomType>
+	void ForEachShape(const TArrayView<ElemType>& Elements, TFunctionRef<void(const ElemType& Elem, const GeomType& Geom, const TTransform& LocalPose, float ContactOffset, float RestOffset)> VisitorFunc) const;
+
+	/** Helper function to determine contact offset params */
+	static void GetContactOffsetParams(float& InOutContactOffsetFactor, float& InOutMinContactOffset, float& InOutMaxContactOffset);
+
+private:
+
+	template <typename ElemType, typename GeomType> bool PopulateGeometryAndTransform(const ElemType& Elem, GeomType& Geom, TTransform& OutTM) const;
+	template <typename GeomType> float ComputeContactOffset(const GeomType& Geom) const;
+	template <typename ElemType> float ComputeRestOffset(const ElemType& Geom) const;
+	template <typename ElemType> FString GetDebugName() const;
+#endif
+
+private:
+	FVector Scale3D;
+	const FTransform& RelativeTM;
+
+	float MinScaleAbs;
+	float MinScale;
+	FVector ShapeScale3DAbs;
+	FVector ShapeScale3D;
+
+	float ContactOffsetFactor;
+	float MinContactOffset;
+	float MaxContactOffset;
+
+	bool bDoubleSidedTriMeshGeo;
+};
+
+#if WITH_PHYSX && PHYSICS_INTERFACE_PHYSX
+
+/// @cond DOXYGEN_WARNINGS
+
+//Explicit export of template instantiation 
+extern template ENGINE_API void FBodySetupShapeIterator::ForEachShape(const TArrayView<FKSphereElem>&, TFunctionRef<void(const FKSphereElem&, const physx::PxSphereGeometry&, const FBodySetupShapeIterator::TTransform&, float, float)>) const;
+extern template ENGINE_API void FBodySetupShapeIterator::ForEachShape(const TArrayView<FKBoxElem>&, TFunctionRef<void(const FKBoxElem&, const physx::PxBoxGeometry&, const FBodySetupShapeIterator::TTransform&, float, float)>) const;
+extern template ENGINE_API void FBodySetupShapeIterator::ForEachShape(const TArrayView<FKSphylElem>&, TFunctionRef<void(const FKSphylElem&, const physx::PxCapsuleGeometry&, const FBodySetupShapeIterator::TTransform&, float, float)>) const;
+extern template ENGINE_API void FBodySetupShapeIterator::ForEachShape(const TArrayView<FKConvexElem>&, TFunctionRef<void(const FKConvexElem&, const physx::PxConvexMeshGeometry&, const FBodySetupShapeIterator::TTransform&, float, float)>) const;
+extern template ENGINE_API void FBodySetupShapeIterator::ForEachShape(const TArrayView<physx::PxTriangleMesh*>&, TFunctionRef<void(physx::PxTriangleMesh* const &, const physx::PxTriangleMeshGeometry&, const FBodySetupShapeIterator::TTransform&, float, float)>) const;
+
+/// @endcond
+
+#endif //WITH_PHYSX
\ No newline at end of file
diff --git a/Engine/Source/Runtime/Engine/Public/Physics/PhysicsInterfaceCore.h b/Engine/Source/Runtime/Engine/Public/Physics/PhysicsInterfaceCore.h
--- a/Engine/Source/Runtime/Engine/Public/Physics/PhysicsInterfaceCore.h
+++ b/Engine/Source/Runtime/Engine/Public/Physics/PhysicsInterfaceCore.h
@@ -13,10 +13,6 @@
 #include "Physics/Experimental/PhysScene_ImmediatePhysX.h"
 #include "Physics/Experimental/PhysicsInterfaceImmediatePhysX.h"
 
-#elif PHYSICS_INTERFACE_LLIMMEDIATE
-	
-#include "Physics/Experimental/PhysicsInterfaceLLImmediate.h"
-
 #elif WITH_CHAOS
 
 #include "Physics/Experimental/PhysInterface_Chaos.h"
diff --git a/Engine/Source/Runtime/Engine/Public/Physics/PhysicsInterfacePhysX.h b/Engine/Source/Runtime/Engine/Public/Physics/PhysicsInterfacePhysX.h
--- a/Engine/Source/Runtime/Engine/Public/Physics/PhysicsInterfacePhysX.h
+++ b/Engine/Source/Runtime/Engine/Public/Physics/PhysicsInterfacePhysX.h
@@ -2,7 +2,7 @@
 
 #pragma once
 
-#if !WITH_CHAOS && !WITH_IMMEDIATE_PHYSX && !PHYSICS_INTERFACE_LLIMMEDIATE
+#if !WITH_CHAOS && !WITH_IMMEDIATE_PHYSX
 
 #include "EngineGlobals.h"
 #include "Engine/EngineTypes.h"
@@ -55,6 +55,7 @@ struct ENGINE_API FPhysicsActorHandle_PhysX
 
 	bool IsValid() const;
 	bool Equals(const FPhysicsActorHandle_PhysX& Other) const;
+	bool operator==(const FPhysicsActorHandle_PhysX& Other) const { return Equals(Other); }
 
 	physx::PxRigidActor* SyncActor;
 
@@ -121,12 +122,18 @@ struct ENGINE_API FPhysicsGeometryCollection_PhysX
 	~FPhysicsGeometryCollection_PhysX();
 
 	ECollisionShapeType GetType() const;
-	physx::PxGeometry& GetGeometry() const;
-	bool GetBoxGeometry(physx::PxBoxGeometry& OutGeom) const;
-	bool GetSphereGeometry(physx::PxSphereGeometry& OutGeom) const;
-	bool GetCapsuleGeometry(physx::PxCapsuleGeometry& OutGeom) const;
-	bool GetConvexGeometry(physx::PxConvexMeshGeometry& OutGeom) const;
-	bool GetTriMeshGeometry(physx::PxTriangleMeshGeometry& OutGeom) const;
+	const physx::PxGeometry& GetGeometry() const;
+	physx::PxGeometry& GetGeometry();
+	const physx::PxBoxGeometry& GetBoxGeometry() const;
+	physx::PxBoxGeometry& GetBoxGeometry();
+	const physx::PxSphereGeometry& GetSphereGeometry() const;
+	physx::PxSphereGeometry& GetSphereGeometry();
+	const physx::PxCapsuleGeometry& GetCapsuleGeometry() const;
+	physx::PxCapsuleGeometry& GetCapsuleGeometry();
+	const physx::PxConvexMeshGeometry& GetConvexGeometry() const;
+	physx::PxConvexMeshGeometry& GetConvexGeometry();
+	const physx::PxTriangleMeshGeometry& GetTriMeshGeometry() const;
+	physx::PxTriangleMeshGeometry& GetTriMeshGeometry();
 
 private:
 	friend struct FPhysicsInterface_PhysX;
@@ -238,6 +245,9 @@ struct ENGINE_API FPhysicsInterface_PhysX : public FGenericPhysicsInterface
 	static void ReleaseActor(FPhysicsActorHandle_PhysX& InHandle, FPhysScene* InScene = nullptr, bool bNeverDeferRelease = false);
 	//////////////////////////////////////////////////////////////////////////
 
+	// NOTE: Adding this temporarily, until we phase out Handle.IsValid().
+	static bool IsValid(const FPhysicsActorHandle& Handle) { return Handle.IsValid(); }
+
 	template<typename AllocatorType>
 	static int32 GetAllShapes_AssumedLocked(const FPhysicsActorHandle_PhysX& InHandle, TArray<FPhysicsShapeHandle_PhysX, AllocatorType>& OutShapes);
 	static int32 GetNumShapes(const FPhysicsActorHandle_PhysX& InHandle);
diff --git a/Engine/Source/Runtime/Engine/Public/Physics/PhysicsInterfaceTypes.h b/Engine/Source/Runtime/Engine/Public/Physics/PhysicsInterfaceTypes.h
--- a/Engine/Source/Runtime/Engine/Public/Physics/PhysicsInterfaceTypes.h
+++ b/Engine/Source/Runtime/Engine/Public/Physics/PhysicsInterfaceTypes.h
@@ -8,6 +8,21 @@
 #include "PhysicsEngine/BodySetupEnums.h"
 #include "PhysicsInterfaceTypesCore.h"
 
+#if WITH_CHAOS
+#include "Chaos/Serializable.h"
+
+namespace Chaos
+{
+	template<typename T, int d>
+	class TImplicitObjectUnion;
+
+	template<typename T, int d>
+	class TImplicitObject;
+
+	template <typename T>
+	class TTriangleMeshImplicitObject;
+}
+#endif
 
 // Defines for enabling hitch repeating (see ScopedSQHitchRepeater.h)
 #if !UE_BUILD_SHIPPING
@@ -40,6 +55,9 @@ struct FGeometryAddParams
 #if WITH_PHYSX
 	TArrayView<physx::PxTriangleMesh*> TriMeshes;
 #endif
+#if WITH_CHAOS
+	TArrayView<TUniquePtr<Chaos::TTriangleMeshImplicitObject<float>>> ChaosTriMeshes;
+#endif
 };
 
 namespace PhysicsInterfaceTypes
diff --git a/Engine/Source/Runtime/Experimental/Chaos/Private/Chaos/BoundingVolumeHierarchy.cpp b/Engine/Source/Runtime/Experimental/Chaos/Private/Chaos/BoundingVolumeHierarchy.cpp
--- a/Engine/Source/Runtime/Experimental/Chaos/Private/Chaos/BoundingVolumeHierarchy.cpp
+++ b/Engine/Source/Runtime/Experimental/Chaos/Private/Chaos/BoundingVolumeHierarchy.cpp
@@ -66,17 +66,24 @@ void Chaos::TBoundingVolumeHierarchy<OBJECT_ARRAY, LEAF_TYPE, T, d>::DebugDraw(I
 }
 #endif
 
-template<class OBJECT_ARRAY, class T, int d>
+template<class OBJECT_ARRAY, class T, int d, typename TPayloadType>
 TArray<int32> MakeNewLeaf(const OBJECT_ARRAY& Objects, const TArray<int32>& AllObjects, const TArray<int32>* /*Type*/)
 {
 	return TArray<int32>(AllObjects);
 }
 
-template<class OBJECT_ARRAY, class T, int d>
-TBoundingVolume<OBJECT_ARRAY, T, d> MakeNewLeaf(const OBJECT_ARRAY& Objects, const TArray<int32>& AllObjects, const TBoundingVolume<OBJECT_ARRAY, T, d>* /*Type*/)
+template<class OBJECT_ARRAY, class T, int d, typename TPayloadType>
+TBoundingVolume<OBJECT_ARRAY, TPayloadType, T, d> MakeNewLeaf(const OBJECT_ARRAY& Objects, const TArray<int32>& AllObjects, const TBoundingVolume<OBJECT_ARRAY, TPayloadType, T, d>* /*Type*/)
 {
+#if CHAOS_PARTICLEHANDLE_TODO
 	// @todo(mlentine): This is pretty dirty but should work for now.
-	return TBoundingVolume<OBJECT_ARRAY, T, d>(Objects, reinterpret_cast<const TArray<uint32>&>(AllObjects));
+	return TBoundingVolume<OBJECT_ARRAY, TPayloadType, T, d>(Objects, reinterpret_cast<const TArray<uint32>&>(AllObjects));
+#else
+	TArray<TSOAView<OBJECT_ARRAY>> SOAs;
+	SOAs.Add(TSOAView<OBJECT_ARRAY>(const_cast<OBJECT_ARRAY*>(&Objects)));	//todo: this const_cast is here because bvh is holding on to things as const but giving it out as non const. Not sure if we should change API or not
+	//return TBoundingVolume<OBJECT_ARRAY, T, d>(MakeParticleView<OBJECT_ARRAY>({ {&Objects} }));
+	return TBoundingVolume<OBJECT_ARRAY, TPayloadType, T, d>(MakeParticleView<OBJECT_ARRAY>(MoveTemp(SOAs)));
+#endif
 }
 
 template<class OBJECT_ARRAY, class LEAF_TYPE, class T, int d>
@@ -167,27 +174,32 @@ void TBoundingVolumeHierarchy<OBJECT_ARRAY, LEAF_TYPE, T, d>::UpdateHierarchyImp
 	else
 	{
 		Elements[0].LeafIndex = Leafs.Num();
-		Leafs.Add(MakeNewLeaf<OBJECT_ARRAY, T, d>(*MObjects, AllObjects, static_cast<LEAF_TYPE*>(nullptr)));
+		Leafs.Add(MakeNewLeaf<OBJECT_ARRAY, T, d, TPayloadType>(*MObjects, AllObjects, static_cast<LEAF_TYPE*>(nullptr)));
 	}
 	check(Elements[0].LeafIndex < Leafs.Num() || Elements[0].MChildren.Num() > 0);
 	UE_LOG(LogChaos, Verbose, TEXT("Generated Tree with %d Nodes"), Elements.Num());
 	//PrintTree("", &Elements[0]);
 }
 
-template<class OBJECT_ARRAY, class T, int d>
+template<class OBJECT_ARRAY, typename TPayloadType, class T, int d>
 TArray<int32> FindAllIntersectionsLeafHelper(const TArray<int32>& Leaf, const TVector<T, d>& Point)
 {
 	return Leaf;
 }
 
-template<class OBJECT_ARRAY, class T, int d>
-TArray<int32> FindAllIntersectionsLeafHelper(const TBoundingVolume<OBJECT_ARRAY, T, d>& Leaf, const TVector<T, d>& Point)
+template<class OBJECT_ARRAY,typename TPayloadType, class T, int d>
+TArray<int32> FindAllIntersectionsLeafHelper(const TBoundingVolume<OBJECT_ARRAY, TPayloadType, T, d>& Leaf, const TVector<T, d>& Point)
 {
+#if CHAOS_PARTICLEHANDLE_TODO
 	return Leaf.FindAllIntersections(Point);
+#else
+	check(false);
+	return TArray<int32>();
+#endif
 }
 
-template<class OBJECT_ARRAY, class LEAF_TYPE, class T, int d>
-TArray<int32> TBoundingVolumeHierarchy<OBJECT_ARRAY, LEAF_TYPE, T, d>::FindAllIntersectionsHelper(const TBVHNode<T,d>& MyNode, const TVector<T, d>& Point) const
+template<class OBJECT_ARRAY, typename TPayloadType, class T, int d>
+TArray<int32> TBoundingVolumeHierarchy<OBJECT_ARRAY, TPayloadType, T, d>::FindAllIntersectionsHelper(const TBVHNode<T,d>& MyNode, const TVector<T, d>& Point) const
 {
 	TBox<T, d> MBox(MyNode.MMin, MyNode.MMax);
 	if (MBox.SignedDistance(Point) > 0)
@@ -196,7 +208,7 @@ TArray<int32> TBoundingVolumeHierarchy<OBJECT_ARRAY, LEAF_TYPE, T, d>::FindAllIn
 	}
 	if (MyNode.MChildren.Num() == 0)
 	{
-		return FindAllIntersectionsLeafHelper<OBJECT_ARRAY, T, d>(Leafs[MyNode.LeafIndex], Point);
+		return FindAllIntersectionsLeafHelper<OBJECT_ARRAY, TPayloadType, T, d>(Leafs[MyNode.LeafIndex], Point);
 	}
 	const TVector<T, d> ObjectCenter = MBox.Center();
 	int32 Child = 0;
@@ -278,16 +290,20 @@ struct TSpecializeParticlesHelper<TParticles<T,d>>
 	}
 };
 
-template<class OBJECT_ARRAY, class T, int d, typename QUERY_OBJECT>
+template<class OBJECT_ARRAY, typename TPayloadType, typename T, int d, typename QUERY_OBJECT>
 TArray<int32> FindAllIntersectionsLeafHelper(const TArray<int32>& Leaf, const QUERY_OBJECT& QueryObject)
 {
 	return Leaf;
 }
 
-template<class OBJECT_ARRAY, class T, int d, typename QUERY_OBJECT>
-TArray<int32> FindAllIntersectionsLeafHelper(const TBoundingVolume<OBJECT_ARRAY, T, d>& Leaf, const QUERY_OBJECT& QueryObject)
+template<class OBJECT_ARRAY, typename TPayloadType, class T, int d, typename QUERY_OBJECT>
+TArray<int32> FindAllIntersectionsLeafHelper(const TBoundingVolume<OBJECT_ARRAY, TPayloadType, T, d>& Leaf, const QUERY_OBJECT& QueryObject)
 {
+#if CHAOS_PARTICLEHANDLE_TODO
 	return Leaf.FindAllIntersectionsImp(QueryObject);
+#else
+	return TArray<int32>();
+#endif
 }
 
 int32 FindAllIntersectionsSingleThreaded = 1;
@@ -304,7 +320,7 @@ void TBoundingVolumeHierarchy<OBJECT_ARRAY, LEAF_TYPE, T, d>::FindAllIntersectio
 	}
 	if (MyNode.MChildren.Num() == 0)
 	{
-		TSpecializeParticlesHelper<OBJECT_ARRAY>::AccumulateChildrenResults(AccumulateElements, FindAllIntersectionsLeafHelper<OBJECT_ARRAY, T, d>(Leafs[MyNode.LeafIndex], QueryObject), QueryObject, MWorldSpaceBoxes);
+		TSpecializeParticlesHelper<OBJECT_ARRAY>::AccumulateChildrenResults(AccumulateElements, FindAllIntersectionsLeafHelper<OBJECT_ARRAY, TPayloadType, T, d>(Leafs[MyNode.LeafIndex], QueryObject), QueryObject, MWorldSpaceBoxes);
 		return;
 	}
 	
@@ -473,7 +489,7 @@ int32 TBoundingVolumeHierarchy<OBJECT_ARRAY, LEAF_TYPE, T, d>::GenerateNextLevel
 		}
 		else
 		{
-			LocalLeafs[i] = MakeNewLeaf<OBJECT_ARRAY, T, d>(*MObjects, LocalObjects[i], static_cast<LEAF_TYPE*>(nullptr));
+			LocalLeafs[i] = MakeNewLeaf<OBJECT_ARRAY, T, d, TPayloadType>(*MObjects, LocalObjects[i], static_cast<LEAF_TYPE*>(nullptr));
 		}
 	});
 	CriticalSection.Lock();
@@ -616,7 +632,7 @@ int32 TBoundingVolumeHierarchy<OBJECT_ARRAY, LEAF_TYPE, T, d>::GenerateNextLevel
 		}
 		else
 		{
-			LocalLeafs[i] = MakeNewLeaf<OBJECT_ARRAY, T, d>(*MObjects, LocalObjects[i], static_cast<LEAF_TYPE*>(nullptr));
+			LocalLeafs[i] = MakeNewLeaf<OBJECT_ARRAY, T, d, TPayloadType>(*MObjects, LocalObjects[i], static_cast<LEAF_TYPE*>(nullptr));
 		}
 	});
 	CriticalSection.Lock();
@@ -637,13 +653,15 @@ int32 TBoundingVolumeHierarchy<OBJECT_ARRAY, LEAF_TYPE, T, d>::GenerateNextLevel
 
 namespace Chaos
 {
-template <typename OBJECT_ARRAY, typename T, int d>
-void FixupLeafObj(const OBJECT_ARRAY& Objects, TArray<TBoundingVolume<OBJECT_ARRAY, T, d>>& Leafs)
+template <typename OBJECT_ARRAY, typename TPayloadType, typename T, int d>
+void FixupLeafObj(const OBJECT_ARRAY& Objects, TArray<TBoundingVolume<OBJECT_ARRAY, TPayloadType, T, d>>& Leafs)
 {
-	for (TBoundingVolume<OBJECT_ARRAY, T, d>& Leaf : Leafs)
+#if CHAOS_PARTICLEHANDLE_TODO
+	for (TBoundingVolume<OBJECT_ARRAY, TPayloadType, T, d>& Leaf : Leafs)
 	{
 		Leaf.SetObjects(Objects);
 	}
+#endif
 }
 
 template <typename OBJECT_ARRAY>
@@ -661,9 +679,11 @@ void TBoundingVolumeHierarchy<OBJECT_ARRAY, LEAF_TYPE, T, d>::Serialize(FArchive
 	}
 }
 }
+
+template class CHAOS_API Chaos::TBoundingVolume<TArray<Chaos::TSphere<float, 3>*>, int32, float, 3>;
 template class CHAOS_API Chaos::TBoundingVolumeHierarchy<TArray<Chaos::TSphere<float, 3>*>, TArray<int32>, float, 3>;
 template class CHAOS_API Chaos::TBoundingVolumeHierarchy<Chaos::TPBDRigidParticles<float, 3>, TArray<int32>, float, 3>;
 template class CHAOS_API Chaos::TBoundingVolumeHierarchy<Chaos::TParticles<float, 3>, TArray<int32>, float, 3>;
 template class CHAOS_API Chaos::TBoundingVolumeHierarchy<Chaos::TGeometryParticles<float, 3>, TArray<int32>, float, 3>;
-template class CHAOS_API Chaos::TBoundingVolumeHierarchy<Chaos::TPBDRigidParticles<float, 3>, TBoundingVolume<Chaos::TPBDRigidParticles<float, 3>, float, 3>, float, 3>;
-template class CHAOS_API Chaos::TBoundingVolumeHierarchy<Chaos::TGeometryParticles<float, 3>, TBoundingVolume<Chaos::TGeometryParticles<float, 3>, float, 3>, float, 3>;
+template class CHAOS_API Chaos::TBoundingVolumeHierarchy<Chaos::TPBDRigidParticles<float, 3>, TBoundingVolume<Chaos::TPBDRigidParticles<float, 3>, TPBDRigidParticleHandle<float,3>*, float, 3>, float, 3>;
+template class CHAOS_API Chaos::TBoundingVolumeHierarchy<Chaos::TGeometryParticles<float, 3>, TBoundingVolume<Chaos::TGeometryParticles<float, 3>, TGeometryParticleHandle<float,3>*, float, 3>, float, 3>;
diff --git a/Engine/Source/Runtime/Experimental/Chaos/Private/Chaos/Box.cpp b/Engine/Source/Runtime/Experimental/Chaos/Private/Chaos/Box.cpp
--- a/Engine/Source/Runtime/Experimental/Chaos/Private/Chaos/Box.cpp
+++ b/Engine/Source/Runtime/Experimental/Chaos/Private/Chaos/Box.cpp
@@ -6,11 +6,12 @@
 namespace Chaos
 {
 template <typename T, int d>
-bool TBox<T, d>::Raycast(const TVector<T, d>& StartPoint, const TVector<T, d>& Dir, const T Length, const T Thickness, T& OutTime, TVector<T, d>& OutPosition, TVector<T, d>& OutNormal) const
+bool TBox<T, d>::Raycast(const TVector<T, d>& StartPoint, const TVector<T, d>& Dir, const T Length, const T Thickness, T& OutTime, TVector<T, d>& OutPosition, TVector<T, d>& OutNormal, int32& OutFaceIndex) const
 {
 	ensure(Length > 0);
 	ensure(FMath::IsNearlyEqual(Dir.SizeSquared(), 1, KINDA_SMALL_NUMBER));
 
+	OutFaceIndex = INDEX_NONE;
 	const TVector<T, d> MinInflated = MMin - Thickness;
 	const TVector<T,d> StartToMin = MinInflated - StartPoint;
 	
@@ -128,7 +129,7 @@ bool TBox<T, d>::Raycast(const TVector<T, d>& StartPoint, const TVector<T, d>& D
 					TVector<T, d> End = GeomStart;
 					End[CurIdx] = End[CurIdx] == MMin[CurIdx] ? MMax[CurIdx] : MMin[CurIdx];
 					TCapsule<T> Capsule(GeomStart, End, Thickness);
-					if (Capsule.Raycast(StartPoint, Dir, Length, 0, CornerTimes[CurIdx], CornerPositions[CurIdx], CornerNormals[CurIdx]))
+					if (Capsule.Raycast(StartPoint, Dir, Length, 0, CornerTimes[CurIdx], CornerPositions[CurIdx], CornerNormals[CurIdx], OutFaceIndex))
 					{
 						if (HitIdx == INDEX_NONE || CornerTimes[CurIdx] < MinTime)
 						{
@@ -156,7 +157,7 @@ bool TBox<T, d>::Raycast(const TVector<T, d>& StartPoint, const TVector<T, d>& D
 			{
 				//capsule: todo(use a cylinder which is cheaper. Our current cylinder raycast implementation doesn't quite work for this setup)
 				TCapsule<T> CapsuleBorder(GeomStart, GeomEnd, Thickness);
-				bHit = CapsuleBorder.Raycast(StartPoint, Dir, Length, 0, OutTime, OutPosition, OutNormal);
+				bHit = CapsuleBorder.Raycast(StartPoint, Dir, Length, 0, OutTime, OutPosition, OutNormal, OutFaceIndex);
 			}
 
 			if (bHit && OutTime > 0)
@@ -168,7 +169,7 @@ bool TBox<T, d>::Raycast(const TVector<T, d>& StartPoint, const TVector<T, d>& D
 	}
 	
 	// didn't hit any rounded parts so just use the box intersection
-	OutTime = LatestStartTime / Length;
+	OutTime = LatestStartTime;
 	OutNormal = Normal;
 	OutPosition = BoxIntersection - Thickness * Normal;
 	return true;
diff --git a/Engine/Source/Runtime/Experimental/Chaos/Private/Chaos/GeometryParticles.cpp b/Engine/Source/Runtime/Experimental/Chaos/Private/Chaos/GeometryParticles.cpp
new file mode 100644
--- /dev/null
+++ b/Engine/Source/Runtime/Experimental/Chaos/Private/Chaos/GeometryParticles.cpp
@@ -0,0 +1,39 @@
+// Copyright 1998-2019 Epic Games, Inc. All Rights Reserved.
+
+#include "Chaos/GeometryParticles.h"
+#include "Chaos/ImplicitObject.h"
+#include "Chaos/ImplicitObjectUnion.h"
+
+namespace Chaos
+{
+	template <typename T, int d>
+	void UpdateShapesArrayFromGeometry(TShapesArray<T, d>& ShapesArray, const TImplicitObject<T, d>* Geometry)
+	{
+		if(Geometry)
+		{
+			if(const auto* Union = Geometry->template GetObject<TImplicitObjectUnion<T, d>>())
+			{
+				ShapesArray.SetNum(Union->GetObjects().Num());
+				int32 Inner = 0;
+				for(const auto& Geom : Union->GetObjects())
+				{
+					ShapesArray[Inner] = MakeUnique<TPerShapeData<T, d>>();
+					ShapesArray[Inner]->Geometry = Geom.Get();
+					++Inner;
+				}
+			}
+			else
+			{
+				ShapesArray.SetNum(1);
+				ShapesArray[0] = MakeUnique<TPerShapeData<T, d>>();
+				ShapesArray[0]->Geometry = Geometry;
+			}
+		}
+		else
+		{
+			ShapesArray.Reset();
+		}
+	}
+
+	template void UpdateShapesArrayFromGeometry(TShapesArray<float, 3>& ShapesArray, const TImplicitObject<float, 3>* Geometry);
+}
diff --git a/Engine/Source/Runtime/Experimental/Chaos/Private/Chaos/HeightField.cpp b/Engine/Source/Runtime/Experimental/Chaos/Private/Chaos/HeightField.cpp
--- a/Engine/Source/Runtime/Experimental/Chaos/Private/Chaos/HeightField.cpp
+++ b/Engine/Source/Runtime/Experimental/Chaos/Private/Chaos/HeightField.cpp
@@ -41,9 +41,9 @@ THeightField<T>::THeightField(TArray<T>&& Height, int32 NumRows, int32 NumCols,
 }
 
 template <typename T>
-bool THeightField<T>::Raycast(const TVector<T, 3>& StartPoint, const TVector<T, 3>& Dir, const T Length, const T Thickness, T& OutTime, TVector<T,3>& OutPosition, TVector<T,3>& OutNormal) const
+bool THeightField<T>::Raycast(const TVector<T, 3>& StartPoint, const TVector<T, 3>& Dir, const T Length, const T Thickness, T& OutTime, TVector<T,3>& OutPosition, TVector<T,3>& OutNormal, int32& OutFaceIndex) const
 {
-	return MTriangleMeshImplicitObject->Raycast(StartPoint, Dir, Length, Thickness, OutTime, OutPosition, OutNormal);
+	return MTriangleMeshImplicitObject->Raycast(StartPoint, Dir, Length, Thickness, OutTime, OutPosition, OutNormal, OutFaceIndex);
 }
 
 template <typename T>
@@ -60,9 +60,15 @@ bool THeightField<T>::OverlapGeom(const TImplicitObject<T, 3>& QueryGeom, const
 }
 
 template <typename T>
-bool THeightField<T>::SweepGeom(const TImplicitObject<T, 3>& QueryGeom, const TRigidTransform<T, 3>& StartTM, const TVector<T, 3>& Dir, const T Length, T& OutTime, TVector<T, 3>& OutPosition, TVector<T, 3>& OutNormal, const T Thickness /* = 0 */) const
+bool THeightField<T>::SweepGeom(const TImplicitObject<T, 3>& QueryGeom, const TRigidTransform<T, 3>& StartTM, const TVector<T, 3>& Dir, const T Length, T& OutTime, TVector<T, 3>& OutPosition, TVector<T, 3>& OutNormal, int32& OutFaceIndex, const T Thickness /* = 0 */) const
 {
-	return MTriangleMeshImplicitObject->SweepGeom(QueryGeom, StartTM, Dir, Length, OutTime, OutPosition, OutNormal, Thickness);
+	return MTriangleMeshImplicitObject->SweepGeom(QueryGeom, StartTM, Dir, Length, OutTime, OutPosition, OutNormal, OutFaceIndex, Thickness);
+}
+
+template <typename T>
+int32 THeightField<T>::FindMostOpposingFace(const TVector<T, 3>& Position, const TVector<T, 3>& UnitDir, int32 HintFaceIndex) const
+{
+	return MTriangleMeshImplicitObject->FindMostOpposingFace(Position, UnitDir, HintFaceIndex);
 }
 
 }
diff --git a/Engine/Source/Runtime/Experimental/Chaos/Private/Chaos/ImplicitObject.cpp b/Engine/Source/Runtime/Experimental/Chaos/Private/Chaos/ImplicitObject.cpp
--- a/Engine/Source/Runtime/Experimental/Chaos/Private/Chaos/ImplicitObject.cpp
+++ b/Engine/Source/Runtime/Experimental/Chaos/Private/Chaos/ImplicitObject.cpp
@@ -10,6 +10,7 @@
 #include "Chaos/Plane.h"
 #include "Chaos/Sphere.h"
 #include "Chaos/TaperedCylinder.h"
+#include "Chaos/TriangleMeshImplicitObject.h"
 #include "HAL/IConsoleManager.h"
 #include "UObject/DestructionObjectVersion.h"
 
@@ -309,6 +310,7 @@ void TImplicitObject<T, d>::StaticSerialize(FChaosArchive& Ar, TSerializablePtr<
 		case ImplicitObjectType::LevelSet: { ImplicitObject = new TLevelSet<T, d>(); break; }
 		case ImplicitObjectType::Convex: { ImplicitObject = new TConvex<T, d>(); break; }
 		case ImplicitObjectType::TaperedCylinder: { ImplicitObject = new TTaperedCylinder<T>(); break; }
+		case ImplicitObjectType::TriangleMesh: { ImplicitObject = new TTriangleMeshImplicitObject<T>(); break; }
 		default:
 			check(false);
 		}
@@ -328,6 +330,7 @@ void TImplicitObject<T, d>::StaticSerialize(FChaosArchive& Ar, TSerializablePtr<
 		case ImplicitObjectType::LevelSet:
 		case ImplicitObjectType::Convex:
 		case ImplicitObjectType::TaperedCylinder:
+		case ImplicitObjectType::TriangleMesh:
 			break;
 		default:
 			check(false);
diff --git a/Engine/Source/Runtime/Experimental/Chaos/Private/Chaos/KinematicGeometryParticles.cpp b/Engine/Source/Runtime/Experimental/Chaos/Private/Chaos/KinematicGeometryParticles.cpp
--- a/Engine/Source/Runtime/Experimental/Chaos/Private/Chaos/KinematicGeometryParticles.cpp
+++ b/Engine/Source/Runtime/Experimental/Chaos/Private/Chaos/KinematicGeometryParticles.cpp
@@ -5,19 +5,24 @@
 
 namespace Chaos
 {
-	template<class T, int d>
-	const typename TKinematicGeometryParticles<T, d>::THandleType* TKinematicGeometryParticles<T, d>::Handle(int32 Index) const
+	template<class T, int d, EGeometryParticlesSimType SimType>
+	const typename TKinematicGeometryParticlesImp<T, d, SimType>::THandleType* TKinematicGeometryParticlesImp<T, d, SimType>::Handle(int32 Index) const
 	{
-		return static_cast<const TKinematicGeometryParticles<T, d>::THandleType*>(TGeometryParticles<T, d>::Handle(Index));
+		checkf(SimType == EGeometryParticlesSimType::RigidBodySim, TEXT("Handles require a rigid body sim type"));
+		return static_cast<const TKinematicGeometryParticlesImp<T, d, SimType>::THandleType*>(TGeometryParticlesImp<T, d, SimType>::Handle(Index));
 	}
 
-	template<class T, int d>
-	typename TKinematicGeometryParticles<T, d>::THandleType* Chaos::TKinematicGeometryParticles<T, d>::Handle(int32 Index)
+	template<class T, int d, EGeometryParticlesSimType SimType>
+	typename TKinematicGeometryParticlesImp<T, d, SimType>::THandleType* Chaos::TKinematicGeometryParticlesImp<T, d, SimType>::Handle(int32 Index)
 	{
-		return static_cast<TKinematicGeometryParticles<T, d>::THandleType*>(TGeometryParticles<T, d>::Handle(Index));
+		//Todo: turn on static_assert CHAOS_PARTICLEHANDLE_TODO
+		//static_assert(SimType == EGeometryParticlesSimType::RigidBodySim, "Handles require a rigid body sim type");
+		return static_cast<TKinematicGeometryParticlesImp<T, d, SimType>::THandleType*>(TGeometryParticlesImp<T, d, SimType>::Handle(Index));
 	}
 }
 
-template class Chaos::TKinematicGeometryParticles<float, 3>;
+template class Chaos::TKinematicGeometryParticlesImp<float, 3, Chaos::EGeometryParticlesSimType::RigidBodySim>;
+template class Chaos::TKinematicGeometryParticlesImp<float, 3, Chaos::EGeometryParticlesSimType::Other>;
+
 // Re-enable when double precision is able to compile
-//template Chaos::TKinematicGeometryParticles<double, 3>;
\ No newline at end of file
+//template Chaos::TKinematicGeometryParticlesImp<double, 3>;
\ No newline at end of file
