Copying IntermediatePhysics from //UE4/Main to resolve conflicts.

#rnx

diff --git a/Engine/Source/Runtime/Engine/Public/Physics/ImmediatePhysics/ImmediatePhysicsActor.h b/Engine/Source/Runtime/Engine/Public/Physics/ImmediatePhysics/ImmediatePhysicsActor.h
deleted file mode 100644
diff --git a/Engine/Source/Runtime/Engine/Public/Physics/ImmediatePhysics/ImmediatePhysicsActorHandle.h b/Engine/Source/Runtime/Engine/Public/Physics/ImmediatePhysics/ImmediatePhysicsActorHandle.h
--- a/Engine/Source/Runtime/Engine/Public/Physics/ImmediatePhysics/ImmediatePhysicsActorHandle.h
+++ b/Engine/Source/Runtime/Engine/Public/Physics/ImmediatePhysics/ImmediatePhysicsActorHandle.h
@@ -2,340 +2,7 @@
 
 #pragma once
 
-#if WITH_PHYSX
-#include "PhysXPublic.h"
-#endif
+#include "Physics/ImmediatePhysics/ImmediatePhysicsPhysX/ImmediatePhysicsActorHandle_PhysX.h"
+#include "Physics/ImmediatePhysics/ImmediatePhysicsChaos/ImmediatePhysicsActorHandle_Chaos.h"
 
-#include "ImmediatePhysicsSimulation.h"
-#include "Engine/EngineTypes.h"
-
-namespace ImmediatePhysics
-{
-
-/** handle associated with a physics actor. This is the proper way to read/write to the physics simulation */
-struct ENGINE_API FActorHandle
-{
-	friend class FPhysInterface_LLImmediate;
-
-public:
-	/** Sets the world transform.*/
-	void SetWorldTransform(const FTransform& WorldTM)
-	{
-#if WITH_PHYSX
-		OwningSimulation.GetLowLevelBody(ActorDataIndex).body2World = U2PTransform(ActorToBody * WorldTM);
-#endif
-	}
-
-	/** Make a body kinematic, or non-kinematic */
-	void SetIsKinematic(bool bKinematic)
-	{
-#if WITH_PHYSX
-		OwningSimulation.SetIsKinematic(ActorDataIndex, bKinematic);
-#endif
-	}
-	
-	/** Is the actor kinematic */
-	bool GetIsKinematic() const
-	{
-#if WITH_PHYSX
-		return OwningSimulation.GetIsKinematic(ActorDataIndex);
-#else
-		return false;
-#endif
-	}
-
-	/** Gets the kinematic target (next transform) for the actor if one is set (check HasKinematicTarget() to see if a target is available) */
-	FImmediateKinematicTarget& GetKinematicTarget()
-	{
-#if WITH_PHYSX
-		return OwningSimulation.GetKinematicTarget(ActorDataIndex);
-#else
-		static FImmediateKinematicTarget EmptyTarget;
-		return EmptyTarget;
-#endif
-	}
-
-	/** Sets the kinematic target. This will affect velocities as expected*/
-	void SetKinematicTarget(const FTransform& WorldTM)
-	{
-#if WITH_PHYSX
-		FImmediateKinematicTarget& KinematicTarget = OwningSimulation.GetKinematicTarget(ActorDataIndex);
-		KinematicTarget.BodyToWorld = U2PTransform(ActorToBody * WorldTM);
-		KinematicTarget.bTargetSet = true;
-#endif
-	}
-
-	/** Does this actor have a kinematic target (next kinematic transform to be applied) */
-	bool HasKinematicTarget() const
-	{
-#if WITH_PHYSX
-		FImmediateKinematicTarget& KinematicTarget = OwningSimulation.GetKinematicTarget(ActorDataIndex);
-		return KinematicTarget.bTargetSet;
-#else
-		return false;
-#endif
-	}
-
-	/** Whether the body is simulating */
-	bool IsSimulated() const
-	{
-		return OwningSimulation.IsSimulated(ActorDataIndex);
-	}
-
-	/** Get the world transform */
-	FTransform GetWorldTransform() const
-	{
-#if WITH_PHYSX
-		return ActorToBody.GetRelativeTransformReverse( P2UTransform(OwningSimulation.GetLowLevelBody(ActorDataIndex).body2World) );
-#else
-		return FTransform::Identity;
-#endif
-	}
-
-	/** Set the linear velocity */
-	void SetLinearVelocity(const FVector& NewLinearVelocity)
-	{
-#if WITH_PHYSX
-		OwningSimulation.GetLowLevelBody(ActorDataIndex).linearVelocity = U2PVector(NewLinearVelocity);
-#endif
-	}
-
-	/** Get the linear velocity */
-	FVector GetLinearVelocity() const
-	{
-#if WITH_PHYSX
-		return P2UVector(OwningSimulation.GetLowLevelBody(ActorDataIndex).linearVelocity);
-#else
-		return FVector::ZeroVector;
-#endif
-	}
-
-	/** Set the angular velocity */
-	void SetAngularVelocity(const FVector& NewAngularVelocity)
-	{
-#if WITH_PHYSX
-		OwningSimulation.GetLowLevelBody(ActorDataIndex).angularVelocity = U2PVector(NewAngularVelocity);
-#endif
-	}
-
-	/** Get the angular velocity */
-	FVector GetAngularVelocity() const
-	{
-#if WITH_PHYSX
-		return P2UVector(OwningSimulation.GetLowLevelBody(ActorDataIndex).angularVelocity);
-#else
-		return FVector::ZeroVector;
-#endif
-	}
-
-	void AddForce(const FVector& Force)
-	{
-#if WITH_PHYSX
-		OwningSimulation.AddForce(ActorDataIndex, Force);
-#endif
-	}
-
-	void AddRadialForce(const FVector& Origin, float Strength, float Radius, ERadialImpulseFalloff Falloff, FSimulation::EForceType ForceType)
-	{
-#if WITH_PHYSX
-		OwningSimulation.AddRadialForce(ActorDataIndex, Origin, Strength, Radius, Falloff, ForceType);
-#endif
-	}
-
-	/** Set the linear damping*/
-	void SetLinearDamping(float NewLinearDamping)
-	{
-#if WITH_PHYSX
-		OwningSimulation.GetLowLevelBody(ActorDataIndex).linearDamping = NewLinearDamping;
-#endif
-	}
-
-	/** Get the linear damping*/
-	float GetLinearDamping() const
-	{
-#if WITH_PHYSX
-		return OwningSimulation.GetLowLevelBody(ActorDataIndex).linearDamping;
-#else
-		return 0.f;
-#endif
-	}
-
-	/** Set the angular damping*/
-	void SetAngularDamping(float NewAngularDamping)
-	{
-#if WITH_PHYSX
-		OwningSimulation.GetLowLevelBody(ActorDataIndex).angularDamping = NewAngularDamping;
-#endif
-	}
-
-	/** Get the angular damping*/
-	float GetAngularDamping() const
-	{
-#if WITH_PHYSX
-		return OwningSimulation.GetLowLevelBody(ActorDataIndex).angularDamping;
-#else
-		return 0.f;
-#endif
-	}
-
-	/** Set the max linear velocity squared*/
-	void SetMaxLinearVelocitySquared(float NewMaxLinearVelocitySquared)
-	{
-#if WITH_PHYSX
-		OwningSimulation.GetLowLevelBody(ActorDataIndex).maxLinearVelocitySq = NewMaxLinearVelocitySquared;
-#endif
-	}
-
-	/** Get the max linear velocity squared*/
-	float GetMaxLinearVelocitySquared() const
-	{
-#if WITH_PHYSX
-		return OwningSimulation.GetLowLevelBody(ActorDataIndex).maxLinearVelocitySq;
-#else
-		return 0.f;
-#endif
-	}
-
-	/** Set the max angular velocity squared*/
-	void SetMaxAngularVelocitySquared(float NewMaxAngularVelocitySquared)
-	{
-#if WITH_PHYSX
-		OwningSimulation.GetLowLevelBody(ActorDataIndex).maxAngularVelocitySq = NewMaxAngularVelocitySquared;
-#endif
-	}
-
-	/** Get the max angular velocity squared*/
-	float GetMaxAngularVelocitySquared() const
-	{
-#if WITH_PHYSX
-		return OwningSimulation.GetLowLevelBody(ActorDataIndex).maxAngularVelocitySq;
-#else
-		return 0.f;
-#endif
-	}
-
-	/** Set the inverse mass. 0 indicates kinematic object */
-	void SetInverseMass(float NewInverseMass)
-	{
-#if WITH_PHYSX
-		OwningSimulation.GetLowLevelBody(ActorDataIndex).invMass = NewInverseMass;
-#endif
-	}
-
-	/** Get the inverse mass. */
-	float GetInverseMass() const
-	{
-#if WITH_PHYSX
-		return OwningSimulation.GetLowLevelBody(ActorDataIndex).invMass;
-#else
-		return 0.f;
-#endif
-	}
-
-	/** Set the inverse inertia. Mass-space inverse inertia diagonal vector */
-	void SetInverseInertia(const FVector& NewInverseInertia)
-	{
-#if WITH_PHYSX
-		OwningSimulation.GetLowLevelBody(ActorDataIndex).invInertia = U2PVector(NewInverseInertia);
-#endif
-	}
-
-	/** Get the inverse inertia. Mass-space inverse inertia diagonal vector */
-	FVector GetInverseInertia() const
-	{
-#if WITH_PHYSX
-		return P2UVector(OwningSimulation.GetLowLevelBody(ActorDataIndex).invInertia);
-#else
-		return FVector(0.f);
-#endif
-	}
-
-	/** Set the max depenetration velocity*/
-	void SetMaxDepenetrationVelocity(float NewMaxDepenetrationVelocity)
-	{
-#if WITH_PHYSX
-		OwningSimulation.GetLowLevelBody(ActorDataIndex).maxDepenetrationVelocity = NewMaxDepenetrationVelocity;
-#endif
-	}
-
-	/** Get the max depenetration velocity*/
-	float GetMaxDepenetrationVelocity(float NewMaxDepenetrationVelocity) const
-	{
-#if WITH_PHYSX
-		return OwningSimulation.GetLowLevelBody(ActorDataIndex).maxDepenetrationVelocity;
-#else
-		return 0.f;
-#endif
-	}
-
-	/** Set the max contact impulse*/
-	void SetMaxContactImpulse(float NewMaxContactImpulse)
-	{
-#if WITH_PHYSX
-		OwningSimulation.GetLowLevelBody(ActorDataIndex).maxContactImpulse = NewMaxContactImpulse;
-#endif
-	}
-
-	/** Get the max contact impulse*/
-	float GetMaxContactImpulse() const
-	{
-#if WITH_PHYSX
-		return OwningSimulation.GetLowLevelBody(ActorDataIndex).maxContactImpulse;
-#else
-		return 0.f;
-#endif
-	}
-
-	const int32 GetActorIndex() const
-	{
-		return ActorDataIndex;
-	}
-
-#if WITH_PHYSX
-	void AddShape(physx::PxShape* InShape)
-	{
-		OwningSimulation.Actors[ActorDataIndex].AddShape(InShape);
-	}
-#endif
-
-	TArray<ImmediatePhysics::FShape>& GetShapes() const
-	{
-		return OwningSimulation.Actors[ActorDataIndex].Shapes;
-	}
-
-	ImmediatePhysics::FActor* GetSimulationActor() const
-	{
-		return &OwningSimulation.Actors[ActorDataIndex];
-	}
-
-#if WITH_PHYSX
-	immediate::PxRigidBodyData* GetSimulationRigidBodyData() const
-	{
-		immediate::PxRigidBodyData& Data = OwningSimulation.GetLowLevelBody(ActorDataIndex);
-		return &Data;
-	}
-#endif
-
-private:
-	/** Converts from actor space (i.e. the transform in world space as the client gives us) to body space (body with its origin at the COM and oriented to inertia tensor) */
-	FTransform ActorToBody;
-	FSimulation& OwningSimulation;
-	int32 ActorDataIndex;
-
-	friend FSimulation;
-	FActorHandle(FSimulation& InOwningSimulation, int32 InActorDataIndex)
-		: ActorToBody(FTransform::Identity)
-		, OwningSimulation(InOwningSimulation)
-		, ActorDataIndex(InActorDataIndex)
-	{
-	}
-
-	~FActorHandle()
-	{
-	}
-
-	FActorHandle(const FActorHandle& );	//Ensure no copying of handles
-};
-
-
-}
\ No newline at end of file
+#include "Physics/ImmediatePhysics/ImmediatePhysicsDeclares.h"
diff --git a/Engine/Source/Runtime/Engine/Public/Physics/ImmediatePhysics/ImmediatePhysicsCacheAllocator.h b/Engine/Source/Runtime/Engine/Public/Physics/ImmediatePhysics/ImmediatePhysicsCacheAllocator.h
deleted file mode 100644
diff --git a/Engine/Source/Runtime/Engine/Public/Physics/ImmediatePhysics/ImmediatePhysicsChaos/ImmediatePhysicsActorHandle_Chaos.h b/Engine/Source/Runtime/Engine/Public/Physics/ImmediatePhysics/ImmediatePhysicsChaos/ImmediatePhysicsActorHandle_Chaos.h
new file mode 100644
--- /dev/null
+++ b/Engine/Source/Runtime/Engine/Public/Physics/ImmediatePhysics/ImmediatePhysicsChaos/ImmediatePhysicsActorHandle_Chaos.h
@@ -0,0 +1,129 @@
+// Copyright 1998-2019 Epic Games, Inc. All Rights Reserved.
+
+#pragma once
+
+#if INCLUDE_CHAOS
+
+#include "Physics/ImmediatePhysics/ImmediatePhysicsChaos/ImmediatePhysicsCore_Chaos.h"
+
+#include "Engine/EngineTypes.h"
+
+namespace ImmediatePhysics_Chaos
+{
+
+	/** handle associated with a physics actor. This is the proper way to read/write to the physics simulation */
+	struct ENGINE_API FActorHandle
+	{
+	public:
+		~FActorHandle();
+
+		void SetEnabled(bool bEnabled);
+
+		/** Sets the world transform.*/
+		void SetWorldTransform(const FTransform& WorldTM);
+
+		/** Make a body kinematic, or non-kinematic */
+		void SetIsKinematic(bool bKinematic);
+
+		/** Is the actor kinematic */
+		bool GetIsKinematic() const;
+
+#if IMMEDIATEPHYSICS_CHAOS_TODO
+		/** Gets the kinematic target (next transform) for the actor if one is set (check HasKinematicTarget() to see if a target is available) */
+		FImmediateKinematicTarget& GetKinematicTarget();
+#endif
+
+		/** Sets the kinematic target. This will affect velocities as expected*/
+		void SetKinematicTarget(const FTransform& WorldTM);
+
+		/** Does this actor have a kinematic target (next kinematic transform to be applied) */
+		bool HasKinematicTarget() const;
+
+		/** Whether the body is simulating */
+		bool IsSimulated() const;
+
+		/** Get the world transform */
+		FTransform GetWorldTransform() const;
+
+		/** Set the linear velocity */
+		void SetLinearVelocity(const FVector& NewLinearVelocity);
+
+		/** Get the linear velocity */
+		FVector GetLinearVelocity() const;
+
+		/** Set the angular velocity */
+		void SetAngularVelocity(const FVector& NewAngularVelocity);
+
+		/** Get the angular velocity */
+		FVector GetAngularVelocity() const;
+
+		void AddForce(const FVector& Force);
+
+		void AddRadialForce(const FVector& Origin, float Strength, float Radius, ERadialImpulseFalloff Falloff, EForceType ForceType);
+
+		/** Set the linear damping*/
+		void SetLinearDamping(float NewLinearDamping);
+
+		/** Get the linear damping*/
+		float GetLinearDamping() const;
+
+		/** Set the angular damping*/
+		void SetAngularDamping(float NewAngularDamping);
+
+		/** Get the angular damping*/
+		float GetAngularDamping() const;
+
+		/** Set the max linear velocity squared*/
+		void SetMaxLinearVelocitySquared(float NewMaxLinearVelocitySquared);
+
+		/** Get the max linear velocity squared*/
+		float GetMaxLinearVelocitySquared() const;
+
+		/** Set the max angular velocity squared*/
+		void SetMaxAngularVelocitySquared(float NewMaxAngularVelocitySquared);
+
+		/** Get the max angular velocity squared*/
+		float GetMaxAngularVelocitySquared() const;
+
+		/** Set the inverse mass. 0 indicates kinematic object */
+		void SetInverseMass(float NewInverseMass);
+
+		/** Get the inverse mass. */
+		float GetInverseMass() const;
+
+		/** Set the inverse inertia. Mass-space inverse inertia diagonal vector */
+		void SetInverseInertia(const FVector& NewInverseInertia);
+
+		/** Get the inverse inertia. Mass-space inverse inertia diagonal vector */
+		FVector GetInverseInertia() const;
+
+		/** Set the max depenetration velocity*/
+		void SetMaxDepenetrationVelocity(float NewMaxDepenetrationVelocity);
+
+		/** Get the max depenetration velocity*/
+		float GetMaxDepenetrationVelocity(float NewMaxDepenetrationVelocity) const;
+
+		/** Set the max contact impulse*/
+		void SetMaxContactImpulse(float NewMaxContactImpulse);
+
+		/** Get the max contact impulse*/
+		float GetMaxContactImpulse() const;
+
+	private:
+		friend struct FSimulation;
+		friend struct FJointHandle;
+
+		FActorHandle(Chaos::TPBDRigidsEvolutionGBF<FReal, Dimensions>* InEvolution, EActorType ActorType, FBodyInstance* BodyInstance, const FTransform& Transform);
+
+		Chaos::TGenericParticleHandle<FReal, Dimensions> Handle() const;
+
+		Chaos::TPBDRigidsEvolutionGBF<FReal, Dimensions>* Evolution;
+		Chaos::TGeometryParticleHandle<FReal, Dimensions>* ParticleHandle;
+		TUniquePtr<Chaos::TImplicitObject<float, 3>> Geometry;
+		TArray<TUniquePtr<Chaos::TPerShapeData<float, 3>>> Shapes;
+		FVector ActorToCoMTranslation;
+	};
+
+}
+
+#endif // INCLUDE_CHAOS
\ No newline at end of file
diff --git a/Engine/Source/Runtime/Engine/Public/Physics/ImmediatePhysics/ImmediatePhysicsChaos/ImmediatePhysicsCore_Chaos.h b/Engine/Source/Runtime/Engine/Public/Physics/ImmediatePhysics/ImmediatePhysicsChaos/ImmediatePhysicsCore_Chaos.h
new file mode 100644
--- /dev/null
+++ b/Engine/Source/Runtime/Engine/Public/Physics/ImmediatePhysics/ImmediatePhysicsChaos/ImmediatePhysicsCore_Chaos.h
@@ -0,0 +1,39 @@
+// Copyright 1998-2019 Epic Games, Inc. All Rights Reserved.
+
+#pragma once
+
+#if INCLUDE_CHAOS
+
+#include "Physics/ImmediatePhysics/ImmediatePhysicsChaos/ImmediatePhysicsDeclares_Chaos.h"
+#include "Physics/ImmediatePhysics/ImmediatePhysicsShared/ImmediatePhysicsCore.h"
+
+#include "Chaos/ParticleHandleFwd.h"
+
+namespace Chaos
+{
+	template<typename T, int D> class TPBDRigidsEvolutionGBF;
+	template<typename T, int D> class TPBDRigidsSOAs;
+	template<typename T, int D> struct TKinematicGeometryParticleParameters;
+	template<typename T, int D> struct TPBDRigidParticleParameters;
+	template<typename T, int D> class TImplicitObject;
+	template<typename T, int D> class TPBDJointConstraints;
+	template<typename P, typename T, int D> class TPBDConstraintIslandRule;
+	template<typename T, int D> class TPerShapeData;
+}
+
+namespace ImmediatePhysics_Chaos
+{
+	using FReal = float;
+	const int Dimensions = 3;
+
+	using EActorType = ImmediatePhysics_Shared::EActorType;
+	using EForceType = ImmediatePhysics_Shared::EForceType;
+}
+
+struct FBodyInstance;
+struct FConstraintInstance;
+
+#endif
+
+// Used to define out code that still has to be implemented to match PhysX
+#define IMMEDIATEPHYSICS_CHAOS_TODO 0
diff --git a/Engine/Source/Runtime/Engine/Public/Physics/ImmediatePhysics/ImmediatePhysicsChaos/ImmediatePhysicsDeclares_Chaos.h b/Engine/Source/Runtime/Engine/Public/Physics/ImmediatePhysics/ImmediatePhysicsChaos/ImmediatePhysicsDeclares_Chaos.h
new file mode 100644
--- /dev/null
+++ b/Engine/Source/Runtime/Engine/Public/Physics/ImmediatePhysics/ImmediatePhysicsChaos/ImmediatePhysicsDeclares_Chaos.h
@@ -0,0 +1,12 @@
+// Copyright 1998-2019 Epic Games, Inc. All Rights Reserved.
+
+#pragma once
+
+#if INCLUDE_CHAOS
+namespace ImmediatePhysics_Chaos
+{
+	struct FSimulation;
+	struct FActorHandle;
+	struct FJointHandle;
+}
+#endif
diff --git a/Engine/Source/Runtime/Engine/Public/Physics/ImmediatePhysics/ImmediatePhysicsChaos/ImmediatePhysicsJointHandle_Chaos.h b/Engine/Source/Runtime/Engine/Public/Physics/ImmediatePhysics/ImmediatePhysicsChaos/ImmediatePhysicsJointHandle_Chaos.h
new file mode 100644
--- /dev/null
+++ b/Engine/Source/Runtime/Engine/Public/Physics/ImmediatePhysics/ImmediatePhysicsChaos/ImmediatePhysicsJointHandle_Chaos.h
@@ -0,0 +1,25 @@
+// Copyright 1998-2019 Epic Games, Inc. All Rights Reserved.
+
+#pragma once
+
+#if INCLUDE_CHAOS
+
+#include "Physics/ImmediatePhysics/ImmediatePhysicsChaos/ImmediatePhysicsCore_Chaos.h"
+
+#include "Engine/EngineTypes.h"
+
+namespace ImmediatePhysics_Chaos
+{
+	/** handle associated with a physics joint. This is the proper way to read/write to the physics simulation */
+	struct ENGINE_API FJointHandle
+	{
+	public:
+		FJointHandle(Chaos::TPBDJointConstraints<FReal, Dimensions>* InConstraints, FConstraintInstance* ConstraintInstance, FActorHandle* Actor1, FActorHandle* Actor2);
+		~FJointHandle();
+
+	private:
+		Chaos::TPBDJointConstraints<FReal, Dimensions>* Constraints;
+		int32 ConstraintIndex;
+	};
+}
+#endif
diff --git a/Engine/Source/Runtime/Engine/Public/Physics/ImmediatePhysics/ImmediatePhysicsChaos/ImmediatePhysicsSimulation_Chaos.h b/Engine/Source/Runtime/Engine/Public/Physics/ImmediatePhysics/ImmediatePhysicsChaos/ImmediatePhysicsSimulation_Chaos.h
new file mode 100644
--- /dev/null
+++ b/Engine/Source/Runtime/Engine/Public/Physics/ImmediatePhysics/ImmediatePhysicsChaos/ImmediatePhysicsSimulation_Chaos.h
@@ -0,0 +1,103 @@
+// Copyright 1998-2019 Epic Games, Inc. All Rights Reserved.
+
+#pragma once
+
+#if INCLUDE_CHAOS
+
+#include "Physics/ImmediatePhysics/ImmediatePhysicsChaos/ImmediatePhysicsCore_Chaos.h"
+
+#include "Engine/EngineTypes.h"
+#include "Templates/UniquePtr.h"
+
+namespace ImmediatePhysics_Chaos
+{
+	/** Owns all the data associated with the simulation. Can be considered a single scene or world */
+	struct ENGINE_API FSimulation
+	{
+	public:
+		FSimulation();
+		~FSimulation();
+
+		int32 NumActors() const { return ActorHandles.Num(); }
+
+		FActorHandle* GetActorHandle(int32 ActorHandleIndex) { return ActorHandles[ActorHandleIndex]; }
+		const FActorHandle* GetActorHandle(int32 ActorHandleIndex) const { return ActorHandles[ActorHandleIndex]; }
+
+		/** Create a static body and add it to the simulation */
+		FActorHandle* CreateStaticActor(FBodyInstance* BodyInstance);
+
+		/** Create a kinematic body and add it to the simulation */
+		FActorHandle* CreateKinematicActor(FBodyInstance* BodyInstance, const FTransform& Transform);
+
+		/** Create a dynamic body and add it to the simulation */
+		FActorHandle* CreateDynamicActor(FBodyInstance* BodyInstance, const FTransform& Transform);
+
+		FActorHandle* CreateActor(EActorType ActorType, FBodyInstance* BodyInstance, const FTransform& Transform);
+		void DestroyActor(FActorHandle* ActorHandle);
+
+		/** Create a physical joint and add it to the simulation */
+		FJointHandle* CreateJoint(FConstraintInstance* ConstraintInstance, FActorHandle* Body1, FActorHandle* Body2);
+		void DestroyJoint(FJointHandle* JointHandle);
+
+		/** Sets the number of active bodies. This number is reset any time a new simulated body is created */
+		void SetNumActiveBodies(int32 NumActiveBodies);
+
+		/** An array of actors to ignore. */
+		struct FIgnorePair
+		{
+			FActorHandle* A;
+			FActorHandle* B;
+		};
+
+		/** Set pair of bodies to ignore collision for */
+		void SetIgnoreCollisionPairTable(const TArray<FIgnorePair>& InIgnoreCollisionPairTable);
+
+		/** Set bodies that require no collision */
+		void SetIgnoreCollisionActors(const TArray<FActorHandle*>& InIgnoreCollisionActors);
+
+		///** Get/Set whether the body is kinematic or not, kinematics do not simulate and move where they're told. They also act as if they have infinite mass */
+		//bool GetIsKinematic(int32 ActorDataIndex);
+		//void SetIsKinematic(int32 ActorDataIndex, bool bKinematic);
+
+		/** Advance the simulation by DeltaTime */
+		void Simulate(float DeltaTime, const FVector& InGravity);
+		void Simulate_AssumesLocked(float DeltaTime, const FVector& InGravity) { Simulate(DeltaTime, InGravity); }
+
+		///** Whether or not an entity is simulated */
+		//bool IsSimulated(int32 ActorDataIndex) const;
+
+
+		///** Add a radial impulse to the given actor */
+		//void AddRadialForce(int32 ActorDataIndex, const FVector& Origin, float Strength, float Radius, ERadialImpulseFalloff Falloff, EForceType ForceType);
+
+		///** Add a force to the given actor */
+		//void AddForce(int32 ActorDataIndex, const FVector& Force);
+
+		///* Set solver iteration counts per step */
+		//void SetPositionIterationCount(int32 InIterationCount);
+		//void SetVelocityIterationCount(int32 InIterationCount);
+
+	private:
+		TUniquePtr<Chaos::TPBDRigidsEvolutionGBF<FReal, Dimensions>> Evolution;
+		TUniquePtr<Chaos::TPBDRigidsSOAs<FReal, Dimensions>> Particles;
+		TUniquePtr<Chaos::TPBDJointConstraints<FReal, Dimensions>> Joints;
+		TUniquePtr<Chaos::TPBDConstraintIslandRule<Chaos::TPBDJointConstraints<FReal, Dimensions>, FReal, Dimensions>> JointsRule;		// @todo(ccaulfield): Don't need islands for anim node physics...
+
+		/** Mapping from entity index to handle */
+		// @todo(ccaulfield): we now have handles pointing to handles which is inefficient - we can do better than this, but don't want to change API yet
+		TArray<FActorHandle*> ActorHandles;
+		int NumActiveActorHandles;
+
+		/** Mapping from constraint index to handle */
+		TArray<FJointHandle*> JointHandles;
+
+		/** Both of these are slow to access. Make sure to use iteration cache when possible */
+		TMap<FActorHandle*, TSet<FActorHandle*>> IgnoreCollisionPairTable;
+		TSet<FActorHandle*> IgnoreCollisionActors;
+
+		FVector Gravity;
+	};
+
+}
+
+#endif // INCLUDE_CHAOS
\ No newline at end of file
diff --git a/Engine/Source/Runtime/Engine/Public/Physics/ImmediatePhysics/ImmediatePhysicsConstraintAllocator.h b/Engine/Source/Runtime/Engine/Public/Physics/ImmediatePhysics/ImmediatePhysicsConstraintAllocator.h
deleted file mode 100644
diff --git a/Engine/Source/Runtime/Engine/Public/Physics/ImmediatePhysics/ImmediatePhysicsContactPair.h b/Engine/Source/Runtime/Engine/Public/Physics/ImmediatePhysics/ImmediatePhysicsContactPair.h
deleted file mode 100644
diff --git a/Engine/Source/Runtime/Engine/Public/Physics/ImmediatePhysics/ImmediatePhysicsContactPointRecorder.h b/Engine/Source/Runtime/Engine/Public/Physics/ImmediatePhysics/ImmediatePhysicsContactPointRecorder.h
deleted file mode 100644
diff --git a/Engine/Source/Runtime/Engine/Public/Physics/ImmediatePhysics/ImmediatePhysicsD6JointData.h b/Engine/Source/Runtime/Engine/Public/Physics/ImmediatePhysics/ImmediatePhysicsD6JointData.h
deleted file mode 100644
diff --git a/Engine/Source/Runtime/Engine/Public/Physics/ImmediatePhysics/ImmediatePhysicsDeclares.h b/Engine/Source/Runtime/Engine/Public/Physics/ImmediatePhysics/ImmediatePhysicsDeclares.h
new file mode 100644
--- /dev/null
+++ b/Engine/Source/Runtime/Engine/Public/Physics/ImmediatePhysics/ImmediatePhysicsDeclares.h
@@ -0,0 +1,37 @@
+// Copyright 1998-2019 Epic Games, Inc. All Rights Reserved.
+
+#pragma once
+
+#include "Physics/ImmediatePhysics/ImmediatePhysicsPhysX/ImmediatePhysicsDeclares_PhysX.h"
+#include "Physics/ImmediatePhysics/ImmediatePhysicsChaos/ImmediatePhysicsDeclares_Chaos.h"
+#include "Physics/ImmediatePhysics/ImmediatePhysicsShared/ImmediatePhysicsCore.h"
+
+
+/**
+ * @file ImmediatePhysicsDeclares.h
+ * 
+ * Move either ImmediatePhysics_PhysX or ImmediatePhysics_Chaos Immediate Physics into the ImmediatePhysics namespace
+ *
+ * We can use PhysX ImmediatePhysics only if PhysX is also providing the global physics interface.
+ * because ImmediatePhysics_PhysX relies on physx types being instantiated by BodyInstance, ConstraintInstance, etc.
+ * 
+ * You can currently run PhysX and Chaos ImmediatePhysics simulations in the same build if PhysX is
+ * providing the global physics interface. This capability will likely disappear once Chaos is established
+ * at which point you will only be able to use Chaos immediate physics if you are also using Chaos
+ * for global physics.
+ */
+namespace ImmediatePhysics
+{
+	using EActorType = ImmediatePhysics_Shared::EActorType;
+	using EForceType = ImmediatePhysics_Shared::EForceType;
+
+#if WITH_PHYSX && PHYSICS_INTERFACE_PHYSX
+	using FSimulation = ImmediatePhysics_PhysX::FSimulation;
+	using FActorHandle = ImmediatePhysics_PhysX::FActorHandle;
+#elif WITH_CHAOS
+	using FSimulation = ImmediatePhysics_Chaos::FSimulation;
+	using FActorHandle = ImmediatePhysics_Chaos::FActorHandle;
+#else
+#error Global ImmediatePhysics is not defined
+#endif
+}
\ No newline at end of file
diff --git a/Engine/Source/Runtime/Engine/Public/Physics/ImmediatePhysics/ImmediatePhysicsJoint.h b/Engine/Source/Runtime/Engine/Public/Physics/ImmediatePhysics/ImmediatePhysicsJoint.h
deleted file mode 100644
diff --git a/Engine/Source/Runtime/Engine/Public/Physics/ImmediatePhysics/ImmediatePhysicsJointHandle.h b/Engine/Source/Runtime/Engine/Public/Physics/ImmediatePhysics/ImmediatePhysicsJointHandle.h
deleted file mode 100644
diff --git a/Engine/Source/Runtime/Engine/Public/Physics/ImmediatePhysics/ImmediatePhysicsKinematicTarget.h b/Engine/Source/Runtime/Engine/Public/Physics/ImmediatePhysics/ImmediatePhysicsKinematicTarget.h
deleted file mode 100644
diff --git a/Engine/Source/Runtime/Engine/Public/Physics/ImmediatePhysics/ImmediatePhysicsLinearBlockAllocator.h b/Engine/Source/Runtime/Engine/Public/Physics/ImmediatePhysics/ImmediatePhysicsLinearBlockAllocator.h
deleted file mode 100644
diff --git a/Engine/Source/Runtime/Engine/Public/Physics/ImmediatePhysics/ImmediatePhysicsMaterial.h b/Engine/Source/Runtime/Engine/Public/Physics/ImmediatePhysics/ImmediatePhysicsMaterial.h
deleted file mode 100644
diff --git a/Engine/Source/Runtime/Engine/Public/Physics/ImmediatePhysics/ImmediatePhysicsMaterialHandle.h b/Engine/Source/Runtime/Engine/Public/Physics/ImmediatePhysics/ImmediatePhysicsMaterialHandle.h
deleted file mode 100644
diff --git a/Engine/Source/Runtime/Engine/Public/Physics/ImmediatePhysics/ImmediatePhysicsPersistentContactPairData.h b/Engine/Source/Runtime/Engine/Public/Physics/ImmediatePhysics/ImmediatePhysicsPersistentContactPairData.h
deleted file mode 100644
diff --git a/Engine/Source/Runtime/Engine/Public/Physics/ImmediatePhysics/ImmediatePhysicsPhysX/ImmediatePhysicsActorHandle_PhysX.h b/Engine/Source/Runtime/Engine/Public/Physics/ImmediatePhysics/ImmediatePhysicsPhysX/ImmediatePhysicsActorHandle_PhysX.h
new file mode 100644
--- /dev/null
+++ b/Engine/Source/Runtime/Engine/Public/Physics/ImmediatePhysics/ImmediatePhysicsPhysX/ImmediatePhysicsActorHandle_PhysX.h
@@ -0,0 +1,248 @@
+// Copyright 1998-2019 Epic Games, Inc. All Rights Reserved.
+
+#pragma once
+
+#if WITH_PHYSX && PHYSICS_INTERFACE_PHYSX
+#include "PhysXPublic.h"
+#include "Physics/ImmediatePhysics/ImmediatePhysicsPhysX/ImmediatePhysicsSimulation_PhysX.h"
+#include "Engine/EngineTypes.h"
+
+namespace ImmediatePhysics_PhysX
+{
+
+	/** handle associated with a physics actor. This is the proper way to read/write to the physics simulation */
+	struct ENGINE_API FActorHandle
+	{
+
+	public:
+		/** Sets the world transform.*/
+		void SetWorldTransform(const FTransform& WorldTM)
+		{
+			OwningSimulation.GetLowLevelBody(ActorDataIndex).body2World = U2PTransform(ActorToBody * WorldTM);
+		}
+
+		/** Make a body kinematic, or non-kinematic */
+		void SetIsKinematic(bool bKinematic)
+		{
+			OwningSimulation.SetIsKinematic(ActorDataIndex, bKinematic);
+		}
+
+		/** Is the actor kinematic */
+		bool GetIsKinematic() const
+		{
+			return OwningSimulation.GetIsKinematic(ActorDataIndex);
+		}
+
+		/** Gets the kinematic target (next transform) for the actor if one is set (check HasKinematicTarget() to see if a target is available) */
+		FImmediateKinematicTarget& GetKinematicTarget()
+		{
+			return OwningSimulation.GetKinematicTarget(ActorDataIndex);
+		}
+
+		/** Sets the kinematic target. This will affect velocities as expected*/
+		void SetKinematicTarget(const FTransform& WorldTM)
+		{
+			FImmediateKinematicTarget& KinematicTarget = OwningSimulation.GetKinematicTarget(ActorDataIndex);
+			KinematicTarget.BodyToWorld = U2PTransform(ActorToBody * WorldTM);
+			KinematicTarget.bTargetSet = true;
+		}
+
+		/** Does this actor have a kinematic target (next kinematic transform to be applied) */
+		bool HasKinematicTarget() const
+		{
+			FImmediateKinematicTarget& KinematicTarget = OwningSimulation.GetKinematicTarget(ActorDataIndex);
+			return KinematicTarget.bTargetSet;
+		}
+
+		/** Whether the body is simulating */
+		bool IsSimulated() const
+		{
+			return OwningSimulation.IsSimulated(ActorDataIndex);
+		}
+
+		/** Get the world transform */
+		FTransform GetWorldTransform() const
+		{
+			return ActorToBody.GetRelativeTransformReverse(P2UTransform(OwningSimulation.GetLowLevelBody(ActorDataIndex).body2World));
+		}
+
+		/** Set the linear velocity */
+		void SetLinearVelocity(const FVector& NewLinearVelocity)
+		{
+			OwningSimulation.GetLowLevelBody(ActorDataIndex).linearVelocity = U2PVector(NewLinearVelocity);
+		}
+
+		/** Get the linear velocity */
+		FVector GetLinearVelocity() const
+		{
+			return P2UVector(OwningSimulation.GetLowLevelBody(ActorDataIndex).linearVelocity);
+		}
+
+		/** Set the angular velocity */
+		void SetAngularVelocity(const FVector& NewAngularVelocity)
+		{
+			OwningSimulation.GetLowLevelBody(ActorDataIndex).angularVelocity = U2PVector(NewAngularVelocity);
+		}
+
+		/** Get the angular velocity */
+		FVector GetAngularVelocity() const
+		{
+			return P2UVector(OwningSimulation.GetLowLevelBody(ActorDataIndex).angularVelocity);
+		}
+
+		void AddForce(const FVector& Force)
+		{
+			OwningSimulation.AddForce(ActorDataIndex, Force);
+		}
+
+		void AddRadialForce(const FVector& Origin, float Strength, float Radius, ERadialImpulseFalloff Falloff, EForceType ForceType)
+		{
+			OwningSimulation.AddRadialForce(ActorDataIndex, Origin, Strength, Radius, Falloff, ForceType);
+		}
+
+		/** Set the linear damping*/
+		void SetLinearDamping(float NewLinearDamping)
+		{
+			OwningSimulation.GetLowLevelBody(ActorDataIndex).linearDamping = NewLinearDamping;
+		}
+
+		/** Get the linear damping*/
+		float GetLinearDamping() const
+		{
+			return OwningSimulation.GetLowLevelBody(ActorDataIndex).linearDamping;
+		}
+
+		/** Set the angular damping*/
+		void SetAngularDamping(float NewAngularDamping)
+		{
+			OwningSimulation.GetLowLevelBody(ActorDataIndex).angularDamping = NewAngularDamping;
+		}
+
+		/** Get the angular damping*/
+		float GetAngularDamping() const
+		{
+			return OwningSimulation.GetLowLevelBody(ActorDataIndex).angularDamping;
+		}
+
+		/** Set the max linear velocity squared*/
+		void SetMaxLinearVelocitySquared(float NewMaxLinearVelocitySquared)
+		{
+			OwningSimulation.GetLowLevelBody(ActorDataIndex).maxLinearVelocitySq = NewMaxLinearVelocitySquared;
+		}
+
+		/** Get the max linear velocity squared*/
+		float GetMaxLinearVelocitySquared() const
+		{
+			return OwningSimulation.GetLowLevelBody(ActorDataIndex).maxLinearVelocitySq;
+		}
+
+		/** Set the max angular velocity squared*/
+		void SetMaxAngularVelocitySquared(float NewMaxAngularVelocitySquared)
+		{
+			OwningSimulation.GetLowLevelBody(ActorDataIndex).maxAngularVelocitySq = NewMaxAngularVelocitySquared;
+		}
+
+		/** Get the max angular velocity squared*/
+		float GetMaxAngularVelocitySquared() const
+		{
+			return OwningSimulation.GetLowLevelBody(ActorDataIndex).maxAngularVelocitySq;
+		}
+
+		/** Set the inverse mass. 0 indicates kinematic object */
+		void SetInverseMass(float NewInverseMass)
+		{
+			OwningSimulation.GetLowLevelBody(ActorDataIndex).invMass = NewInverseMass;
+		}
+
+		/** Get the inverse mass. */
+		float GetInverseMass() const
+		{
+			return OwningSimulation.GetLowLevelBody(ActorDataIndex).invMass;
+		}
+
+		/** Set the inverse inertia. Mass-space inverse inertia diagonal vector */
+		void SetInverseInertia(const FVector& NewInverseInertia)
+		{
+			OwningSimulation.GetLowLevelBody(ActorDataIndex).invInertia = U2PVector(NewInverseInertia);
+		}
+
+		/** Get the inverse inertia. Mass-space inverse inertia diagonal vector */
+		FVector GetInverseInertia() const
+		{
+			return P2UVector(OwningSimulation.GetLowLevelBody(ActorDataIndex).invInertia);
+		}
+
+		/** Set the max depenetration velocity*/
+		void SetMaxDepenetrationVelocity(float NewMaxDepenetrationVelocity)
+		{
+			OwningSimulation.GetLowLevelBody(ActorDataIndex).maxDepenetrationVelocity = NewMaxDepenetrationVelocity;
+		}
+
+		/** Get the max depenetration velocity*/
+		float GetMaxDepenetrationVelocity(float NewMaxDepenetrationVelocity) const
+		{
+			return OwningSimulation.GetLowLevelBody(ActorDataIndex).maxDepenetrationVelocity;
+		}
+
+		/** Set the max contact impulse*/
+		void SetMaxContactImpulse(float NewMaxContactImpulse)
+		{
+			OwningSimulation.GetLowLevelBody(ActorDataIndex).maxContactImpulse = NewMaxContactImpulse;
+		}
+
+		/** Get the max contact impulse*/
+		float GetMaxContactImpulse() const
+		{
+			return OwningSimulation.GetLowLevelBody(ActorDataIndex).maxContactImpulse;
+		}
+
+		const int32 GetActorIndex() const
+		{
+			return ActorDataIndex;
+		}
+
+		void AddShape(physx::PxShape* InShape)
+		{
+			OwningSimulation.Actors[ActorDataIndex].AddShape(InShape);
+		}
+
+		TArray<FShape>& GetShapes() const
+		{
+			return OwningSimulation.Actors[ActorDataIndex].Shapes;
+		}
+
+		FActor* GetSimulationActor() const
+		{
+			return &OwningSimulation.Actors[ActorDataIndex];
+		}
+
+		immediate::PxRigidBodyData* GetSimulationRigidBodyData() const
+		{
+			immediate::PxRigidBodyData& Data = OwningSimulation.GetLowLevelBody(ActorDataIndex);
+			return &Data;
+		}
+
+	private:
+		/** Converts from actor space (i.e. the transform in world space as the client gives us) to body space (body with its origin at the COM and oriented to inertia tensor) */
+		FTransform ActorToBody;
+		FSimulation& OwningSimulation;
+		int32 ActorDataIndex;
+
+		friend FSimulation;
+		FActorHandle(FSimulation& InOwningSimulation, int32 InActorDataIndex)
+			: ActorToBody(FTransform::Identity)
+			, OwningSimulation(InOwningSimulation)
+			, ActorDataIndex(InActorDataIndex)
+		{
+		}
+
+		~FActorHandle()
+		{
+		}
+
+		FActorHandle(const FActorHandle&);	//Ensure no copying of handles
+	};
+
+}
+
+#endif // WITH_PHYSX
diff --git a/Engine/Source/Runtime/Engine/Public/Physics/ImmediatePhysics/ImmediatePhysicsPhysX/ImmediatePhysicsActor_PhysX.h b/Engine/Source/Runtime/Engine/Public/Physics/ImmediatePhysics/ImmediatePhysicsPhysX/ImmediatePhysicsActor_PhysX.h
new file mode 100644
--- /dev/null
+++ b/Engine/Source/Runtime/Engine/Public/Physics/ImmediatePhysics/ImmediatePhysicsPhysX/ImmediatePhysicsActor_PhysX.h
@@ -0,0 +1,32 @@
+// Copyright 1998-2019 Epic Games, Inc. All Rights Reserved.
+
+#pragma once
+
+#if WITH_PHYSX && PHYSICS_INTERFACE_PHYSX
+#include "PhysXPublic.h"
+#include "Physics/ImmediatePhysics/ImmediatePhysicsPhysX/ImmediatePhysicsShape_PhysX.h"
+
+namespace ImmediatePhysics_PhysX
+{
+
+	/** Holds geometry data*/
+	struct FActor
+	{
+		FActor()
+			: UserData(nullptr)
+		{}
+
+		TArray<FShape> Shapes;
+		void* UserData;
+
+		/** Create geometry data for the entity */
+		void CreateGeometry(PxRigidActor* RigidActor, const PxTransform& ActorToBodyTM);
+		bool AddShape(PxShape* InShape);
+
+		/** Ensures all the geometry data has been properly freed */
+		void TerminateGeometry();
+	};
+
+}
+
+#endif // WITH_PHYSX
\ No newline at end of file
diff --git a/Engine/Source/Runtime/Engine/Public/Physics/ImmediatePhysics/ImmediatePhysicsPhysX/ImmediatePhysicsCacheAllocator_PhysX.h b/Engine/Source/Runtime/Engine/Public/Physics/ImmediatePhysics/ImmediatePhysicsPhysX/ImmediatePhysicsCacheAllocator_PhysX.h
new file mode 100644
--- /dev/null
+++ b/Engine/Source/Runtime/Engine/Public/Physics/ImmediatePhysics/ImmediatePhysicsPhysX/ImmediatePhysicsCacheAllocator_PhysX.h
@@ -0,0 +1,35 @@
+// Copyright 1998-2019 Epic Games, Inc. All Rights Reserved.
+
+#pragma once
+
+#if WITH_PHYSX && PHYSICS_INTERFACE_PHYSX
+#include "PhysXPublic.h"
+#include "Physics/ImmediatePhysics/ImmediatePhysicsPhysX/ImmediatePhysicsLinearBlockAllocator_PhysX.h"
+
+namespace ImmediatePhysics_PhysX
+{
+
+	/** TODO: Use a smarter memory allocator */
+	struct FCacheAllocator : public PxCacheAllocator
+	{
+		FCacheAllocator() : External(0) {}
+		PxU8* allocateCacheData(const PxU32 ByteSize) override
+		{
+			return BlockAllocator[External].Alloc(ByteSize);
+		}
+
+		void Reset()
+		{
+#if PERSISTENT_CONTACT_PAIRS
+			External = 1 - External;	//flip buffer so we maintain cache for 1 extra step
+#endif
+			BlockAllocator[External].Reset();
+		}
+
+		FLinearBlockAllocator BlockAllocator[2];
+		int32 External;
+	};
+
+}
+
+#endif // WITH_PHYSX
\ No newline at end of file
diff --git a/Engine/Source/Runtime/Engine/Public/Physics/ImmediatePhysics/ImmediatePhysicsPhysX/ImmediatePhysicsConstraintAllocator_PhysX.h b/Engine/Source/Runtime/Engine/Public/Physics/ImmediatePhysics/ImmediatePhysicsPhysX/ImmediatePhysicsConstraintAllocator_PhysX.h
new file mode 100644
--- /dev/null
+++ b/Engine/Source/Runtime/Engine/Public/Physics/ImmediatePhysics/ImmediatePhysicsPhysX/ImmediatePhysicsConstraintAllocator_PhysX.h
@@ -0,0 +1,43 @@
+// Copyright 1998-2019 Epic Games, Inc. All Rights Reserved.
+
+#pragma once
+
+#if WITH_PHYSX && PHYSICS_INTERFACE_PHYSX
+#include "PhysXPublic.h"
+#include "Physics/ImmediatePhysics/ImmediatePhysicsPhysX/ImmediatePhysicsLinearBlockAllocator_PhysX.h"
+
+namespace ImmediatePhysics_PhysX
+{
+
+	/** TODO: Use a smarter memory allocator */
+	class FConstraintAllocator : public PxConstraintAllocator
+	{
+	public:
+		FConstraintAllocator() : External(0) {}
+
+		PxU8* reserveConstraintData(const PxU32 ByteSize) override
+		{
+			return BlockAllocator[External].Alloc(ByteSize);
+
+		}
+
+		PxU8* reserveFrictionData(const PxU32 ByteSize)
+		{
+			return BlockAllocator[External].Alloc(ByteSize);
+		}
+
+		void Reset()
+		{
+#if PERSISTENT_CONTACT_PAIRS
+			External = 1 - External;	//flip buffer so we maintain cache for 1 extra step
+#endif
+			BlockAllocator[External].Reset();
+		}
+
+		FLinearBlockAllocator BlockAllocator[2];
+		int32 External;
+	};
+
+}
+
+#endif // WITH_PHYSX
\ No newline at end of file
diff --git a/Engine/Source/Runtime/Engine/Public/Physics/ImmediatePhysics/ImmediatePhysicsPhysX/ImmediatePhysicsContactPair_PhysX.h b/Engine/Source/Runtime/Engine/Public/Physics/ImmediatePhysics/ImmediatePhysicsPhysX/ImmediatePhysicsContactPair_PhysX.h
new file mode 100644
--- /dev/null
+++ b/Engine/Source/Runtime/Engine/Public/Physics/ImmediatePhysics/ImmediatePhysicsPhysX/ImmediatePhysicsContactPair_PhysX.h
@@ -0,0 +1,32 @@
+// Copyright 1998-2019 Epic Games, Inc. All Rights Reserved.
+
+#pragma once
+
+#if WITH_PHYSX && PHYSICS_INTERFACE_PHYSX
+#include "CoreMinimal.h"
+
+namespace ImmediatePhysics_PhysX
+{
+
+	/** Contact pair generated between entities */
+	struct FContactPair
+	{
+		/** Index of the dynamic actor that we generated the contact pair for*/
+		int32 DynamicActorDataIndex;
+
+		/** Index of the other actor that we generated the contact pair for. This could be either dynamic or static */
+		uint32 OtherActorDataIndex;
+
+		/** Index into the first contact point associated with this pair*/
+		uint32 StartContactIndex;
+
+		/** Number of contacts associated with this pair. */
+		uint32 NumContacts;
+
+		/** Identifies the pair index from the original contact generation test */
+		int32 PairIdx;
+	};
+
+}
+
+#endif // WITH_PHYSX
\ No newline at end of file
diff --git a/Engine/Source/Runtime/Engine/Public/Physics/ImmediatePhysics/ImmediatePhysicsPhysX/ImmediatePhysicsContactPointRecorder_PhysX.h b/Engine/Source/Runtime/Engine/Public/Physics/ImmediatePhysics/ImmediatePhysicsPhysX/ImmediatePhysicsContactPointRecorder_PhysX.h
new file mode 100644
--- /dev/null
+++ b/Engine/Source/Runtime/Engine/Public/Physics/ImmediatePhysics/ImmediatePhysicsPhysX/ImmediatePhysicsContactPointRecorder_PhysX.h
@@ -0,0 +1,37 @@
+// Copyright 1998-2019 Epic Games, Inc. All Rights Reserved.
+
+#pragma once
+
+#if WITH_PHYSX && PHYSICS_INTERFACE_PHYSX
+#include "PhysXPublic.h"
+
+namespace ImmediatePhysics_PhysX
+{
+	struct FSimulation;
+	struct FMaterial;
+
+	struct FContactPointRecorder : public immediate::PxContactRecorder
+	{
+		FContactPointRecorder(struct FSimulation& InSimulation, int32 InDynamicActorDataIndex, int32 InOtherActorDataIndex, int32 InPairIdx, const FMaterial& InSimulatedShapeMaterial, const FMaterial& InOtherShapeMaterial)
+			: Simulation(InSimulation)
+			, SimulatedShapeMaterial(InSimulatedShapeMaterial)
+			, OtherShapeMaterial(InOtherShapeMaterial)
+			, DynamicActorDataIndex(InDynamicActorDataIndex)
+			, OtherActorDataIndex(InOtherActorDataIndex)
+			, PairIdx(InPairIdx)
+		{
+		}
+
+		bool recordContacts(const Gu::ContactPoint* ContactPoints, const PxU32 NumContacts, const PxU32 Index) override;
+
+		FSimulation& Simulation;
+		const FMaterial& SimulatedShapeMaterial;
+		const FMaterial& OtherShapeMaterial;
+		int32 DynamicActorDataIndex;
+		int32 OtherActorDataIndex;
+		int32 PairIdx;
+	};
+
+}
+
+#endif // WITH_PHYSX
\ No newline at end of file
diff --git a/Engine/Source/Runtime/Engine/Public/Physics/ImmediatePhysics/ImmediatePhysicsPhysX/ImmediatePhysicsCore_PhysX.h b/Engine/Source/Runtime/Engine/Public/Physics/ImmediatePhysics/ImmediatePhysicsPhysX/ImmediatePhysicsCore_PhysX.h
new file mode 100644
--- /dev/null
+++ b/Engine/Source/Runtime/Engine/Public/Physics/ImmediatePhysics/ImmediatePhysicsPhysX/ImmediatePhysicsCore_PhysX.h
@@ -0,0 +1,18 @@
+// Copyright 1998-2019 Epic Games, Inc. All Rights Reserved.
+
+#pragma once
+
+#include "Physics/ImmediatePhysics/ImmediatePhysicsShared/ImmediatePhysicsCore.h"
+
+#if WITH_PHYSX && PHYSICS_INTERFACE_PHYSX
+
+namespace ImmediatePhysics_PhysX
+{
+	using EActorType = ImmediatePhysics_Shared::EActorType;
+	using EForceType = ImmediatePhysics_Shared::EForceType;
+}
+
+struct FBodyInstance;
+struct FConstraintInstance;
+
+#endif
diff --git a/Engine/Source/Runtime/Engine/Public/Physics/ImmediatePhysics/ImmediatePhysicsPhysX/ImmediatePhysicsD6JointData_PhysX.h b/Engine/Source/Runtime/Engine/Public/Physics/ImmediatePhysics/ImmediatePhysicsPhysX/ImmediatePhysicsD6JointData_PhysX.h
new file mode 100644
--- /dev/null
+++ b/Engine/Source/Runtime/Engine/Public/Physics/ImmediatePhysics/ImmediatePhysicsPhysX/ImmediatePhysicsD6JointData_PhysX.h
@@ -0,0 +1,71 @@
+// Copyright 1998-2019 Epic Games, Inc. All Rights Reserved.
+
+#pragma once
+
+#if WITH_PHYSX && PHYSICS_INTERFACE_PHYSX
+#include "PhysXPublic.h"
+
+//TODO: this is just straight copied from physx
+namespace ImmediatePhysics_PhysX
+{
+
+	struct D6JointData
+	{
+		D6JointData(PxD6Joint* Joint);
+
+		PxConstraintInvMassScale	invMassScale;
+		PxTransform					c2b[2];
+
+		PxU32					locked;		// bitmap of locked DOFs
+		PxU32					limited;	// bitmap of limited DOFs
+		PxU32					driving;	// bitmap of active drives (implies driven DOFs not locked)
+
+
+		PxD6Motion::Enum		motion[6];
+		PxJointLinearLimit		linearLimit;
+		PxJointAngularLimitPair	twistLimit;
+		PxJointLimitCone		swingLimit;
+
+		PxD6JointDrive			drive[PxD6Drive::eCOUNT];
+
+		PxTransform				drivePosition;
+		PxVec3					driveLinearVelocity;
+		PxVec3					driveAngularVelocity;
+
+		// derived quantities
+
+
+		// tan-half and tan-quarter angles
+
+		PxReal					thSwingY;
+		PxReal					thSwingZ;
+		PxReal					thSwingPad;
+
+		PxReal					tqSwingY;
+		PxReal					tqSwingZ;
+		PxReal					tqSwingPad;
+
+		PxReal					tqTwistLow;
+		PxReal					tqTwistHigh;
+		PxReal					tqTwistPad;
+
+		PxReal					linearMinDist;	// linear limit minimum distance to get a good direction
+
+		// projection quantities
+		//PxReal					projectionLinearTolerance;
+		//PxReal					projectionAngularTolerance;
+
+		FTransform					ActorToBody[2];
+
+		bool HasConstraints() const
+		{
+			return locked || limited || driving;
+		}
+	};
+
+	void PrepareJointData(D6JointData& JointData);
+	PxU32 D6JointSolverPrep(Px1DConstraint* constraints, PxVec3& body0WorldOffset, PxU32 maxConstraints, PxConstraintInvMassScale& invMassScale, const void* constantBlock, const PxTransform& bA2w, const PxTransform& bB2w);
+
+}
+
+#endif // WITH_PHYSX
diff --git a/Engine/Source/Runtime/Engine/Public/Physics/ImmediatePhysics/ImmediatePhysicsPhysX/ImmediatePhysicsDeclares_PhysX.h b/Engine/Source/Runtime/Engine/Public/Physics/ImmediatePhysics/ImmediatePhysicsPhysX/ImmediatePhysicsDeclares_PhysX.h
new file mode 100644
--- /dev/null
+++ b/Engine/Source/Runtime/Engine/Public/Physics/ImmediatePhysics/ImmediatePhysicsPhysX/ImmediatePhysicsDeclares_PhysX.h
@@ -0,0 +1,11 @@
+// Copyright 1998-2019 Epic Games, Inc. All Rights Reserved.
+
+#pragma once
+
+#if WITH_PHYSX && PHYSICS_INTERFACE_PHYSX
+namespace ImmediatePhysics_PhysX
+{
+	struct FSimulation;
+	struct FActorHandle;
+}
+#endif
diff --git a/Engine/Source/Runtime/Engine/Public/Physics/ImmediatePhysics/ImmediatePhysicsPhysX/ImmediatePhysicsJointHandle_PhysX.h b/Engine/Source/Runtime/Engine/Public/Physics/ImmediatePhysics/ImmediatePhysicsPhysX/ImmediatePhysicsJointHandle_PhysX.h
new file mode 100644
--- /dev/null
+++ b/Engine/Source/Runtime/Engine/Public/Physics/ImmediatePhysics/ImmediatePhysicsPhysX/ImmediatePhysicsJointHandle_PhysX.h
@@ -0,0 +1,36 @@
+// Copyright 1998-2019 Epic Games, Inc. All Rights Reserved.
+
+#pragma once
+
+#if WITH_PHYSX && PHYSICS_INTERFACE_PHYSX
+#include "PhysXPublic.h"
+
+namespace ImmediatePhysics_PhysX
+{
+	struct FSimulation;
+
+	/** handle associated with a physics joint. This is the proper way to read/write to the physics simulation */
+	struct ENGINE_API FJointHandle
+	{
+	private:
+		FSimulation& OwningSimulation;
+		int32 JointDataIndex;
+
+		friend FSimulation;
+		FJointHandle(FSimulation& InOwningSimulation, int32 InJointDataIndex)
+			: OwningSimulation(InOwningSimulation)
+			, JointDataIndex(InJointDataIndex)
+		{
+		}
+
+		~FJointHandle()
+		{
+		}
+
+		FJointHandle(const FJointHandle&);	//Ensure no copying of handles
+
+	};
+
+}
+
+#endif // WITH_PHYSX
diff --git a/Engine/Source/Runtime/Engine/Public/Physics/ImmediatePhysics/ImmediatePhysicsPhysX/ImmediatePhysicsJoint_PhysX.h b/Engine/Source/Runtime/Engine/Public/Physics/ImmediatePhysics/ImmediatePhysicsPhysX/ImmediatePhysicsJoint_PhysX.h
new file mode 100644
--- /dev/null
+++ b/Engine/Source/Runtime/Engine/Public/Physics/ImmediatePhysics/ImmediatePhysicsPhysX/ImmediatePhysicsJoint_PhysX.h
@@ -0,0 +1,20 @@
+// Copyright 1998-2019 Epic Games, Inc. All Rights Reserved.
+
+#pragma once
+
+#if WITH_PHYSX && PHYSICS_INTERFACE_PHYSX
+
+namespace ImmediatePhysics_PhysX
+{
+
+	struct FActorHandle;
+
+	struct FJoint
+	{
+		FActorHandle* DynamicActor;
+		FActorHandle* OtherActor;
+	};
+
+}
+
+#endif // WITH_PHYSX
\ No newline at end of file
diff --git a/Engine/Source/Runtime/Engine/Public/Physics/ImmediatePhysics/ImmediatePhysicsPhysX/ImmediatePhysicsKinematicTarget_PhysX.h b/Engine/Source/Runtime/Engine/Public/Physics/ImmediatePhysics/ImmediatePhysicsPhysX/ImmediatePhysicsKinematicTarget_PhysX.h
new file mode 100644
--- /dev/null
+++ b/Engine/Source/Runtime/Engine/Public/Physics/ImmediatePhysics/ImmediatePhysicsPhysX/ImmediatePhysicsKinematicTarget_PhysX.h
@@ -0,0 +1,24 @@
+// Copyright 1998-2019 Epic Games, Inc. All Rights Reserved.
+
+#pragma once
+
+#if WITH_PHYSX && PHYSICS_INTERFACE_PHYSX
+#include "PhysXPublic.h"
+
+namespace ImmediatePhysics_PhysX
+{
+	/** Holds shape data*/
+	struct FImmediateKinematicTarget
+	{
+		PxTransform BodyToWorld;
+		bool bTargetSet;
+
+		FImmediateKinematicTarget()
+			: bTargetSet(false)
+		{
+		}
+	};
+
+}
+
+#endif // WITH_PHYSX
diff --git a/Engine/Source/Runtime/Engine/Public/Physics/ImmediatePhysics/ImmediatePhysicsPhysX/ImmediatePhysicsLinearBlockAllocator_PhysX.h b/Engine/Source/Runtime/Engine/Public/Physics/ImmediatePhysics/ImmediatePhysicsPhysX/ImmediatePhysicsLinearBlockAllocator_PhysX.h
new file mode 100644
--- /dev/null
+++ b/Engine/Source/Runtime/Engine/Public/Physics/ImmediatePhysics/ImmediatePhysicsPhysX/ImmediatePhysicsLinearBlockAllocator_PhysX.h
@@ -0,0 +1,129 @@
+// Copyright 1998-2019 Epic Games, Inc. All Rights Reserved.
+
+#pragma once
+
+#if WITH_PHYSX && PHYSICS_INTERFACE_PHYSX
+#include "HAL/LowLevelMemTracker.h"
+
+namespace ImmediatePhysics_PhysX
+{
+	const int PageBufferSize = 1024 * 64;	//This has to be here because of VS2013
+
+	struct FLinearBlockAllocator
+	{
+		struct FPageStruct
+		{
+			//We assume FPageStruct is allocated with 16-byte alignment. Do not move Buffer. If we get full support for alignas(16) we could make this better
+			uint8 Buffer[PageBufferSize];
+
+			FPageStruct* NextPage;
+			FPageStruct* PrevPage;
+			int32 SeekPosition;
+
+			FPageStruct()
+				: NextPage(nullptr)
+				, PrevPage(nullptr)
+				, SeekPosition(0)
+			{
+			}
+		};
+
+		FPageStruct* FreePage;
+		FPageStruct* FirstPage;
+
+		FLinearBlockAllocator()
+		{
+			FreePage = AllocPage();
+			FirstPage = FreePage;
+		}
+
+		FPageStruct* AllocPage()
+		{
+			LLM_SCOPE(ELLMTag::PhysX);
+			FPageStruct* ReturnPage = (FPageStruct*)FMemory::Malloc(sizeof(FPageStruct), 16);
+			new(ReturnPage) FPageStruct();
+
+			FPlatformMisc::TagBuffer("ImmediatePhysicsSim", 0, (const void*)ReturnPage, sizeof(FPageStruct));
+			return ReturnPage;
+		}
+
+		void ReleasePage(FPageStruct* Page)
+		{
+			FMemory::Free(Page);
+		}
+
+		uint8* Alloc(int32 Bytes)
+		{
+			check(Bytes < PageBufferSize);	//Page size needs to be increased since we don't allow spillover
+			if (Bytes)
+			{
+				//Assumes 16 byte alignment
+				int32 BytesLeft = PageBufferSize - FreePage->SeekPosition;	//don't switch to uint because negative implies we're out of 16 byte aligned space
+				if (BytesLeft < Bytes)
+				{
+					//no space so allocate new page if needed
+					if (FreePage->NextPage)
+					{
+						FreePage = FreePage->NextPage;
+					}
+					else
+					{
+						FPageStruct* NewPage = AllocPage();
+						NewPage->PrevPage = FreePage;
+						FreePage->NextPage = NewPage;
+						FreePage = NewPage;
+					}
+				}
+
+				uint32 ReturnSlot = FreePage->SeekPosition;
+				FreePage->SeekPosition = (FreePage->SeekPosition + Bytes + 15)&(~15);
+
+				return &FreePage->Buffer[ReturnSlot];
+			}
+			else
+			{
+				return nullptr;
+			}
+		}
+
+		void Reset()
+		{
+			for (FPageStruct* Page = FirstPage; Page; Page = Page->NextPage)
+			{
+				Page->SeekPosition = 0;
+			}
+
+			FreePage = FirstPage;
+		}
+
+		void Empty()
+		{
+			FPageStruct* CurrentPage = FirstPage->NextPage;
+			while (CurrentPage)
+			{
+				FPageStruct* OldPage = CurrentPage;
+				CurrentPage = CurrentPage->NextPage;
+				ReleasePage(OldPage);
+			}
+
+			FirstPage->NextPage = nullptr;
+			FirstPage->SeekPosition = 0;
+			FreePage = FirstPage;
+			//Note we do not deallocate FirstPage until the allocator deallocates
+		}
+
+		~FLinearBlockAllocator()
+		{
+			Empty();
+			ReleasePage(FirstPage);
+		}
+
+	private:
+		//Don't copy these around
+		FLinearBlockAllocator(const FLinearBlockAllocator& Other);
+		const FLinearBlockAllocator& operator=(const FLinearBlockAllocator& Other);
+	};
+
+}
+
+#endif // WITH_PHYSX
\ No newline at end of file
diff --git a/Engine/Source/Runtime/Engine/Public/Physics/ImmediatePhysics/ImmediatePhysicsPhysX/ImmediatePhysicsMaterialHandle_PhysX.h b/Engine/Source/Runtime/Engine/Public/Physics/ImmediatePhysics/ImmediatePhysicsPhysX/ImmediatePhysicsMaterialHandle_PhysX.h
new file mode 100644
--- /dev/null
+++ b/Engine/Source/Runtime/Engine/Public/Physics/ImmediatePhysics/ImmediatePhysicsPhysX/ImmediatePhysicsMaterialHandle_PhysX.h
@@ -0,0 +1,47 @@
+// Copyright 1998-2019 Epic Games, Inc. All Rights Reserved.
+
+#pragma once
+
+#if WITH_PHYSX && PHYSICS_INTERFACE_PHYSX
+
+class UPhysicalMaterial;
+
+namespace ImmediatePhysics_PhysX
+{
+	struct ENGINE_API FMaterialHandle
+	{
+
+	public:
+
+		FMaterial& Get();
+		void SetStaticFriction(float InFriction);
+		void SetDynamicFriction(float InFriction);
+		void SetRestitution(float InRestitution);
+		void SetFrictionCombineMode(EFrictionCombineMode InCombineMode);
+		void SetRestitutionCombineMode(EFrictionCombineMode InCombineMode);
+		void SetFromPhysicalMaterial(UPhysicalMaterial* InMaterial);
+
+	private:
+
+		FSimulation & OwningSimulation;
+		int32 DataIndex;
+
+		friend struct FSimulation;
+
+		FMaterialHandle(FSimulation& InOwningSimulation, int32 InDataIndex)
+			: OwningSimulation(InOwningSimulation)
+			, DataIndex(InDataIndex)
+		{}
+
+		~FMaterialHandle() 
+		{}
+
+		// Removed copy/assign
+		FMaterialHandle(const FMaterialHandle& Other) = delete;
+		FMaterialHandle(FMaterialHandle&& Other) = delete;
+		FMaterialHandle& operator=(const FMaterialHandle& Other) = delete;
+		FMaterialHandle& operator=(FMaterialHandle&& Other) = delete;
+	};
+}
+
+#endif // WITH_PHYSX
\ No newline at end of file
diff --git a/Engine/Source/Runtime/Engine/Public/Physics/ImmediatePhysics/ImmediatePhysicsPhysX/ImmediatePhysicsMaterial_PhysX.h b/Engine/Source/Runtime/Engine/Public/Physics/ImmediatePhysics/ImmediatePhysicsPhysX/ImmediatePhysicsMaterial_PhysX.h
new file mode 100644
--- /dev/null
+++ b/Engine/Source/Runtime/Engine/Public/Physics/ImmediatePhysics/ImmediatePhysicsPhysX/ImmediatePhysicsMaterial_PhysX.h
@@ -0,0 +1,37 @@
+// Copyright 1998-2019 Epic Games, Inc. All Rights Reserved.
+
+#pragma once
+
+#if WITH_PHYSX && PHYSICS_INTERFACE_PHYSX
+#include "PhysicsEngine/PhysicsSettingsEnums.h"
+
+namespace physx
+{
+	class PxMaterial;
+}
+
+namespace ImmediatePhysics_PhysX
+{
+	struct FMaterial
+	{
+		FMaterial()
+			: StaticFriction(0.7f)
+			, DynamicFriction(0.7f)
+			, Restitution(0.3f)
+		{
+		}
+
+		FMaterial(physx::PxMaterial* InPxMaterial);
+
+		float StaticFriction;
+		float DynamicFriction;
+		float Restitution;
+
+		EFrictionCombineMode::Type FrictionCombineMode;
+		EFrictionCombineMode::Type RestitutionCombineMode;
+
+		static FMaterial Default;
+	};
+}
+
+#endif // WITH_PHYSX
\ No newline at end of file
diff --git a/Engine/Source/Runtime/Engine/Public/Physics/ImmediatePhysics/ImmediatePhysicsPhysX/ImmediatePhysicsPersistentContactPairData_PhysX.h b/Engine/Source/Runtime/Engine/Public/Physics/ImmediatePhysics/ImmediatePhysicsPhysX/ImmediatePhysicsPersistentContactPairData_PhysX.h
new file mode 100644
--- /dev/null
+++ b/Engine/Source/Runtime/Engine/Public/Physics/ImmediatePhysics/ImmediatePhysicsPhysX/ImmediatePhysicsPersistentContactPairData_PhysX.h
@@ -0,0 +1,32 @@
+// Copyright 1998-2019 Epic Games, Inc. All Rights Reserved.
+
+#pragma once
+
+#if WITH_PHYSX && PHYSICS_INTERFACE_PHYSX
+
+#ifndef PERSISTENT_CONTACT_PAIRS
+#define PERSISTENT_CONTACT_PAIRS 1
+#endif
+
+namespace ImmediatePhysics_PhysX
+{
+
+#if PERSISTENT_CONTACT_PAIRS
+	/** Persistent data like contact manifold, friction patches, etc... */
+	struct FPersistentContactPairData
+	{
+		PxCache Cache;
+		PxU8* Frictions;
+		PxU32 NumFrictions;
+		uint32 SimCount;
+
+		void Clear()
+		{
+			FPlatformMemory::Memzero(this, sizeof(FPersistentContactPairData));
+		}
+	};
+#endif
+
+}
+
+#endif // WITH_PHYSX
\ No newline at end of file
diff --git a/Engine/Source/Runtime/Engine/Public/Physics/ImmediatePhysics/ImmediatePhysicsPhysX/ImmediatePhysicsShape_PhysX.h b/Engine/Source/Runtime/Engine/Public/Physics/ImmediatePhysics/ImmediatePhysicsPhysX/ImmediatePhysicsShape_PhysX.h
new file mode 100644
--- /dev/null
+++ b/Engine/Source/Runtime/Engine/Public/Physics/ImmediatePhysics/ImmediatePhysicsPhysX/ImmediatePhysicsShape_PhysX.h
@@ -0,0 +1,60 @@
+// Copyright 1998-2019 Epic Games, Inc. All Rights Reserved.
+
+#pragma once
+
+#if WITH_PHYSX && PHYSICS_INTERFACE_PHYSX
+#include "PhysXPublic.h"
+#include "Physics/ImmediatePhysics/ImmediatePhysicsPhysX/ImmediatePhysicsMaterial_PhysX.h"
+
+namespace ImmediatePhysics_PhysX
+{
+	struct FMaterialHandle;
+}
+
+namespace ImmediatePhysics_PhysX
+{
+	/** Holds shape data*/
+	struct FShape
+	{
+		PxTransform LocalTM;
+		FMaterial* Material;
+		FMaterial InternalMaterial;
+		PxGeometry* Geometry;
+		PxVec3 BoundsOffset;
+		float BoundsMagnitude;
+		void* UserData;
+
+		FShape()
+			: LocalTM(PxTransform(PxIDENTITY::PxIdentity))
+			, Material(nullptr)
+			, Geometry(nullptr)
+			, BoundsOffset(PxVec3(PxIDENTITY::PxIdentity))
+			, BoundsMagnitude(0.0f)
+			, UserData(nullptr)
+		{}
+
+		FShape(const PxTransform& InLocalTM, const PxVec3& InBoundsOffset, const float InBoundsMagnitude, PxGeometry* InGeometry, FMaterial* InMaterial = nullptr)
+			: LocalTM(InLocalTM)
+			, Material(InMaterial)
+			, Geometry(InGeometry)
+			, BoundsOffset(InBoundsOffset)
+			, BoundsMagnitude(InBoundsMagnitude)
+			, UserData(nullptr)
+		{
+		}
+
+		FShape(const PxTransform& InLocalTM, const PxVec3& InBoundsOffset, const float InBoundsMagnitude, PxGeometry* InGeometry, FMaterial InMaterial)
+			: LocalTM(InLocalTM)
+			, InternalMaterial(InMaterial)
+			, Geometry(InGeometry)
+			, BoundsOffset(InBoundsOffset)
+			, BoundsMagnitude(InBoundsMagnitude)
+			, UserData(nullptr)
+		{
+			Material = &InternalMaterial;
+		}
+	};
+
+}
+
+#endif // WITH_PHYSX
\ No newline at end of file
diff --git a/Engine/Source/Runtime/Engine/Public/Physics/ImmediatePhysics/ImmediatePhysicsPhysX/ImmediatePhysicsSimulation_PhysX.h b/Engine/Source/Runtime/Engine/Public/Physics/ImmediatePhysics/ImmediatePhysicsPhysX/ImmediatePhysicsSimulation_PhysX.h
new file mode 100644
--- /dev/null
+++ b/Engine/Source/Runtime/Engine/Public/Physics/ImmediatePhysics/ImmediatePhysicsPhysX/ImmediatePhysicsSimulation_PhysX.h
@@ -0,0 +1,276 @@
+// Copyright 1998-2019 Epic Games, Inc. All Rights Reserved.
+
+#pragma once
+
+#if WITH_PHYSX && PHYSICS_INTERFACE_PHYSX
+#include "PhysXPublic.h"
+
+#include "Physics/PhysicsInterfaceTypes.h"
+
+#include "Engine/EngineTypes.h"
+#include "Physics/ImmediatePhysics/ImmediatePhysicsPhysX/ImmediatePhysicsCore_PhysX.h"
+#include "Physics/ImmediatePhysics/ImmediatePhysicsPhysX/ImmediatePhysicsActor_PhysX.h"
+#include "Physics/ImmediatePhysics/ImmediatePhysicsPhysX/ImmediatePhysicsJoint_PhysX.h"
+#include "Physics/ImmediatePhysics/ImmediatePhysicsPhysX/ImmediatePhysicsLinearBlockAllocator_PhysX.h"
+#include "Physics/ImmediatePhysics/ImmediatePhysicsPhysX/ImmediatePhysicsContactPair_PhysX.h"
+#include "Physics/ImmediatePhysics/ImmediatePhysicsPhysX/ImmediatePhysicsPersistentContactPairData_PhysX.h"
+#include "Physics/ImmediatePhysics/ImmediatePhysicsPhysX/ImmediatePhysicsCacheAllocator_PhysX.h"
+#include "Physics/ImmediatePhysics/ImmediatePhysicsPhysX/ImmediatePhysicsConstraintAllocator_PhysX.h"
+#include "Physics/ImmediatePhysics/ImmediatePhysicsPhysX/ImmediatePhysicsKinematicTarget_PhysX.h"
+#include "Physics/ImmediatePhysics/ImmediatePhysicsPhysX/ImmediatePhysicsD6JointData_PhysX.h"
+
+class UBodySetup;
+class UPhysicsConstraintTemplate;
+struct FBodyInstance;
+struct FConstraintInstance;
+
+namespace ImmediatePhysics_PhysX
+{
+	struct FActorHandle;
+	struct FJointHandle;
+    struct FContactPointRecorder;
+}
+
+namespace ImmediatePhysics_PhysX
+{
+	struct FActorData
+	{
+		immediate::PxRigidBodyData RigidBodyData;
+		FTransform InitialTransform;
+
+		bool bStatic;
+		bool bKinematic;
+	};
+
+	FActorData CreateActorData(const FActorCreationParams& InParams);
+
+	/** Owns all the data associated with the simulation. Can be considered a single scene or world*/
+	struct ENGINE_API FSimulation
+	{
+	public:
+
+		FActorHandle* InsertActorData(const FActor& InActor, const FActorData& InData);
+
+		void RemoveActor(FActorHandle* InHandle);
+
+		/** Create a kinematic body with the same setup as the body instance and add it to the simulation */
+		FActorHandle* CreateKinematicActor(FBodyInstance* BodyInstance, const FTransform& TM);
+
+		/** Create a dynamic body with the same setup as the body instance and add it to the simulation */
+		FActorHandle* CreateDynamicActor(FBodyInstance* BodyInstance, const FTransform& TM);
+
+		/** Create a static body with the same setup as the body instance and add it to the simulation */
+		FActorHandle* CreateStaticActor(FBodyInstance* BodyInstance);
+
+		/** Create a body of the specified type with the same setup as the body instance and add it to the simulation */
+		FActorHandle* CreateActor(EActorType ActorType, FBodyInstance* BodyInstance, const FTransform& Transform);
+
+		/** Create a physical joint with the same setup as the constraint instance and add it to the simulation */
+		FJointHandle* CreateJoint(FConstraintInstance* ConstraintInstance, FActorHandle* Body1, FActorHandle* Body2);
+
+		const immediate::PxRigidBodyData& GetLowLevelBody(int32 ActorDataIndex) const
+		{
+			return RigidBodiesData[ActorDataIndex];
+		}
+
+		immediate::PxRigidBodyData& GetLowLevelBody(int32 ActorDataIndex)
+		{
+			return RigidBodiesData[ActorDataIndex];
+		}
+
+		/** Sets the number of active bodies. This number is reset any time a new simulated body is created */
+		void SetNumActiveBodies(int32 NumActiveBodies);
+
+		/** An array of actors to ignore. */
+		struct FIgnorePair
+		{
+			FActorHandle* A;
+			FActorHandle* B;
+		};
+
+		/** Set pair of bodies to ignore collision for */
+		void SetIgnoreCollisionPairTable(const TArray<FIgnorePair>& InIgnoreCollisionPairTable);
+
+		/** Set bodies that require no collision */
+		void SetIgnoreCollisionActors(const TArray<FActorHandle*>& InIgnoreCollision);
+
+		/** Get/Set whether the body is kinematic or not, kinematics do not simulate and move where they're told. They also act as if they have infinite mass */
+		bool GetIsKinematic(int32 ActorDataIndex);
+		void SetIsKinematic(int32 ActorDataIndex, bool bKinematic);
+
+		/** Advance the simulation by DeltaTime */
+		void Simulate(float DeltaTime, const FVector& Gravity);
+		void Simulate_AssumesLocked(float DeltaTime, const FVector& Gravity);
+
+		/** Whether or not an entity is simulated */
+		bool IsSimulated(int32 ActorDataIndex) const
+		{
+			return ActorDataIndex < NumSimulatedBodies;
+		}
+
+		/** Add a radial impulse to the given actor */
+		void AddRadialForce(int32 ActorDataIndex, const FVector& Origin, float Strength, float Radius, ERadialImpulseFalloff Falloff, EForceType ForceType);
+
+		/** Add a force to the given actor */
+		void AddForce(int32 ActorDataIndex, const FVector& Force);
+
+		FSimulation();
+
+		~FSimulation();
+
+		int32 NumActors() const { return ActorHandles.Num(); }
+
+		FActorHandle* GetActorHandle(int ActorHandleIndex) { return ActorHandles[ActorHandleIndex]; }
+		const FActorHandle* GetActorHandle(int ActorHandleIndex) const { return ActorHandles[ActorHandleIndex]; }
+
+		/* Set solver iteration counts per step */
+		void SetPositionIterationCount(int32 InIterationCount);
+		void SetVelocityIterationCount(int32 InIterationCount);
+
+	private:
+		friend FActorHandle;
+
+		const FImmediateKinematicTarget& GetKinematicTarget(int32 ActorDataIndex) const
+		{
+			return KinematicTargets[ActorDataIndex];
+		}
+
+		FImmediateKinematicTarget& GetKinematicTarget(int32 ActorDataIndex)
+		{
+			return KinematicTargets[ActorDataIndex];
+		}
+
+		template <EActorType ActorType>
+		int32 CreateActor(PxRigidActor* RigidActor, const FTransform& TM);
+
+		/** Swap actor data - that is move all data associated with the two actors in the various arrays*/
+		void SwapActorData(int32 Entity1Idx, int32 Entity2Idx);
+
+		/** Swap joint data - that is move all data associated with the two joints in the various arrays*/
+		void SwapJointData(int32 Joint1Idx, int32 Joint2Idx);
+
+		/** Ensure arrays are valid */
+		void ValidateArrays() const;
+
+		/** Constructs solver bodies */
+		void ConstructSolverBodies(float DeltaTime, const FVector& Gravity);
+
+		/** Generate contacts*/
+		void GenerateContacts();
+
+		/** Batch constraints and re-order them for optimal processing */
+		void BatchConstraints();
+
+		/** Prepares the various constraints (contact,joints) for the solver */
+		void PrepareConstraints(float DeltaTime);
+
+		/** Solve constraints and integrate velocities */
+		void SolveAndIntegrate(float DeltaTime);
+
+		/** Prepares iteration cache for generating contacts */
+		void PrepareIterationCache();
+
+		//void EvictCache();
+
+	private:
+
+		/** Mapping from entity index to handle */
+		TArray<FActorHandle*>	ActorHandles;
+
+		/** Mapping from constraint index to handle */
+		TArray<FJointHandle*> JointHandles;
+
+		/** Entities holding loose data. NOTE: for performance reasons we don't automatically cleanup on destructor (needed for tarray swaps etc...) it's very important that Terminate is called */
+		TArray<FActor> Actors;
+		TArray<FJoint> Joints;
+
+		/** Workspace memory that we use for per frame allocations */
+		FLinearBlockAllocator Workspace;
+
+		/** Low level rigid body data */
+		TArray<immediate::PxRigidBodyData> RigidBodiesData;
+
+		/** Low level solver bodies data */
+		TArray<PxSolverBodyData> SolverBodiesData;
+
+		/** Kinematic targets used to implicitly compute the velocity of moving kinematic actors */
+		TArray<FImmediateKinematicTarget> KinematicTargets;
+
+		TArray<PxVec3> PendingAcceleration;
+
+		/** Low level contact points generated for this frame. Points are grouped together by pairs */
+		TArray<Gu::ContactPoint> ContactPoints;
+
+		/** Shapes used in the entire simulation. Shapes are sorted in the same order as actors. Note that an actor can have multiple shapes which will be adjacent*/
+		struct FShapeSOA
+		{
+			TArray<PxTransform> LocalTMs;
+			TArray<FMaterial> Materials;
+			TArray<const PxGeometry*> Geometries;
+			TArray<float> Bounds;
+			TArray<PxVec3> BoundsOffsets;
+			TArray<int32> OwningActors;
+#if PERSISTENT_CONTACT_PAIRS
+			TArray<FPersistentContactPairData> ContactPairData;
+#endif
+		} ShapeSOA;
+
+		/** Low level solver bodies */
+		PxSolverBody* SolverBodies;
+
+		/** Low level constraint descriptors.*/
+		TArray<PxSolverConstraintDesc> OrderedDescriptors;
+		TArray<PxConstraintBatchHeader> BatchHeaders;
+
+		/** JointData as passed in from physics constraint template */
+		TArray<D6JointData> JointData;
+
+		/** When new joints are created we have to update the processing order */
+		bool bDirtyJointData;
+
+		PxU32 NumContactHeaders;
+		PxU32 NumJointHeaders;
+		int32 NumActiveJoints;
+
+		/** Contact pairs generated for this frame */
+		TArray<FContactPair> ContactPairs;
+
+		/** Number of dynamic bodies associated with the simulation */
+		int32 NumSimulatedBodies;
+
+		/** Number of dynamic bodies that are actually active */
+		int32 NumActiveSimulatedBodies;
+
+		/** Number of kinematic bodies (dynamic but not simulated) associated with the simulation */
+		int32 NumKinematicBodies;
+
+		/** Total number of simulated shapes in the scene */
+		int32 NumSimulatedShapesWithCollision;
+
+		/** Number of position iterations used by solver */
+		int32 NumPositionIterations;
+
+		/** Number of velocity iterations used by solver */
+		int32 NumVelocityIterations;
+
+		/** Count of how many times we've ticked. Useful for cache invalidation */
+		int32 SimCount;
+
+		/** Both of these are slow to access. Make sure to use iteration cache when possible */
+		TMap<FActorHandle*, TSet<FActorHandle*>> IgnoreCollisionPairTable;
+		TSet<FActorHandle*> IgnoreCollisionActors;
+
+		/** This cache is used to record which generate contact iteration we can skip. This assumes the iteration order has not changed (add/remove/swap actors must invalidate this) */
+		bool bRecreateIterationCache;
+
+		TArray<int32> SkipCollisionCache;	//Holds the iteration count that we should skip due to ignore filtering
+
+		friend FContactPointRecorder;
+
+		FCacheAllocator CacheAllocator;
+		FConstraintAllocator ConstraintAllocator;
+	};
+
+}
+
+#endif // WITH_PHYSX
\ No newline at end of file
diff --git a/Engine/Source/Runtime/Engine/Public/Physics/ImmediatePhysics/ImmediatePhysicsPhysX/ResourceManager_PhysX.h b/Engine/Source/Runtime/Engine/Public/Physics/ImmediatePhysics/ImmediatePhysicsPhysX/ResourceManager_PhysX.h
new file mode 100644
--- /dev/null
+++ b/Engine/Source/Runtime/Engine/Public/Physics/ImmediatePhysics/ImmediatePhysicsPhysX/ResourceManager_PhysX.h
@@ -0,0 +1,192 @@
+// Copyright 1998-2019 Epic Games, Inc. All Rights Reserved.
+
+#pragma once
+
+#if WITH_PHYSX && PHYSICS_INTERFACE_PHYSX
+
+#include "CoreMinimal.h"
+#include "CoreTypes.h"
+#include "Containers/SparseArray.h"
+#include "Physics/ImmediatePhysics/ImmediatePhysicsPhysX/ImmediatePhysicsMaterial_PhysX.h"
+
+namespace ImmediatePhysics_PhysX
+{
+	struct FMaterial;
+}
+
+namespace ImmediatePhysics_PhysX
+{
+	/* List of available resource types managed here */
+	enum class EResourceType : uint8
+	{
+		Material = 0,
+		NumResources,
+		Invalid
+	};
+
+	/* Handle to a specific shared resource */
+	struct FResourceHandle
+	{
+		FResourceHandle()
+			: Type(EResourceType::Invalid)
+			, ResourceIndex(INDEX_NONE)
+			, ResourceId(INDEX_NONE)
+		{}
+
+		FResourceHandle(EResourceType InType, int32 InIndex, int32 InId)
+			: Type(InType)
+			, ResourceIndex(InIndex)
+			, ResourceId(InId)
+		{}
+
+		/** Get the type of the resource, see the list of supported types above */
+		EResourceType GetType() const { return Type; }
+
+		/** Get the storage index of this resource, note this is not enough to identify a resource, the comparison id is also required */
+		int32 GetIndex() const { return ResourceIndex; }
+
+		/** ID for this resource, used to tell whether a resource is what this handle describes, or just reuses memory */
+		int32 GetId() const { return ResourceId; }
+
+		/** Invalidate the handle */
+		void Invalidate()
+		{
+			Type = EResourceType::Invalid;
+			ResourceIndex = INDEX_NONE;
+			ResourceId = INDEX_NONE;
+		}
+
+	private:
+
+		EResourceType Type;
+		int32 ResourceIndex;
+		int32 ResourceId;
+	};
+
+
+	/** Simple pair for holding resources and IDs, with a per-type ID counter for comparisons */
+	template<typename ResourceType>
+	class TResourceWithId
+	{
+		static int32 ResourceIdCounter;
+
+	public:
+
+		TResourceWithId()
+		{
+			Id = ResourceIdCounter++;
+		}
+
+		ResourceType Resource;
+		int32 Id;
+	};
+
+	template<typename ResourceType>
+	int32 TResourceWithId<ResourceType>::ResourceIdCounter = 0;
+
+	/** Responsible for holding and providing shared resources for the physics system under immediate mode.
+	 *  Any resource that cannot be owned entirely by another physics object should be managed here by this
+	 *  shared resource manager.
+	 */
+	class FSharedResourceManager
+	{
+	public:
+
+		friend struct FScopedSharedResourceReadLock;
+
+		static FSharedResourceManager& Get()
+		{
+			return Instance;
+		}
+
+		/** Material methods */
+		FResourceHandle CreateMaterial();
+		void ReleaseMaterial(int32 InIndex);
+		int32 GetMaterialId(int32 InIndex);
+		FMaterial* GetMaterial(int32 InIndex);
+
+		/** Returns the locking primitive used by this manager, usually the scoped resource lock should suffice */
+		FRWLock& GetLockObject();
+
+	private:
+
+		FSharedResourceManager() {}
+		~FSharedResourceManager() {}
+		FSharedResourceManager(const FSharedResourceManager& Other) = delete;
+		FSharedResourceManager(FSharedResourceManager&& Other) = delete;
+		FSharedResourceManager& operator=(const FSharedResourceManager&& Other) = delete;
+		FSharedResourceManager& operator=(FSharedResourceManager&& Other) = delete;
+
+		/** Storage for managed resources */
+		TSparseArray<TResourceWithId<FMaterial>> Materials;
+
+		/** Locking primitive. Should be used in the desired mode whenever manipulated. Simulations will read lock this
+		 *  so user code cannot write to shared resources while a simulation is in flight
+		 */
+		FRWLock ResourceLock;
+
+		/** Singleton instance of the manager */
+		static FSharedResourceManager Instance;
+	};
+
+	/** Modes for scoped resource locking */
+	enum class LockMode : uint8
+	{
+		Read,
+		Write
+	};
+
+	/** Scoped locking object for physics resources, to be used whenever manipulating shared objects */
+	template<LockMode LockType>
+	struct FScopedSharedResourceLock
+	{
+		FScopedSharedResourceLock()
+		{
+			DoLock();
+		}
+
+		~FScopedSharedResourceLock()
+		{
+			DoUnlock();
+		}
+
+		void DoLock();
+		void DoUnlock();
+
+		FScopedSharedResourceLock(const FScopedSharedResourceLock& Other) = delete;
+		FScopedSharedResourceLock(FScopedSharedResourceLock&& Other) = delete;
+		FScopedSharedResourceLock& operator=(const FScopedSharedResourceLock& Other) = delete;
+		FScopedSharedResourceLock& operator=(FScopedSharedResourceLock&& Other) = delete;
+	};
+
+	template<>
+	inline void FScopedSharedResourceLock<LockMode::Read>::DoLock()
+	{
+		FSharedResourceManager& Manager = FSharedResourceManager::Get();
+		Manager.GetLockObject().ReadLock();
+	}
+
+	template<>
+	inline void FScopedSharedResourceLock<LockMode::Read>::DoUnlock()
+	{
+		FSharedResourceManager& Manager = FSharedResourceManager::Get();
+		Manager.GetLockObject().ReadUnlock();
+	}
+
+	template<>
+	inline void FScopedSharedResourceLock<LockMode::Write>::DoLock()
+	{
+		FSharedResourceManager& Manager = FSharedResourceManager::Get();
+		Manager.GetLockObject().WriteLock();
+	}
+
+	template<>
+	inline void FScopedSharedResourceLock<LockMode::Write>::DoUnlock()
+	{
+		FSharedResourceManager& Manager = FSharedResourceManager::Get();
+		Manager.GetLockObject().WriteUnlock();
+	}
+
+}
+
+#endif // WITH_PHYSX
diff --git a/Engine/Source/Runtime/Engine/Public/Physics/ImmediatePhysics/ImmediatePhysicsShape.h b/Engine/Source/Runtime/Engine/Public/Physics/ImmediatePhysics/ImmediatePhysicsShape.h
deleted file mode 100644
diff --git a/Engine/Source/Runtime/Engine/Public/Physics/ImmediatePhysics/ImmediatePhysicsShared/ImmediatePhysicsCore.h b/Engine/Source/Runtime/Engine/Public/Physics/ImmediatePhysics/ImmediatePhysicsShared/ImmediatePhysicsCore.h
new file mode 100644
--- /dev/null
+++ b/Engine/Source/Runtime/Engine/Public/Physics/ImmediatePhysics/ImmediatePhysicsShared/ImmediatePhysicsCore.h
@@ -0,0 +1,30 @@
+// Copyright 1998-2019 Epic Games, Inc. All Rights Reserved.
+
+#pragma once
+
+/**
+ * @file ImmediatePhysicsCore.h
+ *
+ * This is for declarations that are shared by both PhysX and Chaos implementations.
+ * This will be included in the engine-specific implementations and therefore should never pull in any 
+ * headers or declare any types that are in the engine-specific namespaces (which would results
+ * in some types from each engine being brought into the other).
+ */
+
+namespace ImmediatePhysics_Shared
+{
+	enum class EActorType
+	{
+		StaticActor,		//collision but no movement
+		KinematicActor,		//collision + movement but no dynamics (forces, mass, etc...)
+		DynamicActor		//collision + movement + dynamics
+	};
+
+	enum class EForceType
+	{
+		AddForce,			//use mass and delta time
+		AddAcceleration,	//use delta time, ignore mass
+		AddImpulse,			//use mass, ignore delta time
+		AddVelocity			//ignore mass, ignore delta time
+	};
+}
\ No newline at end of file
diff --git a/Engine/Source/Runtime/Engine/Public/Physics/ImmediatePhysics/ImmediatePhysicsSimulation.h b/Engine/Source/Runtime/Engine/Public/Physics/ImmediatePhysics/ImmediatePhysicsSimulation.h
--- a/Engine/Source/Runtime/Engine/Public/Physics/ImmediatePhysics/ImmediatePhysicsSimulation.h
+++ b/Engine/Source/Runtime/Engine/Public/Physics/ImmediatePhysics/ImmediatePhysicsSimulation.h
@@ -2,301 +2,7 @@
 
 #pragma once
 
-#if WITH_PHYSX
-#include "PhysXPublic.h"
-#include "ImmediatePhysicsD6JointData.h"
-#endif // WITH_PHYSX
+#include "Physics/ImmediatePhysics/ImmediatePhysicsPhysX/ImmediatePhysicsSimulation_PhysX.h"
+#include "Physics/ImmediatePhysics/ImmediatePhysicsChaos/ImmediatePhysicsSimulation_Chaos.h"
 
-#include "../PhysicsInterfaceTypes.h"
-
-#include "Engine/EngineTypes.h"
-#include "ImmediatePhysicsActor.h"
-#include "ImmediatePhysicsJoint.h"
-#include "ImmediatePhysicsLinearBlockAllocator.h"
-#include "ImmediatePhysicsContactPair.h"
-#include "ImmediatePhysicsPersistentContactPairData.h"
-#include "ImmediatePhysicsCacheAllocator.h"
-#include "ImmediatePhysicsConstraintAllocator.h"
-#include "ImmediatePhysicsKinematicTarget.h"
-
-DECLARE_STATS_GROUP(TEXT("Immediate Physics"), STATGROUP_ImmediatePhysics, STATCAT_Advanced);
-
-class UBodySetup;
-class UPhysicsConstraintTemplate;
-
-namespace ImmediatePhysics
-{
-	struct FActorHandle;
-	struct FJointHandle;
-    struct FContactPointRecorder;
-}
-
-namespace ImmediatePhysics
-{
-	struct FActorData
-	{
-#if WITH_PHYSX
-		immediate::PxRigidBodyData RigidBodyData;
-#endif
-		FTransform InitialTransform;
-
-		bool bStatic;
-		bool bKinematic;
-	};
-
-	FActorData CreateActorData(const FActorCreationParams& InParams);
-	
-/** Owns all the data associated with the simulation. Can be considered a single scene or world*/
-struct ENGINE_API FSimulation
-{
-public:
-
-	ImmediatePhysics::FActorHandle* InsertActorData(const FActor& InActor, const FActorData& InData);
-
-	void RemoveActor(ImmediatePhysics::FActorHandle* InHandle);
-
-#if WITH_PHYSX
-	/** Create a kinematic body and add it to the simulation */
-	FActorHandle* CreateKinematicActor(PxRigidBody* RigidBody, const FTransform& TM);
-
-	/** Create a dynamic body and add it to the simulation */
-	FActorHandle* CreateDynamicActor(PxRigidDynamic* RigidDynamic, const FTransform& TM);
-
-	/** Create a static body and add it to the simulation */
-	FActorHandle* CreateStaticActor(PxRigidActor* RigidActor, const FTransform& TM);
-
-	/** Create a physical joint and add it to the simulation */
-	FJointHandle* CreateJoint(PxD6Joint* Joint, FActorHandle* Body1, FActorHandle* Body2);
-
-	const immediate::PxRigidBodyData& GetLowLevelBody(int32 ActorDataIndex) const
-	{
-		return RigidBodiesData[ActorDataIndex];
-	}
-
-	immediate::PxRigidBodyData& GetLowLevelBody(int32 ActorDataIndex)
-	{
-		return RigidBodiesData[ActorDataIndex];
-	}
-#endif
-
-	/** Sets the number of active bodies. This number is reset any time a new simulated body is created */
-	void SetNumActiveBodies(uint32 NumActiveBodies);
-
-	/** An array of actors to ignore. */
-	struct FIgnorePair
-	{
-		FActorHandle* A;
-		FActorHandle* B;
-	};
-
-	/** Set pair of bodies to ignore collision for */
-	void SetIgnoreCollisionPairTable(const TArray<FIgnorePair>& InIgnoreCollisionPairTable);
-
-	/** Set bodies that require no collision */
-	void SetIgnoreCollisionActors(const TArray<FActorHandle*>& InIgnoreCollision);
-
-	/** Get/Set whether the body is kinematic or not, kinematics do not simulate and move where they're told. They also act as if they have infinite mass */
-	bool GetIsKinematic(uint32 ActorDataIndex);
-	void SetIsKinematic(uint32 ActorDataIndex, bool bKinematic);
-
-	/** Advance the simulation by DeltaTime */
-	void Simulate(float DeltaTime, const FVector& Gravity);
-	void Simulate_AssumesLocked(float DeltaTime, const FVector& Gravity);
-
-	/** Whether or not an entity is simulated */
-	bool IsSimulated(uint32 ActorDataIndex) const
-	{
-		return ActorDataIndex < NumSimulatedBodies;
-	}
-
-	enum class EForceType
-	{
-		AddForce,	//use mass and delta time
-		AddAcceleration, //use delta time, ignore mass
-		AddImpulse,	//use mass, ignore delta time
-		AddVelocity	//ignore mass, ignore delta time
-	};
-
-	/** Add a radial impulse to the given actor */
-	void AddRadialForce(int32 ActorDataIndex, const FVector& Origin, float Strength, float Radius, ERadialImpulseFalloff Falloff, EForceType ForceType);
-
-	/** Add a force to the given actor */
-	void AddForce(int32 ActorDataIndex, const FVector& Force);
-
-	FSimulation();
-
-	~FSimulation();
-
-	TArray<FActorHandle*>& GetActorHandles() { return ActorHandles; }
-	const TArray<FActorHandle*>& GetActorHandles() const { return ActorHandles; }
-#if WITH_PHYSX
-	TArray<immediate::PxRigidBodyData>& GetRigidBodyData() { return RigidBodiesData; }
-	const TArray<immediate::PxRigidBodyData>& GetRigidBodyData() const { return RigidBodiesData; }
-#endif
-
-	int32 NumActors() const { return ActorHandles.Num(); }
-
-	/* Set solver iteration counts per step */
-	void SetPositionIterationCount(int32 InIterationCount);
-	void SetVelocityIterationCount(int32 InIterationCount);
-
-private:
-	friend FActorHandle;
-
-#if WITH_PHYSX
-	const FImmediateKinematicTarget& GetKinematicTarget(int32 ActorDataIndex) const
-	{
-		return KinematicTargets[ActorDataIndex];
-	}
-
-	FImmediateKinematicTarget& GetKinematicTarget(int32 ActorDataIndex)
-	{
-		return KinematicTargets[ActorDataIndex];
-	}
-#endif
-
-	enum class ECreateActorType
-	{
-		StaticActor,	//collision but no movement
-		KinematicActor,	//collision + movement but no dynamics (forces, mass, etc...)
-		DynamicActor	//collision + movement + dynamics
-	};
-
-#if WITH_PHYSX
-	template <ECreateActorType ActorType>
-	uint32 CreateActor(PxRigidActor* RigidActor, const FTransform& TM);
-#endif
-
-	/** Swap actor data - that is move all data associated with the two actors in the various arrays*/
-	void SwapActorData(uint32 Entity1Idx, uint32 Entity2Idx);
-
-	/** Swap joint data - that is move all data associated with the two joints in the various arrays*/
-	void SwapJointData(uint32 Joint1Idx, uint32 Joint2Idx);
-
-	/** Ensure arrays are valid */
-	void ValidateArrays() const;
-
-	/** Constructs solver bodies */
-	void ConstructSolverBodies(float DeltaTime, const FVector& Gravity);
-
-	/** Generate contacts*/
-	void GenerateContacts();
-
-	/** Batch constraints and re-order them for optimal processing */
-	void BatchConstraints();
-
-	/** Prepares the various constraints (contact,joints) for the solver */
-	void PrepareConstraints(float DeltaTime);
-
-	/** Solve constraints and integrate velocities */
-	void SolveAndIntegrate(float DeltaTime);
-
-	/** Prepares iteration cache for generating contacts */
-	void PrepareIterationCache();
-
-	//void EvictCache();
-
-private:
-
-	/** Mapping from entity index to handle */
-	TArray<FActorHandle*>	ActorHandles;
-
-	/** Mapping from constraint index to handle */
-	TArray<FJointHandle*> JointHandles;
-
-	/** Entities holding loose data. NOTE: for performance reasons we don't automatically cleanup on destructor (needed for tarray swaps etc...) it's very important that Terminate is called */
-	TArray<FActor> Actors;
-	TArray<FJoint> Joints;
-
-	/** Workspace memory that we use for per frame allocations */
-	FLinearBlockAllocator Workspace;
-
-#if WITH_PHYSX
-	/** Low level rigid body data */
-	TArray<immediate::PxRigidBodyData> RigidBodiesData;
-	
-	/** Low level solver bodies data */
-	TArray<PxSolverBodyData> SolverBodiesData;
-	
-	/** Kinematic targets used to implicitly compute the velocity of moving kinematic actors */
-	TArray<FImmediateKinematicTarget> KinematicTargets;
-
-	TArray<PxVec3> PendingAcceleration;
-
-	/** Low level contact points generated for this frame. Points are grouped together by pairs */
-	TArray<Gu::ContactPoint> ContactPoints;
-
-	/** Shapes used in the entire simulation. Shapes are sorted in the same order as actors. Note that an actor can have multiple shapes which will be adjacent*/
-	struct FShapeSOA
-	{
-		TArray<PxTransform> LocalTMs;
-		TArray<FMaterial> Materials;
-		TArray<const PxGeometry*> Geometries;
-		TArray<float> Bounds;
-		TArray<PxVec3> BoundsOffsets;
-		TArray<int32> OwningActors;
-#if PERSISTENT_CONTACT_PAIRS
-		TArray<FPersistentContactPairData> ContactPairData;
-#endif
-	} ShapeSOA;
-	
-	/** Low level solver bodies */
-	PxSolverBody* SolverBodies;
-
-	/** Low level constraint descriptors.*/
-	TArray<PxSolverConstraintDesc> OrderedDescriptors;
-	TArray<PxConstraintBatchHeader> BatchHeaders;
-
-	/** JointData as passed in from physics constraint template */
-	TArray<D6JointData> JointData;
-
-	/** When new joints are created we have to update the processing order */
-	bool bDirtyJointData;
-
-	PxU32 NumContactHeaders;
-	PxU32 NumJointHeaders;
-	uint32 NumActiveJoints;
-#endif // WITH_PHYSX
-
-	/** Contact pairs generated for this frame */
-	TArray<FContactPair> ContactPairs;
-
-	/** Number of dynamic bodies associated with the simulation */
-	uint32 NumSimulatedBodies;
-
-	/** Number of dynamic bodies that are actually active */
-	uint32 NumActiveSimulatedBodies;
-
-	/** Number of kinematic bodies (dynamic but not simulated) associated with the simulation */
-	uint32 NumKinematicBodies;
-
-	/** Total number of simulated shapes in the scene */
-	uint32 NumSimulatedShapesWithCollision;
-
-	/** Number of position iterations used by solver */
-	uint32 NumPositionIterations;
-
-	/** Number of velocity iterations used by solver */
-	uint32 NumVelocityIterations;
-
-	/** Count of how many times we've ticked. Useful for cache invalidation */
-	uint32 SimCount;
-
-	/** Both of these are slow to access. Make sure to use iteration cache when possible */
-	TMap<FActorHandle*, TSet<FActorHandle*>> IgnoreCollisionPairTable;
-	TSet<FActorHandle*> IgnoreCollisionActors;
-
-	/** This cache is used to record which generate contact iteration we can skip. This assumes the iteration order has not changed (add/remove/swap actors must invalidate this) */
-	bool bRecreateIterationCache;
-
-	TArray<int32> SkipCollisionCache;	//Holds the iteration count that we should skip due to ignore filtering
-
-	friend FContactPointRecorder;
-
-#if WITH_PHYSX
-	FCacheAllocator CacheAllocator;
-	FConstraintAllocator ConstraintAllocator;
-#endif // WITH_PHYSX
-};
-
-}
\ No newline at end of file
+#include "Physics/ImmediatePhysics/ImmediatePhysicsDeclares.h"
diff --git a/Engine/Source/Runtime/Engine/Public/Physics/ImmediatePhysics/ImmediatePhysicsStats.h b/Engine/Source/Runtime/Engine/Public/Physics/ImmediatePhysics/ImmediatePhysicsStats.h
new file mode 100644
--- /dev/null
+++ b/Engine/Source/Runtime/Engine/Public/Physics/ImmediatePhysics/ImmediatePhysicsStats.h
@@ -0,0 +1,7 @@
+// Copyright 1998-2019 Epic Games, Inc. All Rights Reserved.
+
+#pragma once
+
+#include "Stats/Stats2.h"
+
+DECLARE_STATS_GROUP(TEXT("Immediate Physics"), STATGROUP_ImmediatePhysics, STATCAT_Advanced);
diff --git a/Engine/Source/Runtime/Engine/Public/Physics/ImmediatePhysics/ResourceManager.h b/Engine/Source/Runtime/Engine/Public/Physics/ImmediatePhysics/ResourceManager.h
deleted file mode 100644
