#ue4 - Added water immersion info to "p.visualizemovement"

diff --git a/Engine/Source/Runtime/Engine/Private/Components/CharacterMovementComponent.cpp b/Engine/Source/Runtime/Engine/Private/Components/CharacterMovementComponent.cpp
--- a/Engine/Source/Runtime/Engine/Private/Components/CharacterMovementComponent.cpp
+++ b/Engine/Source/Runtime/Engine/Private/Components/CharacterMovementComponent.cpp
@@ -6896,9 +6896,17 @@ void UCharacterMovementComponent::VisualizeMovement() const
 	{
 		const FColor DebugColor = FColor::Blue;
 		HeightOffset += 20.f;
-		const FVector DebugLocation = TopOfCapsule + FVector(0.f,0.f,HeightOffset);
+		FVector DebugLocation = TopOfCapsule + FVector(0.f,0.f,HeightOffset);
 		FString DebugText = FString::Printf(TEXT("MovementMode: %s"), *GetMovementName());
 		DrawDebugString(GetWorld(), DebugLocation, DebugText, nullptr, DebugColor, 0.f, true);
+
+		if (IsInWater())
+		{
+			HeightOffset += 15.f;
+			DebugLocation = TopOfCapsule + FVector(0.f, 0.f, HeightOffset);
+			DebugText = FString::Printf(TEXT("ImmersionDepth: %.2f"), ImmersionDepth());
+			DrawDebugString(GetWorld(), DebugLocation, DebugText, nullptr, DebugColor, 0.f, true);
+		}
 	}
 
 	// Root motion (additive)
