Validate LevelStreamingObject before using it in FStreamLevelAction::UpdateOperation().
Also, print some info to the log in this function so we have more data if it's still crashing after this change

diff --git a/Engine/Source/Runtime/Engine/Classes/Engine/LevelStreaming.h b/Engine/Source/Runtime/Engine/Classes/Engine/LevelStreaming.h
--- a/Engine/Source/Runtime/Engine/Classes/Engine/LevelStreaming.h
+++ b/Engine/Source/Runtime/Engine/Classes/Engine/LevelStreaming.h
@@ -21,7 +21,7 @@ public:
 	bool			bLoading;
 	bool			bMakeVisibleAfterLoad;
 	bool			bShouldBlock;
-	ULevelStreaming* Level;
+	TWeakObjectPtr<ULevelStreaming> Level;
 	FName			LevelName;
 
 	FLatentActionInfo LatentInfo;
diff --git a/Engine/Source/Runtime/Engine/Private/LevelStreaming.cpp b/Engine/Source/Runtime/Engine/Private/LevelStreaming.cpp
--- a/Engine/Source/Runtime/Engine/Private/LevelStreaming.cpp
+++ b/Engine/Source/Runtime/Engine/Private/LevelStreaming.cpp
@@ -83,15 +83,25 @@ FStreamLevelAction::FStreamLevelAction(bool bIsLoading, const FName& InLevelName
 	, LevelName(InLevelName)
 	, LatentInfo(InLatentInfo)
 {
-	Level = FindAndCacheLevelStreamingObject( LevelName, World );
-	ActivateLevel( Level );
+	ULevelStreaming* LocalLevel = FindAndCacheLevelStreamingObject( LevelName, World );
+	Level = LocalLevel;
+	ActivateLevel( LocalLevel );
 }
 
 void FStreamLevelAction::UpdateOperation(FLatentResponse& Response)
 {
-	ULevelStreaming* LevelStreamingObject = Level; // to avoid confusion.
-	bool bIsOperationFinished = UpdateLevel( LevelStreamingObject );
-	Response.FinishAndTriggerIf(bIsOperationFinished, LatentInfo.ExecutionFunction, LatentInfo.Linkage, LatentInfo.CallbackTarget);
+	ULevelStreaming* LevelStreamingObject = Level.Get(); // to avoid confusion.
+	const bool bIsLevelValid = LevelStreamingObject != nullptr;
+	UE_LOG(LogLevelStreaming, Display, TEXT("FStreamLevelAction::UpdateOperation() LevelName %s, bIsLevelValid %d"), *LevelName.ToString(), (int32)bIsLevelValid);
+	if (bIsLevelValid)
+	{
+		bool bIsOperationFinished = UpdateLevel(LevelStreamingObject);
+		Response.FinishAndTriggerIf(bIsOperationFinished, LatentInfo.ExecutionFunction, LatentInfo.Linkage, LatentInfo.CallbackTarget);
+	}
+	else
+	{
+		Response.DoneIf(true);
+	}
 }
 
 #if WITH_EDITOR
