Change ServerUpdateLevelVisibility to use a struct, and serialize a File Name alongside the level package name if they are different.

diff --git a/Engine/Source/Runtime/CoreUObject/Private/UObject/CoreNet.cpp b/Engine/Source/Runtime/CoreUObject/Private/UObject/CoreNet.cpp
--- a/Engine/Source/Runtime/CoreUObject/Private/UObject/CoreNet.cpp
+++ b/Engine/Source/Runtime/CoreUObject/Private/UObject/CoreNet.cpp
@@ -382,24 +382,25 @@ void FPropertyRetirement::CountBytes(FArchive& Ar) const
 //	FNetBitWriter
 // ----------------------------------------------------------------
 FNetBitWriter::FNetBitWriter()
-:	FBitWriter(0)
+	: FBitWriter(0)
+	, PackageMap(nullptr)
 {}
 
-FNetBitWriter::FNetBitWriter( int64 InMaxBits )
-:	FBitWriter (InMaxBits, true)
-,	PackageMap( NULL )
+FNetBitWriter::FNetBitWriter(int64 InMaxBits)
+	: FBitWriter(InMaxBits, true)
+	, PackageMap(nullptr)
 {
 
 }
 
-FNetBitWriter::FNetBitWriter( UPackageMap * InPackageMap, int64 InMaxBits )
-:	FBitWriter (InMaxBits, true)
-,	PackageMap( InPackageMap )
+FNetBitWriter::FNetBitWriter(UPackageMap* InPackageMap, int64 InMaxBits)
+	: FBitWriter(InMaxBits, true)
+	, PackageMap(InPackageMap)
 {
 
 }
 
-FArchive& FNetBitWriter::operator<<( class FName& N )
+FArchive& FNetBitWriter::operator<<(class FName& N)
 {
 	if (PackageMap)
 	{
@@ -413,9 +414,9 @@ FArchive& FNetBitWriter::operator<<( class FName& N )
 	return *this;
 }
 
-FArchive& FNetBitWriter::operator<<( UObject*& Object )
+FArchive& FNetBitWriter::operator<<(UObject*& Object)
 {
-	PackageMap->SerializeObject( *this, UObject::StaticClass(), Object );
+	PackageMap->SerializeObject(*this, UObject::StaticClass(), Object);
 	return *this;
 }
 
@@ -448,19 +449,19 @@ void FNetBitWriter::CountMemory(FArchive& Ar) const
 // ----------------------------------------------------------------
 //	FNetBitReader
 // ----------------------------------------------------------------
-FNetBitReader::FNetBitReader( UPackageMap *InPackageMap, uint8* Src, int64 CountBits )
-:	FBitReader	(Src, CountBits)
-,   PackageMap  ( InPackageMap )
+FNetBitReader::FNetBitReader(UPackageMap* InPackageMap, uint8* Src, int64 CountBits)
+	: FBitReader(Src, CountBits)
+	, PackageMap( InPackageMap )
 {
 }
 
-FArchive& FNetBitReader::operator<<( UObject*& Object )
+FArchive& FNetBitReader::operator<<(UObject*& Object)
 {
-	PackageMap->SerializeObject( *this, UObject::StaticClass(), Object );
+	PackageMap->SerializeObject(*this, UObject::StaticClass(), Object);
 	return *this;
 }
 
-FArchive& FNetBitReader::operator<<( class FName& N )
+FArchive& FNetBitReader::operator<<(class FName& N)
 {
 	if (PackageMap)
 	{
diff --git a/Engine/Source/Runtime/Engine/Classes/Engine/NetConnection.h b/Engine/Source/Runtime/Engine/Classes/Engine/NetConnection.h
--- a/Engine/Source/Runtime/Engine/Classes/Engine/NetConnection.h
+++ b/Engine/Source/Runtime/Engine/Classes/Engine/NetConnection.h
@@ -616,6 +616,9 @@ public:
 	TSet<FName> ClientVisibleLevelNames;
 
 	/** Called by PlayerController to tell connection about client level visiblity change */
+	ENGINE_API void UpdateLevelVisibility(const struct FUpdateLevelVisibilityLevelInfo& LevelVisibility);
+	
+	UE_DEPRECATED(4.24, "This method will be removed. Use UpdateLevelVisibility that takes an FUpdateLevelVisibilityLevelInfo")
 	ENGINE_API void UpdateLevelVisibility(const FName& PackageName, bool bIsVisible);
 
 #if DO_ENABLE_NET_TEST
diff --git a/Engine/Source/Runtime/Engine/Classes/GameFramework/PlayerController.h b/Engine/Source/Runtime/Engine/Classes/GameFramework/PlayerController.h
--- a/Engine/Source/Runtime/Engine/Classes/GameFramework/PlayerController.h
+++ b/Engine/Source/Runtime/Engine/Classes/GameFramework/PlayerController.h
@@ -162,19 +162,47 @@ struct ENGINE_API FUpdateLevelStreamingLevelStatus
 	uint32 bNewShouldBlockOnLoad : 1;
 };
 
-/** This structure is used to pass arguments to ServerUpdateMultipleLevelsVisibility() server RPC function */
+/** This structure is used to pass arguments to ServerUpdateLevelVisibilty() and ServerUpdateMultipleLevelsVisibility() server RPC functions */
 USTRUCT()
 struct ENGINE_API FUpdateLevelVisibilityLevelInfo
 {
 	GENERATED_BODY();
 
-	/** The name of the package for the level whose status changed */
+	FUpdateLevelVisibilityLevelInfo()
+		: PackageName(NAME_None)
+		, FileName(NAME_None)
+		, bIsVisible(false)
+	{
+	}
+
+	/**
+	 * @param Level			Level to pull PackageName and FileName from.
+	 * @param bInIsVisible	Default value for bIsVisible.
+	 */
+	FUpdateLevelVisibilityLevelInfo(const class ULevel* const Level, const bool bInIsVisible);
+
+	/** The name of the package for the level whose status changed. */
 	UPROPERTY()
 	FName PackageName;
 
-	/** The new visibility state for this level */
+	/** The name / path of the asset file for the level whose status changed. */
+	UPROPERTY()
+	FName FileName;
+
+	/** The new visibility state for this level. */
 	UPROPERTY()
 	uint32 bIsVisible : 1;
+
+	bool NetSerialize(FArchive& Ar, UPackageMap* PackageMap, bool& bOutSuccess);
+};
+
+template<>
+struct TStructOpsTypeTraits<FUpdateLevelVisibilityLevelInfo> : public TStructOpsTypeTraitsBase2<FUpdateLevelVisibilityLevelInfo>
+{
+	enum
+	{
+		WithNetSerializer = true
+	};
 };
 
 /** Data structure used to setup an input mode that allows the UI to respond to user input, and if the UI doesn't handle it player input / player controller gets a chance. */
@@ -1087,6 +1115,11 @@ private:
 	UFUNCTION(BlueprintCallable, meta=(Latent, LatentInfo="LatentInfo", ExpandEnumAsExecs="Action", Duration="-1", bAffectsLeftLarge="true", bAffectsLeftSmall="true", bAffectsRightLarge="true", bAffectsRightSmall="true", AdvancedDisplay="bAffectsLeftLarge,bAffectsLeftSmall,bAffectsRightLarge,bAffectsRightSmall"), Category="Game|Feedback")
 	void PlayDynamicForceFeedback(float Intensity, float Duration, bool bAffectsLeftLarge, bool bAffectsLeftSmall, bool bAffectsRightLarge, bool bAffectsRightSmall, TEnumAsByte<EDynamicForceFeedbackAction::Type> Action, FLatentActionInfo LatentInfo);
 
+	//~ This method is purely for debugging purposes.
+	//~ It will trigger a ServerUpdateLevelVisibilityCall with the provided package name.
+	UFUNCTION(Exec)
+	void TestServerLevelVisibilityChange(const FName PackageName, const FName FileName);
+
 public:
 	/** 
 	 * Allows playing of a dynamic force feedback event from native code
@@ -1277,13 +1310,25 @@ public:
 	void ServerUpdateCamera(FVector_NetQuantize CamLoc, int32 CamPitchAndYaw);
 
 	/** 
-	 * Called when the client adds/removes a streamed level
-	 * the server will only replicate references to Actors in visible levels so that it's impossible to send references to
-	 * Actors the client has not initialized
-	 * @param PackageName the name of the package for the level whose status changed
+	 * Called when the client adds/removes a streamed level.
+	 * The server will only replicate references to Actors in visible levels so that it's impossible to send references to
+	 * Actors the client has not initialized.
+	 *
+	 * @param LevelVisibility	Visibility state for the level whose state changed.
 	 */
 	UFUNCTION(reliable, server, WithValidation, SealedEvent)
-	void ServerUpdateLevelVisibility(FName PackageName, bool bIsVisible);
+	void ServerUpdateLevelVisibility(const FUpdateLevelVisibilityLevelInfo& LevelVisibility);
+
+	UE_DEPRECATED(4.24, "Use ServerUpdateLevelVisibility that accepts a LevelVisibility struct.")
+	void ServerUpdateLevelVisibility(FName PackageName, bool bIsVisible)
+	{
+		FUpdateLevelVisibilityLevelInfo LevelVisibility;
+		LevelVisibility.PackageName = PackageName;
+		LevelVisibility.FileName = PackageName;
+		LevelVisibility.bIsVisible = bIsVisible;
+
+		ServerUpdateLevelVisibility(LevelVisibility);
+	}
 
 	/** 
 	 * Called when the client adds/removes a streamed level.  This version of the function allows you to pass the state of 
diff --git a/Engine/Source/Runtime/Engine/Private/NetConnection.cpp b/Engine/Source/Runtime/Engine/Private/NetConnection.cpp
--- a/Engine/Source/Runtime/Engine/Private/NetConnection.cpp
+++ b/Engine/Source/Runtime/Engine/Private/NetConnection.cpp
@@ -69,6 +69,45 @@ DECLARE_CYCLE_STAT(TEXT("NetConnection Tick"), Stat_NetConnectionTick, STATGROUP
 DECLARE_CYCLE_STAT(TEXT("NetConnection ReceivedNak"), Stat_NetConnectionReceivedNak, STATGROUP_Net);
 DECLARE_CYCLE_STAT(TEXT("NetConnection NetConnectionReceivedAcks"), Stat_NetConnectionReceivedAcks, STATGROUP_Net);
 
+namespace UE4_NetConnectionPrivate
+{
+	struct FValidateLevelVisibilityResult
+	{
+		bool bLevelExists = false;
+		UPackage* Package = nullptr;
+		FLinkerLoad* Linker = nullptr;
+	};
+
+	const FValidateLevelVisibilityResult ValidateLevelVisibility(UWorld* World, const FUpdateLevelVisibilityLevelInfo& LevelVisibility)
+	{
+		// Verify that we were passed a valid level name
+		// If we have a linker we know it has been loaded off disk successfully
+		// If we have a file it is fine too
+		// If its in our own streaming level list, its good
+
+		auto IsInLevelList = [&World](FName InPackageName)
+		{
+			for (ULevelStreaming* StreamingLevel : World->GetStreamingLevels())
+			{
+				if (StreamingLevel && (StreamingLevel->GetWorldAssetPackageFName() == InPackageName))
+				{
+					return true;
+				}
+			}
+			return false;
+		};
+
+		FValidateLevelVisibilityResult Result;
+		Result.Package = FindPackage(nullptr, *LevelVisibility.PackageName.ToString());
+		Result.Linker = FLinkerLoad::FindExistingLinkerForPackage(Result.Package);
+		Result.bLevelExists = Result.Linker ||
+			FPackageName::DoesPackageExist(LevelVisibility.FileName.ToString()) ||
+			IsInLevelList(LevelVisibility.PackageName);
+
+		return Result;
+	}
+}
+
 // ChannelRecord Implementation
 namespace FChannelRecordImpl
 {
@@ -902,39 +941,30 @@ void UNetConnection::UpdateAllCachedLevelVisibility() const
 
 void UNetConnection::UpdateLevelVisibility(const FName& PackageName, bool bIsVisible)
 {
+	FUpdateLevelVisibilityLevelInfo LevelVisibility;
+	LevelVisibility.PackageName = PackageName;
+	LevelVisibility.FileName = PackageName;
+	LevelVisibility.bIsVisible = bIsVisible;
+
+	UpdateLevelVisibility(LevelVisibility);
+}
+
+void UNetConnection::UpdateLevelVisibility(const FUpdateLevelVisibilityLevelInfo& LevelVisibility)
+{
+	using namespace UE4_NetConnectionPrivate;
+
 	GNumClientUpdateLevelVisibility++;
 
 	// add or remove the level package name from the list, as requested
-	if (bIsVisible)
+	if (LevelVisibility.bIsVisible)
 	{
 		// verify that we were passed a valid level name
-		FString Filename;
-		UPackage* TempPkg = FindPackage(nullptr, *PackageName.ToString());
-		FLinkerLoad* Linker = FLinkerLoad::FindExistingLinkerForPackage(TempPkg);
-
-		// If we have a linker we know it has been loaded off disk successfully
-		// If we have a file it is fine too
-		// If its in our own streaming level list, its good
-
-		struct Local
-		{
-			static bool IsInLevelList(UWorld* World, FName InPackageName)
-			{
-				for (ULevelStreaming* StreamingLevel : World->GetStreamingLevels())
-				{
-					if (StreamingLevel && (StreamingLevel->GetWorldAssetPackageFName() == InPackageName ))
-					{
-						return true;
-					}
-				}
-				return false;
-			}
-		};
+		FValidateLevelVisibilityResult VisibilityResult = ValidateLevelVisibility(GetWorld(), LevelVisibility);
 
-		if ( Linker || FPackageName::DoesPackageExist(PackageName.ToString(), nullptr, &Filename ) || Local::IsInLevelList(GetWorld(), PackageName ) )
+		if (VisibilityResult.bLevelExists)
 		{
-			ClientVisibleLevelNames.Add(PackageName);
-			UE_LOG( LogPlayerController, Verbose, TEXT("ServerUpdateLevelVisibility() Added '%s'"), *PackageName.ToString() );
+			ClientVisibleLevelNames.Add(LevelVisibility.PackageName);
+			UE_LOG(LogPlayerController, Verbose, TEXT("ServerUpdateLevelVisibility() Added '%s'"), *LevelVisibility.PackageName.ToString());
 
 			QUICK_USE_CYCLE_STAT(NetUpdateLevelVisibility_UpdateDormantActors, STATGROUP_Net);
 
@@ -942,7 +972,7 @@ void UNetConnection::UpdateLevelVisibility(const FName& PackageName, bool bIsVis
 			// destroyed actors list when the level is reloaded, so seek them out and add in
 			for (const auto& DestroyedPair : Driver->DestroyedStartupOrDormantActors)
 			{
-				if (DestroyedPair.Value->StreamingLevelName == PackageName)
+				if (DestroyedPair.Value->StreamingLevelName == LevelVisibility.PackageName)
 				{
 					AddDestructionInfo(DestroyedPair.Value.Get());
 				}
@@ -951,9 +981,9 @@ void UNetConnection::UpdateLevelVisibility(const FName& PackageName, bool bIsVis
 			// Any dormant actor that has changes flushed or made before going dormant needs to be updated on the client 
 			// when the streaming level is loaded, so mark them active for this connection
 			UWorld* LevelWorld = nullptr;
-			if (TempPkg)
+			if (VisibilityResult.Package)
 			{
-				LevelWorld = (UWorld*)FindObjectWithOuter(TempPkg, UWorld::StaticClass());
+				LevelWorld = (UWorld*)FindObjectWithOuter(VisibilityResult.Package, UWorld::StaticClass());
 				if (LevelWorld)
 				{
 					if (LevelWorld->PersistentLevel)
@@ -965,7 +995,7 @@ void UNetConnection::UpdateLevelVisibility(const FName& PackageName, bool bIsVis
 							// to mark Dormant All Actors as (temporarily) active to get the update sent over
 							if (Actor && Actor->GetIsReplicated() && (Actor->NetDormancy == DORM_DormantAll))
 							{
-								NetworkObjectList.MarkActive( Actor, this, Driver );
+								NetworkObjectList.MarkActive(Actor, this, Driver);
 							}
 						}
 					}
@@ -974,33 +1004,33 @@ void UNetConnection::UpdateLevelVisibility(const FName& PackageName, bool bIsVis
 
 			if (ReplicationConnectionDriver)
 			{
-				ReplicationConnectionDriver->NotifyClientVisibleLevelNamesAdd(PackageName, LevelWorld);
+				ReplicationConnectionDriver->NotifyClientVisibleLevelNamesAdd(LevelVisibility.PackageName, LevelWorld);
 			}
 
 		}
 		else
 		{
-			UE_LOG( LogPlayerController, Warning, TEXT("ServerUpdateLevelVisibility() ignored non-existant package '%s'"), *PackageName.ToString() );
+			UE_LOG(LogPlayerController, Warning, TEXT("ServerUpdateLevelVisibility() ignored non-existant package. PackageName='%s', FileName='%s'"), *LevelVisibility.PackageName.ToString(), *LevelVisibility.FileName.ToString());
 			Close();
 		}
 	}
 	else
 	{
-		ClientVisibleLevelNames.Remove(PackageName);
-		UE_LOG( LogPlayerController, Verbose, TEXT("ServerUpdateLevelVisibility() Removed '%s'"), *PackageName.ToString() );
+		ClientVisibleLevelNames.Remove(LevelVisibility.PackageName);
+		UE_LOG(LogPlayerController, Verbose, TEXT("ServerUpdateLevelVisibility() Removed '%s'"), *LevelVisibility.PackageName.ToString());
 		if (ReplicationConnectionDriver)
 		{
-			ReplicationConnectionDriver->NotifyClientVisibleLevelNamesRemove(PackageName);
+			ReplicationConnectionDriver->NotifyClientVisibleLevelNamesRemove(LevelVisibility.PackageName);
 		}
 			
 		// Close any channels now that have actors that were apart of the level the client just unloaded
-		for ( auto It = ActorChannels.CreateIterator(); It; ++It )
+		for (auto It = ActorChannels.CreateIterator(); It; ++It)
 		{
-			UActorChannel* Channel = It.Value();					
+			UActorChannel* Channel = It.Value();
 
-			check( Channel->OpenedLocally );
+			check(Channel->OpenedLocally);
 
-			if ( Channel->Actor && Channel->Actor->GetLevel()->GetOutermost()->GetFName() == PackageName )
+			if (Channel->Actor && Channel->Actor->GetLevel()->GetOutermost()->GetFName() == LevelVisibility.PackageName)
 			{
 				Channel->Close(EChannelCloseReason::LevelUnloaded);
 			}
@@ -3156,9 +3186,8 @@ void UNetConnection::HandleClientPlayer( APlayerController *PC, UNetConnection*
 				const ULevel* Level = LevelStreaming->GetLoadedLevel();
 				if ( Level && Level->bIsVisible && !Level->bClientOnlyVisible )
 				{
-					FUpdateLevelVisibilityLevelInfo& LevelVisibility = *new( LevelVisibilities ) FUpdateLevelVisibilityLevelInfo();
-					LevelVisibility.PackageName = PC->NetworkRemapPath(Level->GetOutermost()->GetFName(), false);
-					LevelVisibility.bIsVisible = true;
+					FUpdateLevelVisibilityLevelInfo& LevelVisibility = *new( LevelVisibilities ) FUpdateLevelVisibilityLevelInfo(Level, true);
+					LevelVisibility.PackageName = PC->NetworkRemapPath(LevelVisibility.PackageName, false);
 				}
 			}
 		}
diff --git a/Engine/Source/Runtime/Engine/Private/PlayerController.cpp b/Engine/Source/Runtime/Engine/Private/PlayerController.cpp
--- a/Engine/Source/Runtime/Engine/Private/PlayerController.cpp
+++ b/Engine/Source/Runtime/Engine/Private/PlayerController.cpp
@@ -86,6 +86,13 @@ namespace PlayerControllerCVars
 		TEXT("Whether to reset server prediction data for the possessed Pawn when the pawn ack handshake completes.\n")
 		TEXT("0: Disable, 1: Enable"),
 		ECVF_Default);
+
+	static bool LevelVisibilityDontSerializeFileName = false;
+	FAutoConsoleVariableRef CVarLevelVisibilityDontSerializeFileName(
+		TEXT("PlayerController.LevelVisibilityDontSerializeFileName"),
+		LevelVisibilityDontSerializeFileName,
+		TEXT("When true, we'll always skip serializing FileName with FUpdateLevelVisibilityLevelInfo's. This will save bandwidth when games don't need both.")
+	);
 }
 
 const float RetryClientRestartThrottleTime = 0.5f;
@@ -95,6 +102,38 @@ const float RetryServerCheckSpectatorThrottleTime = 0.25f;
 // Note: This value should be sufficiently small such that it is considered to be in the past before RetryClientRestartThrottleTime and RetryServerAcknowledgeThrottleTime.
 const float ForceRetryClientRestartTime = -100.0f;
 
+FUpdateLevelVisibilityLevelInfo::FUpdateLevelVisibilityLevelInfo(const ULevel* const Level, const bool bInIsVisible)
+	: bIsVisible(bInIsVisible)
+{
+	const UPackage* const LevelPackage = Level->GetOutermost();
+	PackageName = LevelPackage->GetFName();
+	FileName = LevelPackage->FileName;
+}
+
+bool FUpdateLevelVisibilityLevelInfo::NetSerialize(FArchive& Ar, UPackageMap* PackageMap, bool& bOutSuccess)
+{
+	bool bArePackageAndFileTheSame = !!((PlayerControllerCVars::LevelVisibilityDontSerializeFileName) || (FileName == PackageName));
+	bool bLocalIsVisible = !!bIsVisible;
+
+	Ar.SerializeBits(&bArePackageAndFileTheSame, 1);
+	Ar.SerializeBits(&bLocalIsVisible, 1);
+	Ar << PackageName;
+
+	if (!bArePackageAndFileTheSame)
+	{
+		Ar << FileName;
+	}
+	else if (Ar.IsLoading())
+	{
+		FileName = PackageName;
+	}
+
+	bIsVisible = bLocalIsVisible;
+
+	bOutSuccess = !Ar.IsError();
+	return true;
+}
+
 //////////////////////////////////////////////////////////////////////////
 // APlayerController
 
@@ -327,25 +366,27 @@ void APlayerController::ClientFlushLevelStreaming_Implementation()
 }
 
 
-void APlayerController::ServerUpdateLevelVisibility_Implementation(FName PackageName, bool bIsVisible)
+void APlayerController::ServerUpdateLevelVisibility_Implementation(const FUpdateLevelVisibilityLevelInfo& LevelVisibility)
 {
 	UNetConnection* Connection = Cast<UNetConnection>(Player);
 	if (Connection != NULL)
 	{
-		PackageName = NetworkRemapPath(PackageName, true);
-		Connection->UpdateLevelVisibility(PackageName, bIsVisible);
+		FUpdateLevelVisibilityLevelInfo LevelVisibilityCopy = LevelVisibility;
+		LevelVisibilityCopy.PackageName = NetworkRemapPath(LevelVisibilityCopy.PackageName, true);
+
+		Connection->UpdateLevelVisibility(LevelVisibilityCopy);
 	}
 }
 
-bool APlayerController::ServerUpdateLevelVisibility_Validate(FName PackageName, bool bIsVisible)
+bool APlayerController::ServerUpdateLevelVisibility_Validate(const FUpdateLevelVisibilityLevelInfo& LevelVisibility)
 {
-	RPC_VALIDATE( PackageName.IsValid() );
+	RPC_VALIDATE(LevelVisibility.PackageName.IsValid());
 
 	FText Reason;
 
-	if ( !FPackageName::IsValidLongPackageName( PackageName.ToString(), true, &Reason ) )
+	if (!FPackageName::IsValidLongPackageName(LevelVisibility.PackageName.ToString(), true, &Reason))
 	{
-		UE_LOG( LogPlayerController, Warning, TEXT( "ServerUpdateLevelVisibility() Invalid package name: %s (%s)" ), *PackageName.ToString(), *Reason.ToString() );
+		UE_LOG(LogPlayerController, Warning, TEXT( "ServerUpdateLevelVisibility() Invalid package name: %s (%s)" ), *LevelVisibility.PackageName.ToString(), *Reason.ToString());
 		return false;
 	}
 
@@ -354,17 +395,17 @@ bool APlayerController::ServerUpdateLevelVisibility_Validate(FName PackageName,
 
 void APlayerController::ServerUpdateMultipleLevelsVisibility_Implementation( const TArray<FUpdateLevelVisibilityLevelInfo>& LevelVisibilities )
 {
-	for( const FUpdateLevelVisibilityLevelInfo& LevelVisibility : LevelVisibilities )
+	for(const FUpdateLevelVisibilityLevelInfo& LevelVisibility : LevelVisibilities)
 	{
-		ServerUpdateLevelVisibility_Implementation( LevelVisibility.PackageName, LevelVisibility.bIsVisible );
+		ServerUpdateLevelVisibility_Implementation(LevelVisibility);
 	}
 }
 
 bool APlayerController::ServerUpdateMultipleLevelsVisibility_Validate( const TArray<FUpdateLevelVisibilityLevelInfo>& LevelVisibilities )
 {
-	for( const FUpdateLevelVisibilityLevelInfo& LevelVisibility : LevelVisibilities )
+	for(const FUpdateLevelVisibilityLevelInfo& LevelVisibility : LevelVisibilities)
 	{
-		if( !ServerUpdateLevelVisibility_Validate( LevelVisibility.PackageName, LevelVisibility.bIsVisible ) )
+		if(!ServerUpdateLevelVisibility_Validate(LevelVisibility))
 		{
 			return false;
 		}
@@ -3917,6 +3958,19 @@ void APlayerController::PlayDynamicForceFeedback(float Intensity, float Duration
 	}
 }
 
+
+void APlayerController::TestServerLevelVisibilityChange(const FName PackageName, const FName FileName)
+{
+#if !(UE_BUILD_TEST||UE_BUILD_SHIPPING)
+	FUpdateLevelVisibilityLevelInfo LevelInfo;
+	LevelInfo.bIsVisible = true;
+	LevelInfo.PackageName = PackageName;
+	LevelInfo.FileName = FileName;
+	ServerUpdateLevelVisibility(LevelInfo);
+#endif
+}
+
+
 FDynamicForceFeedbackHandle APlayerController::PlayDynamicForceFeedback(float Intensity, float Duration, bool bAffectsLeftLarge, bool bAffectsLeftSmall, bool bAffectsRightLarge, bool bAffectsRightSmall, EDynamicForceFeedbackAction::Type Action, FDynamicForceFeedbackHandle ActionHandle)
 {
 	FDynamicForceFeedbackHandle FeedbackHandle = 0;
diff --git a/Engine/Source/Runtime/Engine/Private/World.cpp b/Engine/Source/Runtime/Engine/Private/World.cpp
--- a/Engine/Source/Runtime/Engine/Private/World.cpp
+++ b/Engine/Source/Runtime/Engine/Private/World.cpp
@@ -2446,12 +2446,15 @@ void UWorld::AddToWorld( ULevel* Level, const FTransform& LevelTransform, bool b
 		// notify server that the client has finished making this level visible
 		if (!Level->bClientOnlyVisible)
 		{
+			FUpdateLevelVisibilityLevelInfo LevelVisibility(Level, true);
+			const FName UnmappedPackageName = LevelVisibility.PackageName;
+			
 			for (FLocalPlayerIterator It(GEngine, this); It; ++It)
 			{
-				APlayerController* LocalPlayerController = It->GetPlayerController(this);
-				if (LocalPlayerController != NULL)
+				if (APlayerController* LocalPlayerController = It->GetPlayerController(this))
 				{
-					LocalPlayerController->ServerUpdateLevelVisibility(LocalPlayerController->NetworkRemapPath(Level->GetOutermost()->GetFName(), false), true);
+					LevelVisibility.PackageName = LocalPlayerController->NetworkRemapPath(UnmappedPackageName, false);
+					LocalPlayerController->ServerUpdateLevelVisibility(LevelVisibility);
 				}
 			}
 		}
@@ -2613,12 +2616,15 @@ void UWorld::RemoveFromWorld( ULevel* Level, bool bAllowIncrementalRemoval )
 			// notify server that the client has removed this level
 			if (!Level->bClientOnlyVisible)
 			{
+				FUpdateLevelVisibilityLevelInfo LevelVisibility(Level, false);
+				const FName UnmappedPackageName = LevelVisibility.PackageName;
+
 				for (FLocalPlayerIterator It(GEngine, this); It; ++It)
 				{
-					APlayerController* LocalPlayerController = It->GetPlayerController(this);
-					if (LocalPlayerController != NULL)
+					if (APlayerController* LocalPlayerController = It->GetPlayerController(this))
 					{
-						LocalPlayerController->ServerUpdateLevelVisibility(LocalPlayerController->NetworkRemapPath(Level->GetOutermost()->GetFName(), false), false);
+						LevelVisibility.PackageName = LocalPlayerController->NetworkRemapPath(UnmappedPackageName, false);
+						LocalPlayerController->ServerUpdateLevelVisibility(LevelVisibility);
 					}
 				}
 			}
@@ -3982,22 +3988,20 @@ void UWorld::InitializeActorsForPlay(const FURL& InURL, bool bResetTime)
 		{
 			for (FLocalPlayerIterator It(GEngine, this); It; ++It)
 			{
-				APlayerController* LocalPlayerController = It->GetPlayerController(this);
-				if (LocalPlayerController != NULL)
+				if (APlayerController* LocalPlayerController = It->GetPlayerController(this))
 				{
 					TArray<FUpdateLevelVisibilityLevelInfo> LevelVisibilities;
 					for (int32 LevelIndex = 1; LevelIndex < Levels.Num(); LevelIndex++)
 					{
 						ULevel*	SubLevel = Levels[LevelIndex];
 
-						FUpdateLevelVisibilityLevelInfo& LevelVisibility = *new( LevelVisibilities ) FUpdateLevelVisibilityLevelInfo();
-						LevelVisibility.PackageName = LocalPlayerController->NetworkRemapPath(SubLevel->GetOutermost()->GetFName(), false);
-						LevelVisibility.bIsVisible = SubLevel->bIsVisible;
+						FUpdateLevelVisibilityLevelInfo& LevelVisibility = *new (LevelVisibilities) FUpdateLevelVisibilityLevelInfo(SubLevel, SubLevel->bIsVisible);
+						LevelVisibility.PackageName = LocalPlayerController->NetworkRemapPath(LevelVisibility.PackageName, false);
 					}
 
-					if( LevelVisibilities.Num() > 0 )
+					if(LevelVisibilities.Num() > 0)
 					{
-						LocalPlayerController->ServerUpdateMultipleLevelsVisibility( LevelVisibilities );
+						LocalPlayerController->ServerUpdateMultipleLevelsVisibility(LevelVisibilities);
 					}
 				}
 			}
