#ue4 - Add GetFloatPropertyAtSplineInputKey() and GetVectorPropertyAtSplineInputKey() to work with spline input keys, for compatibility with query functions that return a key instead of point.

diff --git a/Engine/Source/Runtime/Engine/Classes/Components/SplineComponent.h b/Engine/Source/Runtime/Engine/Classes/Components/SplineComponent.h
--- a/Engine/Source/Runtime/Engine/Classes/Components/SplineComponent.h
+++ b/Engine/Source/Runtime/Engine/Classes/Components/SplineComponent.h
@@ -357,6 +357,12 @@ public:
 	/** Get scale at the provided input key value */
 	FVector GetScaleAtSplineInputKey(float InKey) const;
 
+	/** Get a metadata property float value along the spline at spline input key */
+	float GetFloatPropertyAtSplineInputKey(float InKey, FName PropertyName) const;
+
+	/** Get a metadata property vector value along the spline at spline input key */
+	FVector GetVectorPropertyAtSplineInputKey(float InKey, FName PropertyName) const;
+
 	/** Specify unselected spline component segment color in the editor */
 	UFUNCTION(BlueprintCallable, Category = Editor)
 	void SetUnselectedSplineSegmentColor(const FLinearColor& SegmentColor);
diff --git a/Engine/Source/Runtime/Engine/Private/Components/SplineComponent.cpp b/Engine/Source/Runtime/Engine/Private/Components/SplineComponent.cpp
--- a/Engine/Source/Runtime/Engine/Private/Components/SplineComponent.cpp
+++ b/Engine/Source/Runtime/Engine/Private/Components/SplineComponent.cpp
@@ -460,6 +460,30 @@ FVector USplineComponent::GetScaleAtSplineInputKey(float InKey) const
 	return Scale;
 }
 
+template<class T>
+T GetPropertyValueAtSplineInputKey(const USplineMetadata* Metadata, float InKey, FName PropertyName)
+{
+	if (Metadata)
+	{
+		if (UProperty* Property = Metadata->GetClass()->FindPropertyByName(PropertyName))
+		{
+			const FInterpCurve<T>* Curve = Property->ContainerPtrToValuePtr<FInterpCurve<T>>(Metadata);
+			return Curve->Eval(InKey, T(0));
+		}
+	}
+
+	return T(0);
+}
+
+float USplineComponent::GetFloatPropertyAtSplineInputKey(float InKey, FName PropertyName) const
+{
+	return GetPropertyValueAtSplineInputKey<float>(GetSplinePointsMetadata(), InKey, PropertyName);
+}
+
+FVector USplineComponent::GetVectorPropertyAtSplineInputKey(float InKey, FName PropertyName) const
+{
+	return GetPropertyValueAtSplineInputKey<FVector>(GetSplinePointsMetadata(), InKey, PropertyName);
+}
 
 void USplineComponent::SetClosedLoop(bool bInClosedLoop, bool bUpdateSpline)
 {
