Fix static analysis warning

diff --git a/Engine/Source/Runtime/Engine/Private/SkeletalMeshLODRenderData.cpp b/Engine/Source/Runtime/Engine/Private/SkeletalMeshLODRenderData.cpp
--- a/Engine/Source/Runtime/Engine/Private/SkeletalMeshLODRenderData.cpp
+++ b/Engine/Source/Runtime/Engine/Private/SkeletalMeshLODRenderData.cpp
@@ -916,6 +916,13 @@ void FSkeletalMeshLODRenderData::Serialize(FArchive& Ar, UObject* Owner, int32 I
 	DECLARE_SCOPE_CYCLE_COUNTER(TEXT("FSkeletalMeshLODRenderData::Serialize"), STAT_SkeletalMeshLODRenderData_Serialize, STATGROUP_LoadTime);
 
 	USkeletalMesh* OwnerMesh = CastChecked<USkeletalMesh>(Owner);
+	
+	// Shouldn't needed but to make some static analyzers happy
+	if (!OwnerMesh)
+	{
+		return;
+	}
+	
 	// Actual flags used during serialization
 	const uint8 ClassDataStripFlags = GenerateClassStripFlags(Ar, OwnerMesh, Idx);
 	FStripDataFlags StripFlags(Ar, ClassDataStripFlags);
@@ -935,12 +942,12 @@ void FSkeletalMeshLODRenderData::Serialize(FArchive& Ar, UObject* Owner, int32 I
 	if (!StripFlags.IsDataStrippedForServer())
 	{
 		// set cpu skinning flag on the vertex buffer so that the resource arrays know if they need to be CPU accessible
-		bNeedsCPUAccess = ShouldKeepCPUResources(OwnerMesh, Idx, bForceKeepCPUResources); // -V595
+		bNeedsCPUAccess = ShouldKeepCPUResources(OwnerMesh, Idx, bForceKeepCPUResources);
 	}
 
 	if (FPlatformProperties::RequiresCookedData())
 	{
-		if (bNeedsCPUAccess && OwnerMesh)
+		if (bNeedsCPUAccess)
 		{
 			UE_LOG(LogStaticMesh, Verbose, TEXT("[%s] Skeletal Mesh is marked for CPU read."), *OwnerMesh->GetName());
 		}
