Avoid calling CPFUO in DialogueSoundWave, which is not guaranteed to not modify the sound wave we're copying from.

diff --git a/Engine/Source/Runtime/Engine/Classes/Sound/DialogueWave.h b/Engine/Source/Runtime/Engine/Classes/Sound/DialogueWave.h
--- a/Engine/Source/Runtime/Engine/Classes/Sound/DialogueWave.h
+++ b/Engine/Source/Runtime/Engine/Classes/Sound/DialogueWave.h
@@ -150,4 +150,7 @@ protected:
 
 	void UpdateMappingProxy(FDialogueContextMapping& ContextMapping);
 
+private:
+	void CopySoundBasePropertiesToProxy(const USoundBase* InSoundBase, USoundBase* Proxy);
+
 };
diff --git a/Engine/Source/Runtime/Engine/Private/DialogueWave.cpp b/Engine/Source/Runtime/Engine/Private/DialogueWave.cpp
--- a/Engine/Source/Runtime/Engine/Private/DialogueWave.cpp
+++ b/Engine/Source/Runtime/Engine/Private/DialogueWave.cpp
@@ -809,7 +809,7 @@ void UDialogueWave::UpdateMappingProxy(FDialogueContextMapping& ContextMapping)
 	{
 		// Copy the properties that the proxy shares with the sound in case it's used as a SoundBase
 		ContextMapping.Proxy->SoundWave = ContextMapping.SoundWave;
-		UEngine::CopyPropertiesForUnrelatedObjects(ContextMapping.SoundWave, ContextMapping.Proxy);
+		CopySoundBasePropertiesToProxy(ContextMapping.SoundWave, ContextMapping.Proxy);
 
 		FSubtitleCue NewSubtitleCue;
 		NewSubtitleCue.Text = GetLocalizedSubtitle(ContextMapping);
@@ -819,3 +819,27 @@ void UDialogueWave::UpdateMappingProxy(FDialogueContextMapping& ContextMapping)
 		ContextMapping.Proxy->Subtitles.Add(NewSubtitleCue);
 	}
 }
+
+void UDialogueWave::CopySoundBasePropertiesToProxy(const USoundBase* InSoundBase, USoundBase* Proxy)
+{
+	Proxy->SoundClassObject = InSoundBase->SoundClassObject;
+	Proxy->bDebug = InSoundBase->bDebug;
+	Proxy->bOverrideConcurrency = InSoundBase->bOverrideConcurrency;
+	Proxy->bOutputToBusOnly = InSoundBase->bOutputToBusOnly;
+	Proxy->bHasDelayNode = InSoundBase->bHasDelayNode;
+	Proxy->bHasConcatenatorNode = InSoundBase->bHasConcatenatorNode;
+	Proxy->bHasVirtualizeWhenSilent = InSoundBase->bHasVirtualizeWhenSilent;
+	Proxy->bBypassVolumeScaleForPriority = InSoundBase->bBypassVolumeScaleForPriority;
+	Proxy->ConcurrencySet = InSoundBase->ConcurrencySet;
+	Proxy->ConcurrencyOverrides = InSoundBase->ConcurrencyOverrides;
+	Proxy->Duration = InSoundBase->Duration;
+	Proxy->MaxDistance = InSoundBase->MaxDistance;
+	Proxy->TotalSamples = InSoundBase->TotalSamples;
+	Proxy->Priority = InSoundBase->Priority;
+	Proxy->AttenuationSettings = InSoundBase->AttenuationSettings;
+	Proxy->SoundSubmixObject = InSoundBase->SoundSubmixObject;
+	Proxy->SoundSubmixSends = InSoundBase->SoundSubmixSends;
+	Proxy->SourceEffectChain = InSoundBase->SourceEffectChain;
+	Proxy->BusSends = InSoundBase->BusSends;
+	Proxy->PreEffectBusSends = InSoundBase->PreEffectBusSends;
+}
