[NO NON-PII COMMIT MESSAGE]

diff --git a/Engine/Source/Runtime/Engine/Private/AudioCompressionSettingsUtils.cpp b/Engine/Source/Runtime/Engine/Private/AudioCompressionSettingsUtils.cpp
--- a/Engine/Source/Runtime/Engine/Private/AudioCompressionSettingsUtils.cpp
+++ b/Engine/Source/Runtime/Engine/Private/AudioCompressionSettingsUtils.cpp
@@ -28,6 +28,21 @@ FAutoConsoleVariableRef CVarCookOverrideCachingIntervalCVar(
 	TEXT("n: Time between caching intervals, in seconds."),
 	ECVF_Default);
 
+
+/**
+ * This value is the minimum potential usage of the stream cache we feasibly want to support.
+ * Setting this to 0.25, for example, cause us to potentially be using 25% of our cache size when we start evicting chunks, worst cast scenario.
+ * The trade off is that when this is increased, we add more elements to our cache, thus linearly increasing the CPU complexity of finding a chunk.
+ * A minimum cache usage of 1.0f is impossible, because it would require an infinite amount of chunks.
+ */
+static float MinimumCacheUsageCvar = 0.75f;
+FAutoConsoleVariableRef CVarMinimumCacheUsage(
+	TEXT("au.streamcaching.MinimumCacheUsage"),
+	MinimumCacheUsageCvar,
+	TEXT("This value is the minimum potential usage of the stream cache we feasibly want to support. Setting this to 0.25, for example, cause us to potentially be using 25% of our cache size when we start evicting chunks, worst cast scenario.\n")
+	TEXT("0.0: limit the number of chunks to our (Cache Size / Max Chunk Size) [0.01-0.99]: Increase our number of chunks to limit disk IO when we have lots of small sounds playing."),
+	ECVF_Default);
+
 const FPlatformRuntimeAudioCompressionOverrides* FPlatformCompressionUtilities::GetRuntimeCompressionOverridesForCurrentPlatform()
 {
 #if PLATFORM_ANDROID && !PLATFORM_LUMIN && ENABLE_PLATFORM_COMPRESSION_OVERRIDES
@@ -186,16 +201,19 @@ const FAudioStreamCachingSettings& FPlatformCompressionUtilities::GetStreamCachi
 FCachedAudioStreamingManagerParams FPlatformCompressionUtilities::BuildCachedStreamingManagerParams()
 {
 	const FAudioStreamCachingSettings& CacheSettings = GetStreamCachingSettingsForCurrentPlatform();
-	int32 ChunkSize = GetChunkSizeForCookOverrides(GetCookOverridesForCurrentPlatform());
+	int32 MaxChunkSize = GetMaxChunkSizeForCookOverrides(GetCookOverridesForCurrentPlatform());
 
-	// Our number of elements in the cache is CacheSize in kilobytes divided by our chunk size in kilobytes:
-	int32 NumElements = CacheSettings.CacheSizeKB / (ChunkSize / 1024);
+	// Our number of elements is tweakable based on the minimum cache usage we want to support.
+	const float MinimumCacheUsage = FMath::Clamp(MinimumCacheUsageCvar, 0.0f, 0.95f);
+	int32 MinChunkSize = (1.0f - MinimumCacheUsage) * MaxChunkSize;
+	int32 NumElements = (CacheSettings.CacheSizeKB * 1024) / MinChunkSize;
 
 	FCachedAudioStreamingManagerParams Params;
 	FCachedAudioStreamingManagerParams::FCacheDimensions CacheDimensions;
 
 	// Primary cache defined here:
-	CacheDimensions.MaxChunkSize = ChunkSize;
+	CacheDimensions.MaxChunkSize = MaxChunkSize;
+	CacheDimensions.MaxMemoryInBytes = CacheSettings.CacheSizeKB * 1024;
 	CacheDimensions.NumElements = NumElements;
 	Params.Caches.Add(CacheDimensions);
 
@@ -204,7 +222,7 @@ FCachedAudioStreamingManagerParams FPlatformCompressionUtilities::BuildCachedStr
 	return Params;
 }
 
-uint32 FPlatformCompressionUtilities::GetChunkSizeForCookOverrides(const FPlatformAudioCookOverrides* InCompressionOverrides)
+uint32 FPlatformCompressionUtilities::GetMaxChunkSizeForCookOverrides(const FPlatformAudioCookOverrides* InCompressionOverrides)
 {
 	check(InCompressionOverrides);
 
diff --git a/Engine/Source/Runtime/Engine/Private/AudioDerivedData.cpp b/Engine/Source/Runtime/Engine/Private/AudioDerivedData.cpp
--- a/Engine/Source/Runtime/Engine/Private/AudioDerivedData.cpp
+++ b/Engine/Source/Runtime/Engine/Private/AudioDerivedData.cpp
@@ -329,7 +329,7 @@ class FStreamedAudioCacheDerivedDataWorker : public FNonAbandonableTask
 				if (bUseStreamCaching)
 				{
 					// Use the chunk size for this duration:
-					MaxChunkSize = FPlatformCompressionUtilities::GetChunkSizeForCookOverrides(CompressionOverrides);
+					MaxChunkSize = FPlatformCompressionUtilities::GetMaxChunkSizeForCookOverrides(CompressionOverrides);
 					UE_LOG(LogAudio, Display, TEXT("Chunk size for %s: %d"), *SoundWave.GetFullName(), MaxChunkSize);
 				}
 				
diff --git a/Engine/Source/Runtime/Engine/Private/AudioStreamingCache.cpp b/Engine/Source/Runtime/Engine/Private/AudioStreamingCache.cpp
--- a/Engine/Source/Runtime/Engine/Private/AudioStreamingCache.cpp
+++ b/Engine/Source/Runtime/Engine/Private/AudioStreamingCache.cpp
@@ -27,6 +27,22 @@ FAutoConsoleVariableRef CVarDebugDisplayCaches(
 	TEXT("n: Number of elements to display on screen."),
 	ECVF_Default);
 
+static int32 ForceBlockForLoadCVar = 0;
+FAutoConsoleVariableRef CVarForceBlockForLoad(
+	TEXT("au.streamcaching.ForceBlockForLoad"),
+	ForceBlockForLoadCVar,
+	TEXT("when set to a nonzero value, blocks GetLoadedChunk until the disk read is complete.\n")
+	TEXT("n: Number of elements to display on screen."),
+	ECVF_Default);
+
+static int32 TrimCacheWhenOverBudgetCVar = 1;
+FAutoConsoleVariableRef CVarTrimCacheWhenOverBudget(
+	TEXT("au.streamcaching.TrimCacheWhenOverBudget"),
+	TrimCacheWhenOverBudgetCVar,
+	TEXT("when set to a nonzero value, TrimMemory will be called in AddOrTouchChunk to ensure we never go over budget.\n")
+	TEXT("n: Number of elements to display on screen."),
+	ECVF_Default);
+
 FCachedAudioStreamingManager::FCachedAudioStreamingManager(const FCachedAudioStreamingManagerParams& InitParams)
 {
 	check(FPlatformCompressionUtilities::IsCurrentPlatformUsingStreamCaching());
@@ -35,7 +51,7 @@ FCachedAudioStreamingManager::FCachedAudioStreamingManager(const FCachedAudioStr
 	// const FAudioStreamCachingSettings& CacheSettings = FPlatformCompressionUtilities::GetStreamCachingSettingsForCurrentPlatform();
 	for (const FCachedAudioStreamingManagerParams::FCacheDimensions& CacheDimensions : InitParams.Caches)
 	{
-		CacheArray.Emplace(CacheDimensions.MaxChunkSize, CacheDimensions.NumElements);
+		CacheArray.Emplace(CacheDimensions.MaxChunkSize, CacheDimensions.NumElements, CacheDimensions.MaxMemoryInBytes);
 	}
 
 	// Here we make sure our CacheArray is sorted from smallest MaxChunkSize to biggest, so that GetCacheForWave can scan through these caches to find the appropriate cache for the chunk size.
@@ -158,6 +174,8 @@ bool FCachedAudioStreamingManager::IsManagedStreamingSoundSource(const FSoundSou
 
 FAudioChunkHandle FCachedAudioStreamingManager::GetLoadedChunk(const USoundWave* SoundWave, uint32 ChunkIndex, bool bBlockForLoad) const
 {
+	bBlockForLoad |= (ForceBlockForLoadCVar != 0);
+
 	// If this sound wave is managed by a cache, use that to get the chunk:
 	FAudioChunkCache* Cache = GetCacheForWave(SoundWave);
 	if (Cache)
@@ -301,11 +319,13 @@ bool FCachedAudioStreamingManager::RequestChunk(USoundWave* SoundWave, uint32 Ch
 	}
 }
 
-FAudioChunkCache::FAudioChunkCache(uint32 InMaxChunkSize, uint32 NumChunks)
+FAudioChunkCache::FAudioChunkCache(uint32 InMaxChunkSize, uint32 NumChunks, uint64 InMemoryLimitInBytes)
 	: MaxChunkSize(InMaxChunkSize)
 	, MostRecentElement(nullptr)
 	, LeastRecentElement(nullptr)
 	, ChunksInUse(0)
+	, MemoryCounterBytes(0)
+	, MemoryLimitBytes(InMemoryLimitInBytes)
 	, bLogCacheMisses(false)
 {
 	CachePool.Reset(NumChunks);
@@ -358,6 +378,12 @@ bool FAudioChunkCache::AddOrTouchChunk(const FChunkKey& InKey, TFunction<void(EA
 		}
 
 		KickOffAsyncLoad(CacheElement, InKey, OnLoadCompleted);
+
+		if (TrimCacheWhenOverBudgetCVar != 0 && MemoryCounterBytes > MemoryLimitBytes)
+		{
+			TrimMemory(MemoryCounterBytes - MemoryLimitBytes);
+		}
+
 		return true;
 	}
 }
@@ -473,7 +499,7 @@ uint64 FAudioChunkCache::TrimMemory(uint64 BytesToFree)
 		if (CurrentElement->CanEvictChunk())
 		{
 			BytesFreed += CurrentElement->ChunkData.Num();
-
+			MemoryCounterBytes -= CurrentElement->ChunkData.Num();
 			// Empty the chunk data and invalidate the key.
 			CurrentElement->ChunkData.Empty();
 			CurrentElement->ChunkDataSize = 0;
@@ -655,6 +681,11 @@ void FAudioChunkCache::TouchElement(FCacheElement* InElement)
 	}
 }
 
+bool FAudioChunkCache::ShouldAddNewChunk() const
+{
+	return (ChunksInUse < CachePool.Num()) && (MemoryCounterBytes.Load() < MemoryLimitBytes);
+}
+
 FAudioChunkCache::FCacheElement* FAudioChunkCache::InsertChunk(const FChunkKey& InKey)
 {
 	FCacheElement* CacheElement = nullptr;
@@ -662,7 +693,7 @@ FAudioChunkCache::FCacheElement* FAudioChunkCache::InsertChunk(const FChunkKey&
 	{
 		FScopeLock ScopeLock(&CacheMutationCriticalSection);
 
-		if (ChunksInUse < CachePool.Num())
+		if (ShouldAddNewChunk())
 		{
 			// We haven't filled up the pool yet, so we don't need to evict anything.
 			CacheElement = &CachePool[ChunksInUse];
@@ -777,8 +808,10 @@ void FAudioChunkCache::KickOffAsyncLoad(FCacheElement* CacheElement, const FChun
 
 	EAsyncIOPriorityAndFlags AsyncIOPriority = AIOP_High;
 
+	MemoryCounterBytes -= CacheElement->ChunkData.Num();
 	// Reallocate our chunk data This allows us to shrink if possible.
 	CacheElement->ChunkData.SetNumUninitialized(Chunk.AudioDataSize, true);
+	MemoryCounterBytes += CacheElement->ChunkData.Num();
 
 #if DEBUG_STREAM_CACHE
 	CacheElement->DebugInfo.NumTotalChunks = InKey.SoundWave->GetNumChunks() - 1;
@@ -967,9 +1000,9 @@ TPair<int, int> FAudioChunkCache::DebugDisplay(UWorld* World, FViewport* Viewpor
 	
 	// Convert to megabytes and print the total size:
 	const double NumMegabytesInUse = (double)NumBytesCounter / (1024 * 1024);
-	const double MaxCacheSizeMB = ((double)MaxChunkSize * CachePool.Num()) / (1024 * 1024);
+	const double MaxCacheSizeMB = ((double)MemoryLimitBytes) / (1024 * 1024);
 
-	FString CacheMemoryUsage = *FString::Printf(TEXT("Using: %.4f Megabytes. Max Potential Usage: %.4f Megabytes."), NumMegabytesInUse, MaxCacheSizeMB);
+	FString CacheMemoryUsage = *FString::Printf(TEXT("Using: %.4f Megabytes (%lu bytes). Max Potential Usage: %.4f Megabytes."), NumMegabytesInUse, MemoryCounterBytes.Load(),  MaxCacheSizeMB);
 
 	// We're going to align this horizontally with the number of elements right above it.
 	Canvas->DrawShadowedString(X + CacheTitleOffsetX, Y, *CacheMemoryUsage, UEngine::GetSmallFont(), FLinearColor::Green);
@@ -1047,4 +1080,3 @@ TPair<int, int> FAudioChunkCache::DebugDisplay(UWorld* World, FViewport* Viewpor
 
 	return TPair<int, int>(X + CacheTitleOffsetX + CacheMemoryTextOffsetX - InitialX, Y - InitialY);
 }
-
diff --git a/Engine/Source/Runtime/Engine/Public/AudioCompressionSettingsUtils.h b/Engine/Source/Runtime/Engine/Public/AudioCompressionSettingsUtils.h
--- a/Engine/Source/Runtime/Engine/Public/AudioCompressionSettingsUtils.h
+++ b/Engine/Source/Runtime/Engine/Public/AudioCompressionSettingsUtils.h
@@ -32,7 +32,7 @@ public:
 	static FCachedAudioStreamingManagerParams BuildCachedStreamingManagerParams();
 
 	/** This is used at runtime in BuildCachedStreamingManagerParams, as well as cooktime in FStreamedAudioCacheDerivedDataWorker::BuildStreamedAudio to split compressed audio.  */
-	static uint32 GetChunkSizeForCookOverrides(const FPlatformAudioCookOverrides* InCompressionOverrides);
+	static uint32 GetMaxChunkSizeForCookOverrides(const FPlatformAudioCookOverrides* InCompressionOverrides);
 
 private:
 	static const FPlatformRuntimeAudioCompressionOverrides* GetRuntimeCompressionOverridesForCurrentPlatform();
diff --git a/Engine/Source/Runtime/Engine/Public/AudioStreamingCache.h b/Engine/Source/Runtime/Engine/Public/AudioStreamingCache.h
--- a/Engine/Source/Runtime/Engine/Public/AudioStreamingCache.h
+++ b/Engine/Source/Runtime/Engine/Public/AudioStreamingCache.h
@@ -36,7 +36,7 @@ public:
 		}
 	};
 
-	FAudioChunkCache(uint32 InMaxChunkSize, uint32 NumChunks);
+	FAudioChunkCache(uint32 InMaxChunkSize, uint32 NumChunks, uint64 InMemoryLimitInBytes);
 	
 	~FAudioChunkCache();
 
@@ -203,16 +203,20 @@ private:
 	FCacheElement* MostRecentElement;
 	FCacheElement* LeastRecentElement;
 
-	// This is incremented on every call of InsertChunk until we hit CachePool.Num().
+	// This is incremented on every call of InsertChunk until we hit CachePool.Num() or MemoryCounterBytes hits MemoryLimitBytes.
 	int32 ChunksInUse;
 
+	// This counter is used to start evicting chunks before we hit CachePool.Num().
+	TAtomic<uint64> MemoryCounterBytes;
+	uint64 MemoryLimitBytes;
+
 	// Number of async load operations we have currently in flight.
 	FThreadSafeCounter NumberOfLoadsInFlight;
 
 	// Critical section: only used when we are modifying element positions in the cache. This only happens in TouchElement, EvictLeastRecentChunk, and TrimMemory.
 	// Individual cache elements should be thread safe to access.
 	FCriticalSection CacheMutationCriticalSection;
-
+	 
 	// This struct is used for logging cache misses.
 	struct FCacheMissInfo
 	{
@@ -240,6 +244,10 @@ private:
 	// This is called once we have more than one chunk in our cache:
 	void SetUpLeastRecentChunk();
 
+	// This is called in InsertChunk. it determines whether we should add a new chunk at the tail of the linked list or
+	// evict the least recent chunk.
+	bool ShouldAddNewChunk() const;
+
 	// Returns the least recent chunk and fixes up the linked list accordingly.
 	FCacheElement* EvictLeastRecentChunk();
 
@@ -264,6 +272,11 @@ struct FCachedAudioStreamingManagerParams
 		// At runtime, this will be clamped to ensure that it is greater than the amount of
 		// sources that can be playing simultaneously.
 		int32 NumElements;
+
+		// The maximum number of elements stored in a single cache before it is evicted.
+		// At runtime, this will be clamped to ensure that it is greater than the amount of
+		// sources that can be playing simultaneously.
+		uint64 MaxMemoryInBytes;
 	};
 
 	// Most use cases will only use a single cache, but applications can optionally
