net.DisableBandwithThrottling cvar for quickly disabling bandwidth throttling
#rnx

diff --git a/Engine/Source/Runtime/Engine/Private/NetConnection.cpp b/Engine/Source/Runtime/Engine/Private/NetConnection.cpp
--- a/Engine/Source/Runtime/Engine/Private/NetConnection.cpp
+++ b/Engine/Source/Runtime/Engine/Private/NetConnection.cpp
@@ -62,6 +62,10 @@ static TAutoConsoleVariable<int32> CVarNetPacketOrderMaxCachedPackets(TEXT("net.
 
 TAutoConsoleVariable<int32> CVarNetEnableDetailedScopeCounters(TEXT("net.EnableDetailedScopeCounters"), 1, TEXT("Enables detailed networking scope cycle counters. There are often lots of these which can negatively impact performance."));
 
+#if !UE_BUILD_SHIPPING
+static TAutoConsoleVariable<int32> CVarDisableBandwithThrottling( TEXT( "net.DisableBandwithThrottling" ), 0, TEXT( "Forces IsNetReady to always return true. Not available in shipping builds." ) );
+#endif
+
 extern int32 GNetDormancyValidate;
 
 DECLARE_CYCLE_STAT(TEXT("NetConnection SendAcks"), Stat_NetConnectionSendAck, STATGROUP_Net);
@@ -1500,6 +1504,13 @@ int32 UNetConnection::IsNetReady( bool Saturate )
 		QueuedBits = -SendBuffer.GetNumBits();
 	}
 
+#if !(UE_BUILD_SHIPPING || UE_BUILD_TEST)
+	if (CVarDisableBandwithThrottling.GetValueOnAnyThread() > 0)
+	{
+		return true;
+	}
+#endif
+
 	return QueuedBits + SendBuffer.GetNumBits() <= 0;
 }
 
