ServerExec command on base Playercontroller class. Just compiles out of shipping/test, not tied to CheatManager instantiation.

#rnx

diff --git a/Engine/Source/Runtime/Engine/Classes/GameFramework/PlayerController.h b/Engine/Source/Runtime/Engine/Classes/GameFramework/PlayerController.h
--- a/Engine/Source/Runtime/Engine/Classes/GameFramework/PlayerController.h
+++ b/Engine/Source/Runtime/Engine/Classes/GameFramework/PlayerController.h
@@ -562,6 +562,14 @@ public:
 	UFUNCTION(exec)
 	virtual void LocalTravel(const FString& URL);
 
+	/** RPC used by ServerExec. Not intended to be called directly */
+	UFUNCTION(Reliable, Server, WithValidation)
+	void ServerExecRPC(const FString& Msg);
+
+	/** Executes command on server (non shipping builds only) */
+	UFUNCTION(exec)
+	void ServerExec(const FString& Msg);
+
 	/** Return the client to the main menu gracefully */
 	UE_DEPRECATED(4.19, "As an FString, the ReturnReason parameter is not easily localized. Please use ClientReturnToMainMenuWithTextReason instead.")
 	UFUNCTION(Reliable, Client)
diff --git a/Engine/Source/Runtime/Engine/Private/PlayerController.cpp b/Engine/Source/Runtime/Engine/Private/PlayerController.cpp
--- a/Engine/Source/Runtime/Engine/Private/PlayerController.cpp
+++ b/Engine/Source/Runtime/Engine/Private/PlayerController.cpp
@@ -1775,6 +1775,31 @@ void APlayerController::ServerUpdateCamera_Implementation(FVector_NetQuantize Ca
 
 /// @endcond
 
+bool APlayerController::ServerExecRPC_Validate(const FString& Msg)
+{
+	return true;
+}
+
+void APlayerController::ServerExecRPC_Implementation(const FString& Msg)
+{
+#if !UE_BUILD_SHIPPING
+	ClientMessage(ConsoleCommand(Msg));
+#endif
+}
+	
+void APlayerController::ServerExec(const FString& Msg)
+{
+#if !UE_BUILD_SHIPPING
+
+	if (Msg.Len() > 128)
+	{
+		UE_LOG(LogPlayerController, Warning, TEXT("APlayerController::ServerExec. Msg too big for network RPC. Truncating to 128 character"));
+	}
+
+	ServerExecRPC(Msg.Left(128));
+#endif
+}
+
 void APlayerController::RestartLevel()
 {
 	if( GetNetMode()==NM_Standalone )
