Fix warning spam
#rnx

diff --git a/Engine/Source/Runtime/Engine/Private/Blueprint.cpp b/Engine/Source/Runtime/Engine/Private/Blueprint.cpp
--- a/Engine/Source/Runtime/Engine/Private/Blueprint.cpp
+++ b/Engine/Source/Runtime/Engine/Private/Blueprint.cpp
@@ -971,7 +971,7 @@ void UBlueprint::GetAssetRegistryTags(TArray<FAssetRegistryTag>& OutTags) const
 			FAssetRegistryTag::TT_Alphabetical ) );
 
 	// Only add the FiB tags in the editor, this now gets run for standalone uncooked games
-	if ( ParentClass && GIsEditor)
+	if ( ParentClass && GIsEditor && !GetOutermost()->HasAnyPackageFlags(PKG_ForDiffing))
 	{
 		FString Value;
 		const bool bRebuildSearchData = false;
