Fix merge conflict causing compile error in DialogueWave.cpp. #jira none #rb none.

diff --git a/Engine/Source/Runtime/Engine/Private/DialogueWave.cpp b/Engine/Source/Runtime/Engine/Private/DialogueWave.cpp
--- a/Engine/Source/Runtime/Engine/Private/DialogueWave.cpp
+++ b/Engine/Source/Runtime/Engine/Private/DialogueWave.cpp
@@ -828,7 +828,7 @@ void UDialogueWave::CopySoundBasePropertiesToProxy(const USoundBase* InSoundBase
 	Proxy->bOutputToBusOnly = InSoundBase->bOutputToBusOnly;
 	Proxy->bHasDelayNode = InSoundBase->bHasDelayNode;
 	Proxy->bHasConcatenatorNode = InSoundBase->bHasConcatenatorNode;
-	Proxy->bHasVirtualizeWhenSilent = InSoundBase->bHasVirtualizeWhenSilent;
+	Proxy->VirtualizationMode = InSoundBase->VirtualizationMode;
 	Proxy->bBypassVolumeScaleForPriority = InSoundBase->bBypassVolumeScaleForPriority;
 	Proxy->ConcurrencySet = InSoundBase->ConcurrencySet;
 	Proxy->ConcurrencyOverrides = InSoundBase->ConcurrencyOverrides;
