Fixed a crash that could occur in composite tables when undoing adding a table to itself.
Made changes to CompositeCurveTable to bring the code in-line with CompositeDataTable.

diff --git a/Engine/Source/Runtime/Engine/Classes/Engine/CompositeCurveTable.h b/Engine/Source/Runtime/Engine/Classes/Engine/CompositeCurveTable.h
--- a/Engine/Source/Runtime/Engine/Classes/Engine/CompositeCurveTable.h
+++ b/Engine/Source/Runtime/Engine/Classes/Engine/CompositeCurveTable.h
@@ -39,9 +39,9 @@ protected:
 	// Returns a pointer to the first table found that depends on itself if a loop exists. Returns nullptr if no loops are found.
 	const UCompositeCurveTable* FindLoops(TArray<const UCompositeCurveTable*> AlreadySeenTables) const;
 
-	void UpdateCachedRowMap();
+	void UpdateCachedRowMap(bool bWarnOnInvalidChildren = true);
 
-	void OnParentTablesUpdated();
+	void OnParentTablesUpdated(EPropertyChangeType::Type ChangeType = EPropertyChangeType::Unspecified);
 
 	// true if this asset is currently being loaded; false otherwise
 	uint8 bIsLoading : 1;
diff --git a/Engine/Source/Runtime/Engine/Private/CompositeCurveTable.cpp b/Engine/Source/Runtime/Engine/Private/CompositeCurveTable.cpp
--- a/Engine/Source/Runtime/Engine/Private/CompositeCurveTable.cpp
+++ b/Engine/Source/Runtime/Engine/Private/CompositeCurveTable.cpp
@@ -64,7 +64,7 @@ void UCompositeCurveTable::Serialize(FArchive& Ar)
 	}
 }
 
-void UCompositeCurveTable::UpdateCachedRowMap()
+void UCompositeCurveTable::UpdateCachedRowMap(bool bWarnOnInvalidChildren)
 {
 #if WITH_EDITOR
 	FCurveTableEditorUtils::BroadcastPreChange(this, FCurveTableEditorUtils::ECurveTableChangeInfo::RowList);
@@ -75,16 +75,19 @@ void UCompositeCurveTable::UpdateCachedRowMap()
 	// Throw up an error message and stop if any loops are found
 	if (const UCompositeCurveTable* LoopTable = FindLoops(TArray<const UCompositeCurveTable*>()))
 	{
-		const FText ErrorMsg = FText::Format(LOCTEXT("FoundLoopError", "Cyclic dependency found. Table {0} depends on itself. Please fix your data"), FText::FromString(LoopTable->GetPathName()));
-#if WITH_EDITOR
-		if (!bIsLoading)
+		if (bWarnOnInvalidChildren)
 		{
-			FMessageDialog::Open(EAppMsgType::Ok, ErrorMsg);
-		}
-		else
+			const FText ErrorMsg = FText::Format(LOCTEXT("FoundLoopError", "Cyclic dependency found. Table {0} depends on itself. Please fix your data"), FText::FromString(LoopTable->GetPathName()));
+#if WITH_EDITOR
+			if (!bIsLoading)
+			{
+				FMessageDialog::Open(EAppMsgType::Ok, ErrorMsg);
+			}
+			else
 #endif
-		{
-			UE_LOG(LogCurveTable, Warning, TEXT("%s"), *ErrorMsg.ToString());
+			{
+				UE_LOG(LogCurveTable, Warning, TEXT("%s"), *ErrorMsg.ToString());
+			}
 		}
 		bLeaveEmpty = true;
 	}
@@ -196,10 +199,8 @@ void UCompositeCurveTable::PostEditUndo()
 }
 #endif // WITH_EDITOR
 
-void UCompositeCurveTable::OnParentTablesUpdated()
+void UCompositeCurveTable::OnParentTablesUpdated(EPropertyChangeType::Type ChangeType)
 {
-	UpdateCachedRowMap();
-
 	for (UCurveTable* Table : OldParentTables)
 	{
 		if (Table && ParentTables.Find(Table) == INDEX_NONE)
@@ -208,11 +209,13 @@ void UCompositeCurveTable::OnParentTablesUpdated()
 		}
 	}
 
+	UpdateCachedRowMap(ChangeType == EPropertyChangeType::ValueSet || ChangeType == EPropertyChangeType::Duplicate);
+
 	for (UCurveTable* Table : ParentTables)
 	{
 		if (Table && OldParentTables.Find(Table) == INDEX_NONE)
 		{
-			Table->OnCurveTableChanged().AddUObject(this, &UCompositeCurveTable::UpdateCachedRowMap);
+			Table->OnCurveTableChanged().AddUObject(this, &UCompositeCurveTable::OnParentTablesUpdated, EPropertyChangeType::Unspecified);
 		}
 	}
 
diff --git a/Engine/Source/Runtime/Engine/Private/CompositeDataTable.cpp b/Engine/Source/Runtime/Engine/Private/CompositeDataTable.cpp
--- a/Engine/Source/Runtime/Engine/Private/CompositeDataTable.cpp
+++ b/Engine/Source/Runtime/Engine/Private/CompositeDataTable.cpp
@@ -274,7 +274,7 @@ void UCompositeDataTable::OnParentTablesUpdated(EPropertyChangeType::Type Change
 	{
 		if (Table && OldParentTables.Find(Table) == INDEX_NONE)
 		{
-			Table->OnDataTableChanged().AddUObject(this, &UCompositeDataTable::UpdateCachedRowMap, true);
+			Table->OnDataTableChanged().AddUObject(this, &UCompositeDataTable::OnParentTablesUpdated, EPropertyChangeType::Unspecified);
 		}
 	}
 
