Check and remove VT material callbacks in EvaluateUniformExpressions
Fixes reported check() crash
Integrated from //UE4/Dev-Rendering
#rnx

diff --git a/Engine/Source/Runtime/Engine/Private/Materials/MaterialShared.cpp b/Engine/Source/Runtime/Engine/Private/Materials/MaterialShared.cpp
--- a/Engine/Source/Runtime/Engine/Private/Materials/MaterialShared.cpp
+++ b/Engine/Source/Runtime/Engine/Private/Materials/MaterialShared.cpp
@@ -2339,7 +2339,12 @@ void FMaterialRenderProxy::EvaluateUniformExpressions(FUniformExpressionCache& O
 	OutUniformExpressionCache.ResetAllocatedVTs();
 	OutUniformExpressionCache.AllocatedVTs.Empty(UniformExpressionSet.VTStacks.Num());
 	OutUniformExpressionCache.OwnedAllocatedVTs.Empty(UniformExpressionSet.VTStacks.Num());
-	check(!HasVirtualTextureCallbacks);
+	
+	if (HasVirtualTextureCallbacks)
+	{
+		GetRendererModule().RemoveAllVirtualTextureProducerDestroyedCallbacks(this);
+		HasVirtualTextureCallbacks = false;
+	}
 
 	for (int32 i = 0; i < UniformExpressionSet.VTStacks.Num(); ++i)
 	{
