[NO NON-PII COMMIT MESSAGE]

diff --git a/Engine/Source/Runtime/Engine/Private/AudioStreamingCache.cpp b/Engine/Source/Runtime/Engine/Private/AudioStreamingCache.cpp
--- a/Engine/Source/Runtime/Engine/Private/AudioStreamingCache.cpp
+++ b/Engine/Source/Runtime/Engine/Private/AudioStreamingCache.cpp
@@ -43,6 +43,14 @@ FAutoConsoleVariableRef CVarTrimCacheWhenOverBudget(
 	TEXT("n: Number of elements to display on screen."),
 	ECVF_Default);
 
+static int32 ReadRequestPriorityCVar = 2;
+FAutoConsoleVariableRef CVarReadRequestPriority(
+	TEXT("au.streamcaching.ReadRequestPriority"),
+	ReadRequestPriorityCVar,
+	TEXT("This cvar sets the default request priority for audio chunks when Stream Caching is turned on.\n")
+	TEXT("0: High, 1: Normal, 2: Below Normal, 3: Low, 4: Min"),
+	ECVF_Default);
+
 FCachedAudioStreamingManager::FCachedAudioStreamingManager(const FCachedAudioStreamingManagerParams& InitParams)
 {
 	check(FPlatformCompressionUtilities::IsCurrentPlatformUsingStreamCaching());
@@ -806,7 +814,7 @@ void FAudioChunkCache::KickOffAsyncLoad(FCacheElement* CacheElement, const FChun
 	const FStreamedAudioChunk& Chunk = InKey.SoundWave->RunningPlatformData->Chunks[InKey.ChunkIndex];
 	int32 ChunkDataSize = Chunk.AudioDataSize;
 
-	EAsyncIOPriorityAndFlags AsyncIOPriority = AIOP_High;
+	EAsyncIOPriorityAndFlags AsyncIOPriority = GetAsyncPriorityForChunk(InKey);
 
 	MemoryCounterBytes -= CacheElement->ChunkData.Num();
 	// Reallocate our chunk data This allows us to shrink if possible.
@@ -912,6 +920,37 @@ void FAudioChunkCache::KickOffAsyncLoad(FCacheElement* CacheElement, const FChun
 	}
 }
 
+EAsyncIOPriorityAndFlags FAudioChunkCache::GetAsyncPriorityForChunk(const FChunkKey& InKey)
+{
+
+	// TODO: In the future we can add an enum to USoundWaves to tweak load priority of individual assets.
+
+	switch (ReadRequestPriorityCVar)
+	{
+		case 4:
+		{
+			return AIOP_MIN;
+		}
+		case 3:
+		{
+			return AIOP_Low;
+		}
+		case 2:
+		{
+			return AIOP_BelowNormal;
+		}
+		case 1:
+		{
+			return AIOP_Normal;
+		}
+		case 0:
+		default:
+		{
+			return AIOP_High;
+		}
+	}
+}
+
 bool FAudioChunkCache::IsKeyValid(const FChunkKey& InKey)
 {
 	return InKey.ChunkIndex < TNumericLimits<uint32>::Max() && ((int32)InKey.ChunkIndex) < InKey.SoundWave->RunningPlatformData->Chunks.Num();
diff --git a/Engine/Source/Runtime/Engine/Public/AudioStreamingCache.h b/Engine/Source/Runtime/Engine/Public/AudioStreamingCache.h
--- a/Engine/Source/Runtime/Engine/Public/AudioStreamingCache.h
+++ b/Engine/Source/Runtime/Engine/Public/AudioStreamingCache.h
@@ -252,6 +252,7 @@ private:
 	FCacheElement* EvictLeastRecentChunk();
 
 	void KickOffAsyncLoad(FCacheElement* CacheElement, const FChunkKey& InKey, TFunction<void(EAudioChunkLoadResult)> OnLoadCompleted);
+	EAsyncIOPriorityAndFlags GetAsyncPriorityForChunk(const FChunkKey& InKey);
 };
 
 // This is used to sort the cache array from smallest chunk size to biggest.
