REL - Fix nullpointer deref when Wireframe index buffer is not available

diff --git a/Engine/Source/Runtime/Engine/Private/StaticMeshRender.cpp b/Engine/Source/Runtime/Engine/Private/StaticMeshRender.cpp
--- a/Engine/Source/Runtime/Engine/Private/StaticMeshRender.cpp
+++ b/Engine/Source/Runtime/Engine/Private/StaticMeshRender.cpp
@@ -746,7 +746,7 @@ uint32 FStaticMeshSceneProxy::SetMeshElementGeometrySource(
 	{
 		const bool bSupportsTessellation = RHISupportsTessellation(GetScene().GetShaderPlatform()) && VertexFactory->GetType()->SupportsTessellationShaders();
 
-		if (LODModel.AdditionalIndexBuffers->WireframeIndexBuffer.IsInitialized() && !bSupportsTessellation)
+		if (LODModel.AdditionalIndexBuffers && LODModel.AdditionalIndexBuffers->WireframeIndexBuffer.IsInitialized() && !bSupportsTessellation)
 		{
 			OutMeshBatch.Type = PT_LineList;
 			OutMeshBatchElement.FirstIndex = 0;
