Removed unused world context object parameter in inset setter helper BP function

diff --git a/Engine/Source/Runtime/Engine/Classes/Kismet/KismetRenderingLibrary.h b/Engine/Source/Runtime/Engine/Classes/Kismet/KismetRenderingLibrary.h
--- a/Engine/Source/Runtime/Engine/Classes/Kismet/KismetRenderingLibrary.h
+++ b/Engine/Source/Runtime/Engine/Classes/Kismet/KismetRenderingLibrary.h
@@ -155,6 +155,6 @@ class UKismetRenderingLibrary : public UBlueprintFunctionLibrary
 	/** Set the inset shadow casting state of the given component and all its child attachments. 
 	 *	Also choose if all attachments should be grouped for the inset shadow rendering. If enabled, one depth target will be shared for all attachments.
 	 */
-	UFUNCTION(BlueprintCallable, Category="Rendering", meta=(Keywords="SetCastShadowForAllAttachments", WorldContext="WorldContextObject", UnsafeDuringActorConstruction="true"))
-	static ENGINE_API void SetCastInsetShadowForAllAttachments(UObject* WorldContextObject, UPrimitiveComponent* PrimitiveComponent, bool bCastInsetShadow, bool bLightAttachmentsAsGroup);
+	UFUNCTION(BlueprintCallable, Category="Rendering", meta=(Keywords="SetCastShadowForAllAttachments", UnsafeDuringActorConstruction="true"))
+	static ENGINE_API void SetCastInsetShadowForAllAttachments(UPrimitiveComponent* PrimitiveComponent, bool bCastInsetShadow, bool bLightAttachmentsAsGroup);
 };
diff --git a/Engine/Source/Runtime/Engine/Private/KismetRenderingLibrary.cpp b/Engine/Source/Runtime/Engine/Private/KismetRenderingLibrary.cpp
--- a/Engine/Source/Runtime/Engine/Private/KismetRenderingLibrary.cpp
+++ b/Engine/Source/Runtime/Engine/Private/KismetRenderingLibrary.cpp
@@ -654,7 +654,7 @@ void UKismetRenderingLibrary::BreakSkinWeightInfo(FSkelMeshSkinWeightInfo InWeig
 	Weight3 = InWeight.Weights[3];
 }
 
-void UKismetRenderingLibrary::SetCastInsetShadowForAllAttachments(UObject* WorldContextObject, UPrimitiveComponent* PrimitiveComponent, bool bCastInsetShadow, bool bLightAttachmentsAsGroup)
+void UKismetRenderingLibrary::SetCastInsetShadowForAllAttachments(UPrimitiveComponent* PrimitiveComponent, bool bCastInsetShadow, bool bLightAttachmentsAsGroup)
 {
 	if (PrimitiveComponent)
 	{
