Merging CL #7774262 from //UE4/Dev-Enterprise to Release-4.23

Fixing error log on material instance

jeanluc.corenthin
#rnx

diff --git a/Engine/Source/Runtime/Engine/Classes/Materials/MaterialExpressionTextureObjectParameter.h b/Engine/Source/Runtime/Engine/Classes/Materials/MaterialExpressionTextureObjectParameter.h
--- a/Engine/Source/Runtime/Engine/Classes/Materials/MaterialExpressionTextureObjectParameter.h
+++ b/Engine/Source/Runtime/Engine/Classes/Materials/MaterialExpressionTextureObjectParameter.h
@@ -16,6 +16,10 @@ UCLASS(collapsecategories, hidecategories=(Object, MaterialExpressionTextureSamp
 class UMaterialExpressionTextureObjectParameter : public UMaterialExpressionTextureSampleParameter
 {
 	GENERATED_UCLASS_BODY()
+	
+	//~ Begin UMaterialExpressionTextureSampleParameter Interface
+	virtual bool TextureIsValid(UTexture* InTexture, FString& OutMessage) override;
+	//~ End UMaterialExpressionTextureSampleParameter Interface
 
 
 	//~ Begin UMaterialExpression Interface
diff --git a/Engine/Source/Runtime/Engine/Private/Materials/MaterialExpressions.cpp b/Engine/Source/Runtime/Engine/Private/Materials/MaterialExpressions.cpp
--- a/Engine/Source/Runtime/Engine/Private/Materials/MaterialExpressions.cpp
+++ b/Engine/Source/Runtime/Engine/Private/Materials/MaterialExpressions.cpp
@@ -2468,6 +2468,17 @@ UMaterialExpressionTextureObjectParameter::UMaterialExpressionTextureObjectParam
 #endif // WITH_EDITORONLY_DATA
 }
 
+bool UMaterialExpressionTextureObjectParameter::TextureIsValid(UTexture* InTexture, FString& OutMessage)
+{
+	if (!InTexture)
+	{
+		OutMessage = TEXT("Requires valid texture");
+		return false;
+	}
+
+	return true;
+}
+
 #if WITH_EDITOR
 void UMaterialExpressionTextureObjectParameter::GetCaption(TArray<FString>& OutCaptions) const
 {
@@ -2483,9 +2494,10 @@ const TArray<FExpressionInput*> UMaterialExpressionTextureObjectParameter::GetIn
 
 int32 UMaterialExpressionTextureObjectParameter::Compile(class FMaterialCompiler* Compiler, int32 OutputIndex)
 {
-	if (!Texture)
+	FString ErrorMessage;
+	if (!TextureIsValid(Texture, ErrorMessage))
 	{
-		return CompilerError(Compiler, TEXT("Requires valid texture"));
+		return CompilerError(Compiler, *ErrorMessage);
 	}
 
 	// It seems like this error should be checked here, but this can break existing materials, see https://jira.it.epicgames.net/browse/UE-68862
@@ -2499,9 +2511,10 @@ int32 UMaterialExpressionTextureObjectParameter::Compile(class FMaterialCompiler
 
 int32 UMaterialExpressionTextureObjectParameter::CompilePreview(class FMaterialCompiler* Compiler, int32 OutputIndex)
 {
-	if (!Texture)
+	FString ErrorMessage;
+	if (!TextureIsValid(Texture, ErrorMessage))
 	{
-		return CompilerError(Compiler, TEXT("Requires valid texture"));
+		return CompilerError(Compiler, *ErrorMessage);
 	}
 
 	// Preview the texture object by actually sampling it
