Moved 2D geometry tools out of paper2D and into engine in GeomTools.h/cpp so they can be shared with others

diff --git a/Engine/Plugins/2D/Paper2D/Source/Paper2D/Private/PaperGeomTools.cpp b/Engine/Plugins/2D/Paper2D/Source/Paper2D/Private/PaperGeomTools.cpp
deleted file mode 100644
diff --git a/Engine/Plugins/2D/Paper2D/Source/Paper2D/Private/PaperSprite.cpp b/Engine/Plugins/2D/Paper2D/Source/Paper2D/Private/PaperSprite.cpp
--- a/Engine/Plugins/2D/Paper2D/Source/Paper2D/Private/PaperSprite.cpp
+++ b/Engine/Plugins/2D/Paper2D/Source/Paper2D/Private/PaperSprite.cpp
@@ -10,7 +10,7 @@
 #include "PhysicsEngine/BodySetup.h"
 
 #include "PaperCustomVersion.h"
-#include "PaperGeomTools.h"
+#include "GeomTools.h"
 #include "PaperSpriteComponent.h"
 #include "PaperFlipbookComponent.h"
 #include "PaperGroupedSpriteComponent.h"
@@ -1078,7 +1078,7 @@ void UPaperSprite::BuildGeometryFromContours(FSpriteGeometryCollection& GeomOwne
 			NewShape.SetNewPivot(AverageCenterSnapped);
 
 			// Get intended winding
-			NewShape.bNegativeWinding = !PaperGeomTools::IsPolygonWindingCCW(NewShape.Vertices);
+			NewShape.bNegativeWinding = !FGeomTools2D::IsPolygonWindingCCW(NewShape.Vertices);
 		}
 	}
 }
@@ -1332,7 +1332,7 @@ void UPaperSprite::FindContours(const FIntPoint& ScanPos, const FIntPoint& ScanS
 						// Remove collinear points from the result
 						RemoveCollinearPoints(/*inout*/ ContourPoly);
 
-						if (!PaperGeomTools::IsPolygonWindingCCW(ContourPoly))
+						if (!FGeomTools2D::IsPolygonWindingCCW(ContourPoly))
 						{
 							// Remove newly added polygon, we don't support holes just yet
 							OutPoints.RemoveAt(OutPoints.Num() - 1);
@@ -1956,7 +1956,7 @@ void FSpriteGeometryCollection::Triangulate(TArray<FVector2D>& Target, bool bInc
 				SourcePolygon.GetTextureSpaceVertices(/*out*/ TextureSpaceVertices);
 
 				TArray<FVector2D>& FixedVertices = *new (ValidPolygonTriangles) TArray<FVector2D>();
-				PaperGeomTools::CorrectPolygonWinding(/*out*/ FixedVertices, TextureSpaceVertices, SourcePolygon.bNegativeWinding);
+				FGeomTools2D::CorrectPolygonWinding(/*out*/ FixedVertices, TextureSpaceVertices, SourcePolygon.bNegativeWinding);
 				PolygonsNegativeWinding.Add(SourcePolygon.bNegativeWinding);
 			}
 
@@ -1968,19 +1968,19 @@ void FSpriteGeometryCollection::Triangulate(TArray<FVector2D>& Target, bool bInc
 	}
 
 	// Check if polygons overlap, or have inconsistent winding, or edges overlap
-	if (!PaperGeomTools::ArePolygonsValid(ValidPolygonTriangles))
+	if (!FGeomTools2D::ArePolygonsValid(ValidPolygonTriangles))
 	{
 		return;
 	}
 
 	// Merge each additive and associated subtractive polygons to form a list of polygons in CCW winding
-	ValidPolygonTriangles = PaperGeomTools::ReducePolygons(ValidPolygonTriangles, PolygonsNegativeWinding);
+	ValidPolygonTriangles = FGeomTools2D::ReducePolygons(ValidPolygonTriangles, PolygonsNegativeWinding);
 
 	// Triangulate the polygons
 	for (int32 PolygonIndex = 0; PolygonIndex < ValidPolygonTriangles.Num(); ++PolygonIndex)
 	{
 		TArray<FVector2D> Generated2DTriangles;
-		if (PaperGeomTools::TriangulatePoly(Generated2DTriangles, ValidPolygonTriangles[PolygonIndex], bAvoidVertexMerging))
+		if (FGeomTools2D::TriangulatePoly(Generated2DTriangles, ValidPolygonTriangles[PolygonIndex], bAvoidVertexMerging))
 		{
 			AllGeneratedTriangles.Append(Generated2DTriangles);
 		}
@@ -1991,7 +1991,7 @@ void FSpriteGeometryCollection::Triangulate(TArray<FVector2D>& Target, bool bInc
 	{
 		TArray<FVector2D> TrianglesCopy = AllGeneratedTriangles;
 		AllGeneratedTriangles.Empty();
-		PaperGeomTools::RemoveRedundantTriangles(/*out*/ AllGeneratedTriangles, TrianglesCopy);
+		FGeomTools2D::RemoveRedundantTriangles(/*out*/ AllGeneratedTriangles, TrianglesCopy);
 	}
 
 	Target.Append(AllGeneratedTriangles);
diff --git a/Engine/Plugins/2D/Paper2D/Source/Paper2D/Private/Terrain/PaperTerrainComponent.cpp b/Engine/Plugins/2D/Paper2D/Source/Paper2D/Private/Terrain/PaperTerrainComponent.cpp
--- a/Engine/Plugins/2D/Paper2D/Source/Paper2D/Private/Terrain/PaperTerrainComponent.cpp
+++ b/Engine/Plugins/2D/Paper2D/Source/Paper2D/Private/Terrain/PaperTerrainComponent.cpp
@@ -8,7 +8,7 @@
 #include "PaperCustomVersion.h"
 
 #include "PaperRenderSceneProxy.h"
-#include "PaperGeomTools.h"
+#include "GeomTools.h"
 #include "PhysicsEngine/ConvexElem.h"
 #include "PhysicsEngine/BoxElem.h"
 #include "PhysicsEngine/BodySetup.h"
@@ -659,10 +659,10 @@ void UPaperTerrainComponent::OnSplineEdited()
 
 			// Always CCW and facing forward regardless of spline winding
 			TArray<FVector2D> CorrectedSplineVertices;
-			PaperGeomTools::CorrectPolygonWinding(CorrectedSplineVertices, SplinePolyVertices2D, false);
+			FGeomTools2D::CorrectPolygonWinding(CorrectedSplineVertices, SplinePolyVertices2D, false);
 
 			TArray<FVector2D> TriangulatedPolygonVertices;
-			PaperGeomTools::TriangulatePoly(/*out*/TriangulatedPolygonVertices, CorrectedSplineVertices, false);
+			FGeomTools2D::TriangulatePoly(/*out*/TriangulatedPolygonVertices, CorrectedSplineVertices, false);
 
 			GenerateCollisionDataFromPolygon(SplinePolyVertices2D, SplineEdgeOffsetAmounts, TriangulatedPolygonVertices);
 
@@ -920,7 +920,7 @@ static void CreateExtrudedConvexHull(TArray<FVector2D>& OutConvexHull, TArray<FV
 		ExtrudedPoints.Add(B - N * extrude);
 	}
 
-	PaperGeomTools::GenerateConvexHullFromPoints(OutConvexHull, ExtrudedPoints);
+	FGeomTools2D::GenerateConvexHullFromPoints(OutConvexHull, ExtrudedPoints);
 }
 
 void UPaperTerrainComponent::GenerateFillRenderDataFromPolygon(const class UPaperSprite* NewSprite, FSpriteDrawCallRecord& FillDrawCall, const FVector2D& TextureSize, const TArray<FVector2D>& TriangulatedPolygonVertices)
@@ -944,7 +944,7 @@ void UPaperTerrainComponent::GenerateCollisionDataFromPolygon(const TArray<FVect
 	{
 		// Generate polygon collider
 		TArray<TArray<FVector2D>> ConvexHulls;
-		PaperGeomTools::GenerateConvexPolygonsFromTriangles(ConvexHulls, TriangulatedPolygonVertices);
+		FGeomTools2D::GenerateConvexPolygonsFromTriangles(ConvexHulls, TriangulatedPolygonVertices);
 		TArray<float> ConvexHullEdgeExtrusionAmount;
 
 		for (TArray<FVector2D> ConvexHull : ConvexHulls)
@@ -1013,20 +1013,20 @@ void UPaperTerrainComponent::InsertConvexCollisionDataFromPolygon(const TArray<F
 
 		// Always CCW and facing forward regardless of spline winding
 		TArray<FVector2D> CorrectedSplineVertices;
-		PaperGeomTools::CorrectPolygonWinding(CorrectedSplineVertices, LocalPolyVertices, false);
+		FGeomTools2D::CorrectPolygonWinding(CorrectedSplineVertices, LocalPolyVertices, false);
 
 		TArray<FVector2D> TriangulatedPolygonVertices;
-		if (!PaperGeomTools::TriangulatePoly(/*out*/TriangulatedPolygonVertices, CorrectedSplineVertices, false))
+		if (!FGeomTools2D::TriangulatePoly(/*out*/TriangulatedPolygonVertices, CorrectedSplineVertices, false))
 		{
 			// Triangulation failed, try triangulating the original non simplified polygon
 			CorrectedSplineVertices.Empty();
-			PaperGeomTools::CorrectPolygonWinding(/*out*/CorrectedSplineVertices, ClosedPolyVertices2D, false);
+			FGeomTools2D::CorrectPolygonWinding(/*out*/CorrectedSplineVertices, ClosedPolyVertices2D, false);
 			TriangulatedPolygonVertices.Empty();
-			PaperGeomTools::TriangulatePoly(/*out*/TriangulatedPolygonVertices, CorrectedSplineVertices, false);
+			FGeomTools2D::TriangulatePoly(/*out*/TriangulatedPolygonVertices, CorrectedSplineVertices, false);
 		}
 
 		TArray<TArray<FVector2D>> ConvexHulls;
-		PaperGeomTools::GenerateConvexPolygonsFromTriangles(ConvexHulls, TriangulatedPolygonVertices);
+		FGeomTools2D::GenerateConvexPolygonsFromTriangles(ConvexHulls, TriangulatedPolygonVertices);
 
 		for (TArray<FVector2D> ConvexHull : ConvexHulls)
 		{
diff --git a/Engine/Plugins/2D/Paper2D/Source/Paper2D/Public/PaperGeomTools.h b/Engine/Plugins/2D/Paper2D/Source/Paper2D/Public/PaperGeomTools.h
deleted file mode 100644
diff --git a/Engine/Plugins/2D/Paper2D/Source/Paper2DEditor/Private/PaperEditorShared/SpriteGeometryEditing.cpp b/Engine/Plugins/2D/Paper2D/Source/Paper2DEditor/Private/PaperEditorShared/SpriteGeometryEditing.cpp
--- a/Engine/Plugins/2D/Paper2D/Source/Paper2DEditor/Private/PaperEditorShared/SpriteGeometryEditing.cpp
+++ b/Engine/Plugins/2D/Paper2D/Source/Paper2DEditor/Private/PaperEditorShared/SpriteGeometryEditing.cpp
@@ -11,7 +11,7 @@
 #include "PaperEditorShared/AssetEditorSelectedItem.h"
 #include "SpriteEditorOnlyTypes.h"
 #include "SpriteEditor/SpriteEditorSelections.h"
-#include "PaperGeomTools.h"
+#include "GeomTools.h"
 #include "DynamicMeshBuilder.h"
 
 #define LOCTEXT_NAMESPACE "PaperGeometryEditing"
@@ -378,7 +378,7 @@ void FSpriteGeometryEditingHelper::DrawGeometry(const FSceneView& View, FPrimiti
 			Shape.GetTextureSpaceVertices(/*out*/ SourceTextureSpaceVertices);
 
 			TArray<FVector2D> TriangulatedPolygonVertices;
-			PaperGeomTools::TriangulatePoly(/*out*/ TriangulatedPolygonVertices, SourceTextureSpaceVertices, /*bKeepColinearVertices=*/ true);
+			FGeomTools2D::TriangulatePoly(/*out*/ TriangulatedPolygonVertices, SourceTextureSpaceVertices, /*bKeepColinearVertices=*/ true);
 
 			if (((TriangulatedPolygonVertices.Num() % 3) == 0) && (TriangulatedPolygonVertices.Num() > 0))
 			{
diff --git a/Engine/Source/Runtime/AIModule/Private/Navigation/NavLocalGridData.cpp b/Engine/Source/Runtime/AIModule/Private/Navigation/NavLocalGridData.cpp
--- a/Engine/Source/Runtime/AIModule/Private/Navigation/NavLocalGridData.cpp
+++ b/Engine/Source/Runtime/AIModule/Private/Navigation/NavLocalGridData.cpp
@@ -4,24 +4,8 @@
 #include "Navigation/NavLocalGridManager.h"
 #include "GraphAStar.h"
 #include "NavigationPath.h"
+#include "GeomTools.h"
 
-namespace FGridMathHelpers
-{
-	// copy & paste from PaperGeomTools.cpp
-	bool IsPointInPolygon(const FVector2D& TestPoint, const TArray<FVector2D>& PolygonPoints)
-	{
-		const int NumPoints = PolygonPoints.Num();
-		float AngleSum = 0.0f;
-		for (int PointIndex = 0; PointIndex < NumPoints; ++PointIndex)
-		{
-			const FVector2D& VecAB = PolygonPoints[PointIndex] - TestPoint;
-			const FVector2D& VecAC = PolygonPoints[(PointIndex + 1) % NumPoints] - TestPoint;
-			const float Angle = FMath::Sign(FVector2D::CrossProduct(VecAB, VecAC)) * FMath::Acos(FMath::Clamp(FVector2D::DotProduct(VecAB, VecAC) / (VecAB.Size() * VecAC.Size()), -1.0f, 1.0f));
-			AngleSum += Angle;
-		}
-		return (FMath::Abs(AngleSum) > 0.001f);
-	}
-}
 
 FNavLocalGridData::FNavLocalGridData(const FVector& Center, float Extent2D) : GridId(0)
 {
@@ -132,7 +116,7 @@ void FNavLocalGridData::MarkBoxObstacle(const FVector& Center, const FVector& Ex
 			for (int32 IdxY = 0; IdxY < (int32)GridSize.Height; IdxY++)
 			{
 				const FVector GridCellLocation = GetWorldCellCenter(IdxX, IdxY);
-				const bool bIsInside = FGridMathHelpers::IsPointInPolygon(FVector2D(GridCellLocation), Verts2D);
+				const bool bIsInside = FGeomTools2D::IsPointInPolygon(FVector2D(GridCellLocation), Verts2D);
 				if (bIsInside)
 				{
 					const int32 CellIndex = GetCellIndexUnsafe(IdxX, IdxY);
diff --git a/Engine/Source/Runtime/Engine/Private/GeomTools.cpp b/Engine/Source/Runtime/Engine/Private/GeomTools.cpp
--- a/Engine/Source/Runtime/Engine/Private/GeomTools.cpp
+++ b/Engine/Source/Runtime/Engine/Private/GeomTools.cpp
@@ -12,68 +12,6 @@
 
 DEFINE_LOG_CATEGORY_STATIC(LogGeomTools, Log, All);
 
-/** Util struct for storing a set of planes defining a convex region, with neighbour information. */
-struct FHullPlanes
-{
-	TArray<FPlane>	Planes;
-	TArray<int32>		PlaneNeighbour;
-	TArray<float>	PlaneNeighbourArea;
-
-	/** 
-	 *	Util for removing redundany planes from the hull. 
-	 *	Also calculates area of connection
-	 */
-	void RemoveRedundantPlanes()
-	{
-		check(Planes.Num() == PlaneNeighbour.Num());
-
-		// For each plane
-		for(int32 i=Planes.Num()-1; i>=0; i--)
-		{
-			// Build big-ass polygon
-			FPoly Polygon;
-			Polygon.Normal = Planes[i];
-
-			FVector AxisX, AxisY;
-			Polygon.Normal.FindBestAxisVectors(AxisX,AxisY);
-
-			const FVector Base = Planes[i] * Planes[i].W;
-
-			new(Polygon.Vertices) FVector(Base + AxisX * HALF_WORLD_MAX + AxisY * HALF_WORLD_MAX);
-			new(Polygon.Vertices) FVector(Base - AxisX * HALF_WORLD_MAX + AxisY * HALF_WORLD_MAX);
-			new(Polygon.Vertices) FVector(Base - AxisX * HALF_WORLD_MAX - AxisY * HALF_WORLD_MAX);
-			new(Polygon.Vertices) FVector(Base + AxisX * HALF_WORLD_MAX - AxisY * HALF_WORLD_MAX);
-
-			// Clip away each plane
-			for(int32 j=0; j<Planes.Num(); j++)
-			{
-				if(i != j)
-				{
-					if(!Polygon.Split(-FVector(Planes[j]), Planes[j] * Planes[j].W))
-					{
-						// At this point we clipped away all the polygon - so this plane is redundant - remove!
-						Polygon.Vertices.Empty();
-
-						Planes.RemoveAt(i);
-						PlaneNeighbour.RemoveAt(i);
-
-						break;
-					}
-				}
-			}
-
-			// If we still have a polygon left - this is a neighbour - calculate area of connection.
-			if(Polygon.Vertices.Num() > 0)
-			{
-				// Insert at front - we are walking planes from end to start.
-				PlaneNeighbourArea.Insert(Polygon.Area(), 0);
-			}
-		}
-
-		check(PlaneNeighbourArea.Num() == PlaneNeighbour.Num());
-		check(Planes.Num() == PlaneNeighbour.Num());
-	}
-};
 	
 /** */
 static FClipSMVertex GetVert(const UStaticMesh* StaticMesh, int32 VertIndex)
@@ -951,3 +889,825 @@ void FGeomTools::Split2DPolysWithPlane(FUtilPoly2DSet& PolySet, const FPlane& Pl
 		}
 	}
 }
+
+
+
+bool FGeomTools2D::IsPolygonWindingCCW(const TArray<FVector2D>& Points)
+{
+	float Sum = 0.0f;
+	const int PointCount = Points.Num();
+	for (int PointIndex = 0; PointIndex < PointCount; ++PointIndex)
+	{
+		const FVector2D& A = Points[PointIndex];
+		const FVector2D& B = Points[(PointIndex + 1) % PointCount];
+		Sum += (B.X - A.X) * (B.Y + A.Y);
+	}
+	return (Sum < 0.0f);
+}
+
+bool FGeomTools2D::IsPolygonWindingCCW(const TArray<FIntPoint>& Points)
+{
+	int32 Sum = 0;
+	const int PointCount = Points.Num();
+	for (int PointIndex = 0; PointIndex < PointCount; ++PointIndex)
+	{
+		const FIntPoint& A = Points[PointIndex];
+		const FIntPoint& B = Points[(PointIndex + 1) % PointCount];
+		Sum += (B.X - A.X) * (B.Y + A.Y);
+	}
+	return (Sum < 0);
+}
+
+// Note: We need to simplify non-simple polygons before this
+static bool IsPolygonConvex(const TArray<FVector2D>& Points)
+{
+	const int PointCount = Points.Num();
+	float Sign = 0;
+	for (int32 PointIndex = 0; PointIndex < PointCount; ++PointIndex)
+	{
+		const FVector2D& A = Points[PointIndex];
+		const FVector2D& B = Points[(PointIndex + 1) % PointCount];
+		const FVector2D& C = Points[(PointIndex + 2) % PointCount];
+		float Det = (B.X - A.X) * (C.Y - B.Y) - (B.Y - A.Y) * (C.X - B.X);
+		float DetSign = FMath::Sign(Det);
+		if (DetSign != 0)
+		{
+			if (Sign == 0)
+			{
+				Sign = DetSign;
+			}
+			else if (Sign != DetSign)
+			{
+				return false;
+			}
+		}
+	}
+
+	return true;
+}
+
+
+static bool IsAdditivePointInTriangle(const FVector2D& TestPoint, const FVector2D& A, const FVector2D& B, const FVector2D& C)
+{
+	const FVector2D AP = TestPoint - A;
+	const FVector2D BP = TestPoint - B;
+	const FVector2D AB = B - A;
+	const FVector2D AC = C - A;
+	const FVector2D BC = C - B;
+	if (FVector2D::CrossProduct(AB, AP) <= 0.0f) return false;
+	if (FVector2D::CrossProduct(AP, AC) <= 0.0f) return false;
+	if (FVector2D::CrossProduct(BC, BP) <= 0.0f) return false;
+	if (AP.SizeSquared() < 2.0f) return false;
+	if (BP.SizeSquared() < 2.0f) return false;
+	const FVector2D CP = TestPoint - C;
+	if (CP.SizeSquared() < 2.0f) return false;
+	return true;
+}
+
+static float VectorSign(const FVector2D& Vec, const FVector2D& A, const FVector2D& B)
+{
+	return FMath::Sign((B.X - A.X) * (Vec.Y - A.Y) - (B.Y - A.Y) * (Vec.X - A.X));
+}
+
+// Returns true when the point is inside the triangle
+// Should not return true when the point is on one of the edges
+static bool IsPointInTriangle(const FVector2D& TestPoint, const FVector2D& A, const FVector2D& B, const FVector2D& C)
+{
+	float BA = VectorSign(B, A, TestPoint);
+	float CB = VectorSign(C, B, TestPoint);
+	float AC = VectorSign(A, C, TestPoint);
+
+	// point is in the same direction of all 3 tri edge lines
+	// must be inside, regardless of tri winding
+	return BA == CB && CB == AC;
+}
+
+// Returns true when the point is on the line segment limited by A and B
+static bool IsPointOnLineSegment(const FVector2D& TestPoint, const FVector2D& A, const FVector2D& B)
+{
+	FVector2D BA = B - A;
+	FVector2D PA = TestPoint - A;
+	float SizeSquaredBA = FVector2D::DotProduct(BA, BA);
+	float AreaCompareThreshold = 0.01f * SizeSquaredBA;
+	float ParallelogramArea = BA.X * PA.Y - BA.Y * PA.X;
+
+	return  TestPoint.X >= FMath::Min(A.X, B.X) && TestPoint.X <= FMath::Max(A.X, B.X) && // X within AB.X, including ON A or B
+		TestPoint.Y >= FMath::Min(A.Y, B.Y) && TestPoint.Y <= FMath::Max(A.Y, B.Y) && // Y within AB.Y, including ON A or B
+		FMath::Abs(ParallelogramArea) < AreaCompareThreshold; // Area is smaller than allowed epsilon = point on line
+}
+
+static void JoinSubtractiveToAdditive(TArray<FVector2D>& AdditivePoly, const TArray<FVector2D>& SubtractivePoly, const int AdditiveJoinIndex, const int SubtractiveJoinIndex)
+{
+	TArray<FVector2D> NewAdditivePoly;
+	for (int AdditiveIndex = 0; AdditiveIndex < AdditivePoly.Num(); ++AdditiveIndex)
+	{
+		NewAdditivePoly.Add(AdditivePoly[AdditiveIndex]);
+		if (AdditiveIndex == AdditiveJoinIndex)
+		{
+			for (int SubtractiveIndex = SubtractiveJoinIndex; SubtractiveIndex < SubtractivePoly.Num(); ++SubtractiveIndex)
+			{
+				NewAdditivePoly.Add(SubtractivePoly[SubtractiveIndex]);
+			}
+			for (int SubtractiveIndex = 0; SubtractiveIndex <= SubtractiveJoinIndex; ++SubtractiveIndex)
+			{
+				NewAdditivePoly.Add(SubtractivePoly[SubtractiveIndex]);
+			}
+			NewAdditivePoly.Add(AdditivePoly[AdditiveIndex]);
+		}
+	}
+	AdditivePoly = NewAdditivePoly;
+}
+
+
+static void JoinMutuallyVisible(TArray<FVector2D>& AdditivePoly, const TArray<FVector2D>& SubtractivePoly)
+{
+	const int NumAdditivePoly = AdditivePoly.Num();
+	const int NumSubtractivePoly = SubtractivePoly.Num();
+	if (NumSubtractivePoly == 0)
+	{
+		return;
+	}
+
+	// Search the inner (subtractive) polygon for the point of maximum x-value
+	int IndexMaxX = 0;
+	{
+		float MaxX = SubtractivePoly[0].X;
+		for (int Index = 1; Index < NumSubtractivePoly; ++Index)
+		{
+			if (SubtractivePoly[Index].X > MaxX)
+			{
+				MaxX = SubtractivePoly[Index].X;
+				IndexMaxX = Index;
+			}
+		}
+	}
+	const FVector2D PointMaxX = SubtractivePoly[IndexMaxX];
+
+	// Intersect a ray from point M facing to the right (a, ab) with the additive shape edges (c, d)
+	// Find the closest intersecting point to M (left-most point)
+	int EdgeStartPointIndex = 0;
+	int EdgeEndPointIndex = 0;
+	bool bIntersectedAtVertex = false;
+	float LeftMostIntersectX = MAX_FLT;
+	const FVector2D A = PointMaxX;
+	const FVector2D AB = FVector2D(1.0f, 0.0f);
+	for (int AdditiveIndex = 0; AdditiveIndex < NumAdditivePoly; ++AdditiveIndex)
+	{
+		const FVector2D& C = AdditivePoly[AdditiveIndex];
+		const FVector2D& D = AdditivePoly[(AdditiveIndex + 1) % NumAdditivePoly];
+		const FVector2D CD = D - C;
+		// Only check edges from the inside, as edges will overlap after mutually visible points are merged.
+		if (CD.Y > 0.0f)
+		{
+			const float DetS = AB.X * CD.Y - AB.Y * CD.X;
+			const float DetT = CD.X * AB.Y - CD.Y * AB.X;
+			if (DetS != 0.0f && DetT != 0.0f)
+			{
+				const float S = (A.Y * CD.X - C.Y * CD.X - A.X * CD.Y + C.X * CD.Y) / DetS;
+				const float T = (C.Y * AB.X - A.Y * AB.X - C.X * AB.Y + A.X * AB.Y) / DetT;
+				if (S >= 0.0f && T >= 0.0f && T <= 1.0f)
+				{
+					const float IntersectX = A.X + AB.X * S;
+					if (IntersectX < LeftMostIntersectX)
+					{
+						LeftMostIntersectX = IntersectX;
+						EdgeStartPointIndex = AdditiveIndex;
+						EdgeEndPointIndex = (AdditiveIndex + 1) % NumAdditivePoly;
+						if (T < FLT_EPSILON)
+						{
+							bIntersectedAtVertex = true;
+						}
+						else if (T > 1.0f - FLT_EPSILON)
+						{
+							bIntersectedAtVertex = true;
+							EdgeStartPointIndex = EdgeEndPointIndex;
+						}
+					}
+				}
+			}
+		}
+	}
+
+	// If the ray intersected a vertex, points are mutually visible
+	if (bIntersectedAtVertex)
+	{
+		JoinSubtractiveToAdditive(AdditivePoly, SubtractivePoly, EdgeStartPointIndex, IndexMaxX);
+		return;
+	}
+
+	// Otherwise, set P to be the edge endpoint with maximum x value
+	const FVector2D Intersect(LeftMostIntersectX, PointMaxX.Y);
+	const int IndexP = (AdditivePoly[EdgeStartPointIndex].X > AdditivePoly[EdgeEndPointIndex].X) ? EdgeStartPointIndex : EdgeEndPointIndex;
+	const FVector2D P = AdditivePoly[IndexP];
+
+	// Search the vertices of the additive shape. If all of these are outside the triangle (M, intersect, P) then M and P are mutually visible
+	// If any vertices lie inside the triangle, find the one that minimizes the angle between (1, 0) and the line M R. (If multiple vertices minimize the angle, choose closest to M)
+	const FVector2D TriA = PointMaxX;
+	const FVector2D TriB = (P.Y < Intersect.Y) ? P : Intersect;
+	const FVector2D TriC = (P.Y < Intersect.Y) ? Intersect : P;
+	float CosAngleMax = 0.0f;
+	float DistanceMin = MAX_FLT;
+	int IndexR = -1;
+	for (int AdditiveIndex = 0; AdditiveIndex < NumAdditivePoly; ++AdditiveIndex)
+	{
+		// Ignore point P
+		if (AdditiveIndex == IndexP)
+		{
+			continue;
+		}
+
+		if (IsAdditivePointInTriangle(AdditivePoly[AdditiveIndex], TriA, TriB, TriC))
+		{
+			const FVector2D MR = AdditivePoly[AdditiveIndex] - PointMaxX;
+			const float CosAngle = MR.X / MR.Size();
+			const float DSq = MR.SizeSquared();
+			if (CosAngle > CosAngleMax || (CosAngle == CosAngleMax && DSq < DistanceMin))
+			{
+				CosAngleMax = CosAngle;
+				DistanceMin = DSq;
+				IndexR = AdditiveIndex;
+			}
+		}
+	}
+
+	JoinSubtractiveToAdditive(AdditivePoly, SubtractivePoly, (IndexR == -1) ? IndexP : IndexR, IndexMaxX);
+}
+
+/** Util that tries to combine two triangles if possible. */
+// Unlike the implementation in GeomTools, we only care about 2D points and not any other vertex attributes
+// If bConvex == true, the triangles are only merged if the resulting polygon is convex
+static bool MergeTriangleIntoPolygon(
+	TArray<FVector2D>& PolygonVertices,
+	const FVector2D& TriangleVertexA, const FVector2D& TriangleVertexB, const FVector2D& TriangleVertexC, bool bConvex = false)
+{
+	// Create indexable copy
+	FVector2D TriangleVertices[3] = { TriangleVertexA, TriangleVertexB, TriangleVertexC };
+
+	for (int32 PolygonEdgeIndex = 0; PolygonEdgeIndex < PolygonVertices.Num(); PolygonEdgeIndex++)
+	{
+		const uint32 PolygonEdgeVertex0 = PolygonEdgeIndex + 0;
+		const uint32 PolygonEdgeVertex1 = (PolygonEdgeIndex + 1) % PolygonVertices.Num();
+
+		for (uint32 TriangleEdgeIndex = 0; TriangleEdgeIndex < 3; TriangleEdgeIndex++)
+		{
+			const uint32 TriangleEdgeVertex0 = TriangleEdgeIndex + 0;
+			const uint32 TriangleEdgeVertex1 = (TriangleEdgeIndex + 1) % 3;
+
+			// If the triangle and polygon share an edge, then the triangle is in the same plane (implied by the above normal check),
+			// and may be merged into the polygon.
+			if (PolygonVertices[PolygonEdgeVertex0].Equals(TriangleVertices[TriangleEdgeVertex1], THRESH_POINTS_ARE_SAME) &&
+				PolygonVertices[PolygonEdgeVertex1].Equals(TriangleVertices[TriangleEdgeVertex0], THRESH_POINTS_ARE_SAME))
+			{
+				bool bMergeTriangle = true;
+				if (bConvex)
+				{
+					TArray<FVector2D> TmpPolygonVertcies = PolygonVertices;
+					const int32 TriangleOppositeVertexIndex = (TriangleEdgeIndex + 2) % 3;
+					TmpPolygonVertcies.Insert(TriangleVertices[TriangleOppositeVertexIndex], PolygonEdgeVertex1);
+					bMergeTriangle = IsPolygonConvex(TmpPolygonVertcies);
+				}
+
+				if (bMergeTriangle)
+				{
+					// Add the triangle's vertex that isn't in the adjacent edge to the polygon in between the vertices of the adjacent edge.
+					const int32 TriangleOppositeVertexIndex = (TriangleEdgeIndex + 2) % 3;
+					PolygonVertices.Insert(TriangleVertices[TriangleOppositeVertexIndex], PolygonEdgeVertex1);
+
+					return true;
+				}
+			}
+		}
+	}
+
+	// Could not merge triangles.
+	return false;
+}
+
+
+TArray<TArray<FVector2D>> FGeomTools2D::ReducePolygons(const TArray<TArray<FVector2D>>& Polygons, const TArray<bool>& PolygonNegativeWinding)
+{
+	TArray<TArray<FVector2D>> ReturnPolygons;
+
+	int NumPolygons = Polygons.Num();
+
+	TArray<float> MaxXValues; // per polygon
+	for (int PolyIndex = 0; PolyIndex < NumPolygons; ++PolyIndex)
+	{
+		float MaxX = -BIG_NUMBER;
+		const TArray<FVector2D>& Vertices = Polygons[PolyIndex];
+		for (int VertexIndex = 0; VertexIndex < Vertices.Num(); ++VertexIndex)
+		{
+			MaxX = FMath::Max(Vertices[VertexIndex].X, MaxX);
+		}
+		MaxXValues.Add(MaxX);
+	}
+
+	// Iterate through additive shapes
+	for (int PolyIndex = 0; PolyIndex < NumPolygons; ++PolyIndex)
+	{
+		if (!PolygonNegativeWinding[PolyIndex])
+		{
+			TArray<FVector2D> Verts(Polygons[PolyIndex]);
+
+			// Store indexes of subtractive shapes that lie inside the additive shape
+			TArray<int> SubtractiveShapeIndices;
+			for (int J = 0; J < NumPolygons; ++J)
+			{
+				const TArray<FVector2D>& CandidateShape = Polygons[J];
+				if (PolygonNegativeWinding[J])
+				{
+					if (IsPointInPolygon(CandidateShape[0], Verts))
+					{
+						SubtractiveShapeIndices.Add(J);
+					}
+				}
+			}
+
+			// Remove subtractive shapes that lie inside the subtractive shapes we've found
+			for (int J = 0; J < SubtractiveShapeIndices.Num();)
+			{
+				const TArray<FVector2D>& ourShape = Polygons[SubtractiveShapeIndices[J]];
+				bool bRemoveOurShape = false;
+				for (int K = 0; K < SubtractiveShapeIndices.Num(); ++K)
+				{
+					if (J == K)
+					{
+						continue;
+					}
+					if (IsPointInPolygon(ourShape[0], Polygons[SubtractiveShapeIndices[K]]))
+					{
+						bRemoveOurShape = true;
+						break;
+					}
+				}
+				if (bRemoveOurShape)
+				{
+					SubtractiveShapeIndices.RemoveAt(J);
+				}
+				else
+				{
+					++J;
+				}
+			}
+
+			// Sort subtractive shapes from right to left by their points' maximum x value
+			const int NumSubtractiveShapes = SubtractiveShapeIndices.Num();
+			for (int J = 0; J < NumSubtractiveShapes; ++J)
+			{
+				for (int K = J + 1; K < NumSubtractiveShapes; ++K)
+				{
+					if (MaxXValues[SubtractiveShapeIndices[J]] < MaxXValues[SubtractiveShapeIndices[K]])
+					{
+						int tempIndex = SubtractiveShapeIndices[J];
+						SubtractiveShapeIndices[J] = SubtractiveShapeIndices[K];
+						SubtractiveShapeIndices[K] = tempIndex;
+					}
+				}
+			}
+
+			for (int SubtractiveIndex = 0; SubtractiveIndex < NumSubtractiveShapes; ++SubtractiveIndex)
+			{
+				JoinMutuallyVisible(Verts, Polygons[SubtractiveShapeIndices[SubtractiveIndex]]);
+			}
+
+			// Add new hole-less polygon to our output shapes
+			ReturnPolygons.Add(Verts);
+		}
+	}
+
+	return ReturnPolygons;
+}
+
+
+void FGeomTools2D::CorrectPolygonWinding(TArray<FVector2D>& OutVertices, const TArray<FVector2D>& Vertices, const bool bNegativeWinding)
+{
+	if (Vertices.Num() >= 3)
+	{
+		// Make sure the polygon winding is correct
+		if ((!bNegativeWinding && !IsPolygonWindingCCW(Vertices)) ||
+			(bNegativeWinding && IsPolygonWindingCCW(Vertices)))
+		{
+			// Reverse vertices
+			for (int32 VertexIndex = Vertices.Num() - 1; VertexIndex >= 0; --VertexIndex)
+			{
+				new (OutVertices) FVector2D(Vertices[VertexIndex]);
+			}
+		}
+		else
+		{
+			// Copy vertices
+			for (int32 VertexIndex = 0; VertexIndex < Vertices.Num(); ++VertexIndex)
+			{
+				new (OutVertices) FVector2D(Vertices[VertexIndex]);
+			}
+		}
+	}
+}
+
+// Assumes polygons are valid, closed and the winding is correct and matches bWindingCCW
+bool FGeomTools2D::ArePolygonsValid(const TArray<TArray<FVector2D>>& Polygons)
+{
+	const int PolygonCount = Polygons.Num();
+
+	// Find every pair of shapes that overlap
+	for (int PolygonIndexA = 0; PolygonIndexA < PolygonCount; ++PolygonIndexA)
+	{
+		const TArray<FVector2D>& PolygonA = Polygons[PolygonIndexA];
+		const bool bIsWindingA_CCW = IsPolygonWindingCCW(PolygonA);
+		for (int PolygonIndexB = PolygonIndexA + 1; PolygonIndexB < PolygonCount; ++PolygonIndexB)
+		{
+			const TArray<FVector2D>& PolygonB = Polygons[PolygonIndexB];
+			const bool bIsWindingB_CCW = IsPolygonWindingCCW(PolygonB);
+
+			// Additive polygons are allowed to intersect
+			if (bIsWindingA_CCW && bIsWindingB_CCW)
+			{
+				continue;
+			}
+
+			// Test each shapeA edge against each shapeB edge for intersection
+			for (int VertexA = 0; VertexA < PolygonA.Num(); ++VertexA)
+			{
+				const FVector2D& A0 = PolygonA[VertexA];
+				const FVector2D& A1 = PolygonA[(VertexA + 1) % PolygonA.Num()];
+				const FVector2D A10 = A1 - A0;
+
+				for (int VertexB = 0; VertexB < PolygonB.Num(); ++VertexB)
+				{
+					const FVector2D B0 = PolygonB[VertexB];
+					const FVector2D B1 = PolygonB[(VertexB + 1) % PolygonB.Num()];
+					const FVector2D B10 = B1 - B0;
+
+					const float DetS = A10.X * B10.Y - A10.Y * B10.X;
+					const float DetT = B10.X * A10.Y - B10.Y * A10.X;
+					if (DetS != 0.0f && DetT != 0.0f)
+					{
+						const float S = (A0.Y * B10.X - B0.Y * B10.X - A0.X * B10.Y + B0.X * B10.Y) / DetS;
+						const float T = (B0.Y * A10.X - A0.Y * A10.X - B0.X * A10.Y + A0.X * A10.Y) / DetT;
+						if (S >= 0.0f && S <= 1.0f && T >= 0.0f && T <= 1.0f)
+						{
+							// Edges intersect
+							UE_LOG(LogGeomTools, Log, TEXT("Edges in polygon %d and %d intersect"), PolygonIndexA, PolygonIndexB);
+							return false;
+						}
+					}
+				}
+			}
+		}
+	}
+
+	// Make sure no polygons are self-intersecting
+	// Disabled for now until we decide what to do with this - contour tracing can generate invalid polys & degenerate edges
+	//for (int PolygonIndex = 0; PolygonIndex < PolygonCount; ++PolygonIndex)
+	//{
+	//	const TArray<FVector2D>& Polygon = Polygons[PolygonIndex].Vertices;
+	//	const int PointCount = Polygon.Num();
+	//	for (int PointIndexA = 0; PointIndexA < PointCount; ++PointIndexA)
+	//	{
+	//		const FVector2D& A0 = Polygon[PointIndexA];
+	//		const FVector2D& A1 = Polygon[(PointIndexA + 1) % PointCount];
+	//		const FVector2D A10 = A1 - A0;
+	//		for (int PointIndexB = 0; PointIndexB < PointCount - 3; ++PointIndexB)
+	//		{
+	//			const FVector2D& B0 = Polygon[(PointIndexA + 2 + PointIndexB) % PointCount];
+	//			const FVector2D& B1 = Polygon[(PointIndexA + 3 + PointIndexB) % PointCount];
+	//			const FVector2D B10 = B1 - B0;
+
+	//			const float DetS = A10.X * B10.Y - A10.Y * B10.X;
+	//			const float DetT = B10.X * A10.Y - B10.Y * A10.X;
+	//			if (DetS != 0.0f && DetT != 0.0f)
+	//			{
+	//				const float S = (A0.Y * B10.X - B0.Y * B10.X - A0.X * B10.Y + B0.X * B10.Y) / DetS;
+	//				const float T = (B0.Y * A10.X - A0.Y * A10.X - B0.X * A10.Y + A0.X * A10.Y) / DetT;
+	//				if (S >= 0.0f && S <= 1.0f && T >= 0.0f && T <= 1.0f)
+	//				{
+	//					// Edges intersect
+	//					UE_LOG(LogGeomTools, Log, TEXT("Polygon %d is self intersecting"), PolygonIndex);
+	//					return false;
+	//				}
+	//			}
+	//		}
+	//	}
+	//}
+
+	return true;
+}
+
+// Determines whether two edges may be merged.
+// Very similar to the one in GeomTools.cpp, but only considers positions (as opposed to other vertex attributes)
+static bool AreEdgesMergeable(const FVector2D& V0, const FVector2D& V1, const FVector2D& V2)
+{
+	const FVector2D MergedEdgeVector = V2 - V0;
+	const float MergedEdgeLengthSquared = MergedEdgeVector.SizeSquared();
+	if (MergedEdgeLengthSquared > DELTA)
+	{
+		// Find the vertex closest to A1/B0 that is on the hypothetical merged edge formed by A0-B1.
+		const float IntermediateVertexEdgeFraction = ((V2 - V0) | (V1 - V0)) / MergedEdgeLengthSquared;
+		const FVector2D InterpolatedVertex = V0 + (V2 - V0) * IntermediateVertexEdgeFraction;
+
+		// The edges are merge-able if the interpolated vertex is close enough to the intermediate vertex.
+		return InterpolatedVertex.Equals(V1, THRESH_POINTS_ARE_SAME);
+	}
+	else
+	{
+		return true;
+	}
+}
+
+// Nearly identical function in GeomTools.cpp - Identical points (exact same fp values, not nearlyEq) are not considered to be inside 
+// the triangle when ear clipping. These are generated by ReducePolygons when additive and subtractive polygons are merged.
+// Expected input - PolygonVertices in CCW order, not overlapping
+bool FGeomTools2D::TriangulatePoly(TArray<FVector2D>& OutTris, const TArray<FVector2D>& InPolyVerts, bool bKeepColinearVertices)
+{
+	// Can't work if not enough verts for 1 triangle
+	if (InPolyVerts.Num() < 3)
+	{
+		// Return true because poly is already a tri
+		return true;
+	}
+
+	// Vertices of polygon in order - make a copy we are going to modify.
+	TArray<FVector2D> PolyVerts = InPolyVerts;
+
+	// Keep iterating while there are still vertices
+	while (true)
+	{
+		if (!bKeepColinearVertices)
+		{
+			// Cull redundant vertex edges from the polygon.
+			for (int32 VertexIndex = 0; VertexIndex < PolyVerts.Num(); VertexIndex++)
+			{
+				const int32 I0 = (VertexIndex + 0) % PolyVerts.Num();
+				const int32 I1 = (VertexIndex + 1) % PolyVerts.Num();
+				const int32 I2 = (VertexIndex + 2) % PolyVerts.Num();
+				if (AreEdgesMergeable(PolyVerts[I0], PolyVerts[I1], PolyVerts[I2]))
+				{
+					PolyVerts.RemoveAt(I1);
+					VertexIndex--;
+				}
+			}
+		}
+
+		if (PolyVerts.Num() < 3)
+		{
+			break;
+		}
+		else
+		{
+			// Look for an 'ear' triangle
+			bool bFoundEar = false;
+			for (int32 EarVertexIndex = 0; EarVertexIndex < PolyVerts.Num(); EarVertexIndex++)
+			{
+				// Triangle is 'this' vert plus the one before and after it
+				const int32 AIndex = (EarVertexIndex == 0) ? PolyVerts.Num() - 1 : EarVertexIndex - 1;
+				const int32 BIndex = EarVertexIndex;
+				const int32 CIndex = (EarVertexIndex + 1) % PolyVerts.Num();
+
+				// Check that this vertex is convex (cross product must be positive)
+				const FVector2D ABEdge = PolyVerts[BIndex] - PolyVerts[AIndex];
+				const FVector2D ACEdge = PolyVerts[CIndex] - PolyVerts[AIndex];
+				if (FMath::IsNegativeFloat(ABEdge ^ ACEdge))
+				{
+					continue;
+				}
+
+				bool bFoundVertInside = false;
+				// Look through all verts before this in array to see if any are inside triangle
+				for (int32 VertexIndex = 0; VertexIndex < PolyVerts.Num(); VertexIndex++)
+				{
+					const FVector2D& CurrentVertex = PolyVerts[VertexIndex];
+					// Test indices first, then make sure we arent comparing identical points
+					// These might have been added by the convex / concave splitting
+					// If a point is not in the triangle, it may be on the new edge we're adding, which isn't allowed as 
+					// it will create a partition in the polygon
+					if (VertexIndex != AIndex && VertexIndex != BIndex && VertexIndex != CIndex &&
+						CurrentVertex != PolyVerts[AIndex] && CurrentVertex != PolyVerts[BIndex] && CurrentVertex != PolyVerts[CIndex] &&
+						(IsPointInTriangle(CurrentVertex, PolyVerts[AIndex], PolyVerts[BIndex], PolyVerts[CIndex]) || IsPointOnLineSegment(CurrentVertex, PolyVerts[CIndex], PolyVerts[AIndex])))
+					{
+						bFoundVertInside = true;
+						break;
+					}
+				}
+
+				// Triangle with no verts inside - its an 'ear'! 
+				if (!bFoundVertInside)
+				{
+					// Add to output list..
+					OutTris.Add(PolyVerts[AIndex]);
+					OutTris.Add(PolyVerts[BIndex]);
+					OutTris.Add(PolyVerts[CIndex]);
+
+					// And remove vertex from polygon
+					PolyVerts.RemoveAt(EarVertexIndex);
+
+					bFoundEar = true;
+					break;
+				}
+			}
+
+			// If we couldn't find an 'ear' it indicates something is bad with this polygon - discard triangles and return.
+			if (!bFoundEar)
+			{
+				// If we couldn't find an 'ear' it indicates something is bad with this polygon - discard triangles and return.
+				UE_LOG(LogGeomTools, Log, TEXT("Triangulation of poly failed."));
+				OutTris.Empty();
+				return false;
+			}
+		}
+	}
+
+	return true;
+}
+
+// 2D version of GeomTools RemoveRedundantTriangles
+void FGeomTools2D::RemoveRedundantTriangles(TArray<FVector2D>& OutTriangles, const TArray<FVector2D>& InTriangleVertices)
+{
+	struct FLocalTriangle
+	{
+		int VertexA, VertexB, VertexC;
+	};
+
+	TArray<FLocalTriangle> Triangles;
+	for (int32 TriangleVertexIndex = 0; TriangleVertexIndex < InTriangleVertices.Num(); TriangleVertexIndex += 3)
+	{
+		FLocalTriangle* NewTriangle = new (Triangles) FLocalTriangle();
+		NewTriangle->VertexA = TriangleVertexIndex + 0;
+		NewTriangle->VertexB = TriangleVertexIndex + 1;
+		NewTriangle->VertexC = TriangleVertexIndex + 2;
+	}
+
+	while (Triangles.Num() > 0)
+	{
+		TArray<FVector2D> PolygonVertices;
+
+		const FLocalTriangle InitialTriangle = Triangles.Pop(/*bAllowShrinking=*/ false);
+		PolygonVertices.Add(InTriangleVertices[InitialTriangle.VertexA]);
+		PolygonVertices.Add(InTriangleVertices[InitialTriangle.VertexB]);
+		PolygonVertices.Add(InTriangleVertices[InitialTriangle.VertexC]);
+
+		// Find triangles that can be merged into the polygon.
+		for (int32 CandidateTriangleIndex = 0; CandidateTriangleIndex < Triangles.Num(); ++CandidateTriangleIndex)
+		{
+			const FLocalTriangle& MergeCandidateTriangle = Triangles[CandidateTriangleIndex];
+			if (MergeTriangleIntoPolygon(PolygonVertices, InTriangleVertices[MergeCandidateTriangle.VertexA], InTriangleVertices[MergeCandidateTriangle.VertexB], InTriangleVertices[MergeCandidateTriangle.VertexC]))
+			{
+				// Remove the merged triangle from the array.
+				Triangles.RemoveAtSwap(CandidateTriangleIndex);
+
+				// Restart the search for mergeable triangles from the start of the array.
+				CandidateTriangleIndex = -1;
+			}
+		}
+
+		// Triangulate merged polygon and append to triangle array
+		TArray<FVector2D> TriangulatedPoly;
+		TriangulatePoly(/*out*/TriangulatedPoly, PolygonVertices);
+		OutTriangles.Append(TriangulatedPoly);
+	}
+}
+
+// Find convex polygons from triangle soup
+void FGeomTools2D::GenerateConvexPolygonsFromTriangles(TArray<TArray<FVector2D>>& OutPolygons, const TArray<FVector2D>& InTriangleVertices)
+{
+	struct FLocalTriangle
+	{
+		int VertexA, VertexB, VertexC;
+	};
+
+	TArray<FLocalTriangle> Triangles;
+	for (int32 TriangleVertexIndex = 0; TriangleVertexIndex < InTriangleVertices.Num(); TriangleVertexIndex += 3)
+	{
+		FLocalTriangle* NewTriangle = new (Triangles) FLocalTriangle();
+		NewTriangle->VertexA = TriangleVertexIndex + 0;
+		NewTriangle->VertexB = TriangleVertexIndex + 1;
+		NewTriangle->VertexC = TriangleVertexIndex + 2;
+	}
+
+	while (Triangles.Num() > 0)
+	{
+		TArray<FVector2D> PolygonVertices;
+
+		const FLocalTriangle InitialTriangle = Triangles.Pop(/*bAllowShrinking=*/ false);
+		PolygonVertices.Add(InTriangleVertices[InitialTriangle.VertexA]);
+		PolygonVertices.Add(InTriangleVertices[InitialTriangle.VertexB]);
+		PolygonVertices.Add(InTriangleVertices[InitialTriangle.VertexC]);
+
+		// Find triangles that can be merged into the polygon.
+		for (int32 CandidateTriangleIndex = 0; CandidateTriangleIndex < Triangles.Num(); ++CandidateTriangleIndex)
+		{
+			const FLocalTriangle& MergeCandidateTriangle = Triangles[CandidateTriangleIndex];
+			if (MergeTriangleIntoPolygon(PolygonVertices, InTriangleVertices[MergeCandidateTriangle.VertexA], InTriangleVertices[MergeCandidateTriangle.VertexB], InTriangleVertices[MergeCandidateTriangle.VertexC], true))
+			{
+				// Remove the merged triangle from the array.
+				Triangles.RemoveAtSwap(CandidateTriangleIndex);
+
+				// Restart the search for mergeable triangles from the start of the array.
+				CandidateTriangleIndex = -1;
+			}
+		}
+
+		OutPolygons.Add(PolygonVertices);
+	}
+}
+
+// Creates a convex hull that contains the given points
+void FGeomTools2D::GenerateConvexHullFromPoints(TArray<FVector2D>& OutConvexHull, TArray<FVector2D>& SourcePoints)
+{
+	int SourcePointsCount = SourcePoints.Num();
+	if (SourcePointsCount < 3)
+	{
+		return;
+	}
+
+	// Find lowest point. If multiple points have the same y, find leftmost
+	int LowestPointIndex = 0;
+	FVector2D& LowestPoint = SourcePoints[0];
+	for (int I = 1; I < SourcePointsCount; ++I)
+	{
+		if ((SourcePoints[I].Y < LowestPoint.Y) || (SourcePoints[I].Y == LowestPoint.Y && SourcePoints[I].X < LowestPoint.X))
+		{
+			LowestPointIndex = I;
+			LowestPoint = SourcePoints[I];
+		}
+	}
+
+	// Get indices sorted by angle they and the lowest point make with the x-axis (excluding the lowest point)
+	FVector2D& P = SourcePoints[LowestPointIndex];
+	TArray<int32> SortedPoints;
+	SortedPoints.Empty(SourcePointsCount);
+	for (int I = 0; I < SourcePointsCount; ++I)
+	{
+		if (I != LowestPointIndex)
+		{
+			SortedPoints.Add(I);
+		}
+	}
+	int SortedPointsCount = SortedPoints.Num();
+	for (int I = 0; I < SortedPointsCount; ++I)
+	{
+		for (int J = I + 1; J < SortedPointsCount; ++J)
+		{
+			FVector2D DI = SourcePoints[SortedPoints[I]] - P;
+			FVector2D DJ = SourcePoints[SortedPoints[J]] - P;
+			if (FMath::Atan2(DI.Y, DI.X) > FMath::Atan2(DJ.Y, DJ.X))
+			{
+				int32 Temp = SortedPoints[I];
+				SortedPoints[I] = SortedPoints[J];
+				SortedPoints[J] = Temp;
+			}
+		}
+	}
+
+	// Initialize hull points and add lowest point to end of sorted points for the algorithm
+	TArray<int> Hull;
+	Hull.Add(LowestPointIndex);
+	SortedPoints.Add(LowestPointIndex);
+	++SortedPointsCount;
+	// Traverse sorted points, removing all prior points that would make a right turn, before adding each new point (Graham scan)
+	for (int I = 0; I < SortedPointsCount; ++I)
+	{
+		int NewPointIndex = SortedPoints[I];
+		FVector2D& C = SourcePoints[NewPointIndex];
+		while (Hull.Num() > 1)
+		{
+			FVector2D& A = SourcePoints[Hull[Hull.Num() - 2]];
+			FVector2D& B = SourcePoints[Hull[Hull.Num() - 1]];
+			if ((B.X - A.X) * (C.Y - A.Y) - (B.Y - A.Y) * (C.X - A.X) < 0.0f)
+			{
+				// Remove last entry
+				Hull.RemoveAt(Hull.Num() - 1);
+			}
+			else
+			{
+				// Left turn, exit loop
+				break;
+			}
+		}
+		Hull.Add(NewPointIndex);
+	}
+
+	// We added the starting point to the end for the algorithm, so remove the duplicate
+	Hull.RemoveAt(Hull.Num() - 1);
+
+	// Result
+	int HullCount = Hull.Num();
+	OutConvexHull.Empty(HullCount);
+	for (int I = 0; I < HullCount; ++I)
+	{
+		OutConvexHull.Add(SourcePoints[Hull[I]]);
+	}
+}
+
+bool FGeomTools2D::IsPointInPolygon(const FVector2D& TestPoint, const TArray<FVector2D>& PolygonPoints)
+{
+	const int NumPoints = PolygonPoints.Num();
+	float AngleSum = 0.0f;
+	for (int PointIndex = 0; PointIndex < NumPoints; ++PointIndex)
+	{
+		const FVector2D& VecAB = PolygonPoints[PointIndex] - TestPoint;
+		const FVector2D& VecAC = PolygonPoints[(PointIndex + 1) % NumPoints] - TestPoint;
+		const float Angle = FMath::Sign(FVector2D::CrossProduct(VecAB, VecAC)) * FMath::Acos(FMath::Clamp(FVector2D::DotProduct(VecAB, VecAC) / (VecAB.Size() * VecAC.Size()), -1.0f, 1.0f));
+		AngleSum += Angle;
+	}
+	return (FMath::Abs(AngleSum) > 0.001f);
+}
diff --git a/Engine/Source/Runtime/Engine/Public/GeomTools.h b/Engine/Source/Runtime/Engine/Public/GeomTools.h
--- a/Engine/Source/Runtime/Engine/Public/GeomTools.h
+++ b/Engine/Source/Runtime/Engine/Public/GeomTools.h
@@ -160,4 +160,43 @@ namespace FGeomTools
 
 	/** Util to see if P lies within triangle created by A, B and C. */
 	ENGINE_API bool PointInTriangle(const FVector& A, const FVector& B, const FVector& C, const FVector& P, const float InsideTriangleDotProductEpsilon = 0.0f);
+
+};
+
+
+class ENGINE_API FGeomTools2D
+{
+public:
+	// Corrects the polygon winding to match bNegativeWinding
+	// Ie. If the vertex order doesn't match the winding, it is reversed
+	// Returns empty array of points if num points < 3
+	static void CorrectPolygonWinding(TArray<FVector2D>& OutVertices, const TArray<FVector2D>& Vertices, const bool bNegativeWinding);
+
+	// Returns true if the points forming a polygon have CCW winding
+	// Returns true if the polygon isn't valid
+	static bool IsPolygonWindingCCW(const TArray<FVector2D>& Points);
+	static bool IsPolygonWindingCCW(const TArray<FIntPoint>& Points);
+
+	// Checks that these polygons can be successfully triangulated	
+	static bool ArePolygonsValid(const TArray<TArray<FVector2D>>& Polygons);
+
+	// Merge additive and subtractive polygons, split them up into additive polygons
+	// Assumes all polygons and overlapping polygons are valid, and the windings match the setting on the polygon
+	static TArray<TArray<FVector2D>> ReducePolygons(const TArray<TArray<FVector2D>>& Polygons, const TArray<bool>& PolygonNegativeWinding);
+
+	// Triangulate a polygon. Check notes in implementation
+	static bool TriangulatePoly(TArray<FVector2D>& OutTris, const TArray<FVector2D>& PolygonVertices, bool bKeepColinearVertices = false);
+
+	// 2D version of RemoveRedundantTriangles from GeomTools
+	static void RemoveRedundantTriangles(TArray<FVector2D>& OutTriangles, const TArray<FVector2D>& InTriangles);
+
+	// Generate convex shapes
+	static void GenerateConvexPolygonsFromTriangles(TArray<TArray<FVector2D>>& OutPolygons, const TArray<FVector2D>& InTriangles);
+
+	// Generate convex hull from points
+	static void GenerateConvexHullFromPoints(TArray<FVector2D>& OutConvexHull, TArray<FVector2D>& Points);
+
+	// Returns true if TestPoint is inside the polygon defined by PolygonPoints
+	static bool IsPointInPolygon(const FVector2D& TestPoint, const TArray<FVector2D>& PolygonPoints);
 };
+
