Fix crash when modifying max texture size for VTs

diff --git a/Engine/Source/Runtime/Engine/Private/VT/VirtualTextureDataBuilder.cpp b/Engine/Source/Runtime/Engine/Private/VT/VirtualTextureDataBuilder.cpp
--- a/Engine/Source/Runtime/Engine/Private/VT/VirtualTextureDataBuilder.cpp
+++ b/Engine/Source/Runtime/Engine/Private/VT/VirtualTextureDataBuilder.cpp
@@ -304,6 +304,9 @@ void FVirtualTextureDataBuilder::Build(const FTextureSourceData& InSourceData, c
 		break;
 	}
 
+	BlockSizeX = FMath::Min<uint32>(BlockSizeX, SettingsPerLayer[0].MaxTextureResolution);
+	BlockSizeY = FMath::Min<uint32>(BlockSizeY, SettingsPerLayer[0].MaxTextureResolution);
+
 	// We require VT blocks (UDIM pages) to be PoT, but multi block textures may have full logical dimension that's not PoT
 	check(FMath::IsPowerOfTwo(BlockSizeX));
 	check(FMath::IsPowerOfTwo(BlockSizeY));
