Allow FMeshBatchElement::PrimitiveUniformBuffer to be assigned to the default scene proxy uniform buffer without triggering the GpuScene assert.

diff --git a/Engine/Source/Runtime/Engine/Private/SceneManagement.cpp b/Engine/Source/Runtime/Engine/Private/SceneManagement.cpp
--- a/Engine/Source/Runtime/Engine/Private/SceneManagement.cpp
+++ b/Engine/Source/Runtime/Engine/Private/SceneManagement.cpp
@@ -243,8 +243,9 @@ void FMeshElementCollector::AddMesh(int32 ViewIndex, FMeshBatch& MeshBatch)
 {
 	DEFINE_LOG_CATEGORY_STATIC(FMeshElementCollector_AddMesh, Warning, All);
 
-	//checkSlow(MeshBatch.GetNumPrimitives() > 0);
-	checkSlow(MeshBatch.VertexFactory && MeshBatch.MaterialRenderProxy);
+	checkSlow(MeshBatch.VertexFactory);
+	checkSlow(MeshBatch.VertexFactory->IsInitialized());
+	checkSlow(MeshBatch.MaterialRenderProxy);
 	checkSlow(PrimitiveSceneProxy);
 
 	PrimitiveSceneProxy->VerifyUsedMaterial(MeshBatch.MaterialRenderProxy);
@@ -273,7 +274,7 @@ void FMeshElementCollector::AddMesh(int32 ViewIndex, FMeshBatch& MeshBatch)
 	}
 
 	// If we are maintaining primitive scene data on the GPU, copy the primitive uniform buffer data to a unified array so it can be uploaded later
-	if (UseGPUScene(GMaxRHIShaderPlatform, FeatureLevel) && MeshBatch.VertexFactory->GetPrimitiveIdStreamIndex(EVertexInputStreamType::Default) >= 0)
+	if (MeshBatch.VertexFactory->IsPrimitiveDataInGPUScene(EVertexInputStreamType::Default))
 	{
 		for (int32 Index = 0; Index < MeshBatch.Elements.Num(); ++Index)
 		{
@@ -1061,28 +1062,46 @@ void FReadOnlyCVARCache::Init()
 void FMeshBatch::PreparePrimitiveUniformBuffer(const FPrimitiveSceneProxy* PrimitiveSceneProxy, ERHIFeatureLevel::Type FeatureLevel)
 {
 	const bool bVFSupportsPrimitiveIdStream = VertexFactory->GetType()->SupportsPrimitiveIdStream();
-	checkf((PrimitiveSceneProxy->DoesVFRequirePrimitiveUniformBuffer() || bVFSupportsPrimitiveIdStream), TEXT("PrimitiveSceneProxy has bVFRequiresPrimitiveUniformBuffer disabled yet tried to draw with a vertex factory (%s) that did not support PrimitiveIdStream."), VertexFactory->GetType()->GetName());
+	checkf(
+		(PrimitiveSceneProxy->DoesVFRequirePrimitiveUniformBuffer() || bVFSupportsPrimitiveIdStream),
+		TEXT("PrimitiveSceneProxy has bVFRequiresPrimitiveUniformBuffer disabled yet tried to draw with a vertex factory (%s) ")
+		TEXT("that did not support PrimitiveIdStream."),
+		VertexFactory->GetType()->GetName());
 
-	const bool bPrimitiveShaderDataComesFromSceneBuffer = VertexFactory->GetPrimitiveIdStreamIndex(EVertexInputStreamType::Default) >= 0;
+	const bool bPrimitiveDataInGPUScene = VertexFactory->IsPrimitiveDataInGPUScene(EVertexInputStreamType::Default);
+
+	FRHIUniformBuffer* SceneProxyUniformBuffer = PrimitiveSceneProxy->GetUniformBuffer();
 
 	for (int32 ElementIndex = 0; ElementIndex < Elements.Num(); ElementIndex++)
 	{
 		FMeshBatchElement& MeshElement = Elements[ElementIndex];
 
-		if (bPrimitiveShaderDataComesFromSceneBuffer)
+		if (bPrimitiveDataInGPUScene)
 		{
-			checkf(!Elements[ElementIndex].PrimitiveUniformBuffer, 
-				TEXT("FMeshBatch was assigned a PrimitiveUniformBuffer even though Vertex Factory %s fetches primitive shader data through a Scene buffer.  The assigned PrimitiveUniformBuffer cannot be respected.  Use PrimitiveUniformBufferResource instead for dynamic primitive data, or leave both null to get FPrimitiveSceneProxy->UniformBuffer."), VertexFactory->GetType()->GetName());
+			// We are using GPU scene. If the user specified the default proxy uniform buffer, we null it out. This is to
+			// allow the user to just provide the default proxy uniform buffer by default without having to check if GPU scene
+			// is enabled for the vertex factory.
+			if (MeshElement.PrimitiveUniformBuffer == SceneProxyUniformBuffer)
+			{
+				MeshElement.PrimitiveUniformBuffer = nullptr;
+			}
+			else
+			{
+				checkf(!MeshElement.PrimitiveUniformBuffer,
+					TEXT("FMeshBatch was assigned a PrimitiveUniformBuffer even though Vertex Factory %s fetches primitive shader data through ")
+					TEXT("a Scene buffer. The assigned PrimitiveUniformBuffer cannot be respected. Use PrimitiveUniformBufferResource instead ")
+					TEXT("for dynamic primitive data, or leave both null to get FPrimitiveSceneProxy->UniformBuffer."),
+					VertexFactory->GetType()->GetName());
+			}
 		}
-
-		// If we are not using GPU Scene, draws using vertex factories that do not support an explicit PrimitiveUniformBuffer on the FMeshBatch need to be setup with the FPrimitiveSceneProxy's uniform buffer
-		if (!MeshElement.PrimitiveUniformBufferResource && !UseGPUScene(GMaxRHIShaderPlatform, FeatureLevel) && bVFSupportsPrimitiveIdStream)
+		else if (!MeshElement.PrimitiveUniformBufferResource)
 		{
-			MeshElement.PrimitiveUniformBuffer = PrimitiveSceneProxy->GetUniformBuffer();
+			// The user did not specify a uniform buffer, and GPU scene is not enabled for this primitive, so use the scene proxy uniform buffer.
+			MeshElement.PrimitiveUniformBuffer = SceneProxyUniformBuffer;
 		}
 
-		checkf(bPrimitiveShaderDataComesFromSceneBuffer || Elements[ElementIndex].PrimitiveUniformBuffer || Elements[ElementIndex].PrimitiveUniformBufferResource != NULL,
-			TEXT("FMeshBatch was not properly setup.  The primitive uniform buffer must be specified."));
+		checkf(bPrimitiveDataInGPUScene || Elements[ElementIndex].PrimitiveUniformBuffer || Elements[ElementIndex].PrimitiveUniformBufferResource,
+			TEXT("FMeshBatch was not properly setup. The primitive uniform buffer must be specified."));
 	}
 }
 
diff --git a/Engine/Source/Runtime/RenderCore/Public/VertexFactory.h b/Engine/Source/Runtime/RenderCore/Public/VertexFactory.h
--- a/Engine/Source/Runtime/RenderCore/Public/VertexFactory.h
+++ b/Engine/Source/Runtime/RenderCore/Public/VertexFactory.h
@@ -13,6 +13,7 @@
 #include "RenderResource.h"
 #include "ShaderCore.h"
 #include "Shader.h"
+#include "RenderUtils.h"
 #include "Misc/EnumClassFlags.h"
 
 class FMaterial;
@@ -556,6 +557,16 @@ public:
 		return PrimitiveIdStreamIndex[static_cast<uint8>(InputStreamType)];
 	}
 
+	inline bool HasPrimitiveIdStream(EVertexInputStreamType InputStreamType) const
+	{
+		return GetPrimitiveIdStreamIndex(InputStreamType) >= 0;
+	}
+
+	inline bool IsPrimitiveDataInGPUScene(EVertexInputStreamType InputStreamType) const
+	{
+		return UseGPUScene(GMaxRHIShaderPlatform, FeatureLevel) && GetPrimitiveIdStreamIndex(InputStreamType) >= 0;
+	}
+
 protected:
 
 	/**
diff --git a/Engine/Source/Runtime/Renderer/Private/DynamicPrimitiveDrawing.inl b/Engine/Source/Runtime/Renderer/Private/DynamicPrimitiveDrawing.inl
--- a/Engine/Source/Runtime/Renderer/Private/DynamicPrimitiveDrawing.inl
+++ b/Engine/Source/Runtime/Renderer/Private/DynamicPrimitiveDrawing.inl
@@ -199,24 +199,27 @@ inline int32 FViewElementPDI::DrawMesh(const FMeshBatch& Mesh)
 			ENQUEUE_RENDER_COMMAND(FCopyDynamicPrimitiveShaderData)(
 				[NewMesh, DynamicPrimitiveShaderDataForRT, FeatureLevel](FRHICommandListImmediate& RHICmdList)
 				{
-					const bool bPrimitiveShaderDataComesFromSceneBuffer = NewMesh->VertexFactory->GetPrimitiveIdStreamIndex(EVertexInputStreamType::Default) >= 0;
+					const bool bPrimitiveDataInGPUScene = NewMesh->VertexFactory->IsPrimitiveDataInGPUScene(EVertexInputStreamType::Default);
 
 					for (int32 ElementIndex = 0; ElementIndex < NewMesh->Elements.Num(); ElementIndex++)
 					{
 						FMeshBatchElement& MeshElement = NewMesh->Elements[ElementIndex];
 
-						if (bPrimitiveShaderDataComesFromSceneBuffer)
+						if (bPrimitiveDataInGPUScene)
 						{
 							checkf(!NewMesh->Elements[ElementIndex].PrimitiveUniformBuffer,
-								TEXT("FMeshBatch was assigned a PrimitiveUniformBuffer even though Vertex Factory %s fetches primitive shader data through a Scene buffer.  The assigned PrimitiveUniformBuffer cannot be respected.  Use PrimitiveUniformBufferResource instead for dynamic primitive data, or leave both null to get FPrimitiveSceneProxy->UniformBuffer."), NewMesh->VertexFactory->GetType()->GetName());
+								TEXT("FMeshBatch was assigned a PrimitiveUniformBuffer even though Vertex Factory %s fetches primitive shader data through ")
+								TEXT("a Scene buffer. The assigned PrimitiveUniformBuffer cannot be respected. Use PrimitiveUniformBufferResource instead for ")
+								TEXT("dynamic primitive data, or leave both null to get FPrimitiveSceneProxy->UniformBuffer."),
+								NewMesh->VertexFactory->GetType()->GetName());
 						}
 
-						checkf(bPrimitiveShaderDataComesFromSceneBuffer || NewMesh->Elements[ElementIndex].PrimitiveUniformBufferResource != NULL,
+						checkf(bPrimitiveDataInGPUScene || NewMesh->Elements[ElementIndex].PrimitiveUniformBufferResource,
 							TEXT("FMeshBatch was not properly setup.  The primitive uniform buffer must be specified."));
 					}
 
 					// If we are maintaining primitive scene data on the GPU, copy the primitive uniform buffer data to a unified array so it can be uploaded later
-					if (UseGPUScene(GMaxRHIShaderPlatform, FeatureLevel) && bPrimitiveShaderDataComesFromSceneBuffer)
+					if (bPrimitiveDataInGPUScene)
 					{
 						for (int32 Index = 0; Index < NewMesh->Elements.Num(); ++Index)
 						{
diff --git a/Engine/Source/Runtime/Renderer/Private/Renderer.cpp b/Engine/Source/Runtime/Renderer/Private/Renderer.cpp
--- a/Engine/Source/Runtime/Renderer/Private/Renderer.cpp
+++ b/Engine/Source/Runtime/Renderer/Private/Renderer.cpp
@@ -105,7 +105,7 @@ void FRendererModule::DrawTileMesh(FRHICommandListImmediate& RHICmdList, FMeshPa
 
 		FSinglePrimitiveStructuredBuffer SinglePrimitiveStructuredBuffer;
 
-		if (Mesh.VertexFactory->GetPrimitiveIdStreamIndex(EVertexInputStreamType::PositionOnly) >= 0)
+		if (Mesh.VertexFactory->HasPrimitiveIdStream(EVertexInputStreamType::PositionOnly))
 		{
 			FMeshBatchElement& MeshElement = Mesh.Elements[0];
 
diff --git a/Engine/Source/Runtime/Renderer/Private/ShaderBaseClasses.cpp b/Engine/Source/Runtime/Renderer/Private/ShaderBaseClasses.cpp
--- a/Engine/Source/Runtime/Renderer/Private/ShaderBaseClasses.cpp
+++ b/Engine/Source/Runtime/Renderer/Private/ShaderBaseClasses.cpp
@@ -547,7 +547,7 @@ void FMeshMaterialShader::GetElementShaderBindings(
 
 	VertexFactoryParameters.GetElementShaderBindings(Scene, ViewIfDynamicMeshCommand, this, InputStreamType, FeatureLevel, VertexFactory, BatchElement, ShaderBindings, VertexStreams);
 		
-	if (UseGPUScene(GMaxRHIShaderPlatform, FeatureLevel) && VertexFactory->GetPrimitiveIdStreamIndex(InputStreamType) >= 0)
+	if (VertexFactory->IsPrimitiveDataInGPUScene(InputStreamType))
 	{
 		ensureMsgf(!GetUniformBufferParameter<FPrimitiveUniformShaderParameters>().IsBound(), TEXT("Shader %s attempted to bind the Primitive uniform buffer even though Vertex Factory computes a PrimitiveId per-instance.  This will break auto-instancing.  Shaders should use GetPrimitiveData(PrimitiveId).Member instead of Primitive.Member."), GetType()->GetName());
 		ensureMsgf(!BatchElement.PrimitiveUniformBuffer, TEXT("FMeshBatchElement was assigned a PrimitiveUniformBuffer even though Vertex Factory %s fetches primitive shader data through a Scene buffer.  The assigned PrimitiveUniformBuffer cannot be respected.  Use PrimitiveUniformBufferResource instead for dynamic primitive data."), GetType()->GetName());
