Make bunch size limit configurable, and check it when sending (not just receiving).

diff --git a/Engine/Source/Runtime/Engine/Private/DataChannel.cpp b/Engine/Source/Runtime/Engine/Private/DataChannel.cpp
--- a/Engine/Source/Runtime/Engine/Private/DataChannel.cpp
+++ b/Engine/Source/Runtime/Engine/Private/DataChannel.cpp
@@ -75,6 +75,21 @@ FAutoConsoleVariableRef CVarNetSkipReplicatorForDestructionInfos(
 extern TAutoConsoleVariable<int32> CVarFilterGuidRemapping;
 extern TAutoConsoleVariable<int32> CVarNetEnableDetailedScopeCounters;
 
+
+// Fairly large number, and probably a bad idea to even have a bunch this size, but want to be safe for now and not throw out legitimate data
+static int32 NetMaxConstructedPartialBunchSizeBytes = 1024 * 64;
+static FAutoConsoleVariableRef CVarNetMaxConstructedPartialBunchSizeBytes(
+	TEXT("net.MaxConstructedPartialBunchSizeBytes"),
+	NetMaxConstructedPartialBunchSizeBytes,
+	TEXT("The maximum size allowed for Partial Bunches.")
+);
+
+template<typename T>
+static const bool IsBunchTooLarge(UNetConnection* Connection, T* Bunch)
+{
+	return !Connection->InternalAck && Bunch != nullptr && Bunch->GetNumBytes() > NetMaxConstructedPartialBunchSizeBytes;
+}
+
 /*-----------------------------------------------------------------------------
 	UChannel implementation.
 -----------------------------------------------------------------------------*/
@@ -703,12 +718,9 @@ bool UChannel::ReceivedNextBunch( FInBunch & Bunch, bool & bOutSkipAck )
 			}
 		}
 
-		// Fairly large number, and probably a bad idea to even have a bunch this size, but want to be safe for now and not throw out legitimate data
-		static const int32 MAX_CONSTRUCTED_PARTIAL_SIZE_IN_BYTES = 1024 * 64;		
-
-		if ( !Connection->InternalAck && InPartialBunch != NULL && InPartialBunch->GetNumBytes() > MAX_CONSTRUCTED_PARTIAL_SIZE_IN_BYTES )
+		if (IsBunchTooLarge(Connection, InPartialBunch))
 		{
-			UE_LOG( LogNetPartialBunch, Error, TEXT( "Final partial bunch too large" ) );
+			UE_LOG(LogNetPartialBunch, Error, TEXT("Received a partial bunch exceeding max allowed size. BunchSize=%d, MaximumSize=%d"), InPartialBunch->GetNumBytes(), NetMaxConstructedPartialBunchSizeBytes);
 			Bunch.SetError();
 			return false;
 		}
@@ -911,6 +923,13 @@ FPacketIdRange UChannel::SendBunch( FOutBunch* Bunch, bool Merge )
 		return FPacketIdRange(INDEX_NONE);
 	}
 
+	if (IsBunchTooLarge(Connection, Bunch))
+	{
+		UE_LOG(LogNetPartialBunch, Error, TEXT("Attempted to send bunch exceeding max allowed size. BunchSize=%d, MaximumSize=%d"), Bunch->GetNumBytes(), NetMaxConstructedPartialBunchSizeBytes);
+		Bunch->SetError();
+		return FPacketIdRange(INDEX_NONE);
+	}
+
 	check(!Closing);
 	checkf(Connection->Channels[ChIndex]==this, TEXT("This: %s, Connection->Channels[ChIndex]: %s"), *Describe(), Connection->Channels[ChIndex] ? *Connection->Channels[ChIndex]->Describe() : TEXT("Null"));
 	check(!Bunch->IsError());
