#ue4 - Refactor code out of ACharacter::CanJumpInternal_Implementation() which deteremined what CharacterMovement state allowed jump attmpts to a new overrideable function UCharacterMovementComponent::CanAttemptJump().

diff --git a/Engine/Source/Runtime/Engine/Classes/GameFramework/CharacterMovementComponent.h b/Engine/Source/Runtime/Engine/Classes/GameFramework/CharacterMovementComponent.h
--- a/Engine/Source/Runtime/Engine/Classes/GameFramework/CharacterMovementComponent.h
+++ b/Engine/Source/Runtime/Engine/Classes/GameFramework/CharacterMovementComponent.h
@@ -1309,13 +1309,18 @@ public:
 	virtual void StartNewPhysics(float deltaTime, int32 Iterations);
 	
 	/**
-	 * Perform jump. Called by Character when a jump has been detected because Character->bPressedJump was true. Checks CanJump().
+	 * Perform jump. Called by Character when a jump has been detected because Character->bPressedJump was true. Checks Character->CanJump().
 	 * Note that you should usually trigger a jump through Character::Jump() instead.
 	 * @param	bReplayingMoves: true if this is being done as part of replaying moves on a locally controlled client after a server correction.
 	 * @return	True if the jump was triggered successfully.
 	 */
 	virtual bool DoJump(bool bReplayingMoves);
 
+	/**
+	 * Returns true if current movement state allows an attempt at jumping. Used by Character::CanJump().
+	 */
+	virtual bool CanAttemptJump() const;
+
 	/** Queue a pending launch with velocity LaunchVel. */
 	virtual void Launch(FVector const& LaunchVel);
 
diff --git a/Engine/Source/Runtime/Engine/Private/Character.cpp b/Engine/Source/Runtime/Engine/Private/Character.cpp
--- a/Engine/Source/Runtime/Engine/Private/Character.cpp
+++ b/Engine/Source/Runtime/Engine/Private/Character.cpp
@@ -263,11 +263,7 @@ bool ACharacter::CanJumpInternal_Implementation() const
 	bool bCanJump = !bIsCrouched;
 
 	// Ensure that the CharacterMovement state is valid
-	bCanJump &= CharacterMovement &&
-				CharacterMovement->IsJumpAllowed() &&
-				!CharacterMovement->bWantsToCrouch &&
-				// Can only jump from the ground, or multi-jump if already falling.
-				(CharacterMovement->IsMovingOnGround() || CharacterMovement->IsFalling());
+	bCanJump &= CharacterMovement->CanAttemptJump();
 
 	if (bCanJump)
 	{
diff --git a/Engine/Source/Runtime/Engine/Private/Components/CharacterMovementComponent.cpp b/Engine/Source/Runtime/Engine/Private/Components/CharacterMovementComponent.cpp
--- a/Engine/Source/Runtime/Engine/Private/Components/CharacterMovementComponent.cpp
+++ b/Engine/Source/Runtime/Engine/Private/Components/CharacterMovementComponent.cpp
@@ -864,6 +864,14 @@ bool UCharacterMovementComponent::DoJump(bool bReplayingMoves)
 	return false;
 }
 
+bool UCharacterMovementComponent::CanAttemptJump() const
+{
+	return IsJumpAllowed() &&
+		   !bWantsToCrouch &&
+		   (IsMovingOnGround() || IsFalling()); // Falling included for double-jump and non-zero jump hold time, but validated by character.
+}
+
+
 FVector UCharacterMovementComponent::GetImpartedMovementBaseVelocity() const
 {
 	FVector Result = FVector::ZeroVector;
