Fix accumulation of duplicated impostors instances during HLOD rebuilds.
#rnx

diff --git a/Engine/Source/Runtime/Engine/Private/LODActor.cpp b/Engine/Source/Runtime/Engine/Private/LODActor.cpp
--- a/Engine/Source/Runtime/Engine/Private/LODActor.cpp
+++ b/Engine/Source/Runtime/Engine/Private/LODActor.cpp
@@ -299,6 +299,27 @@ void ALODActor::PostLoad()
 
 	ParseOverrideDistancesCVar();
 	UpdateOverrideTransitionDistance();
+
+	bool bMarkRenderStateDirty = false;
+	for (const TPair<const UMaterialInterface*, UInstancedStaticMeshComponent*>& Component : ImpostersStaticMeshComponents)
+	{
+		for (int32 i=0; i<Component.Value->PerInstanceSMData.Num() - 1; i++)
+		{
+			for (int32 j=i + 1; j<Component.Value->PerInstanceSMData.Num(); j++)
+			{
+				if (Component.Value->PerInstanceSMData[i].Transform == Component.Value->PerInstanceSMData[j].Transform)
+				{
+					Component.Value->PerInstanceSMData.RemoveAtSwap(j--);
+					bMarkRenderStateDirty = true;
+				}
+			}
+		}
+
+		if (bMarkRenderStateDirty)
+		{
+			Component.Value->MarkRenderStateDirty();
+		}
+	}
 }
 
 void ALODActor::SetComponentsMinDrawDistance(float InMinDrawDistance, bool bInMarkRenderStateDirty)
@@ -999,6 +1020,7 @@ void ALODActor::SetupImposters(UMaterialInterface* InMaterial, UStaticMesh* InSt
 
 	UInstancedStaticMeshComponent* Component = GetOrCreateLODComponentForMaterial(InMaterial);
 	Component->SetStaticMesh(InStaticMesh);
+	Component->PerInstanceSMData.Empty();
 	
 	for (const FTransform& Transform : InTransforms)
 	{
