Fix WITH_EDITOR from previous submit.
#rnx

diff --git a/Engine/Source/Runtime/Engine/Private/LODActor.cpp b/Engine/Source/Runtime/Engine/Private/LODActor.cpp
--- a/Engine/Source/Runtime/Engine/Private/LODActor.cpp
+++ b/Engine/Source/Runtime/Engine/Private/LODActor.cpp
@@ -300,6 +300,7 @@ void ALODActor::PostLoad()
 	ParseOverrideDistancesCVar();
 	UpdateOverrideTransitionDistance();
 
+#if WITH_EDITOR
 	bool bMarkRenderStateDirty = false;
 	for (const TPair<const UMaterialInterface*, UInstancedStaticMeshComponent*>& Component : ImpostersStaticMeshComponents)
 	{
@@ -320,6 +321,7 @@ void ALODActor::PostLoad()
 			Component.Value->MarkRenderStateDirty();
 		}
 	}
+#endif
 }
 
 void ALODActor::SetComponentsMinDrawDistance(float InMinDrawDistance, bool bInMarkRenderStateDirty)
