Animation stops playing on character using Layer Node
#fix Make sure we do not delete the outer anim instance for a layer node
#misc Reinitialise to a self-layer whenever clearing a sub anim instance layer node

diff --git a/Engine/Source/Runtime/Engine/Classes/Animation/AnimNode_Layer.h b/Engine/Source/Runtime/Engine/Classes/Animation/AnimNode_Layer.h
--- a/Engine/Source/Runtime/Engine/Classes/Animation/AnimNode_Layer.h
+++ b/Engine/Source/Runtime/Engine/Classes/Animation/AnimNode_Layer.h
@@ -47,4 +47,7 @@ public:
 	{
 		return *Interface;
 	}
+
+protected:
+	void InitializeSelfLayer(const UAnimInstance* SelfAnimInstance);
 };
\ No newline at end of file
diff --git a/Engine/Source/Runtime/Engine/Private/Animation/AnimNode_Layer.cpp b/Engine/Source/Runtime/Engine/Private/Animation/AnimNode_Layer.cpp
--- a/Engine/Source/Runtime/Engine/Private/Animation/AnimNode_Layer.cpp
+++ b/Engine/Source/Runtime/Engine/Private/Animation/AnimNode_Layer.cpp
@@ -25,34 +25,45 @@ void FAnimNode_Layer::OnInitializeAnimInstance(const FAnimInstanceProxy* InProxy
 	// initialized by the owning anim instance as they may share sub-instances via grouping.
 	if(Interface.Get() == nullptr || InstanceClass.Get() == nullptr)
 	{
-		UAnimInstance* CurrentTarget = GetTargetInstance<UAnimInstance>();
+		InitializeSelfLayer(InAnimInstance);
+	}
+}
+
+void FAnimNode_Layer::InitializeSelfLayer(const UAnimInstance* SelfAnimInstance)
+{
+	UAnimInstance* CurrentTarget = GetTargetInstance<UAnimInstance>();
 
-		USkeletalMeshComponent* MeshComp = InAnimInstance->GetSkelMeshComponent();
-		check(MeshComp);
+	USkeletalMeshComponent* MeshComp = SelfAnimInstance->GetSkelMeshComponent();
+	check(MeshComp);
 
-		if(LinkedRoot)
-		{
-			DynamicUnlink(const_cast<UAnimInstance*>(InAnimInstance));
-		}
+	if (LinkedRoot)
+	{
+		DynamicUnlink(const_cast<UAnimInstance*>(SelfAnimInstance));
+	}
 
-		// Switch from dynamic external to internal, kill old instance
-		if(CurrentTarget && CurrentTarget != InAnimInstance)
-		{
-			MeshComp->SubInstances.Remove(CurrentTarget);
-			CurrentTarget->MarkPendingKill();
-			CurrentTarget = nullptr;
-		}
+	// Switch from dynamic external to internal, kill old instance
+	if (CurrentTarget && CurrentTarget != SelfAnimInstance)
+	{
+		MeshComp->SubInstances.Remove(CurrentTarget);
+		CurrentTarget->MarkPendingKill();
+		CurrentTarget = nullptr;
+	}
 
-		SetTargetInstance(const_cast<UAnimInstance*>(InAnimInstance));
+	SetTargetInstance(const_cast<UAnimInstance*>(SelfAnimInstance));
 
-		// Link before we call InitializeAnimation() so we propgate the call to sub-inputs
-		DynamicLink(const_cast<UAnimInstance*>(InAnimInstance));
+	// Link before we call InitializeAnimation() so we propgate the call to sub-inputs
+	DynamicLink(const_cast<UAnimInstance*>(SelfAnimInstance));
 
-		InitializeProperties(InAnimInstance, InAnimInstance->GetClass());
-	}
+	InitializeProperties(SelfAnimInstance, SelfAnimInstance->GetClass());
 }
 
 void FAnimNode_Layer::SetLayerOverlaySubInstance(const UAnimInstance* InOwningAnimInstance, UAnimInstance* InNewSubInstance)
 {
 	ReinitializeSubAnimInstance(InOwningAnimInstance, InNewSubInstance);
+
+	// Reseting to running as a self-layer, in case it is applicable
+	if ((Interface.Get() == nullptr || InstanceClass.Get() == nullptr) && (InNewSubInstance == nullptr))
+	{
+		InitializeSelfLayer(InOwningAnimInstance);
+	}
 }
diff --git a/Engine/Source/Runtime/Engine/Private/Animation/AnimNode_SubInstance.cpp b/Engine/Source/Runtime/Engine/Private/Animation/AnimNode_SubInstance.cpp
--- a/Engine/Source/Runtime/Engine/Private/Animation/AnimNode_SubInstance.cpp
+++ b/Engine/Source/Runtime/Engine/Private/Animation/AnimNode_SubInstance.cpp
@@ -134,7 +134,11 @@ void FAnimNode_SubInstance::ReinitializeSubAnimInstance(const UAnimInstance* InO
 			DynamicUnlink(const_cast<UAnimInstance*>(InOwningAnimInstance));
 
 			MeshComp->SubInstances.Remove(InstanceToRun);
-			InstanceToRun->MarkPendingKill();
+			// Never delete the owning animation instance
+			if (InstanceToRun != InOwningAnimInstance)
+			{
+				InstanceToRun->MarkPendingKill();
+			}
 			InstanceToRun = nullptr;
 		}
 
