Animation state machines don't update when used as inputs to Layer nodes
#fix Make sure we only update state machines once per frame, as nodes can be updated multiple times through SubAnimInstance(s)

diff --git a/Engine/Source/Runtime/Engine/Private/Animation/AnimInstanceProxy.cpp b/Engine/Source/Runtime/Engine/Private/Animation/AnimInstanceProxy.cpp
--- a/Engine/Source/Runtime/Engine/Private/Animation/AnimInstanceProxy.cpp
+++ b/Engine/Source/Runtime/Engine/Private/Animation/AnimInstanceProxy.cpp
@@ -42,7 +42,11 @@ void FAnimInstanceProxy::UpdateAnimationNode_WithRoot(float DeltaSeconds, FAnimN
 	DECLARE_SCOPE_HIERARCHICAL_COUNTER_FUNC()
 	if(InRootNode != nullptr)
 	{
-		UpdateCounter.Increment();
+		if (FrameCounterForUpdate != GFrameCounter)
+		{
+			UpdateCounter.Increment();
+		}
+
 		InRootNode->Update_AnyThread(FAnimationUpdateContext(this, DeltaSeconds));
 
 		// We've updated the graph, now update the fractured saved pose sections
