Fix spline component having different metadata instances than actors that own the metadata.  Fixed the problem once and for all by having one property for this

diff --git a/Engine/Source/Runtime/Engine/Classes/Components/SplineComponent.h b/Engine/Source/Runtime/Engine/Classes/Components/SplineComponent.h
--- a/Engine/Source/Runtime/Engine/Classes/Components/SplineComponent.h
+++ b/Engine/Source/Runtime/Engine/Classes/Components/SplineComponent.h
@@ -192,9 +192,6 @@ class ENGINE_API USplineComponent : public UPrimitiveComponent
 	UPROPERTY(EditAnywhere, Category=Points)
 	FSplineCurves SplineCurves;
 
-	UPROPERTY()
-	USplineMetadata* SplineMetadata;
-
 	/** Deprecated - please use GetSplinePointsPosition() to fetch this FInterpCurve */
 	UPROPERTY()
 	FInterpCurveVector SplineInfo_DEPRECATED;
@@ -316,10 +313,9 @@ public:
 	const FInterpCurveQuat& GetSplinePointsRotation() const { return SplineCurves.Rotation; }
 	FInterpCurveVector& GetSplinePointsScale() { return SplineCurves.Scale; }
 	const FInterpCurveVector& GetSplinePointsScale() const { return SplineCurves.Scale; }
-	USplineMetadata* GetSplinePointsMetadata() { return SplineMetadata; }
-	const USplineMetadata* GetSplinePointsMetadata() const { return SplineMetadata; }
 
-	void SetSplineMetadata(USplineMetadata* InMetadata) { SplineMetadata = InMetadata; }
+	virtual USplineMetadata* GetSplinePointsMetadata() { return nullptr; }
+	virtual const USplineMetadata* GetSplinePointsMetadata() const { return nullptr; }
 
 	void ApplyComponentInstanceData(struct FSplineInstanceData* ComponentInstanceData, const bool bPostUCS);
 
