Default prioritizing mesh LOD retention
New metric for automatic retention bias selection

diff --git a/Engine/Source/Runtime/Engine/Private/Streaming/AsyncTextureStreaming.cpp b/Engine/Source/Runtime/Engine/Private/Streaming/AsyncTextureStreaming.cpp
--- a/Engine/Source/Runtime/Engine/Private/Streaming/AsyncTextureStreaming.cpp
+++ b/Engine/Source/Runtime/Engine/Private/Streaming/AsyncTextureStreaming.cpp
@@ -335,12 +335,12 @@ void FRenderAssetStreamingMipCalcTask::UpdateBudgetedMips_Async(int64& MemoryUse
 	{
 		if (IsAborted()) break;
 
-		MemoryBudgeted += StreamingRenderAsset.UpdateRetentionPriority_Async(Settings.bPrioritizeMeshRetention);
+		MemoryBudgeted += StreamingRenderAsset.UpdateRetentionPriority_Async(Settings.bPrioritizeMeshLODRetention);
 		const int32 AssetMemUsed = StreamingRenderAsset.GetSize(StreamingRenderAsset.ResidentMips);
 		MemoryUsed += AssetMemUsed;
 
 		// TODO: Use RHI metrics
-		if (StreamingRenderAsset.RenderAssetType != FStreamingRenderAsset::AT_Texture)
+		if (!StreamingRenderAsset.IsTexture())
 		{
 			MemoryUsedByNonTextures += AssetMemUsed;
 		}
@@ -469,7 +469,10 @@ void FRenderAssetStreamingMipCalcTask::UpdateBudgetedMips_Async(int64& MemoryUse
 			for (int32 NumDroppedMips = 0; NumDroppedMips < Settings.GlobalMipBias && MemoryBudgeted > MemoryBudget && !IsAborted(); ++NumDroppedMips)
 			{
 				const int64 PreviousMemoryBudgeted = MemoryBudgeted;
-				const bool bMeshPrivilegedPhase = NumDroppedMips < Settings.MeshRetentionPrivilegeLevel;
+
+				// Heuristic: Only consider dropping max resolution for a mesh if it has reasonable impact on memory reduction.
+				// Currently, reasonable impact is defined as MemDeltaOfDroppingOneLOD >= MinTextureMemDelta in this pass.
+				int64 MinTextureMemDelta = MAX_int64;
 
 				for (int32 PriorityIndex = PrioritizedRenderAssets.Num() - 1; PriorityIndex >= 0 && MemoryBudgeted > MemoryBudget && !IsAborted(); --PriorityIndex)
 				{
@@ -485,18 +488,27 @@ void FRenderAssetStreamingMipCalcTask::UpdateBudgetedMips_Async(int64& MemoryUse
 						continue;
 					}
 
-					// Meshes don't drop max resolution in the first N privileged passes
-					if (bMeshPrivilegedPhase && StreamingRenderAsset.RenderAssetType != FStreamingRenderAsset::AT_Texture)
-					{
-						continue;
-					}
-
 					// If the texture/mesh requires a high resolution mip, consider dropping it. 
 					// When considering dropping the first mip, only textures/meshes using the first mip will drop their resolution, 
 					// But when considering dropping the second mip, textures/meshes using their first and second mips will loose it.
 					if (StreamingRenderAsset.MaxAllowedMips + StreamingRenderAsset.BudgetMipBias - NumDroppedMips <= StreamingRenderAsset.BudgetedMips)
 					{
-						MemoryBudgeted -= StreamingRenderAsset.DropMaxResolution_Async(NumDroppedMips + 1 - StreamingRenderAsset.BudgetMipBias);
+						const int32 NumMipsToDrop = NumDroppedMips + 1 - StreamingRenderAsset.BudgetMipBias;
+
+						if (Settings.bPrioritizeMeshLODRetention)
+						{
+							const bool bIsTexture = StreamingRenderAsset.IsTexture();
+							const int64 MemDeltaFromMaxResDrop = StreamingRenderAsset.GetDropMaxResMemDelta(NumMipsToDrop);
+
+							if (!MemDeltaFromMaxResDrop || (!bIsTexture && MemDeltaFromMaxResDrop < MinTextureMemDelta && MinTextureMemDelta != MAX_int64))
+							{
+								continue;
+							}
+
+							MinTextureMemDelta = bIsTexture ? FMath::Min(MinTextureMemDelta, MemDeltaFromMaxResDrop) : MinTextureMemDelta;
+						}
+
+						MemoryBudgeted -= StreamingRenderAsset.DropMaxResolution_Async(NumMipsToDrop);
 					}
 				}
 
@@ -512,11 +524,12 @@ void FRenderAssetStreamingMipCalcTask::UpdateBudgetedMips_Async(int64& MemoryUse
 		// Drop WantedMip until in budget.
 		//*************************************
 
-		int32 PassCount = 0;
 		while (MemoryBudgeted > MemoryBudget && !IsAborted())
 		{
 			const int64 PreviousMemoryBudgeted = MemoryBudgeted;
-			const bool bMeshPrivilegedPhase = PassCount < Settings.MeshRetentionPrivilegeLevel;
+
+			// Heuristic: only start considering dropping mesh LODs if it has reasonable impact on memory reduction.
+			int64 MinTextureMemDelta = MAX_int64;
 
 			// Drop from the lowest priority first (starting with last elements)
 			for (int32 PriorityIndex = PrioritizedRenderAssets.Num() - 1; PriorityIndex >= 0 && MemoryBudgeted > MemoryBudget && !IsAborted(); --PriorityIndex)
@@ -533,10 +546,15 @@ void FRenderAssetStreamingMipCalcTask::UpdateBudgetedMips_Async(int64& MemoryUse
 					continue;
 				}
 
-				// Don't try to drop mesh LODs in the first N privileged passes
-				if (bMeshPrivilegedPhase && StreamingRenderAsset.RenderAssetType != FStreamingRenderAsset::AT_Texture)
+				const bool bIsTexture = StreamingRenderAsset.IsTexture();
+				const bool bIsMesh = !bIsTexture;
+				if (Settings.bPrioritizeMeshLODRetention && bIsMesh)
 				{
-					continue;
+					const int64 PredictedMemDelta = StreamingRenderAsset.GetDropOneMipMemDelta();
+					if (PredictedMemDelta < MinTextureMemDelta && MinTextureMemDelta != MAX_int64)
+					{
+						continue;
+					}
 				}
 
 				// If this texture/mesh has already missing mips for its normal quality, don't drop more than required..
@@ -546,7 +564,12 @@ void FRenderAssetStreamingMipCalcTask::UpdateBudgetedMips_Async(int64& MemoryUse
 					continue;
 				}
 
-				MemoryBudgeted -= StreamingRenderAsset.DropOneMip_Async();
+				const int64 MemDelta = StreamingRenderAsset.DropOneMip_Async();
+				MemoryBudgeted -= MemDelta;
+				if (Settings.bPrioritizeMeshLODRetention && bIsTexture && MemDelta > 0)
+				{
+					MinTextureMemDelta = FMath::Min(MinTextureMemDelta, MemDelta);
+				}
 			}
 
 			// Break when memory does not change anymore
@@ -554,7 +577,6 @@ void FRenderAssetStreamingMipCalcTask::UpdateBudgetedMips_Async(int64& MemoryUse
 			{
 				break;
 			}
-			++PassCount;
 		}
 	}
 
@@ -876,7 +898,7 @@ void FRenderAssetStreamingMipCalcTask::UpdateStats_Async()
 
 				// Forced mips are not the same as hidden mips as they are loaded because the user wants them absolutly
 				if (StreamingRenderAsset.bForceFullyLoadHeuristic
-					|| (StreamingRenderAsset.RenderAssetType == FStreamingRenderAsset::AT_Texture
+					|| (StreamingRenderAsset.IsTexture()
 						&& StreamingRenderAsset.LODGroup == TEXTUREGROUP_HierarchicalLOD
 						&& Settings.HLODStrategy > 0))
 				{
@@ -894,7 +916,7 @@ void FRenderAssetStreamingMipCalcTask::UpdateStats_Async()
 			Stats.PendingRequests += StreamingRenderAsset.GetSize(StreamingRenderAsset.RequestedMips) - ResidentSize;
 		}
 
-		if (StreamingRenderAsset.RenderAssetType != FStreamingRenderAsset::AT_Texture)
+		if (StreamingRenderAsset.IsMesh())
 		{
 			const bool bOptLODsExist = StreamingRenderAsset.OptionalMipsState == FStreamingRenderAsset::OMS_HasOptionalMips;
 			const int32 NumLODs = bOptLODsExist ? StreamingRenderAsset.MipCount : StreamingRenderAsset.NumNonOptionalMips;
diff --git a/Engine/Source/Runtime/Engine/Private/Streaming/StreamingTexture.cpp b/Engine/Source/Runtime/Engine/Private/Streaming/StreamingTexture.cpp
--- a/Engine/Source/Runtime/Engine/Private/Streaming/StreamingTexture.cpp
+++ b/Engine/Source/Runtime/Engine/Private/Streaming/StreamingTexture.cpp
@@ -52,7 +52,7 @@ void FStreamingRenderAsset::UpdateStaticData(const FRenderAssetStreamingSettings
 		MipCount = RenderAsset->GetNumMipsForStreaming();
 		BudgetMipBias = 0;
 
-		if (RenderAssetType == AT_Texture)
+		if (IsTexture())
 		{
 			check(MipCount <= MAX_TEXTURE_MIP_COUNT);
 			const TextureGroup TextureLODGroup = static_cast<TextureGroup>(LODGroup);
@@ -94,7 +94,7 @@ void FStreamingRenderAsset::UpdateStaticData(const FRenderAssetStreamingSettings
 		NumNonOptionalMips = MipCount - RenderAsset->CalcNumOptionalMips();
 		OptionalMipsState = (NumNonOptionalMips == MipCount) ? EOptionalMipsState::OMS_NoOptionalMips : EOptionalMipsState::OMS_NotCached;
 
-		const int32 MaxNumMips = RenderAssetType == AT_Texture ? MAX_TEXTURE_MIP_COUNT : MaxNumMeshLODs;
+		const int32 MaxNumMips = IsTexture() ? MAX_TEXTURE_MIP_COUNT : MaxNumMeshLODs;
 		for (int32 MipIndex = 0; MipIndex < MaxNumMips; ++MipIndex)
 		{
 			CumulativeLODSizes[MipIndex] = RenderAsset->CalcCumulativeLODSize(FMath::Min(MipIndex + 1, MipCount));
@@ -266,7 +266,7 @@ float FStreamingRenderAsset::GetExtraBoost(TextureGroup	LODGroup, const FRenderA
 
 int32 FStreamingRenderAsset::GetWantedMipsFromSize(float Size, float MaxScreenSizeOverAllViews) const
 {
-	if (RenderAssetType == AT_Texture)
+	if (IsTexture())
 	{
 		float WantedMipsFloat = 1.0f + FMath::Log2(FMath::Max(1.f, Size));
 		int32 WantedMipsInt = FMath::CeilToInt(WantedMipsFloat);
@@ -332,7 +332,7 @@ void FStreamingRenderAsset::SetPerfectWantedMips_Async(
  * Once the wanted mips are computed, the async task will check if everything fits in the budget.
  *  This only consider the highest mip that will be requested eventually, so that slip requests are stables.
  */
-int64 FStreamingRenderAsset::UpdateRetentionPriority_Async(bool bPrioritizeMeshes)
+int64 FStreamingRenderAsset::UpdateRetentionPriority_Async(bool bPrioritizeMesh)
 {
 	// Reserve the budget for the max mip that will be loaded eventually (ignore the effect of split requests)
 	BudgetedMips = GetPerfectWantedMips();
@@ -344,13 +344,12 @@ int64 FStreamingRenderAsset::UpdateRetentionPriority_Async(bool bPrioritizeMeshe
 		const bool bShouldKeep = bIsTerrainTexture || bForceFullyLoadHeuristic || (bLooksLowRes && !bIsHuge);
 		const bool bIsSmall   = GetSize(BudgetedMips) <= 200 * 1024; 
 		const bool bIsVisible = VisibleWantedMips >= HiddenWantedMips; // Whether the first mip dropped would be a visible mip or not.
-		const bool bVIPAsset = bPrioritizeMeshes && RenderAssetType != AT_Texture;
 
 		// Here we try to have a minimal amount of priority flags for last render time to be meaningless.
 		// We mostly want thing not seen from a long time to go first to avoid repeating load / unload patterns.
 
-		if (bShouldKeep)						RetentionPriority += 4096; // Keep forced fully load as much as possible.
-		if (bVIPAsset)							RetentionPriority += 2048; // Prioritize mesh LOD retention is requested.
+		if (bPrioritizeMesh && IsMesh())		RetentionPriority += 4096; // Only consider meshes after textures are processed for faster metric calculation.
+		if (bShouldKeep)						RetentionPriority += 2048; // Keep forced fully load as much as possible.
 		if (bIsVisible)							RetentionPriority += 1024; // Keep visible things as much as possible.
 		if (!bIsHuge)							RetentionPriority += 512; // Drop high resolution which usually target ultra close range quality.
 		if (bIsCharacterTexture || bIsSmall)	RetentionPriority += 256; // Try to keep character of small texture as they don't pay off.
@@ -364,12 +363,17 @@ int64 FStreamingRenderAsset::UpdateRetentionPriority_Async(bool bPrioritizeMeshe
 	}
 }
 
+int32 FStreamingRenderAsset::ClampMaxResChange_Internal(int32 NumMipDropRequested) const
+{
+	// Don't drop bellow min allowed mips. Also ensure that MinAllowedMips < MaxAllowedMips in order allow the BudgetMipBias to reset.
+	return FMath::Min(MaxAllowedMips - MinAllowedMips - 1, NumMipDropRequested);
+}
+
 int64 FStreamingRenderAsset::DropMaxResolution_Async(int32 NumDroppedMips)
 {
 	if (RenderAsset)
 	{
-		// Don't drop bellow min allowed mips. Also ensure that MinAllowedMips < MaxAllowedMips in order allow the BudgetMipBias to reset.
-		NumDroppedMips = FMath::Min<int32>(MaxAllowedMips - MinAllowedMips - 1, NumDroppedMips);
+		NumDroppedMips = ClampMaxResChange_Internal(NumDroppedMips);
 
 		if (NumDroppedMips > 0)
 		{
@@ -422,6 +426,22 @@ int64 FStreamingRenderAsset::KeepOneMip_Async()
 	}
 }
 
+int64 FStreamingRenderAsset::GetDropMaxResMemDelta(int32 NumDroppedMips) const
+{
+	if (RenderAsset)
+	{
+		NumDroppedMips = ClampMaxResChange_Internal(NumDroppedMips);
+		return GetSize(MaxAllowedMips) - GetSize(MaxAllowedMips - NumDroppedMips);
+	}
+
+	return 0;
+}
+
+int64 FStreamingRenderAsset::GetDropOneMipMemDelta() const
+{
+	return GetSize(BudgetedMips + 1) - GetSize(BudgetedMips);
+}
+
 bool FStreamingRenderAsset::UpdateLoadOrderPriority_Async(int32 MinMipForSplitRequest)
 {
 	LoadOrderPriority = 0;
diff --git a/Engine/Source/Runtime/Engine/Private/Streaming/StreamingTexture.h b/Engine/Source/Runtime/Engine/Private/Streaming/StreamingTexture.h
--- a/Engine/Source/Runtime/Engine/Private/Streaming/StreamingTexture.h
+++ b/Engine/Source/Runtime/Engine/Private/Streaming/StreamingTexture.h
@@ -63,7 +63,7 @@ struct FStreamingRenderAsset
 	int32 GetSize( int32 InMipCount ) const
 	{
 		check(InMipCount > 0);
-		check(InMipCount <= RenderAssetType == AT_Texture ? MAX_TEXTURE_MIP_COUNT : MaxNumMeshLODs);
+		check(InMipCount <= (IsTexture() ? MAX_TEXTURE_MIP_COUNT : MaxNumMeshLODs));
 		return CumulativeLODSizes[InMipCount - 1];
 	}
 
@@ -102,7 +102,7 @@ struct FStreamingRenderAsset
 		const FRenderAssetStreamingSettings& Settings);
 
 	/** Init BudgetedMip and update RetentionPriority. Returns the size that would be taken if all budgeted mips where loaded. */
-	int64 UpdateRetentionPriority_Async(bool bPrioritizeMeshes);
+	int64 UpdateRetentionPriority_Async(bool bPrioritizeMesh);
 
 	/** Reduce the maximum allowed resolution by 1 mip. Return the size freed by doing so. */
 	int64 DropMaxResolution_Async(int32 NumDroppedMips);
@@ -113,14 +113,20 @@ struct FStreamingRenderAsset
 	/** Increase BudgetedMip by 1, up to resident mips, and return the size taken. */
 	int64 KeepOneMip_Async();
 
+	/** Return the memory delta in bytes caused by max resolution change. Actual memory reduction is smaller or equal. **/
+	int64 GetDropMaxResMemDelta(int32 NumDroppedMips) const;
+
+	/** Return the memory delta in bytes if a mip is successfully dropped. */
+	int64 GetDropOneMipMemDelta() const;
+
 	float GetMaxAllowedSize(float MaxScreenSizeOverAllViews) const
 	{
-		return RenderAssetType == AT_Texture ? (float)(0x1 << (MaxAllowedMips - 1)) : MaxScreenSizeOverAllViews;
+		return IsTexture() ? (float)(0x1 << (MaxAllowedMips - 1)) : MaxScreenSizeOverAllViews;
 	}
 
 	float GetNormalizedScreenSize(int32 NumMips) const
 	{
-		check(RenderAssetType == AT_StaticMesh || RenderAssetType == AT_SkeletalMesh);
+		check(IsMesh());
 		check(NumMips > 0 && NumMips <= MipCount);
 		return LODScreenSizes[NumMips - 1];
 	}
@@ -128,7 +134,7 @@ struct FStreamingRenderAsset
 	float GetLODScreenSize(int32 NumMips, float MaxScreenSizeOverAllViews) const
 	{
 		check(NumMips > 0 && NumMips <= MipCount);
-		return RenderAssetType == AT_Texture ?
+		return IsTexture() ?
 			static_cast<float>(1 << (NumMips - 1)) :
 			GetNormalizedScreenSize(NumMips) * MaxScreenSizeOverAllViews;
 	}
@@ -148,6 +154,17 @@ struct FStreamingRenderAsset
 	// This allows streaming to happen in parallel with the async update task
 	void StreamWantedMipsUsingCachedData(FRenderAssetStreamingManager& Manager);
 
+	bool IsTexture() const
+	{
+		return RenderAssetType == AT_Texture;
+	}
+
+	bool IsMesh() const
+	{
+		// FRenderAssetStreamingManager only handles textures and meshes currently
+		return RenderAssetType != AT_Texture;
+	}
+
 	FORCEINLINE int32 GetPerfectWantedMips() const { return FMath::Max<int32>(VisibleWantedMips,  HiddenWantedMips); }
 
 	// Whether this texture/mesh can be affected by Global Bias and Budget Bias per texture/mesh.
@@ -155,7 +172,7 @@ struct FStreamingRenderAsset
 	FORCEINLINE bool IsMaxResolutionAffectedByGlobalBias() const 
 	{
 		// In editor, forced stream in should never have reduced mips as they can be edited.
-		return (RenderAssetType != AT_Texture || LODGroup != TEXTUREGROUP_HierarchicalLOD)
+		return (IsMesh() || LODGroup != TEXTUREGROUP_HierarchicalLOD)
 			&& !bIsTerrainTexture
 			&& !bIgnoreStreamingMipBias
 			&& !(GIsEditor && bForceFullyLoadHeuristic); 
@@ -292,4 +309,6 @@ struct FStreamingRenderAsset
 
 private:
 	FORCEINLINE_DEBUGGABLE void StreamWantedMips_Internal(FRenderAssetStreamingManager& Manager, bool bUseCachedData);
+
+	FORCEINLINE_DEBUGGABLE int32 ClampMaxResChange_Internal(int32 NumMipDropRequested) const;
 };
diff --git a/Engine/Source/Runtime/Engine/Private/Streaming/TextureStreamingHelpers.cpp b/Engine/Source/Runtime/Engine/Private/Streaming/TextureStreamingHelpers.cpp
--- a/Engine/Source/Runtime/Engine/Private/Streaming/TextureStreamingHelpers.cpp
+++ b/Engine/Source/Runtime/Engine/Private/Streaming/TextureStreamingHelpers.cpp
@@ -252,18 +252,12 @@ ENGINE_API TAutoConsoleVariable<int32> CVarFramesForFullUpdate(
 	5,
 	TEXT("Texture streaming is time sliced per frame. This values gives the number of frames to visit all textures."));
 
-TAutoConsoleVariable<int32> CVarPrioritizeMeshRetention(
-	TEXT("r.Streaming.PrioritizeMeshRetention"),
-	0,
+static TAutoConsoleVariable<int32> CVarPrioritizeMeshLODRetention(
+	TEXT("r.Streaming.PrioritizeMeshLODRetention"),
+	1,
 	TEXT("Whether to prioritize retaining mesh LODs"),
 	ECVF_Default);
 
-TAutoConsoleVariable<int32> CVarMeshRetentionPrivilegeLevel(
-	TEXT("r.Streaming.MeshRetentionPriviledgeLevel"),
-	0,
-	TEXT("The streamer won't consider evicting mesh LODs until each streaming texture has this many mips evicted."),
-	ECVF_Default);
-
 static TAutoConsoleVariable<int32> CVarStreamingStressTest(
 	TEXT("r.Streaming.StressTest"),
 	0,
@@ -303,8 +297,7 @@ void FRenderAssetStreamingSettings::Update()
 	bUseMaterialData = bUseNewMetrics && CVarStreamingUseMaterialData.GetValueOnAnyThread() != 0;
 	HiddenPrimitiveScale = bUseNewMetrics ? CVarStreamingHiddenPrimitiveScale.GetValueOnAnyThread() : 1.f;
 	bMipCalculationEnablePerLevelList = CVarStreamingMipCalculationEnablePerLevelList.GetValueOnAnyThread() != 0;
-	bPrioritizeMeshRetention = CVarPrioritizeMeshRetention.GetValueOnAnyThread() != 0;
-	MeshRetentionPrivilegeLevel = CVarMeshRetentionPrivilegeLevel.GetValueOnAnyThread();
+	bPrioritizeMeshLODRetention = CVarPrioritizeMeshLODRetention.GetValueOnAnyThread() != 0;
 
 	MaterialQualityLevel = (int32)GetCachedScalabilityCVars().MaterialQualityLevel;
 
diff --git a/Engine/Source/Runtime/Engine/Private/Streaming/TextureStreamingHelpers.h b/Engine/Source/Runtime/Engine/Private/Streaming/TextureStreamingHelpers.h
--- a/Engine/Source/Runtime/Engine/Private/Streaming/TextureStreamingHelpers.h
+++ b/Engine/Source/Runtime/Engine/Private/Streaming/TextureStreamingHelpers.h
@@ -96,8 +96,7 @@ struct FRenderAssetStreamingSettings
 	int32 MaterialQualityLevel;
 	int32 FramesForFullUpdate;
 	bool bMipCalculationEnablePerLevelList;
-	bool bPrioritizeMeshRetention;
-	int32 MeshRetentionPrivilegeLevel;
+	bool bPrioritizeMeshLODRetention;
 
 	bool bStressTest;
 	static int32 ExtraIOLatency;
