Refactor FTickFunction to have a pointer to internal data that only gets allocated when the tick function has been registered. This saves 40 bytes per unregistered tick function at the cost of 8 additional bytes and a dereference to interact with the internal data for registered tick functions.
Reorganize USceneComponent members to get maximum packing efficiency in light of the change in size of FTickFunction

diff --git a/Engine/Source/Runtime/Engine/Classes/Components/SceneComponent.h b/Engine/Source/Runtime/Engine/Classes/Components/SceneComponent.h
--- a/Engine/Source/Runtime/Engine/Classes/Components/SceneComponent.h
+++ b/Engine/Source/Runtime/Engine/Classes/Components/SceneComponent.h
@@ -285,7 +285,20 @@ protected:
 	class UBillboardComponent* SpriteComponent;
 #endif
 
+public:
+	/** Delegate called when this component is moved */
+	FTransformUpdated TransformUpdated;
+
+	/** Returns the current scoped movement update, or NULL if there is none. @see FScopedMovementUpdate */
+	class FScopedMovementUpdate* GetCurrentScopedMovement() const;
+
 private:
+	/** Stack of current movement scopes. */
+	TArray<class FScopedMovementUpdate*> ScopedMovementStack;
+
+	void BeginScopedMovementUpdate(class FScopedMovementUpdate& ScopedUpdate);
+	void EndScopedMovementUpdate(class FScopedMovementUpdate& ScopedUpdate);
+
 	/** Cache that avoids Quat<->Rotator conversions if possible. Only to be used with GetComponentTransform().GetRotation(). */
 	FRotationConversionCache WorldRotationCache;
 
@@ -302,19 +315,7 @@ public:
 	/** Get the RelativeRotationCache.  */
 	FORCEINLINE const FRotationConversionCache& GetRelativeRotationCache() const { return RelativeRotationCache; }
 
-	/** Delegate called when this component is moved */
-	FTransformUpdated TransformUpdated;
-
-	/** Returns the current scoped movement update, or NULL if there is none. @see FScopedMovementUpdate */
-	class FScopedMovementUpdate* GetCurrentScopedMovement() const;
-
 private:
-	/** Stack of current movement scopes. */
-	TArray<class FScopedMovementUpdate*> ScopedMovementStack;
-
-	void BeginScopedMovementUpdate(class FScopedMovementUpdate& ScopedUpdate);
-	void EndScopedMovementUpdate(class FScopedMovementUpdate& ScopedUpdate);
-
 	UFUNCTION()
 	void OnRep_Transform();
 
diff --git a/Engine/Source/Runtime/Engine/Classes/Engine/EngineBaseTypes.h b/Engine/Source/Runtime/Engine/Classes/Engine/EngineBaseTypes.h
--- a/Engine/Source/Runtime/Engine/Classes/Engine/EngineBaseTypes.h
+++ b/Engine/Source/Runtime/Engine/Classes/Engine/EngineBaseTypes.h
@@ -190,13 +190,6 @@ public:
 	UPROPERTY(EditDefaultsOnly, Category="Tick", AdvancedDisplay)
 	TEnumAsByte<enum ETickingGroup> EndTickGroup;
 
-protected:
-	/** Internal data that indicates the tick group we actually started in (it may have been delayed due to prerequisites) **/
-	TEnumAsByte<enum ETickingGroup> ActualStartTickGroup;
-
-	/** Internal data that indicates the tick group we actually started in (it may have been delayed due to prerequisites) **/
-	TEnumAsByte<enum ETickingGroup> ActualEndTickGroup;
-
 public:
 	/** Bool indicating that this function should execute even if the game is paused. Pause ticks are very limited in capabilities. **/
 	UPROPERTY(EditDefaultsOnly, Category="Tick", AdvancedDisplay)
@@ -221,11 +214,6 @@ public:
 	uint8 bRunOnAnyThread:1;
 
 private:
-	/** If true, means that this tick function is in the master array of tick functions */
-	bool bRegistered;
-
-	/** Cache whether this function was rescheduled as an interval function during StartParallel */
-	bool bWasInterval;
 
 	enum class ETickState : uint8
 	{
@@ -241,42 +229,62 @@ private:
 	 **/
 	ETickState TickState;
 
-	/** Internal data to track if we have started visiting this tick function yet this frame **/
-	int32 TickVisitedGFrameCounter;
-
-	/** Internal data to track if we have finshed visiting this tick function yet this frame **/
-	int32 TickQueuedGFrameCounter;
+public:
+	/** The frequency in seconds at which this tick function will be executed.  If less than or equal to 0 then it will tick every frame */
+	UPROPERTY(EditDefaultsOnly, Category="Tick", meta=(DisplayName="Tick Interval (secs)"))
+	float TickInterval;
 
-	/** Pointer to the task, only used during setup. This is often stale. **/
-	void *TaskPointer;
+private:
 
 	/** Prerequisites for this tick function **/
 	TArray<struct FTickPrerequisite> Prerequisites;
 
-	/** The next function in the cooling down list for ticks with an interval*/
-	FTickFunction* Next;
+	/** Internal Data structure that contains members only required for a registered tick function **/
+	struct FInternalData
+	{
+		FInternalData();
 
-	/** 
-	  * If TickFrequency is greater than 0 and tick state is CoolingDown, this is the time, 
-	  * relative to the element ahead of it in the cooling down list, remaining until the next time this function will tick 
-	  */
-	float RelativeTickCooldown;
+		/** Cache whether this function was rescheduled as an interval function during StartParallel */
+		bool bWasInterval;
 
-	/** 
-	* The last world game time at which we were ticked. Game time used is dependent on bTickEvenWhenPaused
-	* Valid only if we've been ticked at least once since having a tick interval; otherwise set to -1.f
-	*/
-	float LastTickGameTimeSeconds;
+		/** Internal data that indicates the tick group we actually started in (it may have been delayed due to prerequisites) **/
+		TEnumAsByte<enum ETickingGroup> ActualStartTickGroup;
 
-public:
+		/** Internal data that indicates the tick group we actually started in (it may have been delayed due to prerequisites) **/
+		TEnumAsByte<enum ETickingGroup> ActualEndTickGroup;
+		
+		/** Internal data to track if we have started visiting this tick function yet this frame **/
+		int32 TickVisitedGFrameCounter;
 
-	/** The frequency in seconds at which this tick function will be executed.  If less than or equal to 0 then it will tick every frame */
-	UPROPERTY(EditDefaultsOnly, Category="Tick", meta=(DisplayName="Tick Interval (secs)"))
-	float TickInterval;
+		/** Internal data to track if we have finshed visiting this tick function yet this frame **/
+		int32 TickQueuedGFrameCounter;
+
+		/** Pointer to the task, only used during setup. This is often stale. **/
+		void* TaskPointer;
 
-	/** Back pointer to the FTickTaskLevel containing this tick function if it is registered **/
-	class FTickTaskLevel*						TickTaskLevel;
+		/** The next function in the cooling down list for ticks with an interval*/
+		FTickFunction* Next;
 
+		/** 
+		 * If TickFrequency is greater than 0 and tick state is CoolingDown, this is the time, 
+		 * relative to the element ahead of it in the cooling down list, remaining until the next time this function will tick 
+		 **/
+		float RelativeTickCooldown;
+
+		/** 
+		 * The last world game time at which we were ticked. Game time used is dependent on bTickEvenWhenPaused
+		 * Valid only if we've been ticked at least once since having a tick interval; otherwise set to -1.f
+		 **/
+		float LastTickGameTimeSeconds;
+
+		/** Back pointer to the FTickTaskLevel containing this tick function if it is registered **/
+		class FTickTaskLevel*						TickTaskLevel;
+	};
+
+	/** Lazily allocated struct that contains the necessary data for a tick function that is registered. **/
+	TUniquePtr<FInternalData> InternalData;
+
+public:
 	/** Default constructor, intitalizes to reasonable defaults **/
 	FTickFunction();
 	/** Destructor, unregisters the tick function **/
@@ -290,14 +298,14 @@ public:
 	/** Removes the tick function from the master list of tick functions. **/
 	void UnRegisterTickFunction();
 	/** See if the tick function is currently registered */
-	bool IsTickFunctionRegistered() const { return bRegistered; }
+	bool IsTickFunctionRegistered() const { return (InternalData != nullptr); }
 
 	/** Enables or disables this tick function. **/
 	void SetTickFunctionEnable(bool bInEnabled);
 	/** Returns whether the tick function is currently enabled */
 	bool IsTickFunctionEnabled() const { return TickState != ETickState::Disabled; }
 	/** Returns whether it is valid to access this tick function's completion handle */
-	bool IsCompletionHandleValid() const { return TaskPointer != nullptr; }
+	bool IsCompletionHandleValid() const { return (InternalData && InternalData->TaskPointer); }
 	/**
 	* Gets the current completion handle of this tick function, so it can be delayed until a later point when some additional
 	* tasks have been completed.  Only valid after TG_PreAsyncWork has started and then only until the TickFunction finishes
@@ -311,7 +319,7 @@ public:
 	**/
 	TEnumAsByte<enum ETickingGroup> GetActualTickGroup() const
 	{
-		return ActualStartTickGroup;
+		return (InternalData ? InternalData->ActualStartTickGroup : TickGroup);
 	}
 
 	/** 
@@ -320,7 +328,7 @@ public:
 	**/
 	TEnumAsByte<enum ETickingGroup> GetActualEndTickGroup() const
 	{
-		return ActualEndTickGroup;
+		return (InternalData ? InternalData->ActualEndTickGroup : EndTickGroup);
 	}
 
 
@@ -478,7 +486,7 @@ struct FActorComponentTickFunction : public FTickFunction
 	ENGINE_API virtual FName DiagnosticContext(bool bDetailed) override;
 
 	/**
-	 * Conditionally calls ExecuteTickFunc if bRegistered == true and a bunch of other criteria are met
+	 * Conditionally calls ExecuteTickFunc if registered and a bunch of other criteria are met
 	 * @param Target - the actor component we are ticking
 	 * @param bTickInEditor - whether the target wants to tick in the editor
 	 * @param DeltaTime - The time since the last tick.
diff --git a/Engine/Source/Runtime/Engine/Private/TickTaskManager.cpp b/Engine/Source/Runtime/Engine/Private/TickTaskManager.cpp
--- a/Engine/Source/Runtime/Engine/Private/TickTaskManager.cpp
+++ b/Engine/Source/Runtime/Engine/Private/TickTaskManager.cpp
@@ -285,7 +285,10 @@ public:
 			Target->ExecuteTick(Target->CalculateDeltaTime(Context), Context.TickType, CurrentThread, MyCompletionGraphEvent);
 			LIGHTWEIGHT_TIME_GUARD_END(FTickFunctionTask, Target->DiagnosticMessage());
 		}
-		Target->TaskPointer = nullptr;  // This is stale and a good time to clear it for safety
+		if (Target->InternalData)
+		{
+			Target->InternalData->TaskPointer = nullptr;  // This is stale and a good time to clear it for safety
+		}
 	}
 };
 
@@ -422,11 +425,12 @@ public:
 	 */
 	FORCEINLINE void StartTickTask(const FGraphEventArray* Prerequisites, FTickFunction* TickFunction, const FTickContext& TickContext)
 	{
-		checkSlow(TickFunction->ActualStartTickGroup >=0 && TickFunction->ActualStartTickGroup < TG_MAX);
+		checkSlow(TickFunction->InternalData);
+		checkSlow(TickFunction->InternalData->ActualStartTickGroup >=0 && TickFunction->InternalData->ActualStartTickGroup < TG_MAX);
 
 		FTickContext UseContext = TickContext;
 
-		bool bIsOriginalTickGroup = (TickFunction->ActualStartTickGroup == TickFunction->TickGroup);
+		bool bIsOriginalTickGroup = (TickFunction->InternalData->ActualStartTickGroup == TickFunction->TickGroup);
 
 		if (TickFunction->bRunOnAnyThread && bAllowConcurrentTicks && bIsOriginalTickGroup)
 		{
@@ -444,7 +448,7 @@ public:
 			UseContext.Thread = ENamedThreads::SetTaskPriority(ENamedThreads::GameThread, TickFunction->bHighPriority ? ENamedThreads::HighTaskPriority : ENamedThreads::NormalTaskPriority);
 		}
 
-		TickFunction->TaskPointer = TGraphTask<FTickFunctionTask>::CreateTask(Prerequisites, TickContext.Thread).ConstructAndHold(TickFunction, &UseContext, bLogTicks, bLogTicksShowPrerequistes);
+		TickFunction->InternalData->TaskPointer = TGraphTask<FTickFunctionTask>::CreateTask(Prerequisites, TickContext.Thread).ConstructAndHold(TickFunction, &UseContext, bLogTicks, bLogTicksShowPrerequistes);
 	}
 
 	/** Add a completion handle to a tick group **/
@@ -497,10 +501,11 @@ public:
 	 */
 	FORCEINLINE void QueueTickTask(const FGraphEventArray* Prerequisites, FTickFunction* TickFunction, const FTickContext& TickContext)
 	{
+		checkSlow(TickFunction->InternalData);
 		checkSlow(TickContext.Thread == ENamedThreads::GameThread);
 		StartTickTask(Prerequisites, TickFunction, TickContext);
-		TGraphTask<FTickFunctionTask>* Task = (TGraphTask<FTickFunctionTask>*)TickFunction->TaskPointer;
-		AddTickTaskCompletion(TickFunction->ActualStartTickGroup, TickFunction->ActualEndTickGroup, Task, TickFunction->bHighPriority);
+		TGraphTask<FTickFunctionTask>* Task = (TGraphTask<FTickFunctionTask>*)TickFunction->InternalData->TaskPointer;
+		AddTickTaskCompletion(TickFunction->InternalData->ActualStartTickGroup, TickFunction->InternalData->ActualEndTickGroup, Task, TickFunction->bHighPriority);
 	}
 
 	/**
@@ -512,10 +517,11 @@ public:
 	 */
 	FORCEINLINE void QueueTickTaskParallel(const FGraphEventArray* Prerequisites, FTickFunction* TickFunction, const FTickContext& TickContext)
 	{
+		checkSlow(TickFunction->InternalData);
 		checkSlow(TickContext.Thread == ENamedThreads::GameThread);
 		StartTickTask(Prerequisites, TickFunction, TickContext);
-		TGraphTask<FTickFunctionTask>* Task = (TGraphTask<FTickFunctionTask>*)TickFunction->TaskPointer;
-		AddTickTaskCompletionParallel(TickFunction->ActualStartTickGroup, TickFunction->ActualEndTickGroup, Task, TickFunction->bHighPriority);
+		TGraphTask<FTickFunctionTask>* Task = (TGraphTask<FTickFunctionTask>*)TickFunction->InternalData->TaskPointer;
+		AddTickTaskCompletionParallel(TickFunction->InternalData->ActualStartTickGroup, TickFunction->InternalData->ActualEndTickGroup, Task, TickFunction->bHighPriority);
 	}
 
 	/**
@@ -709,17 +715,18 @@ public:
 	{
 		for (TSet<FTickFunction*>::TIterator It(AllEnabledTickFunctions); It; ++It)
 		{
-			(*It)->bRegistered = false;
+			(*It)->InternalData.Reset();
 		}
 		for (TSet<FTickFunction*>::TIterator It(AllDisabledTickFunctions); It; ++It)
 		{
-			(*It)->bRegistered = false;
+			(*It)->InternalData.Reset();
 		}
 		FTickFunction* CoolingDownNode = AllCoolingDownTickFunctions.Head;
 		while (CoolingDownNode)
 		{
-			CoolingDownNode->bRegistered = false;
-			CoolingDownNode = CoolingDownNode->Next;
+			FTickFunction* NextNode = CoolingDownNode->InternalData->Next;
+			CoolingDownNode->InternalData.Reset();
+			CoolingDownNode = NextNode;
 		}
 	}
 
@@ -746,15 +753,15 @@ public:
 			FTickFunction* TickFunction = AllCoolingDownTickFunctions.Head;
 			while (TickFunction)
 			{
-				if (CumulativeCooldown + TickFunction->RelativeTickCooldown >= Context.DeltaSeconds)
+				if (CumulativeCooldown + TickFunction->InternalData->RelativeTickCooldown >= Context.DeltaSeconds)
 				{
-					TickFunction->RelativeTickCooldown -= (Context.DeltaSeconds - CumulativeCooldown);
+					TickFunction->InternalData->RelativeTickCooldown -= (Context.DeltaSeconds - CumulativeCooldown);
 					break;
 				}
-				CumulativeCooldown += TickFunction->RelativeTickCooldown;
+				CumulativeCooldown += TickFunction->InternalData->RelativeTickCooldown;
 
 				TickFunction->TickState = FTickFunction::ETickState::Enabled;
-				TickFunction = TickFunction->Next;
+				TickFunction = TickFunction->InternalData->Next;
 				++CooldownTicksEnabled;
 			}
 		}
@@ -791,21 +798,21 @@ public:
 			FTickFunction* TickFunction = AllCoolingDownTickFunctions.Head;
 			while (TickFunction)
 			{
-				if (CumulativeCooldown + TickFunction->RelativeTickCooldown >= Context.DeltaSeconds)
+				if (CumulativeCooldown + TickFunction->InternalData->RelativeTickCooldown >= Context.DeltaSeconds)
 				{
-					TickFunction->RelativeTickCooldown -= (Context.DeltaSeconds - CumulativeCooldown);
+					TickFunction->InternalData->RelativeTickCooldown -= (Context.DeltaSeconds - CumulativeCooldown);
 					break;
 				}
-				CumulativeCooldown += TickFunction->RelativeTickCooldown;
+				CumulativeCooldown += TickFunction->InternalData->RelativeTickCooldown;
 
 				TickFunction->TickState = FTickFunction::ETickState::Enabled;
-				TickFunction->bWasInterval = true;
+				TickFunction->InternalData->bWasInterval = true;
 				AllTickFunctions.Add(TickFunction);
 
 				TickFunctionsToReschedule.Add(FTickScheduleDetails(TickFunction, TickFunction->TickInterval - (Context.DeltaSeconds - CumulativeCooldown))); // Give credit for any overrun
 
-				AllCoolingDownTickFunctions.Head = TickFunction->Next;
-				TickFunction = TickFunction->Next;
+				AllCoolingDownTickFunctions.Head = TickFunction->InternalData->Next;
+				TickFunction = TickFunction->InternalData->Next;
 			}
 		}
 	}
@@ -852,7 +859,7 @@ public:
 			while (ComparisonTickFunction && RescheduleIndex < TickFunctionsToReschedule.Num())
 			{
 				const float CooldownTime = TickFunctionsToReschedule[RescheduleIndex].Cooldown;
-				if ((CumulativeCooldown + ComparisonTickFunction->RelativeTickCooldown) > CooldownTime)
+				if ((CumulativeCooldown + ComparisonTickFunction->InternalData->RelativeTickCooldown) > CooldownTime)
 				{
 					FTickFunction* TickFunction = TickFunctionsToReschedule[RescheduleIndex].TickFunction;
 					if (TickFunction->TickState != FTickFunction::ETickState::Disabled)
@@ -862,29 +869,29 @@ public:
 							verify(AllEnabledTickFunctions.Remove(TickFunction) == 1);
 						}
 						TickFunction->TickState = FTickFunction::ETickState::CoolingDown;
-						TickFunction->RelativeTickCooldown = CooldownTime - CumulativeCooldown;
+						TickFunction->InternalData->RelativeTickCooldown = CooldownTime - CumulativeCooldown;
 
 						if (PrevComparisonTickFunction)
 						{
-							PrevComparisonTickFunction->Next = TickFunction;
+							PrevComparisonTickFunction->InternalData->Next = TickFunction;
 						}
 						else
 						{
 							check(ComparisonTickFunction == AllCoolingDownTickFunctions.Head);
 							AllCoolingDownTickFunctions.Head = TickFunction;
 						}
-						TickFunction->Next = ComparisonTickFunction;
+						TickFunction->InternalData->Next = ComparisonTickFunction;
 						PrevComparisonTickFunction = TickFunction;
-						ComparisonTickFunction->RelativeTickCooldown -= TickFunction->RelativeTickCooldown;
-						CumulativeCooldown += TickFunction->RelativeTickCooldown;
+						ComparisonTickFunction->InternalData->RelativeTickCooldown -= TickFunction->InternalData->RelativeTickCooldown;
+						CumulativeCooldown += TickFunction->InternalData->RelativeTickCooldown;
 					}
 					++RescheduleIndex;
 				}
 				else
 				{
-					CumulativeCooldown += ComparisonTickFunction->RelativeTickCooldown;
+					CumulativeCooldown += ComparisonTickFunction->InternalData->RelativeTickCooldown;
 					PrevComparisonTickFunction = ComparisonTickFunction;
-					ComparisonTickFunction = ComparisonTickFunction->Next;
+					ComparisonTickFunction = ComparisonTickFunction->InternalData->Next;
 				}
 			}
 			for ( ; RescheduleIndex < TickFunctionsToReschedule.Num(); ++RescheduleIndex)
@@ -900,12 +907,12 @@ public:
 					const float CooldownTime = TickFunctionsToReschedule[RescheduleIndex].Cooldown;
 
 					TickFunction->TickState = FTickFunction::ETickState::CoolingDown;
-					TickFunction->RelativeTickCooldown = CooldownTime - CumulativeCooldown;
+					TickFunction->InternalData->RelativeTickCooldown = CooldownTime - CumulativeCooldown;
 
-					TickFunction->Next = nullptr;
+					TickFunction->InternalData->Next = nullptr;
 					if (PrevComparisonTickFunction)
 					{
-						PrevComparisonTickFunction->Next = TickFunction;
+						PrevComparisonTickFunction->InternalData->Next = TickFunction;
 					}
 					else
 					{
@@ -914,7 +921,7 @@ public:
 					}
 					PrevComparisonTickFunction = TickFunction;
 
-					CumulativeCooldown += TickFunction->RelativeTickCooldown;
+					CumulativeCooldown += TickFunction->InternalData->RelativeTickCooldown;
 				}
 			}
 			TickFunctionsToReschedule.Reset();
@@ -942,10 +949,10 @@ public:
 		{
 			if (TickFunction->TickState == FTickFunction::ETickState::Enabled)
 			{
-				CumulativeCooldown += TickFunction->RelativeTickCooldown;
+				CumulativeCooldown += TickFunction->InternalData->RelativeTickCooldown;
 				TickFunction->QueueTickFunction(TTS, Context);
 				TickFunctionsToReschedule.Add(FTickScheduleDetails(TickFunction, TickFunction->TickInterval - (Context.DeltaSeconds - CumulativeCooldown))); // Give credit for any overrun
-				AllCoolingDownTickFunctions.Head = TickFunction->Next;
+				AllCoolingDownTickFunctions.Head = TickFunction->InternalData->Next;
 			}
 			else
 			{
@@ -1016,14 +1023,14 @@ public:
 		FTickFunction* TickFunction = AllCoolingDownTickFunctions.Head;
 		while (TickFunction)
 		{
-			CumulativeCooldown += TickFunction->RelativeTickCooldown;
+			CumulativeCooldown += TickFunction->InternalData->RelativeTickCooldown;
 			if (TickFunction->bTickEvenWhenPaused)
 			{
-				TickFunction->TaskPointer = nullptr; // this is stale, clear it out now
+				TickFunction->InternalData->TaskPointer = nullptr; // this is stale, clear it out now
 				if (CumulativeCooldown < InContext.DeltaSeconds)
 				{
-					TickFunction->TickVisitedGFrameCounter = GFrameCounter;
-					TickFunction->TickQueuedGFrameCounter = GFrameCounter;
+					TickFunction->InternalData->TickVisitedGFrameCounter = GFrameCounter;
+					TickFunction->InternalData->TickQueuedGFrameCounter = GFrameCounter;
 					TickFunction->ExecuteTick(TickFunction->CalculateDeltaTime(InContext), InContext.TickType, ENamedThreads::GameThread, FGraphEventRef());
 
 					TickFunctionsToReschedule.Add(FTickScheduleDetails(TickFunction, TickFunction->TickInterval - (InContext.DeltaSeconds - CumulativeCooldown))); // Give credit for any overrun
@@ -1034,34 +1041,34 @@ public:
 				}
 				if (PrevTickFunction)
 				{
-					PrevTickFunction->Next = TickFunction->Next;
+					PrevTickFunction->InternalData->Next = TickFunction->InternalData->Next;
 				}
 				else
 				{
 					check(TickFunction == AllCoolingDownTickFunctions.Head);
-					AllCoolingDownTickFunctions.Head = TickFunction->Next;
+					AllCoolingDownTickFunctions.Head = TickFunction->InternalData->Next;
 				}
-				if (TickFunction->Next)
+				if (TickFunction->InternalData->Next)
 				{
-					TickFunction->Next->RelativeTickCooldown += TickFunction->RelativeTickCooldown;
-					CumulativeCooldown -= TickFunction->RelativeTickCooldown; // Since the next object in the list will have this cooldown included take it back out of the cumulative
+					TickFunction->InternalData->Next->InternalData->RelativeTickCooldown += TickFunction->InternalData->RelativeTickCooldown;
+					CumulativeCooldown -= TickFunction->InternalData->RelativeTickCooldown; // Since the next object in the list will have this cooldown included take it back out of the cumulative
 				}
 			}
 			else
 			{
 				PrevTickFunction = TickFunction;
 			}
-			TickFunction = TickFunction->Next;
+			TickFunction = TickFunction->InternalData->Next;
 		}
 
 		for (TSet<FTickFunction*>::TIterator It(AllEnabledTickFunctions); It; ++It)
 		{
 			TickFunction = *It;
-			TickFunction->TaskPointer = nullptr; // this is stale, clear it out now
+			TickFunction->InternalData->TaskPointer = nullptr; // this is stale, clear it out now
 			if (TickFunction->bTickEvenWhenPaused && TickFunction->TickState == FTickFunction::ETickState::Enabled)
 			{
-				TickFunction->TickVisitedGFrameCounter = GFrameCounter;
-				TickFunction->TickQueuedGFrameCounter = GFrameCounter;
+				TickFunction->InternalData->TickVisitedGFrameCounter = GFrameCounter;
+				TickFunction->InternalData->TickQueuedGFrameCounter = GFrameCounter;
 				TickFunction->ExecuteTick(TickFunction->CalculateDeltaTime(InContext), InContext.TickType, ENamedThreads::GameThread, FGraphEventRef());
 
 				if (TickFunction->TickInterval > 0.f)
@@ -1130,7 +1137,7 @@ public:
 		Ar.Logf(TEXT("%s, %s, ActualStartTickGroup: %s, Prerequesities: %d"),
 			*Function->DiagnosticMessage(),
 			Function->IsTickFunctionEnabled() ? (RemainingCooldown > 0.f ? *FString::Printf(TEXT("Cooling Down for %.4g seconds"),RemainingCooldown) : TEXT("Enabled")) : TEXT("Disabled"),
-			*TickGroupEnum->GetNameStringByValue(Function->ActualStartTickGroup),
+			*TickGroupEnum->GetNameStringByValue(Function->GetActualTickGroup()),
 			Function->Prerequisites.Num());
 
 		// List all prerequisities
@@ -1162,9 +1169,9 @@ public:
 			FTickFunction* TickFunction = AllCoolingDownTickFunctions.Head;
 			while (TickFunction)
 			{
-				CumulativeCooldown += TickFunction->RelativeTickCooldown;
+				CumulativeCooldown += TickFunction->InternalData->RelativeTickCooldown;
 				DumpTickFunction(Ar, TickFunction, TickGroupEnum, CumulativeCooldown);
-				TickFunction = TickFunction->Next;
+				TickFunction = TickFunction->InternalData->Next;
 				++EnabledCount;
 			}
 		}
@@ -1173,7 +1180,7 @@ public:
 			FTickFunction* TickFunction = AllCoolingDownTickFunctions.Head;
 			while (TickFunction)
 			{
-				TickFunction = TickFunction->Next;
+				TickFunction = TickFunction->InternalData->Next;
 				++EnabledCount;
 			}
 		}
@@ -1212,7 +1219,7 @@ public:
 		while (TickFunction)
 		{
 			AddTickFunctionToMap(ClassNameToCountMap, TickFunction, bDetailed);
-			TickFunction = TickFunction->Next;
+			TickFunction = TickFunction->InternalData->Next;
 			++EnabledCount;
 		}
 	}
@@ -1238,19 +1245,19 @@ public:
 							bFound = true;
 							if (PrevComparisionFunction)
 							{
-								PrevComparisionFunction->Next = TickFunction->Next;
+								PrevComparisionFunction->InternalData->Next = TickFunction->InternalData->Next;
 							}
 							else
 							{
 								check(TickFunction == AllCoolingDownTickFunctions.Head);
-								AllCoolingDownTickFunctions.Head = TickFunction->Next;
+								AllCoolingDownTickFunctions.Head = TickFunction->InternalData->Next;
 							}
-							TickFunction->Next = nullptr;
+							TickFunction->InternalData->Next = nullptr;
 						}
 						else
 						{
 							PrevComparisionFunction = ComparisonFunction;
-							ComparisonFunction = ComparisonFunction->Next;
+							ComparisonFunction = ComparisonFunction->InternalData->Next;
 						}
 					}
 					check(bFound); // otherwise you changed TickState while the tick function was registered. Call SetTickFunctionEnable instead.
@@ -1284,23 +1291,23 @@ public:
 					bFound = true;
 					if (PrevComparisonFunction)
 					{
-						PrevComparisonFunction->Next = TickFunction->Next;
+						PrevComparisonFunction->InternalData->Next = TickFunction->InternalData->Next;
 					}
 					else
 					{
 						check(TickFunction == AllCoolingDownTickFunctions.Head);
-						AllCoolingDownTickFunctions.Head = TickFunction->Next;
+						AllCoolingDownTickFunctions.Head = TickFunction->InternalData->Next;
 					}
-					if (TickFunction->Next)
+					if (TickFunction->InternalData->Next)
 					{
-						TickFunction->Next->RelativeTickCooldown += TickFunction->RelativeTickCooldown;
-						TickFunction->Next = nullptr;
+						TickFunction->InternalData->Next->InternalData->RelativeTickCooldown += TickFunction->InternalData->RelativeTickCooldown;
+						TickFunction->InternalData->Next = nullptr;
 					}
 				}
 				else
 				{
 					PrevComparisonFunction = ComparisonFunction;
-					ComparisonFunction = ComparisonFunction->Next;
+					ComparisonFunction = ComparisonFunction->InternalData->Next;
 				}
 			}
 			check(bFound); // otherwise you changed TickState while the tick function was registered. Call SetTickFunctionEnable instead.
@@ -1330,7 +1337,7 @@ private:
 				{
 					return true;
 				}
-				Node = Node->Next;
+				Node = Node->InternalData->Next;
 			}
 			return false;
 		}
@@ -1577,12 +1584,13 @@ public:
 		check(TickFunction->TickGroup >= 0 && TickFunction->TickGroup < TG_NewlySpawned); // You may not schedule a tick in the newly spawned group...they can only end up there if they are spawned late in a frame.
 		FTickTaskLevel* Level = TickTaskLevelForLevel(InLevel);
 		Level->AddTickFunction(TickFunction);
-		TickFunction->TickTaskLevel = Level;
+		TickFunction->InternalData->TickTaskLevel = Level;
 	}
 	/** Remove the tick function from the master list **/
 	void RemoveTickFunction(FTickFunction* TickFunction)
 	{
-		FTickTaskLevel* Level = TickFunction->TickTaskLevel;
+		check(TickFunction->InternalData);
+		FTickTaskLevel* Level = TickFunction->InternalData->TickTaskLevel;
 		check(Level);
 		Level->RemoveTickFunction(TickFunction);
 	}
@@ -1727,22 +1735,27 @@ private:
 FTickFunction::FTickFunction()
 	: TickGroup(TG_PrePhysics)
 	, EndTickGroup(TG_PrePhysics)
-	, ActualStartTickGroup(TG_PrePhysics)
-	, ActualEndTickGroup(TG_PrePhysics)
 	, bTickEvenWhenPaused(false)
 	, bCanEverTick(false)
 	, bStartWithTickEnabled(false)
 	, bAllowTickOnDedicatedServer(true)
 	, bHighPriority(false)
 	, bRunOnAnyThread(false)
-	, bRegistered(false)
-	, bWasInterval(false)
 	, TickState(ETickState::Enabled)
+	, TickInterval(0.f)
+{
+}
+
+FTickFunction::FInternalData::FInternalData()
+	: bWasInterval(false)
+	, ActualStartTickGroup(TG_PrePhysics)
+	, ActualEndTickGroup(TG_PrePhysics)
 	, TickVisitedGFrameCounter(0)
 	, TickQueuedGFrameCounter(0)
+	, TaskPointer(nullptr)
+	, Next(nullptr)
 	, RelativeTickCooldown(0.f)
 	, LastTickGameTimeSeconds(-1.f)
-	, TickInterval(0.f)
 	, TickTaskLevel(nullptr)
 {
 }
@@ -1760,14 +1773,14 @@ FTickFunction::~FTickFunction()
 **/
 void FTickFunction::RegisterTickFunction(ULevel* Level)
 {
-	if (!bRegistered)
+	if (InternalData == nullptr)
 	{
 		// Only allow registration of tick if we are are allowed on dedicated server, or we are not a dedicated server
 		const UWorld* World = Level ? Level->GetWorld() : nullptr;
 		if(bAllowTickOnDedicatedServer || !(World && World->IsNetMode(NM_DedicatedServer)))
 		{
+			InternalData.Reset(new FInternalData());
 			FTickTaskManager::Get().AddTickFunction(Level, this);
-			bRegistered = true;
 		}
 	}
 	else
@@ -1779,35 +1792,36 @@ void FTickFunction::RegisterTickFunction(ULevel* Level)
 /** Removes the tick function from the master list of tick functions. **/
 void FTickFunction::UnRegisterTickFunction()
 {
-	if (bRegistered)
+	if (InternalData)
 	{
 		FTickTaskManager::Get().RemoveTickFunction(this);
-		bRegistered = false;
+		InternalData.Reset();
 	}
 }
 
 /** Enables or disables this tick function. **/
 void FTickFunction::SetTickFunctionEnable(bool bInEnabled)
 {
-	if (bRegistered)
+	if (InternalData)
 	{
 		if (bInEnabled == (TickState == ETickState::Disabled))
 		{
+			FTickTaskLevel* TickTaskLevel = InternalData->TickTaskLevel;
 			check(TickTaskLevel);
 			TickTaskLevel->RemoveTickFunction(this);
 			TickState = (bInEnabled ? ETickState::Enabled : ETickState::Disabled);
 			TickTaskLevel->AddTickFunction(this);
 		}
+
+		if (TickState == ETickState::Disabled)
+		{
+			InternalData->LastTickGameTimeSeconds = -1.f;
+		}
 	}
 	else
 	{
 		TickState = (bInEnabled ? ETickState::Enabled : ETickState::Disabled);
 	}
-
-	if (TickState == ETickState::Disabled)
-	{
-		LastTickGameTimeSeconds = -1.f;
-	}
 }
 
 void FTickFunction::AddPrerequisite(UObject* TargetObject, struct FTickFunction& TargetTickFunction)
@@ -1855,8 +1869,8 @@ void FTickFunction::ShowPrerequistes(int32 Indent)
 
 FGraphEventRef FTickFunction::GetCompletionHandle() const
 {
-	check(TaskPointer);
-	TGraphTask<FTickFunctionTask>* Task = (TGraphTask<FTickFunctionTask>*)TaskPointer;
+	check(InternalData->TaskPointer);
+	TGraphTask<FTickFunctionTask>* Task = (TGraphTask<FTickFunctionTask>*)InternalData->TaskPointer;
 	return Task->GetCompletionEvent();
 }
 
@@ -1864,11 +1878,11 @@ FGraphEventRef FTickFunction::GetCompletionHandle() const
 void FTickFunction::QueueTickFunction(FTickTaskSequencer& TTS, const struct FTickContext& TickContext)
 {
 	checkSlow(TickContext.Thread == ENamedThreads::GameThread); // we assume same thread here
-	check(bRegistered);
+	check(InternalData);
 
-	if (TickVisitedGFrameCounter != GFrameCounter)
+	if (InternalData->TickVisitedGFrameCounter != GFrameCounter)
 	{
-		TickVisitedGFrameCounter = GFrameCounter;
+		InternalData->TickVisitedGFrameCounter = GFrameCounter;
 		if (TickState != FTickFunction::ETickState::Disabled)
 		{
 			ETickingGroup MaxPrerequisiteTickGroup =  ETickingGroup(0);
@@ -1882,22 +1896,22 @@ void FTickFunction::QueueTickFunction(FTickTaskSequencer& TTS, const struct FTic
 					// stale prereq, delete it
 					Prerequisites.RemoveAtSwap(PrereqIndex--);
 				}
-				else if (Prereq->bRegistered)
+				else if (Prereq->InternalData)
 				{
 					// recursive call to make sure my prerequisite is set up so I can use its completion handle
 					Prereq->QueueTickFunction(TTS, TickContext);
-					if (Prereq->TickQueuedGFrameCounter != GFrameCounter)
+					if (Prereq->InternalData->TickQueuedGFrameCounter != GFrameCounter)
 					{
 						// this must be up the call stack, therefore this is a cycle
 						UE_LOG(LogTick, Warning, TEXT("While processing prerequisites for %s, could use %s because it would form a cycle."),*DiagnosticMessage(), *Prereq->DiagnosticMessage());
 					}
-					else if (!Prereq->TaskPointer)
+					else if (!Prereq->InternalData->TaskPointer)
 					{
 						//ok UE_LOG(LogTick, Warning, TEXT("While processing prerequisites for %s, could use %s because it is disabled."),*DiagnosticMessage(), *Prereq->DiagnosticMessage());
 					}
 					else
 					{
-						MaxPrerequisiteTickGroup =  FMath::Max<ETickingGroup>(MaxPrerequisiteTickGroup, Prereq->ActualStartTickGroup);
+						MaxPrerequisiteTickGroup =  FMath::Max<ETickingGroup>(MaxPrerequisiteTickGroup, Prereq->InternalData->ActualStartTickGroup);
 						TaskPrerequisites.Add(Prereq->GetCompletionHandle());
 					}
 				}
@@ -1913,17 +1927,17 @@ void FTickFunction::QueueTickFunction(FTickTaskSequencer& TTS, const struct FTic
 					MyActualTickGroup = ETickingGroup(MyActualTickGroup + 1);
 				}
 			}
-			ActualStartTickGroup = MyActualTickGroup;
-			ActualEndTickGroup = MyActualTickGroup;
-			if (EndTickGroup > ActualStartTickGroup)
+			InternalData->ActualStartTickGroup = MyActualTickGroup;
+			InternalData->ActualEndTickGroup = MyActualTickGroup;
+			if (EndTickGroup > MyActualTickGroup)
 			{
 				check(EndTickGroup <= TG_NewlySpawned);
-				ETickingGroup TestTickGroup = ETickingGroup(ActualEndTickGroup + 1);
+				ETickingGroup TestTickGroup = ETickingGroup(MyActualTickGroup + 1);
 				while (TestTickGroup <= EndTickGroup)
-			{
+				{
 					if (CanDemoteIntoTickGroup(TestTickGroup))
 					{
-						ActualEndTickGroup = TestTickGroup;
+						InternalData->ActualEndTickGroup = TestTickGroup;
 					}
 					TestTickGroup = ETickingGroup(TestTickGroup + 1);
 				}
@@ -1934,25 +1948,24 @@ void FTickFunction::QueueTickFunction(FTickTaskSequencer& TTS, const struct FTic
 				TTS.QueueTickTask(&TaskPrerequisites, this, TickContext);
 			}
 		}
-		TickQueuedGFrameCounter = GFrameCounter;
+		InternalData->TickQueuedGFrameCounter = GFrameCounter;
 	}
 }
 
 void FTickFunction::QueueTickFunctionParallel(const struct FTickContext& TickContext, TArray<FTickFunction*, TInlineAllocator<8> >& StackForCycleDetection)
 {
-
 	bool bProcessTick;
 
-	int32 OldValue = *(volatile int32*)&TickVisitedGFrameCounter;
+	int32 OldValue = *(volatile int32*)&InternalData->TickVisitedGFrameCounter;
 	if (OldValue != GFrameCounter)
 	{
-		OldValue = FPlatformAtomics::InterlockedCompareExchange(&TickVisitedGFrameCounter , GFrameCounter, OldValue);
+		OldValue = FPlatformAtomics::InterlockedCompareExchange(&InternalData->TickVisitedGFrameCounter , GFrameCounter, OldValue);
 	}
 	bProcessTick = OldValue != GFrameCounter;
 
 	if (bProcessTick)
 	{
-		check(bRegistered);
+		check(InternalData);
 		if (TickState != FTickFunction::ETickState::Disabled)
 		{
 			ETickingGroup MaxPrerequisiteTickGroup = ETickingGroup(0);
@@ -1976,17 +1989,17 @@ void FTickFunction::QueueTickFunctionParallel(const struct FTickContext& TickCon
                     {
                         UE_LOG(LogTick, Warning, TEXT("While processing prerequisites for %s, could use %s because it would form a cycle."), *DiagnosticMessage(), *Prereq->DiagnosticMessage());
                     }
-					else if (Prereq->bRegistered)
+					else if (Prereq->InternalData)
 					{
 						// recursive call to make sure my prerequisite is set up so I can use its completion handle
 						Prereq->QueueTickFunctionParallel(TickContext, StackForCycleDetection);
-						if (!Prereq->TaskPointer)
+						if (!Prereq->InternalData->TaskPointer)
 						{
 							//ok UE_LOG(LogTick, Warning, TEXT("While processing prerequisites for %s, could use %s because it is disabled."),*DiagnosticMessage(), *Prereq->DiagnosticMessage());
 						}
 						else
 						{
-							MaxPrerequisiteTickGroup = FMath::Max<ETickingGroup>(MaxPrerequisiteTickGroup, Prereq->ActualStartTickGroup);
+							MaxPrerequisiteTickGroup = FMath::Max<ETickingGroup>(MaxPrerequisiteTickGroup, Prereq->InternalData->ActualStartTickGroup);
 							TaskPrerequisites.Add(Prereq->GetCompletionHandle());
 						}
 					}
@@ -2004,17 +2017,17 @@ void FTickFunction::QueueTickFunctionParallel(const struct FTickContext& TickCon
 					MyActualTickGroup = ETickingGroup(MyActualTickGroup + 1);
 				}
 			}
-			ActualStartTickGroup = MyActualTickGroup;
-			ActualEndTickGroup = MyActualTickGroup;
-			if (EndTickGroup > ActualStartTickGroup)
+			InternalData->ActualStartTickGroup = MyActualTickGroup;
+			InternalData->ActualEndTickGroup = MyActualTickGroup;
+			if (EndTickGroup > MyActualTickGroup)
 			{
 				check(EndTickGroup <= TG_NewlySpawned);
-				ETickingGroup TestTickGroup = ETickingGroup(ActualEndTickGroup + 1);
+				ETickingGroup TestTickGroup = ETickingGroup(MyActualTickGroup + 1);
 				while (TestTickGroup <= EndTickGroup)
 				{
 					if (CanDemoteIntoTickGroup(TestTickGroup))
 					{
-						ActualEndTickGroup = TestTickGroup;
+						InternalData->ActualEndTickGroup = TestTickGroup;
 					}
 					TestTickGroup = ETickingGroup(TestTickGroup + 1);
 				}
@@ -2023,25 +2036,25 @@ void FTickFunction::QueueTickFunctionParallel(const struct FTickContext& TickCon
 			if (TickState == FTickFunction::ETickState::Enabled)
 			{
 				FTickTaskSequencer::Get().QueueTickTaskParallel(&TaskPrerequisites, this, TickContext);
-				if (!bWasInterval && TickInterval > 0.f)
+				if (!InternalData->bWasInterval && TickInterval > 0.f)
 				{
-					TickTaskLevel->RescheduleForIntervalParallel(this);
+					InternalData->TickTaskLevel->RescheduleForIntervalParallel(this);
 				}
 			}
 		}
-		bWasInterval = false;
+		InternalData->bWasInterval = false;
 		
 		TSAN_BEFORE(&TickQueuedGFrameCounter);
 		FPlatformMisc::MemoryBarrier();
 		
 		// MSVC enforces acq/rel semantics on volatile values, but clang cannot (supports more backend architectures).
 		// consequently on ARM64 you would end up racing
-		FPlatformAtomics::InterlockedExchange(&TickQueuedGFrameCounter, GFrameCounter);
+		FPlatformAtomics::InterlockedExchange(&InternalData->TickQueuedGFrameCounter, GFrameCounter);
 	}
 	else
 	{
 		// if we are not going to process it, we need to at least wait until the other thread finishes it
-		volatile int32* TickQueuedGFrameCounterPtr = &TickQueuedGFrameCounter;
+		volatile int32* TickQueuedGFrameCounterPtr = &InternalData->TickQueuedGFrameCounter;
 		if (*TickQueuedGFrameCounterPtr != GFrameCounter)
 		{
 			QUICK_SCOPE_CYCLE_COUNTER(STAT_FTickFunction_QueueTickFunctionParallel_Spin);
@@ -2059,20 +2072,27 @@ float FTickFunction::CalculateDeltaTime(const FTickContext& TickContext)
 
 	if (TickInterval == 0.f)
 	{
-		// No tick interval. Return the world delta seconds, and make sure to mark that
-		// we're not tracking last-tick-time for this object.
-		LastTickGameTimeSeconds = -1.f;
+		if (InternalData)
+		{
+			// No tick interval. Return the world delta seconds, and make sure to mark that
+			// we're not tracking last-tick-time for this object.
+			InternalData->LastTickGameTimeSeconds = -1.f;
+		}
 	}
-	else
+	else if (InternalData)
 	{
 		// We've got a tick interval. Mark last-tick-time. If we already had last-tick-time, return
 		// the time since then; otherwise, return the world delta seconds.
 		const float CurrentWorldTime = (bTickEvenWhenPaused ? TickContext.World->GetUnpausedTimeSeconds() : TickContext.World->GetTimeSeconds());
-		if (LastTickGameTimeSeconds >= 0.f)
+		if (InternalData->LastTickGameTimeSeconds >= 0.f)
 		{
-			DeltaTimeForFunction = CurrentWorldTime - LastTickGameTimeSeconds;
+			DeltaTimeForFunction = CurrentWorldTime - InternalData->LastTickGameTimeSeconds;
 		}
-		LastTickGameTimeSeconds = CurrentWorldTime;
+		InternalData->LastTickGameTimeSeconds = CurrentWorldTime;
+	}
+	else
+	{
+		UE_LOG(LogTick, Log, TEXT("Unregistered tick function calculating delta time cannot know what its last tick time was. Using tick context delta seconds."));
 	}
 
 	return DeltaTimeForFunction;
