Improved error detection and handling for invalid model component index buffer.

(ushell-p4-cherrypick of 7438602 by graham.wihlidal)

diff --git a/Engine/Source/Runtime/Engine/Private/Components/ModelComponent.cpp b/Engine/Source/Runtime/Engine/Private/Components/ModelComponent.cpp
--- a/Engine/Source/Runtime/Engine/Private/Components/ModelComponent.cpp
+++ b/Engine/Source/Runtime/Engine/Private/Components/ModelComponent.cpp
@@ -663,12 +663,11 @@ bool UModelComponent::GetPhysicsTriMeshData(struct FTriMeshCollisionData* Collis
 	{
 		FModelElement& Element = Elements[ElementIndex];
 		FRawIndexBuffer16or32* IndexBuffer = Element.IndexBuffer;
-		int32 IndexBufferSize = 0;
+		int32 NumIndices = IndexBuffer ? IndexBuffer->Indices.Num() : 0;
 
-		if (IndexBuffer == nullptr || IndexBuffer->Indices.Num() == 0)
+		if (NumIndices == 0)
 		{
-			UE_LOG(LogPhysics, Warning, TEXT("Found bad index buffer when cooking UModelComponent physics data! Component: %s, Buffer: %x, Buffer Size: %d, Element: %d"), *GetPathName(), IndexBuffer, IndexBufferSize, ElementIndex);
-			verify(IndexBufferSize >= 0);
+			UE_LOG(LogPhysics, Error, TEXT("Found bad index buffer when cooking UModelComponent physics data! Component: %s, Buffer: %x, Index Count: %d, Element: %d"), *GetPathName(), IndexBuffer, NumIndices, ElementIndex);
 			continue;
 		}
 
@@ -676,17 +675,15 @@ bool UModelComponent::GetPhysicsTriMeshData(struct FTriMeshCollisionData* Collis
 
 		if (NumVertices < (int32)Element.MaxVertexIndex)
 		{
-			UE_LOG(LogPhysics, Warning, TEXT("Found bad vertex buffer when cooking UModelComponent physics data! Component: %s, Element: %d. Verts Exected: %d, Actual: %d"), 
+			UE_LOG(LogPhysics, Error, TEXT("Found bad vertex buffer when cooking UModelComponent physics data! Component: %s, Element: %d. Expected Vertex Count: %d, Actual Vertex Count: %d"),
 				*GetPathName(), ElementIndex, Element.MaxVertexIndex, NumVertices);
 			continue;
 		}
 
-		IndexBufferSize = IndexBuffer->Indices.Num();
-
-		if (IndexBufferSize < (int32)(Element.FirstIndex + Element.NumTriangles * 3 - 1))
+		if (NumIndices < (int32)(Element.FirstIndex + Element.NumTriangles * 3 - 1))
 		{
-			UE_LOG(LogPhysics, Warning, TEXT("Found bad index buffer when cooking UModelComponent physics data! Not enough indices in buffer for Element triangles. Component: %s, Element: %d, Index Start: %d, Tri Count: %d, Index Buffer Size: %d"),
-				*GetPathName(), ElementIndex, Element.FirstIndex, Element.NumTriangles, IndexBufferSize);
+			UE_LOG(LogPhysics, Warning, TEXT("Found bad index buffer when cooking UModelComponent physics data! Not enough indices in buffer for Element triangles. Component: %s, Element: %d, Index Start: %d, Tri Count: %d, Index Count: %d"),
+				*GetPathName(), ElementIndex, Element.FirstIndex, Element.NumTriangles, NumIndices);
 			continue;
 		}
 
diff --git a/Engine/Source/Runtime/Engine/Private/ModelRender.cpp b/Engine/Source/Runtime/Engine/Private/ModelRender.cpp
--- a/Engine/Source/Runtime/Engine/Private/ModelRender.cpp
+++ b/Engine/Source/Runtime/Engine/Private/ModelRender.cpp
@@ -500,17 +500,17 @@ public:
 
 	virtual void DrawStaticElements(FStaticPrimitiveDrawInterface* PDI) override
 	{
-		if(!HasViewDependentDPG())
+		if (!HasViewDependentDPG())
 		{
 			// Determine the DPG the primitive should be drawn in.
 			uint8 PrimitiveDPG = GetStaticDepthPriorityGroup();
 
 			PDI->ReserveMemoryForMeshes(Elements.Num());
 
-			for(int32 ElementIndex = 0;ElementIndex < Elements.Num();ElementIndex++)
+			for (int32 ElementIndex = 0;ElementIndex < Elements.Num();ElementIndex++)
 			{
 				const FModelElement& ModelElement = Component->GetElements()[ElementIndex];
-				if(ModelElement.NumTriangles > 0)
+				if (ModelElement.NumTriangles > 0)
 				{
 					FMeshBatch MeshElement;
 					FMeshBatchElement& BatchElement = MeshElement.Elements[0];
@@ -526,7 +526,12 @@ public:
 					MeshElement.Type = PT_TriangleList;
 					MeshElement.DepthPriorityGroup = PrimitiveDPG;
 					MeshElement.LODIndex = 0;
-					PDI->DrawMesh(MeshElement, FLT_MAX);
+					const bool bValidIndexBuffer = !BatchElement.IndexBuffer || (BatchElement.IndexBuffer && BatchElement.IndexBuffer->IsInitialized() && BatchElement.IndexBuffer->IndexBufferRHI);
+					ensure(bValidIndexBuffer);
+					if (bValidIndexBuffer)
+					{
+						PDI->DrawMesh(MeshElement, FLT_MAX);
+					}
 				}
 			}
 		}
