PR #4749: Add initialize for FAudioDevice::FCreateComponentParams::bAutoDestroy (to true) (Contributed by [REDACTED_PERSON])

diff --git a/Engine/Source/Runtime/Engine/Private/AudioDevice.cpp b/Engine/Source/Runtime/Engine/Private/AudioDevice.cpp
--- a/Engine/Source/Runtime/Engine/Private/AudioDevice.cpp
+++ b/Engine/Source/Runtime/Engine/Private/AudioDevice.cpp
@@ -4966,6 +4966,7 @@ FAudioDevice::FCreateComponentParams::FCreateComponentParams(FAudioDevice* InAud
 
 void FAudioDevice::FCreateComponentParams::CommonInit()
 {
+	bAutoDestroy = true;
 	bPlay = false;
 	bStopWhenOwnerDestroyed = true;
 	bLocationSet = false;
