PR #4540: Reverb isn't updated when volume changes (Contributed by [REDACTED_PERSON])

diff --git a/Engine/Source/Runtime/Engine/Private/AudioEffect.cpp b/Engine/Source/Runtime/Engine/Private/AudioEffect.cpp
--- a/Engine/Source/Runtime/Engine/Private/AudioEffect.cpp
+++ b/Engine/Source/Runtime/Engine/Private/AudioEffect.cpp
@@ -293,7 +293,8 @@ void FAudioEffectsManager::AddReferencedObjects( FReferenceCollector& Collector
 void FAudioEffectsManager::SetReverbSettings( const FReverbSettings& ReverbSettings, bool bForce)
 {
 	/** Update the settings if the reverb type has changed */
-	if( ReverbSettings.bApplyReverb && (ReverbSettings.ReverbEffect != CurrentReverbAsset || bForce))
+	if( ReverbSettings.bApplyReverb && (ReverbSettings.ReverbEffect != CurrentReverbAsset || bForce || 
+		!FMath::IsNearlyEqual(ReverbSettings.Volume, CurrentReverbSettings.Volume)))
 	{
 		FString CurrentReverbName = CurrentReverbAsset ? CurrentReverbAsset->GetName() : TEXT("None");
 		FString NextReverbName = ReverbSettings.ReverbEffect ? ReverbSettings.ReverbEffect->GetName() : TEXT("None");
