PR #4489: Small fix to Spatialization Plugin Settings tooltip (Contributed by [REDACTED_PERSON])

diff --git a/Engine/Source/Runtime/Engine/Classes/Sound/SoundAttenuation.h b/Engine/Source/Runtime/Engine/Classes/Sound/SoundAttenuation.h
--- a/Engine/Source/Runtime/Engine/Classes/Sound/SoundAttenuation.h
+++ b/Engine/Source/Runtime/Engine/Classes/Sound/SoundAttenuation.h
@@ -57,7 +57,7 @@ struct ENGINE_API FSoundAttenuationPluginSettings
 {
 	GENERATED_USTRUCT_BODY()
 
-	/** Settings to use with occlusion audio plugin. These are defined by the plugin creator. Not all audio plugins utilize this feature. This is an array so multiple plugins can have settings. */
+	/** Settings to use with spatialization audio plugin. These are defined by the plugin creator. Not all audio plugins utilize this feature. This is an array so multiple plugins can have settings. */
 	UPROPERTY(EditAnywhere, BlueprintReadWrite, Category = AttenuationSpatialization, meta = (DisplayName = "Spatialization Plugin Settings"))
 	TArray<USpatializationPluginSourceSettingsBase*> SpatializationPluginSettingsArray;
 
