[NO NON-PII COMMIT MESSAGE]

diff --git a/Engine/Source/Developer/AudioFormatADPCM/Private/AudioFormatADPCM.cpp b/Engine/Source/Developer/AudioFormatADPCM/Private/AudioFormatADPCM.cpp
--- a/Engine/Source/Developer/AudioFormatADPCM/Private/AudioFormatADPCM.cpp
+++ b/Engine/Source/Developer/AudioFormatADPCM/Private/AudioFormatADPCM.cpp
@@ -556,16 +556,19 @@ public:
 			check(InitialMaxChunkSize >= HeaderSize);
 			AddNewChunk(OutBuffers, InitialMaxChunkSize);
 			AddChunkData(OutBuffers, SrcData, HeaderSize);
-
 			int32 CurChunkDataSize = HeaderSize;
 
+			// Start looping through the rest of the chunks:
+			int32 ThisChunkMaxSize = InitialMaxChunkSize;
+
 			for (int32 BlockIndex = 0; BlockIndex < NumBlocksPerChannel; ++BlockIndex)
 			{
 				// If the next block size will put us over the max chunk size, then add a new chunk
-				if (CurChunkDataSize + NumChannels * BlockSize >= MaxChunkSize)
+				if (CurChunkDataSize + NumChannels * BlockSize >= ThisChunkMaxSize)
 				{
 					// Start a new chunk with the reserve memory for the max chunk size
 					AddNewChunk(OutBuffers, MaxChunkSize);
+					ThisChunkMaxSize = MaxChunkSize;
 					CurChunkDataSize = 0;
 				}
 
@@ -585,16 +588,17 @@ public:
 			int32 FrameSize = sizeof(uint16) * NumChannels;
 
 			// Add the first chunk and the header data
-			AddNewChunk(OutBuffers, MaxChunkSize);	
+			AddNewChunk(OutBuffers, InitialMaxChunkSize);	
 
 			// Add the header data
 			int32 HeaderSize = WaveInfo.SampleDataStart - SrcData;
+			check(InitialMaxChunkSize >= HeaderSize);
 			AddChunkData(OutBuffers, SrcData, HeaderSize);
 
 			SrcSize -= HeaderSize;
 			SrcData = WaveInfo.SampleDataStart;
 
-			int32 DataLeftInCurChunk = MaxChunkSize - HeaderSize;
+			int32 DataLeftInCurChunk = InitialMaxChunkSize - HeaderSize;
 
 			while (SrcSize > 0)
 			{
diff --git a/Engine/Source/Runtime/AudioPlatformConfiguration/Public/AudioCompressionSettings.h b/Engine/Source/Runtime/AudioPlatformConfiguration/Public/AudioCompressionSettings.h
--- a/Engine/Source/Runtime/AudioPlatformConfiguration/Public/AudioCompressionSettings.h
+++ b/Engine/Source/Runtime/AudioPlatformConfiguration/Public/AudioCompressionSettings.h
@@ -46,7 +46,7 @@ struct FPlatformAudioCookOverrides
 	// Increment this to force a recook on all Stream Caching assets.
 	// For testing, it's useful to set this to either a negative number or
 	// absurdly large number, to ensure you do not pollute the DDC.
-	static const int32 StreamCachingVersion = 5017;
+	static const int32 StreamCachingVersion = 5019;
 
 	bool bResampleForDevice;
 
diff --git a/Engine/Source/Runtime/Engine/Private/ADPCMAudioInfo.cpp b/Engine/Source/Runtime/Engine/Private/ADPCMAudioInfo.cpp
--- a/Engine/Source/Runtime/Engine/Private/ADPCMAudioInfo.cpp
+++ b/Engine/Source/Runtime/Engine/Private/ADPCMAudioInfo.cpp
@@ -42,6 +42,8 @@ FADPCMAudioInfo::FADPCMAudioInfo(void)
 	, Format(0)
 	, UncompressedBlockData(nullptr)
 	, SamplesPerBlock(0)
+	, FirstChunkSampleDataOffset(0)
+	, FirstChunkSampleDataIndex(0)
 	, bSeekPending(false)
 	, TargetSeekTime(0.0f)
 	, LastSeekTime(0.0f)
@@ -434,10 +436,12 @@ bool FADPCMAudioInfo::StreamCompressedInfoInternal(USoundWave* Wave, struct FSou
 
 	check(StreamingSoundWave == Wave);
 
+	CurrentChunkIndex = 0;
+
 	// Get the first chunk of audio data (should already be loaded)
-	uint8 const* FirstChunk = GetLoadedChunk(Wave, 0, CurrentChunkDataSize);
+	uint8 const* ChunkData = GetLoadedChunk(Wave, CurrentChunkIndex, CurrentChunkDataSize);
 
-	if (FirstChunk == nullptr)
+	if (ChunkData == nullptr)
 	{
 		return false;
 	}
@@ -447,19 +451,33 @@ bool FADPCMAudioInfo::StreamCompressedInfoInternal(USoundWave* Wave, struct FSou
 
 	void* FormatHeader;
 
-	if (!WaveInfo.ReadWaveInfo((uint8*)FirstChunk, Wave->RunningPlatformData->Chunks[0].AudioDataSize, nullptr, true, &FormatHeader))
+	if (!WaveInfo.ReadWaveInfo((uint8*)ChunkData, CurrentChunkDataSize, nullptr, true, &FormatHeader))
 	{
 		UE_LOG(LogAudio, Warning, TEXT("WaveInfo.ReadWaveInfo Failed"));
 		return false;
 	}
+	
+	// if we only included the header in the zeroth chunk, skip to the next chunk.
+	int32 SampleDataOffset = WaveInfo.SampleDataStart - ChunkData;
+	check(SampleDataOffset > 0);
+	if (((uint32)SampleDataOffset) >= CurrentChunkDataSize)
+	{
+		++CurrentChunkIndex;
+		ChunkData = GetLoadedChunk(Wave, CurrentChunkIndex, CurrentChunkDataSize);
+		FirstChunkSampleDataIndex = CurrentChunkIndex;
+		FirstChunkSampleDataOffset = 0;
+	}
+	else
+	{
+		FirstChunkSampleDataOffset = WaveInfo.SampleDataStart - ChunkData;
+		FirstChunkSampleDataIndex = 0;
+	}
 
-	SrcBufferData = FirstChunk;
-
-	FirstChunkSampleDataOffset = WaveInfo.SampleDataStart - FirstChunk;
+	SrcBufferData = ChunkData;
 	CurrentChunkBufferOffset = 0;
 	CurCompressedChunkData = nullptr;
 	CurrentUncompressedBlockSampleIndex = 0;
-	CurrentChunkIndex = 0;
+	
 	TotalSamplesStreamed = 0;
 	Format = *WaveInfo.pFormatTag;
 	NumChannels = *WaveInfo.pChannels;
@@ -588,7 +606,7 @@ bool FADPCMAudioInfo::StreamCompressedData(uint8* Destination, bool bLooping, ui
 					// Set the current buffer offset accounting for the header in the first chunk
 					if (!bSeekPending)
 					{
-						CurrentChunkBufferOffset = CurrentChunkIndex == 0 ? FirstChunkSampleDataOffset : 0;
+						CurrentChunkBufferOffset = CurrentChunkIndex == FirstChunkSampleDataOffset ? FirstChunkSampleDataOffset : 0;
 					}
 
 					bSeekPending = false;
@@ -640,7 +658,7 @@ bool FADPCMAudioInfo::StreamCompressedData(uint8* Destination, bool bLooping, ui
 			{
 				ReachedEndOfSamples = true;
 				CurrentUncompressedBlockSampleIndex = 0;
-				CurrentChunkIndex = 0;
+				CurrentChunkIndex = FirstChunkSampleDataIndex;
 				CurrentChunkBufferOffset = 0;
 				TotalSamplesStreamed = 0;
 				CurCompressedChunkData = nullptr;
@@ -691,7 +709,7 @@ bool FADPCMAudioInfo::StreamCompressedData(uint8* Destination, bool bLooping, ui
 				// Set the current buffer offset accounting for the header in the first chunk
 				if (!bSeekPending)
 				{
-					CurrentChunkBufferOffset = CurrentChunkIndex == 0 ? FirstChunkSampleDataOffset : 0;
+					CurrentChunkBufferOffset = CurrentChunkIndex == FirstChunkSampleDataIndex ? FirstChunkSampleDataOffset : 0;
 				}
 
 				bSeekPending = false;
diff --git a/Engine/Source/Runtime/Engine/Private/AudioStreamingCache.cpp b/Engine/Source/Runtime/Engine/Private/AudioStreamingCache.cpp
--- a/Engine/Source/Runtime/Engine/Private/AudioStreamingCache.cpp
+++ b/Engine/Source/Runtime/Engine/Private/AudioStreamingCache.cpp
@@ -515,7 +515,7 @@ uint64 FAudioChunkCache::TrimMemory(uint64 BytesToFree)
 
 #if DEBUG_STREAM_CACHE
 			// Reset debug info:
-			CurrentElement->DebugInfo = FCacheElementDebugInfo();
+			CurrentElement->DebugInfo.Reset();
 #endif
 		}
 
@@ -613,7 +613,10 @@ FAudioChunkCache::FCacheElement* FAudioChunkCache::FindElementForKey(const FChun
 		{
 			
 #if DEBUG_STREAM_CACHE
-			CurrentElement->DebugInfo.PreviousLocationsBeforeBeingTouched.Add(ElementPosition);
+			{
+				FScopeLock ScopeLock(&CurrentElement->DebugInfo.PreviousLocationsLock);
+				CurrentElement->DebugInfo.PreviousLocationsBeforeBeingTouched.Add(ElementPosition);
+			}
 #endif
 
 			return CurrentElement;
@@ -798,9 +801,12 @@ FAudioChunkCache::FCacheElement* FAudioChunkCache::EvictLeastRecentChunk()
 	}
 
 #if DEBUG_STREAM_CACHE
-	// Reset debug information:
-	CacheElement->DebugInfo.PreviousLocationsBeforeBeingTouched.Reset();
-	CacheElement->DebugInfo.NumTimesTouched = 0;
+	{
+		FScopeLock ScopeLock(&CacheElement->DebugInfo.PreviousLocationsLock);
+		// Reset debug information:
+		CacheElement->DebugInfo.PreviousLocationsBeforeBeingTouched.Reset();
+		CacheElement->DebugInfo.NumTimesTouched = 0;
+	}
 #endif
 
 	return CacheElement;
diff --git a/Engine/Source/Runtime/Engine/Public/ADPCMAudioInfo.h b/Engine/Source/Runtime/Engine/Public/ADPCMAudioInfo.h
--- a/Engine/Source/Runtime/Engine/Public/ADPCMAudioInfo.h
+++ b/Engine/Source/Runtime/Engine/Public/ADPCMAudioInfo.h
@@ -160,6 +160,7 @@ private:
 	uint32			TotalSamplesPerChannel;			// Number of samples per channel, used to detect when an audio waveform has ended
 	uint32			SamplesPerBlock;				// The number of samples per block
 	uint32			FirstChunkSampleDataOffset;		// The size of the header in the first chunk, used to skip over it when looping or starting the sample over
+	uint32          FirstChunkSampleDataIndex;
 	const uint8*	CurCompressedChunkData;			// A pointer to the current chunk of data
 	FAudioChunkHandle CurCompressedChunkHandle;     // Shared reference to the current chunk of data.
 
diff --git a/Engine/Source/Runtime/Engine/Public/AudioStreamingCache.h b/Engine/Source/Runtime/Engine/Public/AudioStreamingCache.h
--- a/Engine/Source/Runtime/Engine/Public/AudioStreamingCache.h
+++ b/Engine/Source/Runtime/Engine/Public/AudioStreamingCache.h
@@ -107,6 +107,7 @@ private:
 		// This allows every call to AddOrTouchChunk to report how far down the cache
 		// it was when it was needed, or INDEX_NONE if it wasn't in the cache when it was needed.
 		TArray<int32> PreviousLocationsBeforeBeingTouched;
+		FCriticalSection PreviousLocationsLock;
 
 		FCacheElementDebugInfo()
 			: NumTotalChunks(0)
@@ -115,6 +116,16 @@ private:
 			, TimeToLoad(0.0)
 		{
 		}
+
+		void Reset()
+		{
+			NumTotalChunks = 0;
+			NumTimesTouched = 0;
+			TimeLoadStarted = 0;
+
+			FScopeLock ScopeLock(&PreviousLocationsLock);
+			PreviousLocationsBeforeBeingTouched.Empty();
+		}
 	};
 #endif
 
