Revert FTickFunction changes while consideration of how to deal with queries to tick function when it is ticking after being unregistered should be handled
#rnx

diff --git a/Engine/Source/Runtime/Engine/Classes/Components/SceneComponent.h b/Engine/Source/Runtime/Engine/Classes/Components/SceneComponent.h
--- a/Engine/Source/Runtime/Engine/Classes/Components/SceneComponent.h
+++ b/Engine/Source/Runtime/Engine/Classes/Components/SceneComponent.h
@@ -285,20 +285,7 @@ protected:
 	class UBillboardComponent* SpriteComponent;
 #endif
 
-public:
-	/** Delegate called when this component is moved */
-	FTransformUpdated TransformUpdated;
-
-	/** Returns the current scoped movement update, or NULL if there is none. @see FScopedMovementUpdate */
-	class FScopedMovementUpdate* GetCurrentScopedMovement() const;
-
 private:
-	/** Stack of current movement scopes. */
-	TArray<class FScopedMovementUpdate*> ScopedMovementStack;
-
-	void BeginScopedMovementUpdate(class FScopedMovementUpdate& ScopedUpdate);
-	void EndScopedMovementUpdate(class FScopedMovementUpdate& ScopedUpdate);
-
 	/** Cache that avoids Quat<->Rotator conversions if possible. Only to be used with GetComponentTransform().GetRotation(). */
 	FRotationConversionCache WorldRotationCache;
 
@@ -315,7 +302,19 @@ public:
 	/** Get the RelativeRotationCache.  */
 	FORCEINLINE const FRotationConversionCache& GetRelativeRotationCache() const { return RelativeRotationCache; }
 
+	/** Delegate called when this component is moved */
+	FTransformUpdated TransformUpdated;
+
+	/** Returns the current scoped movement update, or NULL if there is none. @see FScopedMovementUpdate */
+	class FScopedMovementUpdate* GetCurrentScopedMovement() const;
+
 private:
+	/** Stack of current movement scopes. */
+	TArray<class FScopedMovementUpdate*> ScopedMovementStack;
+
+	void BeginScopedMovementUpdate(class FScopedMovementUpdate& ScopedUpdate);
+	void EndScopedMovementUpdate(class FScopedMovementUpdate& ScopedUpdate);
+
 	UFUNCTION()
 	void OnRep_Transform();
 
diff --git a/Engine/Source/Runtime/Engine/Classes/Engine/EngineBaseTypes.h b/Engine/Source/Runtime/Engine/Classes/Engine/EngineBaseTypes.h
--- a/Engine/Source/Runtime/Engine/Classes/Engine/EngineBaseTypes.h
+++ b/Engine/Source/Runtime/Engine/Classes/Engine/EngineBaseTypes.h
@@ -190,6 +190,13 @@ public:
 	UPROPERTY(EditDefaultsOnly, Category="Tick", AdvancedDisplay)
 	TEnumAsByte<enum ETickingGroup> EndTickGroup;
 
+protected:
+	/** Internal data that indicates the tick group we actually started in (it may have been delayed due to prerequisites) **/
+	TEnumAsByte<enum ETickingGroup> ActualStartTickGroup;
+
+	/** Internal data that indicates the tick group we actually started in (it may have been delayed due to prerequisites) **/
+	TEnumAsByte<enum ETickingGroup> ActualEndTickGroup;
+
 public:
 	/** Bool indicating that this function should execute even if the game is paused. Pause ticks are very limited in capabilities. **/
 	UPROPERTY(EditDefaultsOnly, Category="Tick", AdvancedDisplay)
@@ -214,6 +221,11 @@ public:
 	uint8 bRunOnAnyThread:1;
 
 private:
+	/** If true, means that this tick function is in the master array of tick functions */
+	bool bRegistered;
+
+	/** Cache whether this function was rescheduled as an interval function during StartParallel */
+	bool bWasInterval;
 
 	enum class ETickState : uint8
 	{
@@ -229,62 +241,42 @@ private:
 	 **/
 	ETickState TickState;
 
-public:
-	/** The frequency in seconds at which this tick function will be executed.  If less than or equal to 0 then it will tick every frame */
-	UPROPERTY(EditDefaultsOnly, Category="Tick", meta=(DisplayName="Tick Interval (secs)"))
-	float TickInterval;
+	/** Internal data to track if we have started visiting this tick function yet this frame **/
+	int32 TickVisitedGFrameCounter;
 
-private:
+	/** Internal data to track if we have finshed visiting this tick function yet this frame **/
+	int32 TickQueuedGFrameCounter;
+
+	/** Pointer to the task, only used during setup. This is often stale. **/
+	void *TaskPointer;
 
 	/** Prerequisites for this tick function **/
 	TArray<struct FTickPrerequisite> Prerequisites;
 
-	/** Internal Data structure that contains members only required for a registered tick function **/
-	struct FInternalData
-	{
-		FInternalData();
-
-		/** Cache whether this function was rescheduled as an interval function during StartParallel */
-		bool bWasInterval;
-
-		/** Internal data that indicates the tick group we actually started in (it may have been delayed due to prerequisites) **/
-		TEnumAsByte<enum ETickingGroup> ActualStartTickGroup;
+	/** The next function in the cooling down list for ticks with an interval*/
+	FTickFunction* Next;
 
-		/** Internal data that indicates the tick group we actually started in (it may have been delayed due to prerequisites) **/
-		TEnumAsByte<enum ETickingGroup> ActualEndTickGroup;
-		
-		/** Internal data to track if we have started visiting this tick function yet this frame **/
-		int32 TickVisitedGFrameCounter;
-
-		/** Internal data to track if we have finshed visiting this tick function yet this frame **/
-		int32 TickQueuedGFrameCounter;
-
-		/** Pointer to the task, only used during setup. This is often stale. **/
-		void* TaskPointer;
-
-		/** The next function in the cooling down list for ticks with an interval*/
-		FTickFunction* Next;
+	/** 
+	  * If TickFrequency is greater than 0 and tick state is CoolingDown, this is the time, 
+	  * relative to the element ahead of it in the cooling down list, remaining until the next time this function will tick 
+	  */
+	float RelativeTickCooldown;
 
-		/** 
-		 * If TickFrequency is greater than 0 and tick state is CoolingDown, this is the time, 
-		 * relative to the element ahead of it in the cooling down list, remaining until the next time this function will tick 
-		 **/
-		float RelativeTickCooldown;
+	/** 
+	* The last world game time at which we were ticked. Game time used is dependent on bTickEvenWhenPaused
+	* Valid only if we've been ticked at least once since having a tick interval; otherwise set to -1.f
+	*/
+	float LastTickGameTimeSeconds;
 
-		/** 
-		 * The last world game time at which we were ticked. Game time used is dependent on bTickEvenWhenPaused
-		 * Valid only if we've been ticked at least once since having a tick interval; otherwise set to -1.f
-		 **/
-		float LastTickGameTimeSeconds;
+public:
 
-		/** Back pointer to the FTickTaskLevel containing this tick function if it is registered **/
-		class FTickTaskLevel*						TickTaskLevel;
-	};
+	/** The frequency in seconds at which this tick function will be executed.  If less than or equal to 0 then it will tick every frame */
+	UPROPERTY(EditDefaultsOnly, Category="Tick", meta=(DisplayName="Tick Interval (secs)"))
+	float TickInterval;
 
-	/** Lazily allocated struct that contains the necessary data for a tick function that is registered. **/
-	TUniquePtr<FInternalData> InternalData;
+	/** Back pointer to the FTickTaskLevel containing this tick function if it is registered **/
+	class FTickTaskLevel*						TickTaskLevel;
 
-public:
 	/** Default constructor, intitalizes to reasonable defaults **/
 	FTickFunction();
 	/** Destructor, unregisters the tick function **/
@@ -298,14 +290,14 @@ public:
 	/** Removes the tick function from the master list of tick functions. **/
 	void UnRegisterTickFunction();
 	/** See if the tick function is currently registered */
-	bool IsTickFunctionRegistered() const { return (InternalData != nullptr); }
+	bool IsTickFunctionRegistered() const { return bRegistered; }
 
 	/** Enables or disables this tick function. **/
 	void SetTickFunctionEnable(bool bInEnabled);
 	/** Returns whether the tick function is currently enabled */
 	bool IsTickFunctionEnabled() const { return TickState != ETickState::Disabled; }
 	/** Returns whether it is valid to access this tick function's completion handle */
-	bool IsCompletionHandleValid() const { return (InternalData && InternalData->TaskPointer); }
+	bool IsCompletionHandleValid() const { return TaskPointer != nullptr; }
 	/**
 	* Gets the current completion handle of this tick function, so it can be delayed until a later point when some additional
 	* tasks have been completed.  Only valid after TG_PreAsyncWork has started and then only until the TickFunction finishes
@@ -319,7 +311,7 @@ public:
 	**/
 	TEnumAsByte<enum ETickingGroup> GetActualTickGroup() const
 	{
-		return (InternalData ? InternalData->ActualStartTickGroup : TickGroup);
+		return ActualStartTickGroup;
 	}
 
 	/** 
@@ -328,7 +320,7 @@ public:
 	**/
 	TEnumAsByte<enum ETickingGroup> GetActualEndTickGroup() const
 	{
-		return (InternalData ? InternalData->ActualEndTickGroup : EndTickGroup);
+		return ActualEndTickGroup;
 	}
 
 
@@ -486,7 +478,7 @@ struct FActorComponentTickFunction : public FTickFunction
 	ENGINE_API virtual FName DiagnosticContext(bool bDetailed) override;
 
 	/**
-	 * Conditionally calls ExecuteTickFunc if registered and a bunch of other criteria are met
+	 * Conditionally calls ExecuteTickFunc if bRegistered == true and a bunch of other criteria are met
 	 * @param Target - the actor component we are ticking
 	 * @param bTickInEditor - whether the target wants to tick in the editor
 	 * @param DeltaTime - The time since the last tick.
diff --git a/Engine/Source/Runtime/Engine/Private/TickTaskManager.cpp b/Engine/Source/Runtime/Engine/Private/TickTaskManager.cpp
--- a/Engine/Source/Runtime/Engine/Private/TickTaskManager.cpp
+++ b/Engine/Source/Runtime/Engine/Private/TickTaskManager.cpp
@@ -285,10 +285,7 @@ public:
 			Target->ExecuteTick(Target->CalculateDeltaTime(Context), Context.TickType, CurrentThread, MyCompletionGraphEvent);
 			LIGHTWEIGHT_TIME_GUARD_END(FTickFunctionTask, Target->DiagnosticMessage());
 		}
-		if (Target->InternalData)
-		{
-			Target->InternalData->TaskPointer = nullptr;  // This is stale and a good time to clear it for safety
-		}
+		Target->TaskPointer = nullptr;  // This is stale and a good time to clear it for safety
 	}
 };
 
@@ -425,12 +422,11 @@ public:
 	 */
 	FORCEINLINE void StartTickTask(const FGraphEventArray* Prerequisites, FTickFunction* TickFunction, const FTickContext& TickContext)
 	{
-		checkSlow(TickFunction->InternalData);
-		checkSlow(TickFunction->InternalData->ActualStartTickGroup >=0 && TickFunction->InternalData->ActualStartTickGroup < TG_MAX);
+		checkSlow(TickFunction->ActualStartTickGroup >=0 && TickFunction->ActualStartTickGroup < TG_MAX);
 
 		FTickContext UseContext = TickContext;
 
-		bool bIsOriginalTickGroup = (TickFunction->InternalData->ActualStartTickGroup == TickFunction->TickGroup);
+		bool bIsOriginalTickGroup = (TickFunction->ActualStartTickGroup == TickFunction->TickGroup);
 
 		if (TickFunction->bRunOnAnyThread && bAllowConcurrentTicks && bIsOriginalTickGroup)
 		{
@@ -448,7 +444,7 @@ public:
 			UseContext.Thread = ENamedThreads::SetTaskPriority(ENamedThreads::GameThread, TickFunction->bHighPriority ? ENamedThreads::HighTaskPriority : ENamedThreads::NormalTaskPriority);
 		}
 
-		TickFunction->InternalData->TaskPointer = TGraphTask<FTickFunctionTask>::CreateTask(Prerequisites, TickContext.Thread).ConstructAndHold(TickFunction, &UseContext, bLogTicks, bLogTicksShowPrerequistes);
+		TickFunction->TaskPointer = TGraphTask<FTickFunctionTask>::CreateTask(Prerequisites, TickContext.Thread).ConstructAndHold(TickFunction, &UseContext, bLogTicks, bLogTicksShowPrerequistes);
 	}
 
 	/** Add a completion handle to a tick group **/
@@ -501,11 +497,10 @@ public:
 	 */
 	FORCEINLINE void QueueTickTask(const FGraphEventArray* Prerequisites, FTickFunction* TickFunction, const FTickContext& TickContext)
 	{
-		checkSlow(TickFunction->InternalData);
 		checkSlow(TickContext.Thread == ENamedThreads::GameThread);
 		StartTickTask(Prerequisites, TickFunction, TickContext);
-		TGraphTask<FTickFunctionTask>* Task = (TGraphTask<FTickFunctionTask>*)TickFunction->InternalData->TaskPointer;
-		AddTickTaskCompletion(TickFunction->InternalData->ActualStartTickGroup, TickFunction->InternalData->ActualEndTickGroup, Task, TickFunction->bHighPriority);
+		TGraphTask<FTickFunctionTask>* Task = (TGraphTask<FTickFunctionTask>*)TickFunction->TaskPointer;
+		AddTickTaskCompletion(TickFunction->ActualStartTickGroup, TickFunction->ActualEndTickGroup, Task, TickFunction->bHighPriority);
 	}
 
 	/**
@@ -517,11 +512,10 @@ public:
 	 */
 	FORCEINLINE void QueueTickTaskParallel(const FGraphEventArray* Prerequisites, FTickFunction* TickFunction, const FTickContext& TickContext)
 	{
-		checkSlow(TickFunction->InternalData);
 		checkSlow(TickContext.Thread == ENamedThreads::GameThread);
 		StartTickTask(Prerequisites, TickFunction, TickContext);
-		TGraphTask<FTickFunctionTask>* Task = (TGraphTask<FTickFunctionTask>*)TickFunction->InternalData->TaskPointer;
-		AddTickTaskCompletionParallel(TickFunction->InternalData->ActualStartTickGroup, TickFunction->InternalData->ActualEndTickGroup, Task, TickFunction->bHighPriority);
+		TGraphTask<FTickFunctionTask>* Task = (TGraphTask<FTickFunctionTask>*)TickFunction->TaskPointer;
+		AddTickTaskCompletionParallel(TickFunction->ActualStartTickGroup, TickFunction->ActualEndTickGroup, Task, TickFunction->bHighPriority);
 	}
 
 	/**
@@ -715,18 +709,17 @@ public:
 	{
 		for (TSet<FTickFunction*>::TIterator It(AllEnabledTickFunctions); It; ++It)
 		{
-			(*It)->InternalData.Reset();
+			(*It)->bRegistered = false;
 		}
 		for (TSet<FTickFunction*>::TIterator It(AllDisabledTickFunctions); It; ++It)
 		{
-			(*It)->InternalData.Reset();
+			(*It)->bRegistered = false;
 		}
 		FTickFunction* CoolingDownNode = AllCoolingDownTickFunctions.Head;
 		while (CoolingDownNode)
 		{
-			FTickFunction* NextNode = CoolingDownNode->InternalData->Next;
-			CoolingDownNode->InternalData.Reset();
-			CoolingDownNode = NextNode;
+			CoolingDownNode->bRegistered = false;
+			CoolingDownNode = CoolingDownNode->Next;
 		}
 	}
 
@@ -753,15 +746,15 @@ public:
 			FTickFunction* TickFunction = AllCoolingDownTickFunctions.Head;
 			while (TickFunction)
 			{
-				if (CumulativeCooldown + TickFunction->InternalData->RelativeTickCooldown >= Context.DeltaSeconds)
+				if (CumulativeCooldown + TickFunction->RelativeTickCooldown >= Context.DeltaSeconds)
 				{
-					TickFunction->InternalData->RelativeTickCooldown -= (Context.DeltaSeconds - CumulativeCooldown);
+					TickFunction->RelativeTickCooldown -= (Context.DeltaSeconds - CumulativeCooldown);
 					break;
 				}
-				CumulativeCooldown += TickFunction->InternalData->RelativeTickCooldown;
+				CumulativeCooldown += TickFunction->RelativeTickCooldown;
 
 				TickFunction->TickState = FTickFunction::ETickState::Enabled;
-				TickFunction = TickFunction->InternalData->Next;
+				TickFunction = TickFunction->Next;
 				++CooldownTicksEnabled;
 			}
 		}
@@ -798,21 +791,21 @@ public:
 			FTickFunction* TickFunction = AllCoolingDownTickFunctions.Head;
 			while (TickFunction)
 			{
-				if (CumulativeCooldown + TickFunction->InternalData->RelativeTickCooldown >= Context.DeltaSeconds)
+				if (CumulativeCooldown + TickFunction->RelativeTickCooldown >= Context.DeltaSeconds)
 				{
-					TickFunction->InternalData->RelativeTickCooldown -= (Context.DeltaSeconds - CumulativeCooldown);
+					TickFunction->RelativeTickCooldown -= (Context.DeltaSeconds - CumulativeCooldown);
 					break;
 				}
-				CumulativeCooldown += TickFunction->InternalData->RelativeTickCooldown;
+				CumulativeCooldown += TickFunction->RelativeTickCooldown;
 
 				TickFunction->TickState = FTickFunction::ETickState::Enabled;
-				TickFunction->InternalData->bWasInterval = true;
+				TickFunction->bWasInterval = true;
 				AllTickFunctions.Add(TickFunction);
 
 				TickFunctionsToReschedule.Add(FTickScheduleDetails(TickFunction, TickFunction->TickInterval - (Context.DeltaSeconds - CumulativeCooldown))); // Give credit for any overrun
 
-				AllCoolingDownTickFunctions.Head = TickFunction->InternalData->Next;
-				TickFunction = TickFunction->InternalData->Next;
+				AllCoolingDownTickFunctions.Head = TickFunction->Next;
+				TickFunction = TickFunction->Next;
 			}
 		}
 	}
@@ -859,7 +852,7 @@ public:
 			while (ComparisonTickFunction && RescheduleIndex < TickFunctionsToReschedule.Num())
 			{
 				const float CooldownTime = TickFunctionsToReschedule[RescheduleIndex].Cooldown;
-				if ((CumulativeCooldown + ComparisonTickFunction->InternalData->RelativeTickCooldown) > CooldownTime)
+				if ((CumulativeCooldown + ComparisonTickFunction->RelativeTickCooldown) > CooldownTime)
 				{
 					FTickFunction* TickFunction = TickFunctionsToReschedule[RescheduleIndex].TickFunction;
 					if (TickFunction->TickState != FTickFunction::ETickState::Disabled)
@@ -869,29 +862,29 @@ public:
 							verify(AllEnabledTickFunctions.Remove(TickFunction) == 1);
 						}
 						TickFunction->TickState = FTickFunction::ETickState::CoolingDown;
-						TickFunction->InternalData->RelativeTickCooldown = CooldownTime - CumulativeCooldown;
+						TickFunction->RelativeTickCooldown = CooldownTime - CumulativeCooldown;
 
 						if (PrevComparisonTickFunction)
 						{
-							PrevComparisonTickFunction->InternalData->Next = TickFunction;
+							PrevComparisonTickFunction->Next = TickFunction;
 						}
 						else
 						{
 							check(ComparisonTickFunction == AllCoolingDownTickFunctions.Head);
 							AllCoolingDownTickFunctions.Head = TickFunction;
 						}
-						TickFunction->InternalData->Next = ComparisonTickFunction;
+						TickFunction->Next = ComparisonTickFunction;
 						PrevComparisonTickFunction = TickFunction;
-						ComparisonTickFunction->InternalData->RelativeTickCooldown -= TickFunction->InternalData->RelativeTickCooldown;
-						CumulativeCooldown += TickFunction->InternalData->RelativeTickCooldown;
+						ComparisonTickFunction->RelativeTickCooldown -= TickFunction->RelativeTickCooldown;
+						CumulativeCooldown += TickFunction->RelativeTickCooldown;
 					}
 					++RescheduleIndex;
 				}
 				else
 				{
-					CumulativeCooldown += ComparisonTickFunction->InternalData->RelativeTickCooldown;
+					CumulativeCooldown += ComparisonTickFunction->RelativeTickCooldown;
 					PrevComparisonTickFunction = ComparisonTickFunction;
-					ComparisonTickFunction = ComparisonTickFunction->InternalData->Next;
+					ComparisonTickFunction = ComparisonTickFunction->Next;
 				}
 			}
 			for ( ; RescheduleIndex < TickFunctionsToReschedule.Num(); ++RescheduleIndex)
@@ -907,12 +900,12 @@ public:
 					const float CooldownTime = TickFunctionsToReschedule[RescheduleIndex].Cooldown;
 
 					TickFunction->TickState = FTickFunction::ETickState::CoolingDown;
-					TickFunction->InternalData->RelativeTickCooldown = CooldownTime - CumulativeCooldown;
+					TickFunction->RelativeTickCooldown = CooldownTime - CumulativeCooldown;
 
-					TickFunction->InternalData->Next = nullptr;
+					TickFunction->Next = nullptr;
 					if (PrevComparisonTickFunction)
 					{
-						PrevComparisonTickFunction->InternalData->Next = TickFunction;
+						PrevComparisonTickFunction->Next = TickFunction;
 					}
 					else
 					{
@@ -921,7 +914,7 @@ public:
 					}
 					PrevComparisonTickFunction = TickFunction;
 
-					CumulativeCooldown += TickFunction->InternalData->RelativeTickCooldown;
+					CumulativeCooldown += TickFunction->RelativeTickCooldown;
 				}
 			}
 			TickFunctionsToReschedule.Reset();
@@ -949,10 +942,10 @@ public:
 		{
 			if (TickFunction->TickState == FTickFunction::ETickState::Enabled)
 			{
-				CumulativeCooldown += TickFunction->InternalData->RelativeTickCooldown;
+				CumulativeCooldown += TickFunction->RelativeTickCooldown;
 				TickFunction->QueueTickFunction(TTS, Context);
 				TickFunctionsToReschedule.Add(FTickScheduleDetails(TickFunction, TickFunction->TickInterval - (Context.DeltaSeconds - CumulativeCooldown))); // Give credit for any overrun
-				AllCoolingDownTickFunctions.Head = TickFunction->InternalData->Next;
+				AllCoolingDownTickFunctions.Head = TickFunction->Next;
 			}
 			else
 			{
@@ -1023,14 +1016,14 @@ public:
 		FTickFunction* TickFunction = AllCoolingDownTickFunctions.Head;
 		while (TickFunction)
 		{
-			CumulativeCooldown += TickFunction->InternalData->RelativeTickCooldown;
+			CumulativeCooldown += TickFunction->RelativeTickCooldown;
 			if (TickFunction->bTickEvenWhenPaused)
 			{
-				TickFunction->InternalData->TaskPointer = nullptr; // this is stale, clear it out now
+				TickFunction->TaskPointer = nullptr; // this is stale, clear it out now
 				if (CumulativeCooldown < InContext.DeltaSeconds)
 				{
-					TickFunction->InternalData->TickVisitedGFrameCounter = GFrameCounter;
-					TickFunction->InternalData->TickQueuedGFrameCounter = GFrameCounter;
+					TickFunction->TickVisitedGFrameCounter = GFrameCounter;
+					TickFunction->TickQueuedGFrameCounter = GFrameCounter;
 					TickFunction->ExecuteTick(TickFunction->CalculateDeltaTime(InContext), InContext.TickType, ENamedThreads::GameThread, FGraphEventRef());
 
 					TickFunctionsToReschedule.Add(FTickScheduleDetails(TickFunction, TickFunction->TickInterval - (InContext.DeltaSeconds - CumulativeCooldown))); // Give credit for any overrun
@@ -1041,34 +1034,34 @@ public:
 				}
 				if (PrevTickFunction)
 				{
-					PrevTickFunction->InternalData->Next = TickFunction->InternalData->Next;
+					PrevTickFunction->Next = TickFunction->Next;
 				}
 				else
 				{
 					check(TickFunction == AllCoolingDownTickFunctions.Head);
-					AllCoolingDownTickFunctions.Head = TickFunction->InternalData->Next;
+					AllCoolingDownTickFunctions.Head = TickFunction->Next;
 				}
-				if (TickFunction->InternalData->Next)
+				if (TickFunction->Next)
 				{
-					TickFunction->InternalData->Next->InternalData->RelativeTickCooldown += TickFunction->InternalData->RelativeTickCooldown;
-					CumulativeCooldown -= TickFunction->InternalData->RelativeTickCooldown; // Since the next object in the list will have this cooldown included take it back out of the cumulative
+					TickFunction->Next->RelativeTickCooldown += TickFunction->RelativeTickCooldown;
+					CumulativeCooldown -= TickFunction->RelativeTickCooldown; // Since the next object in the list will have this cooldown included take it back out of the cumulative
 				}
 			}
 			else
 			{
 				PrevTickFunction = TickFunction;
 			}
-			TickFunction = TickFunction->InternalData->Next;
+			TickFunction = TickFunction->Next;
 		}
 
 		for (TSet<FTickFunction*>::TIterator It(AllEnabledTickFunctions); It; ++It)
 		{
 			TickFunction = *It;
-			TickFunction->InternalData->TaskPointer = nullptr; // this is stale, clear it out now
+			TickFunction->TaskPointer = nullptr; // this is stale, clear it out now
 			if (TickFunction->bTickEvenWhenPaused && TickFunction->TickState == FTickFunction::ETickState::Enabled)
 			{
-				TickFunction->InternalData->TickVisitedGFrameCounter = GFrameCounter;
-				TickFunction->InternalData->TickQueuedGFrameCounter = GFrameCounter;
+				TickFunction->TickVisitedGFrameCounter = GFrameCounter;
+				TickFunction->TickQueuedGFrameCounter = GFrameCounter;
 				TickFunction->ExecuteTick(TickFunction->CalculateDeltaTime(InContext), InContext.TickType, ENamedThreads::GameThread, FGraphEventRef());
 
 				if (TickFunction->TickInterval > 0.f)
@@ -1137,7 +1130,7 @@ public:
 		Ar.Logf(TEXT("%s, %s, ActualStartTickGroup: %s, Prerequesities: %d"),
 			*Function->DiagnosticMessage(),
 			Function->IsTickFunctionEnabled() ? (RemainingCooldown > 0.f ? *FString::Printf(TEXT("Cooling Down for %.4g seconds"),RemainingCooldown) : TEXT("Enabled")) : TEXT("Disabled"),
-			*TickGroupEnum->GetNameStringByValue(Function->GetActualTickGroup()),
+			*TickGroupEnum->GetNameStringByValue(Function->ActualStartTickGroup),
 			Function->Prerequisites.Num());
 
 		// List all prerequisities
@@ -1169,9 +1162,9 @@ public:
 			FTickFunction* TickFunction = AllCoolingDownTickFunctions.Head;
 			while (TickFunction)
 			{
-				CumulativeCooldown += TickFunction->InternalData->RelativeTickCooldown;
+				CumulativeCooldown += TickFunction->RelativeTickCooldown;
 				DumpTickFunction(Ar, TickFunction, TickGroupEnum, CumulativeCooldown);
-				TickFunction = TickFunction->InternalData->Next;
+				TickFunction = TickFunction->Next;
 				++EnabledCount;
 			}
 		}
@@ -1180,7 +1173,7 @@ public:
 			FTickFunction* TickFunction = AllCoolingDownTickFunctions.Head;
 			while (TickFunction)
 			{
-				TickFunction = TickFunction->InternalData->Next;
+				TickFunction = TickFunction->Next;
 				++EnabledCount;
 			}
 		}
@@ -1219,7 +1212,7 @@ public:
 		while (TickFunction)
 		{
 			AddTickFunctionToMap(ClassNameToCountMap, TickFunction, bDetailed);
-			TickFunction = TickFunction->InternalData->Next;
+			TickFunction = TickFunction->Next;
 			++EnabledCount;
 		}
 	}
@@ -1245,19 +1238,19 @@ public:
 							bFound = true;
 							if (PrevComparisionFunction)
 							{
-								PrevComparisionFunction->InternalData->Next = TickFunction->InternalData->Next;
+								PrevComparisionFunction->Next = TickFunction->Next;
 							}
 							else
 							{
 								check(TickFunction == AllCoolingDownTickFunctions.Head);
-								AllCoolingDownTickFunctions.Head = TickFunction->InternalData->Next;
+								AllCoolingDownTickFunctions.Head = TickFunction->Next;
 							}
-							TickFunction->InternalData->Next = nullptr;
+							TickFunction->Next = nullptr;
 						}
 						else
 						{
 							PrevComparisionFunction = ComparisonFunction;
-							ComparisonFunction = ComparisonFunction->InternalData->Next;
+							ComparisonFunction = ComparisonFunction->Next;
 						}
 					}
 					check(bFound); // otherwise you changed TickState while the tick function was registered. Call SetTickFunctionEnable instead.
@@ -1291,23 +1284,23 @@ public:
 					bFound = true;
 					if (PrevComparisonFunction)
 					{
-						PrevComparisonFunction->InternalData->Next = TickFunction->InternalData->Next;
+						PrevComparisonFunction->Next = TickFunction->Next;
 					}
 					else
 					{
 						check(TickFunction == AllCoolingDownTickFunctions.Head);
-						AllCoolingDownTickFunctions.Head = TickFunction->InternalData->Next;
+						AllCoolingDownTickFunctions.Head = TickFunction->Next;
 					}
-					if (TickFunction->InternalData->Next)
+					if (TickFunction->Next)
 					{
-						TickFunction->InternalData->Next->InternalData->RelativeTickCooldown += TickFunction->InternalData->RelativeTickCooldown;
-						TickFunction->InternalData->Next = nullptr;
+						TickFunction->Next->RelativeTickCooldown += TickFunction->RelativeTickCooldown;
+						TickFunction->Next = nullptr;
 					}
 				}
 				else
 				{
 					PrevComparisonFunction = ComparisonFunction;
-					ComparisonFunction = ComparisonFunction->InternalData->Next;
+					ComparisonFunction = ComparisonFunction->Next;
 				}
 			}
 			check(bFound); // otherwise you changed TickState while the tick function was registered. Call SetTickFunctionEnable instead.
@@ -1337,7 +1330,7 @@ private:
 				{
 					return true;
 				}
-				Node = Node->InternalData->Next;
+				Node = Node->Next;
 			}
 			return false;
 		}
@@ -1584,13 +1577,12 @@ public:
 		check(TickFunction->TickGroup >= 0 && TickFunction->TickGroup < TG_NewlySpawned); // You may not schedule a tick in the newly spawned group...they can only end up there if they are spawned late in a frame.
 		FTickTaskLevel* Level = TickTaskLevelForLevel(InLevel);
 		Level->AddTickFunction(TickFunction);
-		TickFunction->InternalData->TickTaskLevel = Level;
+		TickFunction->TickTaskLevel = Level;
 	}
 	/** Remove the tick function from the master list **/
 	void RemoveTickFunction(FTickFunction* TickFunction)
 	{
-		check(TickFunction->InternalData);
-		FTickTaskLevel* Level = TickFunction->InternalData->TickTaskLevel;
+		FTickTaskLevel* Level = TickFunction->TickTaskLevel;
 		check(Level);
 		Level->RemoveTickFunction(TickFunction);
 	}
@@ -1735,27 +1727,22 @@ private:
 FTickFunction::FTickFunction()
 	: TickGroup(TG_PrePhysics)
 	, EndTickGroup(TG_PrePhysics)
+	, ActualStartTickGroup(TG_PrePhysics)
+	, ActualEndTickGroup(TG_PrePhysics)
 	, bTickEvenWhenPaused(false)
 	, bCanEverTick(false)
 	, bStartWithTickEnabled(false)
 	, bAllowTickOnDedicatedServer(true)
 	, bHighPriority(false)
 	, bRunOnAnyThread(false)
+	, bRegistered(false)
+	, bWasInterval(false)
 	, TickState(ETickState::Enabled)
-	, TickInterval(0.f)
-{
-}
-
-FTickFunction::FInternalData::FInternalData()
-	: bWasInterval(false)
-	, ActualStartTickGroup(TG_PrePhysics)
-	, ActualEndTickGroup(TG_PrePhysics)
 	, TickVisitedGFrameCounter(0)
 	, TickQueuedGFrameCounter(0)
-	, TaskPointer(nullptr)
-	, Next(nullptr)
 	, RelativeTickCooldown(0.f)
 	, LastTickGameTimeSeconds(-1.f)
+	, TickInterval(0.f)
 	, TickTaskLevel(nullptr)
 {
 }
@@ -1773,14 +1760,14 @@ FTickFunction::~FTickFunction()
 **/
 void FTickFunction::RegisterTickFunction(ULevel* Level)
 {
-	if (InternalData == nullptr)
+	if (!bRegistered)
 	{
 		// Only allow registration of tick if we are are allowed on dedicated server, or we are not a dedicated server
 		const UWorld* World = Level ? Level->GetWorld() : nullptr;
 		if(bAllowTickOnDedicatedServer || !(World && World->IsNetMode(NM_DedicatedServer)))
 		{
-			InternalData.Reset(new FInternalData());
 			FTickTaskManager::Get().AddTickFunction(Level, this);
+			bRegistered = true;
 		}
 	}
 	else
@@ -1792,36 +1779,35 @@ void FTickFunction::RegisterTickFunction(ULevel* Level)
 /** Removes the tick function from the master list of tick functions. **/
 void FTickFunction::UnRegisterTickFunction()
 {
-	if (InternalData)
+	if (bRegistered)
 	{
 		FTickTaskManager::Get().RemoveTickFunction(this);
-		InternalData.Reset();
+		bRegistered = false;
 	}
 }
 
 /** Enables or disables this tick function. **/
 void FTickFunction::SetTickFunctionEnable(bool bInEnabled)
 {
-	if (InternalData)
+	if (bRegistered)
 	{
 		if (bInEnabled == (TickState == ETickState::Disabled))
 		{
-			FTickTaskLevel* TickTaskLevel = InternalData->TickTaskLevel;
 			check(TickTaskLevel);
 			TickTaskLevel->RemoveTickFunction(this);
 			TickState = (bInEnabled ? ETickState::Enabled : ETickState::Disabled);
 			TickTaskLevel->AddTickFunction(this);
 		}
-
-		if (TickState == ETickState::Disabled)
-		{
-			InternalData->LastTickGameTimeSeconds = -1.f;
-		}
 	}
 	else
 	{
 		TickState = (bInEnabled ? ETickState::Enabled : ETickState::Disabled);
 	}
+
+	if (TickState == ETickState::Disabled)
+	{
+		LastTickGameTimeSeconds = -1.f;
+	}
 }
 
 void FTickFunction::AddPrerequisite(UObject* TargetObject, struct FTickFunction& TargetTickFunction)
@@ -1869,8 +1855,8 @@ void FTickFunction::ShowPrerequistes(int32 Indent)
 
 FGraphEventRef FTickFunction::GetCompletionHandle() const
 {
-	check(InternalData->TaskPointer);
-	TGraphTask<FTickFunctionTask>* Task = (TGraphTask<FTickFunctionTask>*)InternalData->TaskPointer;
+	check(TaskPointer);
+	TGraphTask<FTickFunctionTask>* Task = (TGraphTask<FTickFunctionTask>*)TaskPointer;
 	return Task->GetCompletionEvent();
 }
 
@@ -1878,11 +1864,11 @@ FGraphEventRef FTickFunction::GetCompletionHandle() const
 void FTickFunction::QueueTickFunction(FTickTaskSequencer& TTS, const struct FTickContext& TickContext)
 {
 	checkSlow(TickContext.Thread == ENamedThreads::GameThread); // we assume same thread here
-	check(InternalData);
+	check(bRegistered);
 
-	if (InternalData->TickVisitedGFrameCounter != GFrameCounter)
+	if (TickVisitedGFrameCounter != GFrameCounter)
 	{
-		InternalData->TickVisitedGFrameCounter = GFrameCounter;
+		TickVisitedGFrameCounter = GFrameCounter;
 		if (TickState != FTickFunction::ETickState::Disabled)
 		{
 			ETickingGroup MaxPrerequisiteTickGroup =  ETickingGroup(0);
@@ -1896,22 +1882,22 @@ void FTickFunction::QueueTickFunction(FTickTaskSequencer& TTS, const struct FTic
 					// stale prereq, delete it
 					Prerequisites.RemoveAtSwap(PrereqIndex--);
 				}
-				else if (Prereq->InternalData)
+				else if (Prereq->bRegistered)
 				{
 					// recursive call to make sure my prerequisite is set up so I can use its completion handle
 					Prereq->QueueTickFunction(TTS, TickContext);
-					if (Prereq->InternalData->TickQueuedGFrameCounter != GFrameCounter)
+					if (Prereq->TickQueuedGFrameCounter != GFrameCounter)
 					{
 						// this must be up the call stack, therefore this is a cycle
 						UE_LOG(LogTick, Warning, TEXT("While processing prerequisites for %s, could use %s because it would form a cycle."),*DiagnosticMessage(), *Prereq->DiagnosticMessage());
 					}
-					else if (!Prereq->InternalData->TaskPointer)
+					else if (!Prereq->TaskPointer)
 					{
 						//ok UE_LOG(LogTick, Warning, TEXT("While processing prerequisites for %s, could use %s because it is disabled."),*DiagnosticMessage(), *Prereq->DiagnosticMessage());
 					}
 					else
 					{
-						MaxPrerequisiteTickGroup =  FMath::Max<ETickingGroup>(MaxPrerequisiteTickGroup, Prereq->InternalData->ActualStartTickGroup);
+						MaxPrerequisiteTickGroup =  FMath::Max<ETickingGroup>(MaxPrerequisiteTickGroup, Prereq->ActualStartTickGroup);
 						TaskPrerequisites.Add(Prereq->GetCompletionHandle());
 					}
 				}
@@ -1927,17 +1913,17 @@ void FTickFunction::QueueTickFunction(FTickTaskSequencer& TTS, const struct FTic
 					MyActualTickGroup = ETickingGroup(MyActualTickGroup + 1);
 				}
 			}
-			InternalData->ActualStartTickGroup = MyActualTickGroup;
-			InternalData->ActualEndTickGroup = MyActualTickGroup;
-			if (EndTickGroup > MyActualTickGroup)
+			ActualStartTickGroup = MyActualTickGroup;
+			ActualEndTickGroup = MyActualTickGroup;
+			if (EndTickGroup > ActualStartTickGroup)
 			{
 				check(EndTickGroup <= TG_NewlySpawned);
-				ETickingGroup TestTickGroup = ETickingGroup(MyActualTickGroup + 1);
+				ETickingGroup TestTickGroup = ETickingGroup(ActualEndTickGroup + 1);
 				while (TestTickGroup <= EndTickGroup)
-				{
+			{
 					if (CanDemoteIntoTickGroup(TestTickGroup))
 					{
-						InternalData->ActualEndTickGroup = TestTickGroup;
+						ActualEndTickGroup = TestTickGroup;
 					}
 					TestTickGroup = ETickingGroup(TestTickGroup + 1);
 				}
@@ -1948,24 +1934,25 @@ void FTickFunction::QueueTickFunction(FTickTaskSequencer& TTS, const struct FTic
 				TTS.QueueTickTask(&TaskPrerequisites, this, TickContext);
 			}
 		}
-		InternalData->TickQueuedGFrameCounter = GFrameCounter;
+		TickQueuedGFrameCounter = GFrameCounter;
 	}
 }
 
 void FTickFunction::QueueTickFunctionParallel(const struct FTickContext& TickContext, TArray<FTickFunction*, TInlineAllocator<8> >& StackForCycleDetection)
 {
+
 	bool bProcessTick;
 
-	int32 OldValue = *(volatile int32*)&InternalData->TickVisitedGFrameCounter;
+	int32 OldValue = *(volatile int32*)&TickVisitedGFrameCounter;
 	if (OldValue != GFrameCounter)
 	{
-		OldValue = FPlatformAtomics::InterlockedCompareExchange(&InternalData->TickVisitedGFrameCounter , GFrameCounter, OldValue);
+		OldValue = FPlatformAtomics::InterlockedCompareExchange(&TickVisitedGFrameCounter , GFrameCounter, OldValue);
 	}
 	bProcessTick = OldValue != GFrameCounter;
 
 	if (bProcessTick)
 	{
-		check(InternalData);
+		check(bRegistered);
 		if (TickState != FTickFunction::ETickState::Disabled)
 		{
 			ETickingGroup MaxPrerequisiteTickGroup = ETickingGroup(0);
@@ -1989,17 +1976,17 @@ void FTickFunction::QueueTickFunctionParallel(const struct FTickContext& TickCon
                     {
                         UE_LOG(LogTick, Warning, TEXT("While processing prerequisites for %s, could use %s because it would form a cycle."), *DiagnosticMessage(), *Prereq->DiagnosticMessage());
                     }
-					else if (Prereq->InternalData)
+					else if (Prereq->bRegistered)
 					{
 						// recursive call to make sure my prerequisite is set up so I can use its completion handle
 						Prereq->QueueTickFunctionParallel(TickContext, StackForCycleDetection);
-						if (!Prereq->InternalData->TaskPointer)
+						if (!Prereq->TaskPointer)
 						{
 							//ok UE_LOG(LogTick, Warning, TEXT("While processing prerequisites for %s, could use %s because it is disabled."),*DiagnosticMessage(), *Prereq->DiagnosticMessage());
 						}
 						else
 						{
-							MaxPrerequisiteTickGroup = FMath::Max<ETickingGroup>(MaxPrerequisiteTickGroup, Prereq->InternalData->ActualStartTickGroup);
+							MaxPrerequisiteTickGroup = FMath::Max<ETickingGroup>(MaxPrerequisiteTickGroup, Prereq->ActualStartTickGroup);
 							TaskPrerequisites.Add(Prereq->GetCompletionHandle());
 						}
 					}
@@ -2017,17 +2004,17 @@ void FTickFunction::QueueTickFunctionParallel(const struct FTickContext& TickCon
 					MyActualTickGroup = ETickingGroup(MyActualTickGroup + 1);
 				}
 			}
-			InternalData->ActualStartTickGroup = MyActualTickGroup;
-			InternalData->ActualEndTickGroup = MyActualTickGroup;
-			if (EndTickGroup > MyActualTickGroup)
+			ActualStartTickGroup = MyActualTickGroup;
+			ActualEndTickGroup = MyActualTickGroup;
+			if (EndTickGroup > ActualStartTickGroup)
 			{
 				check(EndTickGroup <= TG_NewlySpawned);
-				ETickingGroup TestTickGroup = ETickingGroup(MyActualTickGroup + 1);
+				ETickingGroup TestTickGroup = ETickingGroup(ActualEndTickGroup + 1);
 				while (TestTickGroup <= EndTickGroup)
 				{
 					if (CanDemoteIntoTickGroup(TestTickGroup))
 					{
-						InternalData->ActualEndTickGroup = TestTickGroup;
+						ActualEndTickGroup = TestTickGroup;
 					}
 					TestTickGroup = ETickingGroup(TestTickGroup + 1);
 				}
@@ -2036,25 +2023,25 @@ void FTickFunction::QueueTickFunctionParallel(const struct FTickContext& TickCon
 			if (TickState == FTickFunction::ETickState::Enabled)
 			{
 				FTickTaskSequencer::Get().QueueTickTaskParallel(&TaskPrerequisites, this, TickContext);
-				if (!InternalData->bWasInterval && TickInterval > 0.f)
+				if (!bWasInterval && TickInterval > 0.f)
 				{
-					InternalData->TickTaskLevel->RescheduleForIntervalParallel(this);
+					TickTaskLevel->RescheduleForIntervalParallel(this);
 				}
 			}
 		}
-		InternalData->bWasInterval = false;
+		bWasInterval = false;
 		
 		TSAN_BEFORE(&TickQueuedGFrameCounter);
 		FPlatformMisc::MemoryBarrier();
 		
 		// MSVC enforces acq/rel semantics on volatile values, but clang cannot (supports more backend architectures).
 		// consequently on ARM64 you would end up racing
-		FPlatformAtomics::InterlockedExchange(&InternalData->TickQueuedGFrameCounter, GFrameCounter);
+		FPlatformAtomics::InterlockedExchange(&TickQueuedGFrameCounter, GFrameCounter);
 	}
 	else
 	{
 		// if we are not going to process it, we need to at least wait until the other thread finishes it
-		volatile int32* TickQueuedGFrameCounterPtr = &InternalData->TickQueuedGFrameCounter;
+		volatile int32* TickQueuedGFrameCounterPtr = &TickQueuedGFrameCounter;
 		if (*TickQueuedGFrameCounterPtr != GFrameCounter)
 		{
 			QUICK_SCOPE_CYCLE_COUNTER(STAT_FTickFunction_QueueTickFunctionParallel_Spin);
@@ -2072,27 +2059,20 @@ float FTickFunction::CalculateDeltaTime(const FTickContext& TickContext)
 
 	if (TickInterval == 0.f)
 	{
-		if (InternalData)
-		{
-			// No tick interval. Return the world delta seconds, and make sure to mark that
-			// we're not tracking last-tick-time for this object.
-			InternalData->LastTickGameTimeSeconds = -1.f;
-		}
+		// No tick interval. Return the world delta seconds, and make sure to mark that
+		// we're not tracking last-tick-time for this object.
+		LastTickGameTimeSeconds = -1.f;
 	}
-	else if (InternalData)
+	else
 	{
 		// We've got a tick interval. Mark last-tick-time. If we already had last-tick-time, return
 		// the time since then; otherwise, return the world delta seconds.
 		const float CurrentWorldTime = (bTickEvenWhenPaused ? TickContext.World->GetUnpausedTimeSeconds() : TickContext.World->GetTimeSeconds());
-		if (InternalData->LastTickGameTimeSeconds >= 0.f)
+		if (LastTickGameTimeSeconds >= 0.f)
 		{
-			DeltaTimeForFunction = CurrentWorldTime - InternalData->LastTickGameTimeSeconds;
+			DeltaTimeForFunction = CurrentWorldTime - LastTickGameTimeSeconds;
 		}
-		InternalData->LastTickGameTimeSeconds = CurrentWorldTime;
-	}
-	else
-	{
-		UE_LOG(LogTick, Log, TEXT("Unregistered tick function calculating delta time cannot know what its last tick time was. Using tick context delta seconds."));
+		LastTickGameTimeSeconds = CurrentWorldTime;
 	}
 
 	return DeltaTimeForFunction;
