Fixed: Editor crash when mesh MinLod is equal or greater than total number of mesh LOD

diff --git a/Engine/Source/Runtime/Engine/Private/Components/StaticMeshComponent.cpp b/Engine/Source/Runtime/Engine/Private/Components/StaticMeshComponent.cpp
--- a/Engine/Source/Runtime/Engine/Private/Components/StaticMeshComponent.cpp
+++ b/Engine/Source/Runtime/Engine/Private/Components/StaticMeshComponent.cpp
@@ -2036,24 +2036,24 @@ bool UStaticMeshComponent::UsesTextureLightmaps(int32 InWidth, int32 InHeight) c
 		);
 }
 
-
 bool UStaticMeshComponent::HasLightmapTextureCoordinates() const
 {
-	if ((GetStaticMesh() != NULL) &&
-		(GetStaticMesh()->LightMapCoordinateIndex >= 0) &&
-		(GetStaticMesh()->RenderData != NULL) &&
-		(GetStaticMesh()->RenderData->LODResources.Num() > 0) &&
-		(GetStaticMesh()->LightMapCoordinateIndex >= 0))
+	const UStaticMesh* Mesh = GetStaticMesh();
+	if (Mesh != NULL &&
+		Mesh->LightMapCoordinateIndex >= 0 &&
+		Mesh->RenderData != NULL &&
+		Mesh->RenderData->LODResources.Num() > 0 &&
+		Mesh->LightMapCoordinateIndex >= 0)
 	{
 		const ERHIFeatureLevel::Type FeatureLevel = GetScene() ? GetScene()->GetFeatureLevel() : GMaxRHIFeatureLevel;
-		int32 MeshMinLOD = GetStaticMesh()->MinLOD.GetValueForFeatureLevel(FeatureLevel);
+		int32 MeshMinLOD = Mesh->MinLOD.GetValueForFeatureLevel(FeatureLevel);
+		MeshMinLOD = FMath::Min(MeshMinLOD,  Mesh->RenderData->LODResources.Num() - 1);
 		
-		return ((uint32)GetStaticMesh()->LightMapCoordinateIndex < GetStaticMesh()->RenderData->LODResources[MeshMinLOD].VertexBuffers.StaticMeshVertexBuffer.GetNumTexCoords());
+		return ((uint32)Mesh->LightMapCoordinateIndex < Mesh->RenderData->LODResources[MeshMinLOD].VertexBuffers.StaticMeshVertexBuffer.GetNumTexCoords());
 	}
 	return false;
 }
 
-
 void UStaticMeshComponent::GetTextureLightAndShadowMapMemoryUsage(int32 InWidth, int32 InHeight, int32& OutLightMapMemoryUsage, int32& OutShadowMapMemoryUsage) const
 {
 	// Stored in texture.
