Duplicating CL#7614678 from //UE4/Dev-Anim to fix crash in morph target end of frame updates

Morphtarget refactor w.r.t. register/unreigster (when morph target size changes)

diff --git a/Engine/Source/Runtime/Engine/Classes/Components/SkeletalMeshComponent.h b/Engine/Source/Runtime/Engine/Classes/Components/SkeletalMeshComponent.h
--- a/Engine/Source/Runtime/Engine/Classes/Components/SkeletalMeshComponent.h
+++ b/Engine/Source/Runtime/Engine/Classes/Components/SkeletalMeshComponent.h
@@ -291,7 +291,7 @@ class ENGINE_API USkeletalMeshComponent : public USkinnedMeshComponent, public I
 
 	friend class FSkinnedMeshComponentRecreateRenderStateContext;
 	friend class FParallelAnimationCompletionTask;
-	
+	friend class USkeletalMesh; 
 	/**
 	 * Animation 
 	 */
@@ -2099,6 +2099,7 @@ private:
 	 * Update MorphTargetCurves from mesh - these are not animation curves, but SetMorphTarget and similar functions that can set to this mesh component
 	 */
 	void UpdateMorphTargetOverrideCurves();
+
 	/*
 	 * Reset MorphTarget Curves - Reset all morphtarget curves
 	 */
diff --git a/Engine/Source/Runtime/Engine/Private/Animation/AnimInstance.cpp b/Engine/Source/Runtime/Engine/Private/Animation/AnimInstance.cpp
--- a/Engine/Source/Runtime/Engine/Private/Animation/AnimInstance.cpp
+++ b/Engine/Source/Runtime/Engine/Private/Animation/AnimInstance.cpp
@@ -1217,7 +1217,7 @@ void UAnimInstance::UpdateCurvesToComponents(USkeletalMeshComponent* Component /
 		// this is only any thread because EndOfFrameUpdate update can restart render state
 		// and this needs to restart from worker thread
 		// EndOfFrameUpdate is done after all tick is updated, so in theory you shouldn't have 
-		FAnimInstanceProxy& Proxy = GetProxyOnAnyThread<FAnimInstanceProxy>();
+		FAnimInstanceProxy& Proxy = GetProxyOnGameThread<FAnimInstanceProxy>();
 		Component->ApplyAnimationCurvesToComponent(&Proxy.GetAnimationCurves(EAnimCurveType::MaterialCurve), &Proxy.GetAnimationCurves(EAnimCurveType::MorphTargetCurve));
 	}
 }
diff --git a/Engine/Source/Runtime/Engine/Private/Components/SkinnedMeshComponent.cpp b/Engine/Source/Runtime/Engine/Private/Components/SkinnedMeshComponent.cpp
--- a/Engine/Source/Runtime/Engine/Private/Components/SkinnedMeshComponent.cpp
+++ b/Engine/Source/Runtime/Engine/Private/Components/SkinnedMeshComponent.cpp
@@ -586,9 +586,9 @@ void USkinnedMeshComponent::CreateRenderState_Concurrent()
 				const bool bMorphTargetsAllowed = CVarEnableMorphTargets.GetValueOnAnyThread(true) != 0;
 
 				// Are morph targets disabled for this LOD?
-				if (!bDisableMorphTarget && bMorphTargetsAllowed)
+				if (bDisableMorphTarget || !bMorphTargetsAllowed)
 				{
-					RefreshMorphTargets();
+					ActiveMorphTargets.Empty();
 				}
 
 				MeshObject->Update(PredictedLODLevel, this, ActiveMorphTargets, MorphTargetWeights, EPreviousBoneTransformUpdateMode::UpdatePrevious);  // send to rendering thread
@@ -604,10 +604,6 @@ void USkinnedMeshComponent::DestroyRenderState_Concurrent()
 {
 	Super::DestroyRenderState_Concurrent();
 
-	// clear morphtarget array info while rendering state is destroyed
-	ActiveMorphTargets.Empty();
-	MorphTargetWeights.Empty();
-
 	if(MeshObject)
 	{
 		// Begin releasing the RHI resources used by this skeletal mesh component.
diff --git a/Engine/Source/Runtime/Engine/Private/SkeletalMesh.cpp b/Engine/Source/Runtime/Engine/Private/SkeletalMesh.cpp
--- a/Engine/Source/Runtime/Engine/Private/SkeletalMesh.cpp
+++ b/Engine/Source/Runtime/Engine/Private/SkeletalMesh.cpp
@@ -2197,6 +2197,16 @@ void USkeletalMesh::InitMorphTargetsAndRebuildRenderData()
 	MarkPackageDirty();
 	// need to refresh the map
 	InitMorphTargets();
+
+	// reset all morphtarget for all components
+	for (TObjectIterator<USkeletalMeshComponent> It; It; ++It)
+	{
+		if (It->SkeletalMesh == this)
+		{
+			It->RefreshMorphTargets();
+		}
+	}
+
 	// invalidate render data
 	InvalidateRenderData();
 }
diff --git a/Engine/Source/Runtime/Engine/Private/SkeletalRenderGPUSkin.cpp b/Engine/Source/Runtime/Engine/Private/SkeletalRenderGPUSkin.cpp
--- a/Engine/Source/Runtime/Engine/Private/SkeletalRenderGPUSkin.cpp
+++ b/Engine/Source/Runtime/Engine/Private/SkeletalRenderGPUSkin.cpp
@@ -485,6 +485,9 @@ void FSkeletalMeshObjectGPUSkin::ProcessUpdatedDynamicData(FGPUSkinCache* GPUSki
 			QUICK_SCOPE_CYCLE_COUNTER(STAT_FSkeletalMeshObjectGPUSkin_ProcessUpdatedDynamicData_UpdateMorphBuffer);
 			if (UseGPUMorphTargets(GMaxRHIShaderPlatform))
 			{
+				// sometimes this goes out of bound, we'll ensure here
+				ensureAlways(DynamicData->MorphTargetWeights.Num() == LODData.MorphTargetVertexInfoBuffers.GetNumMorphs());
+
 				// update the morph data for the lod (before SkinCache)
 				TArray<float> MorphTargetWeights;
 				MorphTargetWeights.Reserve(LODData.MorphTargetVertexInfoBuffers.GetNumMorphs());
@@ -495,7 +498,6 @@ void FSkeletalMeshObjectGPUSkin::ProcessUpdatedDynamicData(FGPUSkinCache* GPUSki
 						MorphTargetWeights.Add(DynamicData->MorphTargetWeights[i]);
 					}
 				}
-				ensureAlways(MorphTargetWeights.Num() == LODData.MorphTargetVertexInfoBuffers.GetNumMorphs());
 				LOD.UpdateMorphVertexBufferGPU(RHICmdList, MorphTargetWeights, LODData.MorphTargetVertexInfoBuffers, DynamicData->SectionIdsUseByActiveMorphTargets);
 			}
 			else
