[NO NON-PII COMMIT MESSAGE]

diff --git a/Engine/Source/Runtime/Engine/Private/AudioStreamingCache.cpp b/Engine/Source/Runtime/Engine/Private/AudioStreamingCache.cpp
--- a/Engine/Source/Runtime/Engine/Private/AudioStreamingCache.cpp
+++ b/Engine/Source/Runtime/Engine/Private/AudioStreamingCache.cpp
@@ -613,10 +613,8 @@ FAudioChunkCache::FCacheElement* FAudioChunkCache::FindElementForKey(const FChun
 		{
 			
 #if DEBUG_STREAM_CACHE
-			{
-				FScopeLock ScopeLock(&CurrentElement->DebugInfo.PreviousLocationsLock);
-				CurrentElement->DebugInfo.PreviousLocationsBeforeBeingTouched.Add(ElementPosition);
-			}
+			float& CMA = CurrentElement->DebugInfo.AverageLocationInCacheWhenNeeded;
+			CMA += ((ElementPosition - CMA) / (CurrentElement->DebugInfo.NumTimesTouched + 1));
 #endif
 
 			return CurrentElement;
@@ -801,12 +799,8 @@ FAudioChunkCache::FCacheElement* FAudioChunkCache::EvictLeastRecentChunk()
 	}
 
 #if DEBUG_STREAM_CACHE
-	{
-		FScopeLock ScopeLock(&CacheElement->DebugInfo.PreviousLocationsLock);
-		// Reset debug information:
-		CacheElement->DebugInfo.PreviousLocationsBeforeBeingTouched.Reset();
-		CacheElement->DebugInfo.NumTimesTouched = 0;
-	}
+	// Reset debug information:
+	CacheElement->DebugInfo.Reset();
 #endif
 
 	return CacheElement;
@@ -1072,21 +1066,24 @@ TPair<int, int> FAudioChunkCache::DebugDisplay(UWorld* World, FViewport* Viewpor
 		int32 NumTotalChunks = -1;
 		int32 NumTimesTouched = -1;
 		double TimeToLoad = -1.0;
+		float AveragePlaceInCache = -1.0f;
 
 #if DEBUG_STREAM_CACHE
 		NumTotalChunks = CurrentElement->DebugInfo.NumTotalChunks;
 		NumTimesTouched = CurrentElement->DebugInfo.NumTimesTouched;
 		TimeToLoad = CurrentElement->DebugInfo.TimeToLoad;
+		AveragePlaceInCache = CurrentElement->DebugInfo.AverageLocationInCacheWhenNeeded;
 #endif
 
 		const bool bWasTrimmed = CurrentElement->ChunkData.Num() == 0;
 
-		FString ElementInfo = *FString::Printf(TEXT("%4i. Size: %6.2f KB   Chunk: %d of %d   Request Count: %d    Number of Decoders Using Chunk: %d     Chunk Load Time: %6.4fms      Name: %s "),
+		FString ElementInfo = *FString::Printf(TEXT("%4i. Size: %6.2f KB   Chunk: %d of %d   Request Count: %d    Average Index: %6.2f  Number of Decoders Using Chunk: %d     Chunk Load Time: %6.4fms      Name: %s "),
 			Index,
 			CurrentElement->ChunkData.Num() / 1024.0f,
 			CurrentElement->Key.ChunkIndex,
 			NumTotalChunks,
 			NumTimesTouched,
+			AveragePlaceInCache,
 			CurrentElement->NumConsumers.GetValue(),
 			TimeToLoad,
 			bWasTrimmed ? TEXT("TRIMMED CHUNK") : *CurrentElement->Key.SoundWaveName.ToString()
diff --git a/Engine/Source/Runtime/Engine/Public/AudioStreamingCache.h b/Engine/Source/Runtime/Engine/Public/AudioStreamingCache.h
--- a/Engine/Source/Runtime/Engine/Public/AudioStreamingCache.h
+++ b/Engine/Source/Runtime/Engine/Public/AudioStreamingCache.h
@@ -104,16 +104,15 @@ private:
 		// Amount of time spent loading the audio file.
 		double TimeToLoad;
 
-		// This allows every call to AddOrTouchChunk to report how far down the cache
-		// it was when it was needed, or INDEX_NONE if it wasn't in the cache when it was needed.
-		TArray<int32> PreviousLocationsBeforeBeingTouched;
-		FCriticalSection PreviousLocationsLock;
+		// This is a cumulative moving average of a chunks location before it was 
+		float AverageLocationInCacheWhenNeeded;
 
 		FCacheElementDebugInfo()
 			: NumTotalChunks(0)
 			, NumTimesTouched(0)
 			, TimeLoadStarted(0.0)
 			, TimeToLoad(0.0)
+			, AverageLocationInCacheWhenNeeded(0.0f)
 		{
 		}
 
@@ -123,8 +122,7 @@ private:
 			NumTimesTouched = 0;
 			TimeLoadStarted = 0;
 
-			FScopeLock ScopeLock(&PreviousLocationsLock);
-			PreviousLocationsBeforeBeingTouched.Empty();
+			AverageLocationInCacheWhenNeeded = 0.0f;
 		}
 	};
 #endif
