Add vertex factory dependency to material shader map hash computation
This avoids cases where shader maps for different vertex factories could end up with same ShaderMapId

diff --git a/Engine/Source/Runtime/Engine/Private/Materials/MaterialShader.cpp b/Engine/Source/Runtime/Engine/Private/Materials/MaterialShader.cpp
--- a/Engine/Source/Runtime/Engine/Private/Materials/MaterialShader.cpp
+++ b/Engine/Source/Runtime/Engine/Private/Materials/MaterialShader.cpp
@@ -624,6 +624,11 @@ void FMaterialShaderMapId::GetMaterialHash(FSHAHash& OutHash) const
 		HashState.Update((const uint8*)&ReferencedParameterCollections[CollectionIndex], sizeof(ReferencedParameterCollections[CollectionIndex]));
 	}
 
+	for (int32 VertexFactoryIndex = 0; VertexFactoryIndex < VertexFactoryTypeDependencies.Num(); VertexFactoryIndex++)
+	{
+		HashState.Update((const uint8*)&VertexFactoryTypeDependencies[VertexFactoryIndex].VFSourceHash, sizeof(VertexFactoryTypeDependencies[VertexFactoryIndex].VFSourceHash));
+	}
+
 	HashState.Update((const uint8*)&TextureReferencesHash, sizeof(TextureReferencesHash));
 
 	HashState.Update((const uint8*)&BasePropertyOverridesHash, sizeof(BasePropertyOverridesHash));
diff --git a/Engine/Source/Runtime/Engine/Private/ShaderDerivedDataVersion.h b/Engine/Source/Runtime/Engine/Private/ShaderDerivedDataVersion.h
--- a/Engine/Source/Runtime/Engine/Private/ShaderDerivedDataVersion.h
+++ b/Engine/Source/Runtime/Engine/Private/ShaderDerivedDataVersion.h
@@ -11,4 +11,4 @@
 // In case of merge conflicts with DDC versions, you MUST generate a new GUID and set this new
 // guid as version
 #define GLOBALSHADERMAP_DERIVEDDATA_VER			TEXT("45EDF9DF531B4E7989EB97AF45C38B3E")
-#define MATERIALSHADERMAP_DERIVEDDATA_VER		TEXT("36C14CFABCD54E0E846080318C712DCE")
+#define MATERIALSHADERMAP_DERIVEDDATA_VER		TEXT("D66CBE081B6849509E4D36A915595C28")
