Undo woakround in CL# 7841190

diff --git a/Engine/Source/Runtime/Engine/Private/Materials/MaterialShared.cpp b/Engine/Source/Runtime/Engine/Private/Materials/MaterialShared.cpp
--- a/Engine/Source/Runtime/Engine/Private/Materials/MaterialShared.cpp
+++ b/Engine/Source/Runtime/Engine/Private/Materials/MaterialShared.cpp
@@ -1789,7 +1789,7 @@ bool FMaterial::CacheShaders(const FMaterialShaderMapId& ShaderMapId, EShaderPla
 
 
 	bool bAssumeShaderMapIsComplete = false;
-#if UE_BUILD_SHIPPING
+#if UE_BUILD_SHIPPING || UE_BUILD_TEST
 	bAssumeShaderMapIsComplete = (bContainsInlineShaders || FPlatformProperties::RequiresCookedData()) 
 		&& !bLogShaderMapFailInfo; // if it is the special engine material, we will check it
 #endif
