Exposed custom primitive data on the primitive component to be edited in the detail panel.

diff --git a/Engine/Source/Runtime/Engine/Classes/Components/PrimitiveComponent.h b/Engine/Source/Runtime/Engine/Classes/Components/PrimitiveComponent.h
--- a/Engine/Source/Runtime/Engine/Classes/Components/PrimitiveComponent.h
+++ b/Engine/Source/Runtime/Engine/Classes/Components/PrimitiveComponent.h
@@ -542,7 +542,7 @@ public:
 
 private:
 	/** Custom data that can be read by a material through a material parameter expression. Set data using SetCustomPrimitiveData* functions */
-	UPROPERTY()
+	UPROPERTY(EditAnywhere, AdvancedDisplay, Category=Rendering)
 	FCustomPrimitiveData CustomPrimitiveData;
 
 public:
@@ -762,15 +762,15 @@ public:
 	UFUNCTION(BlueprintCallable, Category="Rendering|Material")
 	void SetCustomPrimitiveDataFloat(int32 DataIndex, float Value);
 
-	/** Set custom primitive data, two floats at once, from index DataIndex to index DataIndex + 2. */
+	/** Set custom primitive data, two floats at once, from index DataIndex to index DataIndex + 1. */
 	UFUNCTION(BlueprintCallable, Category="Rendering|Material")
 	void SetCustomPrimitiveDataVector2(int32 DataIndex, FVector2D Value);
 
-	/** Set custom primitive data, three floats at once, from index DataIndex to index DataIndex + 3. */
+	/** Set custom primitive data, three floats at once, from index DataIndex to index DataIndex + 2. */
 	UFUNCTION(BlueprintCallable, Category="Rendering|Material")
 	void SetCustomPrimitiveDataVector3(int32 DataIndex, FVector Value);
 
-	/** Set custom primitive data, four floats at once, from index DataIndex to index DataIndex + 4. */
+	/** Set custom primitive data, four floats at once, from index DataIndex to index DataIndex + 3. */
 	UFUNCTION(BlueprintCallable, Category="Rendering|Material")
 	void SetCustomPrimitiveDataVector4(int32 DataIndex, FVector4 Value);
 
diff --git a/Engine/Source/Runtime/Engine/Private/Components/PrimitiveComponent.cpp b/Engine/Source/Runtime/Engine/Private/Components/PrimitiveComponent.cpp
--- a/Engine/Source/Runtime/Engine/Private/Components/PrimitiveComponent.cpp
+++ b/Engine/Source/Runtime/Engine/Private/Components/PrimitiveComponent.cpp
@@ -916,6 +916,18 @@ void UPrimitiveComponent::PostEditChangeProperty(FPropertyChangedEvent& Property
 		{
 			MarkRenderStateDirty();
 		}
+
+	}
+
+	if (UProperty* MemberPropertyThatChanged = PropertyChangedEvent.MemberProperty)
+	{
+		const FName MemberPropertyName = MemberPropertyThatChanged->GetFName();
+
+		// Reregister to get the custom primitive data to the proxy
+		if (MemberPropertyName == GET_MEMBER_NAME_CHECKED(UPrimitiveComponent, CustomPrimitiveData))
+		{
+			MarkRenderStateDirty();
+		}
 	}
 
 	if (LightmapType == ELightmapType::ForceSurface && GetStaticLightingType() == LMIT_None)
