Fix padding in FPrimitiveSceneProxy and change RuntimeVirtualTextureMaterialTypes from TSet to TArray (416 bytes w/ 16 bytes padding -> 336 bytes w/ 4 bytes padding)
Also reduce FSkeltalMeshSceneProxy an additional 16 bytes.
#rnx

diff --git a/Engine/Source/Runtime/Engine/Private/PrimitiveSceneProxy.cpp b/Engine/Source/Runtime/Engine/Private/PrimitiveSceneProxy.cpp
--- a/Engine/Source/Runtime/Engine/Private/PrimitiveSceneProxy.cpp
+++ b/Engine/Source/Runtime/Engine/Private/PrimitiveSceneProxy.cpp
@@ -145,11 +145,11 @@ FPrimitiveSceneProxy::FPrimitiveSceneProxy(const UPrimitiveComponent* InComponen
 ,	CustomDepthStencilWriteMask(FRendererStencilMaskEvaluation::ToStencilMask(InComponent->CustomDepthStencilWriteMask))
 ,	LightingChannelMask(GetLightingChannelMaskForStruct(InComponent->LightingChannels))
 ,	IndirectLightingCacheQuality(InComponent->IndirectLightingCacheQuality)
-,	LpvBiasMultiplier(InComponent->LpvBiasMultiplier)
-,	DynamicIndirectShadowMinVisibility(0)
 ,	VirtualTextureLodBias(InComponent->VirtualTextureLodBias)
 ,	VirtualTextureCullMips(InComponent->VirtualTextureCullMips)
 ,	VirtualTextureMinCoverage(InComponent->VirtualTextureMinCoverage)
+,	LpvBiasMultiplier(InComponent->LpvBiasMultiplier)
+,	DynamicIndirectShadowMinVisibility(0)
 ,	PrimitiveComponentId(InComponent->ComponentId)
 ,	Scene(InComponent->GetScene())
 ,	PrimitiveSceneInfo(NULL)
@@ -227,7 +227,7 @@ FPrimitiveSceneProxy::FPrimitiveSceneProxy(const UPrimitiveComponent* InComponen
 			if (VirtualTexture != nullptr && VirtualTexture->GetEnabled())
 			{
 				RuntimeVirtualTextures.Add(VirtualTexture);
-				RuntimeVirtualTextureMaterialTypes.Add(VirtualTexture->GetMaterialType());
+				RuntimeVirtualTextureMaterialTypes.AddUnique(VirtualTexture->GetMaterialType());
 			
 				//todo[vt]: Only flush this specific virtual texture
 				//todo[vt]: Only flush primitive bounds 
diff --git a/Engine/Source/Runtime/Engine/Public/PrimitiveSceneProxy.h b/Engine/Source/Runtime/Engine/Public/PrimitiveSceneProxy.h
--- a/Engine/Source/Runtime/Engine/Public/PrimitiveSceneProxy.h
+++ b/Engine/Source/Runtime/Engine/Public/PrimitiveSceneProxy.h
@@ -923,7 +923,7 @@ protected:
 	 * Whether this proxy ever draws with vertex factories that require a primitive uniform buffer. 
 	 * When false, updating the primitive uniform buffer can be skipped since vertex factories always use GPUScene instead.
 	 */
-	uint32 bVFRequiresPrimitiveUniformBuffer : 1;
+	uint8 bVFRequiresPrimitiveUniformBuffer : 1;
 
 	/** Whether the primitive should always be considered to have velocities, even if it hasn't moved. */
 	uint8 bAlwaysHasVelocity : 1;
@@ -978,6 +978,13 @@ protected:
 	/** Quality of interpolated indirect lighting for Movable components. */
 	TEnumAsByte<EIndirectLightingCacheQuality> IndirectLightingCacheQuality;
 
+	/** Geometry Lod bias when rendering to runtime virtual texture. */
+	int8 VirtualTextureLodBias;
+	/** Number of low mips to skip when rendering to runtime virtual texture. */
+	int8 VirtualTextureCullMips;
+	/** Log2 of minimum estimated pixel coverage before culling from runtime virtual texture. */
+	int8 VirtualTextureMinCoverage;
+
 	/** The bias applied to LPV injection */
 	float LpvBiasMultiplier;
 
@@ -989,14 +996,7 @@ protected:
 	/** Array of runtime virtual textures that this proxy should render to. */
 	TArray<URuntimeVirtualTexture*> RuntimeVirtualTextures;
 	/** Set of unique runtime virtual texture material types referenced by RuntimeVirtualTextures. */
-	TSet<ERuntimeVirtualTextureMaterialType> RuntimeVirtualTextureMaterialTypes;
-
-	/** Geometry Lod bias when rendering to runtime virtual texture. */
-	int8 VirtualTextureLodBias;
-	/** Number of low mips to skip when rendering to runtime virtual texture. */
-	int8 VirtualTextureCullMips;
-	/** Log2 of minimum estimated pixel coverage before culling from runtime virtual texture. */
-	int8 VirtualTextureMinCoverage;
+	TArray<ERuntimeVirtualTextureMaterialType> RuntimeVirtualTextureMaterialTypes;
 
 private:
 	/** The hierarchy of owners of this primitive.  These must not be dereferenced on the rendering thread, but the pointer values can be used for identification.  */
diff --git a/Engine/Source/Runtime/Engine/Public/SkeletalMeshTypes.h b/Engine/Source/Runtime/Engine/Public/SkeletalMeshTypes.h
--- a/Engine/Source/Runtime/Engine/Public/SkeletalMeshTypes.h
+++ b/Engine/Source/Runtime/Engine/Public/SkeletalMeshTypes.h
@@ -301,10 +301,6 @@ public:
 	const TArray<FMatrix>& GetMeshObjectReferenceToLocalMatrices() const;
 	const TIndirectArray<FSkeletalMeshLODRenderData>& GetSkeletalMeshRenderDataLOD() const;
 
-#if RHI_RAYTRACING
-	bool bAnySegmentUsesWorldPositionOffset : 1;
-#endif
-
 protected:
 	AActor* Owner;
 	class FSkeletalMeshObject* MeshObject;
@@ -314,6 +310,12 @@ protected:
 	const USkeletalMesh* SkeletalMeshForDebug;
 	class UPhysicsAsset* PhysicsAssetForDebug;
 
+public:
+#if RHI_RAYTRACING
+	bool bAnySegmentUsesWorldPositionOffset : 1;
+#endif
+
+protected:
 	/** data copied for rendering */
 	uint8 bForceWireframe : 1;
 	uint8 bIsCPUSkinned : 1;
