Fix memory leak when not running the audio thread.

diff --git a/Engine/Source/Runtime/Engine/Private/AudioDevice.cpp b/Engine/Source/Runtime/Engine/Private/AudioDevice.cpp
--- a/Engine/Source/Runtime/Engine/Private/AudioDevice.cpp
+++ b/Engine/Source/Runtime/Engine/Private/AudioDevice.cpp
@@ -3824,10 +3824,8 @@ void FAudioDevice::Update(bool bGameTicking)
 {
 	LLM_SCOPE(ELLMTag::AudioMisc);
 
-	if (!IsInAudioThread())
+	if (IsInGameThread())
 	{
-		check(IsInGameThread());
-
 		// On game thread, look through registered sound waves and remove if we finished precaching (and audio decompressor is cleaned up)
 		// ReferencedSoundWaves is used to make sure GC doesn't run on any sound waves that are actively pre-caching within an async task.
 		// Sounds may be loaded, kick off an async task to decompress, but never actually try to play, so GC can reclaim these while precaches are in-flight.
@@ -3839,6 +3837,11 @@ void FAudioDevice::Update(bool bGameTicking)
 				ReferencedSoundWaves.RemoveAtSwap(i, 1, false);
 			}
 		}
+	}
+
+	if (!IsInAudioThread())
+	{
+		check(IsInGameThread());
 
 		FAudioDevice* AudioDevice = this;
 		FAudioThread::RunCommandOnAudioThread([AudioDevice, bGameTicking]()
