Added Cvar to allow reverb processing on Sound Sources with channel count > 2.  The Cvar is: au.AllowReverbForMultichannelSources
(default behavior is no change)

diff --git a/Engine/Source/Runtime/Engine/Private/Audio.cpp b/Engine/Source/Runtime/Engine/Private/Audio.cpp
--- a/Engine/Source/Runtime/Engine/Private/Audio.cpp
+++ b/Engine/Source/Runtime/Engine/Private/Audio.cpp
@@ -78,6 +78,14 @@ FAutoConsoleVariableRef CVarAllowAudioSpatializationCVar(
 	TEXT("0: Disable, >0: Enable"),
 	ECVF_Default);
 
+static int32 AllowReverbForMultichannelSources = 0;
+FAutoConsoleVariableRef CvarAllowReverbForMultichannelSources(
+	TEXT("au.AllowReverbForMultichannelSources"),
+	AllowReverbForMultichannelSources,
+	TEXT("Controls if we allow Reverb processing for sources with channel counts > 2.\n")
+	TEXT("0: Disable, >0: Enable"),
+	ECVF_Default);
+
 
 bool IsAudioPluginEnabled(EAudioPlugin PluginType)
 {
@@ -314,7 +322,7 @@ bool FSoundSource::SetReverbApplied(bool bHardwareAvailable)
 	}
 
 	// Do not apply reverb to multichannel sounds
-	if (WaveInstance->WaveData->NumChannels > 2)
+	if (!AllowReverbForMultichannelSources && (WaveInstance->WaveData->NumChannels > 2))
 	{
 		bReverbApplied = false;
 	}
