Added a way to set default values for all water body properties that will propagate to all spline points on bodies in the world that have not been modifed

diff --git a/Engine/Source/Runtime/Engine/Classes/Components/SplineComponent.h b/Engine/Source/Runtime/Engine/Classes/Components/SplineComponent.h
--- a/Engine/Source/Runtime/Engine/Classes/Components/SplineComponent.h
+++ b/Engine/Source/Runtime/Engine/Classes/Components/SplineComponent.h
@@ -49,7 +49,7 @@ public:
 	virtual void DuplicatePoint(int32 Index) PURE_VIRTUAL(USplineMetadata::DuplicatePoint, );
 	virtual void CopyPoint(const USplineMetadata* FromSplineMetadata, int32 FromIndex, int32 ToIndex) PURE_VIRTUAL(USplineMetadata::CopyPoint, );
 	virtual void Reset(int32 NumPoints) PURE_VIRTUAL(USplineMetadata::Reset, );
-	virtual void Fixup(int32 NumPoints) PURE_VIRTUAL(USplineMetadata::Fixup, );
+	virtual void Fixup(int32 NumPoints, USplineComponent* SplineComp) PURE_VIRTUAL(USplineMetadata::Fixup, );
 };
 
 USTRUCT()
