record low memory warning data in the OOM crash context so we can better assess the OOM detector accuracy

diff --git a/Engine/Source/Runtime/Core/Private/IOS/IOSPlatformMisc.cpp b/Engine/Source/Runtime/Core/Private/IOS/IOSPlatformMisc.cpp
--- a/Engine/Source/Runtime/Core/Private/IOS/IOSPlatformMisc.cpp
+++ b/Engine/Source/Runtime/Core/Private/IOS/IOSPlatformMisc.cpp
@@ -694,8 +694,14 @@ void FIOSPlatformMisc::SetMemoryWarningHandler(void (* InHandler)(const FGeneric
 	GMemoryWarningHandler = InHandler;
 }
 
+bool FIOSPlatformMisc::HasMemoryWarningHandler()
+{
+	return GMemoryWarningHandler != nullptr;
+}
+
 void FIOSPlatformMisc::HandleLowMemoryWarning()
 {
+	UE_LOG(LogInit, Log, TEXT("Low Memory Warning Triggered"));
 	UE_LOG(LogInit, Log, TEXT("Free Memory at Startup: %d MB"), GStartupFreeMemoryMB);
 	UE_LOG(LogInit, Log, TEXT("Free Memory Now       : %d MB"), GetFreeMemoryMB());
 
diff --git a/Engine/Source/Runtime/Core/Public/GenericPlatform/GenericPlatformMisc.h b/Engine/Source/Runtime/Core/Public/GenericPlatform/GenericPlatformMisc.h
--- a/Engine/Source/Runtime/Core/Public/GenericPlatform/GenericPlatformMisc.h
+++ b/Engine/Source/Runtime/Core/Public/GenericPlatform/GenericPlatformMisc.h
@@ -525,6 +525,14 @@ struct CORE_API FGenericPlatformMisc
 	{
 	}
 
+	/**
+	 * Determines if a warning handler has been set
+	 */
+	static bool HasMemoryWarningHandler()
+	{
+		return false;
+	}
+	
 	FORCEINLINE static uint32 GetLastError()
 	{
 		return 0;
diff --git a/Engine/Source/Runtime/Core/Public/IOS/IOSPlatformMisc.h b/Engine/Source/Runtime/Core/Public/IOS/IOSPlatformMisc.h
--- a/Engine/Source/Runtime/Core/Public/IOS/IOSPlatformMisc.h
+++ b/Engine/Source/Runtime/Core/Public/IOS/IOSPlatformMisc.h
@@ -34,6 +34,7 @@ struct CORE_API FIOSPlatformMisc : public FApplePlatformMisc
 
 	static EAppReturnType::Type MessageBoxExt( EAppMsgType::Type MsgType, const TCHAR* Text, const TCHAR* Caption );
 	static void SetMemoryWarningHandler(void (* Handler)(const FGenericMemoryWarningContext& Context));
+	static bool HasMemoryWarningHandler();
 	static bool HasPlatformFeature(const TCHAR* FeatureName);
 	static bool SetStoredValue(const FString& InStoreId, const FString& InSectionName, const FString& InKeyName, const FString& InValue);
 	static bool GetStoredValue(const FString& InStoreId, const FString& InSectionName, const FString& InKeyName, FString& OutValue);
diff --git a/Engine/Source/Runtime/Engine/Private/UnrealEngine.cpp b/Engine/Source/Runtime/Engine/Private/UnrealEngine.cpp
--- a/Engine/Source/Runtime/Engine/Private/UnrealEngine.cpp
+++ b/Engine/Source/Runtime/Engine/Private/UnrealEngine.cpp
@@ -1097,7 +1097,7 @@ Engine init and exit.
 /** Callback from OS when we get a low memory warning.
 * Note: might not be called from the game thread
 */
-void EngineMemoryWarningHandler(const FGenericMemoryWarningContext& GenericContext)
+static void EngineMemoryWarningHandler(const FGenericMemoryWarningContext& GenericContext)
 {
 	FPlatformMemoryStats Stats = FPlatformMemory::GetStats();
 
@@ -1389,7 +1389,10 @@ void UEngine::Init(IEngineLoop* InEngineLoop)
 #endif
 
 	// Set the memory warning handler
-	FPlatformMisc::SetMemoryWarningHandler(EngineMemoryWarningHandler);
+	if (!FPlatformMisc::HasMemoryWarningHandler())
+	{
+		FPlatformMisc::SetMemoryWarningHandler(EngineMemoryWarningHandler);
+	}
 
 	EngineLoop = InEngineLoop;
 
