Changing default value of au.AllowReverbForMultichannelSources to 1

diff --git a/Engine/Source/Runtime/Engine/Private/Audio.cpp b/Engine/Source/Runtime/Engine/Private/Audio.cpp
--- a/Engine/Source/Runtime/Engine/Private/Audio.cpp
+++ b/Engine/Source/Runtime/Engine/Private/Audio.cpp
@@ -78,7 +78,7 @@ FAutoConsoleVariableRef CVarAllowAudioSpatializationCVar(
 	TEXT("0: Disable, >0: Enable"),
 	ECVF_Default);
 
-static int32 AllowReverbForMultichannelSources = 0;
+static int32 AllowReverbForMultichannelSources = 1;
 FAutoConsoleVariableRef CvarAllowReverbForMultichannelSources(
 	TEXT("au.AllowReverbForMultichannelSources"),
 	AllowReverbForMultichannelSources,
