Disable VT (and VT Crunch compression) by default

diff --git a/Engine/Source/Runtime/Engine/Private/Texture.cpp b/Engine/Source/Runtime/Engine/Private/Texture.cpp
--- a/Engine/Source/Runtime/Engine/Private/Texture.cpp
+++ b/Engine/Source/Runtime/Engine/Private/Texture.cpp
@@ -31,7 +31,7 @@
 
 static TAutoConsoleVariable<int32> CVarVirtualTextures(
 	TEXT("r.VirtualTextures"),
-	1,
+	0,
 	TEXT("Is virtual texture streaming enabled?"),
 	ECVF_RenderThreadSafe | ECVF_ReadOnly);
 
diff --git a/Engine/Source/Runtime/Engine/Private/TextureDerivedData.cpp b/Engine/Source/Runtime/Engine/Private/TextureDerivedData.cpp
--- a/Engine/Source/Runtime/Engine/Private/TextureDerivedData.cpp
+++ b/Engine/Source/Runtime/Engine/Private/TextureDerivedData.cpp
@@ -37,7 +37,7 @@
 
 static TAutoConsoleVariable<int32> CVarVTCompressCrunch(
 	TEXT("r.VT.EnableCompressCrunch"),
-	1,
+	0,
 	TEXT("Enable Crunch compression for virtual textures, for supported formats")
 );
 
