Cherry-pick 7938090: Fixes backward compatibility issues when DOF was not circle DOF, yet had a custom focal distance.

diff --git a/Engine/Source/Runtime/Engine/Private/Scene.cpp b/Engine/Source/Runtime/Engine/Private/Scene.cpp
--- a/Engine/Source/Runtime/Engine/Private/Scene.cpp
+++ b/Engine/Source/Runtime/Engine/Private/Scene.cpp
@@ -966,6 +966,12 @@ void FPostProcessSettings::PostSerialize(const FArchive& Ar)
 			{
 				DepthOfFieldFocalDistance = 1000.0f;
 			}
+			else if (DepthOfFieldMethod_DEPRECATED != DOFM_CircleDOF)
+			{
+				// Aggressively force disable DOF by setting default value on the focal distance to be invalid if the method was not CircleDOF, in case.
+				// it focal distance was modified if even if DOF was in the end disabled.
+				DepthOfFieldFocalDistance = 0.0f;
+			}
 
 			// Make sure gaussian DOF is disabled on mobile if the DOF method was set to something else.
 			if (DepthOfFieldMethod_DEPRECATED != DOFM_Gaussian)
