fix for sound class default 2d reverb sound amount parameter not being observed

diff --git a/Engine/Source/Runtime/Core/Public/UObject/ReleaseObjectVersion.h b/Engine/Source/Runtime/Core/Public/UObject/ReleaseObjectVersion.h
--- a/Engine/Source/Runtime/Core/Public/UObject/ReleaseObjectVersion.h
+++ b/Engine/Source/Runtime/Core/Public/UObject/ReleaseObjectVersion.h
@@ -80,6 +80,9 @@ struct CORE_API FReleaseObjectVersion
 		// DepthOfFieldFocalDistance at 0 now disables DOF instead of DepthOfFieldFstop at 0.
 		FocalDistanceDisablesDOF,
 
+		// SoundClass 2D Reverb Send
+		SoundClass2DReverbSend,
+
 		// -----<new versions can be added above this line>-------------------------------------------------
 		VersionPlusOne,
 		LatestVersion = VersionPlusOne - 1
diff --git a/Engine/Source/Runtime/Engine/Classes/Sound/AudioSettings.h b/Engine/Source/Runtime/Engine/Classes/Sound/AudioSettings.h
--- a/Engine/Source/Runtime/Engine/Classes/Sound/AudioSettings.h
+++ b/Engine/Source/Runtime/Engine/Classes/Sound/AudioSettings.h
@@ -139,8 +139,8 @@ class ENGINE_API UAudioSettings : public UDeveloperSettings
 	float VoipBufferingDelay;
 
 	/** The amount of audio to send to reverb submixes if no reverb send is setup for the source through attenuation settings. Only used in audio mixer. */
-	UPROPERTY(config, EditAnywhere, Category = "Audio", AdvancedDisplay)
-	float DefaultReverbSendLevel;
+	UPROPERTY()
+	float DefaultReverbSendLevel_DEPRECATED;
 
 	/** Enables legacy version of reverb. The legacy reverb runs more slowly, but by most other measures is functionally equivalent. It has a slight perceptual difference. */
 	UPROPERTY(config, EditAnywhere, Category = "Audio", AdvancedDisplay)
diff --git a/Engine/Source/Runtime/Engine/Private/ActiveSound.cpp b/Engine/Source/Runtime/Engine/Private/ActiveSound.cpp
--- a/Engine/Source/Runtime/Engine/Private/ActiveSound.cpp
+++ b/Engine/Source/Runtime/Engine/Private/ActiveSound.cpp
@@ -565,15 +565,23 @@ void FActiveSound::UpdateWaveInstances(TArray<FWaveInstance*> &InWaveInstances,
 	bFinished = true;
 	if (FadeOut == EFadeOut::None || (PlaybackTime <= TargetAdjustVolumeStopTime))
 	{
+		bool bReverbSendLevelWasSet = false;
 		if (bHasAttenuationSettings)
 		{
 			UpdateAttenuation(DeltaTime, ParseParams, *ClosestListenerPtr);
+			bReverbSendLevelWasSet = true;
 		}
 		else
 		{
-			// In the case of no attenuation settings, we still want to setup a default send reverb level
-			ParseParams.ReverbSendMethod = EReverbSendMethod::Manual;
-			ParseParams.ManualReverbSendLevel = AudioDevice->GetDefaultReverbSendLevel();
+			ParseParams.ReverbSendMethod = EReverbSendMethod::Manual;		
+			if (ParseParams.SoundClass)
+			{
+				ParseParams.ManualReverbSendLevel = ParseParams.SoundClass->Properties.Default2DReverbSendAmount;
+			}
+			else
+			{
+				ParseParams.ManualReverbSendLevel = AudioDevice->GetDefaultReverbSendLevel();
+			}
 		}
 
 		ParseParams.ModulationPluginSettings = FindModulationSettings();
diff --git a/Engine/Source/Runtime/Engine/Private/AudioDevice.cpp b/Engine/Source/Runtime/Engine/Private/AudioDevice.cpp
--- a/Engine/Source/Runtime/Engine/Private/AudioDevice.cpp
+++ b/Engine/Source/Runtime/Engine/Private/AudioDevice.cpp
@@ -344,7 +344,7 @@ bool FAudioDevice::Init(int32 InMaxSources)
 	GlobalMaxPitch = FMath::Clamp(AudioSettings->GlobalMaxPitchScale, 0.0001f, 4.0f);
 	bAllowCenterChannel3DPanning = AudioSettings->bAllowCenterChannel3DPanning;
 	bAllowPlayWhenSilent = AudioSettings->bAllowPlayWhenSilent;
-	DefaultReverbSendLevel = AudioSettings->DefaultReverbSendLevel;
+	DefaultReverbSendLevel = AudioSettings->DefaultReverbSendLevel_DEPRECATED;
 
 	const FSoftObjectPath DefaultBaseSoundMixName = GetDefault<UAudioSettings>()->DefaultBaseSoundMix;
 	if (DefaultBaseSoundMixName.IsValid())
diff --git a/Engine/Source/Runtime/Engine/Private/AudioSettings.cpp b/Engine/Source/Runtime/Engine/Private/AudioSettings.cpp
--- a/Engine/Source/Runtime/Engine/Private/AudioSettings.cpp
+++ b/Engine/Source/Runtime/Engine/Private/AudioSettings.cpp
@@ -67,7 +67,7 @@ void UAudioSettings::AddDefaultSettings()
 	DefaultSettings.DisplayName = LOCTEXT("DefaultSettingsName", "Default");
 	QualityLevels.Add(DefaultSettings);
 	bAllowPlayWhenSilent = true;
-	DefaultReverbSendLevel = 0.0f;
+	DefaultReverbSendLevel_DEPRECATED = 0.0f;
 	bEnableLegacyReverb = false;
 	VoiPSampleRate = EVoiceSampleRate::Low16000Hz;
 	VoipBufferingDelay = 0.2f;
diff --git a/Engine/Source/Runtime/Engine/Private/SoundClass.cpp b/Engine/Source/Runtime/Engine/Private/SoundClass.cpp
--- a/Engine/Source/Runtime/Engine/Private/SoundClass.cpp
+++ b/Engine/Source/Runtime/Engine/Private/SoundClass.cpp
@@ -2,6 +2,7 @@
 
 
 #include "Sound/SoundClass.h"
+#include "Sound/AudioSettings.h"
 #include "EngineGlobals.h"
 #include "Engine/Engine.h"
 #include "Audio.h"
@@ -221,6 +222,13 @@ void USoundClass::Serialize( FArchive& Ar )
 		TMap<USoundClass*, FSoundClassEditorData>	EditorData_DEPRECATED;
 		Ar << EditorData_DEPRECATED;
 	}
+
+	const int32 ReleaseObjectVersion = Ar.CustomVer(FReleaseObjectVersion::GUID);
+	if (ReleaseObjectVersion < FReleaseObjectVersion::SoundClass2DReverbSend)
+	{
+		const UAudioSettings* AudioSettings = GetDefault<UAudioSettings>();
+		Properties.Default2DReverbSendAmount = AudioSettings->DefaultReverbSendLevel_DEPRECATED;
+	}
 }
 
 void USoundClass::BeginDestroy()
