Support vertex shader sampling from virtual texture
GPU feedback is still only generated by pixel shader so using vertex shader sampling will require care
Initially we can rely on game side logic to make CPU requests for pages around the camera
But in future the aim would be to support a non screenspace GPU feedback system

diff --git a/Engine/Shaders/Private/SimpleElementTexture2DPreviewPixelShader.usf b/Engine/Shaders/Private/SimpleElementTexture2DPreviewPixelShader.usf
--- a/Engine/Shaders/Private/SimpleElementTexture2DPreviewPixelShader.usf
+++ b/Engine/Shaders/Private/SimpleElementTexture2DPreviewPixelShader.usf
@@ -59,7 +59,7 @@ void Main(
 	if( MipLevel >= 0.0f )
 	{
 #if SAMPLE_VIRTUAL_TEXTURE
-		VTPageTableResult PageTableResult = TextureLoadVirtualPageTableLevel(InPageTableTexture0, InPageTableTexture1, TextureCoordinate, MipLevel, InPosition.xy, VTPageTableUniform_Unpack(VTPackedPageTableUniform[0], VTPackedPageTableUniform[1]), VTADDRESSMODE_CLAMP, VTADDRESSMODE_CLAMP);
+		VTPageTableResult PageTableResult = TextureLoadVirtualPageTableLevel(InPageTableTexture0, InPageTableTexture1, TextureCoordinate, MipLevel, VTPageTableUniform_Unpack(VTPackedPageTableUniform[0], VTPackedPageTableUniform[1]), VTADDRESSMODE_CLAMP, VTADDRESSMODE_CLAMP);
 		Sample = TextureVirtualSampleLevel( InTexture, InTextureSampler, PageTableResult, LayerIndex, VTUniform_Unpack(VTPackedUniform));
 
 
@@ -77,7 +77,7 @@ void Main(
 	{
 #if SAMPLE_VIRTUAL_TEXTURE
 
-		VTPageTableResult PageTableResult = TextureLoadVirtualPageTableLevel(InPageTableTexture0, InPageTableTexture1, TextureCoordinate, 0.0f, InPosition.xy, VTPageTableUniform_Unpack(VTPackedPageTableUniform[0], VTPackedPageTableUniform[1]), VTADDRESSMODE_CLAMP, VTADDRESSMODE_CLAMP);
+		VTPageTableResult PageTableResult = TextureLoadVirtualPageTableLevel(InPageTableTexture0, InPageTableTexture1, TextureCoordinate, 0.0f, VTPageTableUniform_Unpack(VTPackedPageTableUniform[0], VTPackedPageTableUniform[1]), VTADDRESSMODE_CLAMP, VTADDRESSMODE_CLAMP);
 		Sample = TextureVirtualSample( InTexture, InTextureSampler, PageTableResult, LayerIndex, VTUniform_Unpack(VTPackedUniform));
 
 		/*Sample = TextureVirtualSample(	InTexture, InTextureSampler, InPageTableTexture,
diff --git a/Engine/Shaders/Private/VirtualTextureCommon.ush b/Engine/Shaders/Private/VirtualTextureCommon.ush
--- a/Engine/Shaders/Private/VirtualTextureCommon.ush
+++ b/Engine/Shaders/Private/VirtualTextureCommon.ush
@@ -244,7 +244,6 @@ void TextureLoadVirtualPageTableInternal(inout VTPageTableResult OutResult,
 
 VTPageTableResult TextureLoadVirtualPageTableLevel(Texture2D<uint4> PageTable0,
 	float2 UV, float Level,
-	float2 SvPositionXY,
 	VTPageTableUniform PageTableUniform,
 	uint AddressU, uint AddressV)
 {
@@ -256,7 +255,6 @@ VTPageTableResult TextureLoadVirtualPageTableLevel(Texture2D<uint4> PageTable0,
 
 VTPageTableResult TextureLoadVirtualPageTableLevel(Texture2D<uint4> PageTable0, Texture2D<uint4> PageTable1,
 	float2 UV, float Level,
-	float2 SvPositionXY,
 	VTPageTableUniform PageTableUniform,
 	uint AddressU, uint AddressV)
 {
@@ -362,25 +360,43 @@ VTPageTableResult TextureLoadVirtualPageTableBias(Texture2D<uint4> PageTable0, T
 }
 
 VTPageTableResult TextureLoadVirtualPageTableLevel(Texture2D<uint4> PageTable0, VTPageTableUniform PageTableUniform,
-	float2 SvPositionXY, in out FVirtualTextureFeedbackParams Feedback, uint SampleIndex,
+	in out FVirtualTextureFeedbackParams Feedback, uint SampleIndex,
 	float2 UV, uint AddressU, uint AddressV,
 	MaterialFloat Level)
 {
-	VTPageTableResult PageTableResult = TextureLoadVirtualPageTableLevel(PageTable0, UV, Level, SvPositionXY, PageTableUniform, AddressU, AddressV);
+	VTPageTableResult PageTableResult = TextureLoadVirtualPageTableLevel(PageTable0, UV, Level, PageTableUniform, AddressU, AddressV);
 	StoreVirtualTextureFeedback(Feedback, SampleIndex, PageTableResult.PackedRequest);
 	return PageTableResult;
 }
 
+VTPageTableResult TextureLoadVirtualPageTableLevel(Texture2D<uint4> PageTable0, VTPageTableUniform PageTableUniform,
+	uint SampleIndex,
+	float2 UV, uint AddressU, uint AddressV,
+	MaterialFloat Level)
+{
+	VTPageTableResult PageTableResult = TextureLoadVirtualPageTableLevel(PageTable0, UV, Level, PageTableUniform, AddressU, AddressV);
+	return PageTableResult;
+}
+
 VTPageTableResult TextureLoadVirtualPageTableLevel(Texture2D<uint4> PageTable0, Texture2D<uint4> PageTable1, VTPageTableUniform PageTableUniform,
-	float2 SvPositionXY, in out FVirtualTextureFeedbackParams Feedback, uint SampleIndex,
+	in out FVirtualTextureFeedbackParams Feedback, uint SampleIndex,
 	float2 UV, uint AddressU, uint AddressV,
 	MaterialFloat Level)
 {
-	VTPageTableResult PageTableResult = TextureLoadVirtualPageTableLevel(PageTable0, PageTable1, UV, Level, SvPositionXY, PageTableUniform, AddressU, AddressV);
+	VTPageTableResult PageTableResult = TextureLoadVirtualPageTableLevel(PageTable0, PageTable1, UV, Level, PageTableUniform, AddressU, AddressV);
 	StoreVirtualTextureFeedback(Feedback, SampleIndex, PageTableResult.PackedRequest);
 	return PageTableResult;
 }
 
+VTPageTableResult TextureLoadVirtualPageTableLevel(Texture2D<uint4> PageTable0, Texture2D<uint4> PageTable1, VTPageTableUniform PageTableUniform,
+	uint SampleIndex,
+	float2 UV, uint AddressU, uint AddressV,
+	MaterialFloat Level)
+{
+	VTPageTableResult PageTableResult = TextureLoadVirtualPageTableLevel(PageTable0, PageTable1, UV, Level, PageTableUniform, AddressU, AddressV);
+	return PageTableResult;
+}
+
 VTPageTableResult TextureLoadVirtualPageTableGrad(Texture2D<uint4> PageTable0, VTPageTableUniform PageTableUniform,
 	float2 SvPositionXY, in out FVirtualTextureFeedbackParams Feedback, uint SampleIndex,
 	float2 UV, uint AddressU, uint AddressV,
diff --git a/Engine/Source/Runtime/Engine/Private/Materials/HLSLMaterialTranslator.h b/Engine/Source/Runtime/Engine/Private/Materials/HLSLMaterialTranslator.h
--- a/Engine/Source/Runtime/Engine/Private/Materials/HLSLMaterialTranslator.h
+++ b/Engine/Source/Runtime/Engine/Private/Materials/HLSLMaterialTranslator.h
@@ -156,6 +156,7 @@ struct FMaterialVTStackEntry
 	int32 DebugMipValue0Index;
 	int32 DebugMipValue1Index;
 	int32 PreallocatedStackTextureIndex;
+	bool bGenerateFeedback;
 	float AspectRatio;
 
 	int32 CodeIndex;
@@ -783,12 +784,7 @@ public:
 					{
 						Errorf(TEXT("Sampling a virtual texture is currently not supported when connected to the Opacity Mask material attribute."));
 					}
-					else
 #endif
-					if (ShaderFrequencyToValidate != SF_Pixel)
-					{
-						Errorf(TEXT("Sampling a virtual texture is currently only supported from pixel shader."));
-					}
 				}
 			}
 		}
@@ -4057,7 +4053,7 @@ protected:
 		}
 	}
 
-	uint32 AcquireVTStackIndex(ETextureMipValueMode MipValueMode, TextureAddress AddressU, TextureAddress AddressV, float AspectRatio, int32 CoordinateIndex, int32 MipValue0Index, int32 MipValue1Index, int32 PreallocatedStackTextureIndex)
+	uint32 AcquireVTStackIndex(ETextureMipValueMode MipValueMode, TextureAddress AddressU, TextureAddress AddressV, float AspectRatio, int32 CoordinateIndex, int32 MipValue0Index, int32 MipValue1Index, int32 PreallocatedStackTextureIndex, bool bGenerateFeedback)
 	{
 		const uint64 CoordinatHash = GetParameterHash(CoordinateIndex);
 		const uint64 MipValue0Hash = GetParameterHash(MipValue0Index);
@@ -4071,7 +4067,7 @@ protected:
 		Hash = CityHash128to64({ Hash, (uint64)AddressV });
 		Hash = CityHash128to64({ Hash, (uint64)(AspectRatio * 1000.0f) });
 		Hash = CityHash128to64({ Hash, (uint64)PreallocatedStackTextureIndex });
-
+		Hash = CityHash128to64({ Hash, (uint64)(bGenerateFeedback ? 1 : 0) });
 
 		// First check to see if we have an existing VTStack that matches this key, that can still fit another layer
 		for (int32 Index = VTStackHash.First(Hash); VTStackHash.IsValid(Index); Index = VTStackHash.Next(Index))
@@ -4087,7 +4083,8 @@ protected:
 				Entry.AddressU == AddressU &&
 				Entry.AddressV == AddressV &&
 				Entry.AspectRatio == AspectRatio &&
-				Entry.PreallocatedStackTextureIndex == PreallocatedStackTextureIndex)
+				Entry.PreallocatedStackTextureIndex == PreallocatedStackTextureIndex &&
+				Entry.bGenerateFeedback == bGenerateFeedback)
 			{
 				return Index;
 			}
@@ -4109,12 +4106,17 @@ protected:
 		Entry.DebugMipValue0Index = MipValue0Index;
 		Entry.DebugMipValue1Index = MipValue1Index;
 		Entry.PreallocatedStackTextureIndex = PreallocatedStackTextureIndex;
+		Entry.bGenerateFeedback = bGenerateFeedback;
 
 		MaterialCompilationOutput.UniformExpressionSet.VTStacks.Add(FMaterialVirtualTextureStack(PreallocatedStackTextureIndex));
 
 		// these two arrays need to stay in sync
 		check(VTStacks.Num() == MaterialCompilationOutput.UniformExpressionSet.VTStacks.Num());
 
+		// Optionally sample without virtual texture feedback but only for miplevel mode
+		check(bGenerateFeedback || MipValueMode == TMVM_MipLevel)
+		FString FeedbackParameter = bGenerateFeedback ? TEXT("Parameters.VirtualTextureFeedback,") : TEXT("");
+
 		// Code to load the VT page table...this will execute the first time a given VT stack is accessed
 		// Additional stack layers will simply reuse these results
 		switch (MipValueMode)
@@ -4128,8 +4130,8 @@ protected:
 				StackIndex, StackIndex, StackIndex, StackIndex, *CoerceParameter(CoordinateIndex, MCT_Float2), GetVTAddressMode(AddressU), GetVTAddressMode(AddressV), *CoerceParameter(MipValue0Index, MCT_Float1));
 			break;
 		case TMVM_MipLevel:
-			Entry.CodeIndex = AddCodeChunk(MCT_VTPageTableResult, TEXT("TextureLoadVirtualPageTableLevel(VIRTUALTEXTURE_PAGETABLE_%d, VTPageTableUniform_Unpack(Material.VTPackedPageTableUniform[%d*2], Material.VTPackedPageTableUniform[%d*2+1]), Parameters.SvPosition.xy, Parameters.VirtualTextureFeedback, %d + LIGHTMAP_VT_ENABLED, %s, %s, %s, %s)"),
-				StackIndex, StackIndex, StackIndex, StackIndex, *CoerceParameter(CoordinateIndex, MCT_Float2), GetVTAddressMode(AddressU), GetVTAddressMode(AddressV), *CoerceParameter(MipValue0Index, MCT_Float1));
+			Entry.CodeIndex = AddCodeChunk(MCT_VTPageTableResult, TEXT("TextureLoadVirtualPageTableLevel(VIRTUALTEXTURE_PAGETABLE_%d, VTPageTableUniform_Unpack(Material.VTPackedPageTableUniform[%d*2], Material.VTPackedPageTableUniform[%d*2+1]), %s %d + LIGHTMAP_VT_ENABLED, %s, %s, %s, %s)"),
+				StackIndex, StackIndex, StackIndex, *FeedbackParameter, StackIndex, *CoerceParameter(CoordinateIndex, MCT_Float2), GetVTAddressMode(AddressU), GetVTAddressMode(AddressV), *CoerceParameter(MipValue0Index, MCT_Float1));
 			break;
 		case TMVM_Derivative:
 			Entry.CodeIndex = AddCodeChunk(MCT_VTPageTableResult, TEXT("TextureLoadVirtualPageTableGrad(VIRTUALTEXTURE_PAGETABLE_%d, VTPageTableUniform_Unpack(Material.VTPackedPageTableUniform[%d*2], Material.VTPackedPageTableUniform[%d*2+1]), Parameters.SvPosition.xy, Parameters.VirtualTextureFeedback, %d + LIGHTMAP_VT_ENABLED, %s, %s, %s, %s, %s)"),
@@ -4191,11 +4193,7 @@ protected:
 		const bool bVirtualTexture = TextureType == MCT_TextureVirtual;
 		if (bVirtualTexture)
 		{
-			if (ShaderFrequency != SF_Pixel)
-			{
-				return Errorf(TEXT("Sampling a virtual texture is currently only supported in pixel shader."));
-			}
-			else if (Material->GetMaterialDomain() == MD_DeferredDecal)
+			if (Material->GetMaterialDomain() == MD_DeferredDecal)
 			{
 				if (Material->GetDecalBlendMode() == DBM_Volumetric_DistanceFunction)
 				{
@@ -4237,6 +4235,11 @@ protected:
 		{
 			MipValueMode = TMVM_MipLevel;
 			AutomaticViewMipBias = false;
+
+			if (MipValue0Index == INDEX_NONE)
+			{
+				MipValue0Index = Constant(0.f);
+			}
 		}
 
 		// Automatic view mip bias is only for surface and decal domains.
@@ -4561,12 +4564,16 @@ protected:
 				}
 			}
 
+			// Only support GPU feedback from pixel shader
+			//todo[vt]: Support feedback from other shader types
+			const bool bGenerateFeedback = ShaderFrequency == SF_Pixel;
+
 			VTLayerIndex = MaterialCompilationOutput.UniformExpressionSet.UniformVirtualTextureExpressions[VirtualTextureIndex]->GetLayerIndex();
 			if (VTLayerIndex != INDEX_NONE)
 			{
 				// The layer index in the virtual texture stack is already known
 				// Create a page table sample for each new combination of virtual texture and sample parameters
-				VTStackIndex = AcquireVTStackIndex(MipValueMode, AddressU, AddressV, 1.0f, CoordinateIndex, MipValue0Index, MipValue1Index, TextureReferenceIndex);
+				VTStackIndex = AcquireVTStackIndex(MipValueMode, AddressU, AddressV, 1.0f, CoordinateIndex, MipValue0Index, MipValue1Index, TextureReferenceIndex, bGenerateFeedback);
 			}
 			else
 			{
@@ -4581,7 +4588,7 @@ protected:
 				const float TextureAspectRatio = (float)Tex2D->Source.GetSizeX() / (float)Tex2D->Source.GetSizeY();
 
 				// Create a page table sample for each new set of sample parameters
-				VTStackIndex = AcquireVTStackIndex(MipValueMode, AddressU, AddressV, TextureAspectRatio, CoordinateIndex, MipValue0Index, MipValue1Index, INDEX_NONE);
+				VTStackIndex = AcquireVTStackIndex(MipValueMode, AddressU, AddressV, TextureAspectRatio, CoordinateIndex, MipValue0Index, MipValue1Index, INDEX_NONE, bGenerateFeedback);
 				// Allocate a layer in the virtual texture stack for this physical sample
 				VTLayerIndex = MaterialCompilationOutput.UniformExpressionSet.VTStacks[VTStackIndex].AddLayer();
 			}
