#UE4 Remove AnimLoadingDebugData. It is no longer needed and has threading problems with concurrentsave

diff --git a/Engine/Source/Runtime/Engine/Classes/Animation/AnimSequence.h b/Engine/Source/Runtime/Engine/Classes/Animation/AnimSequence.h
--- a/Engine/Source/Runtime/Engine/Classes/Animation/AnimSequence.h
+++ b/Engine/Source/Runtime/Engine/Classes/Animation/AnimSequence.h
@@ -194,53 +194,6 @@ struct ENGINE_API FRequestAnimCompressionParams
 	void InitFrameStrippingFromPlatform(const class ITargetPlatform* TargetPlatform);
 };
 
-#if WITH_EDITOR
-// Cache debugging data in editor for UE-49335
-struct FAnimLoadingDebugData
-{
-public:
-	//Single debug entry
-	struct Entry
-	{
-		FString Tag;
-		int32 RawDataCount;
-		int32 SourceDataCount;
-
-		Entry(const TCHAR* InTag, int32 InRawDataCount, int32 InSourceDataCount)
-			: Tag(InTag)
-			, RawDataCount(InRawDataCount)
-			, SourceDataCount(InSourceDataCount)
-		{}
-	};
-
-	template <typename... ArgsType>
-	void AddEntry(ArgsType&&... Args)
-	{
-		Entries.Emplace(Args...);
-	}
-
-	// Build a string of all the debug entries for output
-	FString GetEntries() const
-	{
-		FString Ret;
-		for (const Entry& E : Entries)
-		{
-			TArray<FStringFormatArg> Args;
-			Args.Add(LexToString(E.Tag));
-			Args.Add(LexToString(E.RawDataCount));
-			Args.Add(LexToString(E.SourceDataCount));
-			Ret += FString::Format(TEXT("\t{0}: Raw:{1} Source:{2}\n"), Args);
-		}
-		return Ret;
-	}
-
-private:
-
-	TArray<Entry> Entries;
-};
-
-#endif
-
 FArchive& operator<<(FArchive& Ar, FCompressedOffsetData& D);
 
 UCLASS(config=Engine, hidecategories=(UObject, Length), BlueprintType)
@@ -902,20 +855,6 @@ public:
 
 	friend class UAnimationAsset;
 	friend struct FScopedAnimSequenceRawDataCache;
-
-#if WITH_EDITOR
-	// Cache debugging data in editor for UE-49335
-	FAnimLoadingDebugData AnimLoadingDebugData;
-#endif
-
-	// Cache debugging data in editor for UE-49335
-	void AddAnimLoadingDebugEntry(const TCHAR* Tag)
-	{
-#if WITH_EDITOR
-		AnimLoadingDebugData.AddEntry(Tag, RawAnimationData.Num(), SourceRawAnimationData.Num());
-#endif
-	}
-
 	friend class UAnimationBlueprintLibrary;
 };
 
diff --git a/Engine/Source/Runtime/Engine/Private/Animation/AnimSequence.cpp b/Engine/Source/Runtime/Engine/Private/Animation/AnimSequence.cpp
--- a/Engine/Source/Runtime/Engine/Private/Animation/AnimSequence.cpp
+++ b/Engine/Source/Runtime/Engine/Private/Animation/AnimSequence.cpp
@@ -679,8 +679,6 @@ void UAnimSequence::Serialize(FArchive& Ar)
 {
 	LLM_SCOPE(ELLMTag::Animation);
 	
-	AddAnimLoadingDebugEntry(TEXT("SerializeStart"));
-
 	Ar.UsingCustomVersion(FFrameworkObjectVersion::GUID);
 
 	FRawCurveTracks RawCurveCache;
@@ -731,8 +729,6 @@ void UAnimSequence::Serialize(FArchive& Ar)
 #endif // WITH_EDITORONLY_DATA
 	}
 
-	AddAnimLoadingDebugEntry(TEXT("PostSerializeRawData"));
-
 	if (Ar.CustomVer(FFrameworkObjectVersion::GUID) < FFrameworkObjectVersion::MoveCompressedAnimDataToTheDDC)
 	{
 		// Serialize the compressed byte stream from the archive to the buffer.
@@ -798,8 +794,6 @@ void UAnimSequence::Serialize(FArchive& Ar)
 	InitCurveCompressionScheme();
 
 #endif // WITH_EDITORONLY_DATA
-
-	AddAnimLoadingDebugEntry(TEXT("PostSerialize"));
 }
 
 #if WITH_EDITOR
@@ -868,8 +862,6 @@ bool UAnimSequence::IsPostLoadThreadSafe() const
 
 void UAnimSequence::PostLoad()
 {
-	AddAnimLoadingDebugEntry(TEXT("PostLoadStart"));
-
 	//Parent PostLoad will ensure that skeleton is fully loaded
 	//before we do anything further in PostLoad
 	Super::PostLoad();
@@ -1017,8 +1009,6 @@ void UAnimSequence::PostLoad()
 		}
 #endif
 	}
-
-	AddAnimLoadingDebugEntry(TEXT("PostLoadEnd"));
 }
 
 #if WITH_EDITOR
@@ -1043,28 +1033,6 @@ void UAnimSequence::VerifyTrackMap(USkeleton* MySkeleton)
 {
 	USkeleton* UseSkeleton = (MySkeleton) ? MySkeleton : GetSkeleton();
 
-#if WITH_EDITOR
-	TArray<FStringFormatArg> Args;
-	Args.Add(UseSkeleton ? TEXT("Yes") : TEXT("No"));
-	if (UseSkeleton)
-	{
-		Args.Add(UseSkeleton->HasAllFlags(EObjectFlags::RF_NeedLoad) ? TEXT("Yes") : TEXT("No"));
-		Args.Add(UseSkeleton->HasAllFlags(EObjectFlags::RF_NeedPostLoad) ? TEXT("Yes") : TEXT("No"));
-		Args.Add(UseSkeleton->GetReferenceSkeleton().GetNum());
-		Args.Add(UseSkeleton->GetReferenceSkeleton().GetRawBoneNum());
-	}
-	else
-	{
-		Args.Add(TEXT("N/A"));
-		Args.Add(TEXT("N/A"));
-		Args.Add(TEXT("N/A"));
-		Args.Add(TEXT("N/A"));
-	}
-
-	const FString DebugStartMessage = FString::Format(TEXT("VerifyTrackMapStart - Skeleton Valid: {0} NeedsLoad: {1} NeedsPostLoad: {2} NumBones: {3} NumRawBones: {4}"), Args);
-	AddAnimLoadingDebugEntry(*DebugStartMessage);
-#endif
-
 	if( AnimationTrackNames.Num() != TrackToSkeletonMapTable.Num() && UseSkeleton!=nullptr)
 	{
 		ShowResaveMessage(this);
@@ -1125,7 +1093,6 @@ void UAnimSequence::VerifyTrackMap(USkeleton* MySkeleton)
 			TrackMap.BoneTreeIndex = UseSkeleton->GetReferenceSkeleton().FindBoneIndex(AnimationTrackNames[I]);
 		}		
 	}
-	AddAnimLoadingDebugEntry(TEXT("VerifyTrackMapEnd"));
 }
 
 #endif // WITH_EDITOR
@@ -1714,13 +1681,7 @@ void UAnimSequence::GetAdditiveBasePose(FCompactPose& OutPose, FBlendedCurve& Ou
 		FString RefSkelName = RefPoseSeq->GetSkeleton() ? RefPoseSeq->GetSkeleton()->GetName() : TEXT("NoRefSeqSkeleton");
 		const TCHAR* NeedsLoad = RefPoseSeq->HasAnyFlags(RF_NeedLoad) ? TEXT("Yes") : TEXT("No");
 
-#if WITH_EDITOR
-		FString LoadingDebugEntriesString = AnimLoadingDebugData.GetEntries();
-		const TCHAR* LoadingDebugEntries = *LoadingDebugEntriesString;
-#else
-		const TCHAR* LoadingDebugEntries = TEXT("No debug entries (non editor build)\n");
-#endif
-		checkf(false, TEXT("Cannot get valid base pose for Anim: ['%s' (Skel:%s)] RefSeq: ['%s' (Skel:%s)] RawAnimDataNum: %i NeedsLoad: %s\nDebug Entries:\n%s"), *Name, *SkelName, *RefName, *RefSkelName, RefPoseSeq->RawAnimationData.Num(), NeedsLoad, LoadingDebugEntries);
+		checkf(false, TEXT("Cannot get valid base pose for Anim: ['%s' (Skel:%s)] RefSeq: ['%s' (Skel:%s)] RawAnimDataNum: %i NeedsLoad: %s"), *Name, *SkelName, *RefName, *RefSkelName, RefPoseSeq->RawAnimationData.Num(), NeedsLoad);
 	}
 
 	switch (RefPoseType)
@@ -2146,8 +2107,6 @@ void UAnimSequence::FlipRotationWForNonRoot(USkeletalMesh * SkelMesh)
 
 void UAnimSequence::RequestAnimCompression(FRequestAnimCompressionParams Params)
 {
-	AddAnimLoadingDebugEntry(TEXT("RequestAnimCompression"));
-
 #if WITH_EDITOR
 	USkeleton* CurrentSkeleton = GetSkeleton();
 	if (CurrentSkeleton == nullptr)
@@ -2458,7 +2417,6 @@ bool IsNewKeyDifferent(const FRichCurveKey& LastKey, float NewValue)
 
 void UAnimSequence::BakeOutAdditiveIntoRawData(TArray<FRawAnimSequenceTrack>& NewRawTracks, TArray<FName>& NewAnimationTrackNames, TArray<FTrackToSkeletonMap>& NewTrackToSkeletonMapTable, FRawCurveTracks& NewCurveTracks, TArray<FRawAnimSequenceTrack>& AdditiveBaseAnimationData)
 {
-	AddAnimLoadingDebugEntry(TEXT("BakeOutAdditiveIntoRawData"));
 	if (!CanBakeAdditive())
 	{
 		return; // Nothing to do
@@ -2925,7 +2883,6 @@ int32 FindValidTransformParentTrack(const URig* Rig, int32 NodeIndex, bool bTran
 
 void UAnimSequence::RemapTracksToNewSkeleton( USkeleton* NewSkeleton, bool bConvertSpaces )
 {
-	AddAnimLoadingDebugEntry(TEXT("RemapTracksToNewSkeletonStart"));
 	// this is not cheap, so make sure it only happens in editor
 
 	// @Todo : currently additive will work fine since we don't bake anything except when we extract
@@ -3406,14 +3363,11 @@ void UAnimSequence::RemapTracksToNewSkeleton( USkeleton* NewSkeleton, bool bConv
 		VerifyTrackMap(NewSkeleton);
 	}
 
-	AddAnimLoadingDebugEntry(TEXT("RemapTracksToNewSkeletonPreSetSkeleton"));
 	Super::RemapTracksToNewSkeleton(NewSkeleton, bConvertSpaces);
-	AddAnimLoadingDebugEntry(TEXT("RemapTracksToNewSkeletonEnd"));
 }
 
 void UAnimSequence::PostProcessSequence(bool bForceNewRawDatGuid)
 {
-	AddAnimLoadingDebugEntry(TEXT("PostProcessSequenceStart"));
 	// pre process before compress raw animation data
 
 	// if scale is too small, zero it out. Cause it hard to retarget when compress
@@ -3461,7 +3415,6 @@ void UAnimSequence::PostProcessSequence(bool bForceNewRawDatGuid)
 	ClampNotifiesAtEndOfSequence();
 	// mark package as dirty
 	MarkPackageDirty();
-	AddAnimLoadingDebugEntry(TEXT("PostProcessSequenceEnd"));
 }
 
 void UAnimSequence::RemoveNaNTracks()
