Added missing config arg to Uproperty

diff --git a/Engine/Source/Runtime/Engine/Classes/Sound/AudioSettings.h b/Engine/Source/Runtime/Engine/Classes/Sound/AudioSettings.h
--- a/Engine/Source/Runtime/Engine/Classes/Sound/AudioSettings.h
+++ b/Engine/Source/Runtime/Engine/Classes/Sound/AudioSettings.h
@@ -139,7 +139,7 @@ class ENGINE_API UAudioSettings : public UDeveloperSettings
 	float VoipBufferingDelay;
 
 	/** The amount of audio to send to reverb submixes if no reverb send is setup for the source through attenuation settings. Only used in audio mixer. */
-	UPROPERTY()
+	UPROPERTY(config)
 	float DefaultReverbSendLevel_DEPRECATED;
 
 	/** Enables legacy version of reverb. The legacy reverb runs more slowly, but by most other measures is functionally equivalent. It has a slight perceptual difference. */
@@ -165,9 +165,6 @@ class ENGINE_API UAudioSettings : public UDeveloperSettings
 	UPROPERTY(config, EditAnywhere, Category = "Quality", AdvancedDisplay)
 	uint32 bAllowPlayWhenSilent:1;
 
-	UPROPERTY()
-	uint32 bAllowVirtualizedSounds_DEPRECATED:1;
-
 	/** Disables master EQ effect in the audio DSP graph. */
 	UPROPERTY(config, EditAnywhere, Category = "Quality", AdvancedDisplay)
 	uint32 bDisableMasterEQ : 1;
