Revert behavior for ServerUpdateLevelVisibility in PIE, because UPackage::FileName may be None.

diff --git a/Engine/Source/Runtime/Engine/Private/PlayerController.cpp b/Engine/Source/Runtime/Engine/Private/PlayerController.cpp
--- a/Engine/Source/Runtime/Engine/Private/PlayerController.cpp
+++ b/Engine/Source/Runtime/Engine/Private/PlayerController.cpp
@@ -107,12 +107,15 @@ FUpdateLevelVisibilityLevelInfo::FUpdateLevelVisibilityLevelInfo(const ULevel* c
 {
 	const UPackage* const LevelPackage = Level->GetOutermost();
 	PackageName = LevelPackage->GetFName();
-	FileName = LevelPackage->FileName;
+
+	// When packages are duplicated for PIE, they may not have a FileName.
+	// For now, just revert to the old behavior.
+	FileName = (LevelPackage->FileName == NAME_None) ? PackageName : LevelPackage->FileName;
 }
 
 bool FUpdateLevelVisibilityLevelInfo::NetSerialize(FArchive& Ar, UPackageMap* PackageMap, bool& bOutSuccess)
 {
-	bool bArePackageAndFileTheSame = !!((PlayerControllerCVars::LevelVisibilityDontSerializeFileName) || (FileName == PackageName));
+	bool bArePackageAndFileTheSame = !!((PlayerControllerCVars::LevelVisibilityDontSerializeFileName) || (FileName == PackageName) || (FileName == NAME_None));
 	bool bLocalIsVisible = !!bIsVisible;
 
 	Ar.SerializeBits(&bArePackageAndFileTheSame, 1);
