Fix base anim blueprints not showing execution state. When a derived instance is watched in the editor via its base blueprint, node visits weren't being recorded. Now we attribute node visits to the base blueprint if it is debugging the derived instance.

#rnx

diff --git a/Engine/Source/Runtime/Engine/Classes/Animation/AnimBlueprint.h b/Engine/Source/Runtime/Engine/Classes/Animation/AnimBlueprint.h
--- a/Engine/Source/Runtime/Engine/Classes/Animation/AnimBlueprint.h
+++ b/Engine/Source/Runtime/Engine/Classes/Animation/AnimBlueprint.h
@@ -13,6 +13,7 @@ class SWidget;
 class UAnimationAsset;
 class USkeletalMesh;
 class USkeleton;
+struct FAnimBlueprintDebugData;
 
 USTRUCT()
 struct FAnimGroupInfo
@@ -130,7 +131,7 @@ class ENGINE_API UAnimBlueprint : public UBlueprint, public IInterface_PreviewMe
 	int32 FindOrAddGroup(FName GroupName);
 
 	/** Returns the most base anim blueprint for a given blueprint (if it is inherited from another anim blueprint, returning null if only native / non-anim BP classes are it's parent) */
-	static UAnimBlueprint* FindRootAnimBlueprint(UAnimBlueprint* DerivedBlueprint);
+	static UAnimBlueprint* FindRootAnimBlueprint(const UAnimBlueprint* DerivedBlueprint);
 
 	DECLARE_MULTICAST_DELEGATE_TwoParams(FOnOverrideChangedMulticaster, FGuid, UAnimationAsset*);
 
@@ -165,6 +166,13 @@ public:
 	virtual USkeletalMesh* GetPreviewMesh(bool bFindIfNotSet = false) override;
 	virtual USkeletalMesh* GetPreviewMesh() const override;
 
+public:
+	/** Check if the anim instance is the active debug object for this anim BP */
+	bool IsObjectBeingDebugged(const UObject* AnimInstance) const;
+
+	/** Get the debug data for this anim BP */
+	FAnimBlueprintDebugData* GetDebugData() const;
+
 #if WITH_EDITORONLY_DATA
 public:
 	// Array of overrides to asset containing nodes in the parent that have been overridden
@@ -180,5 +188,5 @@ private:
 	/** The default skeletal mesh to use when previewing this asset - this only applies when you open Persona using this asset*/
 	UPROPERTY(duplicatetransient, AssetRegistrySearchable)
 	TSoftObjectPtr<class USkeletalMesh> PreviewSkeletalMesh;
-#endif
+#endif // WITH_EDITORONLY_DATA
 };
diff --git a/Engine/Source/Runtime/Engine/Private/AnimBlueprint.cpp b/Engine/Source/Runtime/Engine/Private/AnimBlueprint.cpp
--- a/Engine/Source/Runtime/Engine/Private/AnimBlueprint.cpp
+++ b/Engine/Source/Runtime/Engine/Private/AnimBlueprint.cpp
@@ -77,7 +77,7 @@ int32 UAnimBlueprint::FindOrAddGroup(FName GroupName)
 
 
 /** Returns the most base anim blueprint for a given blueprint (if it is inherited from another anim blueprint, returning null if only native / non-anim BP classes are it's parent) */
-UAnimBlueprint* UAnimBlueprint::FindRootAnimBlueprint(UAnimBlueprint* DerivedBlueprint)
+UAnimBlueprint* UAnimBlueprint::FindRootAnimBlueprint(const UAnimBlueprint* DerivedBlueprint)
 {
 	UAnimBlueprint* ParentBP = NULL;
 
@@ -234,4 +234,32 @@ void UAnimBlueprint::SetPreviewMesh(USkeletalMesh* PreviewMesh, bool bMarkAsDirt
 	}
 	PreviewSkeletalMesh = PreviewMesh;
 #endif
+}
+
+
+bool UAnimBlueprint::IsObjectBeingDebugged(const UObject* Object) const
+{
+#if WITH_EDITOR
+	// Only root anim BPs can have anim graphs and be debugged
+	const UAnimBlueprint* RootBP = UAnimBlueprint::FindRootAnimBlueprint(this);
+	const UAnimBlueprint* DebugBP = RootBP ? RootBP : this;
+	
+	return DebugBP->GetObjectBeingDebugged() == Object;
+#else
+	return false;
+#endif
+}
+
+FAnimBlueprintDebugData* UAnimBlueprint::GetDebugData() const
+{
+#if WITH_EDITORONLY_DATA
+	// Only root anim BPs can have anim graphs and be debugged
+	const UAnimBlueprint* RootBP = UAnimBlueprint::FindRootAnimBlueprint(this);
+	const UAnimBlueprint* DebugBP = RootBP ? RootBP : this;
+	UAnimBlueprintGeneratedClass* AnimClass = DebugBP->GetAnimBlueprintGeneratedClass();
+
+	return AnimClass ? &AnimClass->GetAnimBlueprintDebugData() : nullptr;
+#else
+	return nullptr;
+#endif // WITH_EDITORONLY_DATA
 }
\ No newline at end of file
diff --git a/Engine/Source/Runtime/Engine/Private/Animation/AnimInstanceProxy.cpp b/Engine/Source/Runtime/Engine/Private/Animation/AnimInstanceProxy.cpp
--- a/Engine/Source/Runtime/Engine/Private/Animation/AnimInstanceProxy.cpp
+++ b/Engine/Source/Runtime/Engine/Private/Animation/AnimInstanceProxy.cpp
@@ -373,15 +373,14 @@ void FAnimInstanceProxy::PreUpdate(UAnimInstance* InAnimInstance, float DeltaSec
 	FMemory::Memset(MachineWeights.GetData(), 0, MachineWeights.Num() * sizeof(float));
 
 #if WITH_EDITORONLY_DATA
-	bIsBeingDebugged = false;
-	if (UAnimBlueprint* AnimBlueprint = GetAnimBlueprint())
+	UAnimBlueprint* AnimBP = GetAnimBlueprint();
+	bIsBeingDebugged = AnimBP ? AnimBP->IsObjectBeingDebugged(InAnimInstance) : false;
+	if (bIsBeingDebugged)
 	{
-		bIsBeingDebugged = (AnimBlueprint->GetObjectBeingDebugged() == InAnimInstance);
-		if(bIsBeingDebugged)
+		FAnimBlueprintDebugData* DebugData = AnimBP->GetDebugData();
+		if (DebugData)
 		{
-			UAnimBlueprintGeneratedClass* AnimBlueprintGeneratedClass = Cast<UAnimBlueprintGeneratedClass>(InAnimInstance->GetClass());
-			FAnimBlueprintDebugData& DebugData = AnimBlueprintGeneratedClass->GetAnimBlueprintDebugData();
-			PoseWatchEntriesForThisFrame = DebugData.AnimNodePoseWatch;
+			PoseWatchEntriesForThisFrame = DebugData->AnimNodePoseWatch;
 		}
 	}
 #endif
@@ -458,12 +457,10 @@ void FAnimInstanceProxy::PostUpdate(UAnimInstance* InAnimInstance) const
 	DECLARE_SCOPE_HIERARCHICAL_COUNTER_FUNC()
 
 #if WITH_EDITORONLY_DATA
-	if(bIsBeingDebugged)
+	if (FAnimBlueprintDebugData* DebugData = GetAnimBlueprintDebugData())
 	{
-		UAnimBlueprintGeneratedClass* AnimBlueprintGeneratedClass = Cast<UAnimBlueprintGeneratedClass>(InAnimInstance->GetClass());
-		FAnimBlueprintDebugData& DebugData = AnimBlueprintGeneratedClass->GetAnimBlueprintDebugData();
-		DebugData.RecordNodeVisitArray(UpdatedNodesThisFrame);
-		DebugData.AnimNodePoseWatch = PoseWatchEntriesForThisFrame;
+		DebugData->RecordNodeVisitArray(UpdatedNodesThisFrame);
+		DebugData->AnimNodePoseWatch = PoseWatchEntriesForThisFrame;
 	}
 #endif
 
@@ -2528,4 +2525,17 @@ void FAnimInstanceProxy::AddCurveValue(const FSmartNameMapping& Mapping, const F
 	}
 }
 
+FAnimBlueprintDebugData* FAnimInstanceProxy::GetAnimBlueprintDebugData() const
+{
+#if WITH_EDITORONLY_DATA
+	if (bIsBeingDebugged)
+	{
+		UAnimBlueprint* AnimBP = GetAnimBlueprint();
+		return AnimBP ? AnimBP->GetDebugData() : nullptr;
+	}
+#endif
+
+	return nullptr;
+}
+
 #undef LOCTEXT_NAMESPACE
diff --git a/Engine/Source/Runtime/Engine/Public/Animation/AnimInstanceProxy.h b/Engine/Source/Runtime/Engine/Public/Animation/AnimInstanceProxy.h
--- a/Engine/Source/Runtime/Engine/Public/Animation/AnimInstanceProxy.h
+++ b/Engine/Source/Runtime/Engine/Public/Animation/AnimInstanceProxy.h
@@ -762,6 +762,9 @@ protected:
 	void AddCurveValue(const FSmartNameMapping& Mapping, const FName& CurveName, float Value);
 
 private:
+	/** Get the debug data for this instance's anim bp */
+	FAnimBlueprintDebugData* GetAnimBlueprintDebugData() const;
+
 	/** The component to world transform of the component we are running on */
 	FTransform ComponentTransform;
 
