Fix for complie error. Forgot to check iin cpp file when member variable was deleted in header

diff --git a/Engine/Source/Runtime/Engine/Private/AudioSettings.cpp b/Engine/Source/Runtime/Engine/Private/AudioSettings.cpp
--- a/Engine/Source/Runtime/Engine/Private/AudioSettings.cpp
+++ b/Engine/Source/Runtime/Engine/Private/AudioSettings.cpp
@@ -144,15 +144,6 @@ void UAudioSettings::PostEditChangeChainProperty(FPropertyChangedChainEvent& Pro
 void UAudioSettings::Serialize(FArchive& Ar)
 {
 	Super::Serialize(Ar);
-
-	if (Ar.IsLoading() || Ar.IsSaving())
-	{
-		if (bAllowVirtualizedSounds_DEPRECATED)
-		{
-			bAllowVirtualizedSounds_DEPRECATED = 0;
-			bAllowPlayWhenSilent = 1;
-		}
-	}
 }
 
 const FAudioQualitySettings& UAudioSettings::GetQualityLevelSettings(int32 QualityLevel) const
