Added missing Ar.UsingCustomVersion call preventing causing packaging to fail

diff --git a/Engine/Source/Runtime/Engine/Private/SoundClass.cpp b/Engine/Source/Runtime/Engine/Private/SoundClass.cpp
--- a/Engine/Source/Runtime/Engine/Private/SoundClass.cpp
+++ b/Engine/Source/Runtime/Engine/Private/SoundClass.cpp
@@ -223,11 +223,17 @@ void USoundClass::Serialize( FArchive& Ar )
 		Ar << EditorData_DEPRECATED;
 	}
 
+	Ar.UsingCustomVersion(FReleaseObjectVersion::GUID);
 	const int32 ReleaseObjectVersion = Ar.CustomVer(FReleaseObjectVersion::GUID);
+
 	if (ReleaseObjectVersion < FReleaseObjectVersion::SoundClass2DReverbSend)
 	{
+		UE_LOG(LogAudio, Warning, TEXT("//----------------------------------------------------------------------------------//"));
+		UE_LOG(LogAudio, Warning, TEXT("HIT: ReleaseObjectVersion < FReleaseObjectVersion::SoundClass2DReverbSend, so going to take project Default reverb send level instead of the SoundClass parameter..."));
 		const UAudioSettings* AudioSettings = GetDefault<UAudioSettings>();
 		Properties.Default2DReverbSendAmount = AudioSettings->DefaultReverbSendLevel_DEPRECATED;
+		UE_LOG(LogAudio, Warning, TEXT("We think the project value was: %f"), AudioSettings->DefaultReverbSendLevel_DEPRECATED);
+		UE_LOG(LogAudio, Warning, TEXT("//----------------------------------------------------------------------------------//"));
 	}
 }
 
