#UE4 Fix crash when loading a texture marked to use ETextureMipLoadOptions::OnlyFirstMip, but is loaded on a server where mips are stripped

diff --git a/Engine/Source/Runtime/Engine/Private/Texture2D.cpp b/Engine/Source/Runtime/Engine/Private/Texture2D.cpp
--- a/Engine/Source/Runtime/Engine/Private/Texture2D.cpp
+++ b/Engine/Source/Runtime/Engine/Private/Texture2D.cpp
@@ -740,7 +740,7 @@ void UTexture2D::UpdateResource()
 	// Note that using TF_FirstMip disables texture streaming, because the mip data becomes lost.
 	// Also, the cleanup of the platform data must go between UpdateCachedLODBias() and UpdateResource().
 	const bool bLoadOnlyFirstMip = UDeviceProfileManager::Get().GetActiveProfile()->GetTextureLODSettings()->GetMipLoadOptions(this) == ETextureMipLoadOptions::OnlyFirstMip;
-	if (bLoadOnlyFirstMip && PlatformData && FPlatformProperties::RequiresCookedData())
+	if (bLoadOnlyFirstMip && PlatformData && PlatformData->Mips.Num() > 0 && FPlatformProperties::RequiresCookedData())
 	{
 		const int32 FirstMip = FMath::Clamp(0, GetCachedLODBias(), PlatformData->Mips.Num() - 1);
 		// Remove any mips after the first mip.
