Don't crash if VT max resolution is set below VT tile size

diff --git a/Engine/Source/Runtime/Engine/Private/TextureDerivedData.cpp b/Engine/Source/Runtime/Engine/Private/TextureDerivedData.cpp
--- a/Engine/Source/Runtime/Engine/Private/TextureDerivedData.cpp
+++ b/Engine/Source/Runtime/Engine/Private/TextureDerivedData.cpp
@@ -450,6 +450,9 @@ static void GetTextureBuildSettings(
 		OutBuildSettings.bVirtualTextureEnableCompressZlib = CVarVTCompressZlib.GetValueOnAnyThread() != 0;
 		OutBuildSettings.bVirtualTextureEnableCompressCrunch = CVarVTCompressCrunch.GetValueOnAnyThread() != 0;
 		OutBuildSettings.VirtualTextureTileSize = FMath::RoundUpToPowerOfTwo(CVarVTTileSize.GetValueOnAnyThread());
+
+		// don't all max resolution to be less than VT tile size
+		OutBuildSettings.MaxTextureResolution = FMath::Max<uint32>(OutBuildSettings.MaxTextureResolution, OutBuildSettings.VirtualTextureTileSize);
 		
 		// 0 is a valid value for border size
 		// 1 would be OK in some cases, but breaks BC compressed formats, since it will result in physical tiles that aren't divisible by block size (4)
diff --git a/Engine/Source/Runtime/Engine/Private/VT/VirtualTextureDataBuilder.cpp b/Engine/Source/Runtime/Engine/Private/VT/VirtualTextureDataBuilder.cpp
--- a/Engine/Source/Runtime/Engine/Private/VT/VirtualTextureDataBuilder.cpp
+++ b/Engine/Source/Runtime/Engine/Private/VT/VirtualTextureDataBuilder.cpp
@@ -304,6 +304,7 @@ void FVirtualTextureDataBuilder::Build(const FTextureSourceData& InSourceData, c
 		break;
 	}
 
+	check(SettingsPerLayer[0].MaxTextureResolution >= (uint32)TileSize);
 	BlockSizeX = FMath::Min<uint32>(BlockSizeX, SettingsPerLayer[0].MaxTextureResolution);
 	BlockSizeY = FMath::Min<uint32>(BlockSizeY, SettingsPerLayer[0].MaxTextureResolution);
 
