- Fix stale state issue when concurrent voice stealing & spamming that can occur when concurrency set to 'StopQuietest' on looping sounds with significant voice steal fade times.
- Add more debug visualization support around ActiveSound lifetime

diff --git a/Engine/Source/Runtime/Engine/Classes/Sound/SoundConcurrency.h b/Engine/Source/Runtime/Engine/Classes/Sound/SoundConcurrency.h
--- a/Engine/Source/Runtime/Engine/Classes/Sound/SoundConcurrency.h
+++ b/Engine/Source/Runtime/Engine/Classes/Sound/SoundConcurrency.h
@@ -284,8 +284,8 @@ public:
 	/** Removes the active sound from concurrency tracking when active sound is stopped. */
 	void RemoveActiveSound(FActiveSound& ActiveSound);
 
-	// Fades sound if concurrency settings request to, or calls stop if not set
-	void FadeOutActiveSound(FActiveSound& ActiveSound);
+	/** Stops sound, applying concurrency rules for how to stop */
+	void StopDueToVoiceStealing(FActiveSound& ActiveSound);
 
 	/** Stops any active sounds due to max concurrency quietest sound resolution rule */
 	void UpdateQuietSoundsToStop();
diff --git a/Engine/Source/Runtime/Engine/Private/ActiveSound.cpp b/Engine/Source/Runtime/Engine/Private/ActiveSound.cpp
--- a/Engine/Source/Runtime/Engine/Private/ActiveSound.cpp
+++ b/Engine/Source/Runtime/Engine/Private/ActiveSound.cpp
@@ -76,29 +76,30 @@ FActiveSound::FActiveSound()
 	, FadeOut(EFadeOut::None)
 	, bIsOccluded(false)
 	, bAsyncOcclusionPending(false)
-	, PlaybackTime(0.f)
+	, PlaybackTime(0.0f)
+	, PlaybackTimeNonVirtualized(0.0f)
 	, MinCurrentPitch(1.0f)
-	, RequestedStartTime(0.f)
-	, CurrentAdjustVolumeMultiplier(1.f)
-	, TargetAdjustVolumeMultiplier(1.f)
-	, TargetAdjustVolumeStopTime(-1.f)
-	, VolumeMultiplier(1.f)
-	, PitchMultiplier(1.f)
+	, RequestedStartTime(0.0f)
+	, CurrentAdjustVolumeMultiplier(1.0f)
+	, TargetAdjustVolumeMultiplier(1.0f)
+	, TargetAdjustVolumeStopTime(-1.0f)
+	, VolumeMultiplier(1.0f)
+	, PitchMultiplier(1.0f)
 	, LowPassFilterFrequency(MAX_FILTER_FREQUENCY)
 	, CurrentOcclusionFilterFrequency(MAX_FILTER_FREQUENCY)
 	, CurrentOcclusionVolumeAttenuation(1.0f)
 	, SubtitlePriority(DEFAULT_SUBTITLE_PRIORITY)
 	, Priority(1.0f)
 	, VolumeConcurrency(0.0f)
-	, OcclusionCheckInterval(0.f)
+	, OcclusionCheckInterval(0.0f)
 	, LastOcclusionCheckTime(TNumericLimits<float>::Lowest())
 	, MaxDistance(WORLD_MAX)
 	, LastLocation(FVector::ZeroVector)
 	, AudioVolumeID(0)
-	, LastUpdateTime(0.f)
-	, SourceInteriorVolume(1.f)
+	, LastUpdateTime(0.0f)
+	, SourceInteriorVolume(1.0f)
 	, SourceInteriorLPF(MAX_FILTER_FREQUENCY)
-	, CurrentInteriorVolume(1.f)
+	, CurrentInteriorVolume(1.0f)
 	, CurrentInteriorLPF(MAX_FILTER_FREQUENCY)
 	, EnvelopeFollowerAttackTime(10)
 	, EnvelopeFollowerReleaseTime(100)
@@ -130,7 +131,7 @@ FActiveSound* FActiveSound::CreateVirtualCopy(const FActiveSound& InActiveSoundT
 	ActiveSound->bIsPlayingAudio = false;
 
 	ActiveSound->AudioDevice = &InAudioDevice;
-
+	ActiveSound->PlaybackTimeNonVirtualized = 0.0f;
 	ActiveSound->VolumeConcurrency = 1.0f;
 
 	ActiveSound->ConcurrencyGroupData.Reset();
@@ -458,12 +459,24 @@ bool FActiveSound::GetConcurrencyFadeDuration(float& OutFadeDuration) const
 	GetConcurrencyHandles(Handles);
 	for (FConcurrencyHandle& Handle : Handles)
 	{
+		// StopQuietest can spam if a looping ActiveSound isn't active longer than a virtualization update period, which
+		// can happen when a concurrency group is maxed and constantly evicting.  If the voice steal fade time is particularly
+		// long, this can flood the active sound count. Therefore, only use the voice steal fade time if the sound has been
+		// active for a sufficient period of time.
+		const bool bStopQuietest = Handle.Settings.ResolutionRule == EMaxConcurrentResolutionRule::StopQuietest;
+		if (bStopQuietest && IsLooping() && FMath::IsNearlyZero(PlaybackTimeNonVirtualized, 0.1f))
+		{
+			OutFadeDuration = 0.0f;
+			return false;
+		}
+
 		OutFadeDuration = OutFadeDuration < 0.0f
 			? Handle.Settings.VoiceStealReleaseTime
 			: FMath::Min(Handle.Settings.VoiceStealReleaseTime, OutFadeDuration);
 	}
 
-	if (OutFadeDuration < 0.0f)
+	// Negative if no handles are found, so return no fade required.
+	if (OutFadeDuration <= 0.0f)
 	{
 		OutFadeDuration = 0.0f;
 		return false;
@@ -627,7 +640,7 @@ void FActiveSound::UpdateWaveInstances(TArray<FWaveInstance*> &InWaveInstances,
 			{
 				check(WaveInstance);
 
-				const float WaveInstanceVolume = WaveInstance->GetVolumeWithDistanceAttenuation();
+				float WaveInstanceVolume = WaveInstance->GetVolumeWithDistanceAttenuation();
 				if (WaveInstanceVolume > VolumeConcurrency)
 				{
 					VolumeConcurrency = WaveInstanceVolume;
@@ -710,7 +723,7 @@ void FActiveSound::UpdateWaveInstances(TArray<FWaveInstance*> &InWaveInstances,
 	{
 		DebugColor = FColor::MakeRandomColor();
 	}
-	FAudioDebugger::DrawDebugInfo(*this, ThisSoundsWaveInstances);
+	FAudioDebugger::DrawDebugInfo(*this, ThisSoundsWaveInstances, DeltaTime);
 #endif // ENABLE_AUDIO_DEBUG
 
 	InWaveInstances.Append(ThisSoundsWaveInstances);
diff --git a/Engine/Source/Runtime/Engine/Private/AudioDebug.cpp b/Engine/Source/Runtime/Engine/Private/AudioDebug.cpp
--- a/Engine/Source/Runtime/Engine/Private/AudioDebug.cpp
+++ b/Engine/Source/Runtime/Engine/Private/AudioDebug.cpp
@@ -237,7 +237,7 @@ void FAudioDebugger::DrawDebugInfo(const FSoundSource& SoundSource)
 #endif // ENABLE_DRAW_DEBUG
 }
 
-void FAudioDebugger::DrawDebugInfo(const FActiveSound& ActiveSound, const TArray<FWaveInstance*>& ThisSoundsWaveInstances)
+void FAudioDebugger::DrawDebugInfo(const FActiveSound& ActiveSound, const TArray<FWaveInstance*>& ThisSoundsWaveInstances, const float DeltaTime)
 {
 #if ENABLE_DRAW_DEBUG
 	if (!ActiveSoundVisualizeModeCVar)
@@ -252,6 +252,16 @@ void FAudioDebugger::DrawDebugInfo(const FActiveSound& ActiveSound, const TArray
 		return;
 	}
 
+	const float PlaybackTime = ActiveSound.PlaybackTime;
+	const float PlaybackTimeNonVirtualized = ActiveSound.PlaybackTimeNonVirtualized;
+
+	// StopQuietest sounds can start and immediately stop repeatedly when subscribed
+	// concurrency is flooded, so don't show the initial frame.
+	if (FMath::IsNearlyZero(PlaybackTimeNonVirtualized))
+	{
+		return;
+	}
+
 	if (ActiveSoundVisualizeTypeCVar > 0)
 	{
 		if (ActiveSoundVisualizeTypeCVar == 1 && ActiveSound.GetAudioComponentID() == 0)
@@ -295,7 +305,7 @@ void FAudioDebugger::DrawDebugInfo(const FActiveSound& ActiveSound, const TArray
 		}
 
 		TWeakObjectPtr<UWorld> WorldPtr = ActiveSound.GetWeakWorld();
-		FAudioThread::RunCommandOnGameThread([Name, TextColor, CurTransform, DisplayValue, WorldPtr, CurMaxDistance]()
+		FAudioThread::RunCommandOnGameThread([Name, TextColor, CurTransform, DisplayValue, WorldPtr, CurMaxDistance, PlaybackTime, PlaybackTimeNonVirtualized, DeltaTime]()
 		{
 			if (WorldPtr.IsValid())
 			{
@@ -313,11 +323,11 @@ void FAudioDebugger::DrawDebugInfo(const FActiveSound& ActiveSound, const TArray
 					const float DisplayDbVolume = Audio::ConvertToDecibels(DisplayValue);
 					if (ActiveSoundVisualizeModeCVar == 1)
 					{
-						Descriptor = FString::Printf(TEXT(" (Vol: %.3f)"), DisplayValue);
+						Descriptor = FString::Printf(TEXT(" (Vol: %.3f [Active: %.2fs, Playing: %.2fs])"), DisplayValue, PlaybackTime, PlaybackTimeNonVirtualized);
 					}
 					else
 					{
-						Descriptor = FString::Printf(TEXT(" (Vol: %.3f dB)"), DisplayDbVolume);
+						Descriptor = FString::Printf(TEXT(" (Vol: %.3f dB [Active: %.2fs, Playing: %.2fs])"), DisplayDbVolume, PlaybackTime, PlaybackTimeNonVirtualized);
 					}
 					static const float DbColorMinVol = -30.0f;
 					const float DbVolume = FMath::Clamp(DisplayDbVolume, DbColorMinVol, 0.0f);
@@ -332,7 +342,7 @@ void FAudioDebugger::DrawDebugInfo(const FActiveSound& ActiveSound, const TArray
 				}
 
 				const FString Description = FString::Printf(TEXT("%s%s"), *Name, *Descriptor);
-				DrawDebugString(DebugWorld, Location + FVector(0, 0, 32), *Description, nullptr, Color, 0.03f, false);
+				DrawDebugString(DebugWorld, Location + FVector(0, 0, 32), *Description, nullptr, Color, DeltaTime, false);
 			}
 		}, GET_STATID(STAT_AudioDrawActiveSoundDebugInfo));
 	}
@@ -365,11 +375,12 @@ void FAudioDebugger::DrawDebugInfo(const FAudioVirtualLoop& VirtualLoop)
 		const TWeakObjectPtr<UWorld> World = ActiveSound.GetWeakWorld();
 		const FString Name = Sound->GetName();
 		const float DrawInterval = VirtualLoop.GetUpdateInterval();
-		FAudioThread::RunCommandOnGameThread([World, Transform, Name, DrawInterval]()
+		const float TimeVirtualized = VirtualLoop.GetTimeVirtualized();
+		FAudioThread::RunCommandOnGameThread([World, Transform, Name, DrawInterval, TimeVirtualized]()
 		{
 			if (World.IsValid())
 			{
-				const FString Description = FString::Printf(TEXT("%s [V]"), *Name);
+				const FString Description = FString::Printf(TEXT("%s [Virt: %.2fs]"), *Name, TimeVirtualized);
 				FVector Location = Transform.GetLocation();
 				FRotator Rotation = Transform.GetRotation().Rotator();
 				DrawDebugCrosshairs(World.Get(), Location, Rotation, 20.0f, FColor::Blue, false, DrawInterval, SDPG_Foreground);
diff --git a/Engine/Source/Runtime/Engine/Private/AudioDevice.cpp b/Engine/Source/Runtime/Engine/Private/AudioDevice.cpp
--- a/Engine/Source/Runtime/Engine/Private/AudioDevice.cpp
+++ b/Engine/Source/Runtime/Engine/Private/AudioDevice.cpp
@@ -2173,28 +2173,12 @@ void FAudioDevice::StopQuietSoundsDueToMaxConcurrency(TArray<FWaveInstance*>& Wa
 				continue;
 			}
 
-			if (ActiveSound->bIsStopping)
+			if (IsPendingStop(ActiveSound))
 			{
 				continue;
 			}
 
-			ConcurrencyManager.RemoveActiveSound(*ActiveSound);
-			ConcurrencyManager.FadeOutActiveSound(*ActiveSound);
-
-			const bool bDoRangeCheck = false;
-			FAudioVirtualLoop VirtualLoop;
-			if (FAudioVirtualLoop::Virtualize(*ActiveSound, bDoRangeCheck, VirtualLoop))
-			{
-				// Clear must be called after AddSoundToStop to ensure AudioComponent is properly removed from AudioComponentIDToActiveSoundMap
-				ActiveSound->ClearAudioComponent();
-
-				if (USoundBase* Sound = ActiveSound->GetSound())
-				{
-					UE_LOG(LogAudio, Verbose, TEXT("Playing ActiveSound %s Virtualizing: Invalidated by 'StopQuiestest' Concurrency Rule."), *Sound->GetName());
-				}
-				AddVirtualLoop(VirtualLoop);
-			}
-			ActiveSoundsCopy.RemoveAtSwap(i, 1, false);
+			ConcurrencyManager.StopDueToVoiceStealing(*ActiveSound);
 		}
 	}
 
@@ -2737,22 +2721,30 @@ void FAudioDevice::VirtualizeInactiveLoops()
 	const bool bDoRangeCheck = true;
 	for (FActiveSound* ActiveSound : ActiveSounds)
 	{
+		// Don't virtualize if set to fade out
+		if (ActiveSound->FadeOut != FActiveSound::EFadeOut::None)
+		{
+			continue;
+		}
+
 		// If already pending stop, don't attempt to virtualize
-		if (!IsPendingStop(ActiveSound))
+		if (IsPendingStop(ActiveSound))
 		{
-			FAudioVirtualLoop VirtualLoop;
-			if (FAudioVirtualLoop::Virtualize(*ActiveSound, bDoRangeCheck, VirtualLoop))
-			{
-				AddSoundToStop(ActiveSound);
+			continue;
+		}
 
-				// Clear must be called after AddSoundToStop to ensure AudioComponent is properly removed from AudioComponentIDToActiveSoundMap
-				ActiveSound->ClearAudioComponent();
-				if (USoundBase* Sound = ActiveSound->GetSound())
-				{
-					UE_LOG(LogAudio, Verbose, TEXT("Playing ActiveSound %s Virtualizing: Out of audible range."), *Sound->GetName());
-				}
-				AddVirtualLoop(VirtualLoop);
+		FAudioVirtualLoop VirtualLoop;
+		if (FAudioVirtualLoop::Virtualize(*ActiveSound, bDoRangeCheck, VirtualLoop))
+		{
+			AddSoundToStop(ActiveSound);
+
+			// Clear must be called after AddSoundToStop to ensure AudioComponent is properly removed from AudioComponentIDToActiveSoundMap
+			ActiveSound->ClearAudioComponent();
+			if (USoundBase* Sound = ActiveSound->GetSound())
+			{
+				UE_LOG(LogAudio, Verbose, TEXT("Playing ActiveSound %s Virtualizing: Out of audible range."), *Sound->GetName());
 			}
+			AddVirtualLoop(VirtualLoop);
 		}
 	}
 }
@@ -3519,7 +3511,9 @@ void FAudioDevice::UpdateActiveSoundPlaybackTime(bool bIsGameTicking)
 		for (FActiveSound* ActiveSound : ActiveSounds)
 		{
 			// Scale the playback time with the device delta time and the current "min pitch" of the sounds which would play on it.
-			ActiveSound->PlaybackTime += GetDeviceDeltaTime() * ActiveSound->MinCurrentPitch;
+			const float DeltaTimePitchCorrected = GetDeviceDeltaTime() * ActiveSound->MinCurrentPitch;
+			ActiveSound->PlaybackTime += DeltaTimePitchCorrected;
+			ActiveSound->PlaybackTimeNonVirtualized += DeltaTimePitchCorrected;
 		}
 	}
 	else if (GIsEditor)
@@ -3529,7 +3523,9 @@ void FAudioDevice::UpdateActiveSoundPlaybackTime(bool bIsGameTicking)
 			if (ActiveSound->bIsPreviewSound)
 			{
 				// Scale the playback time with the device delta time and the current "min pitch" of the sounds which would play on it.
-				ActiveSound->PlaybackTime += GetDeviceDeltaTime() * ActiveSound->MinCurrentPitch;
+				const float DeltaTimePitchCorrected = GetDeviceDeltaTime() * ActiveSound->MinCurrentPitch;
+				ActiveSound->PlaybackTime += DeltaTimePitchCorrected;
+				ActiveSound->PlaybackTimeNonVirtualized += DeltaTimePitchCorrected;
 			}
 		}
 	}
@@ -4146,7 +4142,7 @@ void FAudioDevice::AddNewActiveSound(const FActiveSound& NewActiveSound)
 void FAudioDevice::AddNewActiveSoundInternal(const FActiveSound& NewActiveSound, FAudioVirtualLoop* VirtualLoopToRetrigger)
 {
 	LLM_SCOPE(ELLMTag::AudioMisc);
-
+	
 	if (!IsInAudioThread())
 	{
 		DECLARE_CYCLE_STAT(TEXT("FAudioThreadTask.AddNewActiveSound"), STAT_AudioAddNewActiveSound, STATGROUP_AudioThreadCommands);
@@ -4244,9 +4240,9 @@ void FAudioDevice::AddNewActiveSoundInternal(const FActiveSound& NewActiveSound,
 		ActiveSound = ConcurrencyManager.CreateNewActiveSound(NewActiveSound, VirtualLoopToRetrigger != nullptr);
 	}
 
+	// Didn't pass concurrency, and not an attempt to revive from virtualization, so see if candidate for virtualization
 	if (!ActiveSound)
 	{
-		// Didn't pass concurrency, so attempt to add to re-trigger if applicable
 		if (!VirtualLoopToRetrigger)
 		{
 			const bool bDoRangeCheck = false;
@@ -4506,11 +4502,7 @@ void FAudioDevice::AddSoundToStop(FActiveSound* SoundToStop)
 	PendingSoundsToStop.Add(SoundToStop, &bAlreadyPending);
 	if (!bAlreadyPending)
 	{
-		const uint64 AudioComponentID = SoundToStop->GetAudioComponentID();
-		if (AudioComponentID > 0)
-		{
-			AudioComponentIDToActiveSoundMap.Remove(AudioComponentID);
-		}
+		UnlinkActiveSoundFromComponent(*SoundToStop);
 
 		if (VirtualLoops.Contains(SoundToStop))
 		{
@@ -4871,6 +4863,15 @@ int32 FAudioDevice::FindClosestListenerIndex(const FTransform& SoundTransform) c
 	return INDEX_NONE;
 }
 
+void FAudioDevice::UnlinkActiveSoundFromComponent(const FActiveSound& InActiveSound)
+{
+	const uint64 AudioComponentID = InActiveSound.GetAudioComponentID();
+	if (AudioComponentID > 0)
+	{
+		AudioComponentIDToActiveSoundMap.Remove(AudioComponentID);
+	}
+}
+
 void FAudioDevice::GetAzimuth(const FAttenuationListenerData& ListenerData, float& OutAzimuth, float& OutAbsoluteAzimuth) const
 {
 	const FVector& ListenerForwardDir = ListenerData.ListenerTransform.GetUnitAxis(EAxis::X);
@@ -5791,16 +5792,9 @@ void FAudioDevice::UpdateVirtualLoops(bool bForceUpdate)
 				continue;
 			}
 
-			const float DeltaTime = GetDeviceDeltaTime();
-
-			// Keep playback time up-to-date as it may be used to evaluate whether or
-			// not virtual sound is eligible for playback when compared against
-			// actively playing sounds in concurrency checks.
-			ActiveSound.PlaybackTime += DeltaTime * ActiveSound.MinCurrentPitch;
-
 			// If the loop is ready to realize, add to array to be re-triggered
 			// outside of the loop to avoid map manipulation while iterating.
-			if (VirtualLoop.CanRealize(DeltaTime, bForceUpdate))
+			if (VirtualLoop.Update(GetDeviceDeltaTime(), bForceUpdate))
 			{
 				VirtualLoopsToRetrigger.Add(VirtualLoop);
 			}
@@ -5820,13 +5814,9 @@ void FAudioDevice::UpdateVirtualLoops(bool bForceUpdate)
 		for (FVirtualLoopPair& Pair : VirtualLoops)
 		{
 			FActiveSound* ActiveSound = Pair.Key;
+			check(ActiveSound);
 
-			const uint64 ComponentID = ActiveSound->GetAudioComponentID();
-			if (ComponentID > 0)
-			{
-				AudioComponentIDToActiveSoundMap.Remove(ComponentID);
-			}
-
+			UnlinkActiveSoundFromComponent(*ActiveSound);
 			AddNewActiveSound(*ActiveSound);
 
 			ActiveSound->ClearAudioComponent();
diff --git a/Engine/Source/Runtime/Engine/Private/AudioVirtualLoop.cpp b/Engine/Source/Runtime/Engine/Private/AudioVirtualLoop.cpp
--- a/Engine/Source/Runtime/Engine/Private/AudioVirtualLoop.cpp
+++ b/Engine/Source/Runtime/Engine/Private/AudioVirtualLoop.cpp
@@ -45,6 +45,7 @@ FAutoConsoleVariableRef CVarVirtualLoopsUpdateRateMax(
 
 FAudioVirtualLoop::FAudioVirtualLoop()
 	: TimeSinceLastUpdate(0.0f)
+	, TimeVirtualized(0.0f)
 	, UpdateInterval(0.0f)
 	, ActiveSound(nullptr)
 {
@@ -100,6 +101,16 @@ void FAudioVirtualLoop::CalculateUpdateInterval()
 	UpdateInterval = FMath::Lerp(VirtualLoopsUpdateRateMinCVar, VirtualLoopsUpdateRateMaxCVar, DistanceRatioClamped);
 }
 
+float FAudioVirtualLoop::GetTimeVirtualized() const
+{
+	return TimeVirtualized;
+}
+
+float FAudioVirtualLoop::GetUpdateInterval() const
+{
+	return UpdateInterval;
+}
+
 FActiveSound& FAudioVirtualLoop::GetActiveSound()
 {
 	check(ActiveSound);
@@ -184,10 +195,16 @@ void FAudioVirtualLoop::UpdateFocusData(float DeltaTime)
 	ActiveSound->UpdateFocusData(DeltaTime, ListenerData);
 }
 
-bool FAudioVirtualLoop::CanRealize(float DeltaTime, bool bForceUpdate)
+bool FAudioVirtualLoop::Update(float DeltaTime, bool bForceUpdate)
 {
-	const float UpdateDelta = TimeSinceLastUpdate + DeltaTime;
+	// Keep playback time up-to-date as it may be used to evaluate whether or
+	// not virtual sound is eligible for playback when compared against
+	// actively playing sounds in concurrency checks.
+	const float DeltaTimePitchCorrected = DeltaTime * ActiveSound->MinCurrentPitch;
+	ActiveSound->PlaybackTime += DeltaTimePitchCorrected;
+	TimeVirtualized += DeltaTimePitchCorrected;
 
+	const float UpdateDelta = TimeSinceLastUpdate + DeltaTime;
 	if (bForceUpdate)
 	{
 		TimeSinceLastUpdate = 0.0f;
diff --git a/Engine/Source/Runtime/Engine/Private/SoundConcurrency.cpp b/Engine/Source/Runtime/Engine/Private/SoundConcurrency.cpp
--- a/Engine/Source/Runtime/Engine/Private/SoundConcurrency.cpp
+++ b/Engine/Source/Runtime/Engine/Private/SoundConcurrency.cpp
@@ -343,6 +343,7 @@ FActiveSound* FSoundConcurrencyManager::CreateNewActiveSound(const FActiveSound&
 	if (!ConcurrencyHandles.Num())
 	{
 		FActiveSound* ActiveSound = new FActiveSound(NewActiveSound);
+		ActiveSound->PlaybackTimeNonVirtualized = 0.0f;
 		ActiveSound->SetAudioDevice(AudioDevice);
 		return ActiveSound;
 	}
@@ -716,22 +717,7 @@ FActiveSound* FSoundConcurrencyManager::CreateAndEvictActiveSounds(const FActive
 			continue;
 		}
 
-		FadeOutActiveSound(*SoundToEvict);
-
-		const bool bDoRangeCheck = false;
-		FAudioVirtualLoop VirtualLoop;
-		if (FAudioVirtualLoop::Virtualize(*SoundToEvict, bDoRangeCheck, VirtualLoop))
-		{
-			// Don't report that the sound is 'finished' by clearing the audio component of the evicted
-			// copy that is being virtualized and effectively stopped at the end of this frame.
-			SoundToEvict->ClearAudioComponent();
-
-			if (USoundBase* Sound = SoundToEvict->GetSound())
-			{
-				UE_LOG(LogAudio, Verbose, TEXT("Playing ActiveSound %s Virtualizing: Evicted due to concurrency rules."), *Sound->GetName());
-			}
-			AudioDevice->AddVirtualLoop(VirtualLoop);
-		}
+		StopDueToVoiceStealing(*SoundToEvict);
 	}
 
 	return ActiveSound;
@@ -810,20 +796,39 @@ void FSoundConcurrencyManager::RemoveActiveSound(FActiveSound& ActiveSound)
 	ActiveSound.ConcurrencyGroupData.Reset();
 }
 
-void FSoundConcurrencyManager::FadeOutActiveSound(FActiveSound& ActiveSound)
+void FSoundConcurrencyManager::StopDueToVoiceStealing(FActiveSound& ActiveSound)
 {
-	check(AudioDevice);
+	check(ActiveSound.AudioDevice);
 
 	float FadeOutDuration = 0.0f;
-	if (ActiveSound.GetConcurrencyFadeDuration(FadeOutDuration))
+	bool bRequiresConcurrencyFade = ActiveSound.GetConcurrencyFadeDuration(FadeOutDuration);
+	if (bRequiresConcurrencyFade)
 	{
-		ActiveSound.FadeOut = FActiveSound::EFadeOut::Concurrency;
-		ActiveSound.TargetAdjustVolumeMultiplier = 0.0f;
-		ActiveSound.TargetAdjustVolumeStopTime = ActiveSound.PlaybackTime + FadeOutDuration;
+		ActiveSound.AudioDevice->UnlinkActiveSoundFromComponent(ActiveSound);
 	}
 	else
 	{
-		AudioDevice->AddSoundToStop(&ActiveSound);
+		ActiveSound.AudioDevice->AddSoundToStop(&ActiveSound);
+	}
+
+	const bool bDoRangeCheck = false;
+	FAudioVirtualLoop VirtualLoop;
+	if (FAudioVirtualLoop::Virtualize(ActiveSound, bDoRangeCheck, VirtualLoop))
+	{
+		ActiveSound.ClearAudioComponent();
+		if (USoundBase* Sound = ActiveSound.GetSound())
+		{
+			UE_LOG(LogAudio, Verbose, TEXT("Playing ActiveSound %s Virtualizing: Sound's voice stollen due to concurrency group maximum met."), *Sound->GetName());
+		}
+		ActiveSound.AudioDevice->AddVirtualLoop(VirtualLoop);
+	}
+
+	// Apply concurrency fade after potentially virtualizing to avoid transferring undesired new concurrency fade state
+	if (bRequiresConcurrencyFade)
+	{
+		ActiveSound.FadeOut = FActiveSound::EFadeOut::Concurrency;
+		ActiveSound.TargetAdjustVolumeMultiplier = 0.0f;
+		ActiveSound.TargetAdjustVolumeStopTime = ActiveSound.PlaybackTime + FadeOutDuration;
 	}
 }
 
diff --git a/Engine/Source/Runtime/Engine/Private/SoundNodeLooping.cpp b/Engine/Source/Runtime/Engine/Private/SoundNodeLooping.cpp
--- a/Engine/Source/Runtime/Engine/Private/SoundNodeLooping.cpp
+++ b/Engine/Source/Runtime/Engine/Private/SoundNodeLooping.cpp
@@ -30,7 +30,7 @@ void USoundNodeLooping::ParseNodes( FAudioDevice* AudioDevice, const UPTRINT Nod
 	}
 
 #if !(NO_LOGGING || UE_BUILD_SHIPPING || UE_BUILD_TEST)
-	if (bLoopIndefinitely && !ActiveSound.bWarnedAboutOrphanedLooping && ActiveSound.GetAudioComponentID() == 0)
+	if (bLoopIndefinitely && !ActiveSound.bWarnedAboutOrphanedLooping && ActiveSound.GetAudioComponentID() == 0 && ActiveSound.FadeOut == FActiveSound::EFadeOut::None)
 	{
 		UE_LOG(LogAudio, Warning, TEXT("Detected orphaned looping sound '%s'."), *ActiveSound.GetSound()->GetName());
 		ActiveSound.bWarnedAboutOrphanedLooping = true;
diff --git a/Engine/Source/Runtime/Engine/Private/SoundWave.cpp b/Engine/Source/Runtime/Engine/Private/SoundWave.cpp
--- a/Engine/Source/Runtime/Engine/Private/SoundWave.cpp
+++ b/Engine/Source/Runtime/Engine/Private/SoundWave.cpp
@@ -1566,7 +1566,7 @@ void USoundWave::Parse(FAudioDevice* AudioDevice, const UPTRINT NodeWaveInstance
 	{
 		WaveInstance->bIsFinished = false;
 #if !(NO_LOGGING || UE_BUILD_SHIPPING || UE_BUILD_TEST)
-		if (!ActiveSound.bWarnedAboutOrphanedLooping && ActiveSound.GetAudioComponentID() == 0)
+		if (!ActiveSound.bWarnedAboutOrphanedLooping && ActiveSound.GetAudioComponentID() == 0 && ActiveSound.FadeOut == FActiveSound::EFadeOut::None)
 		{
 			UE_LOG(LogAudio, Warning, TEXT("Detected orphaned looping sound '%s'."), *ActiveSound.GetSound()->GetName());
 			ActiveSound.bWarnedAboutOrphanedLooping = true;
diff --git a/Engine/Source/Runtime/Engine/Public/ActiveSound.h b/Engine/Source/Runtime/Engine/Public/ActiveSound.h
--- a/Engine/Source/Runtime/Engine/Public/ActiveSound.h
+++ b/Engine/Source/Runtime/Engine/Public/ActiveSound.h
@@ -462,7 +462,12 @@ public:
 	/** Whether or not there is an async occlusion trace pending */
 	FThreadSafeBool bAsyncOcclusionPending;
 
+	/** Duration between now and when the sound has been started. */
 	float PlaybackTime;
+
+	/** If virtualized, duration between last time virtualized and now. */
+	float PlaybackTimeNonVirtualized;
+
 	float MinCurrentPitch;
 	float RequestedStartTime;
 
diff --git a/Engine/Source/Runtime/Engine/Public/Audio/AudioDebug.h b/Engine/Source/Runtime/Engine/Public/Audio/AudioDebug.h
--- a/Engine/Source/Runtime/Engine/Public/Audio/AudioDebug.h
+++ b/Engine/Source/Runtime/Engine/Public/Audio/AudioDebug.h
@@ -22,8 +22,8 @@ class ENGINE_API FAudioDebugger
 public:
 	FAudioDebugger();
 
-	static void DrawDebugInfo(const FActiveSound& ActiveSound, const TArray<FWaveInstance*>& ThisSoundsWaveInstances);
 	static void DrawDebugInfo(const FSoundSource& SoundSource);
+	static void DrawDebugInfo(const FActiveSound& ActiveSound, const TArray<FWaveInstance*>& ThisSoundsWaveInstances, const float DeltaTime);
 	static void DrawDebugInfo(const FAudioVirtualLoop& VirtualLoop);
 	static bool PostStatModulatorHelp(UWorld* World, FCommonViewportClient* ViewportClient, const TCHAR* Stream);
 	static int32 RenderStatCues(UWorld* World, FViewport* Viewport, FCanvas* Canvas, int32 X, int32 Y, const FVector* ViewLocation, const FRotator* ViewRotation);
diff --git a/Engine/Source/Runtime/Engine/Public/AudioVirtualLoop.h b/Engine/Source/Runtime/Engine/Public/AudioVirtualLoop.h
--- a/Engine/Source/Runtime/Engine/Public/AudioVirtualLoop.h
+++ b/Engine/Source/Runtime/Engine/Public/AudioVirtualLoop.h
@@ -13,6 +13,7 @@ struct ENGINE_API FAudioVirtualLoop
 {
 private:
 	float TimeSinceLastUpdate;
+	float TimeVirtualized;
 	float UpdateInterval;
 
 	FActiveSound* ActiveSound;
@@ -47,10 +48,15 @@ public:
 	 */
 	FActiveSound& GetActiveSound();
 
+	/**
+	 * Returns the time the sound has been virtualized
+	 */
+	float GetTimeVirtualized() const;
+
 	/**
 	 * Returns the wait interval being observed before next update
 	 */
-	float GetUpdateInterval() const { return UpdateInterval; }
+	float GetUpdateInterval() const;
 
 	/**
 	 */
@@ -71,5 +77,5 @@ public:
 	  * Updates the loop and checks if ready to play (or 'realize').
 	  * Returns whether or not the sound is ready to be realized.
 	  */
-	bool CanRealize(float DeltaTime, bool bForceUpdate);
+	bool Update(float DeltaTime, bool bForceUpdate);
 };
