Fix main (missing file submission from ue-release)

diff --git a/Engine/Source/Runtime/Engine/Public/AudioDevice.h b/Engine/Source/Runtime/Engine/Public/AudioDevice.h
--- a/Engine/Source/Runtime/Engine/Public/AudioDevice.h
+++ b/Engine/Source/Runtime/Engine/Public/AudioDevice.h
@@ -1014,6 +1014,9 @@ public:
 	static int32 FindClosestListenerIndex(const FTransform& SoundTransform, const TArray<FListener>& InListeners);
 	int32 FindClosestListenerIndex(const FTransform& SoundTransform) const;
 
+	/** Disables ActiveSound from responding to calls from its associated AudioComponent. */
+	void UnlinkActiveSoundFromComponent(const FActiveSound& InActiveSound);
+
 	/** Return the audio stream time */
 	virtual double GetAudioTime() const
 	{
