early-exit AActor::NeedsLoadForTargetPlatform if there's either no root component or the detail mode is low, which should be the typical default case.  This will avoid much of the overhead from finding device profiles in most instances.

diff --git a/Engine/Source/Runtime/Engine/Private/Actor.cpp b/Engine/Source/Runtime/Engine/Private/Actor.cpp
--- a/Engine/Source/Runtime/Engine/Private/Actor.cpp
+++ b/Engine/Source/Runtime/Engine/Private/Actor.cpp
@@ -349,6 +349,12 @@ bool AActor::IsEditorOnly() const
 
 bool AActor::NeedsLoadForTargetPlatform(const ITargetPlatform* TargetPlatform) const
 {
+	// this is expected to be by far the most common case, saves us some time in the cook.
+	if (!RootComponent || RootComponent->DetailMode == EDetailMode::Low)
+	{
+		return true;
+	}
+
 	if(UDeviceProfile* DeviceProfile = UDeviceProfileManager::Get().FindProfile(TargetPlatform->IniPlatformName()))
 	{
 		// get local scalability CVars that could cull this actor
@@ -361,7 +367,7 @@ bool AActor::NeedsLoadForTargetPlatform(const ITargetPlatform* TargetPlatform) c
 				// Check root component's detail mode.
 				// If e.g. the component's detail mode is High and the platform detail is Medium,
 				// then we should cull it.
-				if(RootComponent && (int32)RootComponent->DetailMode > CVarDetailMode)
+				if((int32)RootComponent->DetailMode > CVarDetailMode)
 				{
 					return false;
 				}
