#UE4 CIS fix

diff --git a/Engine/Source/Runtime/Engine/Private/Actor.cpp b/Engine/Source/Runtime/Engine/Private/Actor.cpp
--- a/Engine/Source/Runtime/Engine/Private/Actor.cpp
+++ b/Engine/Source/Runtime/Engine/Private/Actor.cpp
@@ -350,7 +350,7 @@ bool AActor::IsEditorOnly() const
 bool AActor::NeedsLoadForTargetPlatform(const ITargetPlatform* TargetPlatform) const
 {
 	// this is expected to be by far the most common case, saves us some time in the cook.
-	if (!RootComponent || RootComponent->DetailMode == EDetailMode::Low)
+	if (!RootComponent || RootComponent->DetailMode == EDetailMode::DM_Low)
 	{
 		return true;
 	}
