FLevelUtils::IsLevelLoaded will now return false when the level isn't found, rather than failing a check. Removed FLevelUtils::MarkLevelForLoading and MarkLevelForUnloading, which were unimplemented and unused.

diff --git a/Engine/Source/Runtime/Engine/Private/LevelUtils.cpp b/Engine/Source/Runtime/Engine/Private/LevelUtils.cpp
--- a/Engine/Source/Runtime/Engine/Private/LevelUtils.cpp
+++ b/Engine/Source/Runtime/Engine/Private/LevelUtils.cpp
@@ -222,44 +222,9 @@ bool FLevelUtils::IsLevelLoaded(ULevel* Level)
 	}
 
 	ULevelStreaming* StreamingLevel = FindStreamingLevel( Level );
-	checkf( StreamingLevel, TEXT("Couldn't find streaming level" ) );
-
-	// @todo: Dave, please come talk to me before implementing anything like this.
-	return true;
-}
-
-/**
- * Flags an unloaded level for loading.
- *
- * @param	Level		The level to modify.
- */
-void FLevelUtils::MarkLevelForLoading(ULevel* Level)
-{
-	// If the level is valid and not the persistent level (which is always loaded) . . .
-	if ( Level && !Level->IsPersistentLevel() )
-	{
-		// Mark the level's stream for load.
-		ULevelStreaming* StreamingLevel = FindStreamingLevel( Level );
-		checkf( StreamingLevel, TEXT("Couldn't find streaming level" ) );
-		// @todo: Dave, please come talk to me before implementing anything like this.
-	}
+	return (StreamingLevel != nullptr);
 }
 
-/**
- * Flags a loaded level for unloading.
- *
- * @param	Level		The level to modify.
- */
-void FLevelUtils::MarkLevelForUnloading(ULevel* Level)
-{
-	// If the level is valid and not the persistent level (which is always loaded) . . .
-	if ( Level && !Level->IsPersistentLevel() )
-	{
-		ULevelStreaming* StreamingLevel = FindStreamingLevel( Level );
-		checkf( StreamingLevel, TEXT("Couldn't find streaming level" ) );
-		// @todo: Dave, please come talk to me before implementing anything like this.
-	}
-}
 
 /////////////////////////////////////////////////////////////////////////////////////////
 //
diff --git a/Engine/Source/Runtime/Engine/Public/LevelUtils.h b/Engine/Source/Runtime/Engine/Public/LevelUtils.h
--- a/Engine/Source/Runtime/Engine/Public/LevelUtils.h
+++ b/Engine/Source/Runtime/Engine/Public/LevelUtils.h
@@ -67,19 +67,6 @@ public:
 	 */
 	static bool IsLevelLoaded(ULevel* Level);
 
-	/**
-	 * Flags an unloaded level for loading.
-	 *
-	 * @param	Level		The level to modify.
-	 */
-	static void MarkLevelForLoading(ULevel* Level);
-
-	/**
-	 * Flags a loaded level for unloading.
-	 *
-	 * @param	Level		The level to modify.
-	 */
-	static void MarkLevelForUnloading(ULevel* Level);
 
 	///////////////////////////////////////////////////////////////////////////
 	// Level visibility.
