Implemented different fix, rather than rely on UpdateCounter added a UpdateNodeCounter, for:
"Animation state machines don't update when used as inputs to Layer nodes
#fix Make sure we only update state machines once per frame, as nodes can be updated multiple times through SubAnimInstance(s)

diff --git a/Engine/Source/Runtime/Engine/Private/Animation/AnimInstanceProxy.cpp b/Engine/Source/Runtime/Engine/Private/Animation/AnimInstanceProxy.cpp
--- a/Engine/Source/Runtime/Engine/Private/Animation/AnimInstanceProxy.cpp
+++ b/Engine/Source/Runtime/Engine/Private/Animation/AnimInstanceProxy.cpp
@@ -43,11 +43,12 @@ void FAnimInstanceProxy::UpdateAnimationNode_WithRoot(float DeltaSeconds, FAnimN
 	DECLARE_SCOPE_HIERARCHICAL_COUNTER_FUNC()
 	if(InRootNode != nullptr)
 	{
-		if (FrameCounterForUpdate != GFrameCounter)
+		if (FrameCounterForNodeUpdate != GFrameCounter)
 		{
+			FrameCounterForNodeUpdate = GFrameCounter;
 			UpdateCounter.Increment();
 		}
-
+		
 		InRootNode->Update_AnyThread(FAnimationUpdateContext(this, DeltaSeconds));
 
 		// We've updated the graph, now update the fractured saved pose sections
diff --git a/Engine/Source/Runtime/Engine/Public/Animation/AnimInstanceProxy.h b/Engine/Source/Runtime/Engine/Public/Animation/AnimInstanceProxy.h
--- a/Engine/Source/Runtime/Engine/Public/Animation/AnimInstanceProxy.h
+++ b/Engine/Source/Runtime/Engine/Public/Animation/AnimInstanceProxy.h
@@ -120,6 +120,7 @@ public:
 		, SyncGroupWriteIndex(0)
 		, RootMotionMode(ERootMotionMode::NoRootMotionExtraction)
 		, FrameCounterForUpdate(0)
+		, FrameCounterForNodeUpdate(0)
 		, CacheBonesRecursionCounter(0)
 		, bUpdatingRoot(false)
 		, bBoneCachesInvalidated(false)
@@ -142,6 +143,7 @@ public:
 		, SyncGroupWriteIndex(0)
 		, RootMotionMode(ERootMotionMode::NoRootMotionExtraction)
 		, FrameCounterForUpdate(0)
+		, FrameCounterForNodeUpdate(0)
 		, CacheBonesRecursionCounter(0)
 		, bUpdatingRoot(false)
 		, bBoneCachesInvalidated(false)
@@ -861,6 +863,7 @@ protected:
 
 	// Sync counter
 	uint64 FrameCounterForUpdate;
+	uint64 FrameCounterForNodeUpdate;
 
 private:
 	// Root motion extracted from animation since the last time ConsumeExtractedRootMotion was called
