Update UCameraComponent to have an editable Camera mesh that can be marked as visible at runtime.

diff --git a/Engine/Source/Runtime/Engine/Classes/Camera/CameraComponent.h b/Engine/Source/Runtime/Engine/Classes/Camera/CameraComponent.h
--- a/Engine/Source/Runtime/Engine/Classes/Camera/CameraComponent.h
+++ b/Engine/Source/Runtime/Engine/Classes/Camera/CameraComponent.h
@@ -70,6 +70,12 @@ class ENGINE_API UCameraComponent : public USceneComponent
 	virtual FText GetFilmbackText() const;
 #endif
 
+#if WITH_EDITORONLY_DATA
+	/** If the camera mesh is visible in game. Only relevant when running editor builds. */
+	UPROPERTY(EditAnywhere, BlueprintReadWrite, Category = Camera)
+	uint8 bCameraMeshHiddenInGame : 1;
+#endif
+
 	/** True if the camera's orientation and position should be locked to the HMD */
 	UPROPERTY(EditAnywhere, BlueprintReadWrite, Category = CameraOptions)
 	uint8 bLockToHmd : 1;
@@ -136,7 +142,7 @@ public:
 
 #if WITH_EDITORONLY_DATA
 	virtual void SetCameraMesh(UStaticMesh* Mesh);
-#endif
+#endif 
 
 protected:
 
@@ -144,17 +150,24 @@ protected:
 	// The frustum component used to show visually where the camera field of view is
 	class UDrawFrustumComponent* DrawFrustum;
 
-	UPROPERTY(transient)
+	UPROPERTY(BlueprintReadWrite, EditAnywhere, Category = Camera)
 	class UStaticMesh* CameraMesh;
 
 	// The camera mesh to show visually where the camera is placed
 	class UStaticMeshComponent* ProxyMeshComponent;
-	
+
 	virtual void ResetProxyMeshTransform();
 
 	/** Ensure the proxy mesh is in the correct place */
 	void UpdateProxyMeshTransform();
-#endif
+
+#endif	// WITH_EDITORONLY_DATA
+
+	/**
+	* Internal function for updating the camera mesh visibility in PIE
+	*/
+	UFUNCTION(BlueprintSetter)
+	void OnCameraMeshHiddenChanged();
 
 	/** An optional extra transform to adjust the final view without moving the component, in the camera's local space */
 	FTransform AdditiveOffset;
@@ -212,5 +225,5 @@ public:
 	/** DEPRECATED: use bUsePawnControlRotation instead */
 	UPROPERTY()
 	uint32 bUseControllerViewRotation_DEPRECATED:1;
-#endif
+#endif	// WITH_EDITORONLY_DATA
 };
diff --git a/Engine/Source/Runtime/Engine/Private/Camera/CameraComponent.cpp b/Engine/Source/Runtime/Engine/Private/Camera/CameraComponent.cpp
--- a/Engine/Source/Runtime/Engine/Private/Camera/CameraComponent.cpp
+++ b/Engine/Source/Runtime/Engine/Private/Camera/CameraComponent.cpp
@@ -37,6 +37,7 @@ UCameraComponent::UCameraComponent(const FObjectInitializer& ObjectInitializer)
 	}
 
 	bUseControllerViewRotation_DEPRECATED = true; // the previous default value before bUsePawnControlRotation replaced this var.
+	bCameraMeshHiddenInGame = true;
 #endif
 
 	FieldOfView = 90.0f;
@@ -62,6 +63,18 @@ void UCameraComponent::OnUpdateTransform(EUpdateTransformFlags UpdateTransformFl
 }
 
 #if WITH_EDITORONLY_DATA
+void UCameraComponent::PostEditChangeProperty(FPropertyChangedEvent & PropertyChangedEvent)
+{
+	Super::PostEditChangeProperty(PropertyChangedEvent);
+
+	const FName PropertyName = PropertyChangedEvent.Property ? PropertyChangedEvent.Property->GetFName() : FName();
+	if (PropertyName == GET_MEMBER_NAME_CHECKED(UCameraComponent, bCameraMeshHiddenInGame))
+	{
+		OnCameraMeshHiddenChanged();
+	}
+
+	RefreshVisualRepresentation();
+}
 
 void UCameraComponent::AddReferencedObjects(UObject* InThis, FReferenceCollector& Collector)
 {
@@ -85,7 +98,7 @@ void UCameraComponent::OnComponentDestroyed(bool bDestroyingHierarchy)
 		DrawFrustum->DestroyComponent();
 	}
 }
-#endif
+#endif	// WITH_EDITORONLY_DATA
 
 #if WITH_EDITOR
 FText UCameraComponent::GetFilmbackText() const
@@ -106,7 +119,7 @@ void UCameraComponent::OnRegister()
 			ProxyMeshComponent->SetIsVisualizationComponent(true);
 			ProxyMeshComponent->SetStaticMesh(CameraMesh);
 			ProxyMeshComponent->SetCollisionProfileName(UCollisionProfile::NoCollision_ProfileName);
-			ProxyMeshComponent->bHiddenInGame = true;
+			ProxyMeshComponent->bHiddenInGame = bCameraMeshHiddenInGame;
 			ProxyMeshComponent->CastShadow = false;
 			ProxyMeshComponent->CreationMethod = CreationMethod;
 			ProxyMeshComponent->RegisterComponentWithWorld(GetWorld());
@@ -123,7 +136,7 @@ void UCameraComponent::OnRegister()
 	}
 
 	RefreshVisualRepresentation();
-#endif
+#endif	// WITH_EDITORONLY_DATA
 
 	Super::OnRegister();
 }
@@ -185,6 +198,12 @@ void UCameraComponent::RefreshVisualRepresentation()
 		DrawFrustum->MarkRenderStateDirty();
 	}
 
+	// Update the proxy camera mesh if necessary
+	if (ProxyMeshComponent && ProxyMeshComponent->GetStaticMesh() != CameraMesh)
+	{
+		ProxyMeshComponent->SetStaticMesh(CameraMesh);
+	}
+
 	ResetProxyMeshTransform();
 }
 
@@ -222,15 +241,6 @@ void UCameraComponent::RestoreFrustumColor()
 		//Cam->DrawFrustum->FrustumColor = DefCam->DrawFrustum->FrustumColor;
 	}
 }
-#endif	// WITH_EDITORONLY_DATA
-
-#if WITH_EDITORONLY_DATA
-void UCameraComponent::PostEditChangeProperty(FPropertyChangedEvent& PropertyChangedEvent)
-{
-	Super::PostEditChangeProperty(PropertyChangedEvent);
-
-	RefreshVisualRepresentation();
-}
 
 void UCameraComponent::Serialize(FArchive& Ar)
 {
@@ -243,6 +253,16 @@ void UCameraComponent::Serialize(FArchive& Ar)
 }
 #endif	// WITH_EDITORONLY_DATA
 
+void UCameraComponent::OnCameraMeshHiddenChanged()
+{
+#if WITH_EDITORONLY_DATA
+	if (ProxyMeshComponent)
+	{
+		ProxyMeshComponent->bHiddenInGame = bCameraMeshHiddenInGame;
+	}
+#endif
+}
+
 void UCameraComponent::GetCameraView(float DeltaTime, FMinimalViewInfo& DesiredView)
 {
 	if (GEngine && GEngine->XRSystem.IsValid() && GetWorld() && GetWorld()->WorldType != EWorldType::Editor )
