Added proper safe zone support for the debug console
- Also fixed a typo in stereo 3D 'text safe region' support for canvas and switched over to that as well for the console in stereo 3D mode (untested)

diff --git a/Engine/Source/Runtime/Engine/Private/UserInterface/Canvas.cpp b/Engine/Source/Runtime/Engine/Private/UserInterface/Canvas.cpp
--- a/Engine/Source/Runtime/Engine/Private/UserInterface/Canvas.cpp
+++ b/Engine/Source/Runtime/Engine/Private/UserInterface/Canvas.cpp
@@ -1372,22 +1372,39 @@ void UCanvas::UpdateSafeZoneData()
 
 		SafeZonePadX = (CachedDisplayWidth - (CachedDisplayWidth * SafeRegionPercentage.X))/2.f;
 		SafeZonePadY = (CachedDisplayHeight - (CachedDisplayHeight * SafeRegionPercentage.Y))/2.f;
-		SafeZonePadX = SafeZonePadEX;
-		SafeZonePadY = SafeZonePadEY;
+		SafeZonePadEX = SafeZonePadX;
+		SafeZonePadEY = SafeZonePadY;
 	}
 	else if(FSlateApplication::IsInitialized())
 	{
 		FDisplayMetrics DisplayMetrics;
-
 		FSlateApplication::Get().GetCachedDisplayMetrics(DisplayMetrics);
+ 		CachedDisplayWidth = DisplayMetrics.PrimaryDisplayWidth;
+ 		CachedDisplayHeight = DisplayMetrics.PrimaryDisplayHeight;
+
+#if PLATFORM_DESKTOP
+		TSharedPtr<SWindow> Window = FSlateApplication::Get().GetActiveTopLevelWindow();
+		if (Window.IsValid())
+		{
+			FVector2D WindowSize = Window->GetClientSizeInScreen();
+			if (ISlateViewport* SlateViewport = Window->GetViewport().Get())
+			{
+				WindowSize = SlateViewport->GetSize();
+			}
+
+			CachedDisplayWidth = WindowSize.X;
+			CachedDisplayHeight = WindowSize.Y;
+		}
+#endif
 
-		SafeZonePadX = FMath::CeilToInt(DisplayMetrics.TitleSafePaddingSize.X);
-		SafeZonePadY = FMath::CeilToInt(DisplayMetrics.TitleSafePaddingSize.Y);
-		SafeZonePadEX = FMath::CeilToInt(DisplayMetrics.TitleSafePaddingSize.Z);
-		SafeZonePadEY = FMath::CeilToInt(DisplayMetrics.TitleSafePaddingSize.W);
+		const FVector2D EffectiveScreenSize = FVector2D(CachedDisplayWidth, CachedDisplayHeight);
+		FMargin SafeZone;
+		FSlateApplication::Get().GetSafeZoneSize(/*out*/ SafeZone, EffectiveScreenSize);
 
-		CachedDisplayWidth = DisplayMetrics.PrimaryDisplayWidth;
-		CachedDisplayHeight = DisplayMetrics.PrimaryDisplayHeight;
+		SafeZonePadX = FMath::CeilToInt(SafeZone.Left);
+		SafeZonePadY = FMath::CeilToInt(SafeZone.Top);
+		SafeZonePadEX = FMath::CeilToInt(SafeZone.Right);
+		SafeZonePadEY = FMath::CeilToInt(SafeZone.Bottom);
 	}
 
 }
diff --git a/Engine/Source/Runtime/Engine/Private/UserInterface/Console.cpp b/Engine/Source/Runtime/Engine/Private/UserInterface/Console.cpp
--- a/Engine/Source/Runtime/Engine/Private/UserInterface/Console.cpp
+++ b/Engine/Source/Runtime/Engine/Private/UserInterface/Console.cpp
@@ -1097,20 +1097,6 @@ void UConsole::PostRender_Console_Typing(UCanvas* Canvas)
 	float ClipY = Canvas->ClipY;
 	float LeftPos = 0;
 
-	if (GEngine->IsConsoleBuild())
-	{
-		ClipX	-= 64;
-		ClipY	-= 32;
-		LeftPos	 = 32;
-	}
-
-	if (GEngine->IsStereoscopic3D())
-	{
-		LeftPos = ClipX / 3;
-		ClipX -= LeftPos;
-		ClipY = ClipY * 0.60;
-	}
-
 	PostRender_InputLine(Canvas, FIntPoint(LeftPos, ClipY));
 }
 
@@ -1188,19 +1174,6 @@ void UConsole::PostRender_Console_Open(UCanvas* Canvas)
 	float TopPos = 0;
 	float LeftPos = 0;
 
-	if (GEngine->IsConsoleBuild())
-	{
-		ClipX	-= 80;
-		TopPos	 = 30;
-		LeftPos	 = 40;
-	}
-	if (GEngine->IsStereoscopic3D())
-	{
-		LeftPos = ClipX / 3;
-		ClipX -= LeftPos;
-		Height = Canvas->ClipY * 0.60;
-	}
-
 	UFont* Font = GEngine->GetSmallFont();
 
 	// determine the height of the text
@@ -1317,13 +1290,20 @@ bool UConsole::InputKey( int32 ControllerId, FKey Key, EInputEvent Event, float
 
 void UConsole::PostRender_Console(UCanvas* Canvas)
 {
-	if (ConsoleState == NAME_Typing)
-	{
-		PostRender_Console_Typing(Canvas);
-	}
-	else if (ConsoleState == NAME_Open)
+	if (ConsoleState != NAME_None)
 	{
-		PostRender_Console_Open(Canvas);
+		Canvas->ApplySafeZoneTransform();
+
+		if (ConsoleState == NAME_Typing)
+		{
+			PostRender_Console_Typing(Canvas);
+		}
+		else if (ConsoleState == NAME_Open)
+		{
+			PostRender_Console_Open(Canvas);
+		}
+
+		Canvas->PopSafeZoneTransform();
 	}
 }
 
@@ -1342,18 +1322,6 @@ void UConsole::PostRender_InputLine(UCanvas* Canvas, FIntPoint UserInputLinePos)
 	float ClipX = Canvas->ClipX;
 	float ClipY = Canvas->ClipY;
 
-	if (GEngine->IsConsoleBuild())
-	{
-		ClipX	-= 64;
-		ClipY	-= 32;
-	}
-
-	if (GEngine->IsStereoscopic3D())
-	{
-		ClipX = Canvas->ClipX - UserInputLinePos.X;
-		ClipY = ClipY * 0.60;
-	}
-
 	// Background
 	FLinearColor BackgroundColor = ConsoleDefs::AutocompleteBackgroundColor.ReinterpretAsLinear();
 	BackgroundColor.A = ConsoleSettings->BackgroundOpacityPercentage / 100.0f;
