Replaced use of UEngine::IsConsoleBuild() with !PLATFORM_DESKTOP in AGameMode

diff --git a/Engine/Source/Runtime/Engine/Private/GameMode.cpp b/Engine/Source/Runtime/Engine/Private/GameMode.cpp
--- a/Engine/Source/Runtime/Engine/Private/GameMode.cpp
+++ b/Engine/Source/Runtime/Engine/Private/GameMode.cpp
@@ -642,7 +642,7 @@ void AGameMode::AddInactivePlayer(APlayerState* PlayerState, APlayerController*
 			NewPlayerState->SetLifeSpan(InactivePlayerStateLifeSpan);
 
 			// On console, we have to check the unique net id as network address isn't valid
-			const bool bIsConsole = GEngine->IsConsoleBuild();
+			const bool bIsConsole = !PLATFORM_DESKTOP;
 			// Assume valid unique ids means comparison should be via this method
 			const bool bHasValidUniqueId = NewPlayerState->UniqueId.IsValid();
 			// Don't accidentally compare empty network addresses (already issue with two clients on same machine during development)
@@ -702,7 +702,7 @@ bool AGameMode::FindInactivePlayer(APlayerController* PC)
 	}
 
 	// On console, we have to check the unique net id as network address isn't valid
-	const bool bIsConsole = GEngine->IsConsoleBuild();
+	const bool bIsConsole = !PLATFORM_DESKTOP;
 	// Assume valid unique ids means comparison should be via this method
 	const bool bHasValidUniqueId = PC->PlayerState->UniqueId.IsValid();
 	// Don't accidentally compare empty network addresses (already issue with two clients on same machine during development)
