Removing obsolete EConsoleType and UEngine::IsConsoleBuild / UEngine::bUseConsoleInput

diff --git a/Engine/Source/Runtime/Engine/Classes/Engine/Engine.h b/Engine/Source/Runtime/Engine/Classes/Engine/Engine.h
--- a/Engine/Source/Runtime/Engine/Classes/Engine/Engine.h
+++ b/Engine/Source/Runtime/Engine/Classes/Engine/Engine.h
@@ -102,15 +102,6 @@ enum ETransitionType
 	TT_MAX,
 };
 
-
-UENUM()
-enum EConsoleType
-{
-	CONSOLE_Any,
-	CONSOLE_Mobile,
-	CONSOLE_MAX,
-};
-
 /** Status of dynamic resolution that depends on project setting cvar, game user settings, and pause */
 enum class EDynamicResolutionStatus
 {
@@ -1373,14 +1364,6 @@ public:
 	UPROPERTY(globalconfig)
 	uint32 bShouldGenerateLowQualityLightmaps_DEPRECATED :1;
 
-	/**
-	 * Bool that indicates that 'console' input is desired. This flag is mis named as it is used for a lot of gameplay related things
-	 * (e.g. increasing collision size, changing vehicle turning behavior, modifying put down/up weapon speed, bot behavior)
-	 *
-	 * currently set when you are running a console build (implicitly or explicitly via ?param on the commandline)
-	 */
-	uint32 bUseConsoleInput:1;
-
 	// Color preferences.
 	UPROPERTY()
 	FColor C_WorldBox;
@@ -2322,14 +2305,6 @@ public:
 	 */
 	float GetTimeBetweenGarbageCollectionPasses() const;
 
-	/**
-	 * Returns whether we are running on a console platform or on the PC.
-	 * @param ConsoleType - if specified, only returns true if we're running on the specified platform
-	 *
-	 * @return true if we're on a console, false if we're running on a PC
-	 */
-	bool IsConsoleBuild(EConsoleType ConsoleType = CONSOLE_Any) const;
-
 	/** Add a FString to the On-screen debug message system. bNewerOnTop only works with Key == INDEX_NONE */
 	void AddOnScreenDebugMessage(uint64 Key,float TimeToDisplay,FColor DisplayColor,const FString& DebugMessage, bool bNewerOnTop = true, const FVector2D& TextScale = FVector2D::UnitVector);
 
diff --git a/Engine/Source/Runtime/Engine/Classes/GameFramework/WorldSettings.h b/Engine/Source/Runtime/Engine/Classes/GameFramework/WorldSettings.h
--- a/Engine/Source/Runtime/Engine/Classes/GameFramework/WorldSettings.h
+++ b/Engine/Source/Runtime/Engine/Classes/GameFramework/WorldSettings.h
@@ -736,7 +736,6 @@ public:
 #endif // WITH_EDITOR
 	virtual void PostInitProperties() override;
 	virtual void PreInitializeComponents() override;
-	virtual void PostInitializeComponents() override;
 	virtual void PostRegisterAllComponents() override;
 	//~ End AActor Interface.
 
diff --git a/Engine/Source/Runtime/Engine/Private/UnrealEngine.cpp b/Engine/Source/Runtime/Engine/Private/UnrealEngine.cpp
--- a/Engine/Source/Runtime/Engine/Private/UnrealEngine.cpp
+++ b/Engine/Source/Runtime/Engine/Private/UnrealEngine.cpp
@@ -1465,11 +1465,6 @@ void UEngine::Init(IEngineLoop* InEngineLoop)
 	// Initialize the audio device after a world context is setup
 	InitializeAudioDeviceManager();
 
-	if ( IsConsoleBuild() )
-	{
-		bUseConsoleInput = true;
-	}
-
 	// Make sure networking checksum has access to project version
 	const UGeneralProjectSettings& ProjectSettings = *GetDefault<UGeneralProjectSettings>();
 	FNetworkVersion::SetProjectVersion(*ProjectSettings.ProjectVersion);
@@ -8799,29 +8794,6 @@ FGuid UEngine::GetPackageGuid(FName PackageName, bool bForPIE)
 	return Result;
 }
 
-/** 
-* Returns whether we are running on a console platform or on the PC.
-*
-* @return true if we're on a console, false if we're running on a PC
-*/
-bool UEngine::IsConsoleBuild(EConsoleType ConsoleType) const
-{
-	switch (ConsoleType)
-	{
-	case CONSOLE_Any:
-#if !PLATFORM_DESKTOP
-		return true;
-#else
-		return false;
-#endif
-	case CONSOLE_Mobile:
-		return false;
-	default:
-		UE_LOG(LogEngine, Warning, TEXT("Unknown ConsoleType passed to IsConsoleBuild()"));
-		return false;
-	}
-}
-
 #if !(UE_BUILD_SHIPPING || UE_BUILD_TEST)
 // Critical section used in AddOnScreenDebugMessage to handle concurrent insert.
 static FCriticalSection GOnScreenMessageCS;
diff --git a/Engine/Source/Runtime/Engine/Private/WorldSettings.cpp b/Engine/Source/Runtime/Engine/Private/WorldSettings.cpp
--- a/Engine/Source/Runtime/Engine/Private/WorldSettings.cpp
+++ b/Engine/Source/Runtime/Engine/Private/WorldSettings.cpp
@@ -189,16 +189,6 @@ void AWorldSettings::PreInitializeComponents()
 	}
 }
 
-void AWorldSettings::PostInitializeComponents()
-{
-	Super::PostInitializeComponents();
-
-	if (GEngine->IsConsoleBuild())
-	{
-		GEngine->bUseConsoleInput = true;
-	}
-}
-
 void AWorldSettings::PostRegisterAllComponents()
 {
 	Super::PostRegisterAllComponents();
