Add a check to ComponentIsDamageableFrom that will default to the simple hit results and print a warning.

diff --git a/Engine/Source/Runtime/Engine/Private/GameplayStatics.cpp b/Engine/Source/Runtime/Engine/Private/GameplayStatics.cpp
--- a/Engine/Source/Runtime/Engine/Private/GameplayStatics.cpp
+++ b/Engine/Source/Runtime/Engine/Private/GameplayStatics.cpp
@@ -438,25 +438,34 @@ static bool ComponentIsDamageableFrom(UPrimitiveComponent* VictimComp, FVector c
 		// tiny nudge so LineTraceSingle doesn't early out with no hits
 		TraceStart.Z += 0.01f;
 	}
-	bool const bHadBlockingHit = World->LineTraceSingleByChannel(OutHitResult, TraceStart, TraceEnd, TraceChannel, LineParams);
-	//::DrawDebugLine(World, TraceStart, TraceEnd, FLinearColor::Red, true);
 
-	// If there was a blocking hit, it will be the last one
-	if (bHadBlockingHit)
+	// Only do a line trace if there is a valid channel, if it is invalid then result will have no fall off
+	if (TraceChannel != ECollisionChannel::ECC_MAX)
 	{
-		if (OutHitResult.Component == VictimComp)
-		{
-			// if blocking hit was the victim component, it is visible
-			return true;
-		}
-		else
+		bool const bHadBlockingHit = World->LineTraceSingleByChannel(OutHitResult, TraceStart, TraceEnd, TraceChannel, LineParams);
+		//::DrawDebugLine(World, TraceStart, TraceEnd, FLinearColor::Red, true);
+
+		// If there was a blocking hit, it will be the last one
+		if (bHadBlockingHit)
 		{
-			// if we hit something else blocking, it's not
-			UE_LOG(LogDamage, Log, TEXT("Radial Damage to %s blocked by %s (%s)"), *GetNameSafe(VictimComp), *GetNameSafe(OutHitResult.GetActor()), *GetNameSafe(OutHitResult.Component.Get()) );
-			return false;
+			if (OutHitResult.Component == VictimComp)
+			{
+				// if blocking hit was the victim component, it is visible
+				return true;
+			}
+			else
+			{
+				// if we hit something else blocking, it's not
+				UE_LOG(LogDamage, Log, TEXT("Radial Damage to %s blocked by %s (%s)"), *GetNameSafe(VictimComp), *GetNameSafe(OutHitResult.GetActor()), *GetNameSafe(OutHitResult.Component.Get()));
+				return false;
+			}
 		}
 	}
-		
+	else
+	{
+		UE_LOG(LogDamage, Warning, TEXT("ECollisionChannel::ECC_MAX is not valid! No falloff is added to damage"));
+	}
+
 	// didn't hit anything, assume nothing blocking the damage and victim is consequently visible
 	// but since we don't have a hit result to pass back, construct a simple one, modeling the damage as having hit a point at the component's center.
 	FVector const FakeHitLoc = VictimComp->GetComponentLocation();
@@ -481,23 +490,23 @@ bool UGameplayStatics::ApplyRadialDamageWithFalloff(const UObject* WorldContextO
 	TArray<FOverlapResult> Overlaps;
 	if (UWorld* World = GEngine->GetWorldFromContextObject(WorldContextObject, EGetWorldErrorMode::LogAndReturnNull))
 	{
-	World->OverlapMultiByObjectType(Overlaps, Origin, FQuat::Identity, FCollisionObjectQueryParams(FCollisionObjectQueryParams::InitType::AllDynamicObjects), FCollisionShape::MakeSphere(DamageOuterRadius), SphereParams);
+		World->OverlapMultiByObjectType(Overlaps, Origin, FQuat::Identity, FCollisionObjectQueryParams(FCollisionObjectQueryParams::InitType::AllDynamicObjects), FCollisionShape::MakeSphere(DamageOuterRadius), SphereParams);
 	}
 
 	// collate into per-actor list of hit components
 	TMap<AActor*, TArray<FHitResult> > OverlapComponentMap;
-	for (int32 Idx=0; Idx<Overlaps.Num(); ++Idx)
+	for (int32 Idx = 0; Idx < Overlaps.Num(); ++Idx)
 	{
 		FOverlapResult const& Overlap = Overlaps[Idx];
 		AActor* const OverlapActor = Overlap.GetActor();
 
-		if ( OverlapActor && 
-			OverlapActor->bCanBeDamaged && 
+		if (OverlapActor &&
+			OverlapActor->bCanBeDamaged &&
 			OverlapActor != DamageCauser &&
-			Overlap.Component.IsValid() )
+			Overlap.Component.IsValid())
 		{
 			FHitResult Hit;
-			if (DamagePreventionChannel == ECC_MAX || ComponentIsDamageableFrom(Overlap.Component.Get(), Origin, DamageCauser, IgnoreActors, DamagePreventionChannel, Hit))
+			if (ComponentIsDamageableFrom(Overlap.Component.Get(), Origin, DamageCauser, IgnoreActors, DamagePreventionChannel, Hit))
 			{
 				TArray<FHitResult>& HitList = OverlapComponentMap.FindOrAdd(OverlapActor);
 				HitList.Add(Hit);
@@ -509,25 +518,25 @@ bool UGameplayStatics::ApplyRadialDamageWithFalloff(const UObject* WorldContextO
 
 	if (OverlapComponentMap.Num() > 0)
 	{
-	// make sure we have a good damage type
-	TSubclassOf<UDamageType> const ValidDamageTypeClass = DamageTypeClass ? DamageTypeClass : TSubclassOf<UDamageType>(UDamageType::StaticClass());
+		// make sure we have a good damage type
+		TSubclassOf<UDamageType> const ValidDamageTypeClass = DamageTypeClass ? DamageTypeClass : TSubclassOf<UDamageType>(UDamageType::StaticClass());
 
 		FRadialDamageEvent DmgEvent;
 		DmgEvent.DamageTypeClass = ValidDamageTypeClass;
 		DmgEvent.Origin = Origin;
 		DmgEvent.Params = FRadialDamageParams(BaseDamage, MinimumDamage, DamageInnerRadius, DamageOuterRadius, DamageFalloff);
 
-	// call damage function on each affected actors
-	for (TMap<AActor*, TArray<FHitResult> >::TIterator It(OverlapComponentMap); It; ++It)
-	{
-		AActor* const Victim = It.Key();
-		TArray<FHitResult> const& ComponentHits = It.Value();
-		DmgEvent.ComponentHits = ComponentHits;
+		// call damage function on each affected actors
+		for (TMap<AActor*, TArray<FHitResult> >::TIterator It(OverlapComponentMap); It; ++It)
+		{
+			AActor* const Victim = It.Key();
+			TArray<FHitResult> const& ComponentHits = It.Value();
+			DmgEvent.ComponentHits = ComponentHits;
 
-		Victim->TakeDamage(BaseDamage, DmgEvent, InstigatedByController, DamageCauser);
+			Victim->TakeDamage(BaseDamage, DmgEvent, InstigatedByController, DamageCauser);
 
-		bAppliedDamage = true;
-	}
+			bAppliedDamage = true;
+		}
 	}
 
 	return bAppliedDamage;
