#ue4 - Fix character perch radius at high values close to the capsule radius causing players to fall in-place or on valid geometry. We clamp physics capsules to 0.1 radius so we need something higher than this for our min size.

#udn https://udn.unrealengine.com/questions/501831/inconsistent-character-movement-perch-result.html

diff --git a/Engine/Source/Runtime/Engine/Private/Components/CharacterMovementComponent.cpp b/Engine/Source/Runtime/Engine/Private/Components/CharacterMovementComponent.cpp
--- a/Engine/Source/Runtime/Engine/Private/Components/CharacterMovementComponent.cpp
+++ b/Engine/Source/Runtime/Engine/Private/Components/CharacterMovementComponent.cpp
@@ -6347,7 +6347,7 @@ float UCharacterMovementComponent::GetValidPerchRadius() const
 	if (CharacterOwner)
 	{
 		const float PawnRadius = CharacterOwner->GetCapsuleComponent()->GetScaledCapsuleRadius();
-		return FMath::Clamp(PawnRadius - GetPerchRadiusThreshold(), 0.1f, PawnRadius);
+		return FMath::Clamp(PawnRadius - GetPerchRadiusThreshold(), 0.11f, PawnRadius);
 	}
 	return 0.f;
 }
