Fixed FCompositeNavModifier deprecation

diff --git a/Engine/Source/Runtime/Engine/Public/AI/NavigationModifier.h b/Engine/Source/Runtime/Engine/Public/AI/NavigationModifier.h
--- a/Engine/Source/Runtime/Engine/Public/AI/NavigationModifier.h
+++ b/Engine/Source/Runtime/Engine/Public/AI/NavigationModifier.h
@@ -312,7 +312,7 @@ struct ENGINE_API FCompositeNavModifier : public FNavigationModifier
 	FCompositeNavModifier GetInstantiatedMetaModifier(const struct FNavAgentProperties* NavAgent, TWeakObjectPtr<UObject> WeakOwnerPtr) const;
 	uint32 GetAllocatedSize() const;
 
-	UE_DEPRECATED(4.23, "This method will be removed in future versions. Use FNavigationRelevantData::HasPerInstanceTransforms instead.")
+	UE_DEPRECATED(4.24, "This method will be removed in future versions. Use FNavigationRelevantData::HasPerInstanceTransforms instead.")
 	bool HasPerInstanceTransforms() const;
 
 	TArray<FAreaNavModifier>& GetMutableAreas() { return Areas; }
@@ -321,7 +321,7 @@ struct ENGINE_API FCompositeNavModifier : public FNavigationModifier
 
 public:
 
-	UE_DEPRECATED(4.23, "This property will be removed in future versions. Use FNavigationRelevantData NavDataPerInstanceTransformDelegate instead.")
+	// This property is deprecated and will be removed in future versions. Use FNavigationRelevantData NavDataPerInstanceTransformDelegate instead.
 	// Gathers per instance data for navigation area modifiers in a specified area box
 	FNavDataPerInstanceTransformDelegate NavDataPerInstanceTransformDelegate;
 
