Account for the same command entered multiple times with the DisplayAll command

diff --git a/Engine/Source/Runtime/Engine/Classes/Engine/GameViewportClient.h b/Engine/Source/Runtime/Engine/Classes/Engine/GameViewportClient.h
--- a/Engine/Source/Runtime/Engine/Classes/Engine/GameViewportClient.h
+++ b/Engine/Source/Runtime/Engine/Classes/Engine/GameViewportClient.h
@@ -64,7 +64,7 @@ public:
 	UPROPERTY()
 	class UConsole* ViewportConsole;
 
-	/** @todo document */
+	/** Debug properties that have been added via one of the "displayall" commands */
 	UPROPERTY()
 	TArray<struct FDebugDisplayProperty> DebugProperties;
 
@@ -844,6 +844,17 @@ private:
 	/** Finds available PNG cursor images */
 	bool LoadAvailableCursorPngs(TArray<TSharedPtr<FPngFileData>>& Results, const FString& InPathToCursorWithoutExtension);
 
+	/**
+	* Adds a DebugDisplayProperty to the DebugProperties array if it does not already exist. 
+	* @see FDebugDisplayProperty for more info on debug properties
+	* 
+	* @param Obj				Object that the debug property is on
+	* @param WithinClass		further limit the display to objects that have an Outer of WithinClass
+	* @param PropertyName		name of the property to display
+	* @param bSpecialProperty	whether PropertyName is a "special" value not directly mapping to a real property (e.g. state name)
+	*/
+	void AddDebugDisplayProperty(class UObject* Obj, TSubclassOf<class UObject> WithinClass, const FName& PropertyName, bool bSpecialProperty = false);
+
 private:
 	/** Slate window associated with this viewport client.  The same window may host more than one viewport client. */
 	TWeakPtr<SWindow> Window;
diff --git a/Engine/Source/Runtime/Engine/Private/GameViewportClient.cpp b/Engine/Source/Runtime/Engine/Private/GameViewportClient.cpp
--- a/Engine/Source/Runtime/Engine/Private/GameViewportClient.cpp
+++ b/Engine/Source/Runtime/Engine/Private/GameViewportClient.cpp
@@ -3433,14 +3433,12 @@ bool UGameViewportClient::HandleDisplayCommand( const TCHAR* Cmd, FOutputDevice&
 		FParse::Token(Cmd, PropStr, ARRAY_COUNT(PropStr), true) )
 	{
 		UObject* Obj = FindObject<UObject>(ANY_PACKAGE, ObjectName);
-		if (Obj != NULL)
+		if (Obj != nullptr)
 		{
 			FName PropertyName(PropStr, FNAME_Find);
-			if (PropertyName != NAME_None && FindField<UProperty>(Obj->GetClass(), PropertyName) != NULL)
+			if (PropertyName != NAME_None && FindField<UProperty>(Obj->GetClass(), PropertyName) != nullptr)
 			{
-				FDebugDisplayProperty& NewProp = DebugProperties[DebugProperties.AddZeroed()];
-				NewProp.Obj = Obj;
-				NewProp.PropertyName = PropertyName;
+				AddDebugDisplayProperty(Obj, nullptr, PropertyName);
 			}
 			else
 			{
@@ -3463,7 +3461,7 @@ bool UGameViewportClient::HandleDisplayAllCommand( const TCHAR* Cmd, FOutputDevi
 	if (FParse::Token(Cmd, ClassName, ARRAY_COUNT(ClassName), true))
 	{
 		bool bValidClassToken = true;
-		UClass* WithinClass = NULL;
+		UClass* WithinClass = nullptr;
 		{
 			FString ClassStr(ClassName);
 			int32 DotIndex = ClassStr.Find(TEXT("."));
@@ -3471,7 +3469,7 @@ bool UGameViewportClient::HandleDisplayAllCommand( const TCHAR* Cmd, FOutputDevi
 			{
 				// first part is within class
 				WithinClass = FindObject<UClass>(ANY_PACKAGE, *ClassStr.Left(DotIndex));
-				if (WithinClass == NULL)
+				if (WithinClass == nullptr)
 				{
 					Ar.Logf(TEXT("Within class not found"));
 					bValidClassToken = false;
@@ -3487,10 +3485,10 @@ bool UGameViewportClient::HandleDisplayAllCommand( const TCHAR* Cmd, FOutputDevi
 		{
 			FParse::Token(Cmd, PropStr, ARRAY_COUNT(PropStr), true);
 			UClass* Cls = FindObject<UClass>(ANY_PACKAGE, ClassName);
-			if (Cls != NULL)
+			if (Cls != nullptr)
 			{
 				FName PropertyName(PropStr, FNAME_Find);
-				UProperty* Prop = PropertyName != NAME_None ? FindField<UProperty>(Cls, PropertyName) : NULL;
+				UProperty* Prop = PropertyName != NAME_None ? FindField<UProperty>(Cls, PropertyName) : nullptr;
 				{
 					// add all un-GCable things immediately as that list is static
 					// so then we only have to iterate over dynamic things each frame
@@ -3500,25 +3498,12 @@ bool UGameViewportClient::HandleDisplayAllCommand( const TCHAR* Cmd, FOutputDevi
 						{
 							break;
 						}
-						else if (It->IsA(Cls) && !It->IsTemplate() && (WithinClass == NULL || (It->GetOuter() != NULL && It->GetOuter()->GetClass()->IsChildOf(WithinClass))))
+						else if (It->IsA(Cls) && !It->IsTemplate() && (WithinClass == nullptr || (It->GetOuter() != nullptr && It->GetOuter()->GetClass()->IsChildOf(WithinClass))))
 						{
-							FDebugDisplayProperty& NewProp = DebugProperties[DebugProperties.AddZeroed()];
-							NewProp.Obj = *It;
-							NewProp.PropertyName = PropertyName;
-							if (!Prop)
-							{
-								NewProp.bSpecialProperty = true;
-							}
+							AddDebugDisplayProperty(*It, nullptr, PropertyName, !Prop);
 						}
 					}
-					FDebugDisplayProperty& NewProp = DebugProperties[DebugProperties.AddZeroed()];
-					NewProp.Obj = Cls;
-					NewProp.WithinClass = WithinClass;
-					NewProp.PropertyName = PropertyName;
-					if (!Prop)
-					{
-						NewProp.bSpecialProperty = true;
-					}
+					AddDebugDisplayProperty(Cls, WithinClass, PropertyName, !Prop);
 				}
 			}
 			else
@@ -3537,7 +3522,7 @@ bool UGameViewportClient::HandleDisplayAllLocationCommand( const TCHAR* Cmd, FOu
 	if (FParse::Token(Cmd, ClassName, ARRAY_COUNT(ClassName), true))
 	{
 		UClass* Cls = FindObject<UClass>(ANY_PACKAGE, ClassName);
-		if (Cls != NULL)
+		if (Cls != nullptr)
 		{
 			// add all un-GCable things immediately as that list is static
 			// so then we only have to iterate over dynamic things each frame
@@ -3549,16 +3534,10 @@ bool UGameViewportClient::HandleDisplayAllLocationCommand( const TCHAR* Cmd, FOu
 				}
 				else if (It->IsA(Cls))
 				{
-					FDebugDisplayProperty& NewProp = DebugProperties[DebugProperties.AddZeroed()];
-					NewProp.Obj = *It;
-					NewProp.PropertyName = NAME_Location;
-					NewProp.bSpecialProperty = true;
+					AddDebugDisplayProperty(*It, nullptr, NAME_Location, true);
 				}
 			}
-			FDebugDisplayProperty& NewProp = DebugProperties[DebugProperties.AddZeroed()];
-			NewProp.Obj = Cls;
-			NewProp.PropertyName = NAME_Location;
-			NewProp.bSpecialProperty = true;
+			AddDebugDisplayProperty(Cls, nullptr, NAME_Location, true);
 		}
 		else
 		{
@@ -3575,7 +3554,7 @@ bool UGameViewportClient::HandleDisplayAllRotationCommand( const TCHAR* Cmd, FOu
 	if (FParse::Token(Cmd, ClassName, ARRAY_COUNT(ClassName), true))
 	{
 		UClass* Cls = FindObject<UClass>(ANY_PACKAGE, ClassName);
-		if (Cls != NULL)
+		if (Cls != nullptr)
 		{
 			// add all un-GCable things immediately as that list is static
 			// so then we only have to iterate over dynamic things each frame
@@ -3587,16 +3566,10 @@ bool UGameViewportClient::HandleDisplayAllRotationCommand( const TCHAR* Cmd, FOu
 				}
 				else if (It->IsA(Cls))
 				{
-					FDebugDisplayProperty& NewProp = DebugProperties[DebugProperties.AddZeroed()];
-					NewProp.Obj = *It;
-					NewProp.PropertyName = NAME_Rotation;
-					NewProp.bSpecialProperty = true;
+					AddDebugDisplayProperty(*It, nullptr, NAME_Rotation, true);
 				}
 			}
-			FDebugDisplayProperty& NewProp = DebugProperties[DebugProperties.AddZeroed()];
-			NewProp.Obj = Cls;
-			NewProp.PropertyName = NAME_Rotation;
-			NewProp.bSpecialProperty = true;
+			AddDebugDisplayProperty(Cls, nullptr, NAME_Rotation, true);
 		}
 		else
 		{
@@ -3864,6 +3837,24 @@ void* UGameViewportClient::LoadCursorFromPngs(ICursor& PlatformCursor, const FSt
 	return nullptr;
 }
 
+void UGameViewportClient::AddDebugDisplayProperty(class UObject* Obj, TSubclassOf<class UObject> WithinClass, const FName& PropertyName, bool bSpecialProperty /*= false*/)
+{
+	// If this property already exists than don't add a new one
+	for (const FDebugDisplayProperty& Prop : DebugProperties)
+	{
+		if (Prop.Obj == Obj && Prop.PropertyName == PropertyName)
+		{
+			return;
+		}
+	}
+
+	FDebugDisplayProperty& NewProp = DebugProperties.AddDefaulted_GetRef();
+	NewProp.Obj = Obj;
+	NewProp.WithinClass = WithinClass;
+	NewProp.PropertyName = PropertyName;
+	NewProp.bSpecialProperty = bSpecialProperty;
+}
+
 bool UGameViewportClient::LoadAvailableCursorPngs(TArray< TSharedPtr<FPngFileData> >& Results, const FString& InPathToCursorWithoutExtension)
 {
 	FString CursorsWithSizeSearch = FPaths::GetCleanFilename(InPathToCursorWithoutExtension) + TEXT("*.png");
