Create virtual accessor functions for bCanRenameNode on EdGraphNode

diff --git a/Engine/Source/Runtime/Engine/Classes/EdGraph/EdGraphNode.h b/Engine/Source/Runtime/Engine/Classes/EdGraph/EdGraphNode.h
--- a/Engine/Source/Runtime/Engine/Classes/EdGraph/EdGraphNode.h
+++ b/Engine/Source/Runtime/Engine/Classes/EdGraph/EdGraphNode.h
@@ -325,6 +325,9 @@ public:
 		return (EnabledState == ENodeEnabledState::Enabled)	|| ((EnabledState == ENodeEnabledState::DevelopmentOnly) && IsInDevelopmentMode());
 	}
 
+	/** If true, this node can be renamed in the editor */
+	virtual bool GetCanRenameNode() const;
+
 	/** Returns the specific sort of enable state this node wants */
 	ENodeEnabledState GetDesiredEnabledState() const
 	{
diff --git a/Engine/Source/Runtime/Engine/Private/EdGraph/EdGraphNode.cpp b/Engine/Source/Runtime/Engine/Private/EdGraph/EdGraphNode.cpp
--- a/Engine/Source/Runtime/Engine/Private/EdGraph/EdGraphNode.cpp
+++ b/Engine/Source/Runtime/Engine/Private/EdGraph/EdGraphNode.cpp
@@ -745,6 +745,15 @@ bool UEdGraphNode::CanDuplicateNode() const
 	return true;
 }
 
+bool UEdGraphNode::GetCanRenameNode() const
+{
+#if WITH_EDITORONLY_DATA
+	return bCanRenameNode;
+#else
+	return false;
+#endif
+}
+
 bool UEdGraphNode::CanUserDeleteNode() const
 {
 	return true;
