Based on PR #5820

diff --git a/Engine/Source/Runtime/Engine/Classes/Kismet/GameplayStatics.h b/Engine/Source/Runtime/Engine/Classes/Kismet/GameplayStatics.h
--- a/Engine/Source/Runtime/Engine/Classes/Kismet/GameplayStatics.h
+++ b/Engine/Source/Runtime/Engine/Classes/Kismet/GameplayStatics.h
@@ -793,7 +793,7 @@ public:
 	 * @param HitItem		Primitive-specific data recording which item in the primitive was hit
 	 * @param FaceIndex		If colliding with trimesh or landscape, index of face that was hit.
 	 */
-	UFUNCTION(BlueprintPure, Category = "Collision", meta=(NativeBreakFunc))
+	UFUNCTION(BlueprintPure, Category = "Collision", meta=(NativeBreakFunc, AdvancedDisplay="3"))
 	static void BreakHitResult(const struct FHitResult& Hit, bool& bBlockingHit, bool& bInitialOverlap, float& Time, float& Distance, FVector& Location, FVector& ImpactPoint, FVector& Normal, FVector& ImpactNormal, class UPhysicalMaterial*& PhysMat, class AActor*& HitActor, class UPrimitiveComponent*& HitComponent, FName& HitBoneName, int32& HitItem, int32& FaceIndex, FVector& TraceStart, FVector& TraceEnd);
 
 	/** 
@@ -814,7 +814,7 @@ public:
 	 * @param HitItem		Primitive-specific data recording which item in the primitive was hit
 	 * @param FaceIndex		If colliding with trimesh or landscape, index of face that was hit.
 	 */
-	UFUNCTION(BlueprintPure, Category = "Collision", meta = (NativeMakeFunc, Normal="0,0,1", ImpactNormal="0,0,1"))
+	UFUNCTION(BlueprintPure, Category = "Collision", meta = (NativeMakeFunc, AdvancedDisplay="2", Normal="0,0,1", ImpactNormal="0,0,1"))
 	static FHitResult MakeHitResult(bool bBlockingHit, bool bInitialOverlap, float Time, float Distance, FVector Location, FVector ImpactPoint, FVector Normal, FVector ImpactNormal, class UPhysicalMaterial* PhysMat, class AActor* HitActor, class UPrimitiveComponent* HitComponent, FName HitBoneName, int32 HitItem, int32 FaceIndex, FVector TraceStart, FVector TraceEnd);
 
 
diff --git a/Engine/Source/Runtime/Engine/Classes/Kismet/KismetMathLibrary.h b/Engine/Source/Runtime/Engine/Classes/Kismet/KismetMathLibrary.h
--- a/Engine/Source/Runtime/Engine/Classes/Kismet/KismetMathLibrary.h
+++ b/Engine/Source/Runtime/Engine/Classes/Kismet/KismetMathLibrary.h
@@ -2660,7 +2660,7 @@ class ENGINE_API UKismetMathLibrary : public UBlueprintFunctionLibrary
 	static FDateTime MakeDateTime(int32 Year, int32 Month, int32 Day, int32 Hour = 0, int32 Minute = 0, int32 Second = 0, int32 Millisecond = 0);
 
 	/** Breaks a DateTime into its components */
-	UFUNCTION(BlueprintPure, Category="Math|DateTime", meta=(NativeBreakFunc))
+	UFUNCTION(BlueprintPure, Category="Math|DateTime", meta=(NativeBreakFunc, AdvancedDisplay = "3"))
 	static void BreakDateTime(FDateTime InDateTime, int32& Year, int32& Month, int32& Day, int32& Hour, int32& Minute, int32& Second, int32& Millisecond);
 
 	/** Addition (A + B) */
