There were already native subclasses of these that were blueprintable and they are well exposed to blueprints already, so this will be generally safe.

diff --git a/Engine/Source/Runtime/Engine/Classes/Components/SkeletalMeshComponent.h b/Engine/Source/Runtime/Engine/Classes/Components/SkeletalMeshComponent.h
--- a/Engine/Source/Runtime/Engine/Classes/Components/SkeletalMeshComponent.h
+++ b/Engine/Source/Runtime/Engine/Classes/Components/SkeletalMeshComponent.h
@@ -275,7 +275,7 @@ struct ENGINE_API FClosestPointOnPhysicsAsset
  * @see https://docs.unrealengine.com/latest/INT/Engine/Content/Types/SkeletalMeshes/
  * @see USkeletalMesh
  */
-UCLASS(ClassGroup=(Rendering, Common), hidecategories=Object, config=Engine, editinlinenew, meta=(BlueprintSpawnableComponent))
+UCLASS(Blueprintable, ClassGroup=(Rendering, Common), hidecategories=Object, config=Engine, editinlinenew, meta=(BlueprintSpawnableComponent))
 class ENGINE_API USkeletalMeshComponent : public USkinnedMeshComponent, public IInterface_CollisionDataProvider
 {
 	GENERATED_UCLASS_BODY()
diff --git a/Engine/Source/Runtime/Engine/Classes/Components/StaticMeshComponent.h b/Engine/Source/Runtime/Engine/Classes/Components/StaticMeshComponent.h
--- a/Engine/Source/Runtime/Engine/Classes/Components/StaticMeshComponent.h
+++ b/Engine/Source/Runtime/Engine/Classes/Components/StaticMeshComponent.h
@@ -165,7 +165,7 @@ struct TStructOpsTypeTraits<FStaticMeshComponentLODInfo> : public TStructOpsType
  * @see https://docs.unrealengine.com/latest/INT/Engine/Content/Types/StaticMeshes/
  * @see UStaticMesh
  */
-UCLASS(ClassGroup=(Rendering, Common), hidecategories=(Object,Activation,"Components|Activation"), ShowCategories=(Mobility), editinlinenew, meta=(BlueprintSpawnableComponent))
+UCLASS(Blueprintable, ClassGroup=(Rendering, Common), hidecategories=(Object,Activation,"Components|Activation"), ShowCategories=(Mobility), editinlinenew, meta=(BlueprintSpawnableComponent))
 class ENGINE_API UStaticMeshComponent : public UMeshComponent
 {
 	GENERATED_UCLASS_BODY()
