Fix formatting in Body Instance.cpp

diff --git a/Engine/Source/Runtime/Engine/Private/PhysicsEngine/BodyInstance.cpp b/Engine/Source/Runtime/Engine/Private/PhysicsEngine/BodyInstance.cpp
--- a/Engine/Source/Runtime/Engine/Private/PhysicsEngine/BodyInstance.cpp
+++ b/Engine/Source/Runtime/Engine/Private/PhysicsEngine/BodyInstance.cpp
@@ -1163,44 +1163,44 @@ struct FInitBodiesHelper
 
 
 	void InitBodies()
-				{
+	{
 		LLM_SCOPE(ELLMTag::PhysX);
 
 		check(IsInGameThread());
 
-		if(CreateShapesAndActors())
-					{
-			FPhysicsCommand::ExecuteWrite(PhysScene, [&]()
-	{
-				// If an aggregate present, add to that
-				if(Aggregate.IsValid())
+		if (CreateShapesAndActors())
 		{
-					for(FBodyInstance* BI : Bodies)
+			FPhysicsCommand::ExecuteWrite(PhysScene, [&]()
 			{
-						const FPhysicsActorHandle& ActorHandle = BI->GetPhysicsActorHandle();
-						if(ActorHandle.IsValid())
+				// If an aggregate present, add to that
+				if (Aggregate.IsValid())
 				{
+					for (FBodyInstance* BI : Bodies)
+					{
+						const FPhysicsActorHandle& ActorHandle = BI->GetPhysicsActorHandle();
+						if (ActorHandle.IsValid())
+						{
 							FPhysicsInterface::AddActorToAggregate_AssumesLocked(Aggregate, ActorHandle);
+						}
+					}
 				}
-			}
-		}
-				else if(PhysScene)
-		{
+				else if (PhysScene)
+				{
 					TArray<FPhysicsActorHandle> ActorHandles;
 					ActorHandles.Reserve(Bodies.Num());
 
-					for(FBodyInstance* BI : Bodies)
-			{
+					for (FBodyInstance* BI : Bodies)
+					{
 						const FPhysicsActorHandle& ActorHandle = BI->GetPhysicsActorHandle();
-						if(ActorHandle.IsValid())
-			{
+						if (ActorHandle.IsValid())
+						{
 							ActorHandles.Add(ActorHandle);
-			}
-		}
+						}
+					}
 
 					PhysScene->AddActorsToScene_AssumesLocked(ActorHandles);
 				}
-				
+
 				// Set up dynamic instance data
 				if (!IsStatic())
 				{
@@ -1212,7 +1212,7 @@ struct FInitBodiesHelper
 					}
 				}
 			});
-			}
+		}
 	}
 };
 
