Reorg Blueprint data to pack better and not have editor only concepts in non-editor builds
UBlueprint OnChangedEvent and OnCompiledEvent are now editor only events.

diff --git a/Engine/Plugins/Animation/LiveLink/Source/LiveLink/Private/LiveLinkRemapAsset.cpp b/Engine/Plugins/Animation/LiveLink/Source/LiveLink/Private/LiveLinkRemapAsset.cpp
--- a/Engine/Plugins/Animation/LiveLink/Source/LiveLink/Private/LiveLinkRemapAsset.cpp
+++ b/Engine/Plugins/Animation/LiveLink/Source/LiveLink/Private/LiveLinkRemapAsset.cpp
@@ -10,15 +10,18 @@
 ULiveLinkRemapAsset::ULiveLinkRemapAsset(const FObjectInitializer& ObjectInitializer)
 	: Super(ObjectInitializer)
 {
+#if WITH_EDITOR
 	UBlueprint* Blueprint = Cast<UBlueprint>(GetClass()->ClassGeneratedBy);
 	if (Blueprint)
 	{
 		OnBlueprintCompiledDelegate = Blueprint->OnCompiled().AddUObject(this, &ULiveLinkRemapAsset::OnBlueprintClassCompiled);
 	}
+#endif
 }
 
 void ULiveLinkRemapAsset::BeginDestroy()
 {
+#if WITH_EDITOR
 	if (OnBlueprintCompiledDelegate.IsValid())
 	{
 		UBlueprint* Blueprint = Cast<UBlueprint>(GetClass()->ClassGeneratedBy);
@@ -26,6 +29,7 @@ void ULiveLinkRemapAsset::BeginDestroy()
 		Blueprint->OnCompiled().Remove(OnBlueprintCompiledDelegate);
 		OnBlueprintCompiledDelegate.Reset();
 	}
+#endif
 
 	Super::BeginDestroy();
 }
diff --git a/Engine/Plugins/Animation/LiveLink/Source/LiveLink/Public/LiveLinkRemapAsset.h b/Engine/Plugins/Animation/LiveLink/Source/LiveLink/Public/LiveLinkRemapAsset.h
--- a/Engine/Plugins/Animation/LiveLink/Source/LiveLink/Public/LiveLinkRemapAsset.h
+++ b/Engine/Plugins/Animation/LiveLink/Source/LiveLink/Public/LiveLinkRemapAsset.h
@@ -46,6 +46,8 @@ private:
 	// (returned from GetRemappedCurveName)
 	TMap<FName, FName> CurveNameMap;
 
+#if WITH_EDITOR
 	/** Blueprint.OnCompiled delegate handle */
 	FDelegateHandle OnBlueprintCompiledDelegate;
+#endif
 };
\ No newline at end of file
diff --git a/Engine/Source/Developer/Merge/Private/SBlueprintMerge.cpp b/Engine/Source/Developer/Merge/Private/SBlueprintMerge.cpp
--- a/Engine/Source/Developer/Merge/Private/SBlueprintMerge.cpp
+++ b/Engine/Source/Developer/Merge/Private/SBlueprintMerge.cpp
@@ -339,7 +339,7 @@ void SBlueprintMerge::OnStartMerge()
 		{
 			UBlueprint* LocalBlueprint = GetTargetBlueprint();
 			// we use this to suppress blueprint compilation during StaticDuplicateObject()
-			TGuardValue<bool> GuardDuplicatingReadOnly(LocalBlueprint->bDuplicatingReadOnly, true);
+			FGuardValue_Bitfield(LocalBlueprint->bDuplicatingReadOnly, true);
 
 			UPackage* TransientPackage = GetTransientPackage();
 			Data.BlueprintLocal = Cast<const UBlueprint>(StaticDuplicateObject(LocalBlueprint, TransientPackage,
diff --git a/Engine/Source/Runtime/Engine/Classes/Engine/Blueprint.h b/Engine/Source/Runtime/Engine/Classes/Engine/Blueprint.h
--- a/Engine/Source/Runtime/Engine/Classes/Engine/Blueprint.h
+++ b/Engine/Source/Runtime/Engine/Classes/Engine/Blueprint.h
@@ -427,10 +427,6 @@
 {
 	GENERATED_UCLASS_BODY()
 
-	/** Whether or not this blueprint should recompile itself on load */
-	UPROPERTY(config)
-	uint32 bRecompileOnLoad:1;
-
 	/** 
 	 * Pointer to the parent class that the generated class should derive from. This *can* be null under rare circumstances, 
 	 * one such case can be created by creating a blueprint (A) based on another blueprint (B), shutting down the editor, and
@@ -439,41 +435,97 @@
 	UPROPERTY()
 	TSubclassOf<class UObject> ParentClass;
 
+	/** The type of this blueprint */
+	UPROPERTY(AssetRegistrySearchable)
+	TEnumAsByte<enum EBlueprintType> BlueprintType;
+
+	/** Whether or not this blueprint should recompile itself on load */
+	UPROPERTY(config)
+	uint8 bRecompileOnLoad:1;
+
 	/** When the class generated by this blueprint is loaded, it will be recompiled the first time.  After that initial recompile, subsequent loads will skip the regeneration step */
 	UPROPERTY(transient)
-	uint32 bHasBeenRegenerated:1;
+	uint8 bHasBeenRegenerated:1;
 
 	/** State flag to indicate whether or not the Blueprint is currently being regenerated on load */
 	UPROPERTY(transient)
-	uint32 bIsRegeneratingOnLoad:1;
+	uint8 bIsRegeneratingOnLoad:1;
 
 #if WITH_EDITORONLY_DATA
+
+	/** The blueprint is currently compiled */
+	UPROPERTY(transient)
+	uint8 bBeingCompiled:1;
+
 	/** Whether or not this blueprint is newly created, and hasn't been opened in an editor yet */
 	UPROPERTY(transient)
-	uint32 bIsNewlyCreated : 1;
+	uint8 bIsNewlyCreated : 1;
 
 	/** Whether to force opening the full (non data-only) editor for this blueprint. */
 	UPROPERTY(transient)
-	uint32 bForceFullEditor : 1;
+	uint8 bForceFullEditor : 1;
 
 	UPROPERTY(transient)
-	uint32 bQueuedForCompilation : 1 ;
+	uint8 bQueuedForCompilation : 1 ;
 
 	/**whether or not you want to continuously rerun the construction script for an actor as you drag it in the editor, or only when the drag operation is complete*/
 	UPROPERTY(EditAnywhere, Category=BlueprintOptions)
-	uint32 bRunConstructionScriptOnDrag : 1;
+	uint8 bRunConstructionScriptOnDrag : 1;
 
 	/**whether or not you want to continuously rerun the construction script for an actor in sequencer*/
 	UPROPERTY(EditAnywhere, Category=BlueprintOptions)
-	uint32 bRunConstructionScriptInSequencer : 1;
+	uint8 bRunConstructionScriptInSequencer : 1;
 
 	/** Whether or not this blueprint's class is a const class or not.  Should set CLASS_Const in the KismetCompiler. */
 	UPROPERTY(EditAnywhere, Category=ClassOptions, AdvancedDisplay)
-	uint32 bGenerateConstClass : 1;
+	uint8 bGenerateConstClass : 1;
 
 	/** Whether or not this blueprint's class is a abstract class or not.  Should set CLASS_Abstract in the KismetCompiler. */
 	UPROPERTY(EditAnywhere, Category = ClassOptions, AdvancedDisplay)
-	uint32 bGenerateAbstractClass : 1;
+	uint8 bGenerateAbstractClass : 1;
 
+	/** TRUE to show a warning when attempting to start in PIE and there is a compiler error on this Blueprint */
+	UPROPERTY(transient)
+	uint8 bDisplayCompilePIEWarning:1;
+
+	/** Deprecates the Blueprint, marking the generated class with the CLASS_Deprecated flag */
+	UPROPERTY(EditAnywhere, Category=ClassOptions, AdvancedDisplay)
+	uint8 bDeprecate:1;
+
+	/** 
+	 * Flag indicating that a read only duplicate of this blueprint is being created, used to disable logic in ::PostDuplicate,
+	 *
+	 * This flag needs to be copied on duplication (because it's the duplicated object that we're disabling on PostDuplicate),
+	 * but we don't *need* to serialize it for permanent objects.
+	 *
+	 * Without setting this flag a blueprint will be marked dirty when it is duplicated and if saved while in this dirty
+	 * state you will not be able to open the blueprint. More specifically, UClass::Rename (called by DestroyGeneratedClass)
+	 * sets a dirty flag on the package. Once saved the package will fail to open because some unnamed objects are present in
+	 * the pacakge.
+	 *
+	 * This flag can be used to avoid the package being marked as dirty in the first place. Ideally PostDuplicateObject
+	 * would not rename classes that are still in use by the original object.
+	 */
+	UPROPERTY()
+	mutable uint8 bDuplicatingReadOnly:1;
+
+private:
+	/** Deprecated properties. */
+	UPROPERTY()
+	uint8 bNativize_DEPRECATED:1;
+
+public:
+	/** When exclusive nativization is enabled, then this asset will be nativized. All super classes must be also nativized. */
+	UPROPERTY(transient)
+	EBlueprintNativizationFlag NativizationFlag;
+
+	/** The mode that will be used when compiling this class. */
+	UPROPERTY(EditAnywhere, Category=ClassOptions, AdvancedDisplay)
+	EBlueprintCompileMode CompileMode;
+
+	/** The current status of this blueprint */
+	UPROPERTY(transient)
+	TEnumAsByte<enum EBlueprintStatus> Status;
+
 	/** Overrides the BP's display name in the editor UI */
 	UPROPERTY(EditAnywhere, Category=BlueprintOptions)
@@ -491,27 +543,15 @@
 	UPROPERTY(EditAnywhere, Category=BlueprintOptions)
 	TArray<FString> HideCategories;
 	 
-	/** When exclusive nativization is enabled, then this asset will be nativized. All super classes must be also nativized. */
-	UPROPERTY(transient)
-	EBlueprintNativizationFlag NativizationFlag;
-
-	/** TRUE to show a warning when attempting to start in PIE and there is a compiler error on this Blueprint */
-	UPROPERTY(transient)
-	bool bDisplayCompilePIEWarning;
-
 	/** Guid key for finding searchable data for Blueprint in the DDC */
 	UPROPERTY()
 	FGuid SearchGuid;
-
-	/** Deprecates the Blueprint, marking the generated class with the CLASS_Deprecated flag */
-	UPROPERTY(EditAnywhere, Category=ClassOptions, AdvancedDisplay)
-	bool bDeprecate;
-
-	/** The mode that will be used when compiling this class. */
-	UPROPERTY(EditAnywhere, Category=ClassOptions, AdvancedDisplay)
-	EBlueprintCompileMode CompileMode;
 #endif //WITH_EDITORONLY_DATA
 
+	/** The version of the blueprint system that was used to  create this blueprint */
+	UPROPERTY()
+	int32 BlueprintSystemVersion;
+
 	/** 'Simple' construction script - graph of components to instance */
 	UPROPERTY()
 	class USimpleConstructionScript* SimpleConstructionScript;
@@ -544,23 +584,6 @@
 	/** Cached cosmetic information about macro graphs, use GetCosmeticInfoForMacro() to access */
 	UPROPERTY(Transient)
 	TMap<UEdGraph*, FBlueprintMacroCosmeticInfo> PRIVATE_CachedMacroInfo;
-
-	/** 
-	 * Flag indicating that a read only duplicate of this blueprint is being created, used to disable logic in ::PostDuplicate,
-	 *
-	 * This flag needs to be copied on duplication (because it's the duplicated object that we're disabling on PostDuplicate),
-	 * but we don't *need* to serialize it for permanent objects.
-	 *
-	 * Without setting this flag a blueprint will be marked dirty when it is duplicated and if saved while in this dirty
-	 * state you will not be able to open the blueprint. More specifically, UClass::Rename (called by DestroyGeneratedClass)
-	 * sets a dirty flag on the package. Once saved the package will fail to open because some unnamed objects are present in
-	 * the pacakge.
-	 *
-	 * This flag can be used to avoid the package being marked as dirty in the first place. Ideally PostDuplicateObject
-	 * would not rename classes that are still in use by the original object.
-	 */
-	UPROPERTY()
-	mutable bool bDuplicatingReadOnly;
 #endif // WITH_EDITORONLY_DATA
 
 	/** Array of component template objects, used by AddComponent function */
@@ -575,15 +598,7 @@
 	UPROPERTY()
 	class UInheritableComponentHandler* InheritableComponentHandler;
 
-	/** The type of this blueprint */
-	UPROPERTY(AssetRegistrySearchable)
-	TEnumAsByte<enum EBlueprintType> BlueprintType;
-
 #if WITH_EDITORONLY_DATA
-	/** The current status of this blueprint */
-	UPROPERTY(transient)
-	TEnumAsByte<enum EBlueprintStatus> Status;
-
 	/** Array of new variables to be added to generated class */
 	UPROPERTY()
 	TArray<struct FBPVariableDescription> NewVariables;
@@ -596,13 +611,6 @@
 	UPROPERTY(AssetRegistrySearchable)
 	TArray<struct FBPInterfaceDescription> ImplementedInterfaces;
 	
-#endif // WITH_EDITORONLY_DATA
-
-	/** The version of the blueprint system that was used to  create this blueprint */
-	UPROPERTY()
-	int32 BlueprintSystemVersion;
-
-#if WITH_EDITORONLY_DATA
 	/** Set of documents that were being edited in this blueprint, so we can open them right away */
 	UPROPERTY()
 	TArray<struct FEditedDocumentInfo> LastEditedDocuments;
@@ -636,6 +644,7 @@
 
 public:
 
+#if WITH_EDITOR
 	/** Broadcasts a notification whenever the blueprint has changed. */
 	DECLARE_EVENT_OneParam( UBlueprint, FChangedEvent, class UBlueprint* );
 	FChangedEvent& OnChanged() { return ChangedEvent; }
@@ -647,6 +656,7 @@
 	DECLARE_EVENT_OneParam(UBlueprint, FCompiledEvent, class UBlueprint*);
 	FCompiledEvent& OnCompiled() { return CompiledEvent; }
 	void BroadcastCompiled() { CompiledEvent.Broadcast(this); }
+#endif
 
 	/** Whether or not this blueprint can be considered for a bytecode only compile */
 	virtual bool IsValidForBytecodeOnlyRecompile() const { return true; }
@@ -672,10 +682,6 @@
 	UPROPERTY(VisibleAnywhere, Instanced, Category=Thumbnail)
 	class UThumbnailInfo* ThumbnailInfo;
 
-	/** The blueprint is currently compiled */
-	UPROPERTY(transient)
-	uint32 bBeingCompiled:1;
-
 	/** CRC for CDO calculated right after the latest compilation used by Reinstancer to check if default values were changed */
 	UPROPERTY(transient, duplicatetransient)
 	uint32 CrcLastCompiledCDO;
@@ -987,6 +993,7 @@
 	*/
 	virtual void GetInstanceActions(FBlueprintActionDatabaseRegistrar& ActionRegistrar) const {}
 
+#if WITH_EDITOR
 private:
 
 	/** Broadcasts a notification whenever the blueprint has changed. */
@@ -995,11 +1002,6 @@
 	/** Broadcasts a notification whenever the blueprint is compiled. */
 	FCompiledEvent CompiledEvent;
 
-	/** Deprecated properties. */
-	UPROPERTY()
-	bool bNativize_DEPRECATED;
-
-#if WITH_EDITOR
 public:
 	/** If this blueprint is currently being compiled, the CurrentMessageLog will be the log currently being used to send logs to. */
 	class FCompilerResultsLog* CurrentMessageLog;
diff --git a/Engine/Source/Runtime/Engine/Private/Blueprint.cpp b/Engine/Source/Runtime/Engine/Private/Blueprint.cpp
--- a/Engine/Source/Runtime/Engine/Private/Blueprint.cpp
+++ b/Engine/Source/Runtime/Engine/Private/Blueprint.cpp
@@ -462,6 +462,7 @@ void UBlueprint::Serialize(FArchive& Ar)
 
 		if (Ar.IsLoading())
 		{
+#if WITH_EDITORONLY_DATA
 			if (bNativize_DEPRECATED)
 			{
 				// Migrate to the new transient flag.
@@ -472,6 +473,7 @@ void UBlueprint::Serialize(FArchive& Ar)
 				bSettingsChanged |= PackagingSettings->AddBlueprintAssetToNativizationList(this);
 			}
 			else
+#endif
 			{
 				// Cache whether or not this Blueprint asset was selected for exclusive nativization in the Project Settings.
 				for (int AssetIndex = 0; AssetIndex < PackagingSettings->NativizeBlueprintAssets.Num(); ++AssetIndex)
