Add the ability to Add two DateTime structs together

Added functions to the KismetMathLibrary and the the FDateTime struct.

diff --git a/Engine/Source/Runtime/Core/Public/Misc/DateTime.h b/Engine/Source/Runtime/Core/Public/Misc/DateTime.h
--- a/Engine/Source/Runtime/Core/Public/Misc/DateTime.h
+++ b/Engine/Source/Runtime/Core/Public/Misc/DateTime.h
@@ -124,6 +124,19 @@ public:
 		return *this;
 	}
 
+	/**
+	 * Adds the time from the given date to this date.
+	 *
+	 * @return This date.
+	 * @see FDateTime
+	 */
+	FDateTime& operator+(const FDateTime& Other)
+	{
+		Ticks += Other.Ticks;
+
+		return *this;
+	}
+
 	/**
 	 * Returns time span between this date and the given date.
 	 *
diff --git a/Engine/Source/Runtime/Engine/Classes/Kismet/KismetMathLibrary.h b/Engine/Source/Runtime/Engine/Classes/Kismet/KismetMathLibrary.h
--- a/Engine/Source/Runtime/Engine/Classes/Kismet/KismetMathLibrary.h
+++ b/Engine/Source/Runtime/Engine/Classes/Kismet/KismetMathLibrary.h
@@ -2671,6 +2671,10 @@ class ENGINE_API UKismetMathLibrary : public UBlueprintFunctionLibrary
 	UFUNCTION(BlueprintPure, meta=(DisplayName="DateTime - Timespan", CompactNodeTitle="-", Keywords="- subtract minus"), Category="Math|DateTime")
 	static FDateTime Subtract_DateTimeTimespan(FDateTime A, FTimespan B);
 
+	/** Addition (A + B) */
+	UFUNCTION(BlueprintPure, meta = (DisplayName = "DateTime + DateTime", CompactNodeTitle = "+", Keywords = "+ add plus"), Category = "Math|DateTime")
+	static FDateTime Add_DateTimeDateTime(FDateTime A, FDateTime B);
+
 	/** Subtraction (A - B) */
 	UFUNCTION(BlueprintPure, meta = (DisplayName = "DateTime - DateTime", CompactNodeTitle = "-", Keywords = "- subtract minus"), Category = "Math|DateTime")
 	static FTimespan Subtract_DateTimeDateTime(FDateTime A, FDateTime B);
diff --git a/Engine/Source/Runtime/Engine/Classes/Kismet/KismetMathLibrary.inl b/Engine/Source/Runtime/Engine/Classes/Kismet/KismetMathLibrary.inl
--- a/Engine/Source/Runtime/Engine/Classes/Kismet/KismetMathLibrary.inl
+++ b/Engine/Source/Runtime/Engine/Classes/Kismet/KismetMathLibrary.inl
@@ -2082,6 +2082,12 @@ FTimespan UKismetMathLibrary::Subtract_DateTimeDateTime(FDateTime A, FDateTime B
 	return A - B;
 }
 
+KISMET_MATH_FORCEINLINE
+FDateTime UKismetMathLibrary::Add_DateTimeDateTime(FDateTime A, FDateTime B)
+{
+	return A + B;
+}
+
 KISMET_MATH_FORCEINLINE
 bool UKismetMathLibrary::EqualEqual_DateTimeDateTime(FDateTime A, FDateTime B)
 {
