Make initial controller ID set for ULocalPlayer::PlayerAdded use the overridable SetControllerID function.
Avoid swapping controller IDs with self in UEngine::SwapControllerId

diff --git a/Engine/Source/Runtime/Engine/Private/LocalPlayer.cpp b/Engine/Source/Runtime/Engine/Private/LocalPlayer.cpp
--- a/Engine/Source/Runtime/Engine/Private/LocalPlayer.cpp
+++ b/Engine/Source/Runtime/Engine/Private/LocalPlayer.cpp
@@ -256,8 +256,8 @@ void ULocalPlayer::PostInitProperties()
 
 void ULocalPlayer::PlayerAdded(UGameViewportClient* InViewportClient, int32 InControllerID)
 {
-	ViewportClient		= InViewportClient;
-	ControllerId		= InControllerID;
+	ViewportClient = InViewportClient;
+	SetControllerId(InControllerID);
 
 	SubsystemCollection.Initialize();
 }
diff --git a/Engine/Source/Runtime/Engine/Private/UnrealEngine.cpp b/Engine/Source/Runtime/Engine/Private/UnrealEngine.cpp
--- a/Engine/Source/Runtime/Engine/Private/UnrealEngine.cpp
+++ b/Engine/Source/Runtime/Engine/Private/UnrealEngine.cpp
@@ -3179,7 +3179,7 @@ ULocalPlayer* UEngine::GetLocalPlayerFromControllerId( UWorld * InWorld, const i
 	return GetLocalPlayerFromControllerId_local(GamePlayers, ControllerId);
 }
 
-void UEngine::SwapControllerId(ULocalPlayer *NewPlayer, const int32 CurrentControllerId, const int32 NewControllerID) const
+void UEngine::SwapControllerId(ULocalPlayer* NewPlayer, const int32 CurrentControllerId, const int32 NewControllerID) const
 {
 	for (auto It = WorldList.CreateConstIterator(); It; ++It)
 	{
@@ -3195,9 +3195,9 @@ void UEngine::SwapControllerId(ULocalPlayer *NewPlayer, const int32 CurrentContr
 			// This is the world context that NewPlayer belongs to, see if anyone is using his CurrentControllerId
 			for (ULocalPlayer* LocalPlayer : LocalPlayers)
 			{
-				if(LocalPlayer && LocalPlayer->GetControllerId() == NewControllerID)
+				if (LocalPlayer && LocalPlayer != NewPlayer && LocalPlayer->GetControllerId() == NewControllerID)
 				{
-					LocalPlayer->SetControllerId( CurrentControllerId);
+					LocalPlayer->SetControllerId(CurrentControllerId);
 					return;
 				}
 			}
