Fix non-editor builds by adding #ifdefs around IsDataValid functions.

diff --git a/Engine/Source/Runtime/Engine/Classes/Engine/SCS_Node.h b/Engine/Source/Runtime/Engine/Classes/Engine/SCS_Node.h
--- a/Engine/Source/Runtime/Engine/Classes/Engine/SCS_Node.h
+++ b/Engine/Source/Runtime/Engine/Classes/Engine/SCS_Node.h
@@ -93,10 +93,10 @@ class USCS_Node : public UObject
 
 	/** Make sure, the guid is proper - backward compatibility */
 	ENGINE_API void ValidateGuid();
-#endif
 
 	/** Returns Valid if this object has data validation rules set up for it and the data for this object is valid. Returns Invalid if it does not pass the rules. Returns NotValidated if no rules are set for this object. */
 	virtual EDataValidationResult IsDataValid(TArray<FText>& ValidationErrors) override;
+#endif // WITH_EDITOR
 
 	/**
 	 * Create the specified component on the actor, then call action on children
diff --git a/Engine/Source/Runtime/Engine/Classes/Engine/SimpleConstructionScript.h b/Engine/Source/Runtime/Engine/Classes/Engine/SimpleConstructionScript.h
--- a/Engine/Source/Runtime/Engine/Classes/Engine/SimpleConstructionScript.h
+++ b/Engine/Source/Runtime/Engine/Classes/Engine/SimpleConstructionScript.h
@@ -100,8 +100,10 @@ class USimpleConstructionScript : public UObject
 	/** Checks the root node set for scene components and ensures that it is valid (e.g. after a removal) */
 	ENGINE_API void ValidateSceneRootNodes();
 
+#if WITH_EDITOR
 	/** Returns Valid if this object has data validation rules set up for it and the data for this object is valid. Returns Invalid if it does not pass the rules. Returns NotValidated if no rules are set for this object. */
 	virtual EDataValidationResult IsDataValid(TArray<FText>& ValidationErrors) override;
+#endif // WITH_EDITOR
 
 private:
 	/** Root nodes of the construction script */
diff --git a/Engine/Source/Runtime/Engine/Private/SCS_Node.cpp b/Engine/Source/Runtime/Engine/Private/SCS_Node.cpp
--- a/Engine/Source/Runtime/Engine/Private/SCS_Node.cpp
+++ b/Engine/Source/Runtime/Engine/Private/SCS_Node.cpp
@@ -600,8 +600,6 @@ void USCS_Node::ValidateGuid()
 	}
 }
 
-#endif
-
 EDataValidationResult USCS_Node::IsDataValid(TArray<FText>& ValidationErrors)
 {
 	EDataValidationResult Result = Super::IsDataValid(ValidationErrors);
@@ -626,3 +624,5 @@ EDataValidationResult USCS_Node::IsDataValid(TArray<FText>& ValidationErrors)
 
 	return Result;
 }
+
+#endif // WITH_EDITOR
