PR #5806 : Clarify comment of VectorToVector2D conversion

diff --git a/Engine/Source/Runtime/Engine/Classes/Kismet/KismetMathLibrary.h b/Engine/Source/Runtime/Engine/Classes/Kismet/KismetMathLibrary.h
--- a/Engine/Source/Runtime/Engine/Classes/Kismet/KismetMathLibrary.h
+++ b/Engine/Source/Runtime/Engine/Classes/Kismet/KismetMathLibrary.h
@@ -1161,7 +1161,7 @@ class ENGINE_API UKismetMathLibrary : public UBlueprintFunctionLibrary
 	UFUNCTION(BlueprintPure, meta=(DisplayName = "ToTransform (Vector)", CompactNodeTitle = "->", ScriptMethod = "Transform", Keywords="cast convert", BlueprintAutocast), Category="Math|Conversions")
 	static FTransform Conv_VectorToTransform(FVector InLocation);
 	
-	/** Convert a Vector to a Vector2D */
+	/** Convert a Vector to a Vector2D using the Vector's (X, Y) coordinates */
 	UFUNCTION(BlueprintPure, meta=(DisplayName = "ToVector2D (Vector)", CompactNodeTitle = "->", ScriptMethod = "Vector2D", Keywords="cast convert", BlueprintAutocast), Category="Math|Conversions")
 	static FVector2D Conv_VectorToVector2D(FVector InVector);
 
