Make AActor::InputComponent DuplicateTransient avoiding duplicated actors being left with an InputComponent that has no bindings.

diff --git a/Engine/Source/Runtime/Engine/Classes/GameFramework/Actor.h b/Engine/Source/Runtime/Engine/Classes/GameFramework/Actor.h
--- a/Engine/Source/Runtime/Engine/Classes/GameFramework/Actor.h
+++ b/Engine/Source/Runtime/Engine/Classes/GameFramework/Actor.h
@@ -533,7 +533,7 @@ public:
 	int32 InputPriority;
 
 	/** Component that handles input for this actor, if input is enabled. */
-	UPROPERTY()
+	UPROPERTY(DuplicateTransient)
 	class UInputComponent* InputComponent;
 
 	/** Square of the max distance from the client's viewpoint that this actor is relevant and will be replicated. */
