PR #6073: Moved UPlayerInput::MassageVectorAxisInput and UPlayerInput::MassageAxisInput from private to protected (Contributed by [REDACTED_PERSON])

diff --git a/Engine/Source/Runtime/Engine/Classes/GameFramework/PlayerInput.h b/Engine/Source/Runtime/Engine/Classes/GameFramework/PlayerInput.h
--- a/Engine/Source/Runtime/Engine/Classes/GameFramework/PlayerInput.h
+++ b/Engine/Source/Runtime/Engine/Classes/GameFramework/PlayerInput.h
@@ -622,7 +622,7 @@ public:
 	static const TArray<FInputActionKeyMapping>& GetEngineDefinedActionMappings() { return EngineDefinedActionMappings; }
 	static const TArray<FInputAxisKeyMapping>& GetEngineDefinedAxisMappings() { return EngineDefinedAxisMappings; }
 
-private:
+protected:
 
 	/**
 	* Given raw keystate value of a vector axis, returns the "massaged" value. Override for any custom behavior,
@@ -636,6 +636,8 @@ private:
 	*/
 	virtual float MassageAxisInput(FKey Key, float RawValue);
 
+private:
+
 	/** Process non-axes keystates */
 	void ProcessNonAxesKeys(FKey Inkey, FKeyState* KeyState);
 	
