Add an option to preserve links when removing a node from a graph. Useful for moving nodes to different graphs and preserving connections.

#rnx

diff --git a/Engine/Source/Runtime/Engine/Classes/EdGraph/EdGraph.h b/Engine/Source/Runtime/Engine/Classes/EdGraph/EdGraph.h
--- a/Engine/Source/Runtime/Engine/Classes/EdGraph/EdGraph.h
+++ b/Engine/Source/Runtime/Engine/Classes/EdGraph/EdGraph.h
@@ -166,8 +166,15 @@ public:
 	 */
 	void SelectNodeSet(TSet<const UEdGraphNode*> NodeSelection, bool bFromUI = false);
 
-	/** Remove a node from this graph */
-	bool RemoveNode( UEdGraphNode* NodeToRemove );
+	/** 
+	* Remove a node from this graph
+	* 
+	* @param NodeToRemove		The node to remove from this graph
+	* @param bBreakAllLinks		If true, all links will be broken on the given node. Editor only. Useful for moving nodes to a different graph.
+	* 
+	* @return True if the node has been removed from the graph
+	*/
+	bool RemoveNode( UEdGraphNode* NodeToRemove, bool bBreakAllLinks = true );
 
 	/** Signal to listeners that the graph has changed */
 	virtual void NotifyGraphChanged();
diff --git a/Engine/Source/Runtime/Engine/Private/EdGraph/EdGraph.cpp b/Engine/Source/Runtime/Engine/Private/EdGraph/EdGraph.cpp
--- a/Engine/Source/Runtime/Engine/Private/EdGraph/EdGraph.cpp
+++ b/Engine/Source/Runtime/Engine/Private/EdGraph/EdGraph.cpp
@@ -253,13 +253,16 @@ void UEdGraph::SelectNodeSet(TSet<const UEdGraphNode*> NodeSelection, bool bFrom
 	NotifyGraphChanged(SelectionAction);
 }
 
-bool UEdGraph::RemoveNode( UEdGraphNode* NodeToRemove )
+bool UEdGraph::RemoveNode( UEdGraphNode* NodeToRemove, bool bBreakAllLinks )
 {
 	Modify();
 
 	int32 NumTimesNodeRemoved = Nodes.Remove(NodeToRemove);
 #if WITH_EDITOR
-	NodeToRemove->BreakAllNodeLinks();
+	if (bBreakAllLinks)
+	{
+		NodeToRemove->BreakAllNodeLinks();
+	}
 #endif	//#if WITH_EDITOR
 
 	FEdGraphEditAction RemovalAction;
