bVisualizeComponent should not be a uproperty as it is constant per class and not user editable
#rnx

diff --git a/Engine/Source/Runtime/Engine/Classes/Components/SceneComponent.h b/Engine/Source/Runtime/Engine/Classes/Components/SceneComponent.h
--- a/Engine/Source/Runtime/Engine/Classes/Components/SceneComponent.h
+++ b/Engine/Source/Runtime/Engine/Classes/Components/SceneComponent.h
@@ -262,8 +262,7 @@ public:
 	static int32 SkipUpdateOverlapsOptimEnabled;
 
 #if WITH_EDITORONLY_DATA
-	/** This component is explicitly for visualization in the editor */
-	UPROPERTY()
+	/** This component should create a sprite component for visualization in the editor */
 	uint8 bVisualizeComponent : 1;
 #endif
 
