Non-trivial resolve of CL# 6854616 against CL# 7785269 during integration to Dev-Framework
#rnx

diff --git a/Engine/Source/Runtime/Engine/Private/UserInterface/Console.cpp b/Engine/Source/Runtime/Engine/Private/UserInterface/Console.cpp
--- a/Engine/Source/Runtime/Engine/Private/UserInterface/Console.cpp
+++ b/Engine/Source/Runtime/Engine/Private/UserInterface/Console.cpp
@@ -38,6 +38,26 @@ static const FName NAME_Open = FName(TEXT("Open"));
 
 UConsole::FRegisterConsoleAutoCompleteEntries UConsole::RegisterConsoleAutoCompleteEntries;
 UConsole::FOnConsoleActivationStateChanged UConsole::OnConsoleActivationStateChanged;
+
+static TAutoConsoleVariable<int32> CVarCustomConsolePosEnabled(
+	TEXT("console.position.enable"),
+	0,
+	TEXT("Enable custom console positioning \n"),
+	ECVF_Default);
+
+static TAutoConsoleVariable<int32> CVarConsoleXPos(
+	TEXT("console.position.x"),
+	0,
+	TEXT("Console X offset from left border \n"),
+	ECVF_Default);
+
+static TAutoConsoleVariable<int32> CVarConsoleYPos(
+	TEXT("console.position.y"),
+	0,
+	TEXT("Console Y offset from bottom border \n"),
+	ECVF_Default);
+
+
 namespace ConsoleDefs
 {
 	/** Colors */
@@ -1097,6 +1117,13 @@ void UConsole::PostRender_Console_Typing(UCanvas* Canvas)
 	float ClipY = Canvas->ClipY;
 	float LeftPos = 0;
 
+	if (CVarCustomConsolePosEnabled.GetValueOnAnyThread())
+	{
+		LeftPos = (float)CVarConsoleXPos.GetValueOnAnyThread();
+		float BottomOffset = (float)CVarConsoleYPos.GetValueOnAnyThread();
+		ClipY = ClipY - BottomOffset;
+	}
+
 	PostRender_InputLine(Canvas, FIntPoint(LeftPos, ClipY));
 }
 
@@ -1174,6 +1201,19 @@ void UConsole::PostRender_Console_Open(UCanvas* Canvas)
 	float TopPos = 0;
 	float LeftPos = 0;
 
+	if (GEngine->IsConsoleBuild())
+	{
+		ClipX	-= 80;
+		TopPos	 = 30;
+		LeftPos	 = 40;
+	}
+	if (GEngine->IsStereoscopic3D())
+	{
+		LeftPos = ClipX / 3;
+		ClipX -= LeftPos;
+		Height = Canvas->ClipY * 0.60;
+	}
+
 	UFont* Font = GEngine->GetSmallFont();
 
 	// determine the height of the text
