Redo non-trivial resolve of CL# 6854616 against CL# 7785269 during integration to Dev-Framework
#rnx

diff --git a/Engine/Source/Runtime/Engine/Private/UserInterface/Console.cpp b/Engine/Source/Runtime/Engine/Private/UserInterface/Console.cpp
--- a/Engine/Source/Runtime/Engine/Private/UserInterface/Console.cpp
+++ b/Engine/Source/Runtime/Engine/Private/UserInterface/Console.cpp
@@ -61,30 +61,30 @@ static TAutoConsoleVariable<int32> CVarConsoleYPos(
 namespace ConsoleDefs
 {
 	/** Colors */
-	static const FColor BorderColor( 140, 140, 140 );
-	static const FColor CursorColor( 255, 255, 255 );
-	static const FColor AutocompleteBackgroundColor( 0, 0, 0 );
-	static const FColor CursorLineColor( 0, 50, 0 );
+	static const FColor BorderColor(140, 140, 140);
+	static const FColor CursorColor(255, 255, 255);
+	static const FColor AutocompleteBackgroundColor(0, 0, 0);
+	static const FColor CursorLineColor(0, 50, 0);
 	static const int32 AutocompleteGap = 6;
 
 	/** Text that appears before the user's typed input string that acts as a visual cue for the editable area */
 	static const FString LeadingInputText(TEXT(" > "));
 }
 
-class FConsoleVariableAutoCompleteVisitor 
+class FConsoleVariableAutoCompleteVisitor
 {
 public:
 	// @param Name must not be 0
 	// @param CVar must not be 0
-	static void OnConsoleVariable(const TCHAR *Name, IConsoleObject* CVar,TArray<struct FAutoCompleteCommand>& Sink)
+	static void OnConsoleVariable(const TCHAR* Name, IConsoleObject* CVar, TArray<struct FAutoCompleteCommand>& Sink)
 	{
 #if DISABLE_CHEAT_CVARS
-		if(CVar->TestFlags(ECVF_Cheat))
+		if (CVar->TestFlags(ECVF_Cheat))
 		{
 			return;
 		}
 #endif // DISABLE_CHEAT_CVARS
-		if(CVar->TestFlags(ECVF_Unregistered))
+		if (CVar->TestFlags(ECVF_Unregistered))
 		{
 			return;
 		}
@@ -139,7 +139,7 @@ UConsole::UConsole(const FObjectInitializer& ObjectInitializer)
 UConsole::~UConsole()
 {
 	// At shutdown, GLog may already be null
-	if( GLog != nullptr )
+	if (GLog != nullptr)
 	{
 		GLog->RemoveOutputDevice(this);
 	}
@@ -173,7 +173,7 @@ void UConsole::BuildRuntimeAutoCompleteList(bool bForce)
 	//@todo - probably only need to rebuild the tree + partial command list on level load
 	for (int32 Idx = 0; Idx < AutoCompleteTree.ChildNodes.Num(); Idx++)
 	{
-		FAutoCompleteNode *Node = AutoCompleteTree.ChildNodes[Idx];
+		FAutoCompleteNode* Node = AutoCompleteTree.ChildNodes[Idx];
 		delete Node;
 	}
 
@@ -196,32 +196,32 @@ void UConsole::BuildRuntimeAutoCompleteList(bool bForce)
 		IConsoleManager::Get().ForEachConsoleObjectThatStartsWith(
 			FConsoleObjectVisitor::CreateStatic< TArray<struct FAutoCompleteCommand>& >(
 				&FConsoleVariableAutoCompleteVisitor::OnConsoleVariable,
-				AutoCompleteList ) );
+				AutoCompleteList));
 	}
 
 	// iterate through script exec functions and append to the list
 	for (TObjectIterator<UFunction> It; It; ++It)
 	{
-		UFunction *Func = *It;
+		UFunction* Func = *It;
 
 		// Determine whether or not this is a level script event that we can call (must be defined in the level script actor and not in parent, and has no return value)
 		const UClass* FuncOuter = Cast<UClass>(Func->GetOuter());
-		const bool bIsLevelScriptFunction = FuncOuter 
+		const bool bIsLevelScriptFunction = FuncOuter
 			&& (FuncOuter->IsChildOf(ALevelScriptActor::StaticClass()))
 			&& (FuncOuter != ALevelScriptActor::StaticClass())
-			&& (Func->ReturnValueOffset == MAX_uint16) 
+			&& (Func->ReturnValueOffset == MAX_uint16)
 			&& (Func->GetSuperFunction() == nullptr);
 
 		// exec functions that either have no parent, level script events, or are in the global state (filtering some unnecessary dupes)
-		if ( (Func->HasAnyFunctionFlags(FUNC_Exec) && (Func->GetSuperFunction() == nullptr || FuncOuter))
+		if ((Func->HasAnyFunctionFlags(FUNC_Exec) && (Func->GetSuperFunction() == nullptr || FuncOuter))
 			|| bIsLevelScriptFunction)
 		{
 			FString FuncName = Func->GetName();
-			if(FDefaultValueHelper::HasWhitespaces(FuncName))
+			if (FDefaultValueHelper::HasWhitespaces(FuncName))
 			{
 				FuncName = FString::Printf(TEXT("\"%s\""), *FuncName);
 			}
-			if( bIsLevelScriptFunction )
+			if (bIsLevelScriptFunction)
 			{
 				FuncName = FString(TEXT("ce ")) + FuncName;
 			}
@@ -245,8 +245,8 @@ void UConsole::BuildRuntimeAutoCompleteList(bool bForce)
 			// append each property (and it's type) to the help string
 			for (TFieldIterator<UProperty> PropIt(Func); PropIt && (PropIt->PropertyFlags & CPF_Parm); ++PropIt)
 			{
-				UProperty *Prop = *PropIt;
-				Desc += FString::Printf(TEXT("%s[%s] "),*Prop->GetName(),*Prop->GetCPPType());
+				UProperty* Prop = *PropIt;
+				Desc += FString::Printf(TEXT("%s[%s] "), *Prop->GetName(), *Prop->GetCPPType());
 			}
 			AutoCompleteList[NewIdx].Desc = Desc + AutoCompleteList[NewIdx].Desc;
 		}
@@ -259,7 +259,7 @@ void UConsole::BuildRuntimeAutoCompleteList(bool bForce)
 		{
 			FPackageName::FindPackagesInDirectory(Packages, FString::Printf(TEXT("%s%s"), *FPaths::ProjectDir(), *ConsoleSettings->AutoCompleteMapPaths[PathIdx]));
 		}
-	
+
 		// also include maps in this user's developer dir
 		FPackageName::FindPackagesInDirectory(Packages, FPaths::GameUserDeveloperDir());
 
@@ -340,24 +340,24 @@ void UConsole::BuildRuntimeAutoCompleteList(bool bForce)
 				: AutoCompleteList(InAutoCompleteList)
 			{
 			}
-			
+
 			bool OnEngineShowFlag(uint32 InIndex, const FString& InName)
 			{
 				// Get localized name.
 				FText LocName;
 				FEngineShowFlags::FindShowFlagDisplayName(InName, LocName);
-				
+
 				int32 NewIdx = AutoCompleteList.AddDefaulted();
 				AutoCompleteList[NewIdx].Command = TEXT("show ") + InName;
-				AutoCompleteList[NewIdx].Desc = FString::Printf(TEXT("(toggles the %s showflag)"),*LocName.ToString());
+				AutoCompleteList[NewIdx].Desc = FString::Printf(TEXT("(toggles the %s showflag)"), *LocName.ToString());
 				AutoCompleteList[NewIdx].Color = GetDefault<UConsoleSettings>()->AutoCompleteCommandColor;
-				
+
 				return true;
 			}
-			
+
 			TArray<FAutoCompleteCommand>& AutoCompleteList;
 		};
-		
+
 		FIterSink Sink(AutoCompleteList);
 		FEngineShowFlags::IterateAllFlags(Sink);
 	}
@@ -371,12 +371,12 @@ void UConsole::BuildRuntimeAutoCompleteList(bool bForce)
 	for (int32 ListIdx = 0; ListIdx < AutoCompleteList.Num(); ListIdx++)
 	{
 		FString Command = AutoCompleteList[ListIdx].Command.ToLower();
-		FAutoCompleteNode *Node = &AutoCompleteTree;
+		FAutoCompleteNode* Node = &AutoCompleteTree;
 		for (int32 Depth = 0; Depth < Command.Len(); Depth++)
 		{
 			int32 Char = Command[Depth];
 			int32 FoundNodeIdx = INDEX_NONE;
-			TArray<FAutoCompleteNode*> &NodeList = Node->ChildNodes;
+			TArray<FAutoCompleteNode*>& NodeList = Node->ChildNodes;
 			for (int32 NodeIdx = 0; NodeIdx < NodeList.Num(); NodeIdx++)
 			{
 				if (NodeList[NodeIdx]->IndexChar == Char)
@@ -389,7 +389,7 @@ void UConsole::BuildRuntimeAutoCompleteList(bool bForce)
 			}
 			if (FoundNodeIdx == INDEX_NONE)
 			{
-				FAutoCompleteNode *NewNode = new FAutoCompleteNode(Char);
+				FAutoCompleteNode* NewNode = new FAutoCompleteNode(Char);
 				NewNode->AutoCompleteListIndices.Add(ListIdx);
 				Node->ChildNodes.Add(NewNode);
 				Node = NewNode;
@@ -448,7 +448,7 @@ void UConsole::UpdateCompleteIndices()
 	AutoCompleteIndex = 0;
 	AutoCompleteCursor = 0;
 	AutoComplete.Empty();
-	FAutoCompleteNode *Node = &AutoCompleteTree;
+	FAutoCompleteNode* Node = &AutoCompleteTree;
 	FString LowerTypedStr = TypedStr.ToLower();
 	int32 EndIdx = -1;
 	for (int32 Idx = 0; Idx < TypedStr.Len(); Idx++)
@@ -475,7 +475,7 @@ void UConsole::UpdateCompleteIndices()
 			}
 			else
 			{
-				if(Idx < TypedStr.Len())
+				if (Idx < TypedStr.Len())
 				{
 					// if the first non-matching character is a space we might be adding parameters, stay on the last node we found so users can see the parameter info
 					if (TypedStr[Idx] == TCHAR(' '))
@@ -497,7 +497,7 @@ void UConsole::UpdateCompleteIndices()
 	{
 		const TArray<int32>& Leaf = Node->AutoCompleteListIndices;
 
-		for(uint32 i = 0, Num = (uint32)Leaf.Num(); i < Num; ++i)
+		for (uint32 i = 0, Num = (uint32)Leaf.Num(); i < Num; ++i)
 		{
 			// if we're adding parameters we want to make sure that we only display exact matches
 			// ie Typing "Foo 5" should still show info for "Foo" but not for "FooBar"
@@ -516,11 +516,11 @@ void UConsole::SetAutoCompleteFromHistory()
 	AutoCompleteCursor = 0;
 	AutoComplete.Empty();
 
-	for(int32 i = HistoryBuffer.Num() - 1; i >= 0 ; --i)
+	for (int32 i = HistoryBuffer.Num() - 1; i >= 0; --i)
 	{
 		FAutoCompleteCommand Cmd;
 
-		Cmd.Command = HistoryBuffer[i]; 
+		Cmd.Command = HistoryBuffer[i];
 		Cmd.Color = ConsoleSettings->HistoryColor;
 		Cmd.SetHistory();
 
@@ -528,12 +528,12 @@ void UConsole::SetAutoCompleteFromHistory()
 	}
 }
 
-void UConsole::SetInputText( const FString& Text )
+void UConsole::SetInputText(const FString& Text)
 {
 	TypedStr = Text;
 }
 
-void UConsole::SetCursorPos( int32 Position )
+void UConsole::SetCursorPos(int32 Position)
 {
 	TypedStrPos = Position;
 }
@@ -555,12 +555,12 @@ void UConsole::ConsoleCommand(const FString& Command)
 	OutputText(FString::Printf(TEXT("\n>>> %s <<<"), *Command));
 
 	UGameInstance* GameInstance = GetOuterUGameViewportClient()->GetGameInstance();
-	if(ConsoleTargetPlayer != nullptr)
+	if (ConsoleTargetPlayer != nullptr)
 	{
 		// If there is a console target player, execute the command in the player's context.
 		ConsoleTargetPlayer->PlayerController->ConsoleCommand(Command);
 	}
-	else if(GameInstance && GameInstance->GetFirstLocalPlayerController())
+	else if (GameInstance && GameInstance->GetFirstLocalPlayerController())
 	{
 		// If there are any players, execute the command in the first local player's context.
 		APlayerController* PC = GameInstance->GetFirstLocalPlayerController();
@@ -603,11 +603,11 @@ void UConsole::OutputText(const FString& Text)
 {
 	FString RemainingText = Text;
 	int32 StringLength = Text.Len();
-	while(StringLength > 0)
+	while (StringLength > 0)
 	{
 		// Find the number of characters in the next line of text.
 		int32 LineLength = RemainingText.Find(TEXT("\n"), ESearchCase::CaseSensitive);
-		if(LineLength == -1)
+		if (LineLength == -1)
 		{
 			// There aren't any more newlines in the string, assume there's a newline at the end of the string.
 			LineLength = StringLength;
@@ -642,7 +642,7 @@ void UConsole::FlushPlayerInput()
 	else
 	{
 		UWorld* World = GetOuterUGameViewportClient()->GetWorld();
-		if (ULocalPlayer* LocalPlayer = GEngine->GetFirstGamePlayer(World))
+		if (ULocalPlayer * LocalPlayer = GEngine->GetFirstGamePlayer(World))
 		{
 			PC = LocalPlayer->PlayerController;
 		}
@@ -655,7 +655,7 @@ void UConsole::FlushPlayerInput()
 }
 
 bool UConsole::ProcessControlKey(FKey Key, EInputEvent Event)
-{	
+{
 #if PLATFORM_MAC
 	if (Key == EKeys::LeftCommand || Key == EKeys::RightCommand)
 #else
@@ -746,9 +746,9 @@ void UConsole::AppendInputText(const FString& Text)
 	UpdatePrecompletedInputLine();
 }
 
-bool UConsole::InputChar_Typing( int32 ControllerId, const FString& Unicode )
+bool UConsole::InputChar_Typing(int32 ControllerId, const FString& Unicode)
 {
-	if ( bCaptureKeyInput )
+	if (bCaptureKeyInput)
 	{
 		return true;
 	}
@@ -758,9 +758,9 @@ bool UConsole::InputChar_Typing( int32 ControllerId, const FString& Unicode )
 	return true;
 }
 
-bool UConsole::InputKey_InputLine( int32 ControllerId, FKey Key, EInputEvent Event, float AmountDepressed, bool bGamepad)
+bool UConsole::InputKey_InputLine(int32 ControllerId, FKey Key, EInputEvent Event, float AmountDepressed, bool bGamepad)
 {
-	if ( Event == IE_Pressed )
+	if (Event == IE_Pressed)
 	{
 		bCaptureKeyInput = false;
 	}
@@ -769,7 +769,7 @@ bool UConsole::InputKey_InputLine( int32 ControllerId, FKey Key, EInputEvent Eve
 	bool bModifierDown = bCtrl;
 	FModifierKeysState KeyState = FSlateApplication::Get().GetModifierKeys();
 	bModifierDown |= KeyState.IsAltDown() || KeyState.IsCommandDown() || KeyState.IsShiftDown() || KeyState.IsControlDown();
-	if ( GetDefault<UInputSettings>()->ConsoleKeys.Contains(Key) && Event == IE_Pressed && !bModifierDown )
+	if (GetDefault<UInputSettings>()->ConsoleKeys.Contains(Key) && Event == IE_Pressed && !bModifierDown)
 	{
 		if (ConsoleState == NAME_Typing)
 		{
@@ -848,7 +848,7 @@ bool UConsole::InputKey_InputLine( int32 ControllerId, FKey Key, EInputEvent Eve
 	};
 
 	// if user input is open
-	if(ConsoleState != NAME_None)
+	if (ConsoleState != NAME_None)
 	{
 		if (ProcessControlKey(Key, Event))
 		{
@@ -858,13 +858,13 @@ bool UConsole::InputKey_InputLine( int32 ControllerId, FKey Key, EInputEvent Eve
 		{
 			return true;
 		}
-		else if( bGamepad )
+		else if (bGamepad)
 		{
 			return false;
 		}
-		else if( Key == EKeys::Escape && Event == IE_Released )
+		else if (Key == EKeys::Escape && Event == IE_Released)
 		{
-			if( !TypedStr.IsEmpty() )
+			if (!TypedStr.IsEmpty())
 			{
 				SetInputText("");
 				SetCursorPos(0);
@@ -880,17 +880,17 @@ bool UConsole::InputKey_InputLine( int32 ControllerId, FKey Key, EInputEvent Eve
 			}
 			else
 			{
-				FakeGotoState( NAME_None );
+				FakeGotoState(NAME_None);
 			}
 
 			return true;
 		}
-		else if( Key == EKeys::Enter && Event == IE_Released )
+		else if (Key == EKeys::Enter && Event == IE_Released)
 		{
-			if( !TypedStr.IsEmpty() )
+			if (!TypedStr.IsEmpty())
 			{
 				// Make a local copy of the string.
-				FString Temp=TypedStr;
+				FString Temp = TypedStr;
 
 				SetInputText(TEXT(""));
 				SetCursorPos(0);
@@ -899,9 +899,9 @@ bool UConsole::InputKey_InputLine( int32 ControllerId, FKey Key, EInputEvent Eve
 
 				//OutputText( Localize("Errors","Exec","Core") );
 
-				OutputText( TEXT("") );
+				OutputText(TEXT(""));
 
-				if(ConsoleState == NAME_Typing)
+				if (ConsoleState == NAME_Typing)
 				{
 					// close after each command when in typing mode (single line)
 					FakeGotoState(NAME_None);
@@ -920,21 +920,21 @@ bool UConsole::InputKey_InputLine( int32 ControllerId, FKey Key, EInputEvent Eve
 
 			return true;
 		}
-		else if( Event != IE_Pressed && Event != IE_Repeat )
+		else if (Event != IE_Pressed && Event != IE_Repeat)
 		{
-			if( !bGamepad )
+			if (!bGamepad)
 			{
 				return	Key != EKeys::LeftMouseButton
-					&&	Key != EKeys::MiddleMouseButton
-					&&	Key != EKeys::RightMouseButton;
+					&& Key != EKeys::MiddleMouseButton
+					&& Key != EKeys::RightMouseButton;
 			}
 			return false;
 		}
-		else if( Key == EKeys::Up || (Key == EKeys::Tab && bShift))
+		else if (Key == EKeys::Up || (Key == EKeys::Tab && bShift))
 		{
 			if (!bCtrl)
 			{
-				if(AutoComplete.Num())
+				if (AutoComplete.Num())
 				{
 
 					if (Key == EKeys::Tab)
@@ -958,7 +958,7 @@ bool UConsole::InputKey_InputLine( int32 ControllerId, FKey Key, EInputEvent Eve
 				SetInputLineFromAutoComplete();
 			}
 		}
-		else if( Key == EKeys::Down || (Key == EKeys::Tab && !bShift) )
+		else if (Key == EKeys::Down || (Key == EKeys::Tab && !bShift))
 		{
 			if (!bCtrl)
 			{
@@ -995,33 +995,33 @@ bool UConsole::InputKey_InputLine( int32 ControllerId, FKey Key, EInputEvent Eve
 			}
 			return true;
 		}
-		else if( Key==EKeys::BackSpace )
+		else if (Key == EKeys::BackSpace)
 		{
-			if( TypedStrPos>0 )
+			if (TypedStrPos > 0)
 			{
-				SetInputText(FString::Printf(TEXT("%s%s"), *TypedStr.Left(TypedStrPos-1), *TypedStr.Right( TypedStr.Len() - TypedStrPos)));
-				SetCursorPos(TypedStrPos-1);
+				SetInputText(FString::Printf(TEXT("%s%s"), *TypedStr.Left(TypedStrPos - 1), *TypedStr.Right(TypedStr.Len() - TypedStrPos)));
+				SetCursorPos(TypedStrPos - 1);
 				// unlock auto-complete (@todo - track the lock position so we don't bother unlocking under bogus cases)
 				bAutoCompleteLocked = false;
 			}
 
 			return true;
 		}
-		else if ( Key== EKeys::Delete )
+		else if (Key == EKeys::Delete)
 		{
-			if ( TypedStrPos < TypedStr.Len() )
+			if (TypedStrPos < TypedStr.Len())
 			{
-				SetInputText(FString::Printf(TEXT("%s%s"), *TypedStr.Left(TypedStrPos), *TypedStr.Right( TypedStr.Len() - TypedStrPos - 1)));
+				SetInputText(FString::Printf(TEXT("%s%s"), *TypedStr.Left(TypedStrPos), *TypedStr.Right(TypedStr.Len() - TypedStrPos - 1)));
 			}
 			return true;
 		}
-		else if ( Key==EKeys::Left )
+		else if (Key == EKeys::Left)
 		{
 			if (bCtrl)
 			{
 				// find the nearest '.' or ' '
-				int32 NewPos = FMath::Max(TypedStr.Find(TEXT("."), ESearchCase::CaseSensitive, ESearchDir::FromEnd, TypedStrPos),TypedStr.Find(TEXT(" "),ESearchCase::CaseSensitive, ESearchDir::FromEnd,TypedStrPos));
-				SetCursorPos(FMath::Max(0,NewPos));
+				int32 NewPos = FMath::Max(TypedStr.Find(TEXT("."), ESearchCase::CaseSensitive, ESearchDir::FromEnd, TypedStrPos), TypedStr.Find(TEXT(" "), ESearchCase::CaseSensitive, ESearchDir::FromEnd, TypedStrPos));
+				SetCursorPos(FMath::Max(0, NewPos));
 			}
 			else
 			{
@@ -1029,7 +1029,7 @@ bool UConsole::InputKey_InputLine( int32 ControllerId, FKey Key, EInputEvent Eve
 			}
 			return true;
 		}
-		else if ( Key==EKeys::Right )
+		else if (Key == EKeys::Right)
 		{
 			if (bCtrl)
 			{
@@ -1037,13 +1037,13 @@ bool UConsole::InputKey_InputLine( int32 ControllerId, FKey Key, EInputEvent Eve
 				int32 SpacePos = TypedStr.Find(TEXT(" "));
 				int32 PeriodPos = TypedStr.Find(TEXT("."));
 				// pick the closest valid index
-				int32 NewPos = SpacePos < 0 ? PeriodPos : (PeriodPos < 0 ? SpacePos : FMath::Min(SpacePos,PeriodPos));
+				int32 NewPos = SpacePos < 0 ? PeriodPos : (PeriodPos < 0 ? SpacePos : FMath::Min(SpacePos, PeriodPos));
 				// jump to end if nothing in between
 				if (NewPos == INDEX_NONE)
 				{
 					NewPos = TypedStr.Len();
 				}
-				SetCursorPos(FMath::Min(TypedStr.Len(),FMath::Max(TypedStrPos,NewPos)));
+				SetCursorPos(FMath::Min(TypedStr.Len(), FMath::Max(TypedStrPos, NewPos)));
 			}
 			else
 			{
@@ -1051,12 +1051,12 @@ bool UConsole::InputKey_InputLine( int32 ControllerId, FKey Key, EInputEvent Eve
 			}
 			return true;
 		}
-		else if ( Key==EKeys::Home )
+		else if (Key == EKeys::Home)
 		{
 			SetCursorPos(0);
 			return true;
 		}
-		else if ( Key==EKeys::End )
+		else if (Key == EKeys::End)
 		{
 			SetCursorPos(TypedStr.Len());
 			return true;
@@ -1129,22 +1129,22 @@ void UConsole::PostRender_Console_Typing(UCanvas* Canvas)
 
 void UConsole::BeginState_Typing(FName PreviousStateName)
 {
-	if ( PreviousStateName == NAME_None )
+	if (PreviousStateName == NAME_None)
 	{
 		FlushPlayerInput();
 	}
 	bCaptureKeyInput = true;
 }
 
-void UConsole::EndState_Typing( FName NextStateName )
+void UConsole::EndState_Typing(FName NextStateName)
 {
 	bAutoCompleteLocked = false;
 }
 
 
-bool UConsole::InputChar_Open( int32 ControllerId, const FString& Unicode )
+bool UConsole::InputChar_Open(int32 ControllerId, const FString& Unicode)
 {
-	if ( bCaptureKeyInput )
+	if (bCaptureKeyInput)
 	{
 		return true;
 	}
@@ -1155,33 +1155,33 @@ bool UConsole::InputChar_Open( int32 ControllerId, const FString& Unicode )
 }
 
 
-bool UConsole::InputKey_Open( int32 ControllerId, FKey Key, EInputEvent Event, float AmountDepressed, bool bGamepad)
+bool UConsole::InputKey_Open(int32 ControllerId, FKey Key, EInputEvent Event, float AmountDepressed, bool bGamepad)
 {
-	if ( Key==EKeys::PageUp || Key == EKeys::MouseScrollUp )
+	if (Key == EKeys::PageUp || Key == EKeys::MouseScrollUp)
 	{
-		if (SBPos<Scrollback.Num()-1)
+		if (SBPos < Scrollback.Num() - 1)
 		{
 			if (bCtrl)
-				SBPos+=5;
+				SBPos += 5;
 			else
 				SBPos++;
 
-			if (SBPos>=Scrollback.Num())
-				SBPos = Scrollback.Num()-1;
+			if (SBPos >= Scrollback.Num())
+				SBPos = Scrollback.Num() - 1;
 		}
 
 		return true;
 	}
-	else if ( Key==EKeys::PageDown || Key==EKeys::MouseScrollDown )
+	else if (Key == EKeys::PageDown || Key == EKeys::MouseScrollDown)
 	{
-		if (SBPos>0)
+		if (SBPos > 0)
 		{
 			if (bCtrl)
-				SBPos-=5;
+				SBPos -= 5;
 			else
 				SBPos--;
 
-			if (SBPos<0)
+			if (SBPos < 0)
 				SBPos = 0;
 		}
 
@@ -1201,53 +1201,47 @@ void UConsole::PostRender_Console_Open(UCanvas* Canvas)
 	float TopPos = 0;
 	float LeftPos = 0;
 
-	if (GEngine->IsConsoleBuild())
-	{
-		ClipX	-= 80;
-		TopPos	 = 30;
-		LeftPos	 = 40;
-	}
-	if (GEngine->IsStereoscopic3D())
+	if (CVarCustomConsolePosEnabled.GetValueOnAnyThread())
 	{
-		LeftPos = ClipX / 3;
-		ClipX -= LeftPos;
-		Height = Canvas->ClipY * 0.60;
+		LeftPos = (float)CVarConsoleXPos.GetValueOnAnyThread();
+		float BottomOffset = (float)CVarConsoleYPos.GetValueOnAnyThread();
+		Height = Canvas->ClipY - BottomOffset;
 	}
 
 	UFont* Font = GEngine->GetSmallFont();
 
 	// determine the height of the text
 	float xl, yl;
-	Canvas->StrLen(Font, TEXT("M"),xl,yl);
+	Canvas->StrLen(Font, TEXT("M"), xl, yl);
 	// Background
 	FLinearColor BackgroundColor = ConsoleDefs::AutocompleteBackgroundColor.ReinterpretAsLinear();
 	BackgroundColor.A = ConsoleSettings->BackgroundOpacityPercentage / 100.0f;
-	FCanvasTileItem ConsoleTile( FVector2D( LeftPos, 0.0f ), DefaultTexture_Black->Resource, FVector2D( ClipX, Height + TopPos - yl), FVector2D( 0.0f, 0.0f), FVector2D( 1.0f, 1.0f ), BackgroundColor );
+	FCanvasTileItem ConsoleTile(FVector2D(LeftPos, 0.0f), DefaultTexture_Black->Resource, FVector2D(ClipX, Height + TopPos - yl), FVector2D(0.0f, 0.0f), FVector2D(1.0f, 1.0f), BackgroundColor);
 
 	// Preserve alpha to allow single-pass composite
 	ConsoleTile.BlendMode = SE_BLEND_AlphaBlend;
 
-	Canvas->DrawItem( ConsoleTile );
+	Canvas->DrawItem(ConsoleTile);
 
 	// figure out which element of the scrollback buffer to should appear first (at the top of the screen)
 	int32 idx = SBHead - SBPos;
-	
+
 	float y = Height - yl;
 
-	if(Scrollback.Num())
+	if (Scrollback.Num())
 	{
-		FCanvasTextItem ConsoleText( FVector2D( LeftPos,TopPos+Height-5-yl ), FText::FromString(TEXT("")), GEngine->GetLargeFont(), ConsoleSettings->InputColor );
+		FCanvasTextItem ConsoleText(FVector2D(LeftPos, TopPos + Height - 5 - yl), FText::FromString(TEXT("")), GEngine->GetLargeFont(), ConsoleSettings->InputColor);
 		// change the text color to white
-		ConsoleText.SetColor( FLinearColor::White );
+		ConsoleText.SetColor(FLinearColor::White);
 
 		// while we have enough room to draw another line and there are more lines to draw
-		while (y > -yl && idx>=0)
+		while (y > -yl && idx >= 0)
 		{
 			float PenX;
 			float PenY;
 			float PenZ = 0.1f;
 			PenX = LeftPos;
-			PenY = TopPos+y;
+			PenY = TopPos + y;
 
 			// adjust the location for any word wrapping due to long text lines
 			if (idx < Scrollback.Num())
@@ -1258,14 +1252,14 @@ void UConsole::PostRender_Console_Open(UCanvas* Canvas)
 				{
 					y -= (ScrollLineYL - yl);
 					PenX = LeftPos;
-					PenY = TopPos+y;				
+					PenY = TopPos + y;
 				}
 
-				ConsoleText.Text = FText::FromString( Scrollback[idx] );
-				Canvas->DrawItem( ConsoleText, PenX, PenY );
+				ConsoleText.Text = FText::FromString(Scrollback[idx]);
+				Canvas->DrawItem(ConsoleText, PenX, PenY);
 			}
 			idx--;
-			y-=yl;
+			y -= yl;
 		}
 	}
 
@@ -1275,13 +1269,13 @@ void UConsole::PostRender_Console_Open(UCanvas* Canvas)
 void UConsole::BeginState_Open(FName PreviousStateName)
 {
 	bCaptureKeyInput = true;
-//	HistoryCur = HistoryTop;
+	//	HistoryCur = HistoryTop;
 
 	SBPos = 0;
 	bCtrl = false;
 	bShift = false;
 
-	if ( PreviousStateName == NAME_None )
+	if (PreviousStateName == NAME_None)
 	{
 		FlushPlayerInput();
 	}
@@ -1292,24 +1286,24 @@ void UConsole::EndState_Open(FName NextStateName)
 {
 }
 
-bool UConsole::InputChar( int32 ControllerId, const FString& Unicode )
+bool UConsole::InputChar(int32 ControllerId, const FString& Unicode)
 {
 	if (ConsoleState == NAME_Typing)
 	{
-		return InputChar_Typing(ControllerId,Unicode);
+		return InputChar_Typing(ControllerId, Unicode);
 	}
 	if (ConsoleState == NAME_Open)
 	{
-		return InputChar_Open(ControllerId,Unicode);
+		return InputChar_Open(ControllerId, Unicode);
 	}
 	return bCaptureKeyInput;
 }
 
-bool UConsole::InputKey( int32 ControllerId, FKey Key, EInputEvent Event, float AmountDepressed, bool bGamepad)
+bool UConsole::InputKey(int32 ControllerId, FKey Key, EInputEvent Event, float AmountDepressed, bool bGamepad)
 {
-	bool bWasConsumed = InputKey_InputLine(ControllerId,Key,Event,AmountDepressed,bGamepad);
+	bool bWasConsumed = InputKey_InputLine(ControllerId, Key, Event, AmountDepressed, bGamepad);
 
-	if(!bWasConsumed)
+	if (!bWasConsumed)
 	{
 		if (ConsoleState == NAME_Typing)
 		{
@@ -1318,7 +1312,7 @@ bool UConsole::InputKey( int32 ControllerId, FKey Key, EInputEvent Event, float
 		}
 		if (ConsoleState == NAME_Open)
 		{
-			bWasConsumed = InputKey_Open(ControllerId,Key,Event,AmountDepressed,bGamepad);
+			bWasConsumed = InputKey_Open(ControllerId, Key, Event, AmountDepressed, bGamepad);
 			// if the console is open we don't want any other one to consume the input
 			return true;
 		}
@@ -1365,24 +1359,24 @@ void UConsole::PostRender_InputLine(UCanvas* Canvas, FIntPoint UserInputLinePos)
 	// Background
 	FLinearColor BackgroundColor = ConsoleDefs::AutocompleteBackgroundColor.ReinterpretAsLinear();
 	BackgroundColor.A = ConsoleSettings->BackgroundOpacityPercentage / 100.0f;
-	FCanvasTileItem ConsoleTile( FVector2D( UserInputLinePos.X,UserInputLinePos.Y-6-yl ), DefaultTexture_Black->Resource, FVector2D( ClipX, yl+6 ), FVector2D( 0.0f, 0.0f), FVector2D( 1.0f, 1.0f ), BackgroundColor );
-	
+	FCanvasTileItem ConsoleTile(FVector2D(UserInputLinePos.X, UserInputLinePos.Y - 6 - yl), DefaultTexture_Black->Resource, FVector2D(ClipX, yl + 6), FVector2D(0.0f, 0.0f), FVector2D(1.0f, 1.0f), BackgroundColor);
+
 	// Preserve alpha to allow single-pass composite
 	ConsoleTile.BlendMode = SE_BLEND_AlphaBlend;
 
-	Canvas->DrawItem( ConsoleTile );
+	Canvas->DrawItem(ConsoleTile);
 
 	// Separator line
-	ConsoleTile.SetColor( ConsoleDefs::BorderColor );
+	ConsoleTile.SetColor(ConsoleDefs::BorderColor);
 	ConsoleTile.Texture = DefaultTexture_White->Resource;
-	ConsoleTile.Size = FVector2D( ClipX, 2.0f );
-	Canvas->DrawItem( ConsoleTile );
+	ConsoleTile.Size = FVector2D(ClipX, 2.0f);
+	Canvas->DrawItem(ConsoleTile);
 
 	// Currently typed string
-	FText Str = FText::FromString( TypedInputText );
-	FCanvasTextItem ConsoleText( FVector2D( UserInputLinePos.X,UserInputLinePos.Y-3-yl ), Str , GEngine->GetLargeFont(), ConsoleSettings->InputColor );
+	FText Str = FText::FromString(TypedInputText);
+	FCanvasTextItem ConsoleText(FVector2D(UserInputLinePos.X, UserInputLinePos.Y - 3 - yl), Str, GEngine->GetLargeFont(), ConsoleSettings->InputColor);
 	ConsoleText.EnableShadow(FLinearColor::Black);
-	Canvas->DrawItem( ConsoleText );
+	Canvas->DrawItem(ConsoleText);
 
 	// Precompleted remainder of the typed string (faded out)
 	if (!PrecompletedInputText.IsEmpty())
@@ -1407,7 +1401,7 @@ void UConsole::PostRender_InputLine(UCanvas* Canvas, FIntPoint UserInputLinePos)
 		// Set the background color/texture of the auto-complete section
 		FLinearColor AutoCompleteBackgroundColor = ConsoleDefs::AutocompleteBackgroundColor;
 		AutoCompleteBackgroundColor.A = ConsoleSettings->BackgroundOpacityPercentage / 100.0f;
-		ConsoleTile.SetColor( AutoCompleteBackgroundColor );
+		ConsoleTile.SetColor(AutoCompleteBackgroundColor);
 		ConsoleTile.Texture = DefaultTexture_White->Resource;
 
 		// wasteful memory allocations but when typing in a console command this is fine
@@ -1419,7 +1413,7 @@ void UConsole::PostRender_InputLine(UCanvas* Canvas, FIntPoint UserInputLinePos)
 		float MaxRightWidth = 0;
 		for (int32 MatchIdx = 0; MatchIdx < MAX_AUTOCOMPLETION_LINES && MatchIdx < AutoComplete.Num(); MatchIdx++)
 		{
-			const FAutoCompleteCommand &Cmd = AutoComplete[StartIdx + MatchIdx];
+			const FAutoCompleteCommand& Cmd = AutoComplete[StartIdx + MatchIdx];
 			AutoCompleteElements.Add(&Cmd);
 
 			// Find the longest command and the longest description for left-justification of the descriptions
@@ -1450,13 +1444,13 @@ void UConsole::PostRender_InputLine(UCanvas* Canvas, FIntPoint UserInputLinePos)
 		int32 Border = 4;
 
 		// dark inner part
-		ConsoleTile.Size = FVector2D( MaxWidth + 2 * Border, Height + 2 * Border );
-		ConsoleTile.SetColor( AutoCompleteBackgroundColor );
+		ConsoleTile.Size = FVector2D(MaxWidth + 2 * Border, Height + 2 * Border);
+		ConsoleTile.SetColor(AutoCompleteBackgroundColor);
 
 		// Preserve alpha to allow single-pass composite
 		ConsoleTile.BlendMode = SE_BLEND_AlphaBlend;
 
-		Canvas->DrawItem( ConsoleTile, UserInputLinePos.X + xl - Border, y + yl - Height - Border );
+		Canvas->DrawItem(ConsoleTile, UserInputLinePos.X + xl - Border, y + yl - Height - Border);
 
 		// white border
 		FCanvasBoxItem ConsoleOutline(ConsoleTile.Position, FVector2D(MaxWidth + 2 * Border, Height + 2 * Border));
@@ -1471,11 +1465,11 @@ void UConsole::PostRender_InputLine(UCanvas* Canvas, FIntPoint UserInputLinePos)
 			const bool bMoreMatches = (Num > MAX_AUTOCOMPLETION_LINES && i == Num - 1);
 			const bool bHistory = AutoCompleteElement.IsHistory();
 			int CmdXOffset = 0;
-		
+
 			FColor LeftC = AutoCompleteElement.Color;
 			FColor RightC = ConsoleSettings->AutoCompleteFadedColor;
 
-			if(bCursorLineColor)
+			if (bCursorLineColor)
 			{
 				ConsoleTile.Size = FVector2D(MaxWidth, yl);
 				ConsoleTile.SetColor(ConsoleDefs::CursorLineColor);
@@ -1484,7 +1478,7 @@ void UConsole::PostRender_InputLine(UCanvas* Canvas, FIntPoint UserInputLinePos)
 				LeftC = ConsoleDefs::CursorColor;
 			}
 
-			if(bHistory)
+			if (bHistory)
 			{
 				// > HistoryElement has the strings swapped so we need to swap the colors
 				Swap(LeftC, RightC);
@@ -1501,7 +1495,7 @@ void UConsole::PostRender_InputLine(UCanvas* Canvas, FIntPoint UserInputLinePos)
 
 			ConsoleText.SetColor(LeftC);
 			ConsoleText.Text = FText::FromString(AutoCompleteElement.GetLeft());
-			Canvas->DrawItem( ConsoleText, UserInputLinePos.X + CmdXOffset + xl, y );
+			Canvas->DrawItem(ConsoleText, UserInputLinePos.X + CmdXOffset + xl, y);
 
 			float DescriptionWidth, DescriptionHeight;
 			Canvas->StrLen(Font, AutoCompleteElement.GetRight(), DescriptionWidth, DescriptionHeight);
@@ -1520,7 +1514,7 @@ void UConsole::PostRender_InputLine(UCanvas* Canvas, FIntPoint UserInputLinePos)
 
 			ConsoleText.SetColor(RightC);
 			ConsoleText.Text = FText::FromString(AutoCompleteElement.GetRight());
-			Canvas->DrawItem( ConsoleText, DescriptionX, y );
+			Canvas->DrawItem(ConsoleText, DescriptionX, y);
 			y -= yl;
 		};
 
@@ -1588,7 +1582,7 @@ void UConsole::FakeGotoState(FName NextStateName)
 		// Console has opened
 		OnConsoleActivationStateChanged.Broadcast(true);
 	}
-	else if( NextStateName == NAME_None )
+	else if (NextStateName == NAME_None)
 	{
 		// We need to force the console state name change now otherwise inside the call 
 		// to SetKeyboardFocus the console is still considered active
@@ -1621,7 +1615,7 @@ void UConsole::FakeGotoState(FName NextStateName)
 	ConsoleState = NextStateName;
 }
 
-void UConsole::Serialize( const TCHAR* V, ELogVerbosity::Type Verbosity, const class FName& Category )
+void UConsole::Serialize(const TCHAR* V, ELogVerbosity::Type Verbosity, const class FName& Category)
 {
 	// e.g. UE_LOG(LogConsoleResponse, Display, TEXT("Test"));
 	static const FName ConsoleResponseLog = FName("LogConsoleResponse");
@@ -1635,15 +1629,15 @@ void UConsole::Serialize( const TCHAR* V, ELogVerbosity::Type Verbosity, const c
 	{
 		static const TConsoleVariableData<int32>* CVar = IConsoleManager::Get().FindTConsoleVariableDataInt(TEXT("con.MinLogVerbosity"));
 
-		if(CVar)
+		if (CVar)
 		{
 			int MinVerbosity = CVar->GetValueOnAnyThread();
 
-			if((int)Verbosity <= MinVerbosity)
+			if ((int)Verbosity <= MinVerbosity)
 			{
 				// log all that is >= the specified verbosity
 				OutputText(V);
 			}
 		}
 	}
-}
+}
\ No newline at end of file
