PR #6116: Fix Dead Zone calculation (Contributed by [REDACTED_PERSON])

diff --git a/Engine/Source/Runtime/Engine/Private/UserInterface/PlayerInput.cpp b/Engine/Source/Runtime/Engine/Private/UserInterface/PlayerInput.cpp
--- a/Engine/Source/Runtime/Engine/Private/UserInterface/PlayerInput.cpp
+++ b/Engine/Source/Runtime/Engine/Private/UserInterface/PlayerInput.cpp
@@ -1722,8 +1722,8 @@ FVector UPlayerInput::MassageVectorAxisInput(FKey Key, FVector RawValue)
 			};
 
 			NewVal.X = DeadZoneLambda(NewVal.X);
-			NewVal.X = DeadZoneLambda(NewVal.Y);
-			NewVal.X = DeadZoneLambda(NewVal.Z);
+			NewVal.Y = DeadZoneLambda(NewVal.Y);
+			NewVal.Z = DeadZoneLambda(NewVal.Z);
 		}
 
 		// apply any exponent curvature while we're in the [0..1] range
