Check for empty CurrentSessionSectionName string

When CurrentSessionSectionName isn't set due to FEngineSessionManager::InitializeRecords()
not setting the value, FPlatformMisc::SetStoredValue() will crash due to empty
section name string. Add quick check for this case.

diff --git a/Engine/Source/Runtime/Engine/Private/EngineSessionManager.cpp b/Engine/Source/Runtime/Engine/Private/EngineSessionManager.cpp
--- a/Engine/Source/Runtime/Engine/Private/EngineSessionManager.cpp
+++ b/Engine/Source/Runtime/Engine/Private/EngineSessionManager.cpp
@@ -327,7 +327,10 @@ void FEngineSessionManager::Shutdown()
 	{
 		FSlateApplication::Get().GetOnModalLoopTickEvent().RemoveAll(this);
 
-		FPlatformMisc::SetStoredValue(SessionManagerDefs::StoreId, CurrentSessionSectionName, SessionManagerDefs::WasShutdownStoreKey, SessionManagerDefs::TrueValueString);
+		if (!CurrentSessionSectionName.IsEmpty())
+		{
+			FPlatformMisc::SetStoredValue(SessionManagerDefs::StoreId, CurrentSessionSectionName, SessionManagerDefs::WasShutdownStoreKey, SessionManagerDefs::TrueValueString);
+		}
 	}
 
 	// Clear the session record for this session
