Removing serialization from SoundClass from deprecated audio project setting

diff --git a/Engine/Source/Runtime/Core/Public/UObject/ReleaseObjectVersion.h b/Engine/Source/Runtime/Core/Public/UObject/ReleaseObjectVersion.h
--- a/Engine/Source/Runtime/Core/Public/UObject/ReleaseObjectVersion.h
+++ b/Engine/Source/Runtime/Core/Public/UObject/ReleaseObjectVersion.h
@@ -80,9 +80,6 @@ struct CORE_API FReleaseObjectVersion
 		// DepthOfFieldFocalDistance at 0 now disables DOF instead of DepthOfFieldFstop at 0.
 		FocalDistanceDisablesDOF,
 
-		// SoundClass 2D Reverb Send
-		SoundClass2DReverbSend,
-
 		// -----<new versions can be added above this line>-------------------------------------------------
 		VersionPlusOne,
 		LatestVersion = VersionPlusOne - 1
diff --git a/Engine/Source/Runtime/Engine/Private/SoundClass.cpp b/Engine/Source/Runtime/Engine/Private/SoundClass.cpp
--- a/Engine/Source/Runtime/Engine/Private/SoundClass.cpp
+++ b/Engine/Source/Runtime/Engine/Private/SoundClass.cpp
@@ -222,15 +222,6 @@ void USoundClass::Serialize( FArchive& Ar )
 		TMap<USoundClass*, FSoundClassEditorData>	EditorData_DEPRECATED;
 		Ar << EditorData_DEPRECATED;
 	}
-
-	Ar.UsingCustomVersion(FReleaseObjectVersion::GUID);
-	const int32 ReleaseObjectVersion = Ar.CustomVer(FReleaseObjectVersion::GUID);
-
-	if (ReleaseObjectVersion < FReleaseObjectVersion::SoundClass2DReverbSend)
-	{
-		const UAudioSettings* AudioSettings = GetDefault<UAudioSettings>();
-		Properties.Default2DReverbSendAmount = AudioSettings->DefaultReverbSendLevel_DEPRECATED;
-	}
 }
 
 void USoundClass::BeginDestroy()
