Fix for LPCM non-seamless looping

diff --git a/Engine/Source/Runtime/Engine/Private/ADPCMAudioInfo.cpp b/Engine/Source/Runtime/Engine/Private/ADPCMAudioInfo.cpp
--- a/Engine/Source/Runtime/Engine/Private/ADPCMAudioInfo.cpp
+++ b/Engine/Source/Runtime/Engine/Private/ADPCMAudioInfo.cpp
@@ -26,6 +26,7 @@ FAutoConsoleVariableRef CVarDisableADPCMSeeking(
 #define WAVE_FORMAT_LPCM  1
 #define WAVE_FORMAT_ADPCM 2
 
+
 namespace ADPCM
 {
 	void DecodeBlock(const uint8* EncodedADPCMBlock, int32 BlockSize, int16* DecodedPCMData);
@@ -367,28 +368,33 @@ bool FADPCMAudioInfo::ReadCompressedData(uint8* Destination, bool bLooping, uint
 	}
 	else
 	{
-		uint32 DecompressedSamplesToCopy = BufferSize / ChannelSampleSize;
-
-		// Ensure we don't go over the number of samples left in the audio data
-		if(DecompressedSamplesToCopy > TotalSamplesPerChannel - TotalSamplesStreamed)
+		uint32 OutDataOffset = 0;
+		while (BufferSize > 0)
 		{
-			DecompressedSamplesToCopy = TotalSamplesPerChannel - TotalSamplesStreamed;
-		}
+			uint32 DecompressedSamplesToCopy = BufferSize / ChannelSampleSize;
 
-		FMemory::Memcpy(OutData, WaveInfo.SampleDataStart + (TotalSamplesStreamed * ChannelSampleSize), DecompressedSamplesToCopy * ChannelSampleSize);
-		TotalSamplesStreamed += DecompressedSamplesToCopy;
-		BufferSize -= DecompressedSamplesToCopy * ChannelSampleSize;
+			// Ensure we don't go over the number of samples left in the audio data
+			if (DecompressedSamplesToCopy > TotalSamplesPerChannel - TotalSamplesStreamed)
+			{
+				DecompressedSamplesToCopy = TotalSamplesPerChannel - TotalSamplesStreamed;
+			}
 
-		// Check for the end of the audio samples and loop if needed
-		if(TotalSamplesStreamed >= TotalSamplesPerChannel)
-		{
-			ReachedEndOfSamples = true;
-			TotalSamplesStreamed = 0;
-			if(!bLooping)
+			FMemory::Memcpy(OutData + OutDataOffset, WaveInfo.SampleDataStart + (TotalSamplesStreamed * ChannelSampleSize), DecompressedSamplesToCopy * ChannelSampleSize);
+			TotalSamplesStreamed += DecompressedSamplesToCopy;
+			BufferSize -= DecompressedSamplesToCopy * ChannelSampleSize;
+			OutDataOffset += DecompressedSamplesToCopy;
+
+			// Check for the end of the audio samples and loop if needed
+			if (TotalSamplesStreamed >= TotalSamplesPerChannel)
 			{
-				// Zero remaining buffer
-				FMemory::Memzero(OutData, BufferSize);
-				return true;
+				ReachedEndOfSamples = true;
+				TotalSamplesStreamed = 0;
+				if (!bLooping)
+				{
+					// Zero remaining buffer
+					FMemory::Memzero(OutData, BufferSize);
+					return true;
+				}
 			}
 		}
 	}
