Added stat scope to TimerManager to easier corellate CSV data to CPU captures.

diff --git a/Engine/Source/Runtime/Engine/Private/TimerManager.cpp b/Engine/Source/Runtime/Engine/Private/TimerManager.cpp
--- a/Engine/Source/Runtime/Engine/Private/TimerManager.cpp
+++ b/Engine/Source/Runtime/Engine/Private/TimerManager.cpp
@@ -753,6 +753,7 @@ DECLARE_DWORD_COUNTER_STAT(TEXT("TimerManager Heap Size"),STAT_NumHeapEntries,ST
 
 void FTimerManager::Tick(float DeltaTime)
 {
+	SCOPED_NAMED_EVENT(FTimerManager_Tick, FColor::Orange);
 	CSV_SCOPED_TIMING_STAT_EXCLUSIVE(TimerManager);
 
 #if DO_TIMEGUARD && 0
