[NO NON-PII COMMIT MESSAGE]

diff --git a/Engine/Source/Runtime/Engine/Private/ADPCMAudioInfo.cpp b/Engine/Source/Runtime/Engine/Private/ADPCMAudioInfo.cpp
--- a/Engine/Source/Runtime/Engine/Private/ADPCMAudioInfo.cpp
+++ b/Engine/Source/Runtime/Engine/Private/ADPCMAudioInfo.cpp
@@ -612,7 +612,7 @@ bool FADPCMAudioInfo::StreamCompressedData(uint8* Destination, bool bLooping, ui
 					// Set the current buffer offset accounting for the header in the first chunk
 					if (!bSeekPending)
 					{
-						CurrentChunkBufferOffset = CurrentChunkIndex == FirstChunkSampleDataOffset ? FirstChunkSampleDataOffset : 0;
+						CurrentChunkBufferOffset = CurrentChunkIndex == FirstChunkSampleDataIndex ? FirstChunkSampleDataOffset : 0;
 					}
 
 					bSeekPending = false;
