Correctly invalidate timer handle in ClearTimer if the timer has already expired (can't be found) and invalidate the timer handle in SetTimer if the rate is <= 0

diff --git a/Engine/Source/Runtime/Engine/Private/TimerManager.cpp b/Engine/Source/Runtime/Engine/Private/TimerManager.cpp
--- a/Engine/Source/Runtime/Engine/Private/TimerManager.cpp
+++ b/Engine/Source/Runtime/Engine/Private/TimerManager.cpp
@@ -502,6 +502,10 @@ void FTimerManager::InternalSetTimer(FTimerHandle& InOutHandle, FTimerUnifiedDel
 
 		InOutHandle = NewTimerHandle;
 	}
+	else
+	{
+		InOutHandle.Invalidate();
+	}
 }
 
 FTimerHandle FTimerManager::InternalSetTimerForNextTick(FTimerUnifiedDelegate&& InDelegate)
