Fix for normalized stereo feature

diff --git a/Engine/Source/Runtime/Engine/Private/SoundWave.cpp b/Engine/Source/Runtime/Engine/Private/SoundWave.cpp
--- a/Engine/Source/Runtime/Engine/Private/SoundWave.cpp
+++ b/Engine/Source/Runtime/Engine/Private/SoundWave.cpp
@@ -1682,8 +1682,8 @@ void USoundWave::Parse(FAudioDevice* AudioDevice, const UPTRINT NodeWaveInstance
 	// If we're normalizing 3d stereo spatialized sounds, we need to scale by -6 dB
 	if (WaveInstance->GetUseSpatialization() && ParseParams.bApplyNormalizationToStereoSounds && NumChannels == 2)
 	{
-		float WaveInstanceVolume = WaveInstance->GetVolume();
-		WaveInstance->SetVolume(WaveInstanceVolume * 0.5f);
+		const float ThisVolumeMultiplier = WaveInstance->GetVolumeMultiplier();
+		WaveInstance->SetVolumeMultiplier(ThisVolumeMultiplier * 0.5f);
 	}
 
 	// Copy reverb send settings
