Fix for LPCM non-seamless looping for stereo case as well.

diff --git a/Engine/Source/Runtime/Engine/Private/ADPCMAudioInfo.cpp b/Engine/Source/Runtime/Engine/Private/ADPCMAudioInfo.cpp
--- a/Engine/Source/Runtime/Engine/Private/ADPCMAudioInfo.cpp
+++ b/Engine/Source/Runtime/Engine/Private/ADPCMAudioInfo.cpp
@@ -382,7 +382,7 @@ bool FADPCMAudioInfo::ReadCompressedData(uint8* Destination, bool bLooping, uint
 			FMemory::Memcpy(OutData + OutDataOffset, WaveInfo.SampleDataStart + (TotalSamplesStreamed * ChannelSampleSize), DecompressedSamplesToCopy * ChannelSampleSize);
 			TotalSamplesStreamed += DecompressedSamplesToCopy;
 			BufferSize -= DecompressedSamplesToCopy * ChannelSampleSize;
-			OutDataOffset += DecompressedSamplesToCopy;
+			OutDataOffset += DecompressedSamplesToCopy * NumChannels;
 
 			// Check for the end of the audio samples and loop if needed
 			if (TotalSamplesStreamed >= TotalSamplesPerChannel)
