Ensure volume and duration values passed are positive to new AdjustVolume implementation & fix new stat for StopDelayed

diff --git a/Engine/Source/Runtime/Engine/Private/Components/AudioComponent.cpp b/Engine/Source/Runtime/Engine/Private/Components/AudioComponent.cpp
--- a/Engine/Source/Runtime/Engine/Private/Components/AudioComponent.cpp
+++ b/Engine/Source/Runtime/Engine/Private/Components/AudioComponent.cpp
@@ -575,6 +575,8 @@ void UAudioComponent::AdjustVolumeInternal(float AdjustVolumeDuration, float Adj
 		return;
 	}
 
+	AdjustVolumeDuration = FMath::Max(0.0f, AdjustVolumeDuration);
+	AdjustVolumeLevel = FMath::Max(0.0f, AdjustVolumeLevel);
 	if (FMath::IsNearlyZero(AdjustVolumeDuration) && FMath::IsNearlyZero(AdjustVolumeLevel))
 	{
 		Stop();
@@ -676,7 +678,7 @@ void UAudioComponent::StopDelayed(float DelayTime)
 
 	// 2. Performs delayed stop with no fade
 	const uint64 InAudioComponentID = AudioComponentID;
-	DECLARE_CYCLE_STAT(TEXT("FAudioThreadTask.StopDelayed"), STAT_AudioAdjustVolume, STATGROUP_AudioThreadCommands);
+	DECLARE_CYCLE_STAT(TEXT("FAudioThreadTask.StopDelayed"), STAT_AudioStopDelayed, STATGROUP_AudioThreadCommands);
 	FAudioThread::RunCommandOnAudioThread([AudioDevice, InAudioComponentID, DelayTime]()
 	{
 		FActiveSound* ActiveSound = AudioDevice->FindActiveSound(InAudioComponentID);
@@ -708,7 +710,7 @@ void UAudioComponent::StopDelayed(float DelayTime)
 		}
 
 		ActiveSound->TargetAdjustVolumeStopTime = NewTargetStopTime;
-	}, GET_STATID(STAT_AudioAdjustVolume));
+	}, GET_STATID(STAT_AudioStopDelayed));
 }
 
 void UAudioComponent::SetPaused(bool bPause)
