Fix curve codec being reset

diff --git a/Engine/Source/Runtime/Engine/Classes/Animation/AnimCurveCompressionSettings.h b/Engine/Source/Runtime/Engine/Classes/Animation/AnimCurveCompressionSettings.h
--- a/Engine/Source/Runtime/Engine/Classes/Animation/AnimCurveCompressionSettings.h
+++ b/Engine/Source/Runtime/Engine/Classes/Animation/AnimCurveCompressionSettings.h
@@ -19,7 +19,7 @@ class ENGINE_API UAnimCurveCompressionSettings : public UObject
 	GENERATED_UCLASS_BODY()
 
 	/** An animation curve compression codec. */
-	UPROPERTY(Category = Compression, Instanced, EditAnywhere, NoClear)
+	UPROPERTY(Category = Compression, Export, EditAnywhere, NoClear, meta = (EditInline))
 	UAnimCurveCompressionCodec* Codec;
 
 	//////////////////////////////////////////////////////////////////////////
