Added MarkRenderStateDirtyEvent to UActorComponent
Added UpdateColorAndBrightnessEvent to ULightComponent.

PR #6120: Component changes for NVIDIA USD plugin.

diff --git a/Engine/Source/Runtime/Engine/Classes/Components/ActorComponent.h b/Engine/Source/Runtime/Engine/Classes/Components/ActorComponent.h
--- a/Engine/Source/Runtime/Engine/Classes/Components/ActorComponent.h
+++ b/Engine/Source/Runtime/Engine/Classes/Components/ActorComponent.h
@@ -893,6 +893,11 @@ public:
 	/** Prefix used to identify template component instances */
 	static const FString ComponentTemplateNameSuffix;
 
+	DECLARE_MULTICAST_DELEGATE_OneParam(FOnMarkRenderStateDirty, UActorComponent&);
+
+	/** Called When render state is marked dirty */
+	static FOnMarkRenderStateDirty MarkRenderStateDirtyEvent;
+
 protected:
 	/** Makes sure navigation system has up to date information regarding component's navigation relevancy 
 	 *	and if it can affect navigation at all 
diff --git a/Engine/Source/Runtime/Engine/Classes/Components/LightComponent.h b/Engine/Source/Runtime/Engine/Classes/Components/LightComponent.h
--- a/Engine/Source/Runtime/Engine/Classes/Components/LightComponent.h
+++ b/Engine/Source/Runtime/Engine/Classes/Components/LightComponent.h
@@ -439,6 +439,10 @@ public:
 	 */
 	static void ReassignStationaryLightChannels(UWorld* TargetWorld, bool bAssignForLightingBuild, ULevel* LightingScenario);
 
+	DECLARE_MULTICAST_DELEGATE_OneParam(FOnUpdateColorAndBrightness, ULightComponent&);
+
+	/** Called When light color or brightness needs update */
+	static FOnUpdateColorAndBrightness UpdateColorAndBrightnessEvent;
 };
 
 
diff --git a/Engine/Source/Runtime/Engine/Private/Components/ActorComponent.cpp b/Engine/Source/Runtime/Engine/Private/Components/ActorComponent.cpp
--- a/Engine/Source/Runtime/Engine/Private/Components/ActorComponent.cpp
+++ b/Engine/Source/Runtime/Engine/Private/Components/ActorComponent.cpp
@@ -145,6 +145,8 @@ FGlobalComponentRecreateRenderStateContext::~FGlobalComponentRecreateRenderState
 FActorComponentGlobalCreatePhysicsSignature UActorComponent::GlobalCreatePhysicsDelegate;
 // Destroy Physics global delegate
 FActorComponentGlobalDestroyPhysicsSignature UActorComponent::GlobalDestroyPhysicsDelegate;
+// Render state dirty global delegate
+UActorComponent::FOnMarkRenderStateDirty UActorComponent::MarkRenderStateDirtyEvent;
 
 const FString UActorComponent::ComponentTemplateNameSuffix(TEXT("_GEN_VARIABLE"));
 
@@ -1459,6 +1461,8 @@ void UActorComponent::MarkRenderStateDirty()
 		// Flag as dirty
 		bRenderStateDirty = true;
 		MarkForNeededEndOfFrameRecreate();
+
+		MarkRenderStateDirtyEvent.Broadcast(*this);
 	}
 }
 
diff --git a/Engine/Source/Runtime/Engine/Private/Components/LightComponent.cpp b/Engine/Source/Runtime/Engine/Private/Components/LightComponent.cpp
--- a/Engine/Source/Runtime/Engine/Private/Components/LightComponent.cpp
+++ b/Engine/Source/Runtime/Engine/Private/Components/LightComponent.cpp
@@ -377,6 +377,8 @@ ULightComponentBase::ULightComponentBase(const FObjectInitializer& ObjectInitial
 #endif
 }
 
+ULightComponent::FOnUpdateColorAndBrightness ULightComponent::UpdateColorAndBrightnessEvent;
+
 /**
  * Updates/ resets light GUIDs.
  */
@@ -822,6 +824,8 @@ void ULightComponent::SetIndirectLightingIntensity(float NewIntensity)
 			//@todo - remove from scene if brightness or color becomes 0
 			World->Scene->UpdateLightColorAndBrightness( this );
 		}
+
+		UpdateColorAndBrightnessEvent.Broadcast(*this);
 	}
 }
 
@@ -840,6 +844,8 @@ void ULightComponent::SetVolumetricScatteringIntensity(float NewIntensity)
 			//@todo - remove from scene if brightness or color becomes 0
 			World->Scene->UpdateLightColorAndBrightness( this );
 		}
+
+		UpdateColorAndBrightnessEvent.Broadcast(*this);
 	}
 }
 
@@ -861,6 +867,8 @@ void ULightComponent::SetLightColor(FLinearColor NewLightColor, bool bSRGB)
 			//@todo - remove from scene if brightness or color becomes 0
 			World->Scene->UpdateLightColorAndBrightness( this );
 		}
+
+		UpdateColorAndBrightnessEvent.Broadcast(*this);
 	}
 }
 
@@ -880,6 +888,8 @@ void ULightComponent::SetTemperature(float NewTemperature)
 			//@todo - remove from scene if brightness or color becomes 0
 			World->Scene->UpdateLightColorAndBrightness( this );
 		}
+
+		UpdateColorAndBrightnessEvent.Broadcast(*this);
 	}
 }
 
@@ -1090,6 +1100,8 @@ void ULightComponent::UpdateColorAndBrightness()
 			World->Scene->UpdateLightColorAndBrightness(this);
 		}
 	}
+
+	UpdateColorAndBrightnessEvent.Broadcast(*this);
 }
 
 //
