Changed GetProxyOnGameThread<FAnimInstanceProxy>().GetInstanceAssetPlayers
to GetProxyOnAnyThread<FAnimInstanceProxy>().GetInstanceAssetPlayers

diff --git a/Engine/Source/Runtime/Engine/Private/Animation/AnimInstance.cpp b/Engine/Source/Runtime/Engine/Private/Animation/AnimInstance.cpp
--- a/Engine/Source/Runtime/Engine/Private/Animation/AnimInstance.cpp
+++ b/Engine/Source/Runtime/Engine/Private/Animation/AnimInstance.cpp
@@ -2940,7 +2940,7 @@ int32 UAnimInstance::GetInstanceAssetPlayerIndex(FName MachineName, FName StateN
 
 TArray<FAnimNode_AssetPlayerBase*> UAnimInstance::GetInstanceAssetPlayers(const FName& GraphName)
 {
-	return GetProxyOnGameThread<FAnimInstanceProxy>().GetInstanceAssetPlayers(GraphName);
+	return GetProxyOnAnyThread<FAnimInstanceProxy>().GetInstanceAssetPlayers(GraphName);
 }
 
 FAnimNode_AssetPlayerBase* UAnimInstance::GetRelevantAssetPlayerFromState(int32 MachineIndex, int32 StateIndex)
