Serialize the flag for disabling a runtime virtual texture.
Initially this was intended for debug use only (and not to persist) but it's useful for controlling performance during dev when we want to place an RVT in the world but not use it yet.

diff --git a/Engine/Source/Runtime/Engine/Classes/VT/RuntimeVirtualTexture.h b/Engine/Source/Runtime/Engine/Classes/VT/RuntimeVirtualTexture.h
--- a/Engine/Source/Runtime/Engine/Classes/VT/RuntimeVirtualTexture.h
+++ b/Engine/Source/Runtime/Engine/Classes/VT/RuntimeVirtualTexture.h
@@ -40,8 +40,8 @@ protected:
 	UPROPERTY(EditAnywhere, BlueprintReadOnly, Category = Size, meta = (UIMin = "0", UIMax = "5", DisplayName = "Number of low mips to remove from the virtual texture"))
 	int32 RemoveLowMips = 0;
 
-	/** Enable usage of the virtual texture. This option is intended only for debugging and visualization of the scene without virtual textures. It isn't serialized. */
-	UPROPERTY(EditAnywhere, BlueprintReadOnly, Transient, Category = Debug, meta = (DisplayName = "Enable virtual texture"))
+	/** Enable usage of the virtual texture. When disabled there is no rendering into the virtual texture, and sampling will return zero values. */
+	UPROPERTY(EditAnywhere, BlueprintReadOnly, Category = Layout, meta = (DisplayName = "Enable virtual texture"))
 	bool bEnable = true;
 
 public:
