Added a user preference in the new settings menu to enable / disable HRTF

diff --git a/Engine/Source/Runtime/Engine/Classes/GameFramework/GameUserSettings.h b/Engine/Source/Runtime/Engine/Classes/GameFramework/GameUserSettings.h
--- a/Engine/Source/Runtime/Engine/Classes/GameFramework/GameUserSettings.h
+++ b/Engine/Source/Runtime/Engine/Classes/GameFramework/GameUserSettings.h
@@ -138,6 +138,14 @@ public:
 	UFUNCTION(BlueprintPure, Category=Settings)
 	int32 GetAudioQualityLevel() const { return AudioQualityLevel; }
 
+	/** Enables or disables headphone mode (HRTF) **/
+	UFUNCTION(BlueprintCallable, Category=Settings)
+	void EnableHeadphoneMode(bool bEnabled);
+
+	/** Returns if we're using headphone mode (HRTF) **/
+	UFUNCTION(BlueprintPure, Category=Settings)
+	bool IsHeadphoneModeEnabled() const { return bUseHeadphoneMode; }
+
 	/** Sets the user's frame rate limit (0 will disable frame rate limiting) */
 	UFUNCTION(BlueprintCallable, Category=Settings)
 	void SetFrameRateLimit(float NewLimit);
@@ -391,6 +399,10 @@ protected:
 	UPROPERTY(config)
 	int32 AudioQualityLevel;
 
+	/** Whether to use headphone mode (HRTF) **/
+	UPROPERTY(config)
+	bool bUseHeadphoneMode;
+
 	UPROPERTY(config)
 	int32 LastConfirmedAudioQualityLevel;
 
diff --git a/Engine/Source/Runtime/Engine/Private/GameUserSettings.cpp b/Engine/Source/Runtime/Engine/Private/GameUserSettings.cpp
--- a/Engine/Source/Runtime/Engine/Private/GameUserSettings.cpp
+++ b/Engine/Source/Runtime/Engine/Private/GameUserSettings.cpp
@@ -707,6 +707,17 @@ void UGameUserSettings::SetAudioQualityLevel(int32 QualityLevel)
 	}
 }
 
+void UGameUserSettings::EnableHeadphoneMode(bool bEnabled)
+{
+	static IConsoleVariable* BinauralSpatializationDisabledCVar = IConsoleManager::Get().FindConsoleVariable(TEXT("au.DisableBinauralSpatialization"));
+
+	if (BinauralSpatializationDisabledCVar && (bEnabled != bUseHeadphoneMode))
+	{
+		BinauralSpatializationDisabledCVar->Set(!bEnabled, ECVF_SetByGameSetting);
+		bUseHeadphoneMode = bEnabled;
+	}
+}
+
 void UGameUserSettings::SetFrameRateLimit(float NewLimit)
 {
 	FrameRateLimit = NewLimit;
