Ray Tracing HISM/ISM now uses the right LOD

diff --git a/Engine/Source/Runtime/Engine/Private/InstancedStaticMesh.cpp b/Engine/Source/Runtime/Engine/Private/InstancedStaticMesh.cpp
--- a/Engine/Source/Runtime/Engine/Private/InstancedStaticMesh.cpp
+++ b/Engine/Source/Runtime/Engine/Private/InstancedStaticMesh.cpp
@@ -1052,8 +1052,7 @@ void FInstancedStaticMeshSceneProxy::GetDynamicRayTracingInstances(struct FRayTr
 		return;
 	}
 
-	//#dxr_todo: select the appropriate LOD depending on Context.View
-	uint32 LOD = 0;
+	uint32 LOD = GetCurrentFirstLODIdx_RenderThread();
 	const int32 InstanceCount = InstancedRenderData.Component->PerInstanceSMData.Num();
 
 	if (InstanceCount == 0)
