Fixed potential infinite recursion issue in Get/ContainsPhysicsTriMeshData().
Fixed bug in ContainsPhysicsTriMeshData().

diff --git a/Engine/Source/Runtime/Engine/Private/StaticMesh.cpp b/Engine/Source/Runtime/Engine/Private/StaticMesh.cpp
--- a/Engine/Source/Runtime/Engine/Private/StaticMesh.cpp
+++ b/Engine/Source/Runtime/Engine/Private/StaticMesh.cpp
@@ -5155,7 +5155,7 @@ static int32 GetCollisionVertIndexForMeshVertIndex(int32 MeshVertIndex, TMap<int
 bool UStaticMesh::GetPhysicsTriMeshData(struct FTriMeshCollisionData* CollisionData, bool bInUseAllTriData)
 {
 #if WITH_EDITORONLY_DATA
-	if (ComplexCollisionMesh)
+	if (ComplexCollisionMesh && ComplexCollisionMesh != this)
 	{
 		return ComplexCollisionMesh->GetPhysicsTriMeshData(CollisionData, bInUseAllTriData);
 	}
@@ -5211,9 +5211,9 @@ bool UStaticMesh::GetPhysicsTriMeshData(struct FTriMeshCollisionData* CollisionD
 bool UStaticMesh::ContainsPhysicsTriMeshData(bool bInUseAllTriData) const 
 {
 #if WITH_EDITORONLY_DATA
-	if (ComplexCollisionMesh)
+	if (ComplexCollisionMesh && ComplexCollisionMesh != this)
 	{
-		ComplexCollisionMesh->ContainsPhysicsTriMeshData(bInUseAllTriData);
+		return ComplexCollisionMesh->ContainsPhysicsTriMeshData(bInUseAllTriData);
 	}
 	if(RenderData == nullptr || RenderData->LODResources.Num() == 0)
 	{
