Fix GC hang in audio mixer

- Removing code in IsFinishedDestroy and moved mechanism to prevent GC on playing sound waves to add referenced objects

diff --git a/Engine/Source/Runtime/AudioMixer/Private/AudioMixerSourceBuffer.cpp b/Engine/Source/Runtime/AudioMixer/Private/AudioMixerSourceBuffer.cpp
--- a/Engine/Source/Runtime/AudioMixer/Private/AudioMixerSourceBuffer.cpp
+++ b/Engine/Source/Runtime/AudioMixer/Private/AudioMixerSourceBuffer.cpp
@@ -488,12 +488,6 @@ namespace Audio
 		}
 	}
 
-	bool FMixerSourceBuffer::IsBeginDestroy()
-	{
-		check(SoundWave);
-		return SoundWave->bIsBeginDestroy;
-	}
-
 	void FMixerSourceBuffer::ClearSoundWave()
 	{
 		// Call on end generate right now, before destructor
diff --git a/Engine/Source/Runtime/AudioMixer/Private/AudioMixerSourceBuffer.h b/Engine/Source/Runtime/AudioMixer/Private/AudioMixerSourceBuffer.h
--- a/Engine/Source/Runtime/AudioMixer/Private/AudioMixerSourceBuffer.h
+++ b/Engine/Source/Runtime/AudioMixer/Private/AudioMixerSourceBuffer.h
@@ -88,9 +88,6 @@ namespace Audio
 		// Ensures the async task finishes
 		void EnsureAsyncTaskFinishes();
 
-		// Checks if sound wave is flagged begin destroy
-		bool IsBeginDestroy();
-
 		// Clear the sound wave reference
 		void ClearSoundWave();
 
diff --git a/Engine/Source/Runtime/AudioMixer/Private/AudioMixerSourceManager.cpp b/Engine/Source/Runtime/AudioMixer/Private/AudioMixerSourceManager.cpp
--- a/Engine/Source/Runtime/AudioMixer/Private/AudioMixerSourceManager.cpp
+++ b/Engine/Source/Runtime/AudioMixer/Private/AudioMixerSourceManager.cpp
@@ -2571,37 +2571,6 @@ namespace Audio
 
 	void FMixerSourceManager::UpdatePendingReleaseData(bool bForceWait)
 	{
-		// Check for any pending delete procedural sound waves
-		for (int32 SourceId = 0; SourceId < NumTotalSources; ++SourceId)
-		{
-			FSourceInfo& SourceInfo = SourceInfos[SourceId];
-			if (SourceInfo.MixerSourceBuffer.IsValid())
-			{
-				// If we've been flagged to begin destroy
-				if (SourceInfo.MixerSourceBuffer->IsBeginDestroy())
-				{
-					SourceInfo.MixerSourceBuffer->ClearSoundWave();
-						
-					if (!SourceInfo.bIsDone)
-					{
-						SourceInfo.bIsDone = true;
-						SourceInfo.SourceListener->OnDone();
-					}
-
-					// Clear out the mixer source buffer
-					SourceInfo.MixerSourceBuffer.Reset();
-
-					// Set the sound to be done playing
-					// This will flag the sound to be released
-					SourceInfo.bIsPlaying = false;
-					SourceInfo.bIsPaused = false;
-					SourceInfo.bIsActive = false;
-					SourceInfo.bIsStopping = false;
-				}
-			}
-		}
-
-
 		// Don't block, but let tasks finish naturally
 		for (int32 i = PendingSourceBuffers.Num() - 1; i >= 0; --i)
 		{
diff --git a/Engine/Source/Runtime/Engine/Classes/Sound/SoundWave.h b/Engine/Source/Runtime/Engine/Classes/Sound/SoundWave.h
--- a/Engine/Source/Runtime/Engine/Classes/Sound/SoundWave.h
+++ b/Engine/Source/Runtime/Engine/Classes/Sound/SoundWave.h
@@ -327,9 +327,6 @@ class ENGINE_API USoundWave : public USoundBase
 	/** Set to true for programmatically-generated, streamed audio. */
 	uint8 bProcedural:1;
 
-	/** Whether this sound wave is beginning to be destroyed by GC. */
-	uint8 bIsBeginDestroy:1;
-
 	/** Set to true of this is a bus sound source. This will result in the sound wave not generating audio for itself, but generate audio through instances. Used only in audio mixer. */
 	uint8 bIsBus:1;
 
diff --git a/Engine/Source/Runtime/Engine/Private/AudioDevice.cpp b/Engine/Source/Runtime/Engine/Private/AudioDevice.cpp
--- a/Engine/Source/Runtime/Engine/Private/AudioDevice.cpp
+++ b/Engine/Source/Runtime/Engine/Private/AudioDevice.cpp
@@ -3697,6 +3697,8 @@ void FAudioDevice::StartSources(TArray<FWaveInstance*>& WaveInstances, int32 Fir
 
 	SCOPE_CYCLE_COUNTER(STAT_AudioStartSources);
 
+	TArray<USoundWave*> StartingSoundWaves;
+
 	// Start sources as needed.
 	for (int32 InstanceIndex = FirstActiveIndex; InstanceIndex < WaveInstances.Num(); InstanceIndex++)
 	{
@@ -3728,6 +3730,11 @@ void FAudioDevice::StartSources(TArray<FWaveInstance*>& WaveInstances, int32 Fir
 				Source = FreeSources.Pop();
 				check(Source);
 
+				if (WaveInstance->WaveData)
+				{
+					StartingSoundWaves.AddUnique(WaveInstance->WaveData);
+				}
+
 				// Prepare for initialization...
 				bool bSuccess = false;
 				if (Source->PrepareForInitialization(WaveInstance))
@@ -3814,6 +3821,23 @@ void FAudioDevice::StartSources(TArray<FWaveInstance*>& WaveInstances, int32 Fir
 			}
 		}
 	}
+
+	DECLARE_CYCLE_STAT(TEXT("FGameThreadAudioTask.AudioSendResults"), STAT_AudioAddReferencedSoundWaves, STATGROUP_TaskGraphTasks);
+	
+	// Run a command to make sure we add the starting sounds to the referenced sound waves list
+	if (StartingSoundWaves.Num() > 0)
+	{
+		FAudioThread::RunCommandOnGameThread([this, StartingSoundWaves]()
+		{
+			if (GEngine)
+			{
+				for (USoundWave* SoundWave : StartingSoundWaves)
+				{
+					ReferencedSoundWaves.AddUnique(SoundWave);
+				}
+			}
+		}, GET_STATID(STAT_AudioAddReferencedSoundWaves));
+	}
 }
 
 void FAudioDevice::Update(bool bGameTicking)
@@ -3825,10 +3849,11 @@ void FAudioDevice::Update(bool bGameTicking)
 		// On game thread, look through registered sound waves and remove if we finished precaching (and audio decompressor is cleaned up)
 		// ReferencedSoundWaves is used to make sure GC doesn't run on any sound waves that are actively pre-caching within an async task.
 		// Sounds may be loaded, kick off an async task to decompress, but never actually try to play, so GC can reclaim these while precaches are in-flight.
+		// We are also tracking when a sound wave is actively being used to generate audio in the audio render to prevent GC from happening to sounds till being used in the audio renderer.
 		for (int32 i = ReferencedSoundWaves.Num() - 1; i >= 0; --i)
 		{
 			USoundWave* Wave = ReferencedSoundWaves[i];
-			if (Wave->GetPrecacheState() == ESoundWavePrecacheState::Done)
+			if (Wave->GetPrecacheState() == ESoundWavePrecacheState::Done && !Wave->IsGeneratingAudio())
 			{
 				ReferencedSoundWaves.RemoveAtSwap(i, 1, false);
 			}
@@ -5344,7 +5369,7 @@ void FAudioDevice::Precache(USoundWave* SoundWave, bool bSynchronous, bool bTrac
 		{
 			// On the game thread, add this sound wave to the referenced sound wave nodes so that it doesn't get GC'd
 			SoundWave->SetPrecacheState(ESoundWavePrecacheState::InProgress);
-			ReferencedSoundWaves.Add(SoundWave);
+			ReferencedSoundWaves.AddUnique(SoundWave);
 		}
 
 		// Precache is called from USoundWave::PostLoad, from the game thread, and thus function needs to be called from the audio thread
diff --git a/Engine/Source/Runtime/Engine/Private/SoundWave.cpp b/Engine/Source/Runtime/Engine/Private/SoundWave.cpp
--- a/Engine/Source/Runtime/Engine/Private/SoundWave.cpp
+++ b/Engine/Source/Runtime/Engine/Private/SoundWave.cpp
@@ -867,9 +867,6 @@ void USoundWave::BeginDestroy()
 {
 	Super::BeginDestroy();
 
-	// Flag that this sound wave is beginning destroying. This will ensure that all sounds using this in the audio renderer are stopped before GC finishes.
-	bIsBeginDestroy = true;
-
 #if WITH_EDITOR
 	// Flush any async results so we dont leak them in the DDC
 	if (GetDerivedDataCache() && AsyncLoadingDataFormats.Num() > 0)
@@ -1522,7 +1519,7 @@ bool USoundWave::IsReadyForFinishDestroy()
 			}, GET_STATID(STAT_AudioFreeResources));
 		}
 
-	return NumSourcesPlaying.GetValue() == 0 && ResourceState == ESoundWaveResourceState::Freed;
+	return ResourceState == ESoundWaveResourceState::Freed;
 }
 
 
