Removed invalid check in FLevelUtils::ApplyLevelTransform and adapted code to support reentrancy.

diff --git a/Engine/Source/Runtime/Engine/Private/LevelUtils.cpp b/Engine/Source/Runtime/Engine/Private/LevelUtils.cpp
--- a/Engine/Source/Runtime/Engine/Private/LevelUtils.cpp
+++ b/Engine/Source/Runtime/Engine/Private/LevelUtils.cpp
@@ -385,7 +385,7 @@ void FLevelUtils::ApplyLevelTransform( ULevel* Level, const FTransform& LevelTra
 	if (bTransformActors)
 	{
 #if WITH_EDITOR
-		check(!bApplyingLevelTransform);
+		bool bApplyingLevelTransformBackup = bApplyingLevelTransform;
 		bApplyingLevelTransform = true;
 #endif
 		if (!LevelTransform.GetRotation().IsIdentity())
@@ -420,7 +420,7 @@ void FLevelUtils::ApplyLevelTransform( ULevel* Level, const FTransform& LevelTra
 
 		Level->OnApplyLevelTransform.Broadcast(LevelTransform);
 #if WITH_EDITOR
-		bApplyingLevelTransform = false;
+		bApplyingLevelTransform = bApplyingLevelTransformBackup;
 #endif
 	}
 }
