#UE4 Quick fix for cook crash. Needs investigation

diff --git a/Engine/Source/Runtime/Engine/Private/Components/StaticMeshComponent.cpp b/Engine/Source/Runtime/Engine/Private/Components/StaticMeshComponent.cpp
--- a/Engine/Source/Runtime/Engine/Private/Components/StaticMeshComponent.cpp
+++ b/Engine/Source/Runtime/Engine/Private/Components/StaticMeshComponent.cpp
@@ -2697,7 +2697,7 @@ FArchive& operator<<(FArchive& Ar,FStaticMeshComponentLODInfo& I)
 	const uint8 OverrideColorsStripFlag = 1;
 	bool bStrippedOverrideColors = false;
 #if WITH_EDITORONLY_DATA
-	if( Ar.IsCooking() )
+	if( Ar.IsCooking() && I.OwningComponent )
 	{
 		// Check if override color should be stripped too	
 		int32 LODIndex = 0;
