Initialize GPU particle sorting buffers only on SM5, as ES3.1 does not use particle sorting

diff --git a/Engine/Source/Runtime/Engine/Private/Particles/ParticleSortingGPU.cpp b/Engine/Source/Runtime/Engine/Private/Particles/ParticleSortingGPU.cpp
--- a/Engine/Source/Runtime/Engine/Private/Particles/ParticleSortingGPU.cpp
+++ b/Engine/Source/Runtime/Engine/Private/Particles/ParticleSortingGPU.cpp
@@ -49,7 +49,7 @@ public:
 
 	static bool ShouldCompilePermutation(const FGlobalShaderPermutationParameters& Parameters)
 	{
-		return RHISupportsComputeShaders(Parameters.Platform);
+		return IsFeatureLevelSupported(Parameters.Platform, ERHIFeatureLevel::SM5);
 	}
 
 	static void ModifyCompilationEnvironment(const FGlobalShaderPermutationParameters& Parameters, FShaderCompilerEnvironment& OutEnvironment)
@@ -185,7 +185,7 @@ static int32 GenerateParticleSortKeys(
 	)
 {
 	SCOPED_DRAW_EVENT(RHICmdList, ParticleSortKeyGen);
-	check(RHISupportsComputeShaders(GShaderPlatformForFeatureLevel[FeatureLevel]));
+	check(FeatureLevel == ERHIFeatureLevel::SM5);
 
 	FParticleKeyGenParameters KeyGenParameters;
 	FParticleKeyGenUniformBufferRef KeyGenUniformBuffer;
@@ -253,7 +253,7 @@ static int32 GenerateParticleSortKeys(
  */
 void FParticleSortBuffers::InitRHI()
 {
-	if (RHISupportsComputeShaders(GShaderPlatformForFeatureLevel[GetFeatureLevel()]))
+	if (GetFeatureLevel() == ERHIFeatureLevel::SM5)
 	{
 		for (int32 BufferIndex = 0; BufferIndex < 2; ++BufferIndex)
 		{
