Fixed crash when editing Sky component that has not been spawned in a world yet.
Regular in level behavior is still working.

diff --git a/Engine/Source/Runtime/Engine/Private/Components/SkyAtmosphereComponent.cpp b/Engine/Source/Runtime/Engine/Private/Components/SkyAtmosphereComponent.cpp
--- a/Engine/Source/Runtime/Engine/Private/Components/SkyAtmosphereComponent.cpp
+++ b/Engine/Source/Runtime/Engine/Private/Components/SkyAtmosphereComponent.cpp
@@ -78,12 +78,17 @@ static bool SkyAtmosphereComponentStaticLightingBuilt(const USkyAtmosphereCompon
 {
 	UMapBuildDataRegistry* Registry = Component->GetOwner() && Component->GetOwner()->GetLevel() ? Component->GetOwner()->GetLevel()->GetOrCreateMapBuildData() : nullptr;
 	const FSkyAtmosphereMapBuildData* SkyAtmosphereFogBuildData = Registry ? Registry->GetSkyAtmosphereBuildData(Component->GetStaticLightingBuiltGuid()) : nullptr;
-	class FSceneInterface* Scene = Component->GetWorld()->Scene;
+	UWorld* World = Component->GetWorld();
+	if (World)
+	{
+		class FSceneInterface* Scene = Component->GetWorld()->Scene;
 
-	// Only require building if there is a Sky or Sun light requiring lighting builds, i.e. non movable.
-	const bool StaticLightingDependsOnAtmosphere = Scene->HasSkyLightRequiringLightingBuild() || Scene->HasAtmosphereLightRequiringLightingBuild();
-	// Built data is available or static lighting does not depend any sun/sky components.
-	return (SkyAtmosphereFogBuildData != nullptr && StaticLightingDependsOnAtmosphere) || !StaticLightingDependsOnAtmosphere;
+		// Only require building if there is a Sky or Sun light requiring lighting builds, i.e. non movable.
+		const bool StaticLightingDependsOnAtmosphere = Scene->HasSkyLightRequiringLightingBuild() || Scene->HasAtmosphereLightRequiringLightingBuild();
+		// Built data is available or static lighting does not depend any sun/sky components.
+		return (SkyAtmosphereFogBuildData != nullptr && StaticLightingDependsOnAtmosphere) || !StaticLightingDependsOnAtmosphere;
+	}
+	return true;	// The component has not been spawned in any world yet so let's mark it as built for now.
 }
 
 void USkyAtmosphereComponent::AddToRenderScene() const
