SessionSummary analytics event now includes the SessionDuration in seconds.

Fixed ShutdownTypeString reporting AbnormalShutdown for clean shutdowns.

Fixed Slate delegate not getting unregistered if records weren't initialized correctly.

diff --git a/Engine/Source/Runtime/Engine/Private/EngineSessionManager.cpp b/Engine/Source/Runtime/Engine/Private/EngineSessionManager.cpp
--- a/Engine/Source/Runtime/Engine/Private/EngineSessionManager.cpp
+++ b/Engine/Source/Runtime/Engine/Private/EngineSessionManager.cpp
@@ -69,6 +69,7 @@ namespace SessionManagerDefs
 	static const FString EngineVersionStoreKey(TEXT("EngineVersion"));
 	static const FString TimestampStoreKey(TEXT("Timestamp"));
 	static const FString StartupTimestampStoreKey(TEXT("StartupTimestamp"));
+	static const FString SessionDurationStoreKey(TEXT("SessionDuration"));
 	static const FString Idle1MinStoreKey(TEXT("Idle1Min"));
 	static const FString Idle5MinStoreKey(TEXT("Idle5Min"));
 	static const FString Idle30MinStoreKey(TEXT("Idle30Min"));
@@ -326,16 +327,16 @@ void FEngineSessionManager::Shutdown()
 	if (!CurrentSession.bIsTerminating) // Skip Slate if terminating, since we can't guarantee which thread called us.
 	{
 		FSlateApplication::Get().GetOnModalLoopTickEvent().RemoveAll(this);
-
-		if (!CurrentSessionSectionName.IsEmpty())
-		{
-			FPlatformMisc::SetStoredValue(SessionManagerDefs::StoreId, CurrentSessionSectionName, SessionManagerDefs::WasShutdownStoreKey, SessionManagerDefs::TrueValueString);
-		}
 	}
 
 	// Clear the session record for this session
 	if (bInitializedRecords)
 	{
+		if (!CurrentSession.bIsTerminating)
+		{
+			FPlatformMisc::SetStoredValue(SessionManagerDefs::StoreId, CurrentSessionSectionName, SessionManagerDefs::WasShutdownStoreKey, SessionManagerDefs::TrueValueString);
+		}
+
 		if (!CurrentSession.bCrashed)
 		{
 #if PLATFORM_SUPPORTS_WATCHDOG
@@ -536,7 +537,8 @@ void FEngineSessionManager::SendSessionRecordEvent(const FString& EventName, con
 #if !PLATFORM_PS4
 	FString ShutdownTypeString = Record.bCrashed ? SessionManagerDefs::CrashSessionToken :
 		(Record.bWasEverDebugger ? SessionManagerDefs::DebuggerSessionToken :
-		(Record.bIsTerminating ? SessionManagerDefs::TerminatedSessionToken : SessionManagerDefs::AbnormalSessionToken));
+		(Record.bIsTerminating ? SessionManagerDefs::TerminatedSessionToken : 
+		(Record.bWasShutdown ? SessionManagerDefs::ShutdownSessionToken : SessionManagerDefs::AbnormalSessionToken)));
 #else
 	// PS4 cannot set the crash flag so report abnormal shutdowns with a specific token meaning "crash or abnormal shutdown".
 	FString ShutdownTypeString = Record.bWasEverDebugger ? SessionManagerDefs::DebuggerSessionToken : SessionManagerDefs::PS4SessionToken;
@@ -564,6 +566,9 @@ void FEngineSessionManager::SendSessionRecordEvent(const FString& EventName, con
 	AnalyticsAttributes.Emplace(SessionManagerDefs::IsInEnterpriseStoreKey, Record.bIsInEnterprise);
 	AnalyticsAttributes.Emplace(SessionManagerDefs::IsInVRModeStoreKey, Record.bIsInVRMode);
 
+	double SessionDuration = (Record.Timestamp - Record.StartupTimestamp).GetTotalSeconds();
+	AnalyticsAttributes.Emplace(SessionManagerDefs::SessionDurationStoreKey, SessionDuration);
+
 	AnalyticsAttributes.Emplace(TEXT("1MinIdle"), Record.Idle1Min);
 	AnalyticsAttributes.Emplace(TEXT("5MinIdle"), Record.Idle5Min);
 	AnalyticsAttributes.Emplace(TEXT("30MinIdle"), Record.Idle30Min);
@@ -776,7 +781,7 @@ void FEngineSessionManager::OnCrashing()
 
 void FEngineSessionManager::OnAppReactivate()
 {
-	if (CurrentSession.bIsDeactivated)
+	if (CurrentSession.bIsDeactivated && bInitializedRecords)
 	{
 		CurrentSession.bIsDeactivated = false;
 		FPlatformMisc::SetStoredValue(SessionManagerDefs::StoreId, CurrentSessionSectionName, SessionManagerDefs::DeactivatedStoreKey, SessionManagerDefs::FalseValueString);
@@ -785,7 +790,7 @@ void FEngineSessionManager::OnAppReactivate()
 
 void FEngineSessionManager::OnAppDeactivate()
 {
-	if (!CurrentSession.bIsDeactivated)
+	if (!CurrentSession.bIsDeactivated && bInitializedRecords)
 	{
 		CurrentSession.bIsDeactivated = true;
 		FPlatformMisc::SetStoredValue(SessionManagerDefs::StoreId, CurrentSessionSectionName, SessionManagerDefs::DeactivatedStoreKey, SessionManagerDefs::TrueValueString);
@@ -794,7 +799,7 @@ void FEngineSessionManager::OnAppDeactivate()
 
 void FEngineSessionManager::OnAppBackground()
 {
-	if (!CurrentSession.bIsInBackground)
+	if (!CurrentSession.bIsInBackground && bInitializedRecords)
 	{
 		CurrentSession.bIsInBackground = true;
 		FPlatformMisc::SetStoredValue(SessionManagerDefs::StoreId, CurrentSessionSectionName, SessionManagerDefs::BackgroundStoreKey, SessionManagerDefs::TrueValueString);
@@ -803,7 +808,7 @@ void FEngineSessionManager::OnAppBackground()
 
 void FEngineSessionManager::OnAppForeground()
 {
-	if (CurrentSession.bIsInBackground)
+	if (CurrentSession.bIsInBackground && bInitializedRecords)
 	{
 		CurrentSession.bIsInBackground = false;
 		FPlatformMisc::SetStoredValue(SessionManagerDefs::StoreId, CurrentSessionSectionName, SessionManagerDefs::BackgroundStoreKey, SessionManagerDefs::FalseValueString);
@@ -812,7 +817,7 @@ void FEngineSessionManager::OnAppForeground()
 
 void FEngineSessionManager::OnTerminate()
 {
-	if (!CurrentSession.bIsTerminating)
+	if (!CurrentSession.bIsTerminating && bInitializedRecords)
 	{
 		CurrentSession.bIsTerminating = true;
 		FPlatformMisc::SetStoredValue(SessionManagerDefs::StoreId, CurrentSessionSectionName, SessionManagerDefs::TerminatingKey, SessionManagerDefs::TrueValueString);
