- Fix StopQuietest to not include concurrency ducking volume values as this can cause rapid cycling.  Move UpdateConcurrencyVolumeScalars call to after stopping max concurrency StopQuietest to avoid generational volume being applied prematurally to sounds by first-frame volume calcuation of realized loops.
- Optimization: Calculate update interval on virtual loop initialization to non-zero value to avoid rapid cycling of sounds using StopQuietest setting
mid-frame

diff --git a/Engine/Source/Runtime/Engine/Classes/Sound/SoundConcurrency.h b/Engine/Source/Runtime/Engine/Classes/Sound/SoundConcurrency.h
--- a/Engine/Source/Runtime/Engine/Classes/Sound/SoundConcurrency.h
+++ b/Engine/Source/Runtime/Engine/Classes/Sound/SoundConcurrency.h
@@ -184,7 +184,7 @@ public:
 	void Update(float InElapsed);
 
 	float GetVolume(bool bInDecibels = false) const;
-	float GetTargetVolume(bool bInDecibles = false) const;
+	float GetTargetVolume(bool bInDecibels = false) const;
 
 	void SetTarget(float InTargetVolume, float InLerpTime);
 };
@@ -284,7 +284,7 @@ public:
 	/** Removes the active sound from concurrency tracking when active sound is stopped. */
 	void RemoveActiveSound(FActiveSound& ActiveSound);
 
-	/** Stops sound, applying concurrency rules for how to stop */
+	/** Stops sound, applying concurrency rules for how to stop. */
 	void StopDueToVoiceStealing(FActiveSound& ActiveSound);
 
 	/** Stops any active sounds due to max concurrency quietest sound resolution rule */
diff --git a/Engine/Source/Runtime/Engine/Private/ActiveSound.cpp b/Engine/Source/Runtime/Engine/Private/ActiveSound.cpp
--- a/Engine/Source/Runtime/Engine/Private/ActiveSound.cpp
+++ b/Engine/Source/Runtime/Engine/Private/ActiveSound.cpp
@@ -129,10 +129,16 @@ FActiveSound* FActiveSound::CreateVirtualCopy(const FActiveSound& InActiveSoundT
 	ActiveSound->bAsyncOcclusionPending = false;
 	ActiveSound->bHasVirtualized = true;
 	ActiveSound->bIsPlayingAudio = false;
-
+	ActiveSound->bShouldStopDueToMaxConcurrency = false;
 	ActiveSound->AudioDevice = &InAudioDevice;
 	ActiveSound->PlaybackTimeNonVirtualized = 0.0f;
-	ActiveSound->VolumeConcurrency = 1.0f;
+
+	// If volume concurrency tracking is enabled, reset the value,
+	// otherwise keep disabled
+	if (InActiveSoundToCopy.VolumeConcurrency >= 0.0f)
+	{
+		ActiveSound->VolumeConcurrency = 1.0f;
+	}
 
 	ActiveSound->ConcurrencyGroupData.Reset();
 	ActiveSound->WaveInstances.Reset();
@@ -528,10 +534,6 @@ void FActiveSound::UpdateWaveInstances(TArray<FWaveInstance*> &InWaveInstances,
 	// final value that is correct
 	UpdateAdjustVolumeMultiplier(DeltaTime);
 
-	// Update concurrency volume scalars. Must be done prior to getting collective volume and applying below
-	// in parse params.
-	UpdateConcurrencyVolumeScalars(DeltaTime);
-
 	ParseParams.VolumeMultiplier = GetVolume();
 
 	ParseParams.Priority = Priority;
@@ -646,6 +648,18 @@ void FActiveSound::UpdateWaveInstances(TArray<FWaveInstance*> &InWaveInstances,
 					VolumeConcurrency = WaveInstanceVolume;
 				}
 			}
+
+			// Remove concurrency volume scalars as this can cause ping-ponging to occur with virtualization and loops
+			// utilizing concurrency with StopQuietest.
+			const float VolumeScale = GetTotalConcurrencyVolumeScale();
+			if (VolumeScale > SMALL_NUMBER)
+			{
+				VolumeConcurrency /= VolumeScale;
+			}
+			else
+			{
+				VolumeConcurrency = 0.0f;
+			}
 		}
 
 		// Check to see if we need to broadcast the envelope value of sounds playing with this active sound
diff --git a/Engine/Source/Runtime/Engine/Private/AudioDevice.cpp b/Engine/Source/Runtime/Engine/Private/AudioDevice.cpp
--- a/Engine/Source/Runtime/Engine/Private/AudioDevice.cpp
+++ b/Engine/Source/Runtime/Engine/Private/AudioDevice.cpp
@@ -3480,6 +3480,17 @@ int32 FAudioDevice::GetSortedActiveWaveInstances(TArray<FWaveInstance*>& WaveIns
 		StopQuietSoundsDueToMaxConcurrency(WaveInstances, ActiveSoundsCopy);
 	}
 
+	// Must be completed after StopQuietSoundsDueToMaxConcurrency as it avoids an issue
+	// where quiet loops can wrongfully scale concurrency ducking improperly if they continue
+	// to attempt to be evaluated while being periodically realized to check volumes from virtualized.
+	for (int32 i = 0; i < ActiveSoundsCopy.Num(); ++i)
+	{
+		if (FActiveSound* ActiveSound = ActiveSoundsCopy[i])
+		{
+			ActiveSound->UpdateConcurrencyVolumeScalars(GetGameDeltaTime());
+		}
+	}
+
 	int32 FirstActiveIndex = 0;
 	// Only need to do the wave instance sort if we have any waves and if our wave instances are greater than our max channels.
 	if (WaveInstances.Num() >= 0)
diff --git a/Engine/Source/Runtime/Engine/Private/AudioVirtualLoop.cpp b/Engine/Source/Runtime/Engine/Private/AudioVirtualLoop.cpp
--- a/Engine/Source/Runtime/Engine/Private/AudioVirtualLoop.cpp
+++ b/Engine/Source/Runtime/Engine/Private/AudioVirtualLoop.cpp
@@ -86,6 +86,7 @@ bool FAudioVirtualLoop::Virtualize(const FActiveSound& InActiveSound, FAudioDevi
 
 	FActiveSound* ActiveSound = FActiveSound::CreateVirtualCopy(InActiveSound, InAudioDevice);
 	OutVirtualLoop.ActiveSound = ActiveSound;
+	OutVirtualLoop.CalculateUpdateInterval();
 	return true;
 }
 
diff --git a/Engine/Source/Runtime/Engine/Private/SoundConcurrency.cpp b/Engine/Source/Runtime/Engine/Private/SoundConcurrency.cpp
--- a/Engine/Source/Runtime/Engine/Private/SoundConcurrency.cpp
+++ b/Engine/Source/Runtime/Engine/Private/SoundConcurrency.cpp
@@ -194,19 +194,25 @@ void FConcurrencyGroup::StopQuietSoundsDueToMaxConcurrency()
 		return;
 	}
 
-	// Helper function for sort this concurrency group's active sounds according to their "volume" concurrency
-	// Quieter sounds will be at the front of the array
+	// Comparator for sorting group's ActiveSounds according to their "volume" concurrency. Quieter sounds will be
+	//  at the front of the array. If they share the same volume, newer sounds will be sorted first to avoid loop
+	// realization ping-ponging 
 	struct FCompareActiveSounds
 	{
 		FORCEINLINE bool operator()(const FActiveSound& A, const FActiveSound& B) const
 		{
+			if (FMath::IsNearlyEqual(A.VolumeConcurrency, B.VolumeConcurrency, KINDA_SMALL_NUMBER))
+			{
+				return A.PlaybackTime > B.PlaybackTime;
+			}
 			return A.VolumeConcurrency < B.VolumeConcurrency;
 		}
 	};
 
 	ActiveSounds.Sort(FCompareActiveSounds());
 
-	const int32 NumSoundsToStop = ActiveSounds.Num() - Settings.MaxCount;
+	const int32 NumActiveSounds = ActiveSounds.Num();
+	const int32 NumSoundsToStop = NumActiveSounds - Settings.MaxCount;
 	check(NumSoundsToStop > 0);
 
 	// Need to make a new list when stopping the sounds since the process of stopping an active sound
@@ -216,16 +222,9 @@ void FConcurrencyGroup::StopQuietSoundsDueToMaxConcurrency()
 	{
 		FActiveSound* ActiveSound = ActiveSounds[i];
 		check(ActiveSound);
-
-		// Flag this active sound as needing to be stopped due to volume-based max concurrency.
-		// This will actually be stopped in the audio device update function.
-		if (ActiveSound->FadeOut != FActiveSound::EFadeOut::Concurrency)
-		{
-			ActiveSound->bShouldStopDueToMaxConcurrency = true;
-		}
+		ActiveSound->bShouldStopDueToMaxConcurrency = true;
 	}
 
-	const int32 NumActiveSounds = ActiveSounds.Num();
 	for (; i < NumActiveSounds; ++i)
 	{
 		ActiveSounds[i]->bShouldStopDueToMaxConcurrency = false;
