Mark RemoveLocalPlayer as virtual

diff --git a/Engine/Source/Runtime/Engine/Classes/Engine/GameInstance.h b/Engine/Source/Runtime/Engine/Classes/Engine/GameInstance.h
--- a/Engine/Source/Runtime/Engine/Classes/Engine/GameInstance.h
+++ b/Engine/Source/Runtime/Engine/Classes/Engine/GameInstance.h
@@ -268,7 +268,7 @@ public:
 	 * @param Player - The player to remove.
 	 * @return whether the player was successfully removed. Removal is not allowed while connected to a server.
 	 */
-	bool					RemoveLocalPlayer(ULocalPlayer * ExistingPlayer);
+	virtual bool			RemoveLocalPlayer(ULocalPlayer * ExistingPlayer);
 
 	int32					GetNumLocalPlayers() const;
 	ULocalPlayer*			GetLocalPlayerByIndex(const int32 Index) const;
