Add r.Shadow.MaxNumFarShadowCascades to control the number of far cascades from all directional lights.

diff --git a/Engine/Source/Runtime/Engine/Private/Components/DirectionalLightComponent.cpp b/Engine/Source/Runtime/Engine/Private/Components/DirectionalLightComponent.cpp
--- a/Engine/Source/Runtime/Engine/Private/Components/DirectionalLightComponent.cpp
+++ b/Engine/Source/Runtime/Engine/Private/Components/DirectionalLightComponent.cpp
@@ -45,6 +45,12 @@ static TAutoConsoleVariable<float> CVarPerObjectCastDistanceRadiusScale(
 	ECVF_RenderThreadSafe
 	);
 
+static TAutoConsoleVariable<int32> CVarMaxNumFarShadowCascades(
+	TEXT("r.Shadow.MaxNumFarShadowCascades"),
+	10,
+	TEXT("Max number of far shadow cascades that can be cast from a directional light"),
+	ECVF_RenderThreadSafe);
+
 /**
  * The scene info for a directional light.
  */
@@ -247,8 +253,10 @@ public:
 
 	/** Returns the number of view dependent shadows this light will create, not counting distance field shadow cascades. */
 	virtual uint32 GetNumViewDependentWholeSceneShadows(const FSceneView& View, bool bPrecomputedLightingIsValid) const override
-	{ 
-		uint32 TotalCascades = GetNumShadowMappedCascades(View.MaxShadowCascades, bPrecomputedLightingIsValid) + FarShadowCascadeCount;
+	{
+		uint32 ClampedFarShadowCascadeCount = FMath::Min((uint32)CVarMaxNumFarShadowCascades.GetValueOnAnyThread(), FarShadowCascadeCount);
+
+		uint32 TotalCascades = GetNumShadowMappedCascades(View.MaxShadowCascades, bPrecomputedLightingIsValid) + ClampedFarShadowCascadeCount;
 
 		return TotalCascades;
 	}
@@ -613,7 +621,9 @@ private:
 			else
 			{
 				// the far cascades start at the after the near cascades
-				return CascadeDistanceWithoutFar + ComputeAccumulatedScale(EffectiveCascadeDistributionExponent, SplitIndex - NumNearCascades, FarShadowCascadeCount) * (FarShadowDistance - CascadeDistanceWithoutFar);
+				uint32 ClampedFarShadowCascadeCount = FMath::Min((uint32)CVarMaxNumFarShadowCascades.GetValueOnAnyThread(), FarShadowCascadeCount);
+
+				return CascadeDistanceWithoutFar + ComputeAccumulatedScale(EffectiveCascadeDistributionExponent, SplitIndex - NumNearCascades, ClampedFarShadowCascadeCount) * (FarShadowDistance - CascadeDistanceWithoutFar);
 			}
 		}
 		else
