Fix how the UI safe zone is displayed for splitscreen.
Use the Viewport rather than View[0] that is currently used.

diff --git a/Engine/Source/Runtime/Engine/Private/GameViewportClient.cpp b/Engine/Source/Runtime/Engine/Private/GameViewportClient.cpp
--- a/Engine/Source/Runtime/Engine/Private/GameViewportClient.cpp
+++ b/Engine/Source/Runtime/Engine/Private/GameViewportClient.cpp
@@ -2423,8 +2423,9 @@ void UGameViewportClient::DrawTitleSafeArea( UCanvas* Canvas )
 	FMargin SafeZone;
 	const ULevelEditorPlaySettings* PlayInSettings = GetDefault<ULevelEditorPlaySettings>();
 
-	const float Width = Canvas->UnsafeSizeX;
-	const float Height = Canvas->UnsafeSizeY;
+	float Width, Height;
+	GetPixelSizeOfScreen(Width, Height, Canvas, 0);
+
 	const FLinearColor UnsafeZoneColor(1.0f, 0.0f, 0.0f, 0.25f);
 	FCanvasTileItem TileItem(FVector2D::ZeroVector, GWhiteTexture, UnsafeZoneColor);
 	TileItem.BlendMode = SE_BLEND_Translucent;
