Fix UpdateOverlaps getting called when actor collision is disabled

#rnx

diff --git a/Engine/Source/Runtime/Engine/Classes/Components/PrimitiveComponent.h b/Engine/Source/Runtime/Engine/Classes/Components/PrimitiveComponent.h
--- a/Engine/Source/Runtime/Engine/Classes/Components/PrimitiveComponent.h
+++ b/Engine/Source/Runtime/Engine/Classes/Components/PrimitiveComponent.h
@@ -2166,7 +2166,7 @@ public:
 	
 protected:
 	/** Called when the BodyInstance ResponseToChannels, CollisionEnabled or bNotifyRigidBodyCollision changes, in case subclasses want to use that information. */
-	virtual void OnComponentCollisionSettingsChanged();
+	virtual void OnComponentCollisionSettingsChanged(bool bDeferUpdateOverlaps = false);
 
 	/** Ends all current component overlaps. Generally used when destroying this component or when it can no longer generate overlaps. */
 	void ClearComponentOverlaps(bool bDoNotifies, bool bSkipNotifySelf);
diff --git a/Engine/Source/Runtime/Engine/Classes/Components/SkeletalMeshComponent.h b/Engine/Source/Runtime/Engine/Classes/Components/SkeletalMeshComponent.h
--- a/Engine/Source/Runtime/Engine/Classes/Components/SkeletalMeshComponent.h
+++ b/Engine/Source/Runtime/Engine/Classes/Components/SkeletalMeshComponent.h
@@ -1452,7 +1452,7 @@ public:
 	virtual bool IsAnyRigidBodyAwake() override;
 	virtual void SetEnableGravity(bool bGravityEnabled);
 	virtual bool IsGravityEnabled() const override;
-	virtual void OnComponentCollisionSettingsChanged() override;
+	virtual void OnComponentCollisionSettingsChanged(bool bDeferUpdateOverlaps = false) override;
 	virtual void SetPhysMaterialOverride(UPhysicalMaterial* NewPhysMaterial) override;
 	virtual bool GetSquaredDistanceToCollision(const FVector& Point, float& OutSquaredDistance, FVector& OutClosestPointOnCollision) const override;
 
diff --git a/Engine/Source/Runtime/Engine/Private/Actor.cpp b/Engine/Source/Runtime/Engine/Private/Actor.cpp
--- a/Engine/Source/Runtime/Engine/Private/Actor.cpp
+++ b/Engine/Source/Runtime/Engine/Private/Actor.cpp
@@ -3845,6 +3845,9 @@ void AActor::SetActorEnableCollision(bool bNewActorEnableCollision)
 		{
 			Components[CompIdx]->OnActorEnableCollisionChanged();
 		}
+
+		// update overlaps once after all components have been updated
+		UpdateOverlaps();
 	}
 }
 
diff --git a/Engine/Source/Runtime/Engine/Private/PrimitiveComponentPhysics.cpp b/Engine/Source/Runtime/Engine/Private/PrimitiveComponentPhysics.cpp
--- a/Engine/Source/Runtime/Engine/Private/PrimitiveComponentPhysics.cpp
+++ b/Engine/Source/Runtime/Engine/Private/PrimitiveComponentPhysics.cpp
@@ -942,10 +942,10 @@ FName UPrimitiveComponent::GetCollisionProfileName() const
 void UPrimitiveComponent::OnActorEnableCollisionChanged()
 {
 	BodyInstance.UpdatePhysicsFilterData();
-	OnComponentCollisionSettingsChanged();
+	OnComponentCollisionSettingsChanged(true);
 }
 
-void UPrimitiveComponent::OnComponentCollisionSettingsChanged()
+void UPrimitiveComponent::OnComponentCollisionSettingsChanged(bool bDeferUpdateOverlaps)
 {
 	if (IsRegistered() && !IsTemplate())			// not for CDOs
 	{
@@ -955,7 +955,10 @@ void UPrimitiveComponent::OnComponentCollisionSettingsChanged()
 			ClearSkipUpdateOverlaps();
 		}
 
-		UpdateOverlaps();
+		if (!bDeferUpdateOverlaps)
+		{
+			UpdateOverlaps();
+		}
 
 		// update navigation data if needed
 		const bool bNewNavRelevant = IsNavigationRelevant();
diff --git a/Engine/Source/Runtime/Engine/Private/SkeletalMeshComponentPhysics.cpp b/Engine/Source/Runtime/Engine/Private/SkeletalMeshComponentPhysics.cpp
--- a/Engine/Source/Runtime/Engine/Private/SkeletalMeshComponentPhysics.cpp
+++ b/Engine/Source/Runtime/Engine/Private/SkeletalMeshComponentPhysics.cpp
@@ -205,7 +205,7 @@ void USkeletalMeshComponent::SetSimulatePhysics(bool bSimulate)
  	UpdateClothTickRegisteredState();
 }
 
-void USkeletalMeshComponent::OnComponentCollisionSettingsChanged()
+void USkeletalMeshComponent::OnComponentCollisionSettingsChanged(bool bDeferUpdateOverlaps)
 {
 	for(int32 i=0; i<Bodies.Num(); i++)
 	{
@@ -220,7 +220,7 @@ void USkeletalMeshComponent::OnComponentCollisionSettingsChanged()
 		((FSkeletalMeshSceneProxy*)SceneProxy)->SetCollisionEnabled_GameThread(IsCollisionEnabled());
 	}
 
-	Super::OnComponentCollisionSettingsChanged();
+	Super::OnComponentCollisionSettingsChanged(bDeferUpdateOverlaps);
 }
 
 void USkeletalMeshComponent::AddRadialImpulse(FVector Origin, float Radius, float Strength, ERadialImpulseFalloff Falloff, bool bVelChange)
