Add simple fade selection types for AudioComponent volume adjustment/fading in and out
- Minor bug fix: StopDelayed now accounting for case where prior fade could've already set target volume, which before was causing fade to stop prematurally.

diff --git a/Engine/Plugins/Editor/FacialAnimation/Source/FacialAnimation/Private/AudioCurveSourceComponent.cpp b/Engine/Plugins/Editor/FacialAnimation/Source/FacialAnimation/Private/AudioCurveSourceComponent.cpp
--- a/Engine/Plugins/Editor/FacialAnimation/Source/FacialAnimation/Private/AudioCurveSourceComponent.cpp
+++ b/Engine/Plugins/Editor/FacialAnimation/Source/FacialAnimation/Private/AudioCurveSourceComponent.cpp
@@ -21,6 +21,7 @@ UAudioCurveSourceComponent::UAudioCurveSourceComponent()
 	CachedFadeVolumeLevel = 1.0f;
 	CachedDuration = 0.0f;
 	bCachedLooping = false;
+	CachedFadeType = EAudioFaderCurve::Linear;
 
 #if WITH_EDITORONLY_DATA
 	bVisualizeComponent = false;
@@ -54,24 +55,25 @@ void UAudioCurveSourceComponent::CacheCurveData()
 	}
 }
 
-void UAudioCurveSourceComponent::FadeIn(float FadeInDuration, float FadeVolumeLevel, float StartTime)
+void UAudioCurveSourceComponent::FadeIn(float FadeInDuration, float FadeVolumeLevel, float StartTime, EAudioFaderCurve FadeType)
 {
 	CacheCurveData();
 
 	if (CachedSyncPreRoll <= 0.0f)
 	{
-		PlayInternal(StartTime, FadeInDuration, FadeVolumeLevel);
+		PlayInternal(StartTime, FadeInDuration, FadeVolumeLevel, FadeType);
 	}
 	else
 	{
 		CachedFadeInDuration = FadeInDuration;
 		CachedFadeVolumeLevel = FadeVolumeLevel;
 		CachedStartTime = StartTime;
+		CachedFadeType = FadeType;
 		Delay = 0.0f;
 	}
 }
 
-void UAudioCurveSourceComponent::FadeOut(float FadeOutDuration, float FadeVolumeLevel)
+void UAudioCurveSourceComponent::FadeOut(float FadeOutDuration, float FadeVolumeLevel, EAudioFaderCurve FadeType)
 {
 	if (Delay < CachedSyncPreRoll)
 	{
@@ -82,7 +84,7 @@ void UAudioCurveSourceComponent::FadeOut(float FadeOutDuration, float FadeVolume
 	}
 	else
 	{
-		Super::FadeOut(FadeOutDuration, FadeVolumeLevel);
+		Super::FadeOut(FadeOutDuration, FadeVolumeLevel, FadeType);
 	}
 }
 
@@ -129,7 +131,7 @@ void UAudioCurveSourceComponent::TickComponent(float DeltaTime, enum ELevelTick
 	Delay = FMath::Min(Delay + DeltaTime, CachedSyncPreRoll);
 	if (OldDelay < CachedSyncPreRoll && Delay >= CachedSyncPreRoll)
 	{
-		PlayInternal(CachedStartTime, CachedFadeInDuration, CachedFadeVolumeLevel);
+		PlayInternal(CachedStartTime, CachedFadeInDuration, CachedFadeVolumeLevel, CachedFadeType);
 	}
 	else if(Delay < CachedSyncPreRoll)
 	{
diff --git a/Engine/Plugins/Editor/FacialAnimation/Source/FacialAnimation/Public/AudioCurveSourceComponent.h b/Engine/Plugins/Editor/FacialAnimation/Source/FacialAnimation/Public/AudioCurveSourceComponent.h
--- a/Engine/Plugins/Editor/FacialAnimation/Source/FacialAnimation/Public/AudioCurveSourceComponent.h
+++ b/Engine/Plugins/Editor/FacialAnimation/Source/FacialAnimation/Public/AudioCurveSourceComponent.h
@@ -37,9 +37,9 @@ public:
 	virtual void TickComponent(float DeltaTime, enum ELevelTick TickType, FActorComponentTickFunction *ThisTickFunction) override;
 
 	/** UAudioComponent interface */
-	virtual void FadeIn(float FadeInDuration, float FadeVolumeLevel = 1.f, float StartTime = 0.f) override;
-	virtual	void FadeOut(float FadeOutDuration, float FadeVolumeLevel) override;
-	virtual void Play(float StartTime = 0.f) override;
+	virtual void FadeIn(float FadeInDuration, float FadeVolumeLevel = 1.0f, float StartTime = 0.0f, EAudioFaderCurve FadeType = EAudioFaderCurve::Linear) override;
+	virtual	void FadeOut(float FadeOutDuration, float FadeVolumeLevel, EAudioFaderCurve FadeType = EAudioFaderCurve::Linear) override;
+	virtual void Play(float StartTime = 0.0f) override;
 	virtual void Stop() override;
 	virtual bool IsPlaying() const override;
 
@@ -82,4 +82,7 @@ private:
 
 	/** Cache looping flag */
 	bool bCachedLooping;
+
+	/** Cached param for PlayInternal */
+	EAudioFaderCurve CachedFadeType;
 };
diff --git a/Engine/Source/Runtime/AudioMixer/Private/DSP/VolumeFader.cpp b/Engine/Source/Runtime/AudioMixer/Private/DSP/VolumeFader.cpp
new file mode 100644
--- /dev/null
+++ b/Engine/Source/Runtime/AudioMixer/Private/DSP/VolumeFader.cpp
@@ -0,0 +1,271 @@
+// Copyright 1998-2019 Epic Games, Inc. All Rights Reserved.
+#include "DSP/VolumeFader.h"
+
+#include "DSP/Dsp.h"
+
+
+namespace Audio
+{
+	FVolumeFader::FVolumeFader()
+		: Alpha(1.0f)
+		, Target(1.0f)
+		, ActiveDuration(-1.0f)
+		, FadeDuration(-1.0f)
+		, Elapsed(0.0f)
+		, FadeCurve(EFaderCurve::Linear)
+	{
+	}
+
+	float FVolumeFader::AlphaToVolume(float InAlpha, EFaderCurve InCurve)
+	{
+		switch (InCurve)
+		{
+		case EFaderCurve::Linear:
+		{
+			return InAlpha;
+		}
+		break;
+
+		case EFaderCurve::SCurve:
+		{
+			float Volume = 0.5f * Audio::FastSin(PI * InAlpha - HALF_PI) + 0.5f;
+			return FMath::Max(0.0f, Volume);
+		}
+		break;
+
+		case EFaderCurve::Sin:
+		{
+			float Volume = Audio::FastSin(HALF_PI * InAlpha);
+			return FMath::Max(0.0f, Volume);
+		}
+		break;
+
+		case EFaderCurve::Logarithmic:
+		{
+			return Audio::ConvertToLinear(InAlpha);
+		}
+		break;
+
+		default:
+		{
+			static_assert(static_cast<int32>(EFaderCurve::Count) == 4, "Possible missing switch case coverage for EAudioFade");
+		}
+		break;
+		}
+
+		return 1.0f;
+	}
+
+	void FVolumeFader::Activate()
+	{
+		if (ActiveDuration == 0.0f)
+		{
+			ActiveDuration = -1.0f;
+		}
+	}
+
+	void FVolumeFader::Deactivate()
+	{
+		StopFade();
+		ActiveDuration = 0.0f;
+	}
+
+	float FVolumeFader::GetVolume() const
+	{
+		return AlphaToVolume(Alpha, FadeCurve);
+	}
+
+	float FVolumeFader::GetVolumeAfterTime(float InDeltaTime) const
+	{
+		InDeltaTime = FMath::Max(0.0f, InDeltaTime);
+
+		// Keep stepping towards our target until we hit our stop time & Clamp
+		float FutureAlpha = Alpha;
+		const float Duration = Elapsed + InDeltaTime;
+
+		// If set to deactivate first, return default volume value of 1.0f
+		if (ActiveDuration >= 0.0f && Duration > ActiveDuration)
+		{
+			return 1.0f;
+		}
+
+		if (Duration < FadeDuration)
+		{
+			// Choose min/max bound and clamp dt to prevent unwanted spikes in volume
+			float MinValue = 0.0f;
+			float MaxValue = 0.0f;
+			if (Alpha < Target)
+			{
+				MinValue = Alpha;
+				MaxValue = Target;
+			}
+			else
+			{
+				MinValue = Target;
+				MaxValue = Alpha;
+			}
+
+			FutureAlpha = Alpha + ((Target - Alpha) * InDeltaTime / (FadeDuration - Elapsed));
+			FutureAlpha = FMath::Clamp(FutureAlpha, MinValue, MaxValue);
+		}
+
+		return AlphaToVolume(FutureAlpha, FadeCurve);
+	}
+
+	float FVolumeFader::GetFadeDuration() const
+	{
+		return FadeDuration;
+	}
+
+	EFaderCurve FVolumeFader::GetCurve() const
+	{
+		return FadeCurve;
+	}
+
+	float FVolumeFader::GetTargetVolume() const
+	{
+		switch (FadeCurve)
+		{
+		case EFaderCurve::Linear:
+		case EFaderCurve::SCurve:
+		case EFaderCurve::Sin:
+		{
+			return Target;
+		}
+		break;
+
+		case EFaderCurve::Logarithmic:
+		{
+			return Audio::ConvertToLinear(Target);
+		}
+		break;
+
+		default:
+		{
+			static_assert(static_cast<int32>(EFaderCurve::Count) == 4, "Possible missing switch case coverage for EAudioFade");
+		}
+		break;
+		}
+
+		return 1.0f;
+	}
+
+	bool FVolumeFader::IsActive() const
+	{
+		return Elapsed < ActiveDuration || ActiveDuration < 0.0f;
+	}
+
+	bool FVolumeFader::IsFading() const
+	{
+		return Elapsed < FadeDuration;
+	}
+
+	bool FVolumeFader::IsFadingIn() const
+	{
+		return IsFading() && Target > Alpha;
+	}
+
+	bool FVolumeFader::IsFadingOut() const
+	{
+		return IsFading() && Target < Alpha;
+	}
+
+	void FVolumeFader::SetActiveDuration(float InDuration)
+	{
+		if (InDuration < 0.0f)
+		{
+			Activate();
+		}
+		else
+		{
+			ActiveDuration = InDuration;
+		}
+	}
+
+	void FVolumeFader::SetVolume(float InVolume)
+	{
+		Alpha = InVolume;
+		ActiveDuration = -1.0f;
+		Elapsed = 0.0f;
+		FadeCurve = EFaderCurve::Linear;
+		FadeDuration = -1.0f;
+		Target = InVolume;
+	}
+
+	void FVolumeFader::StartFade(float InVolume, float InDuration, EFaderCurve InCurve)
+	{
+		if (InDuration <= 0.0f)
+		{
+			SetVolume(InVolume);
+			return;
+		}
+
+		if (InCurve != EFaderCurve::Logarithmic)
+		{
+			if (FadeCurve == EFaderCurve::Logarithmic)
+			{
+				Alpha = Audio::ConvertToLinear(Alpha);
+			}
+			Target = InVolume;
+		}
+		else
+		{
+			const float DecibelFloor = KINDA_SMALL_NUMBER; // -80dB
+			if (FadeCurve != EFaderCurve::Logarithmic)
+			{
+				Alpha = Audio::ConvertToDecibels(Alpha, DecibelFloor);
+			}
+			Target = Audio::ConvertToDecibels(InVolume, DecibelFloor);
+		}
+
+		ActiveDuration = -1.0f;
+		Elapsed = 0.0f;
+		FadeCurve = InCurve;
+		FadeDuration = InDuration;
+	}
+
+	void FVolumeFader::StopFade()
+	{
+		if (FadeCurve == EFaderCurve::Logarithmic)
+		{
+			Alpha = Audio::ConvertToLinear(Alpha);
+		}
+		Target = Alpha;
+		FadeCurve = EFaderCurve::Linear;
+		FadeDuration = -1.0f;
+	}
+
+	void FVolumeFader::Update(float InDeltaTime)
+	{
+		Elapsed += InDeltaTime;
+
+		if (!IsFading() || !IsActive())
+		{
+			return;
+		}
+
+		// Keep stepping towards target & clamp until fade duration has expired
+		// Choose min/max bound and clamp dt to prevent unwanted spikes in volume
+		float MinValue = 0.0f;
+		float MaxValue = 0.0f;
+		if (Alpha < Target)
+		{
+			MinValue = Alpha;
+			MaxValue = Target;
+		}
+		else
+		{
+			MinValue = Target;
+			MaxValue = Alpha;
+		}
+
+		Alpha += (Target - Alpha) * InDeltaTime / (FadeDuration - Elapsed);
+		Alpha = FMath::Clamp(Alpha, MinValue, MaxValue);
+
+		// Optimization that avoids fade calculations once alpha reaches target
+		if (InDeltaTime > SMALL_NUMBER && FMath::IsNearlyEqual(Alpha, Target))
+		{
+			StopFade();
+		}
+	}
+} // namespace Audio
\ No newline at end of file
diff --git a/Engine/Source/Runtime/AudioMixer/Public/DSP/VolumeFader.h b/Engine/Source/Runtime/AudioMixer/Public/DSP/VolumeFader.h
new file mode 100644
--- /dev/null
+++ b/Engine/Source/Runtime/AudioMixer/Public/DSP/VolumeFader.h
@@ -0,0 +1,147 @@
+// Copyright 1998-2019 Epic Games, Inc. All Rights Reserved.
+#pragma once
+
+#include "CoreMinimal.h"
+
+
+namespace Audio
+{
+	/** Type of fade to use when adjusting the audio component's volume. */
+	enum class EFaderCurve : uint8
+	{
+		// Linear Fade
+		Linear,
+
+		// Logarithmic Fade
+		Logarithmic,
+
+		// S-Curve, Sinusoidal Fade
+		SCurve,
+
+		// Equal Power, Sinusoidal Fade
+		Sin,
+
+		Count
+	};
+
+
+	/** Control-rate fader for managing volume fades of various standard shapes. */
+	class AUDIOMIXER_API FVolumeFader
+	{
+	public:
+		FVolumeFader();
+
+		/**
+		 * Activates the fader if currently deactivated.
+		 */
+		void Activate();
+
+		/**
+		 * Deactivates the fader, causing it to remain
+		 * at the current value and disregard update.
+		 */
+		void Deactivate();
+
+		/**
+		 * Returns current volume of fader
+		 */
+		float GetVolume() const;
+
+		/**
+		 * Returns the volume given the delta from the current time
+		 * into the future (Effectively like running to retrieve volume
+		 * but without actually updating internal state).
+		 */
+		float GetVolumeAfterTime(float InDeltaTime) const;
+
+		/**
+		 * Returns the duration of the fade.
+		 */
+		float GetFadeDuration() const;
+
+		/**
+		 * Returns the curve type of the fader
+		 */
+		EFaderCurve GetCurve() const;
+
+		/**
+		 * Returns the target volume of the fader
+		 */
+		float GetTargetVolume() const;
+
+		/**
+		 * Whether or not the fader is active.
+		 */
+		bool IsActive() const;
+
+		/**
+		 * Returns whether or not the fader is currently
+		 * fading over time.
+		 */
+		bool IsFading() const;
+
+		/**
+		 * Returns whether or not the fader is currently
+		 * fading over time and value is increasing.
+		 */
+		bool IsFadingIn() const;
+
+		/**
+		 * Returns whether or not the fader is currently
+		 * fading over time and value is decreasing.
+		 */
+		bool IsFadingOut() const;
+
+		/**
+		 * Sets the duration the fader is to be active in the future,
+		 * after which point the fader is disabled.
+		 */
+		void SetActiveDuration(float InDuration);
+
+		/**
+		 * Sets the volume immediately, interrupting any currently active fade.
+		 */
+		void SetVolume(float InVolume);
+
+		/**
+		 * Applies a volume fade over time with the provided parameters.
+		 */
+		void StartFade(float InVolume, float InDuration, EFaderCurve InCurve);
+
+		/**
+		 * Stops fade, maintaining the current value as the target.
+		 */
+		void StopFade();
+
+		/**
+		 * Updates the current value with the given delta.
+		 */
+		void Update(float InDeltaTime);
+
+	private:
+		/** Converts value to final resulting volume */
+		static float AlphaToVolume(float InAlpha, EFaderCurve InCurve);
+
+		/** Current value used to linear interpolate over update delta
+		  * (Normalized value for non-log, -80dB to 0dB for log)
+		  */
+		float Alpha;
+
+		/** Target value used to linear interpolate over update delta
+		  * (Normalized value for non-log, -80dB to 0dB for log)
+		  */
+		float Target;
+
+		/** Duration fader is to be active */
+		float ActiveDuration;
+
+		/** Duration fader is to perform fade */
+		float FadeDuration;
+
+		/** Time elapsed since fade has been initiated */
+		float Elapsed;
+
+		/** Audio fader curve to use */
+		EFaderCurve FadeCurve;
+	};
+} // namespace Audio
\ No newline at end of file
diff --git a/Engine/Source/Runtime/Engine/Classes/Components/AudioComponent.h b/Engine/Source/Runtime/Engine/Classes/Components/AudioComponent.h
--- a/Engine/Source/Runtime/Engine/Classes/Components/AudioComponent.h
+++ b/Engine/Source/Runtime/Engine/Classes/Components/AudioComponent.h
@@ -54,6 +54,28 @@ DECLARE_DYNAMIC_MULTICAST_DELEGATE_ThreeParams(FOnAudioMultiEnvelopeValue, const
 /** shadow delegate declaration for above */
 DECLARE_MULTICAST_DELEGATE_FourParams(FOnAudioMultiEnvelopeValueNative, const class UAudioComponent*, const float, const float, const int32);
 
+
+
+/** Type of fade to use when adjusting the audio component's volume. */
+UENUM(BlueprintType)
+enum class EAudioFaderCurve : uint8
+{
+	// Linear Fade
+	Linear,
+
+	// Logarithmic Fade
+	Logarithmic,
+
+	// S-Curve, Sinusoidal Fade
+	SCurve UMETA(DisplayName = "Sin (S-Curve)"),
+
+	// Equal Power, Sinusoidal Fade
+	Sin UMETA(DisplayName = "Sin (Equal Power)"),
+
+	Count UMETA(Hidden)
+};
+
+
 /**
  *	Struct used for storing one per-instance named parameter for this AudioComponent.
  *	Certain nodes in the SoundCue may reference parameters by name so they can be adjusted per-instance.
@@ -362,7 +384,7 @@ public:
 	 * @param FadeVolumeLevel the percentage of the AudioComponents's calculated volume to fade to
 	 */
 	UFUNCTION(BlueprintCallable, Category="Audio|Components|Audio")
-	virtual void FadeIn(float FadeInDuration, float FadeVolumeLevel = 1.f, float StartTime = 0.f);
+	virtual void FadeIn(float FadeInDuration, float FadeVolumeLevel = 1.0f, float StartTime = 0.0f, const EAudioFaderCurve FadeCurve = EAudioFaderCurve::Linear);
 
 	/**
 	 * This is used in place of "stop" when it is desired to fade the volume of the sound before stopping.
@@ -375,7 +397,7 @@ public:
 	 * @param FadeVolumeLevel the percentage of the AudioComponents's calculated volume in which to fade to
 	 */
 	UFUNCTION(BlueprintCallable, Category="Audio|Components|Audio")
-	virtual	void FadeOut(float FadeOutDuration, float FadeVolumeLevel);
+	virtual	void FadeOut(float FadeOutDuration, float FadeVolumeLevel, const EAudioFaderCurve FadeCurve = EAudioFaderCurve::Linear);
 
 	/** Start a sound playing on an audio component */
 	UFUNCTION(BlueprintCallable, Category="Audio|Components|Audio")
@@ -399,7 +421,7 @@ public:
 
 	/** This will allow one to adjust the volume of an AudioComponent on the fly */
 	UFUNCTION(BlueprintCallable, Category="Audio|Components|Audio")
-	void AdjustVolume(float AdjustVolumeDuration, float AdjustVolumeLevel);
+	void AdjustVolume(float AdjustVolumeDuration, float AdjustVolumeLevel, const EAudioFaderCurve FadeCurve = EAudioFaderCurve::Linear);
 
 	/**  Set a float instance parameter for use in sound cues played by this audio component */
 	UFUNCTION(BlueprintCallable, Category="Audio|Components|Audio")
@@ -533,7 +555,7 @@ public:
 	virtual bool IsReadyForOwnerToAutoDestroy() const override;
 	//~ End ActorComponent Interface.
 
-	void AdjustVolumeInternal(float AdjustVolumeDuration, float AdjustVolumeLevel, bool bIsFadeOut);
+	void AdjustVolumeInternal(float AdjustVolumeDuration, float AdjustVolumeLevel, bool bIsFadeOut, EAudioFaderCurve FadeCurve);
 
 	/** Returns a pointer to the attenuation settings to be used (if any) for this audio component dependent on the SoundAttenuation asset or overrides set. */
 	const FSoundAttenuationSettings* GetAttenuationSettingsToApply() const;
@@ -618,7 +640,7 @@ private:
 
 protected:
 	/** Utility function called by Play and FadeIn to start a sound playing. */
-	void PlayInternal(const float StartTime = 0.f, const float FadeInDuration = 0.f, const float FadeVolumeLevel = 1.f);
+	void PlayInternal(const float StartTime = 0.0f, const float FadeInDuration = 0.0f, const float FadeVolumeLevel = 1.0f, const EAudioFaderCurve FadeCurve = EAudioFaderCurve::Linear);
 
 #if WITH_EDITORONLY_DATA
 	/** Utility function that updates which texture is displayed on the sprite dependent on the properties of the Audio Component. */
@@ -631,6 +653,3 @@ protected:
 	static TMap<uint64, UAudioComponent*> AudioIDToComponentMap;
 	static FCriticalSection AudioIDToComponentMapLock;
 };
-
-
-
diff --git a/Engine/Source/Runtime/Engine/Private/ActiveSound.cpp b/Engine/Source/Runtime/Engine/Private/ActiveSound.cpp
--- a/Engine/Source/Runtime/Engine/Private/ActiveSound.cpp
+++ b/Engine/Source/Runtime/Engine/Private/ActiveSound.cpp
@@ -10,7 +10,6 @@
 #include "Sound/SoundWave.h"
 #include "Sound/SoundNodeAttenuation.h"
 #include "SubtitleManager.h"
-#include "DSP/Dsp.h"
 
 static int32 AudioOcclusionDisabledCvar = 0;
 FAutoConsoleVariableRef CVarAudioOcclusionEnabled(
@@ -80,9 +79,6 @@ FActiveSound::FActiveSound()
 	, PlaybackTimeNonVirtualized(0.0f)
 	, MinCurrentPitch(1.0f)
 	, RequestedStartTime(0.0f)
-	, CurrentAdjustVolumeMultiplier(1.0f)
-	, TargetAdjustVolumeMultiplier(1.0f)
-	, TargetAdjustVolumeStopTime(-1.0f)
 	, VolumeMultiplier(1.0f)
 	, PitchMultiplier(1.0f)
 	, LowPassFilterFrequency(MAX_FILTER_FREQUENCY)
@@ -528,11 +524,7 @@ void FActiveSound::UpdateWaveInstances(TArray<FWaveInstance*> &InWaveInstances,
 	ParseParams.Transform = Transform;
 	ParseParams.StartTime = RequestedStartTime;
 
-	// Default values.
-	// It's all Multiplicative!  So now people are all modifying the multiplier values via various means
-	// (even after the Sound has started playing, and this line takes them all into account and gives us
-	// final value that is correct
-	UpdateAdjustVolumeMultiplier(DeltaTime);
+	ComponentVolumeFader.Update(DeltaTime);
 
 	ParseParams.VolumeMultiplier = GetVolume();
 
@@ -578,7 +570,7 @@ void FActiveSound::UpdateWaveInstances(TArray<FWaveInstance*> &InWaveInstances,
 
 	// Recurse nodes, have SoundWave's create new wave instances and update bFinished unless we finished fading out.
 	bFinished = true;
-	if (FadeOut == EFadeOut::None || (PlaybackTime <= TargetAdjustVolumeStopTime))
+	if (FadeOut == EFadeOut::None || ComponentVolumeFader.IsActive())
 	{
 		bool bReverbSendLevelWasSet = false;
 		if (bHasAttenuationSettings)
@@ -935,38 +927,6 @@ void FActiveSound::RemoveWaveInstance(const UPTRINT WaveInstanceHash)
 	}
 }
 
-void FActiveSound::UpdateAdjustVolumeMultiplier(const float DeltaTime)
-{
-	// Choose min/max bound and clamp dt to prevent unwanted spikes in volume
-	float MinValue = 0.0f;
-	float MaxValue = 0.0f;
-	if (CurrentAdjustVolumeMultiplier < TargetAdjustVolumeMultiplier)
-	{
-		MinValue = CurrentAdjustVolumeMultiplier;
-		MaxValue = TargetAdjustVolumeMultiplier;
-	}
-	else
-	{
-		MinValue = TargetAdjustVolumeMultiplier;
-		MaxValue = CurrentAdjustVolumeMultiplier;
-	}
-
-	float DeltaTimeValue = FMath::Min(DeltaTime, 0.5f);
-
-	// keep stepping towards our target until we hit our stop time
-	if (PlaybackTime < TargetAdjustVolumeStopTime)
-	{
-		CurrentAdjustVolumeMultiplier += (TargetAdjustVolumeMultiplier - CurrentAdjustVolumeMultiplier) * DeltaTimeValue / (TargetAdjustVolumeStopTime - PlaybackTime);
-	}
-	else
-	{
-		CurrentAdjustVolumeMultiplier = TargetAdjustVolumeMultiplier;
-	}
-
-	// Apply final clamp
-	CurrentAdjustVolumeMultiplier = FMath::Clamp(CurrentAdjustVolumeMultiplier, MinValue, MaxValue);
-}
-
 void FActiveSound::OcclusionTraceDone(const FTraceHandle& TraceHandle, FTraceDatum& TraceDatum)
 {
 	// Look for any results that resulted in a blocking hit
@@ -1211,7 +1171,7 @@ bool FActiveSound::GetFloatParameter( const FName InName, float& OutFloat ) cons
 
 float FActiveSound::GetVolume() const
 {
-	const float Volume = VolumeMultiplier * CurrentAdjustVolumeMultiplier * GetTotalConcurrencyVolumeScale();
+	const float Volume = VolumeMultiplier * ComponentVolumeFader.GetVolume() * GetTotalConcurrencyVolumeScale();
 	return Sound ? Volume * Sound->GetVolumeMultiplier() : Volume;
 }
 
diff --git a/Engine/Source/Runtime/Engine/Private/Components/AudioComponent.cpp b/Engine/Source/Runtime/Engine/Private/Components/AudioComponent.cpp
--- a/Engine/Source/Runtime/Engine/Private/Components/AudioComponent.cpp
+++ b/Engine/Source/Runtime/Engine/Private/Components/AudioComponent.cpp
@@ -6,6 +6,7 @@
 #include "ActiveSound.h"
 #include "AudioThread.h"
 #include "AudioDevice.h"
+#include "DSP/VolumeFader.h"
 #include "Sound/SoundNodeAttenuation.h"
 #include "Sound/SoundCue.h"
 #include "Components/BillboardComponent.h"
@@ -354,7 +355,7 @@ void UAudioComponent::Play(float StartTime)
 	PlayInternal(StartTime);
 }
 
-void UAudioComponent::PlayInternal(const float StartTime, const float FadeInDuration, const float FadeVolumeLevel)
+void UAudioComponent::PlayInternal(const float StartTime, const float FadeInDuration, const float FadeVolumeLevel, const EAudioFaderCurve FadeCurve)
 {
 	SCOPE_CYCLE_COUNTER(STAT_AudioComp_Play);
 
@@ -497,19 +498,10 @@ void UAudioComponent::PlayInternal(const float StartTime, const float FadeInDura
 			}
 
 			NewActiveSound.MaxDistance = MaxDistance;
-
 			NewActiveSound.InstanceParameters = InstanceParameters;
-			NewActiveSound.TargetAdjustVolumeMultiplier = FadeVolumeLevel;
 
-			if (FadeInDuration > 0.0f)
-			{
-				NewActiveSound.CurrentAdjustVolumeMultiplier = 0.f;
-				NewActiveSound.TargetAdjustVolumeStopTime = FadeInDuration;
-			}
-			else
-			{
-				NewActiveSound.CurrentAdjustVolumeMultiplier = FadeVolumeLevel;
-			}
+			NewActiveSound.ComponentVolumeFader.SetVolume(0.0f); // Init to 0.0f to fade as default is 1.0f
+			NewActiveSound.ComponentVolumeFader.StartFade(FadeVolumeLevel, FadeInDuration, static_cast<Audio::EFaderCurve>(FadeCurve));
 
 			// Bump ActiveCount... this is used to determine if an audio component is still active after a sound reports back as completed
 			++ActiveCount;
@@ -545,24 +537,24 @@ FAudioDevice* UAudioComponent::GetAudioDevice() const
 	return AudioDevice;
 }
 
-void UAudioComponent::FadeIn( float FadeInDuration, float FadeVolumeLevel, float StartTime )
+void UAudioComponent::FadeIn(float FadeInDuration, float FadeVolumeLevel, float StartTime, const EAudioFaderCurve FadeCurve)
 {
-	PlayInternal(StartTime, FadeInDuration, FadeVolumeLevel);
+	PlayInternal(StartTime, FadeInDuration, FadeVolumeLevel, FadeCurve);
 }
 
-void UAudioComponent::FadeOut(float FadeOutDuration, float FadeVolumeLevel)
+void UAudioComponent::FadeOut(float FadeOutDuration, float FadeVolumeLevel, const EAudioFaderCurve FadeCurve)
 {
 	const bool bIsFadeOut = true;
-	AdjustVolumeInternal(FadeOutDuration, FadeVolumeLevel, bIsFadeOut);
+	AdjustVolumeInternal(FadeOutDuration, FadeVolumeLevel, bIsFadeOut, FadeCurve);
 }
 
-void UAudioComponent::AdjustVolume(float AdjustVolumeDuration, float AdjustVolumeLevel)
+void UAudioComponent::AdjustVolume(float AdjustVolumeDuration, float AdjustVolumeLevel, const EAudioFaderCurve FadeCurve)
 {
 	const bool bIsFadeOut = false;
-	AdjustVolumeInternal(AdjustVolumeDuration, AdjustVolumeLevel, bIsFadeOut);
+	AdjustVolumeInternal(AdjustVolumeDuration, AdjustVolumeLevel, bIsFadeOut, FadeCurve);
 }
 
-void UAudioComponent::AdjustVolumeInternal(float AdjustVolumeDuration, float AdjustVolumeLevel, bool bIsFadeOut)
+void UAudioComponent::AdjustVolumeInternal(float AdjustVolumeDuration, float AdjustVolumeLevel, bool bIsFadeOut, const EAudioFaderCurve FadeCurve)
 {
 	if (!bIsActive)
 	{
@@ -585,7 +577,7 @@ void UAudioComponent::AdjustVolumeInternal(float AdjustVolumeDuration, float Adj
 
 	const uint64 InAudioComponentID = AudioComponentID;
 	DECLARE_CYCLE_STAT(TEXT("FAudioThreadTask.AdjustVolume"), STAT_AudioAdjustVolume, STATGROUP_AudioThreadCommands);
-	FAudioThread::RunCommandOnAudioThread([AudioDevice, InAudioComponentID, AdjustVolumeDuration, AdjustVolumeLevel, bIsFadeOut]()
+	FAudioThread::RunCommandOnAudioThread([AudioDevice, InAudioComponentID, AdjustVolumeDuration, AdjustVolumeLevel, bIsFadeOut, FadeCurve]()
 	{
 		FActiveSound* ActiveSound = AudioDevice->FindActiveSound(InAudioComponentID);
 		if (!ActiveSound)
@@ -593,13 +585,14 @@ void UAudioComponent::AdjustVolumeInternal(float AdjustVolumeDuration, float Adj
 			return;
 		}
 
+		const float InitialTargetVolume = ActiveSound->ComponentVolumeFader.GetTargetVolume();
+
 		// Ignore fade out request if requested volume is higher than current target.
-		if (bIsFadeOut && AdjustVolumeLevel >= ActiveSound->TargetAdjustVolumeMultiplier)
+		if (bIsFadeOut && AdjustVolumeLevel >= InitialTargetVolume)
 		{
 			return;
 		}
 
-		const float NewTargetStopTime = ActiveSound->PlaybackTime + AdjustVolumeDuration;
 		if (ActiveSound->FadeOut == FActiveSound::EFadeOut::Concurrency)
 		{
 			// Ignore adjust volume request if non-zero and currently voice stealing.
@@ -609,7 +602,7 @@ void UAudioComponent::AdjustVolumeInternal(float AdjustVolumeDuration, float Adj
 			}
 
 			// Ignore request of longer fade out than active target if active is concurrency (voice stealing) fade.
-			if (NewTargetStopTime > ActiveSound->TargetAdjustVolumeStopTime)
+			if (AdjustVolumeDuration > ActiveSound->ComponentVolumeFader.GetFadeDuration())
 			{
 				return;
 			}
@@ -619,17 +612,7 @@ void UAudioComponent::AdjustVolumeInternal(float AdjustVolumeDuration, float Adj
 			ActiveSound->FadeOut = FMath::IsNearlyZero(AdjustVolumeLevel) ? FActiveSound::EFadeOut::User : FActiveSound::EFadeOut::None;
 		}
 
-		ActiveSound->TargetAdjustVolumeMultiplier = AdjustVolumeLevel;
-
-		if (AdjustVolumeDuration > 0.0f)
-		{
-			ActiveSound->TargetAdjustVolumeStopTime = NewTargetStopTime;
-		}
-		else
-		{
-			ActiveSound->CurrentAdjustVolumeMultiplier = AdjustVolumeLevel;
-			ActiveSound->TargetAdjustVolumeStopTime = -1.0f;
-		}
+		ActiveSound->ComponentVolumeFader.StartFade(AdjustVolumeLevel, AdjustVolumeDuration, static_cast<Audio::EFaderCurve>(FadeCurve));
 	}, GET_STATID(STAT_AudioAdjustVolume));
 }
 
@@ -687,21 +670,6 @@ void UAudioComponent::StopDelayed(float DelayTime)
 			return;
 		}
 
-		const float NewTargetStopTime = ActiveSound->PlaybackTime + DelayTime;
-		if (ActiveSound->FadeOut == FActiveSound::EFadeOut::Concurrency)
-		{
-			// Ignore request of longer fade out than active target if active is concurrency (voice stealing) fade.
-			if (NewTargetStopTime > ActiveSound->TargetAdjustVolumeStopTime)
-			{
-				return;
-			}
-		}
-		else
-		{
-			// Set fade to user, but don't adjust target volume, which will cause sound to stop abruptly as intended.
-			ActiveSound->FadeOut = FActiveSound::EFadeOut::User;
-		}
-
 		if (const USoundBase* StoppingSound = ActiveSound->GetSound())
 		{
 			UE_LOG(LogAudio, Verbose, TEXT("%g: Delayed Stop requested for sound '%s'"),
@@ -709,7 +677,27 @@ void UAudioComponent::StopDelayed(float DelayTime)
 				*StoppingSound->GetName());
 		}
 
-		ActiveSound->TargetAdjustVolumeStopTime = NewTargetStopTime;
+		switch (ActiveSound->FadeOut)
+		{
+			case FActiveSound::EFadeOut::Concurrency:
+			{
+				// Ignore request of longer fade out than active target if active is concurrency (voice stealing) fade.
+				if (DelayTime < ActiveSound->ComponentVolumeFader.GetFadeDuration())
+				{
+					ActiveSound->ComponentVolumeFader.SetActiveDuration(DelayTime);
+				}
+			}
+			break;
+			
+			case FActiveSound::EFadeOut::User:
+			case FActiveSound::EFadeOut::None:
+			default:
+			{
+				ActiveSound->FadeOut = FActiveSound::EFadeOut::User;
+				ActiveSound->ComponentVolumeFader.SetActiveDuration(DelayTime);
+			}
+			break;
+		}
 	}, GET_STATID(STAT_AudioStopDelayed));
 }
 
diff --git a/Engine/Source/Runtime/Engine/Private/Interpolation.cpp b/Engine/Source/Runtime/Engine/Private/Interpolation.cpp
--- a/Engine/Source/Runtime/Engine/Private/Interpolation.cpp
+++ b/Engine/Source/Runtime/Engine/Private/Interpolation.cpp
@@ -8202,7 +8202,7 @@ void UInterpTrackSound::PreviewUpdateTrack(float NewPosition, UInterpTrackInst*
 			Component->Play(NewPosition - SoundTrackKey.Time);
 
 			const float ScrubDuration = 0.1f;
-			Component->FadeOut(ScrubDuration, 1.0f);
+			Component->FadeOut(ScrubDuration, 1.0f, EAudioFaderCurve::Logarithmic);
 		}
 	}
 #endif
diff --git a/Engine/Source/Runtime/Engine/Private/SoundConcurrency.cpp b/Engine/Source/Runtime/Engine/Private/SoundConcurrency.cpp
--- a/Engine/Source/Runtime/Engine/Private/SoundConcurrency.cpp
+++ b/Engine/Source/Runtime/Engine/Private/SoundConcurrency.cpp
@@ -5,6 +5,7 @@
 #include "ActiveSound.h"
 #include "AudioDevice.h"
 #include "AudioVirtualLoop.h"
+#include "DSP/VolumeFader.h"
 #include "Sound/SoundBase.h"
 
 
@@ -826,8 +827,7 @@ void FSoundConcurrencyManager::StopDueToVoiceStealing(FActiveSound& ActiveSound)
 	if (bRequiresConcurrencyFade)
 	{
 		ActiveSound.FadeOut = FActiveSound::EFadeOut::Concurrency;
-		ActiveSound.TargetAdjustVolumeMultiplier = 0.0f;
-		ActiveSound.TargetAdjustVolumeStopTime = ActiveSound.PlaybackTime + FadeOutDuration;
+		ActiveSound.ComponentVolumeFader.StartFade(0.0f, FadeOutDuration, Audio::EFaderCurve::Logarithmic);
 	}
 }
 
diff --git a/Engine/Source/Runtime/Engine/Private/SoundWave.cpp b/Engine/Source/Runtime/Engine/Private/SoundWave.cpp
--- a/Engine/Source/Runtime/Engine/Private/SoundWave.cpp
+++ b/Engine/Source/Runtime/Engine/Private/SoundWave.cpp
@@ -1726,14 +1726,13 @@ void USoundWave::Parse(FAudioDevice* AudioDevice, const UPTRINT NodeWaveInstance
 
 	WaveInstance->bIsAmbisonics = bIsAmbisonics;
 
-	// Recompute the virtualizability here even though we did it up-front in the active sound parse.
-	// This is because an active sound can generate multiple sound waves, not all of them are necessarily virtualizable.
-	bool bHasSubtitles = ActiveSound.bHandleSubtitles && (ActiveSound.bHasExternalSubtitles || (Subtitles.Num() > 0));
+	const float WaveInstanceVolume = WaveInstance->GetVolumeWithDistanceAttenuation() * WaveInstance->GetDynamicVolume();
 
 	// When the BypassVirtualizeWhenSilent cvar is enabled, we should only honor bVirtualizeWhenSilent for procedural sounds:
+	const bool bHasSubtitles = ActiveSound.bHandleSubtitles && (ActiveSound.bHasExternalSubtitles || (Subtitles.Num() > 0));
+	const bool bStarted = WaveInstanceVolume > KINDA_SMALL_NUMBER || ActiveSound.ComponentVolumeFader.IsFadingIn();
 	const bool bCanPlayWhenSilent = IsPlayWhenSilent() && (!BypassPlayWhenSilentCVar || bProcedural);
-	const float WaveInstanceVolume = WaveInstance->GetVolumeWithDistanceAttenuation() * WaveInstance->GetDynamicVolume();
-	if (WaveInstanceVolume > KINDA_SMALL_NUMBER || (bCanPlayWhenSilent || bHasSubtitles))
+	if (bStarted || bCanPlayWhenSilent || bHasSubtitles)
 	{
 		WaveInstances.Add(WaveInstance);
 		ActiveSound.bFinished = false;
diff --git a/Engine/Source/Runtime/Engine/Public/ActiveSound.h b/Engine/Source/Runtime/Engine/Public/ActiveSound.h
--- a/Engine/Source/Runtime/Engine/Public/ActiveSound.h
+++ b/Engine/Source/Runtime/Engine/Public/ActiveSound.h
@@ -9,9 +9,10 @@
 #include "Audio.h"
 #include "Audio/AudioDebug.h"
 #include "AudioDynamicParameter.h"
-#include "Sound/SoundConcurrency.h"
 #include "Components/AudioComponent.h"
+#include "DSP/VolumeFader.h"
 #include "Sound/AudioVolume.h"
+#include "Sound/SoundConcurrency.h"
 #include "Sound/SoundSubmix.h"
 #include "Sound/SoundSourceBus.h"
 
@@ -471,16 +472,15 @@ public:
 	float MinCurrentPitch;
 	float RequestedStartTime;
 
-	float CurrentAdjustVolumeMultiplier;
-	float TargetAdjustVolumeMultiplier;
-	float TargetAdjustVolumeStopTime;
-
 	float VolumeMultiplier;
 	float PitchMultiplier;
 
 	/** The low-pass filter frequency to apply if bEnableLowPassFilter is true. */
 	float LowPassFilterFrequency;
 
+	/** Fader that tracks component volume */
+	Audio::FVolumeFader ComponentVolumeFader;
+
 	/** The interpolated parameter for the low-pass frequency due to occlusion. */
 	FDynamicParameter CurrentOcclusionFilterFrequency;
 
@@ -696,9 +696,6 @@ private:
 	/** Called when an active sound has been stopped but needs to update it's stopping sounds. Returns true when stopping sources have finished stopping. */
 	bool UpdateStoppingSources(uint64 CurrentTick, bool bEnsureStopped);
 
-	/** Sets the target volume multiplier to achieve over the specified time period */
-	void UpdateAdjustVolumeMultiplier(const float DeltaTime);
-
 	/** Updates ramping concurrency volume scalars */
 	void UpdateConcurrencyVolumeScalars(const float DeltaTime);
 
