Fix for crash with referenced sound waves

diff --git a/Engine/Source/Runtime/Engine/Private/AudioDevice.cpp b/Engine/Source/Runtime/Engine/Private/AudioDevice.cpp
--- a/Engine/Source/Runtime/Engine/Private/AudioDevice.cpp
+++ b/Engine/Source/Runtime/Engine/Private/AudioDevice.cpp
@@ -743,6 +743,9 @@ void FAudioDevice::AddReferencedObjects(FReferenceCollector& Collector)
 		Pair.Key->AddReferencedObjects(Collector);
 	}
 
+	// Make sure our referenced sound waves are up-to-date
+	UpdateReferencedSoundWaves();
+
 	// Make sure we don't try to delete any sound waves which may have in-flight decodes
 	Collector.AddReferencedObjects(ReferencedSoundWaves);
 
@@ -3838,25 +3841,37 @@ void FAudioDevice::StartSources(TArray<FWaveInstance*>& WaveInstances, int32 Fir
 	// Run a command to make sure we add the starting sounds to the referenced sound waves list
 	if (StartingSoundWaves.Num() > 0)
 	{
-		FAudioThread::RunCommandOnGameThread([this, StartingSoundWaves]()
+		FScopeLock ReferencedSoundWaveLock(&ReferencedSoundWaveCritSec);
+
+		for (USoundWave* SoundWave : StartingSoundWaves)
 		{
-			if (GEngine)
-			{
-				for (USoundWave* SoundWave : StartingSoundWaves)
-				{
-					ReferencedSoundWaves.AddUnique(SoundWave);
-				}
-			}
-		}, GET_STATID(STAT_AudioAddReferencedSoundWaves));
+			ReferencedSoundWaves_AudioThread.AddUnique(SoundWave);
+		}
+	}
+}
+
+void FAudioDevice::UpdateReferencedSoundWaves()
+{
+	FScopeLock ReferencedSoundWaveLock(&ReferencedSoundWaveCritSec);
+
+	for (USoundWave* SoundWave : ReferencedSoundWaves_AudioThread)
+	{
+		ReferencedSoundWaves.AddUnique(SoundWave);
 	}
+
+	ReferencedSoundWaves_AudioThread.Reset();
 }
 
+
 void FAudioDevice::Update(bool bGameTicking)
 {
 	LLM_SCOPE(ELLMTag::AudioMisc);
 
 	if (IsInGameThread())
 	{
+		// Make sure our referenced sound waves is up-to-date
+		UpdateReferencedSoundWaves();
+
 		// On game thread, look through registered sound waves and remove if we finished precaching (and audio decompressor is cleaned up)
 		// ReferencedSoundWaves is used to make sure GC doesn't run on any sound waves that are actively pre-caching within an async task.
 		// Sounds may be loaded, kick off an async task to decompress, but never actually try to play, so GC can reclaim these while precaches are in-flight.
diff --git a/Engine/Source/Runtime/Engine/Public/AudioDevice.h b/Engine/Source/Runtime/Engine/Public/AudioDevice.h
--- a/Engine/Source/Runtime/Engine/Public/AudioDevice.h
+++ b/Engine/Source/Runtime/Engine/Public/AudioDevice.h
@@ -1827,7 +1827,13 @@ private:
 	TArray<FActiveSound*> ActiveSounds;
 	/** Array of sound waves to add references to avoid GC until guaranteed to be done with precache or decodes. */
 	TArray<USoundWave*> ReferencedSoundWaves;
+
+	void UpdateReferencedSoundWaves();
+	TArray<USoundWave*> ReferencedSoundWaves_AudioThread;
+	FCriticalSection ReferencedSoundWaveCritSec;
+
 	TArray<USoundWave*> PrecachingSoundWaves;
+
 	TArray<FWaveInstance*> ActiveWaveInstances;
 
 	/** Array of dormant loops stopped due to proximity/applicable concurrency rules
