Cooking fix in StaticMeshComponent : OwningComponent can be NULL for a component created via SpawnActor off of a blueprint default (LODData will be created without a call to SetLODDataCount). When cooking, update LODData's OwningComponent.

diff --git a/Engine/Source/Runtime/Engine/Private/Components/StaticMeshComponent.cpp b/Engine/Source/Runtime/Engine/Private/Components/StaticMeshComponent.cpp
--- a/Engine/Source/Runtime/Engine/Private/Components/StaticMeshComponent.cpp
+++ b/Engine/Source/Runtime/Engine/Private/Components/StaticMeshComponent.cpp
@@ -288,6 +288,17 @@ void UStaticMeshComponent::Serialize(FArchive& Ar)
 	}
 #endif
 
+#if WITH_EDITORONLY_DATA
+	if (Ar.IsCooking())
+	{
+		// LODData's OwningComponent can be NULL for a component created via SpawnActor off of a blueprint default (LODData will be created without a call to SetLODDataCount)
+		for (int32 LODIndex = 0; LODIndex < LODData.Num(); LODIndex++)
+		{
+			LODData[LODIndex].OwningComponent = this;
+		}
+	}
+#endif
+
 	Ar << LODData;
 
 #if WITH_EDITOR
@@ -2697,7 +2708,7 @@ FArchive& operator<<(FArchive& Ar,FStaticMeshComponentLODInfo& I)
 	const uint8 OverrideColorsStripFlag = 1;
 	bool bStrippedOverrideColors = false;
 #if WITH_EDITORONLY_DATA
-	if( Ar.IsCooking() && I.OwningComponent )
+	if( Ar.IsCooking() )
 	{
 		// Check if override color should be stripped too	
 		int32 LODIndex = 0;
