Change some virtual texture properties in PrimitiveComponent from int32 -> int8

diff --git a/Engine/Source/Runtime/Engine/Classes/Components/PrimitiveComponent.h b/Engine/Source/Runtime/Engine/Classes/Components/PrimitiveComponent.h
--- a/Engine/Source/Runtime/Engine/Classes/Components/PrimitiveComponent.h
+++ b/Engine/Source/Runtime/Engine/Classes/Components/PrimitiveComponent.h
@@ -571,23 +571,23 @@ public:
 	TArray<URuntimeVirtualTexture*> RuntimeVirtualTextures;
 
 	/** Bias to the LOD selected for rendering to runtime virtual textures. */
-	UPROPERTY(EditAnywhere, AdvancedDisplay, BlueprintReadOnly, Category = VirtualTexture, meta = (DisplayName = "Virtual Texture LOD Bias", UIMin = "0", UIMax = "7"))
-	int32 VirtualTextureLodBias = 0;
+	UPROPERTY(EditAnywhere, AdvancedDisplay, Category = VirtualTexture, meta = (DisplayName = "Virtual Texture LOD Bias", UIMin = "0", UIMax = "7"))
+	int8 VirtualTextureLodBias = 0;
 
 	/**
 	 * Number of lower mips in the runtime virtual texture to skip for rendering this primitive.
 	 * Larger values reduce the effective draw distance in the runtime virtual texture.
 	 * This culling method doesn't take into account primitive size or virtual texture size.
 	 */
-	UPROPERTY(EditAnywhere, AdvancedDisplay, BlueprintReadOnly, Category = VirtualTexture, meta = (DisplayName = "Virtual Texture Skip Mips", UIMin = "0", UIMax = "15"))
-	int32 VirtualTextureCullMips = 0;
+	UPROPERTY(EditAnywhere, AdvancedDisplay, Category = VirtualTexture, meta = (DisplayName = "Virtual Texture Skip Mips", UIMin = "0", UIMax = "15"))
+	int8 VirtualTextureCullMips = 0;
 
 	/**
 	 * Set the minimum pixel coverage before culling from the runtime virtual texture.
 	 * Larger values reduce the effective draw distance in the runtime virtual texture.
 	 */
-	UPROPERTY(EditAnywhere, AdvancedDisplay, BlueprintReadOnly, Category = VirtualTexture, meta = (UIMin = "0", UIMax = "7"))
-	int32 VirtualTextureMinCoverage = 0;
+	UPROPERTY(EditAnywhere, AdvancedDisplay, Category = VirtualTexture, meta = (UIMin = "0", UIMax = "7"))
+	int8 VirtualTextureMinCoverage = 0;
 
 	/** Render to the main pass based on the virtual texture settings. */
 	UPROPERTY(EditAnywhere, BlueprintReadOnly, Category = VirtualTexture, meta = (DisplayName = "Virtual Texture Pass Type"))
diff --git a/Engine/Source/Runtime/Engine/Public/PrimitiveSceneProxy.h b/Engine/Source/Runtime/Engine/Public/PrimitiveSceneProxy.h
--- a/Engine/Source/Runtime/Engine/Public/PrimitiveSceneProxy.h
+++ b/Engine/Source/Runtime/Engine/Public/PrimitiveSceneProxy.h
@@ -24,7 +24,7 @@ class FStaticPrimitiveDrawInterface;
 class UPrimitiveComponent;
 class URuntimeVirtualTexture;
 class UTexture2D;
-enum class ERuntimeVirtualTextureMaterialType;
+enum class ERuntimeVirtualTextureMaterialType : uint8;
 struct FMeshBatch;
 class FColorVertexBuffer;
 
diff --git a/Engine/Source/Runtime/Engine/Public/VT/RuntimeVirtualTextureEnum.h b/Engine/Source/Runtime/Engine/Public/VT/RuntimeVirtualTextureEnum.h
--- a/Engine/Source/Runtime/Engine/Public/VT/RuntimeVirtualTextureEnum.h
+++ b/Engine/Source/Runtime/Engine/Public/VT/RuntimeVirtualTextureEnum.h
@@ -11,7 +11,7 @@
  * This "fixed function" approach will probably break down if we end up needing to support some complex set of attribute combinations but it is OK to begin with.
  */
 UENUM()
-enum class ERuntimeVirtualTextureMaterialType
+enum class ERuntimeVirtualTextureMaterialType : uint8
 {
 	BaseColor UMETA(DisplayName = "Base Color"),
 	BaseColor_Normal UMETA(DisplayName = "Base Color, Normal"),
@@ -26,7 +26,7 @@ static_assert((uint32)ERuntimeVirtualTextureMaterialType::Count <= (1 << (uint32
 
 /** Enumeration of main pass behaviors when rendering to a runtime virtual texture. */
 UENUM()
-enum class ERuntimeVirtualTextureMainPassType
+enum class ERuntimeVirtualTextureMainPassType : uint8
 {
 	/** If there is no valid virtual texture target we will not render at all. Use this for items that we don't mind removing if there is no virtual texture support. */
 	Never UMETA(DisplayName = "Virtual Texture Only"),
diff --git a/Engine/Source/Runtime/Landscape/Public/LandscapeRender.h b/Engine/Source/Runtime/Landscape/Public/LandscapeRender.h
--- a/Engine/Source/Runtime/Landscape/Public/LandscapeRender.h
+++ b/Engine/Source/Runtime/Landscape/Public/LandscapeRender.h
@@ -34,7 +34,7 @@ LandscapeRender.h: New terrain rendering
 #define LANDSCAPE_MAX_SUBSECTION_NUM 2
 
 class FLandscapeComponentSceneProxy;
-enum class ERuntimeVirtualTextureMaterialType;
+enum class ERuntimeVirtualTextureMaterialType : uint8;
 
 #if WITH_EDITOR
 namespace ELandscapeViewMode
diff --git a/Engine/Source/Runtime/Renderer/Private/VT/RuntimeVirtualTextureProducer.h b/Engine/Source/Runtime/Renderer/Private/VT/RuntimeVirtualTextureProducer.h
--- a/Engine/Source/Runtime/Renderer/Private/VT/RuntimeVirtualTextureProducer.h
+++ b/Engine/Source/Runtime/Renderer/Private/VT/RuntimeVirtualTextureProducer.h
@@ -6,7 +6,7 @@
 #include "UObject/ObjectMacros.h"
 #include "VirtualTexturing.h"
 
-enum class ERuntimeVirtualTextureMaterialType;
+enum class ERuntimeVirtualTextureMaterialType : uint8;
 class FRHITexture2D;
 class FSceneInterface;
 
diff --git a/Engine/Source/Runtime/Renderer/Private/VT/RuntimeVirtualTextureRender.h b/Engine/Source/Runtime/Renderer/Private/VT/RuntimeVirtualTextureRender.h
--- a/Engine/Source/Runtime/Renderer/Private/VT/RuntimeVirtualTextureRender.h
+++ b/Engine/Source/Runtime/Renderer/Private/VT/RuntimeVirtualTextureRender.h
@@ -4,7 +4,7 @@
 
 #include "CoreMinimal.h"
 
-enum class ERuntimeVirtualTextureMaterialType;
+enum class ERuntimeVirtualTextureMaterialType : uint8;
 class FRHICommandListImmediate;
 class FRHITexture2D;
 class FScene;
