Allow a UTexture to override the default virtual texture build settings

diff --git a/Engine/Source/Runtime/Engine/Classes/Engine/Texture.h b/Engine/Source/Runtime/Engine/Classes/Engine/Texture.h
--- a/Engine/Source/Runtime/Engine/Classes/Engine/Texture.h
+++ b/Engine/Source/Runtime/Engine/Classes/Engine/Texture.h
@@ -793,6 +793,9 @@ public:
 		return false;
 	}
 
+	/** Returns the virtual texture build settings. */
+	ENGINE_API virtual void GetVirtualTextureBuildSettings(struct FVirtualTextureBuildSettings& OutSettings) const;
+
 	/**
 	 * Textures that use the derived data cache must override this function and
 	 * provide a pointer to the linked list of platform data.
diff --git a/Engine/Source/Runtime/Engine/Private/TextureDerivedData.cpp b/Engine/Source/Runtime/Engine/Private/TextureDerivedData.cpp
--- a/Engine/Source/Runtime/Engine/Private/TextureDerivedData.cpp
+++ b/Engine/Source/Runtime/Engine/Private/TextureDerivedData.cpp
@@ -23,6 +23,7 @@
 #include "TextureDerivedDataTask.h"
 #include "Streaming/TextureStreamingHelpers.h"
 #include "Engine/VolumeTexture.h"
+#include "VT/VirtualTextureBuildSettings.h"
 #include "VT/VirtualTextureBuiltData.h"
 
 #if WITH_EDITOR
@@ -35,30 +36,6 @@
 #include "ProfilingDebugging/CookStats.h"
 #include "VT/VirtualTextureDataBuilder.h"
 
-static TAutoConsoleVariable<int32> CVarVTCompressCrunch(
-	TEXT("r.VT.EnableCompressCrunch"),
-	0,
-	TEXT("Enable Crunch compression for virtual textures, for supported formats")
-);
-
-static TAutoConsoleVariable<int32> CVarVTCompressZlib(
-	TEXT("r.VT.EnableCompressZlib"),
-	1,
-	TEXT("Enables Zlib compression for virtual textures, if no compression is enabled/supported")
-);
-
-static TAutoConsoleVariable<int32> CVarVTTileSize(
-	TEXT("r.VT.TileSize"),
-	128,
-	TEXT("Size in pixels to use for virtual texture tiles (rounded to next power-of-2)")
-);
-
-static TAutoConsoleVariable<int32> CVarVTTileBorderSize(
-	TEXT("r.VT.TileBorderSize"),
-	4,
-	TEXT("Size in pixels to use for virtual texture tiles borders (rounded to next power-of-2)")
-);
-
 /*------------------------------------------------------------------------------
 	Versioning for texture derived data.
 ------------------------------------------------------------------------------*/
@@ -444,22 +421,24 @@ static void GetTextureBuildSettings(
 	if (OutBuildSettings.bVirtualStreamable)
 	{
 		const UTexture2D *Texture2D = Cast<UTexture2D>(&Texture);
-		checkf(Texture2D, TEXT("Virtual texturing is only supported on 2D textures")); 
+		checkf(Texture2D, TEXT("Virtual texturing is only supported on 2D textures"));
 		OutBuildSettings.VirtualAddressingModeX = Texture2D->AddressX;
 		OutBuildSettings.VirtualAddressingModeY = Texture2D->AddressY;
-		OutBuildSettings.bVirtualTextureEnableCompressZlib = CVarVTCompressZlib.GetValueOnAnyThread() != 0;
-		OutBuildSettings.bVirtualTextureEnableCompressCrunch = CVarVTCompressCrunch.GetValueOnAnyThread() != 0;
-		OutBuildSettings.VirtualTextureTileSize = FMath::RoundUpToPowerOfTwo(CVarVTTileSize.GetValueOnAnyThread());
+
+		FVirtualTextureBuildSettings VirtualTextureBuildSettings;
+		Texture.GetVirtualTextureBuildSettings(VirtualTextureBuildSettings);
+		OutBuildSettings.bVirtualTextureEnableCompressZlib = VirtualTextureBuildSettings.bEnableCompressZlib;
+		OutBuildSettings.bVirtualTextureEnableCompressCrunch = VirtualTextureBuildSettings.bEnableCompressCrunch;
+		OutBuildSettings.VirtualTextureTileSize = FMath::RoundUpToPowerOfTwo(VirtualTextureBuildSettings.TileSize);
 
 		// don't all max resolution to be less than VT tile size
 		OutBuildSettings.MaxTextureResolution = FMath::Max<uint32>(OutBuildSettings.MaxTextureResolution, OutBuildSettings.VirtualTextureTileSize);
-		
+
 		// 0 is a valid value for border size
 		// 1 would be OK in some cases, but breaks BC compressed formats, since it will result in physical tiles that aren't divisible by block size (4)
 		// Could allow border size of 1 for non BC compressed virtual textures, but somewhat complicated to get that correct, especially with multiple layers
 		// Doesn't seem worth the complexity for now, so clamp the size to be at least 2
-		const int32 TileBorderSize = CVarVTTileBorderSize.GetValueOnAnyThread();
-		OutBuildSettings.VirtualTextureBorderSize = (TileBorderSize > 0) ? FMath::RoundUpToPowerOfTwo(FMath::Max(TileBorderSize, 2)) : 0;
+		OutBuildSettings.VirtualTextureBorderSize = (VirtualTextureBuildSettings.TileBorderSize > 0) ? FMath::RoundUpToPowerOfTwo(FMath::Max(VirtualTextureBuildSettings.TileBorderSize, 2)) : 0;
 	}
 	else
 	{
@@ -1904,6 +1883,11 @@ void UTexture::MarkPlatformDataTransient()
 }
 #endif // #if WITH_EDITOR
 
+void UTexture::GetVirtualTextureBuildSettings(FVirtualTextureBuildSettings& OutSettings) const
+{
+	OutSettings.Init();
+}
+
 void UTexture::CleanupCachedRunningPlatformData()
 {
 	FTexturePlatformData **RunningPlatformDataPtr = GetRunningPlatformData();
diff --git a/Engine/Source/Runtime/Engine/Private/VT/VirtualTextureBuildSettings.cpp b/Engine/Source/Runtime/Engine/Private/VT/VirtualTextureBuildSettings.cpp
new file mode 100644
--- /dev/null
+++ b/Engine/Source/Runtime/Engine/Private/VT/VirtualTextureBuildSettings.cpp
@@ -0,0 +1,35 @@
+// Copyright 1998-2019 Epic Games, Inc. All Rights Reserved.
+
+#include "VT/VirtualTextureBuildSettings.h"
+
+static TAutoConsoleVariable<int32> CVarVTCompressCrunch(
+	TEXT("r.VT.EnableCompressCrunch"),
+	0,
+	TEXT("Enable Crunch compression for virtual textures, for supported formats")
+);
+
+static TAutoConsoleVariable<int32> CVarVTCompressZlib(
+	TEXT("r.VT.EnableCompressZlib"),
+	1,
+	TEXT("Enables Zlib compression for virtual textures, if no compression is enabled/supported")
+);
+
+static TAutoConsoleVariable<int32> CVarVTTileSize(
+	TEXT("r.VT.TileSize"),
+	128,
+	TEXT("Size in pixels to use for virtual texture tiles (rounded to next power-of-2)")
+);
+
+static TAutoConsoleVariable<int32> CVarVTTileBorderSize(
+	TEXT("r.VT.TileBorderSize"),
+	4,
+	TEXT("Size in pixels to use for virtual texture tiles borders (rounded to next power-of-2)")
+);
+
+void FVirtualTextureBuildSettings::Init()
+{
+	TileSize = CVarVTTileSize.GetValueOnAnyThread();
+	TileBorderSize = CVarVTTileBorderSize.GetValueOnAnyThread();;
+	bEnableCompressCrunch = CVarVTCompressCrunch.GetValueOnAnyThread() != 0;
+	bEnableCompressZlib = CVarVTCompressZlib.GetValueOnAnyThread() != 0;
+}
diff --git a/Engine/Source/Runtime/Engine/Private/VT/VirtualTextureBuildSettings.h b/Engine/Source/Runtime/Engine/Private/VT/VirtualTextureBuildSettings.h
new file mode 100644
--- /dev/null
+++ b/Engine/Source/Runtime/Engine/Private/VT/VirtualTextureBuildSettings.h
@@ -0,0 +1,28 @@
+// Copyright 1998-2019 Epic Games, Inc. All Rights Reserved.
+
+#pragma once
+
+#include "CoreMinimal.h"
+#include "VirtualTextureBuildSettings.generated.h"
+
+/** Build settings used for virtual textures. */
+USTRUCT()
+struct FVirtualTextureBuildSettings
+{
+	GENERATED_USTRUCT_BODY()
+
+	UPROPERTY()
+	int32 TileSize;
+
+	UPROPERTY()
+	int32 TileBorderSize;
+
+	UPROPERTY()
+	bool bEnableCompressCrunch;
+
+	UPROPERTY()
+	bool bEnableCompressZlib;
+
+	/** Initialize with default build settings. These are defined by the current project setup. */
+	ENGINE_API void Init();
+};
