Add some progress bar support to virtual texture building
Fixed up some includes

diff --git a/Engine/Source/Runtime/Engine/Private/VT/VirtualTextureDataBuilder.cpp b/Engine/Source/Runtime/Engine/Private/VT/VirtualTextureDataBuilder.cpp
--- a/Engine/Source/Runtime/Engine/Private/VT/VirtualTextureDataBuilder.cpp
+++ b/Engine/Source/Runtime/Engine/Private/VT/VirtualTextureDataBuilder.cpp
@@ -9,7 +9,9 @@
 #include "CrunchCompression.h"
 #include "Async/TaskGraphInterfaces.h"
 #include "Async/ParallelFor.h"
+#include "Misc/ScopedSlowTask.h"
 #include "TextureCompressorModule.h"
+#include "TextureDerivedDataTask.h"
 
 // Debugging aid to dump tiles to disc as png files
 #define SAVE_TILES 0
@@ -394,6 +396,8 @@ void FVirtualTextureDataBuilder::BuildPagesMacroBlocks(bool bAllowAsync)
 		MipHeightInTiles = FMath::DivideAndRoundUp(MipHeightInTiles, 2u);
 	}
 
+	FScopedSlowTask BuildTask(NumTiles);
+
 	//
 	TArray<FVTSourceTileEntry> TilesInChunk;
 	TilesInChunk.Reserve(NumTiles);
@@ -431,6 +435,8 @@ void FVirtualTextureDataBuilder::BuildPagesMacroBlocks(bool bAllowAsync)
 					const int32 BlockIndex = FindSourceBlockIndex(Mip, BlockX, BlockY);
 					if (BlockIndex != INDEX_NONE)
 					{
+						BuildTask.EnterProgressFrame();
+
 						const FTextureSourceBlockData& Block = SourceBlocks[BlockIndex];
 						FVTSourceTileEntry* TileEntry = new(TilesInChunk) FVTSourceTileEntry;
 						TileEntry->BlockIndex = BlockIndex;
diff --git a/Engine/Source/Runtime/Engine/Private/VT/VirtualTextureDataBuilder.h b/Engine/Source/Runtime/Engine/Private/VT/VirtualTextureDataBuilder.h
--- a/Engine/Source/Runtime/Engine/Private/VT/VirtualTextureDataBuilder.h
+++ b/Engine/Source/Runtime/Engine/Private/VT/VirtualTextureDataBuilder.h
@@ -3,15 +3,14 @@
 #pragma once
 #include "CoreMinimal.h"
 #include "Engine/Texture.h"
-#include "TextureCompressorModule.h"
-#include "VirtualTextureBuiltData.h"
-#include "TextureDerivedDataTask.h"
 #include "ImageCore.h"
+#include "VT/VirtualTextureBuiltData.h"
 
 struct FImage;
 class ITextureCompressorModule;
 class IImageWrapperModule;
 struct FTextureSourceData;
+struct FTextureSourceBlockData;
 
 struct FVTSourceTileEntry
 {
