Do not create a new guid or FNames for timeline templates when we're going to be serializing them in.
#rnx

diff --git a/Engine/Source/Runtime/Engine/Classes/Engine/TimelineTemplate.h b/Engine/Source/Runtime/Engine/Classes/Engine/TimelineTemplate.h
--- a/Engine/Source/Runtime/Engine/Classes/Engine/TimelineTemplate.h
+++ b/Engine/Source/Runtime/Engine/Classes/Engine/TimelineTemplate.h
@@ -249,6 +249,7 @@ class UTimelineTemplate : public UObject
 	virtual void PostLoad() override;
 	virtual bool Rename(const TCHAR* InName, UObject* NewOuter, ERenameFlags Flags) override;
 	virtual void Serialize(FArchive& Ar) override;
+	virtual void PostInitProperties() override;
 	//~ End UObject Interface
 
 private:
diff --git a/Engine/Source/Runtime/Engine/Private/TimelineTemplate.cpp b/Engine/Source/Runtime/Engine/Private/TimelineTemplate.cpp
--- a/Engine/Source/Runtime/Engine/Private/TimelineTemplate.cpp
+++ b/Engine/Source/Runtime/Engine/Private/TimelineTemplate.cpp
@@ -37,10 +37,18 @@ UTimelineTemplate::UTimelineTemplate(const FObjectInitializer& ObjectInitializer
 	: Super(ObjectInitializer)
 {
 	TimelineLength = 5.0f;
-	TimelineGuid = FGuid::NewGuid();
 	bReplicated = false;
+}
 
-	UpdateCachedNames();
+void UTimelineTemplate::PostInitProperties()
+{
+	Super::PostInitProperties();
+
+	if (HasAnyFlags(RF_ClassDefaultObject | RF_NeedLoad) == false)
+	{
+		TimelineGuid = FGuid::NewGuid();
+		UpdateCachedNames();
+	}
 }
 
 const FString UTimelineTemplate::TemplatePostfix(TEXT("_Template"));
