Avoid unnecessary string copies while dealing with collision profiles in cooked builds.
Pack the collision profile structures more efficiently

diff --git a/Engine/Source/Runtime/Engine/Classes/Engine/CollisionProfile.h b/Engine/Source/Runtime/Engine/Classes/Engine/CollisionProfile.h
--- a/Engine/Source/Runtime/Engine/Classes/Engine/CollisionProfile.h
+++ b/Engine/Source/Runtime/Engine/Classes/Engine/CollisionProfile.h
@@ -54,6 +54,12 @@ struct ENGINE_API FCollisionResponseTemplate
 	// no property anymore
 	TEnumAsByte<enum ECollisionChannel> ObjectType;
 
+	UPROPERTY()
+	bool	bCanModify;
+
+	/** This is result of ResponseToChannel after loaded - please note that it is not property serializable **/
+	struct FCollisionResponseContainer ResponseToChannels;
+
 	UPROPERTY()
 	FName ObjectTypeName;
 
@@ -61,16 +67,11 @@ struct ENGINE_API FCollisionResponseTemplate
 	UPROPERTY()
 	TArray<FResponseChannel>	CustomResponses;
 
+#if WITH_EDITORONLY_DATA
 	/** Help message for collision profile **/
 	UPROPERTY()
 	FString HelpMessage;
-
-	/** Help message for collision profile **/
-	UPROPERTY()
-	bool	bCanModify;
-
-	/** This is result of ResponseToChannel after loaded - please note that it is not property serializable **/
-	struct FCollisionResponseContainer ResponseToChannels;
+#endif
 
 	/** This constructor */
 	FCollisionResponseTemplate();
@@ -95,10 +96,6 @@ struct FCustomChannelSetup
 	UPROPERTY()
 	TEnumAsByte<enum ECollisionChannel> Channel;
 
-	/** Name of channel you'd like to show up **/
-	UPROPERTY()
-	FName Name;
-
 	/** Default Response for the channel */
 	UPROPERTY()
 	TEnumAsByte<enum ECollisionResponse> DefaultResponse;
@@ -111,6 +108,10 @@ struct FCustomChannelSetup
 	UPROPERTY()
 	bool bStaticObject;	
 
+	/** Name of channel you'd like to show up **/
+	UPROPERTY()
+	FName Name;
+
 	FCustomChannelSetup()
 		: Channel(ECC_WorldStatic)
 		, DefaultResponse(ECR_Block)
diff --git a/Engine/Source/Runtime/Engine/Private/Collision/CollisionProfile.cpp b/Engine/Source/Runtime/Engine/Private/Collision/CollisionProfile.cpp
--- a/Engine/Source/Runtime/Engine/Private/Collision/CollisionProfile.cpp
+++ b/Engine/Source/Runtime/Engine/Private/Collision/CollisionProfile.cpp
@@ -22,9 +22,11 @@ FName UCollisionProfile::CustomCollisionProfileName = FName(TEXT("Custom"));
 FCollisionResponseTemplate::FCollisionResponseTemplate()
 	: CollisionEnabled(ECollisionEnabled::NoCollision)
 	, ObjectType(ECollisionChannel::ECC_WorldStatic)
-	, HelpMessage(TEXT("Needs description"))
 	, bCanModify(true)
 	, ResponseToChannels()
+#if WITH_EDITORONLY_DATA
+	, HelpMessage(TEXT("Needs description"))
+#endif
 {
 }
 
@@ -754,7 +756,7 @@ FName UCollisionProfile::ReturnChannelNameFromContainerIndex(int32 ContainerInde
 		return ChannelDisplayNames[ContainerIndex];
 	}
 
-	return TEXT("");
+	return NAME_None;
 }
 
 ECollisionChannel UCollisionProfile::ConvertToCollisionChannel(bool TraceType, int32 Index) const
