Use the actor's class instead of its name when reporting data to the network profiler. This fixes some issues where different actor names would cause data would be grouped incorrectly.

diff --git a/Engine/Source/Runtime/Engine/Private/NetworkProfiler.cpp b/Engine/Source/Runtime/Engine/Private/NetworkProfiler.cpp
--- a/Engine/Source/Runtime/Engine/Private/NetworkProfiler.cpp
+++ b/Engine/Source/Runtime/Engine/Private/NetworkProfiler.cpp
@@ -271,7 +271,7 @@ void FNetworkProfiler::TrackSendRPC(const AActor* Actor, const UFunction* Functi
 
 		SetCurrentConnection(Connection);
 
-		uint32 ActorNameTableIndex = GetNameTableIndex(Actor->GetName());
+		uint32 ActorNameTableIndex = GetNameTableIndex(GetNameSafe(Actor->GetClass()));
 		uint32 FunctionNameTableIndex = GetNameTableIndex(Function->GetName());
 
 		uint8 Type = NPTYPE_SendRPC;
@@ -292,7 +292,7 @@ void FNetworkProfiler::TrackQueuedRPC(UNetConnection* Connection, UObject* Targe
 		
 		FQueuedRPCInfo Info;
 
-		Info.ActorNameIndex = GetNameTableIndex(Actor->GetName());
+		Info.ActorNameIndex = GetNameTableIndex(GetNameSafe(Actor->GetClass()));
 		Info.FunctionNameIndex = GetNameTableIndex(Function->GetName());
 
 		Info.Connection = Connection;
@@ -502,7 +502,7 @@ void FNetworkProfiler::TrackReplicateActor( const AActor* Actor, FReplicationFla
 
 		SetCurrentConnection( Connection );
 
-		uint32 NameTableIndex = GetNameTableIndex( Actor->GetName() );
+		uint32 NameTableIndex = GetNameTableIndex(GetNameSafe(Actor->GetClass()));
 
 		uint8 Type = NPTYPE_ReplicateActor;
 		(*FileWriter) << Type;
