Fix for crash with referenced sound waves

diff --git a/Engine/Source/Runtime/Engine/Private/AudioDevice.cpp b/Engine/Source/Runtime/Engine/Private/AudioDevice.cpp
--- a/Engine/Source/Runtime/Engine/Private/AudioDevice.cpp
+++ b/Engine/Source/Runtime/Engine/Private/AudioDevice.cpp
@@ -3879,7 +3879,19 @@ void FAudioDevice::Update(bool bGameTicking)
 		for (int32 i = ReferencedSoundWaves.Num() - 1; i >= 0; --i)
 		{
 			USoundWave* Wave = ReferencedSoundWaves[i];
-			if (Wave->GetPrecacheState() == ESoundWavePrecacheState::Done && !Wave->IsGeneratingAudio())
+			bool bRemove = true;
+			if (Wave)
+			{
+				const bool bIsPrecacheDone = (Wave->GetPrecacheState() == ESoundWavePrecacheState::Done);
+				const bool bIsGeneratingAudio = Wave->IsGeneratingAudio();
+
+				if (!bIsPrecacheDone || bIsGeneratingAudio)
+				{
+					bRemove = false;
+				}
+			}
+
+			if (bRemove)
 			{
 				ReferencedSoundWaves.RemoveAtSwap(i, 1, false);
 			}
