#UE4 Added console command to display the list of bundles and what paths are in them for a given primary asset

diff --git a/Engine/Source/Runtime/Engine/Classes/Engine/AssetManager.h b/Engine/Source/Runtime/Engine/Classes/Engine/AssetManager.h
--- a/Engine/Source/Runtime/Engine/Classes/Engine/AssetManager.h
+++ b/Engine/Source/Runtime/Engine/Classes/Engine/AssetManager.h
@@ -367,6 +367,9 @@ public:
 	/** Dumps out list of loaded asset bundles to log */
 	static void DumpLoadedAssetState();
 
+	/** Shows a list of all bundles for the specified primary asset by primary asset id (i.e. Map:Entry) */
+	static void DumpBundlesForAsset(const TArray<FString>& Args);
+
 	/** Dumps information about the Asset Registry to log */
 	static void DumpAssetRegistryInfo();
 
diff --git a/Engine/Source/Runtime/Engine/Private/AssetManager.cpp b/Engine/Source/Runtime/Engine/Private/AssetManager.cpp
--- a/Engine/Source/Runtime/Engine/Private/AssetManager.cpp
+++ b/Engine/Source/Runtime/Engine/Private/AssetManager.cpp
@@ -2463,6 +2463,55 @@ void UAssetManager::DumpLoadedAssetState()
 	}
 }
 
+static FAutoConsoleCommand CVarDumpBundlesForAsset(
+	TEXT("AssetManager.DumpBundlesForAsset"),
+	TEXT("Shows a list of all bundles for the specified primary asset by primary asset id (i.e. Map:Entry)"),
+	FConsoleCommandWithArgsDelegate::CreateStatic(UAssetManager::DumpBundlesForAsset),
+	ECVF_Cheat);
+
+void UAssetManager::DumpBundlesForAsset(const TArray<FString>& Args)
+{
+	if (Args.Num() < 1)
+	{
+		UE_LOG(LogAssetManager, Warning, TEXT("Too few arguments for DumpBundlesForAsset. Include the primary asset id (i.e. Map:Entry)"));
+		return;
+	}
+
+	FString PrimaryAssetIdString = Args[0];
+	if (!PrimaryAssetIdString.Contains(TEXT(":")))
+	{
+		UE_LOG(LogAssetManager, Warning, TEXT("Incorrect argument for DumpBundlesForAsset. Arg should be the primary asset id (i.e. Map:Entry)"));
+		return;
+	}
+
+	if (!UAssetManager::IsValid())
+	{
+		UE_LOG(LogAssetManager, Warning, TEXT("DumpBundlesForAsset Failed. Invalid asset manager."));
+		return;
+	}
+
+	UAssetManager& Manager = Get();
+
+	FPrimaryAssetId PrimaryAssetId(PrimaryAssetIdString);
+	const TMap<FName, FAssetBundleEntry>* FoundMap = Manager.CachedAssetBundles.Find(PrimaryAssetId);
+	if (!FoundMap)
+	{
+		UE_LOG(LogAssetManager, Display, TEXT("Could not find bundles for primary asset %s."), *PrimaryAssetIdString);
+		return;
+	}
+
+	UE_LOG(LogAssetManager, Display, TEXT("Dumping bundles for primary asset %s..."), *PrimaryAssetIdString);
+	for (auto MapIt = FoundMap->CreateConstIterator(); MapIt; ++MapIt)
+	{
+		const FAssetBundleEntry& Entry = MapIt.Value();
+		UE_LOG(LogAssetManager, Display, TEXT("  Bundle: %s (%d assets)"), *Entry.BundleName.ToString(), Entry.BundleAssets.Num());
+		for (const FSoftObjectPath& Path : Entry.BundleAssets)
+		{
+			UE_LOG(LogAssetManager, Display, TEXT("    %s"), *Path.ToString());
+		}
+	}
+}
+
 static FAutoConsoleCommand CVarDumpAssetRegistryInfo(
 	TEXT("AssetManager.DumpAssetRegistryInfo"),
 	TEXT("Dumps extended info about asset registry to log"),
