Ensure the MinDrawDistance of newly created LODActor take into account the current HLOD distance override (r.HLOD.DistanceOverride)

diff --git a/Engine/Source/Runtime/Engine/Private/LODActor.cpp b/Engine/Source/Runtime/Engine/Private/LODActor.cpp
--- a/Engine/Source/Runtime/Engine/Private/LODActor.cpp
+++ b/Engine/Source/Runtime/Engine/Private/LODActor.cpp
@@ -221,7 +221,7 @@ void ALODActor::SetupComponent(UStaticMeshComponent* Component)
 	Component->bAllowCullDistanceVolume = false;
 	Component->bNeverDistanceCull = true;
 
-	Component->MinDrawDistance = LODDrawDistance;
+	Component->MinDrawDistance = GetLODDrawDistanceWithOverride();
 }
 
 FString ALODActor::GetDetailedInfoInternal() const
@@ -563,7 +563,7 @@ void ALODActor::UnregisterMeshComponents()
 void ALODActor::SetDrawDistance(float InDistance)
 {
 	LODDrawDistance = InDistance;
-	SetComponentsMinDrawDistance(LODDrawDistance, false);
+	SetComponentsMinDrawDistance(GetLODDrawDistanceWithOverride(), false);
 }
 
 #if !(UE_BUILD_SHIPPING || UE_BUILD_TEST)
@@ -791,7 +791,7 @@ void ALODActor::AddSubActor(AActor* InActor)
 	SubActors.Add(InActor);
 
 	UStaticMeshComponent* LODComponent = GetOrCreateLODComponentForActor(InActor);
-	InActor->SetLODParent(LODComponent, LODDrawDistance);
+	InActor->SetLODParent(LODComponent, GetLODDrawDistanceWithOverride());
 
 	// Adding number of triangles
 	if (!InActor->IsA<ALODActor>())
