Use UInstancedStaticMeshComponent::ClearInstances() method to clear the instance array instead of doing it manually.

diff --git a/Engine/Source/Runtime/Engine/Private/LODActor.cpp b/Engine/Source/Runtime/Engine/Private/LODActor.cpp
--- a/Engine/Source/Runtime/Engine/Private/LODActor.cpp
+++ b/Engine/Source/Runtime/Engine/Private/LODActor.cpp
@@ -1025,7 +1025,7 @@ void ALODActor::SetupImposters(UMaterialInterface* InMaterial, UStaticMesh* InSt
 
 	UInstancedStaticMeshComponent* Component = GetOrCreateLODComponentForMaterial(InMaterial);
 	Component->SetStaticMesh(InStaticMesh);
-	Component->PerInstanceSMData.Empty();
+	Component->ClearInstances();
 	
 	for (const FTransform& Transform : InTransforms)
 	{
