Do not create SRVs for a mesh buffers if platform does not use manual vertex fetch

diff --git a/Engine/Source/Runtime/Engine/Private/InstancedStaticMesh.cpp b/Engine/Source/Runtime/Engine/Private/InstancedStaticMesh.cpp
--- a/Engine/Source/Runtime/Engine/Private/InstancedStaticMesh.cpp
+++ b/Engine/Source/Runtime/Engine/Private/InstancedStaticMesh.cpp
@@ -364,7 +364,11 @@ void FStaticMeshInstanceBuffer::CreateVertexBuffer(FResourceArrayInterface* InRe
 	// TODO: possibility over allocated the vertex buffer when we support partial update for when working in the editor
 	FRHIResourceCreateInfo CreateInfo(InResourceArray);
 	OutVertexBufferRHI = RHICreateVertexBuffer(InResourceArray->GetResourceDataSize(), InUsage, CreateInfo);
-	OutInstanceSRV = RHICreateShaderResourceView(OutVertexBufferRHI, InStride, InFormat);
+	
+	if (RHISupportsManualVertexFetch(GMaxRHIShaderPlatform))
+	{
+		OutInstanceSRV = RHICreateShaderResourceView(OutVertexBufferRHI, InStride, InFormat);
+	}
 }
 
 void FStaticMeshInstanceBuffer::BindInstanceVertexBuffer(const class FVertexFactory* VertexFactory, FInstancedStaticMeshDataType& InstancedStaticMeshData) const
diff --git a/Engine/Source/Runtime/Engine/Private/Rendering/ColorVertexBuffer.cpp b/Engine/Source/Runtime/Engine/Private/Rendering/ColorVertexBuffer.cpp
--- a/Engine/Source/Runtime/Engine/Private/Rendering/ColorVertexBuffer.cpp
+++ b/Engine/Source/Runtime/Engine/Private/Rendering/ColorVertexBuffer.cpp
@@ -419,7 +419,8 @@ FVertexBufferRHIRef FColorVertexBuffer::CreateRHIBuffer_Async()
 void FColorVertexBuffer::InitRHI()
 {
 	VertexBufferRHI = CreateRHIBuffer_RenderThread();
-	if (VertexBufferRHI)
+
+	if (VertexBufferRHI && RHISupportsManualVertexFetch(GMaxRHIShaderPlatform))
 	{
 		ColorComponentsSRV = RHICreateShaderResourceView(VertexData ? VertexBufferRHI : nullptr, 4, PF_R8G8B8A8);
 	}
diff --git a/Engine/Source/Runtime/Engine/Public/Rendering/ColorVertexBuffer.h b/Engine/Source/Runtime/Engine/Public/Rendering/ColorVertexBuffer.h
--- a/Engine/Source/Runtime/Engine/Public/Rendering/ColorVertexBuffer.h
+++ b/Engine/Source/Runtime/Engine/Public/Rendering/ColorVertexBuffer.h
@@ -146,9 +146,11 @@ public:
 	{
 		if (VertexBufferRHI && IntermediateBuffer)
 		{
-			check(ColorComponentsSRV);
 			Batcher.QueueUpdateRequest(VertexBufferRHI, IntermediateBuffer);
-			Batcher.QueueUpdateRequest(ColorComponentsSRV, VertexBufferRHI, 4, PF_R8G8B8A8);
+			if (ColorComponentsSRV)
+			{
+				Batcher.QueueUpdateRequest(ColorComponentsSRV, VertexBufferRHI, 4, PF_R8G8B8A8);
+			}
 		}
 	}
 
