Extended FNavAgentSelector and FNavDataConfig with additional functionalities

Yoan.StAmant, Maxime.Mercier

diff --git a/Engine/Source/Runtime/Engine/Classes/AI/Navigation/NavigationTypes.h b/Engine/Source/Runtime/Engine/Classes/AI/Navigation/NavigationTypes.h
--- a/Engine/Source/Runtime/Engine/Classes/AI/Navigation/NavigationTypes.h
+++ b/Engine/Source/Runtime/Engine/Classes/AI/Navigation/NavigationTypes.h
@@ -100,6 +100,8 @@ struct ENGINE_API FNavAgentSelector
 {
 	GENERATED_USTRUCT_BODY()
 
+	static const uint32 InitializedBit = 0x80000000;
+
 #if CPP
 	union
 	{
@@ -144,21 +146,39 @@ struct ENGINE_API FNavAgentSelector
 	};
 #endif
 
-	FNavAgentSelector();
+	explicit FNavAgentSelector(const uint32 InBits = 0x7fffffff);
 
 	FORCEINLINE bool Contains(int32 AgentIndex) const
 	{
 		return (AgentIndex >= 0 && AgentIndex < 16) ? !!(PackedBits & (1 << AgentIndex)) : false;
 	}
 
+	FORCEINLINE void Set(int32 AgentIndex)
+	{
+		if (AgentIndex >= 0 && AgentIndex < 16)
+		{
+			PackedBits |= (1 << AgentIndex);
+		}
+	}
+
 	FORCEINLINE bool IsInitialized() const
 	{
-		return (PackedBits & 0x80000000) != 0;
+		return (PackedBits & InitializedBit) != 0;
 	}
 
 	FORCEINLINE void MarkInitialized()
 	{
-		PackedBits |= 0x80000000;
+		PackedBits |= InitializedBit;
+	}
+
+	FORCEINLINE void Empty()
+	{
+		PackedBits = 0;
+	}
+
+	bool IsSame(const FNavAgentSelector& Other) const
+	{
+		return (~InitializedBit & PackedBits) == (~InitializedBit & Other.PackedBits);
 	}
 
 	bool Serialize(FArchive& Ar);
@@ -552,12 +572,21 @@ public:
 	FNavDataConfig(float Radius = FNavigationSystem::FallbackAgentRadius, float Height = FNavigationSystem::FallbackAgentHeight);
 	FNavDataConfig(const FNavDataConfig& Other);
 
+	bool IsValid() const 
+	{
+		return FNavAgentProperties::IsValid() && NavDataClass.IsValid();
+	}
+
+	void Invalidate();
+
 	void SetNavDataClass(UClass* InNavDataClass);
 	void SetNavDataClass(TSoftClassPtr<AActor> InNavDataClass);
 	
 	template<typename T>
 	TSubclassOf<T> GetNavDataClass() const { return TSubclassOf<T>(NavDataClass.Get()); }
 
+	FString GetDescription() const;
+
 #if WITH_EDITOR
 	static FName GetNavigationDataClassPropertyName()
 	{
diff --git a/Engine/Source/Runtime/Engine/Private/AI/Navigation/NavigationTypes.cpp b/Engine/Source/Runtime/Engine/Private/AI/Navigation/NavigationTypes.cpp
--- a/Engine/Source/Runtime/Engine/Private/AI/Navigation/NavigationTypes.cpp
+++ b/Engine/Source/Runtime/Engine/Private/AI/Navigation/NavigationTypes.cpp
@@ -64,6 +64,18 @@ void FNavDataConfig::SetNavDataClass(TSoftClassPtr<AActor> InNavDataClass)
 {
 	NavDataClass = InNavDataClass;
 }
+
+void FNavDataConfig::Invalidate()
+{
+	new(this) FNavAgentProperties();
+	SetNavDataClass(nullptr);
+}
+
+FString FNavDataConfig::GetDescription() const
+{
+	return FString::Printf(TEXT("Name %s class %s agent radius %.1f")
+		, *Name.ToString(), *NavDataClass.ToString(), AgentRadius);
+}
 //----------------------------------------------------------------------//
 // FNavigationRelevantData
 //----------------------------------------------------------------------//
@@ -210,7 +222,8 @@ uint16 FNavigationQueryFilter::GetExcludeFlags() const
 //----------------------------------------------------------------------//
 // FNavAgentSelector
 //----------------------------------------------------------------------//
-FNavAgentSelector::FNavAgentSelector() : PackedBits(0x7fffffff)
+FNavAgentSelector::FNavAgentSelector(const uint32 InBits)
+	: PackedBits(InBits)
 {
 }
 
