Fixing comment

diff --git a/Engine/Source/Runtime/Engine/Classes/Sound/SoundSubmix.h b/Engine/Source/Runtime/Engine/Classes/Sound/SoundSubmix.h
--- a/Engine/Source/Runtime/Engine/Classes/Sound/SoundSubmix.h
+++ b/Engine/Source/Runtime/Engine/Classes/Sound/SoundSubmix.h
@@ -179,7 +179,7 @@ class ENGINE_API USoundSubmix : public UObject
 	UPROPERTY(EditAnywhere, BlueprintReadOnly, Category = SoundSubmix)
 	TArray<USoundEffectSubmixPreset*> SubmixEffectChain;
 
-	// TODO: Hide this unless Channel Format is ambisonics. Also, worry about thread safety.
+	/** Optional settings used by plugins which support ambisonics file playback. */
 	UPROPERTY(EditAnywhere, BlueprintReadWrite, Category = SoundSubmix)
 	UAmbisonicsSubmixSettingsBase* AmbisonicsPluginSettings;
 
