Maintain child actor name correctly through initial recreation during rerun construction script on load.

diff --git a/Engine/Source/Runtime/Engine/Private/Components/ChildActorComponent.cpp b/Engine/Source/Runtime/Engine/Private/Components/ChildActorComponent.cpp
--- a/Engine/Source/Runtime/Engine/Private/Components/ChildActorComponent.cpp
+++ b/Engine/Source/Runtime/Engine/Private/Components/ChildActorComponent.cpp
@@ -22,7 +22,17 @@ void UChildActorComponent::OnRegister()
 
 	if (ChildActor)
 	{
-		if (bNeedsRecreate || ChildActor->GetClass() != ChildActorClass)
+		if (ChildActor->GetClass() != ChildActorClass)
+		{
+			bNeedsRecreate = true;
+			ChildActorName = NAME_None;
+		}
+		else
+		{
+			ChildActorName = ChildActor->GetFName();
+		}
+
+		if (bNeedsRecreate)
 		{
 			bNeedsRecreate = false;
 			DestroyChildActor();
@@ -30,8 +40,6 @@ void UChildActorComponent::OnRegister()
 		}
 		else
 		{
-			ChildActorName = ChildActor->GetFName();
-			
 			USceneComponent* ChildRoot = ChildActor->GetRootComponent();
 			if (ChildRoot && ChildRoot->GetAttachParent() != this)
 			{
