textures are now created from render targets in C++ without having to create a new texture in BP manually.

diff --git a/Engine/Source/Runtime/Engine/Classes/Engine/TextureRenderTarget2D.h b/Engine/Source/Runtime/Engine/Classes/Engine/TextureRenderTarget2D.h
--- a/Engine/Source/Runtime/Engine/Classes/Engine/TextureRenderTarget2D.h
+++ b/Engine/Source/Runtime/Engine/Classes/Engine/TextureRenderTarget2D.h
@@ -163,6 +163,7 @@ class UTextureRenderTarget2D : public UTextureRenderTarget
 	 * @return New UTexture2D object.
 	 */
 	ENGINE_API UTexture2D* ConstructTexture2D(UObject* InOuter, const FString& NewTexName, EObjectFlags InObjectFlags, uint32 Flags=CTF_Default, TArray<uint8>* AlphaOverride=NULL);
+	ENGINE_API void UpdateTexture2D(UTexture2D* InTexture2D, ETextureSourceFormat InTextureFormat, uint32 Flags = CTF_Default, TArray<uint8>* AlphaOverride = NULL);
 
 	/**
 	 * Updates (resolves) the render target texture immediately.
diff --git a/Engine/Source/Runtime/Engine/Private/TextureRenderTarget2D.cpp b/Engine/Source/Runtime/Engine/Private/TextureRenderTarget2D.cpp
--- a/Engine/Source/Runtime/Engine/Private/TextureRenderTarget2D.cpp
+++ b/Engine/Source/Runtime/Engine/Private/TextureRenderTarget2D.cpp
@@ -264,15 +264,13 @@ UTexture2D* UTextureRenderTarget2D::ConstructTexture2D(UObject* Outer, const FSt
 
 	const EPixelFormat PixelFormat = GetFormat();
 	ETextureSourceFormat TextureFormat = TSF_Invalid;
-	TextureCompressionSettings CompressionSettingsForTexture = TC_Default;
 	switch (PixelFormat)
 	{
-		case PF_B8G8R8A8:
-			TextureFormat = TSF_BGRA8;
+	case PF_B8G8R8A8:
+		TextureFormat = TSF_BGRA8;
 		break;
-		case PF_FloatRGBA:
-			TextureFormat = TSF_RGBA16F;
-			CompressionSettingsForTexture = TC_HDR;
+	case PF_FloatRGBA:
+		TextureFormat = TSF_RGBA16F;
 		break;
 	}
 
@@ -284,14 +282,57 @@ UTexture2D* UTextureRenderTarget2D::ConstructTexture2D(UObject* Outer, const FSt
 
 	// create the 2d texture
 	Result = NewObject<UTexture2D>(Outer, FName(*NewTexName), InObjectFlags);
+	
+	UpdateTexture2D(Result, TextureFormat, Flags, AlphaOverride);
+
+	// if render target gamma used was 1.0 then disable SRGB for the static texture
+	if (FMath::Abs(RenderTarget->GetDisplayGamma() - 1.0f) < KINDA_SMALL_NUMBER)
+	{
+		Flags &= ~CTF_SRGB;
+	}
+
+	Result->SRGB = (Flags & CTF_SRGB) != 0;
+	Result->MipGenSettings = TMGS_FromTextureGroup;
+
+	if ((Flags & CTF_AllowMips) == 0)
+	{
+		Result->MipGenSettings = TMGS_NoMipmaps;
+	}
+
+
+	if (Flags & CTF_Compress)
+	{
+		// Set compression options.
+		Result->DeferCompression = (Flags & CTF_DeferCompression) ? true : false;
+	}
+	else
+	{
+		// Disable compression
+		Result->CompressionNone = true;
+		Result->DeferCompression = false;
+	}
+	Result->PostEditChange();
+#endif
+	return Result;
+}
+
+PRAGMA_DISABLE_OPTIMIZATION
+void UTextureRenderTarget2D::UpdateTexture2D(UTexture2D* InTexture2D, ETextureSourceFormat InTextureFormat, uint32 Flags, TArray<uint8>* AlphaOverride)
+{
+#if WITH_EDITOR
+	FRenderTarget* RenderTarget = GameThread_GetRenderTargetResource();
+
+	const EPixelFormat PixelFormat = GetFormat();
+	TextureCompressionSettings CompressionSettingsForTexture = PixelFormat == EPixelFormat::PF_FloatRGBA ? TC_HDR : TC_Default;
+
 	// init to the same size as the 2d texture
-	Result->Source.Init(SizeX, SizeY, 1, 1, TextureFormat);
+	InTexture2D->Source.Init(SizeX, SizeY, 1, 1, InTextureFormat);
 
-	uint32* TextureData = (uint32*)Result->Source.LockMip(0);
-	const int32 TextureDataSize = Result->Source.CalcMipSize(0);
+	uint32* TextureData = (uint32*)InTexture2D->Source.LockMip(0);
+	const int32 TextureDataSize = InTexture2D->Source.CalcMipSize(0);
 
 	// read the 2d surface
-	if (TextureFormat == TSF_BGRA8)
+	if (InTextureFormat == TSF_BGRA8)
 	{
 		TArray<FColor> SurfData;
 		RenderTarget->ReadPixels(SurfData);
@@ -321,10 +362,10 @@ UTexture2D* UTextureRenderTarget2D::ConstructTexture2D(UObject* Outer, const FSt
 			}
 		}
 		// copy the 2d surface data to the first mip of the static 2d texture
-		check(TextureDataSize == SurfData.Num()*sizeof(FColor));
+		check(TextureDataSize == SurfData.Num() * sizeof(FColor));
 		FMemory::Memcpy(TextureData, SurfData.GetData(), TextureDataSize);
 	}
-	else if (TextureFormat == TSF_RGBA16F)
+	else if (InTextureFormat == TSF_RGBA16F)
 	{
 		TArray<FFloat16Color> SurfData;
 		RenderTarget->ReadFloat16Pixels(SurfData);
@@ -354,39 +395,14 @@ UTexture2D* UTextureRenderTarget2D::ConstructTexture2D(UObject* Outer, const FSt
 			}
 		}
 		// copy the 2d surface data to the first mip of the static 2d texture
-		check(TextureDataSize == SurfData.Num()*sizeof(FFloat16Color));
+		check(TextureDataSize == SurfData.Num() * sizeof(FFloat16Color));
 		FMemory::Memcpy(TextureData, SurfData.GetData(), TextureDataSize);
 	}
-	Result->Source.UnlockMip(0);
-	// if render target gamma used was 1.0 then disable SRGB for the static texture
-	if (FMath::Abs(RenderTarget->GetDisplayGamma() - 1.0f) < KINDA_SMALL_NUMBER)
-	{
-		Flags &= ~CTF_SRGB;
-	}
 
-	Result->SRGB = (Flags & CTF_SRGB) != 0;
-	Result->MipGenSettings = TMGS_FromTextureGroup;
-
-	if ((Flags & CTF_AllowMips) == 0)
-	{
-		Result->MipGenSettings = TMGS_NoMipmaps;
-	}
+	InTexture2D->Source.UnlockMip(0);
 
-	Result->CompressionSettings = CompressionSettingsForTexture;
-	if (Flags & CTF_Compress)
-	{
-		// Set compression options.
-		Result->DeferCompression = (Flags & CTF_DeferCompression) ? true : false;
-	}
-	else
-	{
-		// Disable compression
-		Result->CompressionNone = true;
-		Result->DeferCompression = false;
-	}
-	Result->PostEditChange();
+	InTexture2D->CompressionSettings = CompressionSettingsForTexture;
 #endif
-	return Result;
 }
 
 /*-----------------------------------------------------------------------------
diff --git a/Engine/Source/Runtime/Landscape/Private/LandscapeBlueprintBrushBase.cpp b/Engine/Source/Runtime/Landscape/Private/LandscapeBlueprintBrushBase.cpp
--- a/Engine/Source/Runtime/Landscape/Private/LandscapeBlueprintBrushBase.cpp
+++ b/Engine/Source/Runtime/Landscape/Private/LandscapeBlueprintBrushBase.cpp
@@ -7,6 +7,8 @@
 
 #define LOCTEXT_NAMESPACE "Landscape"
 
+static int FramePaddingBeforeLayerUpdate = 5;
+
 ALandscapeBlueprintBrushBase::ALandscapeBlueprintBrushBase(const FObjectInitializer& ObjectInitializer)
 #if WITH_EDITORONLY_DATA
 	: OwningLandscape(nullptr)
@@ -57,7 +59,7 @@ void ALandscapeBlueprintBrushBase::Tick(float DeltaSeconds)
 #if WITH_EDITORONLY_DATA
 	// Avoid computing collision and client updates every frame
 	// Wait until we didn't trigger any more landscape update requests (padding of a couple of frames)
-	if (LastRequestLayersContentUpdateFrameNumber + 5 == GFrameNumber)
+	if (LastRequestLayersContentUpdateFrameNumber + FramePaddingBeforeLayerUpdate == GFrameNumber)
 	{
 		uint32 ModeMask = 0;
 		if (AffectHeightmap)
@@ -90,6 +92,11 @@ bool ALandscapeBlueprintBrushBase::ShouldTickIfViewportsOnly() const
 	return true;
 }
 
+bool ALandscapeBlueprintBrushBase::IsLayerUpdatePending() const
+{
+	return GFrameNumber < LastRequestLayersContentUpdateFrameNumber + FramePaddingBeforeLayerUpdate;
+}
+
 void ALandscapeBlueprintBrushBase::SetIsVisible(bool bInIsVisible)
 {
 #if WITH_EDITORONLY_DATA
diff --git a/Engine/Source/Runtime/Landscape/Public/LandscapeBlueprintBrushBase.h b/Engine/Source/Runtime/Landscape/Public/LandscapeBlueprintBrushBase.h
--- a/Engine/Source/Runtime/Landscape/Public/LandscapeBlueprintBrushBase.h
+++ b/Engine/Source/Runtime/Landscape/Public/LandscapeBlueprintBrushBase.h
@@ -51,7 +51,8 @@ public:
 	bool IsAffectingWeightmap() const { return AffectWeightmap; }
 	bool IsAffectingWeightmapLayer(const FName& InLayerName) const;
 	bool IsVisible() const { return bIsVisible; }
-	
+	bool IsLayerUpdatePending() const;
+
 	void SetIsVisible(bool bInIsVisible);
 	void SetAffectsHeightmap(bool bInAffectsHeightmap);
 	void SetAffectsWeightmap(bool bInAffectsWeightmap);
