NetworkEmulation
* Added NetEmulation.Off to remove emulation settings completely
* NetEmulation console commands now triggered via NetEmulation.PktXXX
* NetEmulation console commands can now be applied before the NetDriver is created.
* NetEmulation settings applied via console will now apply to all NetDrivers create during the engine's lifetime.

[RN] minor engine/network

diff --git a/Engine/Source/Runtime/Engine/Classes/Engine/NetDriver.h b/Engine/Source/Runtime/Engine/Classes/Engine/NetDriver.h
--- a/Engine/Source/Runtime/Engine/Classes/Engine/NetDriver.h
+++ b/Engine/Source/Runtime/Engine/Classes/Engine/NetDriver.h
@@ -490,11 +490,6 @@ struct ENGINE_API FPacketSimulationSettings
 	 */
 	void ApplySettings(const FPacketSimulationSettings& NewSettings);
 
-	/**
-	 * Registers commands for auto-completion, etc.
-	 */
-	void RegisterCommands();
-
 	/**
 	 * Ensure that settings have proper values
 	 */
diff --git a/Engine/Source/Runtime/Engine/Private/NetworkDriver.cpp b/Engine/Source/Runtime/Engine/Private/NetworkDriver.cpp
--- a/Engine/Source/Runtime/Engine/Private/NetworkDriver.cpp
+++ b/Engine/Source/Runtime/Engine/Private/NetworkDriver.cpp
@@ -64,6 +64,7 @@
 #include "Engine/ReplicationDriver.h"
 #include "ProfilingDebugging/CsvProfiler.h"
 #include "Engine/LevelScriptActor.h"
+#include "Engine/NetworkSettings.h"
 #include "Net/NetworkGranularMemoryLogging.h"
 #include "SocketSubsystem.h"
 #include "AddressInfoTypes.h"
@@ -155,6 +156,35 @@ int32 GNumSharedSerializationMiss;
 
 extern int32 GNetRPCDebug;
 
+namespace NetEmulationHelper
+{
+#if DO_ENABLE_NET_TEST
+	/** Global that stores the network emulation values outside the NetDriver lifetime */
+	static TOptional<FPacketSimulationSettings> PersistentPacketSimulationSettings;
+
+	void CreatePersistentSimulationSettings()
+	{
+		if (!PersistentPacketSimulationSettings.IsSet())
+		{
+			PersistentPacketSimulationSettings.Emplace(FPacketSimulationSettings());
+		}
+	}
+
+	void ApplySimulationSettingsOnNetDrivers(UWorld* World, const FPacketSimulationSettings& Settings)
+	{
+		// Execute on all active NetDrivers
+		FWorldContext& Context = GEngine->GetWorldContextFromWorldChecked(World);
+		for (const FNamedNetDriver& ActiveNetDriver : Context.ActiveNetDrivers)
+		{
+			if (ActiveNetDriver.NetDriver)
+			{
+				ActiveNetDriver.NetDriver->SetPacketSimulationSettings(Settings);
+			}
+		}
+	}
+#endif //#if DO_ENABLE_NET_TEST
+}
+
 #if UE_BUILD_SHIPPING
 #define DEBUG_REMOTEFUNCTION(Format, ...)
 #else
@@ -210,8 +240,8 @@ static TAutoConsoleVariable<int32> CVarActorChannelPool(
 static TAutoConsoleVariable<int32> CVarAllowReliableMulticastToNonRelevantChannels(
 	TEXT("net.AllowReliableMulticastToNonRelevantChannels"),
 	1,
-	TEXT("Allow Reliable Server Multicasts to be sent to non-Relevant Actors, as long as their is an existing ActorChannel.")
-);
+	TEXT("Allow Reliable Server Multicasts to be sent to non-Relevant Actors, as long as their is an existing ActorChannel."));
+
 
 /*-----------------------------------------------------------------------------
 	UNetDriver implementation.
@@ -294,13 +324,17 @@ void UNetDriver::InitPacketSimulationSettings()
 		return;
 	}
 
-	PacketSimulationSettings.RegisterCommands();
-
 	if (bForcedPacketSettings)
 	{
 		return;
 	}
 
+	if (NetEmulationHelper::PersistentPacketSimulationSettings.IsSet())
+	{
+		SetPacketSimulationSettings(NetEmulationHelper::PersistentPacketSimulationSettings.GetValue());
+		return;
+	}
+
 	// Read the settings from .ini and command line, with the command line taking precedence	
 	PacketSimulationSettings.ResetSettings();
 	PacketSimulationSettings.LoadConfig(*NetDriverName.ToString());
@@ -1436,7 +1470,6 @@ bool UNetDriver::InitBase(bool bInitAsClient, FNetworkNotify* InNotify, const FU
 	if( bSettingFound )
 	{
 		SetPacketSimulationSettings(PacketSettings);
-		bForcedPacketSettings = true;
 	}
 #endif //#if DO_ENABLE_NET_TEST
 
@@ -2783,21 +2816,7 @@ bool UNetDriver::Exec( UWorld* InWorld, const TCHAR* Cmd, FOutputDevice& Ar )
 	else if (bNeverApplyNetworkEmulationSettings == false && 
 			 PacketSimulationSettings.ParseSettings(Cmd, *NetDriverName.ToString()))
 	{
-		if (ServerConnection)
-		{
-			// Notify the server connection of the change
-			ServerConnection->UpdatePacketSimulationSettings();
-		}
-#if WITH_SERVER_CODE
-		else
-		{
-			// Notify all client connections that the settings have changed
-			for (int32 Index = 0; Index < ClientConnections.Num(); Index++)
-			{
-				ClientConnections[Index]->UpdatePacketSimulationSettings();
-			}
-		}		
-#endif
+		OnPacketSimulationSettingsChanged();
 		return true;
 	}
 #endif
@@ -3254,6 +3273,9 @@ void UNetDriver::SetPacketSimulationSettings(const FPacketSimulationSettings& Ne
 		return;
 	}
 
+	// Once set this prevents driver changes from reloading the packet emulation settings
+	bForcedPacketSettings = true;
+
 	PacketSimulationSettings = NewSettings;
 	OnPacketSimulationSettingsChanged();
 }
@@ -3264,6 +3286,7 @@ void UNetDriver::OnPacketSimulationSettingsChanged()
 	{
 		ServerConnection->UpdatePacketSimulationSettings();
 	}
+#if WITH_SERVER_CODE
 	for (UNetConnection* ClientConnection : ClientConnections)
 	{
 		if (ClientConnection)
@@ -3271,6 +3294,7 @@ void UNetDriver::OnPacketSimulationSettingsChanged()
 			ClientConnection->UpdatePacketSimulationSettings();
 		}
 	}
+#endif
 }
 
 class FPacketSimulationConsoleCommandVisitor 
@@ -3411,33 +3435,6 @@ bool FPacketSimulationSettings::ConfigHelperBool(const TCHAR* Name, bool& Value,
 	return false;
 }
 
-void FPacketSimulationSettings::RegisterCommands()
-{
-	IConsoleManager& ConsoleManager = IConsoleManager::Get();
-	
-	// Register exec commands with the console manager for auto-completion if they havent been registered already by another net driver
-	if (!ConsoleManager.IsNameRegistered(TEXT("NetEmulation PktEmulationProfile=")))
-	{
-		ConsoleManager.RegisterConsoleCommand(TEXT("NetEmulation PktEmulationProfile="), TEXT("PktEmulationProfile=<NAME> (apply an emulation profile)"));
-
-		ConsoleManager.RegisterConsoleCommand(TEXT("NetEmulation PktLoss="), TEXT("PktLoss=<n> (simulates network packet loss)"));
-		ConsoleManager.RegisterConsoleCommand(TEXT("NetEmulation PktOrder="), TEXT("PktOrder=<n> (simulates network packet received out of order)"));
-		ConsoleManager.RegisterConsoleCommand(TEXT("NetEmulation PktDup="), TEXT("PktDup=<n> (simulates sending/receiving duplicate network packets)"));
-		ConsoleManager.RegisterConsoleCommand(TEXT("NetEmulation PktLag="), TEXT("PktLag=<n> (simulates network packet lag)"));
-		ConsoleManager.RegisterConsoleCommand(TEXT("NetEmulation PktLagVariance="), TEXT("PktLagVariance=<n> (simulates variable network packet lag)"));
-
-		ConsoleManager.RegisterConsoleCommand(TEXT("NetEmulation PktLagMin="), TEXT("PktLagMin=<n> (minimum outgoing packet latency)"));
-		ConsoleManager.RegisterConsoleCommand(TEXT("NetEmulation PktLagMax="), TEXT("PktLagMax=<n> (maximum outgoing packet latency)"));
-
-		ConsoleManager.RegisterConsoleCommand(TEXT("NetEmulation PktIncomingLagMin="), TEXT("PktIncomingLagMin=<n> (minimum incoming packet latency)"));
-		ConsoleManager.RegisterConsoleCommand(TEXT("NetEmulation PktIncomingLagMax="), TEXT("PktIncomingLagMax=<n> (maximum incoming packet latency)"));
-		ConsoleManager.RegisterConsoleCommand(TEXT("NetEmulation PktIncomingLoss="), TEXT("PktIncomingLoss=<n> (simulates incoming packet loss)"));
-	}
-
-	// Note we never unregister the console commands since net drivers come and go, and we can sometimes have more than 1, etc. 
-	// We could do better bookkeeping for this, but its not worth it right now. Just ensure the commands are always there for tab completion.
-}
-
 /**
  * Reads the settings from a string: command line or an exec
  *
@@ -5659,6 +5656,81 @@ FAutoConsoleCommandWithWorld	DumpRelevantActorsCommand(
 	FConsoleCommandWithWorldDelegate::CreateStatic(DumpRelevantActors)
 	);
 
+
+#if DO_ENABLE_NET_TEST
+
+FAutoConsoleCommandWithWorldArgsAndOutputDevice NetEmulationPktEmulationProfile(TEXT("NetEmulation.PktEmulationProfile"), TEXT("Apply a preconfigured emulation profile."),
+	FConsoleCommandWithWorldArgsAndOutputDeviceDelegate::CreateStatic([](const TArray<FString>& Args, UWorld* World, FOutputDevice& Output)
+{
+	bool bProfileApplied(false);
+	if (Args.Num() > 0)
+	{
+		FString CmdParams = FString::Printf(TEXT("PktEmulationProfile=%s"), *(Args[0]));
+
+		NetEmulationHelper::CreatePersistentSimulationSettings();
+
+		bProfileApplied = NetEmulationHelper::PersistentPacketSimulationSettings.GetValue().ParseSettings(*CmdParams, nullptr);
+		
+		if (bProfileApplied)
+		{
+			NetEmulationHelper::ApplySimulationSettingsOnNetDrivers(World, NetEmulationHelper::PersistentPacketSimulationSettings.GetValue());
+		}
+		else
+		{
+			Output.Log(FString::Printf(TEXT("EmulationProfile: %s was not found in Engine.ini"), *(Args[0])));
+		}
+	}
+	else
+	{
+		Output.Log(FString::Printf(TEXT("Missing emulation profile name")));
+	}
+
+	if (!bProfileApplied)
+	{
+		if (const UNetworkSettings* NetworkSettings = GetDefault<UNetworkSettings>())
+		{
+			Output.Log(TEXT("List of some supported emulation profiles:"));
+			for (const FNetworkEmulationProfileDescription& ProfileDesc : NetworkSettings->NetworkEmulationProfiles)
+			{
+				Output.Log(FString::Printf(TEXT("%s"), *ProfileDesc.ProfileName));
+			}
+		}
+	}
+}));
+
+FAutoConsoleCommandWithWorld NetEmulationOff(TEXT("NetEmulation.Off"), TEXT("Turn off network emulation"),
+	FConsoleCommandWithWorldDelegate::CreateStatic([](UWorld* World)
+{
+	NetEmulationHelper::CreatePersistentSimulationSettings();
+	NetEmulationHelper::PersistentPacketSimulationSettings.GetValue().ResetSettings();
+	NetEmulationHelper::ApplySimulationSettingsOnNetDrivers(World, NetEmulationHelper::PersistentPacketSimulationSettings.GetValue());
+}));
+
+#define BUILD_NETEMULATION_CONSOLE_COMMAND(CommandName, CommandHelp) FAutoConsoleCommandWithWorldAndArgs NetEmulation##CommandName(TEXT("NetEmulation."#CommandName), TEXT(CommandHelp), \
+	FConsoleCommandWithWorldAndArgsDelegate::CreateStatic([](const TArray<FString>& Args, UWorld* World) \
+	{ \
+		if (Args.Num() > 0) \
+		{ \
+			NetEmulationHelper::CreatePersistentSimulationSettings(); \
+			FString CmdParams = FString::Printf(TEXT(#CommandName##"=%s"), *(Args[0])); \
+			NetEmulationHelper::PersistentPacketSimulationSettings.GetValue().ParseSettings(*CmdParams, nullptr); \
+			NetEmulationHelper::ApplySimulationSettingsOnNetDrivers(World, NetEmulationHelper::PersistentPacketSimulationSettings.GetValue()); \
+		} \
+	})); \
+
+BUILD_NETEMULATION_CONSOLE_COMMAND(PktLoss, "Simulates network packet loss");
+BUILD_NETEMULATION_CONSOLE_COMMAND(PktOrder, "Simulates network packets received out of order");
+BUILD_NETEMULATION_CONSOLE_COMMAND(PktDup, "Simulates sending/receiving duplicate network packets");
+BUILD_NETEMULATION_CONSOLE_COMMAND(PktLag, "Simulates network packet lag");
+BUILD_NETEMULATION_CONSOLE_COMMAND(PktLagVariance, "Simulates variable network packet lag");
+BUILD_NETEMULATION_CONSOLE_COMMAND(PktLagMin, "Sets minimum outgoing packet latency");
+BUILD_NETEMULATION_CONSOLE_COMMAND(PktLagMax, "Sets maximum outgoing packet latency)");
+BUILD_NETEMULATION_CONSOLE_COMMAND(PktIncomingLagMin, "Sets minimum incoming packet latency");
+BUILD_NETEMULATION_CONSOLE_COMMAND(PktIncomingLagMax, "Sets maximum incoming packet latency");
+BUILD_NETEMULATION_CONSOLE_COMMAND(PktIncomingLoss, "Simulates incoming packet loss");
+
+#endif //#if DO_ENABLE_NET_TEST
+
 /**
  * Exec handler that routes online specific execs to the proper subsystem
  *
@@ -5673,7 +5745,7 @@ static bool NetDriverExec(UWorld* InWorld, const TCHAR* Cmd, FOutputDevice& Ar)
 	bool bHandled = false;
 
 	// Ignore any execs that don't start with NET
-	if (FParse::Command(&Cmd, TEXT("NET")) || FParse::Command(&Cmd, TEXT("NETEMULATION")))
+	if (FParse::Command(&Cmd, TEXT("NET")))
 	{
 		UNetDriver* NamedDriver = NULL;
 		TCHAR TokenStr[128];
