Removed left over PRAGMA_DISABLE_OPTIMIZATION

#rnx

diff --git a/Engine/Source/Runtime/Engine/Private/TextureRenderTarget2D.cpp b/Engine/Source/Runtime/Engine/Private/TextureRenderTarget2D.cpp
--- a/Engine/Source/Runtime/Engine/Private/TextureRenderTarget2D.cpp
+++ b/Engine/Source/Runtime/Engine/Private/TextureRenderTarget2D.cpp
@@ -316,7 +316,6 @@ UTexture2D* UTextureRenderTarget2D::ConstructTexture2D(UObject* Outer, const FSt
 	return Result;
 }
 
-PRAGMA_DISABLE_OPTIMIZATION
 void UTextureRenderTarget2D::UpdateTexture2D(UTexture2D* InTexture2D, ETextureSourceFormat InTextureFormat, uint32 Flags, TArray<uint8>* AlphaOverride)
 {
 #if WITH_EDITOR
