Temporarily disabling macro build console commands in NetDriver due to Mac compilation error

#rnx

diff --git a/Engine/Source/Runtime/Engine/Private/NetworkDriver.cpp b/Engine/Source/Runtime/Engine/Private/NetworkDriver.cpp
--- a/Engine/Source/Runtime/Engine/Private/NetworkDriver.cpp
+++ b/Engine/Source/Runtime/Engine/Private/NetworkDriver.cpp
@@ -5706,7 +5706,7 @@ FAutoConsoleCommandWithWorld NetEmulationOff(TEXT("NetEmulation.Off"), TEXT("Tur
 	NetEmulationHelper::ApplySimulationSettingsOnNetDrivers(World, NetEmulationHelper::PersistentPacketSimulationSettings.GetValue());
 }));
 
-#define BUILD_NETEMULATION_CONSOLE_COMMAND(CommandName, CommandHelp) FAutoConsoleCommandWithWorldAndArgs NetEmulation##CommandName(TEXT("NetEmulation."#CommandName), TEXT(CommandHelp), \
+/*#define BUILD_NETEMULATION_CONSOLE_COMMAND(CommandName, CommandHelp) FAutoConsoleCommandWithWorldAndArgs NetEmulation##CommandName(TEXT("NetEmulation."#CommandName), TEXT(CommandHelp), \
 	FConsoleCommandWithWorldAndArgsDelegate::CreateStatic([](const TArray<FString>& Args, UWorld* World) \
 	{ \
 		if (Args.Num() > 0) \
@@ -5727,7 +5727,7 @@ BUILD_NETEMULATION_CONSOLE_COMMAND(PktLagMin, "Sets minimum outgoing packet late
 BUILD_NETEMULATION_CONSOLE_COMMAND(PktLagMax, "Sets maximum outgoing packet latency)");
 BUILD_NETEMULATION_CONSOLE_COMMAND(PktIncomingLagMin, "Sets minimum incoming packet latency");
 BUILD_NETEMULATION_CONSOLE_COMMAND(PktIncomingLagMax, "Sets maximum incoming packet latency");
-BUILD_NETEMULATION_CONSOLE_COMMAND(PktIncomingLoss, "Simulates incoming packet loss");
+BUILD_NETEMULATION_CONSOLE_COMMAND(PktIncomingLoss, "Simulates incoming packet loss");*/
 
 #endif //#if DO_ENABLE_NET_TEST
 
