Fix runtime virtual texture baking to support BC5 normal maps

diff --git a/Engine/Source/Runtime/Engine/Classes/VT/RuntimeVirtualTexture.h b/Engine/Source/Runtime/Engine/Classes/VT/RuntimeVirtualTexture.h
--- a/Engine/Source/Runtime/Engine/Classes/VT/RuntimeVirtualTexture.h
+++ b/Engine/Source/Runtime/Engine/Classes/VT/RuntimeVirtualTexture.h
@@ -80,10 +80,10 @@ public:
 
 	/** Returns number of texture layers in the virtual texture */
 	int32 GetLayerCount() const;
+	/** Returns the texture format for the virtual texture layer */
+	EPixelFormat GetLayerFormat(int32 LayerIndex) const;
 	/** Return true if the virtual texture layer should be sampled as sRGB */
 	bool IsLayerSRGB(int32 LayerIndex) const;
-	/** Return true if the virtual texture layer contains an alpha channel */
-	bool IsLayerAlpha(int32 LayerIndex) const;
 
 	/** (Re)Initialize this object. Call this whenever we modify the producer or transform. */
 	void Initialize(IVirtualTexture* InProducer, FTransform const& VolumeToWorld);
diff --git a/Engine/Source/Runtime/Engine/Private/VT/RuntimeVirtualTexture.cpp b/Engine/Source/Runtime/Engine/Private/VT/RuntimeVirtualTexture.cpp
--- a/Engine/Source/Runtime/Engine/Private/VT/RuntimeVirtualTexture.cpp
+++ b/Engine/Source/Runtime/Engine/Private/VT/RuntimeVirtualTexture.cpp
@@ -219,32 +219,10 @@ void URuntimeVirtualTexture::GetProducerDescription(FVTProducerDescription& OutD
 	OutDesc.BlockHeightInTiles = Height / GetTileSize();
 	OutDesc.MaxLevel = FMath::Max((int32)FMath::CeilLogTwo(FMath::Max(OutDesc.BlockWidthInTiles, OutDesc.BlockHeightInTiles)) - GetRemoveLowMips(), 0);
 
-	// Set layer description based on material type
-	switch (MaterialType)
+	OutDesc.NumLayers = GetLayerCount();
+	for (int32 Layer = 0; Layer < OutDesc.NumLayers; Layer++)
 	{
-	case ERuntimeVirtualTextureMaterialType::BaseColor:
-		OutDesc.NumLayers = 1;
-		OutDesc.LayerFormat[0] = bCompressTextures ? PF_DXT1 : PF_B8G8R8A8;
-		break;
-	case ERuntimeVirtualTextureMaterialType::BaseColor_Normal:
-		OutDesc.NumLayers = 2;
-		OutDesc.LayerFormat[0] = bCompressTextures ? PF_DXT1 : PF_B8G8R8A8;
-		OutDesc.LayerFormat[1] = bCompressTextures ? PF_BC5 : PF_B8G8R8A8;
-		break;
-	case ERuntimeVirtualTextureMaterialType::BaseColor_Normal_Specular:
-		OutDesc.NumLayers = 2;
-		OutDesc.LayerFormat[0] = bCompressTextures ? PF_DXT1 : PF_B8G8R8A8;
-		OutDesc.LayerFormat[1] = bCompressTextures ? PF_DXT5 : PF_B8G8R8A8;
-		break;
-	case ERuntimeVirtualTextureMaterialType::WorldHeight:
-		OutDesc.NumLayers = 1;
-		OutDesc.LayerFormat[0] = PF_G16;
-		break;
-	default:
-		checkf(0, TEXT("Invalid Runtime Virtual Texture setup: %s, %d"), *GetName(), MaterialType);
-		OutDesc.NumLayers = 1;
-		OutDesc.LayerFormat[0] = PF_B8G8R8A8;
-		break;
+		OutDesc.LayerFormat[Layer] = GetLayerFormat(Layer);
 	}
 }
 
@@ -259,47 +237,67 @@ int32 URuntimeVirtualTexture::GetLayerCount() const
 	case ERuntimeVirtualTextureMaterialType::BaseColor_Normal_Specular:
 		return 2;
 	default:
-		// Implement logic for any missing material types
-		check(false);
+		break;
 	}
 
+	// Implement logic for any missing material types
+	check(false);
 	return 1;
 }
 
-bool URuntimeVirtualTexture::IsLayerSRGB(int32 LayerIndex) const
+EPixelFormat URuntimeVirtualTexture::GetLayerFormat(int32 LayerIndex) const
 {
-	switch (MaterialType)
+	if (LayerIndex == 0)
 	{
-	case ERuntimeVirtualTextureMaterialType::BaseColor:
-	case ERuntimeVirtualTextureMaterialType::BaseColor_Normal:
-	case ERuntimeVirtualTextureMaterialType::BaseColor_Normal_Specular:
-		// Only BaseColor layer is sRGB
-		return LayerIndex == 0;
-	case ERuntimeVirtualTextureMaterialType::WorldHeight:
-		return false;
-	default:
-		// Implement logic for any missing material types
-		check(false);
+		switch (MaterialType)
+		{
+		case ERuntimeVirtualTextureMaterialType::BaseColor:
+		case ERuntimeVirtualTextureMaterialType::BaseColor_Normal:
+		case ERuntimeVirtualTextureMaterialType::BaseColor_Normal_Specular:
+			return bCompressTextures ? PF_DXT1 : PF_B8G8R8A8;
+		case ERuntimeVirtualTextureMaterialType::WorldHeight:
+			return PF_G16;
+		default:
+			break;
+		}
+	}
+	else if (LayerIndex == 1)
+	{
+		switch (MaterialType)
+		{
+		case ERuntimeVirtualTextureMaterialType::BaseColor_Normal:
+			return bCompressTextures ? PF_BC5 : PF_B8G8R8A8;
+		case ERuntimeVirtualTextureMaterialType::BaseColor_Normal_Specular:
+			return bCompressTextures ? PF_DXT5 : PF_B8G8R8A8;
+		case ERuntimeVirtualTextureMaterialType::BaseColor:
+		case ERuntimeVirtualTextureMaterialType::WorldHeight:
+		default:
+			break;
+		}
 	}
 
-	return false;
+	// Implement logic for any missing material types
+	check(false);
+	return PF_B8G8R8A8;
 }
 
-bool URuntimeVirtualTexture::IsLayerAlpha(int32 LayerIndex) const
+bool URuntimeVirtualTexture::IsLayerSRGB(int32 LayerIndex) const
 {
 	switch (MaterialType)
 	{
 	case ERuntimeVirtualTextureMaterialType::BaseColor:
 	case ERuntimeVirtualTextureMaterialType::BaseColor_Normal:
+	case ERuntimeVirtualTextureMaterialType::BaseColor_Normal_Specular:
+		// Only BaseColor layer is sRGB
+		return LayerIndex == 0;
 	case ERuntimeVirtualTextureMaterialType::WorldHeight:
 		return false;
-	case ERuntimeVirtualTextureMaterialType::BaseColor_Normal_Specular:
-		return LayerIndex == 1;
 	default:
-		// Implement logic for any missing material types
-		check(false);
+		break;
 	}
 
+	// Implement logic for any missing material types
+	check(false);
 	return false;
 }
 
@@ -433,10 +431,13 @@ void URuntimeVirtualTexture::InitializeStreamingTexture(uint32 InSizeX, uint32 I
 	const int32 LayerCount = GetLayerCount();
 	for (int32 Layer = 0; Layer < LayerCount; Layer++)
 	{
+		EPixelFormat LayerFormat = GetLayerFormat(Layer);
+
 		FTextureFormatSettings FormatSettings;
 		FormatSettings.SRGB = IsLayerSRGB(Layer);
 		FormatSettings.CompressionNone = !bCompressTextures;
-		FormatSettings.CompressionNoAlpha = !IsLayerAlpha(Layer);
+		FormatSettings.CompressionNoAlpha = LayerFormat == PF_DXT1 || LayerFormat == PF_BC5;
+		FormatSettings.CompressionSettings = LayerFormat == PF_BC5 ? TC_Normalmap : TC_Default;
 		StreamingTexture->SetLayerFormatSettings(Layer, FormatSettings);
 	}
 
