RenderAssetStreaming:
- Ensured that force fully load heuristic would not be affected by load order in cases the data is overbudget.
- Reset unknow heuristic after the assets haven't been used for 90s. Can be worked around by using ForceFullyLoad logic.

diff --git a/Engine/Source/Runtime/Engine/Private/Streaming/AsyncTextureStreaming.cpp b/Engine/Source/Runtime/Engine/Private/Streaming/AsyncTextureStreaming.cpp
--- a/Engine/Source/Runtime/Engine/Private/Streaming/AsyncTextureStreaming.cpp
+++ b/Engine/Source/Runtime/Engine/Private/Streaming/AsyncTextureStreaming.cpp
@@ -255,7 +255,9 @@ void FAsyncRenderAssetStreamingData::UpdatePerfectWantedMips_Async(FStreamingRen
 				}
 			}
 
-			if (StreamingRenderAsset.bUseUnkownRefHeuristic)
+			// Ignore bUseUnkownRefHeuristic if they haven't been used in the last 90 sec.
+			// If critical, it must be implemented using the ForceFullyLoad logic.
+			if (StreamingRenderAsset.bUseUnkownRefHeuristic && StreamingRenderAsset.LastRenderTime < 90.0f)
 			{
 				if (bOutputToLog) UE_LOG(LogContentStreaming, Log,  TEXT("  UnkownRef"));
 				MaxSize = FMath::Max<int32>(MaxSize, MaxAllowedSize); // affected by HiddenPrimitiveScale
@@ -716,16 +718,13 @@ void FRenderAssetStreamingMipCalcTask::UpdateLoadAndCancelationRequests_Async(in
 				TempMemoryUsed += TempMemoryRequired;
 			}
 		}
-		else if (StreamingRenderAsset.WantedMips > StreamingRenderAsset.ResidentMips
-			&& (MemoryUsed < MemoryBudget || StreamingRenderAsset.NumForcedMips > StreamingRenderAsset.ResidentMips)
-			&& TempMemoryUsed < TempMemoryBudget)
+		else if (StreamingRenderAsset.WantedMips > StreamingRenderAsset.ResidentMips && TempMemoryUsed < TempMemoryBudget)
 		{
 			const int64 UsedMemoryRequired = StreamingRenderAsset.GetSize(StreamingRenderAsset.WantedMips) - StreamingRenderAsset.GetSize(StreamingRenderAsset.ResidentMips);
 			const int64 TempMemoryRequired = StreamingRenderAsset.GetSize(StreamingRenderAsset.WantedMips);
 
 			// Respect the temporary budget unless this is the first load request. This allows a single mip update of any size.
-			if ((StreamingRenderAsset.NumForcedMips > StreamingRenderAsset.ResidentMips || MemoryUsed + UsedMemoryRequired <= MemoryBudget)
-				&& (TempMemoryUsed + TempMemoryRequired <= TempMemoryBudget || !LoadRequests.Num()))
+			if (TempMemoryUsed + TempMemoryRequired <= TempMemoryBudget || !LoadRequests.Num())
 			{
 				LoadRequests.Add(AssetIndex);
 	
