Fix for bug added from new VolumeFader interaction with AudioComponents where looping sounds that were faded out were not reporting as inactive
- Minor clean-up with caching off Fader reference where used

diff --git a/Engine/Source/Runtime/AudioMixer/Private/DSP/VolumeFader.cpp b/Engine/Source/Runtime/AudioMixer/Private/DSP/VolumeFader.cpp
--- a/Engine/Source/Runtime/AudioMixer/Private/DSP/VolumeFader.cpp
+++ b/Engine/Source/Runtime/AudioMixer/Private/DSP/VolumeFader.cpp
@@ -70,6 +70,11 @@ namespace Audio
 		ActiveDuration = 0.0f;
 	}
 
+	float FVolumeFader::GetActiveDuration() const
+	{
+		return ActiveDuration;
+	}
+
 	float FVolumeFader::GetVolume() const
 	{
 		return AlphaToVolume(Alpha, FadeCurve);
@@ -185,7 +190,6 @@ namespace Audio
 	void FVolumeFader::SetVolume(float InVolume)
 	{
 		Alpha = InVolume;
-		ActiveDuration = -1.0f;
 		Elapsed = 0.0f;
 		FadeCurve = EFaderCurve::Linear;
 		FadeDuration = -1.0f;
@@ -218,7 +222,6 @@ namespace Audio
 			Target = Audio::ConvertToDecibels(InVolume, DecibelFloor);
 		}
 
-		ActiveDuration = -1.0f;
 		Elapsed = 0.0f;
 		FadeCurve = InCurve;
 		FadeDuration = InDuration;
diff --git a/Engine/Source/Runtime/AudioMixer/Public/DSP/VolumeFader.h b/Engine/Source/Runtime/AudioMixer/Public/DSP/VolumeFader.h
--- a/Engine/Source/Runtime/AudioMixer/Public/DSP/VolumeFader.h
+++ b/Engine/Source/Runtime/AudioMixer/Public/DSP/VolumeFader.h
@@ -32,7 +32,7 @@ namespace Audio
 		FVolumeFader();
 
 		/**
-		 * Activates the fader if currently deactivated.
+		 * Activates the fader if currently deactivated. Sets duration to indefinite (-1.0f)
 		 */
 		void Activate();
 
@@ -42,6 +42,11 @@ namespace Audio
 		 */
 		void Deactivate();
 
+		/**
+		 * Gets time fader is to remain active (-1.0f if always active).
+		 */
+		float GetActiveDuration() const;
+
 		/**
 		 * Returns current volume of fader
 		 */
@@ -94,7 +99,10 @@ namespace Audio
 
 		/**
 		 * Sets the duration the fader is to be active in the future,
-		 * after which point the fader is disabled.
+		 * after which point the fader is disabled.  When disabled,
+		 * fader will hold the current fade state until activated,
+		 * at which point it will continue applying the fade over the
+		 * remainder of the fade duration.
 		 */
 		void SetActiveDuration(float InDuration);
 
@@ -114,7 +122,7 @@ namespace Audio
 		void StopFade();
 
 		/**
-		 * Updates the current value with the given delta.
+		 * Updates the fader's state with the given delta in time since last update.
 		 */
 		void Update(float InDeltaTime);
 
diff --git a/Engine/Source/Runtime/Engine/Private/Components/AudioComponent.cpp b/Engine/Source/Runtime/Engine/Private/Components/AudioComponent.cpp
--- a/Engine/Source/Runtime/Engine/Private/Components/AudioComponent.cpp
+++ b/Engine/Source/Runtime/Engine/Private/Components/AudioComponent.cpp
@@ -500,8 +500,9 @@ void UAudioComponent::PlayInternal(const float StartTime, const float FadeInDura
 			NewActiveSound.MaxDistance = MaxDistance;
 			NewActiveSound.InstanceParameters = InstanceParameters;
 
-			NewActiveSound.ComponentVolumeFader.SetVolume(0.0f); // Init to 0.0f to fade as default is 1.0f
-			NewActiveSound.ComponentVolumeFader.StartFade(FadeVolumeLevel, FadeInDuration, static_cast<Audio::EFaderCurve>(FadeCurve));
+			Audio::FVolumeFader& Fader = NewActiveSound.ComponentVolumeFader;
+			Fader.SetVolume(0.0f); // Init to 0.0f to fade as default is 1.0f
+			Fader.StartFade(FadeVolumeLevel, FadeInDuration, static_cast<Audio::EFaderCurve>(FadeCurve));
 
 			// Bump ActiveCount... this is used to determine if an audio component is still active after a sound reports back as completed
 			++ActiveCount;
@@ -585,7 +586,8 @@ void UAudioComponent::AdjustVolumeInternal(float AdjustVolumeDuration, float Adj
 			return;
 		}
 
-		const float InitialTargetVolume = ActiveSound->ComponentVolumeFader.GetTargetVolume();
+		Audio::FVolumeFader& Fader = ActiveSound->ComponentVolumeFader;
+		const float InitialTargetVolume = Fader.GetTargetVolume();
 
 		// Ignore fade out request if requested volume is higher than current target.
 		if (bIsFadeOut && AdjustVolumeLevel >= InitialTargetVolume)
@@ -593,6 +595,7 @@ void UAudioComponent::AdjustVolumeInternal(float AdjustVolumeDuration, float Adj
 			return;
 		}
 
+		const bool ToZeroVolume = FMath::IsNearlyZero(AdjustVolumeLevel);
 		if (ActiveSound->FadeOut == FActiveSound::EFadeOut::Concurrency)
 		{
 			// Ignore adjust volume request if non-zero and currently voice stealing.
@@ -602,17 +605,27 @@ void UAudioComponent::AdjustVolumeInternal(float AdjustVolumeDuration, float Adj
 			}
 
 			// Ignore request of longer fade out than active target if active is concurrency (voice stealing) fade.
-			if (AdjustVolumeDuration > ActiveSound->ComponentVolumeFader.GetFadeDuration())
+			if (AdjustVolumeDuration > Fader.GetFadeDuration())
 			{
 				return;
 			}
 		}
 		else
 		{
-			ActiveSound->FadeOut = FMath::IsNearlyZero(AdjustVolumeLevel) ? FActiveSound::EFadeOut::User : FActiveSound::EFadeOut::None;
+			ActiveSound->FadeOut = bIsFadeOut || ToZeroVolume ? FActiveSound::EFadeOut::User : FActiveSound::EFadeOut::None;
 		}
 
-		ActiveSound->ComponentVolumeFader.StartFade(AdjustVolumeLevel, AdjustVolumeDuration, static_cast<Audio::EFaderCurve>(FadeCurve));
+		if (bIsFadeOut || ToZeroVolume)
+		{
+			// If negative, active indefinitely, so always make sure set to minimum positive value for active fade.
+			const float OldActiveDuration = Fader.GetActiveDuration();
+			const float NewActiveDuration = OldActiveDuration < 0.0f
+				? AdjustVolumeDuration
+				: FMath::Min(OldActiveDuration, AdjustVolumeDuration);
+			Fader.SetActiveDuration(NewActiveDuration);
+		}
+
+		Fader.StartFade(AdjustVolumeLevel, AdjustVolumeDuration, static_cast<Audio::EFaderCurve>(FadeCurve));
 	}, GET_STATID(STAT_AudioAdjustVolume));
 }
 
@@ -677,14 +690,15 @@ void UAudioComponent::StopDelayed(float DelayTime)
 				*StoppingSound->GetName());
 		}
 
+		Audio::FVolumeFader& Fader = ActiveSound->ComponentVolumeFader;
 		switch (ActiveSound->FadeOut)
 		{
 			case FActiveSound::EFadeOut::Concurrency:
 			{
 				// Ignore request of longer fade out than active target if active is concurrency (voice stealing) fade.
-				if (DelayTime < ActiveSound->ComponentVolumeFader.GetFadeDuration())
+				if (DelayTime < Fader.GetFadeDuration())
 				{
-					ActiveSound->ComponentVolumeFader.SetActiveDuration(DelayTime);
+					Fader.SetActiveDuration(DelayTime);
 				}
 			}
 			break;
@@ -694,7 +708,7 @@ void UAudioComponent::StopDelayed(float DelayTime)
 			default:
 			{
 				ActiveSound->FadeOut = FActiveSound::EFadeOut::User;
-				ActiveSound->ComponentVolumeFader.SetActiveDuration(DelayTime);
+				Fader.SetActiveDuration(DelayTime);
 			}
 			break;
 		}
