non-unity cis fix

diff --git a/Engine/Source/Runtime/Engine/Classes/AI/NavigationSystemConfig.h b/Engine/Source/Runtime/Engine/Classes/AI/NavigationSystemConfig.h
--- a/Engine/Source/Runtime/Engine/Classes/AI/NavigationSystemConfig.h
+++ b/Engine/Source/Runtime/Engine/Classes/AI/NavigationSystemConfig.h
@@ -2,9 +2,7 @@
 
 #pragma once
 
-#include "CoreMinimal.h"
-#include "UObject/Object.h"
-#include "UObject/SoftObjectPath.h"
+#include "AI/Navigation/NavigationTypes.h"
 #include "NavigationSystemConfig.generated.h"
 
 class UNavigationSystemBase;
