Fixed issue with virtualized loops delaying retrigger when turning off the listener attenuation override.  Turning off the position override could cause a large delta in the position we use for attenuation, so we now check any virutalized loops to see if they should retrigger.

(Also removing my previous fix)

diff --git a/Engine/Source/Runtime/Engine/Private/AudioDevice.cpp b/Engine/Source/Runtime/Engine/Private/AudioDevice.cpp
--- a/Engine/Source/Runtime/Engine/Private/AudioDevice.cpp
+++ b/Engine/Source/Runtime/Engine/Private/AudioDevice.cpp
@@ -2907,8 +2907,14 @@ void FAudioDevice::SetListenerAttenuationOverride(const FVector AttenuationPosit
 	FAudioDevice* AudioDevice = this;
 	FAudioThread::RunCommandOnAudioThread([AudioDevice, AttenuationPosition]()
 	{
+		const bool bPrevAttenuationOverride = AudioDevice->bUseListenerAttenuationOverride;
 		AudioDevice->bUseListenerAttenuationOverride = true;
 		AudioDevice->ListenerAttenuationOverride = AttenuationPosition;
+
+		if (!bPrevAttenuationOverride)
+		{
+			AudioDevice->UpdateVirtualLoops(true);
+		}
 	}, GET_STATID(STAT_AudioSetListenerAttenuationOverride));
 }
 
@@ -2927,6 +2933,7 @@ void FAudioDevice::ClearListenerAttenuationOverride()
 	FAudioThread::RunCommandOnAudioThread([AudioDevice]()
 	{
 		AudioDevice->bUseListenerAttenuationOverride = false;
+		AudioDevice->UpdateVirtualLoops(true);
 	}, GET_STATID(STAT_AudioClearListenerAttenuationOverride));
 }
 
