Adding a Get Settings node from a Post Process Volume in the blueprint editor no longer crashes the editor.

Requested by [REDACTED_PERSON].

diff --git a/Engine/Source/Runtime/Engine/Private/PostProcessVolume.cpp b/Engine/Source/Runtime/Engine/Private/PostProcessVolume.cpp
--- a/Engine/Source/Runtime/Engine/Private/PostProcessVolume.cpp
+++ b/Engine/Source/Runtime/Engine/Private/PostProcessVolume.cpp
@@ -85,10 +85,16 @@ bool APostProcessVolume::CanEditChange(const UProperty* InProperty) const
 		// Settings, can be shared for multiple objects types (volume, component, camera, player)
 		PRAGMA_DISABLE_DEPRECATION_WARNINGS
 		{
-			FSceneInterface* Scene = GetWorld()->Scene;
-			bool bIsMobile = Scene->GetShadingPath(Scene->GetFeatureLevel()) == EShadingPath::Mobile;
-			bool bHaveGaussianDOF = bIsMobile;
+			bool bIsMobile = false;
+
+			if (UWorld* World = GetWorld())
+			{
+				FSceneInterface* Scene = World->Scene;
+				bIsMobile = Scene->GetShadingPath(Scene->GetFeatureLevel()) == EShadingPath::Mobile;
+			}
+
 			bool bHaveCinematicDOF = !bIsMobile;
+			bool bHaveGaussianDOF = bIsMobile;
 
 			if (PropertyName == GET_MEMBER_NAME_STRING_CHECKED(FPostProcessSettings, DepthOfFieldScale) ||
 				PropertyName == GET_MEMBER_NAME_STRING_CHECKED(FPostProcessSettings, DepthOfFieldNearBlurSize) ||
