Fix RVT copy bound transform
#rnx

diff --git a/Engine/Source/Runtime/Engine/Private/Components/RuntimeVirtualTextureComponent.cpp b/Engine/Source/Runtime/Engine/Private/Components/RuntimeVirtualTextureComponent.cpp
--- a/Engine/Source/Runtime/Engine/Private/Components/RuntimeVirtualTextureComponent.cpp
+++ b/Engine/Source/Runtime/Engine/Private/Components/RuntimeVirtualTextureComponent.cpp
@@ -78,10 +78,9 @@ void URuntimeVirtualTextureComponent::SetTransformToBounds()
 		// Calculate the bounds in our local rotation space translated to the BoundsSourceActor center
 		const FQuat TargetRotation = GetComponentToWorld().GetRotation();
 		const FVector InitialPosition = BoundsSourceActor->GetComponentsBoundingBox().GetCenter();
-		const FVector InitialScale = FVector(0.5f, 0.5, 1.f);
 
 		FTransform LocalTransform;
-		LocalTransform.SetComponents(TargetRotation, InitialPosition, InitialScale);
+		LocalTransform.SetComponents(TargetRotation, InitialPosition, FVector::OneVector);
 		FTransform WorldToLocal = LocalTransform.Inverse();
 
 		FBox BoundBox(ForceInit);
@@ -103,6 +102,7 @@ void URuntimeVirtualTextureComponent::SetTransformToBounds()
 		BoundBox.GetCenterAndExtents(Origin, Extent);
 
 		Origin = LocalTransform.TransformPosition(Origin);
+		Extent *= FVector(2.f, 2.f, 1.f); // Account for ARuntimeVirtualTextureVolume:Box offset which centers it on origin
 
 		FTransform Transform;
 		Transform.SetComponents(TargetRotation, Origin, Extent);
