Touching file in hope of fixing build error

diff --git a/Engine/Source/Runtime/Engine/Private/VT/VirtualTextureBuildSettings.h b/Engine/Source/Runtime/Engine/Private/VT/VirtualTextureBuildSettings.h
--- a/Engine/Source/Runtime/Engine/Private/VT/VirtualTextureBuildSettings.h
+++ b/Engine/Source/Runtime/Engine/Private/VT/VirtualTextureBuildSettings.h
@@ -23,6 +23,6 @@ struct FVirtualTextureBuildSettings
 	UPROPERTY()
 	bool bEnableCompressZlib;
 
-	/** Initialize with default build settings. These are defined by the current project setup. */
+	/** Initialize with the default build settings. These are defined by the current project setup. */
 	ENGINE_API void Init();
 };
