Added handle support to FInputActionBinding so trying to remove an action bound to multiple delegates won't potentially remove the wrong one when the current API only supports removal by index.

#rnx

diff --git a/Engine/Source/Runtime/Engine/Classes/Components/InputComponent.h b/Engine/Source/Runtime/Engine/Classes/Components/InputComponent.h
--- a/Engine/Source/Runtime/Engine/Classes/Components/InputComponent.h
+++ b/Engine/Source/Runtime/Engine/Classes/Components/InputComponent.h
@@ -302,6 +302,9 @@ private:
 	/** Friendly name of action, e.g "jump" */
 	FName ActionName;
 
+	/** The handle for this binding action */
+	int32 Handle;
+
 public:
 	/** The delegate bound to the action */
 	FInputActionUnifiedDelegate ActionDelegate;
@@ -311,6 +314,7 @@ public:
 		, bPaired(false)
 		, KeyEvent(EInputEvent::IE_Pressed)
 		, ActionName(NAME_None)
+		, Handle(INDEX_NONE)
 	{ }
 
 	FInputActionBinding(const FName InActionName, const  EInputEvent InKeyEvent)
@@ -318,10 +322,22 @@ public:
 		, bPaired(false)
 		, KeyEvent(InKeyEvent)
 		, ActionName(InActionName)
+		, Handle(INDEX_NONE)
 	{ }
 
 	FName GetActionName() const { return ActionName; }
 	bool IsPaired() const { return bPaired; }
+	int32 GetHandle() const { return Handle; }
+
+	bool operator==(const FInputActionBinding& Rhs)
+	{
+		return (IsValid() && GetHandle() == Rhs.GetHandle());
+	}
+
+	/** Indicates GenerateNewHandle was called */
+	bool IsValid() { return (Handle != INDEX_NONE); }
+
+	void GenerateNewHandle();
 
 	friend class UInputComponent;
 };
@@ -786,6 +802,23 @@ public:
 	 */
 	void RemoveActionBinding( const int32 BindingIndex );
 
+	/**
+	 * Removes the action binding at the specified handle.
+	 *
+	 * @param Handle The handle of the binding to remove.
+	 * @see AddActionBinding, ClearActionBindings, GetActionBinding, GetNumActionBindings
+	 */
+	void RemoveActionBindingForHandle(const int32 Handle);
+
+	/**
+	 * Removes the action binding (index need for multi-name fixups).
+	 *
+	 * @param BindingToRemove The binding to remove.
+	 * @param BindingIndex The binding's index for actions with the same name to fixup their data.
+	 * @see AddActionBinding, ClearActionBindings, GetActionBinding, GetNumActionBindings
+	 */
+	void RemoveActionBinding(const FInputActionBinding &BindingToRemove, const int32 BindingIndex);
+
 	/** Clears all cached binding values. */
 	void ClearBindingValues();
 
diff --git a/Engine/Source/Runtime/Engine/Private/Components/InputComponent.cpp b/Engine/Source/Runtime/Engine/Private/Components/InputComponent.cpp
--- a/Engine/Source/Runtime/Engine/Private/Components/InputComponent.cpp
+++ b/Engine/Source/Runtime/Engine/Private/Components/InputComponent.cpp
@@ -5,6 +5,13 @@
 #include "GameFramework/PlayerInput.h"
 
 
+void FInputActionBinding::GenerateNewHandle()
+{
+	// Must be in C++ to avoid duplicate statics accross execution units
+	static int32 GHandle = 1;
+	Handle = GHandle++;
+}
+
 /* UInputComponent interface
  *****************************************************************************/
 
@@ -200,6 +207,7 @@ FInputActionBinding& UInputComponent::AddActionBinding( FInputActionBinding InBi
 {
 	ActionBindings.Add(MakeShared<FInputActionBinding>(MoveTemp(InBinding)));
 	FInputActionBinding& Binding = *ActionBindings.Last().Get();
+	Binding.GenerateNewHandle();
 
 	if (Binding.KeyEvent == IE_Pressed || Binding.KeyEvent == IE_Released)
 	{
@@ -249,47 +257,64 @@ void UInputComponent::RemoveActionBinding( const int32 BindingIndex )
 	if (BindingIndex >= 0 && BindingIndex < ActionBindings.Num())
 	{
 		const FInputActionBinding& BindingToRemove = *ActionBindings[BindingIndex].Get();
+		RemoveActionBinding(BindingToRemove, BindingIndex);
+	}
+}
 
-		// Potentially need to clear some pairings
-		if (BindingToRemove.bPaired)
+void UInputComponent::RemoveActionBindingForHandle(const int32 Handle)
+{
+	for (int32 ActionIndex = 0; ActionIndex < ActionBindings.Num(); ++ActionIndex)
+	{
+		const FInputActionBinding& BindingToRemove = *ActionBindings[ActionIndex].Get();
+		if (BindingToRemove.GetHandle() == Handle)
 		{
-			TArray<int32> IndicesToClear;
-			const EInputEvent PairedEvent = (BindingToRemove.KeyEvent == IE_Pressed ? IE_Released : IE_Pressed);
+			RemoveActionBinding(BindingToRemove, ActionIndex);
+			return;
+		}
+	}
+}
 
-			for (int32 ActionIndex = 0; ActionIndex < ActionBindings.Num(); ++ActionIndex)
+void UInputComponent::RemoveActionBinding(const FInputActionBinding &BindingToRemove, const int32 BindingIndex)
+{
+	// Potentially need to clear some pairings
+	if (BindingToRemove.bPaired)
+	{
+		TArray<int32> IndicesToClear;
+		const EInputEvent PairedEvent = (BindingToRemove.KeyEvent == IE_Pressed ? IE_Released : IE_Pressed);
+
+		for (int32 ActionIndex = 0; ActionIndex < ActionBindings.Num(); ++ActionIndex)
+		{
+			if (ActionIndex != BindingIndex)
 			{
-				if (ActionIndex != BindingIndex)
+				const FInputActionBinding& ActionBinding = *ActionBindings[ActionIndex].Get();
+				if (ActionBinding.ActionName == BindingToRemove.ActionName)
 				{
-					const FInputActionBinding& ActionBinding = *ActionBindings[ActionIndex].Get();
-					if (ActionBinding.ActionName == BindingToRemove.ActionName)
+					// If we find another of the same key event then the pairing is intact so we're done
+					if (ActionBinding.KeyEvent == BindingToRemove.KeyEvent)
+					{
+						IndicesToClear.Empty();
+						break;
+					}
+					// Otherwise we may need to clear the pairing so track the index
+					else if (ActionBinding.KeyEvent == PairedEvent)
 					{
-						// If we find another of the same key event then the pairing is intact so we're done
-						if (ActionBinding.KeyEvent == BindingToRemove.KeyEvent)
-						{
-							IndicesToClear.Empty();
-							break;
-						}
-						// Otherwise we may need to clear the pairing so track the index
-						else if (ActionBinding.KeyEvent == PairedEvent)
-						{
-							IndicesToClear.Add(ActionIndex);
-						}
+						IndicesToClear.Add(ActionIndex);
 					}
 				}
 			}
-
-			for (int32 ClearIndex = 0; ClearIndex < IndicesToClear.Num(); ++ClearIndex)
-			{
-				ActionBindings[IndicesToClear[ClearIndex]]->bPaired = false;
-			}
 		}
 
-		ActionBindings.RemoveAt(BindingIndex, 1, false);
-		for (FCachedKeyToActionInfo& CachedInfo : CachedKeyToActionInfo)
+		for (int32 ClearIndex = 0; ClearIndex < IndicesToClear.Num(); ++ClearIndex)
 		{
-			CachedInfo.KeyMapBuiltForIndex = 0;
+			ActionBindings[IndicesToClear[ClearIndex]]->bPaired = false;
 		}
 	}
+
+	ActionBindings.RemoveAt(BindingIndex, 1, false);
+	for (FCachedKeyToActionInfo& CachedInfo : CachedKeyToActionInfo)
+	{
+		CachedInfo.KeyMapBuiltForIndex = 0;
+	}
 }
 
 void UInputComponent::ClearBindingValues()
