Fix missing linked child in SoundCueTemplate helper function

diff --git a/Engine/Source/Runtime/Engine/Classes/Sound/SoundCueTemplate.h b/Engine/Source/Runtime/Engine/Classes/Sound/SoundCueTemplate.h
--- a/Engine/Source/Runtime/Engine/Classes/Sound/SoundCueTemplate.h
+++ b/Engine/Source/Runtime/Engine/Classes/Sound/SoundCueTemplate.h
@@ -134,7 +134,9 @@ protected:
 	{
 		check(InputPinIndex >= 0 && InputPinIndex <= InParentNode.GetMaxChildNodes());
 
-		InParentNode.InsertChildNode(InParentNode.ChildNodes.Num());
+		const int32 NewChildNodeIndex = InParentNode.ChildNodes.Num();
+		InParentNode.InsertChildNode(NewChildNodeIndex);
+		InParentNode.ChildNodes[NewChildNodeIndex] = &InChildNode;
 
 		UEdGraphNode* OutNodeGraph = InParentNode.GetGraphNode();
 		check(OutNodeGraph);
