Added blueprint controls for all sky parameters.

diff --git a/Engine/Source/Runtime/Engine/Classes/Components/SkyAtmosphereComponent.h b/Engine/Source/Runtime/Engine/Classes/Components/SkyAtmosphereComponent.h
--- a/Engine/Source/Runtime/Engine/Classes/Components/SkyAtmosphereComponent.h
+++ b/Engine/Source/Runtime/Engine/Classes/Components/SkyAtmosphereComponent.h
@@ -125,6 +125,36 @@ class USkyAtmosphereComponent : public USceneComponent
 
 	UFUNCTION(BlueprintCallable, Category = "Rendering")
 	void OverrideAtmosphereLightDirection(int32 AtmosphereLightIndex, const FVector& LightDirection);
+	
+	UFUNCTION(BlueprintCallable, Category = "Rendering")
+	void SetRayleighScatteringScale(float NewValue);
+	UFUNCTION(BlueprintCallable, Category = "Rendering")
+	void SetRayleighScattering(FColor NewValue);
+	UFUNCTION(BlueprintCallable, Category = "Rendering")
+	void SetRayleighExponentialDistribution(float NewValue);
+
+	UFUNCTION(BlueprintCallable, Category = "Rendering")
+	void SetMieScatteringScale(float NewValue);
+	UFUNCTION(BlueprintCallable, Category = "Rendering")
+	void SetMieScattering(FColor NewValue);
+	UFUNCTION(BlueprintCallable, Category = "Rendering")
+	void SetMieAbsorptionScale(float NewValue);
+	UFUNCTION(BlueprintCallable, Category = "Rendering")
+	void SetMieAbsorption(FColor NewValue);
+	UFUNCTION(BlueprintCallable, Category = "Rendering")
+	void SetMieAnisotropy(float NewValue);
+	UFUNCTION(BlueprintCallable, Category = "Rendering")
+	void SetMieExponentialDistribution(float NewValue);
+
+	UFUNCTION(BlueprintCallable, Category = "Rendering", meta = (DisplayName = "Set Absorption Scale"))
+	void SetOtherAbsorptionScale(float NewValue);
+	UFUNCTION(BlueprintCallable, Category = "Rendering", meta = (DisplayName = "Set Absorption"))
+	void SetOtherAbsorption(FColor NewValue);
+
+	UFUNCTION(BlueprintCallable, Category = "Rendering")
+	void SetSkyLuminanceFactor(FLinearColor NewValue);
+	UFUNCTION(BlueprintCallable, Category = "Rendering")
+	void SetAerialPespectiveViewDistanceScale(float NewValue);
 
 
 
diff --git a/Engine/Source/Runtime/Engine/Private/Components/SkyAtmosphereComponent.cpp b/Engine/Source/Runtime/Engine/Private/Components/SkyAtmosphereComponent.cpp
--- a/Engine/Source/Runtime/Engine/Private/Components/SkyAtmosphereComponent.cpp
+++ b/Engine/Source/Runtime/Engine/Private/Components/SkyAtmosphereComponent.cpp
@@ -244,6 +244,31 @@ void USkyAtmosphereComponent::OverrideAtmosphereLightDirection(int32 AtmosphereL
 	}
 }
 
+#define SKY_DECLARE_BLUEPRINT_SETFUNCTION(MemberType, MemberName) void USkyAtmosphereComponent::Set##MemberName(MemberType NewValue)\
+{\
+	if (AreDynamicDataChangesAllowed() && MemberName != NewValue)\
+	{\
+		MemberName = NewValue;\
+		MarkRenderStateDirty();\
+	}\
+}\
+
+SKY_DECLARE_BLUEPRINT_SETFUNCTION(float, RayleighScatteringScale);
+SKY_DECLARE_BLUEPRINT_SETFUNCTION(FColor, RayleighScattering);
+SKY_DECLARE_BLUEPRINT_SETFUNCTION(float, RayleighExponentialDistribution);
+
+SKY_DECLARE_BLUEPRINT_SETFUNCTION(float, MieScatteringScale);
+SKY_DECLARE_BLUEPRINT_SETFUNCTION(FColor, MieScattering);
+SKY_DECLARE_BLUEPRINT_SETFUNCTION(float, MieAbsorptionScale);
+SKY_DECLARE_BLUEPRINT_SETFUNCTION(FColor, MieAbsorption);
+SKY_DECLARE_BLUEPRINT_SETFUNCTION(float, MieAnisotropy);
+SKY_DECLARE_BLUEPRINT_SETFUNCTION(float, MieExponentialDistribution);
+
+SKY_DECLARE_BLUEPRINT_SETFUNCTION(float, OtherAbsorptionScale);
+SKY_DECLARE_BLUEPRINT_SETFUNCTION(FColor, OtherAbsorption);
+
+SKY_DECLARE_BLUEPRINT_SETFUNCTION(FLinearColor, SkyLuminanceFactor);
+SKY_DECLARE_BLUEPRINT_SETFUNCTION(float, AerialPespectiveViewDistanceScale);
 
 /*=============================================================================
 	ASkyAtmosphere implementation.
