Better fix for crash with referenced sound waves.

diff --git a/Engine/Source/Runtime/Engine/Private/AudioDevice.cpp b/Engine/Source/Runtime/Engine/Private/AudioDevice.cpp
--- a/Engine/Source/Runtime/Engine/Private/AudioDevice.cpp
+++ b/Engine/Source/Runtime/Engine/Private/AudioDevice.cpp
@@ -743,9 +743,6 @@ void FAudioDevice::AddReferencedObjects(FReferenceCollector& Collector)
 		Pair.Key->AddReferencedObjects(Collector);
 	}
 
-	// Make sure our referenced sound waves are up-to-date
-	UpdateReferencedSoundWaves();
-
 	// Make sure we don't try to delete any sound waves which may have in-flight decodes
 	Collector.AddReferencedObjects(ReferencedSoundWaves);
 
@@ -3843,32 +3840,49 @@ void FAudioDevice::StartSources(TArray<FWaveInstance*>& WaveInstances, int32 Fir
 		}
 	}
 
-	DECLARE_CYCLE_STAT(TEXT("FGameThreadAudioTask.AudioSendResults"), STAT_AudioAddReferencedSoundWaves, STATGROUP_TaskGraphTasks);
+	DECLARE_CYCLE_STAT(TEXT("FGameThreadAudioTask.AddReferencedSoundWaves"), STAT_AudioAddReferencedSoundWaves, STATGROUP_TaskGraphTasks);
 	
 	// Run a command to make sure we add the starting sounds to the referenced sound waves list
 	if (StartingSoundWaves.Num() > 0)
 	{
-		FScopeLock ReferencedSoundWaveLock(&ReferencedSoundWaveCritSec);
-
-		for (USoundWave* SoundWave : StartingSoundWaves)
+		FAudioThread::RunCommandOnGameThread([this, StartingSoundWaves]()
 		{
-			ReferencedSoundWaves_AudioThread.AddUnique(SoundWave);
-		}
+			for (USoundWave* SoundWave : StartingSoundWaves)
+			{
+				ReferencedSoundWaves.AddUnique(SoundWave);
+			}
+		}, GET_STATID(STAT_AudioAddReferencedSoundWaves));
 	}
 }
 
 void FAudioDevice::UpdateReferencedSoundWaves()
 {
-	FScopeLock ReferencedSoundWaveLock(&ReferencedSoundWaveCritSec);
-
-	for (USoundWave* SoundWave : ReferencedSoundWaves_AudioThread)
+	// On game thread, look through registered sound waves and remove if we finished precaching (and audio decompressor is cleaned up)
+	// ReferencedSoundWaves is used to make sure GC doesn't run on any sound waves that are actively pre-caching within an async task.
+	// Sounds may be loaded, kick off an async task to decompress, but never actually try to play, so GC can reclaim these while precaches are in-flight.
+	// We are also tracking when a sound wave is actively being used to generate audio in the audio render to prevent GC from happening to sounds till being used in the audio renderer.
+	for (int32 i = ReferencedSoundWaves.Num() - 1; i >= 0; --i)
 	{
-		ReferencedSoundWaves.AddUnique(SoundWave);
-	}
+		USoundWave* Wave = ReferencedSoundWaves[i];
+		bool bRemove = true;
+		// If this is null that means it was nulled out in AddReferencedObjects via mark pending kill
+		if (Wave)
+		{
+			const bool bIsPrecacheDone = (Wave->GetPrecacheState() == ESoundWavePrecacheState::Done);
+			const bool bIsGeneratingAudio = Wave->IsGeneratingAudio();
 
-	ReferencedSoundWaves_AudioThread.Reset();
-}
+			if (!bIsPrecacheDone || bIsGeneratingAudio)
+			{
+				bRemove = false;
+			}
+		}
 
+		if (bRemove)
+		{
+			ReferencedSoundWaves.RemoveAtSwap(i, 1, false);
+		}
+	}
+}
 
 void FAudioDevice::Update(bool bGameTicking)
 {
@@ -3878,31 +3892,6 @@ void FAudioDevice::Update(bool bGameTicking)
 	{
 		// Make sure our referenced sound waves is up-to-date
 		UpdateReferencedSoundWaves();
-
-		// On game thread, look through registered sound waves and remove if we finished precaching (and audio decompressor is cleaned up)
-		// ReferencedSoundWaves is used to make sure GC doesn't run on any sound waves that are actively pre-caching within an async task.
-		// Sounds may be loaded, kick off an async task to decompress, but never actually try to play, so GC can reclaim these while precaches are in-flight.
-		// We are also tracking when a sound wave is actively being used to generate audio in the audio render to prevent GC from happening to sounds till being used in the audio renderer.
-		for (int32 i = ReferencedSoundWaves.Num() - 1; i >= 0; --i)
-		{
-			USoundWave* Wave = ReferencedSoundWaves[i];
-			bool bRemove = true;
-			if (Wave)
-			{
-				const bool bIsPrecacheDone = (Wave->GetPrecacheState() == ESoundWavePrecacheState::Done);
-				const bool bIsGeneratingAudio = Wave->IsGeneratingAudio();
-
-				if (!bIsPrecacheDone || bIsGeneratingAudio)
-				{
-					bRemove = false;
-				}
-			}
-
-			if (bRemove)
-			{
-				ReferencedSoundWaves.RemoveAtSwap(i, 1, false);
-			}
-		}
 	}
 
 	if (!IsInAudioThread())
diff --git a/Engine/Source/Runtime/Engine/Public/AudioDevice.h b/Engine/Source/Runtime/Engine/Public/AudioDevice.h
--- a/Engine/Source/Runtime/Engine/Public/AudioDevice.h
+++ b/Engine/Source/Runtime/Engine/Public/AudioDevice.h
@@ -1827,10 +1827,7 @@ private:
 	TArray<FActiveSound*> ActiveSounds;
 	/** Array of sound waves to add references to avoid GC until guaranteed to be done with precache or decodes. */
 	TArray<USoundWave*> ReferencedSoundWaves;
-
 	void UpdateReferencedSoundWaves();
-	TArray<USoundWave*> ReferencedSoundWaves_AudioThread;
-	FCriticalSection ReferencedSoundWaveCritSec;
 
 	TArray<USoundWave*> PrecachingSoundWaves;
 
