Move headers into public

diff --git a/Engine/Source/Runtime/Engine/Classes/VT/RuntimeVirtualTextureStreamingProxy.h b/Engine/Source/Runtime/Engine/Classes/VT/RuntimeVirtualTextureStreamingProxy.h
new file mode 100644
--- /dev/null
+++ b/Engine/Source/Runtime/Engine/Classes/VT/RuntimeVirtualTextureStreamingProxy.h
@@ -0,0 +1,27 @@
+// Copyright 1998-2019 Epic Games, Inc. All Rights Reserved.
+
+#pragma once
+
+#include "CoreMinimal.h"
+#include "Engine/Texture2D.h"
+#include "VT/VirtualTextureBuildSettings.h"
+#include "RuntimeVirtualTextureStreamingProxy.generated.h"
+
+/** Streaming virtual texture used to store the low mips in a URuntimeVirtualTexture. */
+UCLASS(ClassGroup = Rendering)
+class URuntimeVirtualTextureStreamingProxy : public UTexture2D
+{
+	GENERATED_UCLASS_BODY()
+
+	/** Virtual texture build settings. These should match those of any owning URuntimeVirtualTexture. */
+	UPROPERTY()
+	FVirtualTextureBuildSettings Settings;
+
+	/** Hash of settings used when building this texture. Used to invalidate when build settings have changed. */
+	UPROPERTY()
+	uint32 BuildHash;
+
+	//~ Begin UTexture Interface.
+	virtual void GetVirtualTextureBuildSettings(FVirtualTextureBuildSettings& OutSettings) const override;
+	//~ End UTexture Interface.
+};
diff --git a/Engine/Source/Runtime/Engine/Classes/VT/VirtualTextureBuildSettings.h b/Engine/Source/Runtime/Engine/Classes/VT/VirtualTextureBuildSettings.h
new file mode 100644
--- /dev/null
+++ b/Engine/Source/Runtime/Engine/Classes/VT/VirtualTextureBuildSettings.h
@@ -0,0 +1,28 @@
+// Copyright 1998-2019 Epic Games, Inc. All Rights Reserved.
+
+#pragma once
+
+#include "CoreMinimal.h"
+#include "VirtualTextureBuildSettings.generated.h"
+
+/** Build settings used for virtual textures. */
+USTRUCT()
+struct FVirtualTextureBuildSettings
+{
+	GENERATED_USTRUCT_BODY()
+
+	UPROPERTY()
+	int32 TileSize;
+
+	UPROPERTY()
+	int32 TileBorderSize;
+
+	UPROPERTY()
+	bool bEnableCompressCrunch;
+
+	UPROPERTY()
+	bool bEnableCompressZlib;
+
+	/** Initialize with the default build settings. These are defined by the current project setup. */
+	ENGINE_API void Init();
+};
diff --git a/Engine/Source/Runtime/Engine/Private/VT/RuntimeVirtualTextureStreamingProxy.h b/Engine/Source/Runtime/Engine/Private/VT/RuntimeVirtualTextureStreamingProxy.h
deleted file mode 100644
diff --git a/Engine/Source/Runtime/Engine/Private/VT/VirtualTextureBuildSettings.h b/Engine/Source/Runtime/Engine/Private/VT/VirtualTextureBuildSettings.h
deleted file mode 100644
