Add support for mip level adjustments when sampling runtime virtual textures

diff --git a/Engine/Source/Runtime/Engine/Classes/Materials/MaterialExpressionRuntimeVirtualTextureSample.h b/Engine/Source/Runtime/Engine/Classes/Materials/MaterialExpressionRuntimeVirtualTextureSample.h
--- a/Engine/Source/Runtime/Engine/Classes/Materials/MaterialExpressionRuntimeVirtualTextureSample.h
+++ b/Engine/Source/Runtime/Engine/Classes/Materials/MaterialExpressionRuntimeVirtualTextureSample.h
@@ -8,6 +8,26 @@
 #include "VT/RuntimeVirtualTexture.h"
 #include "MaterialExpressionRuntimeVirtualTextureSample.generated.h"
 
+/**
+ * Defines how MipValue is used.
+ * Internally we will convert to ETextureMipValueMode which is used by internal APIs.
+ * ETextureMipValueMode has more options then are valid for runtime virtual texture.
+ */
+UENUM()
+enum ERuntimeVirtualTextureMipValueMode
+{
+	/* Use hardware computed sample's mip level with automatic anisotropic filtering support. */
+	RVTMVM_None UMETA(DisplayName = "None (use computed mip level)"),
+
+	/* Explicitly compute the sample's mip level. Disables anisotropic filtering. */
+	RVTMVM_MipLevel UMETA(DisplayName = "MipLevel (absolute, 0 is full resolution)"),
+
+	/* Bias the hardware computed sample's mip level. Disables anisotropic filtering. */
+	RVTMVM_MipBias UMETA(DisplayName = "MipBias (relative to the computed mip level)"),
+
+	RVTMVM_MAX,
+};
+
 /** Material expression for sampling from a runtime virtual texture. */
 UCLASS(collapsecategories, hidecategories=Object)
 class UMaterialExpressionRuntimeVirtualTextureSample : public UMaterialExpression
@@ -18,6 +38,10 @@ class UMaterialExpressionRuntimeVirtualTextureSample : public UMaterialExpressio
 	UPROPERTY()
 	FExpressionInput Coordinates;
 
+	/** Meaning depends on MipValueMode. A single unit is one mip level.  */
+	UPROPERTY(meta = (RequiredInput = "false"))
+	FExpressionInput MipValue;
+
 	/** The virtual texture object to sample. */
 	UPROPERTY(EditAnywhere, Category = VirtualTexture)
 	class URuntimeVirtualTexture* VirtualTexture;
@@ -26,6 +50,10 @@ class UMaterialExpressionRuntimeVirtualTextureSample : public UMaterialExpressio
 	UPROPERTY(EditAnywhere, Category = VirtualTexture, meta = (DisplayName = "Virtual texture content"))
 	ERuntimeVirtualTextureMaterialType MaterialType = ERuntimeVirtualTextureMaterialType::BaseColor;
 
+	/** Defines how the MipValue property is applied to the virtual texture lookup. */
+	UPROPERTY(EditAnywhere, Category = VirtualTexture)
+	TEnumAsByte<enum ERuntimeVirtualTextureMipValueMode> MipValueMode;
+
 protected:
 	/** Init MaterialType based on the current VirtualTexture */
 	void InitMaterialType();
diff --git a/Engine/Source/Runtime/Engine/Private/Materials/MaterialExpressions.cpp b/Engine/Source/Runtime/Engine/Private/Materials/MaterialExpressions.cpp
--- a/Engine/Source/Runtime/Engine/Private/Materials/MaterialExpressions.cpp
+++ b/Engine/Source/Runtime/Engine/Private/Materials/MaterialExpressions.cpp
@@ -2192,8 +2192,23 @@ int32 UMaterialExpressionRuntimeVirtualTextureSample::Compile(class FMaterialCom
 		CoordinateIndex = Coordinates.Compile(Compiler);
 	}
 	
+	// Compile the mip level for the current mip value mode
+	ETextureMipValueMode TextureMipLevelMode = TMVM_None;
+	int32 MipValueIndex = INDEX_NONE;
+	if (MipValue.GetTracedInput().Expression != nullptr)
+	{
+		switch (MipValueMode)
+		{
+		case RVTMVM_MipLevel: TextureMipLevelMode = TMVM_MipLevel; break;
+		case RVTMVM_MipBias: TextureMipLevelMode = TMVM_MipBias; break;
+		}
+		if (TextureMipLevelMode != RVTMVM_None)
+		{
+			MipValueIndex = MipValue.Compile(Compiler);
+		}
+	}
+
 	// Compile the texture sample code
-	//todo[vt]: Expose support for mip settings in this expression
 	int32 SampleCodeIndex[MAX_RVT_LAYERS] = { INDEX_NONE };
 	for (int32 Layer = 0; Layer < LayerCount; Layer++)
 	{
@@ -2201,7 +2216,7 @@ int32 UMaterialExpressionRuntimeVirtualTextureSample::Compile(class FMaterialCom
 			TextureCodeIndex[Layer],
 			CoordinateIndex, 
 			SAMPLERTYPE_VirtualMasks,
-			INDEX_NONE, INDEX_NONE, TMVM_None, SSM_Wrap_WorldGroupSettings,
+			MipValueIndex, INDEX_NONE, TextureMipLevelMode, SSM_Wrap_WorldGroupSettings,
 			TextureReferenceIndex[Layer], 
 			false);
 	}
