#UE4 Add some more detail to a warning

diff --git a/Engine/Source/Runtime/Engine/Private/Texture.cpp b/Engine/Source/Runtime/Engine/Private/Texture.cpp
--- a/Engine/Source/Runtime/Engine/Private/Texture.cpp
+++ b/Engine/Source/Runtime/Engine/Private/Texture.cpp
@@ -338,7 +338,7 @@ void UTexture::Serialize(FArchive& Ar)
 	{
 		if (VirtualTextureStreaming && UseVirtualTexturing(GMaxRHIFeatureLevel, Ar.CookingTarget()) == false)
 		{
-			UE_LOG(LogTexture, Warning, TEXT("%s is marked for virtual streaming but virtual texture streaming is not available."), *GetName());
+			UE_LOG(LogTexture, Warning, TEXT("%s is marked for virtual streaming but virtual texture streaming is not available."), *GetPathName());
 		}
 	}
 
