#UE4 Emit a compile error for bad TextureReferenceIndex in UMaterialExpressionRuntimeVirtualTextureSample instead of asserting.

diff --git a/Engine/Source/Runtime/Engine/Private/Materials/MaterialExpressions.cpp b/Engine/Source/Runtime/Engine/Private/Materials/MaterialExpressions.cpp
--- a/Engine/Source/Runtime/Engine/Private/Materials/MaterialExpressions.cpp
+++ b/Engine/Source/Runtime/Engine/Private/Materials/MaterialExpressions.cpp
@@ -2173,7 +2173,14 @@ int32 UMaterialExpressionRuntimeVirtualTextureSample::Compile(class FMaterialCom
 	for (int32 Layer = 0; Layer < LayerCount; Layer++)
 	{
 		TextureCodeIndex[Layer] = Compiler->VirtualTexture(VirtualTexture, Layer, TextureReferenceIndex[Layer], SAMPLERTYPE_VirtualMasks);
-		check(TextureReferenceIndex[Layer] == TextureReferenceIndex[0]);
+		if (TextureReferenceIndex[Layer] != TextureReferenceIndex[0])
+		{
+			Compiler->Errorf(TEXT("Virtual texture %s has invalid TextureReferenceIndex %d, which doesnt match %d in layer %d"),
+				bIsVirtualTextureValid ? *VirtualTexture->GetPathName() : TEXT("<unknown>"),
+				TextureReferenceIndex[Layer],
+				TextureReferenceIndex[0],
+				Layer);
+		}
 	}
 
 	// Compile the coordinates
