#UE4 Fix error in mac editor compile

diff --git a/Engine/Source/Runtime/Engine/Private/Materials/MaterialExpressions.cpp b/Engine/Source/Runtime/Engine/Private/Materials/MaterialExpressions.cpp
--- a/Engine/Source/Runtime/Engine/Private/Materials/MaterialExpressions.cpp
+++ b/Engine/Source/Runtime/Engine/Private/Materials/MaterialExpressions.cpp
@@ -2209,7 +2209,7 @@ int32 UMaterialExpressionRuntimeVirtualTextureSample::Compile(class FMaterialCom
 		case RVTMVM_MipLevel: TextureMipLevelMode = TMVM_MipLevel; break;
 		case RVTMVM_MipBias: TextureMipLevelMode = TMVM_MipBias; break;
 		}
-		if (TextureMipLevelMode != RVTMVM_None)
+		if (TextureMipLevelMode != TMVM_None)
 		{
 			MipValueIndex = MipValue.Compile(Compiler);
 		}
