Fix UActorComponent::BeginPlay comment
#rnx

diff --git a/Engine/Source/Runtime/Engine/Classes/Components/ActorComponent.h b/Engine/Source/Runtime/Engine/Classes/Components/ActorComponent.h
--- a/Engine/Source/Runtime/Engine/Classes/Components/ActorComponent.h
+++ b/Engine/Source/Runtime/Engine/Classes/Components/ActorComponent.h
@@ -553,9 +553,9 @@ public:
 	virtual void InitializeComponent();
 
 	/**
-	 * BeginsPlay for the component.  Occurs at level startup or actor spawn. This is before BeginPlay (Actor or Component).  
-	 * All Components (that want initialization) in the level will be Initialized on load before any 
-	 * Actor/Component gets BeginPlay.
+	 * Begins Play for the component. 
+	 * Called when the owning Actor begins play or when the component is created if the Actor has already begun play.
+	 * Actor BeginPlay normally happens right after PostInitializeComponents but can be delayed for networked or child actors.
 	 * Requires component to be registered and initialized.
 	 */
 	virtual void BeginPlay();
