Fix checks on bBuildDeveloperTools that are actually meant on checks for the editor. It should be possible to set bBuildDeveloperTools for runtime games.

diff --git a/Engine/Plugins/Developer/Concert/MultiUserClient/Source/MultiUserClientLibrary/MultiUserClientLibrary.Build.cs b/Engine/Plugins/Developer/Concert/MultiUserClient/Source/MultiUserClientLibrary/MultiUserClientLibrary.Build.cs
--- a/Engine/Plugins/Developer/Concert/MultiUserClient/Source/MultiUserClientLibrary/MultiUserClientLibrary.Build.cs
+++ b/Engine/Plugins/Developer/Concert/MultiUserClient/Source/MultiUserClientLibrary/MultiUserClientLibrary.Build.cs
@@ -15,7 +15,7 @@ namespace UnrealBuildTool.Rules
 				}
 			);
 
-			if (Target.bBuildDeveloperTools)
+			if (Target.Type == TargetType.Editor || Target.Type == TargetType.Program)
 			{
 				PrivateDefinitions.Add("WITH_CONCERT=1");
 
diff --git a/Engine/Source/Developer/TextureCompressor/TextureCompressor.Build.cs b/Engine/Source/Developer/TextureCompressor/TextureCompressor.Build.cs
--- a/Engine/Source/Developer/TextureCompressor/TextureCompressor.Build.cs
+++ b/Engine/Source/Developer/TextureCompressor/TextureCompressor.Build.cs
@@ -16,9 +16,6 @@ public class TextureCompressor : ModuleRules
 			}
 			);
 
-		if (Target.bBuildDeveloperTools)
-		{
-			AddEngineThirdPartyPrivateStaticDependencies(Target, "nvTextureTools");
-		}
+		AddEngineThirdPartyPrivateStaticDependencies(Target, "nvTextureTools");
 	}
 }
diff --git a/Engine/Source/Editor/LandscapeEditor/LandscapeEditor.Build.cs b/Engine/Source/Editor/LandscapeEditor/LandscapeEditor.Build.cs
--- a/Engine/Source/Editor/LandscapeEditor/LandscapeEditor.Build.cs
+++ b/Engine/Source/Editor/LandscapeEditor/LandscapeEditor.Build.cs
@@ -67,8 +67,7 @@ public class LandscapeEditor : ModuleRules
 		}
 
 		// KissFFT is used by the smooth tool.
-		if (Target.bBuildDeveloperTools &&
-			(Target.Platform == UnrealTargetPlatform.Win64 || Target.Platform == UnrealTargetPlatform.Win32 || Target.Platform == UnrealTargetPlatform.Mac || Target.Platform == UnrealTargetPlatform.Linux))
+		if (Target.Platform == UnrealTargetPlatform.Win64 || Target.Platform == UnrealTargetPlatform.Win32 || Target.Platform == UnrealTargetPlatform.Mac || Target.Platform == UnrealTargetPlatform.Linux)
 		{
 			AddEngineThirdPartyPrivateStaticDependencies(Target, "Kiss_FFT");
 		}
diff --git a/Engine/Source/Runtime/Engine/Engine.Build.cs b/Engine/Source/Runtime/Engine/Engine.Build.cs
--- a/Engine/Source/Runtime/Engine/Engine.Build.cs
+++ b/Engine/Source/Runtime/Engine/Engine.Build.cs
@@ -230,35 +230,25 @@ public class Engine : ModuleRules
 			PrivateDependencyModuleNames.Add("PerfCounters");
 		}
 
-		if (Target.bBuildDeveloperTools)
+		if (Target.Type == TargetType.Editor)
 		{
-			// Add "BlankModule" so that it gets compiled as an example and will be maintained and tested.  This can be removed
-			// at any time if needed.  The module isn't actually loaded by the engine so there is no runtime cost.
-			DynamicallyLoadedModuleNames.Add("BlankModule");
-
-			if (Target.Type != TargetType.Server)
-			{
-				PrivateIncludePathModuleNames.Add("MeshUtilities");
-				PrivateIncludePathModuleNames.Add("MeshUtilitiesCommon");
+			PrivateIncludePathModuleNames.Add("MeshUtilities");
+			PrivateIncludePathModuleNames.Add("MeshUtilitiesCommon");
 
-				DynamicallyLoadedModuleNames.Add("MeshUtilities");
+			DynamicallyLoadedModuleNames.Add("MeshUtilities");
 
-				PrivateDependencyModuleNames.AddRange(
-					new string[] {
-						"ImageCore",
-						"RawMesh"
-					}
-				);
-			}
+			PrivateDependencyModuleNames.AddRange(
+				new string[] {
+					"ImageCore",
+					"RawMesh"
+				}
+			);
 
-			if (Target.Configuration != UnrealTargetConfiguration.Shipping && Target.Configuration != UnrealTargetConfiguration.Test && Target.Type != TargetType.Server)
-			{
-				PrivateDependencyModuleNames.Add("CollisionAnalyzer");
-				CircularlyReferencedDependentModules.Add("CollisionAnalyzer");
+			PrivateDependencyModuleNames.Add("CollisionAnalyzer");
+			CircularlyReferencedDependentModules.Add("CollisionAnalyzer");
 
-				PrivateDependencyModuleNames.Add("LogVisualizer");
-				CircularlyReferencedDependentModules.Add("LogVisualizer");
-			}
+			PrivateDependencyModuleNames.Add("LogVisualizer");
+			CircularlyReferencedDependentModules.Add("LogVisualizer");
 
 			if (Target.Platform == UnrealTargetPlatform.Win64)
 			{
diff --git a/Engine/Source/Runtime/Landscape/Landscape.Build.cs b/Engine/Source/Runtime/Landscape/Landscape.Build.cs
--- a/Engine/Source/Runtime/Landscape/Landscape.Build.cs
+++ b/Engine/Source/Runtime/Landscape/Landscape.Build.cs
@@ -43,7 +43,7 @@ public class Landscape : ModuleRules
 			DynamicallyLoadedModuleNames.Add("PhysXCooking");
 		}
 
-		if (Target.bBuildDeveloperTools && Target.Type != TargetType.Server)
+		if (Target.Type == TargetType.Editor || Target.Type == TargetType.Program)
 		{
 			PrivateDependencyModuleNames.AddRange(
 				new string[] {
diff --git a/Engine/Source/Runtime/MovieSceneCapture/MovieSceneCapture.Build.cs b/Engine/Source/Runtime/MovieSceneCapture/MovieSceneCapture.Build.cs
--- a/Engine/Source/Runtime/MovieSceneCapture/MovieSceneCapture.Build.cs
+++ b/Engine/Source/Runtime/MovieSceneCapture/MovieSceneCapture.Build.cs
@@ -12,7 +12,7 @@ public class MovieSceneCapture : ModuleRules
 			}
 		);
 
-		if (Target.bBuildDeveloperTools)
+		if (Target.Type == TargetType.Editor || Target.Type == TargetType.Program)
 		{
 			PublicIncludePathModuleNames.Add("ImageWrapper");
 			DynamicallyLoadedModuleNames.Add("ImageWrapper");
