Make virtual UObject::Modify WITH_EDITOR. UObject::Modify remains defined as a forceinlined noop without editor to avoid having to fixup all existing code that exists in non-editor that calls Modify.

diff --git a/Engine/Source/Runtime/CoreUObject/Public/UObject/Object.h b/Engine/Source/Runtime/CoreUObject/Public/UObject/Object.h
--- a/Engine/Source/Runtime/CoreUObject/Public/UObject/Object.h
+++ b/Engine/Source/Runtime/CoreUObject/Public/UObject/Object.h
@@ -223,10 +223,15 @@ public:
 	 *								currently recording an active undo/redo transaction
 	 * @return true if the object was saved to the transaction buffer
 	 */
+#if WITH_EDITOR
 	virtual bool Modify( bool bAlwaysMarkDirty=true );
 
 	/** Utility to allow overrides of Modify to avoid doing work if this object cannot be safely modified */
 	bool CanModify() const;
+#else
+	UE_DEPRECATED(4.24,"How many things call this?")
+	FORCEINLINE bool Modify(bool bAlwaysMarkDirty = true) { return false; }
+#endif
 
 #if WITH_EDITOR
 	/** 
diff --git a/Engine/Source/Runtime/Engine/Classes/Engine/Blueprint.h b/Engine/Source/Runtime/Engine/Classes/Engine/Blueprint.h
--- a/Engine/Source/Runtime/Engine/Classes/Engine/Blueprint.h
+++ b/Engine/Source/Runtime/Engine/Classes/Engine/Blueprint.h
@@ -819,7 +819,9 @@ public:
 	virtual UClass* RegenerateClass(UClass* ClassToRegenerate, UObject* PreviousCDO) override;
 	virtual void PostLoad() override;
 	virtual void PostLoadSubobjects( FObjectInstancingGraph* OuterInstanceGraph ) override;
+#if WITH_EDITOR
 	virtual bool Modify(bool bAlwaysMarkDirty = true) override;
+#endif
 	virtual void GetAssetRegistryTags(TArray<FAssetRegistryTag>& OutTags) const override;
 	virtual FPrimaryAssetId GetPrimaryAssetId() const override;
 	virtual void BeginCacheForCookedPlatformData(const ITargetPlatform *TargetPlatform) override;
diff --git a/Engine/Source/Runtime/Engine/Classes/Engine/Polys.h b/Engine/Source/Runtime/Engine/Classes/Engine/Polys.h
--- a/Engine/Source/Runtime/Engine/Classes/Engine/Polys.h
+++ b/Engine/Source/Runtime/Engine/Classes/Engine/Polys.h
@@ -607,7 +607,9 @@ class UPolys : public UObject
 	{}
 
 	//~ Begin UObject Interface
+#if WITH_EDITOR
 	ENGINE_API virtual bool Modify(bool bAlwaysMarkDirty = false) override;	
+#endif
 	ENGINE_API virtual void Serialize( FArchive& Ar ) override;
 	virtual bool IsAsset() const override { return false; }
 	static void AddReferencedObjects(UObject* InThis, FReferenceCollector& Collector);
diff --git a/Engine/Source/Runtime/Engine/Classes/Engine/Selection.h b/Engine/Source/Runtime/Engine/Classes/Engine/Selection.h
--- a/Engine/Source/Runtime/Engine/Classes/Engine/Selection.h
+++ b/Engine/Source/Runtime/Engine/Classes/Engine/Selection.h
@@ -312,18 +312,6 @@ public:
 		return Count;
 	}
 
-	/**
-	 * Gets selected class by index.
-	 *
-	 * @return			The selected class at the specified index.
-	 */
-	UE_DEPRECATED(4.14, "GetSelectedClass is deprecated.  Use IsClassSelected or a class iterator to search through classes")
-	UClass* GetSelectedClass(int32 InIndex) const
-	{
-		FSetElementId Id = FSetElementId::FromInteger(InIndex);
-		return const_cast<UClass*>(SelectedClasses[Id].Class);
-	}
-
 	bool IsClassSelected(UClass* Class) const
 	{
 		const FSelectedClassInfo* Info = SelectedClasses.Find(Class);
@@ -332,7 +320,9 @@ public:
 
 	//~ Begin UObject Interface
 	virtual void Serialize(FArchive& Ar) override;
+#if WITH_EDITOR
 	virtual bool Modify( bool bAlwaysMarkDirty=true) override;
+#endif
 	virtual void BeginDestroy() override;
 
 	//~ End UObject Interface
diff --git a/Engine/Source/Runtime/Engine/Private/Blueprint.cpp b/Engine/Source/Runtime/Engine/Private/Blueprint.cpp
--- a/Engine/Source/Runtime/Engine/Private/Blueprint.cpp
+++ b/Engine/Source/Runtime/Engine/Private/Blueprint.cpp
@@ -1882,11 +1882,13 @@ void UBlueprint::PostLoadSubobjects(FObjectInstancingGraph* OuterInstanceGraph)
 	ConformNativeComponents(this);
 }
 
+#if WITH_EDITOR
 bool UBlueprint::Modify(bool bAlwaysMarkDirty)
 {
 	bCachedDependenciesUpToDate = false;
 	return Super::Modify(bAlwaysMarkDirty);
 }
+#endif
 
 void UBlueprint::GatherDependencies(TSet<TWeakObjectPtr<UBlueprint>>& InDependencies) const
 {
diff --git a/Engine/Source/Runtime/Engine/Private/Brush.cpp b/Engine/Source/Runtime/Engine/Private/Brush.cpp
--- a/Engine/Source/Runtime/Engine/Private/Brush.cpp
+++ b/Engine/Source/Runtime/Engine/Private/Brush.cpp
@@ -310,6 +310,7 @@ bool ABrush::IsStaticBrush() const
 	return BrushComponent && (BrushComponent->Mobility == EComponentMobility::Static);
 }
 
+#if WITH_EDITOR
 bool ABrush::Modify(bool bAlwaysMarkDirty)
 {
 	bool bSavedToTransactionBuffer = Super::Modify(bAlwaysMarkDirty);
@@ -320,4 +321,4 @@ bool ABrush::Modify(bool bAlwaysMarkDirty)
 	}
 	return bSavedToTransactionBuffer;
 }
-
+#endif
\ No newline at end of file
diff --git a/Engine/Source/Runtime/Engine/Private/Components/SplineMeshComponent.cpp b/Engine/Source/Runtime/Engine/Private/Components/SplineMeshComponent.cpp
--- a/Engine/Source/Runtime/Engine/Private/Components/SplineMeshComponent.cpp
+++ b/Engine/Source/Runtime/Engine/Private/Components/SplineMeshComponent.cpp
@@ -554,6 +554,7 @@ void USplineMeshComponent::Serialize(FArchive& Ar)
 #endif
 }
 
+#if WITH_EDITOR
 bool USplineMeshComponent::Modify(bool bAlwaysMarkDirty)
 {
 	bool bSavedToTransactionBuffer = Super::Modify(bAlwaysMarkDirty);
@@ -565,6 +566,7 @@ bool USplineMeshComponent::Modify(bool bAlwaysMarkDirty)
 
 	return bSavedToTransactionBuffer;
 }
+#endif
 
 FPrimitiveSceneProxy* USplineMeshComponent::CreateSceneProxy()
 {
diff --git a/Engine/Source/Runtime/Engine/Private/Materials/MaterialExpressions.cpp b/Engine/Source/Runtime/Engine/Private/Materials/MaterialExpressions.cpp
--- a/Engine/Source/Runtime/Engine/Private/Materials/MaterialExpressions.cpp
+++ b/Engine/Source/Runtime/Engine/Private/Materials/MaterialExpressions.cpp
@@ -9552,21 +9552,16 @@ void UMaterialExpressionComment::PostEditChangeProperty(FPropertyChangedEvent& P
 	}
 }
 
-#endif // WITH_EDITOR
-
 bool UMaterialExpressionComment::Modify( bool bAlwaysMarkDirty/*=true*/ )
 {
 	bool bResult = Super::Modify(bAlwaysMarkDirty);
 
-#if WITH_EDITORONLY_DATA
 	// Don't need to update preview after changing comments
 	bNeedToUpdatePreview = false;
-#endif // WITH_EDITORONLY_DATA
 
 	return bResult;
 }
 
-#if WITH_EDITOR
 void UMaterialExpressionComment::GetCaption(TArray<FString>& OutCaptions) const
 {
 	OutCaptions.Add(TEXT("Comment"));
diff --git a/Engine/Source/Runtime/Engine/Private/Polygon.cpp b/Engine/Source/Runtime/Engine/Private/Polygon.cpp
--- a/Engine/Source/Runtime/Engine/Private/Polygon.cpp
+++ b/Engine/Source/Runtime/Engine/Private/Polygon.cpp
@@ -978,12 +978,14 @@ int32 FPoly::Faces( const FPoly &Test ) const
 	return 0;
 }
 
+#if WITH_EDITOR
 bool UPolys::Modify(bool bAlwaysMarkDirty)
 {
 	Super::Modify(bAlwaysMarkDirty);
 
 	return !!GUndo; // we will make a broad assumption that if we have an undo buffer, we were saved in it
 }
+#endif
 
 void UPolys::AddReferencedObjects(UObject* InThis, FReferenceCollector& Collector)
 {
diff --git a/Engine/Source/Runtime/Engine/Private/Selection.cpp b/Engine/Source/Runtime/Engine/Private/Selection.cpp
--- a/Engine/Source/Runtime/Engine/Private/Selection.cpp
+++ b/Engine/Source/Runtime/Engine/Private/Selection.cpp
@@ -229,6 +229,7 @@ void USelection::Serialize(FArchive& Ar)
 	}
 }
 
+#if WITH_EDITOR
 bool USelection::Modify(bool bAlwaysMarkDirty/* =true */)
 {
 	// If the selection currently contains any PIE objects we should not be including it in the transaction buffer
@@ -243,6 +244,7 @@ bool USelection::Modify(bool bAlwaysMarkDirty/* =true */)
 
 	return Super::Modify(bAlwaysMarkDirty);
 }
+#endif
 
 void USelection::BeginDestroy()
 {
diff --git a/Engine/Source/Runtime/Experimental/GeometryCollectionEngine/Private/GeometryCollection/GeometryCollectionObject.cpp b/Engine/Source/Runtime/Experimental/GeometryCollectionEngine/Private/GeometryCollection/GeometryCollectionObject.cpp
--- a/Engine/Source/Runtime/Experimental/GeometryCollectionEngine/Private/GeometryCollection/GeometryCollectionObject.cpp
+++ b/Engine/Source/Runtime/Experimental/GeometryCollectionEngine/Private/GeometryCollection/GeometryCollectionObject.cpp
@@ -379,7 +379,6 @@ void UGeometryCollection::PostEditChangeProperty(struct FPropertyChangedEvent& P
 		CreateSimulationData();
 	}
 }
-#endif
 
 bool UGeometryCollection::Modify(bool bAlwaysMarkDirty /*= true*/)
 {
@@ -394,12 +393,6 @@ bool UGeometryCollection::Modify(bool bAlwaysMarkDirty /*= true*/)
 	return bSuperResult;
 }
 
-void UGeometryCollection::PostLoad()
-{
-	Super::PostLoad();
-}
-
-#if WITH_EDITOR
 void UGeometryCollection::EnsureDataIsCooked()
 {
 	if (StateGuid != LastBuiltGuid)
diff --git a/Engine/Source/Runtime/Experimental/GeometryCollectionEngine/Public/GeometryCollection/GeometryCollectionObject.h b/Engine/Source/Runtime/Experimental/GeometryCollectionEngine/Public/GeometryCollection/GeometryCollectionObject.h
--- a/Engine/Source/Runtime/Experimental/GeometryCollectionEngine/Public/GeometryCollection/GeometryCollectionObject.h
+++ b/Engine/Source/Runtime/Experimental/GeometryCollectionEngine/Public/GeometryCollection/GeometryCollectionObject.h
@@ -98,11 +98,10 @@ public:
 	UGeometryCollection(const FObjectInitializer& ObjectInitializer = FObjectInitializer::Get());
 
 	/** UObject Interface */
-	virtual void PostLoad() override;
 #if WITH_EDITOR
 	virtual void PostEditChangeProperty(struct FPropertyChangedEvent& PropertyChangedEvent) override;
-#endif
 	virtual bool Modify(bool bAlwaysMarkDirty = true) override;
+#endif
 	/** End UObject Interface */
 
 	void Serialize(FArchive& Ar);
diff --git a/Engine/Source/Runtime/Landscape/Classes/LandscapeSplineSegment.h b/Engine/Source/Runtime/Landscape/Classes/LandscapeSplineSegment.h
--- a/Engine/Source/Runtime/Landscape/Classes/LandscapeSplineSegment.h
+++ b/Engine/Source/Runtime/Landscape/Classes/LandscapeSplineSegment.h
@@ -321,7 +321,6 @@ public:
 	//~ Begin UObject Interface
 	virtual void Serialize(FArchive& Ar) override;
 	virtual void PostLoad() override;
-	virtual bool Modify(bool bAlwaysMarkDirty = true) override;
 #if WITH_EDITOR
 	virtual void PostEditUndo() override;
 	virtual void PostDuplicate(bool bDuplicateForPIE) override;
diff --git a/Engine/Source/Runtime/Landscape/Private/LandscapeSplines.cpp b/Engine/Source/Runtime/Landscape/Private/LandscapeSplines.cpp
--- a/Engine/Source/Runtime/Landscape/Private/LandscapeSplines.cpp
+++ b/Engine/Source/Runtime/Landscape/Private/LandscapeSplines.cpp
@@ -2763,21 +2763,6 @@ void ULandscapeSplineSegment::FindNearest( const FVector& InLocation, float& t,
 	OutTangent = SplineInfo.EvalDerivative(t, FVector::ZeroVector);
 }
 
-bool ULandscapeSplineSegment::Modify(bool bAlwaysMarkDirty /*= true*/)
-{
-	bool bSavedToTransactionBuffer = Super::Modify(bAlwaysMarkDirty);
-
-	//for (auto MeshComponent : MeshComponents)
-	//{
-	//	if (MeshComponent)
-	//	{
-	//		bSavedToTransactionBuffer = MeshComponent->Modify(bAlwaysMarkDirty) || bSavedToTransactionBuffer;
-	//	}
-	//}
-
-	return bSavedToTransactionBuffer;
-}
-
 #if WITH_EDITOR
 void ULandscapeSplineSegment::PostEditUndo()
 {
diff --git a/Engine/Source/Runtime/UMG/Private/Components/Widget.cpp b/Engine/Source/Runtime/UMG/Private/Components/Widget.cpp
--- a/Engine/Source/Runtime/UMG/Private/Components/Widget.cpp
+++ b/Engine/Source/Runtime/UMG/Private/Components/Widget.cpp
@@ -1102,6 +1102,7 @@ void UWidget::DeselectByDesigner()
 #define LOCTEXT_NAMESPACE "UMG"
 #endif
 
+#if WITH_EDITOR
 bool UWidget::Modify(bool bAlwaysMarkDirty)
 {
 	bool Modified = Super::Modify(bAlwaysMarkDirty);
@@ -1114,6 +1115,7 @@ bool UWidget::Modify(bool bAlwaysMarkDirty)
 
 	return Modified;
 }
+#endif
 
 bool UWidget::IsChildOf(UWidget* PossibleParent)
 {
