Exposed UHLODProxy::GetCRC() methods to be able to reuse them elsewhere

diff --git a/Engine/Source/Runtime/Engine/Classes/Engine/HLODProxy.h b/Engine/Source/Runtime/Engine/Classes/Engine/HLODProxy.h
--- a/Engine/Source/Runtime/Engine/Classes/Engine/HLODProxy.h
+++ b/Engine/Source/Runtime/Engine/Classes/Engine/HLODProxy.h
@@ -92,6 +92,12 @@ public:
 
 	/** Build a unique key for the LOD actor, used to determine if the actor needs rebuilding */
 	static FName GenerateKeyForActor(const ALODActor* LODActor);
+
+	static uint32 GetCRC(UMaterialInterface* InMaterialInterface, uint32 InCRC = 0);
+	static uint32 GetCRC(UTexture* InTexture, uint32 InCRC = 0);
+	static uint32 GetCRC(UStaticMesh* InStaticMesh, uint32 InCRC = 0);
+	static uint32 GetCRC(UStaticMeshComponent* InComponent, uint32 InCRC = 0);
+
 #endif
 
 #if !(UE_BUILD_SHIPPING || UE_BUILD_TEST)
diff --git a/Engine/Source/Runtime/Engine/Private/HLODProxy.cpp b/Engine/Source/Runtime/Engine/Private/HLODProxy.cpp
--- a/Engine/Source/Runtime/Engine/Private/HLODProxy.cpp
+++ b/Engine/Source/Runtime/Engine/Private/HLODProxy.cpp
@@ -156,7 +156,7 @@ void UHLODProxy::ExtractComponents(const ALODActor* LODActor, TArray<UPrimitiveC
 	}
 }
 
-uint32 GetCRC(UMaterialInterface* InMaterialInterface, uint32 InCRC = 0)
+uint32 UHLODProxy::GetCRC(UMaterialInterface* InMaterialInterface, uint32 InCRC)
 {
 	TArray<uint8> KeyBuffer;
 
@@ -182,7 +182,7 @@ uint32 GetCRC(UMaterialInterface* InMaterialInterface, uint32 InCRC = 0)
 	return FCrc::MemCrc32(KeyBuffer.GetData(), KeyBuffer.Num(), InCRC);
 }
 
-uint32 GetCRC(UTexture* InTexture, uint32 InCRC = 0)
+uint32 UHLODProxy::GetCRC(UTexture* InTexture, uint32 InCRC)
 {
 	// Default to just the path name if we don't have render data
 	 if (InTexture->GetRunningPlatformData() != nullptr)
@@ -198,7 +198,7 @@ uint32 GetCRC(UTexture* InTexture, uint32 InCRC = 0)
      return FCrc::StrCrc32(*InTexture->GetPathName(), InCRC);
 }
 
-uint32 GetCRC(UStaticMesh* InStaticMesh, uint32 InCRC = 0)
+uint32 UHLODProxy::GetCRC(UStaticMesh* InStaticMesh, uint32 InCRC)
 {
 	TArray<uint8> KeyBuffer;
 
@@ -214,7 +214,7 @@ uint32 GetCRC(UStaticMesh* InStaticMesh, uint32 InCRC = 0)
 	return FCrc::MemCrc32(KeyBuffer.GetData(), KeyBuffer.Num(), InCRC);
 }
 
-uint32 GetCRC(UStaticMeshComponent* InComponent, uint32 InCRC = 0)
+uint32 UHLODProxy::GetCRC(UStaticMeshComponent* InComponent, uint32 InCRC)
 {
 	TArray<uint8> KeyBuffer;
 
