Enable network emulation feature in TEST builds

diff --git a/Engine/Source/Runtime/Engine/Classes/Engine/NetConnection.h b/Engine/Source/Runtime/Engine/Classes/Engine/NetConnection.h
--- a/Engine/Source/Runtime/Engine/Classes/Engine/NetConnection.h
+++ b/Engine/Source/Runtime/Engine/Classes/Engine/NetConnection.h
@@ -220,6 +220,15 @@ struct FDelayedIncomingPacket
 
 	/** Time at which the packet should be reinjected into the connection */
 	double ReinjectionTime = 0.0;
+
+	void CountBytes(FArchive& Ar) const
+	{
+		if (PacketData.IsValid())
+		{
+			PacketData->CountMemory(Ar);
+		}
+		Ar.CountBytes(sizeof(ReinjectionTime), sizeof(ReinjectionTime));
+	}
 };
 
 #endif //#if DO_ENABLE_NET_TEST
@@ -1225,7 +1234,9 @@ private:
 	void UpdateAllCachedLevelVisibility() const;
 
 	/** Returns true if an outgoing packet should be dropped due to packet simulation settings, including loss burst simulation. */
+#if DO_ENABLE_NET_TEST
 	bool ShouldDropOutgoingPacketForLossSimulation(int64 NumBits) const;
+#endif
 
 	/**
 	*	Test the current emulation settings to delay, drop, etc the current packet
diff --git a/Engine/Source/Runtime/Engine/Classes/Engine/NetDriver.h b/Engine/Source/Runtime/Engine/Classes/Engine/NetDriver.h
--- a/Engine/Source/Runtime/Engine/Classes/Engine/NetDriver.h
+++ b/Engine/Source/Runtime/Engine/Classes/Engine/NetDriver.h
@@ -367,7 +367,7 @@ DECLARE_DELEGATE_SevenParams(FOnSendRPC, AActor* /*Actor*/, UFunction* /*Functio
 //
 // Whether to support net lag and packet loss testing.
 //
-#define DO_ENABLE_NET_TEST !(UE_BUILD_SHIPPING || UE_BUILD_TEST)
+#define DO_ENABLE_NET_TEST !(UE_BUILD_SHIPPING)
 
 
 /** Holds the packet simulation settings in one place */
@@ -848,7 +848,9 @@ public:
 	void InitPacketSimulationSettings();
 
 	/** Returns true during the duration of a packet loss burst triggered by the net.pktlossburst command. */
+#if DO_ENABLE_NET_TEST
 	bool IsSimulatingPacketLossBurst() const;
+#endif
 
 	/** Interface for communication network state to others (ie World usually, but anything that implements FNetworkNotify) */
 	class FNetworkNotify*		Notify;
diff --git a/Engine/Source/Runtime/Engine/Private/NetConnection.cpp b/Engine/Source/Runtime/Engine/Private/NetConnection.cpp
--- a/Engine/Source/Runtime/Engine/Private/NetConnection.cpp
+++ b/Engine/Source/Runtime/Engine/Private/NetConnection.cpp
@@ -653,6 +653,14 @@ PRAGMA_ENABLE_DEPRECATION_WARNINGS
 				Packet.CountBytes(Ar);
 			}
 		);
+
+		GRANULAR_NETWORK_MEMORY_TRACKING_TRACK("DelayedIncoming", 
+			DelayedIncomingPackets.CountBytes(Ar);
+			for (const FDelayedIncomingPacket& DelayedPacket : DelayedIncomingPackets)
+			{
+				DelayedPacket.CountBytes(Ar);
+			}
+		);
 #endif
 	}
 }
@@ -1096,11 +1104,13 @@ void UNetConnection::ReceivedRawPacket( void* InData, int32 Count )
 	}
 #endif
 
+#if DO_ENABLE_NET_TEST
 	// Opportunity for packet loss burst simulation to drop the incoming packet.
 	if (Driver && Driver->IsSimulatingPacketLossBurst())
 	{
 		return;
 	}
+#endif
 
 	uint8* Data = (uint8*)InData;
 
@@ -1261,24 +1271,27 @@ void UNetConnection::FlushPacketOrderCache(bool bFlushWholeCache/*=false*/)
 #if DO_ENABLE_NET_TEST
 void UNetConnection::ReinjectDelayedPackets()
 {
-	const double CurrentTime = FPlatformTime::Seconds();
+	if (DelayedIncomingPackets.Num() > 0)
+	{
+		const double CurrentTime = FPlatformTime::Seconds();
 
-	TGuardValue<bool> ReinjectingGuard(bIsReinjectingDelayedPackets, true);
+		TGuardValue<bool> ReinjectingGuard(bIsReinjectingDelayedPackets, true);
 
-	uint32 NbReinjected(0);
-	for (const FDelayedIncomingPacket& DelayedPacket : DelayedIncomingPackets)
-	{
-		if (DelayedPacket.ReinjectionTime > CurrentTime)
+		uint32 NbReinjected(0);
+		for (const FDelayedIncomingPacket& DelayedPacket : DelayedIncomingPackets)
 		{
-			break;
+			if (DelayedPacket.ReinjectionTime > CurrentTime)
+			{
+				break;
+			}
+
+			++NbReinjected;
+			ReceivedPacket(*DelayedPacket.PacketData.Get());
 		}
-		
-		++NbReinjected;
-		ReceivedPacket(*DelayedPacket.PacketData.Get());
-	}
 
-	// Delete processed packets
-	DelayedIncomingPackets.RemoveAt(0, NbReinjected, false);
+		// Delete processed packets
+		DelayedIncomingPackets.RemoveAt(0, NbReinjected, false);
+	}
 }
 #endif //#if DO_ENABLE_NET_TEST
 
@@ -1484,17 +1497,15 @@ bool UNetConnection::CheckOutgoingPacketEmulation(FOutPacketTraits& Traits)
 }
 #endif
 
+#if DO_ENABLE_NET_TEST
 bool UNetConnection::ShouldDropOutgoingPacketForLossSimulation(int64 NumBits) const
 {
-#if DO_ENABLE_NET_TEST
 	return Driver->IsSimulatingPacketLossBurst() || 
 		(PacketSimulationSettings.PktLoss > 0 && 
          PacketSimulationSettings.ShouldDropPacketOfSize(NumBits) && 
          FMath::FRand() * 100.f < PacketSimulationSettings.PktLoss);
-#else
-	return false;
-#endif
 }
+#endif
 
 int32 UNetConnection::IsNetReady( bool Saturate )
 {
@@ -1775,16 +1786,17 @@ void UNetConnection::ReceivedPacket( FBitReader& Reader )
 		}
 		if (PacketSimulationSettings.PktIncomingLagMin > 0 || PacketSimulationSettings.PktIncomingLagMax > 0)
 		{
-			const double LagVariance = FMath::FRand() * double(PacketSimulationSettings.PktIncomingLagMax - PacketSimulationSettings.PktIncomingLagMin);
-			const double ExtraLag = (double(PacketSimulationSettings.PktIncomingLagMin) + LagVariance) / 1000.f;
+			// ExtraLagInSec goes from [PktIncomingLagMin, PktIncomingLagMax]
+			const double LagVarianceInMS = FMath::FRand() * double(PacketSimulationSettings.PktIncomingLagMax - PacketSimulationSettings.PktIncomingLagMin);
+			const double ExtraLagInSec = (double(PacketSimulationSettings.PktIncomingLagMin) + LagVarianceInMS) / 1000.f;
 
 			FDelayedIncomingPacket DelayedPacket;
 			DelayedPacket.PacketData = MakeUnique<FBitReader>(Reader);
-			DelayedPacket.ReinjectionTime = FPlatformTime::Seconds() + ExtraLag;
+			DelayedPacket.ReinjectionTime = FPlatformTime::Seconds() + ExtraLagInSec;
 
 			DelayedIncomingPackets.Emplace(MoveTemp(DelayedPacket));
 
-			UE_LOG(LogNet, VeryVerbose, TEXT("Delaying incoming packet for %f seconds"), ExtraLag);
+			UE_LOG(LogNet, VeryVerbose, TEXT("Delaying incoming packet for %f seconds"), ExtraLagInSec);
 			return;
 		}
 	}
@@ -2811,17 +2823,20 @@ void UNetConnection::Tick()
 
 	AssertValid();
 
+	// Get frame time.
+	const double CurrentRealtimeSeconds = FPlatformTime::Seconds();
+
 	// Lag simulation.
 #if DO_ENABLE_NET_TEST
-	if(Delayed.Num() > 0)
+	if (Delayed.Num() > 0)
 	{
-		for( int32 i=0; i < Delayed.Num(); i++ )
+		uint32 NbPacketsSent(0);
+		for( FDelayedPacket& DelayedPacket : Delayed )
 		{
-			if( FPlatformTime::Seconds() > Delayed[i].SendTime )
+			if( CurrentRealtimeSeconds > DelayedPacket.SendTime )
 			{
-				LowLevelSend((char*)&Delayed[i].Data[0], Delayed[i].SizeBits, Delayed[i].Traits);
-				Delayed.RemoveAt( i );
-				i--;
+				LowLevelSend((char*)&DelayedPacket.Data[0], DelayedPacket.SizeBits, DelayedPacket.Traits);
+				++NbPacketsSent;
 			}
 			else
 			{
@@ -2829,12 +2844,11 @@ void UNetConnection::Tick()
 				break;
 			}
 		}
+
+		Delayed.RemoveAt(0, NbPacketsSent, false);
 	}
 #endif
 
-	// Get frame time.
-	const double CurrentRealtimeSeconds = FPlatformTime::Seconds();
-
 	// if this is 0 it's our first tick since init, so start our real-time tracking from here
 	if (LastTime == 0.0)
 	{
@@ -3289,7 +3303,6 @@ void UChildConnection::HandleClientPlayer(APlayerController* PC, UNetConnection*
 }
 
 #if DO_ENABLE_NET_TEST
-
 void UNetConnection::UpdatePacketSimulationSettings(void)
 {
 	check(Driver);
diff --git a/Engine/Source/Runtime/Engine/Private/NetworkDriver.cpp b/Engine/Source/Runtime/Engine/Private/NetworkDriver.cpp
--- a/Engine/Source/Runtime/Engine/Private/NetworkDriver.cpp
+++ b/Engine/Source/Runtime/Engine/Private/NetworkDriver.cpp
@@ -342,14 +342,12 @@ void UNetDriver::InitPacketSimulationSettings()
 #endif
 }
 
+#if DO_ENABLE_NET_TEST
 bool UNetDriver::IsSimulatingPacketLossBurst() const
 {
-#if DO_ENABLE_NET_TEST
 	return PacketLossBurstEndTime > Time;
-#else
-	return false;
-#endif
 }
+#endif
 
 void UNetDriver::PostInitProperties()
 {
@@ -3534,7 +3532,7 @@ bool FPacketSimulationSettings::ParseHelper(const TCHAR* Cmd, const TCHAR* Name,
 	return false;
 }
 
-#endif
+#endif //#if DO_ENABLE_NET_TEST
 
 FNetViewer::FNetViewer(UNetConnection* InConnection, float DeltaSeconds) :
 	Connection(InConnection),
