Minimize the work that moving actors when adding a streaming level in does by updating the Relative Location and Rotation in place rather than using SetRelativeLocationAndRotation which does a move and a lot of overlap update overhead that is not relevant when transforming a level being added to world.

diff --git a/Engine/Source/Runtime/Engine/Private/LevelUtils.cpp b/Engine/Source/Runtime/Engine/Private/LevelUtils.cpp
--- a/Engine/Source/Runtime/Engine/Private/LevelUtils.cpp
+++ b/Engine/Source/Runtime/Engine/Private/LevelUtils.cpp
@@ -324,20 +324,23 @@ void FLevelUtils::SetEditorTransform(ULevelStreaming* StreamingLevel, const FTra
 void FLevelUtils::ApplyEditorTransform(const ULevelStreaming* StreamingLevel, bool bDoPostEditMove )
 {
 	check(StreamingLevel);
-	ULevel* LoadedLevel = StreamingLevel->GetLoadedLevel();
-	if( LoadedLevel != NULL )
+	if (ULevel* LoadedLevel = StreamingLevel->GetLoadedLevel())
 	{	
-		ApplyLevelTransform( LoadedLevel, StreamingLevel->LevelTransform, bDoPostEditMove );
+		FApplyLevelTransformParams TransformParams(LoadedLevel, StreamingLevel->LevelTransform);
+		TransformParams.bDoPostEditMove = bDoPostEditMove;
+		ApplyLevelTransform(TransformParams);
 	}
 }
 
 void FLevelUtils::RemoveEditorTransform(const ULevelStreaming* StreamingLevel, bool bDoPostEditMove )
 {
 	check(StreamingLevel);
-	ULevel* LoadedLevel = StreamingLevel->GetLoadedLevel();
-	if( LoadedLevel != NULL )
+	if (ULevel* LoadedLevel = StreamingLevel->GetLoadedLevel())
 	{
-		ApplyLevelTransform( LoadedLevel, StreamingLevel->LevelTransform.Inverse(), bDoPostEditMove );
+		const FTransform InverseTransform = StreamingLevel->LevelTransform.Inverse();
+		FApplyLevelTransformParams TransformParams(LoadedLevel, InverseTransform);
+		TransformParams.bDoPostEditMove = bDoPostEditMove;
+		ApplyLevelTransform(TransformParams);
 	}
 }
 
@@ -379,52 +382,65 @@ bool FLevelUtils::IsApplyingLevelTransform()
 
 #endif // WITH_EDITOR
 
-void FLevelUtils::ApplyLevelTransform( ULevel* Level, const FTransform& LevelTransform, bool bDoPostEditMove )
+void FLevelUtils::ApplyLevelTransform(const FLevelUtils::FApplyLevelTransformParams& TransformParams)
 {
-	bool bTransformActors =  !LevelTransform.Equals(FTransform::Identity);
+	const bool bTransformActors =  !TransformParams.LevelTransform.Equals(FTransform::Identity);
 	if (bTransformActors)
 	{
 #if WITH_EDITOR
-		bool bApplyingLevelTransformBackup = bApplyingLevelTransform;
-		bApplyingLevelTransform = true;
+		TGuardValue<bool> ApplyingLevelTransformGuard(bApplyingLevelTransform, true);
 #endif
-		if (!LevelTransform.GetRotation().IsIdentity())
+		if (!TransformParams.LevelTransform.GetRotation().IsIdentity())
 		{
 			// If there is a rotation applied, then the relative precomputed bounds become invalid.
-			Level->bTextureStreamingRotationChanged = true;
+			TransformParams.Level->bTextureStreamingRotationChanged = true;
 		}
 
-		// Iterate over all actors in the level and transform them
-		for( int32 ActorIndex=0; ActorIndex<Level->Actors.Num(); ActorIndex++ )
+		if (TransformParams.bSetRelativeTransformDirectly)
 		{
-			AActor* Actor = Level->Actors[ActorIndex];
+			// Iterate over all actors in the level and transform them
+			for (AActor* Actor : TransformParams.Level->Actors)
+			{
+				if (Actor)
+				{
+					USceneComponent* RootComponent = Actor->GetRootComponent();
 
-			// Don't want to transform children they should stay relative to their parents.
-			if( Actor && Actor->GetAttachParentActor() == nullptr )
+					// Don't want to transform children they should stay relative to their parents.
+					if (RootComponent && RootComponent->GetAttachParent() == nullptr)
+					{
+						RootComponent->RelativeLocation = TransformParams.LevelTransform.TransformPosition(RootComponent->RelativeLocation);
+						RootComponent->RelativeRotation = TransformParams.LevelTransform.TransformRotation(RootComponent->RelativeRotation.Quaternion()).Rotator();
+					}
+				}
+			}
+		}
+		else
+		{
+			// Iterate over all actors in the level and transform them
+			for (AActor* Actor : TransformParams.Level->Actors)
 			{
-				// Has to modify root component directly as GetActorPosition is incorrect this early
-				USceneComponent *RootComponent = Actor->GetRootComponent();
-				if (RootComponent)
+				if (Actor)
 				{
-					RootComponent->SetRelativeLocationAndRotation( LevelTransform.TransformPosition(RootComponent->RelativeLocation), LevelTransform.TransformRotation(RootComponent->RelativeRotation.Quaternion()) );
-				}			
+					USceneComponent* RootComponent = Actor->GetRootComponent();
+
+					// Don't want to transform children they should stay relative to their parents.
+					if (RootComponent && RootComponent->GetAttachParent() == nullptr)
+					{
+						RootComponent->SetRelativeLocationAndRotation(TransformParams.LevelTransform.TransformPosition(RootComponent->RelativeLocation), TransformParams.LevelTransform.TransformRotation(RootComponent->RelativeRotation.Quaternion()));
+					}
+				}
 			}
 		}
 
 #if WITH_EDITOR
-		if( bDoPostEditMove )
+		if (TransformParams.bDoPostEditMove)
 		{
-			ApplyPostEditMove( Level );						
+			ApplyPostEditMove(TransformParams.Level);
 		}
 #endif // WITH_EDITOR
 
-		Level->OnApplyLevelTransform.Broadcast(LevelTransform);
-#if WITH_EDITOR
-		bApplyingLevelTransform = bApplyingLevelTransformBackup;
-#endif
+		TransformParams.Level->OnApplyLevelTransform.Broadcast(TransformParams.LevelTransform);
 	}
 }
 
-
-
 #undef LOCTEXT_NAMESPACE
diff --git a/Engine/Source/Runtime/Engine/Private/World.cpp b/Engine/Source/Runtime/Engine/Private/World.cpp
--- a/Engine/Source/Runtime/Engine/Private/World.cpp
+++ b/Engine/Source/Runtime/Engine/Private/World.cpp
@@ -2285,7 +2285,9 @@ void UWorld::AddToWorld( ULevel* Level, const FTransform& LevelTransform, bool b
 		QUICK_SCOPE_CYCLE_COUNTER(STAT_AddToWorldTime_MoveActors);
 		SCOPE_TIME_TO_VAR(&MoveActorTime);
 
-		FLevelUtils::ApplyLevelTransform( Level, LevelTransform, false );
+		FLevelUtils::FApplyLevelTransformParams TransformParams(Level, LevelTransform);
+		TransformParams.bSetRelativeTransformDirectly = true;
+		FLevelUtils::ApplyLevelTransform(TransformParams);
 
 		Level->bAlreadyMovedActors = true;
 		bExecuteNextStep = (!bConsiderTimeLimit || !IsTimeLimitExceeded( TEXT("moving actors"), StartTime, Level, TimeLimit));
@@ -2352,14 +2354,14 @@ void UWorld::AddToWorld( ULevel* Level, const FTransform& LevelTransform, bool b
 		TGuardValue<bool> IsEditorLoadingPackage(GIsEditorLoadingPackage, (GIsEditor ? true : GIsEditorLoadingPackage));
 #endif
 
-		// Config bool that allows disabling all construction scripts during PIE level streaming.
-		bool bRerunConstructionDuringEditorStreaming = true;
-		GConfig->GetBool(TEXT("Kismet"), TEXT("bRerunConstructionDuringEditorStreaming"), /*out*/ bRerunConstructionDuringEditorStreaming, GEngineIni);
+			// Config bool that allows disabling all construction scripts during PIE level streaming.
+			bool bRerunConstructionDuringEditorStreaming = true;
+			GConfig->GetBool(TEXT("Kismet"), TEXT("bRerunConstructionDuringEditorStreaming"), /*out*/ bRerunConstructionDuringEditorStreaming, GEngineIni);
 
-		// We don't need to rerun construction scripts if we have cooked data or we are playing in editor unless the PIE world was loaded
-		// from disk rather than duplicated
+			// We don't need to rerun construction scripts if we have cooked data or we are playing in editor unless the PIE world was loaded
+			// from disk rather than duplicated
 		const bool bRerunConstructionScript = !(FPlatformProperties::RequiresCookedData() || (IsGameWorld() && (Level->bHasRerunConstructionScripts || !bRerunConstructionDuringEditorStreaming)));
-		
+
 		// Incrementally update components.
 		int32 NumComponentsToUpdate = GLevelStreamingComponentsRegistrationGranularity;
 		do
diff --git a/Engine/Source/Runtime/Engine/Public/LevelUtils.h b/Engine/Source/Runtime/Engine/Public/LevelUtils.h
--- a/Engine/Source/Runtime/Engine/Public/LevelUtils.h
+++ b/Engine/Source/Runtime/Engine/Public/LevelUtils.h
@@ -91,14 +91,41 @@ public:
 	static bool IsLevelVisible(const ULevel* Level);
 
 
-	/**
-	 * Transforms the level to a new world space
-	 *
-	 * @param	Level				The level to Transform.
-	 * @param	LevelTransform		How to Transform the level. 
-	 * @param	bDoPostEditMove		Whether to call PostEditMove on actors after transforming
-	 */
-	static void ApplyLevelTransform( ULevel* Level, const FTransform& LevelTransform, bool bDoPostEditMove = true );
+	struct FApplyLevelTransformParams
+	{
+		FApplyLevelTransformParams(ULevel* InLevel, const FTransform& InLevelTransform)
+			: Level(InLevel)
+			, LevelTransform(InLevelTransform)
+		{
+		}
+
+		// The level to Transform.
+		ULevel* Level;
+
+		// How to Transform the level.
+		const FTransform& LevelTransform;
+
+		// Whether to call SetRelativeTransform or update the Location and Rotation in place without any other updating
+		bool bSetRelativeTransformDirectly = false;
+
+#if WITH_EDITOR
+		// Whether to call PostEditMove on actors after transforming
+		bool bDoPostEditMove = true;
+#endif
+	};
+
+	/** Transforms the level to a new world space */
+	static void ApplyLevelTransform(const FApplyLevelTransformParams& TransformParams);
+
+	UE_DEPRECATED(4.24, "Use version that takes params struct")
+	static void ApplyLevelTransform( ULevel* Level, const FTransform& LevelTransform, bool bDoPostEditMove = true )
+	{
+		FApplyLevelTransformParams Params(Level, LevelTransform);
+#if WITH_EDITOR
+		Params.bDoPostEditMove = bDoPostEditMove;
+#endif
+		ApplyLevelTransform(Params);
+	}
 
 #if WITH_EDITOR
 	///////////////////////////////////////////////////////////////////////////
