Restructure incremental update component calling code to avoid hitting config system in shipping builds and move components to update determination out of loop.
#rnx

diff --git a/Engine/Source/Runtime/Engine/Private/World.cpp b/Engine/Source/Runtime/Engine/Private/World.cpp
--- a/Engine/Source/Runtime/Engine/Private/World.cpp
+++ b/Engine/Source/Runtime/Engine/Private/World.cpp
@@ -2354,21 +2354,25 @@ void UWorld::AddToWorld( ULevel* Level, const FTransform& LevelTransform, bool b
 		TGuardValue<bool> IsEditorLoadingPackage(GIsEditorLoadingPackage, (GIsEditor ? true : GIsEditorLoadingPackage));
 #endif
 
-			// Config bool that allows disabling all construction scripts during PIE level streaming.
-			bool bRerunConstructionDuringEditorStreaming = true;
-			GConfig->GetBool(TEXT("Kismet"), TEXT("bRerunConstructionDuringEditorStreaming"), /*out*/ bRerunConstructionDuringEditorStreaming, GEngineIni);
-
-			// We don't need to rerun construction scripts if we have cooked data or we are playing in editor unless the PIE world was loaded
-			// from disk rather than duplicated
-		const bool bRerunConstructionScript = !(FPlatformProperties::RequiresCookedData() || (IsGameWorld() && (Level->bHasRerunConstructionScripts || !bRerunConstructionDuringEditorStreaming)));
+		bool bRerunConstructionScript = !FPlatformProperties::RequiresCookedData();
+		if (bRerunConstructionScript)
+		{
+		// Config bool that allows disabling all construction scripts during PIE level streaming.
+		bool bRerunConstructionDuringEditorStreaming = true;
+		GConfig->GetBool(TEXT("Kismet"), TEXT("bRerunConstructionDuringEditorStreaming"), /*out*/ bRerunConstructionDuringEditorStreaming, GEngineIni);
 
+		// We don't need to rerun construction scripts if we have cooked data or we are playing in editor unless the PIE world was loaded
+		// from disk rather than duplicated
+			bRerunConstructionScript = !(IsGameWorld() && (Level->bHasRerunConstructionScripts || !bRerunConstructionDuringEditorStreaming));
+		}
+		
 		// Incrementally update components.
-		int32 NumComponentsToUpdate = GLevelStreamingComponentsRegistrationGranularity;
+		int32 NumComponentsToUpdate = (!bConsiderTimeLimit || !IsGameWorld() || IsRunningCommandlet() ? 0 : GLevelStreamingComponentsRegistrationGranularity);
 		do
 		{
-			Level->IncrementalUpdateComponents( (!IsGameWorld() || IsRunningCommandlet()) ? 0 : NumComponentsToUpdate, bRerunConstructionScript );
+			Level->IncrementalUpdateComponents( NumComponentsToUpdate, bRerunConstructionScript );
 		}
-		while( !Level->bAreComponentsCurrentlyRegistered && (!bConsiderTimeLimit || !IsTimeLimitExceeded( TEXT("updating components"), StartTime, Level, TimeLimit )));
+		while (!Level->bAreComponentsCurrentlyRegistered && !IsTimeLimitExceeded(TEXT("updating components"), StartTime, Level, TimeLimit));
 
 		// We are done once all components are attached.
 		Level->bAlreadyUpdatedComponents	= Level->bAreComponentsCurrentlyRegistered;
