[NO NON-PII COMMIT MESSAGE]

diff --git a/Engine/Source/Runtime/Engine/Private/AudioStreamingCache.cpp b/Engine/Source/Runtime/Engine/Private/AudioStreamingCache.cpp
--- a/Engine/Source/Runtime/Engine/Private/AudioStreamingCache.cpp
+++ b/Engine/Source/Runtime/Engine/Private/AudioStreamingCache.cpp
@@ -216,7 +216,6 @@ FAudioChunkHandle FCachedAudioStreamingManager::GetLoadedChunk(const USoundWave*
 			}
 		}
 
-		
 		return BuildChunkHandle(LoadedChunk.GetData(), LoadedChunk.Num(), SoundWave, SoundWave->GetFName(), ChunkIndex);
 	}
 	else
@@ -234,21 +233,7 @@ FAudioChunkCache* FCachedAudioStreamingManager::GetCacheForWave(const USoundWave
 	if (InSoundWave->RunningPlatformData && InSoundWave->RunningPlatformData->Chunks.Num() > 1)
 	{
 		const int32 SoundWaveChunkSize = InSoundWave->RunningPlatformData->Chunks[1].AudioDataSize;
-
-		// Iterate over our caches until we find the lowest MaxChunkSize cache this sound's chunks will fit into. 
-		for (int32 CacheIndex = 0; CacheIndex < CacheArray.Num(); CacheIndex++)
-		{
-			if (SoundWaveChunkSize <= CacheArray[CacheIndex].MaxChunkSize)
-			{
-				return const_cast<FAudioChunkCache*>(&CacheArray[CacheIndex]);
-			}
-		}
-
-		// If we ever hit this, something went wrong during cook.
-		// Please check to make sure this platform's implementation of IAudioFormat honors the MaxChunkSize parameter passed into SplitDataForStreaming,
-		// or that FStreamedAudioCacheDerivedDataWorker::BuildStreamedAudio() is passing the correct MaxChunkSize to IAudioFormat::SplitDataForStreaming.
-		ensureMsgf(false, TEXT("Chunks in SoundWave are too large: %s, Chunk size: %d bytes"), *InSoundWave->GetName(), SoundWaveChunkSize);
-		return nullptr;
+		return GetCacheForChunkSize(SoundWaveChunkSize);
 	}
 	else
 	{
@@ -256,6 +241,25 @@ FAudioChunkCache* FCachedAudioStreamingManager::GetCacheForWave(const USoundWave
 	}
 }
 
+FAudioChunkCache* FCachedAudioStreamingManager::GetCacheForChunkSize(uint32 InChunkSize) const
+{
+	// Iterate over our caches until we find the lowest MaxChunkSize cache this sound's chunks will fit into. 
+	for (int32 CacheIndex = 0; CacheIndex < CacheArray.Num(); CacheIndex++)
+	{
+		check(CacheArray[CacheIndex].MaxChunkSize >= 0);
+		if (InChunkSize <= ((uint32) CacheArray[CacheIndex].MaxChunkSize))
+		{
+			return const_cast<FAudioChunkCache*>(&CacheArray[CacheIndex]);
+		}
+	}
+
+	// If we ever hit this, something went wrong during cook.
+	// Please check to make sure this platform's implementation of IAudioFormat honors the MaxChunkSize parameter passed into SplitDataForStreaming,
+	// or that FStreamedAudioCacheDerivedDataWorker::BuildStreamedAudio() is passing the correct MaxChunkSize to IAudioFormat::SplitDataForStreaming.
+	ensureMsgf(false, TEXT("Chunks in SoundWave are too large: %d bytes"), InChunkSize);
+	return nullptr;
+}
+
 int32 FCachedAudioStreamingManager::GetNextChunkIndex(const USoundWave* InSoundWave, uint32 CurrentChunkIndex) const
 {
 	check(InSoundWave);
@@ -283,7 +287,7 @@ int32 FCachedAudioStreamingManager::GetNextChunkIndex(const USoundWave* InSoundW
 
 void FCachedAudioStreamingManager::AddReferenceToChunk(const FAudioChunkHandle& InHandle)
 {
-	FAudioChunkCache* Cache = GetCacheForWave(InHandle.CorrespondingWave);
+	FAudioChunkCache* Cache = GetCacheForChunkSize(InHandle.CachedDataNumBytes);
 	check(Cache);
 
 	FAudioChunkCache::FChunkKey ChunkKey =
@@ -298,7 +302,7 @@ void FCachedAudioStreamingManager::AddReferenceToChunk(const FAudioChunkHandle&
 
 void FCachedAudioStreamingManager::RemoveReferenceToChunk(const FAudioChunkHandle& InHandle)
 {
-	FAudioChunkCache* Cache = GetCacheForWave(InHandle.CorrespondingWave);
+	FAudioChunkCache* Cache = GetCacheForChunkSize(InHandle.CachedDataNumBytes);
 	check(Cache);
 
 	FAudioChunkCache::FChunkKey ChunkKey =
@@ -876,6 +880,11 @@ void FAudioChunkCache::KickOffAsyncLoad(FCacheElement* CacheElement, const FChun
 		check(Chunk.BulkData.GetFilename().Len());
 		UE_CLOG(Chunk.BulkData.IsStoredCompressedOnDisk(), LogAudio, Fatal, TEXT("Package level compression is no longer supported."));
 
+		if (CacheElement->IsLoadInProgress())
+		{
+			CacheElement->WaitForAsyncLoadCompletion(true);
+		}
+
 		// TODO: Figure out how to cache this file handle somewhere, so that we are not reopening a file on every individual load operation.
 		CacheElement->FileHandle.Reset(FPlatformFileManager::Get().GetPlatformFile().OpenAsyncRead(*Chunk.BulkData.GetFilename()));
 		check(CacheElement->FileHandle.IsValid());
diff --git a/Engine/Source/Runtime/Engine/Public/AudioStreamingCache.h b/Engine/Source/Runtime/Engine/Public/AudioStreamingCache.h
--- a/Engine/Source/Runtime/Engine/Public/AudioStreamingCache.h
+++ b/Engine/Source/Runtime/Engine/Public/AudioStreamingCache.h
@@ -348,6 +348,7 @@ protected:
 	 * or nullptr if MemoryLoadOnDemand is disabled.
 	 */
 	FAudioChunkCache* GetCacheForWave(const USoundWave* InSoundWave) const;
+	FAudioChunkCache* GetCacheForChunkSize(uint32 InChunkSize) const;
 
 	/**
 	 * Returns the next chunk to kick off a load for, or INDEX_NONE if there is only one chunk to cache.
