Mark log in ULevelActorContainer::CreateCluster verbose. Saves ~70% of CreateCluster time due to eliminating GetPathName call.
#rnx

diff --git a/Engine/Source/Runtime/Engine/Private/LevelActorContainer.cpp b/Engine/Source/Runtime/Engine/Private/LevelActorContainer.cpp
--- a/Engine/Source/Runtime/Engine/Private/LevelActorContainer.cpp
+++ b/Engine/Source/Runtime/Engine/Private/LevelActorContainer.cpp
@@ -200,7 +200,7 @@ void ULevelActorContainer::CreateCluster()
 		Cluster.ReferencedClusters.Sort();
 		Cluster.MutableObjects.Sort();
 
-		UE_LOG(LogLevelActorContainer, Log, TEXT("Created LevelActorCluster (%d) for %s with %d objects, %d referenced clusters and %d mutable objects."),
+		UE_LOG(LogLevelActorContainer, Verbose, TEXT("Created LevelActorCluster (%d) for %s with %d objects, %d referenced clusters and %d mutable objects."),
 			ClusterIndex, *GetOuter()->GetPathName(), Cluster.Objects.Num(), Cluster.ReferencedClusters.Num(), Cluster.MutableObjects.Num());
 
 #if UE_GCCLUSTER_VERBOSE_LOGGING
