Fix for ADPCM crash when using the old iOS audio backend with stream caching. Explicitly releases stream chunk when IOSAudioSource::Stop is called.

diff --git a/Engine/Source/Runtime/Engine/Private/ADPCMAudioInfo.cpp b/Engine/Source/Runtime/Engine/Private/ADPCMAudioInfo.cpp
--- a/Engine/Source/Runtime/Engine/Private/ADPCMAudioInfo.cpp
+++ b/Engine/Source/Runtime/Engine/Private/ADPCMAudioInfo.cpp
@@ -438,6 +438,8 @@ void FADPCMAudioInfo::ProcessSeekRequest()
 
 bool FADPCMAudioInfo::StreamCompressedInfoInternal(USoundWave* Wave, struct FSoundQualityInfo* QualityInfo)
 {
+	FScopeLock ScopeLock(&CurCompressedChunkHandleCriticalSection);
+
 	check(QualityInfo);
 
 	check(StreamingSoundWave == Wave);
@@ -539,6 +541,8 @@ bool FADPCMAudioInfo::StreamCompressedInfoInternal(USoundWave* Wave, struct FSou
 
 bool FADPCMAudioInfo::StreamCompressedData(uint8* Destination, bool bLooping, uint32 BufferSize)
 {
+	FScopeLock ScopeLock(&CurCompressedChunkHandleCriticalSection);
+
 	// Initial sanity checks:
 	if (Destination == nullptr || BufferSize == 0)
 	{
@@ -762,6 +766,38 @@ bool FADPCMAudioInfo::StreamCompressedData(uint8* Destination, bool bLooping, ui
 	return ReachedEndOfSamples;
 }
 
+bool FADPCMAudioInfo::ReleaseStreamChunk(bool bBlockUntilReleased)
+{
+	if (bBlockUntilReleased)
+	{
+		// Wait for any pending decode tasks to finish.
+		FScopeLock ScopeLock(&CurCompressedChunkHandleCriticalSection);
+		CurCompressedChunkHandle = FAudioChunkHandle();
+		CurrentChunkBufferOffset = 0;
+		CurCompressedChunkData = nullptr;
+		CurrentChunkDataSize = 0;
+		return true;
+	}
+	else
+	{
+		// If the chunk isn't currently in use, reset the chunk handle.
+		if (CurCompressedChunkHandleCriticalSection.TryLock())
+		{
+			CurCompressedChunkHandle = FAudioChunkHandle();
+			CurrentChunkBufferOffset = 0;
+			CurCompressedChunkData = nullptr;
+			CurrentChunkDataSize = 0;
+			CurCompressedChunkHandleCriticalSection.Unlock();
+			return true;
+		}
+		else
+		{
+			// Otherwise, do nothing and exit.
+			return false;
+		}
+	}
+}
+
 const uint8* FADPCMAudioInfo::GetLoadedChunk(USoundWave* InSoundWave, uint32 ChunkIndex, uint32& OutChunkSize)
 {
 	if (!InSoundWave || ChunkIndex >= InSoundWave->GetNumChunks())
diff --git a/Engine/Source/Runtime/Engine/Public/ADPCMAudioInfo.h b/Engine/Source/Runtime/Engine/Public/ADPCMAudioInfo.h
--- a/Engine/Source/Runtime/Engine/Public/ADPCMAudioInfo.h
+++ b/Engine/Source/Runtime/Engine/Public/ADPCMAudioInfo.h
@@ -113,6 +113,7 @@ public:
 		return CurrentChunkIndex;
 	}
 	virtual int32 GetCurrentChunkOffset() const override {return CurrentChunkBufferOffset;}
+	virtual bool ReleaseStreamChunk(bool bBlockUntilReleased) override;
 	// End of ICompressedAudioInfo Interface
 
 	uint16 GetFormatTag()
@@ -163,6 +164,7 @@ private:
 	uint32          FirstChunkSampleDataIndex;
 	const uint8*	CurCompressedChunkData;			// A pointer to the current chunk of data
 	FAudioChunkHandle CurCompressedChunkHandle;     // Shared reference to the current chunk of data.
+	FCriticalSection CurCompressedChunkHandleCriticalSection;
 
 	uint32			CurrentCompressedBlockIndex;		// For non disk streaming - the current compressed block in the compressed source data
 	uint32			TotalCompressedBlocksPerChannel;	// For non disk streaming - the total number of compressed blocks per channel
diff --git a/Engine/Source/Runtime/Engine/Public/AudioDecompress.h b/Engine/Source/Runtime/Engine/Public/AudioDecompress.h
--- a/Engine/Source/Runtime/Engine/Public/AudioDecompress.h
+++ b/Engine/Source/Runtime/Engine/Public/AudioDecompress.h
@@ -98,6 +98,19 @@ public:
 	 */
 	virtual bool SupportsStreaming() const {return false;}
 
+	/**
+	 * This can be called to explicitly release this decoder's reference to a chunk of compressed audio
+	 * without destroying the decoder itself.
+	 * @param bBlockUntilReleased when set to true will cause this call to block if the decoder is currently using the chunk.
+	 * @returns true if the chunk was released, false otherwise.
+	 */
+	virtual bool ReleaseStreamChunk(bool bBlockUntilReleased)
+	{
+		// If we hit this check, ReleaseStreamChunk needs to be implemented for this codec.
+		checkNoEntry();
+		return false;
+	}
+
 	/**
 	* Streams the header information of a compressed format
 	*
diff --git a/Engine/Source/Runtime/IOS/IOSAudio/Private/IOSAudioBuffer.cpp b/Engine/Source/Runtime/IOS/IOSAudio/Private/IOSAudioBuffer.cpp
--- a/Engine/Source/Runtime/IOS/IOSAudio/Private/IOSAudioBuffer.cpp
+++ b/Engine/Source/Runtime/IOS/IOSAudio/Private/IOSAudioBuffer.cpp
@@ -195,3 +195,14 @@ FIOSAudioSoundBuffer* FIOSAudioSoundBuffer::Init(FIOSAudioDevice* IOSAudioDevice
 	return Buffer;
 }
 
+bool FIOSAudioSoundBuffer::ReleaseCurrentChunk()
+{
+	if (DecompressionState && bStreaming)
+	{
+		return DecompressionState->ReleaseStreamChunk(true);
+	}
+	else
+	{
+		return true;
+	}
+}
diff --git a/Engine/Source/Runtime/IOS/IOSAudio/Private/IOSAudioSource.cpp b/Engine/Source/Runtime/IOS/IOSAudio/Private/IOSAudioSource.cpp
--- a/Engine/Source/Runtime/IOS/IOSAudio/Private/IOSAudioSource.cpp
+++ b/Engine/Source/Runtime/IOS/IOSAudio/Private/IOSAudioSource.cpp
@@ -338,6 +338,13 @@ void FIOSAudioSoundSource::Stop(void)
 	// Call parent class version regardless of if there's a wave instance
 	FSoundSource::Stop();
 
+	if (IOSBuffer != NULL)
+	{
+		// Release the current stream chunk:
+		bool bChunkReleased = IOSBuffer->ReleaseCurrentChunk();
+		check(bChunkReleased);
+	}
+
 	// It's now safe to unlock the callback
 	UnlockCallback(&CallbackLock);
 }
diff --git a/Engine/Source/Runtime/IOS/IOSAudio/Public/IOSAudioDevice.h b/Engine/Source/Runtime/IOS/IOSAudio/Public/IOSAudioDevice.h
--- a/Engine/Source/Runtime/IOS/IOSAudio/Public/IOSAudioDevice.h
+++ b/Engine/Source/Runtime/IOS/IOSAudio/Public/IOSAudioDevice.h
@@ -85,6 +85,7 @@ public:
 	 */
 	virtual bool ReadCompressedData( uint8* Destination, int32 NumFramesToDecode, bool bLooping ) override;
 	
+	bool ReleaseCurrentChunk();
 	
 	int32  RenderCallbackBufferSize;
 	int32  SampleRate;
