If StreamCompressedData is called after ADPCMAudioInfo::ReleaseChunk is called, early exit. #jira none #review-8321600 aaron.mcleran, ryan.mangin, phil.popp, rob.gay, maxwell.hayes

diff --git a/Engine/Source/Runtime/Engine/Private/ADPCMAudioInfo.cpp b/Engine/Source/Runtime/Engine/Private/ADPCMAudioInfo.cpp
--- a/Engine/Source/Runtime/Engine/Private/ADPCMAudioInfo.cpp
+++ b/Engine/Source/Runtime/Engine/Private/ADPCMAudioInfo.cpp
@@ -45,6 +45,7 @@ FADPCMAudioInfo::FADPCMAudioInfo(void)
 	, SamplesPerBlock(0)
 	, FirstChunkSampleDataOffset(0)
 	, FirstChunkSampleDataIndex(0)
+	, bDecompressorReleased(false)
 	, bSeekPending(false)
 	, TargetSeekTime(0.0f)
 	, LastSeekTime(0.0f)
@@ -484,6 +485,7 @@ bool FADPCMAudioInfo::StreamCompressedInfoInternal(USoundWave* Wave, struct FSou
 	SrcBufferData = ChunkData;
 	CurrentChunkBufferOffset = 0;
 	CurCompressedChunkData = nullptr;
+	bDecompressorReleased = false;
 	CurrentUncompressedBlockSampleIndex = 0;
 	
 	TotalSamplesStreamed = 0;
@@ -544,6 +546,13 @@ bool FADPCMAudioInfo::StreamCompressedData(uint8* Destination, bool bLooping, ui
 	FScopeLock ScopeLock(&CurCompressedChunkHandleCriticalSection);
 
 	// Initial sanity checks:
+	if (bDecompressorReleased)
+	{
+		UE_LOG(LogAudio, Error, TEXT("Stream Compressed Data called after release!"));
+		FMemory::Memzero(Destination, BufferSize);
+		return true;
+	}
+
 	if (Destination == nullptr || BufferSize == 0)
 	{
 		UE_LOG(LogAudio, Error, TEXT("Stream Compressed Info not called!"));
@@ -776,6 +785,7 @@ bool FADPCMAudioInfo::ReleaseStreamChunk(bool bBlockUntilReleased)
 		CurrentChunkBufferOffset = 0;
 		CurCompressedChunkData = nullptr;
 		CurrentChunkDataSize = 0;
+		bDecompressorReleased = true;
 		return true;
 	}
 	else
@@ -787,6 +797,7 @@ bool FADPCMAudioInfo::ReleaseStreamChunk(bool bBlockUntilReleased)
 			CurrentChunkBufferOffset = 0;
 			CurCompressedChunkData = nullptr;
 			CurrentChunkDataSize = 0;
+			bDecompressorReleased = true;
 			CurCompressedChunkHandleCriticalSection.Unlock();
 			return true;
 		}
diff --git a/Engine/Source/Runtime/Engine/Public/ADPCMAudioInfo.h b/Engine/Source/Runtime/Engine/Public/ADPCMAudioInfo.h
--- a/Engine/Source/Runtime/Engine/Public/ADPCMAudioInfo.h
+++ b/Engine/Source/Runtime/Engine/Public/ADPCMAudioInfo.h
@@ -165,6 +165,7 @@ private:
 	const uint8*	CurCompressedChunkData;			// A pointer to the current chunk of data
 	FAudioChunkHandle CurCompressedChunkHandle;     // Shared reference to the current chunk of data.
 	FCriticalSection CurCompressedChunkHandleCriticalSection;
+	bool bDecompressorReleased;                     // When the chunk is released, we raise this to true, then early exit on future decodes.
 
 	uint32			CurrentCompressedBlockIndex;		// For non disk streaming - the current compressed block in the compressed source data
 	uint32			TotalCompressedBlocksPerChannel;	// For non disk streaming - the total number of compressed blocks per channel
