Fixed instanced static mesh asset related to missing manual vertex fetch buffers

diff --git a/Engine/Source/Runtime/Engine/Private/InstancedStaticMesh.cpp b/Engine/Source/Runtime/Engine/Private/InstancedStaticMesh.cpp
--- a/Engine/Source/Runtime/Engine/Private/InstancedStaticMesh.cpp
+++ b/Engine/Source/Runtime/Engine/Private/InstancedStaticMesh.cpp
@@ -373,7 +373,7 @@ void FStaticMeshInstanceBuffer::CreateVertexBuffer(FResourceArrayInterface* InRe
 
 void FStaticMeshInstanceBuffer::BindInstanceVertexBuffer(const class FVertexFactory* VertexFactory, FInstancedStaticMeshDataType& InstancedStaticMeshData) const
 {
-	if (InstanceData->GetNumInstances())
+	if (InstanceData->GetNumInstances() && RHISupportsManualVertexFetch(GMaxRHIShaderPlatform))
 	{
 		check(InstanceOriginSRV);
 		check(InstanceTransformSRV);
