Add functionality to allow DeviceProfileManager to query if a device profile is in use and to reapply the current deviceprofile.
#rnx

diff --git a/Engine/Source/Runtime/Engine/Classes/DeviceProfiles/DeviceProfileManager.h b/Engine/Source/Runtime/Engine/Classes/DeviceProfiles/DeviceProfileManager.h
--- a/Engine/Source/Runtime/Engine/Classes/DeviceProfiles/DeviceProfileManager.h
+++ b/Engine/Source/Runtime/Engine/Classes/DeviceProfiles/DeviceProfileManager.h
@@ -32,6 +32,20 @@ public:
 	 */
 	static void InitializeCVarsForActiveDeviceProfile(bool bPushSettings=false);
 
+	/**
+	 * Reapplies the device profile. Useful when configs have changed (i.e. hotfix)
+	 * Applies base and then any overridden device profile.
+	 */
+	void ReapplyDeviceProfile();
+
+	/**
+	 * Examine the currently active or overridden profile for references to entries in DeviceProfilesToQuery
+	 * @param DeviceProfilesToQuery - Collection of device profiles to check.
+	 * 
+	 * @return true if any profiles contained in DeviceProfilesToQuery are referenced by active or overridden profile.
+	 */
+	bool DoActiveProfilesReference(const TSet<FString>& DeviceProfilesToQuery);
+
 	/**
 	 * Create a copy of a device profile from a copy.
 	 *
@@ -166,6 +180,9 @@ private:
 	// values of CVars set in HandleDeviceProfileOverrideChange, to be popped later
 	TMap<FString, FString> PushedSettings;
 
+	// Holds the device profile that has been overridden, null no override active.
+	UDeviceProfile* BaseDeviceProfile = nullptr;
+
 	// Stores any scalability group settings set by the active device profile.
 	static TMap<FString, FString> DeviceProfileScalabilityCVars;
 };
diff --git a/Engine/Source/Runtime/Engine/Private/DeviceProfiles/DeviceProfileManager.cpp b/Engine/Source/Runtime/Engine/Private/DeviceProfiles/DeviceProfileManager.cpp
--- a/Engine/Source/Runtime/Engine/Private/DeviceProfiles/DeviceProfileManager.cpp
+++ b/Engine/Source/Runtime/Engine/Private/DeviceProfiles/DeviceProfileManager.cpp
@@ -294,6 +294,79 @@ void UDeviceProfileManager::InitializeCVarsForActiveDeviceProfile(bool bPushSett
 	}
 }
 
+bool UDeviceProfileManager::DoActiveProfilesReference(const TSet<FString>& DeviceProfilesToQuery)
+{
+	FConfigCacheIni::LoadGlobalIniFile(DeviceProfileFileName, TEXT("DeviceProfiles"));
+	TArray< FString > AvailableProfiles;
+	GConfig->GetSectionNames(DeviceProfileFileName, AvailableProfiles);
+
+	auto DoesProfileReference = [&AvailableProfiles, DeviceProfileFileName = DeviceProfileFileName](const FString& SearchProfile, const TSet<FString>& InDeviceProfilesToQuery)
+	{
+		// For each device profile, starting with the selected and working our way up the BaseProfileName tree,
+		FString BaseDeviceProfileName = SearchProfile;
+		bool bReachedEndOfTree = BaseDeviceProfileName.IsEmpty();
+		while (bReachedEndOfTree == false)
+		{
+			FString CurrentSectionName = FString::Printf(TEXT("%s %s"), *BaseDeviceProfileName, *UDeviceProfile::StaticClass()->GetName());
+			bool bProfileExists = AvailableProfiles.Contains(CurrentSectionName);
+			if (bProfileExists)
+			{
+				if (InDeviceProfilesToQuery.Contains(BaseDeviceProfileName))
+				{
+					return true;
+				}
+
+				// Get the next device profile name
+				FString NextBaseDeviceProfileName;
+				if (GConfig->GetString(*CurrentSectionName, TEXT("BaseProfileName"), NextBaseDeviceProfileName, DeviceProfileFileName))
+				{
+					BaseDeviceProfileName = NextBaseDeviceProfileName;
+				}
+				else
+				{
+					BaseDeviceProfileName.Empty();
+				}
+			}
+			bReachedEndOfTree = !bProfileExists || BaseDeviceProfileName.IsEmpty();
+		}
+		return false;
+	};
+
+	bool bDoActiveProfilesReferenceQuerySet = DoesProfileReference(DeviceProfileManagerSingleton->GetActiveProfile()->GetName(), DeviceProfilesToQuery);
+	if (!bDoActiveProfilesReferenceQuerySet && DeviceProfileManagerSingleton->BaseDeviceProfile != nullptr)
+	{
+		bDoActiveProfilesReferenceQuerySet = DoesProfileReference(DeviceProfileManagerSingleton->BaseDeviceProfile->GetName(), DeviceProfilesToQuery);
+	}
+	return bDoActiveProfilesReferenceQuerySet;
+}
+
+void UDeviceProfileManager::ReapplyDeviceProfile()
+{	
+	UDeviceProfile* OverrideProfile = DeviceProfileManagerSingleton->BaseDeviceProfile ? DeviceProfileManagerSingleton->GetActiveProfile() : nullptr;
+	UDeviceProfile* BaseProfile = DeviceProfileManagerSingleton->BaseDeviceProfile ? DeviceProfileManagerSingleton->BaseDeviceProfile : DeviceProfileManagerSingleton->GetActiveProfile();
+
+	UE_LOG(LogInit, Log, TEXT("ReapplyDeviceProfile applying profile: [%s]"), *BaseProfile->GetName(), OverrideProfile ? *OverrideProfile->GetName() : TEXT("not set.") );
+
+	// pop any pushed settings
+	RestoreDefaultDeviceProfile();
+
+	// Set base profile and re-apply cvars.
+	SetActiveDeviceProfile(BaseProfile);
+	InitializeCVarsForActiveDeviceProfile();
+
+	if (OverrideProfile)
+	{
+		UE_LOG(LogInit, Log, TEXT("ReapplyDeviceProfile applying override profile: [%s]"), *OverrideProfile->GetName());
+		// reapply the override.
+		SetOverrideDeviceProfile(OverrideProfile);
+	}
+	else
+	{
+		// broadcast cvar sinks now that we are done
+		IConsoleManager::Get().CallAllConsoleVariableSinks();
+	}
+}
+
 static void TestProfileForCircularReferences(const FString& ProfileName, const FString& ParentName, const FConfigFile &PlatformConfigFile)
 {
 	TArray<FString> ProfileDependancies;
@@ -533,6 +606,9 @@ void UDeviceProfileManager::SetOverrideDeviceProfile(UDeviceProfile* DeviceProfi
 	// pop any pushed settings
 	HandleDeviceProfileOverridePop();
 
+	// record the currently active profile, needed when we restore the default.
+	BaseDeviceProfile = DeviceProfileManagerSingleton->GetActiveProfile();
+
 	// activate new one!
 	DeviceProfileManagerSingleton->SetActiveDeviceProfile(DeviceProfile);
 	InitializeCVarsForActiveDeviceProfile(true);
@@ -557,6 +633,13 @@ void UDeviceProfileManager::RestoreDefaultDeviceProfile()
 			UE_LOG(LogInit, Log, TEXT("Popping Device Profile CVar: [[%s:%s]]"), *It.Key(), *It.Value());
 		}
 	}
+
+	if(BaseDeviceProfile)
+	{
+		// reset the base profile as we are no longer overriding
+		DeviceProfileManagerSingleton->SetActiveDeviceProfile(BaseDeviceProfile);
+		BaseDeviceProfile = nullptr;
+	}
 }
 
 
