Fix for crash in audio device update

- nullptr check needed due to UWorld flush doing a force pending kill on all uobjects, which results in the uobject getting nulled in addreferenced objects.

diff --git a/Engine/Source/Runtime/Engine/Private/AudioDevice.cpp b/Engine/Source/Runtime/Engine/Private/AudioDevice.cpp
--- a/Engine/Source/Runtime/Engine/Private/AudioDevice.cpp
+++ b/Engine/Source/Runtime/Engine/Private/AudioDevice.cpp
@@ -743,6 +743,8 @@ void FAudioDevice::AddReferencedObjects(FReferenceCollector& Collector)
 		Pair.Key->AddReferencedObjects(Collector);
 	}
 
+	UpdateReferencedSoundWaves();
+
 	// Make sure we don't try to delete any sound waves which may have in-flight decodes
 	Collector.AddReferencedObjects(ReferencedSoundWaves);
 
@@ -5360,6 +5362,11 @@ void FAudioDevice::Flush(UWorld* WorldToFlush, bool bClearActivatedReverb)
 		}
 	}
 
+	if (WorldToFlush == nullptr)
+	{
+		ReferencedSoundWaves.Reset();
+	}
+
 	// Make sure we update any hardware changes that need to happen after flushing
 	if (IsAudioMixerEnabled() && (WorldToFlush == nullptr || WorldToFlush->bIsTearingDown))
 	{
