[NO NON-PII COMMIT MESSAGE]

diff --git a/Engine/Source/Runtime/Engine/Private/AudioStreamingCache.cpp b/Engine/Source/Runtime/Engine/Private/AudioStreamingCache.cpp
--- a/Engine/Source/Runtime/Engine/Private/AudioStreamingCache.cpp
+++ b/Engine/Source/Runtime/Engine/Private/AudioStreamingCache.cpp
@@ -204,6 +204,7 @@ FAudioChunkHandle FCachedAudioStreamingManager::GetLoadedChunk(const USoundWave*
 		// The function call below increments the reference count to the internal chunk.
 		TArrayView<uint8> LoadedChunk = Cache->GetChunk(ChunkKey, bBlockForLoad);
 
+		// Finally, if there's a chunk after this in the sound, request that it is in the cache.
 		const int32 NextChunk = GetNextChunkIndex(SoundWave, ChunkIndex);
 
 		if (NextChunk != INDEX_NONE)
@@ -363,6 +364,8 @@ bool FAudioChunkCache::AddOrTouchChunk(const FChunkKey& InKey, TFunction<void(EA
 		return false;
 	}
 
+	FScopeLock ScopeLock(&CacheMutationCriticalSection);
+
 	FCacheElement* FoundElement = FindElementForKey(InKey);
 	
 	if (FoundElement)
@@ -402,6 +405,8 @@ bool FAudioChunkCache::AddOrTouchChunk(const FChunkKey& InKey, TFunction<void(EA
 
 TArrayView<uint8> FAudioChunkCache::GetChunk(const FChunkKey& InKey, bool bBlockForLoadCompletion)
 {
+	FScopeLock ScopeLock(&CacheMutationCriticalSection);
+
 	FCacheElement* FoundElement = FindElementForKey(InKey);
 	if (FoundElement)
 	{
@@ -443,7 +448,7 @@ TArrayView<uint8> FAudioChunkCache::GetChunk(const FChunkKey& InKey, bool bBlock
 		else if (bLogCacheMisses)
 		{
 			// Chunks missing. Log this as a miss.
-			const uint32 TotalNumChunksInWave = InKey.SoundWave->RunningPlatformData->Chunks.Num();
+			const uint32 TotalNumChunksInWave = InKey.SoundWave->GetNumChunks();
 
 			FCacheMissInfo CacheMissInfo = { InKey.SoundWaveName, InKey.ChunkIndex, TotalNumChunksInWave, false };
 			CacheMissQueue.Enqueue(MoveTemp(CacheMissInfo));
@@ -476,6 +481,7 @@ void FAudioChunkCache::RemoveReferenceToChunk(const FChunkKey& InKey)
 
 void FAudioChunkCache::ClearCache()
 {
+	FScopeLock ScopeLock(&CacheMutationCriticalSection);
 	const uint32 NumChunks = CachePool.Num();
 
 	CachePool.Reset(NumChunks);
@@ -548,18 +554,23 @@ void FAudioChunkCache::BlockForAllPendingLoads() const
 			FPlatformProcess::Sleep(0.0f);
 		}
 
-		// Iterate through every element until we find one with a load in progress.
-		const FCacheElement* CurrentElement = MostRecentElement;
-		while (CurrentElement != nullptr)
 		{
-			bLoadInProgress |= CurrentElement->IsLoadInProgress();
-			CurrentElement = CurrentElement->LessRecentElement;
+			FScopeLock ScopeLock(const_cast<FCriticalSection*>(&CacheMutationCriticalSection));
+
+			// Iterate through every element until we find one with a load in progress.
+			const FCacheElement* CurrentElement = MostRecentElement;
+			while (CurrentElement != nullptr)
+			{
+				bLoadInProgress |= CurrentElement->IsLoadInProgress();
+				CurrentElement = CurrentElement->LessRecentElement;
+			}
 		}
 	} while (bLoadInProgress);
 }
 
 void FAudioChunkCache::CancelAllPendingLoads()
 {
+	FScopeLock ScopeLock(&CacheMutationCriticalSection);
 	FCacheElement* CurrentElement = MostRecentElement;
 	while (CurrentElement != nullptr)
 	{
@@ -604,6 +615,7 @@ FString FAudioChunkCache::FlushCacheMissLog()
 
 FAudioChunkCache::FCacheElement* FAudioChunkCache::FindElementForKey(const FChunkKey& InKey)
 {
+	FScopeLock ScopeLock(&CacheMutationCriticalSection);
 	FCacheElement* CurrentElement = MostRecentElement;
 
 #if DEBUG_STREAM_CACHE
