Fix streaming VT texture copy to staging buffer
Not using D3D_USAGE_DEFAULT triggers an error with d3d debug runtime and can fail at runtime depending on driver etc

diff --git a/Engine/Source/Runtime/Engine/Private/VT/VirtualTextureUploadCache.cpp b/Engine/Source/Runtime/Engine/Private/VT/VirtualTextureUploadCache.cpp
--- a/Engine/Source/Runtime/Engine/Private/VT/VirtualTextureUploadCache.cpp
+++ b/Engine/Source/Runtime/Engine/Private/VT/VirtualTextureUploadCache.cpp
@@ -111,7 +111,7 @@ void FVirtualTextureUploadCache::Finalize(FRHICommandListImmediate& RHICmdList)
 			}
 
 			FRHIResourceCreateInfo CreateInfo;
-			StagingTexture.RHITexture = RHICmdList.CreateTexture2D(TileSize * WidthInTiles, TileSize * HeightInTiles, PoolEntry.Format, 1, 1, TexCreate_CPUWritable, CreateInfo);
+			StagingTexture.RHITexture = RHICmdList.CreateTexture2D(TileSize * WidthInTiles, TileSize * HeightInTiles, PoolEntry.Format, 1, 1, TexCreate_None, CreateInfo);
 			StagingTexture.WidthInTiles = WidthInTiles;
 			StagingTexture.BatchCapacity = WidthInTiles * HeightInTiles;
 			INC_MEMORY_STAT_BY(STAT_TotalGPUUploadSize, CalcTextureSize(TileSize * WidthInTiles, TileSize * HeightInTiles, PoolEntry.Format, 1u));
