Make quantization of Scale, Velocity, and Location on initial actor spawn configurable (while keeping backwards compat).

diff --git a/Engine/Source/Runtime/Core/Public/Misc/NetworkVersion.h b/Engine/Source/Runtime/Core/Public/Misc/NetworkVersion.h
--- a/Engine/Source/Runtime/Core/Public/Misc/NetworkVersion.h
+++ b/Engine/Source/Runtime/Core/Public/Misc/NetworkVersion.h
@@ -44,6 +44,7 @@ enum EEngineNetworkVersionHistory
 	HISTORY_NETEXPORT_SERIALIZE_FIX = 10,			// Bump version to fix net field export name serialization 
 	HISTORY_FAST_ARRAY_DELTA_STRUCT = 11,			// Bump version to allow fast array serialization, delta struct serialization.
 	HISTORY_FIX_ENUM_SERIALIZATION = 12,			// Bump version to fix enum net serialization issues.
+	HISTORY_OPTIONALLY_QUANTIZE_SPAWN_INFO = 13,	// Bump version to conditionally disable quantization for Scale, Location, and Velocity when spawning network actors.
 	// New history items go above here.
 
 	HISTORY_ENGINENETVERSION_PLUS_ONE,
diff --git a/Engine/Source/Runtime/Engine/Private/PackageMapClient.cpp b/Engine/Source/Runtime/Engine/Private/PackageMapClient.cpp
--- a/Engine/Source/Runtime/Engine/Private/PackageMapClient.cpp
+++ b/Engine/Source/Runtime/Engine/Private/PackageMapClient.cpp
@@ -83,6 +83,27 @@ static FAutoConsoleVariableRef CVarAllowClientRemapCacheObject(
 	TEXT("When enabled, we will allow clients to remap read only cache objects and keep the same NetGUID.")
 );
 
+static bool GbQuantizeActorScaleOnSpawn = false;
+static FAutoConsoleVariableRef CVarQuantizeActorScaleOnSpawn(
+	TEXT("net.QuantizeActorScaleOnSpawn"),
+	GbQuantizeActorScaleOnSpawn,
+	TEXT("When enabled, we will quantize Scale for newly spawned actors to a single decimal of precision.")
+);
+
+static bool GbQuantizeActorLocationOnSpawn = true;
+static FAutoConsoleVariableRef CVarQuantizeActorLocationOnSpawn(
+	TEXT("net.QuantizeActorLocationOnSpawn"),
+	GbQuantizeActorLocationOnSpawn,
+	TEXT("When enabled, we will quantize Location for newly spawned actors to a single decimal of precision.")
+);
+
+static bool GbQuantizeActorVelocityOnSpawn = true;
+static FAutoConsoleVariableRef CVarQuantizeActorVelocityOnSpawn(
+	TEXT("net.QuantizeActorVelocityOnSpawn"),
+	GbQuantizeActorVelocityOnSpawn,
+	TEXT("When enabled, we will quantize Velocity for newly spawned actors to a single decimal of precision.")
+);
+
 void BroadcastNetFailure(UNetDriver* Driver, ENetworkFailure::Type FailureType, const FString& ErrorStr)
 {
 	UWorld* World = Driver->GetWorld();
@@ -369,10 +390,9 @@ bool UPackageMapClient::SerializeNewActor(FArchive& Ar, class UActorChannel *Cha
 	{
 		UObject* Archetype = nullptr;
 		UObject* ActorLevel = nullptr;
-		FVector_NetQuantize10 Location;
-		FVector_NetQuantize10 LocalLocation;
-		FVector_NetQuantize10 Scale;
-		FVector_NetQuantize10 Velocity;
+		FVector Location;
+		FVector Scale;
+		FVector Velocity;
 		FRotator Rotation;
 		bool SerSuccess;
 
@@ -398,12 +418,11 @@ bool UPackageMapClient::SerializeNewActor(FArchive& Ar, class UActorChannel *Cha
 
 			if (RootComponent)
 			{
-				LocalLocation = Actor->GetActorLocation();
 				Location = FRepMovement::RebaseOntoZeroOrigin(Actor->GetActorLocation(), Actor);
 			} 
 			else
 			{
-				Location = LocalLocation = FVector::ZeroVector;
+				Location = FVector::ZeroVector;
 			}
 			Rotation = RootComponent ? Actor->GetActorRotation() : FRotator::ZeroRotator;
 			Scale = RootComponent ? Actor->GetActorScale() : FVector::OneVector;
@@ -447,29 +466,63 @@ bool UPackageMapClient::SerializeNewActor(FArchive& Ar, class UActorChannel *Cha
 		bool bSerializeRotation = false;
 		bool bSerializeScale = false;
 		bool bSerializeVelocity = false;
-		const float epsilon = 0.001f;
+
 		{			
 			// Server is serializing an object to be sent to a client
 			if (Ar.IsSaving())
 			{
-				const FVector DefaultScale(1.f, 1.f, 1.f);
+				// We use 0.01f for comparing when using quantization, because we will only send a single point of precision anyway.
+				// We could probably get away with 0.1f, but that may introduce edge cases for rounding.
+				static constexpr float Epsilon_Quantized = 0.01f;
+				
+				// We use KINDA_SMALL_NUMBER for comparing when not using quantization, because that's the default for FVector::Equals.
+				static constexpr float Epsilon = KINDA_SMALL_NUMBER;
 
-				// If the Location isn't the default Location
-				bSerializeLocation = !Location.Equals(FVector::ZeroVector, epsilon);
-				bSerializeRotation = !Rotation.Equals(FRotator::ZeroRotator, epsilon);
-				bSerializeScale = !Scale.Equals(DefaultScale, epsilon);
-				bSerializeVelocity = !Velocity.Equals(FVector::ZeroVector, epsilon);
-			}
+				bSerializeLocation = !Location.Equals(FVector::ZeroVector, GbQuantizeActorLocationOnSpawn ? Epsilon_Quantized : Epsilon);
+				bSerializeVelocity = !Velocity.Equals(FVector::ZeroVector, GbQuantizeActorVelocityOnSpawn ? Epsilon_Quantized : Epsilon);
+				bSerializeScale = !Scale.Equals(FVector::OneVector, GbQuantizeActorScaleOnSpawn ? Epsilon_Quantized : Epsilon);
 
-			Ar.SerializeBits(&bSerializeLocation, 1);
-			if (bSerializeLocation)
-			{
-				Location.NetSerialize(Ar, this, SerSuccess);
+				// We use 0.001f for Rotation comparison to keep consistency with old behavior.
+				bSerializeRotation = !Rotation.IsNearlyZero(0.001f);
+				
 			}
-			else
+
+			auto ConditionallySerializeQuantizedVector = [this, &Ar, &SerSuccess](
+				FVector& InOutValue,
+				const FVector& DefaultValue,
+				bool bShouldQuantize,
+				bool& bWasSerialized)
 			{
-				Location = FVector::ZeroVector;
-			}
+				Ar.SerializeBits(&bWasSerialized, 1);
+				if (bWasSerialized)
+				{
+					if (Ar.EngineNetVer() < HISTORY_OPTIONALLY_QUANTIZE_SPAWN_INFO)
+					{
+						bShouldQuantize = true;
+					}
+					else
+					{
+						Ar.SerializeBits(&bShouldQuantize, 1);
+					}
+
+					if (bShouldQuantize)
+					{
+						FVector_NetQuantize10 Temp = InOutValue;
+						Temp.NetSerialize(Ar, this, SerSuccess);
+						InOutValue = Temp;
+					}
+					else
+					{
+						Ar << InOutValue;
+					}
+				}
+				else
+				{
+					InOutValue = DefaultValue;
+				}
+			};
+
+			ConditionallySerializeQuantizedVector(Location, FVector::ZeroVector, GbQuantizeActorLocationOnSpawn, bSerializeLocation);
 
 			Ar.SerializeBits(&bSerializeRotation, 1);
 			if (bSerializeRotation)
@@ -481,25 +534,8 @@ bool UPackageMapClient::SerializeNewActor(FArchive& Ar, class UActorChannel *Cha
 				Rotation = FRotator::ZeroRotator;
 			}
 
-			Ar.SerializeBits(&bSerializeScale, 1);
-			if (bSerializeScale)
-			{
-				Scale.NetSerialize(Ar, this, SerSuccess);
-			}
-			else
-			{
-				Scale = FVector(1, 1, 1);
-			}
-
-			Ar.SerializeBits(&bSerializeVelocity, 1);
-			if (bSerializeVelocity)
-			{
-				Velocity.NetSerialize(Ar, this, SerSuccess);
-			}
-			else
-			{
-				Velocity = FVector::ZeroVector;
-			}
+			ConditionallySerializeQuantizedVector(Scale, FVector::OneVector, GbQuantizeActorScaleOnSpawn, bSerializeScale);
+			ConditionallySerializeQuantizedVector(Velocity, FVector::ZeroVector, GbQuantizeActorVelocityOnSpawn, bSerializeVelocity);
 		}
 
 		if ( Ar.IsLoading() )
