Remove RVT StreamingTexture from cook when building platfroms that don't support VT

diff --git a/Engine/Source/Runtime/Engine/Classes/VT/RuntimeVirtualTexture.h b/Engine/Source/Runtime/Engine/Classes/VT/RuntimeVirtualTexture.h
--- a/Engine/Source/Runtime/Engine/Classes/VT/RuntimeVirtualTexture.h
+++ b/Engine/Source/Runtime/Engine/Classes/VT/RuntimeVirtualTexture.h
@@ -117,6 +117,7 @@ protected:
 
 	//~ Begin UObject Interface.
 	virtual void GetAssetRegistryTags(TArray<FAssetRegistryTag>& OutTags) const override;
+	virtual void Serialize(FArchive& Ar) override;
 #if WITH_EDITOR
 	virtual void PostEditChangeProperty(FPropertyChangedEvent& PropertyChangedEvent) override;
 #endif
diff --git a/Engine/Source/Runtime/Engine/Private/Texture.cpp b/Engine/Source/Runtime/Engine/Private/Texture.cpp
--- a/Engine/Source/Runtime/Engine/Private/Texture.cpp
+++ b/Engine/Source/Runtime/Engine/Private/Texture.cpp
@@ -334,7 +334,7 @@ void UTexture::Serialize(FArchive& Ar)
 		bUseLegacyGamma = true;
 	}
 
-	if (Ar.IsCooking())
+	if (Ar.IsCooking() && Ar.IsSaving())
 	{
 		if (VirtualTextureStreaming && UseVirtualTexturing(GMaxRHIFeatureLevel, Ar.CookingTarget()) == false)
 		{
diff --git a/Engine/Source/Runtime/Engine/Private/VT/RuntimeVirtualTexture.cpp b/Engine/Source/Runtime/Engine/Private/VT/RuntimeVirtualTexture.cpp
--- a/Engine/Source/Runtime/Engine/Private/VT/RuntimeVirtualTexture.cpp
+++ b/Engine/Source/Runtime/Engine/Private/VT/RuntimeVirtualTexture.cpp
@@ -374,6 +374,18 @@ void URuntimeVirtualTexture::GetAssetRegistryTags(TArray<FAssetRegistryTag>& Out
 	OutTags.Add(FAssetRegistryTag("TileBorderSize", FString::FromInt(GetTileBorderSize()), FAssetRegistryTag::TT_Numerical));
 }
 
+void URuntimeVirtualTexture::Serialize(FArchive& Ar)
+{
+	// Clear StreamingTexture during cook for platforms that don't support virtual texturing
+ 	URuntimeVirtualTextureStreamingProxy* StreamingTextureBackup = StreamingTexture;
+ 	if (Ar.IsCooking() && Ar.IsSaving() && !UseVirtualTexturing(GMaxRHIFeatureLevel, Ar.CookingTarget()))
+ 	{
+ 		StreamingTexture = nullptr;
+ 	}
+	Super::Serialize(Ar);
+	StreamingTexture = StreamingTextureBackup;
+}
+
 #if WITH_EDITOR
 
 void URuntimeVirtualTexture::PostEditChangeProperty(FPropertyChangedEvent& PropertyChangedEvent)
