Add new flags to FRepLayout and deprecate ERepLayoutState.

diff --git a/Engine/Source/Runtime/Engine/Private/DataReplication.cpp b/Engine/Source/Runtime/Engine/Private/DataReplication.cpp
--- a/Engine/Source/Runtime/Engine/Private/DataReplication.cpp
+++ b/Engine/Source/Runtime/Engine/Private/DataReplication.cpp
@@ -826,7 +826,7 @@ void FObjectReplicator::ReceivedNak( int32 NakPacketId )
 	{
 		UE_LOG(LogNet, Verbose, TEXT("FObjectReplicator::ReceivedNak: ObjectClass == nullptr"));
 	}
-	else if (ERepLayoutState::Normal == RepLayout->GetRepLayoutState())
+	else if (!RepLayout->IsEmpty())
 	{
 		if (FSendingRepState* SendingRepState = RepState.IsValid() ? RepState->GetSendingRepState() : nullptr)
 		{
@@ -1567,10 +1567,9 @@ bool FObjectReplicator::ReplicateProperties( FOutBunch & Bunch, FReplicationFlag
 		check(RepLayout.IsValid());
 		check(RepState.IsValid());
 		check(RepState->GetSendingRepState());
-		check(RepLayout->GetRepLayoutState() != ERepLayoutState::Uninitialized);
 		check(ChangelistMgr.IsValid());
 		check(ChangelistMgr->GetRepChangelistState() != nullptr);
-		check((ChangelistMgr->GetRepChangelistState()->StaticBuffer.Num() == 0) == (RepLayout->GetRepLayoutState() == ERepLayoutState::Empty));
+		check((ChangelistMgr->GetRepChangelistState()->StaticBuffer.Num() == 0) == RepLayout->IsEmpty());
 	}
 
 	UNetConnection* OwningChannelConnection = OwningChannel->Connection;
diff --git a/Engine/Source/Runtime/Engine/Private/RepLayout.cpp b/Engine/Source/Runtime/Engine/Private/RepLayout.cpp
--- a/Engine/Source/Runtime/Engine/Private/RepLayout.cpp
+++ b/Engine/Source/Runtime/Engine/Private/RepLayout.cpp
@@ -860,7 +860,7 @@ FReceivingRepState::FReceivingRepState(FRepStateStaticBuffer&& InStaticBuffer) :
 }
 
 FRepLayout::FRepLayout() :
-	LayoutState(ERepLayoutState::Uninitialized),
+	Flags(ERepLayoutFlags::None),
 	RoleIndex(INDEX_NONE),
 	RemoteRoleIndex(-1),
 	Owner(nullptr)
@@ -868,8 +868,6 @@ FRepLayout::FRepLayout() :
 
 FRepLayout::~FRepLayout()
 {
-	// This should never happen.
-	check(ERepLayoutState::Uninitialized != LayoutState);
 }
 
 void FRepLayout::UpdateChangelistMgr(
@@ -965,9 +963,6 @@ static void CompareRoleProperties(
 {
 	if (RepState && SharedParams.RoleIndex != INDEX_NONE && SharedParams.RemoteRoleIndex != INDEX_NONE)
 	{
-		// Verify that the order hasn't changed, otherwise ReceiveProperties will fail
-		check(SharedParams.RemoteRoleIndex < SharedParams.RoleIndex);
-
 		const FRepParentCmd& RemoteRoleParent = SharedParams.Parents[SharedParams.RemoteRoleIndex];
 		const FRepLayoutCmd& RemoteRoleCmd = SharedParams.Cmds[RemoteRoleParent.CmdStart];
 		const uint16 RemoteRoleHandle = RemoteRoleCmd.RelativeHandle;
@@ -1155,7 +1150,7 @@ bool FRepLayout::CompareProperties(
 {
 	CONDITIONAL_SCOPE_CYCLE_COUNTER(STAT_NetReplicateDynamicPropCompareTime, CVarNetEnableDetailedScopeCounters.GetValueOnAnyThread() > 0);
 
-	if (LayoutState == ERepLayoutState::Empty)
+	if (IsEmpty())
 	{
 		return false;
 	}
@@ -1260,7 +1255,7 @@ bool FRepLayout::ReplicateProperties(
 	check(ObjectClass == Owner);
 
 	// If we are an empty RepLayout, there's nothing to do.
-	if (LayoutState == ERepLayoutState::Empty)
+	if (IsEmpty())
 	{
 		return false;
 	}
@@ -2053,7 +2048,7 @@ void FRepLayout::SendProperties_r(
 
 		const FRepSerializedPropertyInfo* SharedPropInfo = nullptr;
 
-		if (bDoSharedSerialization && EnumHasAnyFlags(Cmd.Flags, ERepLayoutFlags::IsSharedSerialization))
+		if (bDoSharedSerialization && EnumHasAnyFlags(Cmd.Flags, ERepLayoutCmdFlags::IsSharedSerialization))
 		{
 			FGuid PropertyGuid(HandleIterator.CmdIndex, HandleIterator.ArrayIndex, ArrayDepth, (int32)((PTRINT)Data.Data & 0xFFFFFFFF));
 
@@ -2149,7 +2144,7 @@ void FRepLayout::SendProperties(
 {
 	SCOPE_CYCLE_COUNTER(STAT_NetReplicateDynamicPropSendTime);
 
-	if (LayoutState == ERepLayoutState::Empty)
+	if (IsEmpty())
 	{
 		return;
 	}
@@ -2900,11 +2895,11 @@ bool FRepLayout::ReceiveProperties(
 	FNetBitReader& InBunch,
 	bool& bOutHasUnmapped,
 	bool& bOutGuidsChanged,
-	const EReceivePropertiesFlags Flags) const
+	const EReceivePropertiesFlags ReceiveFlags) const
 {
 	check(InObjectClass == Owner);
 
-	const bool bEnableRepNotifies = EnumHasAnyFlags(Flags, EReceivePropertiesFlags::RepNotifies);
+	const bool bEnableRepNotifies = EnumHasAnyFlags(ReceiveFlags, EReceivePropertiesFlags::RepNotifies);
 
 	if (OwningChannel->Connection->InternalAck)
 	{
@@ -2924,12 +2919,13 @@ bool FRepLayout::ReceiveProperties(
 	// If we've gotten this far, it means that the server must have sent us something.
 	// That should only happen if there's actually commands to process.
 	// If this is hit, it may mean the Client and Server have different properties!
-	check(LayoutState != ERepLayoutState::Empty);
+	check(!IsEmpty());
 
 	FReceivePropertiesSharedParams Params{
 		bDoChecksum,
-		EnumHasAnyFlags(Flags, EReceivePropertiesFlags::SkipRoleSwap),
-		EnumHasAnyFlags(Flags, EReceivePropertiesFlags::RepNotifies) ? &RepState->RepNotifies : nullptr,
+		// We can skip swapping roles if we're not an Actor layout, or if we've been explicitly told we can skip.
+		EnumHasAnyFlags(ReceiveFlags, EReceivePropertiesFlags::SkipRoleSwap) || !EnumHasAnyFlags(Flags, ERepLayoutFlags::IsActor),
+		EnumHasAnyFlags(ReceiveFlags, EReceivePropertiesFlags::RepNotifies) ? &RepState->RepNotifies : nullptr,
 		InBunch,
 		bOutHasUnmapped,
 		bOutGuidsChanged,
@@ -3245,7 +3241,21 @@ bool FRepLayout::ReceiveProperties_BackwardsCompatible_r(
 		{
 			const int32 ElementOffset = (Data.Data - OldData.Data);
 
-			if (ReceivePropertyHelper(TempReader, GuidReferencesMap, ElementOffset, ShadowData, Data, ShadowData ? &RepState->RepNotifies : nullptr, Parents, Cmds, CmdIndex, false, bOutGuidsChanged, false))
+			if (ReceivePropertyHelper(
+				TempReader,
+				GuidReferencesMap,
+				ElementOffset,
+				ShadowData,
+				Data,
+				ShadowData ? &RepState->RepNotifies : nullptr,
+				Parents,
+				Cmds,
+				CmdIndex,
+				false,
+				bOutGuidsChanged,
+
+				// We can skip role swapping if we're not an actor.
+				!EnumHasAnyFlags(Flags, ERepLayoutFlags::IsActor)))
 			{
 				bOutHasUnmapped = true;
 			}
@@ -3300,7 +3310,7 @@ void FRepLayout::GatherGuidReferences(
 	TSet<FNetworkGUID>& OutReferencedGuids,
 	int32& OutTrackedGuidMemoryBytes) const
 {
-	if (LayoutState == ERepLayoutState::Normal)
+	if (!IsEmpty())
 	{
 		GatherGuidReferences_r(&RepState->GuidReferencesMap, OutReferencedGuids, OutTrackedGuidMemoryBytes);
 
@@ -3360,11 +3370,14 @@ bool FRepLayout::MoveMappedObjectToUnmapped_r(FGuidReferencesMap* GuidReferences
 	return bFoundGUID;
 }
 
-bool FRepLayout::MoveMappedObjectToUnmapped(FReceivingRepState* RESTRICT RepState, FNetDeltaSerializeInfo& Params, const FNetworkGUID& GUID) const
+bool FRepLayout::MoveMappedObjectToUnmapped(
+	FReceivingRepState* RESTRICT RepState,
+	FNetDeltaSerializeInfo& Params,
+	const FNetworkGUID& GUID) const
 {
 	bool bFound = false;
 
-	if (ERepLayoutState::Normal == LayoutState)
+	if (!IsEmpty())
 	{
 		bFound = MoveMappedObjectToUnmapped_r(&RepState->GuidReferencesMap, GUID);
 
@@ -3562,7 +3575,7 @@ void FRepLayout::UpdateUnmappedObjects(
 	bOutHasMoreUnmapped = false;
 	bCalledPreNetReceive = false;
 
-	if (LayoutState == ERepLayoutState::Normal)
+	if (!IsEmpty())
 	{
 		UpdateUnmappedObjects_r(
 			RepState,
@@ -3741,14 +3754,12 @@ void FRepLayout::CallRepNotifies(FReceivingRepState* RepState, UObject* Object)
 		return;
 	}
 	
-	if (LayoutState == ERepLayoutState::Empty)
+	if (IsEmpty())
 	{
 		UE_LOG(LogRep, Error, TEXT("FRepLayout::CallRepNotifies: Empty layout with RepNotifies: %s"), *GetPathNameSafe(Owner));
 		return;
 	}
 
-	check(LayoutState == ERepLayoutState::Normal);
-
 	FRepShadowDataBuffer ShadowData(RepState->StaticBuffer.GetData());
 	FRepObjectDataBuffer ObjectData(Object);
 
@@ -3949,7 +3960,7 @@ void FRepLayout::PruneChangeList(
 
 	PrunedChanged.Empty(1);
 
-	if (ERepLayoutState::Normal == LayoutState)
+	if (!IsEmpty())
 	{
 		FChangelistIterator ChangelistIterator(Changed, 0);
 		FRepHandleIterator HandleIterator(Owner, ChangelistIterator, Cmds, BaseHandleToCmdIndex, 0, 1, 0, Cmds.Num() - 1);
@@ -3968,7 +3979,7 @@ void FRepLayout::MergeChangeList(
 	check(Dirty1.Num() > 0);
 	MergedDirty.Empty(1);
 
-	if (ERepLayoutState::Normal == LayoutState)
+	if (!IsEmpty())
 	{
 		if (Dirty2.Num() == 0)
 		{
@@ -4439,7 +4450,7 @@ static uint32 AddPropertyCmd(
 	{
 		UStructProperty* StructProp = Cast<UStructProperty>(UnderlyingProperty);
 		UScriptStruct* Struct = StructProp->Struct;
-		Cmd.Flags |= ERepLayoutFlags::IsStruct;
+		Cmd.Flags |= ERepLayoutCmdFlags::IsStruct;
 
 		if (Struct->GetFName() == NAME_Vector)
 		{
@@ -4529,7 +4540,7 @@ static uint32 AddPropertyCmd(
 	// Also skip writing the RemoteRole since it can be modified per connection in FObjectReplicator
 	if (Cmd.Property->SupportsNetSharedSerialization() && (Cmd.Property->GetFName() != NAME_RemoteRole))
 	{
-		Cmd.Flags |= ERepLayoutFlags::IsSharedSerialization;
+		Cmd.Flags |= ERepLayoutCmdFlags::IsSharedSerialization;
 	}
 
 	return Cmd.CompatibleChecksum;
@@ -4759,7 +4770,7 @@ static void BuildShadowOffsets_r(TArray<FRepLayoutCmd>::TIterator& CmdIt, int32&
 	// Note, the only time we should see a StructProperty is if we have a NetSerialize struct.
 	// Custom Delta Serialize structs won't have an associated RepLayout command,
 	// and normal structs will flatten their properties.
-	if (CmdIt->Type == ERepLayoutCmdType::DynamicArray || EnumHasAnyFlags(CmdIt->Flags, ERepLayoutFlags::IsStruct))
+	if (CmdIt->Type == ERepLayoutCmdType::DynamicArray || EnumHasAnyFlags(CmdIt->Flags, ERepLayoutCmdFlags::IsStruct))
 	{
 		if (!bAlreadyAligned)
 		{
@@ -4812,16 +4823,15 @@ static void BuildShadowOffsets_r(TArray<FRepLayoutCmd>::TIterator& CmdIt, int32&
 }
 
 template<ERepBuildType ShadowType>
-static void BuildShadowOffsets(UStruct* Owner, TArray<FRepParentCmd>& Parents, TArray<FRepLayoutCmd>& Cmds, int32& ShadowOffset, ERepLayoutState& LayoutState)
+static void BuildShadowOffsets(UStruct* Owner, TArray<FRepParentCmd>& Parents, TArray<FRepLayoutCmd>& Cmds, int32& ShadowOffset)
 {
 	SCOPE_CYCLE_COUNTER(STAT_RepLayout_BuildShadowOffsets);
 
 	if (ShadowType == ERepBuildType::Class && !!GUsePackedShadowBuffers)
 	{
 		ShadowOffset = 0;
-		LayoutState = Parents.Num() > 0 ? ERepLayoutState::Normal : ERepLayoutState::Empty;
 
-		if (ERepLayoutState::Normal == LayoutState)
+		if (0 != Parents.Num())
 		{
 			// Before filling out any ShadowOffset information, we'll sort the Parent Commands by alignment.
 			// This has 2 main benefits:
@@ -4905,7 +4915,6 @@ static void BuildShadowOffsets(UStruct* Owner, TArray<FRepParentCmd>& Parents, T
 	else
 	{
 		ShadowOffset = Owner->GetPropertiesSize();
-		LayoutState = ERepLayoutState::Normal;
 
 		for (auto ParentIt = Parents.CreateIterator(); ParentIt; ++ParentIt)
 		{
@@ -4919,14 +4928,20 @@ static void BuildShadowOffsets(UStruct* Owner, TArray<FRepParentCmd>& Parents, T
 	}
 }
 
-TSharedPtr<FRepLayout> FRepLayout::CreateFromClass(UClass* InClass, const UNetConnection* ServerConnection, const ECreateRepLayoutFlags Flags)
+TSharedPtr<FRepLayout> FRepLayout::CreateFromClass(
+	UClass* InClass,
+	const UNetConnection* ServerConnection,
+	const ECreateRepLayoutFlags CreateFlags)
 {
 	TSharedPtr<FRepLayout> RepLayout = MakeShareable<FRepLayout>(new FRepLayout());
-	RepLayout->InitFromClass(InClass, ServerConnection, Flags);
+	RepLayout->InitFromClass(InClass, ServerConnection, CreateFlags);
 	return RepLayout;
 }
 
-void FRepLayout::InitFromClass(UClass* InObjectClass, const UNetConnection* ServerConnection, const ECreateRepLayoutFlags Flags)
+void FRepLayout::InitFromClass(
+	UClass* InObjectClass,
+	const UNetConnection* ServerConnection,
+	const ECreateRepLayoutFlags CreateFlags)
 {
 	SCOPE_CYCLE_COUNTER(STAT_RepLayout_InitFromObjectClass);
 	SCOPE_CYCLE_UOBJECT(ObjectClass, InObjectClass);
@@ -4935,6 +4950,11 @@ void FRepLayout::InitFromClass(UClass* InObjectClass, const UNetConnection* Serv
 	RoleIndex = INDEX_NONE;
 	RemoteRoleIndex = INDEX_NONE;
 
+	if (bIsObjectActor)
+	{
+		Flags |= ERepLayoutFlags::IsActor;
+	}
+
 	int32 RelativeHandle = 0;
 	int32 LastOffset = INDEX_NONE;
 
@@ -4980,14 +5000,14 @@ void FRepLayout::InitFromClass(UClass* InObjectClass, const UNetConnection* Serv
 		    // Find Role/RemoteRole property indexes so we can swap them on the client
 		    if (Property->GetFName() == NAME_Role)
 		    {
-			    check(RoleIndex == -1);
+			    check(RoleIndex == INDEX_NONE);
 			    check(Parents[ParentHandle].CmdEnd == Parents[ParentHandle].CmdStart + 1);
 			    RoleIndex = ParentHandle;
 		    }
 
 		    if (Property->GetFName() == NAME_RemoteRole)
 		    {
-			    check(RemoteRoleIndex == -1);
+			    check(RemoteRoleIndex == INDEX_NONE);
 			    check(Parents[ParentHandle].CmdEnd == Parents[ParentHandle].CmdStart + 1);
 			    RemoteRoleIndex = ParentHandle;
 		    }
@@ -4995,11 +5015,15 @@ void FRepLayout::InitFromClass(UClass* InObjectClass, const UNetConnection* Serv
 	}
 
 	// Make sure it either found both, or didn't find either
-	check((RoleIndex == -1) == (RemoteRoleIndex == -1));
+	check((RoleIndex == INDEX_NONE) == (RemoteRoleIndex == INDEX_NONE));
 
 	// This is so the receiving side can swap these as it receives them
 	if (RoleIndex != -1)
 	{
+		// Make sure that if we have Role and RemoteRole, that Role comes before RemoteRole.
+		// If this fails, it means that the order of Role and RemoteRole has been changed in AActor, and that
+		// will break assumptions RepLayout makes.
+		check(RemoteRoleIndex < RoleIndex);
 		Parents[RoleIndex].RoleSwapIndex = RemoteRoleIndex;
 		Parents[RemoteRoleIndex].RoleSwapIndex = RoleIndex;
 	}
@@ -5133,12 +5157,12 @@ void FRepLayout::InitFromClass(UClass* InObjectClass, const UNetConnection* Serv
 		}
 	}
 
-	if (!ServerConnection || EnumHasAnyFlags(Flags, ECreateRepLayoutFlags::MaySendProperties))
+	if (!ServerConnection || EnumHasAnyFlags(CreateFlags, ECreateRepLayoutFlags::MaySendProperties))
 	{
 		BuildHandleToCmdIndexTable_r(0, Cmds.Num() - 1, BaseHandleToCmdIndex);
 	}
 
-	BuildShadowOffsets<ERepBuildType::Class>(InObjectClass, Parents, Cmds, ShadowDataBufferSize, LayoutState);
+	BuildShadowOffsets<ERepBuildType::Class>(InObjectClass, Parents, Cmds, ShadowDataBufferSize);
 
 	Owner = InObjectClass;
 }
@@ -5150,7 +5174,10 @@ TSharedPtr<FRepLayout> FRepLayout::CreateFromFunction(UFunction* InFunction, con
 	return RepLayout;
 }
 
-void FRepLayout::InitFromFunction(UFunction* InFunction, const UNetConnection* ServerConnection, const ECreateRepLayoutFlags Flags)
+void FRepLayout::InitFromFunction(
+	UFunction* InFunction,
+	const UNetConnection* ServerConnection,
+	const ECreateRepLayoutFlags CreateFlags)
 {
 	int32 RelativeHandle = 0;
 
@@ -5170,26 +5197,30 @@ void FRepLayout::InitFromFunction(UFunction* InFunction, const UNetConnection* S
 
 	AddReturnCmd(Cmds);
 
-	if (!ServerConnection || EnumHasAnyFlags(Flags, ECreateRepLayoutFlags::MaySendProperties))
+	if (!ServerConnection || EnumHasAnyFlags(CreateFlags, ECreateRepLayoutFlags::MaySendProperties))
 	{
 		BuildHandleToCmdIndexTable_r(0, Cmds.Num() - 1, BaseHandleToCmdIndex);
 	}
 
-	BuildShadowOffsets<ERepBuildType::Function>(InFunction, Parents, Cmds, ShadowDataBufferSize, LayoutState);
+	BuildShadowOffsets<ERepBuildType::Function>(InFunction, Parents, Cmds, ShadowDataBufferSize);
 
 	Owner = InFunction;
-
-	LayoutState = Parents.Num() == 0 ? ERepLayoutState::Empty : ERepLayoutState::Normal;
 }
 
-TSharedPtr<FRepLayout> FRepLayout::CreateFromStruct(UStruct* InStruct, const UNetConnection* ServerConnection, const ECreateRepLayoutFlags Flags)
+TSharedPtr<FRepLayout> FRepLayout::CreateFromStruct(
+	UStruct* InStruct,
+	const UNetConnection* ServerConnection,
+	const ECreateRepLayoutFlags CreateFlags)
 {
 	TSharedPtr<FRepLayout> RepLayout = MakeShareable<FRepLayout>(new FRepLayout());
-	RepLayout->InitFromStruct(InStruct, ServerConnection, Flags);
+	RepLayout->InitFromStruct(InStruct, ServerConnection, CreateFlags);
 	return RepLayout;
 }
 
-void FRepLayout::InitFromStruct(UStruct* InStruct, const UNetConnection* ServerConnection, const ECreateRepLayoutFlags Flags)
+void FRepLayout::InitFromStruct(
+	UStruct* InStruct,
+	const UNetConnection* ServerConnection,
+	const ECreateRepLayoutFlags CreateFlags)
 {
 	int32 RelativeHandle = 0;
 
@@ -5214,12 +5245,12 @@ void FRepLayout::InitFromStruct(UStruct* InStruct, const UNetConnection* ServerC
 
 	AddReturnCmd(Cmds);
 
-	if (!ServerConnection || EnumHasAnyFlags(Flags, ECreateRepLayoutFlags::MaySendProperties))
+	if (!ServerConnection || EnumHasAnyFlags(CreateFlags, ECreateRepLayoutFlags::MaySendProperties))
 	{
 		BuildHandleToCmdIndexTable_r(0, Cmds.Num() - 1, BaseHandleToCmdIndex);
 	}
 
-	BuildShadowOffsets<ERepBuildType::Struct>(InStruct, Parents, Cmds, ShadowDataBufferSize, LayoutState);
+	BuildShadowOffsets<ERepBuildType::Struct>(InStruct, Parents, Cmds, ShadowDataBufferSize);
 
 	Owner = InStruct;
 }
@@ -5314,7 +5345,7 @@ void FRepLayout::SerializeProperties_r(
 
 		const FRepSerializedPropertyInfo* SharedPropInfo = nullptr;
 
-		if ((GNetSharedSerializedData != 0) && Ar.IsSaving() && ((Cmd.Flags & ERepLayoutFlags::IsSharedSerialization) != ERepLayoutFlags::None))
+		if ((GNetSharedSerializedData != 0) && Ar.IsSaving() && EnumHasAnyFlags(Cmd.Flags, ERepLayoutCmdFlags::IsSharedSerialization))
 		{
 			FGuid PropertyGuid(CmdIndex, ArrayIndex, ArrayDepth, (int32)((PTRINT)(const uint8*)(Data + Cmd) & 0xFFFFFFFF));
 
@@ -5496,7 +5527,7 @@ void FRepLayout::BuildSharedSerialization_r(
 			continue;
 		}
 
-		if (EnumHasAnyFlags(Cmd.Flags, ERepLayoutFlags::IsSharedSerialization))
+		if (EnumHasAnyFlags(Cmd.Flags, ERepLayoutCmdFlags::IsSharedSerialization))
 		{
 			SharedInfo.WriteSharedProperty(Cmd, FGuid(HandleIterator.CmdIndex, HandleIterator.ArrayIndex, ArrayDepth, (int32)((PTRINT)Data.Data & 0xFFFFFFFF)), HandleIterator.CmdIndex, HandleIterator.Handle, Data.Data, bWriteHandle, bDoChecksum);
 		}
@@ -5555,7 +5586,7 @@ void FRepLayout::BuildSharedSerializationForRPC_r(
 			continue;
 		}
 
-		if (!Parents[Cmd.ParentIndex].Property->HasAnyPropertyFlags(CPF_OutParm) && ((Cmd.Flags & ERepLayoutFlags::IsSharedSerialization) != ERepLayoutFlags::None))
+		if (!Parents[Cmd.ParentIndex].Property->HasAnyPropertyFlags(CPF_OutParm) && EnumHasAnyFlags(Cmd.Flags, ERepLayoutCmdFlags::IsSharedSerialization))
 		{
 			FGuid PropertyGuid(CmdIndex, ArrayIndex, ArrayDepth, (int32)((PTRINT)(const uint8*)(Data + Cmd) & 0xFFFFFFFF));
 
@@ -5614,7 +5645,7 @@ void FRepLayout::SendPropertiesForRPC(
 {
 	check(Function == Owner);
 
-	if (ERepLayoutState::Normal == LayoutState)
+	if (!IsEmpty())
 	{
 		if (Channel->Connection->InternalAck)
 		{
@@ -5676,7 +5707,7 @@ void FRepLayout::ReceivePropertiesForRPC(
 {
 	check(Function == Owner);
 
-	if (ERepLayoutState::Normal == LayoutState)
+	if (!IsEmpty())
 	{
 		for (int32 i = 0; i < Parents.Num(); i++)
 		{
@@ -5860,13 +5891,16 @@ FRepStateStaticBuffer FRepLayout::CreateShadowBuffer(const FConstRepObjectDataBu
 {
 	FRepStateStaticBuffer ShadowData(AsShared());
 
-	if (ShadowDataBufferSize == 0 && LayoutState != ERepLayoutState::Empty)
+	if (!IsEmpty())
 	{
-		UE_LOG(LogRep, Error, TEXT("FRepLayout::InitShadowData: Invalid RepLayout: %s"), *GetPathNameSafe(Owner));
-	}
-	else if (LayoutState == ERepLayoutState::Normal)
-	{
-		InitRepStateStaticBuffer(ShadowData, Source);
+		if (ShadowDataBufferSize == 0)
+		{
+			UE_LOG(LogRep, Error, TEXT("FRepLayout::InitShadowData: Invalid RepLayout: %s"), *GetPathNameSafe(Owner));
+		}
+		else
+		{
+			InitRepStateStaticBuffer(ShadowData, Source);
+		}
 	}
 
 	return ShadowData;
@@ -5896,8 +5930,12 @@ TSharedPtr<FReplicationChangelistMgr> FRepLayout::CreateReplicationChangelistMgr
 TUniquePtr<FRepState> FRepLayout::CreateRepState(
 	const FConstRepObjectDataBuffer Source,
 	TSharedPtr<FRepChangedPropertyTracker>& InRepChangedPropertyTracker,
-	ECreateRepStateFlags Flags) const
+	ECreateRepStateFlags CreateFlags) const
 {
+	// TODO: We could probably avoid allocating a RepState completely if we detect the RepLayout is empty.
+	//		In that case, there won't be anything interesting to do anyway.
+	//		This would require more sanity checks in code using RepStates though.
+
 	TUniquePtr<FRepState> RepState(new FRepState());
 
 	// If we have a changelist manager, that implies we're acting as a server.
@@ -5918,7 +5956,7 @@ TUniquePtr<FRepState> FRepLayout::CreateRepState(
 		RepState->SendingRepState->InactiveParents.Init(false, Parents.Num());
 	}
 	
-	if (!EnumHasAnyFlags(Flags, ECreateRepStateFlags::SkipCreateReceivingState))
+	if (!EnumHasAnyFlags(CreateFlags, ECreateRepStateFlags::SkipCreateReceivingState))
 	{
 		FRepStateStaticBuffer StaticBuffer(AsShared());
 
@@ -6101,7 +6139,7 @@ void FRepLayout::PreSendCustomDeltaProperties(
 						const FRepParentCmd& FastArrayCmd = Parents[RepIndex];
 
 						void const * const FastArraySerializer = ObjectData + FastArrayCmd;
-						const EFastArraySerializerDeltaFlags Flags = CustomDeltaProperty.GetFastArrayDeltaFlags(FastArraySerializer);
+						const EFastArraySerializerDeltaFlags DeltaFlags = CustomDeltaProperty.GetFastArrayDeltaFlags(FastArraySerializer);
 
 						// Note, we can't rely on EFastArraySerializerDeltaFlags::HasBeenSerialized here.
 						// It's possible we're calling PreSendCustomDeltaProperties **before** the first time the struct
@@ -6113,9 +6151,9 @@ void FRepLayout::PreSendCustomDeltaProperties(
 						//
 						// TODO: If any fast arrays failed this check, we could probably reset their state,
 						//			because we know we should never try sending them again
-						if (EnumHasAnyFlags(Flags, EFastArraySerializerDeltaFlags::IsUsingDeltaSerialization) ||
-							(!EnumHasAnyFlags(Flags, EFastArraySerializerDeltaFlags::HasBeenSerialized) &&
-								EnumHasAnyFlags(Flags, EFastArraySerializerDeltaFlags::HasDeltaBeenRequested)))
+						if (EnumHasAnyFlags(DeltaFlags, EFastArraySerializerDeltaFlags::IsUsingDeltaSerialization) ||
+							(!EnumHasAnyFlags(DeltaFlags, EFastArraySerializerDeltaFlags::HasBeenSerialized) &&
+								EnumHasAnyFlags(DeltaFlags, EFastArraySerializerDeltaFlags::HasDeltaBeenRequested)))
 						{
 							FDeltaArrayHistoryState& FastArrayHistoryState = CustomDeltaChangelistState.ArrayStates[FastArrayNumber];
 
@@ -6717,7 +6755,7 @@ bool FRepLayout::DeltaSerializeFastArrayProperty(FFastArrayDeltaSerializeParams&
 
 				FReceivePropertiesSharedParams SharedParams{
 					/*bDoChecksum=*/ false,
-					/*bSkipRoleSwap=*/ false,
+					/*bSkipRoleSwap=*/ !EnumHasAnyFlags(Flags, ERepLayoutFlags::IsActor),
 					/*RepNotifies=*/ nullptr,
 					Reader,
 					bOutHasUnmapped,
diff --git a/Engine/Source/Runtime/Engine/Public/Net/RepLayout.h b/Engine/Source/Runtime/Engine/Public/Net/RepLayout.h
--- a/Engine/Source/Runtime/Engine/Public/Net/RepLayout.h
+++ b/Engine/Source/Runtime/Engine/Public/Net/RepLayout.h
@@ -792,14 +792,14 @@ public:
 };
 
 /** Various flags that describe how a Property should be handled. */
-enum class ERepLayoutFlags : uint8
+enum class ERepLayoutCmdFlags : uint8
 {
 	None					= 0,		//! No flags.
 	IsSharedSerialization	= (1 << 0),	//! Indicates the property is eligible for shared serialization.
 	IsStruct				= (1 << 1)	//! This is a struct property.
 };
 
-ENUM_CLASS_FLAGS(ERepLayoutFlags)
+ENUM_CLASS_FLAGS(ERepLayoutCmdFlags)
 
 /**
  * Represents a single property, which could be either a Top Level Property, a Nested Struct Property,
@@ -836,7 +836,7 @@ public:
 	uint32 CompatibleChecksum;
 
 	ERepLayoutCmdType Type;
-	ERepLayoutFlags Flags;
+	ERepLayoutCmdFlags Flags;
 };
 	
 /** Converts a relative handle to the appropriate index into the Cmds array */
@@ -1004,7 +1004,14 @@ enum class ECreateRepLayoutFlags
 };
 ENUM_CLASS_FLAGS(ECreateRepLayoutFlags);
 
-enum class ERepLayoutState
+enum class ERepLayoutFlags : uint8
+{
+	None = 0,
+	IsActor = 1 << 1,	//! This RepLayout is for AActor or a subclass of AActor.
+};
+ENUM_CLASS_FLAGS(ERepLayoutFlags);
+
+enum class UE_DEPRECATED(4.24, "Use FRepLayout::IsEmpty() instead.") ERepLayoutState
 {
 	Uninitialized,	//! The RepLayout was never initiliazed, this should not be possible.
 	Empty,			//! The RepLayout was initialized, but doesn't have any RepCommands.
@@ -1413,9 +1420,21 @@ public:
 		return nullptr;
 	}
 
-	const ERepLayoutState GetRepLayoutState() const
+PRAGMA_DISABLE_DEPRECATION_WARNINGS
+
+	UE_DEPRECATED(4.24, "Use GetFlags instead.")
+	const ERepLayoutState GetRepLayoutState() const;
+
+PRAGMA_ENABLE_DEPRECATION_WARNINGS
+
+	const ERepLayoutFlags GetFlags() const
+	{
+		return Flags;
+	}
+
+	const bool IsEmpty() const
 	{
-		return LayoutState;
+		return 0 == Parents.Num();
 	}
 
 	void CountBytes(FArchive& Ar) const;
@@ -1737,7 +1756,7 @@ private:
 
 	const ELifetimeCondition GetLifetimeCustomDeltaPropertyCondition(const uint16 RepIndCustomDeltaPropertyIndexex) const;
 
-	ERepLayoutState LayoutState;
+	ERepLayoutFlags Flags;
 
 	/** Index of the Role property in the Parents list. May be INDEX_NONE if Owner doesn't have the property. */
 	int16 RoleIndex;
@@ -1772,3 +1791,13 @@ private:
 	/** Shared comparison to default state for multicast rpc */
 	TBitArray<> SharedInfoRPCParentsChanged;
 };
+
+
+PRAGMA_DISABLE_DEPRECATION_WARNINGS
+inline const ERepLayoutState FRepLayout::GetRepLayoutState() const
+{
+	// Because our constructor is hidden, we're guaranteed to be initialized.
+	// We're either Empty or Normal.
+	return IsEmpty() ? ERepLayoutState::Empty : ERepLayoutState::Normal;
+}
+PRAGMA_ENABLE_DEPRECATION_WARNINGS
