Fix CIS Issue

diff --git a/Engine/Source/Runtime/Engine/Private/RepLayout.cpp b/Engine/Source/Runtime/Engine/Private/RepLayout.cpp
--- a/Engine/Source/Runtime/Engine/Private/RepLayout.cpp
+++ b/Engine/Source/Runtime/Engine/Private/RepLayout.cpp
@@ -4215,20 +4215,20 @@ bool FRepLayout::DiffProperties(
 	TArray<UProperty*>* RepNotifies,
 	TRepDataBuffer<DestinationType> Destination,
 	TConstRepDataBuffer<SourceType> Source,
-	const EDiffPropertiesFlags Flags) const
+	const EDiffPropertiesFlags DiffFlags) const
 {
 
 	// Currently, only lifetime properties init from their defaults, so default to that,
 	// but also diff conditional properties if requested.
 	ERepParentFlags ParentPropertyFlags = ERepParentFlags::IsLifetime;
-	if (EnumHasAnyFlags(Flags, EDiffPropertiesFlags::IncludeConditionalProperties))
+	if (EnumHasAnyFlags(DiffFlags, EDiffPropertiesFlags::IncludeConditionalProperties))
 	{
 		ParentPropertyFlags |= ERepParentFlags::IsConditional;
 	}
 
 	FDiffPropertiesSharedParams Params{
 		ParentPropertyFlags,
-		Flags,
+		DiffFlags,
 		RepNotifies,
 		Parents,
 		Cmds
@@ -5167,10 +5167,10 @@ void FRepLayout::InitFromClass(
 	Owner = InObjectClass;
 }
 
-TSharedPtr<FRepLayout> FRepLayout::CreateFromFunction(UFunction* InFunction, const UNetConnection* ServerConnection, const ECreateRepLayoutFlags Flags)
+TSharedPtr<FRepLayout> FRepLayout::CreateFromFunction(UFunction* InFunction, const UNetConnection* ServerConnection, const ECreateRepLayoutFlags CreateFlags)
 {
 	TSharedPtr<FRepLayout> RepLayout = MakeShareable<FRepLayout>(new FRepLayout());
-	RepLayout->InitFromFunction(InFunction, ServerConnection, Flags);
+	RepLayout->InitFromFunction(InFunction, ServerConnection, CreateFlags);
 	return RepLayout;
 }
 
