Fix for crash on shutdown

diff --git a/Engine/Source/Runtime/Engine/Private/AudioDevice.cpp b/Engine/Source/Runtime/Engine/Private/AudioDevice.cpp
--- a/Engine/Source/Runtime/Engine/Private/AudioDevice.cpp
+++ b/Engine/Source/Runtime/Engine/Private/AudioDevice.cpp
@@ -594,6 +594,9 @@ int32 FAudioDevice::GetMaxSources() const
 
 void FAudioDevice::Teardown()
 {
+	// Make sure we process any pending game thread tasks before tearing down the audio device.
+	FTaskGraphInterface::Get().ProcessThreadUntilIdle(ENamedThreads::GameThread);
+
 	// Do a fadeout to prevent clicking on shutdown
 	FadeOut();
 
