Exposed OverrideAtmosphereLightDirection as ENGINE_API so game code can call it.

#rnx

diff --git a/Engine/Source/Runtime/Engine/Classes/Components/SkyAtmosphereComponent.h b/Engine/Source/Runtime/Engine/Classes/Components/SkyAtmosphereComponent.h
--- a/Engine/Source/Runtime/Engine/Classes/Components/SkyAtmosphereComponent.h
+++ b/Engine/Source/Runtime/Engine/Classes/Components/SkyAtmosphereComponent.h
@@ -124,7 +124,7 @@ class USkyAtmosphereComponent : public USceneComponent
 
 
 	UFUNCTION(BlueprintCallable, Category = "Rendering")
-	void OverrideAtmosphereLightDirection(int32 AtmosphereLightIndex, const FVector& LightDirection);
+	ENGINE_API void OverrideAtmosphereLightDirection(int32 AtmosphereLightIndex, const FVector& LightDirection);
 	
 	UFUNCTION(BlueprintCallable, Category = "Rendering")
 	void SetRayleighScatteringScale(float NewValue);
