Fix audiomemreport to use appropriate logic to determine whether a soundwave is fully decompressed. #jira none

diff --git a/Engine/Source/Runtime/Engine/Private/AudioDevice.cpp b/Engine/Source/Runtime/Engine/Private/AudioDevice.cpp
--- a/Engine/Source/Runtime/Engine/Private/AudioDevice.cpp
+++ b/Engine/Source/Runtime/Engine/Private/AudioDevice.cpp
@@ -1539,8 +1539,9 @@ bool FAudioDevice::HandleAudioMemoryInfo(const TCHAR* Cmd, FOutputDevice& Ar)
 			USoundWave* SoundWave = *It;
 
 			const FSoundGroup& SoundGroup = GetDefault<USoundGroups>()->GetSoundGroup(SoundWave->SoundGroup);
-			float Duration = SoundWave->GetDuration();
-			bool bDecompressed = SoundGroup.bAlwaysDecompressOnLoad || Duration < SoundGroup.DecompressedDuration;
+
+			float CompressionDurationThreshold = GetCompressionDurationThreshold(SoundGroup);
+			bool bDecompressed = ShouldUseRealtimeDecompression(false, SoundGroup, SoundWave, CompressionDurationThreshold);
 
 			FString SoundGroupName;
 			switch (SoundWave->SoundGroup)
@@ -1571,7 +1572,7 @@ bool FAudioDevice::HandleAudioMemoryInfo(const TCHAR* Cmd, FOutputDevice& Ar)
 			}
 
 			// Add the info to the SoundWaveObjects array
-			SoundWaveObjects.Add(FSoundWaveInfo(SoundWave, TrueResourceSize, SoundGroupName, Duration, bDecompressed));
+			SoundWaveObjects.Add(FSoundWaveInfo(SoundWave, TrueResourceSize, SoundGroupName, SoundWave->Duration, bDecompressed));
 
 			// Track total resource usage
 			TotalResourceSize += TrueResourceSize;
@@ -5478,19 +5479,8 @@ void FAudioDevice::Precache(USoundWave* SoundWave, bool bSynchronous, bool bTrac
 		}
 
 
-		// Check to see if the compression duration threshold is overridden via CVar:
-		float CompressedDurationThreshold = DecompressionThresholdCvar;
-		// If not, check to see if there is an override for the compression duration on this platform in the project settings:
-		if (CompressedDurationThreshold <= 0.0f)
-		{
-			CompressedDurationThreshold = FPlatformCompressionUtilities::GetCompressionDurationForCurrentPlatform();
-		}
+		float CompressedDurationThreshold = GetCompressionDurationThreshold(SoundGroup);
 
-		// If there is neither a CVar override nor a runtime setting override, use the decompression threshold from the sound group directly:
-		if (CompressedDurationThreshold < 0.0f)
-		{
-			CompressedDurationThreshold = SoundGroup.DecompressedDuration;
-		}
 
 		// handle audio decompression
 		if (FPlatformProperties::SupportsAudioStreaming() && SoundWave->IsStreaming())
@@ -5498,11 +5488,7 @@ void FAudioDevice::Precache(USoundWave* SoundWave, bool bSynchronous, bool bTrac
 			SoundWave->DecompressionType = DTYPE_Streaming;
 			SoundWave->bCanProcessAsync = false;
 		}
-		else if (!bForceFullDecompression &&
-				  SupportsRealtimeDecompression() &&
-				  ((bDisableAudioCaching || DisablePCMAudioCaching()) ||
-				  (!SoundGroup.bAlwaysDecompressOnLoad &&
-				  (ForceRealtimeDecompressionCvar || SoundWave->Duration > CompressedDurationThreshold || (RealtimeDecompressZeroDurationSoundsCvar && SoundWave->Duration <= 0.0f)))))
+		else if (ShouldUseRealtimeDecompression(bForceFullDecompression, SoundGroup, SoundWave, CompressedDurationThreshold))
 		{
 			// Store as compressed data and decompress in realtime
 			SoundWave->DecompressionType = DTYPE_RealTime;
@@ -5576,6 +5562,34 @@ void FAudioDevice::Precache(USoundWave* SoundWave, bool bSynchronous, bool bTrac
 	}
 }
 
+float FAudioDevice::GetCompressionDurationThreshold(const FSoundGroup &SoundGroup)
+{
+	// Check to see if the compression duration threshold is overridden via CVar:
+	float CompressedDurationThreshold = DecompressionThresholdCvar;
+	// If not, check to see if there is an override for the compression duration on this platform in the project settings:
+	if (CompressedDurationThreshold <= 0.0f)
+	{
+		CompressedDurationThreshold = FPlatformCompressionUtilities::GetCompressionDurationForCurrentPlatform();
+	}
+
+	// If there is neither a CVar override nor a runtime setting override, use the decompression threshold from the sound group directly:
+	if (CompressedDurationThreshold < 0.0f)
+	{
+		CompressedDurationThreshold = SoundGroup.DecompressedDuration;
+	}
+
+	return CompressedDurationThreshold;
+}
+
+bool FAudioDevice::ShouldUseRealtimeDecompression(bool bForceFullDecompression, const FSoundGroup &SoundGroup, USoundWave* SoundWave, float CompressedDurationThreshold) const
+{
+	return !bForceFullDecompression &&
+		SupportsRealtimeDecompression() &&
+		((bDisableAudioCaching || DisablePCMAudioCaching()) ||
+		(!SoundGroup.bAlwaysDecompressOnLoad &&
+			(ForceRealtimeDecompressionCvar || SoundWave->Duration > CompressedDurationThreshold || (RealtimeDecompressZeroDurationSoundsCvar && SoundWave->Duration <= 0.0f))));
+}
+
 void FAudioDevice::StopSourcesUsingBuffer(FSoundBuffer* SoundBuffer)
 {
 	SCOPED_NAMED_EVENT(FAudioDevice_StopSourcesUsingBuffer, FColor::Blue);
diff --git a/Engine/Source/Runtime/Engine/Public/AudioDevice.h b/Engine/Source/Runtime/Engine/Public/AudioDevice.h
--- a/Engine/Source/Runtime/Engine/Public/AudioDevice.h
+++ b/Engine/Source/Runtime/Engine/Public/AudioDevice.h
@@ -534,6 +534,13 @@ public:
 	 */
 	virtual void Precache(USoundWave* SoundWave, bool bSynchronous = false, bool bTrackMemory = true, bool bForceFullDecompression = false);
 
+	float GetCompressionDurationThreshold(const FSoundGroup &SoundGroup);
+
+	/**
+	 * Returns true if a sound wave should be decompressed.
+	 */
+	bool ShouldUseRealtimeDecompression(bool bForceFullDecompression, const FSoundGroup &SoundGroup, USoundWave* SoundWave, float CompressedDurationThreshold) const;
+
 	/**
 	 * Precaches all existing sounds. Called when audio setup is complete
 	 */
