Modify StreamableManager RequestAsync and Sync Load functions to use MoveTemp and allow avoidance of array and string copies when possible.

diff --git a/Engine/Plugins/Runtime/GameplayAbilities/Source/GameplayAbilities/Private/GameplayCueManager.cpp b/Engine/Plugins/Runtime/GameplayAbilities/Source/GameplayAbilities/Private/GameplayCueManager.cpp
--- a/Engine/Plugins/Runtime/GameplayAbilities/Source/GameplayAbilities/Private/GameplayCueManager.cpp
+++ b/Engine/Plugins/Runtime/GameplayAbilities/Source/GameplayAbilities/Private/GameplayCueManager.cpp
@@ -827,13 +827,14 @@ TSharedPtr<FStreamableHandle> UGameplayCueManager::InitObjectLibrary(FGameplayCu
 
 		if (AssetsToLoad.Num() > 0)
 		{
-			GameplayCueAssetHandle = StreamableManager.RequestAsyncLoad(AssetsToLoad, FStreamableDelegate::CreateStatic( ForwardLambda, AssetsToLoad, Lib.OnLoaded), Lib.AsyncPriority);
+			FStreamableDelegate Del = FStreamableDelegate::CreateStatic(ForwardLambda, AssetsToLoad, Lib.OnLoaded);
+			GameplayCueAssetHandle = StreamableManager.RequestAsyncLoad(MoveTemp(AssetsToLoad), MoveTemp(Del), Lib.AsyncPriority);
 			RetVal = GameplayCueAssetHandle;
 		}
 		else
 		{
 			// Still fire the delegate even if nothing was found to load
-			Lib.OnLoaded.ExecuteIfBound(AssetsToLoad);
+			Lib.OnLoaded.ExecuteIfBound(MoveTemp(AssetsToLoad));
 		}
 	}
 
diff --git a/Engine/Source/Runtime/Engine/Classes/Engine/StreamableManager.h b/Engine/Source/Runtime/Engine/Classes/Engine/StreamableManager.h
--- a/Engine/Source/Runtime/Engine/Classes/Engine/StreamableManager.h
+++ b/Engine/Source/Runtime/Engine/Classes/Engine/StreamableManager.h
@@ -238,12 +238,12 @@ struct ENGINE_API FStreamableManager : public FGCObject
 	 * @param bStartStalled			If true, the handle will start in a stalled state and will not attempt to actually async load until StartStalledHandle is called on it
 	 * @param DebugName				Name of this handle, will be reported in debug tools
 	 */
-	TSharedPtr<FStreamableHandle> RequestAsyncLoad(const TArray<FSoftObjectPath>& TargetsToStream, FStreamableDelegate DelegateToCall = FStreamableDelegate(), TAsyncLoadPriority Priority = DefaultAsyncLoadPriority, bool bManageActiveHandle = false, bool bStartStalled = false, const FString& DebugName = TEXT("RequestAsyncLoad ArrayDelegate"));
-	TSharedPtr<FStreamableHandle> RequestAsyncLoad(const FSoftObjectPath& TargetToStream, FStreamableDelegate DelegateToCall = FStreamableDelegate(), TAsyncLoadPriority Priority = DefaultAsyncLoadPriority, bool bManageActiveHandle = false, bool bStartStalled = false, const FString& DebugName = TEXT("RequestAsyncLoad SingleDelegate"));
+	TSharedPtr<FStreamableHandle> RequestAsyncLoad(TArray<FSoftObjectPath> TargetsToStream, FStreamableDelegate DelegateToCall = FStreamableDelegate(), TAsyncLoadPriority Priority = DefaultAsyncLoadPriority, bool bManageActiveHandle = false, bool bStartStalled = false, FString DebugName = TEXT("RequestAsyncLoad ArrayDelegate"));
+	TSharedPtr<FStreamableHandle> RequestAsyncLoad(const FSoftObjectPath& TargetToStream, FStreamableDelegate DelegateToCall = FStreamableDelegate(), TAsyncLoadPriority Priority = DefaultAsyncLoadPriority, bool bManageActiveHandle = false, bool bStartStalled = false, FString DebugName = TEXT("RequestAsyncLoad SingleDelegate"));
 
 	/** Lambda Wrappers. Be aware that Callback may go off multiple seconds in the future. */
-	TSharedPtr<FStreamableHandle> RequestAsyncLoad(const TArray<FSoftObjectPath>& TargetsToStream, TFunction<void()>&& Callback, TAsyncLoadPriority Priority = DefaultAsyncLoadPriority, bool bManageActiveHandle = false, bool bStartStalled = false, const FString& DebugName = TEXT("RequestAsyncLoad ArrayLambda"));
-	TSharedPtr<FStreamableHandle> RequestAsyncLoad(const FSoftObjectPath& TargetToStream, TFunction<void()>&& Callback, TAsyncLoadPriority Priority = DefaultAsyncLoadPriority, bool bManageActiveHandle = false, bool bStartStalled = false, const FString& DebugName = TEXT("RequestAsyncLoad SingleLambda"));
+	TSharedPtr<FStreamableHandle> RequestAsyncLoad(TArray<FSoftObjectPath> TargetsToStream, TFunction<void()>&& Callback, TAsyncLoadPriority Priority = DefaultAsyncLoadPriority, bool bManageActiveHandle = false, bool bStartStalled = false, FString DebugName = TEXT("RequestAsyncLoad ArrayLambda"));
+	TSharedPtr<FStreamableHandle> RequestAsyncLoad(const FSoftObjectPath& TargetToStream, TFunction<void()>&& Callback, TAsyncLoadPriority Priority = DefaultAsyncLoadPriority, bool bManageActiveHandle = false, bool bStartStalled = false, FString DebugName = TEXT("RequestAsyncLoad SingleLambda"));
 
 	/** 
 	 * Synchronously load a set of assets, and return a handle. This can be very slow and may stall the game thread for several seconds.
@@ -252,8 +252,8 @@ struct ENGINE_API FStreamableManager : public FGCObject
 	 * @param bManageActiveHandle	If true, the manager will keep the streamable handle active until explicitly released
 	 * @param DebugName				Name of this handle, will be reported in debug tools
 	 */
-	TSharedPtr<FStreamableHandle> RequestSyncLoad(const TArray<FSoftObjectPath>& TargetsToStream, bool bManageActiveHandle = false, const FString& DebugName = TEXT("RequestSyncLoad Array"));
-	TSharedPtr<FStreamableHandle> RequestSyncLoad(const FSoftObjectPath& TargetToStream, bool bManageActiveHandle = false, const FString& DebugName = TEXT("RequestSyncLoad Single"));
+	TSharedPtr<FStreamableHandle> RequestSyncLoad(TArray<FSoftObjectPath> TargetsToStream, bool bManageActiveHandle = false, FString DebugName = TEXT("RequestSyncLoad Array"));
+	TSharedPtr<FStreamableHandle> RequestSyncLoad(const FSoftObjectPath& TargetToStream, bool bManageActiveHandle = false, FString DebugName = TEXT("RequestSyncLoad Single"));
 
 	/** 
 	 * Synchronously load the referred asset and return the loaded object, or nullptr if it can't be found. This can be very slow and may stall the game thread for several seconds.
diff --git a/Engine/Source/Runtime/Engine/Private/ObjectLibrary.cpp b/Engine/Source/Runtime/Engine/Private/ObjectLibrary.cpp
--- a/Engine/Source/Runtime/Engine/Private/ObjectLibrary.cpp
+++ b/Engine/Source/Runtime/Engine/Private/ObjectLibrary.cpp
@@ -494,7 +494,7 @@ int32 UObjectLibrary::LoadAssetsFromAssetData()
 	if (bPreloadObjects)
 	{
 		TArray<FSoftObjectPath> AssetsToStream;
-
+		AssetsToStream.Reserve(AssetDataList.Num());
 		for (int32 AssetIdx = 0; AssetIdx < AssetDataList.Num(); AssetIdx++)
 		{
 			FAssetData& Data = AssetDataList[AssetIdx];
@@ -506,7 +506,7 @@ int32 UObjectLibrary::LoadAssetsFromAssetData()
 			FStreamableManager Streamable;
 			
 			// This will either use loadobject or async load + flush as appropriate
-			Streamable.RequestSyncLoad(AssetsToStream);
+			Streamable.RequestSyncLoad(MoveTemp(AssetsToStream));
 		}
 
 	}
diff --git a/Engine/Source/Runtime/Engine/Private/StreamableManager.cpp b/Engine/Source/Runtime/Engine/Private/StreamableManager.cpp
--- a/Engine/Source/Runtime/Engine/Private/StreamableManager.cpp
+++ b/Engine/Source/Runtime/Engine/Private/StreamableManager.cpp
@@ -946,7 +946,7 @@ FStreamable* FStreamableManager::StreamInternal(const FSoftObjectPath& InTargetN
 	return Existing;
 }
 
-TSharedPtr<FStreamableHandle> FStreamableManager::RequestAsyncLoad(const TArray<FSoftObjectPath>& TargetsToStream, FStreamableDelegate DelegateToCall, TAsyncLoadPriority Priority, bool bManageActiveHandle, bool bStartStalled, const FString& DebugName)
+TSharedPtr<FStreamableHandle> FStreamableManager::RequestAsyncLoad(TArray<FSoftObjectPath> TargetsToStream, FStreamableDelegate DelegateToCall, TAsyncLoadPriority Priority, bool bManageActiveHandle, bool bStartStalled, FString DebugName)
 {
 	LLM_SCOPE(ELLMTag::StreamingManager);
 
@@ -954,21 +954,22 @@ TSharedPtr<FStreamableHandle> FStreamableManager::RequestAsyncLoad(const TArray<
 	TSharedRef<FStreamableHandle> NewRequest = MakeShareable(new FStreamableHandle());
 	NewRequest->CompleteDelegate = DelegateToCall;
 	NewRequest->OwningManager = this;
-	NewRequest->RequestedAssets = TargetsToStream;
-	NewRequest->DebugName = DebugName;
+	NewRequest->RequestedAssets = MoveTemp(TargetsToStream);
+	NewRequest->DebugName = MoveTemp(DebugName);
 	NewRequest->Priority = Priority;
 
-	TRACE_LOADTIME_NEW_STREAMABLE_HANDLE(&NewRequest.Get(), *DebugName, false);
+	TRACE_LOADTIME_NEW_STREAMABLE_HANDLE(&NewRequest.Get(), *NewRequest->DebugName, false);
 
-	// Remove null requests
+	int32 NumValidRequests = NewRequest->RequestedAssets.Num();
+	
+	TSet<FSoftObjectPath> TargetSet;
+	TargetSet.Reserve(NumValidRequests);
 
-	for (int32 i = NewRequest->RequestedAssets.Num() - 1; i >= 0 ; i--)
+	for (const FSoftObjectPath& TargetName : NewRequest->RequestedAssets)
 	{
-		FSoftObjectPath& TargetName = NewRequest->RequestedAssets[i];
 		if (TargetName.IsNull())
 		{
-			// Remove null entries
-			NewRequest->RequestedAssets.RemoveAt(i);
+			--NumValidRequests;
 			continue;
 		}
 		else if (FPackageName::IsShortPackageName(TargetName.ToString()))
@@ -977,35 +978,40 @@ TSharedPtr<FStreamableHandle> FStreamableManager::RequestAsyncLoad(const TArray<
 			NewRequest->CancelHandle();
 			return nullptr;
 		}
+		TargetSet.Add(TargetName);
 	}
 
-	if (NewRequest->RequestedAssets.Num() == 0)
+	if (NumValidRequests == 0)
 	{
 		// Original array was empty or all null
 		UE_LOG(LogStreamableManager, Error, TEXT("RequestAsyncLoad called with empty or only null assets!"));
 		NewRequest->CancelHandle();
 		return nullptr;
 	} 
-	else if (NewRequest->RequestedAssets.Num() != TargetsToStream.Num())
+	else if (NewRequest->RequestedAssets.Num() != NumValidRequests)
 	{
 		FString RequestedSet;
 
-		for (const FSoftObjectPath& Asset : TargetsToStream)
+		for (auto It = NewRequest->RequestedAssets.CreateIterator(); It; ++It)
 		{
+			FSoftObjectPath& Asset = *It;
 			if (!RequestedSet.IsEmpty())
 			{
 				RequestedSet += TEXT(", ");
 			}
 			RequestedSet += Asset.ToString();
+
+			// Remove null entries
+			if (Asset.IsNull())
+			{
+				It.RemoveCurrent();
+			}
 		}
 
 		// Some valid, some null
 		UE_LOG(LogStreamableManager, Warning, TEXT("RequestAsyncLoad called with both valid and null assets, null assets removed from %s!"), *RequestedSet);
 	}
 
-	// Remove any duplicates
-	TSet<FSoftObjectPath> TargetSet(NewRequest->RequestedAssets);
-
 	if (TargetSet.Num() != NewRequest->RequestedAssets.Num())
 	{
 #if UE_BUILD_DEBUG
@@ -1044,22 +1050,22 @@ TSharedPtr<FStreamableHandle> FStreamableManager::RequestAsyncLoad(const TArray<
 	return NewRequest;
 }
 
-TSharedPtr<FStreamableHandle> FStreamableManager::RequestAsyncLoad(const FSoftObjectPath& TargetToStream, FStreamableDelegate DelegateToCall, TAsyncLoadPriority Priority, bool bManageActiveHandle, bool bStartStalled, const FString& DebugName)
+TSharedPtr<FStreamableHandle> FStreamableManager::RequestAsyncLoad(const FSoftObjectPath& TargetToStream, FStreamableDelegate DelegateToCall, TAsyncLoadPriority Priority, bool bManageActiveHandle, bool bStartStalled, FString DebugName)
 {
-	return RequestAsyncLoad(TArray<FSoftObjectPath>{TargetToStream}, DelegateToCall, Priority, bManageActiveHandle, bStartStalled, DebugName);
+	return RequestAsyncLoad(TArray<FSoftObjectPath>{TargetToStream}, MoveTemp(DelegateToCall), Priority, bManageActiveHandle, bStartStalled, MoveTemp(DebugName));
 }
 
-TSharedPtr<FStreamableHandle>  FStreamableManager::RequestAsyncLoad(const TArray<FSoftObjectPath>& TargetsToStream, TFunction<void()>&& Callback, TAsyncLoadPriority Priority, bool bManageActiveHandle, bool bStartStalled, const FString& DebugName)
+TSharedPtr<FStreamableHandle> FStreamableManager::RequestAsyncLoad(TArray<FSoftObjectPath> TargetsToStream, TFunction<void()>&& Callback, TAsyncLoadPriority Priority, bool bManageActiveHandle, bool bStartStalled, FString DebugName)
 {
-	return RequestAsyncLoad(TargetsToStream, FStreamableDelegate::CreateLambda( MoveTemp( Callback ) ), Priority, bManageActiveHandle, bStartStalled, DebugName);
+	return RequestAsyncLoad(MoveTemp(TargetsToStream), FStreamableDelegate::CreateLambda( MoveTemp( Callback ) ), Priority, bManageActiveHandle, bStartStalled, MoveTemp(DebugName));
 }
 
-TSharedPtr<FStreamableHandle>  FStreamableManager::RequestAsyncLoad(const FSoftObjectPath& TargetToStream, TFunction<void()>&& Callback, TAsyncLoadPriority Priority, bool bManageActiveHandle, bool bStartStalled, const FString& DebugName)
+TSharedPtr<FStreamableHandle> FStreamableManager::RequestAsyncLoad(const FSoftObjectPath& TargetToStream, TFunction<void()>&& Callback, TAsyncLoadPriority Priority, bool bManageActiveHandle, bool bStartStalled, FString DebugName)
 {
-	return RequestAsyncLoad(TargetToStream, FStreamableDelegate::CreateLambda( MoveTemp( Callback ) ), Priority, bManageActiveHandle, bStartStalled, DebugName);
+	return RequestAsyncLoad(TargetToStream, FStreamableDelegate::CreateLambda( MoveTemp( Callback ) ), Priority, bManageActiveHandle, bStartStalled, MoveTemp(DebugName));
 }
 
-TSharedPtr<FStreamableHandle> FStreamableManager::RequestSyncLoad(const TArray<FSoftObjectPath>& TargetsToStream, bool bManageActiveHandle, const FString& DebugName)
+TSharedPtr<FStreamableHandle> FStreamableManager::RequestSyncLoad(TArray<FSoftObjectPath> TargetsToStream, bool bManageActiveHandle, FString DebugName)
 {
 	// If in async loading thread or from callback always do sync as recursive tick is unsafe
 	// If in EDL always do sync as EDL internally avoids flushing
@@ -1067,7 +1073,7 @@ TSharedPtr<FStreamableHandle> FStreamableManager::RequestSyncLoad(const TArray<F
 	bForceSynchronousLoads = IsInAsyncLoadingThread() || IsEventDrivenLoaderEnabled() || !IsAsyncLoading();
 
 	// Do an async load and wait to complete. In some cases this will do a sync load due to safety issues
-	TSharedPtr<FStreamableHandle> Request = RequestAsyncLoad(TargetsToStream, FStreamableDelegate(), AsyncLoadHighPriority, bManageActiveHandle, false, DebugName);
+	TSharedPtr<FStreamableHandle> Request = RequestAsyncLoad(MoveTemp(TargetsToStream), FStreamableDelegate(), AsyncLoadHighPriority, bManageActiveHandle, false, MoveTemp(DebugName));
 
 	bForceSynchronousLoads = false;
 
@@ -1075,16 +1081,16 @@ TSharedPtr<FStreamableHandle> FStreamableManager::RequestSyncLoad(const TArray<F
 	{
 		EAsyncPackageState::Type Result = Request->WaitUntilComplete();
 
-		ensureMsgf(Result == EAsyncPackageState::Complete, TEXT("RequestSyncLoad of %s resulted in bad async load result %d"), *DebugName, Result);
-		ensureMsgf(Request->HasLoadCompleted(), TEXT("RequestSyncLoad of %s completed early, not actually completed!"), *DebugName);
+		ensureMsgf(Result == EAsyncPackageState::Complete, TEXT("RequestSyncLoad of %s resulted in bad async load result %d"), *Request->DebugName, Result);
+		ensureMsgf(Request->HasLoadCompleted(), TEXT("RequestSyncLoad of %s completed early, not actually completed!"), *Request->DebugName);
 	}
 
 	return Request;
 }
 
-TSharedPtr<FStreamableHandle> FStreamableManager::RequestSyncLoad(const FSoftObjectPath& TargetToStream, bool bManageActiveHandle, const FString& DebugName)
+TSharedPtr<FStreamableHandle> FStreamableManager::RequestSyncLoad(const FSoftObjectPath& TargetToStream, bool bManageActiveHandle, FString DebugName)
 {
-	return RequestSyncLoad(TArray<FSoftObjectPath>{TargetToStream}, bManageActiveHandle, DebugName);
+	return RequestSyncLoad(TArray<FSoftObjectPath>{TargetToStream}, bManageActiveHandle, MoveTemp(DebugName));
 }
 
 void FStreamableManager::StartHandleRequests(TSharedRef<FStreamableHandle> Handle)
