Skip processing Lifetime Properties that are marked COND_Never.

diff --git a/Engine/Source/Runtime/Engine/Private/RepLayout.cpp b/Engine/Source/Runtime/Engine/Private/RepLayout.cpp
--- a/Engine/Source/Runtime/Engine/Private/RepLayout.cpp
+++ b/Engine/Source/Runtime/Engine/Private/RepLayout.cpp
@@ -5049,6 +5049,13 @@ void FRepLayout::InitFromClass(
 			continue;
 		}
 
+		// Don't bother doing any setup work for COND_Never properties.
+		// These are never expected to replicate.
+		if (COND_Never == LifetimeProps[i].Condition)
+		{
+			continue;
+		}
+
 		// Store the condition on the parent in case we need it
 		Parents[ParentIndex].Condition = LifetimeProps[i].Condition;
 		Parents[ParentIndex].RepNotifyCondition = LifetimeProps[i].RepNotifyCondition;
@@ -5127,7 +5134,7 @@ void FRepLayout::InitFromClass(
 				}
 
 				if (!bAddedFastArray)
-		{
+				{
 					UE_LOG(LogRep, Warning, TEXT("FRepLayout::InitFromClass: Unable to find Fast Array Item array in Fast Array Serializer: %s"), *Parents[ParentIndex].CachedPropertyName.ToString());
 				}
 			}
@@ -5142,18 +5149,26 @@ void FRepLayout::InitFromClass(
 
 		Parents[ParentIndex].Flags |= ERepParentFlags::IsLifetime;
 
-		if (ParentIndex == RemoteRoleIndex)
+		if (LifetimeProps[i].Condition == COND_None)
 		{
-			// We handle remote role specially, since it can change between connections when downgraded
-			// So we force it on the conditional list
-			check(LifetimeProps[i].Condition == COND_None || LifetimeProps[i].Condition == COND_Never);
-			LifetimeProps[i].Condition = COND_Custom;
-			continue;
+			Parents[ParentIndex].Flags &= ~ERepParentFlags::IsConditional;
 		}
+	}
 
-		if (LifetimeProps[i].Condition == COND_None)
+	if (bIsObjectActor)
+	{
+		// We handle remote role specially, since it can change between connections when downgraded
+		// So we force it on the conditional list
+		FRepParentCmd& RemoteRoleParent = Parents[RemoteRoleIndex];
+		if (RemoteRoleParent.Condition != COND_Never)
 		{
-			Parents[ParentIndex].Flags &= ~ERepParentFlags::IsConditional;
+			if (COND_None != RemoteRoleParent.Condition)
+			{
+				UE_LOG(LogRep, Warning, TEXT("FRepLayout::InitFromClass: Forcing replication of RemoteRole. Owner=%s"), *InObjectClass->GetPathName());
+			}
+
+			Parents[RemoteRoleIndex].Flags |= ERepParentFlags::IsConditional;
+			Parents[RemoteRoleIndex].Condition = COND_None;
 		}
 	}
 
