Make sure to let non-engine modules use stream caching #jira none

diff --git a/Engine/Source/Runtime/Engine/Public/ContentStreaming.h b/Engine/Source/Runtime/Engine/Public/ContentStreaming.h
--- a/Engine/Source/Runtime/Engine/Public/ContentStreaming.h
+++ b/Engine/Source/Runtime/Engine/Public/ContentStreaming.h
@@ -88,7 +88,7 @@ struct FStreamingViewInfo
  * This structure allows audio chunk data to be accessed, and guarantees that the chunk in question will not be deleted
  * during it's lifecycle.
  */
-class FAudioChunkHandle
+class ENGINE_API FAudioChunkHandle
 {
 public:
 	FAudioChunkHandle();
