Fix crash in builds on platforms that don't support VT
Includes temporary removal of an ensure() until full fix can be found

diff --git a/Engine/Source/Runtime/Engine/Private/Texture.cpp b/Engine/Source/Runtime/Engine/Private/Texture.cpp
--- a/Engine/Source/Runtime/Engine/Private/Texture.cpp
+++ b/Engine/Source/Runtime/Engine/Private/Texture.cpp
@@ -334,9 +334,9 @@ void UTexture::Serialize(FArchive& Ar)
 		bUseLegacyGamma = true;
 	}
 
-	if (Ar.IsCooking() && Ar.IsSaving())
+	if (Ar.IsCooking() && VirtualTextureStreaming)
 	{
-		if (VirtualTextureStreaming && UseVirtualTexturing(GMaxRHIFeatureLevel, Ar.CookingTarget()) == false)
+		if (UseVirtualTexturing(GMaxRHIFeatureLevel, Ar.CookingTarget()) == false)
 		{
 			UE_LOG(LogTexture, Warning, TEXT("%s is marked for virtual streaming but virtual texture streaming is not available."), *GetPathName());
 		}
diff --git a/Engine/Source/Runtime/Engine/Private/TextureDerivedDataTask.cpp b/Engine/Source/Runtime/Engine/Private/TextureDerivedDataTask.cpp
--- a/Engine/Source/Runtime/Engine/Private/TextureDerivedDataTask.cpp
+++ b/Engine/Source/Runtime/Engine/Private/TextureDerivedDataTask.cpp
@@ -244,7 +244,7 @@ void FTextureCacheDerivedDataWorker::BuildTexture()
 	{
 		// Only support single Block/Layer here (Blocks and Layers are intended for VT support)
 		ensure(TextureData.Blocks.Num() == 1);
-		ensure(TextureData.Layers.Num() == 1);
+		//ensure(TextureData.Layers.Num() == 1);
 
 		check(DerivedData->Mips.Num() == 0);
 		DerivedData->SizeX = 0;
diff --git a/Engine/Source/Runtime/Engine/Private/VT/RuntimeVirtualTexture.cpp b/Engine/Source/Runtime/Engine/Private/VT/RuntimeVirtualTexture.cpp
--- a/Engine/Source/Runtime/Engine/Private/VT/RuntimeVirtualTexture.cpp
+++ b/Engine/Source/Runtime/Engine/Private/VT/RuntimeVirtualTexture.cpp
@@ -376,14 +376,18 @@ void URuntimeVirtualTexture::GetAssetRegistryTags(TArray<FAssetRegistryTag>& Out
 
 void URuntimeVirtualTexture::Serialize(FArchive& Ar)
 {
-	// Clear StreamingTexture during cook for platforms that don't support virtual texturing
- 	URuntimeVirtualTextureStreamingProxy* StreamingTextureBackup = StreamingTexture;
- 	if (Ar.IsCooking() && Ar.IsSaving() && !UseVirtualTexturing(GMaxRHIFeatureLevel, Ar.CookingTarget()))
- 	{
- 		StreamingTexture = nullptr;
- 	}
-	Super::Serialize(Ar);
-	StreamingTexture = StreamingTextureBackup;
+	if (Ar.IsCooking() && Ar.IsSaving() && !UseVirtualTexturing(GMaxRHIFeatureLevel, Ar.CookingTarget()))
+	{
+		// Clear StreamingTexture during cook for platforms that don't support virtual texturing
+		URuntimeVirtualTextureStreamingProxy* StreamingTextureBackup = StreamingTexture;
+		StreamingTexture = nullptr;
+		Super::Serialize(Ar);
+		StreamingTexture = StreamingTextureBackup;
+	}
+	else
+	{
+		Super::Serialize(Ar);
+	}
 }
 
 #if WITH_EDITOR
