Added a "LogBotGame" commandline parameter that when added to server commandline will vislog data on bots filtering out all the unnecessary crap.

diff --git a/Engine/Source/Runtime/Engine/Classes/AI/Navigation/NavigationTypes.h b/Engine/Source/Runtime/Engine/Classes/AI/Navigation/NavigationTypes.h
--- a/Engine/Source/Runtime/Engine/Classes/AI/Navigation/NavigationTypes.h
+++ b/Engine/Source/Runtime/Engine/Classes/AI/Navigation/NavigationTypes.h
@@ -182,6 +182,11 @@ struct ENGINE_API FNavAgentSelector
 	}
 
 	bool Serialize(FArchive& Ar);
+
+	uint32 GetAgentBits() const 
+	{
+		return (~InitializedBit & PackedBits);
+	}
 };
 
 template<>
