Slightly better fix for VT build failure on non VT platforms then 8370018

diff --git a/Engine/Source/Runtime/Engine/Classes/VT/RuntimeVirtualTextureStreamingProxy.h b/Engine/Source/Runtime/Engine/Classes/VT/RuntimeVirtualTextureStreamingProxy.h
--- a/Engine/Source/Runtime/Engine/Classes/VT/RuntimeVirtualTextureStreamingProxy.h
+++ b/Engine/Source/Runtime/Engine/Classes/VT/RuntimeVirtualTextureStreamingProxy.h
@@ -23,5 +23,8 @@ class URuntimeVirtualTextureStreamingProxy : public UTexture2D
 
 	//~ Begin UTexture Interface.
 	virtual void GetVirtualTextureBuildSettings(FVirtualTextureBuildSettings& OutSettings) const override;
+#if WITH_EDITOR
+	virtual void BeginCacheForCookedPlatformData(const ITargetPlatform* TargetPlatform) override;
+#endif
 	//~ End UTexture Interface.
 };
diff --git a/Engine/Source/Runtime/Engine/Private/TextureDerivedDataTask.cpp b/Engine/Source/Runtime/Engine/Private/TextureDerivedDataTask.cpp
--- a/Engine/Source/Runtime/Engine/Private/TextureDerivedDataTask.cpp
+++ b/Engine/Source/Runtime/Engine/Private/TextureDerivedDataTask.cpp
@@ -244,7 +244,7 @@ void FTextureCacheDerivedDataWorker::BuildTexture()
 	{
 		// Only support single Block/Layer here (Blocks and Layers are intended for VT support)
 		ensure(TextureData.Blocks.Num() == 1);
-		//ensure(TextureData.Layers.Num() == 1);
+		ensure(TextureData.Layers.Num() == 1);
 
 		check(DerivedData->Mips.Num() == 0);
 		DerivedData->SizeX = 0;
diff --git a/Engine/Source/Runtime/Engine/Private/VT/RuntimeVirtualTexture.cpp b/Engine/Source/Runtime/Engine/Private/VT/RuntimeVirtualTexture.cpp
--- a/Engine/Source/Runtime/Engine/Private/VT/RuntimeVirtualTexture.cpp
+++ b/Engine/Source/Runtime/Engine/Private/VT/RuntimeVirtualTexture.cpp
@@ -175,6 +175,23 @@ void URuntimeVirtualTextureStreamingProxy::GetVirtualTextureBuildSettings(FVirtu
 	OutSettings = Settings;
 }
 
+#if WITH_EDITOR
+
+void URuntimeVirtualTextureStreamingProxy::BeginCacheForCookedPlatformData(const ITargetPlatform* TargetPlatform)
+{
+	// Even though we skip the cook of this object for non VT platforms in URuntimeVirtualTexture::Serialize()
+	// we still load the object at cook time and kick off the DDC build. This will trigger an error in the texture DDC code.
+	// Either we need to make the DDC code more robust for non VT platforms or we can skip the process here...
+	if (!UseVirtualTexturing(GMaxRHIFeatureLevel, TargetPlatform))
+	{
+		return;
+	}
+
+	Super::BeginCacheForCookedPlatformData(TargetPlatform);
+}
+
+#endif
+
 
 URuntimeVirtualTexture::URuntimeVirtualTexture(const FObjectInitializer& ObjectInitializer)
 	: Super(ObjectInitializer)
