audio component state query

diff --git a/Engine/Source/Runtime/Engine/Classes/Components/AudioComponent.h b/Engine/Source/Runtime/Engine/Classes/Components/AudioComponent.h
--- a/Engine/Source/Runtime/Engine/Classes/Components/AudioComponent.h
+++ b/Engine/Source/Runtime/Engine/Classes/Components/AudioComponent.h
@@ -17,6 +17,28 @@ class USoundBase;
 class USoundClass;
 class USoundConcurrency;
 
+// Enum describing the audio component play state
+UENUM(BlueprintType)
+enum class EAudioComponentPlayState : uint8
+{
+	// If the sound is playing (i.e. not fading in, not fading out, not paused)
+	Playing,
+
+	// If the sound is not playing
+	Stopped, 
+
+	// If the sound is playing but paused
+	Paused,
+
+	// If the sound is playing and fading in
+	FadingIn,
+
+	// If the sound is playing and fading out
+	FadingOut,
+
+	Count UMETA(Hidden)
+};
+
 
 /** called when we finish playing audio, either because it played to completion or because a Stop() call turned it off early */
 DECLARE_DYNAMIC_MULTICAST_DELEGATE( FOnAudioFinished );
@@ -212,6 +234,9 @@ class ENGINE_API UAudioComponent : public USceneComponent
 	/** Whether or not this audio component has been paused */
 	uint8 bIsPaused:1;
 
+	/** Whether or not fade out was triggered. */
+	uint8 bIsFadingOut:1;
+
 	/**
 	* True if we should automatically attach to AutoAttachParent when Played, and detach from our parent when playback is completed.
 	* This overrides any current attachment that may be present at the time of activation (deferring initial attachment until activation, if AutoAttachParent is null).
@@ -312,6 +337,12 @@ public:
 	/** while playing, this component will check for occlusion from its closest listener every this many seconds */
 	float OcclusionCheckInterval;
 
+	/** What time the audio component was told to play. Used to compute audio component state. */
+	float TimeAudioComponentPlayed;
+
+	/** How much time the audio component was told to fade in. */
+	float FadeInTimeDuration;
+
 	/**
 	 * Options for how we handle our location when we attach to the AutoAttachParent, if bAutoManageAttachment is true.
 	 * @see bAutoManageAttachment, EAttachmentRule
@@ -415,10 +446,14 @@ public:
 	UFUNCTION(BlueprintCallable, Category = "Audio|Components|Audio")
 	void SetPaused(bool bPause);
 
-	/** Returns true if this component is currently playing a SoundCue. */
-	UFUNCTION(BlueprintCallable, Category="Audio|Components|Audio")
+	/** Returns if the sound playing any audio. Doesn't indicate the play state. Use GetPlayState() to get the actual play state. */
+	UFUNCTION(BlueprintCallable, Category="Audio|Components|Audio", meta=(DeprecatedFunction))
 	virtual bool IsPlaying() const;
 
+	/** Returns the enumerated play states of the audio component. */
+	UFUNCTION(BlueprintCallable, Category = "Audio|Components|Audio")
+	EAudioComponentPlayState GetPlayState() const;
+
 	/** This will allow one to adjust the volume of an AudioComponent on the fly */
 	UFUNCTION(BlueprintCallable, Category="Audio|Components|Audio")
 	void AdjustVolume(float AdjustVolumeDuration, float AdjustVolumeLevel, const EAudioFaderCurve FadeCurve = EAudioFaderCurve::Linear);
diff --git a/Engine/Source/Runtime/Engine/Private/Components/AudioComponent.cpp b/Engine/Source/Runtime/Engine/Private/Components/AudioComponent.cpp
--- a/Engine/Source/Runtime/Engine/Private/Components/AudioComponent.cpp
+++ b/Engine/Source/Runtime/Engine/Private/Components/AudioComponent.cpp
@@ -363,6 +363,9 @@ void UAudioComponent::PlayInternal(const float StartTime, const float FadeInDura
 
 	UE_LOG(LogAudio, Verbose, TEXT("%g: Playing AudioComponent : '%s' with Sound: '%s'"), World ? World->GetAudioTimeSeconds() : 0.0f, *GetFullName(), Sound ? *Sound->GetName() : TEXT("nullptr"));
 
+	// Reset our fading out flag in case this is a reused audio component and we are replaying after previously fading out
+	bIsFadingOut = false;
+
 	if (bIsActive)
 	{
 		// If this is an auto destroy component we need to prevent it from being auto-destroyed since we're really just restarting it
@@ -377,6 +380,17 @@ void UAudioComponent::PlayInternal(const float StartTime, const float FadeInDura
 	{
 		if (FAudioDevice* AudioDevice = GetAudioDevice())
 		{
+			// Store the time that this audio component played
+			if (World)
+			{
+				TimeAudioComponentPlayed = World->GetAudioTimeSeconds();
+			}
+			else
+			{
+				TimeAudioComponentPlayed = 0.0f;
+			}
+			FadeInTimeDuration = FadeInDuration;
+
 			// Auto attach if requested
 			const bool bWasAutoAttached = bDidAutoAttach;
 			bDidAutoAttach = false;
@@ -576,6 +590,8 @@ void UAudioComponent::AdjustVolumeInternal(float AdjustVolumeDuration, float Adj
 		return;
 	}
 
+	bIsFadingOut = bIsFadeOut || FMath::IsNearlyZero(AdjustVolumeLevel);
+
 	const uint64 InAudioComponentID = AudioComponentID;
 	DECLARE_CYCLE_STAT(TEXT("FAudioThreadTask.AdjustVolume"), STAT_AudioAdjustVolume, STATGROUP_AudioThreadCommands);
 	FAudioThread::RunCommandOnAudioThread([AudioDevice, InAudioComponentID, AdjustVolumeDuration, AdjustVolumeLevel, bIsFadeOut, FadeCurve]()
@@ -794,6 +810,35 @@ bool UAudioComponent::IsPlaying() const
 	return bIsActive;
 }
 
+EAudioComponentPlayState UAudioComponent::GetPlayState() const
+{
+	UWorld* World = GetWorld();
+	if (!bIsActive || !World)
+	{
+		return EAudioComponentPlayState::Stopped;
+	}
+
+	if (bIsPaused)
+	{
+		return EAudioComponentPlayState::Paused;
+	}
+
+	if (bIsFadingOut)
+	{
+		return EAudioComponentPlayState::FadingOut;
+	}
+
+	// Get the current audio time seconds and compare when it started and the fade in duration 
+	float CurrentAudioTimeSeconds = World->GetAudioTimeSeconds();
+	if (CurrentAudioTimeSeconds - TimeAudioComponentPlayed < FadeInTimeDuration)
+	{
+		return EAudioComponentPlayState::FadingIn;
+	}
+
+	// If we are not in any of the above states we are "playing"
+	return EAudioComponentPlayState::Playing;
+}
+
 #if WITH_EDITORONLY_DATA
 void UAudioComponent::UpdateSpriteTexture()
 {
