Added comment to explain code change in previous changelist 8360392

diff --git a/Engine/Source/Runtime/Engine/Private/VT/VirtualTextureUploadCache.cpp b/Engine/Source/Runtime/Engine/Private/VT/VirtualTextureUploadCache.cpp
--- a/Engine/Source/Runtime/Engine/Private/VT/VirtualTextureUploadCache.cpp
+++ b/Engine/Source/Runtime/Engine/Private/VT/VirtualTextureUploadCache.cpp
@@ -110,8 +110,9 @@ void FVirtualTextureUploadCache::Finalize(FRHICommandListImmediate& RHICmdList)
 				DEC_MEMORY_STAT_BY(STAT_TotalGPUUploadSize, CalcTextureSize(StagingTexture.RHITexture->GetSizeX(), StagingTexture.RHITexture->GetSizeY(), PoolEntry.Format, 1u));
 			}
 
+			//todo[vt]: Intended to use TexCreate_CPUWritable on PC but it doesn't play well with D3D11 multi-threaded rendering Lock/Unlock
 			FRHIResourceCreateInfo CreateInfo;
-			StagingTexture.RHITexture = RHICmdList.CreateTexture2D(TileSize * WidthInTiles, TileSize * HeightInTiles, PoolEntry.Format, 1, 1, TexCreate_None, CreateInfo);
+			StagingTexture.RHITexture = RHICmdList.CreateTexture2D(TileSize * WidthInTiles, TileSize * HeightInTiles, PoolEntry.Format, 1, 1, TexCreate_None, CreateInfo); 
 			StagingTexture.WidthInTiles = WidthInTiles;
 			StagingTexture.BatchCapacity = WidthInTiles * HeightInTiles;
 			INC_MEMORY_STAT_BY(STAT_TotalGPUUploadSize, CalcTextureSize(TileSize * WidthInTiles, TileSize * HeightInTiles, PoolEntry.Format, 1u));
