Fix wrong computation used for interpolator count stats
Fix number of vertex tex coords allocated by SpeedTree node

diff --git a/Engine/Source/Runtime/Engine/Private/Materials/HLSLMaterialTranslator.h b/Engine/Source/Runtime/Engine/Private/Materials/HLSLMaterialTranslator.h
--- a/Engine/Source/Runtime/Engine/Private/Materials/HLSLMaterialTranslator.h
+++ b/Engine/Source/Runtime/Engine/Private/Materials/HLSLMaterialTranslator.h
@@ -1202,7 +1202,7 @@ ResourcesString = TEXT("");
 				}
 			}
 
-			MaterialCompilationOutput.NumUsedUVScalars = GetNumUserTexCoords();
+			MaterialCompilationOutput.NumUsedUVScalars = GetNumUserTexCoords() * 2;
 			MaterialCompilationOutput.NumUsedCustomInterpolatorScalars = CurrentCustomVertexInterpolatorOffset;
 
 			// Do Normal Chunk first
@@ -3980,15 +3980,21 @@ protected:
 		}
 	}
 
-	void AllocateSlot(TBitArray<>& InBitArray, int32 InSlotIndex) const
+	void AllocateSlot(TBitArray<>& InBitArray, int32 InSlotIndex, int32 InSlotCount = 1) const
 	{
 		// Grow as needed
-		while (InBitArray.Num() <= InSlotIndex)
+		int32 NumSlotsNeeded = InSlotIndex + InSlotCount;
+		int32 CurrentNumSlots = InBitArray.Num();
+		if(NumSlotsNeeded > CurrentNumSlots)
 		{
-			InBitArray.Add(false);
+			InBitArray.Add(false, NumSlotsNeeded - CurrentNumSlots);
 		}
 
-		InBitArray[InSlotIndex] = true;
+		// Allocate the requested slot(s)
+		for (int32 i = InSlotIndex; i < NumSlotsNeeded; ++i)
+		{
+			InBitArray[i] = true;
+		}
 	}
 
 #if WITH_EDITOR
@@ -3996,13 +4002,11 @@ protected:
 	{
 		if (ShaderFrequency == SF_Vertex)
 		{
-			AllocateSlot(AllocatedUserVertexTexCoords, 6);
-			AllocateSlot(AllocatedUserVertexTexCoords, 7);
+			AllocateSlot(AllocatedUserVertexTexCoords, 6, 2);
 		}
 		else
 		{
-			AllocateSlot(AllocatedUserTexCoords, 6);
-			AllocateSlot(AllocatedUserTexCoords, 7);
+			AllocateSlot(AllocatedUserTexCoords, 6, 2);
 		}
 
 		// Note: inlining is important so that on ES2 devices, where half precision is used in the pixel shader, 
@@ -6994,7 +6998,7 @@ protected:
 		{
 			bUsesSpeedTree = true;
 
-			AllocateSlot(AllocatedUserVertexTexCoords, 7);
+			AllocateSlot(AllocatedUserVertexTexCoords, 2, 6);
 
 			// Only generate previous frame's computations if required and opted-in
 			const bool bEnablePreviousFrameInformation = bCompilingPreviousFrame && bAccurateWindVelocities;
