#UE4 Remove deprecation from IsPlaying as it is still broadly used

diff --git a/Engine/Source/Runtime/Engine/Classes/Components/AudioComponent.h b/Engine/Source/Runtime/Engine/Classes/Components/AudioComponent.h
--- a/Engine/Source/Runtime/Engine/Classes/Components/AudioComponent.h
+++ b/Engine/Source/Runtime/Engine/Classes/Components/AudioComponent.h
@@ -447,7 +447,7 @@ public:
 	void SetPaused(bool bPause);
 
 	/** Returns if the sound playing any audio. Doesn't indicate the play state. Use GetPlayState() to get the actual play state. */
-	UFUNCTION(BlueprintCallable, Category="Audio|Components|Audio", meta=(DeprecatedFunction))
+	UFUNCTION(BlueprintCallable, Category="Audio|Components|Audio")
 	virtual bool IsPlaying() const;
 
 	/** Returns the enumerated play states of the audio component. */
