#ue4 - Added SetUseAutoAttachment() to FXSystemComponent, and implemented for ParticleSystemComponent and NiagaraComponent.

diff --git a/Engine/Plugins/FX/Niagara/Source/Niagara/Public/NiagaraComponent.h b/Engine/Plugins/FX/Niagara/Source/Niagara/Public/NiagaraComponent.h
--- a/Engine/Plugins/FX/Niagara/Source/Niagara/Public/NiagaraComponent.h
+++ b/Engine/Plugins/FX/Niagara/Source/Niagara/Public/NiagaraComponent.h
@@ -431,6 +431,8 @@ public:
 	 */
 	void SetAutoAttachmentParameters(USceneComponent* Parent, FName SocketName, EAttachmentRule LocationRule, EAttachmentRule RotationRule, EAttachmentRule ScaleRule) override;
 
+	virtual void SetUseAutoManageAttachment(bool bAutoManage) override { bAutoManageAttachment = bAutoManage; }
+
 	UPROPERTY(EditAnywhere, Category = Preview, Transient, meta=(EditCondition=bEnablePreviewDetailLevel))
 	int32 PreviewDetailLevel;
 
diff --git a/Engine/Source/Runtime/Engine/Classes/Particles/ParticleSystemComponent.h b/Engine/Source/Runtime/Engine/Classes/Particles/ParticleSystemComponent.h
--- a/Engine/Source/Runtime/Engine/Classes/Particles/ParticleSystemComponent.h
+++ b/Engine/Source/Runtime/Engine/Classes/Particles/ParticleSystemComponent.h
@@ -384,6 +384,16 @@ public:
 	UFUNCTION(BlueprintCallable, Category = "Effects|Components|ParticleSystem")
 	virtual void SetAutoAttachmentParameters(USceneComponent* Parent, FName SocketName, EAttachmentRule LocationRule, EAttachmentRule RotationRule, EAttachmentRule ScaleRule) {}
 
+	/**
+	 * Sets whether we should automatically attach to AutoAttachParent when activated, and detach from our parent when completed.
+	 * This overrides any current attachment that may be present at the time of activation (deferring initial attachment until activation, if AutoAttachParent is null).
+	 * When enabled, detachment occurs regardless of whether AutoAttachParent is assigned, and the relative transform from the time of activation is restored.
+	 * This also disables attachment on dedicated servers, where we don't actually activate even if bAutoActivate is true.
+	 * @see SetAutoAttachmentParameters()
+	 */
+	UFUNCTION(BlueprintCallable, Category = "Effects|Components|ParticleSystem")
+	virtual void SetUseAutoManageAttachment(bool bAutoManage) {}
+
 	/**
 	 * Deactivates this system and releases it to the pool on completion.
 	 * Usage of this PSC reference after this call is unsafe.
@@ -762,6 +772,8 @@ public:
 	 */
 	void SetAutoAttachmentParameters(USceneComponent* Parent, FName SocketName, EAttachmentRule LocationRule, EAttachmentRule RotationRule, EAttachmentRule ScaleRule) override;
 
+	virtual void SetUseAutoManageAttachment(bool bAutoManage) override { bAutoManageAttachment = bAutoManage; }
+
 private:
 
 	/** Saved relative transform before auto attachement. Used during detachment to restore the transform if we had automatically attached. */
