Fix memory leaks in immediate physics.

diff --git a/Engine/Source/Runtime/Engine/Private/PhysicsEngine/ImmediatePhysics/ImmediatePhysicsPhysX/ImmediatePhysicsActor_Physx.cpp b/Engine/Source/Runtime/Engine/Private/PhysicsEngine/ImmediatePhysics/ImmediatePhysicsPhysX/ImmediatePhysicsActor_Physx.cpp
--- a/Engine/Source/Runtime/Engine/Private/PhysicsEngine/ImmediatePhysics/ImmediatePhysicsPhysX/ImmediatePhysicsActor_Physx.cpp
+++ b/Engine/Source/Runtime/Engine/Private/PhysicsEngine/ImmediatePhysics/ImmediatePhysicsPhysX/ImmediatePhysicsActor_Physx.cpp
@@ -16,6 +16,11 @@ namespace ImmediatePhysics_PhysX
 		return &FMaterial::Default;
 	}
 
+	FActor::~FActor()
+	{
+		TerminateGeometry();
+	}
+
 	void FActor::CreateGeometry(PxRigidActor* RigidActor, const PxTransform& ActorToBodyTM)
 	{
 		const uint32 NumShapes = RigidActor->getNbShapes();
@@ -110,7 +115,16 @@ namespace ImmediatePhysics_PhysX
 	{
 		for (FShape& Shape : Shapes)
 		{
-			delete Shape.Geometry;
+			switch (Shape.Geometry->getType())
+			{
+			case PxGeometryType::eSPHERE: delete static_cast<PxSphereGeometry*>(Shape.Geometry); break;
+			case PxGeometryType::eCAPSULE: delete static_cast<PxCapsuleGeometry*>(Shape.Geometry); break;
+			case PxGeometryType::eBOX: delete static_cast<PxBoxGeometry*>(Shape.Geometry); break;
+			case PxGeometryType::eCONVEXMESH: delete static_cast<PxConvexMeshGeometry*>(Shape.Geometry); break;
+			case PxGeometryType::eTRIANGLEMESH: delete static_cast<PxTriangleMeshGeometry*>(Shape.Geometry); break;
+			case PxGeometryType::eHEIGHTFIELD: delete static_cast<PxHeightFieldGeometry*>(Shape.Geometry); break;
+			default: break;
+			}
 		}
 
 		Shapes.Empty();
diff --git a/Engine/Source/Runtime/Engine/Private/PhysicsEngine/ImmediatePhysics/ImmediatePhysicsPhysX/ImmediatePhysicsSimulation_PhysX.cpp b/Engine/Source/Runtime/Engine/Private/PhysicsEngine/ImmediatePhysics/ImmediatePhysicsPhysX/ImmediatePhysicsSimulation_PhysX.cpp
--- a/Engine/Source/Runtime/Engine/Private/PhysicsEngine/ImmediatePhysics/ImmediatePhysicsPhysX/ImmediatePhysicsSimulation_PhysX.cpp
+++ b/Engine/Source/Runtime/Engine/Private/PhysicsEngine/ImmediatePhysics/ImmediatePhysicsPhysX/ImmediatePhysicsSimulation_PhysX.cpp
@@ -43,7 +43,12 @@ FSimulation::~FSimulation()
 		delete ActorHandle;
 	}
 
+	for (FJointHandle* JointHandle : JointHandles)
+	{
+		delete JointHandle;
+	}
 	ActorHandles.Empty();
+	JointHandles.Empty();
 }
 
 void FSimulation::SetPositionIterationCount(int32 InIterationCount)
@@ -351,7 +356,10 @@ void FSimulation::RemoveActor(FActorHandle* InHandle)
 	bool bStatic = !bDynamic && !bKinematic;
 
 	Actors.RemoveAt(DataIndex);
+
+	delete ActorHandles[DataIndex];
 	ActorHandles.RemoveAt(DataIndex);
+
 	RigidBodiesData.RemoveAt(DataIndex);
 	SolverBodiesData.RemoveAt(DataIndex);
 	PendingAcceleration.RemoveAt(DataIndex);
diff --git a/Engine/Source/Runtime/Engine/Public/Physics/ImmediatePhysics/ImmediatePhysicsPhysX/ImmediatePhysicsActor_PhysX.h b/Engine/Source/Runtime/Engine/Public/Physics/ImmediatePhysics/ImmediatePhysicsPhysX/ImmediatePhysicsActor_PhysX.h
--- a/Engine/Source/Runtime/Engine/Public/Physics/ImmediatePhysics/ImmediatePhysicsPhysX/ImmediatePhysicsActor_PhysX.h
+++ b/Engine/Source/Runtime/Engine/Public/Physics/ImmediatePhysics/ImmediatePhysicsPhysX/ImmediatePhysicsActor_PhysX.h
@@ -16,6 +16,8 @@ namespace ImmediatePhysics_PhysX
 			: UserData(nullptr)
 		{}
 
+		~FActor();
+
 		TArray<FShape> Shapes;
 		void* UserData;
 
