Optimized APlayerController::ServerShortTimeout and added ForceAllActorsNetUpdateTime to NetDriver
Added ensures to track any CDO or Archetypes getting wrongly added to the NetDriver

diff --git a/Engine/Source/Runtime/Engine/Classes/Engine/NetDriver.h b/Engine/Source/Runtime/Engine/Classes/Engine/NetDriver.h
--- a/Engine/Source/Runtime/Engine/Classes/Engine/NetDriver.h
+++ b/Engine/Source/Runtime/Engine/Classes/Engine/NetDriver.h
@@ -1333,6 +1333,8 @@ public:
 
 	ENGINE_API virtual void ForceNetUpdate(AActor* Actor);
 
+	ENGINE_API void ForceAllActorsNetUpdateTime(float NetUpdateTimeOffset, TFunctionRef<bool(const AActor* const)> ValidActorTestFunc);
+
 	/** Flushes actor from NetDriver's dormancy list, but does not change any state on the Actor itself */
 	ENGINE_API void FlushActorDormancy(AActor *Actor, bool bWasDormInitial=false);
 
diff --git a/Engine/Source/Runtime/Engine/Private/NetworkDriver.cpp b/Engine/Source/Runtime/Engine/Private/NetworkDriver.cpp
--- a/Engine/Source/Runtime/Engine/Private/NetworkDriver.cpp
+++ b/Engine/Source/Runtime/Engine/Private/NetworkDriver.cpp
@@ -445,6 +445,8 @@ void UNetDriver::AssertValid()
 void UNetDriver::AddNetworkActor(AActor* Actor)
 {
 	LLM_SCOPE(ELLMTag::Networking);
+	ensureMsgf(Actor == nullptr || !(Actor->HasAnyFlags(RF_ClassDefaultObject | RF_ArchetypeObject)), TEXT("%s is a CDO or Archetype and should not be replicated."), *GetFullNameSafe(Actor));
+
 	GetNetworkObjectList().FindOrAdd(Actor, this);
 	if (ReplicationDriver)
 	{
@@ -454,6 +456,8 @@ void UNetDriver::AddNetworkActor(AActor* Actor)
 
 FNetworkObjectInfo* UNetDriver::FindOrAddNetworkObjectInfo(const AActor* InActor)
 {
+	ensureMsgf(InActor == nullptr || !(InActor->HasAnyFlags(RF_ClassDefaultObject | RF_ArchetypeObject)), TEXT("%s is a CDO or Archetype and should not be replicated."), *GetFullNameSafe(InActor));
+
 	bool bWasAdded = false;
 	if (TSharedPtr<FNetworkObjectInfo>* InfoPtr = GetNetworkObjectList().FindOrAdd(const_cast<AActor*>(InActor), this, &bWasAdded))
 	{
@@ -3112,6 +3116,27 @@ void UNetDriver::ForceNetUpdate(AActor* Actor)
 	}
 }
 
+void UNetDriver::ForceAllActorsNetUpdateTime(float NetUpdateTimeOffset, TFunctionRef<bool(const AActor* const)> ValidActorTestFunc)
+{
+	for (auto It = GetNetworkObjectList().GetAllObjects().CreateConstIterator(); It; ++It)
+	{
+		FNetworkObjectInfo* NetActorInfo = (*It).Get();
+		if (NetActorInfo)
+		{
+			const AActor* const Actor = NetActorInfo->Actor;
+			if (Actor && !Actor->IsPendingKill())
+			{
+				if (ValidActorTestFunc(Actor))
+				{
+					// Only allow the next update to be sooner than the current one
+					const double NewUpdateTime = World->TimeSeconds + NetUpdateTimeOffset * FMath::FRand();
+					NetActorInfo->NextUpdateTime = FMath::Min(NetActorInfo->NextUpdateTime, NewUpdateTime);
+				}
+			}
+		}
+	}
+}
+
 /** UNetDriver::FlushActorDormancy(AActor* Actor)
  *	 Flushes the actor from the NetDriver's dormant list and/or cancels pending dormancy on the actor channel.
  *
diff --git a/Engine/Source/Runtime/Engine/Private/PlayerController.cpp b/Engine/Source/Runtime/Engine/Private/PlayerController.cpp
--- a/Engine/Source/Runtime/Engine/Private/PlayerController.cpp
+++ b/Engine/Source/Runtime/Engine/Private/PlayerController.cpp
@@ -1038,6 +1038,7 @@ bool APlayerController::ServerShortTimeout_Validate()
 
 void APlayerController::ServerShortTimeout_Implementation()
 {
+	QUICK_SCOPE_CYCLE_COUNTER(STAT_PC_ServerShortTimeout);
 	if (!bShortConnectTimeOut)
 	{
 		UWorld* World = GetWorld();
@@ -1068,20 +1069,14 @@ void APlayerController::ServerShortTimeout_Implementation()
 		}
 		else 
 		{
-			float NetUpdateTimeOffset = (World->GetAuthGameMode()->GetNumPlayers() < 8) ? 0.2f : 0.5f;
-			for (TSharedPtr<FNetworkObjectInfo> NetworkObjectInfo : World->GetNetDriver()->GetNetworkObjectList().GetAllObjects())
+			if (World->GetNetDriver())
 			{
-				if (NetworkObjectInfo.IsValid())
+				float NetUpdateTimeOffset = (World->GetAuthGameMode()->GetNumPlayers() < 8) ? 0.2f : 0.5f;
+				auto ValidActorTest = [](const AActor* const Actor)
 				{
-					AActor* const A = NetworkObjectInfo->Actor;
-					if (A && !A->IsPendingKill())
-					{
-						if ((A->NetUpdateFrequency < 1) && !A->bOnlyRelevantToOwner)
-						{
-							A->SetNetUpdateTime(World->TimeSeconds + NetUpdateTimeOffset * FMath::FRand());
-						}
-					}
-				}
+					return (Actor->NetUpdateFrequency < 1) && !Actor->bOnlyRelevantToOwner;
+				};
+				World->GetNetDriver()->ForceAllActorsNetUpdateTime(NetUpdateTimeOffset, ValidActorTest);
 			}
 		}
 	}
