Exposed public BP functions in Sky Atmosphere component as ENGINE_API so game code can call it.

#rnx

diff --git a/Engine/Source/Runtime/Engine/Classes/Components/SkyAtmosphereComponent.h b/Engine/Source/Runtime/Engine/Classes/Components/SkyAtmosphereComponent.h
--- a/Engine/Source/Runtime/Engine/Classes/Components/SkyAtmosphereComponent.h
+++ b/Engine/Source/Runtime/Engine/Classes/Components/SkyAtmosphereComponent.h
@@ -127,34 +127,34 @@ class USkyAtmosphereComponent : public USceneComponent
 	ENGINE_API void OverrideAtmosphereLightDirection(int32 AtmosphereLightIndex, const FVector& LightDirection);
 	
 	UFUNCTION(BlueprintCallable, Category = "Rendering")
-	void SetRayleighScatteringScale(float NewValue);
+	ENGINE_API void SetRayleighScatteringScale(float NewValue);
 	UFUNCTION(BlueprintCallable, Category = "Rendering")
-	void SetRayleighScattering(FColor NewValue);
+	ENGINE_API void SetRayleighScattering(FColor NewValue);
 	UFUNCTION(BlueprintCallable, Category = "Rendering")
-	void SetRayleighExponentialDistribution(float NewValue);
+	ENGINE_API void SetRayleighExponentialDistribution(float NewValue);
 
 	UFUNCTION(BlueprintCallable, Category = "Rendering")
-	void SetMieScatteringScale(float NewValue);
+	ENGINE_API void SetMieScatteringScale(float NewValue);
 	UFUNCTION(BlueprintCallable, Category = "Rendering")
-	void SetMieScattering(FColor NewValue);
+	ENGINE_API void SetMieScattering(FColor NewValue);
 	UFUNCTION(BlueprintCallable, Category = "Rendering")
-	void SetMieAbsorptionScale(float NewValue);
+	ENGINE_API void SetMieAbsorptionScale(float NewValue);
 	UFUNCTION(BlueprintCallable, Category = "Rendering")
-	void SetMieAbsorption(FColor NewValue);
+	ENGINE_API void SetMieAbsorption(FColor NewValue);
 	UFUNCTION(BlueprintCallable, Category = "Rendering")
-	void SetMieAnisotropy(float NewValue);
+	ENGINE_API void SetMieAnisotropy(float NewValue);
 	UFUNCTION(BlueprintCallable, Category = "Rendering")
-	void SetMieExponentialDistribution(float NewValue);
+	ENGINE_API void SetMieExponentialDistribution(float NewValue);
 
 	UFUNCTION(BlueprintCallable, Category = "Rendering", meta = (DisplayName = "Set Absorption Scale"))
-	void SetOtherAbsorptionScale(float NewValue);
+	ENGINE_API void SetOtherAbsorptionScale(float NewValue);
 	UFUNCTION(BlueprintCallable, Category = "Rendering", meta = (DisplayName = "Set Absorption"))
-	void SetOtherAbsorption(FColor NewValue);
+	ENGINE_API void SetOtherAbsorption(FColor NewValue);
 
 	UFUNCTION(BlueprintCallable, Category = "Rendering")
-	void SetSkyLuminanceFactor(FLinearColor NewValue);
+	ENGINE_API void SetSkyLuminanceFactor(FLinearColor NewValue);
 	UFUNCTION(BlueprintCallable, Category = "Rendering")
-	void SetAerialPespectiveViewDistanceScale(float NewValue);
+	ENGINE_API void SetAerialPespectiveViewDistanceScale(float NewValue);
 
 
 
