Create a templated version of AActor::GetInstigatorController to be match AActor::GetInstigator.

diff --git a/Engine/Source/Runtime/Engine/Classes/GameFramework/Actor.h b/Engine/Source/Runtime/Engine/Classes/GameFramework/Actor.h
--- a/Engine/Source/Runtime/Engine/Classes/GameFramework/Actor.h
+++ b/Engine/Source/Runtime/Engine/Classes/GameFramework/Actor.h
@@ -842,12 +842,25 @@ public:
 	 * @return The instigator for this actor if it is the specified type, nullptr otherwise.
 	 */
 	template <class T>
-	T* GetInstigator() const { return Cast<T>(Instigator); };
+	T* GetInstigator() const
+	{
+		return Cast<T>(GetInstigator());
+	}
 
 	/** Returns the instigator's controller for this actor, or nullptr if there is none. */
 	UFUNCTION(BlueprintCallable, meta=(BlueprintProtected = "true"), Category="Game")
 	AController* GetInstigatorController() const;
 
+	/** 
+	 * Returns the instigator's controller, cast as a specific class.
+	 * @return The instigator's controller for this actor if it is the specified type, nullptr otherwise.
+	 * */
+	template<class T>
+	T* GetInstigatorController() const
+	{
+		return Cast<T>(GetInstigatorController());
+	}
+
 
 	//~=============================================================================
 	// General functions.
diff --git a/Engine/Source/Runtime/Engine/Private/Components/SkinnedMeshComponent.cpp b/Engine/Source/Runtime/Engine/Private/Components/SkinnedMeshComponent.cpp
--- a/Engine/Source/Runtime/Engine/Private/Components/SkinnedMeshComponent.cpp
+++ b/Engine/Source/Runtime/Engine/Private/Components/SkinnedMeshComponent.cpp
@@ -108,8 +108,8 @@ namespace FAnimUpdateRateManager
 		bool IsHumanControlled() const
 		{
 			AActor* Owner = RegisteredComponents[0]->GetOwner();
-			AController * Controller = Owner ? Owner->GetInstigatorController() : NULL;
-			return Cast<APlayerController>(Controller) != NULL;
+			APlayerController* Controller = Owner ? Owner->GetInstigatorController<APlayerController>() : nullptr;
+			return Controller != nullptr;
 		}
 	};
 
