Remove some special cases from Role / RemoteRole in RepLayout, and try to share more code in places where we do check.

diff --git a/Engine/Source/Runtime/Engine/Private/RepLayout.cpp b/Engine/Source/Runtime/Engine/Private/RepLayout.cpp
--- a/Engine/Source/Runtime/Engine/Private/RepLayout.cpp
+++ b/Engine/Source/Runtime/Engine/Private/RepLayout.cpp
@@ -934,6 +934,7 @@ struct FComparePropertiesSharedParams
 	const bool bForceFail;
 	const int16 RoleIndex;
 	const int16 RemoteRoleIndex;
+	const ERepLayoutFlags Flags;
 	const TArray<FRepParentCmd>& Parents;
 	const TArray<FRepLayoutCmd>& Cmds;
 };
@@ -955,35 +956,34 @@ static void CompareProperties_Array_r(
 	const uint16 CmdIndex,
 	const uint16 Handle);
 
+static void CompareRoleProperty(
+	const FComparePropertiesSharedParams& SharedParams,
+	const FConstRepObjectDataBuffer Data,
+	const uint16 RoleOrRemoteRoleIndex,
+	TEnumAsByte<ENetRole>& SavedRoleOrRemoteRole,
+	TArray<uint16>& Changed)
+{
+	const FRepParentCmd& RoleOrRemoteRoleParent = SharedParams.Parents[RoleOrRemoteRoleIndex];
+	const FRepLayoutCmd& RoleOrRemoteRoleCmd = SharedParams.Cmds[RoleOrRemoteRoleParent.CmdStart];
+	const uint16 Handle = RoleOrRemoteRoleCmd.RelativeHandle;
+	const TEnumAsByte<ENetRole> ActorRoleOrRemoteRole = *(const TEnumAsByte<ENetRole>*)(Data + RoleOrRemoteRoleParent).Data;
+	if (SharedParams.bForceFail || SavedRoleOrRemoteRole != ActorRoleOrRemoteRole)
+	{
+		SavedRoleOrRemoteRole = ActorRoleOrRemoteRole;
+		Changed.Add(Handle);
+	}
+}
+
 static void CompareRoleProperties(
 	const FComparePropertiesSharedParams& SharedParams,
 	FSendingRepState* RESTRICT RepState,
 	const FConstRepObjectDataBuffer Data,
 	TArray<uint16>& Changed)
 {
-	if (RepState && SharedParams.RoleIndex != INDEX_NONE && SharedParams.RemoteRoleIndex != INDEX_NONE)
+	if (RepState && EnumHasAnyFlags(SharedParams.Flags, ERepLayoutFlags::IsActor))
 	{
-		const FRepParentCmd& RemoteRoleParent = SharedParams.Parents[SharedParams.RemoteRoleIndex];
-		const FRepLayoutCmd& RemoteRoleCmd = SharedParams.Cmds[RemoteRoleParent.CmdStart];
-		const uint16 RemoteRoleHandle = RemoteRoleCmd.RelativeHandle;
-
-		const TEnumAsByte<ENetRole> ObjectRemoteRole = *(const TEnumAsByte<ENetRole>*)(Data + RemoteRoleParent).Data;
-		if (SharedParams.bForceFail || RepState->SavedRemoteRole != ObjectRemoteRole)
-		{
-			RepState->SavedRemoteRole = ObjectRemoteRole;
-			Changed.Add(RemoteRoleHandle);
-		}
-
-		const FRepParentCmd& RoleParent = SharedParams.Parents[SharedParams.RoleIndex];
-		const FRepLayoutCmd& RoleCmd = SharedParams.Cmds[RoleParent.CmdStart];
-		const uint16 RoleHandle = RoleCmd.RelativeHandle;
-
-		const TEnumAsByte<ENetRole> ObjectRole = *(const TEnumAsByte<ENetRole>*)(Data + RoleParent).Data;
-		if (SharedParams.bForceFail || RepState->SavedRole != ObjectRole)
-		{
-			RepState->SavedRole = ObjectRole;
-			Changed.Add(RoleHandle);
-		}
+		CompareRoleProperty(SharedParams, Data, SharedParams.RemoteRoleIndex, RepState->SavedRemoteRole, Changed);
+		CompareRoleProperty(SharedParams, Data, SharedParams.RoleIndex, RepState->SavedRole, Changed);
 	}
 }
 
@@ -999,6 +999,8 @@ static void CompareParentProperties(
 
 	FRepChangedPropertyTracker* RepChangedPropertyTracker = (RepState ? RepState->RepChangedPropertyTracker.Get() : nullptr);
 
+	const bool bCheckForRole = EnumHasAnyFlags(SharedParams.Flags, ERepLayoutFlags::IsActor) && RepState != nullptr;
+
 	for (uint16 ParentIndex = 0; ParentIndex < SharedParams.Parents.Num(); ++ParentIndex)
 	{
 		const FRepParentCmd& Parent = SharedParams.Parents[ParentIndex];
@@ -1018,34 +1020,24 @@ static void CompareParentProperties(
 			continue;
 		}
 
-		const FRepLayoutCmd& Cmd = SharedParams.Cmds[Parent.CmdStart];
-		const uint16 Handle = Cmd.RelativeHandle;
-
-		// RepState may be null in the case where a deprecated version of this function is called.
-		// In that case, just allow this to fail and perform the old logic.
-		if (UNLIKELY(RepState && ParentIndex == SharedParams.RoleIndex))
+		if (bCheckForRole)
 		{
-			const TEnumAsByte<ENetRole> ObjectRole = *(const TEnumAsByte<ENetRole>*)(Data + Parent).Data;
-			if (SharedParams.bForceFail || RepState->SavedRole != ObjectRole)
+			if (UNLIKELY(ParentIndex == SharedParams.RoleIndex))
 			{
-				RepState->SavedRole = ObjectRole;
-				Changed.Add(Handle);
+				CompareRoleProperty(SharedParams, Data, SharedParams.RoleIndex, RepState->SavedRole, Changed);
+				continue;
 			}
-		}
-		else if (UNLIKELY(RepState && ParentIndex == SharedParams.RemoteRoleIndex))
-		{
-			const TEnumAsByte<ENetRole> ObjectRemoteRole = *(const TEnumAsByte<ENetRole>*)(Data + Parent).Data;
-			if (SharedParams.bForceFail || RepState->SavedRemoteRole != ObjectRemoteRole)
+			else if (UNLIKELY(ParentIndex == SharedParams.RemoteRoleIndex))
 			{
-				RepState->SavedRemoteRole = ObjectRemoteRole;
-				Changed.Add(Handle);
+				CompareRoleProperty(SharedParams, Data, SharedParams.RemoteRoleIndex, RepState->SavedRemoteRole, Changed);
+				continue;
 			}
 		}
-		else
-		{
-			// Note, Handle - 1 to account for CompareProperties_r incrementing handles.
-			CompareProperties_r(SharedParams, Parent.CmdStart, Parent.CmdEnd, ShadowData, Data, Changed, Handle - 1);
-		}
+		
+		// Note, Handle - 1 to account for CompareProperties_r incrementing handles.
+		const FRepLayoutCmd& Cmd = SharedParams.Cmds[Parent.CmdStart];
+		const uint16 Handle = Cmd.RelativeHandle;
+		CompareProperties_r(SharedParams, Parent.CmdStart, Parent.CmdEnd, ShadowData, Data, Changed, Handle - 1);
 	}
 }
 
@@ -1170,6 +1162,7 @@ bool FRepLayout::CompareProperties(
 		/*bForceFail=*/ false,
 		RoleIndex,
 		RemoteRoleIndex,
+		Flags,
 		Parents,
 		Cmds
 	};
@@ -2027,17 +2020,6 @@ void FRepLayout::SendProperties_r(
 			WritePropertyHandle(Writer, 0, bDoChecksum);		// Signify end of dynamic array
 			continue;
 		}
-		else if (RepState)
-		{
-			if (Cmd.ParentIndex == RoleIndex)
-			{
-				Data = reinterpret_cast<const uint8*>(&(RepState->SavedRole));
-			}
-			else if (Cmd.ParentIndex == RemoteRoleIndex)
-			{
-				Data = reinterpret_cast<const uint8*>(&(RepState->SavedRemoteRole));
-			}
-		}
 
 #if !(UE_BUILD_SHIPPING || UE_BUILD_TEST)
 		if (GDoReplicationContextString> 0)
@@ -2356,17 +2338,6 @@ void FRepLayout::SendProperties_BackwardsCompatible_r(
 			HandleIterator.ChangelistIterator.ChangedIndex++;
 			continue;
 		}
-		else
-		{
-			if (Cmd.ParentIndex == RoleIndex)
-			{
-				Data = reinterpret_cast<const uint8*>(&(RepState->SavedRole));
-			}
-			else if (Cmd.ParentIndex == RemoteRoleIndex)
-			{
-				Data = reinterpret_cast<const uint8*>(&(RepState->SavedRemoteRole));
-			}
-		}
 
 		WriteProperty_BackwardsCompatible(Writer, Cmd, HandleIterator.CmdIndex, Owner, Data, bDoChecksum);
 	}
@@ -2435,17 +2406,6 @@ void FRepLayout::SendAllProperties_BackwardsCompatible_r(
 			CmdIndex = Cmd.EndCmd - 1;		// The -1 to handle the ++ in the for loop
 			continue;
 		}
-		else
-		{
-			if (Cmd.ParentIndex == RoleIndex)
-			{
-				Data = reinterpret_cast<const uint8*>(&(RepState->SavedRole));
-			}
-			else if (Cmd.ParentIndex == RemoteRoleIndex)
-			{
-				Data = reinterpret_cast<const uint8*>(&(RepState->SavedRemoteRole));
-			}
-		}
 
 		WriteProperty_BackwardsCompatible(Writer, Cmd, CmdIndex, Owner, Data, bDoChecksum);
 	}
@@ -5017,6 +4977,10 @@ void FRepLayout::InitFromClass(
 	// Make sure it either found both, or didn't find either
 	check((RoleIndex == INDEX_NONE) == (RemoteRoleIndex == INDEX_NONE));
 
+	// Make sure that we only find these if we're an Actor, and if we're
+	// an Actor we always find these.
+	check((RoleIndex == INDEX_NONE) == !bIsObjectActor);
+
 	// This is so the receiving side can swap these as it receives them
 	if (RoleIndex != -1)
 	{
@@ -6518,6 +6482,7 @@ bool FRepLayout::DeltaSerializeFastArrayProperty(FFastArrayDeltaSerializeParams&
 								/*bForceFail=*/ bIsInitial || ShadowArrayItemIsNew[IDIndexPair.Idx],
 								RoleIndex,
 								RemoteRoleIndex,
+								Flags,
 								Parents,
 								Cmds
 							};
