Move implementation of FRepChangedPropertyTracker methods to CPP.

Trivial

diff --git a/Engine/Source/Runtime/Engine/Private/RepLayout.cpp b/Engine/Source/Runtime/Engine/Private/RepLayout.cpp
--- a/Engine/Source/Runtime/Engine/Private/RepLayout.cpp
+++ b/Engine/Source/Runtime/Engine/Private/RepLayout.cpp
@@ -719,6 +719,44 @@ static uint32 GetRepLayoutCmdCompatibleChecksum(
 	return CompatibleChecksum;
 }
 
+FRepChangedPropertyTracker::FRepChangedPropertyTracker(const bool InbIsReplay, const bool InbIsClientReplayRecording):
+		bIsReplay(InbIsReplay),
+		bIsClientReplayRecording(InbIsClientReplayRecording),
+		ExternalDataNumBits(0)
+	{}
+
+FRepChangedPropertyTracker::~FRepChangedPropertyTracker() {}
+
+void FRepChangedPropertyTracker::SetCustomIsActiveOverride(const uint16 RepIndex, const bool bIsActive)
+{
+	FRepChangedParent & Parent = Parents[RepIndex];
+
+	Parent.Active = (bIsActive || bIsClientReplayRecording) ? 1 : 0;
+	Parent.OldActive = Parent.Active;
+}
+
+void FRepChangedPropertyTracker::SetExternalData(const uint8* Src, const int32 NumBits)
+{
+	ExternalDataNumBits = NumBits;
+	const int32 NumBytes = (NumBits + 7) >> 3;
+	ExternalData.Reset(NumBytes);
+	ExternalData.AddUninitialized(NumBytes);
+	FMemory::Memcpy(ExternalData.GetData(), Src, NumBytes);
+}
+
+bool FRepChangedPropertyTracker::IsReplay() const
+{
+	return bIsReplay;
+}
+
+void FRepChangedPropertyTracker::CountBytes(FArchive& Ar) const
+{
+	// Include our size here, because the caller won't know.
+	Ar.CountBytes(sizeof(FRepChangedPropertyTracker), sizeof(FRepChangedPropertyTracker));
+	Parents.CountBytes(Ar);
+	ExternalData.CountBytes(Ar);
+
+}
 
 void FRepStateStaticBuffer::CountBytes(FArchive& Ar) const
 {
diff --git a/Engine/Source/Runtime/Engine/Public/Net/RepLayout.h b/Engine/Source/Runtime/Engine/Public/Net/RepLayout.h
--- a/Engine/Source/Runtime/Engine/Public/Net/RepLayout.h
+++ b/Engine/Source/Runtime/Engine/Public/Net/RepLayout.h
@@ -145,13 +145,10 @@ public:
 class FRepChangedPropertyTracker : public IRepChangedPropertyTracker
 {
 public:
-	FRepChangedPropertyTracker(const bool InbIsReplay, const bool InbIsClientReplayRecording):
-		bIsReplay(InbIsReplay),
-		bIsClientReplayRecording(InbIsClientReplayRecording),
-		ExternalDataNumBits(0)
-	{}
 
-	virtual ~FRepChangedPropertyTracker() {}
+	FRepChangedPropertyTracker(const bool InbIsReplay, const bool InbIsClientReplayRecording);
+
+	virtual ~FRepChangedPropertyTracker();
 
 	//~ Begin IRepChangedPropertyTracker Interface.
 	/**
@@ -163,13 +160,7 @@ public:
 	 * @param RepIndex	Replication index for the Property.
 	 * @param bIsActive	The new Active state.
 	 */
-	virtual void SetCustomIsActiveOverride(const uint16 RepIndex, const bool bIsActive) override
-	{
-		FRepChangedParent & Parent = Parents[RepIndex];
-
-		Parent.Active = (bIsActive || bIsClientReplayRecording) ? 1 : 0;
-		Parent.OldActive = Parent.Active;
-	}
+	virtual void SetCustomIsActiveOverride(const uint16 RepIndex, const bool bIsActive) override;
 
 	/**
 	 * Sets (or resets) the External Data.
@@ -179,29 +170,12 @@ public:
 	 * @param Src		Memory containing the external data.
 	 * @param NumBits	Size of the memory, in bits.
 	 */
-	virtual void SetExternalData(const uint8* Src, const int32 NumBits) override
-	{
-		ExternalDataNumBits = NumBits;
-		const int32 NumBytes = (NumBits + 7) >> 3;
-		ExternalData.Reset(NumBytes);
-		ExternalData.AddUninitialized(NumBytes);
-		FMemory::Memcpy(ExternalData.GetData(), Src, NumBytes);
-	}
+	virtual void SetExternalData(const uint8* Src, const int32 NumBits) override;
 
 	/** Whether or not this is being used for a replay (may be recording or playback). */
-	virtual bool IsReplay() const override
-	{
-		return bIsReplay;
-	}
+	virtual bool IsReplay() const override;
 
-	virtual void CountBytes(FArchive& Ar) const override
-	{
-		// Include our size here, because the caller won't know.
-		Ar.CountBytes(sizeof(FRepChangedPropertyTracker), sizeof(FRepChangedPropertyTracker));
-		Parents.CountBytes(Ar);
-		ExternalData.CountBytes(Ar);
-
-	}
+	virtual void CountBytes(FArchive& Ar) const override;
 	//~ End IRepChangedPropertyTracker Interface
 
 	/** Activation data for top level Properties on the given Actor / Object. */
