Change collision-override mesh for static meshes to be editor only to avoid cooking.

diff --git a/Engine/Source/Runtime/Engine/Classes/Engine/StaticMesh.h b/Engine/Source/Runtime/Engine/Classes/Engine/StaticMesh.h
--- a/Engine/Source/Runtime/Engine/Classes/Engine/StaticMesh.h
+++ b/Engine/Source/Runtime/Engine/Classes/Engine/StaticMesh.h
@@ -809,8 +809,11 @@ public:
 	UPROPERTY(Instanced, VisibleAnywhere, Category = EditableMesh)
 	class UObject* EditableMesh;
 
+#if WITH_EDITORONLY_DATA
 	UPROPERTY(EditAnywhere, Category = Collision)
 	class UStaticMesh* ComplexCollisionMesh;
+#endif
+
 	/**
 	 * Registers the mesh attributes required by the mesh description for a static mesh.
 	 */
diff --git a/Engine/Source/Runtime/Engine/Private/StaticMesh.cpp b/Engine/Source/Runtime/Engine/Private/StaticMesh.cpp
--- a/Engine/Source/Runtime/Engine/Private/StaticMesh.cpp
+++ b/Engine/Source/Runtime/Engine/Private/StaticMesh.cpp
@@ -2893,6 +2893,7 @@ void UStaticMesh::PostEditChangeProperty(FPropertyChangedEvent& PropertyChangedE
 		bool bRebuild = false;
 		SetLODGroup(LODGroup, bRebuild);
 	}
+#if WITH_EDITORONLY_DATA
 	if (PropertyName == GET_MEMBER_NAME_CHECKED(UStaticMesh, ComplexCollisionMesh) && ComplexCollisionMesh != this)
 	{
 		if (BodySetup)
@@ -2901,6 +2902,8 @@ void UStaticMesh::PostEditChangeProperty(FPropertyChangedEvent& PropertyChangedE
 			BodySetup->CreatePhysicsMeshes();
 		}
 	}
+#endif
+
 	LightMapResolution = FMath::Max(LightMapResolution, 0);
 
 	if (PropertyChangedEvent.MemberProperty 
