#ue4 - Add some debug info about jump state to "p.visualizemovement"

diff --git a/Engine/Source/Runtime/Engine/Private/Components/CharacterMovementComponent.cpp b/Engine/Source/Runtime/Engine/Private/Components/CharacterMovementComponent.cpp
--- a/Engine/Source/Runtime/Engine/Private/Components/CharacterMovementComponent.cpp
+++ b/Engine/Source/Runtime/Engine/Private/Components/CharacterMovementComponent.cpp
@@ -6917,6 +6917,15 @@ float UCharacterMovementComponent::VisualizeMovement() const
 		}
 	}
 
+	// Jump
+	{
+		const FColor DebugColor = FColor::Blue;
+		HeightOffset += 15.f;
+		FVector DebugLocation = TopOfCapsule + FVector(0.f, 0.f, HeightOffset);
+		FString DebugText = FString::Printf(TEXT("bIsJumping: %d Count: %d HoldTime: %.2f"), CharacterOwner->bPressedJump, CharacterOwner->JumpCurrentCount, CharacterOwner->JumpKeyHoldTime);
+		DrawDebugString(GetWorld(), DebugLocation, DebugText, nullptr, DebugColor, 0.f, true);
+	}
+
 	// Root motion (additive)
 	if (CurrentRootMotion.HasAdditiveVelocity())
 	{
