Compilation fix for 8405993.

diff --git a/Engine/Source/Runtime/Engine/Private/Components/SkyAtmosphereComponent.cpp b/Engine/Source/Runtime/Engine/Private/Components/SkyAtmosphereComponent.cpp
--- a/Engine/Source/Runtime/Engine/Private/Components/SkyAtmosphereComponent.cpp
+++ b/Engine/Source/Runtime/Engine/Private/Components/SkyAtmosphereComponent.cpp
@@ -42,7 +42,7 @@ USkyAtmosphereComponent::USkyAtmosphereComponent(const FObjectInitializer& Objec
 	// Default: Earth like atmosphere
 	BottomRadius = EarthBottomRadius;
 	AtmosphereHeight = EarthTopRadius - EarthBottomRadius;
-	GroundAlbedo = FColor(0.4f, 0.4f, 0.4f);
+	GroundAlbedo = FColor(170, 170, 170); // 170 => 0.4f linear
 
 	// FLoat to a u8 rgb + float length can lose some precision but it is better UI wise.
 	const FLinearColor RayleightScatteringRaw = FLinearColor(0.005802f, 0.013558f, 0.033100f);
