Fix crash when building layered animation blueprints

Ungrouped layers no longer attempt to create a single instance per 'group' (as there isnt really a group at all). This means that init/update calls are now correctly paired.
This unblocks us to work on developing the layering system more, but needs discussion as to how we develop this in the future.
Also fix odd alightment issue with layer interface details.

diff --git a/Engine/Source/Editor/AnimGraph/Private/AnimGraphDetails.cpp b/Engine/Source/Editor/AnimGraph/Private/AnimGraphDetails.cpp
--- a/Engine/Source/Editor/AnimGraph/Private/AnimGraphDetails.cpp
+++ b/Engine/Source/Editor/AnimGraph/Private/AnimGraphDetails.cpp
@@ -209,6 +209,7 @@ void FAnimGraphDetails::CustomizeDetails(IDetailLayoutBuilder& DetailLayout)
 					]
 				]
 				.ValueContent()
+				.MaxDesiredWidth(250.0f)
 				[
 					SNew(SBox)
 					.Padding(2.0f)
@@ -216,7 +217,7 @@ void FAnimGraphDetails::CustomizeDetails(IDetailLayoutBuilder& DetailLayout)
 						SNew(SHorizontalBox)
 						+SHorizontalBox::Slot()
 						.VAlign(VAlign_Center)
-						.AutoWidth()
+						.FillWidth(1.0f)
 						[
 							SNew(STextBlock)
 							.Text_Lambda(GetSubInputLabel)
@@ -226,7 +227,7 @@ void FAnimGraphDetails::CustomizeDetails(IDetailLayoutBuilder& DetailLayout)
 						.Padding(4.0f, 0.0f, 0.0f, 0.0f)
 						.VAlign(VAlign_Center)
 						.HAlign(HAlign_Right)
-						.FillWidth(1.0f)
+						.AutoWidth()
 						[
 							SNew(SButton)
 							.ButtonStyle(FEditorStyle::Get(), "HoverHintOnly")
diff --git a/Engine/Source/Runtime/Engine/Classes/Animation/AnimInstance.h b/Engine/Source/Runtime/Engine/Classes/Animation/AnimInstance.h
--- a/Engine/Source/Runtime/Engine/Classes/Animation/AnimInstance.h
+++ b/Engine/Source/Runtime/Engine/Classes/Animation/AnimInstance.h
@@ -718,8 +718,9 @@ public:
 	void SetSubInstanceClassByTag(FName InTag, TSubclassOf<UAnimInstance> InClass);
 
 	/** 
-	 * Runs through all layer nodes, attempting to layer nodes that are implemented by the specified class, then sets up a sub instance of the class for each.
-	 * Allocates one sub instance to run each of the groups specified in the class, so state is shared.
+	 * Runs through all layer nodes, attempting to find layer nodes that are implemented by the specified class, then sets up a sub instance of the class for each.
+	 * Allocates one sub instance to run each of the groups specified in the class, so state is shared. If a layer is not grouped (ie. NAME_None), then state is not shared
+	 * and a separate sub-instance is allocated for each layer node.
 	 * If InClass is null, then layers are reset to their defaults.
 	 */
 	UFUNCTION(BlueprintCallable, Category = "Layers")
diff --git a/Engine/Source/Runtime/Engine/Private/Animation/AnimInstance.cpp b/Engine/Source/Runtime/Engine/Private/Animation/AnimInstance.cpp
--- a/Engine/Source/Runtime/Engine/Private/Animation/AnimInstance.cpp
+++ b/Engine/Source/Runtime/Engine/Private/Animation/AnimInstance.cpp
@@ -2569,39 +2569,69 @@ void UAnimInstance::SetLayerOverlay(TSubclassOf<UAnimInstance> InClass)
 		{
 			for(TPair<FName, TArray<FAnimNode_Layer*, TInlineAllocator<4>>> LayerPair : LayerNodesToSet)
 			{
-				// If the class is null, then reset to default (which can be null)
-				UClass* ClassToSet = NewClass != nullptr ? NewClass : LayerPair.Value[0]->InstanceClass.Get();
-				if(ClassToSet != nullptr && ClassToSet != GetClass())
+				if (LayerPair.Key == NAME_None)
 				{
-					// Create and add one sub-instance for this group
-					USkeletalMeshComponent* MeshComp = GetSkelMeshComponent();
-					UAnimInstance* NewSubInstance = NewObject<UAnimInstance>(MeshComp, ClassToSet);
-
+					// Ungrouped path - each layer gets a separate instance
 					for(FAnimNode_Layer* LayerNode : LayerPair.Value)
 					{
-						LayerNode->SetLayerOverlaySubInstance(this, NewSubInstance);
-					}
+						UClass* ClassToSet = NewClass != nullptr ? NewClass : LayerNode->InstanceClass.Get();
+						if (ClassToSet != nullptr && ClassToSet != GetClass())
+						{
+							USkeletalMeshComponent* MeshComp = GetSkelMeshComponent();
+							UAnimInstance* NewSubInstance = NewObject<UAnimInstance>(MeshComp, ClassToSet);
 
-					// Init after we link in the new graph segments, so propagation happens correctly
-					NewSubInstance->InitializeAnimation();
+							LayerNode->SetLayerOverlaySubInstance(this, NewSubInstance);
 
-					// Initialize the correct parts of the sub instance
-					for(FAnimNode_Layer* LayerNode : LayerPair.Value)
-					{
-						if(LayerNode->LinkedRoot)
-						{
-							NewSubInstance->GetProxyOnAnyThread<FAnimInstanceProxy>().InitializeRootNode_WithRoot(LayerNode->LinkedRoot);
+							// Init after we link in the new graph segments, so propagation happens correctly
+							NewSubInstance->InitializeAnimation();
+
+							// Initialize the correct parts of the sub instance
+							if(LayerNode->LinkedRoot)
+							{
+								NewSubInstance->GetProxyOnAnyThread<FAnimInstanceProxy>().InitializeRootNode_WithRoot(LayerNode->LinkedRoot);
+							}
+
+							MeshComp->SubInstances.Add(NewSubInstance);
 						}
 					}
-
-					MeshComp->SubInstances.Add(NewSubInstance);
 				}
 				else
-				{
-					// Clear the node's instance - we didnt find a class to use
-					for(FAnimNode_Layer* LayerNode : LayerPair.Value)
+				{		
+					// Grouped path - each group gets an instance
+					// If the class is null, then reset to default (which can be null)
+					UClass* ClassToSet = NewClass != nullptr ? NewClass : LayerPair.Value[0]->InstanceClass.Get();
+					if(ClassToSet != nullptr && ClassToSet != GetClass())
 					{
-						LayerNode->SetLayerOverlaySubInstance(this, nullptr);
+						// Create and add one sub-instance for this group
+						USkeletalMeshComponent* MeshComp = GetSkelMeshComponent();
+						UAnimInstance* NewSubInstance = NewObject<UAnimInstance>(MeshComp, ClassToSet);
+
+						for(FAnimNode_Layer* LayerNode : LayerPair.Value)
+						{
+							LayerNode->SetLayerOverlaySubInstance(this, NewSubInstance);
+						}
+
+						// Init after we link in the new graph segments, so propagation happens correctly
+						NewSubInstance->InitializeAnimation();
+
+						// Initialize the correct parts of the sub instance
+						for(FAnimNode_Layer* LayerNode : LayerPair.Value)
+						{
+							if(LayerNode->LinkedRoot)
+							{
+								NewSubInstance->GetProxyOnAnyThread<FAnimInstanceProxy>().InitializeRootNode_WithRoot(LayerNode->LinkedRoot);
+							}
+						}
+
+						MeshComp->SubInstances.Add(NewSubInstance);
+					}
+					else
+					{
+						// Clear the node's instance - we didnt find a class to use
+						for(FAnimNode_Layer* LayerNode : LayerPair.Value)
+						{
+							LayerNode->SetLayerOverlaySubInstance(this, nullptr);
+						}
 					}
 				}
 			}
