Adding CSV perf stats

diff --git a/Engine/Source/Runtime/Engine/Private/Collision/SceneQuery.cpp b/Engine/Source/Runtime/Engine/Private/Collision/SceneQuery.cpp
--- a/Engine/Source/Runtime/Engine/Private/Collision/SceneQuery.cpp
+++ b/Engine/Source/Runtime/Engine/Private/Collision/SceneQuery.cpp
@@ -440,6 +440,7 @@ bool FGenericPhysicsInterface::RaycastTest(const UWorld* World, const FVector St
 {
 	SCOPE_CYCLE_COUNTER(STAT_Collision_SceneQueryTotal);
 	SCOPE_CYCLE_COUNTER(STAT_Collision_RaycastAny);
+	CSV_SCOPED_TIMING_STAT_EXCLUSIVE(RaycastTest);
 
 	using TCastTraits = TSQTraits<FHitRaycast, ESweepOrRay::Raycast, ESingleMultiOrTest::Test>;
 	FHitResult DummyHit(NoInit);
@@ -450,6 +451,7 @@ bool FGenericPhysicsInterface::RaycastSingle(const UWorld* World, struct FHitRes
 {
 	SCOPE_CYCLE_COUNTER(STAT_Collision_SceneQueryTotal);
 	SCOPE_CYCLE_COUNTER(STAT_Collision_RaycastSingle);
+	CSV_SCOPED_TIMING_STAT_EXCLUSIVE(RaycastSingle);
 
 	using TCastTraits = TSQTraits<FHitRaycast, ESweepOrRay::Raycast, ESingleMultiOrTest::Single>;
 	return TSceneCastCommon<TCastTraits>(World, OutHit, FRaycastSQAdditionalInputs(), Start, End, TraceChannel, Params, ResponseParams, ObjectParams);
@@ -461,6 +463,7 @@ bool FGenericPhysicsInterface::RaycastMulti(const UWorld* World, TArray<struct F
 {
 	SCOPE_CYCLE_COUNTER(STAT_Collision_SceneQueryTotal);
 	SCOPE_CYCLE_COUNTER(STAT_Collision_RaycastMultiple);
+	CSV_SCOPED_TIMING_STAT_EXCLUSIVE(RaycastMultiple);
 
 	using TCastTraits = TSQTraits<FHitRaycast, ESweepOrRay::Raycast, ESingleMultiOrTest::Multi>;
 	return TSceneCastCommon<TCastTraits> (World, OutHits, FRaycastSQAdditionalInputs(), Start, End, TraceChannel, Params, ResponseParams, ObjectParams);
@@ -473,6 +476,7 @@ bool FGenericPhysicsInterface::GeomSweepTest(const UWorld* World, const struct F
 {
 	SCOPE_CYCLE_COUNTER(STAT_Collision_SceneQueryTotal);
 	SCOPE_CYCLE_COUNTER(STAT_Collision_GeomSweepAny);
+	CSV_SCOPED_TIMING_STAT_EXCLUSIVE(GeomSweepTest);
 
 	using TCastTraits = TSQTraits<FHitSweep, ESweepOrRay::Sweep, ESingleMultiOrTest::Test>;
 	FHitResult DummyHit(NoInit);
@@ -483,6 +487,7 @@ bool FGenericPhysicsInterface::GeomSweepSingle(const UWorld* World, const struct
 {
 	SCOPE_CYCLE_COUNTER(STAT_Collision_SceneQueryTotal);
 	SCOPE_CYCLE_COUNTER(STAT_Collision_GeomSweepSingle);
+	CSV_SCOPED_TIMING_STAT_EXCLUSIVE(GeomSweepSingle);
 
 	using TCastTraits = TSQTraits<FHitSweep, ESweepOrRay::Sweep, ESingleMultiOrTest::Single>;
 	return TSceneCastCommon<TCastTraits>(World, OutHit, FGeomSQAdditionalInputs(CollisionShape, Rot), Start, End, TraceChannel, Params, ResponseParams, ObjectParams);
@@ -493,6 +498,7 @@ bool FGenericPhysicsInterface::GeomSweepMulti(const UWorld* World, const FPhysic
 {
 	SCOPE_CYCLE_COUNTER(STAT_Collision_SceneQueryTotal);
 	SCOPE_CYCLE_COUNTER(STAT_Collision_GeomSweepMultiple);
+	CSV_SCOPED_TIMING_STAT_EXCLUSIVE(GeomSweepMultiple);
 
 	using TCastTraits = TSQTraits<FHitSweep, ESweepOrRay::Sweep, ESingleMultiOrTest::Multi>;
 	return TSceneCastCommon<TCastTraits>(World, OutHits, FGeomCollectionSQAdditionalInputs(InGeom, InGeomRot), Start, End, TraceChannel, Params, ResponseParams, ObjectParams);
@@ -503,6 +509,7 @@ bool FGenericPhysicsInterface::GeomSweepMulti(const UWorld* World, const FCollis
 {
 	SCOPE_CYCLE_COUNTER(STAT_Collision_SceneQueryTotal);
 	SCOPE_CYCLE_COUNTER(STAT_Collision_GeomSweepMultiple);
+	CSV_SCOPED_TIMING_STAT_EXCLUSIVE(GeomSweepMultiple);
 
 	using TCastTraits = TSQTraits<FHitSweep, ESweepOrRay::Sweep, ESingleMultiOrTest::Multi>;
 	return TSceneCastCommon<TCastTraits>(World, OutHits, FGeomSQAdditionalInputs(InGeom, InGeomRot), Start, End, TraceChannel, Params, ResponseParams, ObjectParams);
@@ -610,6 +617,7 @@ bool FGenericPhysicsInterface::GeomOverlapBlockingTest(const UWorld* World, cons
 {
 	SCOPE_CYCLE_COUNTER(STAT_Collision_SceneQueryTotal);
 	SCOPE_CYCLE_COUNTER(STAT_Collision_GeomOverlapBlocking);
+	CSV_SCOPED_TIMING_STAT_EXCLUSIVE(GeomOverlapBlocking);
 
 	TArray<FOverlapResult> Overlaps;	//needed only for template shared code
 	FTransform GeomTransform(Rot, Pos);
@@ -625,6 +633,7 @@ bool FGenericPhysicsInterface::GeomOverlapAnyTest(const UWorld* World, const str
 {
 	SCOPE_CYCLE_COUNTER(STAT_Collision_SceneQueryTotal);
 	SCOPE_CYCLE_COUNTER(STAT_Collision_GeomOverlapAny);
+	CSV_SCOPED_TIMING_STAT_EXCLUSIVE(GeomOverlapAny);
 
 	TArray<FOverlapResult> Overlaps;	//needed only for template shared code
 	FTransform GeomTransform(Rot, Pos);
@@ -641,6 +650,7 @@ bool FGenericPhysicsInterface::GeomOverlapMulti(const UWorld* World, const FPhys
 {
 	SCOPE_CYCLE_COUNTER(STAT_Collision_SceneQueryTotal);
 	SCOPE_CYCLE_COUNTER(STAT_Collision_GeomOverlapMultiple);
+	CSV_SCOPED_TIMING_STAT_EXCLUSIVE(GeomOverlapMultiple);
 
 #if WITH_PHYSX && !WITH_CHAOS_NEEDS_TO_BE_FIXED
 	FTransform GeomTransform(InRotation, InPosition);
@@ -655,6 +665,7 @@ bool FGenericPhysicsInterface::GeomOverlapMulti(const UWorld* World, const FColl
 {
 	SCOPE_CYCLE_COUNTER(STAT_Collision_SceneQueryTotal);
 	SCOPE_CYCLE_COUNTER(STAT_Collision_GeomOverlapMultiple);
+	CSV_SCOPED_TIMING_STAT_EXCLUSIVE(GeomOverlapMultiple);
 
 	FTransform GeomTransform(InRotation, InPosition);
 #if WITH_PHYSX && !WITH_CHAOS_NEEDS_TO_BE_FIXED
diff --git a/Engine/Source/Runtime/Engine/Private/Components/CharacterMovementComponent.cpp b/Engine/Source/Runtime/Engine/Private/Components/CharacterMovementComponent.cpp
--- a/Engine/Source/Runtime/Engine/Private/Components/CharacterMovementComponent.cpp
+++ b/Engine/Source/Runtime/Engine/Private/Components/CharacterMovementComponent.cpp
@@ -4150,6 +4150,7 @@ float UCharacterMovementComponent::BoostAirControl(float DeltaTime, float TickAi
 void UCharacterMovementComponent::PhysFalling(float deltaTime, int32 Iterations)
 {
 	SCOPE_CYCLE_COUNTER(STAT_CharPhysFalling);
+	CSV_SCOPED_TIMING_STAT_EXCLUSIVE(CharPhysFalling);
 
 	if (deltaTime < MIN_TICK_TIME)
 	{
@@ -4775,7 +4776,7 @@ void UCharacterMovementComponent::MaintainHorizontalGroundVelocity()
 void UCharacterMovementComponent::PhysWalking(float deltaTime, int32 Iterations)
 {
 	SCOPE_CYCLE_COUNTER(STAT_CharPhysWalking);
-
+	CSV_SCOPED_TIMING_STAT_EXCLUSIVE(CharPhysWalking);
 
 	if (deltaTime < MIN_TICK_TIME)
 	{
diff --git a/Engine/Source/Runtime/Engine/Private/Components/PrimitiveComponent.cpp b/Engine/Source/Runtime/Engine/Private/Components/PrimitiveComponent.cpp
--- a/Engine/Source/Runtime/Engine/Private/Components/PrimitiveComponent.cpp
+++ b/Engine/Source/Runtime/Engine/Private/Components/PrimitiveComponent.cpp
@@ -1965,6 +1965,7 @@ FCollisionShape UPrimitiveComponent::GetCollisionShape(float Inflation) const
 bool UPrimitiveComponent::MoveComponentImpl( const FVector& Delta, const FQuat& NewRotationQuat, bool bSweep, FHitResult* OutHit, EMoveComponentFlags MoveFlags, ETeleportType Teleport)
 {
 	SCOPE_CYCLE_COUNTER(STAT_MoveComponentTime);
+	CSV_SCOPED_TIMING_STAT_EXCLUSIVE(MoveComponentTime);
 
 #if !(UE_BUILD_SHIPPING || UE_BUILD_TEST) && PERF_MOVECOMPONENT_STATS
 	FScopedMoveCompTimer MoveTimer(this, Delta);
diff --git a/Engine/Source/Runtime/NavigationSystem/Private/NavigationSystem.cpp b/Engine/Source/Runtime/NavigationSystem/Private/NavigationSystem.cpp
--- a/Engine/Source/Runtime/NavigationSystem/Private/NavigationSystem.cpp
+++ b/Engine/Source/Runtime/NavigationSystem/Private/NavigationSystem.cpp
@@ -1056,6 +1056,7 @@ void UNavigationSystemV1::SetNavigationAutoUpdateEnabled(bool bNewEnable, UNavig
 FPathFindingResult UNavigationSystemV1::FindPathSync(const FNavAgentProperties& AgentProperties, FPathFindingQuery Query, EPathFindingMode::Type Mode)
 {
 	SCOPE_CYCLE_COUNTER(STAT_Navigation_PathfindingSync);
+	CSV_SCOPED_TIMING_STAT_EXCLUSIVE(PathfindingSync);
 
 	if (Query.NavData.IsValid() == false)
 	{
@@ -1081,6 +1082,7 @@ FPathFindingResult UNavigationSystemV1::FindPathSync(const FNavAgentProperties&
 FPathFindingResult UNavigationSystemV1::FindPathSync(FPathFindingQuery Query, EPathFindingMode::Type Mode)
 {
 	SCOPE_CYCLE_COUNTER(STAT_Navigation_PathfindingSync);
+	CSV_SCOPED_TIMING_STAT_EXCLUSIVE(PathfindingSync);
 
 	if (Query.NavData.IsValid() == false)
 	{
@@ -1106,6 +1108,7 @@ FPathFindingResult UNavigationSystemV1::FindPathSync(FPathFindingQuery Query, EP
 bool UNavigationSystemV1::TestPathSync(FPathFindingQuery Query, EPathFindingMode::Type Mode, int32* NumVisitedNodes) const
 {
 	SCOPE_CYCLE_COUNTER(STAT_Navigation_PathfindingSync);
+	CSV_SCOPED_TIMING_STAT_EXCLUSIVE(PathfindingSync);
 
 	if (Query.NavData.IsValid() == false)
 	{
@@ -1194,12 +1197,15 @@ void UNavigationSystemV1::TriggerAsyncQueries(TArray<FAsyncPathFindingQuery>& Pa
 
 static void AsyncQueryDone(FAsyncPathFindingQuery Query)
 {
+	CSV_SCOPED_TIMING_STAT_EXCLUSIVE(AsyncNavQueryFinished);
+
 	Query.OnDoneDelegate.ExecuteIfBound(Query.QueryID, Query.Result.Result, Query.Result.Path);
 }
 
 void UNavigationSystemV1::PerformAsyncQueries(TArray<FAsyncPathFindingQuery> PathFindingQueries)
 {
 	SCOPE_CYCLE_COUNTER(STAT_Navigation_PathfindingAsync);
+	CSV_SCOPED_TIMING_STAT_EXCLUSIVE(PathfindingAsync);
 
 	if (PathFindingQueries.Num() == 0)
 	{
