Added UsedWithWater material flag

sebastien.hillaire
#rnx

diff --git a/Engine/Source/Runtime/Engine/Classes/Materials/Material.h b/Engine/Source/Runtime/Engine/Classes/Materials/Material.h
--- a/Engine/Source/Runtime/Engine/Classes/Materials/Material.h
+++ b/Engine/Source/Runtime/Engine/Classes/Materials/Material.h
@@ -725,6 +725,13 @@ public:
 	UPROPERTY(EditAnywhere, BlueprintReadOnly, Category=Usage)
 	uint32 bUsedWithClothing:1;
 
+	/**
+	 * Indicates that the material and its instances can be use with water
+	 * This will result in the shaders required to support water meshes being compiled which will increase shader compile time and memory usage.
+	 */
+	UPROPERTY(EditAnywhere, BlueprintReadOnly, Category = Usage)
+	uint32 bUsedWithWater : 1;
+
 	/** 
 	 * Indicates that the material and its instances can be used with Slate UI and UMG
 	 * This will result in the shaders required to support UI materials being compiled which will increase shader compile time and memory usage.
diff --git a/Engine/Source/Runtime/Engine/Classes/Materials/MaterialInterface.h b/Engine/Source/Runtime/Engine/Classes/Materials/MaterialInterface.h
--- a/Engine/Source/Runtime/Engine/Classes/Materials/MaterialInterface.h
+++ b/Engine/Source/Runtime/Engine/Classes/Materials/MaterialInterface.h
@@ -52,6 +52,7 @@ enum EMaterialUsage
 	MATUSAGE_NiagaraRibbons,
 	MATUSAGE_NiagaraMeshParticles,
 	MATUSAGE_GeometryCache,
+	MATUSAGE_Water,
 
 	MATUSAGE_MAX,
 };
diff --git a/Engine/Source/Runtime/Engine/Private/Materials/Material.cpp b/Engine/Source/Runtime/Engine/Private/Materials/Material.cpp
--- a/Engine/Source/Runtime/Engine/Private/Materials/Material.cpp
+++ b/Engine/Source/Runtime/Engine/Private/Materials/Material.cpp
@@ -1262,6 +1262,7 @@ bool UMaterial::GetUsageByFlag(EMaterialUsage Usage) const
 		case MATUSAGE_GeometryCollections: UsageValue = bUsedWithGeometryCollections; break;
 		case MATUSAGE_Clothing: UsageValue = bUsedWithClothing; break;
 		case MATUSAGE_GeometryCache: UsageValue = bUsedWithGeometryCache; break;
+		case MATUSAGE_Water: UsageValue = bUsedWithWater; break;
 		default: UE_LOG(LogMaterial, Fatal,TEXT("Unknown material usage: %u"), (int32)Usage);
 	};
 	return UsageValue;
@@ -1392,6 +1393,10 @@ void UMaterial::SetUsageByFlag(EMaterialUsage Usage, bool NewValue)
 		{
 			bUsedWithGeometryCache = NewValue; break;
 		}
+		case MATUSAGE_Water:
+		{
+			bUsedWithWater = NewValue; break;
+		}
 		default: UE_LOG(LogMaterial, Fatal,TEXT("Unknown material usage: %u"), (int32)Usage);
 	};
 #if WITH_EDITOR
@@ -1419,6 +1424,7 @@ FString UMaterial::GetUsageName(EMaterialUsage Usage) const
 		case MATUSAGE_GeometryCollections: UsageName = TEXT("bUsedWithGeometryCollections"); break;
 		case MATUSAGE_Clothing: UsageName = TEXT("bUsedWithClothing"); break;
 		case MATUSAGE_GeometryCache: UsageName = TEXT("bUsedWithGeometryCache"); break;
+		case MATUSAGE_Water: UsageName = TEXT("bUsedWithWater"); break;
 		default: UE_LOG(LogMaterial, Fatal,TEXT("Unknown material usage: %u"), (int32)Usage);
 	};
 	return UsageName;
diff --git a/Engine/Source/Runtime/Engine/Private/Materials/MaterialShared.cpp b/Engine/Source/Runtime/Engine/Private/Materials/MaterialShared.cpp
--- a/Engine/Source/Runtime/Engine/Private/Materials/MaterialShared.cpp
+++ b/Engine/Source/Runtime/Engine/Private/Materials/MaterialShared.cpp
@@ -1104,6 +1104,11 @@ bool FMaterialResource::IsUsedWithGeometryCache() const
 	return Material->bUsedWithGeometryCache;
 }
 
+bool FMaterialResource::IsUsedWithWater() const
+{
+	return Material->bUsedWithWater;
+}
+
 bool FMaterialResource::IsUsedWithLandscape() const
 {
 	return false;
diff --git a/Engine/Source/Runtime/Engine/Public/MaterialShared.h b/Engine/Source/Runtime/Engine/Public/MaterialShared.h
--- a/Engine/Source/Runtime/Engine/Public/MaterialShared.h
+++ b/Engine/Source/Runtime/Engine/Public/MaterialShared.h
@@ -1458,6 +1458,7 @@ public:
 	virtual bool IsUsedWithAPEXCloth() const { return false; }
 	virtual bool IsUsedWithUI() const { return false; }
 	virtual bool IsUsedWithGeometryCache() const { return false; }
+	virtual bool IsUsedWithWater() const { return false; }
 	ENGINE_API virtual enum EMaterialTessellationMode GetTessellationMode() const;
 	virtual bool IsCrackFreeDisplacementEnabled() const { return false; }
 	virtual bool IsAdaptiveTessellationEnabled() const { return false; }
@@ -2206,6 +2207,7 @@ public:
 	ENGINE_API virtual bool IsUsedWithGeometryCollections() const override;
 	ENGINE_API virtual bool IsUsedWithAPEXCloth() const override;
 	ENGINE_API virtual bool IsUsedWithGeometryCache() const override;
+	ENGINE_API virtual bool IsUsedWithWater() const override;
 	ENGINE_API virtual enum EMaterialTessellationMode GetTessellationMode() const override;
 	ENGINE_API virtual bool IsCrackFreeDisplacementEnabled() const override;
 	ENGINE_API virtual bool IsAdaptiveTessellationEnabled() const override;
