Error checking and fixup for possibly corrupt indices.

diff --git a/Engine/Source/Runtime/Engine/Private/Particles/ParticleEmitterInstances.cpp b/Engine/Source/Runtime/Engine/Private/Particles/ParticleEmitterInstances.cpp
--- a/Engine/Source/Runtime/Engine/Private/Particles/ParticleEmitterInstances.cpp
+++ b/Engine/Source/Runtime/Engine/Private/Particles/ParticleEmitterInstances.cpp
@@ -2468,27 +2468,42 @@ void FParticleEmitterInstance::KillParticles()
 			}
 		}
 
+		bool bFoundCorruptIndices = false;
 		// Loop over the active particles... If their RelativeTime is > 1.0f (indicating they are dead),
 		// move them to the 'end' of the active particle list.
 		for (int32 i=ActiveParticles-1; i>=0; i--)
 		{
 			const int32	CurrentIndex	= ParticleIndices[i];
-			const uint8* ParticleBase	= ParticleData + CurrentIndex * ParticleStride;
-			FBaseParticle& Particle		= *((FBaseParticle*) ParticleBase);
-			if (Particle.RelativeTime > 1.0f)
-			{
-				if (EventPayload)
+			if (ensure(CurrentIndex < MaxActiveParticles))
+			{ 
+				const uint8* ParticleBase = ParticleData + CurrentIndex * ParticleStride;
+				FBaseParticle& Particle = *((FBaseParticle*)ParticleBase);
+
+				if (Particle.RelativeTime > 1.0f)
 				{
-					LODLevel->EventGenerator->HandleParticleKilled(this, EventPayload, &Particle);
-				}
-				// Move it to the 'back' of the list
-				ParticleIndices[i] = ParticleIndices[ActiveParticles-1];
-				ParticleIndices[ActiveParticles-1]	= CurrentIndex;
-				ActiveParticles--;
+					if (EventPayload)
+					{
+						LODLevel->EventGenerator->HandleParticleKilled(this, EventPayload, &Particle);
+					}
+					// Move it to the 'back' of the list
+					ParticleIndices[i] = ParticleIndices[ActiveParticles-1];
+					ParticleIndices[ActiveParticles-1]	= CurrentIndex;
+					ActiveParticles--;
 
-				INC_DWORD_STAT(STAT_SpriteParticlesKilled);
+					INC_DWORD_STAT(STAT_SpriteParticlesKilled);
+				}
+			}
+			else
+			{
+				bFoundCorruptIndices = true;
 			}
 		}
+
+		if (bFoundCorruptIndices)
+		{
+			UE_LOG(LogParticles, Error, TEXT("Detected corrupted particle indices. Template : %s, Component %s"), *GetNameSafe(Component && Component->IsValidLowLevel() ? Component->Template : nullptr), *GetFullNameSafe(Component));			
+			FixupParticleIndices();
+		}
 	}
 }
 
