#ue4 - Pull in debug text tighter together for "p.visualizemovement 1"

diff --git a/Engine/Source/Runtime/Engine/Private/Components/CharacterMovementComponent.cpp b/Engine/Source/Runtime/Engine/Private/Components/CharacterMovementComponent.cpp
--- a/Engine/Source/Runtime/Engine/Private/Components/CharacterMovementComponent.cpp
+++ b/Engine/Source/Runtime/Engine/Private/Components/CharacterMovementComponent.cpp
@@ -6857,6 +6857,7 @@ void UCharacterMovementComponent::DisplayDebug(UCanvas* Canvas, const FDebugDisp
 float UCharacterMovementComponent::VisualizeMovement() const
 {
 	float HeightOffset = 0.f;
+	const float OffsetPerElement = 10.0f;
 	if (CharacterOwner == nullptr)
 	{
 		return HeightOffset;
@@ -6876,42 +6877,42 @@ float UCharacterMovementComponent::VisualizeMovement() const
 	// Velocity
 	{
 		const FColor DebugColor = FColor::Green;
-		HeightOffset += 15.f;
+		HeightOffset += OffsetPerElement;
 		const FVector DebugLocation = TopOfCapsule + FVector(0.f,0.f,HeightOffset);
-		DrawDebugDirectionalArrow(GetWorld(), DebugLocation, DebugLocation + Velocity, 
+		DrawDebugDirectionalArrow(GetWorld(), DebugLocation - FVector(0.f, 0.f, 5.0f), DebugLocation - FVector(0.f, 0.f, 5.0f) + Velocity,
 			100.f, DebugColor, false, -1.f, (uint8)'\000', 10.f);
 
 		FString DebugText = FString::Printf(TEXT("Velocity: %s (Speed: %.2f) (Max: %.2f)"), *Velocity.ToCompactString(), Velocity.Size(), GetMaxSpeed());
-		DrawDebugString(GetWorld(), DebugLocation + FVector(0.f,0.f,5.f), DebugText, nullptr, DebugColor, 0.f, true);
+		DrawDebugString(GetWorld(), DebugLocation, DebugText, nullptr, DebugColor, 0.f, true);
 	}
 
 	// Acceleration
 	{
 		const FColor DebugColor = FColor::Yellow;
-		HeightOffset += 15.f;
+		HeightOffset += OffsetPerElement;
 		const float MaxAccelerationLineLength = 200.f;
 		const float CurrentMaxAccel = GetMaxAcceleration();
 		const float CurrentAccelAsPercentOfMaxAccel = CurrentMaxAccel > 0.f ? Acceleration.Size() / CurrentMaxAccel : 1.f;
 		const FVector DebugLocation = TopOfCapsule + FVector(0.f,0.f,HeightOffset);
-		DrawDebugDirectionalArrow(GetWorld(), DebugLocation, 
-			DebugLocation + Acceleration.GetSafeNormal(SMALL_NUMBER) * CurrentAccelAsPercentOfMaxAccel * MaxAccelerationLineLength, 
+		DrawDebugDirectionalArrow(GetWorld(), DebugLocation - FVector(0.f, 0.f, 5.0f), 
+			DebugLocation - FVector(0.f, 0.f, 5.0f) + Acceleration.GetSafeNormal(SMALL_NUMBER) * CurrentAccelAsPercentOfMaxAccel * MaxAccelerationLineLength,
 			25.f, DebugColor, false, -1.f, (uint8)'\000', 8.f);
 
 		FString DebugText = FString::Printf(TEXT("Acceleration: %s"), *Acceleration.ToCompactString());
-		DrawDebugString(GetWorld(), DebugLocation + FVector(0.f,0.f,5.f), DebugText, nullptr, DebugColor, 0.f, true);
+		DrawDebugString(GetWorld(), DebugLocation, DebugText, nullptr, DebugColor, 0.f, true);
 	}
 
 	// Movement Mode
 	{
 		const FColor DebugColor = FColor::Blue;
-		HeightOffset += 20.f;
+		HeightOffset += OffsetPerElement;
 		FVector DebugLocation = TopOfCapsule + FVector(0.f,0.f,HeightOffset);
 		FString DebugText = FString::Printf(TEXT("MovementMode: %s"), *GetMovementName());
 		DrawDebugString(GetWorld(), DebugLocation, DebugText, nullptr, DebugColor, 0.f, true);
 
 		if (IsInWater())
 		{
-			HeightOffset += 15.f;
+			HeightOffset += OffsetPerElement;
 			DebugLocation = TopOfCapsule + FVector(0.f, 0.f, HeightOffset);
 			DebugText = FString::Printf(TEXT("ImmersionDepth: %.2f"), ImmersionDepth());
 			DrawDebugString(GetWorld(), DebugLocation, DebugText, nullptr, DebugColor, 0.f, true);
@@ -6921,7 +6922,7 @@ float UCharacterMovementComponent::VisualizeMovement() const
 	// Jump
 	{
 		const FColor DebugColor = FColor::Blue;
-		HeightOffset += 15.f;
+		HeightOffset += OffsetPerElement;
 		FVector DebugLocation = TopOfCapsule + FVector(0.f, 0.f, HeightOffset);
 		FString DebugText = FString::Printf(TEXT("bIsJumping: %d Count: %d HoldTime: %.2f"), CharacterOwner->bPressedJump, CharacterOwner->JumpCurrentCount, CharacterOwner->JumpKeyHoldTime);
 		DrawDebugString(GetWorld(), DebugLocation, DebugText, nullptr, DebugColor, 0.f, true);
@@ -6931,7 +6932,7 @@ float UCharacterMovementComponent::VisualizeMovement() const
 	if (CurrentRootMotion.HasAdditiveVelocity())
 	{
 		const FColor DebugColor = FColor::Cyan;
-		HeightOffset += 15.f;
+		HeightOffset += OffsetPerElement;
 		const FVector DebugLocation = TopOfCapsule + FVector(0.f,0.f,HeightOffset);
 
 		FVector CurrentAdditiveVelocity(FVector::ZeroVector);
@@ -6949,7 +6950,7 @@ float UCharacterMovementComponent::VisualizeMovement() const
 	if (CurrentRootMotion.HasOverrideVelocity())
 	{
 		const FColor DebugColor = FColor::Green;
-		HeightOffset += 15.f;
+		HeightOffset += OffsetPerElement;
 		const FVector DebugLocation = TopOfCapsule + FVector(0.f,0.f,HeightOffset);
 		FString DebugText = FString::Printf(TEXT("Has Override RootMotion"));
 		DrawDebugString(GetWorld(), DebugLocation, DebugText, nullptr, DebugColor, 0.f, true);
