Fixed bounds calculation when there are no active particles.

diff --git a/Engine/Source/Runtime/Engine/Private/Particles/ParticleComponents.cpp b/Engine/Source/Runtime/Engine/Private/Particles/ParticleComponents.cpp
--- a/Engine/Source/Runtime/Engine/Private/Particles/ParticleComponents.cpp
+++ b/Engine/Source/Runtime/Engine/Private/Particles/ParticleComponents.cpp
@@ -4633,8 +4633,7 @@ FBoxSphereBounds UParticleSystemComponent::CalcBounds(const FTransform& LocalToW
 	}
 	else if (FXConsoleVariables::bAllowCulling == false)
 	{
-		BoundingBox.Min = FVector(-HALF_WORLD_MAX);
-		BoundingBox.Max = FVector(+HALF_WORLD_MAX);
+		BoundingBox = FBox(FVector(-HALF_WORLD_MAX), FVector(HALF_WORLD_MAX));
 	}
 	else if(Template && Template->bUseFixedRelativeBoundingBox)
 	{
@@ -4652,6 +4651,12 @@ FBoxSphereBounds UParticleSystemComponent::CalcBounds(const FTransform& LocalToW
 			}
 		}
 
+		// If the bounding box is not valid at this point there were no active particles, return zero-extent/radius box at local origin.
+		if (!BoundingBox.IsValid)
+		{
+			return FBoxSphereBounds(LocalToWorld.GetTranslation(), FVector::ZeroVector, 0.0f);
+		}
+
 		// Expand the actual bounding-box slightly so it will be valid longer in the case of expanding particle systems.
 		const FVector ExpandAmount = BoundingBox.GetExtent() * 0.1f;
 		BoundingBox = FBox(BoundingBox.Min - ExpandAmount, BoundingBox.Max + ExpandAmount);
