Cache skinned mesh component world-space bounds and world-to-local transform instead of local bounds.

- This yields much tighter AABB bounds.

diff --git a/Engine/Source/Runtime/Engine/Classes/Components/SkinnedMeshComponent.h b/Engine/Source/Runtime/Engine/Classes/Components/SkinnedMeshComponent.h
--- a/Engine/Source/Runtime/Engine/Classes/Components/SkinnedMeshComponent.h
+++ b/Engine/Source/Runtime/Engine/Classes/Components/SkinnedMeshComponent.h
@@ -1095,7 +1095,11 @@ public:
 
 	void SetComponentSpaceTransformsDoubleBuffering(bool bInDoubleBufferedComponentSpaceTransforms);
 
-	const FBoxSphereBounds& GetCachedLocalBounds() { return CachedLocalBounds; } 
+	FBoxSphereBounds GetCachedLocalBounds()  const
+	{ 
+		ensure(bCachedLocalBoundsUpToDate);
+		return CachedWorldSpaceBounds.TransformBy(CachedWorldToLocalTransform);
+	} 
 
 	/**
 	* Should update transform in Tick
@@ -1125,9 +1129,11 @@ protected:
 	virtual bool AllocateTransformData();
 	virtual void DeallocateTransformData();
 
-	/** LocalBounds cached, so they're computed just once. */
+	/** Bounds cached, so they're computed just once. */
+	UPROPERTY(Transient)
+	mutable FBoxSphereBounds CachedWorldSpaceBounds;
 	UPROPERTY(Transient)
-	mutable FBoxSphereBounds CachedLocalBounds;
+	mutable FMatrix CachedWorldToLocalTransform;
 
 public:
 
diff --git a/Engine/Source/Runtime/Engine/Private/Components/SkeletalMeshComponent.cpp b/Engine/Source/Runtime/Engine/Private/Components/SkeletalMeshComponent.cpp
--- a/Engine/Source/Runtime/Engine/Private/Components/SkeletalMeshComponent.cpp
+++ b/Engine/Source/Runtime/Engine/Private/Components/SkeletalMeshComponent.cpp
@@ -2525,11 +2525,12 @@ FBoxSphereBounds USkeletalMeshComponent::CalcBounds(const FTransform& LocalToWor
 		if (bIncludeComponentLocationIntoBounds)
 		{
 			const FVector ComponentLocation = GetComponentLocation();
-			return CachedLocalBounds.TransformBy(LocalToWorld) + FBoxSphereBounds(ComponentLocation, FVector(1.0f), 1.0f);
+			return CachedWorldSpaceBounds.TransformBy(CachedWorldToLocalTransform * LocalToWorld.ToMatrixWithScale()) 
+			        + FBoxSphereBounds(ComponentLocation, FVector(1.0f), 1.0f);
 		}
 		else
 		{
-			return CachedLocalBounds.TransformBy(LocalToWorld);
+			return CachedWorldSpaceBounds.TransformBy(CachedWorldToLocalTransform * LocalToWorld.ToMatrixWithScale());
 		}
 	}
 	// Calculate new bounds
@@ -2565,7 +2566,8 @@ FBoxSphereBounds USkeletalMeshComponent::CalcBounds(const FTransform& LocalToWor
 #endif// #if WITH_APEX_CLOTHING
 
 		bCachedLocalBoundsUpToDate = true;
-		CachedLocalBounds = NewBounds.TransformBy(LocalToWorld.ToInverseMatrixWithScale());
+		CachedWorldSpaceBounds = NewBounds;
+		CachedWorldToLocalTransform = LocalToWorld.ToInverseMatrixWithScale();
 
 		return NewBounds;
 	}
