Make sure objects are in place before initializing sub-instance partial graphs

This ensures that access to skeleton and component  in Initialize_Anythread calls succeed

diff --git a/Engine/Source/Runtime/Engine/Private/Animation/AnimInstance.cpp b/Engine/Source/Runtime/Engine/Private/Animation/AnimInstance.cpp
--- a/Engine/Source/Runtime/Engine/Private/Animation/AnimInstance.cpp
+++ b/Engine/Source/Runtime/Engine/Private/Animation/AnimInstance.cpp
@@ -2588,7 +2588,10 @@ void UAnimInstance::SetLayerOverlay(TSubclassOf<UAnimInstance> InClass)
 							// Initialize the correct parts of the sub instance
 							if(LayerNode->LinkedRoot)
 							{
-								NewSubInstance->GetProxyOnAnyThread<FAnimInstanceProxy>().InitializeRootNode_WithRoot(LayerNode->LinkedRoot);
+								FAnimInstanceProxy& Proxy = NewSubInstance->GetProxyOnAnyThread<FAnimInstanceProxy>();
+								Proxy.InitializeObjects(NewSubInstance);
+								Proxy.InitializeRootNode_WithRoot(LayerNode->LinkedRoot);
+								Proxy.ClearObjects();
 							}
 
 							MeshComp->SubInstances.Add(NewSubInstance);
@@ -2614,12 +2617,16 @@ void UAnimInstance::SetLayerOverlay(TSubclassOf<UAnimInstance> InClass)
 						// Init after we link in the new graph segments, so propagation happens correctly
 						NewSubInstance->InitializeAnimation();
 
+						FAnimInstanceProxy& Proxy = NewSubInstance->GetProxyOnAnyThread<FAnimInstanceProxy>();
+
 						// Initialize the correct parts of the sub instance
 						for(FAnimNode_Layer* LayerNode : LayerPair.Value)
 						{
 							if(LayerNode->LinkedRoot)
 							{
-								NewSubInstance->GetProxyOnAnyThread<FAnimInstanceProxy>().InitializeRootNode_WithRoot(LayerNode->LinkedRoot);
+								Proxy.InitializeObjects(NewSubInstance);
+								Proxy.InitializeRootNode_WithRoot(LayerNode->LinkedRoot);
+								Proxy.ClearObjects();
 							}
 						}
 
