Appropriately set the bOverrideWalkableSlopeOnInstance flag on Instanced Static Mesh Components when creating all body instances

Add a getter for bOverrideWalkableSlopeOnInstance

diff --git a/Engine/Source/Runtime/Engine/Classes/PhysicsEngine/BodyInstance.h b/Engine/Source/Runtime/Engine/Classes/PhysicsEngine/BodyInstance.h
--- a/Engine/Source/Runtime/Engine/Classes/PhysicsEngine/BodyInstance.h
+++ b/Engine/Source/Runtime/Engine/Classes/PhysicsEngine/BodyInstance.h
@@ -601,7 +601,10 @@ public:
 	const struct FWalkableSlopeOverride& GetWalkableSlopeOverride() const;
 
 	/** Sets a custom slope override struct for this instance. Implicitly sets bOverrideWalkableSlopeOnInstance to true. */
-	void SetWalkableSlopeOverride(const FWalkableSlopeOverride& NewOverride);
+	void SetWalkableSlopeOverride(const FWalkableSlopeOverride& NewOverride, bool bNewOverideSetting = true);
+
+	/** Gets bOverrideWalkableSlopeOnInstance */
+	bool GetOverrideWalkableSlopeOnInstance() const;
 
 	/** Returns true if the body is not static */
 	bool IsDynamic() const;
diff --git a/Engine/Source/Runtime/Engine/Private/InstancedStaticMesh.cpp b/Engine/Source/Runtime/Engine/Private/InstancedStaticMesh.cpp
--- a/Engine/Source/Runtime/Engine/Private/InstancedStaticMesh.cpp
+++ b/Engine/Source/Runtime/Engine/Private/InstancedStaticMesh.cpp
@@ -1521,6 +1521,11 @@ void UInstancedStaticMeshComponent::CreateAllInstanceBodies()
 	{
 		FPhysScene* PhysScene = GetWorld()->GetPhysicsScene();
 
+		if (!BodyInstance.GetOverrideWalkableSlopeOnInstance())
+		{
+			BodyInstance.SetWalkableSlopeOverride(BodySetup->WalkableSlopeOverride, false);
+		}
+
 		InstanceBodies.SetNumUninitialized(NumBodies);
 
 		// Sanitized array does not contain any nulls
diff --git a/Engine/Source/Runtime/Engine/Private/PhysicsEngine/BodyInstance.cpp b/Engine/Source/Runtime/Engine/Private/PhysicsEngine/BodyInstance.cpp
--- a/Engine/Source/Runtime/Engine/Private/PhysicsEngine/BodyInstance.cpp
+++ b/Engine/Source/Runtime/Engine/Private/PhysicsEngine/BodyInstance.cpp
@@ -2303,13 +2303,18 @@ const FWalkableSlopeOverride& FBodyInstance::GetWalkableSlopeOverride() const
 	}
 }
 
-void FBodyInstance::SetWalkableSlopeOverride(const FWalkableSlopeOverride& NewOverride)
+void FBodyInstance::SetWalkableSlopeOverride(const FWalkableSlopeOverride& NewOverride, bool bNewOverideSetting)
 {
-	bOverrideWalkableSlopeOnInstance = true;
+	bOverrideWalkableSlopeOnInstance = bNewOverideSetting;
 	WalkableSlopeOverride = NewOverride;
 }
 
 
+bool FBodyInstance::GetOverrideWalkableSlopeOnInstance() const
+{
+	return bOverrideWalkableSlopeOnInstance;
+}
+
 /** 
 *	Changes the current PhysMaterialOverride for this body. 
 *	Note that if physics is already running on this component, this will _not_ alter its mass/inertia etc, it will only change its 
