Deprecate UGameViewportClient::GetDisableSplitscreenOverride and UGameViewportClient::SetDisableSplitscreenOverride and create two new better named functions.

diff --git a/Engine/Source/Runtime/Engine/Classes/Engine/GameViewportClient.h b/Engine/Source/Runtime/Engine/Classes/Engine/GameViewportClient.h
--- a/Engine/Source/Runtime/Engine/Classes/Engine/GameViewportClient.h
+++ b/Engine/Source/Runtime/Engine/Classes/Engine/GameViewportClient.h
@@ -368,14 +368,28 @@ public:
 	virtual void LayoutPlayers();
 
 	/** Allows game code to disable splitscreen (useful when in menus) */
-	void SetDisableSplitscreenOverride( const bool bDisabled );
+	void SetForceDisableSplitscreen(const bool bDisabled);
 
 	/** Determines whether splitscreen is forced to be turned off */
-	bool GetDisableSplitscreenOverride() const
+	bool IsSplitscreenForceDisabled() const
 	{
 		return bDisableSplitScreenOverride;
 	}
 
+	/** Allows game code to disable splitscreen (useful when in menus) */
+	UE_DEPRECATED(4.24, "SetDisableSplitscreenOverride is deprecated. Please call UGameViewportClient::SetForceDisableSplitscreen(bDisabled) instead.")
+	void SetDisableSplitscreenOverride( const bool bDisabled )
+	{
+		SetForceDisableSplitscreen(bDisabled);
+	}
+
+	/** Determines whether splitscreen is forced to be turned off */
+	UE_DEPRECATED(4.24, "GetDisableSplitscreenOverride is deprecated. Please call UGameViewportClient::IsSplitscreenForceDisabled() instead.")
+	bool GetDisableSplitscreenOverride() const
+	{
+		return IsSplitscreenForceDisabled();
+	}
+
 	/** called before rending subtitles to allow the game viewport to determine the size of the subtitle area
 	 * @param Min top left bounds of subtitle region (0 to 1)
 	 * @param Max bottom right bounds of subtitle region (0 to 1)
diff --git a/Engine/Source/Runtime/Engine/Private/GameViewportClient.cpp b/Engine/Source/Runtime/Engine/Private/GameViewportClient.cpp
--- a/Engine/Source/Runtime/Engine/Private/GameViewportClient.cpp
+++ b/Engine/Source/Runtime/Engine/Private/GameViewportClient.cpp
@@ -2156,7 +2156,7 @@ void UGameViewportClient::LayoutPlayers()
 	}
 }
 
-void UGameViewportClient::SetDisableSplitscreenOverride( const bool bDisabled )
+void UGameViewportClient::SetForceDisableSplitscreen(const bool bDisabled)
 {
 	bDisableSplitScreenOverride = bDisabled;
 	LayoutPlayers();
