PR #6018: Add two new Blueprint Functions to Get() & Set() the Splitscreen Enabled (Contributed by [REDACTED_PERSON])

diff --git a/Engine/Source/Runtime/Engine/Classes/Kismet/GameplayStatics.h b/Engine/Source/Runtime/Engine/Classes/Kismet/GameplayStatics.h
--- a/Engine/Source/Runtime/Engine/Classes/Kismet/GameplayStatics.h
+++ b/Engine/Source/Runtime/Engine/Classes/Kismet/GameplayStatics.h
@@ -267,6 +267,20 @@ class ENGINE_API UGameplayStatics : public UBlueprintFunctionLibrary
 	UFUNCTION(BlueprintPure, Category="Game", meta=(WorldContext="WorldContextObject") )
 	static bool IsGamePaused(const UObject* WorldContextObject);
 
+	/**
+	 * Enables split screen
+	 * @param	bDisable		Whether the viewport should split screen between local players or not
+	 */
+	UFUNCTION(BlueprintCallable, Category = "Viewport", meta = (WorldContext = "WorldContextObject"))
+	static void SetForceDisableSplitscreen(const UObject* WorldContextObject, bool bDisable);
+
+	/**
+	 * Returns the split screen state
+	 * @return	Whether the game viewport is split screen or not
+	 */
+	UFUNCTION(BlueprintPure, Category = "Viewport", meta = (WorldContext = "WorldContextObject"))
+	static bool IsSplitscreenForceDisabled(const UObject* WorldContextObject);
+
 	/**
 	 * Enabled rendering of the world
 	 * @param	bEnable		Whether the world should be rendered or not
diff --git a/Engine/Source/Runtime/Engine/Private/GameplayStatics.cpp b/Engine/Source/Runtime/Engine/Private/GameplayStatics.cpp
--- a/Engine/Source/Runtime/Engine/Private/GameplayStatics.cpp
+++ b/Engine/Source/Runtime/Engine/Private/GameplayStatics.cpp
@@ -365,6 +365,33 @@ bool UGameplayStatics::IsGamePaused(const UObject* WorldContextObject)
 	return World ? World->IsPaused() : false;
 }
 
+void UGameplayStatics::SetForceDisableSplitscreen(const UObject* WorldContextObject, bool bDisable)
+{
+	UWorld* World = GEngine->GetWorldFromContextObject(WorldContextObject, EGetWorldErrorMode::LogAndReturnNull);
+	if (World)
+	{
+		UGameViewportClient* GameViewportClient = World->GetGameViewport();
+		if (GameViewportClient)
+		{
+			GameViewportClient->SetForceDisableSplitscreen(bDisable);
+		}
+	}
+}
+
+bool UGameplayStatics::IsSplitscreenForceDisabled(const UObject* WorldContextObject)
+{
+	UWorld* const World = GEngine->GetWorldFromContextObject(WorldContextObject, EGetWorldErrorMode::LogAndReturnNull);
+	if (World)
+	{
+		UGameViewportClient* GameViewportClient = World->GetGameViewport();
+		if (GameViewportClient)
+		{
+			return GameViewportClient->IsSplitscreenForceDisabled();
+		}
+	}
+	return false;
+}
+
 void UGameplayStatics::SetEnableWorldRendering(const UObject* WorldContextObject, bool bEnable)
 {
 	UWorld* World = GEngine->GetWorldFromContextObject(WorldContextObject, EGetWorldErrorMode::LogAndReturnNull);
