Deprecate APawn::IsControlled because the name of the function and the actual behavior is confusing.

Add a IsPawnControlled function that simply checks if the controller is valid at all

diff --git a/Engine/Source/Runtime/Engine/Classes/GameFramework/Pawn.h b/Engine/Source/Runtime/Engine/Classes/GameFramework/Pawn.h
--- a/Engine/Source/Runtime/Engine/Classes/GameFramework/Pawn.h
+++ b/Engine/Source/Runtime/Engine/Classes/GameFramework/Pawn.h
@@ -193,9 +193,14 @@ public:
 	virtual float GetDefaultHalfHeight() const;
 
 	/** See if this actor is currently being controlled */
+	UE_DEPRECATED(4.24, "IsControlled is deprecated. Please call APawn::IsPawnControlled() instead.")
 	UFUNCTION(BlueprintCallable, Category=Pawn)
 	bool IsControlled() const;
 
+	/** Check if this actor is currently being controlled at all (the actor has a valid Controller) */
+	UFUNCTION(BlueprintCallable, Category = Pawn)
+	bool IsPawnControlled() const;
+
 	/** Returns controller for this actor. */
 	UFUNCTION(BlueprintCallable, Category=Pawn)
 	AController* GetController() const;
diff --git a/Engine/Source/Runtime/Engine/Private/Pawn.cpp b/Engine/Source/Runtime/Engine/Private/Pawn.cpp
--- a/Engine/Source/Runtime/Engine/Private/Pawn.cpp
+++ b/Engine/Source/Runtime/Engine/Private/Pawn.cpp
@@ -443,6 +443,11 @@ bool APawn::IsControlled() const
 	return(PC != nullptr);
 }
 
+bool APawn::IsPawnControlled() const
+{
+	return (Controller != nullptr);
+}
+
 FRotator APawn::GetControlRotation() const
 {
 	return Controller ? Controller->GetControlRotation() : FRotator::ZeroRotator;
