Update deprecation message for IsControlled to be more explicit

#rnx

diff --git a/Engine/Source/Runtime/Engine/Classes/GameFramework/Pawn.h b/Engine/Source/Runtime/Engine/Classes/GameFramework/Pawn.h
--- a/Engine/Source/Runtime/Engine/Classes/GameFramework/Pawn.h
+++ b/Engine/Source/Runtime/Engine/Classes/GameFramework/Pawn.h
@@ -193,7 +193,7 @@ public:
 	virtual float GetDefaultHalfHeight() const;
 
 	/** See if this actor is currently being controlled */
-	UE_DEPRECATED(4.24, "IsControlled is deprecated. Please call APawn::IsPawnControlled() instead.")
+	UE_DEPRECATED(4.24, "IsControlled is deprecated. To check if this pawn is controlled by anything, then call IsPawnControlled. To check if this pawn is controlled only by the player then call IsPlayerControlled")
 	UFUNCTION(BlueprintCallable, Category=Pawn)
 	bool IsControlled() const;
 
