Don't allow cooking if AssetManagerSettings has non-unique Primary Asset Type entries

diff --git a/Engine/Source/Runtime/Engine/Private/AssetManager.cpp b/Engine/Source/Runtime/Engine/Private/AssetManager.cpp
--- a/Engine/Source/Runtime/Engine/Private/AssetManager.cpp
+++ b/Engine/Source/Runtime/Engine/Private/AssetManager.cpp
@@ -2580,6 +2580,8 @@ void UAssetManager::ScanPrimaryAssetTypesFromConfig()
 			continue;
 		}
 
+		UE_CLOG(AssetTypeMap.Find(TypeInfo.PrimaryAssetType), LogAssetManager, Error, TEXT("Found multiple \"%s\" Primary Asset Type entries in \"Primary Asset Types To Scan\" config. Only a single entry per type is supported."), *TypeInfo.PrimaryAssetType.ToString());
+
 		ScanPathsForPrimaryAssets(TypeInfo.PrimaryAssetType, TypeInfo.AssetScanPaths, TypeInfo.AssetBaseClassLoaded, TypeInfo.bHasBlueprintClasses, TypeInfo.bIsEditorOnly, false);
 
 		SetPrimaryAssetTypeRules(TypeInfo.PrimaryAssetType, TypeInfo.Rules);
