Remove long deprecated and unused member
#rnx

diff --git a/Engine/Source/Runtime/Engine/Classes/Components/ActorComponent.h b/Engine/Source/Runtime/Engine/Classes/Components/ActorComponent.h
--- a/Engine/Source/Runtime/Engine/Classes/Components/ActorComponent.h
+++ b/Engine/Source/Runtime/Engine/Classes/Components/ActorComponent.h
@@ -204,10 +204,6 @@ public:
 	/** If true, we call the virtual InitializeComponent */
 	uint8 bWantsInitializeComponent:1;
 
-	/** If true, we call the virtual BeginPlay */
-	UE_DEPRECATED(4.14, "bWantsBeginPlay was inconsistently enforced and is now unused. BeginPlay will now always be called for Actor Components.")
-	uint8 bWantsBeginPlay:1;
-
 	/** If true, the component will be excluded from non-editor builds */
 	UPROPERTY(EditAnywhere, BlueprintReadOnly, Category = Cooking)
 	uint8 bIsEditorOnly:1;
