PR #6067: PossessedBy now sets the owner of the Pawn (Contributed by [REDACTED_PERSON])

diff --git a/Engine/Source/Runtime/Engine/Private/Pawn.cpp b/Engine/Source/Runtime/Engine/Private/Pawn.cpp
--- a/Engine/Source/Runtime/Engine/Private/Pawn.cpp
+++ b/Engine/Source/Runtime/Engine/Private/Pawn.cpp
@@ -493,6 +493,8 @@ void APawn::SetPlayerState(APlayerState* NewPlayerState)
 
 void APawn::PossessedBy(AController* NewController)
 {
+	SetOwner(NewController);
+	
 	AController* const OldController = Controller;
 
 	Controller = NewController;
