PR #6139: GetActorTransform should return const reference for performance (Contributed by [REDACTED_PERSON])

#rnx

diff --git a/Engine/Source/Runtime/Engine/Classes/GameFramework/Actor.h b/Engine/Source/Runtime/Engine/Classes/GameFramework/Actor.h
--- a/Engine/Source/Runtime/Engine/Classes/GameFramework/Actor.h
+++ b/Engine/Source/Runtime/Engine/Classes/GameFramework/Actor.h
@@ -857,13 +857,13 @@ public:
 	 * @return The transform that transforms from actor space to world space.
 	 */
 	UFUNCTION(BlueprintCallable, meta=(DisplayName = "GetActorTransform", ScriptName = "GetActorTransform"), Category="Utilities|Transformation")
-	FTransform GetTransform() const
+	const FTransform& GetTransform() const
 	{
 		return ActorToWorld();
 	}
 
 	/** Get the local-to-world transform of the RootComponent. Identical to GetTransform(). */
-	FORCEINLINE FTransform ActorToWorld() const
+	FORCEINLINE const FTransform& ActorToWorld() const
 	{
 		return (RootComponent ? RootComponent->GetComponentTransform() : FTransform::Identity);
 	}
@@ -1697,7 +1697,7 @@ public:
 	bool SetRootComponent(USceneComponent* NewRootComponent);
 
 	/** Returns the transform of the RootComponent of this Actor*/ 
-	FORCEINLINE FTransform GetActorTransform() const
+	FORCEINLINE const FTransform& GetActorTransform() const
 	{
 		return TemplateGetActorTransform(RootComponent);
 	}
@@ -3072,9 +3072,9 @@ private:
 	// These are templates for no other reason than to delay compilation until USceneComponent is defined.
 
 	template<class T>
-	static FORCEINLINE FTransform TemplateGetActorTransform(const T* RootComponent)
+	static FORCEINLINE const FTransform& TemplateGetActorTransform(const T* RootComponent)
 	{
-		return (RootComponent != nullptr) ? RootComponent->GetComponentTransform() : FTransform();
+		return (RootComponent != nullptr) ? RootComponent->GetComponentTransform() : FTransform::Identity;
 	}
 
 	template<class T>
