Set bHasFinishedSpawning for loaded actors as they finished spawning when they were placed

diff --git a/Engine/Source/Runtime/Engine/Private/Actor.cpp b/Engine/Source/Runtime/Engine/Private/Actor.cpp
--- a/Engine/Source/Runtime/Engine/Private/Actor.cpp
+++ b/Engine/Source/Runtime/Engine/Private/Actor.cpp
@@ -710,6 +710,9 @@ void AActor::PostLoad()
 		}
 	}
 #endif // WITH_EDITORONLY_DATA
+
+	// Since the actor is being loading, it finished spawning by definition when it was originally spawned, so set to true now
+	bHasFinishedSpawning = true;
 }
 
 void AActor::PostLoadSubobjects(FObjectInstancingGraph* OuterInstanceGraph)
