Change bool flags to uint8:1 flags on MeshComponent

#rnx

diff --git a/Engine/Source/Runtime/Engine/Classes/Components/MeshComponent.h b/Engine/Source/Runtime/Engine/Classes/Components/MeshComponent.h
--- a/Engine/Source/Runtime/Engine/Classes/Components/MeshComponent.h
+++ b/Engine/Source/Runtime/Engine/Classes/Components/MeshComponent.h
@@ -144,9 +144,9 @@ protected:
 	TSortedMap<FName, FMaterialParameterCache, FDefaultAllocator, FNameFastLess> MaterialParameterCache;
 
 	UPROPERTY(EditAnywhere, AdvancedDisplay, BlueprintReadOnly, Category = MaterialParameters)
-	bool bEnableMaterialParameterCaching;
+	uint8 bEnableMaterialParameterCaching : 1;
 
 	/** Flag whether or not the cached material parameter indices map is dirty (defaults to true, and is set from SetMaterial/Set(Skeletal)Mesh */
-	bool bCachedMaterialParameterIndicesAreDirty;
+	uint8 bCachedMaterialParameterIndicesAreDirty : 1;
 
 };
