Add 0 entries to animtion enums to prevent blueprint compiler warnings.

Enums affected:
- ESplineBoneAxis
- EAnimAssetCurveFlags

#rnx

diff --git a/Engine/Source/Runtime/AnimGraphRuntime/Public/BoneControllers/AnimNode_SplineIK.h b/Engine/Source/Runtime/AnimGraphRuntime/Public/BoneControllers/AnimNode_SplineIK.h
--- a/Engine/Source/Runtime/AnimGraphRuntime/Public/BoneControllers/AnimNode_SplineIK.h
+++ b/Engine/Source/Runtime/AnimGraphRuntime/Public/BoneControllers/AnimNode_SplineIK.h
@@ -14,6 +14,7 @@
 UENUM()
 enum class ESplineBoneAxis : uint8
 {
+	None = 0 UMETA(Hidden),
 	X = 1,
 	Y = 2,
 	Z = 3,
diff --git a/Engine/Source/Runtime/Engine/Public/Animation/AnimCurveTypes.h b/Engine/Source/Runtime/Engine/Public/Animation/AnimCurveTypes.h
--- a/Engine/Source/Runtime/Engine/Public/Animation/AnimCurveTypes.h
+++ b/Engine/Source/Runtime/Engine/Public/Animation/AnimCurveTypes.h
@@ -14,6 +14,7 @@
 UENUM(BlueprintType, meta=(Bitflags))
 enum EAnimAssetCurveFlags
 {
+	AACF_NONE = 0 UMETA(Hidden),
 	// Used as morph target curve
 	AACF_DriveMorphTarget_DEPRECATED = 0x00000001 UMETA(Hidden), // This has moved to FAnimCurveType:bMorphTarget. Set per skeleton. DO NOT REMOVE UNTIL FrameworkObjectVersion.MoveCurveTypesToSkeleton expires.
 	// Used as triggering event
diff --git a/Engine/Source/Runtime/MovieScene/Public/Evaluation/Blending/MovieSceneBlendType.h b/Engine/Source/Runtime/MovieScene/Public/Evaluation/Blending/MovieSceneBlendType.h
--- a/Engine/Source/Runtime/MovieScene/Public/Evaluation/Blending/MovieSceneBlendType.h
+++ b/Engine/Source/Runtime/MovieScene/Public/Evaluation/Blending/MovieSceneBlendType.h
@@ -10,6 +10,8 @@
 UENUM()
 enum class EMovieSceneBlendType : uint8
 {
+	None = 0 UMETA(Hidden),
+
 	/** Blends all other weighted values together as an average of the total weight */
 	Absolute			= 0x1,
 	/** Applies this value as a sum total of all other additives */
