PR #6136: Fix LevelScriptActor creating unlimited InputComponents (Contributed by [REDACTED_PERSON])

diff --git a/Engine/Source/Runtime/Engine/Private/LevelScriptActor.cpp b/Engine/Source/Runtime/Engine/Private/LevelScriptActor.cpp
--- a/Engine/Source/Runtime/Engine/Private/LevelScriptActor.cpp
+++ b/Engine/Source/Runtime/Engine/Private/LevelScriptActor.cpp
@@ -32,7 +32,7 @@ ALevelScriptActor::ALevelScriptActor(const FObjectInitializer& ObjectInitializer
 
 void ALevelScriptActor::PreInitializeComponents()
 {
-	if (UInputDelegateBinding::SupportsInputDelegate(GetClass()))
+	if (UInputDelegateBinding::SupportsInputDelegate(GetClass()) && !InputComponent)
 	{
 		// create an InputComponent object so that the level script actor can bind key events
 		InputComponent = NewObject<UInputComponent>(this);
