Improve packing on the ShapeComponent, bring the BoxComponent down a bin size. Improve cache coherencey of SplineComponent

#rnx

diff --git a/Engine/Source/Runtime/Engine/Classes/Components/SplineMeshComponent.h b/Engine/Source/Runtime/Engine/Classes/Components/SplineMeshComponent.h
--- a/Engine/Source/Runtime/Engine/Classes/Components/SplineMeshComponent.h
+++ b/Engine/Source/Runtime/Engine/Classes/Components/SplineMeshComponent.h
@@ -13,6 +13,7 @@
 class FPrimitiveSceneProxy;
 class ULightComponent;
 struct FNavigableGeometryExport;
+class UBodySetup;
 
 UENUM(BlueprintType)
 namespace ESplineMeshAxis
@@ -58,14 +59,14 @@ struct FSplineMeshParams
 	UPROPERTY(EditAnywhere, Category=SplineMesh)
 	FVector EndPos;
 
-	/** End tangent of spline, in component space. */
-	UPROPERTY(EditAnywhere, Category=SplineMesh)
-	FVector EndTangent;
-
 	/** X and Y scale applied to mesh at end of spline. */
 	UPROPERTY(EditAnywhere, Category=SplineMesh, AdvancedDisplay)
 	FVector2D EndScale;
 
+	/** End tangent of spline, in component space. */
+	UPROPERTY(EditAnywhere, Category = SplineMesh)
+	FVector EndTangent;
+
 	/** Roll around spline applied at end. */
 	UPROPERTY(EditAnywhere, Category=SplineMesh, AdvancedDisplay)
 	float EndRoll;
@@ -82,8 +83,8 @@ struct FSplineMeshParams
 		, StartRoll(0)
 		, StartOffset(ForceInit)
 		, EndPos(ForceInit)
-		, EndTangent(ForceInit)
 		, EndScale(ForceInit)
+		, EndTangent(ForceInit)
 		, EndRoll(0)
 		, EndOffset(ForceInit)
 	{
@@ -108,43 +109,43 @@ class ENGINE_API USplineMeshComponent : public UStaticMeshComponent, public IInt
 	UPROPERTY(EditAnywhere, Category=SplineMesh)
 	FVector SplineUpDir;
 
+	/** Minimum coordinate along the spline forward axis which corresponds to start of spline. If set to 0.0, will use bounding box to determine bounds */
+	UPROPERTY(EditAnywhere, Category = SplineMesh, AdvancedDisplay)
+	float SplineBoundaryMin;
+
+	// Used to automatically trigger rebuild of collision data
+	UPROPERTY()
+	FGuid CachedMeshBodySetupGuid;
+
+	// Physics data.
+	UPROPERTY()
+	UBodySetup* BodySetup;
+
+	/** Maximum coordinate along the spline forward axis which corresponds to end of spline. If set to 0.0, will use bounding box to determine bounds */
+	UPROPERTY(EditAnywhere, Category = SplineMesh, AdvancedDisplay)
+	float SplineBoundaryMax;
+
 	/** If true, spline keys may be edited per instance in the level viewport. Otherwise, the spline should be initialized in the construction script. */
 	UPROPERTY(EditDefaultsOnly, Category = Spline)
-	uint32 bAllowSplineEditingPerInstance:1;
+	uint8 bAllowSplineEditingPerInstance:1;
 
 	/** If true, will use smooth interpolation (ease in/out) for Scale, Roll, and Offset along this section of spline. If false, uses linear */
 	UPROPERTY(EditAnywhere, Category=SplineMesh, AdvancedDisplay)
-	uint32 bSmoothInterpRollScale:1;
+	uint8 bSmoothInterpRollScale:1;
+
+	// Indicates that the mesh needs updating
+	UPROPERTY(transient)
+	uint8 bMeshDirty : 1;
 
 	/** Chooses the forward axis for the spline mesh orientation */
 	UPROPERTY(EditAnywhere, Category=SplineMesh)
 	TEnumAsByte<ESplineMeshAxis::Type> ForwardAxis;
 
-	/** Minimum coordinate along the spline forward axis which corresponds to start of spline. If set to 0.0, will use bounding box to determine bounds */
-	UPROPERTY(EditAnywhere, Category = SplineMesh, AdvancedDisplay)
-	float SplineBoundaryMin;
-
-	/** Maximum coordinate along the spline forward axis which corresponds to end of spline. If set to 0.0, will use bounding box to determine bounds */
-	UPROPERTY(EditAnywhere, Category = SplineMesh, AdvancedDisplay)
-	float SplineBoundaryMax;
-	
-	// Physics data.
-	UPROPERTY()
-	class UBodySetup* BodySetup;
-
-	// Used to automatically trigger rebuild of collision data
-	UPROPERTY()
-	FGuid CachedMeshBodySetupGuid;
-
 #if WITH_EDITORONLY_DATA
 	UPROPERTY(transient)
-	uint32 bSelected:1;
+	uint8 bSelected:1;
 #endif
 
-	// Indicates that the mesh needs updating
-	UPROPERTY(transient)
-	uint32 bMeshDirty:1;
-
 	//Begin UObject Interface
 	virtual void Serialize(FArchive& Ar) override;
 	virtual bool Modify(bool bAlwaysMarkDirty = true) override;
