Improve packing particle modules and distriution curves. Brings several particle modules down a bin size, as well as fix up some alignment issues with curves to bring them from 40 bytes to down to 32

#rnx

diff --git a/Engine/Source/Runtime/Engine/Classes/Distributions/Distribution.h b/Engine/Source/Runtime/Engine/Classes/Distributions/Distribution.h
--- a/Engine/Source/Runtime/Engine/Classes/Distributions/Distribution.h
+++ b/Engine/Source/Runtime/Engine/Classes/Distributions/Distribution.h
@@ -26,25 +26,25 @@ USTRUCT(noexport)
 struct FDistributionLookupTable
 {
 	UPROPERTY()
-	uint8 Op;
+	float TimeScale;
 
 	UPROPERTY()
-	uint8 EntryCount;
+	float TimeBias;
 
 	UPROPERTY()
-	uint8 EntryStride;
+	TArray<float> Values;
 
 	UPROPERTY()
-	uint8 SubEntryStride;
+	uint8 Op;
 
 	UPROPERTY()
-	float TimeScale;
+	uint8 EntryCount;
 
 	UPROPERTY()
-	float TimeBias;
+	uint8 EntryStride;
 
 	UPROPERTY()
-	TArray<float> Values;
+	uint8 SubEntryStride;
 
 	UPROPERTY()
 	uint8 LockFlag;
diff --git a/Engine/Source/Runtime/Engine/Classes/Distributions/DistributionFloat.h b/Engine/Source/Runtime/Engine/Classes/Distributions/DistributionFloat.h
--- a/Engine/Source/Runtime/Engine/Classes/Distributions/DistributionFloat.h
+++ b/Engine/Source/Runtime/Engine/Classes/Distributions/DistributionFloat.h
@@ -105,14 +105,14 @@ class ENGINE_API UDistributionFloat : public UDistribution
 
 	/** Can this variable be baked out to a FRawDistribution? Should be true 99% of the time*/
 	UPROPERTY(EditAnywhere, Category=Baked)
-	uint32 bCanBeBaked:1;
+	uint8 bCanBeBaked:1;
 
 	/** Set internally when the distribution is updated so that that FRawDistribution can know to update itself*/
-	uint32 bIsDirty:1;
+	uint8 bIsDirty:1;
 	
 protected:
 	UPROPERTY()
-	uint32 bBakedDataSuccesfully:1;	//It's possible that even though we want to bake we are not able to because of content or code.
+	uint8 bBakedDataSuccesfully:1;	//It's possible that even though we want to bake we are not able to because of content or code.
 public:
 
 	/** Script-accessible way to query a float distribution */
diff --git a/Engine/Source/Runtime/Engine/Classes/Distributions/DistributionVector.h b/Engine/Source/Runtime/Engine/Classes/Distributions/DistributionVector.h
--- a/Engine/Source/Runtime/Engine/Classes/Distributions/DistributionVector.h
+++ b/Engine/Source/Runtime/Engine/Classes/Distributions/DistributionVector.h
@@ -150,14 +150,14 @@ class ENGINE_API UDistributionVector : public UDistribution
 
 	/** Can this variable be baked out to a FRawDistribution? Should be true 99% of the time*/
 	UPROPERTY(EditAnywhere, Category=Baked)
-	uint32 bCanBeBaked:1;
+	uint8 bCanBeBaked:1;
 
 	/** Set internally when the distribution is updated so that that FRawDistribution can know to update itself*/
 	UPROPERTY()
-	uint32 bIsDirty:1;
+	uint8 bIsDirty:1;
 protected:
 	UPROPERTY()
-	uint32 bBakedDataSuccesfully:1;	//It's possible that even though we want to bake we are not able to because of content or code.
+	uint8 bBakedDataSuccesfully:1;	//It's possible that even though we want to bake we are not able to because of content or code.
 public:
 
 	/** Script-accessible way to query a FVector distribution */
diff --git a/Engine/Source/Runtime/Engine/Classes/Particles/Acceleration/ParticleModuleAccelerationBase.h b/Engine/Source/Runtime/Engine/Classes/Particles/Acceleration/ParticleModuleAccelerationBase.h
--- a/Engine/Source/Runtime/Engine/Classes/Particles/Acceleration/ParticleModuleAccelerationBase.h
+++ b/Engine/Source/Runtime/Engine/Classes/Particles/Acceleration/ParticleModuleAccelerationBase.h
@@ -19,8 +19,7 @@ class UParticleModuleAccelerationBase : public UParticleModule
 	 *	If true, then treat the acceleration as world-space
 	 */
 	UPROPERTY(EditAnywhere, Category=Acceleration)
-	uint32 bAlwaysInWorldSpace:1;
-
+	uint8 bAlwaysInWorldSpace : 1;
 
 	//~ Begin UParticleModule Interface
 	virtual void SetToSensibleDefaults(UParticleEmitter* Owner) override;
diff --git a/Engine/Source/Runtime/Engine/Classes/Particles/Color/ParticleModuleColor.h b/Engine/Source/Runtime/Engine/Classes/Particles/Color/ParticleModuleColor.h
--- a/Engine/Source/Runtime/Engine/Classes/Particles/Color/ParticleModuleColor.h
+++ b/Engine/Source/Runtime/Engine/Classes/Particles/Color/ParticleModuleColor.h
@@ -14,6 +14,8 @@ class UInterpCurveEdSetup;
 class UParticleEmitter;
 struct FCurveEdEntry;
 struct FParticleEmitterInstance;
+struct FRawDistributionFloat;
+struct FRawDistributionVector;
 
 UCLASS(editinlinenew, hidecategories=Object, meta=(DisplayName = "Initial Color"))
 class ENGINE_API UParticleModuleColor : public UParticleModuleColorBase
@@ -22,15 +24,15 @@ class ENGINE_API UParticleModuleColor : public UParticleModuleColorBase
 
 	/** Initial color for a particle as a function of Emitter time. */
 	UPROPERTY(EditAnywhere, Category = Color, meta = (TreatAsColor))
-	struct FRawDistributionVector StartColor;
+	FRawDistributionVector StartColor;
 
 	/** Initial alpha for a particle as a function of Emitter time. */
 	UPROPERTY(EditAnywhere, Category=Color)
-	struct FRawDistributionFloat StartAlpha;
+	FRawDistributionFloat StartAlpha;
 
 	/** If true, the alpha value will be clamped to the [0..1] range. */
 	UPROPERTY(EditAnywhere, Category=Color)
-	uint32 bClampAlpha:1;
+	uint8 bClampAlpha:1;
 
 	/** Initializes the default values for this property */
 	void InitializeDefaults();
diff --git a/Engine/Source/Runtime/Engine/Classes/Particles/ParticleModule.h b/Engine/Source/Runtime/Engine/Classes/Particles/ParticleModule.h
--- a/Engine/Source/Runtime/Engine/Classes/Particles/ParticleModule.h
+++ b/Engine/Source/Runtime/Engine/Classes/Particles/ParticleModule.h
@@ -98,14 +98,14 @@ struct FParticleRandomSeedInfo
 	 *	the RandomSeeds array.
 	 */
 	UPROPERTY(EditAnywhere, Category=ParticleRandomSeedInfo)
-	uint32 bGetSeedFromInstance:1;
+	uint8 bGetSeedFromInstance:1;
 
 	/** 
 	 *	If true, the seed value retrieved from the instance will be an
 	 *	index into the array of seeds.
 	 */
 	UPROPERTY(EditAnywhere, Category=ParticleRandomSeedInfo)
-	uint32 bInstanceSeedIsIndex:1;
+	uint8 bInstanceSeedIsIndex:1;
 
 	/** 
 	 *	If true, then reset the seed upon the emitter looping.
@@ -113,13 +113,13 @@ struct FParticleRandomSeedInfo
 	 *	a repeating pattern.
 	 */
 	UPROPERTY(EditAnywhere, Category=ParticleRandomSeedInfo)
-	uint32 bResetSeedOnEmitterLooping:1;
+	uint8 bResetSeedOnEmitterLooping:1;
 
 	/**
 	*	If true, then randomly select a seed entry from the RandomSeeds array
 	*/
 	UPROPERTY(EditAnywhere, Category = ParticleRandomSeedInfo)
-	uint32 bRandomlySelectSeedArray:1;
+	uint8 bRandomlySelectSeedArray:1;
 
 	/** 
 	 *	The random seed values to utilize for the module. 
@@ -151,39 +151,39 @@ class ENGINE_API UParticleModule : public UObject
 
 	/** If true, the module performs operations on particles during Spawning		*/
 	UPROPERTY()
-	uint32 bSpawnModule:1;
+	uint8 bSpawnModule:1;
 
 	/** If true, the module performs operations on particles during Updating		*/
 	UPROPERTY()
-	uint32 bUpdateModule:1;
+	uint8 bUpdateModule:1;
 
 	/** If true, the module performs operations on particles during final update	*/
 	UPROPERTY()
-	uint32 bFinalUpdateModule:1;
+	uint8 bFinalUpdateModule:1;
 
 	/** If true, the module performs operations on particles during update and/or final update for GPU emitters*/
 	UPROPERTY()
-	uint32 bUpdateForGPUEmitter:1;
+	uint8 bUpdateForGPUEmitter:1;
 
 	/** If true, the module displays FVector curves as colors						*/
 	UPROPERTY()
-	uint32 bCurvesAsColor:1;
+	uint8 bCurvesAsColor:1;
 
 	/** If true, the module should render its 3D visualization helper				*/
 	UPROPERTY(EditAnywhere, Category=Cascade)
-	uint32 b3DDrawMode:1;
+	uint8 b3DDrawMode:1;
 
 	/** If true, the module supports rendering a 3D visualization helper			*/
 	UPROPERTY()
-	uint32 bSupported3DDrawMode:1;
+	uint8 bSupported3DDrawMode:1;
 
 	/** If true, the module is enabled												*/
 	UPROPERTY()
-	uint32 bEnabled:1;
+	uint8 bEnabled:1;
 
 	/** If true, the module has had editing enabled on it							*/
 	UPROPERTY()
-	uint32 bEditable:1;
+	uint8 bEditable:1;
 
 	/**
 	*	If true, this flag indicates that auto-generation for LOD will result in
@@ -191,15 +191,15 @@ class ENGINE_API UParticleModule : public UObject
 	*	If false, it will result in a module with different settings.
 	*/
 	UPROPERTY()
-	uint32 LODDuplicate:1;
+	uint8 LODDuplicate:1;
 
 	/** If true, the module supports RandomSeed setting */
 	UPROPERTY()
-	uint32 bSupportsRandomSeed:1;
+	uint8 bSupportsRandomSeed:1;
 
 	/** If true, the module should be told when looping */
 	UPROPERTY()
-	uint32 bRequiresLoopingNotification:1;
+	uint8 bRequiresLoopingNotification:1;
 
 	/**
 	 *	The LOD levels this module is present in.
@@ -274,12 +274,6 @@ class ENGINE_API UParticleModule : public UObject
 	 */
 	virtual uint32	RequiredBytesPerInstance();
 
-	UE_DEPRECATED(4.11, "RequiredBytes now takes a UParticleModuleTypeDataBase*, not per-instance information.")
-	virtual uint32	RequiredBytes(FParticleEmitterInstance* Owner);
-
-	UE_DEPRECATED(4.11, "RequiredBytesPerInstance no longer takes per-instance information.")
-	virtual uint32	RequiredBytesPerInstance(FParticleEmitterInstance* Owner);
-
 	/**
 	 *	Allows the module to prep its 'per-instance' data block.
 	 *	
diff --git a/Engine/Source/Runtime/Engine/Classes/Particles/ParticleModuleRequired.h b/Engine/Source/Runtime/Engine/Classes/Particles/ParticleModuleRequired.h
--- a/Engine/Source/Runtime/Engine/Classes/Particles/ParticleModuleRequired.h
+++ b/Engine/Source/Runtime/Engine/Classes/Particles/ParticleModuleRequired.h
@@ -106,13 +106,13 @@
 
 struct FParticleRequiredModule
 {
-	bool bCutoutTexureIsValid;
 	uint32 NumFrames;
 	uint32 NumBoundingVertices;
 	uint32 NumBoundingTriangles;
 	float AlphaThreshold;
 	TArray<FVector2D> FrameData;
 	FRHIShaderResourceView* BoundingGeometryBufferSRV;
+	uint8 bCutoutTexureIsValid : 1;
 };
 
 
@@ -130,8 +130,17 @@
 	UPROPERTY(EditAnywhere, Category=Emitter)
 	class UMaterialInterface* Material;
 
+	/** The distance at which PSA_FacingCameraDistanceBlend	is fully PSA_Square */
+	UPROPERTY(EditAnywhere, Category = Emitter, meta = (UIMin = "0", DisplayAfter="ScreenAlignment"))
+	float MinFacingCameraBlendDistance;
+
+	/** The distance at which PSA_FacingCameraDistanceBlend	is fully PSA_FacingCameraPosition */
+	UPROPERTY(EditAnywhere, Category = Emitter, meta = (UIMin = "0", DisplayAfter="MinFacingCameraBlendDistance"))
+	float MaxFacingCameraBlendDistance;
+
 	UPROPERTY(EditAnywhere, Category=Emitter)
 	FVector EmitterOrigin;
+
 	UPROPERTY(EditAnywhere, Category=Emitter)
 	FRotator EmitterRotation;
 
@@ -150,25 +159,17 @@
 	UPROPERTY(EditAnywhere, Category=Emitter)
 	TEnumAsByte<EParticleScreenAlignment> ScreenAlignment;
 
-	/** The distance at which PSA_FacingCameraDistanceBlend	is fully PSA_Square */
-	UPROPERTY(EditAnywhere, Category=Emitter, meta=(UIMin = "0"))
-	float MinFacingCameraBlendDistance;
-
-	/** The distance at which PSA_FacingCameraDistanceBlend	is fully PSA_FacingCameraPosition */
-	UPROPERTY(EditAnywhere, Category=Emitter, meta=(UIMin = "0"))
-	float MaxFacingCameraBlendDistance;
-
 	/** If true, update the emitter in local space										*/
-	UPROPERTY(EditAnywhere, Category=Emitter)
-	uint32 bUseLocalSpace:1;
+	UPROPERTY(EditAnywhere, Category=Emitter, meta = (DisplayAfter = "MaxFacingCameraBlendDistance"))
+	uint8 bUseLocalSpace:1;
 
 	/** If true, kill the emitter when the particle system is deactivated				*/
 	UPROPERTY(EditAnywhere, Category=Emitter)
-	uint32 bKillOnDeactivate:1;
+	uint8 bKillOnDeactivate:1;
 
 	/** If true, kill the emitter when it completes										*/
 	UPROPERTY(EditAnywhere, Category=Emitter)
-	uint32 bKillOnCompleted:1;
+	uint8 bKillOnCompleted:1;
 
 	/**
 	 *	The sorting mode to use for this emitter.
@@ -180,7 +181,7 @@
 	 *
 	 */
 	UPROPERTY(EditAnywhere, Category=Emitter)
-	TEnumAsByte<enum EParticleSortMode> SortMode;
+	TEnumAsByte<EParticleSortMode> SortMode;
 
 	/**
 	 *	If true, the EmitterTime for the emitter will be calculated by
@@ -193,42 +194,24 @@
 	 *	resulting in proper looping/delay behavior.
 	 */
 	UPROPERTY(EditAnywhere, Category=Emitter)
-	uint32 bUseLegacyEmitterTime:1;
+	uint8 bUseLegacyEmitterTime:1;
 
 	/** If true, removes the HMD view roll (e.g. in VR) */
 	UPROPERTY(EditAnywhere, Category=Emitter, meta=(DisplayName = "Remove HMD Roll"))
-	uint32 bRemoveHMDRoll:1;
+	uint8 bRemoveHMDRoll:1;
 
-	/** 
-	 *	How long, in seconds, the emitter will run before looping.
-	 */
-	UPROPERTY(EditAnywhere, Category=Duration)
-	float EmitterDuration;
-
-	/** 
-	 *	The low end of the emitter duration if using a range.
-	 */
-	UPROPERTY(EditAnywhere, Category=Duration)
-	float EmitterDurationLow;
-
 	/**
-	 *	If true, select the emitter duration from the range 
+	 *	If true, select the emitter duration from the range
 	 *		[EmitterDurationLow..EmitterDuration]
 	 */
-	UPROPERTY(EditAnywhere, Category=Duration)
-	uint32 bEmitterDurationUseRange:1;
+	UPROPERTY(EditAnywhere, Category = Duration)
+	uint8 bEmitterDurationUseRange : 1;
 
 	/** 
-	 *	If true, recalculate the emitter duration on each loop.
+	 *	How long, in seconds, the emitter will run before looping.
 	 */
-	UPROPERTY(EditAnywhere, Category=Duration)
-	uint32 bDurationRecalcEachLoop:1;
+	UPROPERTY(EditAnywhere, Category=Duration, meta = (DisplayAfter="bRemoveHMDRoll"))
+	float EmitterDuration;
 
-	/** The number of times to loop the emitter.
-	 *	0 indicates loop continuously
-	 */
-	UPROPERTY(EditAnywhere, Category=Duration)
-	int32 EmitterLoops;
-
 	//
 	// Spawn-related.
@@ -238,16 +221,9 @@
 	UPROPERTY()
 	struct FRawDistributionFloat SpawnRate;
 
-	//
-	// Burst-related.
-	//
-	/** The method to utilize when burst-emitting particles						*/
-	UPROPERTY()
-	TEnumAsByte<EParticleBurstMethod> ParticleBurstMethod;
-
 	/** The array of burst entries.												*/
 	UPROPERTY(export, noclear)
-	TArray<struct FParticleBurst> BurstList;
+	TArray<FParticleBurst> BurstList;
 
 	//
 	// Delay-related.
@@ -266,17 +242,10 @@
 	float EmitterDelayLow;
 
 	/**
-	 *	If true, select the emitter delay from the range 
-	 *		[EmitterDelayLow..EmitterDelay]
-	 */
-	UPROPERTY(EditAnywhere, Category=Delay)
-	uint32 bEmitterDelayUseRange:1;
-
-	/**
 	 *	If true, the emitter will be delayed only on the first loop.
 	 */
 	UPROPERTY(EditAnywhere, Category=Delay)
-	uint32 bDelayFirstLoopOnly:1;
+	uint8 bDelayFirstLoopOnly:1;
 
 	//
 	// SubUV-related.
@@ -298,6 +267,45 @@
 	UPROPERTY(EditAnywhere, Category=SubUV)
 	TEnumAsByte<EParticleSubUVInterpMethod> InterpolationMethod;
 
+	/** Whether to scale the UV or not - ie, the model wasn't setup with sub uvs		*/
+	UPROPERTY(EditAnywhere, Category = SubUV)
+	uint8 bScaleUV : 1;
+
+	/**
+	 *	If true, select the emitter delay from the range 
+	 *		[EmitterDelayLow..EmitterDelay]
+	 */
+	UPROPERTY(EditAnywhere, Category=Delay)
+	uint8 bEmitterDelayUseRange:1;
+
+	/** The method to utilize when burst-emitting particles						*/
+	UPROPERTY()
+	TEnumAsByte<EParticleBurstMethod> ParticleBurstMethod;
+
+	/** Override the system MacroUV settings                                            */
+	UPROPERTY(EditAnywhere, Category = MacroUV)
+	uint8 bOverrideSystemMacroUV : 1;
+
+	/**
+	 *	If true, use the MaxDrawCount to limit the number of particles rendered.
+	 *	NOTE: This does not limit the number spawned/updated, only what is drawn.
+	 */
+	UPROPERTY(EditAnywhere, Category = Rendering)
+	uint8 bUseMaxDrawCount : 1;
+
+	UPROPERTY(EditAnywhere, Category=ParticleCutout)
+	TEnumAsByte<enum EOpacitySourceMode> OpacitySourceMode;
+
+	/** Normal generation mode for this emitter LOD. */
+	UPROPERTY(EditAnywhere, Category = Normals)
+	TEnumAsByte<enum EEmitterNormalsMode> EmitterNormalsMode;
+
+	/**
+	* Ensures that movement generated from the orbit module is applied to velocity-aligned particles
+	*/
+	UPROPERTY(EditAnywhere, Category = Emitter)
+	uint8 bOrbitModuleAffectsVelocityAlignment : 1;
+
 	/** The number of sub-images horizontally in the texture							*/
 	UPROPERTY(EditAnywhere, Category=SubUV)
 	int32 SubImages_Horizontal;
@@ -306,10 +314,6 @@
 	UPROPERTY(EditAnywhere, Category=SubUV)
 	int32 SubImages_Vertical;
 
-	/** Whether to scale the UV or not - ie, the model wasn't setup with sub uvs		*/
-	UPROPERTY(EditAnywhere, Category=SubUV)
-	uint32 bScaleUV:1;
-
 	/**
 	 *	The amount of time (particle-relative, 0.0 to 1.0) to 'lock' on a random sub image
 	 *	    0.0 = change every frame
@@ -322,10 +326,6 @@
 	UPROPERTY(EditAnywhere, Category=SubUV)
 	int32 RandomImageChanges;
 
-	/** Override the system MacroUV settings                                            */
-	UPROPERTY(EditAnywhere, Category=MacroUV)
-	uint32 bOverrideSystemMacroUV:1;
-
 	/** Local space position that UVs generated with the ParticleMacroUV material node will be centered on. */
 	UPROPERTY(EditAnywhere, Category=MacroUV)
 	FVector MacroUVPosition;
@@ -335,31 +335,11 @@
 	float MacroUVRadius;
 
 	/**
-	 *	If true, use the MaxDrawCount to limit the number of particles rendered.
-	 *	NOTE: This does not limit the number spawned/updated, only what is drawn.
-	 */
-	UPROPERTY(EditAnywhere, Category=Rendering)
-	uint32 bUseMaxDrawCount:1;
-
-	/**
-	 *	The maximum number of particles to DRAW for this emitter.
-	 *	If set to 0, it will use whatever number are present.
-	 */
-	UPROPERTY(EditAnywhere, Category=Rendering)
-	int32 MaxDrawCount;
-
-	/**
 	* Controls UV Flipping for this emitter.
 	*/
-	UPROPERTY(EditAnywhere, Category=Rendering)
+	UPROPERTY(EditAnywhere, Category = Rendering)
 	EParticleUVFlipMode UVFlippingMode;
 
 	/**
-	* Texture to generate bounding geometry from.
-	*/
-	UPROPERTY(EditAnywhere, Category=ParticleCutout)
-	UTexture2D* CutoutTexture;
-
-	/**
 	* More bounding vertices results in reduced overdraw, but adds more triangle overhead.
 	* The eight vertex mode is best used when the SubUV texture has a lot of space to cut out that is not captured by the four vertex version,
@@ -368,19 +348,11 @@
 	UPROPERTY(EditAnywhere, Category=ParticleCutout)
 	TEnumAsByte<enum ESubUVBoundingVertexCount> BoundingMode;
 
-	UPROPERTY(EditAnywhere, Category=ParticleCutout)
-	TEnumAsByte<enum EOpacitySourceMode> OpacitySourceMode;
-
 	/**
-	* Alpha channel values larger than the threshold are considered occupied and will be contained in the bounding geometry.
-	* Raising this threshold slightly can reduce overdraw in particles using this animation asset.
-	*/
-	UPROPERTY(EditAnywhere, Category=ParticleCutout, meta=(UIMin = "0", UIMax = "1"))
-	float AlphaThreshold;
-
-	/** Normal generation mode for this emitter LOD. */
-	UPROPERTY(EditAnywhere, Category=Normals)
-	TEnumAsByte<enum EEmitterNormalsMode> EmitterNormalsMode;
+	 *	If true, recalculate the emitter duration on each loop.
+	 */
+	UPROPERTY(EditAnywhere, Category = Duration)
+	uint8 bDurationRecalcEachLoop : 1;
 
 	/** 
 	 * When EmitterNormalsMode is ENM_Spherical, particle normals are created to face away from NormalsSphereCenter. 
@@ -389,5 +361,37 @@
 	UPROPERTY(EditAnywhere, Category=Normals)
 	FVector NormalsSphereCenter;
 
+	/**
+	* Alpha channel values larger than the threshold are considered occupied and will be contained in the bounding geometry.
+	* Raising this threshold slightly can reduce overdraw in particles using this animation asset.
+	*/
+	UPROPERTY(EditAnywhere, Category = ParticleCutout, meta = (UIMin = "0", UIMax = "1"))
+	float AlphaThreshold;
+
+	/** The number of times to loop the emitter.
+	 *	0 indicates loop continuously
+	 */
+	UPROPERTY(EditAnywhere, Category = Duration)
+	int32 EmitterLoops;
+
+	/**
+	* Texture to generate bounding geometry from.
+	*/
+	UPROPERTY(EditAnywhere, Category = ParticleCutout)
+	UTexture2D* CutoutTexture;
+
+	/**
+	 *	The maximum number of particles to DRAW for this emitter.
+	 *	If set to 0, it will use whatever number are present.
+	 */
+	UPROPERTY(EditAnywhere, Category = Rendering)
+	int32 MaxDrawCount;
+
+	/**
+	 *	The low end of the emitter duration if using a range.
+	 */
+	UPROPERTY(EditAnywhere, Category = Duration, meta = (DisplayAfter = "EmitterDuration"))
+	float EmitterDurationLow;
+
 	/** 
 	 * When EmitterNormalsMode is ENM_Cylindrical, 
@@ -397,12 +401,6 @@
 	UPROPERTY(EditAnywhere, Category=Normals)
 	FVector NormalsCylinderDirection;
 
-	/**
-	 * Ensures that movement generated from the orbit module is applied to velocity-aligned particles
-	 */
-	UPROPERTY(EditAnywhere, Category=Emitter)
-	uint32 bOrbitModuleAffectsVelocityAlignment:1;
-
 	/** 
 	*	Named material overrides for this emitter. 
 	*	Overrides this emitter's material(s) with those in the correspondingly named slot(s) of the owning system.
diff --git a/Engine/Source/Runtime/Engine/Classes/Particles/SubUV/ParticleModuleSubUV.h b/Engine/Source/Runtime/Engine/Classes/Particles/SubUV/ParticleModuleSubUV.h
--- a/Engine/Source/Runtime/Engine/Classes/Particles/SubUV/ParticleModuleSubUV.h
+++ b/Engine/Source/Runtime/Engine/Classes/Particles/SubUV/ParticleModuleSubUV.h
@@ -39,7 +39,7 @@ class UParticleModuleSubUV : public UParticleModuleSubUVBase
 	 *	The movie will update regardless of the slomo settings of the game.
 	 */
 	UPROPERTY(EditAnywhere, Category=Realtime)
-	uint32 bUseRealTime:1;
+	uint8 bUseRealTime:1;
 
 	/** Initializes the default values for this property */
 	void InitializeDefaults();
diff --git a/Engine/Source/Runtime/Engine/Classes/Particles/TypeData/ParticleModuleTypeDataMesh.h b/Engine/Source/Runtime/Engine/Classes/Particles/TypeData/ParticleModuleTypeDataMesh.h
--- a/Engine/Source/Runtime/Engine/Classes/Particles/TypeData/ParticleModuleTypeDataMesh.h
+++ b/Engine/Source/Runtime/Engine/Classes/Particles/TypeData/ParticleModuleTypeDataMesh.h
@@ -14,6 +14,7 @@
 class UParticleEmitter;
 class UParticleSystemComponent;
 struct FParticleEmitterInstance;
+class UStaticMesh;
 
 UENUM()
 enum EMeshScreenAlignment
@@ -64,24 +65,26 @@ class UParticleModuleTypeDataMesh : public UParticleModuleTypeDataBase
 
 	/** The static mesh to render at the particle positions */
 	UPROPERTY(EditAnywhere, Category=Mesh)
-	class UStaticMesh* Mesh;
+	UStaticMesh* Mesh;
 
-	/** use the static mesh's LOD setup and switch LODs based on largest particle's screen size*/
-	UPROPERTY(EditAnywhere, Category = Mesh)
-	bool bUseStaticMeshLODs;
+	/** Random stream for the initial rotation distribution */
+	FRandomStream RandomStream;
 
 	/** use the static mesh's LOD setup and switch LODs based on largest particle's screen size*/
 	UPROPERTY(EditAnywhere, Category = Mesh)
 	float LODSizeScale;
 
+	/** use the static mesh's LOD setup and switch LODs based on largest particle's screen size*/
+	UPROPERTY(EditAnywhere, Category = Mesh)
+	uint8 bUseStaticMeshLODs : 1;
 
 	/** If true, has the meshes cast shadows */
 	UPROPERTY()
-	uint32 CastShadows:1;
+	uint8 CastShadows:1;
 
 	/** UNUSED (the collision module dictates doing collisions) */
 	UPROPERTY()
-	uint32 DoCollisions:1;
+	uint8 DoCollisions:1;
 
 	/** 
 	 *	The alignment to use on the meshes emitted.
@@ -103,30 +106,18 @@ class UParticleModuleTypeDataMesh : public UParticleModuleTypeDataBase
 	 *	to the static mesh model.
 	 */
 	UPROPERTY(EditAnywhere, Category=Mesh)
-	uint32 bOverrideMaterial:1;
+	uint8 bOverrideMaterial:1;
 
 	UPROPERTY(EditAnywhere, Category = Mesh)
-	uint32 bOverrideDefaultMotionBlurSettings : 1;
+	uint8 bOverrideDefaultMotionBlurSettings : 1;
 
 	UPROPERTY(EditAnywhere, Category = Mesh, meta=(EditCondition="bOverrideDefaultMotionBlurSettings"))
-	uint32 bEnableMotionBlur : 1;
-
-	/** deprecated properties for initial orientation */
-	UPROPERTY()
-	float Pitch_DEPRECATED;
-	UPROPERTY()
-	float Roll_DEPRECATED;
-	UPROPERTY()
-	float Yaw_DEPRECATED;
-
+	uint8 bEnableMotionBlur : 1;
 
 	/** The 'pre' rotation pitch (in degrees) to apply to the static mesh used. */
 	UPROPERTY(EditAnywhere, Category=Orientation)
 	FRawDistributionVector RollPitchYawRange;
 
-	/** Random stream for the initial rotation distribution */
-	FRandomStream RandomStream;
-
 	/**
 	 *	The axis to lock the mesh on. This overrides TypeSpecific mesh alignment as well as the LockAxis module.
 	 *		EPAL_NONE		 -	No locking to an axis.
@@ -148,7 +139,7 @@ class UParticleModuleTypeDataMesh : public UParticleModuleTypeDataBase
 	 *	When set, AxisLockOption as well as all other locked axis/screen alignment settings are ignored.
 	 */
 	UPROPERTY(EditAnywhere, Category=CameraFacing)
-	uint32 bCameraFacing:1;
+	uint8 bCameraFacing:1;
 
 	/**
 	 *	The axis of the mesh to point up when camera facing the X-axis.
@@ -189,21 +180,21 @@ class UParticleModuleTypeDataMesh : public UParticleModuleTypeDataBase
 	 *	If false, apply 'sprite' particle rotation about the camera facing axis.
 	 */
 	UPROPERTY(EditAnywhere, Category=CameraFacing)
-	uint32 bApplyParticleRotationAsSpin:1;
+	uint8 bApplyParticleRotationAsSpin : 1;
 
 	/** 
 	*	If true, all camera facing options will point the mesh against the camera's view direction rather than pointing at the cameras location. 
 	*	If false, the camera facing will point to the cameras position as normal.
 	*/
 	UPROPERTY(EditAnywhere, Category=CameraFacing)
-	uint32 bFaceCameraDirectionRatherThanPosition:1;
+	uint8 bFaceCameraDirectionRatherThanPosition : 1;
 
 	/**
 	*	If true, all collisions for mesh particle on this emitter will take the particle size into account.
 	*	If false, particle size will be ignored in collision checks.
 	*/
 	UPROPERTY(EditAnywhere, Category = Collision)
-	uint32 bCollisionsConsiderPartilceSize : 1;
+	uint8 bCollisionsConsiderPartilceSize : 1;
 
 	static int32 GetCurrentDetailMode();
 	static int32 GetMeshParticleMotionBlurMinDetailMode();
diff --git a/Engine/Source/Runtime/Engine/Private/Particles/ParticleModules.cpp b/Engine/Source/Runtime/Engine/Private/Particles/ParticleModules.cpp
--- a/Engine/Source/Runtime/Engine/Private/Particles/ParticleModules.cpp
+++ b/Engine/Source/Runtime/Engine/Private/Particles/ParticleModules.cpp
@@ -158,36 +158,15 @@ void UParticleModule::FinalUpdate(FParticleEmitterInstance* Owner, int32 Offset,
 }
 
 uint32 UParticleModule::RequiredBytes(UParticleModuleTypeDataBase* TypeData)
-{
-PRAGMA_DISABLE_DEPRECATION_WARNINGS
-		// To try and handle deprecation, if function is not overriden, call old function
-	// Will fail though if actually trying to use Owner
-	return RequiredBytes(static_cast<FParticleEmitterInstance*>(nullptr));
-PRAGMA_ENABLE_DEPRECATION_WARNINGS
-}
-
-uint32 UParticleModule::RequiredBytesPerInstance()
-{
-PRAGMA_DISABLE_DEPRECATION_WARNINGS
-	// To try and handle deprecation, if function is not overriden, call old function
-	// Will fail though if actually trying to use Owner
-	return RequiredBytesPerInstance(nullptr);
-PRAGMA_ENABLE_DEPRECATION_WARNINGS
-}
-
-// DEPRECATED 4.11
-uint32 UParticleModule::RequiredBytes(FParticleEmitterInstance* Owner)
 {
 	return 0;
 }
 
-// DEPRECATED 4.11
-uint32 UParticleModule::RequiredBytesPerInstance(FParticleEmitterInstance* Owner)
+uint32 UParticleModule::RequiredBytesPerInstance()
 {
 	return 0;
 }
 
-
 uint32 UParticleModule::PrepPerInstanceBlock(FParticleEmitterInstance* Owner, void* InstData)
 {
 	return 0xffffffff;
@@ -3426,7 +3405,7 @@ void UParticleModuleTypeDataMesh::Serialize(FArchive& Ar)
 	Super::Serialize(Ar);
 	if (Ar.IsLoading() && Ar.UE4Ver() < VER_UE4_MESH_EMITTER_INITIAL_ORIENTATION_DISTRIBUTION)
 	{
-		FVector oldOrient(Roll_DEPRECATED, Pitch_DEPRECATED, Yaw_DEPRECATED);
+		FVector oldOrient(0.0f, 0.0f, 0.0f);
 		CreateDistribution();
 		UDistributionVectorUniform* RPYDistribution = Cast<UDistributionVectorUniform>(RollPitchYawRange.Distribution);
 		RPYDistribution->Min = oldOrient;
diff --git a/Engine/Source/Runtime/Engine/Public/Distributions.h b/Engine/Source/Runtime/Engine/Public/Distributions.h
--- a/Engine/Source/Runtime/Engine/Public/Distributions.h
+++ b/Engine/Source/Runtime/Engine/Public/Distributions.h
@@ -28,6 +28,12 @@ enum ERawDistributionOperation
  */
 struct FDistributionLookupTable
 {
+	/** Time between values in the lookup table */
+	float TimeScale;
+	/** Absolute time of the first value */
+	float TimeBias;
+	/** Values in the table. */
+	TArray<float> Values;
 	/** Operation for which the table was built. */
 	uint8 Op;
 	/** Number of entries in the table. */
@@ -36,22 +42,17 @@ struct FDistributionLookupTable
 	uint8 EntryStride;
 	/** Number of values between sub-entries [0,4]. */
 	uint8 SubEntryStride;
-	/** Time between values in the lookup table */
-	float TimeScale;
-	/** Absolute time of the first value */
-	float TimeBias;
-	/** Values in the table. */
-	TArray<float> Values;
 	/** Lock axes flag for vector distributions. */
 	uint8 LockFlag;
+	
 	/** Default constructor. */
 	FDistributionLookupTable()
-		: Op(RDO_Uninitialized)
+		: TimeScale(0.0f)
+		, TimeBias(0.0f)
+		, Op(RDO_Uninitialized)
 		, EntryCount(0)
 		, EntryStride(0)
 		, SubEntryStride(0)
-		, TimeScale(0.0f)
-		, TimeBias(0.0f)
 		, LockFlag(0)
 	{}
 
