Improve packing on SoundCue and SoundNodeRandom. Add missing override specifier

#rnx

diff --git a/Engine/Source/Runtime/Engine/Classes/Sound/SoundCue.h b/Engine/Source/Runtime/Engine/Classes/Sound/SoundCue.h
--- a/Engine/Source/Runtime/Engine/Classes/Sound/SoundCue.h
+++ b/Engine/Source/Runtime/Engine/Classes/Sound/SoundCue.h
@@ -77,14 +77,6 @@ class ENGINE_API USoundCue : public USoundBase
 {
 	GENERATED_UCLASS_BODY()
 
-	/* Indicates whether attenuation should use the Attenuation Overrides or the Attenuation Settings asset */
-	UPROPERTY(EditAnywhere, Category=Attenuation)
-	uint32 bOverrideAttenuation:1;
-
-	/* Makes this sound cue ignore per-platform random node culling for memory purposes */
-	UPROPERTY(EditAnywhere, Category=Culling)
-	uint32 bExcludeFromRandomNodeBranchCulling:1;
-
 	UPROPERTY()
 	USoundNode* FirstNode;
 
@@ -116,6 +108,16 @@ protected:
 private:
 	float MaxAudibleDistance;
 
+public:
+	/* Indicates whether attenuation should use the Attenuation Overrides or the Attenuation Settings asset */
+	UPROPERTY(EditAnywhere, Category=Attenuation)
+	uint8 bOverrideAttenuation : 1;
+
+	/* Makes this sound cue ignore per-platform random node culling for memory purposes */
+	UPROPERTY(EditAnywhere, Category=Culling)
+	uint8 bExcludeFromRandomNodeBranchCulling : 1;
+
+private:
 	/** Whether a sound has play when silent enabled (i.e. for a sound cue, if any sound wave player has it enabled). */
 	UPROPERTY()
 	uint8 bHasPlayWhenSilent : 1;
diff --git a/Engine/Source/Runtime/Engine/Classes/Sound/SoundNodeRandom.h b/Engine/Source/Runtime/Engine/Classes/Sound/SoundNodeRandom.h
--- a/Engine/Source/Runtime/Engine/Classes/Sound/SoundNodeRandom.h
+++ b/Engine/Source/Runtime/Engine/Classes/Sound/SoundNodeRandom.h
@@ -24,6 +24,17 @@ class USoundNodeRandom : public USoundNode
 	UPROPERTY(EditAnywhere, editfixedsize, Category=Random)
 	TArray<float> Weights;
 
+	/**
+	 * Internal state of which sounds have been played.  This is only used at runtime
+	 * to keep track of which sounds have been played
+	 */
+	UPROPERTY(transient)
+	TArray<bool> HasBeenUsed;
+
+	/** Counter var so we don't have to count all of the used sounds each time we choose a sound **/
+	UPROPERTY(transient)
+	int32 NumRandomUsed;
+
 	/** If greater than 0, then upon each level load such a number of inputs will be randomly selected
 	 *  and the rest will be removed. This can be used to cut down the memory usage of large randomizing
 	 *  cues.
@@ -53,18 +64,7 @@ class USoundNodeRandom : public USoundNode
 	 * So one could play the same sound over and over if the probabilities don't go your way :-)
 	 */
 	UPROPERTY(EditAnywhere, Category=Random)
-	uint32 bRandomizeWithoutReplacement:1;
-
-	/**
-	 * Internal state of which sounds have been played.  This is only used at runtime
-	 * to keep track of which sounds have been played
-	 */
-	UPROPERTY(transient)
-	TArray<bool> HasBeenUsed;
-
-	/** Counter var so we don't have to count all of the used sounds each time we choose a sound **/
-	UPROPERTY(transient)
-	int32 NumRandomUsed;
+	uint8 bRandomizeWithoutReplacement : 1;
 
 #if WITH_EDITORONLY_DATA
 	/** Editor only list of nodes hidden to duplicate behavior of PreselectAtLevelLoad */
@@ -79,7 +79,7 @@ public:
 
 	//~ Begin USoundNode Interface.
 	virtual void ParseNodes( FAudioDevice* AudioDevice, const UPTRINT NodeWaveInstanceHash, FActiveSound& ActiveSound, const FSoundParseParameters& ParseParams, TArray<FWaveInstance*>& WaveInstances ) override;
-	virtual int32 GetNumSounds(const UPTRINT NodeWaveInstanceHash, FActiveSound& ActiveSound) const;
+	virtual int32 GetNumSounds(const UPTRINT NodeWaveInstanceHash, FActiveSound& ActiveSound) const override;
 	virtual int32 GetMaxChildNodes() const override 
 	{ 
 		return MAX_ALLOWED_CHILD_NODES; 
