Improve packing of Materials , experssions, and interfaces to go down a bin

#rnx

diff --git a/Engine/Source/Runtime/Engine/Classes/Materials/Material.h b/Engine/Source/Runtime/Engine/Classes/Materials/Material.h
--- a/Engine/Source/Runtime/Engine/Classes/Materials/Material.h
+++ b/Engine/Source/Runtime/Engine/Classes/Materials/Material.h
@@ -199,7 +199,7 @@ struct FColorMaterialInput : public FMaterialInput
 {
 #if WITH_EDITORONLY_DATA
 	UPROPERTY()
-	uint32 UseConstant:1;
+	uint32 UseConstant : 1;
 
 	UPROPERTY()
 	FColor Constant;
@@ -213,7 +213,7 @@ struct FScalarMaterialInput : public FMaterialInput
 {
 #if WITH_EDITORONLY_DATA
 	UPROPERTY()
-	uint32 UseConstant:1;
+	uint32 UseConstant : 1;
 
 	UPROPERTY()
 	float Constant;
@@ -235,7 +235,7 @@ struct FVectorMaterialInput : public FMaterialInput
 {
 #if WITH_EDITORONLY_DATA
 	UPROPERTY()
-	uint32 UseConstant:1;
+	uint32 UseConstant : 1;
 
 	UPROPERTY()
 	FVector Constant;
@@ -249,7 +249,7 @@ struct FVector2MaterialInput : public FMaterialInput
 {
 #if WITH_EDITORONLY_DATA
 	UPROPERTY()
-	uint32 UseConstant:1;
+	uint32 UseConstant : 1;
 
 	UPROPERTY()
 	float ConstantX;
@@ -411,6 +411,15 @@ private:
 	UPROPERTY(EditAnywhere, Category=Material, AssetRegistrySearchable)
 	TEnumAsByte<enum EMaterialShadingModel> ShadingModel; 
 
+public:
+	/**
+	* If true, translucent materials will cast dynamic shadows according to their opacity.
+	* OpacityMaskClipValue is used as the threshold value.
+	*/
+	UPROPERTY(EditAnywhere, Category = Material, AdvancedDisplay)
+	uint8 bCastDynamicShadowAsMasked : 1;
+
+private:
 	UPROPERTY(AssetRegistrySearchable)
 	FMaterialShadingModelField ShadingModels;
 
@@ -426,13 +435,6 @@ public:
 	UPROPERTY(EditAnywhere, Category=Material, AdvancedDisplay)
 	float OpacityMaskClipValue;
 
-	/**
-	* If true, translucent materials will cast dynamic shadows according to their opacity.
-	* OpacityMaskClipValue is used as the threshold value.
-	*/
-	UPROPERTY(EditAnywhere, Category=Material, AdvancedDisplay)
-	uint32 bCastDynamicShadowAsMasked:1;
-
 	/** Adds to world position in the vertex shader. */
 	UPROPERTY()
 	FVectorMaterialInput WorldPositionOffset;
@@ -490,51 +492,51 @@ public:
 
 	/** Indicates that the material should be rendered in the SeparateTranslucency Pass (not affected by DOF, requires bAllowSeparateTranslucency to be set in .ini). */
 	UPROPERTY(EditAnywhere, Category=Translucency, meta=(DisplayName = "Render After DOF"), AdvancedDisplay)
-	uint32 bEnableSeparateTranslucency:1;
-
-	/** Indicates that the translucent material should not be affected by bloom or DOF. (Note: Depth testing is not available) */
-	UPROPERTY(EditAnywhere, Category = Translucency, meta = (DisplayName = "Mobile Separate Translucency"), AdvancedDisplay)
-	uint32 bEnableMobileSeparateTranslucency : 1;
+	uint8 bEnableSeparateTranslucency : 1;
 
 	/**
 	 * Indicates that the material should be rendered using responsive anti-aliasing. Improves sharpness of small moving particles such as sparks.
 	 * Only use for small moving features because it will cause aliasing of the background.
 	 */
 	UPROPERTY(EditAnywhere, Category=Translucency, meta=(DisplayName = "Responsive AA"), AdvancedDisplay)
-	uint32 bEnableResponsiveAA:1;
+	uint8 bEnableResponsiveAA : 1;
 
 	/** SSR on translucency */
 	UPROPERTY(EditAnywhere, Category=Translucency, meta=(DisplayName = "Screen Space Reflections"))
-	uint32 bScreenSpaceReflections:1;
+	uint8 bScreenSpaceReflections : 1;
 
 	/** Contact shadows on translucency */
 	UPROPERTY(EditAnywhere, Category = Translucency, meta = (DisplayName = "Contact Shadows"))
-	uint32 bContactShadows : 1;
+	uint8 bContactShadows : 1;
 
 	/** Indicates that the material should be rendered without backface culling and the normal should be flipped for backfaces. */
 	UPROPERTY(EditAnywhere, Category=Material)
-	uint32 TwoSided:1;
+	uint8 TwoSided : 1;
 
 	/** Whether meshes rendered with the material should support dithered LOD transitions. */
 	UPROPERTY(EditAnywhere, Category = Material, AdvancedDisplay, meta = (DisplayName = "Dithered LOD Transition"))
-	uint32 DitheredLODTransition : 1;
+	uint8 DitheredLODTransition : 1;
 
 	/** Dither opacity mask. When combined with Temporal AA this can be used as a form of limited translucency which supports all lighting features. */
 	UPROPERTY(EditAnywhere, Category=Material, AdvancedDisplay)
-	uint32 DitherOpacityMask:1;
+	uint8 DitherOpacityMask : 1;
 
 	/** Whether the material should allow outputting negative emissive color values.  Only allowed on unlit materials. */
 	UPROPERTY(EditAnywhere, Category=Material, AdvancedDisplay)
-	uint32 bAllowNegativeEmissiveColor:1;
-
-	/** Number of customized UV inputs to display.  Unconnected customized UV inputs will just pass through the vertex UVs. */
-	UPROPERTY(EditAnywhere, Category=Material, AdvancedDisplay, meta=(ClampMin=0))
-	int32 NumCustomizedUVs;
+	uint8 bAllowNegativeEmissiveColor : 1;
 
 	/** Sets the lighting mode that will be used on this material if it is translucent. */
 	UPROPERTY(EditAnywhere, AssetRegistrySearchable, Category=Translucency, meta=(DisplayName = "Lighting Mode"))
 	TEnumAsByte<enum ETranslucencyLightingMode> TranslucencyLightingMode;
 
+	/** Indicates that the translucent material should not be affected by bloom or DOF. (Note: Depth testing is not available) */
+	UPROPERTY(EditAnywhere, Category = Translucency, meta = (DisplayName = "Mobile Separate Translucency"), AdvancedDisplay)
+	uint8 bEnableMobileSeparateTranslucency : 1;
+
+	/** Number of customized UV inputs to display.  Unconnected customized UV inputs will just pass through the vertex UVs. */
+	UPROPERTY(EditAnywhere, Category = Material, AdvancedDisplay, meta = (ClampMin = 0))
+	int32 NumCustomizedUVs;
+
 	/** 
 	 * Useful for artificially increasing the influence of the normal on the lighting result for translucency.
 	 * A value larger than 1 increases the influence of the normal, a value smaller than 1 makes the lighting more ambient.
@@ -542,10 +544,6 @@ public:
 	UPROPERTY(EditAnywhere, Category=Translucency, meta=(DisplayName = "Directional Lighting Intensity"))
 	float TranslucencyDirectionalLightingIntensity;
 
-	/** Allows a translucent material to be used with custom depth writing by compiling additional shaders. */
-	UPROPERTY(EditAnywhere, Category=Translucency, AdvancedDisplay, meta=(DisplayName = "Allow Custom Depth Writes"))
-	uint32 AllowTranslucentCustomDepthWrites:1;
-
 	/** Scale used to make translucent shadows more or less opaque than the material's actual opacity. */
 	UPROPERTY(EditAnywhere, Category=TranslucencySelfShadowing, meta=(DisplayName = "Shadow Density Scale"))
 	float TranslucentShadowDensityScale;
@@ -586,34 +584,34 @@ public:
 
 	/** Whether to draw on top of opaque pixels even if behind them. This only has meaning for translucency. */
 	UPROPERTY(EditAnywhere, Category=Translucency, AdvancedDisplay)
-	uint32 bDisableDepthTest:1;
+	uint8 bDisableDepthTest : 1;
 
 	/** Whether the transluency pass should write its alpha, and only the alpha, into the framebuffer */
 	UPROPERTY(EditAnywhere, Category = Translucency, AdvancedDisplay)
-	uint32 bWriteOnlyAlpha : 1;
+	uint8 bWriteOnlyAlpha : 1;
 
 	/** Whether to generate spherical normals for particles that use this material. */
 	UPROPERTY(EditAnywhere, Category=Material, AdvancedDisplay)
-	uint32 bGenerateSphericalParticleNormals:1;
+	uint8 bGenerateSphericalParticleNormals : 1;
 
 	/**
 	 * Whether the material takes a tangent space normal or a world space normal as input.
 	 * (TangentSpace requires extra instructions but is often more convenient).
 	 */
 	UPROPERTY(EditAnywhere, Category=Material, AdvancedDisplay)
-	uint32 bTangentSpaceNormal:1;
+	uint8 bTangentSpaceNormal : 1;
 
 	/**
 	 * If enabled, the material's emissive colour is injected into the LightPropagationVolume
 	 */
 	UPROPERTY(EditAnywhere, BlueprintReadWrite, Category=Material, meta=(DisplayName = "Emissive (Dynamic Area Light)"), AdvancedDisplay)
-	uint32 bUseEmissiveForDynamicAreaLighting : 1;
+	uint8 bUseEmissiveForDynamicAreaLighting : 1;
 
 	/**
 	 * If enabled, the material's opacity defines how much GI is blocked when using the LightPropagationVolume feature
 	 */
 	UPROPERTY(EditAnywhere, BlueprintReadWrite, Category=Material, meta=(DisplayName = "Block Global Illumination"), AdvancedDisplay)
-	uint32 bBlockGI : 1;
+	uint8 bBlockGI : 1;
 
 	/** 
 	 * This is a special usage flag that allows a material to be assignable to any primitive type.
@@ -622,42 +620,42 @@ public:
 	 * This flag should only be enabled when absolutely necessary, and is purposefully not exposed to the UI to prevent abuse.
 	 */
 	UPROPERTY(duplicatetransient)
-	uint32 bUsedAsSpecialEngineMaterial:1;
+	uint8 bUsedAsSpecialEngineMaterial : 1;
 
 	/** 
 	 * Indicates that the material and its instances can be used with skeletal meshes.  
 	 * This will result in the shaders required to support skeletal meshes being compiled which will increase shader compile time and memory usage.
 	 */
 	UPROPERTY(EditAnywhere, BlueprintReadOnly, Category=Usage)
-	uint32 bUsedWithSkeletalMesh:1;
+	uint8 bUsedWithSkeletalMesh : 1;
 
 	/** 
 	 * Indicates that the material and its instances can be used with editor compositing  
 	 * This will result in the shaders required to support editor compositing being compiled which will increase shader compile time and memory usage.
 	 */
 	UPROPERTY(EditAnywhere, BlueprintReadOnly, Category=Usage)
-	uint32 bUsedWithEditorCompositing:1;
+	uint8 bUsedWithEditorCompositing : 1;
 
 	/** 
 	 * Indicates that the material and its instances can be used with particle sprites 
 	 * This will result in the shaders required to support particle sprites being compiled which will increase shader compile time and memory usage.
 	 */
 	UPROPERTY(EditAnywhere, BlueprintReadOnly, Category=Usage)
-	uint32 bUsedWithParticleSprites:1;
+	uint8 bUsedWithParticleSprites : 1;
 
 	/** 
 	 * Indicates that the material and its instances can be used with beam trails
 	 * This will result in the shaders required to support beam trails being compiled which will increase shader compile time and memory usage.
 	 */
 	UPROPERTY(EditAnywhere, BlueprintReadOnly, Category=Usage)
-	uint32 bUsedWithBeamTrails:1;
+	uint8 bUsedWithBeamTrails : 1;
 
 	/** 
 	 * Indicates that the material and its instances can be used with mesh particles
 	 * This will result in the shaders required to support mesh particles being compiled which will increase shader compile time and memory usage.
 	 */
 	UPROPERTY(EditAnywhere, BlueprintReadOnly, Category=Usage)
-	uint32 bUsedWithMeshParticles:1;
+	uint8 bUsedWithMeshParticles : 1;
 
 
 	/**
@@ -665,83 +663,83 @@ public:
 	* This will result in the shaders required to support Niagara sprites being compiled which will increase shader compile time and memory usage.
 	*/
 	UPROPERTY(EditAnywhere, BlueprintReadOnly, Category = Usage)
-	uint32 bUsedWithNiagaraSprites : 1;
+	uint8 bUsedWithNiagaraSprites : 1;
 
 	UPROPERTY(EditAnywhere, BlueprintReadOnly, Category = Usage)
-	uint32 bUsedWithNiagaraRibbons : 1;
+	uint8 bUsedWithNiagaraRibbons : 1;
 
 	UPROPERTY(EditAnywhere, BlueprintReadOnly, Category = Usage)
-	uint32 bUsedWithNiagaraMeshParticles : 1;
+	uint8 bUsedWithNiagaraMeshParticles : 1;
 
 	UPROPERTY(EditAnywhere, BlueprintReadOnly, Category = Usage)
-	uint32 bUsedWithGeometryCache : 1;
+	uint8 bUsedWithGeometryCache : 1;
 
 	/** 
 	 * Indicates that the material and its instances can be used with static lighting
 	 * This will result in the shaders required to support static lighting being compiled which will increase shader compile time and memory usage.
 	 */
 	UPROPERTY(EditAnywhere, BlueprintReadOnly, Category=Usage)
-	uint32 bUsedWithStaticLighting:1;
+	uint8 bUsedWithStaticLighting : 1;
 
 	/** 
 	 * Indicates that the material and its instances can be used with morph targets
 	 * This will result in the shaders required to support morph targets being compiled which will increase shader compile time and memory usage.
 	 */
 	UPROPERTY(EditAnywhere, BlueprintReadOnly, Category=Usage)
-	uint32 bUsedWithMorphTargets:1;
+	uint8 bUsedWithMorphTargets : 1;
 
 	/** 
 	 * Indicates that the material and its instances can be used with spline meshes
 	 * This will result in the shaders required to support spline meshes being compiled which will increase shader compile time and memory usage.
 	 */
 	UPROPERTY(EditAnywhere, BlueprintReadOnly, Category=Usage)
-	uint32 bUsedWithSplineMeshes:1;
+	uint8 bUsedWithSplineMeshes : 1;
 
 	/** 
 	 * Indicates that the material and its instances can be used with instanced static meshes
 	 * This will result in the shaders required to support instanced static meshes being compiled which will increase shader compile time and memory usage.
 	 */
 	UPROPERTY(EditAnywhere, BlueprintReadOnly, Category=Usage)
-	uint32 bUsedWithInstancedStaticMeshes:1;
+	uint8 bUsedWithInstancedStaticMeshes : 1;
 
 	/**
 	 * Indicates that the material and its instances can be use with geometry collections
 	 * This will result in the shaders required to support geometry collections being compiled which will increase shader compile time and memory usage.
 	 */
 	UPROPERTY(EditAnywhere, BlueprintReadOnly, Category = Usage)
-	uint32 bUsedWithGeometryCollections : 1;
+	uint8 bUsedWithGeometryCollections : 1;
 
 	/** 
 	 * Indicates that the material and its instances can be used with distortion
 	 * This will result in the shaders required to support distortion being compiled which will increase shader compile time and memory usage.
 	 */
 	UPROPERTY()
-	uint32 bUsesDistortion:1;
+	uint8 bUsesDistortion : 1;
 
 	/** 
 	 * Indicates that the material and its instances can be used with clothing
 	 * This will result in the shaders required to support clothing being compiled which will increase shader compile time and memory usage.
 	 */
 	UPROPERTY(EditAnywhere, BlueprintReadOnly, Category=Usage)
-	uint32 bUsedWithClothing:1;
+	uint8 bUsedWithClothing : 1;
 
 	/** 
 	 * Indicates that the material and its instances can be used with Slate UI and UMG
 	 * This will result in the shaders required to support UI materials being compiled which will increase shader compile time and memory usage.
 	 */
 	UPROPERTY()
-	uint32 bUsedWithUI_DEPRECATED:1;
+	uint8 bUsedWithUI_DEPRECATED : 1;
 
 	/** 
 	 * Whether to automatically set usage flags based on what the material is applied to in the editor.
 	 * It can be useful to disable this on a base material with many instances, where adding another usage flag accidentally (eg bUsedWithSkeletalMeshes) can add a lot of shader permutations.
 	 */
 	UPROPERTY(EditAnywhere, BlueprintReadOnly, Category=Usage, AdvancedDisplay)
-	uint32 bAutomaticallySetUsageInEditor:1;
+	uint8 bAutomaticallySetUsageInEditor : 1;
 
 	/* Forces the material to be completely rough. Saves a number of instructions and one sampler. */
 	UPROPERTY(EditAnywhere, BlueprintReadOnly, Category=Material, AdvancedDisplay)
-	uint32 bFullyRough:1;
+	uint8 bFullyRough : 1;
 
 	/** 
 	 *	Forces this material to use full (highp) precision in the pixel shader.
@@ -749,26 +747,26 @@ public:
 	 *	This setting has no effect on older mobile devices that do not support high precision. 
 	 */
 	UPROPERTY(EditAnywhere, BlueprintReadOnly, Category = Mobile)
-	uint32 bUseFullPrecision : 1;
+	uint8 bUseFullPrecision : 1;
 
 	/* Use lightmap directionality and per pixel normals. If disabled, lighting from lightmaps will be flat but cheaper. */
 	UPROPERTY(EditAnywhere, BlueprintReadOnly, Category=Mobile)
-	uint32 bUseLightmapDirectionality:1;
+	uint8 bUseLightmapDirectionality : 1;
 
 	/* 
 	 * Forward renderer: enables multiple parallax-corrected reflection captures that blend together.
 	 * Mobile renderer: blend between nearest 3 reflection captures, but reduces the number of samplers available to the material as two more samplers will be used for reflection cubemaps.
 	 */
 	UPROPERTY(EditAnywhere, BlueprintReadOnly, Category = ForwardShading, meta = (DisplayName = "High Quality Reflections"))
-	uint32 bUseHQForwardReflections : 1;
+	uint8 bUseHQForwardReflections : 1;
 
 	/* Enables planar reflection when using the forward renderer or mobile. Enabling this setting reduces the number of samplers available to the material as one more sampler will be used for the planar reflection. */
 	UPROPERTY(EditAnywhere, BlueprintReadOnly, Category = ForwardShading, meta = (DisplayName = "Planar Reflections"))
-	uint32 bUsePlanarForwardReflections : 1;
+	uint8 bUsePlanarForwardReflections : 1;
 
 	/* Reduce roughness based on screen space normal changes. */
 	UPROPERTY(EditAnywhere, BlueprintReadOnly, Category=Material, AdvancedDisplay)
-	uint32 bNormalCurvatureToRoughness : 1;
+	uint8 bNormalCurvatureToRoughness : 1;
 
 	/** The type of tessellation to apply to this object.  Note D3D11 required for anything except MTM_NoTessellation. */
 	UPROPERTY(EditAnywhere, BlueprintReadOnly, Category=Tessellation)
@@ -776,22 +774,23 @@ public:
 
 	/** Prevents cracks in the surface of the mesh when using tessellation. */
 	UPROPERTY(EditAnywhere, BlueprintReadOnly, Category=Tessellation, meta=(DisplayName = "Crack Free Displacement"))
-	uint32 bEnableCrackFreeDisplacement:1;
+	uint8 bEnableCrackFreeDisplacement : 1;
 
 	/** Enables adaptive tessellation, which tries to maintain a uniform number of pixels per triangle. */
 	UPROPERTY(EditAnywhere, BlueprintReadOnly, Category=Tessellation, meta=(DisplayName = "Adaptive Tessellation"))
-	uint32 bEnableAdaptiveTessellation:1;
+	uint8 bEnableAdaptiveTessellation : 1;
 
-	UPROPERTY(EditAnywhere, Category=Tessellation)
-	float MaxDisplacement;
+	/** Allows a translucent material to be used with custom depth writing by compiling additional shaders. */
+	UPROPERTY(EditAnywhere, Category = Translucency, AdvancedDisplay, meta = (DisplayName = "Allow Custom Depth Writes"))
+	uint8 AllowTranslucentCustomDepthWrites : 1;
 
 	/** Enables a wireframe view of the mesh the material is applied to.  */
 	UPROPERTY(EditAnywhere, Category=Material, AdvancedDisplay)
-	uint32 Wireframe:1;
+	uint8 Wireframe : 1;
 
 	/** Skips outputting velocity during the base pass. */
 	UPROPERTY(EditAnywhere, Category=Material, AdvancedDisplay, meta=(DisplayName = "Support accurate velocities from Vertex Deformation"))
-	uint32 bOutputVelocityOnBasePass:1;
+	uint8 bOutputVelocityOnBasePass : 1;
 
 #if WITH_EDITORONLY_DATA
 	UPROPERTY()
@@ -831,67 +830,59 @@ public:
 
 	/** true if this Material can be assumed Opaque when set to masked. */
 	UPROPERTY()
-	uint32 bCanMaskedBeAssumedOpaque : 1;
+	uint8 bCanMaskedBeAssumedOpaque : 1;
 
 	/** true if Material is masked and uses custom opacity */
 	UPROPERTY()
-	uint32 bIsMasked_DEPRECATED:1;
+	uint8 bIsMasked_DEPRECATED : 1;
 
 	/** true if Material is the preview material used in the material editor. */
 	UPROPERTY(transient, duplicatetransient)
-	uint32 bIsPreviewMaterial:1;
+	uint8 bIsPreviewMaterial : 1;
 
 	/** true if Material is the function preview material used in the material instance editor. */
 	UPROPERTY(transient, duplicatetransient)
-	uint32 bIsFunctionPreviewMaterial:1;
+	uint8 bIsFunctionPreviewMaterial : 1;
 
 	/** when true, the material attributes pin is used instead of the regular pins. */
 	UPROPERTY(EditAnywhere, Category=Material)
-	uint32 bUseMaterialAttributes:1;
+	uint8 bUseMaterialAttributes : 1;
 
 	/** when true, the material casts ray tracing shadows. */
 	UPROPERTY(EditAnywhere, Category = Material)
-	uint32 bCastRayTracedShadows : 1;
+	uint8 bCastRayTracedShadows : 1;
 
 	/** When true, translucent materials are fogged. Defaults to true. */
 	UPROPERTY(EditAnywhere, Category=Translucency, meta=(DisplayName = "Apply Fogging"))
-	uint32 bUseTranslucencyVertexFog:1;
+	uint8 bUseTranslucencyVertexFog : 1;
 
 	/** Unlit and Opaque materials can be used as sky material on a sky dome mesh. When IsSky is true, these meshes will not receive any contribution from the aerial perspective. Height and Volumetric fog effects will still be applied. */
 	UPROPERTY(EditAnywhere, Category=Material, AdvancedDisplay)
-	uint32 bIsSky:1;
+	uint8 bIsSky : 1;
 
 	/** When true, translucent materials have fog computed for every pixel, which costs more but fixes artifacts due to low tessellation. */
 	UPROPERTY(EditAnywhere, Category=Translucency)
-	uint32 bComputeFogPerPixel:1;
+	uint8 bComputeFogPerPixel : 1;
 
 	/** When true, translucent materials will output motion vectors in velocity pass. */
 	UPROPERTY(EditAnywhere, Category = Translucency, meta = (DisplayName = "Output Velocity"))
-	uint32 bOutputTranslucentVelocity : 1;
+	uint8 bOutputTranslucentVelocity : 1;
 
 	/** If true the compilation environment will be changed to remove the global COMPILE_SHADERS_FOR_DEVELOPMENT flag. */
 	UPROPERTY(transient, duplicatetransient)
-	uint32 bAllowDevelopmentShaderCompile:1;
+	uint8 bAllowDevelopmentShaderCompile : 1;
 
 	/** true if this is a special material used for stats by the material editor. */
 	UPROPERTY(transient, duplicatetransient)
-	uint32 bIsMaterialEditorStatsMaterial:1;
-
-	/** true if we have printed a warning about material usage for a given usage flag. */
-	UPROPERTY(transient, duplicatetransient)
-	uint32 UsageFlagWarnings;
+	uint8 bIsMaterialEditorStatsMaterial : 1;
 	
 	/** Where the node is inserted in the (post processing) graph, only used if domain is PostProcess */
 	UPROPERTY(EditAnywhere, BlueprintReadWrite, Category=PostProcessMaterial, meta=(DisplayName = "Blendable Location"))
 	TEnumAsByte<enum EBlendableLocation> BlendableLocation;
 
-	/** If multiple nodes with the same  type are inserted at the same point, this defined order and if they get combined, only used if domain is PostProcess */
-	UPROPERTY(EditAnywhere, BlueprintReadWrite, Category=PostProcessMaterial, meta=(DisplayName = "Blendable Priority"))
-	int32 BlendablePriority;
-
 	/** If this is enabled, the blendable will output alpha */
 	UPROPERTY(EditAnywhere, BlueprintReadWrite, Category = PostProcessMaterial, meta = (DisplayName = "Output Alpha"))
-	bool BlendableOutputAlpha;
+	uint8 BlendableOutputAlpha : 1;
 
 	/** 
 	 * Selectively execute post process material only for pixels that pass the stencil test against the Custom Depth/Stencil buffer. 
@@ -910,6 +901,14 @@ public:
 	UPROPERTY(EditAnywhere, Category=Refraction)
 	TEnumAsByte<enum ERefractionMode> RefractionMode;
 
+	/** If multiple nodes with the same  type are inserted at the same point, this defined order and if they get combined, only used if domain is PostProcess */
+	UPROPERTY(EditAnywhere, BlueprintReadWrite, Category = PostProcessMaterial, meta = (DisplayName = "Blendable Priority"))
+	int32 BlendablePriority;
+
+	/** true if we have printed a warning about material usage for a given usage flag. */
+	UPROPERTY(transient, duplicatetransient)
+	uint32 UsageFlagWarnings;
+
 	/** This is the refraction depth bias, larger values offset distortion to prevent closer objects from rendering into the distorted surface at acute viewing angles but increases the disconnect between surface and where the refraction starts. */
 	UPROPERTY(EditAnywhere, Category=Refraction)
 	float RefractionDepthBias;
@@ -923,6 +922,9 @@ public:
 	UPROPERTY()
 	FGuid StateId;
 
+	UPROPERTY(EditAnywhere, Category = Tessellation)
+	float MaxDisplacement;
+
 #if STORE_ONLY_ACTIVE_SHADERMAPS
 	// Relative offset to the beginning of the package containing this
 	uint32 OffsetToFirstResource;
diff --git a/Engine/Source/Runtime/Engine/Classes/Materials/MaterialExpression.h b/Engine/Source/Runtime/Engine/Classes/Materials/MaterialExpression.h
--- a/Engine/Source/Runtime/Engine/Classes/Materials/MaterialExpression.h
+++ b/Engine/Source/Runtime/Engine/Classes/Materials/MaterialExpression.h
@@ -150,7 +150,7 @@ class ENGINE_API UMaterialExpression : public UObject
 
 	/** Indicates that this is a 'parameter' type of expression and should always be loaded (ie not cooked away) because we might want the default parameter. */
 	UPROPERTY()
-	uint32 bIsParameterExpression:1;
+	uint8 bIsParameterExpression : 1;
 
 #if WITH_EDITORONLY_DATA
 	/** If true, the comment bubble will be visible in the graph editor */
diff --git a/Engine/Source/Runtime/Engine/Classes/Materials/MaterialExpressionTextureBase.h b/Engine/Source/Runtime/Engine/Classes/Materials/MaterialExpressionTextureBase.h
--- a/Engine/Source/Runtime/Engine/Classes/Materials/MaterialExpressionTextureBase.h
+++ b/Engine/Source/Runtime/Engine/Classes/Materials/MaterialExpressionTextureBase.h
@@ -28,7 +28,7 @@ class ENGINE_API UMaterialExpressionTextureBase : public UMaterialExpression
 	
 	/** Is default selected texture when using mesh paint mode texture painting */
 	UPROPERTY(EditAnywhere, Category=MaterialExpressionTextureBase)
-	uint32 IsDefaultMeshpaintTexture:1;
+	uint8 IsDefaultMeshpaintTexture:1;
 	
 	//~ Begin UObject Interface
 #if WITH_EDITOR
diff --git a/Engine/Source/Runtime/Engine/Classes/Materials/MaterialExpressionTextureSample.h b/Engine/Source/Runtime/Engine/Classes/Materials/MaterialExpressionTextureSample.h
--- a/Engine/Source/Runtime/Engine/Classes/Materials/MaterialExpressionTextureSample.h
+++ b/Engine/Source/Runtime/Engine/Classes/Materials/MaterialExpressionTextureSample.h
@@ -73,18 +73,24 @@ class ENGINE_API UMaterialExpressionTextureSample : public UMaterialExpressionTe
 	UPROPERTY(EditAnywhere, Category=MaterialExpressionTextureSample)
 	TEnumAsByte<enum ESamplerSourceMode> SamplerSource;
 
+	/** Whether the texture should be sampled with per view mip biasing for sharper output with Temporal AA. */
+	UPROPERTY(EditAnywhere, Category = MaterialExpressionTextureSample)
+	uint8 AutomaticViewMipBias : 1;
+
+protected:
+	// Inherited parameter expressions can hide unused input pin
+	uint8 bShowTextureInputPin : 1;
+
+public:
+
 	/** only used if Coordinates is not hooked up */
 	UPROPERTY(EditAnywhere, Category = MaterialExpressionTextureSample)
-	uint32 ConstCoordinate;
+	uint8 ConstCoordinate;
 
 	/** only used if MipValue is not hooked up */
 	UPROPERTY(EditAnywhere, Category = MaterialExpressionTextureSample)
 	int32 ConstMipValue;
 
-	/** Whether the texture should be sampled with per view mip biasing for sharper output with Temporal AA. */
-	UPROPERTY(EditAnywhere, Category = MaterialExpressionTextureSample)
-	uint32 AutomaticViewMipBias : 1;
-
 	//~ Begin UObject Interface
 #if WITH_EDITOR
 	virtual bool CanEditChange(const UProperty* InProperty) const override;
@@ -114,11 +120,4 @@ class ENGINE_API UMaterialExpressionTextureSample : public UMaterialExpressionTe
 	int32 CompileMipValue0(class FMaterialCompiler* Compiler);
 	int32 CompileMipValue1(class FMaterialCompiler* Compiler);
 #endif // WITH_EDITOR
-
-protected:
-	// Inherited parameter expressions can hide unused input pin
-	bool bShowTextureInputPin;
-};
-
-
-
+};
\ No newline at end of file
diff --git a/Engine/Source/Runtime/Engine/Classes/Materials/MaterialFunction.h b/Engine/Source/Runtime/Engine/Classes/Materials/MaterialFunction.h
--- a/Engine/Source/Runtime/Engine/Classes/Materials/MaterialFunction.h
+++ b/Engine/Source/Runtime/Engine/Classes/Materials/MaterialFunction.h
@@ -14,6 +14,7 @@
 class UMaterial;
 class UTexture;
 struct FPropertyChangedEvent;
+class UMaterialExpression;
 
 /**
  * A Material Function is a collection of material expressions that can be reused in different materials
@@ -33,13 +34,17 @@ class UMaterialFunction : public UMaterialFunctionInterface
 	UPROPERTY(EditAnywhere, Category=MaterialFunction, AssetRegistrySearchable)
 	FString Description;
 
+	/** Array of material expressions, excluding Comments.  Used by the material editor. */
+	UPROPERTY()
+	TArray<UMaterialExpression*> FunctionExpressions;
+
 	/** Whether to list this function in the material function library, which is a window in the material editor that lists categorized functions. */
 	UPROPERTY(EditAnywhere, Category=MaterialFunction, AssetRegistrySearchable)
-	uint32 bExposeToLibrary:1;
+	uint8 bExposeToLibrary:1;
 	
 	/** If true, parameters in this function will have a prefix added to their group name. */
 	UPROPERTY(EditAnywhere, Category=MaterialFunction)
-	uint32 bPrefixParameterNames:1;
+	uint8 bPrefixParameterNames:1;
 
 #if WITH_EDITORONLY_DATA
 	/** 
@@ -56,9 +61,6 @@ class UMaterialFunction : public UMaterialFunctionInterface
 	UPROPERTY(EditAnywhere, Category=MaterialFunction, AssetRegistrySearchable)
 	TArray<FText> LibraryCategoriesText;
 #endif
-	/** Array of material expressions, excluding Comments.  Used by the material editor. */
-	UPROPERTY()
-	TArray<class UMaterialExpression*> FunctionExpressions;
 
 #if WITH_EDITORONLY_DATA
 	/** Array of comments associated with this material; viewed in the material editor. */
@@ -72,7 +74,7 @@ class UMaterialFunction : public UMaterialFunctionInterface
 private:
 	/** Transient flag used to track re-entrance in recursive functions like IsDependent. */
 	UPROPERTY(transient)
-	uint32 bReentrantFlag:1;
+	uint8 bReentrantFlag:1;
 
 public:
 	//~ Begin UObject Interface.
diff --git a/Engine/Source/Runtime/Engine/Classes/Materials/MaterialInstance.h b/Engine/Source/Runtime/Engine/Classes/Materials/MaterialInstance.h
--- a/Engine/Source/Runtime/Engine/Classes/Materials/MaterialInstance.h
+++ b/Engine/Source/Runtime/Engine/Classes/Materials/MaterialInstance.h
@@ -295,9 +295,6 @@ class ENGINE_VTABLE UMaterialInstance : public UMaterialInterface
 	 */
 	ENGINE_API const FStaticParameterSet& GetStaticParameters() const;
 
-	/** Flag to detect cycles in the material instance graph. */
-	bool ReentrantFlag[2];
-
 	/**
 	 * Indicates whether the instance has static permutation resources (which are required when static parameters are present) 
 	 * Read directly from the rendering thread, can only be modified with the use of a FMaterialUpdateContext.
@@ -316,7 +313,9 @@ class ENGINE_VTABLE UMaterialInstance : public UMaterialInterface
 	uint8 bIsShadingModelFromMaterialExpression : 1;
 
 	TEnumAsByte<EBlendMode> BlendMode;
-	FMaterialShadingModelField ShadingModels;
+
+	//Cached copies of the base property overrides or the value from the parent to avoid traversing the parent chain for each access.
+	float OpacityMaskClipValue;
 
 	FORCEINLINE bool GetReentrantFlag() const
 	{
@@ -357,8 +356,10 @@ class ENGINE_VTABLE UMaterialInstance : public UMaterialInterface
 	uint32 OffsetToFirstResource;
 #endif
 
-	//Cached copies of the base property overrides or the value from the parent to avoid traversing the parent chain for each access.
-	float OpacityMaskClipValue;
+	FMaterialShadingModelField ShadingModels;
+
+	/** Flag to detect cycles in the material instance graph. */
+	bool ReentrantFlag[2];
 
 	/** 
 	 * FMaterialRenderProxy derivative that represent this material instance to the renderer, when the renderer needs to fetch parameter values. 
diff --git a/Engine/Source/Runtime/Engine/Classes/Materials/MaterialInterface.h b/Engine/Source/Runtime/Engine/Classes/Materials/MaterialInterface.h
--- a/Engine/Source/Runtime/Engine/Classes/Materials/MaterialInterface.h
+++ b/Engine/Source/Runtime/Engine/Classes/Materials/MaterialInterface.h
@@ -117,10 +117,6 @@ struct FLightmassMaterialInterfaceSettings
 {
 	GENERATED_USTRUCT_BODY()
 
-	/** If true, forces translucency to cast static shadows as if the material were masked. */
-	UPROPERTY(EditAnywhere, Category=Material)
-	uint32 bCastShadowAsMasked:1;
-
 	/** Scales the emissive contribution of this material to static lighting. */
 	UPROPERTY()
 	float EmissiveBoost;
@@ -136,28 +132,32 @@ struct FLightmassMaterialInterfaceSettings
 	UPROPERTY(EditAnywhere, Category=Material)
 	float ExportResolutionScale;
 
+	/** If true, forces translucency to cast static shadows as if the material were masked. */
+	UPROPERTY(EditAnywhere, Category = Material)
+	uint8 bCastShadowAsMasked : 1;
+
 	/** Boolean override flags - only used in MaterialInstance* cases. */
 	/** If true, override the bCastShadowAsMasked setting of the parent material. */
 	UPROPERTY()
-	uint32 bOverrideCastShadowAsMasked:1;
+	uint8 bOverrideCastShadowAsMasked:1;
 
 	/** If true, override the emissive boost setting of the parent material. */
 	UPROPERTY()
-	uint32 bOverrideEmissiveBoost:1;
+	uint8 bOverrideEmissiveBoost:1;
 
 	/** If true, override the diffuse boost setting of the parent material. */
 	UPROPERTY()
-	uint32 bOverrideDiffuseBoost:1;
+	uint8 bOverrideDiffuseBoost:1;
 
 	/** If true, override the export resolution scale setting of the parent material. */
 	UPROPERTY()
-	uint32 bOverrideExportResolutionScale:1;
+	uint8 bOverrideExportResolutionScale:1;
 
 	FLightmassMaterialInterfaceSettings()
-		: bCastShadowAsMasked(false)
-		, EmissiveBoost(1.0f)
+		: EmissiveBoost(1.0f)
 		, DiffuseBoost(1.0f)
 		, ExportResolutionScale(1.0f)
+		, bCastShadowAsMasked(false)
 		, bOverrideCastShadowAsMasked(false)
 		, bOverrideEmissiveBoost(false)
 		, bOverrideDiffuseBoost(false)
@@ -230,10 +230,6 @@ protected:
 	UPROPERTY(EditAnywhere, Category=Lightmass)
 	struct FLightmassMaterialInterfaceSettings LightmassSettings;
 
-private:
-	/** Feature levels to force to compile. */
-	uint32 FeatureLevelsToForceCompile;
-
 protected:
 #if WITH_EDITORONLY_DATA
 	/** Because of redirector, the texture names need to be resorted at each load in case they changed. */
@@ -251,6 +247,10 @@ protected:
 	UPROPERTY(EditAnywhere, AdvancedDisplay, Instanced, Category = Material)
 	TArray<UAssetUserData*> AssetUserData;
 
+private:
+	/** Feature levels to force to compile. */
+	uint32 FeatureLevelsToForceCompile;
+
 public:
 
 	//~ Begin IInterface_AssetUserData Interface
