Improve packing of various components and the particle system

#rnx

diff --git a/Engine/Source/Runtime/Engine/Classes/Components/ArrowComponent.h b/Engine/Source/Runtime/Engine/Classes/Components/ArrowComponent.h
--- a/Engine/Source/Runtime/Engine/Classes/Components/ArrowComponent.h
+++ b/Engine/Source/Runtime/Engine/Classes/Components/ArrowComponent.h
@@ -26,17 +26,17 @@ class ENGINE_API UArrowComponent : public UPrimitiveComponent
 	UPROPERTY(EditAnywhere, BlueprintReadWrite, Category=ArrowComponent)
 	float ArrowSize;
 
-	/** Set to limit the screen size of this arrow */
-	UPROPERTY(EditAnywhere, BlueprintReadWrite, Category=ArrowComponent)
-	bool bIsScreenSizeScaled;
-
 	/** The size on screen to limit this arrow to (in screen space) */
-	UPROPERTY(EditAnywhere, BlueprintReadWrite, Category=ArrowComponent)
+	UPROPERTY(EditAnywhere, BlueprintReadWrite, Category = ArrowComponent)
 	float ScreenSize;
 
+	/** Set to limit the screen size of this arrow */
+	UPROPERTY(EditAnywhere, BlueprintReadWrite, Category=ArrowComponent)
+	uint8 bIsScreenSizeScaled : 1;
+
 	/** If true, don't show the arrow when EngineShowFlags.BillboardSprites is disabled. */
 	UPROPERTY(EditAnywhere, BlueprintReadWrite, Category=ArrowComponent)
-	uint32 bTreatAsASprite:1;
+	uint8 bTreatAsASprite : 1;
 
 #if WITH_EDITORONLY_DATA
 	/** Sprite category that the arrow component belongs to, if being treated as a sprite. Value serves as a key into the localization file. */
diff --git a/Engine/Source/Runtime/Engine/Classes/Particles/ParticleSystem.h b/Engine/Source/Runtime/Engine/Classes/Particles/ParticleSystem.h
--- a/Engine/Source/Runtime/Engine/Classes/Particles/ParticleSystem.h
+++ b/Engine/Source/Runtime/Engine/Classes/Particles/ParticleSystem.h
@@ -12,6 +12,10 @@
 
 enum class EParticleSignificanceLevel : uint8;
 enum class EParticleSystemInsignificanceReaction: uint8;
+class UInterpCurveEdSetup;
+class UMaterialInterface;
+class UParticleSystemComponent;
+class UParticleEmitter;
 
 /**
  *	ParticleSystemUpdateMode
@@ -100,7 +104,7 @@
 	FName Name;
 
 	UPROPERTY(EditAnywhere, Category = NamedMaterial)
-	class UMaterialInterface* Material;
+	UMaterialInterface* Material;
 };
 
 UCLASS(Abstract, MinimalAPI, BlueprintType)
@@ -132,9 +136,6 @@
 
 	}
 
-	UPROPERTY(EditAnywhere, Category=ParticleSystem, AssetRegistrySearchable)
-	TEnumAsByte<enum EParticleSystemUpdateMode> SystemUpdateMode;
-
 	/** UpdateTime_FPS	- the frame per second to update at in FixedTime mode		*/
 	UPROPERTY(EditAnywhere, Category=ParticleSystem)
 	float UpdateTime_FPS;
@@ -159,11 +160,11 @@
 
 	/** Emitters	- internal - the array of emitters in the system				*/
 	UPROPERTY(instanced)
-	TArray<class UParticleEmitter*> Emitters;
+	TArray<UParticleEmitter*> Emitters;
 
 	/** The component used to preview the particle system in Cascade				*/
 	UPROPERTY(transient)
-	class UParticleSystemComponent* PreviewComponent;
+	UParticleSystemComponent* PreviewComponent;
 
 #if WITH_EDITORONLY_DATA
 	/** The angle to use when rendering the thumbnail image							*/
@@ -181,11 +182,7 @@
 #endif // WITH_EDITORONLY_DATA
 	/** Used for curve editor to remember curve-editing setup.						*/
 	UPROPERTY(export)
-	class UInterpCurveEdSetup* CurveEdSetup;
+	UInterpCurveEdSetup* CurveEdSetup;
 
-	/** If true, the system's Z axis will be oriented toward the camera				*/
-	UPROPERTY(EditAnywhere, Category=ParticleSystem)
-	uint32 bOrientZAxisTowardCamera:1;
-
 	//
 	//	LOD
@@ -196,14 +193,9 @@
 	UPROPERTY(EditAnywhere, Category=LOD, AssetRegistrySearchable)
 	float LODDistanceCheckTime;
 
-	/**
-	 *	The method of LOD level determination to utilize for this particle system
-	 *	  PARTICLESYSTEMLODMETHOD_Automatic - Automatically set the LOD level, checking every LODDistanceCheckTime seconds.
-	 *    PARTICLESYSTEMLODMETHOD_DirectSet - LOD level is directly set by the game code.
-	 *    PARTICLESYSTEMLODMETHOD_ActivateAutomatic - LOD level is determined at Activation time, then left alone unless directly set by game code.
-	 */
-	UPROPERTY(EditAnywhere, Category=LOD)
-	TEnumAsByte<enum ParticleSystemLODMethod> LODMethod;
+	/** World space radius that UVs generated with the ParticleMacroUV material node will tile based on. */
+	UPROPERTY(EditAnywhere, Category = MacroUV)
+	float MacroUVRadius;
 
 	/**
 	 *	The array of distances for each LOD level in the system.
@@ -228,20 +220,9 @@
 	int32 EditorLODSetting;
 
 #endif // WITH_EDITORONLY_DATA
-	/**
-	 *	Internal value that tracks the regenerate LOD levels preference.
-	 *	If true, when autoregenerating LOD levels in code, the low level will
-	 *	be a duplicate of the high.
-	 */
-	UPROPERTY()
-	uint32 bRegenerateLODDuplicate:1;
 
 	UPROPERTY(EditAnywhere, Category=LOD)
-	TArray<struct FParticleSystemLOD> LODSettings;
+	TArray<FParticleSystemLOD> LODSettings;
 
-	/** Whether to use the fixed relative bounding box or calculate it every frame. */
-	UPROPERTY(EditAnywhere, Category=Bounds)
-	uint32 bUseFixedRelativeBoundingBox:1;
-
 	/**	Fixed relative bounding box for particle system.							*/
 	UPROPERTY(EditAnywhere, Category=Bounds)
@@ -278,27 +259,44 @@
 	FColor BackgroundColor;
 
 #endif // WITH_EDITORONLY_DATA
+	
+	/** How long this Particle system should delay when ActivateSystem is called on it. */
+	UPROPERTY(EditAnywhere, Category=Delay, AssetRegistrySearchable)
+	float Delay;
+
+	/** The low end of the emitter delay if using a range. */
+	UPROPERTY(EditAnywhere, Category=Delay)
+	float DelayLow;
+
+	/** If true, the system's Z axis will be oriented toward the camera				*/
+	UPROPERTY(EditAnywhere, Category = ParticleSystem)
+	uint8 bOrientZAxisTowardCamera : 1;
+
+	/** Whether to use the fixed relative bounding box or calculate it every frame. */
+	UPROPERTY(EditAnywhere, Category = Bounds)
+	uint8 bUseFixedRelativeBoundingBox : 1;
+
 	/** EDITOR ONLY: Indicates that Cascade would like to have the PeakActiveParticles count reset */
 	UPROPERTY()
-	uint32 bShouldResetPeakCounts:1;
+	uint8 bShouldResetPeakCounts : 1;
 
 	/** Set during load time to indicate that physics is used... */
 	UPROPERTY(transient)
-	uint32 bHasPhysics:1;
+	uint8 bHasPhysics : 1;
 
 	/** Inidicates the old 'real-time' thumbnail rendering should be used	*/
 	UPROPERTY(EditAnywhere, Category=Thumbnail)
-	uint32 bUseRealtimeThumbnail:1;
+	uint8 bUseRealtimeThumbnail : 1;
 
 	/** Internal: Indicates the PSys thumbnail image is out of date			*/
 	UPROPERTY()
-	uint32 ThumbnailImageOutOfDate:1;
+	uint8 ThumbnailImageOutOfDate : 1;
 
 private:
 	/** if true, this psys can tick in any thread **/
-	uint32 bIsElligibleForAsyncTick:1;
+	uint8 bIsElligibleForAsyncTick : 1;
 	/** if true, bIsElligibleForAsyncTick is set up **/
-	uint32 bIsElligibleForAsyncTickComputed:1;
+	uint8 bIsElligibleForAsyncTickComputed : 1;
 public:
 
 #if WITH_EDITORONLY_DATA
@@ -307,58 +305,73 @@
 	class UTexture2D* ThumbnailImage;
 
 #endif // WITH_EDITORONLY_DATA
-	/** How long this Particle system should delay when ActivateSystem is called on it. */
-	UPROPERTY(EditAnywhere, Category=Delay, AssetRegistrySearchable)
-	float Delay;
-
-	/** The low end of the emitter delay if using a range. */
-	UPROPERTY(EditAnywhere, Category=Delay)
-	float DelayLow;
 
 	/**
 	 *	If true, select the emitter delay from the range 
 	 *		[DelayLow..Delay]
 	 */
 	UPROPERTY(EditAnywhere, Category=Delay)
-	uint32 bUseDelayRange:1;
+	uint8 bUseDelayRange : 1;
 
 	UPROPERTY(EditAnywhere, Category = Performance, meta = (ToolTip = "Whether or not to allow instances of this system to have their ticks managed."), AdvancedDisplay)
-	uint32 bAllowManagedTicking:1;
+	uint8 bAllowManagedTicking : 1;
 
 	UPROPERTY(EditAnywhere, Category = Performance, meta = (ToolTip = "Auto-deactivate system if all emitters are determined to not spawn particles again, regardless of lifetime."))
-	bool bAutoDeactivate;
+	uint8 bAutoDeactivate : 1;
 
-	UPROPERTY(EditAnywhere, Category = Performance, meta = (ToolTip = "Minimum duration between ticks; 33=tick at max. 30FPS, 16=60FPS, 8=120FPS"))
-	uint32 MinTimeBetweenTicks;
+	/**
+	 *	Internal value that tracks the regenerate LOD levels preference.
+	 *	If true, when autoregenerating LOD levels in code, the low level will
+	 *	be a duplicate of the high.
+	 */	
+	UPROPERTY()
+	uint8 bRegenerateLODDuplicate : 1;
 
+	UPROPERTY(EditAnywhere, Category = ParticleSystem, AssetRegistrySearchable)
+	TEnumAsByte<enum EParticleSystemUpdateMode> SystemUpdateMode;
+
+	/**
+	 *	The method of LOD level determination to utilize for this particle system
+	 *	  PARTICLESYSTEMLODMETHOD_Automatic - Automatically set the LOD level, checking every LODDistanceCheckTime seconds.
+	 *    PARTICLESYSTEMLODMETHOD_DirectSet - LOD level is directly set by the game code.
+	 *    PARTICLESYSTEMLODMETHOD_ActivateAutomatic - LOD level is determined at Activation time, then left alone unless directly set by game code.
+	 */
+	UPROPERTY(EditAnywhere, Category=LOD)
+	TEnumAsByte<enum ParticleSystemLODMethod> LODMethod;
+
 	/** The reaction this system takes when all emitters are insignificant. */
 	UPROPERTY(EditAnywhere, Category = Performance)
 	EParticleSystemInsignificanceReaction InsignificantReaction;
 
-	/** Time delay between all emitters becoming insignificant and the systems insignificant reaction. */
-	UPROPERTY(EditAnywhere, Category = Performance)
-	float InsignificanceDelay;
-
-	/** The maximum level of significance for emitters in this system. Any emitters with a higher significance will be capped at this significance level. */
-	UPROPERTY(EditAnywhere, Category = Performance)
-	EParticleSignificanceLevel MaxSignificanceLevel;
-
-	/** Local space position that UVs generated with the ParticleMacroUV material node will be centered on. */
-	UPROPERTY(EditAnywhere, Category=MacroUV)
-	FVector MacroUVPosition;
-
-	/** World space radius that UVs generated with the ParticleMacroUV material node will tile based on. */
-	UPROPERTY(EditAnywhere, Category=MacroUV)
-	float MacroUVRadius;
-
-	/** 
+	/**
 	 *	Which occlusion bounds method to use for this particle system.
 	 *	EPSOBM_None - Don't determine occlusion for this system.
 	 *	EPSOBM_ParticleBounds - Use the bounds of the component when determining occlusion.
 	 */
-	UPROPERTY(EditAnywhere, Category=Occlusion)
+	UPROPERTY(EditAnywhere, Category = Occlusion)
 	TEnumAsByte<enum EParticleSystemOcclusionBoundsMethod> OcclusionBoundsMethod;
 
+private:
+	/** Does any emitter loop forever? */
+	uint8 bAnyEmitterLoopsForever : 1;
+
+public:
+
+	/** The maximum level of significance for emitters in this system. Any emitters with a higher significance will be capped at this significance level. */
+	UPROPERTY(EditAnywhere, Category = Performance)
+	EParticleSignificanceLevel MaxSignificanceLevel;
+
+	UPROPERTY(EditAnywhere, Category = Performance, meta = (ToolTip = "Minimum duration between ticks; 33=tick at max. 30FPS, 16=60FPS, 8=120FPS"))
+	uint32 MinTimeBetweenTicks;
+
+	/** Time delay between all emitters becoming insignificant and the systems insignificant reaction. */
+	UPROPERTY(EditAnywhere, Category = Performance)
+	float InsignificanceDelay;
+
+	/** Local space position that UVs generated with the ParticleMacroUV material node will be centered on. */
+	UPROPERTY(EditAnywhere, Category=MacroUV)
+	FVector MacroUVPosition;
+
 	/** The occlusion bounds to use if OcclusionBoundsMethod is set to EPSOBM_CustomBounds */
 	UPROPERTY(EditAnywhere, Category=Occlusion)
 	FBox CustomOcclusionBounds;
@@ -384,7 +397,7 @@
 	virtual void GetAssetRegistryTags(TArray<FAssetRegistryTag>& OutTags) const override;
 	bool UsesCPUCollision() const;
 	virtual bool CanBeClusterRoot() const override;
-	virtual void Serialize(FArchive& Ar);
+	virtual void Serialize(FArchive& Ar) override;
 
 	//~ End UObject Interface.
 
@@ -551,15 +564,10 @@
 	/** The lowest significance of any emitter. Clamped by MaxSignificanceLevel.*/
 	EParticleSignificanceLevel LowestSignificance;
 	
-	uint32 bShouldManageSignificance : 1;
+	uint8 bShouldManageSignificance : 1;
 
-	/** Does any emitter loop forever? */
-	uint32 bAnyEmitterLoopsForever : 1;
-	/** Does any emiter never die due to infinie looping AND indefinite duration? */
-	uint32 bIsImmortal : 1;
+	/** Does any emitter never die due to infinite looping AND indefinite duration? */
+	uint8 bIsImmortal : 1;
 	/** Does any emitter ever become a zombie (is immortal AND stops spawning at some point, i.e. is burst only)? */
-	uint32 bWillBecomeZombie : 1;
+	uint8 bWillBecomeZombie : 1;
 };
-
-
-
diff --git a/Engine/Source/Runtime/Engine/Classes/Particles/Spawn/ParticleModuleSpawnPerUnit.h b/Engine/Source/Runtime/Engine/Classes/Particles/Spawn/ParticleModuleSpawnPerUnit.h
--- a/Engine/Source/Runtime/Engine/Classes/Particles/Spawn/ParticleModuleSpawnPerUnit.h
+++ b/Engine/Source/Runtime/Engine/Classes/Particles/Spawn/ParticleModuleSpawnPerUnit.h
@@ -10,6 +10,7 @@
 #include "ParticleModuleSpawnPerUnit.generated.h"
 
 struct FParticleEmitterInstance;
+struct FRawDistributionFloat;
 
 UCLASS(editinlinenew, MinimalAPI, hidecategories=Object, meta=(DisplayName = "Spawn PerUnit"))
 class UParticleModuleSpawnPerUnit : public UParticleModuleSpawnBase
@@ -24,21 +25,6 @@ class UParticleModuleSpawnPerUnit : public UParticleModuleSpawnBase
 	UPROPERTY(EditAnywhere, Category=Spawn)
 	float UnitScalar;
 
-	/** 
-	 *	The amount to spawn per meter distribution.
-	 *	The value is retrieved using the EmitterTime.
-	 */
-	UPROPERTY(EditAnywhere, Category=Spawn)
-	struct FRawDistributionFloat SpawnPerUnit;
-
-	/**
-	 *	If true, process the default spawn rate when not moving...
-	 *	When not moving, skip the default spawn rate.
-	 *	If false, return the bProcessSpawnRate setting.
-	 */
-	UPROPERTY(EditAnywhere, Category=Spawn)
-	uint32 bIgnoreSpawnRateWhenMoving:1;
-
 	/**
 	 *	The tolerance for moving vs. not moving w.r.t. the bIgnoreSpawnRateWhenMoving flag.
 	 *	Ie, if (DistanceMoved < (UnitScalar x MovementTolerance)) then consider it not moving.
@@ -46,28 +32,43 @@ class UParticleModuleSpawnPerUnit : public UParticleModuleSpawnBase
 	UPROPERTY(EditAnywhere, Category=Spawn)
 	float MovementTolerance;
 
+	/** 
+	 *	The amount to spawn per meter distribution.
+	 *	The value is retrieved using the EmitterTime.
+	 */
+	UPROPERTY(EditAnywhere, Category=Spawn)
+	FRawDistributionFloat SpawnPerUnit;
+
 	/**
 	 *	The maximum valid movement for a single frame.
 	 *	If 0.0, then the check is not performed.
 	 *	Currently, if the distance moved between frames is greater than this
 	 *	then NO particles will be spawned.
-	 *	This is primiarily intended to cover cases where the PSystem is 
+	 *	This is primiarily intended to cover cases where the PSystem is
 	 *	attached to teleporting objects.
 	 */
-	UPROPERTY(EditAnywhere, Category=Spawn)
+	UPROPERTY(EditAnywhere, Category = Spawn)
 	float MaxFrameDistance;
 
+	/**
+	 *	If true, process the default spawn rate when not moving...
+	 *	When not moving, skip the default spawn rate.
+	 *	If false, return the bProcessSpawnRate setting.
+	 */
+	UPROPERTY(EditAnywhere, Category=Spawn)
+	uint8 bIgnoreSpawnRateWhenMoving : 1;
+
 	/** If true, ignore the X-component of the movement */
 	UPROPERTY(EditAnywhere, Category=Spawn)
-	uint32 bIgnoreMovementAlongX:1;
+	uint8 bIgnoreMovementAlongX : 1;
 
 	/** If true, ignore the Y-component of the movement */
 	UPROPERTY(EditAnywhere, Category=Spawn)
-	uint32 bIgnoreMovementAlongY:1;
+	uint8 bIgnoreMovementAlongY : 1;
 
 	/** If true, ignore the Z-component of the movement */
 	UPROPERTY(EditAnywhere, Category=Spawn)
-	uint32 bIgnoreMovementAlongZ:1;
+	uint8 bIgnoreMovementAlongZ : 1;
 
 	/** Initializes the default values for this property */
 	void InitializeDefaults();
diff --git a/Engine/Source/Runtime/Engine/Private/Particles/ParticleComponents.cpp b/Engine/Source/Runtime/Engine/Private/Particles/ParticleComponents.cpp
--- a/Engine/Source/Runtime/Engine/Private/Particles/ParticleComponents.cpp
+++ b/Engine/Source/Runtime/Engine/Private/Particles/ParticleComponents.cpp
@@ -2087,6 +2087,7 @@ UParticleSystem::UParticleSystem(const FObjectInitializer& ObjectInitializer)
 	, OcclusionBoundsMethod(EPSOBM_ParticleBounds)
 	, HighestSignificance(EParticleSignificanceLevel::Critical)
 	, LowestSignificance(EParticleSignificanceLevel::Low)
+	, bShouldManageSignificance(false)
 	, bAnyEmitterLoopsForever(false)
 	, bIsImmortal(false)
 	, bWillBecomeZombie(false)
@@ -2127,7 +2128,6 @@ UParticleSystem::UParticleSystem(const FObjectInitializer& ObjectInitializer)
 	InsignificanceDelay = 0.0f;
 	MaxSignificanceLevel = EParticleSignificanceLevel::Critical;
 	MaxPoolSize = 32;
-	bShouldManageSignificance = false;
 
 
 	bAllowManagedTicking = true;
