Copying //UE4/Dev-Framework to Dev-Main (//UE4/Dev-Main) @ 8491425
#rnx

diff --git a/Engine/Documentation/Source/Shared/GraphNodes/Blueprint/UKismetMathLibrary/UKismetMathLibrary.INT.udn b/Engine/Documentation/Source/Shared/GraphNodes/Blueprint/UKismetMathLibrary/UKismetMathLibrary.INT.udn
--- a/Engine/Documentation/Source/Shared/GraphNodes/Blueprint/UKismetMathLibrary/UKismetMathLibrary.INT.udn
+++ b/Engine/Documentation/Source/Shared/GraphNodes/Blueprint/UKismetMathLibrary/UKismetMathLibrary.INT.udn
@@ -224,7 +224,7 @@ Return Value - This outputs A's string if "Pick A" was true, otherwise it output
 
 The "FInterp to" (float interpolation) node allows you to smoothly transition a float value to a different target value. When utilized, this node will take your current float value and (based on a time-step and speed) output a new float approaching (or equal to) the target value.
 
-NOTE: That the output value should be used as the "Current" input for the following frame. If the input value sits unchanged, then this node will make the same step each frame (resulting in the same output every time, failing to ever reach the "Target" value).
+NOTE: That the output value should be used as the "Current" input for the following frame. If the input value sits unchanged, then this node will make the same step each frame (resulting in the same output every time, failing to ever reach the "Target" value). If a Speed of 0 is given, then the result will be the target. 
 
 
 ---
@@ -252,7 +252,7 @@ Return Value - This outputs an interpolated float value approaching (or equal to
 
 The "RInterp to" (Rotator interpolation) node allows you to smoothly transition from one Rotator to another. When utilized, this node will take your current Rotator and (based on a time-step and speed) output a new Rotator approaching (or equal to) the target one.
 
-NOTE: That the output value should be used as the "Current" input for the following frame. If the input value sits unchanged, then this node will make the same step each frame (resulting in the same output every time, failing to ever reach the "Target" value).
+NOTE: That the output value should be used as the "Current" input for the following frame. If the input value sits unchanged, then this node will make the same step each frame (resulting in the same output every time, failing to ever reach the "Target" value). If a Speed of 0 is given, then the result will be the target. 
 
 
 ---
@@ -280,7 +280,7 @@ Return Value - This outputs an interpolated Rotator approaching (or equal to) th
 
 The "TInterp to" (Transform interpolation) node allows you to smoothly transition from one Transform to another. When utilized, this node will take your current Transform and (based on a time-step and speed) output a new Transform approaching (or equal to) the target one.
 
-NOTE: That the output value should be used as the "Current" input for the following frame. If the input value sits unchanged, then this node will make the same step each frame (resulting in the same output every time, failing to ever reach the "Target" value).
+NOTE: That the output value should be used as the "Current" input for the following frame. If the input value sits unchanged, then this node will make the same step each frame (resulting in the same output every time, failing to ever reach the "Target" value). If a Speed of 0 is given, then the result will be the target. 
 
 
 ---
@@ -308,7 +308,7 @@ Return Value - This outputs an interpolated Transform approaching (or equal to)
 
 The "VInterp to" (Vector interpolation) node allows you to smoothly transition from one Vector to another. When utilized, this node will take your current Vector and (based on a time-step and speed) output a new Vector approaching (or equal to) the target one.
 
-NOTE: That the output value should be used as the "Current" input for the following frame. If the input value sits unchanged, then this node will make the same step each frame (resulting in the same output every time, failing to ever reach the "Target" value).
+NOTE: That the output value should be used as the "Current" input for the following frame. If the input value sits unchanged, then this node will make the same step each frame (resulting in the same output every time, failing to ever reach the "Target" value). If a Speed of 0 is given, then the result will be the target.  If a Speed of 0 is given, then the result will be the target. 
 
 
 ---
diff --git a/Engine/Plugins/Experimental/AlembicImporter/Source/AlembicLibrary/Public/AbcImportSettings.h b/Engine/Plugins/Experimental/AlembicImporter/Source/AlembicLibrary/Public/AbcImportSettings.h
--- a/Engine/Plugins/Experimental/AlembicImporter/Source/AlembicLibrary/Public/AbcImportSettings.h
+++ b/Engine/Plugins/Experimental/AlembicImporter/Source/AlembicLibrary/Public/AbcImportSettings.h
@@ -21,6 +21,7 @@ enum class EAlembicImportType : uint8
 UENUM(BlueprintType)
 enum class EBaseCalculationType : uint8
 {
+	None = 0 UMETA(Hidden),
 	/** Determines the number of bases that should be used with the given percentage*/
 	PercentageBased = 1,
 	/** Set a fixed number of bases to import*/
diff --git a/Engine/Plugins/Experimental/HairStrands/Source/HairStrandsNiagara/Public/NiagaraDataInterfaceHairStrands.h b/Engine/Plugins/Experimental/HairStrands/Source/HairStrandsNiagara/Public/NiagaraDataInterfaceHairStrands.h
--- a/Engine/Plugins/Experimental/HairStrands/Source/HairStrandsNiagara/Public/NiagaraDataInterfaceHairStrands.h
+++ b/Engine/Plugins/Experimental/HairStrands/Source/HairStrandsNiagara/Public/NiagaraDataInterfaceHairStrands.h
@@ -12,6 +12,7 @@
 UENUM(BlueprintType)
 enum class EHairStrandsSize : uint8
 {
+	None = 0 UMETA(Hidden),
 	Size2 = 0x02 UMETA(DisplatName = "2"),
 	Size4 = 0x04 UMETA(DisplatName = "4"),
 	Size8 = 0x08 UMETA(DisplatName = "8"),
diff --git a/Engine/Plugins/Experimental/SimpleHMD/Source/SimpleHMD/Private/SimpleHMD.cpp b/Engine/Plugins/Experimental/SimpleHMD/Source/SimpleHMD/Private/SimpleHMD.cpp
--- a/Engine/Plugins/Experimental/SimpleHMD/Source/SimpleHMD/Private/SimpleHMD.cpp
+++ b/Engine/Plugins/Experimental/SimpleHMD/Source/SimpleHMD/Private/SimpleHMD.cpp
@@ -287,9 +287,9 @@ void FSimpleHMD::SetupViewFamily(FSceneViewFamily& InViewFamily)
 	InViewFamily.EngineShowFlags.HMDDistortion = true;
 	InViewFamily.EngineShowFlags.StereoRendering = IsStereoEnabled();
 
-	if (GWorld)
+	if (UWorld* World = GWorld)
 	{
-		WorldToMeters = GWorld->GetWorldSettings()->WorldToMeters;
+		WorldToMeters = World->GetWorldSettings()->WorldToMeters;
 	}
 }
 
diff --git a/Engine/Plugins/Runtime/AR/Google/GoogleARCore/Source/GoogleARCoreBase/Public/GoogleARCoreTypes.h b/Engine/Plugins/Runtime/AR/Google/GoogleARCore/Source/GoogleARCoreBase/Public/GoogleARCoreTypes.h
--- a/Engine/Plugins/Runtime/AR/Google/GoogleARCore/Source/GoogleARCoreBase/Public/GoogleARCoreTypes.h
+++ b/Engine/Plugins/Runtime/AR/Google/GoogleARCore/Source/GoogleARCoreBase/Public/GoogleARCoreTypes.h
@@ -187,6 +187,7 @@ ENUM_CLASS_FLAGS(EGoogleARCoreLineTraceChannel);
 UENUM(BlueprintType, Category = "GoogleARCore|CameraConfig")
 enum class EGoogleARCoreCameraFacing : uint8
 {
+	None = 0 UMETA(Hidden),
 	/* Camera facing away from user. */
 	Back = 1,
 	/* Camera facing towards user. */
diff --git a/Engine/Source/Developer/BlueprintCompilerCppBackend/Private/BlueprintCompilerCppBackendBase.cpp b/Engine/Source/Developer/BlueprintCompilerCppBackend/Private/BlueprintCompilerCppBackendBase.cpp
--- a/Engine/Source/Developer/BlueprintCompilerCppBackend/Private/BlueprintCompilerCppBackendBase.cpp
+++ b/Engine/Source/Developer/BlueprintCompilerCppBackend/Private/BlueprintCompilerCppBackendBase.cpp
@@ -357,6 +357,7 @@ FString FBlueprintCompilerCppBackendBase::GenerateCodeFromClass(UClass* SourceCl
 			EmitterContext.Header.AddLine(TEXT("{"));
 			EmitterContext.Header.IncreaseIndent();
 			EmitterContext.Header.AddLine(TEXT("GENERATED_BODY()"));
+			EmitterContext.Header.AddLine(TEXT("static void __CustomDynamicClassInitialization(UDynamicClass* InDynamicClass) {}"));
 			EmitterContext.Header.DecreaseIndent();
 			EmitterContext.Header.AddLine(TEXT("};"));
 			EmitterContext.Header.AddLine(FString::Printf(TEXT("class %s"), *CppClassName));
diff --git a/Engine/Source/Developer/BlueprintCompilerCppBackend/Private/BlueprintCompilerCppBackendValueHelper.cpp b/Engine/Source/Developer/BlueprintCompilerCppBackend/Private/BlueprintCompilerCppBackendValueHelper.cpp
--- a/Engine/Source/Developer/BlueprintCompilerCppBackend/Private/BlueprintCompilerCppBackendValueHelper.cpp
+++ b/Engine/Source/Developer/BlueprintCompilerCppBackend/Private/BlueprintCompilerCppBackendValueHelper.cpp
@@ -1645,6 +1645,7 @@ void FEmitDefaultValueHelper::GenerateCustomDynamicClassInitialization(FEmitterL
 		Context.AddLine(FString::Printf(TEXT("ensure(0 == InDynamicClass->%s.Num());"), GET_MEMBER_NAME_STRING_CHECKED(UDynamicClass, DynamicBindingObjects)));
 		Context.AddLine(FString::Printf(TEXT("ensure(0 == InDynamicClass->%s.Num());"), GET_MEMBER_NAME_STRING_CHECKED(UDynamicClass, ComponentTemplates)));
 		Context.AddLine(FString::Printf(TEXT("ensure(0 == InDynamicClass->%s.Num());"), GET_MEMBER_NAME_STRING_CHECKED(UDynamicClass, Timelines)));
+		Context.AddLine(FString::Printf(TEXT("ensure(0 == InDynamicClass->%s.Num());"), GET_MEMBER_NAME_STRING_CHECKED(UDynamicClass, ComponentClassOverrides)));
 		Context.AddLine(FString::Printf(TEXT("ensure(nullptr == InDynamicClass->%s);"), GET_MEMBER_NAME_STRING_CHECKED(UDynamicClass, AnimClassImplementation)));
 		Context.AddLine(FString::Printf(TEXT("InDynamicClass->%s();"), GET_FUNCTION_NAME_STRING_CHECKED(UDynamicClass, AssembleReferenceTokenStream)));
 
@@ -1774,6 +1775,18 @@ void FEmitDefaultValueHelper::GenerateCustomDynamicClassInitialization(FEmitterL
 		CreateAndInitializeClassSubobjects(true, false);
 		CreateAndInitializeClassSubobjects(false, true);
 
+		for (const FBPComponentClassOverride& Override : BPGC->ComponentClassOverrides)
+		{
+			if (Override.ComponentClass)
+			{
+				Context.AddLine(FString::Printf(TEXT("InDynamicClass->ComponentClassOverrides.Emplace(MakeTuple(FName(\"%s\"), FindObject<UClass>(ANY_PACKAGE,TEXT(\"%s\"))));"), *Override.ComponentName.ToString(), *Override.ComponentClass->GetPathName()));
+			}
+			else
+			{
+				Context.AddLine(FString::Printf(TEXT("InDynamicClass->ComponentClassOverrides.Emplace(MakeTuple(FName(\"%s\"), nullptr));"), *Override.ComponentName.ToString()));
+			}
+		}
+
 		FBackendHelperAnim::CreateAnimClassData(Context);
 
 		Context.DecreaseIndent();
@@ -1813,14 +1826,6 @@ void FEmitDefaultValueHelper::GenerateConstructor(FEmitterLocalContext& Context)
 		Context.AddLine(TEXT("{"));
 		Context.IncreaseIndent();
 
-		// Call CustomDynamicClassInitialization
-		Context.AddLine(FString::Printf(TEXT("if(HasAnyFlags(RF_ClassDefaultObject) && (%s::StaticClass() == GetClass()))"), *CppClassName));
-		Context.AddLine(TEXT("{"));
-		Context.IncreaseIndent();
-		Context.AddLine(FString::Printf(TEXT("%s::__CustomDynamicClassInitialization(CastChecked<UDynamicClass>(GetClass()));"), *CppClassName));
-		Context.DecreaseIndent();
-		Context.AddLine(TEXT("}"));
-
 		// Subobjects that must be fixed after serialization
 		TArray<FDefaultSubobjectData> SubobjectsToInit;
 		TArray<FNonNativeComponentData> ComponentsToInit;
diff --git a/Engine/Source/Editor/BlueprintGraph/Classes/EdGraphSchema_K2_Actions.h b/Engine/Source/Editor/BlueprintGraph/Classes/EdGraphSchema_K2_Actions.h
--- a/Engine/Source/Editor/BlueprintGraph/Classes/EdGraphSchema_K2_Actions.h
+++ b/Engine/Source/Editor/BlueprintGraph/Classes/EdGraphSchema_K2_Actions.h
@@ -569,6 +569,7 @@ public:
 	// Simple type info
 	static FName StaticGetTypeId() {static FName Type("FEdGraphSchemaAction_K2LocalVar"); return Type;}
 	virtual FName GetTypeId() const override { return StaticGetTypeId(); } 
+	virtual int32 GetReorderIndexInContainer() const override;
 
 	FEdGraphSchemaAction_K2LocalVar() 
 		: FEdGraphSchemaAction_BlueprintVariableBase()
diff --git a/Engine/Source/Editor/BlueprintGraph/Classes/K2Node_CustomEvent.h b/Engine/Source/Editor/BlueprintGraph/Classes/K2Node_CustomEvent.h
--- a/Engine/Source/Editor/BlueprintGraph/Classes/K2Node_CustomEvent.h
+++ b/Engine/Source/Editor/BlueprintGraph/Classes/K2Node_CustomEvent.h
@@ -13,6 +13,7 @@
 class FBlueprintActionDatabaseRegistrar;
 class INameValidatorInterface;
 class UEdGraph;
+struct FKismetUserDeclaredFunctionMetadata;
 
 UCLASS(MinimalAPI)
 class UK2Node_CustomEvent : public UK2Node_Event
@@ -103,7 +104,15 @@ class UK2Node_CustomEvent : public UK2Node_Event
 private:
 	/** Constructing FText strings can be costly, so we cache the node's title */
 	FNodeTextCache CachedNodeTitle;
-};
 
+	/** Custom event metadata that can be used for adding custom keywords */
+	UPROPERTY()
+	FKismetUserDeclaredFunctionMetadata MetaData;
 
+public:
 
+	FKismetUserDeclaredFunctionMetadata& GetUserDefinedMetaData()
+	{
+		return MetaData;
+	}
+};
\ No newline at end of file
diff --git a/Engine/Source/Editor/BlueprintGraph/Classes/K2Node_EditablePinBase.h b/Engine/Source/Editor/BlueprintGraph/Classes/K2Node_EditablePinBase.h
--- a/Engine/Source/Editor/BlueprintGraph/Classes/K2Node_EditablePinBase.h
+++ b/Engine/Source/Editor/BlueprintGraph/Classes/K2Node_EditablePinBase.h
@@ -144,6 +144,14 @@ class BLUEPRINTGRAPH_VTABLE UK2Node_EditablePinBase : public UK2Node
 	 */
 	BLUEPRINTGRAPH_API void RemoveUserDefinedPinByName(const FName PinName);
 
+	/**
+	* Check if a pin with this name exists in the user defined pin set
+	* 
+	* @param	PinName name of pin check existence of
+	* @return	True if a user defined pin with this name exists
+	*/
+	bool UserDefinedPinExists(const FName PinName) const;
+
 	/**
 	 * Creates a new pin on the node from the specified user pin info.
 	 * Must be overridden so each type of node can ensure that the pin is created in the proper direction, etc
diff --git a/Engine/Source/Editor/BlueprintGraph/Private/EdGraphSchema_K2.cpp b/Engine/Source/Editor/BlueprintGraph/Private/EdGraphSchema_K2.cpp
--- a/Engine/Source/Editor/BlueprintGraph/Private/EdGraphSchema_K2.cpp
+++ b/Engine/Source/Editor/BlueprintGraph/Private/EdGraphSchema_K2.cpp
@@ -6556,18 +6556,16 @@ UEdGraph* UEdGraphSchema_K2::DuplicateGraph(UEdGraph* GraphToDuplicate) const
 			for (int32 NodeIndex = 0; NodeIndex < NewGraph->Nodes.Num(); ++NodeIndex)
 			{
 				UEdGraphNode* Node = NewGraph->Nodes[NodeIndex];
-				if (UK2Node_FunctionEntry* EntryNode = Cast<UK2Node_FunctionEntry>(Node))
+				if (UK2Node_FunctionTerminator* TerminatorNode = Cast<UK2Node_FunctionTerminator>(Node))
 				{
-					if (EntryNode->FunctionReference.GetMemberName() == GraphToDuplicate->GetFName())
+					if (TerminatorNode->FunctionReference.GetMemberName() == GraphToDuplicate->GetFName())
 					{
-						EntryNode->Modify();
+						TerminatorNode->Modify();
 						
 						// We're duplicating the graph, so fully reset the member reference (including the GUID!)
 						FMemberReference NewRef;
 						NewRef.SetMemberName(NewGraph->GetFName());
-						EntryNode->FunctionReference = NewRef;
-
-						break;
+						TerminatorNode->FunctionReference = NewRef;
 					}
 				}
 				// Rename any custom events to be unique
diff --git a/Engine/Source/Editor/BlueprintGraph/Private/EdGraphSchema_K2_Actions.cpp b/Engine/Source/Editor/BlueprintGraph/Private/EdGraphSchema_K2_Actions.cpp
--- a/Engine/Source/Editor/BlueprintGraph/Private/EdGraphSchema_K2_Actions.cpp
+++ b/Engine/Source/Editor/BlueprintGraph/Private/EdGraphSchema_K2_Actions.cpp
@@ -57,16 +57,16 @@ int32 FEdGraphSchemaAction_BlueprintVariableBase::GetReorderIndexInContainer() c
 
 bool FEdGraphSchemaAction_BlueprintVariableBase::ReorderToBeforeAction(TSharedRef<FEdGraphSchemaAction> OtherAction)
 {
-	if ((OtherAction->GetTypeId() == StaticGetTypeId()) && (OtherAction->GetPersistentItemDefiningObject() == GetPersistentItemDefiningObject()))
+	if (OtherAction->GetPersistentItemDefiningObject() == GetPersistentItemDefiningObject())
 	{
 		FEdGraphSchemaAction_BlueprintVariableBase* VarAction = (FEdGraphSchemaAction_BlueprintVariableBase*)&OtherAction.Get();
 
 		// Only let you drag and drop if variables are from same BP class, and not onto itself
 		UBlueprint* BP = GetSourceBlueprint();
 		FName TargetVarName = VarAction->GetVariableName();
-		if ((BP != nullptr) && (VarName != TargetVarName) && (VariableSource == VarAction->GetVariableClass()))
+		if ((BP != nullptr) && (VarName != TargetVarName) && (VariableSource == VarAction->GetVariableScope()))
 		{
-			if (FBlueprintEditorUtils::MoveVariableBeforeVariable(BP, VarName, TargetVarName, true))
+			if (FBlueprintEditorUtils::MoveVariableBeforeVariable(BP, VarAction->GetVariableScope(), VarName, TargetVarName, true))
 			{
 				// Change category of var to match the one we dragged on to as well
 				FText TargetVarCategory = FBlueprintEditorUtils::GetBlueprintVariableCategory(BP, TargetVarName, GetVariableScope());
@@ -107,6 +107,14 @@ UBlueprint* FEdGraphSchemaAction_BlueprintVariableBase::GetSourceBlueprint() con
 	return UBlueprint::GetBlueprintFromClass(ClassToCheck);
 }
 
+/////////////////////////////////////////////////////
+// FEdGraphSchemaAction_K2LocalVar
+
+int32 FEdGraphSchemaAction_K2LocalVar::GetReorderIndexInContainer() const
+{
+	return FBlueprintEditorUtils::FindLocalVariableIndex(GetSourceBlueprint(), GetVariableScope(), GetVariableName());
+}
+
 /////////////////////////////////////////////////////
 // FEdGraphSchemaAction_K2Graph
 
@@ -132,7 +140,7 @@ int32 FEdGraphSchemaAction_K2Graph::GetReorderIndexInContainer() const
 
 bool FEdGraphSchemaAction_K2Graph::ReorderToBeforeAction(TSharedRef<FEdGraphSchemaAction> OtherAction)
 {
-	if ((OtherAction->GetTypeId() == GetTypeId()) && (OtherAction->GetPersistentItemDefiningObject() == GetPersistentItemDefiningObject()))
+	if (OtherAction->GetTypeId() == GetTypeId())
 	{
 		const int32 OldIndex = GetReorderIndexInContainer();
 		const int32 NewIndexToGoBefore = OtherAction->GetReorderIndexInContainer();
@@ -708,3 +716,4 @@ UEdGraphNode* FEdGraphSchemaAction_K2PasteHere::PerformAction( class UEdGraph* P
 	return nullptr;
 }
 
+
diff --git a/Engine/Source/Editor/BlueprintGraph/Private/K2Node.cpp b/Engine/Source/Editor/BlueprintGraph/Private/K2Node.cpp
--- a/Engine/Source/Editor/BlueprintGraph/Private/K2Node.cpp
+++ b/Engine/Source/Editor/BlueprintGraph/Private/K2Node.cpp
@@ -271,6 +271,13 @@ bool UK2Node::CreatePinsForFunctionEntryExit(const UFunction* Function, bool bFo
 	{
 		UProperty* Param = *PropIt;
 
+		// Don't create a new pin if one exists already! 
+		// @see UE-79032, UE-58390
+		if (FindPin(Param->GetFName()))
+		{
+			continue;
+		}
+
 		const bool bIsFunctionInput = !Param->HasAnyPropertyFlags(CPF_OutParm) || Param->HasAnyPropertyFlags(CPF_ReferenceParm);
 
 		if (bIsFunctionInput == bForFunctionEntry)
diff --git a/Engine/Source/Editor/BlueprintGraph/Private/K2Node_EditablePinBase.cpp b/Engine/Source/Editor/BlueprintGraph/Private/K2Node_EditablePinBase.cpp
--- a/Engine/Source/Editor/BlueprintGraph/Private/K2Node_EditablePinBase.cpp
+++ b/Engine/Source/Editor/BlueprintGraph/Private/K2Node_EditablePinBase.cpp
@@ -80,7 +80,7 @@ void UK2Node_EditablePinBase::AllocateDefaultPins()
 	for(int32 i = 0; i < UserDefinedPins.Num(); i++)
 	{
 		FText DummyErrorMsg;
-		if (!IsEditable() || CanCreateUserDefinedPin(UserDefinedPins[i]->PinType, UserDefinedPins[i]->DesiredPinDirection, DummyErrorMsg))
+		if ((!IsEditable() || CanCreateUserDefinedPin(UserDefinedPins[i]->PinType, UserDefinedPins[i]->DesiredPinDirection, DummyErrorMsg)) && !FindPin(UserDefinedPins[i]->PinName))
 		{
 			CreatePinFromUserDefinition(UserDefinedPins[i]);
 		}
@@ -137,6 +137,14 @@ void UK2Node_EditablePinBase::RemoveUserDefinedPinByName(const FName PinName)
 	});
 }
 
+bool UK2Node_EditablePinBase::UserDefinedPinExists(const FName PinName) const
+{
+	return UserDefinedPins.ContainsByPredicate([&PinName](const TSharedPtr<FUserPinInfo>& UDPin)
+	{
+		return UDPin.IsValid() && (UDPin->PinName == PinName);
+	});
+}
+
 void UK2Node_EditablePinBase::ExportCustomProperties(FOutputDevice& Out, uint32 Indent)
 {
 	Super::ExportCustomProperties(Out, Indent);
diff --git a/Engine/Source/Editor/BlueprintGraph/Private/K2Node_FunctionTerminator.cpp b/Engine/Source/Editor/BlueprintGraph/Private/K2Node_FunctionTerminator.cpp
--- a/Engine/Source/Editor/BlueprintGraph/Private/K2Node_FunctionTerminator.cpp
+++ b/Engine/Source/Editor/BlueprintGraph/Private/K2Node_FunctionTerminator.cpp
@@ -109,12 +109,17 @@ void UK2Node_FunctionTerminator::PromoteFromInterfaceOverride(bool bIsPrimaryTer
 {
 	// Remove the signature class, that is not relevant.
 	FunctionReference.SetSelfMember(FunctionReference.GetMemberName());
-	TArray<UEdGraphPin*> OriginalPins = Pins;
-	for (const UEdGraphPin* Pin : OriginalPins)
+	
+	// For every pin that has been defined, make it a user defined pin if we can
+	for (UEdGraphPin* const Pin : Pins)
 	{
-		if (Pin->PinType.PinCategory != UEdGraphSchema_K2::PC_Exec)
+		if (Pin->PinType.PinCategory != UEdGraphSchema_K2::PC_Exec && !UserDefinedPinExists(Pin->PinName))
 		{
-			CreateUserDefinedPin(Pin->PinName, Pin->PinType, Pin->Direction, false);
+			TSharedPtr<FUserPinInfo> NewPinInfo = MakeShareable(new FUserPinInfo());
+			NewPinInfo->PinName = Pin->PinName;
+			NewPinInfo->PinType = Pin->PinType;
+			NewPinInfo->DesiredPinDirection = Pin->Direction;
+			UserDefinedPins.Add(NewPinInfo);
 		}
 	}
 	const UEdGraphSchema_K2* Schema = GetDefault<UEdGraphSchema_K2>();
diff --git a/Engine/Source/Editor/DetailCustomizations/Private/InputSettingsDetails.cpp b/Engine/Source/Editor/DetailCustomizations/Private/InputSettingsDetails.cpp
--- a/Engine/Source/Editor/DetailCustomizations/Private/InputSettingsDetails.cpp
+++ b/Engine/Source/Editor/DetailCustomizations/Private/InputSettingsDetails.cpp
@@ -51,6 +51,7 @@ void FActionMappingsNodeBuilder::GenerateHeaderRowContent( FDetailWidgetRow& Nod
 		LOCTEXT("ClearActionMappingToolTip", "Removes all Action Mappings") );
 
 	NodeRow
+	.FilterString(ActionMappingsPropertyHandle->GetPropertyDisplayName())
 	[
 		SNew( SHorizontalBox )
 		+SHorizontalBox::Slot()
diff --git a/Engine/Source/Editor/EditorStyle/Private/SlateEditorStyle.cpp b/Engine/Source/Editor/EditorStyle/Private/SlateEditorStyle.cpp
--- a/Engine/Source/Editor/EditorStyle/Private/SlateEditorStyle.cpp
+++ b/Engine/Source/Editor/EditorStyle/Private/SlateEditorStyle.cpp
@@ -6162,6 +6162,7 @@ void FSlateEditorStyle::FStyle::SetupClassIconsAndThumbnails()
 		Set("ClassIcon.Deleted", new IMAGE_BRUSH("Icons/ActorIcons/DeletedActor_16px", Icon16x16));
 		Set("ClassIcon.StaticMeshActor", new IMAGE_BRUSH("Icons/AssetIcons/StaticMesh_16x", Icon16x16));
 		Set("ClassIcon.SkeletalMeshActor", new IMAGE_BRUSH("Icons/AssetIcons/SkeletalMesh_16x", Icon16x16));
+		Set("ClassIcon.LevelScriptActor", new IMAGE_BRUSH("Icons/AssetIcons/Blueprint_16x", Icon16x16));
 
 		// Component classes
 		Set("ClassIcon.AudioComponent", new IMAGE_BRUSH("Icons/ActorIcons/SoundActor_16x", Icon16x16));
diff --git a/Engine/Source/Editor/Kismet/Private/BPFunctionDragDropAction.cpp b/Engine/Source/Editor/Kismet/Private/BPFunctionDragDropAction.cpp
--- a/Engine/Source/Editor/Kismet/Private/BPFunctionDragDropAction.cpp
+++ b/Engine/Source/Editor/Kismet/Private/BPFunctionDragDropAction.cpp
@@ -7,6 +7,7 @@
 #include "EdGraphSchema_K2_Actions.h"
 #include "K2Node_CustomEvent.h"
 #include "K2Node_MacroInstance.h"
+#include "Kismet2/KismetEditorUtilities.h"
 
 #include "Kismet2/BlueprintEditorUtils.h"
 #include "BlueprintNodeBinder.h"
@@ -162,11 +163,18 @@ bool FKismetDragDropAction::ActionWillShowExistingNode() const
 			else if (SourceAction->GetTypeId() == FEdGraphSchemaAction_K2Event::StaticGetTypeId())
 			{
 				FEdGraphSchemaAction_K2Event* FuncAction = (FEdGraphSchemaAction_K2Event*)SourceAction.Get();
-				UK2Node_CustomEvent* CustomEvent = Cast<UK2Node_CustomEvent>(FuncAction->NodeTemplate);
-				// Drag and dropping custom event's will place a Call Function and will not focus the existing event
-				if (CustomEvent == nullptr)
+				// Drag and dropping custom events will place a Call Function and will not focus the existing event
+				if (UK2Node_Event* Event = CastChecked<UK2Node_Event>(FuncAction->NodeTemplate))
 				{
-					bWillFocusOnExistingNode = true;
+					UFunction* Function = FFunctionFromNodeHelper::FunctionFromNode(Event);
+					if (Function && (Function->FunctionFlags & (FUNC_BlueprintCallable | FUNC_BlueprintPure)))
+					{
+						bWillFocusOnExistingNode = false;
+					}
+					else
+					{
+						bWillFocusOnExistingNode = true;
+					}
 				}
 			}
 		}
@@ -235,7 +243,7 @@ FReply FKismetFunctionDragDropAction::DroppedOnPin(FVector2D ScreenPosition, FVe
 		if (!CanBeDroppedDelegate.IsBound() || CanBeDroppedDelegate.Execute(nullptr, Graph, CannotDropReason))
 		{
 			UFunction const* Function = GetFunctionProperty();
-			if ((Function != NULL) && UEdGraphSchema_K2::CanUserKismetCallFunction(Function))
+			if ((Function != nullptr) && UEdGraphSchema_K2::CanUserKismetCallFunction(Function))
 			{
 				AnalyticCallback.ExecuteIfBound();
 
@@ -257,6 +265,13 @@ FReply FKismetFunctionDragDropAction::DroppedOnPin(FVector2D ScreenPosition, FVe
 			}
 		}
 	}
+	else if (SourceAction->GetTypeId() == FEdGraphSchemaAction_K2Event::StaticGetTypeId())
+	{
+		if (FEdGraphSchemaAction_K2Event* FuncAction = (FEdGraphSchemaAction_K2Event*)SourceAction.Get())
+		{
+			FKismetEditorUtilities::BringKismetToFocusAttentionOnObject(FuncAction->NodeTemplate);
+		}
+	}
 
 	return Reply;
 }
@@ -288,7 +303,10 @@ UBlueprintFunctionNodeSpawner* FKismetFunctionDragDropAction::GetDropAction(UEdG
 
 			if (UFunction const* Function = GetFunctionProperty())
 			{
-				return UBlueprintFunctionNodeSpawner::Create(Function);
+				if (Function->FunctionFlags & (FUNC_BlueprintCallable | FUNC_BlueprintPure))
+				{
+					return UBlueprintFunctionNodeSpawner::Create(Function);
+				}
 			}
 		}
 	}
diff --git a/Engine/Source/Editor/Kismet/Private/BlueprintCompilationManager.cpp b/Engine/Source/Editor/Kismet/Private/BlueprintCompilationManager.cpp
--- a/Engine/Source/Editor/Kismet/Private/BlueprintCompilationManager.cpp
+++ b/Engine/Source/Editor/Kismet/Private/BlueprintCompilationManager.cpp
@@ -2390,7 +2390,7 @@ UClass* FBlueprintCompilationManagerImpl::FastGenerateSkeletonClass(UBlueprint*
 	UField** CurrentFieldStorageLocation = &Ret->Children;
 	
 	// Helper function for making UFunctions generated for 'event' nodes, e.g. custom event and timelines
-	const auto MakeEventFunction = [&CurrentFieldStorageLocation, MakeFunction, Schema]( FName InName, EFunctionFlags ExtraFnFlags, const TArray<UEdGraphPin*>& InputPins, const TArray< TSharedPtr<FUserPinInfo> >& UserPins, UFunction* InSourceFN, bool bInCallInEditor, bool bIsDeprecated, const FString& DeprecationMessage )
+	const auto MakeEventFunction = [&CurrentFieldStorageLocation, MakeFunction, Schema]( FName InName, EFunctionFlags ExtraFnFlags, const TArray<UEdGraphPin*>& InputPins, const TArray< TSharedPtr<FUserPinInfo> >& UserPins, UFunction* InSourceFN, bool bInCallInEditor, bool bIsDeprecated, const FString& DeprecationMessage, FKismetUserDeclaredFunctionMetadata* UserDefinedMetaData = nullptr)
 	{
 		UField** CurrentParamStorageLocation = nullptr;
 
@@ -2424,6 +2424,11 @@ UClass* FBlueprintCompilationManagerImpl::FastGenerateSkeletonClass(UBlueprint*
 				NewFunction->SetMetaData(FBlueprintMetadata::MD_CallInEditor, TEXT( "true" ));
 			}
 
+			if (UserDefinedMetaData)
+			{
+				NewFunction->SetMetaData(FBlueprintMetadata::MD_FunctionKeywords, *(UserDefinedMetaData->Keywords).ToString());
+			}
+
 			NewFunction->Bind();
 			NewFunction->StaticLink(true);
 		}
@@ -2461,6 +2466,7 @@ UClass* FBlueprintCompilationManagerImpl::FastGenerateSkeletonClass(UBlueprint*
 			FString DeprecationMessage;
 			bool bIsDeprecated = false;
 			bool bCallInEditor = false;
+			FKismetUserDeclaredFunctionMetadata* UserMetaData = nullptr;
 			if(UK2Node_CustomEvent* CustomEvent = Cast<UK2Node_CustomEvent>(Event))
 			{
 				bCallInEditor = CustomEvent->bCallInEditor;
@@ -2469,6 +2475,7 @@ UClass* FBlueprintCompilationManagerImpl::FastGenerateSkeletonClass(UBlueprint*
 				{
 					DeprecationMessage = CustomEvent->DeprecationMessage;
 				}
+				UserMetaData = &(CustomEvent->GetUserDefinedMetaData());
 			}
 			MakeEventFunction(
 				CompilerContext.GetEventStubFunctionName(Event), 
@@ -2478,7 +2485,8 @@ UClass* FBlueprintCompilationManagerImpl::FastGenerateSkeletonClass(UBlueprint*
 				Event->FindEventSignatureFunction(),
 				bCallInEditor,
 				bIsDeprecated,
-				DeprecationMessage
+				DeprecationMessage,
+				UserMetaData
 			);
 		}
 	}
diff --git a/Engine/Source/Editor/Kismet/Private/BlueprintDetailsCustomization.cpp b/Engine/Source/Editor/Kismet/Private/BlueprintDetailsCustomization.cpp
--- a/Engine/Source/Editor/Kismet/Private/BlueprintDetailsCustomization.cpp
+++ b/Engine/Source/Editor/Kismet/Private/BlueprintDetailsCustomization.cpp
@@ -2888,8 +2888,24 @@ void FBlueprintGraphArgumentLayout::GenerateChildContent( IDetailChildrenBuilder
 			}
 		}
 
+		bool bMacroGraph = false;
+		if (TargetNode)
+		{
+			if (const UBlueprint* Blueprint = TargetNode->GetBlueprint())
+			{
+				if (Blueprint->BlueprintType == BPTYPE_MacroLibrary)
+				{
+					bMacroGraph = true;
+				}
+				else if (const UEdGraph* Graph = TargetNode->GetGraph()) 
+				{
+					bMacroGraph = Blueprint->MacroGraphs.Contains(Graph);
+				}
+			}
+		}
+
 		// Exec pins can't be passed by reference
-		if (FoundPin == nullptr || !UEdGraphSchema_K2::IsExecPin(*FoundPin))
+		if (FoundPin && !UEdGraphSchema_K2::IsExecPin(*FoundPin) && !bMacroGraph)
 		{
 			ChildrenBuilder.AddCustomRow(LOCTEXT("FunctionArgDetailsPassByReference", "Pass-by-Reference"))
 				.NameContent()
@@ -3477,6 +3493,27 @@ void FBlueprintGraphActionDetails::CustomizeDetails( IDetailLayoutBuilder& Detai
 										FSlateIcon(),
 										FUIAction(FExecuteAction::CreateStatic( &FBlueprintGraphActionDetails::SetNetFlags, FunctionEntryNodePtr, static_cast<uint32>(FUNC_NetClient) ), CanExecuteDelegate));
 
+			const FString DocLink = TEXT("Shared/Editors/BlueprintEditor/GraphDetails");
+			TSharedPtr<SToolTip> KeywordsTooltip = IDocumentation::Get()->CreateToolTip(LOCTEXT("EditEventKeywords_Tooltip", "Keywords for searching for the event."), nullptr, DocLink, TEXT("Keywords"));
+			Category.AddCustomRow(LOCTEXT("EventsKeywordsLabel", "Keywords"))
+			.NameContent()
+			[
+				SNew(STextBlock)
+					.Text(LOCTEXT("KeywordsLabel", "Keywords"))
+					.ToolTip(KeywordsTooltip)
+					.Font(IDetailLayoutBuilder::GetDetailFont())
+			]
+			.ValueContent()
+			[
+				SNew(SEditableTextBox)
+					.Text(this, &FBlueprintGraphActionDetails::OnGetKeywordsText)
+					.OnTextCommitted(this, &FBlueprintGraphActionDetails::OnKeywordsTextCommitted)
+					.ToolTip(KeywordsTooltip)
+					.RevertTextOnEscape(true)
+					.Font(IDetailLayoutBuilder::GetDetailFont())
+			];
+
+
 			Category.AddCustomRow( LOCTEXT( "FunctionReplicate", "Replicates" ) )
 			.NameContent()
 			[
@@ -3802,7 +3839,7 @@ void FBlueprintGraphActionDetails::SetNetFlags( TWeakObjectPtr<UK2Node_EditableP
 				TypedEntryNode->SetExtraFlags(ExtraFlags);
 				bBlueprintModified = true;
 			}
-			if (UK2Node_CustomEvent * CustomEventNode = Cast<UK2Node_CustomEvent>(FunctionEntryNode.Get()))
+			if (UK2Node_CustomEvent* CustomEventNode = Cast<UK2Node_CustomEvent>(FunctionEntryNode.Get()))
 			{
 				CustomEventNode->FunctionFlags &= ~FlagsToClear;
 				CustomEventNode->FunctionFlags |= FlagsToSet;
@@ -3819,7 +3856,7 @@ void FBlueprintGraphActionDetails::SetNetFlags( TWeakObjectPtr<UK2Node_EditableP
 
 FText FBlueprintGraphActionDetails::GetCurrentReplicatedEventString() const
 {
-	const UK2Node_EditablePinBase * FunctionEntryNode = FunctionEntryNodePtr.Get();
+	const UK2Node_EditablePinBase* FunctionEntryNode = FunctionEntryNodePtr.Get();
 	const UK2Node_CustomEvent* CustomEvent = Cast<const UK2Node_CustomEvent>(FunctionEntryNode);
 
 	uint32 const ReplicatedNetMask = (FUNC_NetMulticast | FUNC_NetServer | FUNC_NetClient);
@@ -4444,10 +4481,13 @@ UK2Node_EditablePinBase* FBlueprintGraphActionDetails::GetEditableNode() const
 
 UFunction* FBlueprintGraphActionDetails::FindFunction() const
 {
-	UEdGraph* Graph = GetGraph();
-	if(Graph)
+	if (UK2Node_CustomEvent* EventNode = Cast<UK2Node_CustomEvent>(FunctionEntryNodePtr.Get()))
 	{
-		if(UBlueprint* Blueprint = FBlueprintEditorUtils::FindBlueprintForGraph(Graph))
+		return FFunctionFromNodeHelper::FunctionFromNode(EventNode);
+	}
+	else if (UEdGraph* Graph = GetGraph())
+	{
+		if (UBlueprint* Blueprint = FBlueprintEditorUtils::FindBlueprintForGraph(Graph))
 		{
 			UClass* Class = Blueprint->SkeletonGeneratedClass;
 
@@ -4461,7 +4501,7 @@ UFunction* FBlueprintGraphActionDetails::FindFunction() const
 			}
 		}
 	}
-	return NULL;
+	return nullptr;
 }
 
 FKismetUserDeclaredFunctionMetadata* FBlueprintGraphActionDetails::GetMetadataBlock() const
@@ -4476,8 +4516,12 @@ FKismetUserDeclaredFunctionMetadata* FBlueprintGraphActionDetails::GetMetadataBl
 		// Must be exactly a tunnel, not a macro instance
 		return &(TunnelNode->MetaData);
 	}
+	else if (UK2Node_CustomEvent* EventNode = Cast<UK2Node_CustomEvent>(FunctionEntryNode))
+	{
+		return &(EventNode->GetUserDefinedMetaData());
+	}
 
-	return NULL;
+	return nullptr;
 }
 
 FText FBlueprintGraphActionDetails::OnGetTooltipText() const
@@ -4682,7 +4726,7 @@ TSharedRef<ITableRow> FBlueprintGraphActionDetails::HandleGenerateRowAccessSpeci
 FText FBlueprintGraphActionDetails::GetCurrentAccessSpecifierName() const
 {
 	uint32 AccessSpecifierFlag = 0;
-	UK2Node_EditablePinBase * FunctionEntryNode = FunctionEntryNodePtr.Get();
+	UK2Node_EditablePinBase* FunctionEntryNode = FunctionEntryNodePtr.Get();
 	if(UK2Node_FunctionEntry* EntryNode = Cast<UK2Node_FunctionEntry>(FunctionEntryNode))
 	{
 		AccessSpecifierFlag = FUNC_AccessSpecifiers & EntryNode->GetFunctionFlags();
@@ -4698,7 +4742,7 @@ bool FBlueprintGraphActionDetails::IsAccessSpecifierVisible() const
 {
 	bool bSupportedType = false;
 	bool bIsEditable = false;
-	UK2Node_EditablePinBase * FunctionEntryNode = FunctionEntryNodePtr.Get();
+	UK2Node_EditablePinBase* FunctionEntryNode = FunctionEntryNodePtr.Get();
 	if(FunctionEntryNode)
 	{
 		UBlueprint* Blueprint = FunctionEntryNode->GetBlueprint();
@@ -4818,7 +4862,7 @@ bool FBlueprintGraphActionDetails::IsCustomEvent() const
 
 void FBlueprintGraphActionDetails::OnIsReliableReplicationFunctionModified(const ECheckBoxState NewCheckedState)
 {
-	UK2Node_EditablePinBase * FunctionEntryNode = FunctionEntryNodePtr.Get();
+	UK2Node_EditablePinBase* FunctionEntryNode = FunctionEntryNodePtr.Get();
 	UK2Node_CustomEvent* CustomEvent = Cast<UK2Node_CustomEvent>(FunctionEntryNode);
 	if( CustomEvent )
 	{
@@ -4828,7 +4872,7 @@ void FBlueprintGraphActionDetails::OnIsReliableReplicationFunctionModified(const
 			{
 				TypedEntryNode->AddExtraFlags(FUNC_NetReliable);
 			}
-			if (UK2Node_CustomEvent * CustomEventNode = Cast<UK2Node_CustomEvent>(FunctionEntryNode))
+			if (UK2Node_CustomEvent* CustomEventNode = Cast<UK2Node_CustomEvent>(FunctionEntryNode))
 			{
 				CustomEventNode->FunctionFlags |= FUNC_NetReliable;
 			}
@@ -4839,7 +4883,7 @@ void FBlueprintGraphActionDetails::OnIsReliableReplicationFunctionModified(const
 			{
 				TypedEntryNode->ClearExtraFlags(FUNC_NetReliable);
 			}
-			if (UK2Node_CustomEvent * CustomEventNode = Cast<UK2Node_CustomEvent>(FunctionEntryNode))
+			if (UK2Node_CustomEvent* CustomEventNode = Cast<UK2Node_CustomEvent>(FunctionEntryNode))
 			{
 				CustomEventNode->FunctionFlags &= ~FUNC_NetReliable;
 			}
@@ -4851,7 +4895,7 @@ void FBlueprintGraphActionDetails::OnIsReliableReplicationFunctionModified(const
 
 ECheckBoxState FBlueprintGraphActionDetails::GetIsReliableReplicatedFunction() const
 {
-	const UK2Node_EditablePinBase * FunctionEntryNode = FunctionEntryNodePtr.Get();
+	const UK2Node_EditablePinBase* FunctionEntryNode = FunctionEntryNodePtr.Get();
 	const UK2Node_CustomEvent* CustomEvent = Cast<const UK2Node_CustomEvent>(FunctionEntryNode);
 	if(!CustomEvent)
 	{
@@ -4871,7 +4915,7 @@ bool FBlueprintGraphActionDetails::IsPureFunctionVisible() const
 {
 	bool bSupportedType = false;
 	bool bIsEditable = false;
-	UK2Node_EditablePinBase * FunctionEntryNode = FunctionEntryNodePtr.Get();
+	UK2Node_EditablePinBase* FunctionEntryNode = FunctionEntryNodePtr.Get();
 	if(FunctionEntryNode)
 	{
 		UBlueprint* Blueprint = FunctionEntryNode->GetBlueprint();
@@ -4885,7 +4929,7 @@ bool FBlueprintGraphActionDetails::IsPureFunctionVisible() const
 
 void FBlueprintGraphActionDetails::OnIsPureFunctionModified( const ECheckBoxState NewCheckedState )
 {
-	UK2Node_EditablePinBase * FunctionEntryNode = FunctionEntryNodePtr.Get();
+	UK2Node_EditablePinBase* FunctionEntryNode = FunctionEntryNodePtr.Get();
 	UFunction* Function = FindFunction();
 	UK2Node_FunctionEntry* EntryNode = Cast<UK2Node_FunctionEntry>(FunctionEntryNode);
 	if (EntryNode && Function)
@@ -4916,7 +4960,7 @@ bool FBlueprintGraphActionDetails::IsConstFunctionVisible() const
 {
 	bool bSupportedType = false;
 	bool bIsEditable = false;
-	UK2Node_EditablePinBase * FunctionEntryNode = FunctionEntryNodePtr.Get();
+	UK2Node_EditablePinBase* FunctionEntryNode = FunctionEntryNodePtr.Get();
 	if(FunctionEntryNode)
 	{
 		bSupportedType = FunctionEntryNode->IsA<UK2Node_FunctionEntry>();
@@ -4927,7 +4971,7 @@ bool FBlueprintGraphActionDetails::IsConstFunctionVisible() const
 
 void FBlueprintGraphActionDetails::OnIsConstFunctionModified( const ECheckBoxState NewCheckedState )
 {
-	UK2Node_EditablePinBase * FunctionEntryNode = FunctionEntryNodePtr.Get();
+	UK2Node_EditablePinBase* FunctionEntryNode = FunctionEntryNodePtr.Get();
 	UFunction* Function = FindFunction();
 	UK2Node_FunctionEntry* EntryNode = Cast<UK2Node_FunctionEntry>(FunctionEntryNode);
 	if(EntryNode && Function)
@@ -4956,7 +5000,7 @@ ECheckBoxState FBlueprintGraphActionDetails::GetIsConstFunction() const
 
 FReply FBaseBlueprintGraphActionDetails::OnAddNewInputClicked()
 {
-	UK2Node_EditablePinBase * FunctionEntryNode = FunctionEntryNodePtr.Get();
+	UK2Node_EditablePinBase* FunctionEntryNode = FunctionEntryNodePtr.Get();
 
 	if( FunctionEntryNode )
 	{
@@ -5736,6 +5780,40 @@ void FBlueprintComponentDetails::CustomizeDetails(IDetailLayoutBuilder& DetailLa
 			]
 		];
 
+		UActorComponent* ComponentTemplate = (CachedNodePtr->IsNative() ? CachedNodePtr->GetComponentTemplate() : nullptr);
+		if (ComponentTemplate)
+		{
+			UClass* BaseClass = ComponentTemplate->GetClass();
+
+			if (const FBPComponentClassOverride* Override = BlueprintObj->ComponentClassOverrides.FindByKey(ComponentTemplate->GetFName()))
+			{
+				AActor* Owner = ComponentTemplate->GetOwner();
+				AActor* OwnerArchetype = CastChecked<AActor>(Owner->GetArchetype());
+				if (UActorComponent* ArchetypeComponent = Cast<UActorComponent>((UObject*)FindObjectWithOuter(OwnerArchetype, UActorComponent::StaticClass(), ComponentTemplate->GetFName())))
+				{
+					BaseClass = ArchetypeComponent->GetClass();
+				}
+			}
+
+			const FText ComponentClassTooltip = LOCTEXT("BlueprintComponentDetails_ComponentClassOverrideTooltip", "The class to use when creating this component for this class. This can only be done for components defined in native at this time.");
+
+			VariableCategory.AddCustomRow( LOCTEXT("BlueprintComponentDetails_ComponentClassOverride", "Component Class") )
+			.NameContent()
+			[
+				SNew(STextBlock)
+				.Text(LOCTEXT("BlueprintComponentDetails_ComponentClassOverrideLabel", "Component Class"))
+				.ToolTipText(ComponentClassTooltip)
+				.Font(IDetailLayoutBuilder::GetDetailFont())
+			]
+			.ValueContent()
+			[
+				SNew(SClassPropertyEntryBox)
+				.MetaClass(BaseClass)
+				.AllowNone(false)
+				.SelectedClass(this, &FBlueprintComponentDetails::GetSelectedEntryClass)
+				.OnSetClass(this, &FBlueprintComponentDetails::HandleNewEntryClassSelected)];
+		}
+
 		IDetailCategoryBuilder& SocketsCategory = DetailLayout.EditCategory("Sockets", LOCTEXT("BlueprintComponentDetailsCategory", "Sockets"), ECategoryPriority::Important);
 
 		SocketsCategory.AddCustomRow(LOCTEXT("BlueprintComponentDetails_Sockets", "Sockets"))
@@ -5995,6 +6073,87 @@ void FBlueprintComponentDetails::PopulateVariableCategories()
 	}
 }
 
+const UClass* FBlueprintComponentDetails::GetSelectedEntryClass() const
+{
+	check(CachedNodePtr.IsValid());
+
+	if (UActorComponent* ComponentTemplate = CachedNodePtr->GetComponentTemplate())
+	{
+		UBlueprint* BlueprintObj = GetBlueprintObj();
+		check(BlueprintObj);
+
+		do
+		{
+			if (const FBPComponentClassOverride* Override = BlueprintObj->ComponentClassOverrides.FindByKey(ComponentTemplate->GetFName()))
+			{
+				return Override->ComponentClass;
+			}
+
+			BlueprintObj = UBlueprint::GetBlueprintFromClass(Cast<UBlueprintGeneratedClass>(BlueprintObj->ParentClass));
+
+		} while (BlueprintObj);
+
+		return ComponentTemplate->GetClass();
+	}
+
+	return nullptr;
+
+}
+
+void FBlueprintComponentDetails::HandleNewEntryClassSelected(const UClass* NewEntryClass) const
+{
+	const UClass* PreviousClass = GetSelectedEntryClass();
+
+	if (PreviousClass != NewEntryClass)
+	{
+		check(CachedNodePtr.IsValid());
+
+		if (UActorComponent* ComponentTemplate = CachedNodePtr->GetComponentTemplate())
+		{
+			const FName ComponentTemplateName = ComponentTemplate->GetFName();
+
+			UBlueprint* BlueprintObj = GetBlueprintObj();
+			check(BlueprintObj);
+
+			if (FBPComponentClassOverride* Override = BlueprintObj->ComponentClassOverrides.FindByKey(ComponentTemplateName))
+			{
+				bool bRemoveEntry = false;
+				bool bFoundOverride = false;
+
+				UBlueprint* ParentBP = UBlueprint::GetBlueprintFromClass(Cast<UBlueprintGeneratedClass>(BlueprintObj->ParentClass));
+				while (ParentBP)
+				{
+					if (FBPComponentClassOverride* ParentOverride = BlueprintObj->ComponentClassOverrides.FindByKey(ComponentTemplateName))
+					{
+						bRemoveEntry = (ParentOverride->ComponentClass == NewEntryClass);
+						bFoundOverride = true;
+						break;
+					}
+				}
+				if (!bFoundOverride)
+				{
+					bRemoveEntry = (ComponentTemplate->GetClass() == NewEntryClass);
+				}
+				if (bRemoveEntry)
+				{
+					BlueprintObj->ComponentClassOverrides.RemoveAllSwap([ComponentTemplateName](const FBPComponentClassOverride& CCOverride) { return (CCOverride.ComponentName == ComponentTemplateName); });
+				}
+				else
+				{
+					Override->ComponentClass = const_cast<UClass*>(NewEntryClass);
+				}
+			}
+			else
+			{
+				BlueprintObj->ComponentClassOverrides.Emplace(FBPComponentClassOverride(ComponentTemplateName, const_cast<UClass*>(NewEntryClass)));
+			}
+
+			FBlueprintEditorUtils::MarkBlueprintAsStructurallyModified(BlueprintObj);
+		}
+	}
+}
+
+
 FText FBlueprintComponentDetails::GetSocketName() const
 {
 	check(CachedNodePtr.IsValid());
diff --git a/Engine/Source/Editor/Kismet/Private/BlueprintDetailsCustomization.h b/Engine/Source/Editor/Kismet/Private/BlueprintDetailsCustomization.h
--- a/Engine/Source/Editor/Kismet/Private/BlueprintDetailsCustomization.h
+++ b/Engine/Source/Editor/Kismet/Private/BlueprintDetailsCustomization.h
@@ -820,6 +820,8 @@ protected:
 	void OnVariableCategoryTextCommitted(const FText& NewText, ETextCommit::Type InTextCommit, FName VarName);
 	void OnVariableCategorySelectionChanged(TSharedPtr<FText> ProposedSelection, ESelectInfo::Type /*SelectInfo*/);
 	TSharedRef<ITableRow> MakeVariableCategoryViewWidget(TSharedPtr<FText> Item, const TSharedRef< STableViewBase >& OwnerTable);
+	const UClass* GetSelectedEntryClass() const;
+	void HandleNewEntryClassSelected(const UClass* NewEntryClass) const;
 
 	FText GetSocketName() const;
 	bool CanChangeSocket() const;
diff --git a/Engine/Source/Editor/Kismet/Private/BlueprintEditor.cpp b/Engine/Source/Editor/Kismet/Private/BlueprintEditor.cpp
--- a/Engine/Source/Editor/Kismet/Private/BlueprintEditor.cpp
+++ b/Engine/Source/Editor/Kismet/Private/BlueprintEditor.cpp
@@ -5374,38 +5374,79 @@ bool FBlueprintEditor::CanConvertFunctionToEvent() const
 	{
 		if (UK2Node_FunctionEntry* const SelectedCallFunctionNode = Cast<UK2Node_FunctionEntry>(SelectedNode))
 		{
-			UFunction* const Func = SelectedCallFunctionNode->FindSignatureFunction();
-
-			// If the function has no output parameters and is not an interface or construction script
-			if (Func && !UEdGraphSchema_K2::HasFunctionAnyOutputParameter(Func) &&
-				!FBlueprintEditorUtils::IsInterfaceFunction(GetBlueprintObj(), Func) &&
-				SelectedCallFunctionNode->FunctionReference.GetMemberName() != UEdGraphSchema_K2::FN_UserConstructionScript)
+			if (SelectedCallFunctionNode->FindSignatureFunction())
 			{
 				return true;
 			}
 		}
 	}
-	// We cannot convert functions that have output parameters to events
+
 	return false;
 }
 
 void FBlueprintEditor::OnConvertFunctionToEvent()
 {
+	FCompilerResultsLog LogResults;
+	FMessageLog MessageLog("BlueprintLog");
+	FText SpecificErrorMessage;
+
 	if (UEdGraphNode* SelectedNode = GetSingleSelectedNode())
 	{
 		if (UK2Node_FunctionEntry* SelectedCallFunctionNode = Cast<UK2Node_FunctionEntry>(SelectedNode))
 		{
-			ConvertFunctionToEvent(SelectedCallFunctionNode);
+			UFunction* Func = SelectedCallFunctionNode->FindSignatureFunction();
+			check(Func);
+
+			// Make sure there are no output parameters
+			if (UEdGraphSchema_K2::HasFunctionAnyOutputParameter(Func))
+			{
+				LogResults.Error(*LOCTEXT("FunctionHasOutput_Error", "A function can only be converted if it does not have any output parameters.").ToString());
+				SpecificErrorMessage = LOCTEXT("FunctionHasOutput_Error_Title", "Function cannot have output parameters");
+			}
+			// Make sure this is not an interface function
+			else if (FBlueprintEditorUtils::IsInterfaceFunction(GetBlueprintObj(), Func))
+			{
+				LogResults.Error(*LOCTEXT("FunctionIsFromInterface_Error", "Functions from interfaces cannot be converted to events.").ToString());
+				SpecificErrorMessage = LOCTEXT("FunctionIsFromInterface_Error_Title", "Cannot convert interface functions");
+			}
+			// Make sure this is not a blueprint/macro function library
+			else if (GetBlueprintObj()->BlueprintType == BPTYPE_FunctionLibrary || GetBlueprintObj()->BlueprintType == BPTYPE_MacroLibrary)
+			{
+				LogResults.Error(*LOCTEXT("BlueprintFunctionLibarary_Error", "Cannot convert functions from blueprint or macro libraries.").ToString());
+				SpecificErrorMessage = LOCTEXT("BlueprintFunctionLibarary_Error_Title", "Cannot convert blueprint or macro library functions");
+			}
+			// Ensure that this is no the construction script
+			else if (SelectedCallFunctionNode->FunctionReference.GetMemberName() == UEdGraphSchema_K2::FN_UserConstructionScript)
+			{
+				LogResults.Error(*LOCTEXT("ConvertConstructionScript_Error", "Cannot convert the construction script!").ToString());
+				SpecificErrorMessage = LOCTEXT("ConvertConstructionScript_Error_Title", "Cannot convert construction script");
+			}
+			// Make sure we are not on the animation graph
+			else if (IsEditingAnimGraph())
+			{
+				LogResults.Error(*LOCTEXT("ConvertAnimGraph_Error", "Cannot convert functions on the anim graph!").ToString());
+				SpecificErrorMessage = LOCTEXT("ConvertAnimGraph_Error_Title", "Cannot convert on the anim graph");
+			}
+			else
+			{
+				ConvertFunctionToEvent(SelectedCallFunctionNode);
+			}
 		}
 	}
 	else
 	{
-		FCompilerResultsLog LogResults;
 		LogResults.Error(*LOCTEXT("MultipleNodesSelectred_Error", "Only one node can be selected for conversion!").ToString());
-		FMessageLog MessageLog("BlueprintLog");
+	}
+
+	// Show the log results if there were any errors
+	if (LogResults.NumErrors)
+	{
 		MessageLog.NewPage(LOCTEXT("OnConvertEventToFunctionLabel", "Convert Event to Function"));
 		MessageLog.AddMessages(LogResults.Messages);
-		MessageLog.Notify(LOCTEXT("OnConvertEventToFunctionError", "Convert Event to Function Failed!"));
+
+		FFormatNamedArguments Arguments;
+		Arguments.Add(TEXT("SpecificErrorMessage"), SpecificErrorMessage);
+		MessageLog.Notify(FText::Format(LOCTEXT("OnConvertEventToFunctionError", "Convert Event to Function Failed!\n{SpecificErrorMessage}"), Arguments));
 	}
 }
 
@@ -5566,36 +5607,72 @@ bool FBlueprintEditor::CanConvertEventToFunction() const
 {
 	// Only allow when only a _single_ event node is selected, not allowed on the anim graph
 	UEdGraphNode* SelectedNode = GetSingleSelectedNode();
-	if (SelectedNode && !IsEditingAnimGraph())
+	if (SelectedNode && SelectedNode->IsA<UK2Node_Event>())
+	{
+		return true;
+	}
+
+	return false;
+}
+
+void FBlueprintEditor::OnConvertEventToFunction()
+{
+	FCompilerResultsLog LogResults;
+	FMessageLog MessageLog("BlueprintLog");
+	FText SpecificErrorMessage;
+
+	if (UEdGraphNode* SelectedNode = GetSingleSelectedNode())
 	{
 		if (UK2Node_Event* SelectedEventNode = Cast<UK2Node_Event>(SelectedNode))
 		{
 			UFunction* const Func = FFunctionFromNodeHelper::FunctionFromNode(SelectedEventNode);
 
-			// If the function has no output parameters it is valid to make a function
-			if (!FBlueprintEditorUtils::IsInterfaceFunction(GetBlueprintObj(), Func) &&
-				!SelectedEventNode->IsInterfaceEventNode())
+			// Ensure we are not trying to do this on an interface event
+			if (FBlueprintEditorUtils::IsInterfaceFunction(GetBlueprintObj(), Func) || SelectedEventNode->IsInterfaceEventNode())
 			{
-				return true;
+				LogResults.Error(*LOCTEXT("EventIsOnInterface_Error", "Only non-interface events can be converted to functions!").ToString());
+				SpecificErrorMessage = LOCTEXT("EventIsOnInterface_Error_Title", "Cannot convert interface events");
+			}
+			// Make sure that we are not on the animation graph
+			else if (IsEditingAnimGraph())
+			{
+				LogResults.Error(*LOCTEXT("ConvertAnimGraphEvent_Error", "Cannot convert events on the anim graph!").ToString());
+				SpecificErrorMessage = LOCTEXT("ConvertAnimGraphEvent_Error_Title", "Cannot convert on the anim graph");
+			}
+			else
+			{
+				ConvertEventToFunction(SelectedEventNode);
 			}
 		}
 	}
+	else
+	{
+		SpecificErrorMessage = LOCTEXT("MultipleNodesSelectred_Error_Title", "Only one node can be selected for conversion");
+		LogResults.Error(*LOCTEXT("MultipleNodesSelectred_Error", "Only one node can be selected for conversion!").ToString());
+	}
 
-	return false;
+	// Show the log results if there were any errors
+	if (LogResults.NumErrors)
+	{
+		MessageLog.NewPage(LOCTEXT("OnConvertEventToFunctionLabel", "Convert Event to Function"));
+		MessageLog.AddMessages(LogResults.Messages);
+
+		// Format the title node
+		FFormatNamedArguments Arguments;
+		Arguments.Add(TEXT("SpecificErrorMessage"), SpecificErrorMessage);
+		MessageLog.Notify(FText::Format(LOCTEXT("OnConvertEventToFunctionError", "Convert Event to Function Failed!\n{SpecificErrorMessage}"), Arguments));
+	}
 }
 
-void FBlueprintEditor::OnConvertEventToFunction()
-{
-	const FScopedTransaction Transaction(FGraphEditorCommands::Get().ConvertEventToFunction->GetLabel());
+void FBlueprintEditor::ConvertEventToFunction(UK2Node_Event* SelectedEventNode)
+{	
+	if (SelectedEventNode)
+	{
+		const FScopedTransaction Transaction(FGraphEditorCommands::Get().ConvertEventToFunction->GetLabel());
 
-	UBlueprint* const NodeBP = GetBlueprintObj();
-	NodeBP->Modify();
-	
-	// Keep track of the current selected node
-	UEdGraphNode* const SelectedNode = GetSingleSelectedNode();
+		UBlueprint* const NodeBP = GetBlueprintObj();
+		NodeBP->Modify();
 
-	if (UK2Node_Event* SelectedEventNode = Cast<UK2Node_Event>(SelectedNode))
-	{
 		const UEdGraphSchema_K2* Schema = GetDefault<UEdGraphSchema_K2>();
 
 		// Break any connections to the delegate pin, because that won't exist 
@@ -5611,9 +5688,9 @@ void FBlueprintEditor::OnConvertEventToFunction()
 
 		// Keep track of the old connections from the event node
 		TMap<FString, TSet<UEdGraphPin*>> PinConnections;
-		GetPinConnectionMap(SelectedNode, PinConnections);
+		GetPinConnectionMap(SelectedEventNode, PinConnections);
 
-		TArray<UEdGraphNode*> CollapsableNodes = GetAllConnectedNodes(SelectedNode);
+		TArray<UEdGraphNode*> CollapsableNodes = GetAllConnectedNodes(SelectedEventNode);
 
 		UFunction* const FunctionSig = FFunctionFromNodeHelper::FunctionFromNode(SelectedEventNode);
 		UEdGraph* const SourceGraph = SelectedEventNode->GetGraph();
diff --git a/Engine/Source/Editor/Kismet/Private/SMyBlueprint.cpp b/Engine/Source/Editor/Kismet/Private/SMyBlueprint.cpp
--- a/Engine/Source/Editor/Kismet/Private/SMyBlueprint.cpp
+++ b/Engine/Source/Editor/Kismet/Private/SMyBlueprint.cpp
@@ -1313,7 +1313,27 @@ void SMyBlueprint::CollectAllActions(FGraphActionListBuilderBase& OutAllActions)
 		}
 	}
 
-	// Also function implemented for interfaces
+	auto IsInAnimBPLambda = [&BlueprintObj](const FName FunctionName, FText& FunctionCategory) -> bool
+	{
+		if (BlueprintObj->SkeletonGeneratedClass != nullptr)
+		{
+			if (UFunction * Function = BlueprintObj->SkeletonGeneratedClass->FindFunctionByName(FunctionName))
+			{
+				FunctionCategory = Function->GetMetaDataText(FBlueprintMetadata::MD_FunctionCategory, TEXT("UObjectCategory"), Function->GetFullGroupName(false));
+
+				if (IAnimClassInterface * AnimClassInterface = IAnimClassInterface::GetFromClass(BlueprintObj->SkeletonGeneratedClass))
+				{
+					if (IAnimClassInterface::IsAnimBlueprintFunction(AnimClassInterface, Function))
+					{
+						return true;
+					}
+				}
+			}
+		}
+		return false;
+	};
+
+	// Also functions implemented from interfaces
 	for (int32 i=0; i < BlueprintObj->ImplementedInterfaces.Num(); i++)
 	{
 		FBPInterfaceDescription& InterfaceDesc = BlueprintObj->ImplementedInterfaces[i];
@@ -1322,33 +1342,18 @@ void SMyBlueprint::CollectAllActions(FGraphActionListBuilderBase& OutAllActions)
 		{
 			for (TFieldIterator<UFunction> FunctionIt(InterfaceClass, EFieldIteratorFlags::IncludeSuper); FunctionIt; ++FunctionIt)
 			{
-				const FName FunctionName = (*FunctionIt)->GetFName();
+				const UFunction* Function = *FunctionIt;
+				const FName FunctionName = Function->GetFName();
 
 				if (FunctionName != UEdGraphSchema_K2::FN_ExecuteUbergraphBase)
 				{
-					FString FunctionTooltip = FunctionName.ToString();
-					FString FunctionDesc = FunctionName.ToString();
+					FText FunctionTooltip = Function->GetToolTipText();
+					FText FunctionDesc = K2Schema->GetFriendlySignatureName(Function);
 
 					FText FunctionCategory;
-					bool bIsAnimFunction = false;
-
-					if (BlueprintObj->SkeletonGeneratedClass != nullptr)
-					{
-						if (UFunction * Function = BlueprintObj->SkeletonGeneratedClass->FindFunctionByName(FunctionName))
-						{
-							FunctionCategory = Function->GetMetaDataText(FBlueprintMetadata::MD_FunctionCategory, TEXT("UObjectCategory"), Function->GetFullGroupName(false));
-
-							if (IAnimClassInterface * AnimClassInterface = IAnimClassInterface::GetFromClass(BlueprintObj->SkeletonGeneratedClass))
-							{
-								if (IAnimClassInterface::IsAnimBlueprintFunction(AnimClassInterface, Function))
-								{
-									bIsAnimFunction = true;
-								}
-							}
-						}
-					}
+					bool bIsAnimFunction = IsInAnimBPLambda(FunctionName, FunctionCategory);
 
-					TSharedPtr<FEdGraphSchemaAction_K2Graph> NewFuncAction = MakeShareable(new FEdGraphSchemaAction_K2Graph(EEdGraphSchemaAction_K2Graph::Interface, FunctionCategory, FText::FromString(FunctionDesc), FText::FromString(FunctionTooltip), 1, bIsAnimFunction ? NodeSectionID::ANIMLAYER : NodeSectionID::INTERFACE));
+					TSharedPtr<FEdGraphSchemaAction_K2Graph> NewFuncAction = MakeShareable(new FEdGraphSchemaAction_K2Graph(EEdGraphSchemaAction_K2Graph::Interface, FunctionCategory, FunctionDesc, FunctionTooltip, 1, bIsAnimFunction ? NodeSectionID::ANIMLAYER : NodeSectionID::INTERFACE));
 
 					NewFuncAction->FuncName = FunctionName;
 					OutAllActions.AddAction(NewFuncAction);
@@ -1380,27 +1385,31 @@ void SMyBlueprint::CollectAllActions(FGraphActionListBuilderBase& OutAllActions)
 		for (int32 Idx = 0; Idx < TempClass->Interfaces.Num(); ++Idx)
 		{
 			FImplementedInterface const& I = TempClass->Interfaces[Idx];
-			if (!I.bImplementedByK2)
+			
+			// same as above, make a function?
+			for (TFieldIterator<UFunction> FunctionIt(I.Class, EFieldIteratorFlags::IncludeSuper); FunctionIt; ++FunctionIt)
 			{
-				// same as above, make a function?
-				for (TFieldIterator<UFunction> FunctionIt(I.Class, EFieldIteratorFlags::IncludeSuper); FunctionIt; ++FunctionIt)
-				{
-					const UFunction* Function = *FunctionIt;
-					const FName FunctionName = Function->GetFName();
+				const UFunction* Function = *FunctionIt;
+				const FName FunctionName = Function->GetFName();
 
-					if (UEdGraphSchema_K2::CanKismetOverrideFunction(Function) && !ImplementedFunctionCache.Contains(FunctionName) && !UEdGraphSchema_K2::FunctionCanBePlacedAsEvent(Function))
-					{
-						FText FunctionTooltip = Function->GetToolTipText();
-						FText FunctionDesc = K2Schema->GetFriendlySignatureName(Function);
+				if (UEdGraphSchema_K2::CanKismetOverrideFunction(Function) && !ImplementedFunctionCache.Contains(FunctionName))
+				{
+					FText FunctionTooltip = Function->GetToolTipText();
+					FText FunctionDesc = K2Schema->GetFriendlySignatureName(Function);
 
-						FText FunctionCategory = Function->GetMetaDataText(FBlueprintMetadata::MD_FunctionCategory, TEXT("UObjectCategory"), Function->GetFullGroupName(false));
+					FText FunctionCategory = Function->GetMetaDataText(FBlueprintMetadata::MD_FunctionCategory, TEXT("UObjectCategory"), Function->GetFullGroupName(false));
+					bool bIsAnimFunction = IsInAnimBPLambda(FunctionName, FunctionCategory);
 
-						TSharedPtr<FEdGraphSchemaAction_K2Graph> NewFuncAction = MakeShareable(new FEdGraphSchemaAction_K2Graph(EEdGraphSchemaAction_K2Graph::Function, FunctionCategory, FunctionDesc, FunctionTooltip, 1, NodeSectionID::INTERFACE));
-						NewFuncAction->FuncName = FunctionName;
+					TSharedPtr<FEdGraphSchemaAction_K2Graph> NewFuncAction = MakeShareable(new FEdGraphSchemaAction_K2Graph(EEdGraphSchemaAction_K2Graph::Interface, FunctionCategory, FunctionDesc, FunctionTooltip, 1, bIsAnimFunction ? NodeSectionID::ANIMLAYER : NodeSectionID::INTERFACE));
+					NewFuncAction->FuncName = FunctionName;
+					
+					if (!OverridableFunctionNames.Contains(FunctionName))
+					{
 						OverridableFunctionActions.Add(NewFuncAction);
 						OverridableFunctionNames.Add(FunctionName);
-						OutAllActions.AddAction(NewFuncAction);
 					}
+
+					OutAllActions.AddAction(NewFuncAction);
 				}
 			}
 		}
@@ -1567,7 +1576,7 @@ FReply SMyBlueprint::OnActionDragged( const TArray< TSharedPtr<FEdGraphSchemaAct
 		return FReply::Unhandled();
 	}
 
-	TSharedPtr<FEdGraphSchemaAction> InAction( InActions.Num() > 0 ? InActions[0] : NULL );
+	TSharedPtr<FEdGraphSchemaAction> InAction( InActions.Num() > 0 ? InActions[0] : nullptr );
 	if(InAction.IsValid())
 	{
 		auto AnalyticsDelegate = FNodeCreationAnalytic::CreateSP( this, &SMyBlueprint::UpdateNodeCreation );
@@ -1655,17 +1664,32 @@ FReply SMyBlueprint::OnActionDragged( const TArray< TSharedPtr<FEdGraphSchemaAct
 		{	
 			// Check if it's a custom event, it is preferable to drop a call function for custom events than to focus on the node
 			FEdGraphSchemaAction_K2Event* FuncAction = (FEdGraphSchemaAction_K2Event*)InAction.Get();
-			if (UK2Node_CustomEvent* CustomEvent = Cast<UK2Node_CustomEvent>(FuncAction->NodeTemplate))
+			if (UK2Node_Event* Event = Cast<UK2Node_Event>(FuncAction->NodeTemplate))
 			{
-				return FReply::Handled().BeginDragDrop(FKismetFunctionDragDropAction::New(InAction, CustomEvent->GetFunctionName(), GetBlueprintObj()->SkeletonGeneratedClass, FMemberReference(), AnalyticsDelegate, FKismetDragDropAction::FCanBeDroppedDelegate()));
-			}
-			else
-			{
-				// don't need a valid FCanBeDroppedDelegate because this entry means we already have this 
-				// event placed (so this action will just focus it)
-				TSharedRef<FKismetDragDropAction> DragOperation = FKismetDragDropAction::New(InAction, AnalyticsDelegate, FKismetDragDropAction::FCanBeDroppedDelegate());
+				UFunction* const Function = FFunctionFromNodeHelper::FunctionFromNode(Event);
 
-				return FReply::Handled().BeginDragDrop(DragOperation);
+				// Callback function to report that the user cannot drop this function in the graph
+				auto CanDragDropAction = [](TSharedPtr<FEdGraphSchemaAction> /*DropAction*/, UEdGraph* /*HoveredGraphIn*/, FText& ImpededReasonOut, UFunction* Func)->bool
+				{
+					// If this function is not BP callable then don't let it be dropped
+					if (Func && !(Func->FunctionFlags & (FUNC_BlueprintCallable | FUNC_BlueprintPure)))
+					{
+						ImpededReasonOut = LOCTEXT("NonBlueprintCallableEvent", "This event was not marked as Blueprint Callable and cannot be placed in a graph!");
+						return false;
+					}
+
+					return true;
+				};
+
+				TSharedRef< FKismetFunctionDragDropAction> DragOperation =
+					FKismetFunctionDragDropAction::New(
+						InAction, Function->GetFName(),
+						GetBlueprintObj()->SkeletonGeneratedClass,
+						FMemberReference(),
+						AnalyticsDelegate,
+						FKismetDragDropAction::FCanBeDroppedDelegate::CreateLambda(CanDragDropAction, Function)
+					);
+				return FReply::Handled().BeginDragDrop(DragOperation);	
 			}
 		}
 	}
diff --git a/Engine/Source/Editor/Kismet/Public/BlueprintEditor.h b/Engine/Source/Editor/Kismet/Public/BlueprintEditor.h
--- a/Engine/Source/Editor/Kismet/Public/BlueprintEditor.h
+++ b/Engine/Source/Editor/Kismet/Public/BlueprintEditor.h
@@ -46,6 +46,7 @@ class UUserDefinedStruct;
 class UBlueprintEditorOptions;
 struct Rect;
 class UK2Node_FunctionEntry;
+class UK2Node_Event;
 
 /* Enums to use when grouping the blueprint members in the list panel. The order here will determine the order in the list */
 namespace NodeSectionID
@@ -853,6 +854,7 @@ protected:
 	void ReconnectPinMap(UEdGraphNode* Node, const TMap<FString, TSet<UEdGraphPin*>>& PinConnections);
 
 	void OnConvertEventToFunction();
+	void ConvertEventToFunction(UK2Node_Event* SelectedEventNode);
 	bool CanConvertEventToFunction() const;
 
 	/**
diff --git a/Engine/Source/Editor/KismetCompiler/Private/KismetCompiler.cpp b/Engine/Source/Editor/KismetCompiler/Private/KismetCompiler.cpp
--- a/Engine/Source/Editor/KismetCompiler/Private/KismetCompiler.cpp
+++ b/Engine/Source/Editor/KismetCompiler/Private/KismetCompiler.cpp
@@ -591,6 +591,44 @@ void FKismetCompilerContext::ValidateVariableNames()
 	}
 }
 
+void FKismetCompilerContext::ValidateComponentClassOverrides()
+{
+	if (UClass* ParentClass = Blueprint->ParentClass)
+	{
+		if (UObject* CDO = ParentClass->GetDefaultObject(false))
+		{
+			for (auto It(Blueprint->ComponentClassOverrides.CreateIterator()); It; ++It)
+			{
+				const FBPComponentClassOverride& Override = *It;
+				if (UObject* OverridenObject = (UObject*)FindObjectWithOuter(CDO, nullptr, Override.ComponentName))
+				{
+					if (Override.ComponentClass && !Override.ComponentClass->IsChildOf(OverridenObject->GetClass()))
+					{
+						MessageLog.Error(
+							*FText::Format(
+								LOCTEXT("InvalidOverride", "{0} is not a legal override for component {1} because it does not derive from {2}."),
+								FText::FromName(Override.ComponentClass->GetFName()),
+								FText::FromName(Override.ComponentName),
+								FText::FromName(OverridenObject->GetClass()->GetFName())
+								).ToString()
+						);
+					}
+				}
+				else
+				{
+					MessageLog.Note(
+						*FText::Format(
+							LOCTEXT("UnneededOverride", "Removing class override for component {0} that no longer exists."),
+							FText::FromName(Override.ComponentName)
+						).ToString()
+					);
+					It.RemoveCurrent();
+				}
+			}
+		}
+	}
+}
+
 void FKismetCompilerContext::ValidateTimelineNames()
 {
 	TSharedPtr<FKismetNameValidator> ParentBPNameValidator;
@@ -2449,6 +2487,7 @@ void FKismetCompilerContext::FinishCompilingClass(UClass* Class)
 		BPGClass->Timelines = Blueprint->Timelines;
 		BPGClass->SimpleConstructionScript = Blueprint->SimpleConstructionScript;
 		BPGClass->InheritableComponentHandler = Blueprint->InheritableComponentHandler;
+		BPGClass->ComponentClassOverrides = Blueprint->ComponentClassOverrides;
 	}
 
 	//@TODO: Not sure if doing this again is actually necessary
@@ -3952,6 +3991,8 @@ void FKismetCompilerContext::CompileClassLayout(EInternalCompilerFlags InternalF
 	// This validation requires CDO object.
 	ValidateVariableNames();
 
+	ValidateComponentClassOverrides();
+
 	OldCDO = NULL;
 	OldGenLinkerIdx = INDEX_NONE;
 	OldLinker = Blueprint->GetLinker();
diff --git a/Engine/Source/Editor/KismetCompiler/Public/KismetCompiler.h b/Engine/Source/Editor/KismetCompiler/Public/KismetCompiler.h
--- a/Engine/Source/Editor/KismetCompiler/Public/KismetCompiler.h
+++ b/Engine/Source/Editor/KismetCompiler/Public/KismetCompiler.h
@@ -267,7 +267,10 @@ public:
 
 	/** Ensures that all variables have valid names for compilation/replication */
 	void ValidateVariableNames();
-	
+
+	/** Ensures that all component class overrides are legal overrides of the parent class */
+	void ValidateComponentClassOverrides();
+
 	/** Creates a class variable for each entry in the Blueprint NewVars array */
 	virtual void CreateClassVariablesFromBlueprint();
 
diff --git a/Engine/Source/Editor/PropertyEditor/Private/DetailItemNode.cpp b/Engine/Source/Editor/PropertyEditor/Private/DetailItemNode.cpp
--- a/Engine/Source/Editor/PropertyEditor/Private/DetailItemNode.cpp
+++ b/Engine/Source/Editor/PropertyEditor/Private/DetailItemNode.cpp
@@ -495,7 +495,7 @@ static bool PassesAllFilters( FDetailItemNode* ItemNode, const FDetailLayoutCust
 			const bool bIsSeenDueToFiltering = PropertyNodePin->HasNodeFlags(EPropertyNodeFlags::IsSeenDueToFiltering) != 0;
 			const bool bIsParentSeenDueToFiltering = PropertyNodePin->HasNodeFlags(EPropertyNodeFlags::IsParentSeenDueToFiltering) != 0;
 
-			const bool bPassesSearchFilter = bSearchFilterIsEmpty || ( bIsNotBeingFiltered || bIsSeenDueToFiltering || bIsParentSeenDueToFiltering );
+			const bool bPassesSearchFilter = bPassesCategoryFilter || bSearchFilterIsEmpty || ( bIsNotBeingFiltered || bIsSeenDueToFiltering || bIsParentSeenDueToFiltering );
 			const bool bPassesModifiedFilter = bPassesSearchFilter && ( InFilter.bShowOnlyModifiedProperties == false || PropertyNodePin->GetDiffersFromDefault() == true );
 			const bool bPassesDifferingFilter = InFilter.bShowOnlyDiffering ? InFilter.WhitelistedProperties.Find(*FPropertyNode::CreatePropertyPath(PropertyNodePin.ToSharedRef())) != nullptr : true;
 
@@ -516,16 +516,25 @@ static bool PassesAllFilters( FDetailItemNode* ItemNode, const FDetailLayoutCust
 			const bool bPassesAnimatedFilter = (InFilter.bShowAnimated == false || Local::ItemIsAnimated(ItemNode, PropertyNodePin));
 
 			// The property node is visible (note categories are never visible unless they have a child that is visible )
-			bPassesAllFilters = (bPassesSearchFilter && bPassesModifiedFilter && bPassesDifferingFilter && bPassesKeyableFilter && bPassesAnimatedFilter) || bPassesCategoryFilter;
+			bPassesAllFilters = bPassesSearchFilter && bPassesModifiedFilter && bPassesDifferingFilter && bPassesKeyableFilter && bPassesAnimatedFilter;
 		}
 		else if (InCustomization.HasCustomWidget())
 		{
-			const bool bPassesTextFilter = Local::StringPassesFilter(InFilter, InCustomization.WidgetDecl->FilterTextString.ToString());
+			const bool bPassesTextFilter = bPassesCategoryFilter || Local::StringPassesFilter(InFilter, InCustomization.WidgetDecl->FilterTextString.ToString());
 			const bool bPassesModifiedFilter = (InFilter.bShowOnlyModifiedProperties == false || InCustomization.WidgetDecl->DiffersFromDefaultAttr.Get() == true);
 			//@todo we need to support custom widgets for keyable,animated, in particularly for transforms(ComponentTransformDetails).
 			const bool bPassesKeyableFilter = (InFilter.bShowKeyable == false);
 			const bool bPassesAnimatedFilter = (InFilter.bShowAnimated == false);
-			bPassesAllFilters = (bPassesTextFilter && bPassesModifiedFilter && bPassesKeyableFilter && bPassesAnimatedFilter) || bPassesCategoryFilter;
+			bPassesAllFilters = bPassesTextFilter && bPassesModifiedFilter && bPassesKeyableFilter && bPassesAnimatedFilter;
+		}
+		else if (InCustomization.HasCustomBuilder())
+		{
+			const bool bPassesTextFilter = bPassesCategoryFilter || Local::StringPassesFilter(InFilter, InCustomization.CustomBuilderRow->GetWidgetRow().FilterTextString.ToString());
+			//@todo we need to support custom builders for modified, keyable, animated, in particularly for transforms(ComponentTransformDetails).
+			const bool bPassesModifiedFilter = (InFilter.bShowOnlyModifiedProperties == false);
+			const bool bPassesKeyableFilter = (InFilter.bShowKeyable == false);
+			const bool bPassesAnimatedFilter = (InFilter.bShowAnimated == false);
+			bPassesAllFilters = bPassesTextFilter && bPassesModifiedFilter && bPassesKeyableFilter && bPassesAnimatedFilter;
 		}
 	}
 
diff --git a/Engine/Source/Editor/UnrealEd/Private/Commandlets/ContentCommandlets.cpp b/Engine/Source/Editor/UnrealEd/Private/Commandlets/ContentCommandlets.cpp
--- a/Engine/Source/Editor/UnrealEd/Private/Commandlets/ContentCommandlets.cpp
+++ b/Engine/Source/Editor/UnrealEd/Private/Commandlets/ContentCommandlets.cpp
@@ -1543,7 +1543,7 @@ void UResavePackagesCommandlet::PerformAdditionalOperations(class UWorld* World,
 			// If we are (minimally) rebuilding HLOD, set the visible streamed-in levels packages to clean (as FlushLevelStreaming will dirty their packages in this commandlet context)
 			if(bShouldBuildHLOD)
 			{
-				for (const ULevel* Level : GWorld->GetLevels())
+				for (const ULevel* Level : World->GetLevels())
 				{
 					if (Level->bIsVisible)
 					{
@@ -1569,7 +1569,7 @@ void UResavePackagesCommandlet::PerformAdditionalOperations(class UWorld* World,
 					TSubclassOf<UHierarchicalLODSetup> NewHLODSetupAsset = LoadClass<UHierarchicalLODSetup>( NULL, *ForceHLODSetupAsset, NULL, LOAD_None, NULL );
 					if( NewHLODSetupAsset != nullptr )
 					{
-						GWorld->GetWorldSettings()->HLODSetupAsset = NewHLODSetupAsset;
+						World->GetWorldSettings()->HLODSetupAsset = NewHLODSetupAsset;
 					}
 					else
 					{
@@ -1580,16 +1580,16 @@ void UResavePackagesCommandlet::PerformAdditionalOperations(class UWorld* World,
 				// Force HLOD support on for this level if we were asked to
 				if( bForceEnableHLODForLevel )
 				{
-					GWorld->GetWorldSettings()->bEnableHierarchicalLODSystem = true;
+					World->GetWorldSettings()->bEnableHierarchicalLODSystem = true;
 				}
 
 				// Use a single cluster for all actors in the level if we were asked to
 				if( bForceSingleClusterForLevel )
 				{
-					GWorld->GetWorldSettings()->bGenerateSingleClusterForLevel = true;
+					World->GetWorldSettings()->bGenerateSingleClusterForLevel = true;
 				}
 
-				FHierarchicalLODBuilder Builder(GWorld);
+				FHierarchicalLODBuilder Builder(World);
 
 				if (bForceClusterGeneration)
 				{
@@ -1625,7 +1625,7 @@ void UResavePackagesCommandlet::PerformAdditionalOperations(class UWorld* World,
 
 				// Get the list of packages needs to be saved.
 				TSet<UPackage*> PackagesToSave;
-				for(ULevel* Level : GWorld->GetLevels())
+				for(ULevel* Level : World->GetLevels())
 				{
 					if(Level->bIsVisible)
 					{
diff --git a/Engine/Source/Editor/UnrealEd/Private/Commandlets/PackageUtilities.cpp b/Engine/Source/Editor/UnrealEd/Private/Commandlets/PackageUtilities.cpp
--- a/Engine/Source/Editor/UnrealEd/Private/Commandlets/PackageUtilities.cpp
+++ b/Engine/Source/Editor/UnrealEd/Private/Commandlets/PackageUtilities.cpp
@@ -2726,9 +2726,9 @@ int32 UReplaceActorCommandlet::Main(const FString& Params)
 	}
 
 	// UEditorEngine::FinishDestroy() expects GWorld to exist
-	if( GWorld )
+	if( UWorld* World = GWorld )
 	{
-		GWorld->DestroyWorld( false );
+		World->DestroyWorld( false );
 	}
 	GWorld = UWorld::CreateWorld(EWorldType::Editor, false );
 	return 0;
diff --git a/Engine/Source/Editor/UnrealEd/Private/EditorBuildUtils.cpp b/Engine/Source/Editor/UnrealEd/Private/EditorBuildUtils.cpp
--- a/Engine/Source/Editor/UnrealEd/Private/EditorBuildUtils.cpp
+++ b/Engine/Source/Editor/UnrealEd/Private/EditorBuildUtils.cpp
@@ -718,12 +718,13 @@ bool FEditorBuildUtils::PrepForAutomatedBuild( const FEditorAutomatedBuildSettin
 	if ( bBuildSuccessful )
 	{
 		bool bVisibilityToggled = false;
-		if ( !FLevelUtils::IsLevelVisible( GWorld->PersistentLevel ) )
+		UWorld* World = GWorld;
+		if ( !FLevelUtils::IsLevelVisible( World->PersistentLevel ) )
 		{
-			EditorLevelUtils::SetLevelVisibility( GWorld->PersistentLevel, true, false );
+			EditorLevelUtils::SetLevelVisibility( World->PersistentLevel, true, false );
 			bVisibilityToggled = true;
 		}
-		for (ULevelStreaming* CurStreamingLevel : GWorld->GetStreamingLevels())
+		for (ULevelStreaming* CurStreamingLevel : World->GetStreamingLevels())
 		{
 			if ( CurStreamingLevel && !FLevelUtils::IsStreamingLevelVisibleInEditor( CurStreamingLevel ) )
 			{
@@ -733,7 +734,7 @@ bool FEditorBuildUtils::PrepForAutomatedBuild( const FEditorAutomatedBuildSettin
 		}
 		if ( bVisibilityToggled )
 		{
-			GWorld->FlushLevelStreaming();
+			World->FlushLevelStreaming();
 		}
 	}
 
diff --git a/Engine/Source/Editor/UnrealEd/Private/Kismet2/BlueprintEditorUtils.cpp b/Engine/Source/Editor/UnrealEd/Private/Kismet2/BlueprintEditorUtils.cpp
--- a/Engine/Source/Editor/UnrealEd/Private/Kismet2/BlueprintEditorUtils.cpp
+++ b/Engine/Source/Editor/UnrealEd/Private/Kismet2/BlueprintEditorUtils.cpp
@@ -3116,7 +3116,7 @@ bool FBlueprintEditorUtils::IsDataOnlyBlueprint(const UBlueprint* Blueprint)
 		return false;
 	}
 
-	if (Blueprint->ComponentTemplates.Num() || Blueprint->Timelines.Num())
+	if (Blueprint->ComponentTemplates.Num() > 0 || Blueprint->Timelines.Num() > 0 || Blueprint->ComponentClassOverrides.Num() > 0)
 	{
 		return false;
 	}
@@ -3238,6 +3238,11 @@ bool FBlueprintEditorUtils::IsInterfaceBlueprint(const UBlueprint* Blueprint)
 	return (Blueprint && Blueprint->BlueprintType == BPTYPE_Interface);
 }
 
+bool FBlueprintEditorUtils::IsInterfaceGraph(const UEdGraph* Graph)
+{
+	return IsInterfaceBlueprint(FindBlueprintForGraph(Graph));
+}
+
 bool FBlueprintEditorUtils::IsLevelScriptBlueprint(const UBlueprint* Blueprint)
 {
 	return (Blueprint && Blueprint->BlueprintType == BPTYPE_LevelScript);
@@ -3461,26 +3466,81 @@ int32 FBlueprintEditorUtils::FindNewVariableIndex(const UBlueprint* Blueprint, c
 	return INDEX_NONE;
 }
 
-bool FBlueprintEditorUtils::MoveVariableBeforeVariable(UBlueprint* Blueprint, FName VarNameToMove, FName TargetVarName, bool bDontRecompile)
+int32 FBlueprintEditorUtils::FindLocalVariableIndex(const UBlueprint* Blueprint, UStruct* VariableScope, const FName& InVariableName)
+{
+	UK2Node_FunctionEntry* FunctionEntryNode = nullptr;
+	FindLocalVariable(Blueprint, VariableScope, InVariableName, &FunctionEntryNode);
+
+	if (FunctionEntryNode != nullptr)
+	{
+		for (int32 i = 0; i < FunctionEntryNode->LocalVariables.Num(); i++)
+		{
+			if (FunctionEntryNode->LocalVariables[i].VarName == InVariableName)
+			{
+				return i;
+			}
+		}
+	}
+	return INDEX_NONE;
+}
+
+bool FBlueprintEditorUtils::MoveVariableBeforeVariable(UBlueprint* Blueprint, UStruct* VariableScope, FName VarNameToMove, FName TargetVarName, bool bDontRecompile)
 {
+	check(Blueprint && VariableScope);
+
 	bool bMoved = false;
-	int32 VarIndexToMove = FindNewVariableIndex(Blueprint, VarNameToMove);
-	int32 TargetVarIndex = FindNewVariableIndex(Blueprint, TargetVarName);
-	if(VarIndexToMove != INDEX_NONE && TargetVarIndex != INDEX_NONE)
+	int32 VarIndexToMove = INDEX_NONE;
+	int32 TargetVarIndex = INDEX_NONE;
+
+	//Get the indices of the variables to be re-ordered
+	if (VariableScope->IsA(UFunction::StaticClass()))
+	{
+		VarIndexToMove = FindLocalVariableIndex(Blueprint, VariableScope, VarNameToMove);
+		TargetVarIndex = FindLocalVariableIndex(Blueprint, VariableScope, TargetVarName);
+	}
+	else
+	{
+		VarIndexToMove = FindNewVariableIndex(Blueprint, VarNameToMove);
+		TargetVarIndex = FindNewVariableIndex(Blueprint, TargetVarName);
+	}
+
+	if (VarIndexToMove != INDEX_NONE && TargetVarIndex != INDEX_NONE)
 	{
-		// Copy var we want to move
-		FBPVariableDescription MoveVar = Blueprint->NewVariables[VarIndexToMove];
 		// When we remove item, will back all items after it. If your target is after it, need to adjust
-		if(TargetVarIndex > VarIndexToMove)
+		if (TargetVarIndex > VarIndexToMove)
 		{
 			TargetVarIndex--;
 		}
-		// Remove var we are moving
-		Blueprint->NewVariables.RemoveAt(VarIndexToMove);
-		// Add in before target variable
-		Blueprint->NewVariables.Insert(MoveVar, TargetVarIndex);
 
-		if(!bDontRecompile)
+		//Handle Local vs Class scope 
+		if (VariableScope->IsA(UFunction::StaticClass()))
+		{
+			UK2Node_FunctionEntry* FunctionEntryNode = nullptr;
+			FindLocalVariable(Blueprint, VariableScope, VarNameToMove, &FunctionEntryNode);
+
+			if (FunctionEntryNode != nullptr)
+			{
+				// Copy var we want to move
+				FBPVariableDescription MoveVar = FunctionEntryNode->LocalVariables[VarIndexToMove];
+
+				// Remove var we are moving
+				FunctionEntryNode->LocalVariables.RemoveAt(VarIndexToMove);
+				// Add in before target variable
+				FunctionEntryNode->LocalVariables.Insert(MoveVar, TargetVarIndex);				
+			}
+		}
+		else
+		{			
+			// Copy var we want to move
+			FBPVariableDescription MoveVar = Blueprint->NewVariables[VarIndexToMove];
+			
+			// Remove var we are moving
+			Blueprint->NewVariables.RemoveAt(VarIndexToMove);
+			// Add in before target variable
+			Blueprint->NewVariables.Insert(MoveVar, TargetVarIndex);
+		}
+
+		if (!bDontRecompile)
 		{
 			FBlueprintEditorUtils::MarkBlueprintAsStructurallyModified(Blueprint);
 		}
@@ -6023,6 +6083,8 @@ void FBlueprintEditorUtils::RemoveInterface(UBlueprint* Blueprint, const FName&
 	{
 		FBPInterfaceDescription& CurrentInterface = Blueprint->ImplementedInterfaces[Idx];
 		const UClass* InterfaceClass = Blueprint->ImplementedInterfaces[Idx].Interface;
+		const FScopedTransaction Transaction(LOCTEXT("RemoveInterface", "Remove interface"));
+		Blueprint->Modify();
 
 		// For every function and event in the interface...
 		for (TFieldIterator<UFunction> FunctionIt(InterfaceClass); FunctionIt; ++FunctionIt)
@@ -6045,6 +6107,10 @@ void FBlueprintEditorUtils::RemoveInterface(UBlueprint* Blueprint, const FName&
 				UK2Node_Event* EventNode = *NodeIt;
 				if (EventNode->EventReference.GetMemberParentClass(EventNode->GetBlueprintClassFromNode()) == InterfaceClass)
 				{
+					UEdGraph* EventNodeGraph = EventNode->GetGraph();
+					check(EventNodeGraph);
+					EventNodeGraph->Modify();
+
 					if (bPreserveFunctions)
 					{
 						// Create a custom event with the same name and signature
@@ -6062,7 +6128,7 @@ void FBlueprintEditorUtils::RemoveInterface(UBlueprint* Blueprint, const FName&
 						}
 					}
 
-					EventNode->GetGraph()->RemoveNode(EventNode);
+					EventNodeGraph->RemoveNode(EventNode);
 					break;
 				}
 			}
@@ -6079,7 +6145,7 @@ void FBlueprintEditorUtils::RemoveInterface(UBlueprint* Blueprint, const FName&
 	}
 }
 
-bool FBlueprintEditorUtils::RemoveInterfaceFunction(class UBlueprint* Blueprint, FBPInterfaceDescription& Interface, UFunction* Function, bool bPreserveFunction)
+bool FBlueprintEditorUtils::RemoveInterfaceFunction(UBlueprint* Blueprint, FBPInterfaceDescription& Interface, UFunction* Function, bool bPreserveFunction)
 {
 	for (TArray<UEdGraph*>::TIterator it(Interface.Graphs); it; ++it)
 	{
diff --git a/Engine/Source/Editor/UnrealEd/Private/Kismet2/Kismet2.cpp b/Engine/Source/Editor/UnrealEd/Private/Kismet2/Kismet2.cpp
--- a/Engine/Source/Editor/UnrealEd/Private/Kismet2/Kismet2.cpp
+++ b/Engine/Source/Editor/UnrealEd/Private/Kismet2/Kismet2.cpp
@@ -1124,9 +1124,25 @@ static void ConformComponentsUtils::ConformRemovedNativeComponents(UObject* BpCd
 		{
 			// Keep track of components inherited from the native super class that are still valid.
 			NewNativeComponents.Add(Component);
-
 			continue;
 		}
+
+		// If we have overriden the class of a native component then ensure that the component still exists and that the overriden class is a valid subclass of it
+		if (const FBPComponentClassOverride* BPCO = CastChecked<UBlueprintGeneratedClass>(BlueprintClass)->ComponentClassOverrides.FindByKey(Component->GetFName()))
+		{
+			if (BPCO->ComponentClass == Component->GetClass())
+			{
+				if (UObject* OverridenComponent = (UObject*)FindObjectWithOuter(NativeCDO, UActorComponent::StaticClass(), Component->GetFName()))
+				{
+					if (Component->IsA(OverridenComponent->GetClass()))
+					{
+						NewNativeComponents.Add(Component);
+						continue;
+					}
+				}
+			}
+		}
+
 		// else, the component has been removed from our native super class
 
 		Component->DestroyComponent(/*bPromoteChildren =*/false);
diff --git a/Engine/Source/Editor/UnrealEd/Private/SEditorViewportViewMenu.cpp b/Engine/Source/Editor/UnrealEd/Private/SEditorViewportViewMenu.cpp
--- a/Engine/Source/Editor/UnrealEd/Private/SEditorViewportViewMenu.cpp
+++ b/Engine/Source/Editor/UnrealEd/Private/SEditorViewportViewMenu.cpp
@@ -317,7 +317,8 @@ TSharedRef<SWidget> SEditorViewportViewMenu::GenerateViewMenuContent() const
 					{
 						const FEditorViewportCommands& BaseViewportCommands = FEditorViewportCommands::Get();
 
-						const ERHIFeatureLevel::Type FeatureLevel = (IsInGameThread() && GWorld) ? (ERHIFeatureLevel::Type)GWorld->FeatureLevel : GMaxRHIFeatureLevel;
+						UWorld* World = GWorld;
+						const ERHIFeatureLevel::Type FeatureLevel = (IsInGameThread() && World) ? (ERHIFeatureLevel::Type)World->FeatureLevel : GMaxRHIFeatureLevel;
 
 						Menu.BeginSection("OptimizationViewmodes", LOCTEXT("OptimizationSubMenuHeader", "Optimization Viewmodes"));
 						{
diff --git a/Engine/Source/Editor/UnrealEd/Private/StaticLightingSystem/StaticLightingSystem.cpp b/Engine/Source/Editor/UnrealEd/Private/StaticLightingSystem/StaticLightingSystem.cpp
--- a/Engine/Source/Editor/UnrealEd/Private/StaticLightingSystem/StaticLightingSystem.cpp
+++ b/Engine/Source/Editor/UnrealEd/Private/StaticLightingSystem/StaticLightingSystem.cpp
@@ -289,17 +289,18 @@ void FStaticLightingManager::CreateStaticLightingSystem(const FLightingBuildOpti
 
 		bBuildReflectionCapturesOnFinish = !Options.bOnlyBuildVisibility;
 
-		for (ULevel* Level : GWorld->GetLevels())
+		UWorld* World = GWorld;
+		for (ULevel* Level : World->GetLevels())
 		{
 			if (Level->bIsLightingScenario && Level->bIsVisible)
 			{
-				StaticLightingSystems.Emplace(new FStaticLightingSystem(Options, GWorld, Level));
+				StaticLightingSystems.Emplace(new FStaticLightingSystem(Options, World, Level));
 			}
 		}
 
 		if (StaticLightingSystems.Num() == 0)
 		{
-			StaticLightingSystems.Emplace(new FStaticLightingSystem(Options, GWorld, nullptr));
+			StaticLightingSystems.Emplace(new FStaticLightingSystem(Options, World, nullptr));
 		}
 
 		ActiveStaticLightingSystem = StaticLightingSystems[0].Get();
@@ -315,7 +316,7 @@ void FStaticLightingManager::CreateStaticLightingSystem(const FLightingBuildOpti
 			// BeginLightmassProcess returns false if there are errors or no precomputed lighting is allowed. Handle both cases.
 			static const auto AllowStaticLightingVar = IConsoleManager::Get().FindTConsoleVariableDataInt(TEXT("r.AllowStaticLighting"));
 			const bool bAllowStaticLighting = (!AllowStaticLightingVar || AllowStaticLightingVar->GetValueOnGameThread() != 0);
-			const bool bForceNoPrecomputedLighting = GWorld->GetWorldSettings()->bForceNoPrecomputedLighting || !bAllowStaticLighting;
+			const bool bForceNoPrecomputedLighting = World->GetWorldSettings()->bForceNoPrecomputedLighting || !bAllowStaticLighting;
 
 
 			if (bForceNoPrecomputedLighting)
diff --git a/Engine/Source/Editor/UnrealEd/Public/Kismet2/BlueprintEditorUtils.h b/Engine/Source/Editor/UnrealEd/Public/Kismet2/BlueprintEditorUtils.h
--- a/Engine/Source/Editor/UnrealEd/Public/Kismet2/BlueprintEditorUtils.h
+++ b/Engine/Source/Editor/UnrealEd/Public/Kismet2/BlueprintEditorUtils.h
@@ -30,6 +30,9 @@ class UTimelineTemplate;
 struct FBlueprintCookedComponentInstancingData;
 struct FComponentKey;
 class UAnimGraphNode_Root;
+class UBlueprint;
+struct FBPInterfaceDescription;
+class UFunction;
 
 /** 
   * Flags describing how to handle graph removal
@@ -552,6 +555,13 @@ public:
 	 */
 	static bool IsInterfaceBlueprint(const UBlueprint* Blueprint);
 
+	/**
+	* Whether or not this graph is an interface graph (i.e. is from an interface blueprint)
+	* 
+	* @return	Whether or not this is an interface graph
+	*/
+	static bool IsInterfaceGraph(const UEdGraph* Graph);
+
 	/**
 	 * Whether or not this blueprint is an interface, used only for defining functions to implement
 	 *
@@ -1142,8 +1152,19 @@ public:
 	 */
 	static int32 FindNewVariableIndex(const UBlueprint* Blueprint, const FName& InName);
 
+	/**
+	 * Find the index of a local variable declared in this blueprint. Returns INDEX_NONE if not found.
+	 *
+	 * @param	VariableScope	Struct of owning function.
+	 *
+	 * @param	InVariableName	Name of the variable to find.
+	 *
+	 * @return	The index of the variable, or INDEX_NONE if it wasn't introduced in this blueprint.
+	 */
+	static int32 FindLocalVariableIndex(const UBlueprint* Blueprint, UStruct* VariableScope, const FName& InVariableName);
+
 	/** Change the order of variables in the Blueprint */
-	static bool MoveVariableBeforeVariable(UBlueprint* Blueprint, FName VarNameToMove, FName TargetVarName, bool bDontRecompile);
+	static bool MoveVariableBeforeVariable(UBlueprint* Blueprint, UStruct* VariableScope, FName VarNameToMove, FName TargetVarName, bool bDontRecompile);
 
 	/**
 	 * Find the index of a timeline first declared in this blueprint. Returns INDEX_NONE if not found.
@@ -1275,7 +1296,7 @@ public:
 	* 
 	* @return	True if the function was removed from the blueprint
 	*/
-	static bool RemoveInterfaceFunction(class UBlueprint* Blueprint, struct FBPInterfaceDescription& Interface, class UFunction* Function, bool bPreserveFunction);
+	static bool RemoveInterfaceFunction(UBlueprint* Blueprint, FBPInterfaceDescription& Interface, UFunction* Function, bool bPreserveFunction);
 
 	/**
 	* Promotes a Graph from being an Interface Override to a full member function
diff --git a/Engine/Source/Programs/UnrealHeaderTool/Private/HeaderParser.cpp b/Engine/Source/Programs/UnrealHeaderTool/Private/HeaderParser.cpp
--- a/Engine/Source/Programs/UnrealHeaderTool/Private/HeaderParser.cpp
+++ b/Engine/Source/Programs/UnrealHeaderTool/Private/HeaderParser.cpp
@@ -1784,6 +1784,12 @@ UEnum* FHeaderParser::CompileEnum()
 	}
 
 	CheckDocumentationPolicyForEnum(Enum, EnumValueMetaData, EntryMetaData);
+
+	if (!Enum->IsValidEnumValue(0) && EnumToken.MetaData.Contains(BlueprintTypeName))
+	{
+		UE_LOG_WARNING_UHT(TEXT("'%s' does not have a 0 entry! (This is a problem when the enum is initalized by default)"), *Enum->GetName());
+	}
+
 	return Enum;
 }
 
diff --git a/Engine/Source/Runtime/AnimGraphRuntime/Public/BoneControllers/AnimNode_SplineIK.h b/Engine/Source/Runtime/AnimGraphRuntime/Public/BoneControllers/AnimNode_SplineIK.h
--- a/Engine/Source/Runtime/AnimGraphRuntime/Public/BoneControllers/AnimNode_SplineIK.h
+++ b/Engine/Source/Runtime/AnimGraphRuntime/Public/BoneControllers/AnimNode_SplineIK.h
@@ -14,6 +14,7 @@
 UENUM()
 enum class ESplineBoneAxis : uint8
 {
+	None = 0 UMETA(Hidden),
 	X = 1,
 	Y = 2,
 	Z = 3,
diff --git a/Engine/Source/Runtime/Core/Public/Math/Color.h b/Engine/Source/Runtime/Core/Public/Math/Color.h
--- a/Engine/Source/Runtime/Core/Public/Math/Color.h
+++ b/Engine/Source/Runtime/Core/Public/Math/Color.h
@@ -14,7 +14,7 @@ class FFloat16Color;
 /**
  * Enum for the different kinds of gamma spaces we expect to need to convert from/to.
  */
-enum class EGammaSpace
+enum class EGammaSpace : uint8
 {
 	/** No gamma correction is applied to this space, the incoming colors are assumed to already be in linear space. */
 	Linear,
diff --git a/Engine/Source/Runtime/CoreUObject/Private/UObject/Class.cpp b/Engine/Source/Runtime/CoreUObject/Private/UObject/Class.cpp
--- a/Engine/Source/Runtime/CoreUObject/Private/UObject/Class.cpp
+++ b/Engine/Source/Runtime/CoreUObject/Private/UObject/Class.cpp
@@ -3137,6 +3137,13 @@ UObject* UClass::CreateDefaultObject()
 					USparseDelegateFunction* SparseDelegateFunction = CastChecked<USparseDelegateFunction>(SparseDelegateIt->SignatureFunction);
 					FSparseDelegateStorage::RegisterDelegateOffset(ClassDefaultObject, SparseDelegateFunction->DelegateName, (size_t)&SparseDelegate - (size_t)ClassDefaultObject);
 				}
+				if (HasAnyClassFlags(CLASS_CompiledFromBlueprint))
+				{
+					if (UDynamicClass* DynamicClass = Cast<UDynamicClass>(this))
+					{
+						(*(DynamicClass->DynamicClassInitializer))(DynamicClass);
+					}
+				}
 				(*ClassConstructor)(FObjectInitializer(ClassDefaultObject, ParentDefaultObject, false, bShouldInitializeProperties));
 				if (bDoNotify)
 				{
@@ -3221,7 +3228,7 @@ FFeedbackContext& UClass::GetDefaultPropertiesFeedbackContext()
 * Get the name of the CDO for the this class
 * @return The name of the CDO
 */
-FName UClass::GetDefaultObjectName()
+FName UClass::GetDefaultObjectName() const
 {
 	FString DefaultName;
 	DefaultName.Reserve(NAME_SIZE);
@@ -4612,7 +4619,8 @@ void GetPrivateStaticClassBody(
 	UClass::ClassAddReferencedObjectsType InClassAddReferencedObjects,
 	UClass::StaticClassFunctionType InSuperClassFn,
 	UClass::StaticClassFunctionType InWithinClassFn,
-	bool bIsDynamic /*= false*/
+	bool bIsDynamic /*= false*/,
+	UDynamicClass::DynamicClassInitializerType InDynamicClassInitializerFn /*= nullptr*/
 	)
 {
 #if WITH_HOT_RELOAD
@@ -4690,7 +4698,8 @@ void GetPrivateStaticClassBody(
 			EObjectFlags(RF_Public | RF_Standalone | RF_Transient | RF_Dynamic | (GIsInitialLoad ? RF_MarkAsRootSet : RF_NoFlags)),
 			InClassConstructor,
 			InClassVTableHelperCtorCaller,
-			InClassAddReferencedObjects
+			InClassAddReferencedObjects,
+			InDynamicClassInitializerFn
 			);
 		check(ReturnClass);
 	}
@@ -5277,7 +5286,8 @@ UDynamicClass::UDynamicClass(
 	EObjectFlags	InFlags,
 	ClassConstructorType InClassConstructor,
 	ClassVTableHelperCtorCallerType InClassVTableHelperCtorCaller,
-	ClassAddReferencedObjectsType InClassAddReferencedObjects)
+	ClassAddReferencedObjectsType InClassAddReferencedObjects,
+	DynamicClassInitializerType InDynamicClassInitializer)
 : UClass(
   EC_StaticConstructor
 , InName
@@ -5291,6 +5301,7 @@ UDynamicClass::UDynamicClass(
 , InClassVTableHelperCtorCaller
 , InClassAddReferencedObjects)
 , AnimClassImplementation(nullptr)
+, DynamicClassInitializer(InDynamicClassInitializer)
 {
 }
 
@@ -5305,6 +5316,11 @@ void UDynamicClass::AddReferencedObjects(UObject* InThis, FReferenceCollector& C
 	Collector.AddReferencedObjects(This->ComponentTemplates, This);
 	Collector.AddReferencedObjects(This->Timelines, This);
 
+	for (TPair<FName, UClass*>& Override : This->ComponentClassOverrides)
+	{
+		Collector.AddReferencedObject(Override.Value);
+	}
+
 	Collector.AddReferencedObject(This->AnimClassImplementation, This);
 
 	Super::AddReferencedObjects(This, Collector);
@@ -5332,14 +5348,14 @@ void UDynamicClass::PurgeClass(bool bRecompilingOnLoad)
 	DynamicBindingObjects.Empty();
 	ComponentTemplates.Empty();
 	Timelines.Empty();
+	ComponentClassOverrides.Empty();
 
 	AnimClassImplementation = nullptr;
 }
 
-UObject* UDynamicClass::FindArchetype(UClass* ArchetypeClass, const FName ArchetypeName) const
+UObject* UDynamicClass::FindArchetype(const UClass* ArchetypeClass, const FName ArchetypeName) const
 {
-	UDynamicClass* ThisClass = const_cast<UDynamicClass*>(this);
-	UObject* Archetype = static_cast<UObject*>(FindObjectWithOuter(ThisClass, ArchetypeClass, ArchetypeName));
+	UObject* Archetype = static_cast<UObject*>(FindObjectWithOuter(this, ArchetypeClass, ArchetypeName));
 	if (!Archetype)
 	{
 		// See UBlueprintGeneratedClass::FindArchetype, UE-35259, UE-37480
@@ -5358,6 +5374,17 @@ UObject* UDynamicClass::FindArchetype(UClass* ArchetypeClass, const FName Archet
 		(SuperClass ? SuperClass->FindArchetype(ArchetypeClass, ArchetypeName) : nullptr);
 }
 
+void UDynamicClass::SetupObjectInitializer(FObjectInitializer& ObjectInitializer) const
+{
+	for (const TPair<FName, UClass*>& Override : ComponentClassOverrides)
+	{
+		ObjectInitializer.SetDefaultSubobjectClass(Override.Key, Override.Value);
+	}
+
+	GetSuperClass()->SetupObjectInitializer(ObjectInitializer);
+}
+
+
 UStructProperty* UDynamicClass::FindStructPropertyChecked(const TCHAR* PropertyName) const
 {
 	return FindFieldChecked<UStructProperty>(this, PropertyName);
diff --git a/Engine/Source/Runtime/CoreUObject/Private/UObject/UObjectArchetype.cpp b/Engine/Source/Runtime/CoreUObject/Private/UObject/UObjectArchetype.cpp
--- a/Engine/Source/Runtime/CoreUObject/Private/UObject/UObjectArchetype.cpp
+++ b/Engine/Source/Runtime/CoreUObject/Private/UObject/UObjectArchetype.cpp
@@ -12,7 +12,7 @@
 #include "UObject/UObjectAnnotation.h"
 #include "Stats/StatsMisc.h"
 
-UObject* GetArchetypeFromRequiredInfoImpl(UClass* Class, UObject* Outer, FName Name, EObjectFlags ObjectFlags, bool bUseUpToDateClass)
+UObject* GetArchetypeFromRequiredInfoImpl(const UClass* Class, const UObject* Outer, FName Name, EObjectFlags ObjectFlags, bool bUseUpToDateClass)
 {
 	UObject* Result = NULL;
 	const bool bIsCDO = !!(ObjectFlags&RF_ClassDefaultObject);
@@ -37,8 +37,8 @@ UObject* GetArchetypeFromRequiredInfoImpl(UClass* Class, UObject* Outer, FName N
 			}
 			else if (!!(ObjectFlags&RF_InheritableComponentTemplate) && Outer->IsA<UClass>())
 			{
-				UClass* OuterSuperClass = static_cast<UClass*>(Outer)->GetSuperClass();
-				for (UClass* SuperClassArchetype = bUseUpToDateClass && OuterSuperClass ? OuterSuperClass->GetAuthoritativeClass() : OuterSuperClass;
+				const UClass* OuterSuperClass = static_cast<const UClass*>(Outer)->GetSuperClass();
+				for (const UClass* SuperClassArchetype = bUseUpToDateClass && OuterSuperClass ? OuterSuperClass->GetAuthoritativeClass() : OuterSuperClass;
 					SuperClassArchetype && SuperClassArchetype->HasAllClassFlags(CLASS_CompiledFromBlueprint);
 					SuperClassArchetype = SuperClassArchetype->GetSuperClass())
 				{
@@ -92,7 +92,7 @@ UObject* GetArchetypeFromRequiredInfoImpl(UClass* Class, UObject* Outer, FName N
 	return Result;
 }
 
-UObject* UObject::GetArchetypeFromRequiredInfo(UClass* Class, UObject* Outer, FName Name, EObjectFlags ObjectFlags)
+UObject* UObject::GetArchetypeFromRequiredInfo(const UClass* Class, const UObject* Outer, FName Name, EObjectFlags ObjectFlags)
 {
 	bool bUseUpToDateClass = false;
 #if WITH_EDITOR
diff --git a/Engine/Source/Runtime/CoreUObject/Private/UObject/UObjectGlobals.cpp b/Engine/Source/Runtime/CoreUObject/Private/UObject/UObjectGlobals.cpp
--- a/Engine/Source/Runtime/CoreUObject/Private/UObject/UObjectGlobals.cpp
+++ b/Engine/Source/Runtime/CoreUObject/Private/UObject/UObjectGlobals.cpp
@@ -251,7 +251,7 @@ struct FPerClassNumberSuffixAnnotation
 		return false;
 	}
 
-	TMap<UClass*, int32> Suffixes;
+	TMap<const UClass*, int32> Suffixes;
 };
 
 /**
@@ -259,7 +259,7 @@ struct FPerClassNumberSuffixAnnotation
  *
  * Updating is done via a callback because a lock needs to be maintained while this happens.
  */
-int32 UpdateSuffixForNextNewObject(UObject* Parent, UClass* Class, TFunctionRef<void(int32&)> IndexMutator)
+int32 UpdateSuffixForNextNewObject(UObject* Parent, const UClass* Class, TFunctionRef<void(int32&)> IndexMutator)
 {
 	static FCriticalSection PerClassNumberSuffixAnnotationMutex;
 	static FUObjectAnnotationDense<FPerClassNumberSuffixAnnotation, true> PerClassNumberSuffixAnnotation;
@@ -1747,7 +1747,7 @@ void EndLoad(FUObjectSerializeContext* LoadContext)
 	Object name functions.
 -----------------------------------------------------------------------------*/
 
-FName MakeUniqueObjectName( UObject* Parent, UClass* Class, FName InBaseName/*=NAME_None*/ )
+FName MakeUniqueObjectName( UObject* Parent, const UClass* Class, FName InBaseName/*=NAME_None*/ )
 {
 	CSV_SCOPED_TIMING_STAT(UObject, MakeUniqueObjectName);
 	check(Class);
@@ -2171,7 +2171,7 @@ void SnapshotTransactionBuffer(UObject* Object)
 	}
 }
 
-bool StaticAllocateObjectErrorTests( UClass* InClass, UObject* InOuter, FName InName, EObjectFlags InFlags)
+bool StaticAllocateObjectErrorTests( const UClass* InClass, UObject* InOuter, FName InName, EObjectFlags InFlags)
 {
 	// Validation checks.
 	if( !InClass )
@@ -2239,7 +2239,7 @@ COREUOBJECT_API bool GOutputCookingWarnings = false;
 
 UObject* StaticAllocateObject
 (
-	UClass*			InClass,
+	const UClass*	InClass,
 	UObject*		InOuter,
 	FName			InName,
 	EObjectFlags	InFlags,
@@ -2435,7 +2435,7 @@ UObject* StaticAllocateObject
 	if (!bSubObject)
 	{
 		FMemory::Memzero((void *)Obj, TotalSize);
-		new ((void *)Obj) UObjectBase(InClass, InFlags|RF_NeedInitialization, InternalSetFlags, InOuter, InName);
+		new ((void *)Obj) UObjectBase(const_cast<UClass*>(InClass), InFlags|RF_NeedInitialization, InternalSetFlags, InOuter, InName);
 	}
 	else
 	{
@@ -2557,6 +2557,11 @@ FObjectInitializer::FObjectInitializer(UObject* InObj, UObject* InObjectArchetyp
 	LastConstructedObject = ThreadContext.ConstructedObject;
 	ThreadContext.ConstructedObject = Obj;
 	ThreadContext.PushInitializer(this);
+
+	if (Obj)
+	{
+		Obj->GetClass()->SetupObjectInitializer(*this);
+	}
 }
 
 /**
@@ -3081,7 +3086,7 @@ void FObjectInitializer::AssertIfSubobjectSetupIsNotAllowed(const FName Subobjec
 }
 
 #if DO_CHECK
-void CheckIsClassChildOf_Internal(UClass* Parent, UClass* Child)
+void CheckIsClassChildOf_Internal(const UClass* Parent, const UClass* Child)
 {
 	// This is a function to avoid platform compilation issues
 	checkf(Child, TEXT("NewObject called with a nullptr class object"));
@@ -3091,7 +3096,7 @@ void CheckIsClassChildOf_Internal(UClass* Parent, UClass* Child)
 
 UObject* StaticConstructObject_Internal
 (
-	UClass*			InClass,
+	const UClass*	InClass,
 	UObject*		InOuter								/*=GetTransientPackage()*/,
 	FName			InName								/*=NAME_None*/,
 	EObjectFlags	InFlags								/*=0*/,
diff --git a/Engine/Source/Runtime/CoreUObject/Public/UObject/Class.h b/Engine/Source/Runtime/CoreUObject/Public/UObject/Class.h
--- a/Engine/Source/Runtime/CoreUObject/Public/UObject/Class.h
+++ b/Engine/Source/Runtime/CoreUObject/Public/UObject/Class.h
@@ -2245,7 +2245,7 @@ public:
 	ClassAddReferencedObjectsType ClassAddReferencedObjects;
 
 	/** Class pseudo-unique counter; used to accelerate unique instance name generation */
-	uint32 ClassUnique:31;
+	mutable uint32 ClassUnique:31;
 
 	/** Used to check if the class was cooked or not */
 	uint32 bCooked:1;
@@ -2525,11 +2525,11 @@ public:
 	 * @param	bCreateIfNeeded if true (default) then the CDO is created if it is null
 	 * @return		the CDO for this class
 	 */
-	UObject* GetDefaultObject(bool bCreateIfNeeded = true)
+	UObject* GetDefaultObject(bool bCreateIfNeeded = true) const
 	{
 		if (ClassDefaultObject == nullptr && bCreateIfNeeded)
 		{
-			CreateDefaultObject();
+			const_cast<UClass*>(this)->CreateDefaultObject();
 		}
 
 		return ClassDefaultObject;
@@ -2549,11 +2549,16 @@ public:
 	 */
 	virtual void InitPropertiesFromCustomList(uint8* DataPtr, const uint8* DefaultDataPtr) {}
 
+	/**
+	 * Allows class to provide data to the object initializer that can affect how native class subobjects are created.
+	 */
+	virtual void SetupObjectInitializer(FObjectInitializer& ObjectInitializer) const {}
+
 	/**
 	 * Get the name of the CDO for the this class
 	 * @return The name of the CDO
 	 */
-	FName GetDefaultObjectName();
+	FName GetDefaultObjectName() const;
 
 	/** Returns memory used to store temporary data on an instance, used by blueprints */
 	virtual uint8* GetPersistentUberGraphFrame(UObject* Obj, UFunction* FuncToCheck) const
@@ -2576,7 +2581,7 @@ public:
 	 * @return		the CDO for this class
 	 */
 	template<class T>
-	T* GetDefaultObject()
+	T* GetDefaultObject() const
 	{
 		UObject *Ret = GetDefaultObject();
 		check(Ret->IsA(T::StaticClass()));
@@ -2788,7 +2793,7 @@ public:
 	virtual bool HasProperty(UProperty* InProperty) const;
 
 	/** Finds the object that is used as the parent object when serializing properties, overridden for blueprints */
-	virtual UObject* FindArchetype(UClass* ArchetypeClass, const FName ArchetypeName) const { return nullptr; }
+	virtual UObject* FindArchetype(const UClass* ArchetypeClass, const FName ArchetypeName) const { return nullptr; }
 
 	/** Returns archetype object for CDO */
 	virtual UObject* GetArchetypeForCDO() const;
@@ -2871,12 +2876,15 @@ class COREUOBJECT_API UDynamicClass : public UClass
 
 public:
 
+	typedef void (*DynamicClassInitializerType)	(UDynamicClass*);
+
 	UDynamicClass(const FObjectInitializer& ObjectInitializer = FObjectInitializer::Get());
 	explicit UDynamicClass(const FObjectInitializer& ObjectInitializer, UClass* InSuperClass);
 	UDynamicClass(EStaticConstructor, FName InName, uint32 InSize, uint32 InAlignment, EClassFlags InClassFlags, EClassCastFlags InClassCastFlags,
 		const TCHAR* InClassConfigName, EObjectFlags InFlags, ClassConstructorType InClassConstructor,
 		ClassVTableHelperCtorCallerType InClassVTableHelperCtorCaller,
-		ClassAddReferencedObjectsType InClassAddReferencedObjects);
+		ClassAddReferencedObjectsType InClassAddReferencedObjects,
+		DynamicClassInitializerType InDynamicClassInitializer);
 
 	// UObject interface.
 	static void AddReferencedObjects(UObject* InThis, FReferenceCollector& Collector);
@@ -2884,7 +2892,8 @@ public:
 	// UClass interface
 	virtual UObject* CreateDefaultObject();
 	virtual void PurgeClass(bool bRecompilingOnLoad) override;
-	virtual UObject* FindArchetype(UClass* ArchetypeClass, const FName ArchetypeName) const override;
+	virtual UObject* FindArchetype(const UClass* ArchetypeClass, const FName ArchetypeName) const override;
+	virtual void SetupObjectInitializer(FObjectInitializer& ObjectInitializer) const override;
 
 	/** Find a struct property, called from generated code */
 	UStructProperty* FindStructPropertyChecked(const TCHAR* PropertyName) const;
@@ -2903,8 +2912,13 @@ public:
 	TArray<UObject*> ComponentTemplates;
 	TArray<UObject*> Timelines;
 
+	/** Array of blueprint overrides of component classes in parent classes */
+	TArray<TPair<FName, UClass*>> ComponentClassOverrides;
+
 	/** IAnimClassInterface (UAnimClassData) or null */
 	UObject* AnimClassImplementation;
+
+	DynamicClassInitializerType DynamicClassInitializer;
 };
 
 /**
@@ -2968,7 +2982,8 @@ COREUOBJECT_API void GetPrivateStaticClassBody(
 	UClass::ClassAddReferencedObjectsType InClassAddReferencedObjects,
 	UClass::StaticClassFunctionType InSuperClassFn,
 	UClass::StaticClassFunctionType InWithinClassFn,
-	bool bIsDynamic = false);
+	bool bIsDynamic = false,
+	UDynamicClass::DynamicClassInitializerType InDynamicClassInitializer = nullptr);
 
 /*-----------------------------------------------------------------------------
 	FObjectInstancingGraph.
diff --git a/Engine/Source/Runtime/CoreUObject/Public/UObject/Object.h b/Engine/Source/Runtime/CoreUObject/Public/UObject/Object.h
--- a/Engine/Source/Runtime/CoreUObject/Public/UObject/Object.h
+++ b/Engine/Source/Runtime/CoreUObject/Public/UObject/Object.h
@@ -1120,7 +1120,7 @@ public:
 	 * 
 	 * @return the archetype for this object
 	 */
-	static UObject* GetArchetypeFromRequiredInfo(UClass* Class, UObject* Outer, FName Name, EObjectFlags ObjectFlags);
+	static UObject* GetArchetypeFromRequiredInfo(const UClass* Class, const UObject* Outer, FName Name, EObjectFlags ObjectFlags);
 
 	/**
 	 * Return the template this object is based on. 
diff --git a/Engine/Source/Runtime/CoreUObject/Public/UObject/ObjectMacros.h b/Engine/Source/Runtime/CoreUObject/Public/UObject/ObjectMacros.h
--- a/Engine/Source/Runtime/CoreUObject/Public/UObject/ObjectMacros.h
+++ b/Engine/Source/Runtime/CoreUObject/Public/UObject/ObjectMacros.h
@@ -1665,7 +1665,8 @@ public: \
 			&TClass::AddReferencedObjects, \
 			&TClass::Super::StaticClass, \
 			&TClass::WithinClass::StaticClass, \
-			true \
+			true, \
+			&TClass::__CustomDynamicClassInitialization \
 			); \
 		} \
 		return PrivateStaticClass; \
diff --git a/Engine/Source/Runtime/CoreUObject/Public/UObject/UObjectGlobals.h b/Engine/Source/Runtime/CoreUObject/Public/UObject/UObjectGlobals.h
--- a/Engine/Source/Runtime/CoreUObject/Public/UObject/UObjectGlobals.h
+++ b/Engine/Source/Runtime/CoreUObject/Public/UObject/UObjectGlobals.h
@@ -188,7 +188,7 @@ COREUOBJECT_API bool ResolveName(UObject*& Outer, FString& ObjectsReferenceStrin
 COREUOBJECT_API bool SafeLoadError( UObject* Outer, uint32 LoadFlags, const TCHAR* ErrorMessage);
 
 /** Internal function used to update the suffix to be given to the next newly-created unnamed object. */
-COREUOBJECT_API int32 UpdateSuffixForNextNewObject(UObject* Parent, UClass* Class, TFunctionRef<void(int32&)> IndexMutator);
+COREUOBJECT_API int32 UpdateSuffixForNextNewObject(UObject* Parent, const UClass* Class, TFunctionRef<void(int32&)> IndexMutator);
 
 
 /**
@@ -278,7 +278,7 @@ COREUOBJECT_API UClass* StaticLoadClass(UClass* BaseClass, UObject* InOuter, con
  *
  * @return	A pointer to a fully initialized object of the specified class.
  */
-COREUOBJECT_API UObject* StaticConstructObject_Internal(UClass* Class, UObject* InOuter = (UObject*)GetTransientPackage(), FName Name = NAME_None, EObjectFlags SetFlags = RF_NoFlags, EInternalObjectFlags InternalSetFlags = EInternalObjectFlags::None, UObject* Template = nullptr, bool bCopyTransientsFromClassDefaults = false, struct FObjectInstancingGraph* InstanceGraph = nullptr, bool bAssumeTemplateIsArchetype = false);
+COREUOBJECT_API UObject* StaticConstructObject_Internal(const UClass* Class, UObject* InOuter = (UObject*)GetTransientPackage(), FName Name = NAME_None, EObjectFlags SetFlags = RF_NoFlags, EInternalObjectFlags InternalSetFlags = EInternalObjectFlags::None, UObject* Template = nullptr, bool bCopyTransientsFromClassDefaults = false, struct FObjectInstancingGraph* InstanceGraph = nullptr, bool bAssumeTemplateIsArchetype = false);
 
 /**
  * Creates a copy of SourceObject using the Outer and Name specified, as well as copies of all objects contained by SourceObject.  
@@ -500,7 +500,7 @@ COREUOBJECT_API void IncrementalPurgeGarbage( bool bUseTimeLimit, float TimeLimi
  * @return	name is the form BaseName_##, where ## is the number of objects of this
  *			type that have been created since the last time the class was garbage collected.
  */
-COREUOBJECT_API FName MakeUniqueObjectName( UObject* Outer, UClass* Class, FName BaseName=NAME_None );
+COREUOBJECT_API FName MakeUniqueObjectName( UObject* Outer, const UClass* Class, FName BaseName=NAME_None );
 
 /**
  * Given a display label string, generates an FName slug that is a valid FName for that label.
@@ -640,7 +640,7 @@ COREUOBJECT_API void SnapshotTransactionBuffer(UObject* Object);
  * @param	SetFlags	the ObjectFlags to assign to the new object. some flags can affect the behavior of constructing the object.
  * @return	true if nullptr should be returned; there was a problem reported 
  */
-bool StaticAllocateObjectErrorTests( UClass* Class, UObject* InOuter, FName Name, EObjectFlags SetFlags);
+bool StaticAllocateObjectErrorTests( const UClass* Class, UObject* InOuter, FName Name, EObjectFlags SetFlags);
 
 /**
  * Create a new instance of an object or replace an existing object.  If both an Outer and Name are specified, and there is an object already in memory with the same Class, Outer, and Name, the
@@ -655,7 +655,7 @@ bool StaticAllocateObjectErrorTests( UClass* Class, UObject* InOuter, FName Name
  * @param bOutReusedSubobject	flag indicating if the object is a subobject that has already been created (in which case further initialization is not necessary).
  * @return	a pointer to a fully initialized object of the specified class.
  */
-COREUOBJECT_API UObject* StaticAllocateObject(UClass* Class, UObject* InOuter, FName Name, EObjectFlags SetFlags, EInternalObjectFlags InternalSetFlags = EInternalObjectFlags::None, bool bCanReuseSubobjects = false, bool* bOutReusedSubobject = nullptr);
+COREUOBJECT_API UObject* StaticAllocateObject(const UClass* Class, UObject* InOuter, FName Name, EObjectFlags SetFlags, EInternalObjectFlags InternalSetFlags = EInternalObjectFlags::None, bool bCanReuseSubobjects = false, bool* bOutReusedSubobject = nullptr);
 
 /** @deprecated Use raw pointers or TWeakObjectPtr instead */
 class COREUOBJECT_API FSubobjectPtr
@@ -1176,7 +1176,7 @@ public:
 
 #if DO_CHECK
 /** Called by NewObject to make sure Child is actually a child of Parent */
-COREUOBJECT_API void CheckIsClassChildOf_Internal(UClass* Parent, UClass* Child);
+COREUOBJECT_API void CheckIsClassChildOf_Internal(const UClass* Parent, const UClass* Child);
 #endif
 
 /**
@@ -1194,7 +1194,7 @@ COREUOBJECT_API void CheckIsClassChildOf_Internal(UClass* Parent, UClass* Child)
  */
 template< class T >
 FUNCTION_NON_NULL_RETURN_START
-	T* NewObject(UObject* Outer, UClass* Class, FName Name = NAME_None, EObjectFlags Flags = RF_NoFlags, UObject* Template = nullptr, bool bCopyTransientsFromClassDefaults = false, FObjectInstancingGraph* InInstanceGraph = nullptr)
+	T* NewObject(UObject* Outer, const UClass* Class, FName Name = NAME_None, EObjectFlags Flags = RF_NoFlags, UObject* Template = nullptr, bool bCopyTransientsFromClassDefaults = false, FObjectInstancingGraph* InInstanceGraph = nullptr)
 FUNCTION_NON_NULL_RETURN_END
 {
 	if (Name == NAME_None)
diff --git a/Engine/Source/Runtime/Engine/Classes/Components/ActorComponent.h b/Engine/Source/Runtime/Engine/Classes/Components/ActorComponent.h
--- a/Engine/Source/Runtime/Engine/Classes/Components/ActorComponent.h
+++ b/Engine/Source/Runtime/Engine/Classes/Components/ActorComponent.h
@@ -204,10 +204,6 @@ public:
 	/** If true, we call the virtual InitializeComponent */
 	uint8 bWantsInitializeComponent:1;
 
-	/** If true, we call the virtual BeginPlay */
-	UE_DEPRECATED(4.14, "bWantsBeginPlay was inconsistently enforced and is now unused. BeginPlay will now always be called for Actor Components.")
-	uint8 bWantsBeginPlay:1;
-
 	/** If true, the component will be excluded from non-editor builds */
 	UPROPERTY(EditAnywhere, BlueprintReadOnly, Category = Cooking)
 	uint8 bIsEditorOnly:1;
diff --git a/Engine/Source/Runtime/Engine/Classes/Components/ArrowComponent.h b/Engine/Source/Runtime/Engine/Classes/Components/ArrowComponent.h
--- a/Engine/Source/Runtime/Engine/Classes/Components/ArrowComponent.h
+++ b/Engine/Source/Runtime/Engine/Classes/Components/ArrowComponent.h
@@ -26,17 +26,17 @@ class ENGINE_API UArrowComponent : public UPrimitiveComponent
 	UPROPERTY(EditAnywhere, BlueprintReadWrite, Category=ArrowComponent)
 	float ArrowSize;
 
-	/** Set to limit the screen size of this arrow */
-	UPROPERTY(EditAnywhere, BlueprintReadWrite, Category=ArrowComponent)
-	bool bIsScreenSizeScaled;
-
 	/** The size on screen to limit this arrow to (in screen space) */
-	UPROPERTY(EditAnywhere, BlueprintReadWrite, Category=ArrowComponent)
+	UPROPERTY(EditAnywhere, BlueprintReadWrite, Category = ArrowComponent)
 	float ScreenSize;
 
+	/** Set to limit the screen size of this arrow */
+	UPROPERTY(EditAnywhere, BlueprintReadWrite, Category=ArrowComponent)
+	uint8 bIsScreenSizeScaled : 1;
+
 	/** If true, don't show the arrow when EngineShowFlags.BillboardSprites is disabled. */
 	UPROPERTY(EditAnywhere, BlueprintReadWrite, Category=ArrowComponent)
-	uint32 bTreatAsASprite:1;
+	uint8 bTreatAsASprite : 1;
 
 #if WITH_EDITORONLY_DATA
 	/** Sprite category that the arrow component belongs to, if being treated as a sprite. Value serves as a key into the localization file. */
diff --git a/Engine/Source/Runtime/Engine/Classes/Components/MeshComponent.h b/Engine/Source/Runtime/Engine/Classes/Components/MeshComponent.h
--- a/Engine/Source/Runtime/Engine/Classes/Components/MeshComponent.h
+++ b/Engine/Source/Runtime/Engine/Classes/Components/MeshComponent.h
@@ -144,9 +144,9 @@ protected:
 	TSortedMap<FName, FMaterialParameterCache, FDefaultAllocator, FNameFastLess> MaterialParameterCache;
 
 	UPROPERTY(EditAnywhere, AdvancedDisplay, BlueprintReadOnly, Category = MaterialParameters)
-	bool bEnableMaterialParameterCaching;
+	uint8 bEnableMaterialParameterCaching : 1;
 
 	/** Flag whether or not the cached material parameter indices map is dirty (defaults to true, and is set from SetMaterial/Set(Skeletal)Mesh */
-	bool bCachedMaterialParameterIndicesAreDirty;
+	uint8 bCachedMaterialParameterIndicesAreDirty : 1;
 
 };
diff --git a/Engine/Source/Runtime/Engine/Classes/Components/SplineMeshComponent.h b/Engine/Source/Runtime/Engine/Classes/Components/SplineMeshComponent.h
--- a/Engine/Source/Runtime/Engine/Classes/Components/SplineMeshComponent.h
+++ b/Engine/Source/Runtime/Engine/Classes/Components/SplineMeshComponent.h
@@ -13,6 +13,7 @@
 class FPrimitiveSceneProxy;
 class ULightComponent;
 struct FNavigableGeometryExport;
+class UBodySetup;
 
 UENUM(BlueprintType)
 namespace ESplineMeshAxis
@@ -58,14 +59,14 @@ struct FSplineMeshParams
 	UPROPERTY(EditAnywhere, Category=SplineMesh)
 	FVector EndPos;
 
-	/** End tangent of spline, in component space. */
-	UPROPERTY(EditAnywhere, Category=SplineMesh)
-	FVector EndTangent;
-
 	/** X and Y scale applied to mesh at end of spline. */
 	UPROPERTY(EditAnywhere, Category=SplineMesh, AdvancedDisplay)
 	FVector2D EndScale;
 
+	/** End tangent of spline, in component space. */
+	UPROPERTY(EditAnywhere, Category = SplineMesh)
+	FVector EndTangent;
+
 	/** Roll around spline applied at end. */
 	UPROPERTY(EditAnywhere, Category=SplineMesh, AdvancedDisplay)
 	float EndRoll;
@@ -82,8 +83,8 @@ struct FSplineMeshParams
 		, StartRoll(0)
 		, StartOffset(ForceInit)
 		, EndPos(ForceInit)
-		, EndTangent(ForceInit)
 		, EndScale(ForceInit)
+		, EndTangent(ForceInit)
 		, EndRoll(0)
 		, EndOffset(ForceInit)
 	{
@@ -108,43 +109,43 @@ class ENGINE_API USplineMeshComponent : public UStaticMeshComponent, public IInt
 	UPROPERTY(EditAnywhere, Category=SplineMesh)
 	FVector SplineUpDir;
 
+	/** Minimum coordinate along the spline forward axis which corresponds to start of spline. If set to 0.0, will use bounding box to determine bounds */
+	UPROPERTY(EditAnywhere, Category = SplineMesh, AdvancedDisplay)
+	float SplineBoundaryMin;
+
+	// Used to automatically trigger rebuild of collision data
+	UPROPERTY()
+	FGuid CachedMeshBodySetupGuid;
+
+	// Physics data.
+	UPROPERTY()
+	UBodySetup* BodySetup;
+
+	/** Maximum coordinate along the spline forward axis which corresponds to end of spline. If set to 0.0, will use bounding box to determine bounds */
+	UPROPERTY(EditAnywhere, Category = SplineMesh, AdvancedDisplay)
+	float SplineBoundaryMax;
+
 	/** If true, spline keys may be edited per instance in the level viewport. Otherwise, the spline should be initialized in the construction script. */
 	UPROPERTY(EditDefaultsOnly, Category = Spline)
-	uint32 bAllowSplineEditingPerInstance:1;
+	uint8 bAllowSplineEditingPerInstance:1;
 
 	/** If true, will use smooth interpolation (ease in/out) for Scale, Roll, and Offset along this section of spline. If false, uses linear */
 	UPROPERTY(EditAnywhere, Category=SplineMesh, AdvancedDisplay)
-	uint32 bSmoothInterpRollScale:1;
+	uint8 bSmoothInterpRollScale:1;
+
+	// Indicates that the mesh needs updating
+	UPROPERTY(transient)
+	uint8 bMeshDirty : 1;
 
 	/** Chooses the forward axis for the spline mesh orientation */
 	UPROPERTY(EditAnywhere, Category=SplineMesh)
 	TEnumAsByte<ESplineMeshAxis::Type> ForwardAxis;
 
-	/** Minimum coordinate along the spline forward axis which corresponds to start of spline. If set to 0.0, will use bounding box to determine bounds */
-	UPROPERTY(EditAnywhere, Category = SplineMesh, AdvancedDisplay)
-	float SplineBoundaryMin;
-
-	/** Maximum coordinate along the spline forward axis which corresponds to end of spline. If set to 0.0, will use bounding box to determine bounds */
-	UPROPERTY(EditAnywhere, Category = SplineMesh, AdvancedDisplay)
-	float SplineBoundaryMax;
-	
-	// Physics data.
-	UPROPERTY()
-	class UBodySetup* BodySetup;
-
-	// Used to automatically trigger rebuild of collision data
-	UPROPERTY()
-	FGuid CachedMeshBodySetupGuid;
-
 #if WITH_EDITORONLY_DATA
 	UPROPERTY(transient)
-	uint32 bSelected:1;
+	uint8 bSelected:1;
 #endif
 
-	// Indicates that the mesh needs updating
-	UPROPERTY(transient)
-	uint32 bMeshDirty:1;
-
 	//Begin UObject Interface
 	virtual void Serialize(FArchive& Ar) override;
 #if WITH_EDITOR
diff --git a/Engine/Source/Runtime/Engine/Classes/Distributions/Distribution.h b/Engine/Source/Runtime/Engine/Classes/Distributions/Distribution.h
--- a/Engine/Source/Runtime/Engine/Classes/Distributions/Distribution.h
+++ b/Engine/Source/Runtime/Engine/Classes/Distributions/Distribution.h
@@ -26,25 +26,25 @@ USTRUCT(noexport)
 struct FDistributionLookupTable
 {
 	UPROPERTY()
-	uint8 Op;
+	float TimeScale;
 
 	UPROPERTY()
-	uint8 EntryCount;
+	float TimeBias;
 
 	UPROPERTY()
-	uint8 EntryStride;
+	TArray<float> Values;
 
 	UPROPERTY()
-	uint8 SubEntryStride;
+	uint8 Op;
 
 	UPROPERTY()
-	float TimeScale;
+	uint8 EntryCount;
 
 	UPROPERTY()
-	float TimeBias;
+	uint8 EntryStride;
 
 	UPROPERTY()
-	TArray<float> Values;
+	uint8 SubEntryStride;
 
 	UPROPERTY()
 	uint8 LockFlag;
diff --git a/Engine/Source/Runtime/Engine/Classes/Distributions/DistributionFloat.h b/Engine/Source/Runtime/Engine/Classes/Distributions/DistributionFloat.h
--- a/Engine/Source/Runtime/Engine/Classes/Distributions/DistributionFloat.h
+++ b/Engine/Source/Runtime/Engine/Classes/Distributions/DistributionFloat.h
@@ -105,14 +105,14 @@ class ENGINE_API UDistributionFloat : public UDistribution
 
 	/** Can this variable be baked out to a FRawDistribution? Should be true 99% of the time*/
 	UPROPERTY(EditAnywhere, Category=Baked)
-	uint32 bCanBeBaked:1;
+	uint8 bCanBeBaked:1;
 
 	/** Set internally when the distribution is updated so that that FRawDistribution can know to update itself*/
-	uint32 bIsDirty:1;
+	uint8 bIsDirty:1;
 	
 protected:
 	UPROPERTY()
-	uint32 bBakedDataSuccesfully:1;	//It's possible that even though we want to bake we are not able to because of content or code.
+	uint8 bBakedDataSuccesfully:1;	//It's possible that even though we want to bake we are not able to because of content or code.
 public:
 
 	/** Script-accessible way to query a float distribution */
diff --git a/Engine/Source/Runtime/Engine/Classes/Distributions/DistributionVector.h b/Engine/Source/Runtime/Engine/Classes/Distributions/DistributionVector.h
--- a/Engine/Source/Runtime/Engine/Classes/Distributions/DistributionVector.h
+++ b/Engine/Source/Runtime/Engine/Classes/Distributions/DistributionVector.h
@@ -150,14 +150,14 @@ class ENGINE_API UDistributionVector : public UDistribution
 
 	/** Can this variable be baked out to a FRawDistribution? Should be true 99% of the time*/
 	UPROPERTY(EditAnywhere, Category=Baked)
-	uint32 bCanBeBaked:1;
+	uint8 bCanBeBaked:1;
 
 	/** Set internally when the distribution is updated so that that FRawDistribution can know to update itself*/
 	UPROPERTY()
-	uint32 bIsDirty:1;
+	uint8 bIsDirty:1;
 protected:
 	UPROPERTY()
-	uint32 bBakedDataSuccesfully:1;	//It's possible that even though we want to bake we are not able to because of content or code.
+	uint8 bBakedDataSuccesfully:1;	//It's possible that even though we want to bake we are not able to because of content or code.
 public:
 
 	/** Script-accessible way to query a FVector distribution */
diff --git a/Engine/Source/Runtime/Engine/Classes/Engine/Blueprint.h b/Engine/Source/Runtime/Engine/Classes/Engine/Blueprint.h
--- a/Engine/Source/Runtime/Engine/Classes/Engine/Blueprint.h
+++ b/Engine/Source/Runtime/Engine/Classes/Engine/Blueprint.h
@@ -11,6 +11,7 @@
 #include "Engine/EngineTypes.h"
 #include "EdGraph/EdGraphPin.h"
 #include "Engine/BlueprintCore.h"
+#include "Engine/BlueprintGeneratedClass.h"
 #include "UObject/SoftObjectPath.h"
 #include "Blueprint/BlueprintSupport.h"
 #include "Blueprint.generated.h"
@@ -594,6 +595,10 @@ public:
 	UPROPERTY()
 	TArray<class UTimelineTemplate*> Timelines;
 
+	/** Array of blueprint overrides of component classes in parent classes */
+	UPROPERTY()
+	TArray<FBPComponentClassOverride> ComponentClassOverrides;
+
 	/** Stores data to override (in children classes) components (created by SCS) from parent classes */
 	UPROPERTY()
 	class UInheritableComponentHandler* InheritableComponentHandler;
diff --git a/Engine/Source/Runtime/Engine/Classes/Engine/BlueprintGeneratedClass.h b/Engine/Source/Runtime/Engine/Classes/Engine/BlueprintGeneratedClass.h
--- a/Engine/Source/Runtime/Engine/Classes/Engine/BlueprintGeneratedClass.h
+++ b/Engine/Source/Runtime/Engine/Classes/Engine/BlueprintGeneratedClass.h
@@ -539,13 +539,6 @@ struct ENGINE_API FBlueprintCookedComponentInstancingData
 {
 	GENERATED_USTRUCT_BODY()
 
-	/** Flag indicating whether or not this contains valid cooked data. Note that an empty changed property list can also be a valid template data context. */
-	UPROPERTY()
-	bool bHasValidCookedData;
-
-	/** Source template object flags (recorded at load time and used for instancing). */
-	EObjectFlags ComponentTemplateFlags;
-
 	/** List of property info records with values that differ between the template and the component class CDO. This list will be generated at cook time. */
 	UPROPERTY()
 	TArray<struct FBlueprintComponentChangedPropertyInfo> ChangedPropertyList;
@@ -556,6 +549,13 @@ struct ENGINE_API FBlueprintCookedComponentInstancingData
 	/** Source template object class (recorded at load time and used for instancing). */
 	UClass* ComponentTemplateClass;
 
+	/** Source template object flags (recorded at load time and used for instancing). */
+	TEnumAsByte<EObjectFlags> ComponentTemplateFlags;
+
+	/** Flag indicating whether or not this contains valid cooked data. Note that an empty changed property list can also be a valid template data context. */
+	UPROPERTY()
+	bool bHasValidCookedData;
+
 	/** Default constructor. */
 	FBlueprintCookedComponentInstancingData()
 	{
@@ -591,6 +591,37 @@ private:
 	TArray<uint8> CachedPropertyData;
 };
 
+/** Utility struct to store class overrides for components. */
+USTRUCT()
+struct FBPComponentClassOverride
+{
+	GENERATED_BODY()
+
+	/** The component name an override is being specified for. */
+	UPROPERTY()
+	FName ComponentName;
+
+	/** The class to use when constructing the component. */
+	UPROPERTY()
+	UClass* ComponentClass;
+
+	FBPComponentClassOverride()
+		: ComponentClass(nullptr)
+	{
+	}
+
+	FBPComponentClassOverride(FName InComponentName, UClass* InComponentClass)
+		: ComponentName(InComponentName)
+		, ComponentClass(InComponentClass)
+	{
+	}
+
+	bool operator==(const FName OtherComponentName) const
+	{
+		return (ComponentName == OtherComponentName);
+	}
+};
+
 UCLASS()
 class ENGINE_API UBlueprintGeneratedClass : public UClass
 {
@@ -626,6 +657,10 @@ public:
 	UPROPERTY()
 	TArray<class UTimelineTemplate*> Timelines;
 
+	/** Array of blueprint overrides of component classes in parent classes */
+	UPROPERTY()
+	TArray<FBPComponentClassOverride> ComponentClassOverrides;
+
 	/** 'Simple' construction script - graph of components to instance */
 	UPROPERTY()
 	class USimpleConstructionScript* SimpleConstructionScript;
@@ -712,9 +747,10 @@ public:
 	virtual void PurgeClass(bool bRecompilingOnLoad) override;
 	virtual void Bind() override;
 	virtual void GetDefaultObjectPreloadDependencies(TArray<UObject*>& OutDeps) override;
-	virtual UObject* FindArchetype(UClass* ArchetypeClass, const FName ArchetypeName) const override;
+	virtual UObject* FindArchetype(const UClass* ArchetypeClass, const FName ArchetypeName) const override;
 
 	virtual void InitPropertiesFromCustomList(uint8* DataPtr, const uint8* DefaultDataPtr) override;
+	virtual void SetupObjectInitializer(FObjectInitializer& ObjectInitializer) const override;
 
 protected:
 
diff --git a/Engine/Source/Runtime/Engine/Classes/Engine/GameViewportClient.h b/Engine/Source/Runtime/Engine/Classes/Engine/GameViewportClient.h
--- a/Engine/Source/Runtime/Engine/Classes/Engine/GameViewportClient.h
+++ b/Engine/Source/Runtime/Engine/Classes/Engine/GameViewportClient.h
@@ -368,14 +368,28 @@ public:
 	virtual void LayoutPlayers();
 
 	/** Allows game code to disable splitscreen (useful when in menus) */
-	void SetDisableSplitscreenOverride( const bool bDisabled );
+	void SetForceDisableSplitscreen(const bool bDisabled);
 
 	/** Determines whether splitscreen is forced to be turned off */
-	bool GetDisableSplitscreenOverride() const
+	bool IsSplitscreenForceDisabled() const
 	{
 		return bDisableSplitScreenOverride;
 	}
 
+	/** Allows game code to disable splitscreen (useful when in menus) */
+	UE_DEPRECATED(4.24, "SetDisableSplitscreenOverride is deprecated. Please call UGameViewportClient::SetForceDisableSplitscreen(bDisabled) instead.")
+	void SetDisableSplitscreenOverride( const bool bDisabled )
+	{
+		SetForceDisableSplitscreen(bDisabled);
+	}
+
+	/** Determines whether splitscreen is forced to be turned off */
+	UE_DEPRECATED(4.24, "GetDisableSplitscreenOverride is deprecated. Please call UGameViewportClient::IsSplitscreenForceDisabled() instead.")
+	bool GetDisableSplitscreenOverride() const
+	{
+		return IsSplitscreenForceDisabled();
+	}
+
 	/** called before rending subtitles to allow the game viewport to determine the size of the subtitle area
 	 * @param Min top left bounds of subtitle region (0 to 1)
 	 * @param Max bottom right bounds of subtitle region (0 to 1)
diff --git a/Engine/Source/Runtime/Engine/Classes/GameFramework/Actor.h b/Engine/Source/Runtime/Engine/Classes/GameFramework/Actor.h
--- a/Engine/Source/Runtime/Engine/Classes/GameFramework/Actor.h
+++ b/Engine/Source/Runtime/Engine/Classes/GameFramework/Actor.h
@@ -870,13 +870,13 @@ public:
 	 * @return The transform that transforms from actor space to world space.
 	 */
 	UFUNCTION(BlueprintCallable, meta=(DisplayName = "GetActorTransform", ScriptName = "GetActorTransform"), Category="Utilities|Transformation")
-	FTransform GetTransform() const
+	const FTransform& GetTransform() const
 	{
 		return ActorToWorld();
 	}
 
 	/** Get the local-to-world transform of the RootComponent. Identical to GetTransform(). */
-	FORCEINLINE FTransform ActorToWorld() const
+	FORCEINLINE const FTransform& ActorToWorld() const
 	{
 		return (RootComponent ? RootComponent->GetComponentTransform() : FTransform::Identity);
 	}
@@ -1710,7 +1710,7 @@ public:
 	bool SetRootComponent(USceneComponent* NewRootComponent);
 
 	/** Returns the transform of the RootComponent of this Actor*/ 
-	FORCEINLINE FTransform GetActorTransform() const
+	FORCEINLINE const FTransform& GetActorTransform() const
 	{
 		return TemplateGetActorTransform(RootComponent);
 	}
@@ -3085,9 +3085,9 @@ private:
 	// These are templates for no other reason than to delay compilation until USceneComponent is defined.
 
 	template<class T>
-	static FORCEINLINE FTransform TemplateGetActorTransform(const T* RootComponent)
+	static FORCEINLINE const FTransform& TemplateGetActorTransform(const T* RootComponent)
 	{
-		return (RootComponent != nullptr) ? RootComponent->GetComponentTransform() : FTransform();
+		return (RootComponent != nullptr) ? RootComponent->GetComponentTransform() : FTransform::Identity;
 	}
 
 	template<class T>
diff --git a/Engine/Source/Runtime/Engine/Classes/GameFramework/Pawn.h b/Engine/Source/Runtime/Engine/Classes/GameFramework/Pawn.h
--- a/Engine/Source/Runtime/Engine/Classes/GameFramework/Pawn.h
+++ b/Engine/Source/Runtime/Engine/Classes/GameFramework/Pawn.h
@@ -193,9 +193,14 @@ public:
 	virtual float GetDefaultHalfHeight() const;
 
 	/** See if this actor is currently being controlled */
+	UE_DEPRECATED(4.24, "IsControlled is deprecated. To check if this pawn is controlled by anything, then call IsPawnControlled. To check if this pawn is controlled only by the player then call IsPlayerControlled")
 	UFUNCTION(BlueprintCallable, Category=Pawn)
 	bool IsControlled() const;
 
+	/** Check if this actor is currently being controlled at all (the actor has a valid Controller) */
+	UFUNCTION(BlueprintCallable, Category = Pawn)
+	bool IsPawnControlled() const;
+
 	/** Returns controller for this actor. */
 	UFUNCTION(BlueprintCallable, Category=Pawn)
 	AController* GetController() const;
diff --git a/Engine/Source/Runtime/Engine/Classes/Kismet/GameplayStatics.h b/Engine/Source/Runtime/Engine/Classes/Kismet/GameplayStatics.h
--- a/Engine/Source/Runtime/Engine/Classes/Kismet/GameplayStatics.h
+++ b/Engine/Source/Runtime/Engine/Classes/Kismet/GameplayStatics.h
@@ -267,6 +267,20 @@ class ENGINE_API UGameplayStatics : public UBlueprintFunctionLibrary
 	UFUNCTION(BlueprintPure, Category="Game", meta=(WorldContext="WorldContextObject") )
 	static bool IsGamePaused(const UObject* WorldContextObject);
 
+	/**
+	 * Enables split screen
+	 * @param	bDisable		Whether the viewport should split screen between local players or not
+	 */
+	UFUNCTION(BlueprintCallable, Category = "Viewport", meta = (WorldContext = "WorldContextObject"))
+	static void SetForceDisableSplitscreen(const UObject* WorldContextObject, bool bDisable);
+
+	/**
+	 * Returns the split screen state
+	 * @return	Whether the game viewport is split screen or not
+	 */
+	UFUNCTION(BlueprintPure, Category = "Viewport", meta = (WorldContext = "WorldContextObject"))
+	static bool IsSplitscreenForceDisabled(const UObject* WorldContextObject);
+
 	/**
 	 * Enabled rendering of the world
 	 * @param	bEnable		Whether the world should be rendered or not
diff --git a/Engine/Source/Runtime/Engine/Classes/Kismet/KismetMathLibrary.h b/Engine/Source/Runtime/Engine/Classes/Kismet/KismetMathLibrary.h
--- a/Engine/Source/Runtime/Engine/Classes/Kismet/KismetMathLibrary.h
+++ b/Engine/Source/Runtime/Engine/Classes/Kismet/KismetMathLibrary.h
@@ -3265,11 +3265,11 @@ class ENGINE_API UKismetMathLibrary : public UBlueprintFunctionLibrary
 
 	/**
 	 * Tries to reach Target based on distance from Current position, giving a nice smooth feeling when tracking a position.
-	 *
+	 * 
 	 * @param		Current			Actual position
 	 * @param		Target			Target position
 	 * @param		DeltaTime		Time since last tick
-	 * @param		InterpSpeed		Interpolation speed
+	 * @param		InterpSpeed		Interpolation speed, if the speed given is 0, then jump to the target.
 	 * @return		New interpolated position
 	 */
 	UFUNCTION(BlueprintPure, Category="Math|Interpolation")
@@ -3293,7 +3293,7 @@ class ENGINE_API UKismetMathLibrary : public UBlueprintFunctionLibrary
 	 * @param		Current			Actual rotation
 	 * @param		Target			Target rotation
 	 * @param		DeltaTime		Time since last tick
-	 * @param		InterpSpeed		Interpolation speed
+	 * @param		InterpSpeed		Interpolation speed, if the speed given is 0, then jump to the target.
 	 * @return		New interpolated position
 	 */
 	UFUNCTION(BlueprintPure, Category="Math|Interpolation", meta=(Keywords="rotation rotate"))
diff --git a/Engine/Source/Runtime/Engine/Classes/Materials/Material.h b/Engine/Source/Runtime/Engine/Classes/Materials/Material.h
--- a/Engine/Source/Runtime/Engine/Classes/Materials/Material.h
+++ b/Engine/Source/Runtime/Engine/Classes/Materials/Material.h
@@ -199,7 +199,7 @@ struct FColorMaterialInput : public FMaterialInput
 {
 #if WITH_EDITORONLY_DATA
 	UPROPERTY()
-	uint32 UseConstant:1;
+	uint32 UseConstant : 1;
 
 	UPROPERTY()
 	FColor Constant;
@@ -213,7 +213,7 @@ struct FScalarMaterialInput : public FMaterialInput
 {
 #if WITH_EDITORONLY_DATA
 	UPROPERTY()
-	uint32 UseConstant:1;
+	uint32 UseConstant : 1;
 
 	UPROPERTY()
 	float Constant;
@@ -235,7 +235,7 @@ struct FVectorMaterialInput : public FMaterialInput
 {
 #if WITH_EDITORONLY_DATA
 	UPROPERTY()
-	uint32 UseConstant:1;
+	uint32 UseConstant : 1;
 
 	UPROPERTY()
 	FVector Constant;
@@ -249,7 +249,7 @@ struct FVector2MaterialInput : public FMaterialInput
 {
 #if WITH_EDITORONLY_DATA
 	UPROPERTY()
-	uint32 UseConstant:1;
+	uint32 UseConstant : 1;
 
 	UPROPERTY()
 	float ConstantX;
@@ -411,6 +411,15 @@ private:
 	UPROPERTY(EditAnywhere, Category=Material, AssetRegistrySearchable)
 	TEnumAsByte<enum EMaterialShadingModel> ShadingModel; 
 
+public:
+	/**
+	* If true, translucent materials will cast dynamic shadows according to their opacity.
+	* OpacityMaskClipValue is used as the threshold value.
+	*/
+	UPROPERTY(EditAnywhere, Category = Material, AdvancedDisplay)
+	uint8 bCastDynamicShadowAsMasked : 1;
+
+private:
 	UPROPERTY(AssetRegistrySearchable)
 	FMaterialShadingModelField ShadingModels;
 
@@ -426,13 +435,6 @@ public:
 	UPROPERTY(EditAnywhere, Category=Material, AdvancedDisplay)
 	float OpacityMaskClipValue;
 
-	/**
-	* If true, translucent materials will cast dynamic shadows according to their opacity.
-	* OpacityMaskClipValue is used as the threshold value.
-	*/
-	UPROPERTY(EditAnywhere, Category=Material, AdvancedDisplay)
-	uint32 bCastDynamicShadowAsMasked:1;
-
 	/** Adds to world position in the vertex shader. */
 	UPROPERTY()
 	FVectorMaterialInput WorldPositionOffset;
@@ -490,51 +492,51 @@ public:
 
 	/** Indicates that the material should be rendered in the SeparateTranslucency Pass (not affected by DOF, requires bAllowSeparateTranslucency to be set in .ini). */
 	UPROPERTY(EditAnywhere, Category=Translucency, meta=(DisplayName = "Render After DOF"), AdvancedDisplay)
-	uint32 bEnableSeparateTranslucency:1;
-
-	/** Indicates that the translucent material should not be affected by bloom or DOF. (Note: Depth testing is not available) */
-	UPROPERTY(EditAnywhere, Category = Translucency, meta = (DisplayName = "Mobile Separate Translucency"), AdvancedDisplay)
-	uint32 bEnableMobileSeparateTranslucency : 1;
+	uint8 bEnableSeparateTranslucency : 1;
 
 	/**
 	 * Indicates that the material should be rendered using responsive anti-aliasing. Improves sharpness of small moving particles such as sparks.
 	 * Only use for small moving features because it will cause aliasing of the background.
 	 */
 	UPROPERTY(EditAnywhere, Category=Translucency, meta=(DisplayName = "Responsive AA"), AdvancedDisplay)
-	uint32 bEnableResponsiveAA:1;
+	uint8 bEnableResponsiveAA : 1;
 
 	/** SSR on translucency */
 	UPROPERTY(EditAnywhere, Category=Translucency, meta=(DisplayName = "Screen Space Reflections"))
-	uint32 bScreenSpaceReflections:1;
+	uint8 bScreenSpaceReflections : 1;
 
 	/** Contact shadows on translucency */
 	UPROPERTY(EditAnywhere, Category = Translucency, meta = (DisplayName = "Contact Shadows"))
-	uint32 bContactShadows : 1;
+	uint8 bContactShadows : 1;
 
 	/** Indicates that the material should be rendered without backface culling and the normal should be flipped for backfaces. */
 	UPROPERTY(EditAnywhere, Category=Material)
-	uint32 TwoSided:1;
+	uint8 TwoSided : 1;
 
 	/** Whether meshes rendered with the material should support dithered LOD transitions. */
 	UPROPERTY(EditAnywhere, Category = Material, AdvancedDisplay, meta = (DisplayName = "Dithered LOD Transition"))
-	uint32 DitheredLODTransition : 1;
+	uint8 DitheredLODTransition : 1;
 
 	/** Dither opacity mask. When combined with Temporal AA this can be used as a form of limited translucency which supports all lighting features. */
 	UPROPERTY(EditAnywhere, Category=Material, AdvancedDisplay)
-	uint32 DitherOpacityMask:1;
+	uint8 DitherOpacityMask : 1;
 
 	/** Whether the material should allow outputting negative emissive color values.  Only allowed on unlit materials. */
 	UPROPERTY(EditAnywhere, Category=Material, AdvancedDisplay)
-	uint32 bAllowNegativeEmissiveColor:1;
-
-	/** Number of customized UV inputs to display.  Unconnected customized UV inputs will just pass through the vertex UVs. */
-	UPROPERTY(EditAnywhere, Category=Material, AdvancedDisplay, meta=(ClampMin=0))
-	int32 NumCustomizedUVs;
+	uint8 bAllowNegativeEmissiveColor : 1;
 
 	/** Sets the lighting mode that will be used on this material if it is translucent. */
 	UPROPERTY(EditAnywhere, AssetRegistrySearchable, Category=Translucency, meta=(DisplayName = "Lighting Mode"))
 	TEnumAsByte<enum ETranslucencyLightingMode> TranslucencyLightingMode;
 
+	/** Indicates that the translucent material should not be affected by bloom or DOF. (Note: Depth testing is not available) */
+	UPROPERTY(EditAnywhere, Category = Translucency, meta = (DisplayName = "Mobile Separate Translucency"), AdvancedDisplay)
+	uint8 bEnableMobileSeparateTranslucency : 1;
+
+	/** Number of customized UV inputs to display.  Unconnected customized UV inputs will just pass through the vertex UVs. */
+	UPROPERTY(EditAnywhere, Category = Material, AdvancedDisplay, meta = (ClampMin = 0))
+	int32 NumCustomizedUVs;
+
 	/** 
 	 * Useful for artificially increasing the influence of the normal on the lighting result for translucency.
 	 * A value larger than 1 increases the influence of the normal, a value smaller than 1 makes the lighting more ambient.
@@ -542,10 +544,6 @@ public:
 	UPROPERTY(EditAnywhere, Category=Translucency, meta=(DisplayName = "Directional Lighting Intensity"))
 	float TranslucencyDirectionalLightingIntensity;
 
-	/** Allows a translucent material to be used with custom depth writing by compiling additional shaders. */
-	UPROPERTY(EditAnywhere, Category=Translucency, AdvancedDisplay, meta=(DisplayName = "Allow Custom Depth Writes"))
-	uint32 AllowTranslucentCustomDepthWrites:1;
-
 	/** Scale used to make translucent shadows more or less opaque than the material's actual opacity. */
 	UPROPERTY(EditAnywhere, Category=TranslucencySelfShadowing, meta=(DisplayName = "Shadow Density Scale"))
 	float TranslucentShadowDensityScale;
@@ -586,34 +584,34 @@ public:
 
 	/** Whether to draw on top of opaque pixels even if behind them. This only has meaning for translucency. */
 	UPROPERTY(EditAnywhere, Category=Translucency, AdvancedDisplay)
-	uint32 bDisableDepthTest:1;
+	uint8 bDisableDepthTest : 1;
 
 	/** Whether the transluency pass should write its alpha, and only the alpha, into the framebuffer */
 	UPROPERTY(EditAnywhere, Category = Translucency, AdvancedDisplay)
-	uint32 bWriteOnlyAlpha : 1;
+	uint8 bWriteOnlyAlpha : 1;
 
 	/** Whether to generate spherical normals for particles that use this material. */
 	UPROPERTY(EditAnywhere, Category=Material, AdvancedDisplay)
-	uint32 bGenerateSphericalParticleNormals:1;
+	uint8 bGenerateSphericalParticleNormals : 1;
 
 	/**
 	 * Whether the material takes a tangent space normal or a world space normal as input.
 	 * (TangentSpace requires extra instructions but is often more convenient).
 	 */
 	UPROPERTY(EditAnywhere, Category=Material, AdvancedDisplay)
-	uint32 bTangentSpaceNormal:1;
+	uint8 bTangentSpaceNormal : 1;
 
 	/**
 	 * If enabled, the material's emissive colour is injected into the LightPropagationVolume
 	 */
 	UPROPERTY(EditAnywhere, BlueprintReadWrite, Category=Material, meta=(DisplayName = "Emissive (Dynamic Area Light)"), AdvancedDisplay)
-	uint32 bUseEmissiveForDynamicAreaLighting : 1;
+	uint8 bUseEmissiveForDynamicAreaLighting : 1;
 
 	/**
 	 * If enabled, the material's opacity defines how much GI is blocked when using the LightPropagationVolume feature
 	 */
 	UPROPERTY(EditAnywhere, BlueprintReadWrite, Category=Material, meta=(DisplayName = "Block Global Illumination"), AdvancedDisplay)
-	uint32 bBlockGI : 1;
+	uint8 bBlockGI : 1;
 
 	/** 
 	 * This is a special usage flag that allows a material to be assignable to any primitive type.
@@ -622,42 +620,42 @@ public:
 	 * This flag should only be enabled when absolutely necessary, and is purposefully not exposed to the UI to prevent abuse.
 	 */
 	UPROPERTY(duplicatetransient)
-	uint32 bUsedAsSpecialEngineMaterial:1;
+	uint8 bUsedAsSpecialEngineMaterial : 1;
 
 	/** 
 	 * Indicates that the material and its instances can be used with skeletal meshes.  
 	 * This will result in the shaders required to support skeletal meshes being compiled which will increase shader compile time and memory usage.
 	 */
 	UPROPERTY(EditAnywhere, BlueprintReadOnly, Category=Usage)
-	uint32 bUsedWithSkeletalMesh:1;
+	uint8 bUsedWithSkeletalMesh : 1;
 
 	/** 
 	 * Indicates that the material and its instances can be used with editor compositing  
 	 * This will result in the shaders required to support editor compositing being compiled which will increase shader compile time and memory usage.
 	 */
 	UPROPERTY(EditAnywhere, BlueprintReadOnly, Category=Usage)
-	uint32 bUsedWithEditorCompositing:1;
+	uint8 bUsedWithEditorCompositing : 1;
 
 	/** 
 	 * Indicates that the material and its instances can be used with particle sprites 
 	 * This will result in the shaders required to support particle sprites being compiled which will increase shader compile time and memory usage.
 	 */
 	UPROPERTY(EditAnywhere, BlueprintReadOnly, Category=Usage)
-	uint32 bUsedWithParticleSprites:1;
+	uint8 bUsedWithParticleSprites : 1;
 
 	/** 
 	 * Indicates that the material and its instances can be used with beam trails
 	 * This will result in the shaders required to support beam trails being compiled which will increase shader compile time and memory usage.
 	 */
 	UPROPERTY(EditAnywhere, BlueprintReadOnly, Category=Usage)
-	uint32 bUsedWithBeamTrails:1;
+	uint8 bUsedWithBeamTrails : 1;
 
 	/** 
 	 * Indicates that the material and its instances can be used with mesh particles
 	 * This will result in the shaders required to support mesh particles being compiled which will increase shader compile time and memory usage.
 	 */
 	UPROPERTY(EditAnywhere, BlueprintReadOnly, Category=Usage)
-	uint32 bUsedWithMeshParticles:1;
+	uint8 bUsedWithMeshParticles : 1;
 
 
 	/**
@@ -665,65 +663,65 @@ public:
 	* This will result in the shaders required to support Niagara sprites being compiled which will increase shader compile time and memory usage.
 	*/
 	UPROPERTY(EditAnywhere, BlueprintReadOnly, Category = Usage)
-	uint32 bUsedWithNiagaraSprites : 1;
+	uint8 bUsedWithNiagaraSprites : 1;
 
 	UPROPERTY(EditAnywhere, BlueprintReadOnly, Category = Usage)
-	uint32 bUsedWithNiagaraRibbons : 1;
+	uint8 bUsedWithNiagaraRibbons : 1;
 
 	UPROPERTY(EditAnywhere, BlueprintReadOnly, Category = Usage)
-	uint32 bUsedWithNiagaraMeshParticles : 1;
+	uint8 bUsedWithNiagaraMeshParticles : 1;
 
 	UPROPERTY(EditAnywhere, BlueprintReadOnly, Category = Usage)
-	uint32 bUsedWithGeometryCache : 1;
+	uint8 bUsedWithGeometryCache : 1;
 
 	/** 
 	 * Indicates that the material and its instances can be used with static lighting
 	 * This will result in the shaders required to support static lighting being compiled which will increase shader compile time and memory usage.
 	 */
 	UPROPERTY(EditAnywhere, BlueprintReadOnly, Category=Usage)
-	uint32 bUsedWithStaticLighting:1;
+	uint8 bUsedWithStaticLighting : 1;
 
 	/** 
 	 * Indicates that the material and its instances can be used with morph targets
 	 * This will result in the shaders required to support morph targets being compiled which will increase shader compile time and memory usage.
 	 */
 	UPROPERTY(EditAnywhere, BlueprintReadOnly, Category=Usage)
-	uint32 bUsedWithMorphTargets:1;
+	uint8 bUsedWithMorphTargets : 1;
 
 	/** 
 	 * Indicates that the material and its instances can be used with spline meshes
 	 * This will result in the shaders required to support spline meshes being compiled which will increase shader compile time and memory usage.
 	 */
 	UPROPERTY(EditAnywhere, BlueprintReadOnly, Category=Usage)
-	uint32 bUsedWithSplineMeshes:1;
+	uint8 bUsedWithSplineMeshes : 1;
 
 	/** 
 	 * Indicates that the material and its instances can be used with instanced static meshes
 	 * This will result in the shaders required to support instanced static meshes being compiled which will increase shader compile time and memory usage.
 	 */
 	UPROPERTY(EditAnywhere, BlueprintReadOnly, Category=Usage)
-	uint32 bUsedWithInstancedStaticMeshes:1;
+	uint8 bUsedWithInstancedStaticMeshes : 1;
 
 	/**
 	 * Indicates that the material and its instances can be use with geometry collections
 	 * This will result in the shaders required to support geometry collections being compiled which will increase shader compile time and memory usage.
 	 */
 	UPROPERTY(EditAnywhere, BlueprintReadOnly, Category = Usage)
-	uint32 bUsedWithGeometryCollections : 1;
+	uint8 bUsedWithGeometryCollections : 1;
 
 	/** 
 	 * Indicates that the material and its instances can be used with distortion
 	 * This will result in the shaders required to support distortion being compiled which will increase shader compile time and memory usage.
 	 */
 	UPROPERTY()
-	uint32 bUsesDistortion:1;
+	uint8 bUsesDistortion : 1;
 
 	/** 
 	 * Indicates that the material and its instances can be used with clothing
 	 * This will result in the shaders required to support clothing being compiled which will increase shader compile time and memory usage.
 	 */
 	UPROPERTY(EditAnywhere, BlueprintReadOnly, Category=Usage)
-	uint32 bUsedWithClothing:1;
+	uint8 bUsedWithClothing : 1;
 
 	/**
 	 * Indicates that the material and its instances can be use with water
@@ -737,18 +735,18 @@ public:
 	 * This will result in the shaders required to support UI materials being compiled which will increase shader compile time and memory usage.
 	 */
 	UPROPERTY()
-	uint32 bUsedWithUI_DEPRECATED:1;
+	uint8 bUsedWithUI_DEPRECATED : 1;
 
 	/** 
 	 * Whether to automatically set usage flags based on what the material is applied to in the editor.
 	 * It can be useful to disable this on a base material with many instances, where adding another usage flag accidentally (eg bUsedWithSkeletalMeshes) can add a lot of shader permutations.
 	 */
 	UPROPERTY(EditAnywhere, BlueprintReadOnly, Category=Usage, AdvancedDisplay)
-	uint32 bAutomaticallySetUsageInEditor:1;
+	uint8 bAutomaticallySetUsageInEditor : 1;
 
 	/* Forces the material to be completely rough. Saves a number of instructions and one sampler. */
 	UPROPERTY(EditAnywhere, BlueprintReadOnly, Category=Material, AdvancedDisplay)
-	uint32 bFullyRough:1;
+	uint8 bFullyRough : 1;
 
 	/** 
 	 *	Forces this material to use full (highp) precision in the pixel shader.
@@ -756,26 +754,26 @@ public:
 	 *	This setting has no effect on older mobile devices that do not support high precision. 
 	 */
 	UPROPERTY(EditAnywhere, BlueprintReadOnly, Category = Mobile)
-	uint32 bUseFullPrecision : 1;
+	uint8 bUseFullPrecision : 1;
 
 	/* Use lightmap directionality and per pixel normals. If disabled, lighting from lightmaps will be flat but cheaper. */
 	UPROPERTY(EditAnywhere, BlueprintReadOnly, Category=Mobile)
-	uint32 bUseLightmapDirectionality:1;
+	uint8 bUseLightmapDirectionality : 1;
 
 	/* 
 	 * Forward renderer: enables multiple parallax-corrected reflection captures that blend together.
 	 * Mobile renderer: blend between nearest 3 reflection captures, but reduces the number of samplers available to the material as two more samplers will be used for reflection cubemaps.
 	 */
 	UPROPERTY(EditAnywhere, BlueprintReadOnly, Category = ForwardShading, meta = (DisplayName = "High Quality Reflections"))
-	uint32 bUseHQForwardReflections : 1;
+	uint8 bUseHQForwardReflections : 1;
 
 	/* Enables planar reflection when using the forward renderer or mobile. Enabling this setting reduces the number of samplers available to the material as one more sampler will be used for the planar reflection. */
 	UPROPERTY(EditAnywhere, BlueprintReadOnly, Category = ForwardShading, meta = (DisplayName = "Planar Reflections"))
-	uint32 bUsePlanarForwardReflections : 1;
+	uint8 bUsePlanarForwardReflections : 1;
 
 	/* Reduce roughness based on screen space normal changes. */
 	UPROPERTY(EditAnywhere, BlueprintReadOnly, Category=Material, AdvancedDisplay)
-	uint32 bNormalCurvatureToRoughness : 1;
+	uint8 bNormalCurvatureToRoughness : 1;
 
 	/** The type of tessellation to apply to this object.  Note D3D11 required for anything except MTM_NoTessellation. */
 	UPROPERTY(EditAnywhere, BlueprintReadOnly, Category=Tessellation)
@@ -783,22 +781,23 @@ public:
 
 	/** Prevents cracks in the surface of the mesh when using tessellation. */
 	UPROPERTY(EditAnywhere, BlueprintReadOnly, Category=Tessellation, meta=(DisplayName = "Crack Free Displacement"))
-	uint32 bEnableCrackFreeDisplacement:1;
+	uint8 bEnableCrackFreeDisplacement : 1;
 
 	/** Enables adaptive tessellation, which tries to maintain a uniform number of pixels per triangle. */
 	UPROPERTY(EditAnywhere, BlueprintReadOnly, Category=Tessellation, meta=(DisplayName = "Adaptive Tessellation"))
-	uint32 bEnableAdaptiveTessellation:1;
+	uint8 bEnableAdaptiveTessellation : 1;
 
-	UPROPERTY(EditAnywhere, Category=Tessellation)
-	float MaxDisplacement;
+	/** Allows a translucent material to be used with custom depth writing by compiling additional shaders. */
+	UPROPERTY(EditAnywhere, Category = Translucency, AdvancedDisplay, meta = (DisplayName = "Allow Custom Depth Writes"))
+	uint8 AllowTranslucentCustomDepthWrites : 1;
 
 	/** Enables a wireframe view of the mesh the material is applied to.  */
 	UPROPERTY(EditAnywhere, Category=Material, AdvancedDisplay)
-	uint32 Wireframe:1;
+	uint8 Wireframe : 1;
 
 	/** Skips outputting velocity during the base pass. */
 	UPROPERTY(EditAnywhere, Category=Material, AdvancedDisplay, meta=(DisplayName = "Support accurate velocities from Vertex Deformation"))
-	uint32 bOutputVelocityOnBasePass:1;
+	uint8 bOutputVelocityOnBasePass : 1;
 
 #if WITH_EDITORONLY_DATA
 	UPROPERTY()
@@ -838,67 +837,59 @@ public:
 
 	/** true if this Material can be assumed Opaque when set to masked. */
 	UPROPERTY()
-	uint32 bCanMaskedBeAssumedOpaque : 1;
+	uint8 bCanMaskedBeAssumedOpaque : 1;
 
 	/** true if Material is masked and uses custom opacity */
 	UPROPERTY()
-	uint32 bIsMasked_DEPRECATED:1;
+	uint8 bIsMasked_DEPRECATED : 1;
 
 	/** true if Material is the preview material used in the material editor. */
 	UPROPERTY(transient, duplicatetransient)
-	uint32 bIsPreviewMaterial:1;
+	uint8 bIsPreviewMaterial : 1;
 
 	/** true if Material is the function preview material used in the material instance editor. */
 	UPROPERTY(transient, duplicatetransient)
-	uint32 bIsFunctionPreviewMaterial:1;
+	uint8 bIsFunctionPreviewMaterial : 1;
 
 	/** when true, the material attributes pin is used instead of the regular pins. */
 	UPROPERTY(EditAnywhere, Category=Material)
-	uint32 bUseMaterialAttributes:1;
+	uint8 bUseMaterialAttributes : 1;
 
 	/** when true, the material casts ray tracing shadows. */
 	UPROPERTY(EditAnywhere, Category = Material)
-	uint32 bCastRayTracedShadows : 1;
+	uint8 bCastRayTracedShadows : 1;
 
 	/** When true, translucent materials are fogged. Defaults to true. */
 	UPROPERTY(EditAnywhere, Category=Translucency, meta=(DisplayName = "Apply Fogging"))
-	uint32 bUseTranslucencyVertexFog:1;
+	uint8 bUseTranslucencyVertexFog : 1;
 
 	/** Unlit and Opaque materials can be used as sky material on a sky dome mesh. When IsSky is true, these meshes will not receive any contribution from the aerial perspective. Height and Volumetric fog effects will still be applied. */
 	UPROPERTY(EditAnywhere, Category=Material, AdvancedDisplay)
-	uint32 bIsSky:1;
+	uint8 bIsSky : 1;
 
 	/** When true, translucent materials have fog computed for every pixel, which costs more but fixes artifacts due to low tessellation. */
 	UPROPERTY(EditAnywhere, Category=Translucency)
-	uint32 bComputeFogPerPixel:1;
+	uint8 bComputeFogPerPixel : 1;
 
 	/** When true, translucent materials will output motion vectors in velocity pass. */
 	UPROPERTY(EditAnywhere, Category = Translucency, meta = (DisplayName = "Output Velocity"))
-	uint32 bOutputTranslucentVelocity : 1;
+	uint8 bOutputTranslucentVelocity : 1;
 
 	/** If true the compilation environment will be changed to remove the global COMPILE_SHADERS_FOR_DEVELOPMENT flag. */
 	UPROPERTY(transient, duplicatetransient)
-	uint32 bAllowDevelopmentShaderCompile:1;
+	uint8 bAllowDevelopmentShaderCompile : 1;
 
 	/** true if this is a special material used for stats by the material editor. */
 	UPROPERTY(transient, duplicatetransient)
-	uint32 bIsMaterialEditorStatsMaterial:1;
-
-	/** true if we have printed a warning about material usage for a given usage flag. */
-	UPROPERTY(transient, duplicatetransient)
-	uint32 UsageFlagWarnings;
+	uint8 bIsMaterialEditorStatsMaterial : 1;
 	
 	/** Where the node is inserted in the (post processing) graph, only used if domain is PostProcess */
 	UPROPERTY(EditAnywhere, BlueprintReadWrite, Category=PostProcessMaterial, meta=(DisplayName = "Blendable Location"))
 	TEnumAsByte<enum EBlendableLocation> BlendableLocation;
 
-	/** If multiple nodes with the same  type are inserted at the same point, this defined order and if they get combined, only used if domain is PostProcess */
-	UPROPERTY(EditAnywhere, BlueprintReadWrite, Category=PostProcessMaterial, meta=(DisplayName = "Blendable Priority"))
-	int32 BlendablePriority;
-
 	/** If this is enabled, the blendable will output alpha */
 	UPROPERTY(EditAnywhere, BlueprintReadWrite, Category = PostProcessMaterial, meta = (DisplayName = "Output Alpha"))
-	bool BlendableOutputAlpha;
+	uint8 BlendableOutputAlpha : 1;
 
 	/** 
 	 * Selectively execute post process material only for pixels that pass the stencil test against the Custom Depth/Stencil buffer. 
@@ -917,6 +908,14 @@ public:
 	UPROPERTY(EditAnywhere, Category=Refraction)
 	TEnumAsByte<enum ERefractionMode> RefractionMode;
 
+	/** If multiple nodes with the same  type are inserted at the same point, this defined order and if they get combined, only used if domain is PostProcess */
+	UPROPERTY(EditAnywhere, BlueprintReadWrite, Category = PostProcessMaterial, meta = (DisplayName = "Blendable Priority"))
+	int32 BlendablePriority;
+
+	/** true if we have printed a warning about material usage for a given usage flag. */
+	UPROPERTY(transient, duplicatetransient)
+	uint32 UsageFlagWarnings;
+
 	/** This is the refraction depth bias, larger values offset distortion to prevent closer objects from rendering into the distorted surface at acute viewing angles but increases the disconnect between surface and where the refraction starts. */
 	UPROPERTY(EditAnywhere, Category=Refraction)
 	float RefractionDepthBias;
@@ -930,6 +929,9 @@ public:
 	UPROPERTY()
 	FGuid StateId;
 
+	UPROPERTY(EditAnywhere, Category = Tessellation)
+	float MaxDisplacement;
+
 #if STORE_ONLY_ACTIVE_SHADERMAPS
 	// Relative offset to the beginning of the package containing this
 	uint32 OffsetToFirstResource;
diff --git a/Engine/Source/Runtime/Engine/Classes/Materials/MaterialExpression.h b/Engine/Source/Runtime/Engine/Classes/Materials/MaterialExpression.h
--- a/Engine/Source/Runtime/Engine/Classes/Materials/MaterialExpression.h
+++ b/Engine/Source/Runtime/Engine/Classes/Materials/MaterialExpression.h
@@ -150,7 +150,7 @@ class ENGINE_API UMaterialExpression : public UObject
 
 	/** Indicates that this is a 'parameter' type of expression and should always be loaded (ie not cooked away) because we might want the default parameter. */
 	UPROPERTY()
-	uint32 bIsParameterExpression:1;
+	uint8 bIsParameterExpression : 1;
 
 #if WITH_EDITORONLY_DATA
 	/** If true, the comment bubble will be visible in the graph editor */
diff --git a/Engine/Source/Runtime/Engine/Classes/Materials/MaterialExpressionTextureBase.h b/Engine/Source/Runtime/Engine/Classes/Materials/MaterialExpressionTextureBase.h
--- a/Engine/Source/Runtime/Engine/Classes/Materials/MaterialExpressionTextureBase.h
+++ b/Engine/Source/Runtime/Engine/Classes/Materials/MaterialExpressionTextureBase.h
@@ -28,7 +28,7 @@ class ENGINE_API UMaterialExpressionTextureBase : public UMaterialExpression
 	
 	/** Is default selected texture when using mesh paint mode texture painting */
 	UPROPERTY(EditAnywhere, Category=MaterialExpressionTextureBase)
-	uint32 IsDefaultMeshpaintTexture:1;
+	uint8 IsDefaultMeshpaintTexture:1;
 	
 	//~ Begin UObject Interface
 #if WITH_EDITOR
diff --git a/Engine/Source/Runtime/Engine/Classes/Materials/MaterialExpressionTextureSample.h b/Engine/Source/Runtime/Engine/Classes/Materials/MaterialExpressionTextureSample.h
--- a/Engine/Source/Runtime/Engine/Classes/Materials/MaterialExpressionTextureSample.h
+++ b/Engine/Source/Runtime/Engine/Classes/Materials/MaterialExpressionTextureSample.h
@@ -73,18 +73,24 @@ class ENGINE_API UMaterialExpressionTextureSample : public UMaterialExpressionTe
 	UPROPERTY(EditAnywhere, Category=MaterialExpressionTextureSample)
 	TEnumAsByte<enum ESamplerSourceMode> SamplerSource;
 
+	/** Whether the texture should be sampled with per view mip biasing for sharper output with Temporal AA. */
+	UPROPERTY(EditAnywhere, Category = MaterialExpressionTextureSample)
+	uint8 AutomaticViewMipBias : 1;
+
+protected:
+	// Inherited parameter expressions can hide unused input pin
+	uint8 bShowTextureInputPin : 1;
+
+public:
+
 	/** only used if Coordinates is not hooked up */
 	UPROPERTY(EditAnywhere, Category = MaterialExpressionTextureSample)
-	uint32 ConstCoordinate;
+	uint8 ConstCoordinate;
 
 	/** only used if MipValue is not hooked up */
 	UPROPERTY(EditAnywhere, Category = MaterialExpressionTextureSample)
 	int32 ConstMipValue;
 
-	/** Whether the texture should be sampled with per view mip biasing for sharper output with Temporal AA. */
-	UPROPERTY(EditAnywhere, Category = MaterialExpressionTextureSample)
-	uint32 AutomaticViewMipBias : 1;
-
 	//~ Begin UObject Interface
 #if WITH_EDITOR
 	virtual bool CanEditChange(const UProperty* InProperty) const override;
@@ -114,11 +120,4 @@ class ENGINE_API UMaterialExpressionTextureSample : public UMaterialExpressionTe
 	int32 CompileMipValue0(class FMaterialCompiler* Compiler);
 	int32 CompileMipValue1(class FMaterialCompiler* Compiler);
 #endif // WITH_EDITOR
-
-protected:
-	// Inherited parameter expressions can hide unused input pin
-	bool bShowTextureInputPin;
-};
-
-
-
+};
\ No newline at end of file
diff --git a/Engine/Source/Runtime/Engine/Classes/Materials/MaterialFunction.h b/Engine/Source/Runtime/Engine/Classes/Materials/MaterialFunction.h
--- a/Engine/Source/Runtime/Engine/Classes/Materials/MaterialFunction.h
+++ b/Engine/Source/Runtime/Engine/Classes/Materials/MaterialFunction.h
@@ -14,6 +14,7 @@
 class UMaterial;
 class UTexture;
 struct FPropertyChangedEvent;
+class UMaterialExpression;
 
 /**
  * A Material Function is a collection of material expressions that can be reused in different materials
@@ -33,13 +34,17 @@ class UMaterialFunction : public UMaterialFunctionInterface
 	UPROPERTY(EditAnywhere, Category=MaterialFunction, AssetRegistrySearchable)
 	FString Description;
 
+	/** Array of material expressions, excluding Comments.  Used by the material editor. */
+	UPROPERTY()
+	TArray<UMaterialExpression*> FunctionExpressions;
+
 	/** Whether to list this function in the material function library, which is a window in the material editor that lists categorized functions. */
 	UPROPERTY(EditAnywhere, Category=MaterialFunction, AssetRegistrySearchable)
-	uint32 bExposeToLibrary:1;
+	uint8 bExposeToLibrary:1;
 	
 	/** If true, parameters in this function will have a prefix added to their group name. */
 	UPROPERTY(EditAnywhere, Category=MaterialFunction)
-	uint32 bPrefixParameterNames:1;
+	uint8 bPrefixParameterNames:1;
 
 #if WITH_EDITORONLY_DATA
 	/** 
@@ -56,9 +61,6 @@ class UMaterialFunction : public UMaterialFunctionInterface
 	UPROPERTY(EditAnywhere, Category=MaterialFunction, AssetRegistrySearchable)
 	TArray<FText> LibraryCategoriesText;
 #endif
-	/** Array of material expressions, excluding Comments.  Used by the material editor. */
-	UPROPERTY()
-	TArray<class UMaterialExpression*> FunctionExpressions;
 
 #if WITH_EDITORONLY_DATA
 	/** Array of comments associated with this material; viewed in the material editor. */
@@ -72,7 +74,7 @@ class UMaterialFunction : public UMaterialFunctionInterface
 private:
 	/** Transient flag used to track re-entrance in recursive functions like IsDependent. */
 	UPROPERTY(transient)
-	uint32 bReentrantFlag:1;
+	uint8 bReentrantFlag:1;
 
 public:
 	//~ Begin UObject Interface.
diff --git a/Engine/Source/Runtime/Engine/Classes/Materials/MaterialInstance.h b/Engine/Source/Runtime/Engine/Classes/Materials/MaterialInstance.h
--- a/Engine/Source/Runtime/Engine/Classes/Materials/MaterialInstance.h
+++ b/Engine/Source/Runtime/Engine/Classes/Materials/MaterialInstance.h
@@ -295,9 +295,6 @@ class ENGINE_VTABLE UMaterialInstance : public UMaterialInterface
 	 */
 	ENGINE_API const FStaticParameterSet& GetStaticParameters() const;
 
-	/** Flag to detect cycles in the material instance graph. */
-	bool ReentrantFlag[2];
-
 	/**
 	 * Indicates whether the instance has static permutation resources (which are required when static parameters are present) 
 	 * Read directly from the rendering thread, can only be modified with the use of a FMaterialUpdateContext.
@@ -316,7 +313,9 @@ class ENGINE_VTABLE UMaterialInstance : public UMaterialInterface
 	uint8 bIsShadingModelFromMaterialExpression : 1;
 
 	TEnumAsByte<EBlendMode> BlendMode;
-	FMaterialShadingModelField ShadingModels;
+
+	//Cached copies of the base property overrides or the value from the parent to avoid traversing the parent chain for each access.
+	float OpacityMaskClipValue;
 
 	FORCEINLINE bool GetReentrantFlag() const
 	{
@@ -357,8 +356,10 @@ class ENGINE_VTABLE UMaterialInstance : public UMaterialInterface
 	uint32 OffsetToFirstResource;
 #endif
 
-	//Cached copies of the base property overrides or the value from the parent to avoid traversing the parent chain for each access.
-	float OpacityMaskClipValue;
+	FMaterialShadingModelField ShadingModels;
+
+	/** Flag to detect cycles in the material instance graph. */
+	bool ReentrantFlag[2];
 
 	/** 
 	 * FMaterialRenderProxy derivative that represent this material instance to the renderer, when the renderer needs to fetch parameter values. 
diff --git a/Engine/Source/Runtime/Engine/Classes/Materials/MaterialInterface.h b/Engine/Source/Runtime/Engine/Classes/Materials/MaterialInterface.h
--- a/Engine/Source/Runtime/Engine/Classes/Materials/MaterialInterface.h
+++ b/Engine/Source/Runtime/Engine/Classes/Materials/MaterialInterface.h
@@ -118,10 +118,6 @@ struct FLightmassMaterialInterfaceSettings
 {
 	GENERATED_USTRUCT_BODY()
 
-	/** If true, forces translucency to cast static shadows as if the material were masked. */
-	UPROPERTY(EditAnywhere, Category=Material)
-	uint32 bCastShadowAsMasked:1;
-
 	/** Scales the emissive contribution of this material to static lighting. */
 	UPROPERTY()
 	float EmissiveBoost;
@@ -137,28 +133,32 @@ struct FLightmassMaterialInterfaceSettings
 	UPROPERTY(EditAnywhere, Category=Material)
 	float ExportResolutionScale;
 
+	/** If true, forces translucency to cast static shadows as if the material were masked. */
+	UPROPERTY(EditAnywhere, Category = Material)
+	uint8 bCastShadowAsMasked : 1;
+
 	/** Boolean override flags - only used in MaterialInstance* cases. */
 	/** If true, override the bCastShadowAsMasked setting of the parent material. */
 	UPROPERTY()
-	uint32 bOverrideCastShadowAsMasked:1;
+	uint8 bOverrideCastShadowAsMasked:1;
 
 	/** If true, override the emissive boost setting of the parent material. */
 	UPROPERTY()
-	uint32 bOverrideEmissiveBoost:1;
+	uint8 bOverrideEmissiveBoost:1;
 
 	/** If true, override the diffuse boost setting of the parent material. */
 	UPROPERTY()
-	uint32 bOverrideDiffuseBoost:1;
+	uint8 bOverrideDiffuseBoost:1;
 
 	/** If true, override the export resolution scale setting of the parent material. */
 	UPROPERTY()
-	uint32 bOverrideExportResolutionScale:1;
+	uint8 bOverrideExportResolutionScale:1;
 
 	FLightmassMaterialInterfaceSettings()
-		: bCastShadowAsMasked(false)
-		, EmissiveBoost(1.0f)
+		: EmissiveBoost(1.0f)
 		, DiffuseBoost(1.0f)
 		, ExportResolutionScale(1.0f)
+		, bCastShadowAsMasked(false)
 		, bOverrideCastShadowAsMasked(false)
 		, bOverrideEmissiveBoost(false)
 		, bOverrideDiffuseBoost(false)
@@ -231,10 +231,6 @@ protected:
 	UPROPERTY(EditAnywhere, Category=Lightmass)
 	struct FLightmassMaterialInterfaceSettings LightmassSettings;
 
-private:
-	/** Feature levels to force to compile. */
-	uint32 FeatureLevelsToForceCompile;
-
 protected:
 #if WITH_EDITORONLY_DATA
 	/** Because of redirector, the texture names need to be resorted at each load in case they changed. */
@@ -252,6 +248,10 @@ protected:
 	UPROPERTY(EditAnywhere, AdvancedDisplay, Instanced, Category = Material)
 	TArray<UAssetUserData*> AssetUserData;
 
+private:
+	/** Feature levels to force to compile. */
+	uint32 FeatureLevelsToForceCompile;
+
 public:
 
 	//~ Begin IInterface_AssetUserData Interface
diff --git a/Engine/Source/Runtime/Engine/Classes/Particles/Acceleration/ParticleModuleAccelerationBase.h b/Engine/Source/Runtime/Engine/Classes/Particles/Acceleration/ParticleModuleAccelerationBase.h
--- a/Engine/Source/Runtime/Engine/Classes/Particles/Acceleration/ParticleModuleAccelerationBase.h
+++ b/Engine/Source/Runtime/Engine/Classes/Particles/Acceleration/ParticleModuleAccelerationBase.h
@@ -19,8 +19,7 @@ class UParticleModuleAccelerationBase : public UParticleModule
 	 *	If true, then treat the acceleration as world-space
 	 */
 	UPROPERTY(EditAnywhere, Category=Acceleration)
-	uint32 bAlwaysInWorldSpace:1;
-
+	uint8 bAlwaysInWorldSpace : 1;
 
 	//~ Begin UParticleModule Interface
 	virtual void SetToSensibleDefaults(UParticleEmitter* Owner) override;
diff --git a/Engine/Source/Runtime/Engine/Classes/Particles/Color/ParticleModuleColor.h b/Engine/Source/Runtime/Engine/Classes/Particles/Color/ParticleModuleColor.h
--- a/Engine/Source/Runtime/Engine/Classes/Particles/Color/ParticleModuleColor.h
+++ b/Engine/Source/Runtime/Engine/Classes/Particles/Color/ParticleModuleColor.h
@@ -14,6 +14,8 @@ class UInterpCurveEdSetup;
 class UParticleEmitter;
 struct FCurveEdEntry;
 struct FParticleEmitterInstance;
+struct FRawDistributionFloat;
+struct FRawDistributionVector;
 
 UCLASS(editinlinenew, hidecategories=Object, meta=(DisplayName = "Initial Color"))
 class ENGINE_API UParticleModuleColor : public UParticleModuleColorBase
@@ -22,15 +24,15 @@ class ENGINE_API UParticleModuleColor : public UParticleModuleColorBase
 
 	/** Initial color for a particle as a function of Emitter time. */
 	UPROPERTY(EditAnywhere, Category = Color, meta = (TreatAsColor))
-	struct FRawDistributionVector StartColor;
+	FRawDistributionVector StartColor;
 
 	/** Initial alpha for a particle as a function of Emitter time. */
 	UPROPERTY(EditAnywhere, Category=Color)
-	struct FRawDistributionFloat StartAlpha;
+	FRawDistributionFloat StartAlpha;
 
 	/** If true, the alpha value will be clamped to the [0..1] range. */
 	UPROPERTY(EditAnywhere, Category=Color)
-	uint32 bClampAlpha:1;
+	uint8 bClampAlpha:1;
 
 	/** Initializes the default values for this property */
 	void InitializeDefaults();
diff --git a/Engine/Source/Runtime/Engine/Classes/Particles/ParticleModule.h b/Engine/Source/Runtime/Engine/Classes/Particles/ParticleModule.h
--- a/Engine/Source/Runtime/Engine/Classes/Particles/ParticleModule.h
+++ b/Engine/Source/Runtime/Engine/Classes/Particles/ParticleModule.h
@@ -98,14 +98,14 @@ struct FParticleRandomSeedInfo
 	 *	the RandomSeeds array.
 	 */
 	UPROPERTY(EditAnywhere, Category=ParticleRandomSeedInfo)
-	uint32 bGetSeedFromInstance:1;
+	uint8 bGetSeedFromInstance:1;
 
 	/** 
 	 *	If true, the seed value retrieved from the instance will be an
 	 *	index into the array of seeds.
 	 */
 	UPROPERTY(EditAnywhere, Category=ParticleRandomSeedInfo)
-	uint32 bInstanceSeedIsIndex:1;
+	uint8 bInstanceSeedIsIndex:1;
 
 	/** 
 	 *	If true, then reset the seed upon the emitter looping.
@@ -113,13 +113,13 @@ struct FParticleRandomSeedInfo
 	 *	a repeating pattern.
 	 */
 	UPROPERTY(EditAnywhere, Category=ParticleRandomSeedInfo)
-	uint32 bResetSeedOnEmitterLooping:1;
+	uint8 bResetSeedOnEmitterLooping:1;
 
 	/**
 	*	If true, then randomly select a seed entry from the RandomSeeds array
 	*/
 	UPROPERTY(EditAnywhere, Category = ParticleRandomSeedInfo)
-	uint32 bRandomlySelectSeedArray:1;
+	uint8 bRandomlySelectSeedArray:1;
 
 	/** 
 	 *	The random seed values to utilize for the module. 
@@ -151,39 +151,39 @@ class ENGINE_API UParticleModule : public UObject
 
 	/** If true, the module performs operations on particles during Spawning		*/
 	UPROPERTY()
-	uint32 bSpawnModule:1;
+	uint8 bSpawnModule:1;
 
 	/** If true, the module performs operations on particles during Updating		*/
 	UPROPERTY()
-	uint32 bUpdateModule:1;
+	uint8 bUpdateModule:1;
 
 	/** If true, the module performs operations on particles during final update	*/
 	UPROPERTY()
-	uint32 bFinalUpdateModule:1;
+	uint8 bFinalUpdateModule:1;
 
 	/** If true, the module performs operations on particles during update and/or final update for GPU emitters*/
 	UPROPERTY()
-	uint32 bUpdateForGPUEmitter:1;
+	uint8 bUpdateForGPUEmitter:1;
 
 	/** If true, the module displays FVector curves as colors						*/
 	UPROPERTY()
-	uint32 bCurvesAsColor:1;
+	uint8 bCurvesAsColor:1;
 
 	/** If true, the module should render its 3D visualization helper				*/
 	UPROPERTY(EditAnywhere, Category=Cascade)
-	uint32 b3DDrawMode:1;
+	uint8 b3DDrawMode:1;
 
 	/** If true, the module supports rendering a 3D visualization helper			*/
 	UPROPERTY()
-	uint32 bSupported3DDrawMode:1;
+	uint8 bSupported3DDrawMode:1;
 
 	/** If true, the module is enabled												*/
 	UPROPERTY()
-	uint32 bEnabled:1;
+	uint8 bEnabled:1;
 
 	/** If true, the module has had editing enabled on it							*/
 	UPROPERTY()
-	uint32 bEditable:1;
+	uint8 bEditable:1;
 
 	/**
 	*	If true, this flag indicates that auto-generation for LOD will result in
@@ -191,15 +191,15 @@ class ENGINE_API UParticleModule : public UObject
 	*	If false, it will result in a module with different settings.
 	*/
 	UPROPERTY()
-	uint32 LODDuplicate:1;
+	uint8 LODDuplicate:1;
 
 	/** If true, the module supports RandomSeed setting */
 	UPROPERTY()
-	uint32 bSupportsRandomSeed:1;
+	uint8 bSupportsRandomSeed:1;
 
 	/** If true, the module should be told when looping */
 	UPROPERTY()
-	uint32 bRequiresLoopingNotification:1;
+	uint8 bRequiresLoopingNotification:1;
 
 	/**
 	 *	The LOD levels this module is present in.
@@ -274,12 +274,6 @@ class ENGINE_API UParticleModule : public UObject
 	 */
 	virtual uint32	RequiredBytesPerInstance();
 
-	UE_DEPRECATED(4.11, "RequiredBytes now takes a UParticleModuleTypeDataBase*, not per-instance information.")
-	virtual uint32	RequiredBytes(FParticleEmitterInstance* Owner);
-
-	UE_DEPRECATED(4.11, "RequiredBytesPerInstance no longer takes per-instance information.")
-	virtual uint32	RequiredBytesPerInstance(FParticleEmitterInstance* Owner);
-
 	/**
 	 *	Allows the module to prep its 'per-instance' data block.
 	 *	
diff --git a/Engine/Source/Runtime/Engine/Classes/Particles/ParticleModuleRequired.h b/Engine/Source/Runtime/Engine/Classes/Particles/ParticleModuleRequired.h
--- a/Engine/Source/Runtime/Engine/Classes/Particles/ParticleModuleRequired.h
+++ b/Engine/Source/Runtime/Engine/Classes/Particles/ParticleModuleRequired.h
@@ -106,13 +106,13 @@
 
 struct FParticleRequiredModule
 {
-	bool bCutoutTexureIsValid;
 	uint32 NumFrames;
 	uint32 NumBoundingVertices;
 	uint32 NumBoundingTriangles;
 	float AlphaThreshold;
 	TArray<FVector2D> FrameData;
 	FRHIShaderResourceView* BoundingGeometryBufferSRV;
+	uint8 bCutoutTexureIsValid : 1;
 };
 
 
@@ -130,8 +130,17 @@
 	UPROPERTY(EditAnywhere, Category=Emitter)
 	class UMaterialInterface* Material;
 
+	/** The distance at which PSA_FacingCameraDistanceBlend	is fully PSA_Square */
+	UPROPERTY(EditAnywhere, Category = Emitter, meta = (UIMin = "0", DisplayAfter="ScreenAlignment"))
+	float MinFacingCameraBlendDistance;
+
+	/** The distance at which PSA_FacingCameraDistanceBlend	is fully PSA_FacingCameraPosition */
+	UPROPERTY(EditAnywhere, Category = Emitter, meta = (UIMin = "0", DisplayAfter="MinFacingCameraBlendDistance"))
+	float MaxFacingCameraBlendDistance;
+
 	UPROPERTY(EditAnywhere, Category=Emitter)
 	FVector EmitterOrigin;
+
 	UPROPERTY(EditAnywhere, Category=Emitter)
 	FRotator EmitterRotation;
 
@@ -150,25 +159,17 @@
 	UPROPERTY(EditAnywhere, Category=Emitter)
 	TEnumAsByte<EParticleScreenAlignment> ScreenAlignment;
 
-	/** The distance at which PSA_FacingCameraDistanceBlend	is fully PSA_Square */
-	UPROPERTY(EditAnywhere, Category=Emitter, meta=(UIMin = "0"))
-	float MinFacingCameraBlendDistance;
-
-	/** The distance at which PSA_FacingCameraDistanceBlend	is fully PSA_FacingCameraPosition */
-	UPROPERTY(EditAnywhere, Category=Emitter, meta=(UIMin = "0"))
-	float MaxFacingCameraBlendDistance;
-
 	/** If true, update the emitter in local space										*/
-	UPROPERTY(EditAnywhere, Category=Emitter)
-	uint32 bUseLocalSpace:1;
+	UPROPERTY(EditAnywhere, Category=Emitter, meta = (DisplayAfter = "MaxFacingCameraBlendDistance"))
+	uint8 bUseLocalSpace:1;
 
 	/** If true, kill the emitter when the particle system is deactivated				*/
 	UPROPERTY(EditAnywhere, Category=Emitter)
-	uint32 bKillOnDeactivate:1;
+	uint8 bKillOnDeactivate:1;
 
 	/** If true, kill the emitter when it completes										*/
 	UPROPERTY(EditAnywhere, Category=Emitter)
-	uint32 bKillOnCompleted:1;
+	uint8 bKillOnCompleted:1;
 
 	/**
 	 *	The sorting mode to use for this emitter.
@@ -180,7 +181,7 @@
 	 *
 	 */
 	UPROPERTY(EditAnywhere, Category=Emitter)
-	TEnumAsByte<enum EParticleSortMode> SortMode;
+	TEnumAsByte<EParticleSortMode> SortMode;
 
 	/**
 	 *	If true, the EmitterTime for the emitter will be calculated by
@@ -193,42 +194,24 @@
 	 *	resulting in proper looping/delay behavior.
 	 */
 	UPROPERTY(EditAnywhere, Category=Emitter)
-	uint32 bUseLegacyEmitterTime:1;
+	uint8 bUseLegacyEmitterTime:1;
 
 	/** If true, removes the HMD view roll (e.g. in VR) */
 	UPROPERTY(EditAnywhere, Category=Emitter, meta=(DisplayName = "Remove HMD Roll"))
-	uint32 bRemoveHMDRoll:1;
+	uint8 bRemoveHMDRoll:1;
 
-	/** 
-	 *	How long, in seconds, the emitter will run before looping.
-	 */
-	UPROPERTY(EditAnywhere, Category=Duration)
-	float EmitterDuration;
-
-	/** 
-	 *	The low end of the emitter duration if using a range.
-	 */
-	UPROPERTY(EditAnywhere, Category=Duration)
-	float EmitterDurationLow;
-
 	/**
-	 *	If true, select the emitter duration from the range 
+	 *	If true, select the emitter duration from the range
 	 *		[EmitterDurationLow..EmitterDuration]
 	 */
-	UPROPERTY(EditAnywhere, Category=Duration)
-	uint32 bEmitterDurationUseRange:1;
+	UPROPERTY(EditAnywhere, Category = Duration)
+	uint8 bEmitterDurationUseRange : 1;
 
 	/** 
-	 *	If true, recalculate the emitter duration on each loop.
+	 *	How long, in seconds, the emitter will run before looping.
 	 */
-	UPROPERTY(EditAnywhere, Category=Duration)
-	uint32 bDurationRecalcEachLoop:1;
+	UPROPERTY(EditAnywhere, Category=Duration, meta = (DisplayAfter="bRemoveHMDRoll"))
+	float EmitterDuration;
 
-	/** The number of times to loop the emitter.
-	 *	0 indicates loop continuously
-	 */
-	UPROPERTY(EditAnywhere, Category=Duration)
-	int32 EmitterLoops;
-
 	//
 	// Spawn-related.
@@ -238,16 +221,9 @@
 	UPROPERTY()
 	struct FRawDistributionFloat SpawnRate;
 
-	//
-	// Burst-related.
-	//
-	/** The method to utilize when burst-emitting particles						*/
-	UPROPERTY()
-	TEnumAsByte<EParticleBurstMethod> ParticleBurstMethod;
-
 	/** The array of burst entries.												*/
 	UPROPERTY(export, noclear)
-	TArray<struct FParticleBurst> BurstList;
+	TArray<FParticleBurst> BurstList;
 
 	//
 	// Delay-related.
@@ -266,17 +242,10 @@
 	float EmitterDelayLow;
 
 	/**
-	 *	If true, select the emitter delay from the range 
-	 *		[EmitterDelayLow..EmitterDelay]
-	 */
-	UPROPERTY(EditAnywhere, Category=Delay)
-	uint32 bEmitterDelayUseRange:1;
-
-	/**
 	 *	If true, the emitter will be delayed only on the first loop.
 	 */
 	UPROPERTY(EditAnywhere, Category=Delay)
-	uint32 bDelayFirstLoopOnly:1;
+	uint8 bDelayFirstLoopOnly:1;
 
 	//
 	// SubUV-related.
@@ -298,6 +267,45 @@
 	UPROPERTY(EditAnywhere, Category=SubUV)
 	TEnumAsByte<EParticleSubUVInterpMethod> InterpolationMethod;
 
+	/** Whether to scale the UV or not - ie, the model wasn't setup with sub uvs		*/
+	UPROPERTY(EditAnywhere, Category = SubUV)
+	uint8 bScaleUV : 1;
+
+	/**
+	 *	If true, select the emitter delay from the range 
+	 *		[EmitterDelayLow..EmitterDelay]
+	 */
+	UPROPERTY(EditAnywhere, Category=Delay)
+	uint8 bEmitterDelayUseRange:1;
+
+	/** The method to utilize when burst-emitting particles						*/
+	UPROPERTY()
+	TEnumAsByte<EParticleBurstMethod> ParticleBurstMethod;
+
+	/** Override the system MacroUV settings                                            */
+	UPROPERTY(EditAnywhere, Category = MacroUV)
+	uint8 bOverrideSystemMacroUV : 1;
+
+	/**
+	 *	If true, use the MaxDrawCount to limit the number of particles rendered.
+	 *	NOTE: This does not limit the number spawned/updated, only what is drawn.
+	 */
+	UPROPERTY(EditAnywhere, Category = Rendering)
+	uint8 bUseMaxDrawCount : 1;
+
+	UPROPERTY(EditAnywhere, Category=ParticleCutout)
+	TEnumAsByte<enum EOpacitySourceMode> OpacitySourceMode;
+
+	/** Normal generation mode for this emitter LOD. */
+	UPROPERTY(EditAnywhere, Category = Normals)
+	TEnumAsByte<enum EEmitterNormalsMode> EmitterNormalsMode;
+
+	/**
+	* Ensures that movement generated from the orbit module is applied to velocity-aligned particles
+	*/
+	UPROPERTY(EditAnywhere, Category = Emitter)
+	uint8 bOrbitModuleAffectsVelocityAlignment : 1;
+
 	/** The number of sub-images horizontally in the texture							*/
 	UPROPERTY(EditAnywhere, Category=SubUV)
 	int32 SubImages_Horizontal;
@@ -306,10 +314,6 @@
 	UPROPERTY(EditAnywhere, Category=SubUV)
 	int32 SubImages_Vertical;
 
-	/** Whether to scale the UV or not - ie, the model wasn't setup with sub uvs		*/
-	UPROPERTY(EditAnywhere, Category=SubUV)
-	uint32 bScaleUV:1;
-
 	/**
 	 *	The amount of time (particle-relative, 0.0 to 1.0) to 'lock' on a random sub image
 	 *	    0.0 = change every frame
@@ -322,10 +326,6 @@
 	UPROPERTY(EditAnywhere, Category=SubUV)
 	int32 RandomImageChanges;
 
-	/** Override the system MacroUV settings                                            */
-	UPROPERTY(EditAnywhere, Category=MacroUV)
-	uint32 bOverrideSystemMacroUV:1;
-
 	/** Local space position that UVs generated with the ParticleMacroUV material node will be centered on. */
 	UPROPERTY(EditAnywhere, Category=MacroUV)
 	FVector MacroUVPosition;
@@ -335,31 +335,11 @@
 	float MacroUVRadius;
 
 	/**
-	 *	If true, use the MaxDrawCount to limit the number of particles rendered.
-	 *	NOTE: This does not limit the number spawned/updated, only what is drawn.
-	 */
-	UPROPERTY(EditAnywhere, Category=Rendering)
-	uint32 bUseMaxDrawCount:1;
-
-	/**
-	 *	The maximum number of particles to DRAW for this emitter.
-	 *	If set to 0, it will use whatever number are present.
-	 */
-	UPROPERTY(EditAnywhere, Category=Rendering)
-	int32 MaxDrawCount;
-
-	/**
 	* Controls UV Flipping for this emitter.
 	*/
-	UPROPERTY(EditAnywhere, Category=Rendering)
+	UPROPERTY(EditAnywhere, Category = Rendering)
 	EParticleUVFlipMode UVFlippingMode;
 
 	/**
-	* Texture to generate bounding geometry from.
-	*/
-	UPROPERTY(EditAnywhere, Category=ParticleCutout)
-	UTexture2D* CutoutTexture;
-
-	/**
 	* More bounding vertices results in reduced overdraw, but adds more triangle overhead.
 	* The eight vertex mode is best used when the SubUV texture has a lot of space to cut out that is not captured by the four vertex version,
@@ -368,19 +348,11 @@
 	UPROPERTY(EditAnywhere, Category=ParticleCutout)
 	TEnumAsByte<enum ESubUVBoundingVertexCount> BoundingMode;
 
-	UPROPERTY(EditAnywhere, Category=ParticleCutout)
-	TEnumAsByte<enum EOpacitySourceMode> OpacitySourceMode;
-
 	/**
-	* Alpha channel values larger than the threshold are considered occupied and will be contained in the bounding geometry.
-	* Raising this threshold slightly can reduce overdraw in particles using this animation asset.
-	*/
-	UPROPERTY(EditAnywhere, Category=ParticleCutout, meta=(UIMin = "0", UIMax = "1"))
-	float AlphaThreshold;
-
-	/** Normal generation mode for this emitter LOD. */
-	UPROPERTY(EditAnywhere, Category=Normals)
-	TEnumAsByte<enum EEmitterNormalsMode> EmitterNormalsMode;
+	 *	If true, recalculate the emitter duration on each loop.
+	 */
+	UPROPERTY(EditAnywhere, Category = Duration)
+	uint8 bDurationRecalcEachLoop : 1;
 
 	/** 
 	 * When EmitterNormalsMode is ENM_Spherical, particle normals are created to face away from NormalsSphereCenter. 
@@ -389,5 +361,37 @@
 	UPROPERTY(EditAnywhere, Category=Normals)
 	FVector NormalsSphereCenter;
 
+	/**
+	* Alpha channel values larger than the threshold are considered occupied and will be contained in the bounding geometry.
+	* Raising this threshold slightly can reduce overdraw in particles using this animation asset.
+	*/
+	UPROPERTY(EditAnywhere, Category = ParticleCutout, meta = (UIMin = "0", UIMax = "1"))
+	float AlphaThreshold;
+
+	/** The number of times to loop the emitter.
+	 *	0 indicates loop continuously
+	 */
+	UPROPERTY(EditAnywhere, Category = Duration)
+	int32 EmitterLoops;
+
+	/**
+	* Texture to generate bounding geometry from.
+	*/
+	UPROPERTY(EditAnywhere, Category = ParticleCutout)
+	UTexture2D* CutoutTexture;
+
+	/**
+	 *	The maximum number of particles to DRAW for this emitter.
+	 *	If set to 0, it will use whatever number are present.
+	 */
+	UPROPERTY(EditAnywhere, Category = Rendering)
+	int32 MaxDrawCount;
+
+	/**
+	 *	The low end of the emitter duration if using a range.
+	 */
+	UPROPERTY(EditAnywhere, Category = Duration, meta = (DisplayAfter = "EmitterDuration"))
+	float EmitterDurationLow;
+
 	/** 
 	 * When EmitterNormalsMode is ENM_Cylindrical, 
@@ -397,12 +401,6 @@
 	UPROPERTY(EditAnywhere, Category=Normals)
 	FVector NormalsCylinderDirection;
 
-	/**
-	 * Ensures that movement generated from the orbit module is applied to velocity-aligned particles
-	 */
-	UPROPERTY(EditAnywhere, Category=Emitter)
-	uint32 bOrbitModuleAffectsVelocityAlignment:1;
-
 	/** 
 	*	Named material overrides for this emitter. 
 	*	Overrides this emitter's material(s) with those in the correspondingly named slot(s) of the owning system.
diff --git a/Engine/Source/Runtime/Engine/Classes/Particles/ParticleSystem.h b/Engine/Source/Runtime/Engine/Classes/Particles/ParticleSystem.h
--- a/Engine/Source/Runtime/Engine/Classes/Particles/ParticleSystem.h
+++ b/Engine/Source/Runtime/Engine/Classes/Particles/ParticleSystem.h
@@ -12,6 +12,10 @@
 
 enum class EParticleSignificanceLevel : uint8;
 enum class EParticleSystemInsignificanceReaction: uint8;
+class UInterpCurveEdSetup;
+class UMaterialInterface;
+class UParticleSystemComponent;
+class UParticleEmitter;
 
 /**
  *	ParticleSystemUpdateMode
@@ -100,7 +104,7 @@
 	FName Name;
 
 	UPROPERTY(EditAnywhere, Category = NamedMaterial)
-	class UMaterialInterface* Material;
+	UMaterialInterface* Material;
 };
 
 UCLASS(Abstract, MinimalAPI, BlueprintType)
@@ -132,9 +136,6 @@
 
 	}
 
-	UPROPERTY(EditAnywhere, Category=ParticleSystem, AssetRegistrySearchable)
-	TEnumAsByte<enum EParticleSystemUpdateMode> SystemUpdateMode;
-
 	/** UpdateTime_FPS	- the frame per second to update at in FixedTime mode		*/
 	UPROPERTY(EditAnywhere, Category=ParticleSystem)
 	float UpdateTime_FPS;
@@ -159,11 +160,11 @@
 
 	/** Emitters	- internal - the array of emitters in the system				*/
 	UPROPERTY(instanced)
-	TArray<class UParticleEmitter*> Emitters;
+	TArray<UParticleEmitter*> Emitters;
 
 	/** The component used to preview the particle system in Cascade				*/
 	UPROPERTY(transient)
-	class UParticleSystemComponent* PreviewComponent;
+	UParticleSystemComponent* PreviewComponent;
 
 #if WITH_EDITORONLY_DATA
 	/** The angle to use when rendering the thumbnail image							*/
@@ -181,11 +182,7 @@
 #endif // WITH_EDITORONLY_DATA
 	/** Used for curve editor to remember curve-editing setup.						*/
 	UPROPERTY(export)
-	class UInterpCurveEdSetup* CurveEdSetup;
+	UInterpCurveEdSetup* CurveEdSetup;
 
-	/** If true, the system's Z axis will be oriented toward the camera				*/
-	UPROPERTY(EditAnywhere, Category=ParticleSystem)
-	uint32 bOrientZAxisTowardCamera:1;
-
 	//
 	//	LOD
@@ -196,14 +193,9 @@
 	UPROPERTY(EditAnywhere, Category=LOD, AssetRegistrySearchable)
 	float LODDistanceCheckTime;
 
-	/**
-	 *	The method of LOD level determination to utilize for this particle system
-	 *	  PARTICLESYSTEMLODMETHOD_Automatic - Automatically set the LOD level, checking every LODDistanceCheckTime seconds.
-	 *    PARTICLESYSTEMLODMETHOD_DirectSet - LOD level is directly set by the game code.
-	 *    PARTICLESYSTEMLODMETHOD_ActivateAutomatic - LOD level is determined at Activation time, then left alone unless directly set by game code.
-	 */
-	UPROPERTY(EditAnywhere, Category=LOD)
-	TEnumAsByte<enum ParticleSystemLODMethod> LODMethod;
+	/** World space radius that UVs generated with the ParticleMacroUV material node will tile based on. */
+	UPROPERTY(EditAnywhere, Category = MacroUV)
+	float MacroUVRadius;
 
 	/**
 	 *	The array of distances for each LOD level in the system.
@@ -228,20 +220,9 @@
 	int32 EditorLODSetting;
 
 #endif // WITH_EDITORONLY_DATA
-	/**
-	 *	Internal value that tracks the regenerate LOD levels preference.
-	 *	If true, when autoregenerating LOD levels in code, the low level will
-	 *	be a duplicate of the high.
-	 */
-	UPROPERTY()
-	uint32 bRegenerateLODDuplicate:1;
 
 	UPROPERTY(EditAnywhere, Category=LOD)
-	TArray<struct FParticleSystemLOD> LODSettings;
+	TArray<FParticleSystemLOD> LODSettings;
 
-	/** Whether to use the fixed relative bounding box or calculate it every frame. */
-	UPROPERTY(EditAnywhere, Category=Bounds)
-	uint32 bUseFixedRelativeBoundingBox:1;
-
 	/**	Fixed relative bounding box for particle system.							*/
 	UPROPERTY(EditAnywhere, Category=Bounds)
@@ -278,27 +259,44 @@
 	FColor BackgroundColor;
 
 #endif // WITH_EDITORONLY_DATA
+	
+	/** How long this Particle system should delay when ActivateSystem is called on it. */
+	UPROPERTY(EditAnywhere, Category=Delay, AssetRegistrySearchable)
+	float Delay;
+
+	/** The low end of the emitter delay if using a range. */
+	UPROPERTY(EditAnywhere, Category=Delay)
+	float DelayLow;
+
+	/** If true, the system's Z axis will be oriented toward the camera				*/
+	UPROPERTY(EditAnywhere, Category = ParticleSystem)
+	uint8 bOrientZAxisTowardCamera : 1;
+
+	/** Whether to use the fixed relative bounding box or calculate it every frame. */
+	UPROPERTY(EditAnywhere, Category = Bounds)
+	uint8 bUseFixedRelativeBoundingBox : 1;
+
 	/** EDITOR ONLY: Indicates that Cascade would like to have the PeakActiveParticles count reset */
 	UPROPERTY()
-	uint32 bShouldResetPeakCounts:1;
+	uint8 bShouldResetPeakCounts : 1;
 
 	/** Set during load time to indicate that physics is used... */
 	UPROPERTY(transient)
-	uint32 bHasPhysics:1;
+	uint8 bHasPhysics : 1;
 
 	/** Inidicates the old 'real-time' thumbnail rendering should be used	*/
 	UPROPERTY(EditAnywhere, Category=Thumbnail)
-	uint32 bUseRealtimeThumbnail:1;
+	uint8 bUseRealtimeThumbnail : 1;
 
 	/** Internal: Indicates the PSys thumbnail image is out of date			*/
 	UPROPERTY()
-	uint32 ThumbnailImageOutOfDate:1;
+	uint8 ThumbnailImageOutOfDate : 1;
 
 private:
 	/** if true, this psys can tick in any thread **/
-	uint32 bIsElligibleForAsyncTick:1;
+	uint8 bIsElligibleForAsyncTick : 1;
 	/** if true, bIsElligibleForAsyncTick is set up **/
-	uint32 bIsElligibleForAsyncTickComputed:1;
+	uint8 bIsElligibleForAsyncTickComputed : 1;
 public:
 
 #if WITH_EDITORONLY_DATA
@@ -307,58 +305,73 @@
 	class UTexture2D* ThumbnailImage;
 
 #endif // WITH_EDITORONLY_DATA
-	/** How long this Particle system should delay when ActivateSystem is called on it. */
-	UPROPERTY(EditAnywhere, Category=Delay, AssetRegistrySearchable)
-	float Delay;
-
-	/** The low end of the emitter delay if using a range. */
-	UPROPERTY(EditAnywhere, Category=Delay)
-	float DelayLow;
 
 	/**
 	 *	If true, select the emitter delay from the range 
 	 *		[DelayLow..Delay]
 	 */
 	UPROPERTY(EditAnywhere, Category=Delay)
-	uint32 bUseDelayRange:1;
+	uint8 bUseDelayRange : 1;
 
 	UPROPERTY(EditAnywhere, Category = Performance, meta = (ToolTip = "Whether or not to allow instances of this system to have their ticks managed."), AdvancedDisplay)
-	uint32 bAllowManagedTicking:1;
+	uint8 bAllowManagedTicking : 1;
 
 	UPROPERTY(EditAnywhere, Category = Performance, meta = (ToolTip = "Auto-deactivate system if all emitters are determined to not spawn particles again, regardless of lifetime."))
-	bool bAutoDeactivate;
+	uint8 bAutoDeactivate : 1;
 
-	UPROPERTY(EditAnywhere, Category = Performance, meta = (ToolTip = "Minimum duration between ticks; 33=tick at max. 30FPS, 16=60FPS, 8=120FPS"))
-	uint32 MinTimeBetweenTicks;
+	/**
+	 *	Internal value that tracks the regenerate LOD levels preference.
+	 *	If true, when autoregenerating LOD levels in code, the low level will
+	 *	be a duplicate of the high.
+	 */	
+	UPROPERTY()
+	uint8 bRegenerateLODDuplicate : 1;
 
+	UPROPERTY(EditAnywhere, Category = ParticleSystem, AssetRegistrySearchable)
+	TEnumAsByte<enum EParticleSystemUpdateMode> SystemUpdateMode;
+
+	/**
+	 *	The method of LOD level determination to utilize for this particle system
+	 *	  PARTICLESYSTEMLODMETHOD_Automatic - Automatically set the LOD level, checking every LODDistanceCheckTime seconds.
+	 *    PARTICLESYSTEMLODMETHOD_DirectSet - LOD level is directly set by the game code.
+	 *    PARTICLESYSTEMLODMETHOD_ActivateAutomatic - LOD level is determined at Activation time, then left alone unless directly set by game code.
+	 */
+	UPROPERTY(EditAnywhere, Category=LOD)
+	TEnumAsByte<enum ParticleSystemLODMethod> LODMethod;
+
 	/** The reaction this system takes when all emitters are insignificant. */
 	UPROPERTY(EditAnywhere, Category = Performance)
 	EParticleSystemInsignificanceReaction InsignificantReaction;
 
-	/** Time delay between all emitters becoming insignificant and the systems insignificant reaction. */
-	UPROPERTY(EditAnywhere, Category = Performance)
-	float InsignificanceDelay;
-
-	/** The maximum level of significance for emitters in this system. Any emitters with a higher significance will be capped at this significance level. */
-	UPROPERTY(EditAnywhere, Category = Performance)
-	EParticleSignificanceLevel MaxSignificanceLevel;
-
-	/** Local space position that UVs generated with the ParticleMacroUV material node will be centered on. */
-	UPROPERTY(EditAnywhere, Category=MacroUV)
-	FVector MacroUVPosition;
-
-	/** World space radius that UVs generated with the ParticleMacroUV material node will tile based on. */
-	UPROPERTY(EditAnywhere, Category=MacroUV)
-	float MacroUVRadius;
-
-	/** 
+	/**
 	 *	Which occlusion bounds method to use for this particle system.
 	 *	EPSOBM_None - Don't determine occlusion for this system.
 	 *	EPSOBM_ParticleBounds - Use the bounds of the component when determining occlusion.
 	 */
-	UPROPERTY(EditAnywhere, Category=Occlusion)
+	UPROPERTY(EditAnywhere, Category = Occlusion)
 	TEnumAsByte<enum EParticleSystemOcclusionBoundsMethod> OcclusionBoundsMethod;
 
+private:
+	/** Does any emitter loop forever? */
+	uint8 bAnyEmitterLoopsForever : 1;
+
+public:
+
+	/** The maximum level of significance for emitters in this system. Any emitters with a higher significance will be capped at this significance level. */
+	UPROPERTY(EditAnywhere, Category = Performance)
+	EParticleSignificanceLevel MaxSignificanceLevel;
+
+	UPROPERTY(EditAnywhere, Category = Performance, meta = (ToolTip = "Minimum duration between ticks; 33=tick at max. 30FPS, 16=60FPS, 8=120FPS"))
+	uint32 MinTimeBetweenTicks;
+
+	/** Time delay between all emitters becoming insignificant and the systems insignificant reaction. */
+	UPROPERTY(EditAnywhere, Category = Performance)
+	float InsignificanceDelay;
+
+	/** Local space position that UVs generated with the ParticleMacroUV material node will be centered on. */
+	UPROPERTY(EditAnywhere, Category=MacroUV)
+	FVector MacroUVPosition;
+
 	/** The occlusion bounds to use if OcclusionBoundsMethod is set to EPSOBM_CustomBounds */
 	UPROPERTY(EditAnywhere, Category=Occlusion)
 	FBox CustomOcclusionBounds;
@@ -384,7 +397,7 @@
 	virtual void GetAssetRegistryTags(TArray<FAssetRegistryTag>& OutTags) const override;
 	bool UsesCPUCollision() const;
 	virtual bool CanBeClusterRoot() const override;
-	virtual void Serialize(FArchive& Ar);
+	virtual void Serialize(FArchive& Ar) override;
 
 	//~ End UObject Interface.
 
@@ -551,15 +564,10 @@
 	/** The lowest significance of any emitter. Clamped by MaxSignificanceLevel.*/
 	EParticleSignificanceLevel LowestSignificance;
 	
-	uint32 bShouldManageSignificance : 1;
+	uint8 bShouldManageSignificance : 1;
 
-	/** Does any emitter loop forever? */
-	uint32 bAnyEmitterLoopsForever : 1;
-	/** Does any emiter never die due to infinie looping AND indefinite duration? */
-	uint32 bIsImmortal : 1;
+	/** Does any emitter never die due to infinite looping AND indefinite duration? */
+	uint8 bIsImmortal : 1;
 	/** Does any emitter ever become a zombie (is immortal AND stops spawning at some point, i.e. is burst only)? */
-	uint32 bWillBecomeZombie : 1;
+	uint8 bWillBecomeZombie : 1;
 };
-
-
-
diff --git a/Engine/Source/Runtime/Engine/Classes/Particles/Spawn/ParticleModuleSpawnPerUnit.h b/Engine/Source/Runtime/Engine/Classes/Particles/Spawn/ParticleModuleSpawnPerUnit.h
--- a/Engine/Source/Runtime/Engine/Classes/Particles/Spawn/ParticleModuleSpawnPerUnit.h
+++ b/Engine/Source/Runtime/Engine/Classes/Particles/Spawn/ParticleModuleSpawnPerUnit.h
@@ -10,6 +10,7 @@
 #include "ParticleModuleSpawnPerUnit.generated.h"
 
 struct FParticleEmitterInstance;
+struct FRawDistributionFloat;
 
 UCLASS(editinlinenew, MinimalAPI, hidecategories=Object, meta=(DisplayName = "Spawn PerUnit"))
 class UParticleModuleSpawnPerUnit : public UParticleModuleSpawnBase
@@ -24,21 +25,6 @@ class UParticleModuleSpawnPerUnit : public UParticleModuleSpawnBase
 	UPROPERTY(EditAnywhere, Category=Spawn)
 	float UnitScalar;
 
-	/** 
-	 *	The amount to spawn per meter distribution.
-	 *	The value is retrieved using the EmitterTime.
-	 */
-	UPROPERTY(EditAnywhere, Category=Spawn)
-	struct FRawDistributionFloat SpawnPerUnit;
-
-	/**
-	 *	If true, process the default spawn rate when not moving...
-	 *	When not moving, skip the default spawn rate.
-	 *	If false, return the bProcessSpawnRate setting.
-	 */
-	UPROPERTY(EditAnywhere, Category=Spawn)
-	uint32 bIgnoreSpawnRateWhenMoving:1;
-
 	/**
 	 *	The tolerance for moving vs. not moving w.r.t. the bIgnoreSpawnRateWhenMoving flag.
 	 *	Ie, if (DistanceMoved < (UnitScalar x MovementTolerance)) then consider it not moving.
@@ -46,28 +32,43 @@ class UParticleModuleSpawnPerUnit : public UParticleModuleSpawnBase
 	UPROPERTY(EditAnywhere, Category=Spawn)
 	float MovementTolerance;
 
+	/** 
+	 *	The amount to spawn per meter distribution.
+	 *	The value is retrieved using the EmitterTime.
+	 */
+	UPROPERTY(EditAnywhere, Category=Spawn)
+	FRawDistributionFloat SpawnPerUnit;
+
 	/**
 	 *	The maximum valid movement for a single frame.
 	 *	If 0.0, then the check is not performed.
 	 *	Currently, if the distance moved between frames is greater than this
 	 *	then NO particles will be spawned.
-	 *	This is primiarily intended to cover cases where the PSystem is 
+	 *	This is primiarily intended to cover cases where the PSystem is
 	 *	attached to teleporting objects.
 	 */
-	UPROPERTY(EditAnywhere, Category=Spawn)
+	UPROPERTY(EditAnywhere, Category = Spawn)
 	float MaxFrameDistance;
 
+	/**
+	 *	If true, process the default spawn rate when not moving...
+	 *	When not moving, skip the default spawn rate.
+	 *	If false, return the bProcessSpawnRate setting.
+	 */
+	UPROPERTY(EditAnywhere, Category=Spawn)
+	uint8 bIgnoreSpawnRateWhenMoving : 1;
+
 	/** If true, ignore the X-component of the movement */
 	UPROPERTY(EditAnywhere, Category=Spawn)
-	uint32 bIgnoreMovementAlongX:1;
+	uint8 bIgnoreMovementAlongX : 1;
 
 	/** If true, ignore the Y-component of the movement */
 	UPROPERTY(EditAnywhere, Category=Spawn)
-	uint32 bIgnoreMovementAlongY:1;
+	uint8 bIgnoreMovementAlongY : 1;
 
 	/** If true, ignore the Z-component of the movement */
 	UPROPERTY(EditAnywhere, Category=Spawn)
-	uint32 bIgnoreMovementAlongZ:1;
+	uint8 bIgnoreMovementAlongZ : 1;
 
 	/** Initializes the default values for this property */
 	void InitializeDefaults();
diff --git a/Engine/Source/Runtime/Engine/Classes/Particles/SubUV/ParticleModuleSubUV.h b/Engine/Source/Runtime/Engine/Classes/Particles/SubUV/ParticleModuleSubUV.h
--- a/Engine/Source/Runtime/Engine/Classes/Particles/SubUV/ParticleModuleSubUV.h
+++ b/Engine/Source/Runtime/Engine/Classes/Particles/SubUV/ParticleModuleSubUV.h
@@ -39,7 +39,7 @@ class UParticleModuleSubUV : public UParticleModuleSubUVBase
 	 *	The movie will update regardless of the slomo settings of the game.
 	 */
 	UPROPERTY(EditAnywhere, Category=Realtime)
-	uint32 bUseRealTime:1;
+	uint8 bUseRealTime:1;
 
 	/** Initializes the default values for this property */
 	void InitializeDefaults();
diff --git a/Engine/Source/Runtime/Engine/Classes/Particles/TypeData/ParticleModuleTypeDataMesh.h b/Engine/Source/Runtime/Engine/Classes/Particles/TypeData/ParticleModuleTypeDataMesh.h
--- a/Engine/Source/Runtime/Engine/Classes/Particles/TypeData/ParticleModuleTypeDataMesh.h
+++ b/Engine/Source/Runtime/Engine/Classes/Particles/TypeData/ParticleModuleTypeDataMesh.h
@@ -14,6 +14,7 @@
 class UParticleEmitter;
 class UParticleSystemComponent;
 struct FParticleEmitterInstance;
+class UStaticMesh;
 
 UENUM()
 enum EMeshScreenAlignment
@@ -64,24 +65,26 @@ class UParticleModuleTypeDataMesh : public UParticleModuleTypeDataBase
 
 	/** The static mesh to render at the particle positions */
 	UPROPERTY(EditAnywhere, Category=Mesh)
-	class UStaticMesh* Mesh;
+	UStaticMesh* Mesh;
 
-	/** use the static mesh's LOD setup and switch LODs based on largest particle's screen size*/
-	UPROPERTY(EditAnywhere, Category = Mesh)
-	bool bUseStaticMeshLODs;
+	/** Random stream for the initial rotation distribution */
+	FRandomStream RandomStream;
 
 	/** use the static mesh's LOD setup and switch LODs based on largest particle's screen size*/
 	UPROPERTY(EditAnywhere, Category = Mesh)
 	float LODSizeScale;
 
+	/** use the static mesh's LOD setup and switch LODs based on largest particle's screen size*/
+	UPROPERTY(EditAnywhere, Category = Mesh)
+	uint8 bUseStaticMeshLODs : 1;
 
 	/** If true, has the meshes cast shadows */
 	UPROPERTY()
-	uint32 CastShadows:1;
+	uint8 CastShadows:1;
 
 	/** UNUSED (the collision module dictates doing collisions) */
 	UPROPERTY()
-	uint32 DoCollisions:1;
+	uint8 DoCollisions:1;
 
 	/** 
 	 *	The alignment to use on the meshes emitted.
@@ -103,30 +106,18 @@ class UParticleModuleTypeDataMesh : public UParticleModuleTypeDataBase
 	 *	to the static mesh model.
 	 */
 	UPROPERTY(EditAnywhere, Category=Mesh)
-	uint32 bOverrideMaterial:1;
+	uint8 bOverrideMaterial:1;
 
 	UPROPERTY(EditAnywhere, Category = Mesh)
-	uint32 bOverrideDefaultMotionBlurSettings : 1;
+	uint8 bOverrideDefaultMotionBlurSettings : 1;
 
 	UPROPERTY(EditAnywhere, Category = Mesh, meta=(EditCondition="bOverrideDefaultMotionBlurSettings"))
-	uint32 bEnableMotionBlur : 1;
-
-	/** deprecated properties for initial orientation */
-	UPROPERTY()
-	float Pitch_DEPRECATED;
-	UPROPERTY()
-	float Roll_DEPRECATED;
-	UPROPERTY()
-	float Yaw_DEPRECATED;
-
+	uint8 bEnableMotionBlur : 1;
 
 	/** The 'pre' rotation pitch (in degrees) to apply to the static mesh used. */
 	UPROPERTY(EditAnywhere, Category=Orientation)
 	FRawDistributionVector RollPitchYawRange;
 
-	/** Random stream for the initial rotation distribution */
-	FRandomStream RandomStream;
-
 	/**
 	 *	The axis to lock the mesh on. This overrides TypeSpecific mesh alignment as well as the LockAxis module.
 	 *		EPAL_NONE		 -	No locking to an axis.
@@ -148,7 +139,7 @@ class UParticleModuleTypeDataMesh : public UParticleModuleTypeDataBase
 	 *	When set, AxisLockOption as well as all other locked axis/screen alignment settings are ignored.
 	 */
 	UPROPERTY(EditAnywhere, Category=CameraFacing)
-	uint32 bCameraFacing:1;
+	uint8 bCameraFacing:1;
 
 	/**
 	 *	The axis of the mesh to point up when camera facing the X-axis.
@@ -189,21 +180,21 @@ class UParticleModuleTypeDataMesh : public UParticleModuleTypeDataBase
 	 *	If false, apply 'sprite' particle rotation about the camera facing axis.
 	 */
 	UPROPERTY(EditAnywhere, Category=CameraFacing)
-	uint32 bApplyParticleRotationAsSpin:1;
+	uint8 bApplyParticleRotationAsSpin : 1;
 
 	/** 
 	*	If true, all camera facing options will point the mesh against the camera's view direction rather than pointing at the cameras location. 
 	*	If false, the camera facing will point to the cameras position as normal.
 	*/
 	UPROPERTY(EditAnywhere, Category=CameraFacing)
-	uint32 bFaceCameraDirectionRatherThanPosition:1;
+	uint8 bFaceCameraDirectionRatherThanPosition : 1;
 
 	/**
 	*	If true, all collisions for mesh particle on this emitter will take the particle size into account.
 	*	If false, particle size will be ignored in collision checks.
 	*/
 	UPROPERTY(EditAnywhere, Category = Collision)
-	uint32 bCollisionsConsiderPartilceSize : 1;
+	uint8 bCollisionsConsiderPartilceSize : 1;
 
 	static int32 GetCurrentDetailMode();
 	static int32 GetMeshParticleMotionBlurMinDetailMode();
diff --git a/Engine/Source/Runtime/Engine/Classes/PhysicsEngine/BodyInstance.h b/Engine/Source/Runtime/Engine/Classes/PhysicsEngine/BodyInstance.h
--- a/Engine/Source/Runtime/Engine/Classes/PhysicsEngine/BodyInstance.h
+++ b/Engine/Source/Runtime/Engine/Classes/PhysicsEngine/BodyInstance.h
@@ -601,7 +601,10 @@ public:
 	const struct FWalkableSlopeOverride& GetWalkableSlopeOverride() const;
 
 	/** Sets a custom slope override struct for this instance. Implicitly sets bOverrideWalkableSlopeOnInstance to true. */
-	void SetWalkableSlopeOverride(const FWalkableSlopeOverride& NewOverride);
+	void SetWalkableSlopeOverride(const FWalkableSlopeOverride& NewOverride, bool bNewOverideSetting = true);
+
+	/** Gets bOverrideWalkableSlopeOnInstance */
+	bool GetOverrideWalkableSlopeOnInstance() const;
 
 	/** Returns true if the body is not static */
 	bool IsDynamic() const;
diff --git a/Engine/Source/Runtime/Engine/Classes/Sound/SoundCue.h b/Engine/Source/Runtime/Engine/Classes/Sound/SoundCue.h
--- a/Engine/Source/Runtime/Engine/Classes/Sound/SoundCue.h
+++ b/Engine/Source/Runtime/Engine/Classes/Sound/SoundCue.h
@@ -77,14 +77,6 @@ class ENGINE_API USoundCue : public USoundBase
 {
 	GENERATED_UCLASS_BODY()
 
-	/* Indicates whether attenuation should use the Attenuation Overrides or the Attenuation Settings asset */
-	UPROPERTY(EditAnywhere, Category=Attenuation)
-	uint32 bOverrideAttenuation:1;
-
-	/* Makes this sound cue ignore per-platform random node culling for memory purposes */
-	UPROPERTY(EditAnywhere, Category=Culling)
-	uint32 bExcludeFromRandomNodeBranchCulling:1;
-
 	UPROPERTY()
 	USoundNode* FirstNode;
 
@@ -116,6 +108,16 @@ protected:
 private:
 	float MaxAudibleDistance;
 
+public:
+	/* Indicates whether attenuation should use the Attenuation Overrides or the Attenuation Settings asset */
+	UPROPERTY(EditAnywhere, Category=Attenuation)
+	uint8 bOverrideAttenuation : 1;
+
+	/* Makes this sound cue ignore per-platform random node culling for memory purposes */
+	UPROPERTY(EditAnywhere, Category=Culling)
+	uint8 bExcludeFromRandomNodeBranchCulling : 1;
+
+private:
 	/** Whether a sound has play when silent enabled (i.e. for a sound cue, if any sound wave player has it enabled). */
 	UPROPERTY()
 	uint8 bHasPlayWhenSilent : 1;
diff --git a/Engine/Source/Runtime/Engine/Classes/Sound/SoundNodeRandom.h b/Engine/Source/Runtime/Engine/Classes/Sound/SoundNodeRandom.h
--- a/Engine/Source/Runtime/Engine/Classes/Sound/SoundNodeRandom.h
+++ b/Engine/Source/Runtime/Engine/Classes/Sound/SoundNodeRandom.h
@@ -24,6 +24,17 @@ class USoundNodeRandom : public USoundNode
 	UPROPERTY(EditAnywhere, editfixedsize, Category=Random)
 	TArray<float> Weights;
 
+	/**
+	 * Internal state of which sounds have been played.  This is only used at runtime
+	 * to keep track of which sounds have been played
+	 */
+	UPROPERTY(transient)
+	TArray<bool> HasBeenUsed;
+
+	/** Counter var so we don't have to count all of the used sounds each time we choose a sound **/
+	UPROPERTY(transient)
+	int32 NumRandomUsed;
+
 	/** If greater than 0, then upon each level load such a number of inputs will be randomly selected
 	 *  and the rest will be removed. This can be used to cut down the memory usage of large randomizing
 	 *  cues.
@@ -53,18 +64,7 @@ class USoundNodeRandom : public USoundNode
 	 * So one could play the same sound over and over if the probabilities don't go your way :-)
 	 */
 	UPROPERTY(EditAnywhere, Category=Random)
-	uint32 bRandomizeWithoutReplacement:1;
-
-	/**
-	 * Internal state of which sounds have been played.  This is only used at runtime
-	 * to keep track of which sounds have been played
-	 */
-	UPROPERTY(transient)
-	TArray<bool> HasBeenUsed;
-
-	/** Counter var so we don't have to count all of the used sounds each time we choose a sound **/
-	UPROPERTY(transient)
-	int32 NumRandomUsed;
+	uint8 bRandomizeWithoutReplacement : 1;
 
 #if WITH_EDITORONLY_DATA
 	/** Editor only list of nodes hidden to duplicate behavior of PreselectAtLevelLoad */
@@ -79,7 +79,7 @@ public:
 
 	//~ Begin USoundNode Interface.
 	virtual void ParseNodes( FAudioDevice* AudioDevice, const UPTRINT NodeWaveInstanceHash, FActiveSound& ActiveSound, const FSoundParseParameters& ParseParams, TArray<FWaveInstance*>& WaveInstances ) override;
-	virtual int32 GetNumSounds(const UPTRINT NodeWaveInstanceHash, FActiveSound& ActiveSound) const;
+	virtual int32 GetNumSounds(const UPTRINT NodeWaveInstanceHash, FActiveSound& ActiveSound) const override;
 	virtual int32 GetMaxChildNodes() const override 
 	{ 
 		return MAX_ALLOWED_CHILD_NODES; 
diff --git a/Engine/Source/Runtime/Engine/Private/Actor.cpp b/Engine/Source/Runtime/Engine/Private/Actor.cpp
--- a/Engine/Source/Runtime/Engine/Private/Actor.cpp
+++ b/Engine/Source/Runtime/Engine/Private/Actor.cpp
@@ -710,6 +710,9 @@ void AActor::PostLoad()
 		}
 	}
 #endif // WITH_EDITORONLY_DATA
+
+	// Since the actor is being loading, it finished spawning by definition when it was originally spawned, so set to true now
+	bHasFinishedSpawning = true;
 }
 
 void AActor::PostLoadSubobjects(FObjectInstancingGraph* OuterInstanceGraph)
diff --git a/Engine/Source/Runtime/Engine/Private/AssetManager.cpp b/Engine/Source/Runtime/Engine/Private/AssetManager.cpp
--- a/Engine/Source/Runtime/Engine/Private/AssetManager.cpp
+++ b/Engine/Source/Runtime/Engine/Private/AssetManager.cpp
@@ -2629,6 +2629,8 @@ void UAssetManager::ScanPrimaryAssetTypesFromConfig()
 			continue;
 		}
 
+		UE_CLOG(AssetTypeMap.Find(TypeInfo.PrimaryAssetType), LogAssetManager, Error, TEXT("Found multiple \"%s\" Primary Asset Type entries in \"Primary Asset Types To Scan\" config. Only a single entry per type is supported."), *TypeInfo.PrimaryAssetType.ToString());
+
 		ScanPathsForPrimaryAssets(TypeInfo.PrimaryAssetType, TypeInfo.AssetScanPaths, TypeInfo.AssetBaseClassLoaded, TypeInfo.bHasBlueprintClasses, TypeInfo.bIsEditorOnly, false);
 
 		SetPrimaryAssetTypeRules(TypeInfo.PrimaryAssetType, TypeInfo.Rules);
diff --git a/Engine/Source/Runtime/Engine/Private/BlueprintGeneratedClass.cpp b/Engine/Source/Runtime/Engine/Private/BlueprintGeneratedClass.cpp
--- a/Engine/Source/Runtime/Engine/Private/BlueprintGeneratedClass.cpp
+++ b/Engine/Source/Runtime/Engine/Private/BlueprintGeneratedClass.cpp
@@ -5,6 +5,7 @@
 #include "Stats/StatsMisc.h"
 #include "UObject/UObjectHash.h"
 #include "UObject/CoreNet.h"
+#include "UObject/CoreRedirects.h"
 #include "UObject/Package.h"
 #include "UObject/LinkerLoad.h"
 #include "Serialization/ObjectReader.h"
@@ -127,6 +128,28 @@ void UBlueprintGeneratedClass::PostLoad()
 	}
 #endif // WITH_EDITORONLY_DATA
 
+	// Update any component names that have been redirected
+	if (!FPlatformProperties::RequiresCookedData())
+	{
+		for (FBPComponentClassOverride& Override : ComponentClassOverrides)
+		{
+			const FString ComponentName = Override.ComponentName.ToString();
+			UClass* ClassToCheck = this;
+			while (ClassToCheck)
+			{
+				if (const TMap<FString, FString>* ValueChanges = FCoreRedirects::GetValueRedirects(ECoreRedirectFlags::Type_Class, ClassToCheck))
+				{
+					if (const FString* NewComponentName = ValueChanges->Find(ComponentName))
+					{
+						Override.ComponentName = **NewComponentName;
+						break;
+					}
+				}
+				ClassToCheck = ClassToCheck->GetSuperClass();
+			}
+		}
+	}
+
 	AssembleReferenceTokenStream(true);
 }
 
@@ -603,6 +626,16 @@ void UBlueprintGeneratedClass::UpdateCustomPropertyListForPostConstruction()
 	bCustomPropertyListForPostConstructionInitialized = true;
 }
 
+void UBlueprintGeneratedClass::SetupObjectInitializer(FObjectInitializer& ObjectInitializer) const
+{
+	for (const FBPComponentClassOverride& Override : ComponentClassOverrides)
+	{
+		ObjectInitializer.SetDefaultSubobjectClass(Override.ComponentName, Override.ComponentClass);
+	}
+
+	GetSuperClass()->SetupObjectInitializer(ObjectInitializer);
+}
+
 void UBlueprintGeneratedClass::InitPropertiesFromCustomList(uint8* DataPtr, const uint8* DefaultDataPtr)
 {
 	FScopeLock SerializeAndPostLoadLock(&SerializeAndPostLoadCritical);
@@ -712,7 +745,7 @@ UInheritableComponentHandler* UBlueprintGeneratedClass::GetInheritableComponentH
 }
 
 
-UObject* UBlueprintGeneratedClass::FindArchetype(UClass* ArchetypeClass, const FName ArchetypeName) const
+UObject* UBlueprintGeneratedClass::FindArchetype(const UClass* ArchetypeClass, const FName ArchetypeName) const
 {
 	UObject* Archetype = nullptr;
 
@@ -1428,6 +1461,15 @@ void UBlueprintGeneratedClass::GetDefaultObjectPreloadDependencies(TArray<UObjec
 			OutDeps.Add(ComponentTemplate);
 		}
 	}
+
+	// Add the classes that will be used for overriding components defined in base classes
+	for (const FBPComponentClassOverride& Override : ComponentClassOverrides)
+	{
+		if (Override.ComponentClass)
+		{
+			OutDeps.Add(Override.ComponentClass);
+		}
+	}
 }
 
 bool UBlueprintGeneratedClass::NeedsLoadForServer() const
diff --git a/Engine/Source/Runtime/Engine/Private/GameViewportClient.cpp b/Engine/Source/Runtime/Engine/Private/GameViewportClient.cpp
--- a/Engine/Source/Runtime/Engine/Private/GameViewportClient.cpp
+++ b/Engine/Source/Runtime/Engine/Private/GameViewportClient.cpp
@@ -2159,7 +2159,7 @@ void UGameViewportClient::LayoutPlayers()
 	}
 }
 
-void UGameViewportClient::SetDisableSplitscreenOverride( const bool bDisabled )
+void UGameViewportClient::SetForceDisableSplitscreen(const bool bDisabled)
 {
 	bDisableSplitScreenOverride = bDisabled;
 	LayoutPlayers();
diff --git a/Engine/Source/Runtime/Engine/Private/GameplayStatics.cpp b/Engine/Source/Runtime/Engine/Private/GameplayStatics.cpp
--- a/Engine/Source/Runtime/Engine/Private/GameplayStatics.cpp
+++ b/Engine/Source/Runtime/Engine/Private/GameplayStatics.cpp
@@ -365,6 +365,33 @@ bool UGameplayStatics::IsGamePaused(const UObject* WorldContextObject)
 	return World ? World->IsPaused() : false;
 }
 
+void UGameplayStatics::SetForceDisableSplitscreen(const UObject* WorldContextObject, bool bDisable)
+{
+	UWorld* World = GEngine->GetWorldFromContextObject(WorldContextObject, EGetWorldErrorMode::LogAndReturnNull);
+	if (World)
+	{
+		UGameViewportClient* GameViewportClient = World->GetGameViewport();
+		if (GameViewportClient)
+		{
+			GameViewportClient->SetForceDisableSplitscreen(bDisable);
+		}
+	}
+}
+
+bool UGameplayStatics::IsSplitscreenForceDisabled(const UObject* WorldContextObject)
+{
+	UWorld* const World = GEngine->GetWorldFromContextObject(WorldContextObject, EGetWorldErrorMode::LogAndReturnNull);
+	if (World)
+	{
+		UGameViewportClient* GameViewportClient = World->GetGameViewport();
+		if (GameViewportClient)
+		{
+			return GameViewportClient->IsSplitscreenForceDisabled();
+		}
+	}
+	return false;
+}
+
 void UGameplayStatics::SetEnableWorldRendering(const UObject* WorldContextObject, bool bEnable)
 {
 	UWorld* World = GEngine->GetWorldFromContextObject(WorldContextObject, EGetWorldErrorMode::LogAndReturnNull);
diff --git a/Engine/Source/Runtime/Engine/Private/InstancedStaticMesh.cpp b/Engine/Source/Runtime/Engine/Private/InstancedStaticMesh.cpp
--- a/Engine/Source/Runtime/Engine/Private/InstancedStaticMesh.cpp
+++ b/Engine/Source/Runtime/Engine/Private/InstancedStaticMesh.cpp
@@ -1524,6 +1524,11 @@ void UInstancedStaticMeshComponent::CreateAllInstanceBodies()
 	{
 		FPhysScene* PhysScene = GetWorld()->GetPhysicsScene();
 
+		if (!BodyInstance.GetOverrideWalkableSlopeOnInstance())
+		{
+			BodyInstance.SetWalkableSlopeOverride(BodySetup->WalkableSlopeOverride, false);
+		}
+
 		InstanceBodies.SetNumUninitialized(NumBodies);
 
 		// Sanitized array does not contain any nulls
diff --git a/Engine/Source/Runtime/Engine/Private/LevelScriptActor.cpp b/Engine/Source/Runtime/Engine/Private/LevelScriptActor.cpp
--- a/Engine/Source/Runtime/Engine/Private/LevelScriptActor.cpp
+++ b/Engine/Source/Runtime/Engine/Private/LevelScriptActor.cpp
@@ -32,7 +32,7 @@ ALevelScriptActor::ALevelScriptActor(const FObjectInitializer& ObjectInitializer
 
 void ALevelScriptActor::PreInitializeComponents()
 {
-	if (UInputDelegateBinding::SupportsInputDelegate(GetClass()))
+	if (UInputDelegateBinding::SupportsInputDelegate(GetClass()) && !InputComponent)
 	{
 		// create an InputComponent object so that the level script actor can bind key events
 		InputComponent = NewObject<UInputComponent>(this);
diff --git a/Engine/Source/Runtime/Engine/Private/Particles/ParticleComponents.cpp b/Engine/Source/Runtime/Engine/Private/Particles/ParticleComponents.cpp
--- a/Engine/Source/Runtime/Engine/Private/Particles/ParticleComponents.cpp
+++ b/Engine/Source/Runtime/Engine/Private/Particles/ParticleComponents.cpp
@@ -2085,9 +2085,10 @@ UFXSystemAsset::UFXSystemAsset(const FObjectInitializer& ObjectInitializer)
 UParticleSystem::UParticleSystem(const FObjectInitializer& ObjectInitializer)
 	: Super(ObjectInitializer)
 	, OcclusionBoundsMethod(EPSOBM_ParticleBounds)
+	, bAnyEmitterLoopsForever(false)
 	, HighestSignificance(EParticleSignificanceLevel::Critical)
 	, LowestSignificance(EParticleSignificanceLevel::Low)
-	, bAnyEmitterLoopsForever(false)
+	, bShouldManageSignificance(false)
 	, bIsImmortal(false)
 	, bWillBecomeZombie(false)
 {
@@ -2127,7 +2128,6 @@ UParticleSystem::UParticleSystem(const FObjectInitializer& ObjectInitializer)
 	InsignificanceDelay = 0.0f;
 	MaxSignificanceLevel = EParticleSignificanceLevel::Critical;
 	MaxPoolSize = 32;
-	bShouldManageSignificance = false;
 
 
 	bAllowManagedTicking = true;
diff --git a/Engine/Source/Runtime/Engine/Private/Particles/ParticleModules.cpp b/Engine/Source/Runtime/Engine/Private/Particles/ParticleModules.cpp
--- a/Engine/Source/Runtime/Engine/Private/Particles/ParticleModules.cpp
+++ b/Engine/Source/Runtime/Engine/Private/Particles/ParticleModules.cpp
@@ -158,36 +158,15 @@ void UParticleModule::FinalUpdate(FParticleEmitterInstance* Owner, int32 Offset,
 }
 
 uint32 UParticleModule::RequiredBytes(UParticleModuleTypeDataBase* TypeData)
-{
-PRAGMA_DISABLE_DEPRECATION_WARNINGS
-		// To try and handle deprecation, if function is not overriden, call old function
-	// Will fail though if actually trying to use Owner
-	return RequiredBytes(static_cast<FParticleEmitterInstance*>(nullptr));
-PRAGMA_ENABLE_DEPRECATION_WARNINGS
-}
-
-uint32 UParticleModule::RequiredBytesPerInstance()
-{
-PRAGMA_DISABLE_DEPRECATION_WARNINGS
-	// To try and handle deprecation, if function is not overriden, call old function
-	// Will fail though if actually trying to use Owner
-	return RequiredBytesPerInstance(nullptr);
-PRAGMA_ENABLE_DEPRECATION_WARNINGS
-}
-
-// DEPRECATED 4.11
-uint32 UParticleModule::RequiredBytes(FParticleEmitterInstance* Owner)
 {
 	return 0;
 }
 
-// DEPRECATED 4.11
-uint32 UParticleModule::RequiredBytesPerInstance(FParticleEmitterInstance* Owner)
+uint32 UParticleModule::RequiredBytesPerInstance()
 {
 	return 0;
 }
 
-
 uint32 UParticleModule::PrepPerInstanceBlock(FParticleEmitterInstance* Owner, void* InstData)
 {
 	return 0xffffffff;
@@ -3426,7 +3405,7 @@ void UParticleModuleTypeDataMesh::Serialize(FArchive& Ar)
 	Super::Serialize(Ar);
 	if (Ar.IsLoading() && Ar.UE4Ver() < VER_UE4_MESH_EMITTER_INITIAL_ORIENTATION_DISTRIBUTION)
 	{
-		FVector oldOrient(Roll_DEPRECATED, Pitch_DEPRECATED, Yaw_DEPRECATED);
+		FVector oldOrient(0.0f, 0.0f, 0.0f);
 		CreateDistribution();
 		UDistributionVectorUniform* RPYDistribution = Cast<UDistributionVectorUniform>(RollPitchYawRange.Distribution);
 		RPYDistribution->Min = oldOrient;
diff --git a/Engine/Source/Runtime/Engine/Private/Pawn.cpp b/Engine/Source/Runtime/Engine/Private/Pawn.cpp
--- a/Engine/Source/Runtime/Engine/Private/Pawn.cpp
+++ b/Engine/Source/Runtime/Engine/Private/Pawn.cpp
@@ -443,6 +443,11 @@ bool APawn::IsControlled() const
 	return(PC != nullptr);
 }
 
+bool APawn::IsPawnControlled() const
+{
+	return (Controller != nullptr);
+}
+
 FRotator APawn::GetControlRotation() const
 {
 	return Controller ? Controller->GetControlRotation() : FRotator::ZeroRotator;
@@ -488,6 +493,8 @@ void APawn::SetPlayerState(APlayerState* NewPlayerState)
 
 void APawn::PossessedBy(AController* NewController)
 {
+	SetOwner(NewController);
+	
 	AController* const OldController = Controller;
 
 	Controller = NewController;
diff --git a/Engine/Source/Runtime/Engine/Private/PhysicsEngine/BodyInstance.cpp b/Engine/Source/Runtime/Engine/Private/PhysicsEngine/BodyInstance.cpp
--- a/Engine/Source/Runtime/Engine/Private/PhysicsEngine/BodyInstance.cpp
+++ b/Engine/Source/Runtime/Engine/Private/PhysicsEngine/BodyInstance.cpp
@@ -2303,13 +2303,18 @@ const FWalkableSlopeOverride& FBodyInstance::GetWalkableSlopeOverride() const
 	}
 }
 
-void FBodyInstance::SetWalkableSlopeOverride(const FWalkableSlopeOverride& NewOverride)
+void FBodyInstance::SetWalkableSlopeOverride(const FWalkableSlopeOverride& NewOverride, bool bNewOverideSetting)
 {
-	bOverrideWalkableSlopeOnInstance = true;
+	bOverrideWalkableSlopeOnInstance = bNewOverideSetting;
 	WalkableSlopeOverride = NewOverride;
 }
 
 
+bool FBodyInstance::GetOverrideWalkableSlopeOnInstance() const
+{
+	return bOverrideWalkableSlopeOnInstance;
+}
+
 /** 
 *	Changes the current PhysMaterialOverride for this body. 
 *	Note that if physics is already running on this component, this will _not_ alter its mass/inertia etc, it will only change its 
diff --git a/Engine/Source/Runtime/Engine/Public/Animation/AnimCurveTypes.h b/Engine/Source/Runtime/Engine/Public/Animation/AnimCurveTypes.h
--- a/Engine/Source/Runtime/Engine/Public/Animation/AnimCurveTypes.h
+++ b/Engine/Source/Runtime/Engine/Public/Animation/AnimCurveTypes.h
@@ -14,6 +14,7 @@
 UENUM(BlueprintType, meta=(Bitflags))
 enum EAnimAssetCurveFlags
 {
+	AACF_NONE = 0 UMETA(Hidden),
 	// Used as morph target curve
 	AACF_DriveMorphTarget_DEPRECATED = 0x00000001 UMETA(Hidden), // This has moved to FAnimCurveType:bMorphTarget. Set per skeleton. DO NOT REMOVE UNTIL FrameworkObjectVersion.MoveCurveTypesToSkeleton expires.
 	// Used as triggering event
diff --git a/Engine/Source/Runtime/Engine/Public/Distributions.h b/Engine/Source/Runtime/Engine/Public/Distributions.h
--- a/Engine/Source/Runtime/Engine/Public/Distributions.h
+++ b/Engine/Source/Runtime/Engine/Public/Distributions.h
@@ -28,6 +28,12 @@ enum ERawDistributionOperation
  */
 struct FDistributionLookupTable
 {
+	/** Time between values in the lookup table */
+	float TimeScale;
+	/** Absolute time of the first value */
+	float TimeBias;
+	/** Values in the table. */
+	TArray<float> Values;
 	/** Operation for which the table was built. */
 	uint8 Op;
 	/** Number of entries in the table. */
@@ -36,22 +42,17 @@ struct FDistributionLookupTable
 	uint8 EntryStride;
 	/** Number of values between sub-entries [0,4]. */
 	uint8 SubEntryStride;
-	/** Time between values in the lookup table */
-	float TimeScale;
-	/** Absolute time of the first value */
-	float TimeBias;
-	/** Values in the table. */
-	TArray<float> Values;
 	/** Lock axes flag for vector distributions. */
 	uint8 LockFlag;
+	
 	/** Default constructor. */
 	FDistributionLookupTable()
-		: Op(RDO_Uninitialized)
+		: TimeScale(0.0f)
+		, TimeBias(0.0f)
+		, Op(RDO_Uninitialized)
 		, EntryCount(0)
 		, EntryStride(0)
 		, SubEntryStride(0)
-		, TimeScale(0.0f)
-		, TimeBias(0.0f)
 		, LockFlag(0)
 	{}
 
diff --git a/Engine/Source/Runtime/Experimental/GeometryCollectionSimulationCore/Public/GeometryCollection/GeometryCollectionSimulationTypes.h b/Engine/Source/Runtime/Experimental/GeometryCollectionSimulationCore/Public/GeometryCollection/GeometryCollectionSimulationTypes.h
--- a/Engine/Source/Runtime/Experimental/GeometryCollectionSimulationCore/Public/GeometryCollection/GeometryCollectionSimulationTypes.h
+++ b/Engine/Source/Runtime/Experimental/GeometryCollectionSimulationCore/Public/GeometryCollection/GeometryCollectionSimulationTypes.h
@@ -30,6 +30,7 @@ enum class EImplicitTypeEnum : uint8
 UENUM(BlueprintType)
 enum class EObjectStateTypeEnum : uint8
 {
+	Chaos_NONE = 0 UMETA(Hidden),
 	Chaos_Object_Sleeping  = 1 /*Chaos::EObjectStateType::Sleeping*/   UMETA(DisplayName = "Sleeping"),
 	Chaos_Object_Kinematic = 2 /*Chaos::EObjectStateType::Kinematic*/  UMETA(DisplayName = "Kinematic"),
 	Chaos_Object_Static = 3    /*Chaos::EObjectStateType::Static*/     UMETA(DisplayName = "Static"),
diff --git a/Engine/Source/Runtime/Foliage/Public/ProceduralFoliageInstance.h b/Engine/Source/Runtime/Foliage/Public/ProceduralFoliageInstance.h
--- a/Engine/Source/Runtime/Foliage/Public/ProceduralFoliageInstance.h
+++ b/Engine/Source/Runtime/Foliage/Public/ProceduralFoliageInstance.h
@@ -37,6 +37,7 @@ namespace ESimulationQuery
 {
 	enum Type
 	{
+		None = 0 UMETA(Hidden),
 		/*Instances overlap with collision*/
 		CollisionOverlap = 1,
 		/*Instances overlap with shade*/
diff --git a/Engine/Source/Runtime/GameplayTags/Classes/GameplayTagContainer.h b/Engine/Source/Runtime/GameplayTags/Classes/GameplayTagContainer.h
--- a/Engine/Source/Runtime/GameplayTags/Classes/GameplayTagContainer.h
+++ b/Engine/Source/Runtime/GameplayTags/Classes/GameplayTagContainer.h
@@ -985,24 +985,26 @@ public:
 	FGameplayTagQuery& operator=(FGameplayTagQuery&& Other);
 
 private:
+	// Note: Properties need to be editable to allow FComponentPropertyWriter to serialize them, but are hidden in the editor by the customizations mentioned above.
+
 	/** Versioning for future token stream protocol changes. See EGameplayTagQueryStreamVersion. */
-	UPROPERTY()
+	UPROPERTY(EditAnywhere, Category = Hidden)
 	int32 TokenStreamVersion;
 
 	/** List of tags referenced by this entire query. Token stream stored indices into this list. */
-	UPROPERTY()
+	UPROPERTY(EditAnywhere, Category = Hidden)
 	TArray<FGameplayTag> TagDictionary;
 
 	/** Stream representation of the actual hierarchical query */
-	UPROPERTY()
+	UPROPERTY(EditAnywhere, Category = Hidden)
 	TArray<uint8> QueryTokenStream;
 
 	/** User-provided string describing the query */
-	UPROPERTY()
+	UPROPERTY(EditAnywhere, Category = Hidden)
 	FString UserDescription;
 
 	/** Auto-generated string describing the query */
-	UPROPERTY()
+	UPROPERTY(EditAnywhere, Category = Hidden)
 	FString AutoDescription;
 
 	/** Returns a gameplay tag from the tag dictionary */
diff --git a/Engine/Source/Runtime/HeadMountedDisplay/Private/HeadMountedDisplayConsoleCommands.cpp b/Engine/Source/Runtime/HeadMountedDisplay/Private/HeadMountedDisplayConsoleCommands.cpp
--- a/Engine/Source/Runtime/HeadMountedDisplay/Private/HeadMountedDisplayConsoleCommands.cpp
+++ b/Engine/Source/Runtime/HeadMountedDisplay/Private/HeadMountedDisplayConsoleCommands.cpp
@@ -57,9 +57,10 @@ static void DrawDebugTrackingSensorLocations(UCanvas* Canvas, APlayerController*
 		return;
 	}
 
+	UWorld* World = (PlayerController ? PlayerController->GetWorld() : GWorld);
 	if (!PlayerController)
 	{
-		PlayerController = GWorld->GetFirstPlayerController();
+		PlayerController = World->GetFirstPlayerController();
 		if (!PlayerController)
 		{
 			return;
@@ -133,30 +134,30 @@ static void DrawDebugTrackingSensorLocations(UCanvas* Canvas, APlayerController*
 		FVector NearBR = Matrix.TransformPosition(EdgeBR * SensorProperties.NearPlane);
 
 		// Draw a point at the sensor position
-		DrawDebugPoint(GWorld, Tip, 5, FrustrumColor);
+		DrawDebugPoint(World, Tip, 5, FrustrumColor);
 
 		// Draw the four edges of the pyramid
-		DrawDebugLine(GWorld, Tip, BaseTR, FrustrumColor);
-		DrawDebugLine(GWorld, Tip, BaseTL, FrustrumColor);
-		DrawDebugLine(GWorld, Tip, BaseBL, FrustrumColor);
-		DrawDebugLine(GWorld, Tip, BaseBR, FrustrumColor);
+		DrawDebugLine(World, Tip, BaseTR, FrustrumColor);
+		DrawDebugLine(World, Tip, BaseTL, FrustrumColor);
+		DrawDebugLine(World, Tip, BaseBL, FrustrumColor);
+		DrawDebugLine(World, Tip, BaseBR, FrustrumColor);
 
 		// Draw the base (far plane)
-		DrawDebugLine(GWorld, BaseTR, BaseTL, FrustrumColor);
-		DrawDebugLine(GWorld, BaseTL, BaseBL, FrustrumColor);
-		DrawDebugLine(GWorld, BaseBL, BaseBR, FrustrumColor);
-		DrawDebugLine(GWorld, BaseBR, BaseTR, FrustrumColor);
+		DrawDebugLine(World, BaseTR, BaseTL, FrustrumColor);
+		DrawDebugLine(World, BaseTL, BaseBL, FrustrumColor);
+		DrawDebugLine(World, BaseBL, BaseBR, FrustrumColor);
+		DrawDebugLine(World, BaseBR, BaseTR, FrustrumColor);
 
 		// Draw the near plane
-		DrawDebugLine(GWorld, NearTR, NearTL, FrustrumColor);
-		DrawDebugLine(GWorld, NearTL, NearBL, FrustrumColor);
-		DrawDebugLine(GWorld, NearBL, NearBR, FrustrumColor);
-		DrawDebugLine(GWorld, NearBR, NearTR, FrustrumColor);
+		DrawDebugLine(World, NearTR, NearTL, FrustrumColor);
+		DrawDebugLine(World, NearTL, NearBL, FrustrumColor);
+		DrawDebugLine(World, NearBL, NearBR, FrustrumColor);
+		DrawDebugLine(World, NearBR, NearTR, FrustrumColor);
 
 		// Draw a center line from the sensor to the focal point
 		FVector CenterLine = Matrix.TransformPosition(FVector(SensorProperties.CameraDistance, 0, 0));
-		DrawDebugLine(GWorld, Tip, CenterLine, CenterLineColor);
-		DrawDebugPoint(GWorld, CenterLine, 5, CenterLineColor);
+		DrawDebugLine(World, Tip, CenterLine, CenterLineColor);
+		DrawDebugPoint(World, CenterLine, 5, CenterLineColor);
 	}
 }
 
diff --git a/Engine/Source/Runtime/ImageCore/Public/ImageCore.h b/Engine/Source/Runtime/ImageCore/Public/ImageCore.h
--- a/Engine/Source/Runtime/ImageCore/Public/ImageCore.h
+++ b/Engine/Source/Runtime/ImageCore/Public/ImageCore.h
@@ -16,7 +16,7 @@ namespace ERawImageFormat
 	/**
 	 * Enumerates supported raw image formats.
 	 */
-	enum Type
+	enum Type : uint8
 	{
 		G8,
 		BGRA8,
diff --git a/Engine/Source/Runtime/MediaAssets/Public/Misc/MediaBlueprintFunctionLibrary.h b/Engine/Source/Runtime/MediaAssets/Public/Misc/MediaBlueprintFunctionLibrary.h
--- a/Engine/Source/Runtime/MediaAssets/Public/Misc/MediaBlueprintFunctionLibrary.h
+++ b/Engine/Source/Runtime/MediaAssets/Public/Misc/MediaBlueprintFunctionLibrary.h
@@ -15,6 +15,8 @@
 UENUM(BlueprintType, meta=(BitFlags, UseEnumValuesAsMaskValuesInEditor = "true"))
 enum class EMediaAudioCaptureDeviceFilter : uint8
 {
+	None = 0 UMETA(Hidden),
+
 	/** Audio capture cards. */
 	Card = 0x1,
 
@@ -37,6 +39,8 @@ ENUM_CLASS_FLAGS(EMediaAudioCaptureDeviceFilter)
 UENUM(BlueprintType, meta=(BitFlags, UseEnumValuesAsMaskValuesInEditor = "true"))
 enum class EMediaVideoCaptureDeviceFilter : uint8
 {
+	None = 0 UMETA(Hidden),
+
 	/** Video capture card. */
 	Card = 0x1,
 
@@ -59,6 +63,8 @@ ENUM_CLASS_FLAGS(EMediaVideoCaptureDeviceFilter)
 UENUM(BlueprintType, meta=(BitFlags, UseEnumValuesAsMaskValuesInEditor = "true"))
 enum class EMediaWebcamCaptureDeviceFilter : uint8
 {
+	None = 0 UMETA(Hidden),
+
 	/** Depth sensor. */
 	DepthSensor = 0x1,
 
diff --git a/Engine/Source/Runtime/MovieScene/Public/Evaluation/Blending/MovieSceneBlendType.h b/Engine/Source/Runtime/MovieScene/Public/Evaluation/Blending/MovieSceneBlendType.h
--- a/Engine/Source/Runtime/MovieScene/Public/Evaluation/Blending/MovieSceneBlendType.h
+++ b/Engine/Source/Runtime/MovieScene/Public/Evaluation/Blending/MovieSceneBlendType.h
@@ -10,6 +10,8 @@
 UENUM()
 enum class EMovieSceneBlendType : uint8
 {
+	None = 0 UMETA(Hidden),
+
 	/** Blends all other weighted values together as an average of the total weight */
 	Absolute			= 0x1,
 	/** Applies this value as a sum total of all other additives */
diff --git a/Engine/Source/Runtime/UMG/Public/Components/HorizontalBoxSlot.h b/Engine/Source/Runtime/UMG/Public/Components/HorizontalBoxSlot.h
--- a/Engine/Source/Runtime/UMG/Public/Components/HorizontalBoxSlot.h
+++ b/Engine/Source/Runtime/UMG/Public/Components/HorizontalBoxSlot.h
@@ -17,6 +17,9 @@ class UMG_API UHorizontalBoxSlot : public UPanelSlot
 {
 	GENERATED_UCLASS_BODY()
 
+private:
+	SHorizontalBox::FSlot* Slot;
+
 public:
 	
 	/** The amount of padding between the slots parent and the content. */
@@ -33,8 +36,6 @@ public:
 	UPROPERTY(EditAnywhere, BlueprintReadOnly, Category="Layout|Horizontal Box Slot")
 	TEnumAsByte<EVerticalAlignment> VerticalAlignment;
 
-public:
-
 	UFUNCTION(BlueprintCallable, Category="Layout|Horizontal Box Slot")
 	void SetPadding(FMargin InPadding);
 
@@ -55,6 +56,4 @@ public:
 
 	virtual void ReleaseSlateResources(bool bReleaseChildren) override;
 
-private:
-	SHorizontalBox::FSlot* Slot;
 };
diff --git a/Engine/Source/Runtime/UMG/Public/Components/SizeBox.h b/Engine/Source/Runtime/UMG/Public/Components/SizeBox.h
--- a/Engine/Source/Runtime/UMG/Public/Components/SizeBox.h
+++ b/Engine/Source/Runtime/UMG/Public/Components/SizeBox.h
@@ -22,116 +22,117 @@ class UMG_API USizeBox : public UContentWidget
 {
 	GENERATED_UCLASS_BODY()
 
-public:
-
-	/**  */
-	UPROPERTY(EditAnywhere, Category="Child Layout", meta=(InlineEditConditionToggle))
-		uint32 bOverride_WidthOverride : 1;
-
-	/**  */
-	UPROPERTY(EditAnywhere, Category="Child Layout", meta=(InlineEditConditionToggle))
-		uint32 bOverride_HeightOverride : 1;
-
-	/**  */
-	UPROPERTY(EditAnywhere, Category="Child Layout", meta=(InlineEditConditionToggle))
-		uint32 bOverride_MinDesiredWidth : 1;
-
-	/**  */
-	UPROPERTY(EditAnywhere, Category="Child Layout", meta=(InlineEditConditionToggle))
-		uint32 bOverride_MinDesiredHeight : 1;
-
-	/**  */
-	UPROPERTY(EditAnywhere, Category="Child Layout", meta=(InlineEditConditionToggle))
-		uint32 bOverride_MaxDesiredWidth : 1;
-
-	/**  */
-	UPROPERTY(EditAnywhere, Category="Child Layout", meta=(InlineEditConditionToggle))
-		uint32 bOverride_MaxDesiredHeight : 1;
-
-	/**  */
-	UPROPERTY(EditAnywhere, Category="Child Layout", meta=(InlineEditConditionToggle))
-	uint32 bOverride_MinAspectRatio : 1;
-
-	/**  */
-	UPROPERTY(EditAnywhere, Category = "Child Layout", meta=(InlineEditConditionToggle))
-	uint32 bOverride_MaxAspectRatio : 1;
+protected:
+	TSharedPtr<SBox> MySizeBox;
 
+public:
 
 	/** When specified, ignore the content's desired size and report the WidthOverride as the Box's desired width. */
 	UPROPERTY(EditAnywhere, BlueprintReadOnly, Category="Child Layout", meta=( editcondition="bOverride_WidthOverride" ))
-		float WidthOverride;
+	float WidthOverride;
 
 	/** When specified, ignore the content's desired size and report the HeightOverride as the Box's desired height. */
 	UPROPERTY(EditAnywhere, BlueprintReadOnly, Category="Child Layout", meta=( editcondition="bOverride_HeightOverride" ))
-		float HeightOverride;
+	float HeightOverride;
 
 	/** When specified, will report the MinDesiredWidth if larger than the content's desired width. */
 	UPROPERTY(EditAnywhere, BlueprintReadOnly, Category="Child Layout", meta=( editcondition="bOverride_MinDesiredWidth" ))
-		float MinDesiredWidth;
+	float MinDesiredWidth;
 
 	/** When specified, will report the MinDesiredHeight if larger than the content's desired height. */
 	UPROPERTY(EditAnywhere, BlueprintReadOnly, Category="Child Layout", meta=( editcondition="bOverride_MinDesiredHeight" ))
-		float MinDesiredHeight;
+	float MinDesiredHeight;
 
 	/** When specified, will report the MaxDesiredWidth if smaller than the content's desired width. */
 	UPROPERTY(EditAnywhere, BlueprintReadOnly, Category="Child Layout", meta=( editcondition="bOverride_MaxDesiredWidth" ))
-		float MaxDesiredWidth;
+	float MaxDesiredWidth;
 
 	/** When specified, will report the MaxDesiredHeight if smaller than the content's desired height. */
 	UPROPERTY(EditAnywhere, BlueprintReadOnly, Category="Child Layout", meta=( editcondition="bOverride_MaxDesiredHeight" ))
-		float MaxDesiredHeight;
+	float MaxDesiredHeight;
 
 	/** */
 	UPROPERTY(EditAnywhere, BlueprintReadOnly, Category = "Child Layout", meta = (editcondition = "bOverride_MinAspectRatio"))
-		float MinAspectRatio;
+	float MinAspectRatio;
 
 	/** */
 	UPROPERTY(EditAnywhere, BlueprintReadOnly, Category="Child Layout", meta=( editcondition="bOverride_MaxAspectRatio" ))
-		float MaxAspectRatio;
+	float MaxAspectRatio;
+
+	/**  */
+	UPROPERTY(EditAnywhere, Category="Child Layout", meta=(InlineEditConditionToggle))
+	uint8 bOverride_WidthOverride : 1;
+
+	/**  */
+	UPROPERTY(EditAnywhere, Category="Child Layout", meta=(InlineEditConditionToggle))
+	uint8 bOverride_HeightOverride : 1;
+
+	/**  */
+	UPROPERTY(EditAnywhere, Category="Child Layout", meta=(InlineEditConditionToggle))
+	uint8 bOverride_MinDesiredWidth : 1;
+
+	/**  */
+	UPROPERTY(EditAnywhere, Category="Child Layout", meta=(InlineEditConditionToggle))
+	uint8 bOverride_MinDesiredHeight : 1;
+
+	/**  */
+	UPROPERTY(EditAnywhere, Category="Child Layout", meta=(InlineEditConditionToggle))
+	uint8 bOverride_MaxDesiredWidth : 1;
+
+	/**  */
+	UPROPERTY(EditAnywhere, Category="Child Layout", meta=(InlineEditConditionToggle))
+	uint8 bOverride_MaxDesiredHeight : 1;
+
+	/**  */
+	UPROPERTY(EditAnywhere, Category="Child Layout", meta=(InlineEditConditionToggle))
+	uint8 bOverride_MinAspectRatio : 1;
+
+	/**  */
+	UPROPERTY(EditAnywhere, Category = "Child Layout", meta=(InlineEditConditionToggle))
+	uint8 bOverride_MaxAspectRatio : 1;
 
-public:
 		
 	/** When specified, ignore the content's desired size and report the WidthOverride as the Box's desired width. */
 	UFUNCTION(BlueprintCallable, Category="Layout|Size Box")
-		void SetWidthOverride(float InWidthOverride);
+	void SetWidthOverride(float InWidthOverride);
 
 	UFUNCTION(BlueprintCallable, Category="Layout|Size Box")
-		void ClearWidthOverride();
+	void ClearWidthOverride();
 
 	/** When specified, ignore the content's desired size and report the HeightOverride as the Box's desired height. */
 	UFUNCTION(BlueprintCallable, Category="Layout|Size Box")
-		void SetHeightOverride(float InHeightOverride);
+	void SetHeightOverride(float InHeightOverride);
 
 	UFUNCTION(BlueprintCallable, Category="Layout|Size Box")
-		void ClearHeightOverride();
+	void ClearHeightOverride();
 
 	/** When specified, will report the MinDesiredWidth if larger than the content's desired width. */
 	UFUNCTION(BlueprintCallable, Category="Layout|Size Box")
-		void SetMinDesiredWidth(float InMinDesiredWidth);
+	void SetMinDesiredWidth(float InMinDesiredWidth);
 
 	UFUNCTION(BlueprintCallable, Category="Layout|Size Box")
-		void ClearMinDesiredWidth();
+	void ClearMinDesiredWidth();
 
 	/** When specified, will report the MinDesiredHeight if larger than the content's desired height. */
 	UFUNCTION(BlueprintCallable, Category="Layout|Size Box")
-		void SetMinDesiredHeight(float InMinDesiredHeight);
+	void SetMinDesiredHeight(float InMinDesiredHeight);
 
 	UFUNCTION(BlueprintCallable, Category="Layout|Size Box")
-		void ClearMinDesiredHeight();
+	void ClearMinDesiredHeight();
 
 	/** When specified, will report the MaxDesiredWidth if smaller than the content's desired width. */
 	UFUNCTION(BlueprintCallable, Category="Layout|Size Box")
-		void SetMaxDesiredWidth(float InMaxDesiredWidth);
+	void SetMaxDesiredWidth(float InMaxDesiredWidth);
 
 	UFUNCTION(BlueprintCallable, Category="Layout|Size Box")
-		void ClearMaxDesiredWidth();
+	void ClearMaxDesiredWidth();
 
 	/** When specified, will report the MaxDesiredHeight if smaller than the content's desired height. */
 	UFUNCTION(BlueprintCallable, Category="Layout|Size Box")
-		void SetMaxDesiredHeight(float InMaxDesiredHeight);
+	void SetMaxDesiredHeight(float InMaxDesiredHeight);
 
 	UFUNCTION(BlueprintCallable, Category="Layout|Size Box")
-		void ClearMaxDesiredHeight();
+	void ClearMaxDesiredHeight();
 
 	UFUNCTION(BlueprintCallable, Category="Layout|Size Box")
 	void SetMinAspectRatio(float InMinAspectRatio);
@@ -145,7 +146,6 @@ public:
 	UFUNCTION(BlueprintCallable, Category="Layout|Size Box")
 	void ClearMaxAspectRatio();
 
-public:
 
 	// UWidget interface
 	virtual void SynchronizeProperties() override;
@@ -167,10 +167,6 @@ protected:
 	virtual void OnSlotRemoved(UPanelSlot* Slot) override;
 	// End UPanelWidget
 
-protected:
-	TSharedPtr<SBox> MySizeBox;
-
-protected:
 	// UWidget interface
 	virtual TSharedRef<SWidget> RebuildWidget() override;
 	// End of UWidget interface
diff --git a/Engine/Source/Runtime/UMG/Public/Components/SizeBoxSlot.h b/Engine/Source/Runtime/UMG/Public/Components/SizeBoxSlot.h
--- a/Engine/Source/Runtime/UMG/Public/Components/SizeBoxSlot.h
+++ b/Engine/Source/Runtime/UMG/Public/Components/SizeBoxSlot.h
@@ -23,6 +23,11 @@ public:
 	UPROPERTY(EditAnywhere, BlueprintReadOnly, Category="Layout|SizeBox Slot")
 	FMargin Padding;
 
+private:
+	/** A pointer to the button to allow us to adjust the size, padding...etc at runtime. */
+	TWeakPtr<SBox> SizeBox;
+
+public:
 	/** The alignment of the object horizontally. */
 	UPROPERTY(EditAnywhere, BlueprintReadOnly, Category="Layout|SizeBox Slot")
 	TEnumAsByte<EHorizontalAlignment> HorizontalAlignment;
@@ -31,7 +36,6 @@ public:
 	UPROPERTY(EditAnywhere, BlueprintReadOnly, Category="Layout|SizeBox Slot")
 	TEnumAsByte<EVerticalAlignment> VerticalAlignment;
 
-public:
 
 	UFUNCTION(BlueprintCallable, Category="Layout|SizeBox Slot")
 	void SetPadding(FMargin InPadding);
@@ -42,8 +46,6 @@ public:
 	UFUNCTION(BlueprintCallable, Category="Layout|SizeBox Slot")
 	void SetVerticalAlignment(EVerticalAlignment InVerticalAlignment);
 
-public:
-
 	// UPanelSlot interface
 	virtual void SynchronizeProperties() override;
 	// End of UPanelSlot interface
@@ -52,9 +54,4 @@ public:
 	void BuildSlot(TSharedRef<SBox> InSizeBox);
 
 	virtual void ReleaseSlateResources(bool bReleaseChildren) override;
-
-private:
-
-	/** A pointer to the button to allow us to adjust the size, padding...etc at runtime. */
-	TWeakPtr<SBox> SizeBox;
 };
diff --git a/Engine/Source/Runtime/UMG/Public/Components/VerticalBoxSlot.h b/Engine/Source/Runtime/UMG/Public/Components/VerticalBoxSlot.h
--- a/Engine/Source/Runtime/UMG/Public/Components/VerticalBoxSlot.h
+++ b/Engine/Source/Runtime/UMG/Public/Components/VerticalBoxSlot.h
@@ -20,13 +20,21 @@ class UMG_API UVerticalBoxSlot : public UPanelSlot
 
 public:
 	
+	/** How much space this slot should occupy in the direction of the panel. */
+	UPROPERTY(EditAnywhere, BlueprintReadOnly, Category="Layout|Vertical Box Slot", meta = (DisplayAfter = "Padding"))
+	FSlateChildSize Size;
+
 	/** The padding area between the slot and the content it contains. */
 	UPROPERTY(EditAnywhere, BlueprintReadOnly, Category="Layout|Vertical Box Slot")
 	FMargin Padding;
 
-	/** How much space this slot should occupy in the direction of the panel. */
-	UPROPERTY(EditAnywhere, BlueprintReadOnly, Category="Layout|Vertical Box Slot")
-	FSlateChildSize Size;
+private:
+	//TODO UMG Slots should hold weak or shared refs to slots.
+
+	/** A raw pointer to the slot to allow us to adjust the size, padding...etc at runtime. */
+	SVerticalBox::FSlot* Slot;
+
+public:
 
 	/** The alignment of the object horizontally. */
 	UPROPERTY(EditAnywhere, BlueprintReadOnly, Category="Layout|Vertical Box Slot")
@@ -60,10 +68,4 @@ public:
 
 	/** Builds the underlying FSlot for the Slate layout panel. */
 	void BuildSlot(TSharedRef<SVerticalBox> InVerticalBox);
-
-private:
-	//TODO UMG Slots should hold weak or shared refs to slots.
-
-	/** A raw pointer to the slot to allow us to adjust the size, padding...etc at runtime. */
-	SVerticalBox::FSlot* Slot;
 };
