Porting over fix as edit from dev-rendering cl: 8364385

IWYU Build Fix of multiple define definitions.

diff --git a/Engine/Source/Runtime/Engine/Classes/Materials/MaterialInterface.h b/Engine/Source/Runtime/Engine/Classes/Materials/MaterialInterface.h
--- a/Engine/Source/Runtime/Engine/Classes/Materials/MaterialInterface.h
+++ b/Engine/Source/Runtime/Engine/Classes/Materials/MaterialInterface.h
@@ -16,12 +16,9 @@
 #include "Engine/BlendableInterface.h"
 #include "Materials/MaterialLayersFunctions.h"
 #include "Interfaces/Interface_AssetUserData.h"
+#include "MaterialSceneTextureId.h"
 #include "MaterialInterface.generated.h"
 
-#ifndef STORE_ONLY_ACTIVE_SHADERMAPS
-#define STORE_ONLY_ACTIVE_SHADERMAPS 0
-#endif
-
 class FMaterialCompiler;
 class FMaterialRenderProxy;
 class FMaterialResource;
diff --git a/Engine/Source/Runtime/Engine/Public/MaterialSceneTextureId.h b/Engine/Source/Runtime/Engine/Public/MaterialSceneTextureId.h
--- a/Engine/Source/Runtime/Engine/Public/MaterialSceneTextureId.h
+++ b/Engine/Source/Runtime/Engine/Public/MaterialSceneTextureId.h
@@ -7,6 +7,10 @@
 #include "UObject/ObjectMacros.h"
 #include "MaterialSceneTextureId.generated.h"
 
+#ifndef STORE_ONLY_ACTIVE_SHADERMAPS
+#define STORE_ONLY_ACTIVE_SHADERMAPS 0
+#endif
+
 /** like EPassInputId but can expose more e.g. GBuffer */
 UENUM()
 enum ESceneTextureId
diff --git a/Engine/Source/Runtime/Engine/Public/MaterialShared.h b/Engine/Source/Runtime/Engine/Public/MaterialShared.h
--- a/Engine/Source/Runtime/Engine/Public/MaterialShared.h
+++ b/Engine/Source/Runtime/Engine/Public/MaterialShared.h
@@ -69,10 +69,6 @@ template <class ElementType> class TLinkedList;
 // disallow debug data in shipping or on non-desktop Test
 #define ALLOW_SHADERMAP_DEBUG_DATA (!(UE_BUILD_SHIPPING || (UE_BUILD_TEST && !PLATFORM_DESKTOP)))
 
-#ifndef STORE_ONLY_ACTIVE_SHADERMAPS
-#define STORE_ONLY_ACTIVE_SHADERMAPS 0
-#endif
-
 #define MATERIAL_OPACITYMASK_DOESNT_SUPPORT_VIRTUALTEXTURE 1
 
 DECLARE_LOG_CATEGORY_EXTERN(LogMaterial,Log,Verbose);
