Cherry-pick CL 7609630:
Added "Affect Global Illumination" attribute on light types.

This attribute optionally includes/excludes a given light from contributing to the ray tracing global illumination pass.

diff --git a/Engine/Source/Runtime/Engine/Classes/Components/LightComponentBase.h b/Engine/Source/Runtime/Engine/Classes/Components/LightComponentBase.h
--- a/Engine/Source/Runtime/Engine/Classes/Components/LightComponentBase.h
+++ b/Engine/Source/Runtime/Engine/Classes/Components/LightComponentBase.h
@@ -85,6 +85,10 @@ class ENGINE_API ULightComponentBase : public USceneComponent
 	UPROPERTY(EditAnywhere, BlueprintReadOnly, Category = Light, AdvancedDisplay)
 	uint32 bAffectReflection : 1;
 
+	/** Whether the light affects objects in reflections, when ray-traced global illumination is enabled. */
+	UPROPERTY(EditAnywhere, BlueprintReadOnly, Category = Light, AdvancedDisplay)
+	uint32 bAffectGlobalIllumination : 1;
+
 	/** 
 	 * Scales the indirect lighting contribution from this light. 
 	 * A value of 0 disables any GI from this light. Default is 1.
@@ -132,6 +136,9 @@ class ENGINE_API ULightComponentBase : public USceneComponent
 	UFUNCTION(BlueprintCallable, Category = "Rendering|Components|Light")
 	void SetAffectReflection(bool bNewValue);
 
+	UFUNCTION(BlueprintCallable, Category = "Rendering|Components|Light")
+	void SetAffectGlobalIllumination(bool bNewValue);
+
 	UFUNCTION(BlueprintCallable, Category = "Rendering|Components|Light")
 	void SetCastRaytracedShadow(bool bNewValue);
 
diff --git a/Engine/Source/Runtime/Engine/Private/Components/LightComponent.cpp b/Engine/Source/Runtime/Engine/Private/Components/LightComponent.cpp
--- a/Engine/Source/Runtime/Engine/Private/Components/LightComponent.cpp
+++ b/Engine/Source/Runtime/Engine/Private/Components/LightComponent.cpp
@@ -81,6 +81,16 @@ void ULightComponentBase::SetAffectReflection(bool bNewValue)
 	}
 }
 
+void ULightComponentBase::SetAffectGlobalIllumination(bool bNewValue)
+{
+	if (AreDynamicDataChangesAllowed()
+		&& bAffectGlobalIllumination != bNewValue)
+	{
+		bAffectGlobalIllumination = bNewValue;
+		MarkRenderStateDirty();
+	}
+}
+
 void ULightComponentBase::SetCastRaytracedShadow(bool bNewValue)
 {
 	if (AreDynamicDataChangesAllowed()
@@ -269,6 +279,7 @@ FLightSceneProxy::FLightSceneProxy(const ULightComponent* InLightComponent)
 	, bForceCachedShadowsForMovablePrimitives(InLightComponent->bForceCachedShadowsForMovablePrimitives)
 	, bCastRaytracedShadow(InLightComponent->bCastRaytracedShadow)
 	, bAffectReflection(InLightComponent->bAffectReflection)
+	, bAffectGlobalIllumination(InLightComponent->bAffectGlobalIllumination)
 	, bAffectTranslucentLighting(InLightComponent->bAffectTranslucentLighting)
 	, bUsedAsAtmosphereSunLight(InLightComponent->IsUsedAsAtmosphereSunLight())
 	, bAffectDynamicIndirectLighting(InLightComponent->bAffectDynamicIndirectLighting)
@@ -374,6 +385,7 @@ ULightComponentBase::ULightComponentBase(const FObjectInitializer& ObjectInitial
 	CastDynamicShadows = true;
 	bCastRaytracedShadow = true;
 	bAffectReflection = true;
+	bAffectGlobalIllumination = true;
 #if WITH_EDITORONLY_DATA
 	bVisualizeComponent = true;
 #endif
@@ -663,6 +675,7 @@ void ULightComponent::PostEditChangeProperty(FPropertyChangedEvent& PropertyChan
 		PropertyName != GET_MEMBER_NAME_STRING_CHECKED(ULightComponent, bCastVolumetricShadow) &&
 		PropertyName != GET_MEMBER_NAME_STRING_CHECKED(ULightComponent, bCastRaytracedShadow) &&
 		PropertyName != GET_MEMBER_NAME_STRING_CHECKED(ULightComponent, bAffectReflection) &&
+		PropertyName != GET_MEMBER_NAME_STRING_CHECKED(ULightComponent, bAffectGlobalIllumination) &&
 		// Point light properties that shouldn't unbuild lighting
 		PropertyName != GET_MEMBER_NAME_STRING_CHECKED(UPointLightComponent, SourceRadius) &&
 		PropertyName != GET_MEMBER_NAME_STRING_CHECKED(UPointLightComponent, SoftSourceRadius) &&
diff --git a/Engine/Source/Runtime/Engine/Public/SceneManagement.h b/Engine/Source/Runtime/Engine/Public/SceneManagement.h
--- a/Engine/Source/Runtime/Engine/Public/SceneManagement.h
+++ b/Engine/Source/Runtime/Engine/Public/SceneManagement.h
@@ -1357,6 +1357,7 @@ public:
 	inline bool CastsVolumetricShadow() const { return bCastVolumetricShadow; }
 	inline bool CastsRaytracedShadow() const { return bCastRaytracedShadow; }
 	inline bool AffectReflection() const { return bAffectReflection; }
+	inline bool AffectGlobalIllumination() const { return bAffectGlobalIllumination; }
 	inline bool CastsShadowsFromCinematicObjectsOnly() const { return bCastShadowsFromCinematicObjectsOnly; }
 	inline bool CastsModulatedShadows() const { return bCastModulatedShadows; }
 	inline const FLinearColor& GetModulatedShadowColor() const { return ModulatedShadowColor; }
@@ -1531,6 +1532,9 @@ protected:
 	/** Whether the light affects objects in reflections, when ray-traced reflection is enabled. */
 	const uint8 bAffectReflection : 1;
 
+	/** Whether the light affects objects in reflections, when ray-traced global illumination is enabled. */
+	const uint8 bAffectGlobalIllumination : 1;
+
 	/** Whether the light affects translucency or not.  Disabling this can save GPU time when there are many small lights. */
 	const uint8 bAffectTranslucentLighting : 1;
 
diff --git a/Engine/Source/Runtime/Renderer/Private/LightSceneInfo.cpp b/Engine/Source/Runtime/Renderer/Private/LightSceneInfo.cpp
--- a/Engine/Source/Runtime/Renderer/Private/LightSceneInfo.cpp
+++ b/Engine/Source/Runtime/Renderer/Private/LightSceneInfo.cpp
@@ -31,6 +31,7 @@ void FLightSceneInfoCompact::Init(FLightSceneInfo* InLightSceneInfo)
 	bCastStaticShadow = InLightSceneInfo->Proxy->CastsStaticShadow();
 	bStaticLighting = InLightSceneInfo->Proxy->HasStaticLighting();
 	bAffectReflection = InLightSceneInfo->Proxy->AffectReflection();
+	bAffectGlobalIllumination = InLightSceneInfo->Proxy->AffectGlobalIllumination();
 	bCastRaytracedShadow = InLightSceneInfo->Proxy->CastsRaytracedShadow();
 }
 
diff --git a/Engine/Source/Runtime/Renderer/Private/LightSceneInfo.h b/Engine/Source/Runtime/Renderer/Private/LightSceneInfo.h
--- a/Engine/Source/Runtime/Renderer/Private/LightSceneInfo.h
+++ b/Engine/Source/Runtime/Renderer/Private/LightSceneInfo.h
@@ -38,6 +38,7 @@ public:
 	uint32 bCastStaticShadow : 1;
 	uint32 bStaticLighting : 1;
 	uint32 bAffectReflection : 1;
+	uint32 bAffectGlobalIllumination : 1;
 	uint32 bCastRaytracedShadow : 1;
 
 	/** Initializes the compact scene info from the light's full scene info. */
diff --git a/Engine/Source/Runtime/Renderer/Private/RayTracing/RayTracingGlobalIllumination.cpp b/Engine/Source/Runtime/Renderer/Private/RayTracing/RayTracingGlobalIllumination.cpp
--- a/Engine/Source/Runtime/Renderer/Private/RayTracing/RayTracingGlobalIllumination.cpp
+++ b/Engine/Source/Runtime/Renderer/Private/RayTracing/RayTracingGlobalIllumination.cpp
@@ -135,6 +135,7 @@ void SetupLightParameters(
 		if (LightParameters->Count >= GLightCountMax) break;
 
 		if (Light.LightSceneInfo->Proxy->HasStaticLighting() && Light.LightSceneInfo->IsPrecomputedLightingValid()) continue;
+		if (!Light.LightSceneInfo->Proxy->AffectGlobalIllumination()) continue;
 
 		FLightShaderParameters LightShaderParameters;
 		Light.LightSceneInfo->Proxy->GetLightShaderParameters(LightShaderParameters);
