Fixed black lightmaps after opening the static mesh editor while using mobile preview mode.

Lightmap resources were not recreated correctly when when the feature level changed. Also, modified FRenderResource::ChangeFeatureLevel() so that it doesn't recreate resources which are already using the correct feature level.

diff --git a/Engine/Source/Runtime/Engine/Private/LightMap.cpp b/Engine/Source/Runtime/Engine/Private/LightMap.cpp
--- a/Engine/Source/Runtime/Engine/Private/LightMap.cpp
+++ b/Engine/Source/Runtime/Engine/Private/LightMap.cpp
@@ -3293,14 +3293,18 @@ void FLightmapResourceCluster::UpdateUniformBuffer(ERHIFeatureLevel::Type InFeat
 	});
 }
 
+bool FLightmapResourceCluster::GetUseVirtualTexturing() const
+{
+	const bool bAllowHighQualityLightMaps = AllowHighQualityLightmaps(FeatureLevel);
+	return bAllowHighQualityLightMaps && (CVarVirtualTexturedLightMaps.GetValueOnRenderThread() != 0) && UseVirtualTexturing(FeatureLevel);
+}
+
 void FLightmapResourceCluster::UpdateUniformBuffer_RenderThread()
 {
 	check(IsInRenderingThread());
-	const bool bAllowHighQualityLightMaps = AllowHighQualityLightmaps(FeatureLevel);
-	const bool bUseVirtualTextures = bAllowHighQualityLightMaps && (CVarVirtualTexturedLightMaps.GetValueOnRenderThread() != 0) && UseVirtualTexturing(FeatureLevel);
 
 	FLightmapResourceClusterShaderParameters Parameters;
-	GetLightmapClusterResourceParameters(FeatureLevel, Input, bUseVirtualTextures ? AcquireAllocatedVT() : nullptr, Parameters);
+	GetLightmapClusterResourceParameters(FeatureLevel, Input, GetUseVirtualTexturing() ? AcquireAllocatedVT() : nullptr, Parameters);
 
 	RHIUpdateUniformBuffer(UniformBuffer, &Parameters);
 }
@@ -3379,7 +3383,18 @@ void FLightmapResourceCluster::ReleaseAllocatedVT()
 void FLightmapResourceCluster::InitRHI()
 {
 	FLightmapResourceClusterShaderParameters Parameters;
-	GetLightmapClusterResourceParameters(GMaxRHIFeatureLevel, FLightmapClusterResourceInput(), nullptr, Parameters);
+
+	// Lightmap resources are normally created before the feature level is known, so we'll use defaults and rely on a subsequent call to UpdateUniformBuffer()
+	// to set the correct level and update things accordingly. However, when we're coming from FRenderResource::ChangeFeatureLevel(), the feature level
+	// has already been set, so we can go ahead and use the correct level and input from the start (UpdateUniformBuffer() is not being called in that case).
+	if (FeatureLevel < ERHIFeatureLevel::Num)
+	{
+		GetLightmapClusterResourceParameters(FeatureLevel, Input, GetUseVirtualTexturing() ? AcquireAllocatedVT() : nullptr, Parameters);
+	}
+	else
+	{
+		GetLightmapClusterResourceParameters(GMaxRHIFeatureLevel, FLightmapClusterResourceInput(), nullptr, Parameters);
+	}
 
 	UniformBuffer = FLightmapResourceClusterShaderParameters::CreateUniformBuffer(Parameters, UniformBuffer_MultiFrame);
 }
diff --git a/Engine/Source/Runtime/Engine/Public/LightMap.h b/Engine/Source/Runtime/Engine/Public/LightMap.h
--- a/Engine/Source/Runtime/Engine/Public/LightMap.h
+++ b/Engine/Source/Runtime/Engine/Public/LightMap.h
@@ -653,6 +653,7 @@ public:
 	 */
 	IAllocatedVirtualTexture* AcquireAllocatedVT() const;
 	void ReleaseAllocatedVT();
+	bool GetUseVirtualTexturing() const;
 
 	FLightmapClusterResourceInput Input;
 	mutable IAllocatedVirtualTexture* AllocatedVT;
diff --git a/Engine/Source/Runtime/RenderCore/Private/RenderResource.cpp b/Engine/Source/Runtime/RenderCore/Private/RenderResource.cpp
--- a/Engine/Source/Runtime/RenderCore/Private/RenderResource.cpp
+++ b/Engine/Source/Runtime/RenderCore/Private/RenderResource.cpp
@@ -66,7 +66,7 @@ void FRenderResource::ChangeFeatureLevel(ERHIFeatureLevel::Type NewFeatureLevel)
 			FRenderResource* Resource = *It;
 
 			// Only resources configured for a specific feature level need to be updated
-			if (Resource->HasValidFeatureLevel())
+			if (Resource->HasValidFeatureLevel() && (Resource->FeatureLevel != NewFeatureLevel))
 			{
 				Resource->ReleaseRHI();
 				Resource->ReleaseDynamicRHI();
