#UE4 Add more info to a debugging log

diff --git a/Engine/Source/Runtime/Engine/Private/Materials/MaterialUniformExpressions.h b/Engine/Source/Runtime/Engine/Private/Materials/MaterialUniformExpressions.h
--- a/Engine/Source/Runtime/Engine/Private/Materials/MaterialUniformExpressions.h
+++ b/Engine/Source/Runtime/Engine/Private/Materials/MaterialUniformExpressions.h
@@ -277,7 +277,7 @@ static TextureType* GetIndexedTexture(const FMaterial& Material, int32 TextureIn
 		static bool bWarnedOnce = false;
 		if (!bWarnedOnce)
 		{
-			UE_LOG(LogMaterial, Warning, TEXT("Requesting an invalid TextureIndex! (%u / %u)"), TextureIndex, ReferencedTextures.Num());
+			UE_LOG(LogMaterial, Warning, TEXT("%s: Requesting an invalid TextureIndex! (%u / %u)"), *Material.GetFriendlyName(), TextureIndex, ReferencedTextures.Num());
 			bWarnedOnce = true;
 		}
 	}
@@ -287,7 +287,7 @@ static TextureType* GetIndexedTexture(const FMaterial& Material, int32 TextureIn
 		static bool bWarnedOnce = false;
 		if (!bWarnedOnce)
 		{
-			UE_LOG(LogMaterial, Warning, TEXT("GetIndexedTexture returning NULL (%u)"), TextureIndex);
+			UE_LOG(LogMaterial, Warning, TEXT("%s: GetIndexedTexture returning NULL (%u)"), *Material.GetFriendlyName(), TextureIndex);
 			bWarnedOnce = true;
 		}
 	}
