Fixed crash between seamless travel in PIE

diff --git a/Engine/Source/Runtime/Engine/Private/Streaming/LevelTextureManager.cpp b/Engine/Source/Runtime/Engine/Private/Streaming/LevelTextureManager.cpp
--- a/Engine/Source/Runtime/Engine/Private/Streaming/LevelTextureManager.cpp
+++ b/Engine/Source/Runtime/Engine/Private/Streaming/LevelTextureManager.cpp
@@ -74,7 +74,7 @@ float FLevelRenderAssetManager::GetWorldTime() const
 		{
 			// In the editor, we only return a time for the PIE world.
 			// TODO: figure out why there are more than one PIE world
-			if (!GIsEditor || (World->IsPlayInEditor() && World->Scene->GetFrameNumber()))
+			if (!GIsEditor || (World->IsPlayInEditor() && World->Scene && World->Scene->GetFrameNumber()))
 			{
 				return World->GetTimeSeconds();
 			}
