Add logging for a PreExit to catch entering the PreExit state during a game

diff --git a/Engine/Source/Runtime/Engine/Private/GameEngine.cpp b/Engine/Source/Runtime/Engine/Private/GameEngine.cpp
--- a/Engine/Source/Runtime/Engine/Private/GameEngine.cpp
+++ b/Engine/Source/Runtime/Engine/Private/GameEngine.cpp
@@ -1137,6 +1137,8 @@ void UGameEngine::Start()
 
 void UGameEngine::PreExit()
 {
+	UE_LOG(LogInit, Display, TEXT("PreExit Game."));
+
 	GetGameUserSettings()->SaveSettings();
 
 	// Stop tracking, automatically flushes.
