Remove code from bad merge, causing CIS errors.

#rnx

diff --git a/Engine/Source/Runtime/Engine/Private/Collision/SceneQuery.cpp b/Engine/Source/Runtime/Engine/Private/Collision/SceneQuery.cpp
--- a/Engine/Source/Runtime/Engine/Private/Collision/SceneQuery.cpp
+++ b/Engine/Source/Runtime/Engine/Private/Collision/SceneQuery.cpp
@@ -20,10 +20,6 @@
 
 #include "Collision/CollisionDebugDrawing.h"
 
-#if WITH_CHAOS || PHYSICS_INTERFACE_LLIMMEDIATE
-#define STARTQUERYTIMER()
-#endif
-
 float DebugLineLifetime = 2.f;
 
 #include "PhysicsEngine/PhysXSupport.h"
