Ray tracing translucency does not work in PIE with the default pawn.

diff --git a/Engine/Source/Runtime/Engine/Private/DefaultPawn.cpp b/Engine/Source/Runtime/Engine/Private/DefaultPawn.cpp
--- a/Engine/Source/Runtime/Engine/Private/DefaultPawn.cpp
+++ b/Engine/Source/Runtime/Engine/Private/DefaultPawn.cpp
@@ -65,6 +65,7 @@ ADefaultPawn::ADefaultPawn(const FObjectInitializer& ObjectInitializer)
 		MeshComponent->bCastDynamicShadow = true;
 		MeshComponent->bAffectDynamicIndirectLighting = false;
 		MeshComponent->bAffectDistanceFieldLighting = false;
+		MeshComponent->bVisibleInRayTracing = false;
 		MeshComponent->PrimaryComponentTick.TickGroup = TG_PrePhysics;
 		MeshComponent->SetupAttachment(RootComponent);
 		MeshComponent->SetCollisionProfileName(UCollisionProfile::Pawn_ProfileName);
