Clean up AActor::GetComponents API. GetComponentsByClass was intended only for blueprints. The BP version has been renamed K2_GetComponentsByClass and made private, while the C++ GetComponentsByClass has been deprecated. A new non-templated GetComponents function has been added that takes a UActorComponent subclass and populates a passed in array with UActorComponents of that subclass.

diff --git a/Engine/Config/BaseEngine.ini b/Engine/Config/BaseEngine.ini
--- a/Engine/Config/BaseEngine.ini
+++ b/Engine/Config/BaseEngine.ini
@@ -690,6 +690,8 @@ bUseStreamingPause=false
 ; 4.24
 +PropertyRedirects=(OldName="NavDataConfig.NavigationDataClassName", NewName="NavDataConfig.NavDataClass")
 
++FunctionRedirects=(OldName="Actor.GetComponentsByClass", NewName="Actor.K2_GetComponentsByClass")
+
 [CoreUObject.Metadata]
 ; Note: UnrealHeaderTool should be rerun after making changes to MetadataRedirects to catch any keys specified in class headers
 +MetadataRedirects=(OldKey="K2Protected", NewKey="BlueprintProtected")
diff --git a/Engine/Source/Runtime/Engine/Classes/GameFramework/Actor.h b/Engine/Source/Runtime/Engine/Classes/GameFramework/Actor.h
--- a/Engine/Source/Runtime/Engine/Classes/GameFramework/Actor.h
+++ b/Engine/Source/Runtime/Engine/Classes/GameFramework/Actor.h
@@ -78,6 +78,20 @@ DECLARE_CYCLE_STAT_EXTERN(TEXT("GetComponentsTime"),STAT_GetComponentsTime,STATG
 extern ENGINE_API FUObjectAnnotationSparseBool GSelectedActorAnnotation;
 #endif
 
+/**
+ * TInlineComponentArray is simply a TArray that reserves a fixed amount of space on the stack
+ * to try to avoid heap allocation when there are fewer than a specified number of elements expected in the result.
+ */
+template<class T, uint32 NumElements = NumInlinedActorComponents>
+class TInlineComponentArray : public TArray<T, TInlineAllocator<NumElements>>
+{
+	typedef TArray<T, TInlineAllocator<NumElements>> Super;
+
+public:
+	TInlineComponentArray() : Super() { }
+	TInlineComponentArray(const AActor* Actor, bool bIncludeFromChildActors = false);
+};
+
 /**
  * Actor is the base class for an Object that can be placed or spawned in a level.
  * Actors may contain a collection of ActorComponents, which can be used to control how actors move, how they are rendered, etc.
@@ -2787,13 +2801,15 @@ public:
 	UFUNCTION(BlueprintCallable, Category = "Actor", meta = (ComponentClass = "ActorComponent"), meta = (DeterminesOutputType = "ComponentClass"))
 	UActorComponent* GetComponentByClass(TSubclassOf<UActorComponent> ComponentClass) const;
 
+private:
 	/**
 	 * Gets all the components that inherit from the given class.
 	 * Currently returns an array of UActorComponent which must be cast to the correct type.
 	 */
-	UFUNCTION(BlueprintCallable, Category = "Actor", meta = (ComponentClass = "ActorComponent"), meta = (DeterminesOutputType = "ComponentClass"))
-	TArray<UActorComponent*> GetComponentsByClass(TSubclassOf<UActorComponent> ComponentClass) const;
+	UFUNCTION(BlueprintCallable, Category = "Actor", meta = (ComponentClass = "ActorComponent", DisplayName = "GetComponentsByClass", DeterminesOutputType = "ComponentClass", AllowPrivateAccess=true))
+	TArray<UActorComponent*> K2_GetComponentsByClass(TSubclassOf<UActorComponent> ComponentClass) const;
 
+public:
 	/** Gets all the components that inherit from the given class with a given tag. */
 	UFUNCTION(BlueprintCallable, Category = "Actor", meta = (ComponentClass = "ActorComponent"), meta = (DeterminesOutputType = "ComponentClass"))
 	TArray<UActorComponent*> GetComponentsByTag(TSubclassOf<UActorComponent> ComponentClass, FName Tag) const;
@@ -2807,105 +2823,118 @@ public:
 		return (T*)FindComponentByClass(T::StaticClass());
 	}
 
+private:
 	/**
-	 * Get all components derived from class 'T' and fill in the OutComponents array with the result.
-	 * It's recommended to use TArrays with a TInlineAllocator to potentially avoid memory allocation costs.
-	 * TInlineComponentArray is defined to make this easier, for example:
-	 * {
-	 * 	   TInlineComponentArray<UPrimitiveComponent*> PrimComponents(Actor);
-	 * }
-	 *
-	 * @param bIncludeFromChildActors	If true then recurse in to ChildActor components and find components of the appropriate type in those Actors as well
+	 * Internal helper function to centralize GetComponents logic. 
+	 * Use template parameter bClassIsActorComponent to avoid doing unnecessary IsA checks when the ComponentClass is exactly UActorComponent
 	 */
-	template<class T, class AllocatorType>
-	void GetComponents(TArray<T*, AllocatorType>& OutComponents, bool bIncludeFromChildActors = false) const
+	template<bool bClassIsActorComponent, class AllocatorType>
+	void GetComponents_Internal(TSubclassOf<UActorComponent> ComponentClass, TArray<UActorComponent*, AllocatorType>& OutComponents, bool bIncludeFromChildActors = false) const
 	{
-		static_assert(TPointerIsConvertibleFromTo<T, const UActorComponent>::Value, "'T' template parameter to GetComponents must be derived from UActorComponent");
-		SCOPE_CYCLE_COUNTER(STAT_GetComponentsTime);
-
-		// Empty input array, but don't affect allocated size.
-		OutComponents.Reset(0);
+		check(bClassIsActorComponent == false || ComponentClass == UActorComponent::StaticClass());
 
-		TArray<UChildActorComponent*> ChildActorComponents;
+		TArray<AActor*> ChildActors;
 
 		for (UActorComponent* OwnedComponent : OwnedComponents)
 		{
-			if (T* Component = Cast<T>(OwnedComponent))
-			{
-				OutComponents.Add(Component);
-			}
-			if (bIncludeFromChildActors)
+			if (OwnedComponent)
 			{
-				if (UChildActorComponent* ChildActorComponent = Cast<UChildActorComponent>(OwnedComponent))
+				if (bClassIsActorComponent || OwnedComponent->IsA(ComponentClass))
+				{
+					OutComponents.Add(OwnedComponent);
+				}
+				if (bIncludeFromChildActors)
 				{
-					ChildActorComponents.Add(ChildActorComponent);
+					if (UChildActorComponent* ChildActorComponent = Cast<UChildActorComponent>(OwnedComponent))
+					{
+						if (AActor* ChildActor = ChildActorComponent->GetChildActor())
+						{
+							ChildActors.Add(ChildActor);
+						}
+					}
 				}
 			}
 		}
 
-		if (bIncludeFromChildActors)
+		for (AActor* ChildActor : ChildActors)
 		{
-			TArray<T*, AllocatorType> ComponentsInChildActor;
-			for (UChildActorComponent* ChildActorComponent : ChildActorComponents)
-			{
-				if (AActor* ChildActor = ChildActorComponent->GetChildActor())
-				{
-					ChildActor->GetComponents(ComponentsInChildActor, true);
-					OutComponents.Append(MoveTemp(ComponentsInChildActor));
-				}
-			}
+			ChildActor->GetComponents_Internal<bClassIsActorComponent>(ComponentClass, OutComponents, true);
 		}
 	}
 
+public:
+
+	UE_DEPRECATED(4.24, "Use one of the GetComponents implementations as appropriate")
+	TArray<UActorComponent*> GetComponentsByClass(TSubclassOf<UActorComponent> ComponentClass) const
+	{
+		return K2_GetComponentsByClass(ComponentClass);
+	}
+
 	/**
-	 * UActorComponent specialization of GetComponents() to avoid unnecessary casts.
+	 * Get all components derived from specified ComponentClass and fill in the OutComponents array with the result.
 	 * It's recommended to use TArrays with a TInlineAllocator to potentially avoid memory allocation costs.
 	 * TInlineComponentArray is defined to make this easier, for example:
 	 * {
-	 * 	   TInlineComponentArray<UActorComponent*> PrimComponents;
-	 *     Actor->GetComponents(PrimComponents);
+	 * 	   TInlineComponentArray<UPrimitiveComponent*> PrimComponents(Actor);
 	 * }
 	 *
-	 * @param bIncludeFromChildActors	If true then recurse in to ChildActor components and find components of the appropriate type in those Actors as well
+	 * @param ComponentClass            The component class to find all components of a class derived from
+	 * @param bIncludeFromChildActors   If true then recurse in to ChildActor components and find components of the appropriate type in those Actors as well
 	 */
 	template<class AllocatorType>
-	void GetComponents(TArray<UActorComponent*, AllocatorType>& OutComponents, bool bIncludeFromChildActors = false) const
+	void GetComponents(TSubclassOf<UActorComponent> ComponentClass, TArray<UActorComponent*, AllocatorType>& OutComponents, bool bIncludeFromChildActors = false) const
 	{
 		SCOPE_CYCLE_COUNTER(STAT_GetComponentsTime);
 
-		OutComponents.Reset(OwnedComponents.Num());
-
-		TArray<UChildActorComponent*> ChildActorComponents;
-
-		for (UActorComponent* Component : OwnedComponents)
+		OutComponents.Reset();
+		if (ComponentClass == UActorComponent::StaticClass())
 		{
-			if (Component)
-			{
-				OutComponents.Add(Component);
-
-				if (bIncludeFromChildActors)
-				{
-					if (UChildActorComponent* ChildActorComponent = Cast<UChildActorComponent>(Component))
-					{
-						ChildActorComponents.Add(ChildActorComponent);
-					}
-				}
-			}
+			GetComponents_Internal<true>(ComponentClass, OutComponents, bIncludeFromChildActors);
 		}
-
-		if (bIncludeFromChildActors)
+		else
 		{
-			TArray<UActorComponent*, AllocatorType> ComponentsInChildActor;
-			for (UChildActorComponent* ChildActorComponent : ChildActorComponents)
-			{
-				if (AActor* ChildActor = ChildActorComponent->GetChildActor())
-				{
-					ChildActor->GetComponents(ComponentsInChildActor, true);
-					OutComponents.Append(MoveTemp(ComponentsInChildActor));
-				}
-			}
+			GetComponents_Internal<false>(ComponentClass, OutComponents, bIncludeFromChildActors);
 		}
+	}
 
+	/**
+	 * Get all components derived from class 'T' and fill in the OutComponents array with the result.
+	 * It's recommended to use TArrays with a TInlineAllocator to potentially avoid memory allocation costs.
+	 * TInlineComponentArray is defined to make this easier, for example:
+	 * {
+	 * 	   TInlineComponentArray<UPrimitiveComponent*> PrimComponents(Actor);
+	 * }
+	 *
+	 * @param bIncludeFromChildActors	If true then recurse in to ChildActor components and find components of the appropriate type in those Actors as well
+	 */
+	template<class T, class AllocatorType>
+	void GetComponents(TArray<T*, AllocatorType>& OutComponents, bool bIncludeFromChildActors = false) const
+	{
+		SCOPE_CYCLE_COUNTER(STAT_GetComponentsTime);
+
+		static_assert(TPointerIsConvertibleFromTo<T, const UActorComponent>::Value, "'T' template parameter to GetComponents must be derived from UActorComponent");
+		OutComponents.Reset();
+		GetComponents_Internal<false>(T::StaticClass(), *reinterpret_cast<TArray<UActorComponent*, AllocatorType>*>(&OutComponents), bIncludeFromChildActors);
+	}
+
+	/**
+	 * UActorComponent specialization of GetComponents() to avoid unnecessary casts.
+	 * It's recommended to use TArrays with a TInlineAllocator to potentially avoid memory allocation costs.
+	 * TInlineComponentArray is defined to make this easier, for example:
+	 * {
+	 * 	   TInlineComponentArray<UActorComponent*> PrimComponents;
+	 *     Actor->GetComponents(PrimComponents);
+	 * }
+	 *
+	 * @param bIncludeFromChildActors	If true then recurse in to ChildActor components and find components of the appropriate type in those Actors as well
+	 */
+	template<class AllocatorType>
+	void GetComponents(TArray<UActorComponent*, AllocatorType>& OutComponents, bool bIncludeFromChildActors = false) const
+	{
+		SCOPE_CYCLE_COUNTER(STAT_GetComponentsTime);
+
+		OutComponents.Reset();
+		GetComponents_Internal<true>(UActorComponent::StaticClass(), OutComponents, bIncludeFromChildActors);
 	}
 
 	/**
@@ -3189,26 +3218,6 @@ private:
 };
 #endif
 
-/**
- * TInlineComponentArray is simply a TArray that reserves a fixed amount of space on the stack
- * to try to avoid heap allocation when there are fewer than a specified number of elements expected in the result.
- */
-template<class T, uint32 NumElements = NumInlinedActorComponents>
-class TInlineComponentArray : public TArray<T, TInlineAllocator<NumElements>>
-{
-	typedef TArray<T, TInlineAllocator<NumElements>> Super;
-
-public:
-	TInlineComponentArray() : Super() { }
-	TInlineComponentArray(const class AActor* Actor, bool bIncludeFromChildActors = false) : Super()
-	{
-		if (Actor)
-		{
-			Actor->GetComponents(*this, bIncludeFromChildActors);
-		}
-	};
-};
-
 /** Helper function for executing tick functions based on the normal conditions previous found in UActorComponent::ConditionalTick */
 template <typename ExecuteTickLambda>
 void FActorComponentTickFunction::ExecuteTickHelper(UActorComponent* Target, bool bTickInEditor, float DeltaTime, ELevelTick TickType, const ExecuteTickLambda& ExecuteTickFunc)
@@ -3234,6 +3243,15 @@ void FActorComponentTickFunction::ExecuteTickHelper(UActorComponent* Target, boo
 	}
 }
 
+template<class T, uint32 NumElements>
+TInlineComponentArray<T, NumElements>::TInlineComponentArray(const AActor* Actor, bool bIncludeFromChildActors) 
+	: Super()
+{
+	if (Actor)
+	{
+		Actor->GetComponents(*this, bIncludeFromChildActors);
+	}
+};
 
 //////////////////////////////////////////////////////////////////////////
 // Inlines
diff --git a/Engine/Source/Runtime/Engine/Private/Actor.cpp b/Engine/Source/Runtime/Engine/Private/Actor.cpp
--- a/Engine/Source/Runtime/Engine/Private/Actor.cpp
+++ b/Engine/Source/Runtime/Engine/Private/Actor.cpp
@@ -2796,50 +2796,29 @@ UActorComponent* AActor::GetComponentByClass(TSubclassOf<UActorComponent> Compon
 	return FindComponentByClass(ComponentClass);
 }
 
-TArray<UActorComponent*> AActor::GetComponentsByClass(TSubclassOf<UActorComponent> ComponentClass) const
+TArray<UActorComponent*> AActor::K2_GetComponentsByClass(TSubclassOf<UActorComponent> ComponentClass) const
 {
-	TArray<UActorComponent*> ValidComponents;
-
-	// In the UActorComponent case we can skip the IsA checks for a slight performance benefit
-	if (ComponentClass == UActorComponent::StaticClass())
-	{
-		for (UActorComponent* Component : OwnedComponents)
-		{
-			if (Component)
-			{
-				ValidComponents.Add(Component);
-			}
-		}
-	}
-	else if (*ComponentClass)
-	{
-		for (UActorComponent* Component : OwnedComponents)
-		{
-			if (Component && Component->IsA(ComponentClass))
-			{
-				ValidComponents.Add(Component);
-			}
-		}
-	}
-	
-	return ValidComponents;
+	TArray<UActorComponent*> Components;
+	GetComponents(ComponentClass, Components);
+	return MoveTemp(Components);
 }
 
 TArray<UActorComponent*> AActor::GetComponentsByTag(TSubclassOf<UActorComponent> ComponentClass, FName Tag) const
 {
-	TArray<UActorComponent*> ComponentsByClass = GetComponentsByClass(ComponentClass);
+	TInlineComponentArray<UActorComponent*> ComponentsByClass;
+	GetComponents(ComponentClass, ComponentsByClass);
 
 	TArray<UActorComponent*> ComponentsByTag;
 	ComponentsByTag.Reserve(ComponentsByClass.Num());
-	for (int i = 0; i < ComponentsByClass.Num(); ++i)
+	for (UActorComponent* Component : ComponentsByClass)
 	{
-		if (ComponentsByClass[i]->ComponentHasTag(Tag))
+		if (Component->ComponentHasTag(Tag))
 		{
-			ComponentsByTag.Push(ComponentsByClass[i]);
+			ComponentsByTag.Push(Component);
 		}
 	}
 
-	return ComponentsByTag;
+	return MoveTemp(ComponentsByTag);
 }
 
 void AActor::DisableComponentsSimulatePhysics()
