Use FNetworkObjectInfo's TWeakPtr to access actor pointer instead of cached pointer to try to prevent crash.

diff --git a/Engine/Source/Runtime/Engine/Private/NetworkDriver.cpp b/Engine/Source/Runtime/Engine/Private/NetworkDriver.cpp
--- a/Engine/Source/Runtime/Engine/Private/NetworkDriver.cpp
+++ b/Engine/Source/Runtime/Engine/Private/NetworkDriver.cpp
@@ -3123,8 +3123,8 @@ void UNetDriver::ForceAllActorsNetUpdateTime(float NetUpdateTimeOffset, TFunctio
 		FNetworkObjectInfo* NetActorInfo = (*It).Get();
 		if (NetActorInfo)
 		{
-			const AActor* const Actor = NetActorInfo->Actor;
-			if (Actor && !Actor->IsPendingKill())
+			const AActor* const Actor = NetActorInfo->WeakActor.Get();
+			if (Actor)
 			{
 				if (ValidActorTestFunc(Actor))
 				{
diff --git a/Engine/Source/Runtime/Engine/Private/PlayerController.cpp b/Engine/Source/Runtime/Engine/Private/PlayerController.cpp
--- a/Engine/Source/Runtime/Engine/Private/PlayerController.cpp
+++ b/Engine/Source/Runtime/Engine/Private/PlayerController.cpp
@@ -1056,8 +1056,8 @@ void APlayerController::ServerShortTimeout_Implementation()
 			{
 				if (NetworkObjectInfo.IsValid())
 				{
-					AActor* const A = NetworkObjectInfo->Actor;
-					if (A && !A->IsPendingKill())
+					AActor* const A = NetworkObjectInfo->WeakActor.Get();
+					if (A)
 					{
 						if (!A->bOnlyRelevantToOwner)
 						{
