Potential fix and added check for safety for DyanmicSoundVolume crash

diff --git a/Engine/Source/Runtime/Engine/Private/Audio.cpp b/Engine/Source/Runtime/Engine/Private/Audio.cpp
--- a/Engine/Source/Runtime/Engine/Private/Audio.cpp
+++ b/Engine/Source/Runtime/Engine/Private/Audio.cpp
@@ -914,17 +914,23 @@ float FWaveInstance::GetDistanceAttenuation() const
 float FWaveInstance::GetDynamicVolume() const
 {
 	float OutVolume = 1.0f;
-	if (FAudioDeviceManager* DeviceManager = GEngine->GetAudioDeviceManager())
+
+	if (GEngine)
 	{
-		if (WaveData)
+		if (FAudioDeviceManager* DeviceManager = GEngine->GetAudioDeviceManager())
 		{
-			OutVolume *= DeviceManager->GetDynamicSoundVolume(ESoundType::Wave, WaveData->GetFName());
-		}
+			if (WaveData)
+			{
+				OutVolume *= DeviceManager->GetDynamicSoundVolume(ESoundType::Wave, WaveData->GetFName());
+			}
 
-		check(ActiveSound);
-		if (const USoundCue* Sound = Cast<USoundCue>(ActiveSound->GetSound()))
-		{
-			OutVolume *= DeviceManager->GetDynamicSoundVolume(ESoundType::Cue, Sound->GetFName());
+			if (ActiveSound)
+			{
+				if (const USoundCue* Sound = Cast<USoundCue>(ActiveSound->GetSound()))
+				{
+					OutVolume *= DeviceManager->GetDynamicSoundVolume(ESoundType::Cue, Sound->GetFName());
+				}
+			}
 		}
 	}
 
diff --git a/Engine/Source/Runtime/Engine/Private/AudioDeviceManager.cpp b/Engine/Source/Runtime/Engine/Private/AudioDeviceManager.cpp
--- a/Engine/Source/Runtime/Engine/Private/AudioDeviceManager.cpp
+++ b/Engine/Source/Runtime/Engine/Private/AudioDeviceManager.cpp
@@ -1093,6 +1093,8 @@ bool FAudioDeviceManager::GetAudioDebugSound(FString& OutDebugSound)
 
 float FAudioDeviceManager::GetDynamicSoundVolume(ESoundType SoundType, const FName& SoundName) const
 {
+	check(IsInAudioThread());
+
 	TTuple<ESoundType, FName> SoundKey(SoundType, SoundName);
 	if (const float* Volume = DynamicSoundVolumes.Find(SoundKey))
 	{
