Added UAnimInstance::ClearLayerOverlay

This new call allows overlaid layer classes to be selectively reset to their defaults, allowing partial graph resets instead of needing to restore and rebuild state each time.
Also added a call to get a specific layer sub instance by class, a useful counterpart to the new functionality.

diff --git a/Engine/Source/Runtime/Engine/Classes/Animation/AnimInstance.h b/Engine/Source/Runtime/Engine/Classes/Animation/AnimInstance.h
--- a/Engine/Source/Runtime/Engine/Classes/Animation/AnimInstance.h
+++ b/Engine/Source/Runtime/Engine/Classes/Animation/AnimInstance.h
@@ -31,6 +31,7 @@ struct FAnimNode_SubInput;
 struct FBakedAnimationStateMachine;
 class FCompilerResultsLog;
 struct FBoneContainer;
+struct FAnimNode_Layer;
 
 typedef TArray<FTransform> FTransformArrayA2;
 
@@ -721,18 +722,34 @@ public:
 	 * Runs through all layer nodes, attempting to find layer nodes that are implemented by the specified class, then sets up a sub instance of the class for each.
 	 * Allocates one sub instance to run each of the groups specified in the class, so state is shared. If a layer is not grouped (ie. NAME_None), then state is not shared
 	 * and a separate sub-instance is allocated for each layer node.
-	 * If InClass is null, then layers are reset to their defaults.
+	 * If InClass is null, then all layers are reset to their defaults.
 	 */
 	UFUNCTION(BlueprintCallable, Category = "Layers")
 	void SetLayerOverlay(TSubclassOf<UAnimInstance> InClass);
 
-	/** Gets the sub instance corresponding to the specified group */
+	/** 
+	 * Runs through all layer nodes, attempting to find layer nodes that are currently running the specified class, then resets each to its default value.
+	 * State sharing rules are as with SetLayerOverlay.
+	 * If InClass is null, does nothing.
+	 */
+	UFUNCTION(BlueprintCallable, Category = "Layers")
+	void ClearLayerOverlay(TSubclassOf<UAnimInstance> InClass);
+
+	/** Gets the layer sub instance corresponding to the specified group */
 	UFUNCTION(BlueprintCallable, Category = "Layers")
 	UAnimInstance* GetLayerSubInstanceByGroup(FName InGroup) const;
 
+	/** Gets the first layer sub instance corresponding to the specified class */
+	UFUNCTION(BlueprintCallable, Category = "Layers")
+	UAnimInstance* GetLayerSubInstanceByClass(TSubclassOf<UAnimInstance> InClass) const;
+
 	/** Sets up initial layer groupings */
 	void InitializeGroupedLayers();
 
+private:
+	/** Helper function to perform layer overlay actions (set, clear) */
+	void PerformLayerOverlayOperation(TSubclassOf<UAnimInstance> InClass, TFunctionRef<UClass*(UClass*, FAnimNode_Layer*)> InClassSelectorFunction);
+
 protected:
 	/** Map between Active Montages and their FAnimMontageInstance */
 	TMap<class UAnimMontage*, struct FAnimMontageInstance*> ActiveMontagesMap;
diff --git a/Engine/Source/Runtime/Engine/Classes/Components/SkeletalMeshComponent.h b/Engine/Source/Runtime/Engine/Classes/Components/SkeletalMeshComponent.h
--- a/Engine/Source/Runtime/Engine/Classes/Components/SkeletalMeshComponent.h
+++ b/Engine/Source/Runtime/Engine/Classes/Components/SkeletalMeshComponent.h
@@ -787,17 +787,30 @@ public:
 	void GetSubInstancesByTag(FName InTag, TArray<UAnimInstance*>& OutSubInstances) const;
 
 	/** 
-	 * Runs through all layer nodes, attempting to find nodes that are implemented by the specified class, then sets up a sub instance of the class for each.
-	 * Allocates one sub instance to run each of the groups specified in the class, so state is shared.
-	 * If InClass is null, then layers are reset to their defaults.
+	 * Runs through all layer nodes, attempting to find layer nodes that are implemented by the specified class, then sets up a sub instance of the class for each.
+	 * Allocates one sub instance to run each of the groups specified in the class, so state is shared. If a layer is not grouped (ie. NAME_None), then state is not shared
+	 * and a separate sub-instance is allocated for each layer node.
+	 * If InClass is null, then all layers are reset to their defaults.
 	 */
 	UFUNCTION(BlueprintCallable, Category = "Components|SkeletalMesh|Layers")
 	void SetLayerOverlay(TSubclassOf<UAnimInstance> InClass);
 
-	/** Gets the sub instance corresponding to the specified group */
+	/** 
+	 * Runs through all layer nodes, attempting to find layer nodes that are currently running the specified class, then resets each to its default value.
+	 * State sharing rules are as with SetLayerOverlay.
+	 * If InClass is null, does nothing.
+	 */
+	UFUNCTION(BlueprintCallable, Category = "Components|SkeletalMesh|Layers")
+	void ClearLayerOverlay(TSubclassOf<UAnimInstance> InClass);
+
+	/** Gets the layer sub instance corresponding to the specified group */
 	UFUNCTION(BlueprintCallable, Category = "Components|SkeletalMesh|Layers")
 	UAnimInstance* GetLayerSubInstanceByGroup(FName InGroup) const;
 
+	/** Gets the first layer sub instance corresponding to the specified class */
+	UFUNCTION(BlueprintCallable, Category = "Components|SkeletalMesh|Layers")
+	UAnimInstance* GetLayerSubInstanceByClass(TSubclassOf<UAnimInstance> InClass) const;
+
 	/** 
 	 * Returns whether there are any valid instances to run, currently this means whether we have
 	 * have an animation instance or a post process instance available to process.
diff --git a/Engine/Source/Runtime/Engine/Private/Animation/AnimInstance.cpp b/Engine/Source/Runtime/Engine/Private/Animation/AnimInstance.cpp
--- a/Engine/Source/Runtime/Engine/Private/Animation/AnimInstance.cpp
+++ b/Engine/Source/Runtime/Engine/Private/Animation/AnimInstance.cpp
@@ -2515,7 +2515,7 @@ void UAnimInstance::SetSubInstanceClassByTag(FName InTag, TSubclassOf<UAnimInsta
 	}
 }
 
-void UAnimInstance::SetLayerOverlay(TSubclassOf<UAnimInstance> InClass)
+void UAnimInstance::PerformLayerOverlayOperation(TSubclassOf<UAnimInstance> InClass, TFunctionRef<UClass*(UClass* InClass, FAnimNode_Layer*)> InClassSelectorFunction)
 {
 	if (IAnimClassInterface* AnimBlueprintClass = IAnimClassInterface::GetFromClass(GetClass()))
 	{
@@ -2533,6 +2533,7 @@ void UAnimInstance::SetLayerOverlay(TSubclassOf<UAnimInstance> InClass)
 			}
 		}
 
+		// Make sure we have valid objects as initialization can route back out of sub-instances into this outer graph
 		GetProxyOnAnyThread<FAnimInstanceProxy>().InitializeObjects(this);
 
 		// Map of group name->nodes to run under that group instance
@@ -2567,17 +2568,22 @@ void UAnimInstance::SetLayerOverlay(TSubclassOf<UAnimInstance> InClass)
 			}
 		}
 
-		if(LayerNodesToSet.Num() > 0)
+		for(TPair<FName, TArray<FAnimNode_Layer*, TInlineAllocator<4>>> LayerPair : LayerNodesToSet)
 		{
-			for(TPair<FName, TArray<FAnimNode_Layer*, TInlineAllocator<4>>> LayerPair : LayerNodesToSet)
+			if (LayerPair.Key == NAME_None)
 			{
-				if (LayerPair.Key == NAME_None)
+				// Ungrouped path - each layer gets a separate instance
+				for(FAnimNode_Layer* LayerNode : LayerPair.Value)
 				{
-					// Ungrouped path - each layer gets a separate instance
-					for(FAnimNode_Layer* LayerNode : LayerPair.Value)
+					UClass* ClassToSet = InClassSelectorFunction(NewClass, LayerNode);
+
+					// Disallow setting the same class as this instance, which would create infinite recursion
+					if (ClassToSet != nullptr && ClassToSet != GetClass())
 					{
-						UClass* ClassToSet = NewClass != nullptr ? NewClass : LayerNode->InstanceClass.Get();
-						if (ClassToSet != nullptr && ClassToSet != GetClass())
+						UAnimInstance* TargetInstance = LayerNode->GetTargetInstance<UAnimInstance>();
+
+						// Skip setting if the class is the same
+						if (TargetInstance == nullptr || ClassToSet != TargetInstance->GetClass())
 						{
 							USkeletalMeshComponent* MeshComp = GetSkelMeshComponent();
 							UAnimInstance* NewSubInstance = NewObject<UAnimInstance>(MeshComp, ClassToSet);
@@ -2599,13 +2605,26 @@ void UAnimInstance::SetLayerOverlay(TSubclassOf<UAnimInstance> InClass)
 							MeshComp->SubInstances.Add(NewSubInstance);
 						}
 					}
+					else
+					{
+						LayerNode->SetLayerOverlaySubInstance(this, nullptr);
+					}
 				}
-				else
-				{		
-					// Grouped path - each group gets an instance
-					// If the class is null, then reset to default (which can be null)
-					UClass* ClassToSet = NewClass != nullptr ? NewClass : LayerPair.Value[0]->InstanceClass.Get();
-					if(ClassToSet != nullptr && ClassToSet != GetClass())
+			}
+			else
+			{		
+				// Grouped path - each group gets an instance
+				// If the class is null, then reset to default (which can be null)
+				FAnimNode_Layer* FirstLayerNode = LayerPair.Value[0];
+				UClass* ClassToSet = InClassSelectorFunction(NewClass, FirstLayerNode);
+
+				// Disallow setting the same class as this instance, which would create infinite recursion
+				if(ClassToSet != nullptr && ClassToSet != GetClass())
+				{
+					UAnimInstance* TargetInstance = FirstLayerNode->GetTargetInstance<UAnimInstance>();
+
+					// Skip setting if the class is the same
+					if (TargetInstance == nullptr || ClassToSet != TargetInstance->GetClass())
 					{
 						// Create and add one sub-instance for this group
 						USkeletalMeshComponent* MeshComp = GetSkelMeshComponent();
@@ -2634,13 +2653,13 @@ void UAnimInstance::SetLayerOverlay(TSubclassOf<UAnimInstance> InClass)
 
 						MeshComp->SubInstances.Add(NewSubInstance);
 					}
-					else
+				}
+				else
+				{
+					// Clear the node's instance - we didnt find a class to use
+					for(FAnimNode_Layer* LayerNode : LayerPair.Value)
 					{
-						// Clear the node's instance - we didnt find a class to use
-						for(FAnimNode_Layer* LayerNode : LayerPair.Value)
-						{
-							LayerNode->SetLayerOverlaySubInstance(this, nullptr);
-						}
+						LayerNode->SetLayerOverlaySubInstance(this, nullptr);
 					}
 				}
 			}
@@ -2648,6 +2667,44 @@ void UAnimInstance::SetLayerOverlay(TSubclassOf<UAnimInstance> InClass)
 	}
 }
 
+void UAnimInstance::SetLayerOverlay(TSubclassOf<UAnimInstance> InClass)
+{
+	auto SelectResolvedClassIfValid = [](UClass* InResolvedClass, FAnimNode_Layer* InLayerNode)
+	{
+		if(InResolvedClass != nullptr)
+		{
+			// If we have a valid resolved class, use that as an overlay
+			return InResolvedClass;
+		}
+		else
+		{
+			// Otherwise use the default (which can be null)
+			return InLayerNode->InstanceClass.Get();
+		}
+	};
+
+	PerformLayerOverlayOperation(InClass, SelectResolvedClassIfValid);
+}
+
+void UAnimInstance::ClearLayerOverlay(TSubclassOf<UAnimInstance> InClass)
+{
+	auto ConditionallySelectDefaultClass = [](UClass* InResolvedClass, FAnimNode_Layer* InLayerNode)
+	{
+		if (InResolvedClass != nullptr && InLayerNode->GetTargetInstance<UAnimInstance>()->GetClass() == InResolvedClass)
+		{
+			// Reset to default if the classes match
+			return InLayerNode->InstanceClass.Get();
+		}
+		else
+		{
+			// No change
+			return InLayerNode->GetTargetInstance<UAnimInstance>()->GetClass();
+		}
+	};
+
+	PerformLayerOverlayOperation(InClass, ConditionallySelectDefaultClass);
+}
+
 void UAnimInstance::InitializeGroupedLayers()
 {
 	SetLayerOverlay(nullptr);
@@ -2686,6 +2743,25 @@ UAnimInstance* UAnimInstance::GetLayerSubInstanceByGroup(FName InGroup) const
 	return nullptr;
 }
 
+UAnimInstance* UAnimInstance::GetLayerSubInstanceByClass(TSubclassOf<UAnimInstance> InClass) const
+{
+	if (IAnimClassInterface* AnimBlueprintClass = IAnimClassInterface::GetFromClass(GetClass()))
+	{
+		const TArray<UStructProperty*>& LayerNodeProperties = AnimBlueprintClass->GetLayerNodeProperties();
+		for(UStructProperty* LayerNodeProperty : LayerNodeProperties)
+		{
+			const FAnimNode_Layer* Layer = LayerNodeProperty->ContainerPtrToValuePtr<FAnimNode_Layer>(this);
+			UAnimInstance* TargetInstance = Layer->GetTargetInstance<UAnimInstance>();
+			if (TargetInstance && TargetInstance->GetClass() == InClass.Get())
+			{
+				return TargetInstance;
+			}
+		}
+	}
+
+	return nullptr;
+}
+
 FAnimMontageInstance* UAnimInstance::GetActiveInstanceForMontage(UAnimMontage const& Montage) const
 {
 	return GetActiveInstanceForMontage(&Montage);
diff --git a/Engine/Source/Runtime/Engine/Private/Components/SkeletalMeshComponent.cpp b/Engine/Source/Runtime/Engine/Private/Components/SkeletalMeshComponent.cpp
--- a/Engine/Source/Runtime/Engine/Private/Components/SkeletalMeshComponent.cpp
+++ b/Engine/Source/Runtime/Engine/Private/Components/SkeletalMeshComponent.cpp
@@ -2769,6 +2769,14 @@ void USkeletalMeshComponent::SetLayerOverlay(TSubclassOf<UAnimInstance> InClass)
 	}
 }
 
+void USkeletalMeshComponent::ClearLayerOverlay(TSubclassOf<UAnimInstance> InClass)
+{
+	if(AnimScriptInstance)
+	{
+		AnimScriptInstance->ClearLayerOverlay(InClass);
+	}
+}
+
 UAnimInstance* USkeletalMeshComponent::GetLayerSubInstanceByGroup(FName InGroup) const
 {
 	if(AnimScriptInstance)
@@ -2778,6 +2786,15 @@ UAnimInstance* USkeletalMeshComponent::GetLayerSubInstanceByGroup(FName InGroup)
 	return nullptr;
 }
 
+UAnimInstance* USkeletalMeshComponent::GetLayerSubInstanceByClass(TSubclassOf<UAnimInstance> InClass) const
+{
+	if(AnimScriptInstance)
+	{
+		return AnimScriptInstance->GetLayerSubInstanceByClass(InClass);
+	}
+	return nullptr;
+}
+
 bool USkeletalMeshComponent::HasValidAnimationInstance() const
 {
 	return AnimScriptInstance || PostProcessAnimInstance;
