Expose RVT GetLayerCount() as static so that we don't need a object to get a layer count
Opens up the way for RVT sample parameter node with no RVT set

diff --git a/Engine/Source/Runtime/Engine/Classes/VT/RuntimeVirtualTexture.h b/Engine/Source/Runtime/Engine/Classes/VT/RuntimeVirtualTexture.h
--- a/Engine/Source/Runtime/Engine/Classes/VT/RuntimeVirtualTexture.h
+++ b/Engine/Source/Runtime/Engine/Classes/VT/RuntimeVirtualTexture.h
@@ -84,6 +84,8 @@ public:
 
 	/** Returns number of texture layers in the virtual texture */
 	int32 GetLayerCount() const;
+	/** Returns number of texture layers in the virtual texture of a given material type */
+	static int32 GetLayerCount(ERuntimeVirtualTextureMaterialType InMaterialType);
 	/** Returns the texture format for the virtual texture layer */
 	EPixelFormat GetLayerFormat(int32 LayerIndex) const;
 	/** Return true if the virtual texture layer should be sampled as sRGB */
@@ -103,7 +105,7 @@ public:
 	IAllocatedVirtualTexture* GetAllocatedVirtualTexture() const;
 
 	/** Getter for the shader uniform parameters. */
-	FVector4 GetUniformParameter(int32 Index);
+	FVector4 GetUniformParameter(int32 Index) const;
 
 #if WITH_EDITOR
 	/** Get a hash of the current state to use for streaming texture invalidation. */
diff --git a/Engine/Source/Runtime/Engine/Private/Materials/MaterialExpressions.cpp b/Engine/Source/Runtime/Engine/Private/Materials/MaterialExpressions.cpp
--- a/Engine/Source/Runtime/Engine/Private/Materials/MaterialExpressions.cpp
+++ b/Engine/Source/Runtime/Engine/Private/Materials/MaterialExpressions.cpp
@@ -2165,7 +2165,7 @@ int32 UMaterialExpressionRuntimeVirtualTextureSample::Compile(class FMaterialCom
 	
 	// Compile the texture object references
 	enum { MAX_RVT_LAYERS = 2 };
-	const int32 LayerCount = VirtualTexture->GetLayerCount();
+	const int32 LayerCount = URuntimeVirtualTexture::GetLayerCount(MaterialType);
 	check(LayerCount <= MAX_RVT_LAYERS);
 
 	int32 TextureCodeIndex[MAX_RVT_LAYERS] = { INDEX_NONE };
diff --git a/Engine/Source/Runtime/Engine/Private/VT/RuntimeVirtualTexture.cpp b/Engine/Source/Runtime/Engine/Private/VT/RuntimeVirtualTexture.cpp
--- a/Engine/Source/Runtime/Engine/Private/VT/RuntimeVirtualTexture.cpp
+++ b/Engine/Source/Runtime/Engine/Private/VT/RuntimeVirtualTexture.cpp
@@ -243,9 +243,9 @@ void URuntimeVirtualTexture::GetProducerDescription(FVTProducerDescription& OutD
 	}
 }
 
-int32 URuntimeVirtualTexture::GetLayerCount() const
+int32 URuntimeVirtualTexture::GetLayerCount(ERuntimeVirtualTextureMaterialType InMaterialType)
 {
-	switch (MaterialType)
+	switch (InMaterialType)
 	{
 	case ERuntimeVirtualTextureMaterialType::BaseColor:
 	case ERuntimeVirtualTextureMaterialType::WorldHeight:
@@ -263,6 +263,11 @@ int32 URuntimeVirtualTexture::GetLayerCount() const
 	return 1;
 }
 
+int32 URuntimeVirtualTexture::GetLayerCount() const
+{
+	return GetLayerCount(MaterialType);
+}
+
 EPixelFormat URuntimeVirtualTexture::GetLayerFormat(int32 LayerIndex) const
 {
 	if (LayerIndex == 0)
@@ -345,7 +350,7 @@ IAllocatedVirtualTexture* URuntimeVirtualTexture::GetAllocatedVirtualTexture() c
 	return Resource->GetAllocatedVirtualTexture();
 }
 
-FVector4 URuntimeVirtualTexture::GetUniformParameter(int32 Index)
+FVector4 URuntimeVirtualTexture::GetUniformParameter(int32 Index) const
 {
 	check(Index >= 0 && Index < sizeof(WorldToUVTransformParameters)/sizeof(WorldToUVTransformParameters[0]));
 	
