Add interface that will allow FMaterialUniformExpressionTextureParameter to return a URuntimeVirtualTexture
Deprecate FMaterialUniformExpressionRuntimeVirtualTextureParameter in favor of FMaterialUniformExpressionRuntimeVirtualTextureUniform which will support material instance RVT parameters

diff --git a/Engine/Source/Runtime/Engine/Private/Materials/HLSLMaterialTranslator.h b/Engine/Source/Runtime/Engine/Private/Materials/HLSLMaterialTranslator.h
--- a/Engine/Source/Runtime/Engine/Private/Materials/HLSLMaterialTranslator.h
+++ b/Engine/Source/Runtime/Engine/Private/Materials/HLSLMaterialTranslator.h
@@ -5074,7 +5074,7 @@ protected:
 
 	virtual int32 VirtualTextureParam(int32 TextureIndex, int32 ParamIndex) override
 	{
-		return AddUniformExpression(new FMaterialUniformExpressionRuntimeVirtualTextureParameter(TextureIndex, ParamIndex), MCT_Float3, TEXT(""));
+		return AddUniformExpression(new FMaterialUniformExpressionRuntimeVirtualTextureUniform(TextureIndex, ParamIndex), MCT_Float3, TEXT(""));
 	}
 
 	virtual int32 VirtualTextureWorldToUV(int32 WorldPositionIndex, int32 P0, int32 P1, int32 P2) override
diff --git a/Engine/Source/Runtime/Engine/Private/Materials/MaterialUniformExpressions.cpp b/Engine/Source/Runtime/Engine/Private/Materials/MaterialUniformExpressions.cpp
--- a/Engine/Source/Runtime/Engine/Private/Materials/MaterialUniformExpressions.cpp
+++ b/Engine/Source/Runtime/Engine/Private/Materials/MaterialUniformExpressions.cpp
@@ -491,7 +491,7 @@ void FUniformExpressionSet::FillUniformBuffer(const FMaterialRenderContext& Mate
 			if (!bFoundTexture)
 			{
 				const URuntimeVirtualTexture* Texture = nullptr;
-				UniformVirtualTextureExpressions[ExpressionIndex]->GetTextureValue(MaterialRenderContext.Material, Texture);
+				UniformVirtualTextureExpressions[ExpressionIndex]->GetTextureValue(MaterialRenderContext, MaterialRenderContext.Material, Texture);
 				if (Texture != nullptr)
 				{
 					int32 LayerIndex = UniformVirtualTextureExpressions[ExpressionIndex]->GetLayerIndex();
@@ -815,7 +815,7 @@ void FUniformExpressionSet::FillUniformBuffer(const FMaterialRenderContext& Mate
 			if (!bValidResources)
 			{
 				const URuntimeVirtualTexture* Texture;
-				UniformVirtualTextureExpressions[ExpressionIndex]->GetTextureValue(MaterialRenderContext.Material, Texture);
+				UniformVirtualTextureExpressions[ExpressionIndex]->GetTextureValue(MaterialRenderContext, MaterialRenderContext.Material, Texture);
 				if (Texture != nullptr)
 				{
 					IAllocatedVirtualTexture const* AllocatedVT = Texture->GetAllocatedVirtualTexture();
@@ -946,7 +946,7 @@ void FMaterialUniformExpressionTexture::GetTextureValue(const FMaterialRenderCon
 	}
 }
 
-void FMaterialUniformExpressionTexture::GetTextureValue(const FMaterial& Material, const URuntimeVirtualTexture*& OutValue) const
+void FMaterialUniformExpressionTexture::GetTextureValue(const FMaterialRenderContext& Context, const FMaterial& Material, const URuntimeVirtualTexture*& OutValue) const
 {
 	check(IsInParallelRenderingThread());
 	OutValue = GetIndexedTexture<URuntimeVirtualTexture>(Material, TextureIndex);
@@ -976,6 +976,61 @@ bool FMaterialUniformExpressionTexture::IsIdentical(const FMaterialUniformExpres
 	return TextureIndex == OtherTextureExpression->TextureIndex && LayerIndex == OtherTextureExpression->LayerIndex && bVirtualTexture == OtherTextureExpression->bVirtualTexture;
 }
 
+void FMaterialUniformExpressionTextureParameter::GetTextureValue(const FMaterialRenderContext& Context, const FMaterial& Material, const UTexture*& OutValue) const
+{
+	check(IsInParallelRenderingThread());
+	if (TransientOverrideValue_RenderThread != NULL)
+	{
+		OutValue = TransientOverrideValue_RenderThread;
+	}
+	else
+	{
+		if (!Context.MaterialRenderProxy || !Context.MaterialRenderProxy->GetTextureValue(ParameterInfo, &OutValue, Context))
+		{
+			UTexture* Value = nullptr;
+
+			if (AreExperimentalMaterialLayersEnabled())
+			{
+				UMaterialInterface* Interface = Context.Material.GetMaterialInterface();
+				if (!Interface || !Interface->GetTextureParameterDefaultValue(ParameterInfo, Value))
+				{
+					Value = GetIndexedTexture<UTexture>(Material, TextureIndex);
+				}
+			}
+			else
+			{
+				Value = GetIndexedTexture<UTexture>(Material, TextureIndex);
+			}
+
+			OutValue = Value;
+		}
+	}
+}
+
+void FMaterialUniformExpressionTextureParameter::GetTextureValue(const FMaterialRenderContext& Context, const FMaterial& Material, const URuntimeVirtualTexture*& OutValue) const
+{
+	// todo[vt]: Add runtime virtual texture parameter
+	check(false);
+}
+
+void FMaterialUniformExpressionTextureParameter::GetGameThreadTextureValue(const UMaterialInterface* MaterialInterface, const FMaterial& Material, UTexture*& OutValue, bool bAllowOverride) const
+{
+	check(IsInGameThread());
+	if (bAllowOverride && TransientOverrideValue_GameThread != NULL)
+	{
+		OutValue = TransientOverrideValue_GameThread;
+	}
+	else
+	{
+		OutValue = NULL;
+		const bool bOverrideValuesOnly = !AreExperimentalMaterialLayersEnabled();
+		if (!MaterialInterface->GetTextureParameterValue(ParameterInfo, OutValue, bOverrideValuesOnly))
+		{
+			OutValue = GetIndexedTexture<UTexture>(Material, TextureIndex);
+		}
+	}
+}
+
 FMaterialUniformExpressionExternalTextureBase::FMaterialUniformExpressionExternalTextureBase(int32 InSourceTextureIndex)
 	: SourceTextureIndex(InSourceTextureIndex)
 {}
@@ -1225,41 +1280,56 @@ void FMaterialUniformExpressionExternalTextureCoordinateOffset::GetNumberValue(c
 	}
 }
 
-FMaterialUniformExpressionRuntimeVirtualTextureParameter::FMaterialUniformExpressionRuntimeVirtualTextureParameter()
-	: TextureIndex(INDEX_NONE)
-	, ParamIndex(INDEX_NONE)
+FMaterialUniformExpressionRuntimeVirtualTextureUniform::FMaterialUniformExpressionRuntimeVirtualTextureUniform()
+	: bParameter(false)
+	, TextureIndex(INDEX_NONE)
+	, VectorIndex(INDEX_NONE)
 {
 }
 
-FMaterialUniformExpressionRuntimeVirtualTextureParameter::FMaterialUniformExpressionRuntimeVirtualTextureParameter(int32 InTextureIndex, int32 InParamIndex)
-	: TextureIndex(InTextureIndex)
-	, ParamIndex(InParamIndex)
+FMaterialUniformExpressionRuntimeVirtualTextureUniform::FMaterialUniformExpressionRuntimeVirtualTextureUniform(int32 InTextureIndex, int32 InVectorIndex)
+	: bParameter(false)
+	, TextureIndex(InTextureIndex)
+	, VectorIndex(InVectorIndex)
+{
+}
+
+FMaterialUniformExpressionRuntimeVirtualTextureUniform::FMaterialUniformExpressionRuntimeVirtualTextureUniform(const FMaterialParameterInfo& InParameterInfo, int32 InTextureIndex, int32 InVectorIndex)
+	: bParameter(true)
+	, ParameterInfo(InParameterInfo)
+	, TextureIndex(InTextureIndex)
+	, VectorIndex(InVectorIndex)
 {
 }
 
-void FMaterialUniformExpressionRuntimeVirtualTextureParameter::Serialize(FArchive& Ar)
+void FMaterialUniformExpressionRuntimeVirtualTextureUniform::Serialize(FArchive& Ar)
 {
+	Ar << bParameter;
+	Ar << ParameterInfo;
 	Ar << TextureIndex;
-	Ar << ParamIndex;
+	Ar << VectorIndex;
 }
 
-bool FMaterialUniformExpressionRuntimeVirtualTextureParameter::IsIdentical(const FMaterialUniformExpression* OtherExpression) const
+bool FMaterialUniformExpressionRuntimeVirtualTextureUniform::IsIdentical(const FMaterialUniformExpression* OtherExpression) const
 {
 	if (GetType() != OtherExpression->GetType())
 	{
 		return false;
 	}
 
-	const auto* Other = static_cast<const FMaterialUniformExpressionRuntimeVirtualTextureParameter*>(OtherExpression);
-	return TextureIndex == Other->TextureIndex && ParamIndex == Other->ParamIndex;
+	const auto* Other = static_cast<const FMaterialUniformExpressionRuntimeVirtualTextureUniform*>(OtherExpression);
+	return ParameterInfo == Other->ParameterInfo && TextureIndex == Other->TextureIndex && VectorIndex == Other->VectorIndex;
 }
 
-void FMaterialUniformExpressionRuntimeVirtualTextureParameter::GetNumberValue(const struct FMaterialRenderContext& Context, FLinearColor& OutValue) const
+void FMaterialUniformExpressionRuntimeVirtualTextureUniform::GetNumberValue(const struct FMaterialRenderContext& Context, FLinearColor& OutValue) const
 {
+	//todo[vt]: Add support for parameters
+	check(!bParameter)
+
 	URuntimeVirtualTexture* Texture = GetIndexedTexture<URuntimeVirtualTexture>(Context.Material, TextureIndex);
-	if (Texture != nullptr && ParamIndex != INDEX_NONE)
+	if (Texture != nullptr && VectorIndex != INDEX_NONE)
 	{
-		OutValue = FLinearColor(Texture->GetUniformParameter(ParamIndex));
+		OutValue = FLinearColor(Texture->GetUniformParameter(VectorIndex));
 	}
 	else
 	{
@@ -1267,6 +1337,24 @@ void FMaterialUniformExpressionRuntimeVirtualTextureParameter::GetNumberValue(co
 	}
 }
 
+/**
+ * Deprecated FMaterialUniformExpressionRuntimeVirtualTextureParameter in favor of FMaterialUniformExpressionRuntimeVirtualTextureUniform
+ * Keep around until we no longer need to support serialization of 4.23 data
+ */
+class FMaterialUniformExpressionRuntimeVirtualTextureParameter_DEPRECATED : public FMaterialUniformExpressionRuntimeVirtualTextureUniform
+{
+	DECLARE_MATERIALUNIFORMEXPRESSION_TYPE(FMaterialUniformExpressionRuntimeVirtualTextureParameter_DEPRECATED);
+
+public:
+	virtual void Serialize(FArchive& Ar) override
+	{
+		Ar << TextureIndex;
+		Ar << VectorIndex;
+	}
+};
+// Specific implementation of FMaterialUniformExpressionRuntimeVirtualTextureParameter_DEPRECATED to serialize old class name
+FMaterialUniformExpressionType FMaterialUniformExpressionRuntimeVirtualTextureParameter_DEPRECATED::StaticType(TEXT("FMaterialUniformExpressionRuntimeVirtualTextureParameter"), &FMaterialUniformExpressionRuntimeVirtualTextureParameter_DEPRECATED::SerializationConstructor);
+
 
 IMPLEMENT_MATERIALUNIFORMEXPRESSION_TYPE(FMaterialUniformExpressionTexture);
 IMPLEMENT_MATERIALUNIFORMEXPRESSION_TYPE(FMaterialUniformExpressionConstant);
@@ -1278,7 +1366,7 @@ IMPLEMENT_MATERIALUNIFORMEXPRESSION_TYPE(FMaterialUniformExpressionExternalTextu
 IMPLEMENT_MATERIALUNIFORMEXPRESSION_TYPE(FMaterialUniformExpressionExternalTextureParameter);
 IMPLEMENT_MATERIALUNIFORMEXPRESSION_TYPE(FMaterialUniformExpressionExternalTextureCoordinateScaleRotation);
 IMPLEMENT_MATERIALUNIFORMEXPRESSION_TYPE(FMaterialUniformExpressionExternalTextureCoordinateOffset);
-IMPLEMENT_MATERIALUNIFORMEXPRESSION_TYPE(FMaterialUniformExpressionRuntimeVirtualTextureParameter);
+IMPLEMENT_MATERIALUNIFORMEXPRESSION_TYPE(FMaterialUniformExpressionRuntimeVirtualTextureUniform);
 IMPLEMENT_MATERIALUNIFORMEXPRESSION_TYPE(FMaterialUniformExpressionFlipBookTextureParameter);
 IMPLEMENT_MATERIALUNIFORMEXPRESSION_TYPE(FMaterialUniformExpressionSine);
 IMPLEMENT_MATERIALUNIFORMEXPRESSION_TYPE(FMaterialUniformExpressionSquareRoot);
diff --git a/Engine/Source/Runtime/Engine/Private/Materials/MaterialUniformExpressions.h b/Engine/Source/Runtime/Engine/Private/Materials/MaterialUniformExpressions.h
--- a/Engine/Source/Runtime/Engine/Private/Materials/MaterialUniformExpressions.h
+++ b/Engine/Source/Runtime/Engine/Private/Materials/MaterialUniformExpressions.h
@@ -306,9 +306,14 @@ public:
 
 	FMaterialUniformExpressionTextureParameter() {}
 
-	FMaterialUniformExpressionTextureParameter(const FMaterialParameterInfo& InParameterInfo, int32 InTextureIndex, EMaterialSamplerType InSamplerType, ESamplerSourceMode InSourceMode, bool InVirtualTexture) :
-		Super(InTextureIndex, InSamplerType, InSourceMode, InVirtualTexture),
-		ParameterInfo(InParameterInfo)
+	FMaterialUniformExpressionTextureParameter(const FMaterialParameterInfo& InParameterInfo, int32 InTextureIndex, EMaterialSamplerType InSamplerType, ESamplerSourceMode InSourceMode, bool InVirtualTexture)
+		: Super(InTextureIndex, InSamplerType, InSourceMode, InVirtualTexture)
+		, ParameterInfo(InParameterInfo)
+	{}
+
+	FMaterialUniformExpressionTextureParameter(const FMaterialParameterInfo& InParameterInfo, int32 InTextureIndex, int32 InLayerIndex, EMaterialSamplerType InSamplerType)
+		: Super(InTextureIndex, InLayerIndex, InSamplerType)
+		, ParameterInfo(InParameterInfo)
 	{}
 
 	// FMaterialUniformExpression interface.
@@ -319,53 +324,10 @@ public:
 		Ar << ParameterInfo;
 		Super::Serialize(Ar);
 	}
-	virtual void GetTextureValue(const FMaterialRenderContext& Context, const FMaterial& Material, const UTexture*& OutValue) const override
-	{
-		check(IsInParallelRenderingThread());
-		if (TransientOverrideValue_RenderThread != NULL)
-		{
-			OutValue = TransientOverrideValue_RenderThread;
-		}
-		else
-		{
-			if (!Context.MaterialRenderProxy || !Context.MaterialRenderProxy->GetTextureValue(ParameterInfo, &OutValue, Context))
-			{
-				UTexture* Value = nullptr;
 
-				if (AreExperimentalMaterialLayersEnabled())
-				{
-					UMaterialInterface* Interface = Context.Material.GetMaterialInterface();
-					if (!Interface || !Interface->GetTextureParameterDefaultValue(ParameterInfo, Value))
-					{
-						Value = GetIndexedTexture<UTexture>(Material, TextureIndex);
-					}
-				}
-				else
-				{
-					Value = GetIndexedTexture<UTexture>(Material, TextureIndex);
-				}
-
-				OutValue = Value;
-			}
-		}
-	}
-	virtual void GetGameThreadTextureValue(const UMaterialInterface* MaterialInterface,const FMaterial& Material,UTexture*& OutValue,bool bAllowOverride=true) const
-	{
-		check(IsInGameThread());
-		if( bAllowOverride && TransientOverrideValue_GameThread != NULL )
-		{
-			OutValue = TransientOverrideValue_GameThread;
-		}
-		else
-		{
-			OutValue = NULL;
-		const bool bOverrideValuesOnly = !AreExperimentalMaterialLayersEnabled();
-			if(!MaterialInterface->GetTextureParameterValue(ParameterInfo,OutValue,bOverrideValuesOnly))
-		{
-				OutValue = GetIndexedTexture<UTexture>(Material, TextureIndex);
-			}
-		}
-	}
+	virtual void GetTextureValue(const FMaterialRenderContext& Context, const FMaterial& Material, const UTexture*& OutValue) const override;
+	virtual void GetTextureValue(const FMaterialRenderContext& Context, const FMaterial& Material, const URuntimeVirtualTexture*& OutValue) const override;
+	virtual void GetGameThreadTextureValue(const UMaterialInterface* MaterialInterface, const FMaterial& Material, UTexture*& OutValue, bool bAllowOverride = true) const override;
 
 	virtual bool IsConstant() const
 	{
@@ -1760,16 +1722,18 @@ protected:
 };
 
 /**
- * A uniform expression to retrieve one of the parameters associated with a URuntimeVirtualTexture
+ * A uniform expression to retrieve one of the vector uniform parameters stored in a URuntimeVirtualTexture
  */
-class FMaterialUniformExpressionRuntimeVirtualTextureParameter : public FMaterialUniformExpression
+class FMaterialUniformExpressionRuntimeVirtualTextureUniform : public FMaterialUniformExpression
 {
-	DECLARE_MATERIALUNIFORMEXPRESSION_TYPE(FMaterialUniformExpressionRuntimeVirtualTextureParameter);
+	DECLARE_MATERIALUNIFORMEXPRESSION_TYPE(FMaterialUniformExpressionRuntimeVirtualTextureUniform);
 
 public:
-	FMaterialUniformExpressionRuntimeVirtualTextureParameter();
-	/** Construct with the index of the texture reference and the parameter index that we want to retrieve. */
-	FMaterialUniformExpressionRuntimeVirtualTextureParameter(int32 InTextureIndex, int32 InParamIndex);
+	FMaterialUniformExpressionRuntimeVirtualTextureUniform();
+	/** Construct with the index of the texture reference and the vector index that we want to retrieve. */
+	FMaterialUniformExpressionRuntimeVirtualTextureUniform(int32 InTextureIndex, int32 InVectorIndex);
+	/** Construct with a URuntimeVirtualTexture parameter and the vector index that we want to retrieve. */
+	FMaterialUniformExpressionRuntimeVirtualTextureUniform(const FMaterialParameterInfo& InParameterInfo, int32 InTextureIndex, int32 InVectorIndex);
 
 	//~ Begin FMaterialUniformExpression Interface.
 	virtual bool IsConstant() const override { return false; }
@@ -1779,8 +1743,12 @@ public:
 	//~ End FMaterialUniformExpression Interface.
 
 protected:
-	/** Index of a URuntimeVirtualTexture in the material texture references. */
+	/** Is this expression using a material instance parameter. */
+	bool bParameter;
+	/** Contains the parameter info used if bParameter is true. */
+	FMaterialParameterInfo ParameterInfo;
+	/** Index of the associated URuntimeVirtualTexture in the material texture references used if bParameter is false. */
 	int32 TextureIndex;
-	/** Index of the parameter to fetch from the URuntimeVirtualTexture. */
-	int32 ParamIndex;
+	/** Index of the uniform vector to fetch from the URuntimeVirtualTexture. */
+	int32 VectorIndex;
 };
diff --git a/Engine/Source/Runtime/Engine/Public/MaterialShared.h b/Engine/Source/Runtime/Engine/Public/MaterialShared.h
--- a/Engine/Source/Runtime/Engine/Public/MaterialShared.h
+++ b/Engine/Source/Runtime/Engine/Public/MaterialShared.h
@@ -311,7 +311,7 @@ public:
 	/** Get the UTexture referenced by this expression (can be a nullptr). */
 	virtual void GetTextureValue(const FMaterialRenderContext& Context, const FMaterial& Material, const UTexture*& OutValue) const;
 	/** Get the URuntimeVirtualTexture referenced by this expression (can be a nullptr). */
-	virtual void GetTextureValue(const FMaterial& Material, const URuntimeVirtualTexture*& OutValue) const;
+	virtual void GetTextureValue(const FMaterialRenderContext& Context, const FMaterial& Material, const URuntimeVirtualTexture*& OutValue) const;
 	/** Get the UTexture referenced by this expression including any transient override logic. */
 	virtual void GetGameThreadTextureValue(const UMaterialInterface* MaterialInterface, const FMaterial& Material, UTexture*& OutValue, bool bAllowOverride = true) const;
 
