Fix bug where changing RVT tile border didn't invalidate streaming texture

diff --git a/Engine/Source/Runtime/Engine/Private/VT/RuntimeVirtualTexture.cpp b/Engine/Source/Runtime/Engine/Private/VT/RuntimeVirtualTexture.cpp
--- a/Engine/Source/Runtime/Engine/Private/VT/RuntimeVirtualTexture.cpp
+++ b/Engine/Source/Runtime/Engine/Private/VT/RuntimeVirtualTexture.cpp
@@ -450,6 +450,7 @@ uint32 URuntimeVirtualTexture::GetStreamingTextureBuildHash() const
 	Settings.CompressTextures = (uint32)bCompressTextures;
 	Settings.Size = (uint32)Size;
 	Settings.TileSize = (uint32)TileSize;
+	Settings.TileBorderSize = (uint32)TileBorderSize;
 	Settings.StreamLowMips = (uint32)GetStreamLowMips();
 
 	return Settings.PackedValue;
