Added material expression for RVT sample parameter
Needs another round of implementation to manipulate the parameter in a material instance (incoming)

diff --git a/Engine/Source/Runtime/Engine/Classes/Materials/MaterialExpressionRuntimeVirtualTextureSampleParameter.h b/Engine/Source/Runtime/Engine/Classes/Materials/MaterialExpressionRuntimeVirtualTextureSampleParameter.h
new file mode 100644
--- /dev/null
+++ b/Engine/Source/Runtime/Engine/Classes/Materials/MaterialExpressionRuntimeVirtualTextureSampleParameter.h
@@ -0,0 +1,57 @@
+// Copyright 1998-2019 Epic Games, Inc. All Rights Reserved.
+
+#pragma once
+
+#include "CoreMinimal.h"
+#include "Materials/MaterialExpressionRuntimeVirtualTextureSample.h"
+#include "MaterialExpressionRuntimeVirtualTextureSampleParameter.generated.h"
+
+UCLASS(collapsecategories, hidecategories = Object, MinimalAPI)
+class UMaterialExpressionRuntimeVirtualTextureSampleParameter : public UMaterialExpressionRuntimeVirtualTextureSample
+{
+	GENERATED_UCLASS_BODY()
+
+	/** Name to be referenced when we want to find and set this parameter */
+	UPROPERTY(EditAnywhere, Category = MaterialParameter)
+	FName ParameterName;
+
+	/** GUID that should be unique within the material, this is used for parameter renaming. */
+	UPROPERTY()
+	FGuid ExpressionGUID;
+
+	/** The name of the parameter Group to display in MaterialInstance Editor. Default is None group */
+	UPROPERTY(EditAnywhere, Category = MaterialParameter)
+	FName Group;
+
+#if WITH_EDITORONLY_DATA
+	/** Controls where the this parameter is displayed in a material instance parameter list. The lower the number the higher up in the parameter list. */
+	UPROPERTY(EditAnywhere, Category = MaterialParameter)
+	int32 SortPriority;
+#endif
+
+#if WITH_EDITOR
+	/** If this is the named parameter from this material expression, then set its value. */
+	bool SetParameterValue(FName InParameterName, URuntimeVirtualTexture* InValue);
+#endif
+
+	/** Return whether this is the named parameter from this material expression, and if it is then return its value. */
+	bool IsNamedParameter(const FMaterialParameterInfo& ParameterInfo, URuntimeVirtualTexture*& OutValue) const;
+	/** Adds to arrays of parameter info and id with the values used by this material expression. */
+	void GetAllParameterInfo(TArray<FMaterialParameterInfo> &OutParameterInfo, TArray<FGuid> &OutParameterIds, const FMaterialParameterInfo& InBaseParameterInfo) const;
+
+	//~ Begin UMaterialExpression Interface
+#if WITH_EDITOR
+	virtual bool CanRenameNode() const override { return true; }
+	virtual void SetEditableName(const FString& NewName) override;
+	virtual FString GetEditableName() const override;
+	virtual bool HasAParameterName() const override { return true; }
+	virtual void SetParameterName(const FName& Name) override { ParameterName = Name; }
+	virtual FName GetParameterName() const override { return ParameterName; }
+	virtual void ValidateParameterName(const bool bAllowDuplicateName) override;
+	virtual void GetCaption(TArray<FString>& OutCaptions) const override;
+	virtual bool MatchesSearchQuery(const TCHAR* SearchQuery) override;
+	virtual void SetValueToMatchingExpression(UMaterialExpression* OtherExpression) override;
+#endif
+	virtual FGuid& GetParameterExpressionId() override { return ExpressionGUID; }
+	//~ End UMaterialExpression Interface
+};
diff --git a/Engine/Source/Runtime/Engine/Private/Materials/MaterialExpressions.cpp b/Engine/Source/Runtime/Engine/Private/Materials/MaterialExpressions.cpp
--- a/Engine/Source/Runtime/Engine/Private/Materials/MaterialExpressions.cpp
+++ b/Engine/Source/Runtime/Engine/Private/Materials/MaterialExpressions.cpp
@@ -161,6 +161,7 @@
 #include "Materials/MaterialExpressionRuntimeVirtualTextureOutput.h"
 #include "Materials/MaterialExpressionRuntimeVirtualTextureReplace.h"
 #include "Materials/MaterialExpressionRuntimeVirtualTextureSample.h"
+#include "Materials/MaterialExpressionRuntimeVirtualTextureSampleParameter.h"
 #include "Materials/MaterialExpressionVirtualTextureFeatureSwitch.h"
 #include "Materials/MaterialExpressionSaturate.h"
 #include "Materials/MaterialExpressionSceneColor.h"
@@ -2240,6 +2241,116 @@ void UMaterialExpressionRuntimeVirtualTextureSample::GetCaption(TArray<FString>&
 
 #endif // WITH_EDITOR
 
+UMaterialExpressionRuntimeVirtualTextureSampleParameter::UMaterialExpressionRuntimeVirtualTextureSampleParameter(const FObjectInitializer& ObjectInitializer)
+	: Super(ObjectInitializer)
+{
+	// Structure to hold one-time initialization
+	struct FConstructorStatics
+	{
+		FText NAME_Parameters;
+		FConstructorStatics()
+			: NAME_Parameters(LOCTEXT("Parameters", "Parameters"))
+		{
+		}
+	};
+	static FConstructorStatics ConstructorStatics;
+
+	bIsParameterExpression = true;
+
+#if WITH_EDITORONLY_DATA
+	MenuCategories.Add(ConstructorStatics.NAME_Parameters);
+#endif
+}
+
+bool UMaterialExpressionRuntimeVirtualTextureSampleParameter::SetParameterValue(FName InParameterName, URuntimeVirtualTexture* InValue)
+{
+	if (InParameterName == ParameterName)
+	{
+		VirtualTexture = InValue;
+		return true;
+	}
+
+	return false;
+}
+
+bool UMaterialExpressionRuntimeVirtualTextureSampleParameter::IsNamedParameter(const FMaterialParameterInfo& ParameterInfo, URuntimeVirtualTexture*& OutValue) const
+{
+	if (ParameterInfo.Name == ParameterName)
+	{
+		OutValue = VirtualTexture;
+		return true;
+	}
+
+	return false;
+}
+
+void UMaterialExpressionRuntimeVirtualTextureSampleParameter::GetAllParameterInfo(TArray<FMaterialParameterInfo> &OutParameterInfo, TArray<FGuid> &OutParameterIds, const FMaterialParameterInfo& InBaseParameterInfo) const
+{
+	int32 CurrentSize = OutParameterInfo.Num();
+	FMaterialParameterInfo NewParameter(ParameterName, InBaseParameterInfo.Association, InBaseParameterInfo.Index);
+#if WITH_EDITOR
+	NewParameter.ParameterLocation = Material;
+	if (Function != nullptr)
+	{
+		NewParameter.ParameterLocation = Function;
+	}
+	if (HasConnectedOutputs())
+#endif
+	{
+		OutParameterInfo.AddUnique(NewParameter);
+		if (CurrentSize != OutParameterInfo.Num())
+		{
+			OutParameterIds.Add(ExpressionGUID);
+		}
+	}
+}
+
+#if WITH_EDITOR
+
+void UMaterialExpressionRuntimeVirtualTextureSampleParameter::SetEditableName(const FString& NewName)
+{
+	ParameterName = *NewName;
+}
+
+FString UMaterialExpressionRuntimeVirtualTextureSampleParameter::GetEditableName() const
+{
+	return ParameterName.ToString();
+}
+
+void UMaterialExpressionRuntimeVirtualTextureSampleParameter::ValidateParameterName(const bool bAllowDuplicateName)
+{
+	ValidateParameterNameInternal(this, Material, bAllowDuplicateName);
+}
+
+void UMaterialExpressionRuntimeVirtualTextureSampleParameter::GetCaption(TArray<FString>& OutCaptions) const
+{
+	OutCaptions.Add(FString(TEXT("Runtime Virtual Texture Sample Param ")));
+	OutCaptions.Add(FString::Printf(TEXT("'%s'"), *ParameterName.ToString()));
+}
+
+bool UMaterialExpressionRuntimeVirtualTextureSampleParameter::MatchesSearchQuery(const TCHAR* SearchQuery)
+{
+	if (ParameterName.ToString().Contains(SearchQuery))
+	{
+		return true;
+	}
+
+	return Super::MatchesSearchQuery(SearchQuery);
+}
+
+void UMaterialExpressionRuntimeVirtualTextureSampleParameter::SetValueToMatchingExpression(UMaterialExpression* OtherExpression)
+{
+	URuntimeVirtualTexture* Value = nullptr;
+	//todo[vt]: Store RVT params in the material
+	//Material->GetRuntimeVirtualTextureParameterValue(FMaterialParameterInfo(OtherExpression->GetParameterName()), Value);
+	VirtualTexture = Value;
+	UProperty* ParamProperty = FindField<UProperty>(UMaterialExpressionRuntimeVirtualTextureSampleParameter::StaticClass(), GET_MEMBER_NAME_STRING_CHECKED(UMaterialExpressionRuntimeVirtualTextureSampleParameter, VirtualTexture));
+	FPropertyChangedEvent PropertyChangedEvent(ParamProperty);
+	PostEditChangeProperty(PropertyChangedEvent);
+}
+
+#endif
+
 UMaterialExpressionRuntimeVirtualTextureReplace::UMaterialExpressionRuntimeVirtualTextureReplace(const FObjectInitializer& ObjectInitializer)
 	: Super(ObjectInitializer)
 {
