Add hooks to material compiler to compile RVT parameters
Rename VirtualTextureParam() functions to VirtualTextureUniform() functions to try and avoid confusion with material parameter naming

diff --git a/Engine/Source/Runtime/Engine/Private/Materials/HLSLMaterialTranslator.h b/Engine/Source/Runtime/Engine/Private/Materials/HLSLMaterialTranslator.h
--- a/Engine/Source/Runtime/Engine/Private/Materials/HLSLMaterialTranslator.h
+++ b/Engine/Source/Runtime/Engine/Private/Materials/HLSLMaterialTranslator.h
@@ -5072,9 +5072,33 @@ protected:
 		return AddUniformExpression(new FMaterialUniformExpressionTexture(TextureReferenceIndex, LayerIndex, SamplerType), MCT_TextureVirtual, TEXT(""));
 	}
 
-	virtual int32 VirtualTextureParam(int32 TextureIndex, int32 ParamIndex) override
+	virtual int32 VirtualTextureParameter(FName ParameterName, URuntimeVirtualTexture* DefaultValue, int32 LayerIndex, int32& TextureReferenceIndex, EMaterialSamplerType SamplerType) override
 	{
-		return AddUniformExpression(new FMaterialUniformExpressionRuntimeVirtualTextureUniform(TextureIndex, ParamIndex), MCT_Float3, TEXT(""));
+		if (!UseVirtualTexturing(FeatureLevel, TargetPlatform))
+		{
+			return INDEX_NONE;
+		}
+
+		TextureReferenceIndex = Material->GetReferencedTextures().Find(DefaultValue);
+		checkf(TextureReferenceIndex != INDEX_NONE, TEXT("Material expression called Compiler->VirtualTexture() without implementing UMaterialExpression::GetReferencedTexture properly"));
+
+		FMaterialParameterInfo ParameterInfo = GetParameterAssociationInfo();
+		ParameterInfo.Name = ParameterName;
+
+		return AddUniformExpression(new FMaterialUniformExpressionTextureParameter(ParameterInfo, TextureReferenceIndex, LayerIndex, SamplerType), MCT_TextureVirtual, TEXT(""));
+	}
+
+	virtual int32 VirtualTextureUniform(int32 TextureIndex, int32 VectorIndex) override
+	{
+		return AddUniformExpression(new FMaterialUniformExpressionRuntimeVirtualTextureUniform(TextureIndex, VectorIndex), MCT_Float3, TEXT(""));
+	}
+
+	virtual int32 VirtualTextureUniform(FName ParameterName, int32 TextureIndex, int32 VectorIndex) override
+	{
+		FMaterialParameterInfo ParameterInfo = GetParameterAssociationInfo();
+		ParameterInfo.Name = ParameterName;
+
+		return AddUniformExpression(new FMaterialUniformExpressionRuntimeVirtualTextureUniform(ParameterInfo, TextureIndex, VectorIndex), MCT_Float3, TEXT(""));
 	}
 
 	virtual int32 VirtualTextureWorldToUV(int32 WorldPositionIndex, int32 P0, int32 P1, int32 P2) override
diff --git a/Engine/Source/Runtime/Engine/Private/Materials/MaterialExpressions.cpp b/Engine/Source/Runtime/Engine/Private/Materials/MaterialExpressions.cpp
--- a/Engine/Source/Runtime/Engine/Private/Materials/MaterialExpressions.cpp
+++ b/Engine/Source/Runtime/Engine/Private/Materials/MaterialExpressions.cpp
@@ -2190,9 +2190,9 @@ int32 UMaterialExpressionRuntimeVirtualTextureSample::Compile(class FMaterialCom
 	if (Coordinates.GetTracedInput().Expression == nullptr)
 	{
 		int32 WorldPositionIndex = Compiler->WorldPosition(WPT_Default);
-		int32 P0 = Compiler->VirtualTextureParam(TextureReferenceIndex[0], 0);
-		int32 P1 = Compiler->VirtualTextureParam(TextureReferenceIndex[0], 1);
-		int32 P2 = Compiler->VirtualTextureParam(TextureReferenceIndex[0], 2);
+		int32 P0 = Compiler->VirtualTextureUniform(TextureReferenceIndex[0], 0);
+		int32 P1 = Compiler->VirtualTextureUniform(TextureReferenceIndex[0], 1);
+		int32 P2 = Compiler->VirtualTextureUniform(TextureReferenceIndex[0], 2);
 		CoordinateIndex = Compiler->VirtualTextureWorldToUV(WorldPositionIndex, P0, P1, P2);
 	}
 	else
diff --git a/Engine/Source/Runtime/Engine/Public/MaterialCompiler.h b/Engine/Source/Runtime/Engine/Public/MaterialCompiler.h
--- a/Engine/Source/Runtime/Engine/Public/MaterialCompiler.h
+++ b/Engine/Source/Runtime/Engine/Public/MaterialCompiler.h
@@ -188,7 +188,9 @@ public:
 	virtual int32 TextureParameter(FName ParameterName,UTexture* DefaultTexture,int32& TextureReferenceIndex, EMaterialSamplerType SamplerType, ESamplerSourceMode SamplerSource=SSM_FromTextureAsset) = 0;
 
 	virtual int32 VirtualTexture(URuntimeVirtualTexture* InTexture, int32 LayerIndex, int32& TextureReferenceIndex, EMaterialSamplerType SamplerType) = 0;
-	virtual int32 VirtualTextureParam(int32 TextureIndex, int32 ParamIndex) = 0;
+	virtual int32 VirtualTextureParameter(FName ParameterName, URuntimeVirtualTexture* DefaultValue, int32 LayerIndex, int32& TextureReferenceIndex, EMaterialSamplerType SamplerType) = 0;
+	virtual int32 VirtualTextureUniform(int32 TextureIndex, int32 VectorIndex) = 0;
+	virtual int32 VirtualTextureUniform(FName ParameterName, int32 TextureIndex, int32 VectorIndex) = 0;
 	virtual int32 VirtualTextureWorldToUV(int32 WorldPositionIndex, int32 P0, int32 P1, int32 P2) = 0;
 	virtual int32 VirtualTextureUnpack(int32 CodeIndex0, int32 CodeIndex1, EVirtualTextureUnpackType UnpackType) = 0;
 
@@ -214,6 +216,12 @@ public:
 		return VirtualTexture(InTexture, LayerIndex, TextureReferenceIndex, SamplerType);
 	}
 
+	int32 VirtualTextureParameter(FName ParameterName, URuntimeVirtualTexture* DefaultValue, int32 LayerIndex, EMaterialSamplerType SamplerType)
+	{
+		int32 TextureReferenceIndex = INDEX_NONE;
+		return VirtualTextureParameter(ParameterName, DefaultValue, LayerIndex, TextureReferenceIndex, SamplerType);
+	}
+
 	int32 ExternalTexture(UTexture* DefaultTexture)
 	{
 		int32 TextureReferenceIndex = INDEX_NONE;
@@ -462,7 +470,9 @@ public:
 	virtual int32 TextureParameter(FName ParameterName,UTexture* DefaultValue,int32& TextureReferenceIndex, EMaterialSamplerType SamplerType, ESamplerSourceMode SamplerSource=SSM_FromTextureAsset) override { return Compiler->TextureParameter(ParameterName,DefaultValue,TextureReferenceIndex, SamplerType, SamplerSource); }
 
 	virtual int32 VirtualTexture(URuntimeVirtualTexture* InTexture, int32 LayerIndex, int32& TextureReferenceIndex, EMaterialSamplerType SamplerType) override { return Compiler->VirtualTexture(InTexture, LayerIndex, TextureReferenceIndex, SamplerType); }
-	virtual int32 VirtualTextureParam(int32 TextureIndex, int32 ParamIndex) override { return Compiler->VirtualTextureParam(TextureIndex, ParamIndex); }
+	virtual int32 VirtualTextureParameter(FName ParameterName, URuntimeVirtualTexture* DefaultValue, int32 LayerIndex, int32& TextureReferenceIndex, EMaterialSamplerType SamplerType) override { return Compiler->VirtualTextureParameter(ParameterName, DefaultValue, LayerIndex, TextureReferenceIndex, SamplerType); }
+	virtual int32 VirtualTextureUniform(int32 TextureIndex, int32 VectorIndex) override { return Compiler->VirtualTextureUniform(TextureIndex, VectorIndex); }
+	virtual int32 VirtualTextureUniform(FName ParameterName, int32 TextureIndex, int32 VectorIndex) override { return Compiler->VirtualTextureUniform(ParameterName, TextureIndex, VectorIndex); }
 	virtual int32 VirtualTextureWorldToUV(int32 WorldPositionIndex, int32 P0, int32 P1, int32 P2) override { return Compiler->VirtualTextureWorldToUV(WorldPositionIndex, P0, P1, P2); }
 	virtual int32 VirtualTextureUnpack(int32 CodeIndex0, int32 CodeIndex1, EVirtualTextureUnpackType UnpackType) override { return Compiler->VirtualTextureUnpack(CodeIndex0, CodeIndex1, UnpackType); }
 
