Fix Editor crash when particle template is null

diff --git a/Engine/Source/Runtime/Engine/Private/Particles/ParticleComponents.cpp b/Engine/Source/Runtime/Engine/Private/Particles/ParticleComponents.cpp
--- a/Engine/Source/Runtime/Engine/Private/Particles/ParticleComponents.cpp
+++ b/Engine/Source/Runtime/Engine/Private/Particles/ParticleComponents.cpp
@@ -3599,7 +3599,10 @@ void UParticleSystemComponent::ForceReset()
 {
 #if WITH_EDITOR
 	//If we're resetting in the editor, cached emitter values may now be invalid.
-	Template->UpdateAllModuleLists();
+	if (Template != nullptr)
+	{
+		Template->UpdateAllModuleLists();
+	}
 #endif
 
 	bool bOldActive = bIsActive;
