Fix small potential bug in VT stack logic
Use full material uniform expression lookup of RVT as this is what will be needed for RVT parameters

diff --git a/Engine/Source/Runtime/Engine/Private/Materials/MaterialShared.cpp b/Engine/Source/Runtime/Engine/Private/Materials/MaterialShared.cpp
--- a/Engine/Source/Runtime/Engine/Private/Materials/MaterialShared.cpp
+++ b/Engine/Source/Runtime/Engine/Private/Materials/MaterialShared.cpp
@@ -2202,7 +2202,7 @@ void FMaterialVirtualTextureStack::GetTextureValues(const FMaterialRenderContext
 	for (uint32 LayerIndex = 0u; LayerIndex < NumLayers; ++LayerIndex)
 	{
 		const int32 ExpressionIndex = LayerUniformExpressionIndices[LayerIndex];
-		if (LayerIndex != INDEX_NONE)
+		if (ExpressionIndex != INDEX_NONE)
 		{
 			const FMaterialUniformExpressionTexture* UniformExpression = UniformExpressionSet.UniformVirtualTextureExpressions[ExpressionIndex];
 
@@ -2215,7 +2215,12 @@ void FMaterialVirtualTextureStack::GetTextureValues(const FMaterialRenderContext
 
 void FMaterialVirtualTextureStack::GetTextureValue(const FMaterialRenderContext& Context, const FUniformExpressionSet& UniformExpressionSet, const URuntimeVirtualTexture*& OutValue) const
 {
-	OutValue = GetIndexedTexture<URuntimeVirtualTexture>(Context.Material, PreallocatedStackTextureIndex);
+	const int32 ExpressionIndex = LayerUniformExpressionIndices[0];
+	if (ExpressionIndex != INDEX_NONE)
+	{
+		const FMaterialUniformExpressionTexture* UniformExpression = UniformExpressionSet.UniformVirtualTextureExpressions[ExpressionIndex];
+		UniformExpression->GetTextureValue(Context, Context.Material, OutValue);
+	}
 }
 
 void FMaterialVirtualTextureStack::Serialize(FArchive& Ar)
