Recursively evaluate nodes instead of using a node queue.
#rnx

diff --git a/Engine/Source/Runtime/Engine/Private/SoundCue.cpp b/Engine/Source/Runtime/Engine/Private/SoundCue.cpp
--- a/Engine/Source/Runtime/Engine/Private/SoundCue.cpp
+++ b/Engine/Source/Runtime/Engine/Private/SoundCue.cpp
@@ -185,31 +185,33 @@ void USoundCue::EvaluateNodes(bool bAddToRoot)
 		}
 	}
 
-	TArray<USoundNode*> NodesToEvaluate;
-	NodesToEvaluate.Push(FirstNode);
-
-	while (NodesToEvaluate.Num() > 0)
+	TFunction<void(USoundNode*)> EvaluateNodes_Internal = [&](USoundNode* SoundNode)
 	{
-		if (USoundNode* SoundNode = NodesToEvaluate.Pop(false))
+		if (USoundNodeAssetReferencer* AssetReferencerNode = Cast<USoundNodeAssetReferencer>(SoundNode))
 		{
-			if (USoundNodeAssetReferencer* AssetReferencerNode = Cast<USoundNodeAssetReferencer>(SoundNode))
+			AssetReferencerNode->ConditionalPostLoad();
+			AssetReferencerNode->LoadAsset(bAddToRoot);
+		}
+		else if (USoundNodeQualityLevel* QualityLevelNode = Cast<USoundNodeQualityLevel>(SoundNode))
+		{
+			if (CachedQualityLevel < QualityLevelNode->ChildNodes.Num())
 			{
-				AssetReferencerNode->ConditionalPostLoad();
-				AssetReferencerNode->LoadAsset(bAddToRoot);
+				EvaluateNodes_Internal(QualityLevelNode->ChildNodes[CachedQualityLevel]);
 			}
-			else if (USoundNodeQualityLevel* QualityLevelNode = Cast<USoundNodeQualityLevel>(SoundNode))
+		}
+		else if (SoundNode)
+		{
+			for (USoundNode* ChildNode : SoundNode->ChildNodes)
 			{
-				if (CachedQualityLevel < QualityLevelNode->ChildNodes.Num())
+				if (ChildNode)
 				{
-					NodesToEvaluate.Add(QualityLevelNode->ChildNodes[CachedQualityLevel]);
+					EvaluateNodes_Internal(ChildNode);
 				}
 			}
-			else
-			{
-				NodesToEvaluate.Append(SoundNode->ChildNodes);
-			}
 		}
-	}
+	};
+
+	EvaluateNodes_Internal(FirstNode);
 }
 
 float USoundCue::FindMaxDistanceInternal() const
