Don't do unnecessary versioning work in cooked build. Fix typo in PURE_VIRTUAL class name.
#rnx

diff --git a/Engine/Source/Runtime/Engine/Classes/Sound/SoundBase.h b/Engine/Source/Runtime/Engine/Classes/Sound/SoundBase.h
--- a/Engine/Source/Runtime/Engine/Classes/Sound/SoundBase.h
+++ b/Engine/Source/Runtime/Engine/Classes/Sound/SoundBase.h
@@ -175,7 +175,9 @@ public:
 
 	//~ Begin UObject Interface.
 	virtual void PostInitProperties() override;
+#if WITH_EDITORONLY_DATA
 	virtual void PostLoad() override;
+#endif
 	virtual bool CanBeClusterRoot() const override;
 	virtual bool CanBeInCluster() const override;
 	virtual void Serialize(FArchive& Ar) override;
diff --git a/Engine/Source/Runtime/Engine/Classes/Sound/SoundNodeAssetReferencer.h b/Engine/Source/Runtime/Engine/Classes/Sound/SoundNodeAssetReferencer.h
--- a/Engine/Source/Runtime/Engine/Classes/Sound/SoundNodeAssetReferencer.h
+++ b/Engine/Source/Runtime/Engine/Classes/Sound/SoundNodeAssetReferencer.h
@@ -17,8 +17,8 @@ class ENGINE_API USoundNodeAssetReferencer : public USoundNode
 	GENERATED_BODY()
 
 public:
-	virtual void LoadAsset(bool bAddToRoot = false) PURE_VIRTUAL(USoundNodesAssetReferencer::LoadAsset,);
-	virtual void ClearAssetReferences() PURE_VIRTUAL(USoundNodesAssetReferencer::ClearAssetReferences, );
+	virtual void LoadAsset(bool bAddToRoot = false) PURE_VIRTUAL(USoundNodeAssetReferencer::LoadAsset,);
+	virtual void ClearAssetReferences() PURE_VIRTUAL(USoundNodeAssetReferencer::ClearAssetReferences, );
 
 	bool ShouldHardReferenceAsset() const;
 
diff --git a/Engine/Source/Runtime/Engine/Private/SoundClass.cpp b/Engine/Source/Runtime/Engine/Private/SoundClass.cpp
--- a/Engine/Source/Runtime/Engine/Private/SoundClass.cpp
+++ b/Engine/Source/Runtime/Engine/Private/SoundClass.cpp
@@ -37,6 +37,7 @@ void USoundClass::PostLoad()
 {
 	Super::PostLoad();
 
+#if WITH_EDITORONLY_DATA
 	for (int32 ChildIndex = ChildClasses.Num()-1; ChildIndex >= 0; ChildIndex--)
 	{
 		if (ChildClasses[ChildIndex] != NULL && ChildClasses[ChildIndex]->GetLinkerUE4Version() < VER_UE4_SOUND_CLASS_GRAPH_EDITOR)
@@ -54,6 +55,8 @@ void USoundClass::PostLoad()
 			}
 		}
 	}
+#endif
+
 	// Use the main/default audio device for storing and retrieving sound class properties
 	FAudioDeviceManager* AudioDeviceManager = (GEngine ? GEngine->GetAudioDeviceManager() : nullptr);
 
