Add support to material interface for querying RVT paramters
Stubbed out material instance implementation for now
Use the new interface in RVT parameter material expressions

diff --git a/Engine/Plugins/2D/Paper2D/Source/Paper2D/Private/PaperRenderSceneProxy.cpp b/Engine/Plugins/2D/Paper2D/Source/Paper2D/Private/PaperRenderSceneProxy.cpp
--- a/Engine/Plugins/2D/Paper2D/Source/Paper2D/Private/PaperRenderSceneProxy.cpp
+++ b/Engine/Plugins/2D/Paper2D/Source/Paper2D/Private/PaperRenderSceneProxy.cpp
@@ -100,6 +100,11 @@ public:
 		return Parent->GetTextureValue(ParameterInfo, OutValue, Context);
 	}
 
+	virtual bool GetTextureValue(const FMaterialParameterInfo& ParameterInfo, const URuntimeVirtualTexture** OutValue, const FMaterialRenderContext& Context) const override
+	{
+		return Parent->GetTextureValue(ParameterInfo, OutValue, Context);
+	}
+
 #if WITH_EDITOR
 	inline const UTexture* ApplyEditorOverrides(const UTexture* InTexture) const
 	{
diff --git a/Engine/Source/Developer/MaterialBaking/Private/ExportMaterialProxy.h b/Engine/Source/Developer/MaterialBaking/Private/ExportMaterialProxy.h
--- a/Engine/Source/Developer/MaterialBaking/Private/ExportMaterialProxy.h
+++ b/Engine/Source/Developer/MaterialBaking/Private/ExportMaterialProxy.h
@@ -329,6 +329,11 @@ public:
 		return MaterialInterface->GetRenderProxy()->GetTextureValue(ParameterInfo, OutValue, Context);
 	}
 
+	virtual bool GetTextureValue(const FMaterialParameterInfo& ParameterInfo, const URuntimeVirtualTexture** OutValue, const FMaterialRenderContext& Context) const override
+	{
+		return MaterialInterface->GetRenderProxy()->GetTextureValue(ParameterInfo, OutValue, Context);
+	}
+
 	// Material properties.
 	/** Entry point for compiling a specific material property.  This must call SetMaterialProperty. */
 	virtual int32 CompilePropertyAndSetMaterialProperty(EMaterialProperty Property, FMaterialCompiler* Compiler, EShaderFrequency OverrideShaderFrequency, bool bUsePreviousFrameTime) const override
diff --git a/Engine/Source/Developer/MaterialUtilities/Private/MaterialUtilities.cpp b/Engine/Source/Developer/MaterialUtilities/Private/MaterialUtilities.cpp
--- a/Engine/Source/Developer/MaterialUtilities/Private/MaterialUtilities.cpp
+++ b/Engine/Source/Developer/MaterialUtilities/Private/MaterialUtilities.cpp
@@ -509,6 +509,11 @@ public:
 		return MaterialInterface->GetRenderProxy()->GetTextureValue(ParameterInfo,OutValue,Context);
 	}
 
+	virtual bool GetTextureValue(const FMaterialParameterInfo& ParameterInfo, const URuntimeVirtualTexture** OutValue, const FMaterialRenderContext& Context) const override
+	{
+		return MaterialInterface->GetRenderProxy()->GetTextureValue(ParameterInfo, OutValue, Context);
+	}
+
 	// Material properties.
 	/** Entry point for compiling a specific material property.  This must call SetMaterialProperty. */
 	virtual int32 CompilePropertyAndSetMaterialProperty(EMaterialProperty Property, FMaterialCompiler* Compiler, EShaderFrequency OverrideShaderFrequency, bool bUsePreviousFrameTime) const override
diff --git a/Engine/Source/Editor/MaterialEditor/Private/MaterialEditor.h b/Engine/Source/Editor/MaterialEditor/Private/MaterialEditor.h
--- a/Engine/Source/Editor/MaterialEditor/Private/MaterialEditor.h
+++ b/Engine/Source/Editor/MaterialEditor/Private/MaterialEditor.h
@@ -139,7 +139,16 @@ public:
 		}
 		return false;
 	}
-	
+
+	virtual bool GetTextureValue(const FMaterialParameterInfo& ParameterInfo, const URuntimeVirtualTexture** OutValue, const FMaterialRenderContext& Context) const override
+	{
+		if (Expression.IsValid() && Expression->Material)
+		{
+			return Expression->Material->GetRenderProxy()->GetTextureValue(ParameterInfo, OutValue, Context);
+		}
+		return false;
+	}
+
 	// Material properties.
 	/** Entry point for compiling a specific material property.  This must call SetMaterialProperty. */
 	virtual int32 CompilePropertyAndSetMaterialProperty(EMaterialProperty Property, FMaterialCompiler* Compiler, EShaderFrequency OverrideShaderFrequency, bool bUsePreviousFrameTime) const override;
diff --git a/Engine/Source/Editor/UnrealEd/Private/Lightmass/LightmassRender.cpp b/Engine/Source/Editor/UnrealEd/Private/Lightmass/LightmassRender.cpp
--- a/Engine/Source/Editor/UnrealEd/Private/Lightmass/LightmassRender.cpp
+++ b/Engine/Source/Editor/UnrealEd/Private/Lightmass/LightmassRender.cpp
@@ -336,6 +336,11 @@ public:
 		return MaterialInterface->GetRenderProxy()->GetTextureValue(ParameterInfo,OutValue,Context);
 	}
 
+	virtual bool GetTextureValue(const FMaterialParameterInfo& ParameterInfo, const URuntimeVirtualTexture** OutValue, const FMaterialRenderContext& Context) const override
+	{
+		return MaterialInterface->GetRenderProxy()->GetTextureValue(ParameterInfo, OutValue, Context);
+	}
+
 	// Material properties.
 
 	/** Entry point for compiling a specific material property.  This must call SetMaterialProperty. */
diff --git a/Engine/Source/Editor/UnrealEd/Private/PreviewMaterial.cpp b/Engine/Source/Editor/UnrealEd/Private/PreviewMaterial.cpp
--- a/Engine/Source/Editor/UnrealEd/Private/PreviewMaterial.cpp
+++ b/Engine/Source/Editor/UnrealEd/Private/PreviewMaterial.cpp
@@ -235,6 +235,11 @@ public:
 	{
 		return Material->GetRenderProxy()->GetTextureValue(ParameterInfo,OutValue,Context);
 	}
+
+	virtual bool GetTextureValue(const FMaterialParameterInfo& ParameterInfo, const URuntimeVirtualTexture** OutValue, const FMaterialRenderContext& Context) const
+	{
+		return Material->GetRenderProxy()->GetTextureValue(ParameterInfo, OutValue, Context);
+	}
 };
 
 /** Implementation of Preview Material functions*/
diff --git a/Engine/Source/Runtime/Engine/Classes/Materials/MaterialInstance.h b/Engine/Source/Runtime/Engine/Classes/Materials/MaterialInstance.h
--- a/Engine/Source/Runtime/Engine/Classes/Materials/MaterialInstance.h
+++ b/Engine/Source/Runtime/Engine/Classes/Materials/MaterialInstance.h
@@ -423,6 +423,7 @@ public:
 	virtual ENGINE_API bool GetVectorParameterValue(const FMaterialParameterInfo& ParameterInfo, FLinearColor& OutValue, bool bOveriddenOnly = false) const override;
 	virtual ENGINE_API bool IsVectorParameterUsedAsChannelMask(const FMaterialParameterInfo& ParameterInfo, bool& OutValue) const override;
 	virtual ENGINE_API bool GetTextureParameterValue(const FMaterialParameterInfo& ParameterInfo, class UTexture*& OutValue, bool bOveriddenOnly = false) const override;
+	virtual ENGINE_API bool GetRuntimeVirtualTextureParameterValue(const FMaterialParameterInfo& ParameterInfo, class URuntimeVirtualTexture*& OutValue, bool bOveriddenOnly = false) const override;
 	virtual ENGINE_API bool GetFontParameterValue(const FMaterialParameterInfo& ParameterInfo, class UFont*& OutFontValue, int32& OutFontPage, bool bOveriddenOnly = false) const override;
 	virtual ENGINE_API void GetUsedTextures(TArray<UTexture*>& OutTextures, EMaterialQualityLevel::Type QualityLevel, bool bAllQualityLevels, ERHIFeatureLevel::Type FeatureLevel, bool bAllFeatureLevels) const override;
 	virtual ENGINE_API void GetUsedTexturesAndIndices(TArray<UTexture*>& OutTextures, TArray< TArray<int32> >& OutIndices, EMaterialQualityLevel::Type QualityLevel, ERHIFeatureLevel::Type FeatureLevel) const;
diff --git a/Engine/Source/Runtime/Engine/Classes/Materials/MaterialInterface.h b/Engine/Source/Runtime/Engine/Classes/Materials/MaterialInterface.h
--- a/Engine/Source/Runtime/Engine/Classes/Materials/MaterialInterface.h
+++ b/Engine/Source/Runtime/Engine/Classes/Materials/MaterialInterface.h
@@ -486,6 +486,8 @@ public:
 		PURE_VIRTUAL(UMaterialInterface::GetAllVectorParameterInfo,return;);
 	virtual void GetAllTextureParameterInfo(TArray<FMaterialParameterInfo>& OutParameterInfo, TArray<FGuid>& OutParameterIds) const
 		PURE_VIRTUAL(UMaterialInterface::GetAllTextureParameterInfo,return;);
+	virtual void GetAllRuntimeVirtualTextureParameterInfo(TArray<FMaterialParameterInfo>& OutParameterInfo, TArray<FGuid>& OutParameterIds) const
+		PURE_VIRTUAL(UMaterialInterface::GetAllRuntimeVirtualTextureParameterInfo, return;);
 	virtual void GetAllFontParameterInfo(TArray<FMaterialParameterInfo>& OutParameterInfo, TArray<FGuid>& OutParameterIds) const
 		PURE_VIRTUAL(UMaterialInterface::GetAllFontParameterInfo,return;);
 	virtual void GetAllMaterialLayersParameterInfo(TArray<FMaterialParameterInfo>& OutParameterInfo, TArray<FGuid>& OutParameterIds) const
@@ -504,6 +506,8 @@ public:
 		PURE_VIRTUAL(UMaterialInterface::GetVectorParameterDefaultValue,return false;);
 	virtual bool GetTextureParameterDefaultValue(const FMaterialParameterInfo& ParameterInfo, class UTexture*& OutValue, bool bCheckOwnedGlobalOverrides = false) const
 		PURE_VIRTUAL(UMaterialInterface::GetTextureParameterDefaultValue,return false;);
+	virtual bool GetRuntimeVirtualTextureParameterDefaultValue(const FMaterialParameterInfo& ParameterInfo, class URuntimeVirtualTexture*& OutValue, bool bCheckOwnedGlobalOverrides = false) const
+		PURE_VIRTUAL(UMaterialInterface::GetRuntimeVirtualTextureParameterDefaultValue, return false;);
 	virtual bool GetFontParameterDefaultValue(const FMaterialParameterInfo& ParameterInfo, class UFont*& OutFontValue, int32& OutFontPage, bool bCheckOwnedGlobalOverrides = false) const
 		PURE_VIRTUAL(UMaterialInterface::GetFontParameterDefaultValue,return false;);
 	virtual bool GetStaticSwitchParameterDefaultValue(const FMaterialParameterInfo& ParameterInfo, bool& OutValue, FGuid& OutExpressionGuid, bool bCheckOwnedGlobalOverrides = false) const
@@ -703,6 +707,7 @@ public:
 	ENGINE_API virtual bool GetLinearColorParameterValue(const FMaterialParameterInfo& ParameterInfo, FLinearColor& OutValue) const;
 	ENGINE_API virtual bool GetLinearColorCurveParameterValue(const FMaterialParameterInfo& ParameterInfo, FInterpCurveLinearColor& OutValue) const;
 	ENGINE_API virtual bool GetTextureParameterValue(const FMaterialParameterInfo& ParameterInfo, class UTexture*& OutValue, bool bOveriddenOnly = false) const;
+	ENGINE_API virtual bool GetRuntimeVirtualTextureParameterValue(const FMaterialParameterInfo& ParameterInfo, class URuntimeVirtualTexture*& OutValue, bool bOveriddenOnly = false) const;
 	ENGINE_API virtual bool GetFontParameterValue(const FMaterialParameterInfo& ParameterInfo,class UFont*& OutFontValue, int32& OutFontPage, bool bOveriddenOnly = false) const;
 	ENGINE_API virtual bool GetRefractionSettings(float& OutBiasValue) const;
 
diff --git a/Engine/Source/Runtime/Engine/Private/DebugViewModeMaterialProxy.cpp b/Engine/Source/Runtime/Engine/Private/DebugViewModeMaterialProxy.cpp
--- a/Engine/Source/Runtime/Engine/Private/DebugViewModeMaterialProxy.cpp
+++ b/Engine/Source/Runtime/Engine/Private/DebugViewModeMaterialProxy.cpp
@@ -120,6 +120,11 @@ bool FDebugViewModeMaterialProxy::GetTextureValue(const FMaterialParameterInfo&
 	return MaterialInterface->GetRenderProxy()->GetTextureValue(ParameterInfo,OutValue,Context);
 }
 
+bool FDebugViewModeMaterialProxy::GetTextureValue(const FMaterialParameterInfo& ParameterInfo, const URuntimeVirtualTexture** OutValue, const FMaterialRenderContext& Context) const
+{
+	return MaterialInterface->GetRenderProxy()->GetTextureValue(ParameterInfo, OutValue, Context);
+}
+
 EMaterialDomain FDebugViewModeMaterialProxy::GetMaterialDomain() const
 {
 	return Material ? (EMaterialDomain)(Material->MaterialDomain) : MD_Surface;
diff --git a/Engine/Source/Runtime/Engine/Private/DebugViewModeMaterialProxy.h b/Engine/Source/Runtime/Engine/Private/DebugViewModeMaterialProxy.h
--- a/Engine/Source/Runtime/Engine/Private/DebugViewModeMaterialProxy.h
+++ b/Engine/Source/Runtime/Engine/Private/DebugViewModeMaterialProxy.h
@@ -129,7 +129,8 @@ public:
 	virtual const FMaterial& GetMaterialWithFallback(ERHIFeatureLevel::Type InFeatureLevel, const FMaterialRenderProxy*& OutFallbackMaterialRenderProxy) const override;
 	virtual bool GetVectorValue(const FMaterialParameterInfo& ParameterInfo, FLinearColor* OutValue, const FMaterialRenderContext& Context) const override;
 	virtual bool GetScalarValue(const FMaterialParameterInfo& ParameterInfo, float* OutValue, const FMaterialRenderContext& Context) const override;
-	virtual bool GetTextureValue(const FMaterialParameterInfo& ParameterInfo,const UTexture** OutValue, const FMaterialRenderContext& Context) const override;
+	virtual bool GetTextureValue(const FMaterialParameterInfo& ParameterInfo, const UTexture** OutValue, const FMaterialRenderContext& Context) const override;
+	virtual bool GetTextureValue(const FMaterialParameterInfo& ParameterInfo, const URuntimeVirtualTexture** OutValue, const FMaterialRenderContext& Context) const override;
 
 	virtual EMaterialDomain GetMaterialDomain() const override;
 	virtual bool IsTwoSided() const  override;
diff --git a/Engine/Source/Runtime/Engine/Private/Materials/Material.cpp b/Engine/Source/Runtime/Engine/Private/Materials/Material.cpp
--- a/Engine/Source/Runtime/Engine/Private/Materials/Material.cpp
+++ b/Engine/Source/Runtime/Engine/Private/Materials/Material.cpp
@@ -364,6 +364,19 @@ public:
 		}
 	}
 
+	virtual bool GetTextureValue(const FMaterialParameterInfo& ParameterInfo, const URuntimeVirtualTexture** OutValue, const FMaterialRenderContext& Context) const
+	{
+		const FMaterialResource* MaterialResource = Material->GetMaterialResource(Context.Material.GetFeatureLevel());
+		if (MaterialResource && MaterialResource->GetRenderingThreadShaderMap())
+		{
+			return false;
+		}
+		else
+		{
+			return GetFallbackRenderProxy().GetTextureValue(ParameterInfo, OutValue, Context);
+		}
+	}
+
 	// FRenderResource interface.
 	virtual FString GetFriendlyName() const { return Material->GetName(); }
 
diff --git a/Engine/Source/Runtime/Engine/Private/Materials/MaterialExpressions.cpp b/Engine/Source/Runtime/Engine/Private/Materials/MaterialExpressions.cpp
--- a/Engine/Source/Runtime/Engine/Private/Materials/MaterialExpressions.cpp
+++ b/Engine/Source/Runtime/Engine/Private/Materials/MaterialExpressions.cpp
@@ -2341,12 +2341,13 @@ bool UMaterialExpressionRuntimeVirtualTextureSampleParameter::MatchesSearchQuery
 void UMaterialExpressionRuntimeVirtualTextureSampleParameter::SetValueToMatchingExpression(UMaterialExpression* OtherExpression)
 {
 	URuntimeVirtualTexture* Value = nullptr;
-	//todo[vt]: Store RVT params in the material
-	//Material->GetRuntimeVirtualTextureParameterValue(FMaterialParameterInfo(OtherExpression->GetParameterName()), Value);
-	VirtualTexture = Value;
-	UProperty* ParamProperty = FindField<UProperty>(UMaterialExpressionRuntimeVirtualTextureSampleParameter::StaticClass(), GET_MEMBER_NAME_STRING_CHECKED(UMaterialExpressionRuntimeVirtualTextureSampleParameter, VirtualTexture));
-	FPropertyChangedEvent PropertyChangedEvent(ParamProperty);
-	PostEditChangeProperty(PropertyChangedEvent);
+	if (Material->GetRuntimeVirtualTextureParameterValue(FMaterialParameterInfo(OtherExpression->GetParameterName()), Value))
+	{
+		VirtualTexture = Value;
+		UProperty* ParamProperty = FindField<UProperty>(UMaterialExpressionRuntimeVirtualTextureSampleParameter::StaticClass(), GET_MEMBER_NAME_STRING_CHECKED(UMaterialExpressionRuntimeVirtualTextureSampleParameter, VirtualTexture));
+		FPropertyChangedEvent PropertyChangedEvent(ParamProperty);
+		PostEditChangeProperty(PropertyChangedEvent);
+	}
 }
 
 #endif
diff --git a/Engine/Source/Runtime/Engine/Private/Materials/MaterialInstance.cpp b/Engine/Source/Runtime/Engine/Private/Materials/MaterialInstance.cpp
--- a/Engine/Source/Runtime/Engine/Private/Materials/MaterialInstance.cpp
+++ b/Engine/Source/Runtime/Engine/Private/Materials/MaterialInstance.cpp
@@ -17,6 +17,7 @@
 #include "Materials/MaterialExpressionTextureSampleParameter.h"
 #include "Materials/MaterialExpressionFontSampleParameter.h"
 #include "Materials/MaterialExpressionMaterialAttributeLayers.h"
+#include "Materials/MaterialExpressionRuntimeVirtualTextureSampleParameter.h"
 #include "Materials/MaterialExpressionStaticSwitchParameter.h"
 #include "Materials/MaterialExpressionStaticComponentMaskParameter.h"
 #include "Materials/MaterialFunctionInstance.h"
@@ -286,6 +287,16 @@ bool FMaterialInstanceResource::GetTextureValue(
 	}
 }
 
+bool FMaterialInstanceResource::GetTextureValue(
+	const FMaterialParameterInfo& ParameterInfo,
+	const URuntimeVirtualTexture** OutValue,
+	const FMaterialRenderContext& Context
+) const
+{
+	//todo[vt]
+	return false;
+}
+
 void UMaterialInstance::PropagateDataToMaterialProxy()
 {
 	if (Resource)
@@ -838,6 +849,12 @@ bool UMaterialInstance::GetTextureParameterValue(const FMaterialParameterInfo& P
 	return false;
 }
 
+bool UMaterialInstance::GetRuntimeVirtualTextureParameterValue(const FMaterialParameterInfo& ParameterInfo, URuntimeVirtualTexture*& OutValue, bool bOveriddenOnly) const
+{
+	// todo[vt]:
+	return false;
+}
+
 bool UMaterialInstance::GetFontParameterValue(const FMaterialParameterInfo& ParameterInfo,class UFont*& OutFontValue, int32& OutFontPage, bool bOveriddenOnly) const
 {
 	bool bFoundAValue = false;
diff --git a/Engine/Source/Runtime/Engine/Private/Materials/MaterialInstanceSupport.h b/Engine/Source/Runtime/Engine/Private/Materials/MaterialInstanceSupport.h
--- a/Engine/Source/Runtime/Engine/Private/Materials/MaterialInstanceSupport.h
+++ b/Engine/Source/Runtime/Engine/Private/Materials/MaterialInstanceSupport.h
@@ -95,9 +95,11 @@ public:
 	/** Get the FMaterial that should be used for rendering, but might not be in a valid state to actually use.  Can return NULL. */
 	virtual FMaterial* GetMaterialNoFallback(ERHIFeatureLevel::Type FeatureLevel) const override;
 	virtual UMaterialInterface* GetMaterialInterface() const override;
+	
 	virtual bool GetVectorValue(const FMaterialParameterInfo& ParameterInfo, FLinearColor* OutValue, const FMaterialRenderContext& Context) const override;
 	virtual bool GetScalarValue(const FMaterialParameterInfo& ParameterInfo, float* OutValue, const FMaterialRenderContext& Context) const override;
 	virtual bool GetTextureValue(const FMaterialParameterInfo& ParameterInfo, const UTexture** OutValue, const FMaterialRenderContext& Context) const override;
+	virtual bool GetTextureValue(const FMaterialParameterInfo& ParameterInfo, const URuntimeVirtualTexture** OutValue, const FMaterialRenderContext& Context) const override;
 
 	void GameThread_SetParent(UMaterialInterface* ParentMaterialInterface);
 
diff --git a/Engine/Source/Runtime/Engine/Private/Materials/MaterialInterface.cpp b/Engine/Source/Runtime/Engine/Private/Materials/MaterialInterface.cpp
--- a/Engine/Source/Runtime/Engine/Private/Materials/MaterialInterface.cpp
+++ b/Engine/Source/Runtime/Engine/Private/Materials/MaterialInterface.cpp
@@ -341,6 +341,11 @@ bool UMaterialInterface::GetTextureParameterValue(const FMaterialParameterInfo&
 	return false;
 }
 
+bool UMaterialInterface::GetRuntimeVirtualTextureParameterValue(const FMaterialParameterInfo& ParameterInfo, URuntimeVirtualTexture*& OutValue, bool bOveriddenOnly) const
+{
+	return false;
+}
+
 bool UMaterialInterface::GetFontParameterValue(const FMaterialParameterInfo& ParameterInfo, class UFont*& OutFontValue, int32& OutFontPage, bool bOveriddenOnly) const
 {
 	return false;
diff --git a/Engine/Source/Runtime/Engine/Private/Materials/MaterialShared.cpp b/Engine/Source/Runtime/Engine/Private/Materials/MaterialShared.cpp
--- a/Engine/Source/Runtime/Engine/Private/Materials/MaterialShared.cpp
+++ b/Engine/Source/Runtime/Engine/Private/Materials/MaterialShared.cpp
@@ -2618,6 +2618,11 @@ bool FColoredMaterialRenderProxy::GetTextureValue(const FMaterialParameterInfo&
 	return Parent->GetTextureValue(ParameterInfo,OutValue,Context);
 }
 
+bool FColoredMaterialRenderProxy::GetTextureValue(const FMaterialParameterInfo& ParameterInfo, const URuntimeVirtualTexture** OutValue, const FMaterialRenderContext& Context) const
+{
+	return Parent->GetTextureValue(ParameterInfo, OutValue, Context);
+}
+
 /*-----------------------------------------------------------------------------
 	FColoredTexturedMaterialRenderProxy
 -----------------------------------------------------------------------------*/
@@ -2666,6 +2671,11 @@ bool FOverrideSelectionColorMaterialRenderProxy::GetTextureValue(const FMaterial
 	return Parent->GetTextureValue(ParameterInfo, OutValue, Context);
 }
 
+bool FOverrideSelectionColorMaterialRenderProxy::GetTextureValue(const FMaterialParameterInfo& ParameterInfo, const URuntimeVirtualTexture** OutValue, const FMaterialRenderContext& Context) const
+{
+	return Parent->GetTextureValue(ParameterInfo, OutValue, Context);
+}
+
 /*-----------------------------------------------------------------------------
 	FLightingDensityMaterialRenderProxy
 -----------------------------------------------------------------------------*/
diff --git a/Engine/Source/Runtime/Engine/Private/Materials/MaterialUniformExpressions.cpp b/Engine/Source/Runtime/Engine/Private/Materials/MaterialUniformExpressions.cpp
--- a/Engine/Source/Runtime/Engine/Private/Materials/MaterialUniformExpressions.cpp
+++ b/Engine/Source/Runtime/Engine/Private/Materials/MaterialUniformExpressions.cpp
@@ -1009,8 +1009,25 @@ void FMaterialUniformExpressionTextureParameter::GetTextureValue(const FMaterial
 
 void FMaterialUniformExpressionTextureParameter::GetTextureValue(const FMaterialRenderContext& Context, const FMaterial& Material, const URuntimeVirtualTexture*& OutValue) const
 {
-	// todo[vt]: Add runtime virtual texture parameter
-	check(false);
+	if (!Context.MaterialRenderProxy || !Context.MaterialRenderProxy->GetTextureValue(ParameterInfo, &OutValue, Context))
+	{
+		URuntimeVirtualTexture* Value = nullptr;
+
+		if (AreExperimentalMaterialLayersEnabled())
+		{
+			UMaterialInterface* Interface = Context.Material.GetMaterialInterface();
+			if (!Interface || !Interface->GetRuntimeVirtualTextureParameterDefaultValue(ParameterInfo, Value))
+			{
+				Value = GetIndexedTexture<URuntimeVirtualTexture>(Material, TextureIndex);
+			}
+		}
+		else
+		{
+			Value = GetIndexedTexture<URuntimeVirtualTexture>(Material, TextureIndex);
+		}
+
+		OutValue = Value;
+	}
 }
 
 void FMaterialUniformExpressionTextureParameter::GetGameThreadTextureValue(const UMaterialInterface* MaterialInterface, const FMaterial& Material, UTexture*& OutValue, bool bAllowOverride) const
@@ -1323,10 +1340,11 @@ bool FMaterialUniformExpressionRuntimeVirtualTextureUniform::IsIdentical(const F
 
 void FMaterialUniformExpressionRuntimeVirtualTextureUniform::GetNumberValue(const struct FMaterialRenderContext& Context, FLinearColor& OutValue) const
 {
-	//todo[vt]: Add support for parameters
-	check(!bParameter)
-
-	URuntimeVirtualTexture* Texture = GetIndexedTexture<URuntimeVirtualTexture>(Context.Material, TextureIndex);
+	const URuntimeVirtualTexture* Texture = nullptr;
+	if (!bParameter || !Context.MaterialRenderProxy || !Context.MaterialRenderProxy->GetTextureValue(ParameterInfo, &Texture, Context))
+	{
+		Texture = GetIndexedTexture<URuntimeVirtualTexture>(Context.Material, TextureIndex);
+	}
 	if (Texture != nullptr && VectorIndex != INDEX_NONE)
 	{
 		OutValue = FLinearColor(Texture->GetUniformParameter(VectorIndex));
diff --git a/Engine/Source/Runtime/Engine/Public/MaterialShared.h b/Engine/Source/Runtime/Engine/Public/MaterialShared.h
--- a/Engine/Source/Runtime/Engine/Public/MaterialShared.h
+++ b/Engine/Source/Runtime/Engine/Public/MaterialShared.h
@@ -312,6 +312,7 @@ public:
 	virtual void GetTextureValue(const FMaterialRenderContext& Context, const FMaterial& Material, const UTexture*& OutValue) const;
 	/** Get the URuntimeVirtualTexture referenced by this expression (can be a nullptr). */
 	virtual void GetTextureValue(const FMaterialRenderContext& Context, const FMaterial& Material, const URuntimeVirtualTexture*& OutValue) const;
+	
 	/** Get the UTexture referenced by this expression including any transient override logic. */
 	virtual void GetGameThreadTextureValue(const UMaterialInterface* MaterialInterface, const FMaterial& Material, UTexture*& OutValue, bool bAllowOverride = true) const;
 
@@ -1973,9 +1974,12 @@ public:
 	/** Returns the FMaterial, without using a fallback if the FMaterial doesn't have a valid shader map. Can return NULL. */
 	virtual FMaterial* GetMaterialNoFallback(ERHIFeatureLevel::Type InFeatureLevel) const { return NULL; }
 	virtual UMaterialInterface* GetMaterialInterface() const { return NULL; }
+
 	virtual bool GetVectorValue(const FMaterialParameterInfo& ParameterInfo, FLinearColor* OutValue, const FMaterialRenderContext& Context) const = 0;
 	virtual bool GetScalarValue(const FMaterialParameterInfo& ParameterInfo, float* OutValue, const FMaterialRenderContext& Context) const = 0;
 	virtual bool GetTextureValue(const FMaterialParameterInfo& ParameterInfo,const UTexture** OutValue, const FMaterialRenderContext& Context) const = 0;
+	virtual bool GetTextureValue(const FMaterialParameterInfo& ParameterInfo, const URuntimeVirtualTexture** OutValue, const FMaterialRenderContext& Context) const = 0;
+
 	bool IsDeleted() const
 	{
 		return DeletedFlag != 0;
@@ -2058,6 +2062,7 @@ public:
 	ENGINE_API virtual bool GetVectorValue(const FMaterialParameterInfo& ParameterInfo, FLinearColor* OutValue, const FMaterialRenderContext& Context) const;
 	ENGINE_API virtual bool GetScalarValue(const FMaterialParameterInfo& ParameterInfo, float* OutValue, const FMaterialRenderContext& Context) const;
 	ENGINE_API virtual bool GetTextureValue(const FMaterialParameterInfo& ParameterInfo,const UTexture** OutValue, const FMaterialRenderContext& Context) const;
+	ENGINE_API virtual bool GetTextureValue(const FMaterialParameterInfo& ParameterInfo, const URuntimeVirtualTexture** OutValue, const FMaterialRenderContext& Context) const;
 };
 
 
@@ -2104,6 +2109,7 @@ public:
 	ENGINE_API virtual bool GetVectorValue(const FMaterialParameterInfo& ParameterInfo, FLinearColor* OutValue, const FMaterialRenderContext& Context) const;
 	ENGINE_API virtual bool GetScalarValue(const FMaterialParameterInfo& ParameterInfo, float* OutValue, const FMaterialRenderContext& Context) const;
 	ENGINE_API virtual bool GetTextureValue(const FMaterialParameterInfo& ParameterInfo, const UTexture** OutValue, const FMaterialRenderContext& Context) const;
+	ENGINE_API virtual bool GetTextureValue(const FMaterialParameterInfo& ParameterInfo, const URuntimeVirtualTexture** OutValue, const FMaterialRenderContext& Context) const;
 };
 
 
diff --git a/Engine/Source/Runtime/Landscape/Private/LandscapeGizmoActor.cpp b/Engine/Source/Runtime/Landscape/Private/LandscapeGizmoActor.cpp
--- a/Engine/Source/Runtime/Landscape/Private/LandscapeGizmoActor.cpp
+++ b/Engine/Source/Runtime/Landscape/Private/LandscapeGizmoActor.cpp
@@ -111,6 +111,10 @@ public:
 		}
 		return Parent->GetTextureValue(ParameterInfo, OutValue, Context);
 	}
+	virtual bool GetTextureValue(const FMaterialParameterInfo& ParameterInfo, const URuntimeVirtualTexture** OutValue, const FMaterialRenderContext& Context) const
+	{
+		return Parent->GetTextureValue(ParameterInfo, OutValue, Context);
+	}
 };
 
 /** Represents a LandscapeGizmoRenderingComponent to the scene manager. */
diff --git a/Engine/Source/Runtime/Landscape/Public/LandscapeRender.h b/Engine/Source/Runtime/Landscape/Public/LandscapeRender.h
--- a/Engine/Source/Runtime/Landscape/Public/LandscapeRender.h
+++ b/Engine/Source/Runtime/Landscape/Public/LandscapeRender.h
@@ -861,6 +861,10 @@ public:
 			return Parent->GetTextureValue(ParameterInfo, OutValue, Context);
 		}
 	}
+	virtual bool GetTextureValue(const FMaterialParameterInfo& ParameterInfo, const URuntimeVirtualTexture** OutValue, const FMaterialRenderContext& Context) const
+	{
+		return Parent->GetTextureValue(ParameterInfo, OutValue, Context);
+	}
 };
 
 class FLandscapeSelectMaterialRenderProxy : public FMaterialRenderProxy
@@ -908,6 +912,10 @@ public:
 			return Parent->GetTextureValue(ParameterInfo, OutValue, Context);
 		}
 	}
+	virtual bool GetTextureValue(const FMaterialParameterInfo& ParameterInfo, const URuntimeVirtualTexture** OutValue, const FMaterialRenderContext& Context) const
+	{
+		return Parent->GetTextureValue(ParameterInfo, OutValue, Context);
+	}
 };
 
 class FLandscapeMaskMaterialRenderProxy : public FMaterialRenderProxy
@@ -954,6 +962,10 @@ public:
 			return Parent->GetTextureValue(ParameterInfo, OutValue, Context);
 		}
 	}
+	virtual bool GetTextureValue(const FMaterialParameterInfo& ParameterInfo, const URuntimeVirtualTexture** OutValue, const FMaterialRenderContext& Context) const
+	{
+		return Parent->GetTextureValue(ParameterInfo, OutValue, Context);
+	}
 };
 
 class FLandscapeLayerUsageRenderProxy : public FMaterialRenderProxy
@@ -1031,4 +1043,8 @@ public:
 	{
 		return Parent->GetTextureValue(ParameterInfo, OutValue, Context);
 	}
+	virtual bool GetTextureValue(const FMaterialParameterInfo& ParameterInfo, const URuntimeVirtualTexture** OutValue, const FMaterialRenderContext& Context) const
+	{
+		return Parent->GetTextureValue(ParameterInfo, OutValue, Context);
+	}
 };
