Added "r.Streaming.PoolSize.VRAMPercentageClamp" to prevent reserving too much space
when using PoolSizeVRAMPercentage on system with high amount of VRAM. Default=1024 MB

diff --git a/Engine/Source/Runtime/Engine/Private/Streaming/AsyncTextureStreaming.cpp b/Engine/Source/Runtime/Engine/Private/Streaming/AsyncTextureStreaming.cpp
--- a/Engine/Source/Runtime/Engine/Private/Streaming/AsyncTextureStreaming.cpp
+++ b/Engine/Source/Runtime/Engine/Private/Streaming/AsyncTextureStreaming.cpp
@@ -380,8 +380,9 @@ void FRenderAssetStreamingMipCalcTask::UpdateBudgetedMips_Async(int64& MemoryUse
 	// The later helps significantly in low budget settings, where NonStreamingRenderAssetMemory would take too much of the pool.
 	if (GPoolSizeVRAMPercentage > 0 && TotalGraphicsMemory > 0)
 	{
-		const int64 UsableVRAM = TotalGraphicsMemory * GPoolSizeVRAMPercentage / 100 - (int64)GCurrentRendertargetMemorySize * 1024ll; // Add any other...
-		const int64 AvailableVRAMForStreaming = FMath::Min<int64>(UsableVRAM - NonStreamingRenderAssetMemory - MemoryMargin, PoolSize);
+		const int64 UsableVRAM = FMath::Max<int64>(TotalGraphicsMemory * GPoolSizeVRAMPercentage / 100, TotalGraphicsMemory - Settings.VRAMPercentageClamp * 1024ll * 1024ll);
+		const int64 UsedVRAM = (int64)GCurrentRendertargetMemorySize * 1024ll + NonStreamingRenderAssetMemory; // Add any other...
+		const int64 AvailableVRAMForStreaming = FMath::Min<int64>(UsableVRAM - UsedVRAM - MemoryMargin, PoolSize);
 		if (Settings.bLimitPoolSizeToVRAM || AvailableVRAMForStreaming > AvailableMemoryForStreaming)
 		{
 			AvailableMemoryForStreaming = AvailableVRAMForStreaming;
diff --git a/Engine/Source/Runtime/Engine/Private/Streaming/TextureStreamingHelpers.cpp b/Engine/Source/Runtime/Engine/Private/Streaming/TextureStreamingHelpers.cpp
--- a/Engine/Source/Runtime/Engine/Private/Streaming/TextureStreamingHelpers.cpp
+++ b/Engine/Source/Runtime/Engine/Private/Streaming/TextureStreamingHelpers.cpp
@@ -277,6 +277,13 @@ static TAutoConsoleVariable<int32> CVarStreamingStressTestFramesForFullUpdate(
 	TEXT("Num frames to update texture states when doing the stress tests."),
 	ECVF_Cheat);
 
+static TAutoConsoleVariable<int32> CVarStreamingVRAMPercentageClamp(
+	TEXT("r.Streaming.PoolSize.VRAMPercentageClamp"),
+	1024,
+	TEXT("When using PoolSizeVRAMPercentage, a maximum amout of memory to reserve in MB.\n")
+	TEXT("This avoids reserving too much space for systems with a lot of VRAM. (default=1024)"),
+	ECVF_Cheat);
+
 void FRenderAssetStreamingSettings::Update()
 {
 	MaxEffectiveScreenSize = CVarStreamingScreenSizeEffectiveMax.GetValueOnAnyThread();
@@ -298,6 +305,7 @@ void FRenderAssetStreamingSettings::Update()
 	HiddenPrimitiveScale = bUseNewMetrics ? CVarStreamingHiddenPrimitiveScale.GetValueOnAnyThread() : 1.f;
 	bMipCalculationEnablePerLevelList = CVarStreamingMipCalculationEnablePerLevelList.GetValueOnAnyThread() != 0;
 	bPrioritizeMeshLODRetention = CVarPrioritizeMeshLODRetention.GetValueOnAnyThread() != 0;
+	VRAMPercentageClamp = CVarStreamingVRAMPercentageClamp.GetValueOnAnyThread();
 
 	MaterialQualityLevel = (int32)GetCachedScalabilityCVars().MaterialQualityLevel;
 
diff --git a/Engine/Source/Runtime/Engine/Private/Streaming/TextureStreamingHelpers.h b/Engine/Source/Runtime/Engine/Private/Streaming/TextureStreamingHelpers.h
--- a/Engine/Source/Runtime/Engine/Private/Streaming/TextureStreamingHelpers.h
+++ b/Engine/Source/Runtime/Engine/Private/Streaming/TextureStreamingHelpers.h
@@ -97,6 +97,7 @@ struct FRenderAssetStreamingSettings
 	int32 FramesForFullUpdate;
 	bool bMipCalculationEnablePerLevelList;
 	bool bPrioritizeMeshLODRetention;
+	int32 VRAMPercentageClamp;
 
 	bool bStressTest;
 	static int32 ExtraIOLatency;
