Implement RVT parameter support for material instances

diff --git a/Engine/Source/Runtime/Engine/Classes/Materials/MaterialFunctionInterface.h b/Engine/Source/Runtime/Engine/Classes/Materials/MaterialFunctionInterface.h
--- a/Engine/Source/Runtime/Engine/Classes/Materials/MaterialFunctionInterface.h
+++ b/Engine/Source/Runtime/Engine/Classes/Materials/MaterialFunctionInterface.h
@@ -421,6 +421,11 @@ public:
 		return false;
 	}
 	
+	virtual bool OverrideNamedRuntimeVirtualTextureParameter(const FMaterialParameterInfo& ParameterInfo, class URuntimeVirtualTexture*& OutValue)
+	{
+		return false;
+	}
+
 	virtual bool OverrideNamedFontParameter(const FMaterialParameterInfo& ParameterInfo, class UFont*& OutFontValue, int32& OutFontPage)
 	{
 		return false;
diff --git a/Engine/Source/Runtime/Engine/Classes/Materials/MaterialInstance.h b/Engine/Source/Runtime/Engine/Classes/Materials/MaterialInstance.h
--- a/Engine/Source/Runtime/Engine/Classes/Materials/MaterialInstance.h
+++ b/Engine/Source/Runtime/Engine/Classes/Materials/MaterialInstance.h
@@ -186,6 +186,42 @@ struct FTextureParameterValue
 	static ValueType GetValue(const FTextureParameterValue& Parameter) { return Parameter.ParameterValue; }
 };
 
+/** Editable runtime virtual texture parameter. */
+USTRUCT(BlueprintType)
+struct FRuntimeVirtualTextureParameterValue
+{
+	GENERATED_USTRUCT_BODY()
+
+	UPROPERTY(EditAnywhere, BlueprintReadWrite, Category = RuntimeVirtualTextureParameterValue)
+	FMaterialParameterInfo ParameterInfo;
+
+	UPROPERTY(EditAnywhere, BlueprintReadWrite, Category = RuntimeVirtualTextureParameterValue)
+	class URuntimeVirtualTexture* ParameterValue;
+
+	UPROPERTY()
+	FGuid ExpressionGUID;
+
+	FRuntimeVirtualTextureParameterValue()
+		: ParameterValue(NULL)
+	{
+	}
+
+	bool operator==(const FRuntimeVirtualTextureParameterValue& Other) const
+	{
+		return
+			ParameterInfo == Other.ParameterInfo &&
+			ParameterValue == Other.ParameterValue &&
+			ExpressionGUID == Other.ExpressionGUID;
+	}
+	bool operator!=(const FRuntimeVirtualTextureParameterValue& Other) const
+	{
+		return !((*this) == Other);
+	}
+
+	typedef const URuntimeVirtualTexture* ValueType;
+	static ValueType GetValue(const FRuntimeVirtualTextureParameterValue& Parameter) { return Parameter.ParameterValue; }
+};
+
 /** Editable font parameter. */
 USTRUCT(BlueprintType)
 struct FFontParameterValue
@@ -339,6 +375,10 @@ class ENGINE_VTABLE UMaterialInstance : public UMaterialInterface
 	UPROPERTY(EditAnywhere, BlueprintReadOnly, Category=MaterialInstance)
 	TArray<struct FTextureParameterValue> TextureParameterValues;
 
+	/** RuntimeVirtualTexture parameters. */
+	UPROPERTY(EditAnywhere, BlueprintReadOnly, Category = MaterialInstance)
+	TArray<struct FRuntimeVirtualTextureParameterValue> RuntimeVirtualTextureParameterValues;
+
 	/** Font parameters. */
 	UPROPERTY(EditAnywhere, BlueprintReadOnly, Category=MaterialInstance)
 	TArray<struct FFontParameterValue> FontParameterValues;
@@ -557,6 +597,7 @@ public:
 	ENGINE_API virtual void GetAllScalarParameterInfo(TArray<FMaterialParameterInfo>& OutParameterInfo, TArray<FGuid>& OutParameterIds) const override;
 	ENGINE_API virtual void GetAllVectorParameterInfo(TArray<FMaterialParameterInfo>& OutParameterInfo, TArray<FGuid>& OutParameterIds) const override;
 	ENGINE_API virtual void GetAllTextureParameterInfo(TArray<FMaterialParameterInfo>& OutParameterInfo, TArray<FGuid>& OutParameterIds) const override;
+	ENGINE_API virtual void GetAllRuntimeVirtualTextureParameterInfo(TArray<FMaterialParameterInfo>& OutParameterInfo, TArray<FGuid>& OutParameterIds) const override;
 	ENGINE_API virtual void GetAllFontParameterInfo(TArray<FMaterialParameterInfo>& OutParameterInfo, TArray<FGuid>& OutParameterIds) const override;
 	ENGINE_API virtual void GetAllMaterialLayersParameterInfo(TArray<FMaterialParameterInfo>& OutParameterInfo, TArray<FGuid>& OutParameterIds) const override;
 	ENGINE_API virtual void GetAllStaticSwitchParameterInfo(TArray<FMaterialParameterInfo>& OutParameterInfo, TArray<FGuid>& OutParameterIds) const override;
@@ -567,6 +608,7 @@ public:
 	ENGINE_API virtual bool GetScalarParameterDefaultValue(const FMaterialParameterInfo& ParameterInfo, float& OutValue, bool bOveriddenOnly = false, bool bCheckOwnedGlobalOverrides = false) const override;
 	ENGINE_API virtual bool GetVectorParameterDefaultValue(const FMaterialParameterInfo& ParameterInfo, FLinearColor& OutValue, bool bOveriddenOnly = false, bool bCheckOwnedGlobalOverrides = false) const override;
 	ENGINE_API virtual bool GetTextureParameterDefaultValue(const FMaterialParameterInfo& ParameterInfo, UTexture*& OutValue, bool bCheckOwnedGlobalOverrides = false) const override;
+	ENGINE_API virtual bool GetRuntimeVirtualTextureParameterDefaultValue(const FMaterialParameterInfo& ParameterInfo, URuntimeVirtualTexture*& OutValue, bool bCheckOwnedGlobalOverrides = false) const override;
 	ENGINE_API virtual bool GetFontParameterDefaultValue(const FMaterialParameterInfo& ParameterInfo, UFont*& OutFontValue, int32& OutFontPage, bool bCheckOwnedGlobalOverrides = false) const override;
 	ENGINE_API virtual bool GetStaticSwitchParameterDefaultValue(const FMaterialParameterInfo& ParameterInfo, bool& OutValue, FGuid& OutExpressionGuid, bool bCheckOwnedGlobalOverrides = false) const override;
 	ENGINE_API virtual bool GetStaticComponentMaskParameterDefaultValue(const FMaterialParameterInfo& ParameterInfo, bool& OutR, bool& OutG, bool& OutB, bool& OutA, FGuid& OutExpressionGuid, bool bCheckOwnedGlobalOverrides = false) const override;
@@ -661,6 +703,7 @@ protected:
 	void SetScalarParameterAtlasInternal(const FMaterialParameterInfo& ParameterInfo, FScalarParameterAtlasInstanceData AtlasData);
 #endif
 	void SetTextureParameterValueInternal(const FMaterialParameterInfo& ParameterInfo, class UTexture* Value);
+	void SetRuntimeVirtualTextureParameterValueInternal(const FMaterialParameterInfo& ParameterInfo, class URuntimeVirtualTexture* Value);
 	void SetFontParameterValueInternal(const FMaterialParameterInfo& ParameterInfo, class UFont* FontValue, int32 FontPage);
 	void ClearParameterValuesInternal(const bool bAllParameters = true);
 
diff --git a/Engine/Source/Runtime/Engine/Classes/Materials/MaterialInstanceConstant.h b/Engine/Source/Runtime/Engine/Classes/Materials/MaterialInstanceConstant.h
--- a/Engine/Source/Runtime/Engine/Classes/Materials/MaterialInstanceConstant.h
+++ b/Engine/Source/Runtime/Engine/Classes/Materials/MaterialInstanceConstant.h
@@ -74,6 +74,7 @@ class ENGINE_VTABLE UMaterialInstanceConstant : public UMaterialInstance
 	ENGINE_API void SetScalarParameterValueEditorOnly(const FMaterialParameterInfo& ParameterInfo, float Value);
 	ENGINE_API void SetScalarParameterAtlasEditorOnly(const FMaterialParameterInfo& ParameterInfo, FScalarParameterAtlasInstanceData AtlasData);
 	ENGINE_API void SetTextureParameterValueEditorOnly(const FMaterialParameterInfo& ParameterInfo, class UTexture* Value);
+	ENGINE_API void SetRuntimeVirtualTextureParameterValueEditorOnly(const FMaterialParameterInfo& ParameterInfo, class URuntimeVirtualTexture* Value);
 	ENGINE_API void SetFontParameterValueEditorOnly(const FMaterialParameterInfo& ParameterInfo, class UFont* FontValue, int32 FontPage);
 
 	/**
diff --git a/Engine/Source/Runtime/Engine/Private/Materials/MaterialInstance.cpp b/Engine/Source/Runtime/Engine/Private/Materials/MaterialInstance.cpp
--- a/Engine/Source/Runtime/Engine/Private/Materials/MaterialInstance.cpp
+++ b/Engine/Source/Runtime/Engine/Private/Materials/MaterialInstance.cpp
@@ -293,8 +293,21 @@ bool FMaterialInstanceResource::GetTextureValue(
 	const FMaterialRenderContext& Context
 ) const
 {
-	//todo[vt]
-	return false;
+	checkSlow(IsInParallelRenderingThread());
+	const URuntimeVirtualTexture* const * Value = RenderThread_FindParameterByName<const URuntimeVirtualTexture*>(ParameterInfo);
+	if (Value && *Value)
+	{
+		*OutValue = *Value;
+		return true;
+	}
+	else if (Parent)
+	{
+		return Parent->GetRenderProxy()->GetTextureValue(ParameterInfo, OutValue, Context);
+	}
+	else
+	{
+		return false;
+	}
 }
 
 void UMaterialInstance::PropagateDataToMaterialProxy()
@@ -383,6 +396,9 @@ bool UMaterialInstance::UpdateParameters()
 			// Texture parameters
 			bDirty = UpdateParameterSet<FTextureParameterValue, UMaterialExpressionTextureSampleParameter>(TextureParameterValues, ParentMaterial) || bDirty;
 
+			// Runtime Virtual Texture parameters
+			bDirty = UpdateParameterSet<FRuntimeVirtualTextureParameterValue, UMaterialExpressionRuntimeVirtualTextureSampleParameter>(RuntimeVirtualTextureParameterValues, ParentMaterial) || bDirty;
+
 			// Font parameters
 			bDirty = UpdateParameterSet<FFontParameterValue, UMaterialExpressionFontSampleParameter>(FontParameterValues, ParentMaterial) || bDirty;
 
@@ -481,6 +497,7 @@ void UMaterialInstance::InitResources()
 	GameThread_InitMIParameters(this, VectorParameterValues);
 	GameThread_InitMIParameters(this, TextureParameterValues);
 	GameThread_InitMIParameters(this, FontParameterValues);
+	GameThread_InitMIParameters(this, RuntimeVirtualTextureParameterValues);
 
 	PropagateDataToMaterialProxy();
 
@@ -851,7 +868,56 @@ bool UMaterialInstance::GetTextureParameterValue(const FMaterialParameterInfo& P
 
 bool UMaterialInstance::GetRuntimeVirtualTextureParameterValue(const FMaterialParameterInfo& ParameterInfo, URuntimeVirtualTexture*& OutValue, bool bOveriddenOnly) const
 {
-	// todo[vt]:
+	bool bFoundAValue = false;
+
+	if (GetReentrantFlag())
+	{
+		return false;
+	}
+
+	// Instance override
+	const FRuntimeVirtualTextureParameterValue* ParameterValue = GameThread_FindParameterByName(RuntimeVirtualTextureParameterValues, ParameterInfo);
+	if (ParameterValue)
+	{
+		OutValue = ParameterValue->ParameterValue;
+		return true;
+	}
+
+	// Instance-included default
+	if (ParameterInfo.Association != EMaterialParameterAssociation::GlobalParameter)
+	{
+		UMaterialExpressionRuntimeVirtualTextureSampleParameter* Parameter = nullptr;
+		for (const FStaticMaterialLayersParameter& LayersParam : StaticParameters.MaterialLayersParameters)
+		{
+			if (LayersParam.bOverride)
+			{
+				UMaterialFunctionInterface* Function = LayersParam.GetParameterAssociatedFunction(ParameterInfo);
+				UMaterialFunctionInterface* ParameterOwner = nullptr;
+
+				if (Function && Function->OverrideNamedRuntimeVirtualTextureParameter(ParameterInfo, OutValue))
+				{
+					return true;
+				}
+
+				if (Function && Function->GetNamedParameterOfType(ParameterInfo, Parameter, &ParameterOwner))
+				{
+					if (!ParameterOwner->OverrideNamedRuntimeVirtualTextureParameter(ParameterInfo, OutValue))
+					{
+						Parameter->IsNamedParameter(ParameterInfo, OutValue);
+					}
+					return true;
+				}
+			}
+		}
+	}
+
+	// Next material in hierarchy
+	if (Parent)
+	{
+		FMICReentranceGuard	Guard(this);
+		return Parent->GetRuntimeVirtualTextureParameterValue(ParameterInfo, OutValue, bOveriddenOnly);
+	}
+
 	return false;
 }
 
@@ -1634,6 +1700,22 @@ void UMaterialInstance::CopyMaterialInstanceParameters(UMaterialInterface* Sourc
 			}
 		}
 
+		// Now do the runtime virtual texture params
+		OutParameterInfo.Reset();
+		Guids.Reset();
+		GetAllRuntimeVirtualTextureParameterInfo(OutParameterInfo, Guids);
+		for (const FMaterialParameterInfo& ParameterInfo : OutParameterInfo)
+		{
+			URuntimeVirtualTexture* Value = nullptr;
+			if (Source->GetRuntimeVirtualTextureParameterValue(ParameterInfo, Value))
+			{
+				FRuntimeVirtualTextureParameterValue* ParameterValue = new(RuntimeVirtualTextureParameterValues) FRuntimeVirtualTextureParameterValue;
+				ParameterValue->ParameterInfo = ParameterInfo;
+				ParameterValue->ExpressionGUID.Invalidate();
+				ParameterValue->ParameterValue = Value;
+			}
+		}
+
 		// Now, init the resources
 		InitResources();
 	}
@@ -1837,6 +1919,16 @@ void UMaterialInstance::GetAllTextureParameterInfo(TArray<FMaterialParameterInfo
 	}
 }
 
+void UMaterialInstance::GetAllRuntimeVirtualTextureParameterInfo(TArray<FMaterialParameterInfo>& OutParameterInfo, TArray<FGuid>& OutParameterIds) const
+{
+	OutParameterInfo.Empty();
+	OutParameterIds.Empty();
+	if (const UMaterial* Material = GetMaterial())
+	{
+		Material->GetAllParameterInfo<UMaterialExpressionRuntimeVirtualTextureSampleParameter>(OutParameterInfo, OutParameterIds, &StaticParameters.MaterialLayersParameters);
+	}
+}
+
 void UMaterialInstance::GetAllFontParameterInfo(TArray<FMaterialParameterInfo>& OutParameterInfo, TArray<FGuid>& OutParameterIds) const
 {
 	OutParameterInfo.Empty();
@@ -2098,6 +2190,69 @@ bool UMaterialInstance::GetTextureParameterDefaultValue(const FMaterialParameter
 	return false;
 }
 
+bool UMaterialInstance::GetRuntimeVirtualTextureParameterDefaultValue(const FMaterialParameterInfo& ParameterInfo, URuntimeVirtualTexture*& OutValue, bool bCheckOwnedGlobalOverrides) const
+{
+#if WITH_EDITOR
+	if (GetReentrantFlag())
+	{
+		return false;
+	}
+#endif
+
+	// In the case of duplicate parameters with different values, this will return the
+	// first matching expression found, not necessarily the one that's used for rendering
+	UMaterialExpressionRuntimeVirtualTextureSampleParameter* Parameter = nullptr;
+
+	if (ParameterInfo.Association != EMaterialParameterAssociation::GlobalParameter)
+	{
+		// Parameters introduced by this instance's layer stack
+		for (const FStaticMaterialLayersParameter& LayersParam : StaticParameters.MaterialLayersParameters)
+		{
+			if (LayersParam.bOverride)
+			{
+				UMaterialFunctionInterface* Function = LayersParam.GetParameterAssociatedFunction(ParameterInfo);
+				UMaterialFunctionInterface* ParameterOwner = nullptr;
+
+				if (Function && Function->OverrideNamedRuntimeVirtualTextureParameter(ParameterInfo, OutValue))
+				{
+					return true;
+				}
+
+				if (Function && Function->GetNamedParameterOfType(ParameterInfo, Parameter, &ParameterOwner))
+				{
+					if (!ParameterOwner->OverrideNamedRuntimeVirtualTextureParameter(ParameterInfo, OutValue))
+					{
+						Parameter->IsNamedParameter(ParameterInfo, OutValue);
+					}
+					return true;
+				}
+			}
+		}
+	}
+	else if (bCheckOwnedGlobalOverrides)
+	{
+		// Parameters overridden by this instance
+		for (const FRuntimeVirtualTextureParameterValue& RuntimeVirtualTextureParam : RuntimeVirtualTextureParameterValues)
+		{
+			if (RuntimeVirtualTextureParam.ParameterInfo == ParameterInfo)
+			{
+				OutValue = RuntimeVirtualTextureParam.ParameterValue;
+				return true;
+			}
+		}
+	}
+
+	if (Parent)
+	{
+#if WITH_EDITOR
+		FMICReentranceGuard	Guard(this);
+#endif
+		return Parent->GetRuntimeVirtualTextureParameterDefaultValue(ParameterInfo, OutValue, true);
+	}
+
+	return false;
+}
+
 bool UMaterialInstance::GetFontParameterDefaultValue(const FMaterialParameterInfo& ParameterInfo, UFont*& OutFontValue, int32& OutFontPage, bool bCheckOwnedGlobalOverrides) const
 {
 #if WITH_EDITOR
@@ -3217,6 +3372,17 @@ void UMaterialInstance::PostLoad()
 		}
 	}
 
+	// do the same for runtime virtual textures
+	for (int32 ValueIndex = 0; ValueIndex < RuntimeVirtualTextureParameterValues.Num(); ValueIndex++)
+	{
+		// Make sure the texture is postloaded so the resource isn't null.
+		URuntimeVirtualTexture* Value = RuntimeVirtualTextureParameterValues[ValueIndex].ParameterValue;
+		if (Value)
+		{
+			Value->ConditionalPostLoad();
+		}
+	}
+
 	// do the same for font textures
 	for( int32 ValueIndex=0; ValueIndex < FontParameterValues.Num(); ValueIndex++ )
 	{
@@ -3566,6 +3732,36 @@ void UMaterialInstance::SetTextureParameterValueInternal(const FMaterialParamete
 	}
 }
 
+void UMaterialInstance::SetRuntimeVirtualTextureParameterValueInternal(const FMaterialParameterInfo& ParameterInfo, URuntimeVirtualTexture* Value)
+{
+	LLM_SCOPE(ELLMTag::MaterialInstance);
+
+	FRuntimeVirtualTextureParameterValue* ParameterValue = GameThread_FindParameterByName(RuntimeVirtualTextureParameterValues, ParameterInfo);
+
+	bool bForceUpdate = false;
+	if (!ParameterValue)
+	{
+		// If there's no element for the named parameter in array yet, add one.
+		ParameterValue = new(RuntimeVirtualTextureParameterValues) FRuntimeVirtualTextureParameterValue;
+		ParameterValue->ParameterInfo = ParameterInfo;
+		ParameterValue->ExpressionGUID.Invalidate();
+		bForceUpdate = true;
+	}
+
+	// Don't enqueue an update if it isn't needed
+	if (bForceUpdate || ParameterValue->ParameterValue != Value)
+	{
+		// set as an ensure, because it is somehow possible to accidentally pass non-textures into here via blueprints...
+		if (Value && ensureMsgf(Value->IsA(URuntimeVirtualTexture::StaticClass()), TEXT("Expecting a URuntimeVirtualTexture! Value='%s' class='%s'"), *Value->GetName(), *Value->GetClass()->GetName()))
+		{
+			ParameterValue->ParameterValue = Value;
+			// Update the material instance data in the rendering thread.
+			GameThread_UpdateMIParameter(this, *ParameterValue);
+			CacheMaterialInstanceUniformExpressions(this);
+		}
+	}
+}
+
 void UMaterialInstance::SetFontParameterValueInternal(const FMaterialParameterInfo& ParameterInfo,class UFont* FontValue,int32 FontPage)
 {
 	LLM_SCOPE(ELLMTag::MaterialInstance);
@@ -3603,6 +3799,7 @@ void UMaterialInstance::ClearParameterValuesInternal(const bool bAllParameters)
 	if(bAllParameters)
 	{
 		TextureParameterValues.Empty();
+		RuntimeVirtualTextureParameterValues.Empty();
 		FontParameterValues.Empty();
 	}
 
@@ -3912,6 +4109,7 @@ void UMaterialInstance::GetResourceSizeEx(FResourceSizeEx& CumulativeResourceSiz
 		CumulativeResourceSize.AddDedicatedSystemMemoryBytes(ScalarParameterValues.Num() * sizeof(FMaterialInstanceResource::TNamedParameter<float>));
 		CumulativeResourceSize.AddDedicatedSystemMemoryBytes(VectorParameterValues.Num() * sizeof(FMaterialInstanceResource::TNamedParameter<FLinearColor>));
 		CumulativeResourceSize.AddDedicatedSystemMemoryBytes(TextureParameterValues.Num() * sizeof(FMaterialInstanceResource::TNamedParameter<const UTexture*>));
+		CumulativeResourceSize.AddDedicatedSystemMemoryBytes(RuntimeVirtualTextureParameterValues.Num() * sizeof(FMaterialInstanceResource::TNamedParameter<const URuntimeVirtualTexture*>));
 		CumulativeResourceSize.AddDedicatedSystemMemoryBytes(FontParameterValues.Num() * sizeof(FMaterialInstanceResource::TNamedParameter<const UTexture*>));
 	}
 }
@@ -4221,6 +4419,10 @@ bool UMaterialInstance::Equivalent(const UMaterialInstance* CompareTo) const
 	{
 		return false;
 	}
+	if (!CompareValueArraysByExpressionGUID(RuntimeVirtualTextureParameterValues, CompareTo->RuntimeVirtualTextureParameterValues))
+	{
+		return false;
+	}
 	if (!CompareValueArraysByExpressionGUID(FontParameterValues, CompareTo->FontParameterValues))
 	{
 		return false;
@@ -4443,6 +4645,28 @@ void UMaterialInstance::CopyMaterialUniformParametersInternal(UMaterialInterface
 						TextureParameterValues.Add(Parameter);
 					}
 				}
+
+				// Runtime Virtual Textures
+				for (FRuntimeVirtualTextureParameterValue& Parameter : AsInstance->RuntimeVirtualTextureParameterValues)
+				{
+					// If the parameter already exists, override it
+					bool bExisting = false;
+					for (FRuntimeVirtualTextureParameterValue& ExistingParameter : RuntimeVirtualTextureParameterValues)
+					{
+						if (ExistingParameter.ParameterInfo.Name == Parameter.ParameterInfo.Name)
+						{
+							ExistingParameter.ParameterValue = Parameter.ParameterValue;
+							bExisting = true;
+							break;
+						}
+					}
+
+					// Instance has introduced a new parameter via static param set
+					if (!bExisting)
+					{
+						RuntimeVirtualTextureParameterValues.Add(Parameter);
+					}
+				}
 			}
 			else if (UMaterial* AsMaterial = Cast<UMaterial>(Interface))
 			{
@@ -4450,6 +4674,7 @@ void UMaterialInstance::CopyMaterialUniformParametersInternal(UMaterialInterface
 				checkSlow(ScalarParameterValues.Num() == 0);
 				checkSlow(VectorParameterValues.Num() == 0);
 				checkSlow(TextureParameterValues.Num() == 0);
+				checkSlow(RuntimeVirtualTextureParameterValues.Num() == 0);
 
 				const FMaterialResource* MaterialResource = nullptr;
 				if (UWorld* World = AsMaterial->GetWorld())
diff --git a/Engine/Source/Runtime/Engine/Private/Materials/MaterialInstanceConstant.cpp b/Engine/Source/Runtime/Engine/Private/Materials/MaterialInstanceConstant.cpp
--- a/Engine/Source/Runtime/Engine/Private/Materials/MaterialInstanceConstant.cpp
+++ b/Engine/Source/Runtime/Engine/Private/Materials/MaterialInstanceConstant.cpp
@@ -102,6 +102,12 @@ void UMaterialInstanceConstant::SetTextureParameterValueEditorOnly(const FMateri
 	SetTextureParameterValueInternal(ParameterInfo,Value);
 }
 
+void UMaterialInstanceConstant::SetRuntimeVirtualTextureParameterValueEditorOnly(const FMaterialParameterInfo& ParameterInfo, URuntimeVirtualTexture* Value)
+{
+	check(GIsEditor || IsRunningCommandlet());
+	SetRuntimeVirtualTextureParameterValueInternal(ParameterInfo, Value);
+}
+
 void UMaterialInstanceConstant::SetFontParameterValueEditorOnly(const FMaterialParameterInfo& ParameterInfo,class UFont* FontValue,int32 FontPage)
 {
 	check(GIsEditor || IsRunningCommandlet());
diff --git a/Engine/Source/Runtime/Engine/Private/Materials/MaterialInstanceSupport.h b/Engine/Source/Runtime/Engine/Private/Materials/MaterialInstanceSupport.h
--- a/Engine/Source/Runtime/Engine/Private/Materials/MaterialInstanceSupport.h
+++ b/Engine/Source/Runtime/Engine/Private/Materials/MaterialInstanceSupport.h
@@ -111,6 +111,7 @@ public:
 		VectorParameterArray.Reset();
 		ScalarParameterArray.Reset();
 		TextureParameterArray.Reset();
+		RuntimeVirtualTextureParameterArray.Reset();
 		InvalidateUniformExpressionCache(false);
 	}
 
@@ -181,14 +182,18 @@ private:
 	TArray<TNamedParameter<float> > ScalarParameterArray;
 	/** Texture parameters for this material instance. */
 	TArray<TNamedParameter<const UTexture*> > TextureParameterArray;
+	/** Runtime Virtual Texture parameters for this material instance. */
+	TArray<TNamedParameter<const URuntimeVirtualTexture*> > RuntimeVirtualTextureParameterArray; 
 };
 
 template <> FORCEINLINE TArray<FMaterialInstanceResource::TNamedParameter<float> >& FMaterialInstanceResource::GetValueArray() { return ScalarParameterArray; }
 template <> FORCEINLINE TArray<FMaterialInstanceResource::TNamedParameter<FLinearColor> >& FMaterialInstanceResource::GetValueArray() { return VectorParameterArray; }
 template <> FORCEINLINE TArray<FMaterialInstanceResource::TNamedParameter<const UTexture*> >& FMaterialInstanceResource::GetValueArray() { return TextureParameterArray; }
+template <> FORCEINLINE TArray<FMaterialInstanceResource::TNamedParameter<const URuntimeVirtualTexture*> >& FMaterialInstanceResource::GetValueArray() { return RuntimeVirtualTextureParameterArray; }
 template <> FORCEINLINE const TArray<FMaterialInstanceResource::TNamedParameter<float> >& FMaterialInstanceResource::GetValueArray() const { return ScalarParameterArray; }
 template <> FORCEINLINE const TArray<FMaterialInstanceResource::TNamedParameter<FLinearColor> >& FMaterialInstanceResource::GetValueArray() const { return VectorParameterArray; }
 template <> FORCEINLINE const TArray<FMaterialInstanceResource::TNamedParameter<const UTexture*> >& FMaterialInstanceResource::GetValueArray() const { return TextureParameterArray; }
+template <> FORCEINLINE const TArray<FMaterialInstanceResource::TNamedParameter<const URuntimeVirtualTexture*> >& FMaterialInstanceResource::GetValueArray() const { return RuntimeVirtualTextureParameterArray; }
 
 /** Finds a parameter by name from the game thread. */
 template <typename ParameterType>
