Implement RVT parameter support for material functions

diff --git a/Engine/Source/Runtime/Engine/Classes/Materials/MaterialFunction.h b/Engine/Source/Runtime/Engine/Classes/Materials/MaterialFunction.h
--- a/Engine/Source/Runtime/Engine/Classes/Materials/MaterialFunction.h
+++ b/Engine/Source/Runtime/Engine/Classes/Materials/MaterialFunction.h
@@ -143,6 +143,7 @@ public:
 	ENGINE_API bool SetVectorParameterValueEditorOnly(FName ParameterName, FLinearColor InValue);
 	ENGINE_API bool SetScalarParameterValueEditorOnly(FName ParameterName, float InValue);
 	ENGINE_API bool SetTextureParameterValueEditorOnly(FName ParameterName, class UTexture* InValue);
+	ENGINE_API bool SetRuntimeVirtualTextureParameterValueEditorOnly(FName ParameterName, class URuntimeVirtualTexture* InValue);
 	ENGINE_API bool SetFontParameterValueEditorOnly(FName ParameterName, class UFont* InFontValue, int32 InFontPage);
 	ENGINE_API bool SetStaticComponentMaskParameterValueEditorOnly(FName ParameterName, bool R, bool G, bool B, bool A, FGuid OutExpressionGuid);
 	ENGINE_API bool SetStaticSwitchParameterValueEditorOnly(FName ParameterName, bool OutValue, FGuid OutExpressionGuid);
diff --git a/Engine/Source/Runtime/Engine/Classes/Materials/MaterialFunctionInstance.h b/Engine/Source/Runtime/Engine/Classes/Materials/MaterialFunctionInstance.h
--- a/Engine/Source/Runtime/Engine/Classes/Materials/MaterialFunctionInstance.h
+++ b/Engine/Source/Runtime/Engine/Classes/Materials/MaterialFunctionInstance.h
@@ -61,6 +61,10 @@ class UMaterialFunctionInstance : public UMaterialFunctionInterface
 	UPROPERTY(EditAnywhere, Category=MaterialFunctionInstance)
 	TArray<struct FStaticComponentMaskParameter> StaticComponentMaskParameterValues;
 
+	/** Runtime virtual texture parameters. */
+	UPROPERTY(EditAnywhere, Category = MaterialFunctionInstance)
+	TArray<struct FRuntimeVirtualTextureParameterValue> RuntimeVirtualTextureParameterValues;
+
 	ENGINE_API void UpdateParameterSet();
 #if WITH_EDITOR
 	ENGINE_API void OverrideMaterialInstanceParameterValues(class UMaterialInstance* Instance);
@@ -160,6 +164,7 @@ public:
 	virtual bool OverrideNamedScalarParameter(const FMaterialParameterInfo& ParameterInfo, float& OutValue) override;
 	virtual bool OverrideNamedVectorParameter(const FMaterialParameterInfo& ParameterInfo, FLinearColor& OutValue) override;
 	virtual bool OverrideNamedTextureParameter(const FMaterialParameterInfo& ParameterInfo, class UTexture*& OutValue) override;
+	virtual bool OverrideNamedRuntimeVirtualTextureParameter(const FMaterialParameterInfo& ParameterInfo, class URuntimeVirtualTexture*& OutValue) override;
 	virtual bool OverrideNamedFontParameter(const FMaterialParameterInfo& ParameterInfo, class UFont*& OutFontValue, int32& OutFontPage) override;
 	virtual bool OverrideNamedStaticSwitchParameter(const FMaterialParameterInfo& ParameterInfo, bool& OutValue, FGuid& OutExpressionGuid) override;
 	virtual bool OverrideNamedStaticComponentMaskParameter(const FMaterialParameterInfo& ParameterInfo, bool& OutR, bool& OutG, bool& OutB, bool& OutA, FGuid& OutExpressionGuid) override;
diff --git a/Engine/Source/Runtime/Engine/Private/Materials/MaterialExpressions.cpp b/Engine/Source/Runtime/Engine/Private/Materials/MaterialExpressions.cpp
--- a/Engine/Source/Runtime/Engine/Private/Materials/MaterialExpressions.cpp
+++ b/Engine/Source/Runtime/Engine/Private/Materials/MaterialExpressions.cpp
@@ -11539,6 +11539,49 @@ bool UMaterialFunction::SetTextureParameterValueEditorOnly(FName ParameterName,
 	return false;
 };
 
+bool UMaterialFunction::SetRuntimeVirtualTextureParameterValueEditorOnly(FName ParameterName, class URuntimeVirtualTexture* InValue)
+{
+	for (UMaterialExpression* Expression : FunctionExpressions)
+	{
+		if (UMaterialExpressionRuntimeVirtualTextureSampleParameter* Parameter = Cast<UMaterialExpressionRuntimeVirtualTextureSampleParameter>(Expression))
+		{
+			if (Parameter->SetParameterValue(ParameterName, InValue))
+			{
+				return true;
+				// Warning: in the case of duplicate parameters with different default values, this will find the first in the expression array, not necessarily the one that's used for rendering
+			}
+		}
+		else if (UMaterialExpressionMaterialFunctionCall* FunctionCall = Cast<UMaterialExpressionMaterialFunctionCall>(Expression))
+		{
+			if (FunctionCall->MaterialFunction)
+			{
+				TArray<UMaterialFunctionInterface*> Functions;
+				Functions.Add(FunctionCall->MaterialFunction);
+				FunctionCall->MaterialFunction->GetDependentFunctions(Functions);
+
+				for (UMaterialFunctionInterface* Function : Functions)
+				{
+					const TArray<UMaterialExpression*>* ExpressionPtr = Function->GetFunctionExpressions();
+					if (ExpressionPtr)
+					{
+						for (UMaterialExpression* FunctionExpression : *ExpressionPtr)
+						{
+							if (UMaterialExpressionRuntimeVirtualTextureSampleParameter* FunctionExpressionParameter = Cast<UMaterialExpressionRuntimeVirtualTextureSampleParameter>(FunctionExpression))
+							{
+								if (FunctionExpressionParameter->SetParameterValue(ParameterName, InValue))
+								{
+									return true;
+								}
+							}
+						}
+					}
+				}
+			}
+		}
+	}
+	return false;
+};
+
 bool UMaterialFunction::SetFontParameterValueEditorOnly(FName ParameterName, class UFont* InFontValue, int32 InFontPage)
 {
 	for (UMaterialExpression* Expression : FunctionExpressions)
@@ -11728,6 +11771,17 @@ void UMaterialFunctionInstance::UpdateParameterSet()
 						}
 					}
 				}
+				else if (const UMaterialExpressionRuntimeVirtualTextureSampleParameter* RuntimeVirtualTextureParameter = Cast<const UMaterialExpressionRuntimeVirtualTextureSampleParameter>(FunctionExpression))
+				{
+					for (FRuntimeVirtualTextureParameterValue& RuntimeVirtualTextureParameterValue : RuntimeVirtualTextureParameterValues)
+					{
+						if (RuntimeVirtualTextureParameterValue.ExpressionGUID == RuntimeVirtualTextureParameter->ExpressionGUID)
+						{
+							RuntimeVirtualTextureParameterValue.ParameterInfo.Name = RuntimeVirtualTextureParameter->ParameterName;
+							break;
+						}
+					}
+				}
 				else if (const UMaterialExpressionFontSampleParameter* FontParameter = Cast<const UMaterialExpressionFontSampleParameter>(FunctionExpression))
 				{
 					for (FFontParameterValue& FontParameterValue : FontParameterValues)
@@ -11773,6 +11827,7 @@ void UMaterialFunctionInstance::OverrideMaterialInstanceParameterValues(UMateria
 	Instance->ScalarParameterValues = ScalarParameterValues;
 	Instance->VectorParameterValues = VectorParameterValues;
 	Instance->TextureParameterValues = TextureParameterValues;
+	Instance->RuntimeVirtualTextureParameterValues = RuntimeVirtualTextureParameterValues;
 	Instance->FontParameterValues = FontParameterValues;
 
 	// Static parameters
@@ -11936,6 +11991,20 @@ bool UMaterialFunctionInstance::OverrideNamedTextureParameter(const FMaterialPar
 	return false;
 }
 
+bool UMaterialFunctionInstance::OverrideNamedRuntimeVirtualTextureParameter(const FMaterialParameterInfo& ParameterInfo, URuntimeVirtualTexture*& OutValue)
+{
+	for (const FRuntimeVirtualTextureParameterValue& RuntimeVirtualTextureParameter : RuntimeVirtualTextureParameterValues)
+	{
+		if (RuntimeVirtualTextureParameter.ParameterInfo.Name == ParameterInfo.Name)
+		{
+			OutValue = RuntimeVirtualTextureParameter.ParameterValue;
+			return true;
+		}
+	}
+
+	return false;
+}
+
 bool UMaterialFunctionInstance::OverrideNamedFontParameter(const FMaterialParameterInfo& ParameterInfo, UFont*& OutFontValue, int32& OutFontPage)
 {
 	for (const FFontParameterValue& FontParameter : FontParameterValues)
