Add RVT parameter support for UMaterial

diff --git a/Engine/Source/Runtime/Engine/Classes/Materials/Material.h b/Engine/Source/Runtime/Engine/Classes/Materials/Material.h
--- a/Engine/Source/Runtime/Engine/Classes/Materials/Material.h
+++ b/Engine/Source/Runtime/Engine/Classes/Materials/Material.h
@@ -998,6 +998,7 @@ public:
 	ENGINE_API virtual bool GetVectorParameterValue(const FMaterialParameterInfo& ParameterInfo,FLinearColor& OutValue, bool bOveriddenOnly = false) const override;
 	ENGINE_API virtual bool IsVectorParameterUsedAsChannelMask(const FMaterialParameterInfo& ParameterInfo, bool& OutValue) const override;
 	ENGINE_API virtual bool GetTextureParameterValue(const FMaterialParameterInfo& ParameterInfo,class UTexture*& OutValue, bool bOveriddenOnly = false) const override;
+	ENGINE_API virtual bool GetRuntimeVirtualTextureParameterValue(const FMaterialParameterInfo& ParameterInfo, class URuntimeVirtualTexture*& OutValue, bool bOveriddenOnly = false) const override;
 	ENGINE_API virtual bool GetFontParameterValue(const FMaterialParameterInfo& ParameterInfo,class UFont*& OutFontValue,int32& OutFontPage, bool bOveriddenOnly = false) const override;
 	ENGINE_API virtual bool GetRefractionSettings(float& OutBiasValue) const override;
 	ENGINE_API virtual FMaterialRenderProxy* GetRenderProxy() const override;
@@ -1127,6 +1128,7 @@ public:
 	ENGINE_API bool SetVectorParameterValueEditorOnly(FName ParameterName, FLinearColor InValue);
 	ENGINE_API bool SetScalarParameterValueEditorOnly(FName ParameterName, float InValue);
 	ENGINE_API bool SetTextureParameterValueEditorOnly(FName ParameterName, class UTexture* InValue);
+	ENGINE_API bool SetRuntimeVirtualTextureParameterValueEditorOnly(FName ParameterName, class URuntimeVirtualTexture* InValue);
 	ENGINE_API bool SetFontParameterValueEditorOnly(FName ParameterName, class UFont* InFontValue, int32 InFontPage);
 	ENGINE_API bool SetMaterialLayersParameterValueEditorOnly(FName ParameterName, struct FMaterialLayersFunctions InValue, FGuid OutExpressionGuid);
 	ENGINE_API bool SetStaticComponentMaskParameterValueEditorOnly(FName ParameterName, bool R, bool G, bool B, bool A, FGuid OutExpressionGuid);
@@ -1327,6 +1329,7 @@ public:
 	ENGINE_API virtual bool GetScalarParameterDefaultValue(const FMaterialParameterInfo& ParameterInfo, float& OutValue, bool bOveriddenOnly = false, bool bCheckOwnedGlobalOverrides = false) const override;
 	ENGINE_API virtual bool GetVectorParameterDefaultValue(const FMaterialParameterInfo& ParameterInfo, FLinearColor& OutValue, bool bOveriddenOnly = false, bool bCheckOwnedGlobalOverrides = false) const override;
 	ENGINE_API virtual bool GetTextureParameterDefaultValue(const FMaterialParameterInfo& ParameterInfo, class UTexture*& OutValue, bool bCheckOwnedGlobalOverrides = false) const override;
+	ENGINE_API virtual bool GetRuntimeVirtualTextureParameterDefaultValue(const FMaterialParameterInfo& ParameterInfo, class URuntimeVirtualTexture*& OutValue, bool bCheckOwnedGlobalOverrides = false) const override;
 	ENGINE_API virtual bool GetFontParameterDefaultValue(const FMaterialParameterInfo& ParameterInfo, class UFont*& OutFontValue, int32& OutFontPage, bool bCheckOwnedGlobalOverrides = false) const override;
 	ENGINE_API virtual bool GetStaticComponentMaskParameterDefaultValue(const FMaterialParameterInfo& ParameterInfo, bool& OutR, bool& OutG, bool& OutB, bool& OutA, FGuid& OutExpressionGuid, bool bCheckOwnedGlobalOverrides = false) const override;
 	ENGINE_API virtual bool GetStaticSwitchParameterDefaultValue(const FMaterialParameterInfo& ParameterInfo, bool& OutValue, FGuid& OutExpressionGuid, bool bCheckOwnedGlobalOverrides = false) const override;
diff --git a/Engine/Source/Runtime/Engine/Classes/Materials/MaterialExpressionRuntimeVirtualTextureSample.h b/Engine/Source/Runtime/Engine/Classes/Materials/MaterialExpressionRuntimeVirtualTextureSample.h
--- a/Engine/Source/Runtime/Engine/Classes/Materials/MaterialExpressionRuntimeVirtualTextureSample.h
+++ b/Engine/Source/Runtime/Engine/Classes/Materials/MaterialExpressionRuntimeVirtualTextureSample.h
@@ -67,9 +67,10 @@ protected:
 
 #if WITH_EDITOR
 	virtual void PostLoad() override;
-	virtual void PostEditChangeProperty(FPropertyChangedEvent& PropertyChangedEvent) override;
 	virtual int32 Compile(class FMaterialCompiler* Compiler, int32 OutputIndex) override;
 	virtual void GetCaption(TArray<FString>& OutCaptions) const override;
+public:
+	virtual void PostEditChangeProperty(FPropertyChangedEvent& PropertyChangedEvent) override;
 #endif
 	//~ End UMaterialExpression Interface
 };
diff --git a/Engine/Source/Runtime/Engine/Private/Materials/Material.cpp b/Engine/Source/Runtime/Engine/Private/Materials/Material.cpp
--- a/Engine/Source/Runtime/Engine/Private/Materials/Material.cpp
+++ b/Engine/Source/Runtime/Engine/Private/Materials/Material.cpp
@@ -27,6 +27,7 @@
 #include "Materials/MaterialExpressionShadingPathSwitch.h"
 #include "Materials/MaterialExpressionMakeMaterialAttributes.h"
 #include "Materials/MaterialExpressionParameter.h"
+#include "Materials/MaterialExpressionRuntimeVirtualTextureSampleParameter.h"
 #include "Materials/MaterialExpressionScalarParameter.h"
 #include "Materials/MaterialExpressionStaticBoolParameter.h"
 #include "Materials/MaterialExpressionStaticComponentMaskParameter.h"
@@ -1315,6 +1316,11 @@ bool UMaterial::SetTextureParameterValueEditorOnly(FName ParameterName, class UT
 	return SetParameterValueEditorOnly<UMaterialExpressionTextureSampleParameter>(ParameterName, InValue);
 }
 
+bool UMaterial::SetRuntimeVirtualTextureParameterValueEditorOnly(FName ParameterName, class URuntimeVirtualTexture* InValue)
+{
+	return SetParameterValueEditorOnly<UMaterialExpressionRuntimeVirtualTextureSampleParameter>(ParameterName, InValue);
+}
+
 bool UMaterial::SetFontParameterValueEditorOnly(FName ParameterName, class UFont* InFontValue, int32 InFontPage)
 {
 	return SetParameterValueEditorOnly<UMaterialExpressionFontSampleParameter>(ParameterName, InFontValue, InFontPage);
@@ -1788,6 +1794,11 @@ bool UMaterial::GetTextureParameterDefaultValue(const FMaterialParameterInfo& Pa
 	return GetTextureParameterValue(ParameterInfo, OutValue);
 }
 
+bool UMaterial::GetRuntimeVirtualTextureParameterDefaultValue(const FMaterialParameterInfo& ParameterInfo, class URuntimeVirtualTexture*& OutValue, bool bCheckOwnedGlobalOverrides) const
+{
+	return GetRuntimeVirtualTextureParameterValue(ParameterInfo, OutValue);
+}
+
 bool UMaterial::GetFontParameterDefaultValue(const FMaterialParameterInfo& ParameterInfo, class UFont*& OutFontValue, int32& OutFontPage, bool bCheckOwnedGlobalOverrides) const
 {
 	return GetFontParameterValue(ParameterInfo, OutFontValue, OutFontPage);
@@ -2453,6 +2464,76 @@ bool UMaterial::GetTextureParameterValue(const FMaterialParameterInfo& Parameter
 	return false;
 }
 
+bool UMaterial::GetRuntimeVirtualTextureParameterValue(const FMaterialParameterInfo& ParameterInfo, URuntimeVirtualTexture*& OutValue, bool bOveriddenOnly) const
+{
+	if (bOveriddenOnly && !AreExperimentalMaterialLayersEnabled())
+	{
+		return false;
+	}
+
+	// In the case of duplicate parameters with different values, this will return the
+	// first matching expression found, not necessarily the one that's used for rendering
+	UMaterialExpressionRuntimeVirtualTextureSampleParameter* Parameter = nullptr;
+
+	for (UMaterialExpression* Expression : Expressions)
+	{
+		// Only need to check parameters that match in associated scope
+		if (ParameterInfo.Association == EMaterialParameterAssociation::GlobalParameter)
+		{
+			if (UMaterialExpressionRuntimeVirtualTextureSampleParameter* ExpressionParameter = Cast<UMaterialExpressionRuntimeVirtualTextureSampleParameter>(Expression))
+			{
+				if (ExpressionParameter->IsNamedParameter(ParameterInfo, OutValue))
+				{
+					return true;
+				}
+			}
+			else if (UMaterialExpressionMaterialFunctionCall* FunctionCall = Cast<UMaterialExpressionMaterialFunctionCall>(Expression))
+			{
+				UMaterialFunctionInterface* Function = FunctionCall->MaterialFunction;
+				UMaterialFunctionInterface* ParameterOwner = nullptr;
+
+				if (Function && Function->OverrideNamedRuntimeVirtualTextureParameter(ParameterInfo, OutValue))
+				{
+					return true;
+				}
+
+				if (Function && Function->GetNamedParameterOfType(ParameterInfo, Parameter, &ParameterOwner))
+				{
+					if (!ParameterOwner->OverrideNamedRuntimeVirtualTextureParameter(ParameterInfo, OutValue))
+					{
+						Parameter->IsNamedParameter(ParameterInfo, OutValue);
+					}
+					return true;
+				}
+			}
+		}
+		else
+		{
+			if (UMaterialExpressionMaterialAttributeLayers* LayersExpression = Cast<UMaterialExpressionMaterialAttributeLayers>(Expression))
+			{
+				UMaterialFunctionInterface* Function = LayersExpression->GetParameterAssociatedFunction(ParameterInfo);
+				UMaterialFunctionInterface* ParameterOwner = nullptr;
+
+				if (Function && Function->OverrideNamedRuntimeVirtualTextureParameter(ParameterInfo, OutValue))
+				{
+					return true;
+				}
+
+				if (Function && Function->GetNamedParameterOfType(ParameterInfo, Parameter, &ParameterOwner))
+				{
+					if (!ParameterOwner->OverrideNamedRuntimeVirtualTextureParameter(ParameterInfo, OutValue))
+					{
+						Parameter->IsNamedParameter(ParameterInfo, OutValue);
+					}
+					return true;
+				}
+			}
+		}
+	}
+
+	return false;
+}
+
 bool UMaterial::GetFontParameterValue(const FMaterialParameterInfo& ParameterInfo, UFont*& OutFontValue, int32& OutFontPage, bool bOveriddenOnly) const
 {
 	// In the case of duplicate parameters with different values, this will return the
