Implement RVT parameter usage in the RVT Sample material expression

diff --git a/Engine/Source/Runtime/Engine/Private/Materials/MaterialExpressions.cpp b/Engine/Source/Runtime/Engine/Private/Materials/MaterialExpressions.cpp
--- a/Engine/Source/Runtime/Engine/Private/Materials/MaterialExpressions.cpp
+++ b/Engine/Source/Runtime/Engine/Private/Materials/MaterialExpressions.cpp
@@ -2061,11 +2061,17 @@ void UMaterialExpressionRuntimeVirtualTextureSample::PostEditChangeProperty(FPro
 
 int32 UMaterialExpressionRuntimeVirtualTextureSample::Compile(class FMaterialCompiler* Compiler, int32 OutputIndex)
 {
+	// Is this a valid UMaterialExpressionRuntimeVirtualTextureSampleParameter?
+	const bool bIsParameter = HasAParameterName() && GetParameterName().IsValid() && !GetParameterName().IsNone();
+
 	// Check validity of current virtual texture
 	bool bIsVirtualTextureValid = VirtualTexture != nullptr;
 	if (!bIsVirtualTextureValid)
 	{
-		Compiler->Error(TEXT("Missing input Virtual Texture"));
+		if (!bIsParameter)
+		{
+			Compiler->Error(TEXT("Missing input Virtual Texture"));
+		}
 	}
 	else if (VirtualTexture->GetMaterialType() != MaterialType)
 	{
@@ -2102,7 +2108,7 @@ int32 UMaterialExpressionRuntimeVirtualTextureSample::Compile(class FMaterialCom
 	switch (OutputIndex)
 	{
 	case 0: 
-		if (bIsVirtualTextureValid && bIsBaseColorValid)
+		if ((bIsParameter || bIsVirtualTextureValid) && bIsBaseColorValid)
 		{
 			UnpackType = EVirtualTextureUnpackType::BaseColor;
 		}
@@ -2112,7 +2118,7 @@ int32 UMaterialExpressionRuntimeVirtualTextureSample::Compile(class FMaterialCom
 		}
 		break;
 	case 1:
-		if (bIsVirtualTextureValid && bIsSpecularValid)
+		if ((bIsParameter || bIsVirtualTextureValid) && bIsSpecularValid)
 		{
 			UnpackType = EVirtualTextureUnpackType::SpecularR8;
 		}
@@ -2122,7 +2128,7 @@ int32 UMaterialExpressionRuntimeVirtualTextureSample::Compile(class FMaterialCom
 		}
 		break;
 	case 2:
-		if (bIsVirtualTextureValid && bIsSpecularValid)
+		if ((bIsParameter || bIsVirtualTextureValid) && bIsSpecularValid)
 		{
 			switch (MaterialType)
 			{
@@ -2136,7 +2142,7 @@ int32 UMaterialExpressionRuntimeVirtualTextureSample::Compile(class FMaterialCom
 		}
 		break;
 	case 3:
-		if (bIsVirtualTextureValid && bIsNormalValid)
+		if ((bIsParameter || bIsVirtualTextureValid) && bIsNormalValid)
 		{
 			switch (MaterialType)
 			{
@@ -2151,7 +2157,7 @@ int32 UMaterialExpressionRuntimeVirtualTextureSample::Compile(class FMaterialCom
 		}
 		break;
 	case 4:
-		if (bIsVirtualTextureValid && bIsWorldHeightValid)
+		if ((bIsParameter || bIsVirtualTextureValid) && bIsWorldHeightValid)
 		{
 			UnpackType = EVirtualTextureUnpackType::HeightR16;
 		}
@@ -2173,14 +2179,13 @@ int32 UMaterialExpressionRuntimeVirtualTextureSample::Compile(class FMaterialCom
 	int32 TextureReferenceIndex[MAX_RVT_LAYERS] = { INDEX_NONE };
 	for (int32 Layer = 0; Layer < LayerCount; Layer++)
 	{
-		TextureCodeIndex[Layer] = Compiler->VirtualTexture(VirtualTexture, Layer, TextureReferenceIndex[Layer], SAMPLERTYPE_VirtualMasks);
-		if (TextureReferenceIndex[Layer] != TextureReferenceIndex[0])
+		if (bIsParameter)
+		{
+			TextureCodeIndex[Layer] = Compiler->VirtualTextureParameter(GetParameterName(), VirtualTexture, Layer, TextureReferenceIndex[Layer], SAMPLERTYPE_VirtualMasks);
+		}
+		else
 		{
-			Compiler->Errorf(TEXT("Virtual texture %s has invalid TextureReferenceIndex %d, which doesnt match %d in layer %d"),
-				bIsVirtualTextureValid ? *VirtualTexture->GetPathName() : TEXT("<unknown>"),
-				TextureReferenceIndex[Layer],
-				TextureReferenceIndex[0],
-				Layer);
+			TextureCodeIndex[Layer] = Compiler->VirtualTexture(VirtualTexture, Layer, TextureReferenceIndex[Layer], SAMPLERTYPE_VirtualMasks);
 		}
 	}
 
@@ -2190,9 +2195,19 @@ int32 UMaterialExpressionRuntimeVirtualTextureSample::Compile(class FMaterialCom
 	if (Coordinates.GetTracedInput().Expression == nullptr)
 	{
 		int32 WorldPositionIndex = Compiler->WorldPosition(WPT_Default);
-		int32 P0 = Compiler->VirtualTextureUniform(TextureReferenceIndex[0], 0);
-		int32 P1 = Compiler->VirtualTextureUniform(TextureReferenceIndex[0], 1);
-		int32 P2 = Compiler->VirtualTextureUniform(TextureReferenceIndex[0], 2);
+		int32 P0, P1, P2;
+		if (bIsParameter)
+		{
+			P0 = Compiler->VirtualTextureUniform(GetParameterName(), TextureReferenceIndex[0], 0);
+			P1 = Compiler->VirtualTextureUniform(GetParameterName(), TextureReferenceIndex[0], 1);
+			P2 = Compiler->VirtualTextureUniform(GetParameterName(), TextureReferenceIndex[0], 2);
+		}
+		else
+		{
+			P0 = Compiler->VirtualTextureUniform(TextureReferenceIndex[0], 0);
+			P1 = Compiler->VirtualTextureUniform(TextureReferenceIndex[0], 1);
+			P2 = Compiler->VirtualTextureUniform(TextureReferenceIndex[0], 2);
+		}
 		CoordinateIndex = Compiler->VirtualTextureWorldToUV(WorldPositionIndex, P0, P1, P2);
 	}
 	else
