Fixing spatialization issue with "au.EnableBinauralAudioForAllSpatialSounds" cvar

diff --git a/Engine/Source/Runtime/Engine/Private/SoundWave.cpp b/Engine/Source/Runtime/Engine/Private/SoundWave.cpp
--- a/Engine/Source/Runtime/Engine/Private/SoundWave.cpp
+++ b/Engine/Source/Runtime/Engine/Private/SoundWave.cpp
@@ -1669,18 +1669,37 @@ void USoundWave::Parse(FAudioDevice* AudioDevice, const UPTRINT NodeWaveInstance
 	WaveInstance->Location = ParseParams.Transform.GetTranslation();
 	WaveInstance->bIsStarted = true;
 	WaveInstance->bAlreadyNotifiedHook = false;
-	WaveInstance->SetUseSpatialization(ParseParams.bUseSpatialization);
-	WaveInstance->SpatializationMethod = ParseParams.SpatializationMethod;
+
 	WaveInstance->WaveData = this;
 	WaveInstance->NotifyBufferFinishedHooks = ParseParams.NotifyBufferFinishedHooks;
 	WaveInstance->LoopingMode = ((bLooping || ParseParams.bLooping) ? LOOP_Forever : LOOP_Never);
 	WaveInstance->bIsPaused = ParseParams.bIsPaused;
 
 	// If we're normalizing 3d stereo spatialized sounds, we need to scale by -6 dB
-	if (WaveInstance->GetUseSpatialization() && ParseParams.bApplyNormalizationToStereoSounds && NumChannels == 2)
+	WaveInstance->SetUseSpatialization(ParseParams.bUseSpatialization);
+
+	// Setup the spat method if we're actually spatializing (note a cvar can turn this off so we use the getter here)
+	if (WaveInstance->GetUseSpatialization())
 	{
-		const float ThisVolumeMultiplier = WaveInstance->GetVolumeMultiplier();
-		WaveInstance->SetVolumeMultiplier(ThisVolumeMultiplier * 0.5f);
+		WaveInstance->SpatializationMethod = ParseParams.SpatializationMethod;
+
+		// Check for possible HRTF-enforcement if this is a spatialized sound
+		if (AudioDevice->IsHRTFEnabledForAll() && ParseParams.SpatializationMethod == ESoundSpatializationAlgorithm::SPATIALIZATION_Default)
+		{
+			WaveInstance->SpatializationMethod = ESoundSpatializationAlgorithm::SPATIALIZATION_HRTF;
+		}
+		else
+		{
+			WaveInstance->SpatializationMethod = ParseParams.SpatializationMethod;
+		}
+
+		// Apply stereo normalization to wave instances if enabled
+		if (ParseParams.bApplyNormalizationToStereoSounds && NumChannels == 2)
+		{
+			const float ThisVolumeMultiplier = WaveInstance->GetVolumeMultiplier();
+			WaveInstance->SetVolumeMultiplier(ThisVolumeMultiplier * 0.5f);
+		}
+
 	}
 
 	// Copy reverb send settings
@@ -1709,15 +1728,6 @@ void USoundWave::Parse(FAudioDevice* AudioDevice, const UPTRINT NodeWaveInstance
 		WaveInstance->SoundSourceBusSends[BusSendType] = ParseParams.SoundSourceBusSends[BusSendType];
 	}
 
-	if (AudioDevice->IsHRTFEnabledForAll() && ParseParams.SpatializationMethod == ESoundSpatializationAlgorithm::SPATIALIZATION_Default)
-	{
-		WaveInstance->SpatializationMethod = ESoundSpatializationAlgorithm::SPATIALIZATION_HRTF;
-	}
-	else
-	{
-		WaveInstance->SpatializationMethod = ParseParams.SpatializationMethod;
-	}
-
 	// Pass along plugin settings to the wave instance
 	WaveInstance->SpatializationPluginSettings = ParseParams.SpatializationPluginSettings;
 	WaveInstance->OcclusionPluginSettings = ParseParams.OcclusionPluginSettings;
diff --git a/Engine/Source/Runtime/Engine/Public/Audio.h b/Engine/Source/Runtime/Engine/Public/Audio.h
--- a/Engine/Source/Runtime/Engine/Public/Audio.h
+++ b/Engine/Source/Runtime/Engine/Public/Audio.h
@@ -473,7 +473,7 @@ public:
 	/** Whether to use spatialization, which controls 3D effects like panning */
 	void SetUseSpatialization(const bool InUseSpatialization) { bUseSpatialization = InUseSpatialization; }
 
-	/** Whether this wave will be spatializized, which controls 3D effects like panning */
+	/** Whether this wave will be spatialized, which controls 3D effects like panning */
 	bool GetUseSpatialization() const;
 };
 
