Remove incorrect warning message from texture expressions

diff --git a/Engine/Source/Runtime/Engine/Private/Materials/MaterialUniformExpressions.h b/Engine/Source/Runtime/Engine/Private/Materials/MaterialUniformExpressions.h
--- a/Engine/Source/Runtime/Engine/Private/Materials/MaterialUniformExpressions.h
+++ b/Engine/Source/Runtime/Engine/Private/Materials/MaterialUniformExpressions.h
@@ -265,21 +265,11 @@ private:
 template<typename TextureType>
 static TextureType* GetIndexedTexture(const FMaterial& Material, int32 TextureIndex)
 {
-	TextureType* IndexedTexture = NULL;
-
+	UObject* IndexedTexture = nullptr;
 	const TArray<UObject*>& ReferencedTextures = Material.GetReferencedTextures();
 	if (ReferencedTextures.IsValidIndex(TextureIndex))
 	{
-		IndexedTexture = Cast<TextureType>(ReferencedTextures[TextureIndex]);
-	}
-	else
-	{
-		static bool bWarnedOnce = false;
-		if (!bWarnedOnce)
-		{
-			UE_LOG(LogMaterial, Warning, TEXT("%s: Requesting an invalid TextureIndex! (%u / %u)"), *Material.GetFriendlyName(), TextureIndex, ReferencedTextures.Num());
-			bWarnedOnce = true;
-		}
+		IndexedTexture = ReferencedTextures[TextureIndex];
 	}
 
 	if (IndexedTexture == nullptr)
@@ -287,12 +277,13 @@ static TextureType* GetIndexedTexture(const FMaterial& Material, int32 TextureIn
 		static bool bWarnedOnce = false;
 		if (!bWarnedOnce)
 		{
-			UE_LOG(LogMaterial, Warning, TEXT("%s: GetIndexedTexture returning NULL (%u)"), *Material.GetFriendlyName(), TextureIndex);
+			UE_LOG(LogMaterial, Warning, TEXT("%s: Requesting an invalid TextureIndex! (%u / %u)"), *Material.GetFriendlyName(), TextureIndex, ReferencedTextures.Num());
 			bWarnedOnce = true;
 		}
 	}
 
-	return IndexedTexture;
+	// Can return nullptr if TextureType doesn't match type of indexed texture
+	return Cast<TextureType>(IndexedTexture);
 }
 
 /**
