Fix poses breaking and IK being disabled in certain circumstances when using layers

diff --git a/Engine/Source/Runtime/Engine/Private/Animation/AnimInstance.cpp b/Engine/Source/Runtime/Engine/Private/Animation/AnimInstance.cpp
--- a/Engine/Source/Runtime/Engine/Private/Animation/AnimInstance.cpp
+++ b/Engine/Source/Runtime/Engine/Private/Animation/AnimInstance.cpp
@@ -2664,6 +2664,8 @@ void UAnimInstance::PerformLayerOverlayOperation(TSubclassOf<UAnimInstance> InCl
 				}
 			}
 		}
+
+		RecalcRequiredBones();
 	}
 }
 
diff --git a/Engine/Source/Runtime/Engine/Private/Animation/AnimNode_SubInput.cpp b/Engine/Source/Runtime/Engine/Private/Animation/AnimNode_SubInput.cpp
--- a/Engine/Source/Runtime/Engine/Private/Animation/AnimNode_SubInput.cpp
+++ b/Engine/Source/Runtime/Engine/Private/Animation/AnimNode_SubInput.cpp
@@ -38,7 +38,7 @@ void FAnimNode_SubInput::Evaluate_AnyThread(FPoseContext& Output)
 	{
 		Output.Pose.SetBoneContainer(&InputProxy->GetRequiredBones());
 
-		FPoseContext InputContext(InputProxy);
+		FPoseContext InputContext(InputProxy, Output.ExpectsAdditivePose());
 		InputPose.Evaluate(InputContext);
 
 		Output.Pose.MoveBonesFrom(InputContext.Pose);
diff --git a/Engine/Source/Runtime/Engine/Private/Animation/AnimNode_SubInstance.cpp b/Engine/Source/Runtime/Engine/Private/Animation/AnimNode_SubInstance.cpp
--- a/Engine/Source/Runtime/Engine/Private/Animation/AnimNode_SubInstance.cpp
+++ b/Engine/Source/Runtime/Engine/Private/Animation/AnimNode_SubInstance.cpp
@@ -62,7 +62,7 @@ void FAnimNode_SubInstance::Evaluate_AnyThread(FPoseContext& Output)
 		Output.Pose.SetBoneContainer(&Proxy.GetRequiredBones());
 
 		// Create an evaluation context
-		FPoseContext EvaluationContext(&Proxy);
+		FPoseContext EvaluationContext(&Proxy, Output.ExpectsAdditivePose());
 		EvaluationContext.ResetToRefPose();
 			
 		// Run the anim blueprint
