Actor registration to NavSystem
Moved from Initialize/RouteEndPlay to PreRegisterComponents/PostUnregisterComponents so Actors registration flow is the same in the editor and in gameplay.
It also prevents some actors from not being unregistered in Editor since some operations may register actors (ex: moving an actor) but since the actor may never be initialized for gameplay session it was not unregistered when unloaded (ex: removing sublevel)

diff --git a/Engine/Source/Runtime/Engine/Classes/GameFramework/Actor.h b/Engine/Source/Runtime/Engine/Classes/GameFramework/Actor.h
--- a/Engine/Source/Runtime/Engine/Classes/GameFramework/Actor.h
+++ b/Engine/Source/Runtime/Engine/Classes/GameFramework/Actor.h
@@ -2329,7 +2329,7 @@ public:
 	virtual void UnregisterAllComponents(bool bForReregister = false);
 
 	/** Called after all currently registered components are cleared */
-	virtual void PostUnregisterAllComponents() {}
+	virtual void PostUnregisterAllComponents();
 
 	/** Will reregister all components on this actor. Does a lot of work - should only really be used in editor, generally use UpdateComponentTransforms or MarkComponentsRenderStateDirty. */
 	virtual void ReregisterAllComponents();
diff --git a/Engine/Source/Runtime/Engine/Private/Actor.cpp b/Engine/Source/Runtime/Engine/Private/Actor.cpp
--- a/Engine/Source/Runtime/Engine/Private/Actor.cpp
+++ b/Engine/Source/Runtime/Engine/Private/Actor.cpp
@@ -2145,8 +2145,6 @@ void AActor::RouteEndPlay(const EEndPlayReason::Type EndPlayReason)
 		{
 			SetLifeSpan(0.f);
 		}
-
-		FNavigationSystem::OnActorUnregistered(*this);
 	}
 
 	UninitializeComponents();
@@ -2834,6 +2832,7 @@ void AActor::DisableComponentsSimulatePhysics()
 
 void AActor::PreRegisterAllComponents()
 {
+	FNavigationSystem::OnActorRegistered(*this);
 }
 
 void AActor::PostRegisterAllComponents() 
@@ -4284,6 +4283,11 @@ void AActor::UnregisterAllComponents(const bool bForReregister)
 	PostUnregisterAllComponents();
 }
 
+void AActor::PostUnregisterAllComponents()
+{
+	FNavigationSystem::OnActorUnregistered(*this);
+}
+
 void AActor::RegisterAllComponents()
 {
 	PreRegisterAllComponents();
@@ -4666,8 +4670,6 @@ void AActor::PostInitializeComponents()
 	if( !IsPendingKill() )
 	{
 		bActorInitialized = true;
-
-		FNavigationSystem::OnActorRegistered(*this);
 		
 		UpdateAllReplicatedComponents();
 	}
