Fix missing vtables breaking Linux build.

diff --git a/Engine/Source/Runtime/Engine/Public/PrimitiveUniformShaderParameters.h b/Engine/Source/Runtime/Engine/Public/PrimitiveUniformShaderParameters.h
--- a/Engine/Source/Runtime/Engine/Public/PrimitiveUniformShaderParameters.h
+++ b/Engine/Source/Runtime/Engine/Public/PrimitiveUniformShaderParameters.h
@@ -175,7 +175,7 @@ struct FPrimitiveSceneShaderData
 	ENGINE_API void Setup(const FPrimitiveUniformShaderParameters& PrimitiveUniformShaderParameters);
 };
 
-class FSinglePrimitiveStructuredBuffer : public FRenderResource
+class ENGINE_VTABLE FSinglePrimitiveStructuredBuffer : public FRenderResource
 {
 public:
 
diff --git a/Engine/Source/Runtime/Engine/Public/RawIndexBuffer.h b/Engine/Source/Runtime/Engine/Public/RawIndexBuffer.h
--- a/Engine/Source/Runtime/Engine/Public/RawIndexBuffer.h
+++ b/Engine/Source/Runtime/Engine/Public/RawIndexBuffer.h
@@ -118,7 +118,7 @@ private:
 	bool b32Bit;
 };
 
-class FRawStaticIndexBuffer : public FIndexBuffer
+class ENGINE_VTABLE FRawStaticIndexBuffer : public FIndexBuffer
 {
 public:	
 	/**
