Virtual function was added to the public interface which requires exporting the vtable of the class if done. Causes failure to link on Linux/Mac

diff --git a/Engine/Source/Runtime/Engine/Public/FXSystem.h b/Engine/Source/Runtime/Engine/Public/FXSystem.h
--- a/Engine/Source/Runtime/Engine/Public/FXSystem.h
+++ b/Engine/Source/Runtime/Engine/Public/FXSystem.h
@@ -120,7 +120,7 @@ DECLARE_DELEGATE_RetVal_TwoParams(FFXSystemInterface*, FCreateCustomFXSystemDele
 /**
  * The interface to an effects system.
  */
-class FFXSystemInterface
+class ENGINE_VTABLE FFXSystemInterface
 {
 public:
 
