Updated Material If expression to properly evaluate whether result is a material attribute

diff --git a/Engine/Source/Runtime/Engine/Private/Materials/MaterialExpressions.cpp b/Engine/Source/Runtime/Engine/Private/Materials/MaterialExpressions.cpp
--- a/Engine/Source/Runtime/Engine/Private/Materials/MaterialExpressions.cpp
+++ b/Engine/Source/Runtime/Engine/Private/Materials/MaterialExpressions.cpp
@@ -8545,8 +8545,10 @@ uint32 UMaterialExpressionIf::GetInputType(int32 InputIndex)
 
 bool UMaterialExpressionIf::IsResultMaterialAttributes(int32 OutputIndex)
 {
-	if ((A.GetTracedInput().Expression && !A.Expression->ContainsInputLoop() && A.Expression->IsResultMaterialAttributes(A.OutputIndex)) ||
-		(B.GetTracedInput().Expression && !B.Expression->ContainsInputLoop() && B.Expression->IsResultMaterialAttributes(B.OutputIndex)))
+	if ((AGreaterThanB.GetTracedInput().Expression && !AGreaterThanB.Expression->ContainsInputLoop() && AGreaterThanB.Expression->IsResultMaterialAttributes(A.OutputIndex))
+		&& (!AEqualsB.GetTracedInput().Expression || (!AEqualsB.Expression->ContainsInputLoop() && AEqualsB.Expression->IsResultMaterialAttributes(AEqualsB.OutputIndex)))
+		&& (ALessThanB.GetTracedInput().Expression && !ALessThanB.Expression->ContainsInputLoop() && ALessThanB.Expression->IsResultMaterialAttributes(ALessThanB.OutputIndex))
+		)
 	{
 		return true;
 	}
