Fixed conditional for material If node validation for material attributes

#rnx

diff --git a/Engine/Source/Runtime/Engine/Private/Materials/MaterialExpressions.cpp b/Engine/Source/Runtime/Engine/Private/Materials/MaterialExpressions.cpp
--- a/Engine/Source/Runtime/Engine/Private/Materials/MaterialExpressions.cpp
+++ b/Engine/Source/Runtime/Engine/Private/Materials/MaterialExpressions.cpp
@@ -8545,7 +8545,7 @@ uint32 UMaterialExpressionIf::GetInputType(int32 InputIndex)
 
 bool UMaterialExpressionIf::IsResultMaterialAttributes(int32 OutputIndex)
 {
-	if ((AGreaterThanB.GetTracedInput().Expression && !AGreaterThanB.Expression->ContainsInputLoop() && AGreaterThanB.Expression->IsResultMaterialAttributes(A.OutputIndex))
+	if ((AGreaterThanB.GetTracedInput().Expression && !AGreaterThanB.Expression->ContainsInputLoop() && AGreaterThanB.Expression->IsResultMaterialAttributes(AGreaterThanB.OutputIndex))
 		&& (!AEqualsB.GetTracedInput().Expression || (!AEqualsB.Expression->ContainsInputLoop() && AEqualsB.Expression->IsResultMaterialAttributes(AEqualsB.OutputIndex)))
 		&& (ALessThanB.GetTracedInput().Expression && !ALessThanB.Expression->ContainsInputLoop() && ALessThanB.Expression->IsResultMaterialAttributes(ALessThanB.OutputIndex))
 		)
