Added support for per world subsystems

diff --git a/Engine/Source/Editor/BlueprintGraph/Private/K2Node_GetSubsystem.cpp b/Engine/Source/Editor/BlueprintGraph/Private/K2Node_GetSubsystem.cpp
--- a/Engine/Source/Editor/BlueprintGraph/Private/K2Node_GetSubsystem.cpp
+++ b/Engine/Source/Editor/BlueprintGraph/Private/K2Node_GetSubsystem.cpp
@@ -131,6 +131,10 @@ void UK2Node_GetSubsystem::ExpandNode(class FKismetCompilerContext& CompilerCont
 	{
 		Get_FunctionName = GET_FUNCTION_NAME_CHECKED(USubsystemBlueprintLibrary, GetGameInstanceSubsystem);
 	}
+	else if (CustomClass->IsChildOf<UWorldSubsystem>())
+	{
+		Get_FunctionName = GET_FUNCTION_NAME_CHECKED(USubsystemBlueprintLibrary, GetWorldSubsystem);
+	}
 	else if (CustomClass->IsChildOf<ULocalPlayerSubsystem>())
 	{
 		Get_FunctionName = GET_FUNCTION_NAME_CHECKED(USubsystemBlueprintLibrary, GetLocalPlayerSubsystem);
@@ -208,6 +212,7 @@ void UK2Node_GetSubsystem::GetMenuActions(FBlueprintActionDatabaseRegistrar& Act
 	static TArray<UClass*> Subclasses;
 	Subclasses.Reset();
 	GetDerivedClasses(UGameInstanceSubsystem::StaticClass(), Subclasses);
+	GetDerivedClasses(UWorldSubsystem::StaticClass(), Subclasses);
 	GetDerivedClasses(ULocalPlayerSubsystem::StaticClass(), Subclasses);
 
 	auto CustomizeCallback = [](UEdGraphNode* Node, bool bIsTemplateNode, UClass* Subclass)
@@ -240,6 +245,10 @@ FText UK2Node_GetSubsystem::GetMenuCategory() const
 	{
 		return NSLOCTEXT("K2Node", "GetSubsystem_LocalPlayerSubsystemsMenuCategory", "LocalPlayer Subsystems");
 	}
+	else if (CustomClass->IsChildOf<UWorldSubsystem>())
+	{
+		return NSLOCTEXT("K2Node", "GetSubsystem_WorldSubsystemsMenuCategory", "World Subsystems");
+	}
 
 	return NSLOCTEXT("K2Node", "GetSubsystem_InvalidSubsystemTypeMenuCategory", "Invalid Subsystem Type");
 }
@@ -253,6 +262,10 @@ FText UK2Node_GetSubsystem::GetTooltipText() const
 		{
 			SubsystemTypeText =  NSLOCTEXT("K2Node", "GetSubsystem_GameInstanceSubsystemTooltip", "GameInstance Subsystem");
 		}
+		else if (CustomClass->IsChildOf<UWorldSubsystem>())
+		{
+			SubsystemTypeText = NSLOCTEXT("K2Node", "GetSubsystem_WorldSubsystemTooltip", "World Subsystem");
+		}
 		else
 		{
 			SubsystemTypeText = NSLOCTEXT("K2Node", "GetSubsystem_LocalPlayerSubsystemTooltip", "LocalPlayer Subsystem");
diff --git a/Engine/Source/Runtime/Engine/Classes/Engine/World.h b/Engine/Source/Runtime/Engine/Classes/Engine/World.h
--- a/Engine/Source/Runtime/Engine/Classes/Engine/World.h
+++ b/Engine/Source/Runtime/Engine/Classes/Engine/World.h
@@ -23,6 +23,9 @@
 #include "Particles/WorldPSCPool.h"
 #include "Containers/SortedMap.h"
 
+#include "Subsystems/WorldSubsystem.h"
+#include "Subsystems/SubsystemCollection.h"
+
 #include "World.generated.h"
 
 class ABrush;
@@ -3446,6 +3449,50 @@ public:
 		return (OwningGameInstance ? OwningGameInstance->GetLatentActionManager() : LatentActionManager);
 	}
 
+	/**
+	 * Get a Subsystem of specified type
+	 */
+	UWorldSubsystem* GetSubsystemBase(TSubclassOf<UWorldSubsystem> SubsystemClass) const
+	{
+		return SubsystemCollection.GetSubsystem<UWorldSubsystem>(SubsystemClass);
+	}
+
+	/**
+	 * Get a Subsystem of specified type
+	 */
+	template <typename TSubsystemClass>
+	TSubsystemClass* GetSubsystem() const
+	{
+		return SubsystemCollection.GetSubsystem<TSubsystemClass>(TSubsystemClass::StaticClass());
+	}
+
+	/**
+	 * Get a Subsystem of specified type from the provided GameInstance
+	 * returns nullptr if the Subsystem cannot be found or the GameInstance is null
+	 */
+	template <typename TSubsystemClass>
+	static FORCEINLINE TSubsystemClass* GetSubsystem(const UWorld* World)
+	{
+		if (World)
+		{
+			return World->GetSubsystem<TSubsystemClass>();
+		}
+		return nullptr;
+	}
+
+	/**
+	 * Get all Subsystem of specified type, this is only necessary for interfaces that can have multiple implementations instanced at a time.
+	 *
+	 * Do not hold onto this Array reference unless you are sure the lifetime is less than that of UGameInstance
+	 */
+	template <typename TSubsystemClass>
+	const TArray<TSubsystemClass*>& GetSubsystemArray() const
+	{
+		return SubsystemCollection.GetSubsystemArray<TSubsystemClass>(TSubsystemClass::StaticClass());
+	}
+
+
+
 	/** Sets the owning game instance for this world */
 	inline void SetGameInstance(UGameInstance* NewGI)
 	{
@@ -3528,6 +3575,7 @@ public:
 	FWorldPSCPool PSCPool;
 
 	//PSC Pooling END
+	FSubsystemCollection<UWorldSubsystem> SubsystemCollection;
 };
 
 /** Global UWorld pointer. Use of this pointer should be avoided whenever possible. */
diff --git a/Engine/Source/Runtime/Engine/Private/Subsystems/SubsystemBlueprintLibrary.cpp b/Engine/Source/Runtime/Engine/Private/Subsystems/SubsystemBlueprintLibrary.cpp
--- a/Engine/Source/Runtime/Engine/Private/Subsystems/SubsystemBlueprintLibrary.cpp
+++ b/Engine/Source/Runtime/Engine/Private/Subsystems/SubsystemBlueprintLibrary.cpp
@@ -49,6 +49,15 @@
 	return nullptr;
 }
 
+UWorldSubsystem* USubsystemBlueprintLibrary::GetWorldSubsystem(UObject* ContextObject, TSubclassOf<UWorldSubsystem> Class)
+{
+	if (const UWorld* World = ThisClass::GetWorldFrom(ContextObject))
+	{
+		return World->GetSubsystemBase(Class);
+	}
+	return nullptr;
+}
+
 /*static*/ ULocalPlayerSubsystem* USubsystemBlueprintLibrary::GetLocalPlayerSubSystemFromPlayerController(APlayerController* PlayerController, TSubclassOf<ULocalPlayerSubsystem> Class)
 {
 	if (PlayerController)
diff --git a/Engine/Source/Runtime/Engine/Private/Subsystems/WorldSubsystem.cpp b/Engine/Source/Runtime/Engine/Private/Subsystems/WorldSubsystem.cpp
new file mode 100644
--- /dev/null
+++ b/Engine/Source/Runtime/Engine/Private/Subsystems/WorldSubsystem.cpp
@@ -0,0 +1,15 @@
+// Copyright 1998-2019 Epic Games, Inc. All Rights Reserved.
+
+#include "Subsystems/WorldSubsystem.h"
+#include "Engine/World.h"
+
+UWorldSubsystem::UWorldSubsystem()
+	: USubsystem()
+{
+
+}
+
+UWorld* UWorldSubsystem::GetWorld() const
+{
+	return Cast<UWorld>(GetOuter());
+}
\ No newline at end of file
diff --git a/Engine/Source/Runtime/Engine/Private/World.cpp b/Engine/Source/Runtime/Engine/Private/World.cpp
--- a/Engine/Source/Runtime/Engine/Private/World.cpp
+++ b/Engine/Source/Runtime/Engine/Private/World.cpp
@@ -352,6 +352,7 @@ UWorld::UWorld( const FObjectInitializer& ObjectInitializer )
 ,	TickTaskLevel(FTickTaskManagerInterface::Get().AllocateTickTaskLevel())
 ,	FlushLevelStreamingType(EFlushLevelStreamingType::None)
 ,	NextTravelType(TRAVEL_Relative)
+,	SubsystemCollection(this)
 {
 	TimerManager = new FTimerManager();
 #if WITH_EDITOR
@@ -1324,6 +1325,11 @@ void UWorld::InitWorld(const InitializationValues IVS)
 		return;
 	}
 
+	if (WorldType == EWorldType::Game || WorldType == EWorldType::Editor || WorldType == EWorldType::PIE)
+	{
+		SubsystemCollection.Initialize();
+	}
+
 	FWorldDelegates::OnPreWorldInitialization.Broadcast(this, IVS);
 
 	AWorldSettings* WorldSettings = GetWorldSettings();
@@ -4269,6 +4275,8 @@ void UWorld::CleanupWorld(bool bSessionEnded, bool bCleanupResources, UWorld* Ne
 		check(WorldToResetCleanedUpFlag->bCleanedUpWorld);
 		WorldToResetCleanedUpFlag->bCleanedUpWorld = false;
 	}
+
+	SubsystemCollection.Deinitialize();
 }
 
 UGameViewportClient* UWorld::GetGameViewport() const
diff --git a/Engine/Source/Runtime/Engine/Public/Subsystems/Subsystem.h b/Engine/Source/Runtime/Engine/Public/Subsystems/Subsystem.h
--- a/Engine/Source/Runtime/Engine/Public/Subsystems/Subsystem.h
+++ b/Engine/Source/Runtime/Engine/Public/Subsystems/Subsystem.h
@@ -15,6 +15,7 @@ class FSubsystemCollectionBase;
  *		Engine		 -> inherit UEngineSubsystem
  *		Editor		 -> inherit UEditorSubsystem
  *		GameInstance -> inherit UGameInstanceSubsystem
+ *		World		 -> inherit UWorldSubsystem
  *		LocalPlayer	 -> inherit ULocalPlayerSubsystem
  *
  *
diff --git a/Engine/Source/Runtime/Engine/Public/Subsystems/SubsystemBlueprintLibrary.h b/Engine/Source/Runtime/Engine/Public/Subsystems/SubsystemBlueprintLibrary.h
--- a/Engine/Source/Runtime/Engine/Public/Subsystems/SubsystemBlueprintLibrary.h
+++ b/Engine/Source/Runtime/Engine/Public/Subsystems/SubsystemBlueprintLibrary.h
@@ -29,6 +29,10 @@ public:
 	UFUNCTION(BlueprintPure, Category = "LocalPlayer Subsystems", meta = (WorldContext = "ContextObject", BlueprintInternalUseOnly = "true"))
 	static ULocalPlayerSubsystem* GetLocalPlayerSubsystem(UObject* ContextObject, TSubclassOf<ULocalPlayerSubsystem> Class);
 
+	/** Get a World Subsystem from the World associated with the provided context */
+	UFUNCTION(BlueprintPure, Category = "GameInstance Subsystems", meta = (WorldContext = "ContextObject", BlueprintInternalUseOnly = "true"))
+	static UWorldSubsystem* GetWorldSubsystem(UObject* ContextObject, TSubclassOf<UWorldSubsystem> Class);
+
 	/** 
 	 * Get a Local Player Subsystem from the LocalPlayer associated with the provided context
 	 * If the player controller isn't associated to a LocalPlayer nullptr is returned
diff --git a/Engine/Source/Runtime/Engine/Public/Subsystems/WorldSubsystem.h b/Engine/Source/Runtime/Engine/Public/Subsystems/WorldSubsystem.h
new file mode 100644
--- /dev/null
+++ b/Engine/Source/Runtime/Engine/Public/Subsystems/WorldSubsystem.h
@@ -0,0 +1,24 @@
+// Copyright 1998-2019 Epic Games, Inc. All Rights Reserved.
+
+#pragma once
+
+#include "Subsystem.h"
+
+#include "WorldSubsystem.generated.h"
+
+/**
+ * UWorldSubsystem
+ * Base class for auto instanced and initialized systems that share the lifetime of a UWorld
+
+ */
+
+UCLASS(Abstract)
+class ENGINE_API UWorldSubsystem : public USubsystem
+{
+	GENERATED_BODY()
+
+public:
+	UWorldSubsystem();
+
+	virtual UWorld* GetWorld() const override;
+};
