Fixed -nan issue with view uniform buffer due to render target context not being allocated for the main view.

diff --git a/Engine/Source/Runtime/Engine/Private/SceneView.cpp b/Engine/Source/Runtime/Engine/Private/SceneView.cpp
--- a/Engine/Source/Runtime/Engine/Private/SceneView.cpp
+++ b/Engine/Source/Runtime/Engine/Private/SceneView.cpp
@@ -2144,31 +2144,38 @@ void FSceneView::SetupViewRectUniformBufferParameters(FViewUniformShaderParamete
 {
 	checkfSlow(EffectiveViewRect.Area() > 0, TEXT("Invalid-size EffectiveViewRect passed to CreateUniformBufferParameters [%d * %d]."), EffectiveViewRect.Width(), EffectiveViewRect.Height());
 
-	// Calculate the vector used by shaders to convert clip space coordinates to texture space.
-	const float InvBufferSizeX = 1.0f / BufferSize.X;
-	const float InvBufferSizeY = 1.0f / BufferSize.Y;
-	// to bring NDC (-1..1, 1..-1) into 0..1 UV for BufferSize textures
-	const FVector4 ScreenPositionScaleBias(
-		EffectiveViewRect.Width() * InvBufferSizeX / +2.0f,
-		EffectiveViewRect.Height() * InvBufferSizeY / (-2.0f * GProjectionSignY),
-		(EffectiveViewRect.Height() / 2.0f + EffectiveViewRect.Min.Y) * InvBufferSizeY,
-		(EffectiveViewRect.Width() / 2.0f + EffectiveViewRect.Min.X) * InvBufferSizeX
-		);
-
-	ViewUniformShaderParameters.ScreenPositionScaleBias = ScreenPositionScaleBias;
-
 	ViewUniformShaderParameters.ViewRectMin = FVector4(EffectiveViewRect.Min.X, EffectiveViewRect.Min.Y, 0.0f, 0.0f);
 	ViewUniformShaderParameters.ViewSizeAndInvSize = FVector4(EffectiveViewRect.Width(), EffectiveViewRect.Height(), 1.0f / float(EffectiveViewRect.Width()), 1.0f / float(EffectiveViewRect.Height()));
-	ViewUniformShaderParameters.BufferSizeAndInvSize = FVector4(BufferSize.X, BufferSize.Y, InvBufferSizeX, InvBufferSizeY);
-	ViewUniformShaderParameters.BufferBilinearUVMinMax = FVector4(
-		InvBufferSizeX * (EffectiveViewRect.Min.X + 0.5),
-		InvBufferSizeY * (EffectiveViewRect.Min.Y + 0.5),
-		InvBufferSizeX * (EffectiveViewRect.Max.X - 0.5),
-		InvBufferSizeY * (EffectiveViewRect.Max.Y - 0.5));
-
-	/* Texture Level-of-Detail Strategies for Real-Time Ray Tracing https://developer.nvidia.com/raytracinggems Equation 20 */
-	float RadFOV = (PI / 180.0f) * FOV;
-	ViewUniformShaderParameters.EyeToPixelSpreadAngle = FPlatformMath::Atan((2.0f * FPlatformMath::Tan(RadFOV * 0.5f)) / BufferSize.Y);
+
+	// A buffer size of zero means that we don't have a scene context allocated. This can happen when rendering to a one-off render target (e.g. canvas) prior to
+	// main view rendering. In that case we just skip the scene color target specific parameters, since they wouldn't be relevant to those views anyway.
+	if (BufferSize != FIntPoint::ZeroValue)
+	{
+		ensureMsgf((BufferSize.X > 0 && BufferSize.Y > 0), TEXT("Invalid-size BufferSize passed to CreateUniformBufferParameters [%d * %d]."), BufferSize.X, BufferSize.Y);
+
+		// Calculate the vector used by shaders to convert clip space coordinates to texture space.
+		const float InvBufferSizeX = 1.0f / BufferSize.X;
+		const float InvBufferSizeY = 1.0f / BufferSize.Y;
+		// to bring NDC (-1..1, 1..-1) into 0..1 UV for BufferSize textures
+		const FVector4 ScreenPositionScaleBias(
+			EffectiveViewRect.Width() * InvBufferSizeX / +2.0f,
+			EffectiveViewRect.Height() * InvBufferSizeY / (-2.0f * GProjectionSignY),
+			(EffectiveViewRect.Height() / 2.0f + EffectiveViewRect.Min.Y) * InvBufferSizeY,
+			(EffectiveViewRect.Width() / 2.0f + EffectiveViewRect.Min.X) * InvBufferSizeX);
+
+		ViewUniformShaderParameters.ScreenPositionScaleBias = ScreenPositionScaleBias;
+
+		ViewUniformShaderParameters.BufferSizeAndInvSize = FVector4(BufferSize.X, BufferSize.Y, InvBufferSizeX, InvBufferSizeY);
+		ViewUniformShaderParameters.BufferBilinearUVMinMax = FVector4(
+			InvBufferSizeX * (EffectiveViewRect.Min.X + 0.5),
+			InvBufferSizeY * (EffectiveViewRect.Min.Y + 0.5),
+			InvBufferSizeX * (EffectiveViewRect.Max.X - 0.5),
+			InvBufferSizeY * (EffectiveViewRect.Max.Y - 0.5));
+
+		/* Texture Level-of-Detail Strategies for Real-Time Ray Tracing https://developer.nvidia.com/raytracinggems Equation 20 */
+		float RadFOV = (PI / 180.0f) * FOV;
+		ViewUniformShaderParameters.EyeToPixelSpreadAngle = FPlatformMath::Atan((2.0f * FPlatformMath::Tan(RadFOV * 0.5f)) / BufferSize.Y);
+	}
 
 	ViewUniformShaderParameters.MotionBlurNormalizedToPixel = FinalPostProcessSettings.MotionBlurMax * EffectiveViewRect.Width() / 100.0f;
 
