#UE4 Added some non-shipping logging for bad body instances due to a zero scale. An all-zero scale is already changed to a value thats nearly zero, but this warning covers non-uniform scales that have a zero component

diff --git a/Engine/Source/Runtime/Engine/Private/Components/PrimitiveComponent.cpp b/Engine/Source/Runtime/Engine/Private/Components/PrimitiveComponent.cpp
--- a/Engine/Source/Runtime/Engine/Private/Components/PrimitiveComponent.cpp
+++ b/Engine/Source/Runtime/Engine/Private/Components/PrimitiveComponent.cpp
@@ -739,6 +739,12 @@ void UPrimitiveComponent::OnCreatePhysicsState()
 			{
 				BodyTransform.SetScale3D(FVector(KINDA_SMALL_NUMBER));
 			}
+#if !(UE_BUILD_SHIPPING || UE_BUILD_TEST)
+			if (FMath::IsNearlyZero(BodyScale.X) || FMath::IsNearlyZero(BodyScale.Y) || FMath::IsNearlyZero(BodyScale.Z))
+			{
+				UE_LOG(LogPhysics, Warning, TEXT("Scale for %s has a component set to zero, which will result in a bad body instance. Scale:%s"), *GetPathNameSafe(this), *BodyScale.ToString());
+			}
+#endif
 
 			// Create the body.
 			BodyInstance.InitBody(BodySetup, BodyTransform, this, GetWorld()->GetPhysicsScene());		
