Fixed 'Foliage' Engine Show Flag ignored by Actors painted by [REDACTED_PERSON]

diff --git a/Engine/Source/Runtime/Engine/Engine.Build.cs b/Engine/Source/Runtime/Engine/Engine.Build.cs
--- a/Engine/Source/Runtime/Engine/Engine.Build.cs
+++ b/Engine/Source/Runtime/Engine/Engine.Build.cs
@@ -36,7 +36,7 @@ public class Engine : ModuleRules
 				"AutomationWorker",
 				"MovieSceneCapture",
 				"DesktopPlatform",
-			}
+            }
 		);
 
 		if (Target.Configuration != UnrealTargetConfiguration.Shipping)
@@ -187,11 +187,18 @@ public class Engine : ModuleRules
 
 		CircularlyReferencedDependentModules.Add("GameplayTags");
 		CircularlyReferencedDependentModules.Add("Landscape");
-		CircularlyReferencedDependentModules.Add("UMG");
+        CircularlyReferencedDependentModules.Add("UMG");
 		CircularlyReferencedDependentModules.Add("MaterialShaderQualitySettings");
 		CircularlyReferencedDependentModules.Add("CinematicCamera");
 		CircularlyReferencedDependentModules.Add("AudioMixer");
 
+		if (Target.Type == TargetType.Editor)
+		{
+			PrivateIncludePathModuleNames.Add("Foliage");
+			PrivateDependencyModuleNames.Add("Foliage");
+			CircularlyReferencedDependentModules.Add("Foliage");
+		}
+
 		// The AnimGraphRuntime module is not needed by Engine proper, but it is loaded in LaunchEngineLoop.cpp,
 		// and needs to be listed in an always-included module in order to be compiled into standalone games
 		DynamicallyLoadedModuleNames.Add("AnimGraphRuntime");
diff --git a/Engine/Source/Runtime/Engine/Private/PrimitiveSceneProxy.cpp b/Engine/Source/Runtime/Engine/Private/PrimitiveSceneProxy.cpp
--- a/Engine/Source/Runtime/Engine/Private/PrimitiveSceneProxy.cpp
+++ b/Engine/Source/Runtime/Engine/Private/PrimitiveSceneProxy.cpp
@@ -15,6 +15,9 @@
 #include "Materials/Material.h"
 #include "SceneManagement.h"
 #include "VT/RuntimeVirtualTexture.h"
+#if WITH_EDITOR
+#include "InstancedFoliageActor.h"
+#endif
 
 static TAutoConsoleVariable<int32> CVarForceSingleSampleShadowingFromStationary(
 	TEXT("r.Shadow.ForceSingleSampleShadowingFromStationary"),
@@ -160,6 +163,7 @@ FPrimitiveSceneProxy::FPrimitiveSceneProxy(const UPrimitiveComponent* InComponen
 // by default we are always drawn
 ,	HiddenEditorViews(0)
 ,	DrawInAnyEditMode(0)
+,   bIsFoliage(false)
 #endif
 ,	VisibilityId(InComponent->VisibilityId)
 ,	MaxDrawDistance(InComponent->CachedMaxDrawDistance > 0 ? InComponent->CachedMaxDrawDistance : FLT_MAX)
@@ -209,6 +213,7 @@ FPrimitiveSceneProxy::FPrimitiveSceneProxy(const UPrimitiveComponent* InComponen
 		// cache the actor's group membership
 		HiddenEditorViews = InComponent->GetHiddenEditorViews();
 		DrawInAnyEditMode = InComponent->GetOwner()->IsEditorOnly();
+		bIsFoliage = AInstancedFoliageActor::IsOwnedByFoliage(InComponent->GetOwner());
 #endif
 	}
 	
@@ -619,6 +624,11 @@ bool FPrimitiveSceneProxy::IsShown(const FSceneView* View) const
 		}
 	}
 
+	if (bIsFoliage && !View->Family->EngineShowFlags.InstancedFoliage)
+	{
+		return false;
+	}
+
 	// After checking for VR/Desktop Edit mode specific actors, check for Editor vs. Game
 	if(View->Family->EngineShowFlags.Editor)
 	{
diff --git a/Engine/Source/Runtime/Engine/Public/PrimitiveSceneProxy.h b/Engine/Source/Runtime/Engine/Public/PrimitiveSceneProxy.h
--- a/Engine/Source/Runtime/Engine/Public/PrimitiveSceneProxy.h
+++ b/Engine/Source/Runtime/Engine/Public/PrimitiveSceneProxy.h
@@ -1042,6 +1042,8 @@ private:
 
 	/** Whether this should only draw in any editing mode*/
 	uint32 DrawInAnyEditMode : 1;
+
+	uint32 bIsFoliage : 1;
 #endif
 
 	/** Used for precomputed visibility */
