Fix crash on undo in physics asset editor after clicking 'Generate All Bodies' button

diff --git a/Engine/Source/Runtime/Engine/Private/WorldSettings.cpp b/Engine/Source/Runtime/Engine/Private/WorldSettings.cpp
--- a/Engine/Source/Runtime/Engine/Private/WorldSettings.cpp
+++ b/Engine/Source/Runtime/Engine/Private/WorldSettings.cpp
@@ -661,7 +661,7 @@ void AWorldSettings::InternalPostPropertyChanged(FName PropertyName)
 			UpdateBookmarkClass();
 		}
 
-	if (GetWorld() != nullptr && GetWorld()->PersistentLevel->GetWorldSettings() == this)
+	if (GetWorld() != nullptr && GetWorld()->PersistentLevel && GetWorld()->PersistentLevel->GetWorldSettings() == this)
 	{
 		if (PropertyName == GET_MEMBER_NAME_CHECKED(FHierarchicalSimplification, TransitionScreenSize))
 		{
