Fixed editor crash when reading cooked static mesh

#rnx

diff --git a/Engine/Source/Runtime/Engine/Private/StaticMesh.cpp b/Engine/Source/Runtime/Engine/Private/StaticMesh.cpp
--- a/Engine/Source/Runtime/Engine/Private/StaticMesh.cpp
+++ b/Engine/Source/Runtime/Engine/Private/StaticMesh.cpp
@@ -1362,16 +1362,13 @@ void FStaticMeshRenderData::Serialize(FArchive& Ar, UStaticMesh* Owner, bool bCo
 				if (bValid)
 				{
 #if WITH_EDITOR
-					check(LOD.DistanceFieldData != nullptr);
-
-					bool bDownSampling = Ar.IsCooking() && Ar.IsSaving();
-					
-					if (bDownSampling)
+					if (Ar.IsCooking() && Ar.IsSaving())
 					{
+						check(LOD.DistanceFieldData != nullptr);
+
 						float Divider = Ar.CookingTarget()->GetDownSampleMeshDistanceFieldDivider();
-						bDownSampling = Divider > 1;
 
-						if (bDownSampling)
+						if (Divider > 1)
 						{
 							FDistanceFieldVolumeData DownSampledDFVolumeData = *LOD.DistanceFieldData;
 							IMeshUtilities& MeshUtilities = FModuleManager::Get().LoadModuleChecked<IMeshUtilities>(TEXT("MeshUtilities"));
@@ -1380,20 +1377,21 @@ void FStaticMeshRenderData::Serialize(FArchive& Ar, UStaticMesh* Owner, bool bCo
 
 							Ar << DownSampledDFVolumeData;
 						}
+						else
+						{
+							Ar << *(LOD.DistanceFieldData);
+						}
 					}
-
-					if (!bDownSampling)
+					else
+#endif
 					{
+						if (LOD.DistanceFieldData == nullptr)
+						{
+							LOD.DistanceFieldData = new FDistanceFieldVolumeData();
+						}
+
 						Ar << *(LOD.DistanceFieldData);
 					}
-#else
-					if (LOD.DistanceFieldData == nullptr)
-					{
-						LOD.DistanceFieldData = new FDistanceFieldVolumeData();
-					}
-
-					Ar << *(LOD.DistanceFieldData);
-#endif
 				}
 			}
 		}
