#UE4 Add more log info to a warning related to render targets

diff --git a/Engine/Source/Runtime/Engine/Private/KismetRenderingLibrary.cpp b/Engine/Source/Runtime/Engine/Private/KismetRenderingLibrary.cpp
--- a/Engine/Source/Runtime/Engine/Private/KismetRenderingLibrary.cpp
+++ b/Engine/Source/Runtime/Engine/Private/KismetRenderingLibrary.cpp
@@ -169,19 +169,19 @@ void UKismetRenderingLibrary::ExportRenderTarget(UObject* WorldContextObject, UT
 
 	if (!TextureRenderTarget)
 	{
-		FMessageLog("Blueprint").Warning(LOCTEXT("ExportRenderTarget_InvalidTextureRenderTarget", "ExportRenderTarget: TextureRenderTarget must be non-null."));
+		FMessageLog("Blueprint").Warning(FText::Format(LOCTEXT("ExportRenderTarget_InvalidTextureRenderTarget", "ExportRenderTarget[{0}]: TextureRenderTarget must be non-null."), FText::FromString(GetPathNameSafe(WorldContextObject))));
 	}
 	else if (!TextureRenderTarget->Resource)
 	{
-		FMessageLog("Blueprint").Warning(LOCTEXT("ExportRenderTarget_ReleasedTextureRenderTarget", "ExportRenderTarget: render target has been released."));
+		FMessageLog("Blueprint").Warning(FText::Format(LOCTEXT("ExportRenderTarget_ReleasedTextureRenderTarget", "ExportRenderTarget[{0}]: render target has been released."), FText::FromString(GetPathNameSafe(WorldContextObject))));
 	}
 	else if (!PathError.IsEmpty())
 	{
-		FMessageLog("Blueprint").Warning(FText::Format(LOCTEXT("ExportRenderTarget_InvalidFilePath", "ExportRenderTarget: Invalid file path provided: '{0}'"), PathError));
+		FMessageLog("Blueprint").Warning(FText::Format(LOCTEXT("ExportRenderTarget_InvalidFilePath", "ExportRenderTarget[{0}]: Invalid file path provided: '{1}'"), FText::FromString(GetPathNameSafe(WorldContextObject)), PathError));
 	}
 	else if (FileName.IsEmpty())
 	{
-		FMessageLog("Blueprint").Warning(LOCTEXT("ExportRenderTarget_InvalidFileName", "ExportRenderTarget: FileName must be non-empty."));
+		FMessageLog("Blueprint").Warning(FText::Format(LOCTEXT("ExportRenderTarget_InvalidFileName", "ExportRenderTarget[{0}]: FileName must be non-empty."), FText::FromString(GetPathNameSafe(WorldContextObject))));
 	}
 	else
 	{
@@ -449,15 +449,15 @@ void UKismetRenderingLibrary::ConvertRenderTargetToTexture2DEditorOnly( UObject*
 #if WITH_EDITOR
 	if (!RenderTarget)
 	{
-		FMessageLog("Blueprint").Warning(LOCTEXT("ConvertRenderTargetToTexture2D_InvalidRenderTarget", "ConvertRenderTargetToTexture2DEditorOnly: RenderTarget must be non-null."));
+		FMessageLog("Blueprint").Warning(FText::Format(LOCTEXT("ConvertRenderTargetToTexture2D_InvalidRenderTarget", "ConvertRenderTargetToTexture2DEditorOnly[{0}]: RenderTarget must be non-null."), FText::FromString(GetPathNameSafe(WorldContextObject))));
 	}
 	else if (!RenderTarget->Resource)
 	{
-		FMessageLog("Blueprint").Warning(LOCTEXT("ConvertRenderTargetToTexture2D_ReleasedTextureRenderTarget", "ConvertRenderTargetToTexture2DEditorOnly: render target has been released."));
+		FMessageLog("Blueprint").Warning(FText::Format(LOCTEXT("ConvertRenderTargetToTexture2D_ReleasedTextureRenderTarget", "ConvertRenderTargetToTexture2DEditorOnly[{0}]: render target has been released."), FText::FromString(GetPathNameSafe(WorldContextObject))));
 	}
 	else if (!Texture)
 	{
-		FMessageLog("Blueprint").Warning(LOCTEXT("ConvertRenderTargetToTexture2D_InvalidTexture", "ConvertRenderTargetToTexture2DEditorOnly: Texture must be non-null."));
+		FMessageLog("Blueprint").Warning(FText::Format(LOCTEXT("ConvertRenderTargetToTexture2D_InvalidTexture", "ConvertRenderTargetToTexture2DEditorOnly[{0}]: Texture must be non-null."), FText::FromString(GetPathNameSafe(WorldContextObject))));
 	}
 	else
 	{
@@ -542,11 +542,11 @@ void UKismetRenderingLibrary::BeginDrawCanvasToRenderTarget(UObject* WorldContex
 	}
 	else if (!TextureRenderTarget)
 	{
-		FMessageLog("Blueprint").Warning(LOCTEXT("BeginDrawCanvasToRenderTarget_InvalidTextureRenderTarget", "BeginDrawCanvasToRenderTarget: TextureRenderTarget must be non-null."));
+		FMessageLog("Blueprint").Warning(FText::Format(LOCTEXT("BeginDrawCanvasToRenderTarget_InvalidTextureRenderTarget", "BeginDrawCanvasToRenderTarget[{0}]: TextureRenderTarget must be non-null."), FText::FromString(GetPathNameSafe(WorldContextObject))));
 	}
 	else if (!TextureRenderTarget->Resource)
 	{
-		FMessageLog("Blueprint").Warning(LOCTEXT("BeginDrawCanvasToRenderTarget_ReleasedTextureRenderTarget", "BeginDrawCanvasToRenderTarget: render target has been released."));
+		FMessageLog("Blueprint").Warning(FText::Format(LOCTEXT("BeginDrawCanvasToRenderTarget_ReleasedTextureRenderTarget", "BeginDrawCanvasToRenderTarget[{0}]: render target has been released."), FText::FromString(GetPathNameSafe(WorldContextObject))));
 	}
 	else
 	{
