Add an option to force mips resident ASAP when visible. Behavior of bIgnoreStreamingMipBias is changed to only affect forced fully resident textures.

diff --git a/Engine/Source/Runtime/Engine/Classes/Engine/StreamableRenderAsset.h b/Engine/Source/Runtime/Engine/Classes/Engine/StreamableRenderAsset.h
--- a/Engine/Source/Runtime/Engine/Classes/Engine/StreamableRenderAsset.h
+++ b/Engine/Source/Runtime/Engine/Classes/Engine/StreamableRenderAsset.h
@@ -254,7 +254,7 @@ public:
 	UPROPERTY(transient)
 	uint8 bForceMiplevelsToBeResident : 1;
 
-	/** Ignores the streaming mip bias used to accommodate memory constraints. */
+	/** When forced fully resident, ignores the streaming mip bias used to accommodate memory constraints. */
 	UPROPERTY(transient)
 	uint8 bIgnoreStreamingMipBias : 1;
 
diff --git a/Engine/Source/Runtime/Engine/Classes/Materials/MaterialInterface.h b/Engine/Source/Runtime/Engine/Classes/Materials/MaterialInterface.h
--- a/Engine/Source/Runtime/Engine/Classes/Materials/MaterialInterface.h
+++ b/Engine/Source/Runtime/Engine/Classes/Materials/MaterialInterface.h
@@ -746,9 +746,10 @@ public:
 	 * @param bForceMiplevelsToBeResidentValue		- true forces all mips to stream in. false lets other factors decide what to do with the mips.
 	 * @param ForceDuration							- Number of seconds to keep all mip-levels in memory, disregarding the normal priority logic. Negative value turns it off.
 	 * @param CinematicTextureGroups				- Bitfield indicating texture groups that should use extra high-resolution mips
+	 * @param bFastResponse							- USE WITH EXTREME CAUTION! Fast response textures incur sizable GT overhead and disturb streaming metric calculation. Avoid whenever possible.
 	 */
 	UFUNCTION(BlueprintCallable, Category = "Rendering|Material")
-	ENGINE_API virtual void SetForceMipLevelsToBeResident( bool OverrideForceMiplevelsToBeResident, bool bForceMiplevelsToBeResidentValue, float ForceDuration, int32 CinematicTextureGroups = 0 );
+	ENGINE_API virtual void SetForceMipLevelsToBeResident( bool OverrideForceMiplevelsToBeResident, bool bForceMiplevelsToBeResidentValue, float ForceDuration, int32 CinematicTextureGroups = 0, bool bFastResponse = false );
 
 	/**
 	 * Re-caches uniform expressions for all material interfaces
diff --git a/Engine/Source/Runtime/Engine/Private/Materials/MaterialInterface.cpp b/Engine/Source/Runtime/Engine/Private/Materials/MaterialInterface.cpp
--- a/Engine/Source/Runtime/Engine/Private/Materials/MaterialInterface.cpp
+++ b/Engine/Source/Runtime/Engine/Private/Materials/MaterialInterface.cpp
@@ -195,7 +195,7 @@ int32 UMaterialInterface::GetHeight() const
 }
 
 
-void UMaterialInterface::SetForceMipLevelsToBeResident( bool OverrideForceMiplevelsToBeResident, bool bForceMiplevelsToBeResidentValue, float ForceDuration, int32 CinematicTextureGroups )
+void UMaterialInterface::SetForceMipLevelsToBeResident( bool OverrideForceMiplevelsToBeResident, bool bForceMiplevelsToBeResidentValue, float ForceDuration, int32 CinematicTextureGroups, bool bFastResponse )
 {
 	TArray<UTexture*> Textures;
 	
@@ -210,6 +210,14 @@ void UMaterialInterface::SetForceMipLevelsToBeResident( bool OverrideForceMiplev
 			{
 				Texture->bForceMiplevelsToBeResident = bForceMiplevelsToBeResidentValue;
 			}
+
+			if (bFastResponse && (ForceDuration > 0.f || Texture->bForceMiplevelsToBeResident))
+			{
+				static IConsoleVariable* CVarAllowFastForceResident = IConsoleManager::Get().FindConsoleVariable(TEXT("r.Streaming.AllowFastForceResident"));
+
+				Texture->bIgnoreStreamingMipBias = CVarAllowFastForceResident && CVarAllowFastForceResident->GetInt();
+				IStreamingManager::Get().GetRenderAssetStreamingManager().FastForceFullyResident(Texture);
+			}
 		}
 	}
 }
diff --git a/Engine/Source/Runtime/Engine/Private/Streaming/StreamingManagerTexture.cpp b/Engine/Source/Runtime/Engine/Private/Streaming/StreamingManagerTexture.cpp
--- a/Engine/Source/Runtime/Engine/Private/Streaming/StreamingManagerTexture.cpp
+++ b/Engine/Source/Runtime/Engine/Private/Streaming/StreamingManagerTexture.cpp
@@ -35,6 +35,14 @@ static TAutoConsoleVariable<int32> CVarStreamingOverlapAssetAndLevelTicks(
 	TEXT("Ticks render asset streaming info on a high priority task thread while ticking levels on GT"),
 	ECVF_Default);
 
+static TAutoConsoleVariable<int32> CVarStreamingAllowFastForceResident(
+	TEXT("r.Streaming.AllowFastForceResident"),
+	0,
+	TEXT("Whether it is allowed to load in missing mips for fast-force-resident assets ASAP. ")
+	TEXT("Useful to accelerate force-resident process but risks disturbing streaming metric calculation. ")
+	TEXT("Fast-force-resident mips can't be sacrificed even when overbudget so use with caution."),
+	ECVF_Default);
+
 bool TrackRenderAsset( const FString& AssetName );
 bool UntrackRenderAsset( const FString& AssetName );
 void ListTrackedRenderAssets( FOutputDevice& Ar, int32 NumTextures );
@@ -69,6 +77,7 @@ FRenderAssetStreamingManager::FRenderAssetStreamingManager()
 ,	MaxEverRequired(0)
 ,	bPauseRenderAssetStreaming(false)
 ,	LastWorldUpdateTime(GIsEditor ? -FLT_MAX : 0) // In editor, visibility is not taken into consideration.
+,	LastWorldUpdateTime_MipCalcTask(LastWorldUpdateTime)
 {
 	// Read settings from ini file.
 	int32 TempInt;
@@ -335,6 +344,62 @@ void FRenderAssetStreamingManager::IncrementalUpdate(float Percentage, bool bUpd
 	SetRenderAssetsRemovedTimestamp(RemovedRenderAssets);
 }
 
+void FRenderAssetStreamingManager::TickFastResponseAssets()
+{
+	if (!CVarStreamingAllowFastForceResident.GetValueOnGameThread())
+	{
+		return;
+	}
+
+	for (TSet<UStreamableRenderAsset*>::TIterator It(FastResponseRenderAssets); It; ++It)
+	{
+		UStreamableRenderAsset* RenderAsset = *It;
+		
+		if (!RenderAsset
+			|| !RenderAsset->bIgnoreStreamingMipBias
+			|| (!RenderAsset->bForceMiplevelsToBeResident && RenderAsset->ForceMipLevelsToBeResidentTimestamp < FApp::GetCurrentTime()))
+		{
+			// No longer qualified for fast response
+			It.RemoveCurrent();
+			continue;
+		}
+
+		if (RenderAsset->GetLastRenderTimeForStreaming() < LastWorldUpdateTime_MipCalcTask)
+		{
+			// Not visible
+			continue;
+		}
+
+		const int32 StreamingIdx = RenderAsset->StreamingIndex;
+		FStreamingRenderAsset& Asset = StreamingRenderAssets[StreamingIdx];
+
+		check(RenderAsset == Asset.RenderAsset);
+
+		Asset.UpdateStreamingStatus(false);
+
+		if (!Asset.bInFlight && Asset.ResidentMips < Asset.NumNonOptionalMips)
+		{
+			RenderAsset->StreamIn(Asset.NumNonOptionalMips, true);
+			RenderAsset->bHasStreamingUpdatePending = true;
+			Asset.bHasUpdatePending = true;
+			VisibleFastResponseRenderAssetIndices.Add(StreamingIdx);
+
+			for (int32 Idx = 0; Idx < PendingMipCopyRequests.Num(); ++Idx)
+			{
+				FPendingMipCopyRequest& PendingRequest = PendingMipCopyRequests[Idx];
+
+				if (PendingRequest.RenderAsset == RenderAsset)
+				{
+					PendingRequest.RenderAsset = nullptr;
+				}
+			}
+
+			Asset.UpdateStreamingStatus(false);
+			TrackRenderAssetEvent(&Asset, RenderAsset, true, this);
+		}
+	}
+}
+
 void FRenderAssetStreamingManager::ProcessRemovedRenderAssets()
 {
 	for (int32 AssetIndex : RemovedRenderAssetIndices)
@@ -349,7 +414,14 @@ void FRenderAssetStreamingManager::ProcessRemovedRenderAssets()
 		if (StreamingRenderAssets.IsValidIndex(AssetIndex))
 		{
 			// Update the texture with its new index.
-			StreamingRenderAssets[AssetIndex].RenderAsset->StreamingIndex = AssetIndex;
+			int32& StreamingIdx = StreamingRenderAssets[AssetIndex].RenderAsset->StreamingIndex;
+
+			if (VisibleFastResponseRenderAssetIndices.Remove(StreamingIdx))
+			{
+				VisibleFastResponseRenderAssetIndices.Add(AssetIndex);
+			}
+
+			StreamingIdx = AssetIndex;
 		}
 	}
 	RemovedRenderAssetIndices.Empty();
@@ -550,6 +622,8 @@ void FRenderAssetStreamingManager::PrepareAsyncTask(bool bProcessEverything)
 		AsyncTask.Reset(0, Stats.AllocatedMemorySize, MAX_int64, MAX_int64 / 2, 0);
 	}
 	AsyncTask.StreamingData.Init(CurrentViewInfos, LastWorldUpdateTime, LevelRenderAssetManagers, DynamicComponentManager);
+
+	LastWorldUpdateTime_MipCalcTask = LastWorldUpdateTime;
 }
 
 /**
@@ -727,6 +801,9 @@ void FRenderAssetStreamingManager::RemoveStreamingRenderAsset( UStreamableRender
 	{
 		StreamingRenderAssets[Idx].RenderAsset = nullptr;
 		RemovedRenderAssetIndices.Add(Idx);
+
+		FastResponseRenderAssets.Remove(RenderAsset);
+		VisibleFastResponseRenderAssetIndices.Remove(Idx);
 	}
 
 	RenderAsset->StreamingIndex = INDEX_NONE;
@@ -1014,6 +1091,31 @@ void FRenderAssetStreamingManager::UpdateIndividualRenderAsset( UStreamableRende
 	StreamingRenderAsset->StreamWantedMips(*this);
 }
 
+void FRenderAssetStreamingManager::FastForceFullyResident(UStreamableRenderAsset* RenderAsset)
+{
+	check(IsInGameThread());
+
+	if (CVarStreamingAllowFastForceResident.GetValueOnGameThread()
+		&& IStreamingManager::Get().IsStreamingEnabled()
+		&& !bPauseRenderAssetStreaming
+		&& RenderAsset
+		&& RenderAsset->bIgnoreStreamingMipBias
+		&& (RenderAsset->bForceMiplevelsToBeResident || RenderAsset->ForceMipLevelsToBeResidentTimestamp >= FApp::GetCurrentTime())
+		&& StreamingRenderAssets.IsValidIndex(RenderAsset->StreamingIndex)
+		&& StreamingRenderAssets[RenderAsset->StreamingIndex].RenderAsset == RenderAsset)
+	{
+		const int32 StreamingIdx = RenderAsset->StreamingIndex;
+		FStreamingRenderAsset& Asset = StreamingRenderAssets[StreamingIdx];
+
+		Asset.UpdateStreamingStatus(false);
+
+		if (Asset.ResidentMips < Asset.NumNonOptionalMips)
+		{
+			FastResponseRenderAssets.Add(Asset.RenderAsset);
+		}
+	}
+}
+
 /**
  * Not thread-safe: Updates a portion (as indicated by 'StageIndex') of all streaming textures,
  * allowing their streaming state to progress.
@@ -1098,7 +1200,7 @@ void FRenderAssetStreamingManager::StreamRenderAssets( bool bProcessEverything )
 	{
 		for (int32 AssetIndex : AsyncTask.GetCancelationRequests())
 		{
-			if (StreamingRenderAssets.IsValidIndex(AssetIndex))
+			if (StreamingRenderAssets.IsValidIndex(AssetIndex) && !VisibleFastResponseRenderAssetIndices.Contains(AssetIndex))
 			{
 				StreamingRenderAssets[AssetIndex].CancelPendingMipChangeRequest();
 			}
@@ -1114,7 +1216,8 @@ void FRenderAssetStreamingManager::StreamRenderAssets( bool bProcessEverything )
 			for (int32 AssetIndex : AsyncTask.GetLoadRequests())
 			{
 				if (StreamingRenderAssets.IsValidIndex(AssetIndex)
-					&& StreamingRenderAssets[AssetIndex].RenderAsset)
+					&& StreamingRenderAssets[AssetIndex].RenderAsset
+					&& !VisibleFastResponseRenderAssetIndices.Contains(AssetIndex))
 				{
 					FStreamingRenderAsset& StreamingRenderAsset = StreamingRenderAssets[AssetIndex];
 					StreamingRenderAsset.CacheStreamingMetaData();
@@ -1126,7 +1229,7 @@ void FRenderAssetStreamingManager::StreamRenderAssets( bool bProcessEverything )
 		{
 			for (int32 AssetIndex : AsyncTask.GetLoadRequests())
 			{
-				if (StreamingRenderAssets.IsValidIndex(AssetIndex))
+				if (StreamingRenderAssets.IsValidIndex(AssetIndex) && !VisibleFastResponseRenderAssetIndices.Contains(AssetIndex))
 				{
 					StreamingRenderAssets[AssetIndex].StreamWantedMips(*this);
 				}
@@ -1136,7 +1239,7 @@ void FRenderAssetStreamingManager::StreamRenderAssets( bool bProcessEverything )
 	
 	for (int32 AssetIndex : AsyncTask.GetPendingUpdateDirties())
 	{
-		if (StreamingRenderAssets.IsValidIndex(AssetIndex))
+		if (StreamingRenderAssets.IsValidIndex(AssetIndex) && !VisibleFastResponseRenderAssetIndices.Contains(AssetIndex))
 		{
 			FStreamingRenderAsset& StreamingRenderAsset = StreamingRenderAssets[AssetIndex];
 			const bool bNewState = StreamingRenderAsset.HasUpdatePending(bPauseRenderAssetStreaming, AsyncTask.HasAnyView());
@@ -1149,6 +1252,23 @@ void FRenderAssetStreamingManager::StreamRenderAssets( bool bProcessEverything )
 			}
 		}
 	}
+
+	// Reset BudgetMipBias and MaxAllowedMips before we forget. Otherwise, new requests may stream out forced mips
+	for (TSet<int32>::TConstIterator It(VisibleFastResponseRenderAssetIndices); It; ++It)
+	{
+		const int32 Idx = *It;
+
+		if (StreamingRenderAssets.IsValidIndex(Idx))
+		{
+			FStreamingRenderAsset& Asset = StreamingRenderAssets[Idx];
+
+			Asset.BudgetMipBias = 0;
+			// If optional mips is available but not counted, they will be counted later in UpdateDynamicData
+			Asset.MaxAllowedMips = FMath::Max(Asset.MaxAllowedMips, Asset.NumNonOptionalMips);
+		}
+	}
+
+	VisibleFastResponseRenderAssetIndices.Empty();
 }
 
 void FRenderAssetStreamingManager::ProcessPendingMipCopyRequests()
@@ -1382,6 +1502,8 @@ void FRenderAssetStreamingManager::UpdateResourceStreaming( float DeltaTime, boo
 		PrepareAsyncTask(bProcessEverything || Settings.bStressTest);
 		AsyncWork->StartSynchronousTask();
 
+		TickFastResponseAssets();
+
 		StreamRenderAssets(bProcessEverything);
 
 		STAT(GatheredStats.SetupAsyncTaskCycles = 0);
@@ -1409,6 +1531,7 @@ void FRenderAssetStreamingManager::UpdateResourceStreaming( float DeltaTime, boo
 		UpdatePendingStates(false);
 		PrepareAsyncTask(bProcessEverything || Settings.bStressTest);
 		AsyncWork->StartBackgroundTask(CVarUseBackgroundThreadPool.GetValueOnGameThread() ? GBackgroundPriorityThreadPool : GThreadPool);
+		TickFastResponseAssets();
 		++ProcessingStage;
 
 		STAT(GatheredStats.SetupAsyncTaskCycles += FPlatformTime::Cycles();)
@@ -1422,6 +1545,8 @@ void FRenderAssetStreamingManager::UpdateResourceStreaming( float DeltaTime, boo
 			SetLastUpdateTime();
 		}
 
+		TickFastResponseAssets();
+
 		FEvent* SyncEvent = nullptr;
 		// Optimization: overlapping UpdateStreamingRenderAssets() and IncrementalUpdate();
 		// Restrict this optimization to platforms tested to have a win;
@@ -1461,6 +1586,8 @@ void FRenderAssetStreamingManager::UpdateResourceStreaming( float DeltaTime, boo
 			StreamingRenderAssets[TextureIndex].UpdateStreamingStatus(DeltaTime > 0);
 		}
 
+		TickFastResponseAssets();
+
 		StreamRenderAssets(bProcessEverything);
 		// Release the old view now as the destructors can be expensive. Now only the dynamic manager holds a ref.
 		AsyncWork->GetTask().ReleaseAsyncViews();
diff --git a/Engine/Source/Runtime/Engine/Private/Streaming/StreamingManagerTexture.h b/Engine/Source/Runtime/Engine/Private/Streaming/StreamingManagerTexture.h
--- a/Engine/Source/Runtime/Engine/Private/Streaming/StreamingManagerTexture.h
+++ b/Engine/Source/Runtime/Engine/Private/Streaming/StreamingManagerTexture.h
@@ -53,6 +53,15 @@ struct FRenderAssetStreamingManager final : public IRenderAssetStreamingManager
 	 */
 	virtual void UpdateIndividualRenderAsset( UStreamableRenderAsset* RenderAsset ) override;
 
+	/**
+	 * Register an asset whose non-resident mips need to be loaded ASAP when visible.
+	 * bIgnoreStreamingMipBias must be set on the asset.
+	 * Either bForceMiplevelsToBeResident or ForceMipLevelsToBeResidentTimestamp need to be set on the asset.
+	 *
+	 * @param RenderAsset The asset to register
+	 */
+	virtual void FastForceFullyResident(UStreamableRenderAsset* RenderAsset) override;
+
 	/**
 	 * Blocks till all pending requests are fulfilled.
 	 *
@@ -187,6 +196,9 @@ private:
 		 */
 		void UpdateStreamingRenderAssets( int32 StageIndex, int32 NumStages, bool bWaitForMipFading );
 
+		/** Check visibility of fast response assets and initiate stream-in requests if necessary. */
+		void TickFastResponseAssets();
+
 		void ProcessRemovedRenderAssets();
 		void ProcessAddedRenderAssets();
 		void ConditionalUpdateStaticData();
@@ -333,6 +345,12 @@ private:
 	/** The list of indices with null render asset in StreamingRenderAssets. */
 	TArray<int32>	RemovedRenderAssetIndices;
 
+	/** [Game Thread] Forced fully resident assets that need to be loaded ASAP when visible. */
+	TSet<UStreamableRenderAsset*> FastResponseRenderAssets;
+
+	/** [Game Thread] Fast response assets detected visible before the end of full streaming update. */
+	TSet<int32> VisibleFastResponseRenderAssetIndices;
+
 	// Represent a pending request to stream in/out mips to/from GPU for a texture or mesh
 	struct FPendingMipCopyRequest
 	{
@@ -385,6 +403,9 @@ private:
 	/** Last time all data were fully updated. Instances are considered visible if they were rendered between that last time and the current time. */
 	float LastWorldUpdateTime;
 
+	/** LastWorldUpdateTime seen by the async task. */
+	float LastWorldUpdateTime_MipCalcTask;
+
 	FRenderAssetStreamingStats DisplayedStats;
 	FRenderAssetStreamingStats GatheredStats;
 
diff --git a/Engine/Source/Runtime/Engine/Private/Streaming/StreamingTexture.h b/Engine/Source/Runtime/Engine/Private/Streaming/StreamingTexture.h
--- a/Engine/Source/Runtime/Engine/Private/Streaming/StreamingTexture.h
+++ b/Engine/Source/Runtime/Engine/Private/Streaming/StreamingTexture.h
@@ -174,7 +174,7 @@ struct FStreamingRenderAsset
 		// In editor, forced stream in should never have reduced mips as they can be edited.
 		return (IsMesh() || LODGroup != TEXTUREGROUP_HierarchicalLOD)
 			&& !bIsTerrainTexture
-			&& !bIgnoreStreamingMipBias
+			&& !(bForceFullyLoadHeuristic && bIgnoreStreamingMipBias)
 			&& !(GIsEditor && bForceFullyLoadHeuristic); 
 	}
 
@@ -286,7 +286,7 @@ struct FStreamingRenderAsset
 	uint32			bReadyForStreaming : 1;
 	/** (2) Whether the texture should be forcibly fully loaded. */
 	uint32			bForceFullyLoad : 1;
-	/** (2) Whether the texture resolution should be affected by the memory budget. */
+	/** (2) Whether the texture resolution should be affected by the memory budget. Only effective when forced fully resident. */
 	uint32			bIgnoreStreamingMipBias : 1;
 
 	/** (3) Whether the texture is currently being streamed in/out. */
diff --git a/Engine/Source/Runtime/Engine/Public/ContentStreaming.h b/Engine/Source/Runtime/Engine/Public/ContentStreaming.h
--- a/Engine/Source/Runtime/Engine/Public/ContentStreaming.h
+++ b/Engine/Source/Runtime/Engine/Public/ContentStreaming.h
@@ -398,6 +398,9 @@ struct IRenderAssetStreamingManager : public IStreamingManager
 	*/
 	virtual void UpdateIndividualRenderAsset(UStreamableRenderAsset* RenderAsset) = 0;
 
+	/** Stream in non-resident mips for an asset ASAP. */
+	virtual void FastForceFullyResident(UStreamableRenderAsset* RenderAsset) = 0;
+
 	/**
 	* Temporarily boosts the streaming distance factor by the specified number.
 	* This factor is automatically reset to 1.0 after it's been used for mip-calculations.
diff --git a/Engine/Source/Runtime/UMG/Private/Components/Image.cpp b/Engine/Source/Runtime/UMG/Private/Components/Image.cpp
--- a/Engine/Source/Runtime/UMG/Private/Components/Image.cpp
+++ b/Engine/Source/Runtime/UMG/Private/Components/Image.cpp
@@ -167,6 +167,7 @@ void UImage::SetBrushFromTexture(UTexture2D* Texture, bool bMatchSize)
 
 		if (Texture) // Since this texture is used as UI, don't allow it affected by budget.
 		{
+			Texture->bForceMiplevelsToBeResident = true;
 			Texture->bIgnoreStreamingMipBias = true;
 		}
 
