#UE4 Nonunity fix

diff --git a/Engine/Source/Runtime/Engine/Private/Materials/MaterialInterface.cpp b/Engine/Source/Runtime/Engine/Private/Materials/MaterialInterface.cpp
--- a/Engine/Source/Runtime/Engine/Private/Materials/MaterialInterface.cpp
+++ b/Engine/Source/Runtime/Engine/Private/Materials/MaterialInterface.cpp
@@ -19,6 +19,7 @@
 #include "Algo/BinarySearch.h"
 #include "Interfaces/ITargetPlatform.h"
 #include "Components.h"
+#include "ContentStreaming.h"
 
 /**
  * This is used to deprecate data that has been built with older versions.
