Adding LPF control to sound classes and sound mixes.

diff --git a/Engine/Source/Runtime/Engine/Classes/Sound/SoundClass.h b/Engine/Source/Runtime/Engine/Classes/Sound/SoundClass.h
--- a/Engine/Source/Runtime/Engine/Classes/Sound/SoundClass.h
+++ b/Engine/Source/Runtime/Engine/Classes/Sound/SoundClass.h
@@ -4,6 +4,7 @@
 #include "CoreMinimal.h"
 #include "UObject/ObjectMacros.h"
 #include "UObject/Object.h"
+#include "AudioDefines.h"
 #if WITH_EDITOR
 #include "EdGraph/EdGraph.h"
 #endif
@@ -59,6 +60,10 @@ struct FSoundClassProperties
 	UPROPERTY(EditAnywhere, BlueprintReadOnly, Category=SoundClassProperties)
 	float Pitch;
 
+	/** Lowpass filter frequency */
+	UPROPERTY(EditAnywhere, BlueprintReadOnly, Category = SoundClassProperties)
+	float LowPassFilterFrequency;
+
 	/** The amount of stereo sounds to bleed to the rear speakers */
 	UPROPERTY(EditAnywhere, BlueprintReadOnly, Category=SoundClassProperties)
 	float StereoBleed;
@@ -118,6 +123,7 @@ struct FSoundClassProperties
 	FSoundClassProperties()
 		: Volume(1.0f)
 		, Pitch(1.0f)
+		, LowPassFilterFrequency(MAX_FILTER_FREQUENCY)
 		, StereoBleed(0.25f)
 		, LFEBleed(0.5f)
 		, VoiceCenterChannelVolume(0.0f)
diff --git a/Engine/Source/Runtime/Engine/Classes/Sound/SoundMix.h b/Engine/Source/Runtime/Engine/Classes/Sound/SoundMix.h
--- a/Engine/Source/Runtime/Engine/Classes/Sound/SoundMix.h
+++ b/Engine/Source/Runtime/Engine/Classes/Sound/SoundMix.h
@@ -107,25 +107,30 @@ struct FSoundClassAdjuster
 	USoundClass* SoundClassObject;
 
 	/* A multiplier applied to the volume. */
-	UPROPERTY(EditAnywhere, Category=SoundClassAdjuster )
+	UPROPERTY(EditAnywhere, Category=SoundClassAdjuster, meta = (ClampMin = "0.0", UIMin = "0.0", UIMax = "4.0"))
 	float VolumeAdjuster;
 
 	/* A multiplier applied to the pitch. */
-	UPROPERTY(EditAnywhere, Category=SoundClassAdjuster )
+	UPROPERTY(EditAnywhere, Category=SoundClassAdjuster, meta = (ClampMin = "0.0", ClampMax = "8.0", UIMin = "0.0", UIMax = "8.0"))
 	float PitchAdjuster;
 
+	/* Lowpass filter cutoff frequency to apply to sound sources. */
+	UPROPERTY(EditAnywhere, Category = SoundClassAdjuster, meta = (ClampMin = "0.0", ClampMax = "20000.0", UIMin = "0.0", UIMax = "20000.0"))
+	float LowPassFilterFrequency;
+
 	/* Set to true to apply this adjuster to all children of the sound class. */
 	UPROPERTY(EditAnywhere, Category=SoundClassAdjuster )
 	uint32 bApplyToChildren:1;
 
 	/* A multiplier applied to VoiceCenterChannelVolume. */
-	UPROPERTY(EditAnywhere, Category=SoundClassAdjuster )
+	UPROPERTY(EditAnywhere, Category=SoundClassAdjuster, meta = (ClampMin = "0.0", UIMin = "0.0", UIMax = "4.0"))
 	float VoiceCenterChannelVolumeAdjuster;
 
 	FSoundClassAdjuster()
 		: SoundClassObject(NULL)
 		, VolumeAdjuster(1)
 		, PitchAdjuster(1)
+		, LowPassFilterFrequency(MAX_FILTER_FREQUENCY)
 		, bApplyToChildren(false)
 		, VoiceCenterChannelVolumeAdjuster(1)
 		{
diff --git a/Engine/Source/Runtime/Engine/Private/Audio.cpp b/Engine/Source/Runtime/Engine/Private/Audio.cpp
--- a/Engine/Source/Runtime/Engine/Private/Audio.cpp
+++ b/Engine/Source/Runtime/Engine/Private/Audio.cpp
@@ -384,6 +384,7 @@ void FSoundSource::SetFilterFrequency()
 			LPFFrequency = FMath::Min(LPFFrequency, WaveInstance->AmbientZoneFilterFrequency);
 			LPFFrequency = FMath::Min(LPFFrequency, WaveInstance->AttenuationLowpassFilterFrequency);
 			LPFFrequency = FMath::Min(LPFFrequency, WaveInstance->SoundModulationControls.Lowpass);
+			LPFFrequency = FMath::Min(LPFFrequency, WaveInstance->SoundClassFilterFrequency);
 		}
 		break;
 	}
@@ -802,6 +803,7 @@ FWaveInstance::FWaveInstance(const UPTRINT InWaveInstanceHash, FActiveSound& InA
 	, ReverbPluginSettings(nullptr)
 	, OutputTarget(EAudioOutputTarget::Speaker)
 	, LowPassFilterFrequency(MAX_FILTER_FREQUENCY)
+	, SoundClassFilterFrequency(MAX_FILTER_FREQUENCY)
 	, OcclusionFilterFrequency(MAX_FILTER_FREQUENCY)
 	, AmbientZoneFilterFrequency(MAX_FILTER_FREQUENCY)
 	, AttenuationLowpassFilterFrequency(MAX_FILTER_FREQUENCY)
diff --git a/Engine/Source/Runtime/Engine/Private/AudioDevice.cpp b/Engine/Source/Runtime/Engine/Private/AudioDevice.cpp
--- a/Engine/Source/Runtime/Engine/Private/AudioDevice.cpp
+++ b/Engine/Source/Runtime/Engine/Private/AudioDevice.cpp
@@ -2143,6 +2143,12 @@ void FAudioDevice::RecursiveApplyAdjuster(const FSoundClassAdjuster& InAdjuster,
 		Properties->Pitch *= InAdjuster.PitchAdjuster;
 		Properties->VoiceCenterChannelVolume *= InAdjuster.VoiceCenterChannelVolumeAdjuster;
 
+		// Only set the LPF frequency if the input adjuster is *less* than the sound class' property
+		if (InAdjuster.LowPassFilterFrequency < Properties->LowPassFilterFrequency)
+		{
+			Properties->LowPassFilterFrequency = InAdjuster.LowPassFilterFrequency;
+		}
+
 		// Recurse through this classes children
 		for (int32 ChildIdx = 0; ChildIdx < InSoundClass->ChildClasses.Num(); ++ChildIdx)
 		{
@@ -2508,6 +2514,12 @@ static void UpdateClassAdjustorOverrideEntry(FSoundClassAdjuster& ClassAdjustor,
 	}
 }
 
+float FAudioDevice::GetInterpolatedFrequency(const float InFrequency, const float InterpValue) const
+{
+	const float NormFrequency = InterpolateAdjuster(InFrequency / MAX_FILTER_FREQUENCY, InterpValue);
+	return Audio::GetLogFrequencyClamped(NormFrequency, FVector2D(0.0f, 1.0f), FVector2D(MIN_FILTER_FREQUENCY, MAX_FILTER_FREQUENCY));
+}
+
 void FAudioDevice::ApplyClassAdjusters(USoundMix* SoundMix, float InterpValue, float DeltaTime)
 {
 	if (!SoundMix)
@@ -2544,6 +2556,7 @@ void FAudioDevice::ApplyClassAdjusters(USoundMix* SoundMix, float InterpValue, f
 				Entry.VolumeAdjuster = InterpolateAdjuster(Entry.VolumeAdjuster, InterpValue);
 				Entry.PitchAdjuster = InterpolateAdjuster(Entry.PitchAdjuster, InterpValue);
 				Entry.VoiceCenterChannelVolumeAdjuster = InterpolateAdjuster(Entry.VoiceCenterChannelVolumeAdjuster, InterpValue);
+				Entry.LowPassFilterFrequency = GetInterpolatedFrequency(Entry.LowPassFilterFrequency, InterpValue);
 			}
 		}
 
@@ -2621,6 +2634,7 @@ void FAudioDevice::ApplyClassAdjusters(USoundMix* SoundMix, float InterpValue, f
 					EntryCopy.VolumeAdjuster = InterpolateAdjuster(Entry.VolumeAdjuster, InterpValue);
 					EntryCopy.PitchAdjuster = InterpolateAdjuster(Entry.PitchAdjuster, InterpValue);
 					EntryCopy.VoiceCenterChannelVolumeAdjuster = InterpolateAdjuster(Entry.VoiceCenterChannelVolumeAdjuster, InterpValue);
+					EntryCopy.LowPassFilterFrequency = GetInterpolatedFrequency(Entry.LowPassFilterFrequency, InterpValue);
 
 					RecursiveApplyAdjuster(EntryCopy, Entry.SoundClassObject);
 				}
@@ -2637,6 +2651,11 @@ void FAudioDevice::ApplyClassAdjusters(USoundMix* SoundMix, float InterpValue, f
 						Properties->Volume *= Entry.VolumeAdjuster;
 						Properties->Pitch *= Entry.PitchAdjuster;
 						Properties->VoiceCenterChannelVolume *= Entry.VoiceCenterChannelVolumeAdjuster;
+
+						if (Entry.LowPassFilterFrequency < Properties->LowPassFilterFrequency)
+						{
+							Properties->LowPassFilterFrequency = Entry.LowPassFilterFrequency;
+						}
 					}
 					// Otherwise, we need to use the "static" data and compute the adjustment interpolations now
 					else
@@ -2644,6 +2663,12 @@ void FAudioDevice::ApplyClassAdjusters(USoundMix* SoundMix, float InterpValue, f
 						Properties->Volume *= InterpolateAdjuster(Entry.VolumeAdjuster, InterpValue);
 						Properties->Pitch *= InterpolateAdjuster(Entry.PitchAdjuster, InterpValue);
 						Properties->VoiceCenterChannelVolume *= InterpolateAdjuster(Entry.VoiceCenterChannelVolumeAdjuster, InterpValue);
+
+						const float NewLPF = GetInterpolatedFrequency(Entry.LowPassFilterFrequency, InterpValue);
+						if (NewLPF < Properties->LowPassFilterFrequency)
+						{
+							Properties->LowPassFilterFrequency = NewLPF;
+						}
 					}
 				}
 				else
diff --git a/Engine/Source/Runtime/Engine/Private/SoundWave.cpp b/Engine/Source/Runtime/Engine/Private/SoundWave.cpp
--- a/Engine/Source/Runtime/Engine/Private/SoundWave.cpp
+++ b/Engine/Source/Runtime/Engine/Private/SoundWave.cpp
@@ -1615,8 +1615,8 @@ void USoundWave::Parse(FAudioDevice* AudioDevice, const UPTRINT NodeWaveInstance
 		float VolumeMultiplier = WaveInstance->GetVolumeMultiplier();
 		WaveInstance->SetVolumeMultiplier(VolumeMultiplier* SoundClassProperties->Volume);
 		WaveInstance->SetPitch(WaveInstance->Pitch * SoundClassProperties->Pitch);
-		//TODO: Add in HighFrequencyGainMultiplier property to sound classes
 
+		WaveInstance->SoundClassFilterFrequency = SoundClassProperties->LowPassFilterFrequency;
 		WaveInstance->VoiceCenterChannelVolume = SoundClassProperties->VoiceCenterChannelVolume;
 		WaveInstance->RadioFilterVolume = SoundClassProperties->RadioFilterVolume * ParseParams.VolumeMultiplier;
 		WaveInstance->RadioFilterVolumeThreshold = SoundClassProperties->RadioFilterVolumeThreshold * ParseParams.VolumeMultiplier;
diff --git a/Engine/Source/Runtime/Engine/Public/Audio.h b/Engine/Source/Runtime/Engine/Public/Audio.h
--- a/Engine/Source/Runtime/Engine/Public/Audio.h
+++ b/Engine/Source/Runtime/Engine/Public/Audio.h
@@ -7,6 +7,7 @@
 #pragma once
 
 #include "CoreMinimal.h"
+#include "AudioDefines.h"
 #include "Stats/Stats.h"
 #include "HAL/ThreadSafeBool.h"
 #include "Sound/SoundClass.h"
@@ -27,85 +28,6 @@ struct FActiveSound;
 struct FWaveInstance;
 struct FSoundSourceBusSendInfo;
 
-
-ENGINE_API DECLARE_LOG_CATEGORY_EXTERN(LogAudio, Warning, All);
-
-// Special log category used for temporary programmer debugging code of audio
-ENGINE_API DECLARE_LOG_CATEGORY_EXTERN(LogAudioDebug, Display, All);
- 
-/** 
- * Maximum number of channels that can be set using the ini setting
- */
-#define MAX_AUDIOCHANNELS				64
-
-
-/** 
- * Length of sound in seconds to be considered as looping forever
- */
-#define INDEFINITELY_LOOPING_DURATION	10000.0f
-
-/**
- * Some defaults to help cross platform consistency
- */
-#define SPEAKER_COUNT					6
-
-#define DEFAULT_LOW_FREQUENCY			600.0f
-#define DEFAULT_MID_FREQUENCY			1000.0f
-#define DEFAULT_HIGH_FREQUENCY			2000.0f
-
-#define MAX_VOLUME						4.0f
-#define MIN_PITCH						0.4f
-#define MAX_PITCH						2.0f
-
-#define MIN_SOUND_PRIORITY				0.0f
-#define MAX_SOUND_PRIORITY				100.0f
-
-#define DEFAULT_SUBTITLE_PRIORITY		10000.0f
-
-/**
- * Some filters don't work properly with extreme values, so these are the limits 
- */
-#define MIN_FILTER_GAIN					0.126f
-#define MAX_FILTER_GAIN					7.94f
-
-#define MIN_FILTER_FREQUENCY			20.0f
-#define MAX_FILTER_FREQUENCY			20000.0f
-
-#define MIN_FILTER_BANDWIDTH			0.1f
-#define MAX_FILTER_BANDWIDTH			2.0f
-
-/**
- * Audio stats
- */
-DECLARE_DWORD_COUNTER_STAT_EXTERN( TEXT( "Active Sounds" ), STAT_ActiveSounds, STATGROUP_Audio , );
-DECLARE_CYCLE_STAT_EXTERN(TEXT("Audio Evaluate Concurrency"), STAT_AudioEvaluateConcurrency, STATGROUP_Audio, );
-DECLARE_DWORD_COUNTER_STAT_EXTERN( TEXT( "Audio Sources" ), STAT_AudioSources, STATGROUP_Audio , );
-DECLARE_DWORD_COUNTER_STAT_EXTERN( TEXT( "Wave Instances" ), STAT_WaveInstances, STATGROUP_Audio , );
-DECLARE_DWORD_COUNTER_STAT_EXTERN( TEXT( "Wave Instances Dropped" ), STAT_WavesDroppedDueToPriority, STATGROUP_Audio , );
-DECLARE_DWORD_COUNTER_STAT_EXTERN(TEXT("Virtualized Loops"), STAT_AudioVirtualLoops, STATGROUP_Audio, );
-DECLARE_DWORD_COUNTER_STAT_EXTERN( TEXT( "Audible Wave Instances Dropped" ), STAT_AudibleWavesDroppedDueToPriority, STATGROUP_Audio , );
-DECLARE_DWORD_COUNTER_STAT_EXTERN(TEXT("Max Channels"), STAT_AudioMaxChannels, STATGROUP_Audio, );
-DECLARE_DWORD_COUNTER_STAT_EXTERN(TEXT("Max Stopping Sources"), STAT_AudioMaxStoppingSources, STATGROUP_Audio, );
-DECLARE_DWORD_COUNTER_STAT_EXTERN( TEXT( "Finished delegates called" ), STAT_AudioFinishedDelegatesCalled, STATGROUP_Audio , );
-DECLARE_CYCLE_STAT_EXTERN( TEXT( "Finished delegates time" ), STAT_AudioFinishedDelegates, STATGROUP_Audio , );
-DECLARE_MEMORY_STAT_EXTERN( TEXT( "Audio Memory Used" ), STAT_AudioMemorySize, STATGROUP_Audio , );
-DECLARE_FLOAT_ACCUMULATOR_STAT_EXTERN( TEXT( "Audio Buffer Time" ), STAT_AudioBufferTime, STATGROUP_Audio , );
-DECLARE_FLOAT_ACCUMULATOR_STAT_EXTERN( TEXT( "Audio Buffer Time (w/ Channels)" ), STAT_AudioBufferTimeChannels, STATGROUP_Audio , );
-DECLARE_CYCLE_STAT_EXTERN( TEXT( "Gathering WaveInstances" ), STAT_AudioGatherWaveInstances, STATGROUP_Audio , );
-DECLARE_CYCLE_STAT_EXTERN( TEXT( "Processing Sources" ), STAT_AudioStartSources, STATGROUP_Audio , );
-DECLARE_CYCLE_STAT_EXTERN( TEXT( "Updating Sources" ), STAT_AudioUpdateSources, STATGROUP_Audio , ENGINE_API);
-DECLARE_CYCLE_STAT_EXTERN( TEXT( "Updating Effects" ), STAT_AudioUpdateEffects, STATGROUP_Audio , );
-DECLARE_CYCLE_STAT_EXTERN( TEXT( "Source Init" ), STAT_AudioSourceInitTime, STATGROUP_Audio , ENGINE_API);
-DECLARE_CYCLE_STAT_EXTERN( TEXT( "Source Create" ), STAT_AudioSourceCreateTime, STATGROUP_Audio , ENGINE_API);
-DECLARE_CYCLE_STAT_EXTERN( TEXT( "Submit Buffers" ), STAT_AudioSubmitBuffersTime, STATGROUP_Audio , ENGINE_API);
-DECLARE_CYCLE_STAT_EXTERN( TEXT( "Decompress Audio" ), STAT_AudioDecompressTime, STATGROUP_Audio , );
-DECLARE_CYCLE_STAT_EXTERN( TEXT( "Decompress Vorbis" ), STAT_VorbisDecompressTime, STATGROUP_Audio , );
-DECLARE_CYCLE_STAT_EXTERN( TEXT( "Prepare Audio Decompression" ), STAT_AudioPrepareDecompressionTime, STATGROUP_Audio , );
-DECLARE_CYCLE_STAT_EXTERN( TEXT( "Prepare Vorbis Decompression" ), STAT_VorbisPrepareDecompressionTime, STATGROUP_Audio , );
-DECLARE_CYCLE_STAT_EXTERN( TEXT( "Finding Nearest Location" ), STAT_AudioFindNearestLocation, STATGROUP_Audio , );
-DECLARE_CYCLE_STAT_EXTERN( TEXT( "Decompress Streamed" ), STAT_AudioStreamedDecompressTime, STATGROUP_Audio , );
-DECLARE_CYCLE_STAT_EXTERN( TEXT( "Buffer Creation" ), STAT_AudioResourceCreationTime, STATGROUP_Audio , );
-
 /**
  * Channel definitions for multistream waves
  *
@@ -342,6 +264,9 @@ public:
 	/** The low pass filter frequency to use */
 	float LowPassFilterFrequency;
 
+	/** The low pass filter frequency to use from sound class. */
+	float SoundClassFilterFrequency;
+
 	/** The low pass filter frequency to use if the sound is occluded. */
 	float OcclusionFilterFrequency;
 
diff --git a/Engine/Source/Runtime/Engine/Public/AudioDefines.h b/Engine/Source/Runtime/Engine/Public/AudioDefines.h
new file mode 100644
--- /dev/null
+++ b/Engine/Source/Runtime/Engine/Public/AudioDefines.h
@@ -0,0 +1,86 @@
+// Copyright 1998-2019 Epic Games, Inc. All Rights Reserved.
+
+/*=============================================================================
+	AudioDefines.h: Defines for audio system
+=============================================================================*/
+
+#pragma once
+
+#include "CoreMinimal.h"
+
+ENGINE_API DECLARE_LOG_CATEGORY_EXTERN(LogAudio, Warning, All);
+
+// Special log category used for temporary programmer debugging code of audio
+ENGINE_API DECLARE_LOG_CATEGORY_EXTERN(LogAudioDebug, Display, All);
+ 
+/** 
+ * Maximum number of channels that can be set using the ini setting
+ */
+#define MAX_AUDIOCHANNELS				64
+
+/** 
+ * Length of sound in seconds to be considered as looping forever
+ */
+#define INDEFINITELY_LOOPING_DURATION	10000.0f
+
+/**
+ * Some defaults to help cross platform consistency
+ */
+#define SPEAKER_COUNT					6
+
+#define DEFAULT_LOW_FREQUENCY			600.0f
+#define DEFAULT_MID_FREQUENCY			1000.0f
+#define DEFAULT_HIGH_FREQUENCY			2000.0f
+
+#define MAX_VOLUME						4.0f
+#define MIN_PITCH						0.4f
+#define MAX_PITCH						2.0f
+
+#define MIN_SOUND_PRIORITY				0.0f
+#define MAX_SOUND_PRIORITY				100.0f
+
+#define DEFAULT_SUBTITLE_PRIORITY		10000.0f
+
+/**
+ * Some filters don't work properly with extreme values, so these are the limits 
+ */
+#define MIN_FILTER_GAIN					0.126f
+#define MAX_FILTER_GAIN					7.94f
+
+#define MIN_FILTER_FREQUENCY			20.0f
+#define MAX_FILTER_FREQUENCY			20000.0f
+
+#define MIN_FILTER_BANDWIDTH			0.1f
+#define MAX_FILTER_BANDWIDTH			2.0f
+
+/**
+ * Audio stats
+ */
+DECLARE_DWORD_COUNTER_STAT_EXTERN( TEXT( "Active Sounds" ), STAT_ActiveSounds, STATGROUP_Audio , );
+DECLARE_CYCLE_STAT_EXTERN(TEXT("Audio Evaluate Concurrency"), STAT_AudioEvaluateConcurrency, STATGROUP_Audio, );
+DECLARE_DWORD_COUNTER_STAT_EXTERN( TEXT( "Audio Sources" ), STAT_AudioSources, STATGROUP_Audio , );
+DECLARE_DWORD_COUNTER_STAT_EXTERN( TEXT( "Wave Instances" ), STAT_WaveInstances, STATGROUP_Audio , );
+DECLARE_DWORD_COUNTER_STAT_EXTERN( TEXT( "Wave Instances Dropped" ), STAT_WavesDroppedDueToPriority, STATGROUP_Audio , );
+DECLARE_DWORD_COUNTER_STAT_EXTERN(TEXT("Virtualized Loops"), STAT_AudioVirtualLoops, STATGROUP_Audio, );
+DECLARE_DWORD_COUNTER_STAT_EXTERN( TEXT( "Audible Wave Instances Dropped" ), STAT_AudibleWavesDroppedDueToPriority, STATGROUP_Audio , );
+DECLARE_DWORD_COUNTER_STAT_EXTERN(TEXT("Max Channels"), STAT_AudioMaxChannels, STATGROUP_Audio, );
+DECLARE_DWORD_COUNTER_STAT_EXTERN(TEXT("Max Stopping Sources"), STAT_AudioMaxStoppingSources, STATGROUP_Audio, );
+DECLARE_DWORD_COUNTER_STAT_EXTERN( TEXT( "Finished delegates called" ), STAT_AudioFinishedDelegatesCalled, STATGROUP_Audio , );
+DECLARE_CYCLE_STAT_EXTERN( TEXT( "Finished delegates time" ), STAT_AudioFinishedDelegates, STATGROUP_Audio , );
+DECLARE_MEMORY_STAT_EXTERN( TEXT( "Audio Memory Used" ), STAT_AudioMemorySize, STATGROUP_Audio , );
+DECLARE_FLOAT_ACCUMULATOR_STAT_EXTERN( TEXT( "Audio Buffer Time" ), STAT_AudioBufferTime, STATGROUP_Audio , );
+DECLARE_FLOAT_ACCUMULATOR_STAT_EXTERN( TEXT( "Audio Buffer Time (w/ Channels)" ), STAT_AudioBufferTimeChannels, STATGROUP_Audio , );
+DECLARE_CYCLE_STAT_EXTERN( TEXT( "Gathering WaveInstances" ), STAT_AudioGatherWaveInstances, STATGROUP_Audio , );
+DECLARE_CYCLE_STAT_EXTERN( TEXT( "Processing Sources" ), STAT_AudioStartSources, STATGROUP_Audio , );
+DECLARE_CYCLE_STAT_EXTERN( TEXT( "Updating Sources" ), STAT_AudioUpdateSources, STATGROUP_Audio , ENGINE_API);
+DECLARE_CYCLE_STAT_EXTERN( TEXT( "Updating Effects" ), STAT_AudioUpdateEffects, STATGROUP_Audio , );
+DECLARE_CYCLE_STAT_EXTERN( TEXT( "Source Init" ), STAT_AudioSourceInitTime, STATGROUP_Audio , ENGINE_API);
+DECLARE_CYCLE_STAT_EXTERN( TEXT( "Source Create" ), STAT_AudioSourceCreateTime, STATGROUP_Audio , ENGINE_API);
+DECLARE_CYCLE_STAT_EXTERN( TEXT( "Submit Buffers" ), STAT_AudioSubmitBuffersTime, STATGROUP_Audio , ENGINE_API);
+DECLARE_CYCLE_STAT_EXTERN( TEXT( "Decompress Audio" ), STAT_AudioDecompressTime, STATGROUP_Audio , );
+DECLARE_CYCLE_STAT_EXTERN( TEXT( "Decompress Vorbis" ), STAT_VorbisDecompressTime, STATGROUP_Audio , );
+DECLARE_CYCLE_STAT_EXTERN( TEXT( "Prepare Audio Decompression" ), STAT_AudioPrepareDecompressionTime, STATGROUP_Audio , );
+DECLARE_CYCLE_STAT_EXTERN( TEXT( "Prepare Vorbis Decompression" ), STAT_VorbisPrepareDecompressionTime, STATGROUP_Audio , );
+DECLARE_CYCLE_STAT_EXTERN( TEXT( "Finding Nearest Location" ), STAT_AudioFindNearestLocation, STATGROUP_Audio , );
+DECLARE_CYCLE_STAT_EXTERN( TEXT( "Decompress Streamed" ), STAT_AudioStreamedDecompressTime, STATGROUP_Audio , );
+DECLARE_CYCLE_STAT_EXTERN( TEXT( "Buffer Creation" ), STAT_AudioResourceCreationTime, STATGROUP_Audio , );
diff --git a/Engine/Source/Runtime/Engine/Public/AudioDevice.h b/Engine/Source/Runtime/Engine/Public/AudioDevice.h
--- a/Engine/Source/Runtime/Engine/Public/AudioDevice.h
+++ b/Engine/Source/Runtime/Engine/Public/AudioDevice.h
@@ -1369,6 +1369,9 @@ private:
 		return Adjuster * InterpValue + 1.0f - InterpValue;
 	}
 
+	/** Retrieve the filter frequency to use. Takes into account logarithmic nature of frequency. */
+	float GetInterpolatedFrequency(const float InFrequency, const float InterpValue) const;
+
 	/** Allow platforms to optionally specify low-level audio platform settings. */
 	virtual FAudioPlatformSettings GetPlatformSettings() const { return FAudioPlatformSettings(); }
 
