#UE4 Zero-scale warning now only applies to components that have collision

diff --git a/Engine/Source/Runtime/Engine/Private/Components/PrimitiveComponent.cpp b/Engine/Source/Runtime/Engine/Private/Components/PrimitiveComponent.cpp
--- a/Engine/Source/Runtime/Engine/Private/Components/PrimitiveComponent.cpp
+++ b/Engine/Source/Runtime/Engine/Private/Components/PrimitiveComponent.cpp
@@ -740,7 +740,7 @@ void UPrimitiveComponent::OnCreatePhysicsState()
 				BodyTransform.SetScale3D(FVector(KINDA_SMALL_NUMBER));
 			}
 #if !(UE_BUILD_SHIPPING || UE_BUILD_TEST)
-			if (FMath::IsNearlyZero(BodyScale.X) || FMath::IsNearlyZero(BodyScale.Y) || FMath::IsNearlyZero(BodyScale.Z))
+			if ((BodyInstance.GetCollisionEnabled() != ECollisionEnabled::NoCollision) && (FMath::IsNearlyZero(BodyScale.X) || FMath::IsNearlyZero(BodyScale.Y) || FMath::IsNearlyZero(BodyScale.Z)))
 			{
 				UE_LOG(LogPhysics, Warning, TEXT("Scale for %s has a component set to zero, which will result in a bad body instance. Scale:%s"), *GetPathNameSafe(this), *BodyScale.ToString());
 			}
