CSVProfiler EventWait scope for LastDemotable (Niagara wait) and PostUpdate.

#rnx

diff --git a/Engine/Source/Runtime/Engine/Private/LevelTick.cpp b/Engine/Source/Runtime/Engine/Private/LevelTick.cpp
--- a/Engine/Source/Runtime/Engine/Private/LevelTick.cpp
+++ b/Engine/Source/Runtime/Engine/Private/LevelTick.cpp
@@ -1646,11 +1646,13 @@ void UWorld::Tick( ELevelTick TickType, float DeltaSeconds )
 			{
 				SCOPE_CYCLE_COUNTER(STAT_TG_PostUpdateWork);
 				SCOPE_TIME_GUARD_MS(TEXT("UWorld::Tick - PostUpdateWork"), 5);
+				CSV_SCOPED_SET_WAIT_STAT(PostUpdateWork);
 				RunTickGroup(TG_PostUpdateWork);
 			}
 			{
 				SCOPE_CYCLE_COUNTER(STAT_TG_LastDemotable);
 				SCOPE_TIME_GUARD_MS(TEXT("UWorld::Tick - TG_LastDemotable"), 5);
+				CSV_SCOPED_SET_WAIT_STAT(LastDemotable);
 				RunTickGroup(TG_LastDemotable);
 			}
 
