Adding a "Safe Divide" function to UKismetMathLibrary

diff --git a/Engine/Source/Runtime/Engine/Classes/Kismet/KismetMathLibrary.h b/Engine/Source/Runtime/Engine/Classes/Kismet/KismetMathLibrary.h
--- a/Engine/Source/Runtime/Engine/Classes/Kismet/KismetMathLibrary.h
+++ b/Engine/Source/Runtime/Engine/Classes/Kismet/KismetMathLibrary.h
@@ -686,6 +686,10 @@ class ENGINE_API UKismetMathLibrary : public UBlueprintFunctionLibrary
 	UFUNCTION(BlueprintPure, meta=(DisplayName = "Clamp (float)", Min="0.0", Max="1.0"), Category="Math|Float")
 	static float FClamp(float Value, float Min, float Max);
 
+	/** This functions returns 0 if B (the denominator) is zero */
+	UFUNCTION(BlueprintPure, Category = "Math|Float", meta = (Keywords = "percent"))
+	static float SafeDivide(float A, float B);
+
 	/** Returns max of all array entries and the index at which it was found. Returns value of 0 and index of -1 if the supplied array is empty. */
 	UFUNCTION(BlueprintPure, Category="Math|Integer")
 	static void MaxOfIntArray(const TArray<int32>& IntArray, int32& IndexOfMaxValue, int32& MaxValue);
diff --git a/Engine/Source/Runtime/Engine/Private/KismetMathLibrary.cpp b/Engine/Source/Runtime/Engine/Private/KismetMathLibrary.cpp
--- a/Engine/Source/Runtime/Engine/Private/KismetMathLibrary.cpp
+++ b/Engine/Source/Runtime/Engine/Private/KismetMathLibrary.cpp
@@ -285,6 +285,11 @@ float UKismetMathLibrary::MakePulsatingValue(float InCurrentTime, float InPulses
 	return FMath::MakePulsatingValue((double)InCurrentTime, InPulsesPerSecond, InPhase);
 }
 
+float UKismetMathLibrary::SafeDivide(float A, float B)
+{
+	return (B != 0.0f) ? (A / B) : 0.0f;
+}
+
 void UKismetMathLibrary::MaxOfIntArray(const TArray<int32>& IntArray, int32& IndexOfMaxValue, int32& MaxValue)
 {
 	MaxValue = FMath::Max(IntArray, &IndexOfMaxValue);
