A bit more clean-up for Audio::FDeviceId
(Missing from last CL)

diff --git a/Engine/Source/Runtime/Engine/Private/AudioDeviceManager.cpp b/Engine/Source/Runtime/Engine/Private/AudioDeviceManager.cpp
--- a/Engine/Source/Runtime/Engine/Private/AudioDeviceManager.cpp
+++ b/Engine/Source/Runtime/Engine/Private/AudioDeviceManager.cpp
@@ -502,7 +502,7 @@ bool FAudioDeviceManager::CreateAudioDevice(bool bCreateNewDevice, FCreateAudioD
 	return (OutResults.AudioDevice != nullptr);
 }
 
-bool FAudioDeviceManager::IsValidAudioDeviceHandle(uint32 Handle) const
+bool FAudioDeviceManager::IsValidAudioDeviceHandle(Audio::FDeviceId Handle) const
 {
 	if (AudioDeviceModule == nullptr || Handle == INDEX_NONE)
 	{
@@ -519,7 +519,7 @@ bool FAudioDeviceManager::IsValidAudioDeviceHandle(uint32 Handle) const
 	return Generations[Index] == Generation;
 }
 
-bool FAudioDeviceManager::ShutdownAudioDevice(uint32 Handle)
+bool FAudioDeviceManager::ShutdownAudioDevice(Audio::FDeviceId Handle)
 {
 	if (!IsValidAudioDeviceHandle(Handle))
 	{
@@ -597,7 +597,7 @@ bool FAudioDeviceManager::ShutdownAllAudioDevices()
 	return true;
 }
 
-FAudioDevice* FAudioDeviceManager::GetAudioDevice(uint32 Handle)
+FAudioDevice* FAudioDeviceManager::GetAudioDevice(Audio::FDeviceId Handle)
 {
 	if (!IsValidAudioDeviceHandle(Handle))
 	{
@@ -803,7 +803,7 @@ void FAudioDeviceManager::SetActiveDevice(uint32 InAudioDeviceHandle)
 	}
 }
 
-void FAudioDeviceManager::SetSoloDevice(uint32 InAudioDeviceHandle)
+void FAudioDeviceManager::SetSoloDevice(Audio::FDeviceId InAudioDeviceHandle)
 {
 	SoloDeviceHandle = InAudioDeviceHandle;
 	if (SoloDeviceHandle != INDEX_NONE)
diff --git a/Engine/Source/Runtime/Engine/Public/AudioDeviceManager.h b/Engine/Source/Runtime/Engine/Public/AudioDeviceManager.h
--- a/Engine/Source/Runtime/Engine/Public/AudioDeviceManager.h
+++ b/Engine/Source/Runtime/Engine/Public/AudioDeviceManager.h
@@ -1,11 +1,11 @@
 // Copyright 1998-2019 Epic Games, Inc. All Rights Reserved.
 #pragma once
 
+#include "Audio.h"
 #include "Audio/AudioDebug.h"
 #include "AudioThread.h"
 #include "Containers/Queue.h"
 #include "CoreMinimal.h"
-#include "Sound/SoundEffectSource.h"
 
 
 #if ENABLE_AUDIO_DEBUG
@@ -53,7 +53,7 @@ public:
 	bool Initialize();
 
 	/** Returns the handle to the main audio device. */
-	uint32 GetMainAudioDeviceHandle() const { return MainAudioDeviceHandle; }
+	Audio::FDeviceId GetMainAudioDeviceHandle() const { return MainAudioDeviceHandle; }
 
 	/** Returns true if we're currently using the audio mixer. */
 	bool IsUsingAudioMixer() const;
@@ -78,13 +78,13 @@ public:
 	* Returns whether the audio device handle is valid (i.e. points to
 	* an actual audio device instance)
 	*/
-	bool IsValidAudioDeviceHandle(uint32 Handle) const;
+	bool IsValidAudioDeviceHandle(Audio::FDeviceId Handle) const;
 
 	/**
 	* Shutsdown the audio device associated with the handle. The handle
 	* will become invalid after the audio device is shut down.
 	*/
-	bool ShutdownAudioDevice(uint32 Handle);
+	bool ShutdownAudioDevice(Audio::FDeviceId Handle);
 
 	/**
 	* Shuts down all active audio devices
@@ -95,7 +95,7 @@ public:
 	* Returns a ptr to the audio device associated with the handle. If the
 	* handle is invalid then a NULL device ptr will be returned.
 	*/
-	FAudioDevice* GetAudioDevice(uint32 Handle);
+	FAudioDevice* GetAudioDevice(Audio::FDeviceId Handle);
 
 	/**
 	* Returns a ptr to the active audio device. If there is no active
@@ -149,7 +149,7 @@ public:
 	void SetActiveDevice(uint32 InAudioDeviceHandle);
 
 	/** Sets an audio device to be solo'd */
-	void SetSoloDevice(uint32 InAudioDeviceHandle);
+	void SetSoloDevice(Audio::FDeviceId InAudioDeviceHandle);
 
 	/** Links up the resource data indices for looking up and cleaning up. */
 	void TrackResource(USoundWave* SoundWave, FSoundBuffer* Buffer);
@@ -288,7 +288,7 @@ private:
 	FString AudioMixerModuleName;
 
 	/** Handle to the main audio device. */
-	uint32 MainAudioDeviceHandle;
+	Audio::FDeviceId MainAudioDeviceHandle;
 
 	/** Count of the number of free slots in the audio device array. */
 	uint32 FreeIndicesSize;
@@ -306,16 +306,16 @@ private:
 	* Array of audio device pointers. If the device has been free, then
 	* the device ptr at the array index will be null.
 	*/
-	TArray<FAudioDevice* > Devices;
+	TArray<FAudioDevice*> Devices;
 
 	/** Next resource ID to assign out to a wave/buffer */
 	int32 NextResourceID;
 
 	/** Which audio device is solo'd */
-	uint32 SoloDeviceHandle;
+	Audio::FDeviceId SoloDeviceHandle;
 
 	/** Which audio device is currently active */
-	uint32 ActiveAudioDeviceHandle;
+	Audio::FDeviceId ActiveAudioDeviceHandle;
 
 	/** Instance of the debug names struct. */
 	FDebugNames DebugNames;
