Material Editor - Fixing a crash in the editor when swapping to UI materials.  We are no generating the LocalVertexFactory again, but only when GIsEditor is true.

diff --git a/Engine/Source/Runtime/Engine/Private/LocalVertexFactory.cpp b/Engine/Source/Runtime/Engine/Private/LocalVertexFactory.cpp
--- a/Engine/Source/Runtime/Engine/Private/LocalVertexFactory.cpp
+++ b/Engine/Source/Runtime/Engine/Private/LocalVertexFactory.cpp
@@ -181,9 +181,13 @@ void FLocalVertexFactoryShaderParameters::GetElementShaderBindings(
  */
 bool FLocalVertexFactory::ShouldCompilePermutation(EShaderPlatform Platform, const class FMaterial* Material, const class FShaderType* ShaderType)
 {
-	if (Material->GetMaterialDomain() == MD_UI)
+	// Only compile this permutation inside the editor - it's not applicable in games, but occasionally the editor needs it.
+	if (!GIsEditor)
 	{
-		return false;
+		if (Material->GetMaterialDomain() == MD_UI)
+		{
+			return false;
+		}
 	}
 
 	return true; 
