Apply UObject::ProcessEvent's early outs in AActor::ProcessEvent for blueprint implementable events that have no script code.
#rnx

diff --git a/Engine/Source/Runtime/Engine/Private/Actor.cpp b/Engine/Source/Runtime/Engine/Private/Actor.cpp
--- a/Engine/Source/Runtime/Engine/Private/Actor.cpp
+++ b/Engine/Source/Runtime/Engine/Private/Actor.cpp
@@ -769,6 +769,18 @@ void AActor::ProcessEvent(UFunction* Function, void* Parameters)
 {
 	LLM_SCOPE(ELLMTag::EngineMisc);
 
+#if !UE_BUILD_SHIPPING
+	if (!ProcessEventDelegate.IsBound())
+#endif
+	{
+		// Apply UObject::ProcessEvent's early outs before doing any other work
+		// If the process event delegate is bound, we need to allow the process to play out
+		if ((Function->FunctionFlags & FUNC_Native) == 0 && (Function->Script.Num() == 0))
+		{
+			return;
+		}
+	}
+
 	#if WITH_EDITOR
 	static const FName CallInEditorMeta(TEXT("CallInEditor"));
 	const bool bAllowScriptExecution = GAllowActorScriptExecutionInEditor || Function->GetBoolMetaData(CallInEditorMeta);
