#UE4 SubsystemCollection Initialize now takes the outer so it can be initialized again after a deinitialize call which nulls the outer. Fixes a crash when a subsystemcollection gets deinitialized and reinitialized

diff --git a/Engine/Source/Editor/UnrealEd/Private/EditorEngine.cpp b/Engine/Source/Editor/UnrealEd/Private/EditorEngine.cpp
--- a/Engine/Source/Editor/UnrealEd/Private/EditorEngine.cpp
+++ b/Engine/Source/Editor/UnrealEd/Private/EditorEngine.cpp
@@ -310,7 +310,6 @@ UEditorEngine* GEditor = nullptr;
 
 UEditorEngine::UEditorEngine(const FObjectInitializer& ObjectInitializer)
 	: Super(ObjectInitializer)
-	, EditorSubsystemCollection(this)
 {
 	if (!IsRunningCommandlet() && !IsRunningDedicatedServer())
 	{
@@ -836,7 +835,7 @@ void UEditorEngine::HandleSettingChanged( FName Name )
 
 void UEditorEngine::InitializeObjectReferences()
 {
-	EditorSubsystemCollection.Initialize();
+	EditorSubsystemCollection.Initialize(this);
 
 	Super::InitializeObjectReferences();
 
diff --git a/Engine/Source/Runtime/Engine/Classes/Engine/LocalPlayer.h b/Engine/Source/Runtime/Engine/Classes/Engine/LocalPlayer.h
--- a/Engine/Source/Runtime/Engine/Classes/Engine/LocalPlayer.h
+++ b/Engine/Source/Runtime/Engine/Classes/Engine/LocalPlayer.h
@@ -170,7 +170,7 @@ class ENGINE_API ULocalPlayer : public UPlayer
 public:
 
 	/** DO NOT USE. This constructor is for internal usage only for hot-reload purposes. */
-	ULocalPlayer(FVTableHelper& Helper) : Super(Helper), SubsystemCollection(this), SlateOperations(FReply::Unhandled()) {}
+	ULocalPlayer(FVTableHelper& Helper) : Super(Helper), SlateOperations(FReply::Unhandled()) {}
 
 	/** The FUniqueNetId which this player is associated with. */
 	FUniqueNetIdRepl CachedUniqueNetId;
diff --git a/Engine/Source/Runtime/Engine/Private/GameEngine.cpp b/Engine/Source/Runtime/Engine/Private/GameEngine.cpp
--- a/Engine/Source/Runtime/Engine/Private/GameEngine.cpp
+++ b/Engine/Source/Runtime/Engine/Private/GameEngine.cpp
@@ -685,7 +685,6 @@ void UGameEngine::OnViewportResized(FViewport* Viewport, uint32 Unused)
 UEngine::UEngine(const FObjectInitializer& ObjectInitializer)
 	: Super(ObjectInitializer)
 	, ViewExtensions( new FSceneViewExtensions() )
-	, EngineSubsystemCollection(this)
 {
 	C_WorldBox = FColor(0, 0, 40, 255);
 	C_BrushWire = FColor(192, 0, 0, 255);
diff --git a/Engine/Source/Runtime/Engine/Private/GameInstance.cpp b/Engine/Source/Runtime/Engine/Private/GameInstance.cpp
--- a/Engine/Source/Runtime/Engine/Private/GameInstance.cpp
+++ b/Engine/Source/Runtime/Engine/Private/GameInstance.cpp
@@ -39,7 +39,6 @@ UGameInstance::UGameInstance(const FObjectInitializer& ObjectInitializer)
 	: Super(ObjectInitializer)
 	, TimerManager(new FTimerManager(this))
 	, LatentActionManager(new FLatentActionManager())
-	, SubsystemCollection(this)
 {
 }
 
@@ -97,7 +96,7 @@ void UGameInstance::Init()
 		FNetDelegates::OnReceivedNetworkEncryptionAck.BindUObject(this, &ThisClass::ReceivedNetworkEncryptionAck);
 	}
 
-	SubsystemCollection.Initialize();
+	SubsystemCollection.Initialize(this);
 }
 
 void UGameInstance::OnConsoleInput(const FString& Command)
diff --git a/Engine/Source/Runtime/Engine/Private/LocalPlayer.cpp b/Engine/Source/Runtime/Engine/Private/LocalPlayer.cpp
--- a/Engine/Source/Runtime/Engine/Private/LocalPlayer.cpp
+++ b/Engine/Source/Runtime/Engine/Private/LocalPlayer.cpp
@@ -227,7 +227,6 @@ bool FLocalPlayerContext::IsFromLocalPlayer(const AActor* ActorToTest) const
 
 ULocalPlayer::ULocalPlayer(const FObjectInitializer& ObjectInitializer)
 	: Super(ObjectInitializer)
-	, SubsystemCollection(this)
 	, SlateOperations( FReply::Unhandled() )
 {
 	PendingLevelPlayerControllerClass = APlayerController::StaticClass();
@@ -259,7 +258,7 @@ void ULocalPlayer::PlayerAdded(UGameViewportClient* InViewportClient, int32 InCo
 	ViewportClient = InViewportClient;
 	SetControllerId(InControllerID);
 
-	SubsystemCollection.Initialize();
+	SubsystemCollection.Initialize(this);
 }
 
 void ULocalPlayer::InitOnlineSession()
diff --git a/Engine/Source/Runtime/Engine/Private/Subsystems/SubsystemCollection.cpp b/Engine/Source/Runtime/Engine/Private/Subsystems/SubsystemCollection.cpp
--- a/Engine/Source/Runtime/Engine/Private/Subsystems/SubsystemCollection.cpp
+++ b/Engine/Source/Runtime/Engine/Private/Subsystems/SubsystemCollection.cpp
@@ -38,13 +38,12 @@ FSubsystemCollectionBase::FSubsystemCollectionBase()
 {
 }
 
-FSubsystemCollectionBase::FSubsystemCollectionBase(UObject* InOuter, TSubclassOf<USubsystem> InBaseType)
+FSubsystemCollectionBase::FSubsystemCollectionBase(TSubclassOf<USubsystem> InBaseType)
 	: BaseType(InBaseType)
-	, Outer(InOuter)
+	, Outer(nullptr)
 	, bPopulating(false)
 {
 	check(BaseType);
-	check(Outer);
 }
 
 USubsystem* FSubsystemCollectionBase::GetSubsystemInternal(TSubclassOf<USubsystem> SubsystemClass) const
@@ -89,11 +88,12 @@ const TArray<USubsystem*>& FSubsystemCollectionBase::GetSubsystemArrayInternal(T
 	return List;
 }
 
-void FSubsystemCollectionBase::Initialize()
+void FSubsystemCollectionBase::Initialize(UObject* NewOuter)
 {
+	Outer = NewOuter;
+	check(Outer);
 	if (ensure(BaseType) && ensureMsgf(SubsystemMap.Num() == 0, TEXT("Currently don't support repopulation of Subsystem Collections.")))
 	{
-		check(Outer);
 		check(!bPopulating); //Populating collections on multiple threads?
 		
 		if (SubsystemCollections.Num() == 0)
diff --git a/Engine/Source/Runtime/Engine/Private/UnrealEngine.cpp b/Engine/Source/Runtime/Engine/Private/UnrealEngine.cpp
--- a/Engine/Source/Runtime/Engine/Private/UnrealEngine.cpp
+++ b/Engine/Source/Runtime/Engine/Private/UnrealEngine.cpp
@@ -1399,7 +1399,7 @@ void UEngine::Init(IEngineLoop* InEngineLoop)
 	// Subsystems.
 	FURL::StaticInit();
 	FLinkerLoad::StaticInit(UTexture2D::StaticClass());
-	EngineSubsystemCollection.Initialize();
+	EngineSubsystemCollection.Initialize(this);
 
 #if !UE_BUILD_SHIPPING
 	// Check for overrides to the default map on the command line
diff --git a/Engine/Source/Runtime/Engine/Private/World.cpp b/Engine/Source/Runtime/Engine/Private/World.cpp
--- a/Engine/Source/Runtime/Engine/Private/World.cpp
+++ b/Engine/Source/Runtime/Engine/Private/World.cpp
@@ -352,7 +352,6 @@ UWorld::UWorld( const FObjectInitializer& ObjectInitializer )
 ,	TickTaskLevel(FTickTaskManagerInterface::Get().AllocateTickTaskLevel())
 ,	FlushLevelStreamingType(EFlushLevelStreamingType::None)
 ,	NextTravelType(TRAVEL_Relative)
-,	SubsystemCollection(this)
 {
 	TimerManager = new FTimerManager();
 #if WITH_EDITOR
@@ -1327,7 +1326,7 @@ void UWorld::InitWorld(const InitializationValues IVS)
 
 	if (WorldType == EWorldType::Game || WorldType == EWorldType::Editor || WorldType == EWorldType::PIE)
 	{
-		SubsystemCollection.Initialize();
+		SubsystemCollection.Initialize(this);
 	}
 
 	FWorldDelegates::OnPreWorldInitialization.Broadcast(this, IVS);
diff --git a/Engine/Source/Runtime/Engine/Public/Subsystems/SubsystemCollection.h b/Engine/Source/Runtime/Engine/Public/Subsystems/SubsystemCollection.h
--- a/Engine/Source/Runtime/Engine/Public/Subsystems/SubsystemCollection.h
+++ b/Engine/Source/Runtime/Engine/Public/Subsystems/SubsystemCollection.h
@@ -13,7 +13,7 @@ class ENGINE_API FSubsystemCollectionBase : public FGCObject
 {
 public:
 	/** Initialize the collection of systems, systems will be created and initialized */
-	void Initialize();
+	void Initialize(UObject* NewOuter);
 
 	/* Clears the collection, while deinitializing the systems */
 	void Deinitialize();
@@ -30,7 +30,7 @@ public:
 
 protected:
 	/** protected constructor - for use by the template only(FSubsystemCollection<TBaseType>) */
-	FSubsystemCollectionBase(UObject* InOuter, TSubclassOf<USubsystem> InBaseType);
+	FSubsystemCollectionBase(TSubclassOf<USubsystem> InBaseType);
 
 	/** protected constructor - Use the FSubsystemCollection<TBaseType> class */
 	FSubsystemCollectionBase();
@@ -99,12 +99,9 @@ public:
 public:
 
 	/** Construct a FSubsystemCollection, pass in the owning object almost certainly (this). */
-	FSubsystemCollection(UObject* InOuter)
-		: FSubsystemCollectionBase(InOuter, TBaseType::StaticClass())
+	FSubsystemCollection()
+		: FSubsystemCollectionBase(TBaseType::StaticClass())
 	{
 	}
-
-	/** DO NOT USE - required for default UObject default constructors unfortunately. */
-	FSubsystemCollection() {}
 };
 
