Additional compile warnings if the curve is not in the atlas
#rnx

diff --git a/Engine/Source/Runtime/Engine/Private/Materials/MaterialExpressions.cpp b/Engine/Source/Runtime/Engine/Private/Materials/MaterialExpressions.cpp
--- a/Engine/Source/Runtime/Engine/Private/Materials/MaterialExpressions.cpp
+++ b/Engine/Source/Runtime/Engine/Private/Materials/MaterialExpressions.cpp
@@ -16399,26 +16399,33 @@ int32 UMaterialExpressionCurveAtlasRowParameter::Compile(class FMaterialCompiler
 	{
 		// Retrieve the curve index directly from the atlas rather than relying on the scalar parameter defaults
 		int32 CurveIndex = 0;
-		Atlas->GetCurveIndex(Curve, CurveIndex);
-		DefaultValue = (float)CurveIndex;
-		int32 Slot = Compiler->ScalarParameter(ParameterName, DefaultValue);
+		
+		if (Atlas->GetCurveIndex(Curve, CurveIndex))
+		{
+			DefaultValue = (float)CurveIndex;
+			int32 Slot = Compiler->ScalarParameter(ParameterName, DefaultValue);
 
-		// Get Atlas texture object and texture size
-		int32 AtlasRef = INDEX_NONE;
-		int32 AtlasCode = Compiler->Texture(Atlas, AtlasRef, SAMPLERTYPE_LinearColor, SSM_Clamp_WorldGroupSettings, TMVM_None);
-		int32 AtlasSize = Compiler->ForceCast(Compiler->TextureProperty(AtlasCode, TMTM_TextureSize), MCT_Float1);
+			// Get Atlas texture object and texture size
+			int32 AtlasRef = INDEX_NONE;
+			int32 AtlasCode = Compiler->Texture(Atlas, AtlasRef, SAMPLERTYPE_LinearColor, SSM_Clamp_WorldGroupSettings, TMVM_None);
+			int32 AtlasSize = Compiler->ForceCast(Compiler->TextureProperty(AtlasCode, TMTM_TextureSize), MCT_Float1);
 
-		// Calculate UVs from size and slot
-		// if the input is hooked up, use it, otherwise use the internal constant
-		int32 Arg1 = InputTime.GetTracedInput().Expression ? InputTime.Compile(Compiler) : Compiler->Constant(0);
-		int32 Arg2 = Compiler->Div(Compiler->Add(Slot, Compiler->Constant(0.5)), AtlasSize);
+			// Calculate UVs from size and slot
+			// if the input is hooked up, use it, otherwise use the internal constant
+			int32 Arg1 = InputTime.GetTracedInput().Expression ? InputTime.Compile(Compiler) : Compiler->Constant(0);
+			int32 Arg2 = Compiler->Div(Compiler->Add(Slot, Compiler->Constant(0.5)), AtlasSize);
 
-		int32 UV = Compiler->AppendVector(Arg1, Arg2);
+			int32 UV = Compiler->AppendVector(Arg1, Arg2);
 
-		// Sample texture
-		return Compiler->TextureSample(AtlasCode, UV, SAMPLERTYPE_LinearColor, INDEX_NONE, INDEX_NONE, TMVM_None, SSM_Clamp_WorldGroupSettings, AtlasRef, false);
+			// Sample texture
+			return Compiler->TextureSample(AtlasCode, UV, SAMPLERTYPE_LinearColor, INDEX_NONE, INDEX_NONE, TMVM_None, SSM_Clamp_WorldGroupSettings, AtlasRef, false);
+		}
+		else
+		{
+			return CompilerError(Compiler, TEXT("The curve is not contained within the atlas."));
+		}
 	}
-	return INDEX_NONE;
+	return CompilerError(Compiler, TEXT("Either the curve or atlas is not currently set."));
 }
 
 void UMaterialExpressionCurveAtlasRowParameter::PostEditChangeProperty(FPropertyChangedEvent& PropertyChangedEvent)
