Gate on a class being native rather than compiled from blueprint to allow python to receive blueprint implementable events correctly
#rnx

diff --git a/Engine/Source/Runtime/CoreUObject/Private/UObject/SparseDelegate.cpp b/Engine/Source/Runtime/CoreUObject/Private/UObject/SparseDelegate.cpp
--- a/Engine/Source/Runtime/CoreUObject/Private/UObject/SparseDelegate.cpp
+++ b/Engine/Source/Runtime/CoreUObject/Private/UObject/SparseDelegate.cpp
@@ -57,7 +57,7 @@ FSparseDelegate* FSparseDelegateStorage::ResolveSparseDelegate(const UObject* Ow
 	const UClass* OwningClass = OwningObject->GetClass();
 	while (OwningClass)
 	{
-		if (!OwningClass->HasAnyClassFlags(CLASS_CompiledFromBlueprint))
+		if (OwningClass->HasAnyClassFlags(CLASS_Native))
 		{
 			if (size_t* DelegateOffset = SparseDelegateObjectOffsets.Find(TPair<FName, FName>(OwningClass->GetFName(), DelegateName)))
 			{
diff --git a/Engine/Source/Runtime/Engine/Private/Actor.cpp b/Engine/Source/Runtime/Engine/Private/Actor.cpp
--- a/Engine/Source/Runtime/Engine/Private/Actor.cpp
+++ b/Engine/Source/Runtime/Engine/Private/Actor.cpp
@@ -998,7 +998,7 @@ void AActor::TickActor( float DeltaSeconds, ELevelTick TickType, FActorTickFunct
 
 void AActor::Tick( float DeltaSeconds )
 {
-	if (GetClass()->HasAnyClassFlags(CLASS_CompiledFromBlueprint))
+	if (!GetClass()->HasAnyClassFlags(CLASS_Native))
 	{
 		// Blueprint code outside of the construction script should not run in the editor
 		// Allow tick if we are not a dedicated server, or we allow this tick on dedicated servers
diff --git a/Engine/Source/Runtime/Engine/Private/Components/ActorComponent.cpp b/Engine/Source/Runtime/Engine/Private/Components/ActorComponent.cpp
--- a/Engine/Source/Runtime/Engine/Private/Components/ActorComponent.cpp
+++ b/Engine/Source/Runtime/Engine/Private/Components/ActorComponent.cpp
@@ -804,7 +804,7 @@ void UActorComponent::BeginPlay()
 	check(!bHasBegunPlay);
 	checkSlow(bTickFunctionsRegistered); // If this fails, someone called BeginPlay() without first calling RegisterAllComponentTickFunctions().
 
-	if (GetClass()->HasAnyClassFlags(CLASS_CompiledFromBlueprint))
+	if (!GetClass()->HasAnyClassFlags(CLASS_Native))
 	{
 		ReceiveBeginPlay();
 	}
@@ -817,7 +817,7 @@ void UActorComponent::EndPlay(const EEndPlayReason::Type EndPlayReason)
 	check(bHasBegunPlay);
 
 	// If we're in the process of being garbage collected it is unsafe to call out to blueprints
-	if (!HasAnyFlags(RF_BeginDestroyed) && !IsUnreachable() && GetClass()->HasAnyClassFlags(CLASS_CompiledFromBlueprint))
+	if (!HasAnyFlags(RF_BeginDestroyed) && !IsUnreachable() && !GetClass()->HasAnyClassFlags(CLASS_Native))
 	{
 		ReceiveEndPlay(EndPlayReason);
 	}
@@ -957,7 +957,7 @@ void UActorComponent::TickComponent(float DeltaTime, enum ELevelTick TickType, F
 {
 	check(bRegistered);
 
-	if (GetClass()->HasAnyClassFlags(CLASS_CompiledFromBlueprint))
+	if (!GetClass()->HasAnyClassFlags(CLASS_Native))
 	{
 		ReceiveTick(DeltaTime);
 
