Fixing non-unity compile issue

Moving stats and log category declare to audio.h

diff --git a/Engine/Source/Runtime/Engine/Public/Audio.h b/Engine/Source/Runtime/Engine/Public/Audio.h
--- a/Engine/Source/Runtime/Engine/Public/Audio.h
+++ b/Engine/Source/Runtime/Engine/Public/Audio.h
@@ -17,6 +17,44 @@
 #include "Sound/SoundSourceBusSend.h"
 #include "IAudioExtensionPlugin.h"
 
+ENGINE_API DECLARE_LOG_CATEGORY_EXTERN(LogAudio, Warning, All);
+
+// Special log category used for temporary programmer debugging code of audio
+ENGINE_API DECLARE_LOG_CATEGORY_EXTERN(LogAudioDebug, Display, All);
+
+/**
+ * Audio stats
+ */
+DECLARE_DWORD_COUNTER_STAT_EXTERN(TEXT("Active Sounds"), STAT_ActiveSounds, STATGROUP_Audio, );
+DECLARE_CYCLE_STAT_EXTERN(TEXT("Audio Evaluate Concurrency"), STAT_AudioEvaluateConcurrency, STATGROUP_Audio, );
+DECLARE_DWORD_COUNTER_STAT_EXTERN(TEXT("Audio Sources"), STAT_AudioSources, STATGROUP_Audio, );
+DECLARE_DWORD_COUNTER_STAT_EXTERN(TEXT("Wave Instances"), STAT_WaveInstances, STATGROUP_Audio, );
+DECLARE_DWORD_COUNTER_STAT_EXTERN(TEXT("Wave Instances Dropped"), STAT_WavesDroppedDueToPriority, STATGROUP_Audio, );
+DECLARE_DWORD_COUNTER_STAT_EXTERN(TEXT("Virtualized Loops"), STAT_AudioVirtualLoops, STATGROUP_Audio, );
+DECLARE_DWORD_COUNTER_STAT_EXTERN(TEXT("Audible Wave Instances Dropped"), STAT_AudibleWavesDroppedDueToPriority, STATGROUP_Audio, );
+DECLARE_DWORD_COUNTER_STAT_EXTERN(TEXT("Max Channels"), STAT_AudioMaxChannels, STATGROUP_Audio, );
+DECLARE_DWORD_COUNTER_STAT_EXTERN(TEXT("Max Stopping Sources"), STAT_AudioMaxStoppingSources, STATGROUP_Audio, );
+DECLARE_DWORD_COUNTER_STAT_EXTERN(TEXT("Finished delegates called"), STAT_AudioFinishedDelegatesCalled, STATGROUP_Audio, );
+DECLARE_CYCLE_STAT_EXTERN(TEXT("Finished delegates time"), STAT_AudioFinishedDelegates, STATGROUP_Audio, );
+DECLARE_MEMORY_STAT_EXTERN(TEXT("Audio Memory Used"), STAT_AudioMemorySize, STATGROUP_Audio, );
+DECLARE_FLOAT_ACCUMULATOR_STAT_EXTERN(TEXT("Audio Buffer Time"), STAT_AudioBufferTime, STATGROUP_Audio, );
+DECLARE_FLOAT_ACCUMULATOR_STAT_EXTERN(TEXT("Audio Buffer Time (w/ Channels)"), STAT_AudioBufferTimeChannels, STATGROUP_Audio, );
+DECLARE_CYCLE_STAT_EXTERN(TEXT("Gathering WaveInstances"), STAT_AudioGatherWaveInstances, STATGROUP_Audio, );
+DECLARE_CYCLE_STAT_EXTERN(TEXT("Processing Sources"), STAT_AudioStartSources, STATGROUP_Audio, );
+DECLARE_CYCLE_STAT_EXTERN(TEXT("Updating Sources"), STAT_AudioUpdateSources, STATGROUP_Audio, ENGINE_API);
+DECLARE_CYCLE_STAT_EXTERN(TEXT("Updating Effects"), STAT_AudioUpdateEffects, STATGROUP_Audio, );
+DECLARE_CYCLE_STAT_EXTERN(TEXT("Source Init"), STAT_AudioSourceInitTime, STATGROUP_Audio, ENGINE_API);
+DECLARE_CYCLE_STAT_EXTERN(TEXT("Source Create"), STAT_AudioSourceCreateTime, STATGROUP_Audio, ENGINE_API);
+DECLARE_CYCLE_STAT_EXTERN(TEXT("Submit Buffers"), STAT_AudioSubmitBuffersTime, STATGROUP_Audio, ENGINE_API);
+DECLARE_CYCLE_STAT_EXTERN(TEXT("Decompress Audio"), STAT_AudioDecompressTime, STATGROUP_Audio, );
+DECLARE_CYCLE_STAT_EXTERN(TEXT("Decompress Vorbis"), STAT_VorbisDecompressTime, STATGROUP_Audio, );
+DECLARE_CYCLE_STAT_EXTERN(TEXT("Prepare Audio Decompression"), STAT_AudioPrepareDecompressionTime, STATGROUP_Audio, );
+DECLARE_CYCLE_STAT_EXTERN(TEXT("Prepare Vorbis Decompression"), STAT_VorbisPrepareDecompressionTime, STATGROUP_Audio, );
+DECLARE_CYCLE_STAT_EXTERN(TEXT("Finding Nearest Location"), STAT_AudioFindNearestLocation, STATGROUP_Audio, );
+DECLARE_CYCLE_STAT_EXTERN(TEXT("Decompress Streamed"), STAT_AudioStreamedDecompressTime, STATGROUP_Audio, );
+DECLARE_CYCLE_STAT_EXTERN(TEXT("Buffer Creation"), STAT_AudioResourceCreationTime, STATGROUP_Audio, );
+
+
 class FAudioDevice;
 class USoundNode;
 struct FSoundModulationControls;
diff --git a/Engine/Source/Runtime/Engine/Public/AudioDefines.h b/Engine/Source/Runtime/Engine/Public/AudioDefines.h
--- a/Engine/Source/Runtime/Engine/Public/AudioDefines.h
+++ b/Engine/Source/Runtime/Engine/Public/AudioDefines.h
@@ -8,11 +8,7 @@
 
 #include "CoreMinimal.h"
 
-ENGINE_API DECLARE_LOG_CATEGORY_EXTERN(LogAudio, Warning, All);
 
-// Special log category used for temporary programmer debugging code of audio
-ENGINE_API DECLARE_LOG_CATEGORY_EXTERN(LogAudioDebug, Display, All);
- 
 /** 
  * Maximum number of channels that can be set using the ini setting
  */
@@ -53,34 +49,3 @@ ENGINE_API DECLARE_LOG_CATEGORY_EXTERN(LogAudioDebug, Display, All);
 #define MIN_FILTER_BANDWIDTH			0.1f
 #define MAX_FILTER_BANDWIDTH			2.0f
 
-/**
- * Audio stats
- */
-DECLARE_DWORD_COUNTER_STAT_EXTERN( TEXT( "Active Sounds" ), STAT_ActiveSounds, STATGROUP_Audio , );
-DECLARE_CYCLE_STAT_EXTERN(TEXT("Audio Evaluate Concurrency"), STAT_AudioEvaluateConcurrency, STATGROUP_Audio, );
-DECLARE_DWORD_COUNTER_STAT_EXTERN( TEXT( "Audio Sources" ), STAT_AudioSources, STATGROUP_Audio , );
-DECLARE_DWORD_COUNTER_STAT_EXTERN( TEXT( "Wave Instances" ), STAT_WaveInstances, STATGROUP_Audio , );
-DECLARE_DWORD_COUNTER_STAT_EXTERN( TEXT( "Wave Instances Dropped" ), STAT_WavesDroppedDueToPriority, STATGROUP_Audio , );
-DECLARE_DWORD_COUNTER_STAT_EXTERN(TEXT("Virtualized Loops"), STAT_AudioVirtualLoops, STATGROUP_Audio, );
-DECLARE_DWORD_COUNTER_STAT_EXTERN( TEXT( "Audible Wave Instances Dropped" ), STAT_AudibleWavesDroppedDueToPriority, STATGROUP_Audio , );
-DECLARE_DWORD_COUNTER_STAT_EXTERN(TEXT("Max Channels"), STAT_AudioMaxChannels, STATGROUP_Audio, );
-DECLARE_DWORD_COUNTER_STAT_EXTERN(TEXT("Max Stopping Sources"), STAT_AudioMaxStoppingSources, STATGROUP_Audio, );
-DECLARE_DWORD_COUNTER_STAT_EXTERN( TEXT( "Finished delegates called" ), STAT_AudioFinishedDelegatesCalled, STATGROUP_Audio , );
-DECLARE_CYCLE_STAT_EXTERN( TEXT( "Finished delegates time" ), STAT_AudioFinishedDelegates, STATGROUP_Audio , );
-DECLARE_MEMORY_STAT_EXTERN( TEXT( "Audio Memory Used" ), STAT_AudioMemorySize, STATGROUP_Audio , );
-DECLARE_FLOAT_ACCUMULATOR_STAT_EXTERN( TEXT( "Audio Buffer Time" ), STAT_AudioBufferTime, STATGROUP_Audio , );
-DECLARE_FLOAT_ACCUMULATOR_STAT_EXTERN( TEXT( "Audio Buffer Time (w/ Channels)" ), STAT_AudioBufferTimeChannels, STATGROUP_Audio , );
-DECLARE_CYCLE_STAT_EXTERN( TEXT( "Gathering WaveInstances" ), STAT_AudioGatherWaveInstances, STATGROUP_Audio , );
-DECLARE_CYCLE_STAT_EXTERN( TEXT( "Processing Sources" ), STAT_AudioStartSources, STATGROUP_Audio , );
-DECLARE_CYCLE_STAT_EXTERN( TEXT( "Updating Sources" ), STAT_AudioUpdateSources, STATGROUP_Audio , ENGINE_API);
-DECLARE_CYCLE_STAT_EXTERN( TEXT( "Updating Effects" ), STAT_AudioUpdateEffects, STATGROUP_Audio , );
-DECLARE_CYCLE_STAT_EXTERN( TEXT( "Source Init" ), STAT_AudioSourceInitTime, STATGROUP_Audio , ENGINE_API);
-DECLARE_CYCLE_STAT_EXTERN( TEXT( "Source Create" ), STAT_AudioSourceCreateTime, STATGROUP_Audio , ENGINE_API);
-DECLARE_CYCLE_STAT_EXTERN( TEXT( "Submit Buffers" ), STAT_AudioSubmitBuffersTime, STATGROUP_Audio , ENGINE_API);
-DECLARE_CYCLE_STAT_EXTERN( TEXT( "Decompress Audio" ), STAT_AudioDecompressTime, STATGROUP_Audio , );
-DECLARE_CYCLE_STAT_EXTERN( TEXT( "Decompress Vorbis" ), STAT_VorbisDecompressTime, STATGROUP_Audio , );
-DECLARE_CYCLE_STAT_EXTERN( TEXT( "Prepare Audio Decompression" ), STAT_AudioPrepareDecompressionTime, STATGROUP_Audio , );
-DECLARE_CYCLE_STAT_EXTERN( TEXT( "Prepare Vorbis Decompression" ), STAT_VorbisPrepareDecompressionTime, STATGROUP_Audio , );
-DECLARE_CYCLE_STAT_EXTERN( TEXT( "Finding Nearest Location" ), STAT_AudioFindNearestLocation, STATGROUP_Audio , );
-DECLARE_CYCLE_STAT_EXTERN( TEXT( "Decompress Streamed" ), STAT_AudioStreamedDecompressTime, STATGROUP_Audio , );
-DECLARE_CYCLE_STAT_EXTERN( TEXT( "Buffer Creation" ), STAT_AudioResourceCreationTime, STATGROUP_Audio , );
