- Handle reimporting a static mesh in Cascade mesh type data

#rnx

diff --git a/Engine/Source/Runtime/Engine/Classes/Particles/TypeData/ParticleModuleTypeDataMesh.h b/Engine/Source/Runtime/Engine/Classes/Particles/TypeData/ParticleModuleTypeDataMesh.h
--- a/Engine/Source/Runtime/Engine/Classes/Particles/TypeData/ParticleModuleTypeDataMesh.h
+++ b/Engine/Source/Runtime/Engine/Classes/Particles/TypeData/ParticleModuleTypeDataMesh.h
@@ -199,12 +199,18 @@ class UParticleModuleTypeDataMesh : public UParticleModuleTypeDataBase
 	static int32 GetCurrentDetailMode();
 	static int32 GetMeshParticleMotionBlurMinDetailMode();
 
+#if WITH_EDITOR
+	void OnMeshChanged();
+#endif
+
 	virtual void PostLoad();
 
 	//~ Begin UObject Interface
 	virtual bool IsPostLoadThreadSafe() const override;
 #if WITH_EDITOR
-	virtual void	PostEditChangeProperty(FPropertyChangedEvent& PropertyChangedEvent) override;
+	virtual void BeginDestroy() override;
+	virtual void PreEditChange(UProperty* PropertyThatWillChange) override;
+	virtual void PostEditChangeProperty(FPropertyChangedEvent& PropertyChangedEvent) override;
 #endif // WITH_EDITOR
 
 	virtual void	Serialize(FArchive& Ar) override;
diff --git a/Engine/Source/Runtime/Engine/Private/Particles/ParticleModules.cpp b/Engine/Source/Runtime/Engine/Private/Particles/ParticleModules.cpp
--- a/Engine/Source/Runtime/Engine/Private/Particles/ParticleModules.cpp
+++ b/Engine/Source/Runtime/Engine/Private/Particles/ParticleModules.cpp
@@ -3389,6 +3389,13 @@ FParticleEmitterInstance* UParticleModuleTypeDataMesh::CreateInstance(UParticleE
 
 	CreateDistribution();
 
+#if WITH_EDITOR
+	if (GIsEditor && (Mesh != nullptr))
+	{
+		Mesh->GetOnMeshChanged().AddUObject(this, &UParticleModuleTypeDataMesh::OnMeshChanged);
+	}
+#endif
+
 	return Instance;
 }
 
@@ -3434,6 +3441,12 @@ void UParticleModuleTypeDataMesh::PostLoad()
 	if (Mesh != nullptr)
 	{
 		Mesh->ConditionalPostLoad();
+#if WITH_EDITOR
+		if ( GIsEditor )
+		{
+			Mesh->GetOnMeshChanged().AddUObject(this, &UParticleModuleTypeDataMesh::OnMeshChanged);
+		}
+#endif
 	}
 }
 
@@ -3443,6 +3456,49 @@ bool UParticleModuleTypeDataMesh::IsPostLoadThreadSafe() const
 }
 
 #if WITH_EDITOR
+void UParticleModuleTypeDataMesh::OnMeshChanged()
+{
+	UObject* OuterObj = GetOuter();
+	check(OuterObj);
+	UParticleLODLevel* LODLevel = Cast<UParticleLODLevel>(OuterObj);
+	if (LODLevel)
+	{
+		// The outer is incorrect - warn the user and handle it
+		UE_LOG(LogParticles, Warning, TEXT("UParticleModuleTypeDataMesh has an incorrect outer... run FixupEmitters on package %s"),
+			*(OuterObj->GetOutermost()->GetPathName()));
+		OuterObj = LODLevel->GetOuter();
+		UParticleEmitter* Emitter = Cast<UParticleEmitter>(OuterObj);
+		check(Emitter);
+		OuterObj = Emitter->GetOuter();
+	}
+
+	UProperty* MeshProperty = FindField<UProperty>(UParticleModuleTypeDataMesh::StaticClass(), FName(TEXT("Mesh")));
+	FPropertyChangedEvent PropertyChangedEvent(MeshProperty);
+
+	UParticleSystem* PartSys = CastChecked<UParticleSystem>(OuterObj);
+	PartSys->PostEditChangeProperty(PropertyChangedEvent);
+}
+
+void UParticleModuleTypeDataMesh::PreEditChange(UProperty* PropertyThatWillChange)
+{
+	if ( PropertyThatWillChange->GetFName() == FName(TEXT("Mesh")) )
+	{
+		if ( GIsEditor && (Mesh != nullptr) )
+		{
+			Mesh->GetOnMeshChanged().RemoveAll(this);
+		}
+	}
+}
+
+void UParticleModuleTypeDataMesh::BeginDestroy()
+{
+	Super::BeginDestroy();
+	if (GIsEditor && (Mesh != nullptr))
+	{
+		Mesh->GetOnMeshChanged().RemoveAll(this);
+	}
+}
+
 void UParticleModuleTypeDataMesh::PostEditChangeProperty(FPropertyChangedEvent& PropertyChangedEvent)
 {
 	UProperty* PropertyThatChanged = PropertyChangedEvent.Property;
@@ -3450,22 +3506,11 @@ void UParticleModuleTypeDataMesh::PostEditChangeProperty(FPropertyChangedEvent&
 	{
 		if (PropertyThatChanged->GetFName() == FName(TEXT("Mesh")))
 		{
-			UObject* OuterObj = GetOuter();
-			check(OuterObj);
-			UParticleLODLevel* LODLevel = Cast<UParticleLODLevel>(OuterObj);
-			if (LODLevel)
-			{
-				// The outer is incorrect - warn the user and handle it
-				UE_LOG(LogParticles, Warning, TEXT("UParticleModuleTypeDataMesh has an incorrect outer... run FixupEmitters on package %s"),
-					*(OuterObj->GetOutermost()->GetPathName()));
-				OuterObj = LODLevel->GetOuter();
-				UParticleEmitter* Emitter = Cast<UParticleEmitter>(OuterObj);
-				check(Emitter);
-				OuterObj = Emitter->GetOuter();
-			}
-			UParticleSystem* PartSys = CastChecked<UParticleSystem>(OuterObj);
-
-			PartSys->PostEditChangeProperty(PropertyChangedEvent);
+			OnMeshChanged();
+		}
+		if (GIsEditor && (Mesh != nullptr))
+		{
+			Mesh->GetOnMeshChanged().AddUObject(this, &UParticleModuleTypeDataMesh::OnMeshChanged);
 		}
 	}
 	Super::PostEditChangeProperty(PropertyChangedEvent);
