Fixed crash when nodes with cached bone references exist beyond a sub-instance in the anim graph traversal

Also fixed up deprecated code usage in graph traversal debug code.
Also addresses perf issues when spawning anim instances, where cache bones (etc) counters were being invalidated with each layer graph traversal.

diff --git a/Engine/Source/Runtime/Engine/Private/Animation/AnimInstance.cpp b/Engine/Source/Runtime/Engine/Private/Animation/AnimInstance.cpp
--- a/Engine/Source/Runtime/Engine/Private/Animation/AnimInstance.cpp
+++ b/Engine/Source/Runtime/Engine/Private/Animation/AnimInstance.cpp
@@ -2568,6 +2568,8 @@ void UAnimInstance::PerformLayerOverlayOperation(TSubclassOf<UAnimInstance> InCl
 			}
 		}
 
+		USkeletalMeshComponent* MeshComp = GetSkelMeshComponent();
+
 		for(TPair<FName, TArray<FAnimNode_Layer*, TInlineAllocator<4>>> LayerPair : LayerNodesToSet)
 		{
 			if (LayerPair.Key == NAME_None)
@@ -2585,7 +2587,6 @@ void UAnimInstance::PerformLayerOverlayOperation(TSubclassOf<UAnimInstance> InCl
 						// Skip setting if the class is the same
 						if (TargetInstance == nullptr || ClassToSet != TargetInstance->GetClass())
 						{
-							USkeletalMeshComponent* MeshComp = GetSkelMeshComponent();
 							UAnimInstance* NewSubInstance = NewObject<UAnimInstance>(MeshComp, ClassToSet);
 
 							LayerNode->SetLayerOverlaySubInstance(this, NewSubInstance);
@@ -2599,6 +2600,7 @@ void UAnimInstance::PerformLayerOverlayOperation(TSubclassOf<UAnimInstance> InCl
 								FAnimInstanceProxy& Proxy = NewSubInstance->GetProxyOnAnyThread<FAnimInstanceProxy>();
 								Proxy.InitializeObjects(NewSubInstance);
 								Proxy.InitializeRootNode_WithRoot(LayerNode->LinkedRoot);
+								Proxy.CacheBones_WithRoot(LayerNode->LinkedRoot);
 								Proxy.ClearObjects();
 							}
 
@@ -2627,7 +2629,6 @@ void UAnimInstance::PerformLayerOverlayOperation(TSubclassOf<UAnimInstance> InCl
 					if (TargetInstance == nullptr || ClassToSet != TargetInstance->GetClass())
 					{
 						// Create and add one sub-instance for this group
-						USkeletalMeshComponent* MeshComp = GetSkelMeshComponent();
 						UAnimInstance* NewSubInstance = NewObject<UAnimInstance>(MeshComp, ClassToSet);
 
 						for(FAnimNode_Layer* LayerNode : LayerPair.Value)
@@ -2647,6 +2648,7 @@ void UAnimInstance::PerformLayerOverlayOperation(TSubclassOf<UAnimInstance> InCl
 							{
 								Proxy.InitializeObjects(NewSubInstance);
 								Proxy.InitializeRootNode_WithRoot(LayerNode->LinkedRoot);
+								Proxy.CacheBones_WithRoot(LayerNode->LinkedRoot);
 								Proxy.ClearObjects();
 							}
 						}
@@ -2664,8 +2666,6 @@ void UAnimInstance::PerformLayerOverlayOperation(TSubclassOf<UAnimInstance> InCl
 				}
 			}
 		}
-
-		RecalcRequiredBones();
 	}
 }
 
diff --git a/Engine/Source/Runtime/Engine/Private/Animation/AnimInstanceProxy.cpp b/Engine/Source/Runtime/Engine/Private/Animation/AnimInstanceProxy.cpp
--- a/Engine/Source/Runtime/Engine/Private/Animation/AnimInstanceProxy.cpp
+++ b/Engine/Source/Runtime/Engine/Private/Animation/AnimInstanceProxy.cpp
@@ -43,9 +43,8 @@ void FAnimInstanceProxy::UpdateAnimationNode_WithRoot(float DeltaSeconds, FAnimN
 	DECLARE_SCOPE_HIERARCHICAL_COUNTER_FUNC()
 	if(InRootNode != nullptr)
 	{
-		if (FrameCounterForNodeUpdate != GFrameCounter)
+		if (InRootNode == RootNode)
 		{
-			FrameCounterForNodeUpdate = GFrameCounter;
 			UpdateCounter.Increment();
 		}
 		
@@ -265,7 +264,10 @@ void FAnimInstanceProxy::InitializeRootNode_WithRoot(FAnimNode_Base* InRootNode)
 
 	if (InRootNode != nullptr)
 	{
-		InitializationCounter.Increment();
+		if(InRootNode == RootNode)
+		{
+			InitializationCounter.Increment();
+		}
 		FAnimationInitializeContext InitContext(this);
 		InRootNode->Initialize_AnyThread(InitContext);
 	}
@@ -1281,7 +1283,10 @@ void FAnimInstanceProxy::CacheBones_WithRoot(FAnimNode_Base* InRootNode)
 	{
 		CacheBonesRecursionCounter++;
 
-		CachedBonesCounter.Increment();
+		if(InRootNode == RootNode)
+		{
+			CachedBonesCounter.Increment();
+		}
 		FAnimationCacheBonesContext Context(this);
 		InRootNode->CacheBones_AnyThread(Context);
 
@@ -1306,7 +1311,10 @@ void FAnimInstanceProxy::EvaluateAnimationNode_WithRoot(FPoseContext& Output, FA
 	if (InRootNode != NULL)
 	{
 		ANIM_MT_SCOPE_CYCLE_COUNTER(EvaluateAnimGraph, !IsInGameThread());
-		EvaluationCounter.Increment();
+		if(InRootNode == RootNode)
+		{
+			EvaluationCounter.Increment();
+		}
 		InRootNode->Evaluate_AnyThread(Output);
 	}
 	else
diff --git a/Engine/Source/Runtime/Engine/Private/Animation/AnimNodeBase.cpp b/Engine/Source/Runtime/Engine/Private/Animation/AnimNodeBase.cpp
--- a/Engine/Source/Runtime/Engine/Private/Animation/AnimNodeBase.cpp
+++ b/Engine/Source/Runtime/Engine/Private/Animation/AnimNodeBase.cpp
@@ -317,7 +317,7 @@ void FPoseLinkBase::Update(const FAnimationUpdateContext& Context)
 #endif
 
 #if ENABLE_ANIMGRAPH_TRAVERSAL_DEBUG
-	checkf(InitializationCounter.IsSynchronizedWith(Context.AnimInstanceProxy->GetInitializationCounter()), TEXT("Calling Update without initialization!"));
+	checkf(InitializationCounter.IsSynchronized_Counter(Context.AnimInstanceProxy->GetInitializationCounter()), TEXT("Calling Update without initialization!"));
 	UpdateCounter.SynchronizeWith(Context.AnimInstanceProxy->GetUpdateCounter());
 #endif
 
@@ -355,9 +355,9 @@ void FPoseLink::Evaluate(FPoseContext& Output)
 #endif
 
 #if ENABLE_ANIMGRAPH_TRAVERSAL_DEBUG
-	checkf(InitializationCounter.IsSynchronizedWith(Output.AnimInstanceProxy->GetInitializationCounter()), TEXT("Calling Evaluate without initialization!"));
-	checkf(UpdateCounter.IsSynchronizedWith(Output.AnimInstanceProxy->GetUpdateCounter()), TEXT("Calling Evaluate without Update for this node!"));
-	checkf(CachedBonesCounter.IsSynchronizedWith(Output.AnimInstanceProxy->GetCachedBonesCounter()), TEXT("Calling Evaluate without CachedBones!"));
+	checkf(InitializationCounter.IsSynchronized_Counter(Output.AnimInstanceProxy->GetInitializationCounter()), TEXT("Calling Evaluate without initialization!"));
+	checkf(UpdateCounter.IsSynchronized_Counter(Output.AnimInstanceProxy->GetUpdateCounter()), TEXT("Calling Evaluate without Update for this node!"));
+	checkf(CachedBonesCounter.IsSynchronized_Counter(Output.AnimInstanceProxy->GetCachedBonesCounter()), TEXT("Calling Evaluate without CachedBones!"));
 	EvaluationCounter.SynchronizeWith(Output.AnimInstanceProxy->GetEvaluationCounter());
 #endif
 
@@ -423,9 +423,9 @@ void FComponentSpacePoseLink::EvaluateComponentSpace(FComponentSpacePoseContext&
 #endif
 
 #if ENABLE_ANIMGRAPH_TRAVERSAL_DEBUG
-	checkf(InitializationCounter.IsSynchronizedWith(Output.AnimInstanceProxy->GetInitializationCounter()), TEXT("Calling EvaluateComponentSpace without initialization!"));
-	checkf(CachedBonesCounter.IsSynchronizedWith(Output.AnimInstanceProxy->GetCachedBonesCounter()), TEXT("Calling EvaluateComponentSpace without CachedBones!"));
-	checkf(UpdateCounter.IsSynchronizedWith(Output.AnimInstanceProxy->GetUpdateCounter()), TEXT("Calling EvaluateComponentSpace without Update for this node!"));
+	checkf(InitializationCounter.IsSynchronized_Counter(Output.AnimInstanceProxy->GetInitializationCounter()), TEXT("Calling EvaluateComponentSpace without initialization!"));
+	checkf(CachedBonesCounter.IsSynchronized_Counter(Output.AnimInstanceProxy->GetCachedBonesCounter()), TEXT("Calling EvaluateComponentSpace without CachedBones!"));
+	checkf(UpdateCounter.IsSynchronized_Counter(Output.AnimInstanceProxy->GetUpdateCounter()), TEXT("Calling EvaluateComponentSpace without Update for this node!"));
 	EvaluationCounter.SynchronizeWith(Output.AnimInstanceProxy->GetEvaluationCounter());
 #endif
 
diff --git a/Engine/Source/Runtime/Engine/Private/Animation/AnimNode_SubInstance.cpp b/Engine/Source/Runtime/Engine/Private/Animation/AnimNode_SubInstance.cpp
--- a/Engine/Source/Runtime/Engine/Private/Animation/AnimNode_SubInstance.cpp
+++ b/Engine/Source/Runtime/Engine/Private/Animation/AnimNode_SubInstance.cpp
@@ -20,8 +20,18 @@ void FAnimNode_SubInstance::Initialize_AnyThread(const FAnimationInitializeConte
 	if(InstanceToRun && LinkedRoot)
 	{
 		FAnimInstanceProxy& Proxy = InstanceToRun->GetProxyOnAnyThread<FAnimInstanceProxy>();
+		Proxy.InitializationCounter.SynchronizeWith(Context.AnimInstanceProxy->InitializationCounter);
 		Proxy.InitializeRootNode_WithRoot(LinkedRoot);
 	}
+	else
+	{
+		// If we have no valid instance (self or otherwise), we need to propagate down the graph to make sure
+		// subsequent nodes get properly initialized
+		for(FPoseLink& InputPose : InputPoses)
+		{
+			InputPose.Initialize(Context);
+		}
+	}
 }
 
 void FAnimNode_SubInstance::CacheBones_AnyThread(const FAnimationCacheBonesContext& Context)
@@ -30,7 +40,21 @@ void FAnimNode_SubInstance::CacheBones_AnyThread(const FAnimationCacheBonesConte
 	if(InstanceToRun && LinkedRoot)
 	{
 		FAnimInstanceProxy& Proxy = InstanceToRun->GetProxyOnAnyThread<FAnimInstanceProxy>();
-		Proxy.CacheBones_WithRoot(LinkedRoot);
+		Proxy.CachedBonesCounter.SynchronizeWith(Context.AnimInstanceProxy->CachedBonesCounter);
+
+		// Note not calling Proxy.CacheBones_WithRoot here as it is guarded by
+		// bBoneCachesInvalidated, which is handled at a higher level
+		FAnimationCacheBonesContext LinkedContext(&Proxy);
+		LinkedRoot->CacheBones_AnyThread(LinkedContext);
+	}
+	else
+	{
+		// If we have no valid instance (self or otherwise), we need to propagate down the graph to make sure
+		// subsequent nodes get properly their bones properly cached
+		for(FPoseLink& InputPose : InputPoses)
+		{
+			InputPose.CacheBones(Context);
+		}
 	}
 }
 
@@ -42,6 +66,7 @@ void FAnimNode_SubInstance::Update_AnyThread(const FAnimationUpdateContext& Cont
 	if(InstanceToRun && LinkedRoot)
 	{
 		FAnimInstanceProxy& Proxy = InstanceToRun->GetProxyOnAnyThread<FAnimInstanceProxy>();
+		Proxy.UpdateCounter.SynchronizeWith(Context.AnimInstanceProxy->UpdateCounter);
 
 		PropagateInputProperties(Context.AnimInstanceProxy->GetAnimInstanceObject());
 
@@ -51,6 +76,12 @@ void FAnimNode_SubInstance::Update_AnyThread(const FAnimationUpdateContext& Cont
 			Proxy.UpdateAnimation_WithRoot(LinkedRoot, GetDynamicLinkFunctionName());
 		}
 	}
+	else if(InputPoses.Num() > 0)
+	{
+		// If we have no valid instance (self or otherwise), we need to propagate down the graph to make sure
+		// subsequent nodes get properly updated
+		InputPoses[0].Update(Context);
+	}
 }
 
 void FAnimNode_SubInstance::Evaluate_AnyThread(FPoseContext& Output)
@@ -59,6 +90,7 @@ void FAnimNode_SubInstance::Evaluate_AnyThread(FPoseContext& Output)
 	if(InstanceToRun && LinkedRoot)
 	{
 		FAnimInstanceProxy& Proxy = InstanceToRun->GetProxyOnAnyThread<FAnimInstanceProxy>();
+		Proxy.EvaluationCounter.SynchronizeWith(Output.AnimInstanceProxy->EvaluationCounter);
 		Output.Pose.SetBoneContainer(&Proxy.GetRequiredBones());
 
 		// Create an evaluation context
@@ -72,6 +104,12 @@ void FAnimNode_SubInstance::Evaluate_AnyThread(FPoseContext& Output)
 		Output.Curve.MoveFrom(EvaluationContext.Curve);
 		Output.Pose.MoveBonesFrom(EvaluationContext.Pose);
 	}
+	else if(InputPoses.Num() > 0)
+	{
+		// If we have no valid instance (self or otherwise), we need to propagate down the graph to make sure
+		// subsequent nodes get properly evaluated
+		InputPoses[0].Evaluate(Output);
+	}
 	else
 	{
 		Output.ResetToRefPose();
@@ -93,6 +131,12 @@ void FAnimNode_SubInstance::GatherDebugData(FNodeDebugData& DebugData)
 		FAnimInstanceProxy& Proxy = InstanceToRun->GetProxyOnAnyThread<FAnimInstanceProxy>();
 		Proxy.GatherDebugData_WithRoot(DebugData.BranchFlow(1.0f), LinkedRoot, GetDynamicLinkFunctionName());
 	}
+	else if(InputPoses.Num() > 0)
+	{
+		// If we have no valid instance (self or otherwise), we need to propagate down the graph to make sure
+		// subsequent nodes get their debug data properly collected to reflect relevancy
+		InputPoses[0].GatherDebugData(DebugData);
+	}
 }
 
 void FAnimNode_SubInstance::OnInitializeAnimInstance(const FAnimInstanceProxy* InProxy, const UAnimInstance* InAnimInstance)
