Block Actors from being added to the Network List if they are being destroyed.

diff --git a/Engine/Source/Runtime/Engine/Private/World.cpp b/Engine/Source/Runtime/Engine/Private/World.cpp
--- a/Engine/Source/Runtime/Engine/Private/World.cpp
+++ b/Engine/Source/Runtime/Engine/Private/World.cpp
@@ -4437,19 +4437,19 @@ FConstCameraActorIterator UWorld::GetAutoActivateCameraIterator() const
 }
 
 
-void UWorld::AddNetworkActor( AActor* Actor )
+void UWorld::AddNetworkActor(AActor* Actor)
 {
-	if ( Actor == nullptr )
+	if (Actor == nullptr)
 	{
 		return;
 	}
 
-	if ( Actor->IsPendingKill() ) 
+	if (Actor->IsPendingKillPending())
 	{
 		return;
 	}
 
-	if ( !ContainsLevel(Actor->GetLevel()) )
+	if (!ContainsLevel(Actor->GetLevel()))
 	{
 		return;
 	}
