Make r.Shadow.MaxNumFarShadowCascades a scalability cvar so overriding it in scalability ini files works. Also fixes an ensure in autotests:
>Ensure: Scalability.ini can only set ECVF_Scalability console variables ('r.Shadow.MaxNumFarShadowCascades'='0' is ignored):

diff --git a/Engine/Source/Runtime/Engine/Private/Components/DirectionalLightComponent.cpp b/Engine/Source/Runtime/Engine/Private/Components/DirectionalLightComponent.cpp
--- a/Engine/Source/Runtime/Engine/Private/Components/DirectionalLightComponent.cpp
+++ b/Engine/Source/Runtime/Engine/Private/Components/DirectionalLightComponent.cpp
@@ -49,7 +49,7 @@ static TAutoConsoleVariable<int32> CVarMaxNumFarShadowCascades(
 	TEXT("r.Shadow.MaxNumFarShadowCascades"),
 	10,
 	TEXT("Max number of far shadow cascades that can be cast from a directional light"),
-	ECVF_RenderThreadSafe);
+	ECVF_RenderThreadSafe | ECVF_Scalability );
 
 /**
  * The scene info for a directional light.
