Added EditorSessionSummary with combines information present in the current Heartbeat, SessionStart and SessionEnd events and gets fired after each editor session ends. The addition of this event does not affect the existing events.

This event should simplify writing MTBF analytics queries in the future, since events need not be collated across the entire session's lifetime.

diff --git a/Engine/Source/Runtime/Engine/Private/Analytics/EditorSessionSummary.cpp b/Engine/Source/Runtime/Engine/Private/Analytics/EditorSessionSummary.cpp
--- a/Engine/Source/Runtime/Engine/Private/Analytics/EditorSessionSummary.cpp
+++ b/Engine/Source/Runtime/Engine/Private/Analytics/EditorSessionSummary.cpp
@@ -179,9 +179,10 @@ namespace EditorSessionSummaryUtils
 /* FEditorSessionSummaryWriter */
 
 FEditorSessionSummaryWriter::FEditorSessionSummaryWriter() :
-	bInitializedRecords(false)
-	, CurrentSession(nullptr)
+	CurrentSession(nullptr)
+	, LastUserInteractionTime(0.0f)
 	, HeartbeatTimeElapsed(0.0f)
+	, bInitializedRecords(false)
 	, bShutdown(false)
 {
 
@@ -249,7 +250,7 @@ void FEditorSessionSummaryWriter::UpdateTimestamps()
 	CurrentSession->Timestamp = FDateTime::UtcNow();
 	FPlatformMisc::SetStoredValue(SessionSummaryDefs::StoreId, CurrentSessionSectionName, SessionSummaryDefs::TimestampStoreKey, EditorSessionSummaryUtils::TimestampToString(CurrentSession->Timestamp));
 
-	const float IdleSeconds = FPlatformTime::Seconds() - FSlateApplication::Get().GetLastUserInteractionTime();
+	const float IdleSeconds = FPlatformTime::Seconds() - LastUserInteractionTime;
 
 	// 1 + 1 minutes
 	if (IdleSeconds > (60 + 60))
@@ -275,6 +276,14 @@ void FEditorSessionSummaryWriter::UpdateTimestamps()
 
 void FEditorSessionSummaryWriter::Tick(float DeltaTime)
 {
+	if (bShutdown)
+	{
+		return;
+	}
+
+	// cache the last user interaction time so that during a crash we have access to it
+	LastUserInteractionTime = FSlateApplication::Get().GetLastUserInteractionTime();
+
 	HeartbeatTimeElapsed += DeltaTime;
 
 	if (HeartbeatTimeElapsed > (float) SessionSummaryDefs::HeartbeatPeriodSeconds)
@@ -337,7 +346,7 @@ void FEditorSessionSummaryWriter::Shutdown()
 	FCoreDelegates::IsVanillaProductChanged.RemoveAll(this);
 
 	// Skip Slate if terminating, since we can't guarantee which thread called us.
-	if (GIsRequestingExit) 
+	if (!CurrentSession->bIsTerminating) 
 	{
 		FSlateApplication::Get().GetOnModalLoopTickEvent().RemoveAll(this);
 	}
diff --git a/Engine/Source/Runtime/Engine/Private/Analytics/EditorSessionSummary.h b/Engine/Source/Runtime/Engine/Private/Analytics/EditorSessionSummary.h
--- a/Engine/Source/Runtime/Engine/Private/Analytics/EditorSessionSummary.h
+++ b/Engine/Source/Runtime/Engine/Private/Analytics/EditorSessionSummary.h
@@ -32,10 +32,11 @@ private:
 	void UpdateTimestamps();
 
 private:
-	bool bInitializedRecords;
 	FEditorSessionRecord* CurrentSession;
 	FString CurrentSessionSectionName;
+	double LastUserInteractionTime;
 	float HeartbeatTimeElapsed;
+	bool bInitializedRecords;
 	bool bShutdown;
 };
 
