Special case Weak Object Ptrs and Special Object Ptrs in RepLayouts. This fixes issues with RepNotifies and property comparisons.

diff --git a/Engine/Source/Runtime/Engine/Private/RepLayout.cpp b/Engine/Source/Runtime/Engine/Private/RepLayout.cpp
--- a/Engine/Source/Runtime/Engine/Private/RepLayout.cpp
+++ b/Engine/Source/Runtime/Engine/Private/RepLayout.cpp
@@ -501,7 +501,6 @@ static FORCEINLINE bool CompareObject(
 	const void* 			A,
 	const void* 			B)
 {
-#if 1
 	// Until UObjectPropertyBase::Identical is made safe for GC'd objects, we need to do it manually
 	// This saves us from having to add referenced objects during GC
 	UObjectPropertyBase* ObjProperty = CastChecked<UObjectPropertyBase>(Cmd.Property);
@@ -510,9 +509,28 @@ static FORCEINLINE bool CompareObject(
 	UObject* ObjectB = ObjProperty->GetObjectPropertyValue(B);
 
 	return ObjectA == ObjectB;
-#else
+}
+
+static FORCEINLINE bool CompareSoftObject(
+	const FRepLayoutCmd& Cmd,
+	const void* A,
+	const void* B)
+{
+	// USoftObjectProperty::Identical will get the SoftObjectPath for each pointer, and compare the Path etc.
+	// It should also handle null checks, and doesn't try to dereference the object, so is GC safe.
 	return Cmd.Property->Identical(A, B);
-#endif
+}
+
+static FORCEINLINE bool CompareWeakObject(
+	const FRepLayoutCmd& Cmd,
+	const void* A,
+	const void* B)
+{
+	const UWeakObjectProperty* const WeakObjectProperty = CastChecked<UWeakObjectProperty>(Cmd.Property);
+	const FWeakObjectPtr ObjectA = WeakObjectProperty->GetPropertyValue(A);
+	const FWeakObjectPtr ObjectB = WeakObjectProperty->GetPropertyValue(B);
+
+	return ObjectA.HasSameIndexAndSerialNumber(ObjectB);
 }
 
 template<typename T>
@@ -541,6 +559,8 @@ static FORCEINLINE bool PropertiesAreIdenticalNative(
 		case ERepLayoutCmdType::PropertyInt:			return CompareValue<int32>(A, B);
 		case ERepLayoutCmdType::PropertyName:			return CompareValue<FName>(A, B);
 		case ERepLayoutCmdType::PropertyObject:			return CompareObject(Cmd, A, B);
+		case ERepLayoutCmdType::PropertySoftObject:		return CompareSoftObject(Cmd, A, B);
+		case ERepLayoutCmdType::PropertyWeakObject:		return CompareWeakObject(Cmd, A, B);
 		case ERepLayoutCmdType::PropertyUInt32:			return CompareValue<uint32>(A, B);
 		case ERepLayoutCmdType::PropertyUInt64:			return CompareValue<uint64>(A, B);
 		case ERepLayoutCmdType::PropertyVector:			return CompareValue<FVector>(A, B);
@@ -4344,7 +4364,9 @@ static bool DiffStableProperties_r(FDiffStablePropertiesSharedParams& Params, TD
 					continue;
 				}
 
-				if (Cmd.Type == ERepLayoutCmdType::PropertyObject)
+				if (Cmd.Type == ERepLayoutCmdType::PropertyObject ||
+					Cmd.Type == ERepLayoutCmdType::PropertyWeakObject ||
+					Cmd.Type == ERepLayoutCmdType::PropertySoftObject)
 				{
 					if (UObjectPropertyBase* ObjProperty = CastChecked<UObjectPropertyBase>(Cmd.Property))
 					{
@@ -4510,7 +4532,18 @@ static uint32 AddPropertyCmd(
 	}
 	else if (UnderlyingProperty->IsA(UObjectPropertyBase::StaticClass()))
 	{
-		Cmd.Type = ERepLayoutCmdType::PropertyObject;
+		if (UnderlyingProperty->IsA(USoftObjectProperty::StaticClass()))
+		{
+			Cmd.Type = ERepLayoutCmdType::PropertySoftObject;
+		}
+		else if (UnderlyingProperty->IsA(UWeakObjectProperty::StaticClass()))
+		{
+			Cmd.Type = ERepLayoutCmdType::PropertyWeakObject;
+		}
+		else
+		{
+			Cmd.Type = ERepLayoutCmdType::PropertyObject;
+		}
 	}
 	else if (UnderlyingProperty->IsA(UNameProperty::StaticClass()))
 	{
diff --git a/Engine/Source/Runtime/Engine/Public/Net/RepLayout.h b/Engine/Source/Runtime/Engine/Public/Net/RepLayout.h
--- a/Engine/Source/Runtime/Engine/Public/Net/RepLayout.h
+++ b/Engine/Source/Runtime/Engine/Public/Net/RepLayout.h
@@ -673,6 +673,8 @@ enum class ERepLayoutCmdType : uint8
 	PropertyString			= 19,
 	PropertyUInt64			= 20,
 	PropertyNativeBool		= 21,
+	PropertySoftObject		= 22,
+	PropertyWeakObject		= 23,
 };
 
 /** Various flags that describe how a Top Level Property should be handled. */
