Allow a GameInstance to be notified when the current world changes.

diff --git a/Engine/Source/Runtime/Engine/Classes/Engine/GameInstance.h b/Engine/Source/Runtime/Engine/Classes/Engine/GameInstance.h
--- a/Engine/Source/Runtime/Engine/Classes/Engine/GameInstance.h
+++ b/Engine/Source/Runtime/Engine/Classes/Engine/GameInstance.h
@@ -220,6 +220,9 @@ public:
 	struct FWorldContext* GetWorldContext() const { return WorldContext; };
 	class UGameViewportClient* GetGameViewportClient() const;
 
+	/** Callback from the world context when the world changes */
+	virtual void OnWorldChanged(UWorld* OldWorld, UWorld* NewWorld) {}
+
 	/** Starts the GameInstance state machine running */
 	virtual void StartGameInstance();
 	virtual bool JoinSession(ULocalPlayer* LocalPlayer, int32 SessionIndexInSearchResults) { return false; }
diff --git a/Engine/Source/Runtime/Engine/Private/UnrealEngine.cpp b/Engine/Source/Runtime/Engine/Private/UnrealEngine.cpp
--- a/Engine/Source/Runtime/Engine/Private/UnrealEngine.cpp
+++ b/Engine/Source/Runtime/Engine/Private/UnrealEngine.cpp
@@ -1046,7 +1046,13 @@ void FWorldContext::SetCurrentWorld(UWorld* World)
 		}
 	}
 
+	UWorld* OldWorld = ThisCurrentWorld;
 	ThisCurrentWorld = World;
+
+	if (OwningGameInstance)
+	{
+		OwningGameInstance->OnWorldChanged(OldWorld, ThisCurrentWorld);
+	}
 }
 
 void FWorldContext::AddReferencedObjects(FReferenceCollector& Collector, const UObject* ReferencingObject)
