Make sure FStaticMeshLODResources::Serialize supports loading cooked data in editor

diff --git a/Engine/Source/Runtime/Engine/Private/StaticMesh.cpp b/Engine/Source/Runtime/Engine/Private/StaticMesh.cpp
--- a/Engine/Source/Runtime/Engine/Private/StaticMesh.cpp
+++ b/Engine/Source/Runtime/Engine/Private/StaticMesh.cpp
@@ -601,6 +601,11 @@ void FStaticMeshLODResources::Serialize(FArchive& Ar, UObject* Owner, int32 Inde
 {
 	DECLARE_SCOPE_CYCLE_COUNTER(TEXT("FStaticMeshLODResources::Serialize"), STAT_StaticMeshLODResources_Serialize, STATGROUP_LoadTime);
 
+	bool bUsingCookedData = false;
+#if WITH_EDITORONLY_DATA
+	bUsingCookedData = Owner->GetOutermost()->bIsCookedForEditor;
+#endif
+
 	UStaticMesh* OwnerStaticMesh = Cast<UStaticMesh>(Owner);
 	// Actual flags used during serialization
 	const uint8 ClassDataStripFlags = GenerateClassStripFlags(Ar, OwnerStaticMesh, Index);
@@ -626,7 +631,7 @@ void FStaticMeshLODResources::Serialize(FArchive& Ar, UObject* Owner, int32 Inde
 			Ar << TmpBuffersSize;
 			BuffersSize = TmpBuffersSize.CalcBuffersSize();
 		}
-		else if (FPlatformProperties::RequiresCookedData() || Ar.IsCooking())
+		else if (FPlatformProperties::RequiresCookedData() || Ar.IsCooking() || bUsingCookedData)
 		{
 #if WITH_EDITOR
 			if (Ar.IsSaving())
