Early out if we do not get a valid index back when compiling a texture.
#rnx

diff --git a/Engine/Source/Runtime/Engine/Private/Materials/MaterialExpressions.cpp b/Engine/Source/Runtime/Engine/Private/Materials/MaterialExpressions.cpp
--- a/Engine/Source/Runtime/Engine/Private/Materials/MaterialExpressions.cpp
+++ b/Engine/Source/Runtime/Engine/Private/Materials/MaterialExpressions.cpp
@@ -16408,17 +16408,24 @@ int32 UMaterialExpressionCurveAtlasRowParameter::Compile(class FMaterialCompiler
 			// Get Atlas texture object and texture size
 			int32 AtlasRef = INDEX_NONE;
 			int32 AtlasCode = Compiler->Texture(Atlas, AtlasRef, SAMPLERTYPE_LinearColor, SSM_Clamp_WorldGroupSettings, TMVM_None);
-			int32 AtlasSize = Compiler->ForceCast(Compiler->TextureProperty(AtlasCode, TMTM_TextureSize), MCT_Float1);
+			if (AtlasCode != INDEX_NONE)
+			{
+				int32 AtlasSize = Compiler->ForceCast(Compiler->TextureProperty(AtlasCode, TMTM_TextureSize), MCT_Float1);
 
-			// Calculate UVs from size and slot
-			// if the input is hooked up, use it, otherwise use the internal constant
-			int32 Arg1 = InputTime.GetTracedInput().Expression ? InputTime.Compile(Compiler) : Compiler->Constant(0);
-			int32 Arg2 = Compiler->Div(Compiler->Add(Slot, Compiler->Constant(0.5)), AtlasSize);
+				// Calculate UVs from size and slot
+				// if the input is hooked up, use it, otherwise use the internal constant
+				int32 Arg1 = InputTime.GetTracedInput().Expression ? InputTime.Compile(Compiler) : Compiler->Constant(0);
+				int32 Arg2 = Compiler->Div(Compiler->Add(Slot, Compiler->Constant(0.5)), AtlasSize);
 
-			int32 UV = Compiler->AppendVector(Arg1, Arg2);
+				int32 UV = Compiler->AppendVector(Arg1, Arg2);
 
-			// Sample texture
-			return Compiler->TextureSample(AtlasCode, UV, SAMPLERTYPE_LinearColor, INDEX_NONE, INDEX_NONE, TMVM_None, SSM_Clamp_WorldGroupSettings, AtlasRef, false);
+				// Sample texture
+				return Compiler->TextureSample(AtlasCode, UV, SAMPLERTYPE_LinearColor, INDEX_NONE, INDEX_NONE, TMVM_None, SSM_Clamp_WorldGroupSettings, AtlasRef, false);
+			}
+			else
+			{
+				return CompilerError(Compiler, TEXT("There was an error when compiling the texture."));
+			}
 		}
 		else
 		{
