Fix crash in some platform tests due to data issues
Will remove when underlying data problem is fixed and cleaned up

diff --git a/Engine/Source/Runtime/Engine/Private/VT/RuntimeVirtualTexture.cpp b/Engine/Source/Runtime/Engine/Private/VT/RuntimeVirtualTexture.cpp
--- a/Engine/Source/Runtime/Engine/Private/VT/RuntimeVirtualTexture.cpp
+++ b/Engine/Source/Runtime/Engine/Private/VT/RuntimeVirtualTexture.cpp
@@ -500,7 +500,8 @@ IVirtualTexture* URuntimeVirtualTexture::CreateStreamingTextureProducer(IVirtual
 		FTexturePlatformData** StreamingTextureData = StreamingTexture->GetRunningPlatformData();
 		if (StreamingTextureData != nullptr && *StreamingTextureData != nullptr)
 		{
-			OutTransitionLevel = FMath::Max(InMaxLevel - GetStreamLowMips() + 1, 0);
+			int32 NumStreamMips = FMath::Min(GetStreamLowMips(), (*StreamingTextureData)->GetNumVTMips()); // Min() is a hack while we fix crash due to data breakage
+			OutTransitionLevel = FMath::Max(InMaxLevel - NumStreamMips + 1, 0);
 			IVirtualTexture* StreamingProducer = new FUploadingVirtualTexture((*StreamingTextureData)->VTData, 0);
 			return new FVirtualTextureLevelRedirector(InProducer, StreamingProducer, OutTransitionLevel);
 		}
