Made VisLog's device session ID reset with every new recording

#rnx

diff --git a/Engine/Source/Runtime/Engine/Private/VisualLogger/VisualLoggerBinaryFileDevice.cpp b/Engine/Source/Runtime/Engine/Private/VisualLogger/VisualLoggerBinaryFileDevice.cpp
--- a/Engine/Source/Runtime/Engine/Private/VisualLogger/VisualLoggerBinaryFileDevice.cpp
+++ b/Engine/Source/Runtime/Engine/Private/VisualLogger/VisualLoggerBinaryFileDevice.cpp
@@ -33,6 +33,9 @@ void FVisualLoggerBinaryFileDevice::StartRecordingToFile(float TimeStamp)
 		return;
 	}
 
+	// start new session
+	SessionGUID = FGuid::NewGuid();
+
 	StartRecordingTime = TimeStamp;
 	LastLogTimeStamp = StartRecordingTime;
 	TempFileName = FVisualLoggerHelpers::GenerateTemporaryFilename(VISLOG_FILENAME_EXT);
@@ -61,7 +64,7 @@ void FVisualLoggerBinaryFileDevice::StopRecordingToFile(float TimeStamp)
 	FileArchive = nullptr;
 
 	const FString TempFullFilename = FPaths::Combine(*FPaths::ProjectLogDir(), *TempFileName);
-	const FString NewFileName = FString::Printf(TEXT("%u_%s"), GetShortID(), *FVisualLoggerHelpers::GenerateFilename(TempFileName, FileName, StartRecordingTime, LastLogTimeStamp));
+	const FString NewFileName = FString::Printf(TEXT("%u_%s"), GetShortSessionID(), *FVisualLoggerHelpers::GenerateFilename(TempFileName, FileName, StartRecordingTime, LastLogTimeStamp));
 	const FString NewFullFileName = FPaths::Combine(*FPaths::ProjectLogDir(), *NewFileName);
 
 	if (TotalSize > 0)
diff --git a/Engine/Source/Runtime/Engine/Private/VisualLogger/VisualLoggerTypes.cpp b/Engine/Source/Runtime/Engine/Private/VisualLogger/VisualLoggerTypes.cpp
--- a/Engine/Source/Runtime/Engine/Private/VisualLogger/VisualLoggerTypes.cpp
+++ b/Engine/Source/Runtime/Engine/Private/VisualLogger/VisualLoggerTypes.cpp
@@ -20,14 +20,6 @@ UVisualLoggerDebugSnapshotInterface::UVisualLoggerDebugSnapshotInterface(const F
 #define DEPRECATED_VISUAL_LOGGER_MAGIC_NUMBER 0xFAFAAFAF
 #define VISUAL_LOGGER_MAGIC_NUMBER 0xAFAFFAFA
 
-//----------------------------------------------------------------------//
-// FVisualLogDevice 
-//----------------------------------------------------------------------//
-FVisualLogDevice::FVisualLogDevice()
-	: GUID(FGuid::NewGuid())
-{
-}
-
 //----------------------------------------------------------------------//
 // FVisualLogShapeElement 
 //----------------------------------------------------------------------//
diff --git a/Engine/Source/Runtime/Engine/Public/VisualLogger/VisualLoggerTypes.h b/Engine/Source/Runtime/Engine/Public/VisualLogger/VisualLoggerTypes.h
--- a/Engine/Source/Runtime/Engine/Public/VisualLogger/VisualLoggerTypes.h
+++ b/Engine/Source/Runtime/Engine/Public/VisualLogger/VisualLoggerTypes.h
@@ -266,7 +266,6 @@ public:
 		FVisualLogEntry Entry;
 	};
 
-	FVisualLogDevice();
 	virtual ~FVisualLogDevice() { }
 	virtual void Serialize(const UObject* LogOwner, FName OwnerName, FName InOwnerClassName, const FVisualLogEntry& LogEntry) = 0;
 	virtual void Cleanup(bool bReleaseMemory = false) { /* Empty */ }
@@ -276,9 +275,10 @@ public:
 	virtual void SetFileName(const FString& InFileName) { /* Empty */ }
 	virtual void GetRecordedLogs(TArray<FVisualLogDevice::FVisualLogEntryItem>& OutLogs)  const { /* Empty */ }
 	virtual bool HasFlags(int32 InFlags) const { return false; }
-	uint32 GetShortID() const { return GUID[0]; }
-
-	const FGuid GUID;
+	FGuid GetSessionGUID() const { return SessionGUID; }
+	uint32 GetShortSessionID() const { return SessionGUID[0]; }
+protected:
+	FGuid SessionGUID;
 };
 
 struct ENGINE_API FVisualLoggerCategoryVerbosityPair
