Remove LODBias in cook for virtual textures
Fixes issues with low mips build configurations breaking the VT

diff --git a/Engine/Source/Runtime/Engine/Private/TextureLODSettings.cpp b/Engine/Source/Runtime/Engine/Private/TextureLODSettings.cpp
--- a/Engine/Source/Runtime/Engine/Private/TextureLODSettings.cpp
+++ b/Engine/Source/Runtime/Engine/Private/TextureLODSettings.cpp
@@ -144,6 +144,12 @@ int32 UTextureLODSettings::CalculateLODBias(int32 Width, int32 Height, int32 Max
 		return 0;
 	}
 
+	// Ignore LODBias for virtual textures, as these are not restricted by any GPU max size
+	if (bVirtualTexture)
+	{
+		return 0;
+	}
+
 	// Calculate maximum number of miplevels.
 	if (MaxSize > 0)
 	{
@@ -166,9 +172,8 @@ int32 UTextureLODSettings::CalculateLODBias(int32 Width, int32 Height, int32 Max
 		UsedLODBias += GUITextureLODBias;
 	}
 
-	// Ignore the group's MaxLODMipCount for virtual textures, as these are not restricted by any GPU max size
 	int32 MinLOD		= LODGroupInfo.MinLODMipCount;
-	int32 MaxLOD		= bVirtualTexture ? 32 : LODGroupInfo.MaxLODMipCount;
+	int32 MaxLOD		= LODGroupInfo.MaxLODMipCount;
 	int32 WantedMaxLOD	= FMath::Clamp( TextureMaxLOD - UsedLODBias, MinLOD, MaxLOD );
 	WantedMaxLOD		= FMath::Clamp( WantedMaxLOD, 0, TextureMaxLOD );
 	UsedLODBias			= TextureMaxLOD - WantedMaxLOD;
