Force world subsystems to initialize if the world was inactive (due to loading in the content browser) and becoming the active world

diff --git a/Engine/Source/Editor/UnrealEd/Private/EditorServer.cpp b/Engine/Source/Editor/UnrealEd/Private/EditorServer.cpp
--- a/Engine/Source/Editor/UnrealEd/Private/EditorServer.cpp
+++ b/Engine/Source/Editor/UnrealEd/Private/EditorServer.cpp
@@ -2597,6 +2597,9 @@ bool UEditorEngine::Map_Load(const TCHAR* Str, FOutputDevice& Ar)
 					// Inactive worlds are already initialized but lack these two objects for memory reasons.
 					World->ClearWorldComponents();
 
+					// If the world was inactive subsystems would not have been initialized.  When we transition to the editor world initialize them
+					World->InitializeSubsystems();
+
 					if (World->FeatureLevel == FeatureLevelIndex)
 					{
 						if (World->GetPhysicsScene() == nullptr)
diff --git a/Engine/Source/Runtime/Engine/Classes/Engine/World.h b/Engine/Source/Runtime/Engine/Classes/Engine/World.h
--- a/Engine/Source/Runtime/Engine/Classes/Engine/World.h
+++ b/Engine/Source/Runtime/Engine/Classes/Engine/World.h
@@ -1261,6 +1261,9 @@ public:
 	/** Creates a new FX system for this world */
 	void CreateFXSystem();
 
+	/** Initialize all world subsystems */
+	void InitializeSubsystems();
+
 #if WITH_EDITOR
 
 	/** Change the feature level that this world is current rendering with */
diff --git a/Engine/Source/Runtime/Engine/Private/Subsystems/SubsystemCollection.cpp b/Engine/Source/Runtime/Engine/Private/Subsystems/SubsystemCollection.cpp
--- a/Engine/Source/Runtime/Engine/Private/Subsystems/SubsystemCollection.cpp
+++ b/Engine/Source/Runtime/Engine/Private/Subsystems/SubsystemCollection.cpp
@@ -90,6 +90,12 @@ const TArray<USubsystem*>& FSubsystemCollectionBase::GetSubsystemArrayInternal(T
 
 void FSubsystemCollectionBase::Initialize(UObject* NewOuter)
 {
+	if (Outer != nullptr)
+	{
+		// already initialized
+		return;
+	}
+
 	Outer = NewOuter;
 	check(Outer);
 	if (ensure(BaseType) && ensureMsgf(SubsystemMap.Num() == 0, TEXT("Currently don't support repopulation of Subsystem Collections.")))
diff --git a/Engine/Source/Runtime/Engine/Private/World.cpp b/Engine/Source/Runtime/Engine/Private/World.cpp
--- a/Engine/Source/Runtime/Engine/Private/World.cpp
+++ b/Engine/Source/Runtime/Engine/Private/World.cpp
@@ -1346,10 +1346,7 @@ void UWorld::InitWorld(const InitializationValues IVS)
 		return;
 	}
 
-	if (WorldType == EWorldType::Game || WorldType == EWorldType::Editor || WorldType == EWorldType::PIE)
-	{
-		SubsystemCollection.Initialize(this);
-	}
+	InitializeSubsystems();
 
 	FWorldDelegates::OnPreWorldInitialization.Broadcast(this, IVS);
 
@@ -7531,6 +7528,14 @@ void UWorld::RemovePostProcessVolume(IInterface_PostProcessVolume* InVolume)
 	PostProcessVolumes.RemoveSingle(InVolume);
 }
 
+void UWorld::InitializeSubsystems()
+{
+	if (WorldType == EWorldType::Game || WorldType == EWorldType::Editor || WorldType == EWorldType::PIE)
+	{
+		SubsystemCollection.Initialize(this);
+	}
+}
+
 /**
 * Dump visible actors in current world.
 */
