Correct fix. RVT static meshes weren't created with forced LOD settings.
Put check() back in.

diff --git a/Engine/Source/Runtime/Engine/Private/StaticMeshRender.cpp b/Engine/Source/Runtime/Engine/Private/StaticMeshRender.cpp
--- a/Engine/Source/Runtime/Engine/Private/StaticMeshRender.cpp
+++ b/Engine/Source/Runtime/Engine/Private/StaticMeshRender.cpp
@@ -906,6 +906,16 @@ inline bool AllowShadowOnlyMesh(ERHIFeatureLevel::Type InFeatureLevel)
 	return !UseLightPropagationVolumeRT2(InFeatureLevel);
 }
 
+inline void SetupMeshBatchForRuntimeVirtualTexture(FMeshBatch& MeshBatch)
+{
+	MeshBatch.CastShadow = 0;
+	MeshBatch.bUseAsOccluder = 0;
+	MeshBatch.bUseForDepthPass = 0;
+	MeshBatch.bUseForMaterial = 0;
+	MeshBatch.bDitheredLODTransition = 0;
+	MeshBatch.bRenderToVirtualTexture = 1;
+}
+
 void FStaticMeshSceneProxy::DrawStaticElements(FStaticPrimitiveDrawInterface* PDI)
 {
 	checkSlow(IsInParallelRenderingThread());
@@ -918,6 +928,7 @@ void FStaticMeshSceneProxy::DrawStaticElements(FStaticPrimitiveDrawInterface* PD
 		bool bIsMeshElementSelected = false;
 		const auto FeatureLevel = GetScene().GetFeatureLevel();
 		const bool IsMobile = IsMobilePlatform(GetScene().GetShaderPlatform());
+		const int32 NumRuntimeVirtualTextureTypes = RuntimeVirtualTextureMaterialTypes.Num();
 
 		//check if a LOD is being forced
 		if (ForcedLodModel > 0) 
@@ -938,8 +949,7 @@ void FStaticMeshSceneProxy::DrawStaticElements(FStaticPrimitiveDrawInterface* PD
 #endif // WITH_EDITOR
 
 				const int32 NumBatches = GetNumMeshBatches();
-
-				PDI->ReserveMemoryForMeshes(NumBatches);
+				PDI->ReserveMemoryForMeshes(NumBatches * (1 + NumRuntimeVirtualTextureTypes));
 
 				for (int32 BatchIndex = 0; BatchIndex < NumBatches; BatchIndex++)
 				{
@@ -948,6 +958,17 @@ void FStaticMeshSceneProxy::DrawStaticElements(FStaticPrimitiveDrawInterface* PD
 					if (GetMeshElement(LODIndex, BatchIndex, SectionIndex, PrimitiveDPG, bIsMeshElementSelected, true, MeshBatch))
 					{
 						PDI->DrawMesh(MeshBatch, FLT_MAX);
+
+						if (NumRuntimeVirtualTextureTypes > 0)
+						{
+							// Runtime virtual texture mesh elements.
+							SetupMeshBatchForRuntimeVirtualTexture(MeshBatch);
+							for (ERuntimeVirtualTextureMaterialType MaterialType : RuntimeVirtualTextureMaterialTypes)
+							{
+								MeshBatch.RuntimeVirtualTextureMaterialType = (uint32)MaterialType;
+								PDI->DrawMesh(MeshBatch, FLT_MAX);
+							}
+						}
 					}
 				}
 			}
@@ -1041,8 +1062,6 @@ void FStaticMeshSceneProxy::DrawStaticElements(FStaticPrimitiveDrawInterface* PD
 #endif // WITH_EDITOR
 
 					const int32 NumBatches = GetNumMeshBatches();
-					const int32 NumRuntimeVirtualTextureTypes = RuntimeVirtualTextureMaterialTypes.Num();
-
 					PDI->ReserveMemoryForMeshes(NumBatches * (1 + NumRuntimeVirtualTextureTypes));
 
 					for (int32 BatchIndex = 0; BatchIndex < NumBatches; BatchIndex++)
@@ -1063,12 +1082,7 @@ void FStaticMeshSceneProxy::DrawStaticElements(FStaticPrimitiveDrawInterface* PD
 							{
 								// Runtime virtual texture mesh elements.
 								FMeshBatch MeshBatch(BaseMeshBatch);
-								MeshBatch.CastShadow = 0;
-								MeshBatch.bUseAsOccluder = 0;
-								MeshBatch.bUseForDepthPass = 0;
-								MeshBatch.bUseForMaterial = 0;
-								MeshBatch.bDitheredLODTransition = 0;
-								MeshBatch.bRenderToVirtualTexture = 1;
+								SetupMeshBatchForRuntimeVirtualTexture(MeshBatch);
 
 								for (ERuntimeVirtualTextureMaterialType MaterialType : RuntimeVirtualTextureMaterialTypes)
 								{
diff --git a/Engine/Source/Runtime/Renderer/Private/PrimitiveSceneInfo.cpp b/Engine/Source/Runtime/Renderer/Private/PrimitiveSceneInfo.cpp
--- a/Engine/Source/Runtime/Renderer/Private/PrimitiveSceneInfo.cpp
+++ b/Engine/Source/Runtime/Renderer/Private/PrimitiveSceneInfo.cpp
@@ -463,7 +463,7 @@ static void GetRuntimeVirtualTextureLODRange(TArray<class FStaticMeshBatchReleva
 		}
 	}
 
-	//check(OutMinLOD <= OutMaxLOD);
+	check(OutMinLOD <= OutMaxLOD);
 }
 
 int32 FPrimitiveSceneInfo::UpdateStaticLightingBuffer()
