#ue4 - Fix Actor overlap events triggering on the same actor from sub components overlapping at startup. Fix overlap events triggering while the level is being loaded, except where allowed explicitly by the actor overlap flags.

(Dupe 8671782 from Dev-Framework)

diff --git a/Engine/Source/Runtime/Engine/Private/Components/PrimitiveComponent.cpp b/Engine/Source/Runtime/Engine/Private/Components/PrimitiveComponent.cpp
--- a/Engine/Source/Runtime/Engine/Private/Components/PrimitiveComponent.cpp
+++ b/Engine/Source/Runtime/Engine/Private/Components/PrimitiveComponent.cpp
@@ -2567,13 +2567,15 @@ void UPrimitiveComponent::BeginComponentOverlap(const FOverlapInfo& OtherOverlap
 			AActor* const OtherActor = OtherComp->GetOwner();
 			AActor* const MyActor = GetOwner();
 
-			const bool bNotifyActorTouch = bDoNotifies && !AreActorsOverlapping(*MyActor, *OtherActor);
+			const bool bSameActor = (MyActor == OtherActor);
+			const bool bNotifyActorTouch = bDoNotifies && !bSameActor && !AreActorsOverlapping(*MyActor, *OtherActor);
 
 			// Perform reflexive touch.
 			OverlappingComponents.Add(OtherOverlap);												// already verified uniqueness above
 			AddUniqueOverlapFast(OtherComp->OverlappingComponents, FOverlapInfo(this, INDEX_NONE));	// uniqueness unverified, so addunique
 			
-			if (bDoNotifies)
+			const UWorld* World = GetWorld();
+			if (bDoNotifies && World && World->HasBegunPlay())
 			{
 				// first execute component delegates
 				if (!IsPendingKill())
@@ -2641,7 +2643,8 @@ void UPrimitiveComponent::EndComponentOverlap(const FOverlapInfo& OtherOverlap,
 		GlobalOverlapEventsCounter++;
 		OverlappingComponents.RemoveAtSwap(OverlapIdx, 1, false);
 
-		if (bDoNotifies)
+		const UWorld* World = GetWorld();
+		if (bDoNotifies && World && World->HasBegunPlay())
 		{
 			AActor* const OtherActor = OtherComp->GetOwner();
 			AActor* const MyActor = GetOwner();
@@ -2658,7 +2661,8 @@ void UPrimitiveComponent::EndComponentOverlap(const FOverlapInfo& OtherOverlap,
 				}
 	
 				// if this was the last touch on the other actor by this actor, notify that we've untouched the actor as well
-				if (MyActor && !AreActorsOverlapping(*MyActor, *OtherActor))
+				const bool bSameActor = (MyActor == OtherActor);
+				if (MyActor && !bSameActor && !AreActorsOverlapping(*MyActor, *OtherActor))
 				{			
 					if (IsActorValidToNotify(MyActor))
 					{
diff --git a/Engine/Source/Runtime/Engine/Private/WorldSettings.cpp b/Engine/Source/Runtime/Engine/Private/WorldSettings.cpp
--- a/Engine/Source/Runtime/Engine/Private/WorldSettings.cpp
+++ b/Engine/Source/Runtime/Engine/Private/WorldSettings.cpp
@@ -242,7 +242,8 @@ void AWorldSettings::NotifyBeginPlay()
 		for (FActorIterator It(World); It; ++It)
 		{
 			SCOPE_CYCLE_COUNTER(STAT_ActorBeginPlay);
-			It->DispatchBeginPlay();
+			const bool bFromLevelLoad = true;
+			It->DispatchBeginPlay(bFromLevelLoad);
 		}
 		World->bBegunPlay = true;
 	}
