Spawning optimizations

A collection of major and minor optimizations to anim instance initialization.

Major:
- Added defered root node init. When an anim instance is set to init to refpose, this defers the first call to initialize the blend tree until just before the first Update() call, on a worker thread.
- Skip ConditionallyCacheBonesForState in state machines in the case where we would call CacheBones_AnyThread during first Initialize_AnyThread. This avoids duplicate work being done, as graph traversal counters will let through sychronization calls if they have yet to be updated (ie. their count is INDEX_NONE).
- Cache access to nodes that need pre-update, dynamic reset, one-time initialization, as well as state machines in the generated class to avoid iterating all anim nodes on initialization.

Minor:
- Cache access to node properties when re-linking nodes to avoid duplicate memory round-trips.
- Cache access to LUT size in FBaseBlendedCurve::GetValidElementCount, avoiding duplicate memory round-trips.
- Prevent FName rehashing from string (intended to strip whitespace) when initializing FBoneReferences.
- Cache access to AnimCurveMapping smart name container on USkeleton, preventing a TMap lookup every time we access the smart name container.
- USkeleton::GetCurveMetaData now uses a 'fast path' when accessing using smart name UID. In this case we can avoid the TMap lookup as the UID is always a direct index to metadata at runtime (slow path still applies in editor builds).
- Avoid calling FAnimNode_LayeredBoneBlend::ReinitializeBoneBlendWeights in Initialize_AnyThread, as CacheBones_AnyThread will be called straight after anyways.
- Cache access to curve UID count and LinkedBones array in FAnimNode_LayeredBoneBlend::ReinitializeBoneBlendWeights, avoiding duplicate memory round-trips.

Upgrade notes:
- Breaking change! Custom anim nodes that implement OnInitializeAnimInstance now need to return true for the new accessor NeedsOnInitializeAnimInstance to be called at init time.
- The meaning of HasPreUpdate and NeedsDynamicReset has changed slightly. Previously these were called at init time, so if their results were mutable it was possible to re-init an anim instance with new settings. Now they are called at CDO load time this is no longer possible.

diff --git a/Engine/Plugins/Experimental/ControlRig/Source/ControlRig/Public/AnimNode_ControlRigBase.h b/Engine/Plugins/Experimental/ControlRig/Source/ControlRig/Public/AnimNode_ControlRigBase.h
--- a/Engine/Plugins/Experimental/ControlRig/Source/ControlRig/Public/AnimNode_ControlRigBase.h
+++ b/Engine/Plugins/Experimental/ControlRig/Source/ControlRig/Public/AnimNode_ControlRigBase.h
@@ -22,6 +22,7 @@ struct CONTROLRIG_API FAnimNode_ControlRigBase : public FAnimNode_CustomProperty
 
 	// FAnimNode_Base interface
 	virtual void OnInitializeAnimInstance(const FAnimInstanceProxy* InProxy, const UAnimInstance* InAnimInstance) override;
+	virtual bool NeedsOnInitializeAnimInstance() const override { return true; }
 	virtual void Initialize_AnyThread(const FAnimationInitializeContext& Context) override;
 	virtual void GatherDebugData(FNodeDebugData& DebugData) override;
 	virtual void Update_AnyThread(const FAnimationUpdateContext& Context) override;
diff --git a/Engine/Source/Runtime/AnimGraphRuntime/Private/AnimNodes/AnimNode_BlendListBase.cpp b/Engine/Source/Runtime/AnimGraphRuntime/Private/AnimNodes/AnimNode_BlendListBase.cpp
--- a/Engine/Source/Runtime/AnimGraphRuntime/Private/AnimNodes/AnimNode_BlendListBase.cpp
+++ b/Engine/Source/Runtime/AnimGraphRuntime/Private/AnimNodes/AnimNode_BlendListBase.cpp
@@ -12,7 +12,7 @@ void FAnimNode_BlendListBase::Initialize_AnyThread(const FAnimationInitializeCon
 	DECLARE_SCOPE_HIERARCHICAL_COUNTER_ANIMNODE(Initialize_AnyThread)
 	FAnimNode_Base::Initialize_AnyThread(Context);
 
-	const int NumPoses = BlendPose.Num();
+	const int32 NumPoses = BlendPose.Num();
 	checkSlow(BlendTime.Num() == NumPoses);
 
 	BlendWeights.Reset(NumPoses);
@@ -78,7 +78,7 @@ void FAnimNode_BlendListBase::Update_AnyThread(const FAnimationUpdateContext& Co
 	DECLARE_SCOPE_HIERARCHICAL_COUNTER_ANIMNODE(Update_AnyThread)
 	GetEvaluateGraphExposedInputs().Execute(Context);
 
-	const int NumPoses = BlendPose.Num();
+	const int32 NumPoses = BlendPose.Num();
 	checkSlow((BlendTime.Num() == NumPoses) && (BlendWeights.Num() == NumPoses));
 
 	PosesToEvaluate.Empty(NumPoses);
@@ -252,7 +252,7 @@ void FAnimNode_BlendListBase::Evaluate_AnyThread(FPoseContext& Output)
 void FAnimNode_BlendListBase::GatherDebugData(FNodeDebugData& DebugData)
 {
 	DECLARE_SCOPE_HIERARCHICAL_COUNTER_ANIMNODE(GatherDebugData)
-	const int NumPoses = BlendPose.Num();
+	const int32 NumPoses = BlendPose.Num();
 	const int32 ChildIndex = GetActiveChildIndex();
 
 	FString DebugLine = GetNodeName(DebugData);
diff --git a/Engine/Source/Runtime/AnimGraphRuntime/Private/AnimNodes/AnimNode_LayeredBoneBlend.cpp b/Engine/Source/Runtime/AnimGraphRuntime/Private/AnimNodes/AnimNode_LayeredBoneBlend.cpp
--- a/Engine/Source/Runtime/AnimGraphRuntime/Private/AnimNodes/AnimNode_LayeredBoneBlend.cpp
+++ b/Engine/Source/Runtime/AnimGraphRuntime/Private/AnimNodes/AnimNode_LayeredBoneBlend.cpp
@@ -25,10 +25,6 @@ void FAnimNode_LayeredBoneBlend::Initialize_AnyThread(const FAnimationInitialize
 		{
 			BlendPoses[ChildIndex].Initialize(Context);
 		}
-
-		// initialize mask weight now
-		check (Context.AnimInstanceProxy->GetSkeleton());
-		ReinitializeBoneBlendWeights(Context.AnimInstanceProxy->GetRequiredBones(), Context.AnimInstanceProxy->GetSkeleton());
 	}
 }
 
@@ -82,13 +78,14 @@ void FAnimNode_LayeredBoneBlend::ReinitializeBoneBlendWeights(const FBoneContain
 	FAnimationRuntime::UpdateDesiredBoneWeight(DesiredBoneBlendWeights, CurrentBoneBlendWeights, BlendWeights);
 
 	TArray<uint16> const& CurveUIDFinder = RequiredBones.GetUIDToArrayLookupTable();
+	const int32 CurveUIDCount = CurveUIDFinder.Num();
 	const int32 TotalCount = FBlendedCurve::GetValidElementCount(&CurveUIDFinder);
 	CurvePoseSourceIndices.Reset(TotalCount);
 	// initialize with FF - which is default
 	CurvePoseSourceIndices.Init(DEFAULT_SOURCEINDEX, TotalCount);
 
 	// now go through point to correct source indices. Curve only picks one source index
-	for (int32 UIDIndex = 0; UIDIndex < CurveUIDFinder.Num(); ++UIDIndex)
+	for (int32 UIDIndex = 0; UIDIndex < CurveUIDCount; ++UIDIndex)
 	{
 		int32 CurrentPoseIndex = CurveUIDFinder[UIDIndex];
 		if (CurrentPoseIndex != MAX_uint16)
@@ -96,11 +93,13 @@ void FAnimNode_LayeredBoneBlend::ReinitializeBoneBlendWeights(const FBoneContain
 			SmartName::UID_Type CurveUID = (SmartName::UID_Type)UIDIndex;
 
 			const FCurveMetaData* CurveMetaData = Skeleton->GetCurveMetaData(CurveUID);
-			if (CurveMetaData && CurveMetaData->LinkedBones.Num() > 0)
+			if (CurveMetaData)
 			{
-				for (int32 LinkedBoneIndex = 0; LinkedBoneIndex < CurveMetaData->LinkedBones.Num(); ++LinkedBoneIndex)
+				const TArray<FBoneReference>& LinkedBones = CurveMetaData->LinkedBones;
+				const int32 NumLinkedBones = LinkedBones.Num();
+				for (int32 LinkedBoneIndex = 0; LinkedBoneIndex < NumLinkedBones; ++LinkedBoneIndex)
 				{
-					FCompactPoseBoneIndex CompactPoseIndex = CurveMetaData->LinkedBones[LinkedBoneIndex].GetCompactPoseIndex(RequiredBones);
+					FCompactPoseBoneIndex CompactPoseIndex = LinkedBones[LinkedBoneIndex].GetCompactPoseIndex(RequiredBones);
 					if (CompactPoseIndex != INDEX_NONE)
 					{
 						if (DesiredBoneBlendWeights[CompactPoseIndex.GetInt()].BlendWeight > 0.f)
diff --git a/Engine/Source/Runtime/AnimGraphRuntime/Public/AnimNodes/AnimNode_AimOffsetLookAt.h b/Engine/Source/Runtime/AnimGraphRuntime/Public/AnimNodes/AnimNode_AimOffsetLookAt.h
--- a/Engine/Source/Runtime/AnimGraphRuntime/Public/AnimNodes/AnimNode_AimOffsetLookAt.h
+++ b/Engine/Source/Runtime/AnimGraphRuntime/Public/AnimNodes/AnimNode_AimOffsetLookAt.h
@@ -69,6 +69,7 @@ public:
 
 	// FAnimNode_Base interface
 	virtual void OnInitializeAnimInstance(const FAnimInstanceProxy* InProxy, const UAnimInstance* InAnimInstance) override;
+	virtual bool NeedsOnInitializeAnimInstance() const override { return true; }
 	virtual void Initialize_AnyThread(const FAnimationInitializeContext& Context) override;
 	virtual void CacheBones_AnyThread(const FAnimationCacheBonesContext& Context) override;
 	virtual void UpdateAssetPlayer(const FAnimationUpdateContext& Context) override;
diff --git a/Engine/Source/Runtime/AnimGraphRuntime/Public/BoneControllers/AnimNode_ApplyLimits.h b/Engine/Source/Runtime/AnimGraphRuntime/Public/BoneControllers/AnimNode_ApplyLimits.h
--- a/Engine/Source/Runtime/AnimGraphRuntime/Public/BoneControllers/AnimNode_ApplyLimits.h
+++ b/Engine/Source/Runtime/AnimGraphRuntime/Public/BoneControllers/AnimNode_ApplyLimits.h
@@ -49,6 +49,7 @@ struct ANIMGRAPHRUNTIME_API FAnimNode_ApplyLimits : public FAnimNode_SkeletalCon
 	// FAnimNode_Base interface
 	virtual void GatherDebugData(FNodeDebugData& DebugData) override;
 	virtual void OnInitializeAnimInstance(const FAnimInstanceProxy* InProxy, const UAnimInstance* InAnimInstance) override;
+	virtual bool NeedsOnInitializeAnimInstance() const override { return true; }
 	// End of FAnimNode_Base interface
 
 	// FAnimNode_SkeletalControlBase interface
diff --git a/Engine/Source/Runtime/AnimGraphRuntime/Public/BoneControllers/AnimNode_ResetRoot.h b/Engine/Source/Runtime/AnimGraphRuntime/Public/BoneControllers/AnimNode_ResetRoot.h
--- a/Engine/Source/Runtime/AnimGraphRuntime/Public/BoneControllers/AnimNode_ResetRoot.h
+++ b/Engine/Source/Runtime/AnimGraphRuntime/Public/BoneControllers/AnimNode_ResetRoot.h
@@ -21,6 +21,7 @@ struct ANIMGRAPHRUNTIME_API FAnimNode_ResetRoot : public FAnimNode_SkeletalContr
 	// FAnimNode_Base interface
 	virtual void GatherDebugData(FNodeDebugData& DebugData) override;
 	virtual void OnInitializeAnimInstance(const FAnimInstanceProxy* InProxy, const UAnimInstance* InAnimInstance) override;
+	virtual bool NeedsOnInitializeAnimInstance() const override { return true; }
 	// End of FAnimNode_Base interface
 
 	// FAnimNode_SkeletalControlBase interface
diff --git a/Engine/Source/Runtime/AnimGraphRuntime/Public/BoneControllers/AnimNode_RigidBody.h b/Engine/Source/Runtime/AnimGraphRuntime/Public/BoneControllers/AnimNode_RigidBody.h
--- a/Engine/Source/Runtime/AnimGraphRuntime/Public/BoneControllers/AnimNode_RigidBody.h
+++ b/Engine/Source/Runtime/AnimGraphRuntime/Public/BoneControllers/AnimNode_RigidBody.h
@@ -45,6 +45,7 @@ struct ANIMGRAPHRUNTIME_API FAnimNode_RigidBody : public FAnimNode_SkeletalContr
 	virtual void EvaluateComponentPose_AnyThread(FComponentSpacePoseContext& Output) override;
 	virtual void EvaluateSkeletalControl_AnyThread(FComponentSpacePoseContext& Output, TArray<FBoneTransform>& OutBoneTransforms) override;
 	virtual void OnInitializeAnimInstance(const FAnimInstanceProxy* InProxy, const UAnimInstance* InAnimInstance) override;
+	virtual bool NeedsOnInitializeAnimInstance() const override { return true; }
 	virtual void PreUpdate(const UAnimInstance* InAnimInstance) override;
 	virtual void UpdateInternal(const FAnimationUpdateContext& Context) override;
 	virtual bool HasPreUpdate() const override { return true; }
diff --git a/Engine/Source/Runtime/AnimGraphRuntime/Public/BoneControllers/AnimNode_RigidBody_Chaos.h b/Engine/Source/Runtime/AnimGraphRuntime/Public/BoneControllers/AnimNode_RigidBody_Chaos.h
--- a/Engine/Source/Runtime/AnimGraphRuntime/Public/BoneControllers/AnimNode_RigidBody_Chaos.h
+++ b/Engine/Source/Runtime/AnimGraphRuntime/Public/BoneControllers/AnimNode_RigidBody_Chaos.h
@@ -47,6 +47,7 @@ public:
 	virtual void EvaluateComponentPose_AnyThread(FComponentSpacePoseContext& Output) override;
 	virtual void EvaluateSkeletalControl_AnyThread(FComponentSpacePoseContext& Output, TArray<FBoneTransform>& OutBoneTransforms) override;
 	virtual void OnInitializeAnimInstance(const FAnimInstanceProxy* InProxy, const UAnimInstance* InAnimInstance) override;
+	virtual bool NeedsOnInitializeAnimInstance() const override { return true; }
 	virtual void PreUpdate(const UAnimInstance* InAnimInstance) override;
 	virtual void UpdateInternal(const FAnimationUpdateContext& Context) override;
 	virtual bool HasPreUpdate() const override { return true; }
diff --git a/Engine/Source/Runtime/AnimGraphRuntime/Public/BoneControllers/AnimNode_SplineIK.h b/Engine/Source/Runtime/AnimGraphRuntime/Public/BoneControllers/AnimNode_SplineIK.h
--- a/Engine/Source/Runtime/AnimGraphRuntime/Public/BoneControllers/AnimNode_SplineIK.h
+++ b/Engine/Source/Runtime/AnimGraphRuntime/Public/BoneControllers/AnimNode_SplineIK.h
@@ -106,6 +106,7 @@ struct ANIMGRAPHRUNTIME_API FAnimNode_SplineIK : public FAnimNode_SkeletalContro
 	// FAnimNode_Base interface
 	virtual void GatherDebugData(FNodeDebugData& DebugData) override;
 	virtual void OnInitializeAnimInstance(const FAnimInstanceProxy* InProxy, const UAnimInstance* InAnimInstance) override;
+	virtual bool NeedsOnInitializeAnimInstance() const override { return true; }
 	// End of FAnimNode_Base interface
 
 	// FAnimNode_SkeletalControlBase interface
diff --git a/Engine/Source/Runtime/Engine/Classes/Animation/AnimBlueprintGeneratedClass.h b/Engine/Source/Runtime/Engine/Classes/Animation/AnimBlueprintGeneratedClass.h
--- a/Engine/Source/Runtime/Engine/Classes/Animation/AnimBlueprintGeneratedClass.h
+++ b/Engine/Source/Runtime/Engine/Classes/Animation/AnimBlueprintGeneratedClass.h
@@ -201,6 +201,10 @@ class ENGINE_API UAnimBlueprintGeneratedClass : public UBlueprintGeneratedClass,
 	TArray<UStructProperty*> AnimNodeProperties;
 	TArray<UStructProperty*> SubInstanceNodeProperties;
 	TArray<UStructProperty*> LayerNodeProperties;
+	TArray<UStructProperty*> PreUpdateNodeProperties;
+	TArray<UStructProperty*> DynamicResetNodeProperties;
+	TArray<UStructProperty*> StateMachineNodeProperties;
+	TArray<UStructProperty*> InitializationNodeProperties;
 
 	// Array of sync group names in the order that they are requested during compile
 	UPROPERTY()
@@ -222,6 +226,10 @@ public:
 	virtual const TArray<UStructProperty*>& GetAnimNodeProperties() const override { return AnimNodeProperties; }
 	virtual const TArray<UStructProperty*>& GetSubInstanceNodeProperties() const override { return SubInstanceNodeProperties; }
 	virtual const TArray<UStructProperty*>& GetLayerNodeProperties() const override { return LayerNodeProperties; }
+	virtual const TArray<UStructProperty*>& GetPreUpdateNodeProperties() const override { return PreUpdateNodeProperties; }
+	virtual const TArray<UStructProperty*>& GetDynamicResetNodeProperties() const override { return DynamicResetNodeProperties; }
+	virtual const TArray<UStructProperty*>& GetStateMachineNodeProperties() const override { return StateMachineNodeProperties; }
+	virtual const TArray<UStructProperty*>& GetInitializationNodeProperties() const override { return InitializationNodeProperties; }
 	virtual const TArray<FName>& GetSyncGroupNames() const override { return SyncGroupNames; }
 	virtual const TMap<FName, FCachedPoseIndices>& GetOrderedSavedPoseNodeIndicesMap() const override { return OrderedSavedPoseIndicesMap; }
 	virtual int32 GetSyncGroupIndex(FName SyncGroupName) const override { return SyncGroupNames.IndexOfByKey(SyncGroupName); }
diff --git a/Engine/Source/Runtime/Engine/Classes/Animation/AnimClassData.h b/Engine/Source/Runtime/Engine/Classes/Animation/AnimClassData.h
--- a/Engine/Source/Runtime/Engine/Classes/Animation/AnimClassData.h
+++ b/Engine/Source/Runtime/Engine/Classes/Animation/AnimClassData.h
@@ -50,6 +50,22 @@ public:
 	UPROPERTY()
 	TArray<UStructProperty*> LayerNodeProperties;
 
+	// Array of nodes that need a PreUpdate() call
+	UPROPERTY()
+	TArray<UStructProperty*> PreUpdateNodeProperties;
+
+	// Array of nodes that need a DynamicReset() call
+	UPROPERTY()
+	TArray<UStructProperty*> DynamicResetNodeProperties;
+
+	// Array of state machine nodes
+	UPROPERTY()
+	TArray<UStructProperty*> StateMachineNodeProperties;
+
+	// Array of nodes that need an OnInitializeAnimInstance call
+	UPROPERTY()
+	TArray<UStructProperty*> InitializationNodeProperties;
+
 	// Indices for any Asset Player found within a specific (named) Anim Layer Graph, or implemented Anim Interface Graph
 	UPROPERTY()
 	TMap<FName, FGraphAssetPlayerInformation> GraphNameAssetPlayers;
@@ -72,6 +88,10 @@ public:
 	virtual const TArray<UStructProperty*>& GetAnimNodeProperties() const override { return AnimNodeProperties; }
 	virtual const TArray<UStructProperty*>& GetSubInstanceNodeProperties() const override { return SubInstanceNodeProperties; }
 	virtual const TArray<UStructProperty*>& GetLayerNodeProperties() const override { return LayerNodeProperties; }
+	virtual const TArray<UStructProperty*>& GetPreUpdateNodeProperties() const override { return PreUpdateNodeProperties; }
+	virtual const TArray<UStructProperty*>& GetDynamicResetNodeProperties() const override { return DynamicResetNodeProperties; }
+	virtual const TArray<UStructProperty*>& GetStateMachineNodeProperties() const override { return StateMachineNodeProperties; }
+	virtual const TArray<UStructProperty*>& GetInitializationNodeProperties() const override { return InitializationNodeProperties; }
 	virtual const TArray<FName>& GetSyncGroupNames() const override { return SyncGroupNames; }
 	virtual int32 GetSyncGroupIndex(FName SyncGroupName) const override { return SyncGroupNames.IndexOfByKey(SyncGroupName); }
 	virtual const TArray<FExposedValueHandler>& GetExposedValueHandlers() const { return EvaluateGraphExposedInputs; }
@@ -89,6 +109,10 @@ public:
 		AnimNodeProperties = AnimClass->GetAnimNodeProperties();
 		SubInstanceNodeProperties = AnimClass->GetSubInstanceNodeProperties();
 		LayerNodeProperties = AnimClass->GetLayerNodeProperties();
+		PreUpdateNodeProperties = AnimClass->GetPreUpdateNodeProperties();
+		DynamicResetNodeProperties = AnimClass->GetDynamicResetNodeProperties();
+		StateMachineNodeProperties = AnimClass->GetStateMachineNodeProperties();
+		InitializationNodeProperties = AnimClass->GetInitializationNodeProperties();
 		SyncGroupNames = AnimClass->GetSyncGroupNames();
 		EvaluateGraphExposedInputs = AnimClass->GetExposedValueHandlers();
 		GraphNameAssetPlayers = AnimClass->GetGraphAssetPlayerInformation();
diff --git a/Engine/Source/Runtime/Engine/Classes/Animation/AnimClassInterface.h b/Engine/Source/Runtime/Engine/Classes/Animation/AnimClassInterface.h
--- a/Engine/Source/Runtime/Engine/Classes/Animation/AnimClassInterface.h
+++ b/Engine/Source/Runtime/Engine/Classes/Animation/AnimClassInterface.h
@@ -121,6 +121,10 @@ public:
 	virtual const TArray<UStructProperty*>& GetAnimNodeProperties() const = 0;
 	virtual const TArray<UStructProperty*>& GetSubInstanceNodeProperties() const = 0;
 	virtual const TArray<UStructProperty*>& GetLayerNodeProperties() const = 0;
+	virtual const TArray<UStructProperty*>& GetPreUpdateNodeProperties() const = 0;
+	virtual const TArray<UStructProperty*>& GetDynamicResetNodeProperties() const = 0;
+	virtual const TArray<UStructProperty*>& GetStateMachineNodeProperties() const = 0;
+	virtual const TArray<UStructProperty*>& GetInitializationNodeProperties() const = 0;
 	virtual const TArray<FExposedValueHandler>& GetExposedValueHandlers() const = 0;
 	virtual const TArray<FName>& GetSyncGroupNames() const = 0;
 	virtual const TMap<FName, FCachedPoseIndices>& GetOrderedSavedPoseNodeIndicesMap() const = 0;
diff --git a/Engine/Source/Runtime/Engine/Classes/Animation/AnimInstance.h b/Engine/Source/Runtime/Engine/Classes/Animation/AnimInstance.h
--- a/Engine/Source/Runtime/Engine/Classes/Animation/AnimInstance.h
+++ b/Engine/Source/Runtime/Engine/Classes/Animation/AnimInstance.h
@@ -1026,7 +1026,8 @@ public:
 	// evaluate happens when condition is met - i.e. depending on your skeletalmeshcomponent update flag
 	// post eval happens after evaluation is done
 	// uninitialize happens when owner is unregistered
-	void InitializeAnimation();
+	// @param	bInDeferRootNodeInitialization	When set to true, defer init of the blend tree until the first Update() call
+	void InitializeAnimation(bool bInDeferRootNodeInitialization = false);
 
 	/** Update Animation code-paths, updates and advances animation state, returns whether or not the actual update should have been called immediately */
 	bool UpdateAnimation(float DeltaSeconds, bool bNeedsValidRootMotion, EUpdateAnimationFlag UpdateFlag = EUpdateAnimationFlag::Default );
diff --git a/Engine/Source/Runtime/Engine/Classes/Animation/AnimNodeBase.h b/Engine/Source/Runtime/Engine/Classes/Animation/AnimNodeBase.h
--- a/Engine/Source/Runtime/Engine/Classes/Animation/AnimNodeBase.h
+++ b/Engine/Source/Runtime/Engine/Classes/Animation/AnimNodeBase.h
@@ -690,8 +690,9 @@ struct ENGINE_API FAnimNode_Base
 
 	/**
 	 * Override this to indicate that PreUpdate() should be called on the game thread (usually to 
-	 * This is called on the game thread.
 	 * gather non-thread safe data) before Update() is called.
+	 * Note that this is called at load on the UAnimInstance CDO to avoid needing to call this at runtime.
+	 * This is called on the game thread.
 	 */
 	virtual bool HasPreUpdate() const { return false; }
 
@@ -700,7 +701,8 @@ struct ENGINE_API FAnimNode_Base
 
 	/**
 	 * For nodes that implement some kind of simulation, return true here so ResetDynamics() gets called
-	 * when things like teleports, time skips etc. occur that might require special handling
+	 * when things like teleports, time skips etc. occur that might require special handling.
+	 * Note that this is called at load on the UAnimInstance CDO to avoid needing to call this at runtime.
 	 * This is called on the game thread.
 	 */
 	virtual bool NeedsDynamicReset() const { return false; }
@@ -711,6 +713,13 @@ struct ENGINE_API FAnimNode_Base
 	/** Called after compilation */
 	virtual void PostCompile(const class USkeleton* InSkeleton) {}
 
+	/** 
+	 * For nodes that need some kind of initialization that is not dependent on node relevancy 
+	 * (i.e. it is insufficent or inefficent to use Initialize_AnyThread), return true here.
+	 * Note that this is called at load on the UAnimInstance CDO to avoid needing to call this at runtime.
+	 */
+	virtual bool NeedsOnInitializeAnimInstance() const { return false; }
+
 	virtual ~FAnimNode_Base() {}
 
 	/** Deprecated functions */
diff --git a/Engine/Source/Runtime/Engine/Classes/Animation/AnimNode_Layer.h b/Engine/Source/Runtime/Engine/Classes/Animation/AnimNode_Layer.h
--- a/Engine/Source/Runtime/Engine/Classes/Animation/AnimNode_Layer.h
+++ b/Engine/Source/Runtime/Engine/Classes/Animation/AnimNode_Layer.h
@@ -39,6 +39,7 @@ public:
 
 	/** FAnimNode_Base interface */
 	virtual void OnInitializeAnimInstance(const FAnimInstanceProxy* InProxy, const UAnimInstance* InAnimInstance) override;
+	virtual bool NeedsOnInitializeAnimInstance() const override { return true; }
 
 	/** FAnimNode_SubInstance interface */
 	virtual FName GetDynamicLinkFunctionName() const override;
diff --git a/Engine/Source/Runtime/Engine/Classes/Animation/AnimNode_SubInstance.h b/Engine/Source/Runtime/Engine/Classes/Animation/AnimNode_SubInstance.h
--- a/Engine/Source/Runtime/Engine/Classes/Animation/AnimNode_SubInstance.h
+++ b/Engine/Source/Runtime/Engine/Classes/Animation/AnimNode_SubInstance.h
@@ -66,6 +66,7 @@ public:
 
 protected:
 	virtual void OnInitializeAnimInstance(const FAnimInstanceProxy* InProxy, const UAnimInstance* InAnimInstance) override;
+	virtual bool NeedsOnInitializeAnimInstance() const override { return true; }
 	// End of FAnimNode_Base interface
 
 	// Re-create the sub instances for this node
diff --git a/Engine/Source/Runtime/Engine/Classes/Animation/Skeleton.h b/Engine/Source/Runtime/Engine/Classes/Animation/Skeleton.h
--- a/Engine/Source/Runtime/Engine/Classes/Animation/Skeleton.h
+++ b/Engine/Source/Runtime/Engine/Classes/Animation/Skeleton.h
@@ -384,6 +384,9 @@ protected:
 	UPROPERTY()
 	FSmartNameContainer SmartNames;
 
+	// Cached ptr to the persistent AnimCurveMapping
+	FSmartNameMapping* AnimCurveMapping;
+
 	// this is default curve uid list used like ref pose, as default value
 	// don't use this unless you want all curves from the skeleton
 	// FBoneContainer contains only list that is used by current LOD
diff --git a/Engine/Source/Runtime/Engine/Classes/Animation/SmartName.h b/Engine/Source/Runtime/Engine/Classes/Animation/SmartName.h
--- a/Engine/Source/Runtime/Engine/Classes/Animation/SmartName.h
+++ b/Engine/Source/Runtime/Engine/Classes/Animation/SmartName.h
@@ -142,6 +142,13 @@ struct ENGINE_API FSmartNameMapping
 		return CurveMetaDataMap.Find(CurveName);
 	}
 
+#if !WITH_EDITOR
+	const FCurveMetaData& GetCurveMetaData(SmartName::UID_Type CurveUID) const
+	{
+		return CurveMetaDataList[CurveUID];
+	}
+#endif
+
 	// Serialize this to the provided archive; required for TMap serialization
 	void Serialize(FArchive& Ar);
 
@@ -154,8 +161,14 @@ struct ENGINE_API FSmartNameMapping
 	SmartName::UID_Type GetMaxUID() const { return CurveNameList.Num() - 1; }
 
 private:
-	
-	TArray<FName> CurveNameList; // List of curve names, order gives UID
+	// List of curve names, indexed by UID
+	TArray<FName> CurveNameList;
+
+#if !WITH_EDITOR
+	// List of curve metadata, indexed by UID
+	TArray<FCurveMetaData> CurveMetaDataList;
+#endif
+
 	TMap<FName, FCurveMetaData> CurveMetaDataMap;
 };
 
@@ -165,7 +178,7 @@ struct ENGINE_API FSmartNameContainer
 	GENERATED_USTRUCT_BODY();
 
 	// Add a new smartname container with the provided name
-	void AddContainer(FName NewContainerName);
+	FSmartNameMapping* AddContainer(FName NewContainerName);
 
 	// Get a container by name	
 	const FSmartNameMapping* GetContainer(FName ContainerName) const;
diff --git a/Engine/Source/Runtime/Engine/Classes/Components/SkeletalMeshComponent.h b/Engine/Source/Runtime/Engine/Classes/Components/SkeletalMeshComponent.h
--- a/Engine/Source/Runtime/Engine/Classes/Components/SkeletalMeshComponent.h
+++ b/Engine/Source/Runtime/Engine/Classes/Components/SkeletalMeshComponent.h
@@ -1064,7 +1064,7 @@ public:
 	virtual void SkelMeshCompOnParticleSystemFinished( class UParticleSystemComponent* PSC );
 
 	class UAnimSingleNodeInstance * GetSingleNodeInstance() const;
-	bool InitializeAnimScriptInstance(bool bForceReinit=true);
+	bool InitializeAnimScriptInstance(bool bForceReinit = true, bool bInDeferRootNodeInitialization = false);
 
 	/** @return true if wind is enabled */
 	virtual bool IsWindEnabled() const;
diff --git a/Engine/Source/Runtime/Engine/Private/Animation/AnimBlueprintGeneratedClass.cpp b/Engine/Source/Runtime/Engine/Private/Animation/AnimBlueprintGeneratedClass.cpp
--- a/Engine/Source/Runtime/Engine/Private/Animation/AnimBlueprintGeneratedClass.cpp
+++ b/Engine/Source/Runtime/Engine/Private/Animation/AnimBlueprintGeneratedClass.cpp
@@ -17,6 +17,7 @@
 #include "Animation/AnimNode_SubInput.h"
 #include "Animation/AnimNode_Layer.h"
 #include "Animation/AnimNode_AssetPlayerBase.h"
+#include "Animation/AnimNode_StateMachine.h"
 
 /////////////////////////////////////////////////////
 // FStateMachineDebugData
@@ -256,6 +257,10 @@ void UAnimBlueprintGeneratedClass::Link(FArchive& Ar, bool bRelinkExistingProper
 	AnimNodeProperties.Empty();
 	SubInstanceNodeProperties.Empty();
 	LayerNodeProperties.Empty();
+	PreUpdateNodeProperties.Empty();
+	DynamicResetNodeProperties.Empty();
+	StateMachineNodeProperties.Empty();
+	InitializationNodeProperties.Empty();
 
 #if WITH_EDITOR
 	// This relies on the entire class being fully loaded, this is not the case with EDL async-loading, in which case the functions are generated in PostLoad
@@ -277,6 +282,10 @@ void UAnimBlueprintGeneratedClass::Link(FArchive& Ar, bool bRelinkExistingProper
 				{
 					LayerNodeProperties.Add(StructProp);
 				}
+				else if(StructProp->Struct == FAnimNode_StateMachine::StaticStruct())
+				{
+					StateMachineNodeProperties.Add(StructProp);
+				}
 				AnimNodeProperties.Add(StructProp);
 			}
 		}
@@ -442,6 +451,10 @@ void UAnimBlueprintGeneratedClass::GenerateAnimationBlueprintFunctions()
 
 void UAnimBlueprintGeneratedClass::LinkFunctionsToDefaultObjectNodes(UObject* DefaultObject)
 {
+	PreUpdateNodeProperties.Empty();
+	DynamicResetNodeProperties.Empty();
+	InitializationNodeProperties.Empty();
+
 	// Link functions to their nodes
 	for(int32 AnimNodeIndex = 0; AnimNodeIndex < AnimNodeProperties.Num(); ++AnimNodeIndex)
 	{
@@ -471,6 +484,24 @@ void UAnimBlueprintGeneratedClass::LinkFunctionsToDefaultObjectNodes(UObject* De
 				}
 			}
 		}
+		else if(StructProperty->Struct->IsChildOf(FAnimNode_Base::StaticStruct()))
+		{
+			FAnimNode_Base* Node = StructProperty->ContainerPtrToValuePtr<FAnimNode_Base>(DefaultObject);
+			if(Node->NeedsDynamicReset())
+			{
+				DynamicResetNodeProperties.Add(StructProperty);
+			}
+
+			if(Node->HasPreUpdate())
+			{
+				PreUpdateNodeProperties.Add(StructProperty);
+			}
+
+			if(Node->NeedsOnInitializeAnimInstance())
+			{
+				InitializationNodeProperties.Add(StructProperty);
+			}
+		}
 	}
 
 	// Determine whether functions are 'implemented'
diff --git a/Engine/Source/Runtime/Engine/Private/Animation/AnimInstance.cpp b/Engine/Source/Runtime/Engine/Private/Animation/AnimInstance.cpp
--- a/Engine/Source/Runtime/Engine/Private/Animation/AnimInstance.cpp
+++ b/Engine/Source/Runtime/Engine/Private/Animation/AnimInstance.cpp
@@ -168,7 +168,7 @@ UWorld* UAnimInstance::GetWorld() const
 	return (HasAnyFlags(RF_ClassDefaultObject) ? nullptr : GetSkelMeshComponent()->GetWorld());
 }
 
-void UAnimInstance::InitializeAnimation()
+void UAnimInstance::InitializeAnimation(bool bInDeferRootNodeInitialization)
 {
 	FScopeCycleCounterUObject ContextScope(this);
 	SCOPE_CYCLE_COUNTER(STAT_AnimInitTime);
@@ -213,7 +213,7 @@ void UAnimInstance::InitializeAnimation()
 	NativeInitializeAnimation();
 	BlueprintInitializeAnimation();
 
-	GetProxyOnGameThread<FAnimInstanceProxy>().InitializeRootNode();
+	GetProxyOnGameThread<FAnimInstanceProxy>().InitializeRootNode(bInDeferRootNodeInitialization);
 
 	// we can bind rules & events now the graph has been initialized
 	GetProxyOnGameThread<FAnimInstanceProxy>().BindNativeDelegates();
diff --git a/Engine/Source/Runtime/Engine/Private/Animation/AnimInstanceProxy.cpp b/Engine/Source/Runtime/Engine/Private/Animation/AnimInstanceProxy.cpp
--- a/Engine/Source/Runtime/Engine/Private/Animation/AnimInstanceProxy.cpp
+++ b/Engine/Source/Runtime/Engine/Private/Animation/AnimInstanceProxy.cpp
@@ -180,7 +180,7 @@ void FAnimInstanceProxy::Initialize(UAnimInstance* InAnimInstance)
 	}
 }
 
-void FAnimInstanceProxy::InitializeRootNode()
+void FAnimInstanceProxy::InitializeRootNode(bool bInDeferRootNodeInitialization)
 {
 	DECLARE_SCOPE_HIERARCHICAL_COUNTER_FUNC()
 
@@ -188,40 +188,35 @@ void FAnimInstanceProxy::InitializeRootNode()
 	GameThreadPreUpdateNodes.Reset();
 	DynamicResetNodes.Reset();
 
-	auto InitializeNode = [this](FAnimNode_Base* AnimNode)
+	if(AnimClassInterface)
 	{
-		AnimNode->OnInitializeAnimInstance(this, CastChecked<UAnimInstance>(GetAnimInstanceObject()));
+		// Init any nodes that need non-relevancy based initialization
+		UAnimInstance* AnimInstance = CastChecked<UAnimInstance>(GetAnimInstanceObject());
+		for (UStructProperty* Property : AnimClassInterface->GetInitializationNodeProperties())
+		{
+			FAnimNode_Base* AnimNode = Property->ContainerPtrToValuePtr<FAnimNode_Base>(AnimInstanceObject);
+			AnimNode->OnInitializeAnimInstance(this, AnimInstance);
+		}
 
-		if (AnimNode->HasPreUpdate())
+		// Cache any preupdate nodes
+		for (UStructProperty* Property : AnimClassInterface->GetPreUpdateNodeProperties())
 		{
+			FAnimNode_Base* AnimNode = Property->ContainerPtrToValuePtr<FAnimNode_Base>(AnimInstanceObject);
 			GameThreadPreUpdateNodes.Add(AnimNode);
 		}
 
-		if (AnimNode->NeedsDynamicReset())
+		// Cache any dynamic reset nodes
+		for (UStructProperty* Property : AnimClassInterface->GetDynamicResetNodeProperties())
 		{
+			FAnimNode_Base* AnimNode = Property->ContainerPtrToValuePtr<FAnimNode_Base>(AnimInstanceObject);
 			DynamicResetNodes.Add(AnimNode);
 		}
-	};
 
-	if(AnimClassInterface)
-	{
 		// cache any state machine descriptions we have
-		for (UStructProperty* Property : AnimClassInterface->GetAnimNodeProperties())
+		for (UStructProperty* Property : AnimClassInterface->GetStateMachineNodeProperties())
 		{
-			if (Property->Struct->IsChildOf(FAnimNode_Base::StaticStruct()))
-			{
-				FAnimNode_Base* AnimNode = Property->ContainerPtrToValuePtr<FAnimNode_Base>(AnimInstanceObject);
-				if (AnimNode)
-				{
-					InitializeNode(AnimNode);
-
-					if (Property->Struct->IsChildOf(FAnimNode_StateMachine::StaticStruct()))
-					{
-						FAnimNode_StateMachine* StateMachine = static_cast<FAnimNode_StateMachine*>(AnimNode);
-						StateMachine->CacheMachineDescription(AnimClassInterface);
-					}
-				}
-			}
+			FAnimNode_StateMachine* StateMachine = Property->ContainerPtrToValuePtr<FAnimNode_StateMachine>(AnimInstanceObject);
+			StateMachine->CacheMachineDescription(AnimClassInterface);
 		}
 
 		// Cache default sub-input 
@@ -243,6 +238,24 @@ void FAnimInstanceProxy::InitializeRootNode()
 	}
 	else
 	{
+		auto InitializeNode = [this](FAnimNode_Base* AnimNode)
+		{
+			if(AnimNode->NeedsOnInitializeAnimInstance())
+			{
+				AnimNode->OnInitializeAnimInstance(this, CastChecked<UAnimInstance>(GetAnimInstanceObject()));
+			}
+
+			if (AnimNode->HasPreUpdate())
+			{
+				GameThreadPreUpdateNodes.Add(AnimNode);
+			}
+
+			if (AnimNode->NeedsDynamicReset())
+			{
+				DynamicResetNodes.Add(AnimNode);
+			}
+		};
+
 		//We have a custom root node, so get the associated nodes and initialize them
 		TArray<FAnimNode_Base*> CustomNodes;
 		GetCustomNodes(CustomNodes);
@@ -255,7 +268,14 @@ void FAnimInstanceProxy::InitializeRootNode()
 		}
 	}
 
-	InitializeRootNode_WithRoot(RootNode);
+	if(!bInDeferRootNodeInitialization)
+	{
+		InitializeRootNode_WithRoot(RootNode);
+	}
+	else
+	{
+		bDeferRootNodeInitialization = true;
+	}
 }
 
 void FAnimInstanceProxy::InitializeRootNode_WithRoot(FAnimNode_Base* InRootNode)
@@ -1173,6 +1193,12 @@ void FAnimInstanceProxy::UpdateAnimation_WithRoot(FAnimNode_Base* InRootNode, FN
 
 	if(InRootNode == RootNode)
 	{
+		if(bDeferRootNodeInitialization)
+		{
+			InitializeRootNode_WithRoot(RootNode);
+			bDeferRootNodeInitialization = false;
+		}
+
 		// Call the correct override point if this is the root node
 		CacheBones();
 	}
diff --git a/Engine/Source/Runtime/Engine/Private/Animation/AnimNodeBase.cpp b/Engine/Source/Runtime/Engine/Private/Animation/AnimNodeBase.cpp
--- a/Engine/Source/Runtime/Engine/Private/Animation/AnimNodeBase.cpp
+++ b/Engine/Source/Runtime/Engine/Private/Animation/AnimNodeBase.cpp
@@ -187,9 +187,10 @@ void FPoseLinkBase::AttemptRelink(const FAnimationBaseContext& Context)
 		check(AnimBlueprintClass);
 
 		// adding ensure. We had a crash on here
-		if ( ensure(AnimBlueprintClass->GetAnimNodeProperties().IsValidIndex(LinkID)) )
+		const TArray<UStructProperty*> AnimNodeProperties = AnimBlueprintClass->GetAnimNodeProperties();
+		if ( ensure(AnimNodeProperties.IsValidIndex(LinkID)) )
 		{
-			UProperty* LinkedProperty = AnimBlueprintClass->GetAnimNodeProperties()[LinkID];
+			UProperty* LinkedProperty = AnimNodeProperties[LinkID];
 			void* LinkedNodePtr = LinkedProperty->ContainerPtrToValuePtr<void>(Context.AnimInstanceProxy->GetAnimInstanceObject());
 			LinkedNode = (FAnimNode_Base*)LinkedNodePtr;
 		}
diff --git a/Engine/Source/Runtime/Engine/Private/Animation/AnimNode_StateMachine.cpp b/Engine/Source/Runtime/Engine/Private/Animation/AnimNode_StateMachine.cpp
--- a/Engine/Source/Runtime/Engine/Private/Animation/AnimNode_StateMachine.cpp
+++ b/Engine/Source/Runtime/Engine/Private/Animation/AnimNode_StateMachine.cpp
@@ -942,7 +942,14 @@ void FAnimNode_StateMachine::SetState(const FAnimationBaseContext& Context, int3
 			StatePoseLinks[NewStateIndex].Initialize(InitContext);
 
 			// Also call cache bones if needed
-			ConditionallyCacheBonesForState(NewStateIndex, Context);
+			// Note dont call CacheBones here if we are in the process of whole-graph initialization as a 'never updated' counter
+			// will not perform its 'minimal update guard' duty and every call will end up getting though, performing duplicate work
+			// over Save/UseCachedPose boundaries etc.
+			// This is OK because CacheBones is actually called before updating the graph anyway after whole-graph initialization
+			if(Context.AnimInstanceProxy->GetCachedBonesCounter().HasEverBeenUpdated())
+			{
+				ConditionallyCacheBonesForState(NewStateIndex, Context);
+			}
 		}
 
 		if(CurrentState != INDEX_NONE && CurrentState < OnGraphStatesEntered.Num())
diff --git a/Engine/Source/Runtime/Engine/Private/Animation/Skeleton.cpp b/Engine/Source/Runtime/Engine/Private/Animation/Skeleton.cpp
--- a/Engine/Source/Runtime/Engine/Private/Animation/Skeleton.cpp
+++ b/Engine/Source/Runtime/Engine/Private/Animation/Skeleton.cpp
@@ -104,7 +104,7 @@ USkeleton::USkeleton(const FObjectInitializer& ObjectInitializer)
 	: Super(ObjectInitializer)
 {
 	// Make sure we have somewhere for curve names.
-	SmartNames.AddContainer(AnimCurveMappingName);
+	AnimCurveMapping = SmartNames.AddContainer(AnimCurveMappingName);
 
 	AnimCurveUidVersion = 0;
 }
@@ -229,6 +229,8 @@ void USkeleton::Serialize( FArchive& Ar )
 	if(Ar.UE4Ver() >= VER_UE4_SKELETON_ADD_SMARTNAMES)
 	{
 		SmartNames.Serialize(Ar);
+
+		AnimCurveMapping = SmartNames.GetContainerInternal(USkeleton::AnimCurveMappingName);
 	}
 
 	// Build look up table between Slot nodes and their Group.
@@ -1423,6 +1425,7 @@ FSmartNameMapping* USkeleton::GetOrAddSmartNameContainer(const FName& ContainerN
 		Modify();
 		IncreaseAnimCurveUidVersion();
 		SmartNames.AddContainer(ContainerName);
+		AnimCurveMapping = SmartNames.GetContainerInternal(USkeleton::AnimCurveMappingName);
 		Mapping = SmartNames.GetContainerInternal(ContainerName);		
 	}
 
@@ -1452,8 +1455,8 @@ void USkeleton::IncreaseAnimCurveUidVersion()
 	++AnimCurveUidVersion;
 
 	// update default uid list
-	const FSmartNameMapping* Mapping = GetSmartNameContainer(USkeleton::AnimCurveMappingName);
-	if (Mapping != nullptr)
+	const FSmartNameMapping* Mapping = AnimCurveMapping;
+	if (ensureAlways(Mapping))
 	{
 		DefaultCurveUIDList.Reset();
 		// this should exactly work with what current index is
@@ -1463,7 +1466,7 @@ void USkeleton::IncreaseAnimCurveUidVersion()
 
 FCurveMetaData* USkeleton::GetCurveMetaData(const FName& CurveName)
 {
-	FSmartNameMapping* Mapping = SmartNames.GetContainerInternal(USkeleton::AnimCurveMappingName);
+	FSmartNameMapping* Mapping = AnimCurveMapping;
 	if (ensureAlways(Mapping))
 	{
 		return Mapping->GetCurveMetaData(CurveName);
@@ -1474,7 +1477,7 @@ FCurveMetaData* USkeleton::GetCurveMetaData(const FName& CurveName)
 
 const FCurveMetaData* USkeleton::GetCurveMetaData(const FName& CurveName) const
 {
-	const FSmartNameMapping* Mapping = SmartNames.GetContainerInternal(USkeleton::AnimCurveMappingName);
+	const FSmartNameMapping* Mapping = AnimCurveMapping;
 	if (ensureAlways(Mapping))
 	{
 		return Mapping->GetCurveMetaData(CurveName);
@@ -1485,14 +1488,18 @@ const FCurveMetaData* USkeleton::GetCurveMetaData(const FName& CurveName) const
 
 const FCurveMetaData* USkeleton::GetCurveMetaData(const SmartName::UID_Type CurveUID) const
 {
-	const FSmartNameMapping* Mapping = SmartNames.GetContainerInternal(USkeleton::AnimCurveMappingName);
+	const FSmartNameMapping* Mapping = AnimCurveMapping;
 	if (ensureAlways(Mapping))
 	{
+#if WITH_EDITOR
 		FSmartName SmartName;
 		if (Mapping->FindSmartNameByUID(CurveUID, SmartName))
 		{
 			return Mapping->GetCurveMetaData(SmartName.DisplayName);
 		}
+#else
+		return &Mapping->GetCurveMetaData(CurveUID);
+#endif
 	}
 
 	return nullptr;
@@ -1500,7 +1507,7 @@ const FCurveMetaData* USkeleton::GetCurveMetaData(const SmartName::UID_Type Curv
 
 FCurveMetaData* USkeleton::GetCurveMetaData(const FSmartName& CurveName)
 {
-	FSmartNameMapping* Mapping = SmartNames.GetContainerInternal(USkeleton::AnimCurveMappingName);
+	FSmartNameMapping* Mapping = AnimCurveMapping;
 	if (ensureAlways(Mapping))
 	{
 		// the name might have changed, make sure it's up-to-date
@@ -1514,7 +1521,7 @@ FCurveMetaData* USkeleton::GetCurveMetaData(const FSmartName& CurveName)
 
 const FCurveMetaData* USkeleton::GetCurveMetaData(const FSmartName& CurveName) const
 {
-	const FSmartNameMapping* Mapping = SmartNames.GetContainerInternal(USkeleton::AnimCurveMappingName);
+	const FSmartNameMapping* Mapping = AnimCurveMapping;
 	if (ensureAlways(Mapping))
 	{
 		// the name might have changed, make sure it's up-to-date
@@ -1531,7 +1538,7 @@ void USkeleton::AccumulateCurveMetaData(FName CurveName, bool bMaterialSet, bool
 {
 	if (bMaterialSet || bMorphtargetSet)
 	{
-		const FSmartNameMapping* Mapping = SmartNames.GetContainerInternal(USkeleton::AnimCurveMappingName);
+		const FSmartNameMapping* Mapping = AnimCurveMapping;
 		if (ensureAlways(Mapping))
 		{
 			// if we don't have name, add one
diff --git a/Engine/Source/Runtime/Engine/Private/Animation/SmartName.cpp b/Engine/Source/Runtime/Engine/Private/Animation/SmartName.cpp
--- a/Engine/Source/Runtime/Engine/Private/Animation/SmartName.cpp
+++ b/Engine/Source/Runtime/Engine/Private/Animation/SmartName.cpp
@@ -23,6 +23,9 @@ FSmartName FSmartNameMapping::AddName(FName Name)
 
 	FSmartName NewSmartName(Name, CurveNameList.Add(Name));
 	CurveMetaDataMap.Add(Name);
+#if !WITH_EDITOR
+	CurveMetaDataList.AddDefaulted();
+#endif
 	return NewSmartName;
 }
 
@@ -114,6 +117,9 @@ void FSmartNameMapping::Serialize(FArchive& Ar)
 	if (Ar.IsLoading())
 	{
 		CurveMetaDataMap.GenerateKeyArray(CurveNameList);
+#if !WITH_EDITOR
+		CurveMetaDataMap.GenerateValueArray(CurveMetaDataList);
+#endif
 	}
 }
 
@@ -269,12 +275,15 @@ void FSmartNameMapping::InitializeCurveMetaData(class USkeleton* Skeleton)
 // FSmartNameContainer
 //
 ///////////////////////////////////////////////////////////////////////
-void FSmartNameContainer::AddContainer(FName NewContainerName)
+FSmartNameMapping* FSmartNameContainer::AddContainer(FName NewContainerName)
 {
-	if(!NameMappings.Find(NewContainerName))
+	FSmartNameMapping* ExistingMapping = NameMappings.Find(NewContainerName);
+	if(ExistingMapping == nullptr)
 	{
-		NameMappings.Add(NewContainerName);
+		return &NameMappings.Add(NewContainerName);
 	}
+
+	return ExistingMapping;
 }
 
 const FSmartNameMapping* FSmartNameContainer::GetContainer(FName ContainerName) const
diff --git a/Engine/Source/Runtime/Engine/Private/BoneContainer.cpp b/Engine/Source/Runtime/Engine/Private/BoneContainer.cpp
--- a/Engine/Source/Runtime/Engine/Private/BoneContainer.cpp
+++ b/Engine/Source/Runtime/Engine/Private/BoneContainer.cpp
@@ -439,7 +439,9 @@ void FBoneContainer::RemapFromSkeleton(USkeleton const & SourceSkeleton)
 
 bool FBoneReference::Initialize(const FBoneContainer& RequiredBones)
 {
+#if WITH_EDITOR
 	BoneName = *BoneName.ToString().TrimStartAndEnd();
+#endif
 	BoneIndex = RequiredBones.GetPoseBoneIndexForBoneName(BoneName);
 
 	bUseSkeletonIndex = false;
@@ -467,7 +469,9 @@ bool FBoneReference::Initialize(const USkeleton* Skeleton)
 {
 	if (Skeleton && (BoneName != NAME_None))
 	{
+#if WITH_EDITOR
 		BoneName = *BoneName.ToString().TrimStartAndEnd();
+#endif
 		BoneIndex = Skeleton->GetReferenceSkeleton().FindBoneIndex(BoneName);
 		bUseSkeletonIndex = true;
 	}
diff --git a/Engine/Source/Runtime/Engine/Private/Components/SkeletalMeshComponent.cpp b/Engine/Source/Runtime/Engine/Private/Components/SkeletalMeshComponent.cpp
--- a/Engine/Source/Runtime/Engine/Private/Components/SkeletalMeshComponent.cpp
+++ b/Engine/Source/Runtime/Engine/Private/Components/SkeletalMeshComponent.cpp
@@ -649,7 +649,13 @@ void USkeletalMeshComponent::InitAnim(bool bForceReinit)
 		// this has to be called before Initialize Animation because it will required RequiredBones list when InitializeAnimScript
 		RecalcRequiredBones(PredictedLODLevel);
 
-		const bool bInitializedAnimInstance = InitializeAnimScriptInstance(bForceReinit);
+		// In Editor, animations won't get ticked. So Update once to get accurate representation instead of T-Pose.
+		// Also allow this to be an option to support pre-4.19 games that might need it..
+		const bool bTickAnimationNow =
+			((GetWorld()->WorldType == EWorldType::Editor) && !bUseRefPoseOnInitAnim)
+			|| UAnimationSettings::Get()->bTickAnimationOnSkeletalMeshInit;
+
+		const bool bInitializedAnimInstance = InitializeAnimScriptInstance(bForceReinit, !bTickAnimationNow);
 
 		// Make sure we have a valid pose.
 		// We don't allocate transform data when using MasterPoseComponent, so we have nothing to render.
@@ -657,12 +663,6 @@ void USkeletalMeshComponent::InitAnim(bool bForceReinit)
 		{	
 			if (bInitializedAnimInstance || (AnimScriptInstance == nullptr))
 			{ 
-				// In Editor, animations won't get ticked. So Update once to get accurate representation instead of T-Pose.
-				// Also allow this to be an option to support pre-4.19 games that might need it..
-				const bool bTickAnimationNow =
-					((GetWorld()->WorldType == EWorldType::Editor) && !bUseRefPoseOnInitAnim)
-					|| UAnimationSettings::Get()->bTickAnimationOnSkeletalMeshInit;
-
 				if (bTickAnimationNow)
 				{
 					TickAnimation(0.f, false);
@@ -690,7 +690,7 @@ void USkeletalMeshComponent::InitAnim(bool bForceReinit)
 	}
 }
 
-bool USkeletalMeshComponent::InitializeAnimScriptInstance(bool bForceReinit)
+bool USkeletalMeshComponent::InitializeAnimScriptInstance(bool bForceReinit, bool bInDeferRootNodeInitialization)
 {
 	bool bInitializedMainInstance = false;
 	bool bInitializedPostInstance = false;
@@ -709,7 +709,7 @@ bool USkeletalMeshComponent::InitializeAnimScriptInstance(bool bForceReinit)
 				// If we have any sub-instances left we need to clear them out now, we're about to have a new master instance
 				SubInstances.Empty();
 
-				AnimScriptInstance->InitializeAnimation();
+				AnimScriptInstance->InitializeAnimation(bInDeferRootNodeInitialization);
 				bInitializedMainInstance = true;
 			}
 		}
@@ -730,7 +730,7 @@ bool USkeletalMeshComponent::InitializeAnimScriptInstance(bool bForceReinit)
 
 				if (AnimScriptInstance)
 				{
-					AnimScriptInstance->InitializeAnimation();
+					AnimScriptInstance->InitializeAnimation(bInDeferRootNodeInitialization);
 					bInitializedMainInstance = true;
 				}
 
@@ -784,7 +784,7 @@ bool USkeletalMeshComponent::InitializeAnimScriptInstance(bool bForceReinit)
 
 		if (AnimScriptInstance && !bInitializedMainInstance && bForceReinit)
 		{
-			AnimScriptInstance->InitializeAnimation();
+			AnimScriptInstance->InitializeAnimation(bInDeferRootNodeInitialization);
 			bInitializedMainInstance = true;
 		}
 
diff --git a/Engine/Source/Runtime/Engine/Public/Animation/AnimCurveTypes.h b/Engine/Source/Runtime/Engine/Public/Animation/AnimCurveTypes.h
--- a/Engine/Source/Runtime/Engine/Public/Animation/AnimCurveTypes.h
+++ b/Engine/Source/Runtime/Engine/Public/Animation/AnimCurveTypes.h
@@ -383,7 +383,8 @@ struct FBaseBlendedCurve
 		int32 Count = 0;
 		if (InUIDToArrayIndexLUT)
 		{
-			for (int32 Index = 0; Index<InUIDToArrayIndexLUT->Num(); ++Index)
+			const int32 ArraySize = InUIDToArrayIndexLUT->Num();
+			for (int32 Index = 0; Index < ArraySize; ++Index)
 			{
 				if ((*InUIDToArrayIndexLUT)[Index] != MAX_uint16)
 				{
diff --git a/Engine/Source/Runtime/Engine/Public/Animation/AnimInstanceProxy.h b/Engine/Source/Runtime/Engine/Public/Animation/AnimInstanceProxy.h
--- a/Engine/Source/Runtime/Engine/Public/Animation/AnimInstanceProxy.h
+++ b/Engine/Source/Runtime/Engine/Public/Animation/AnimInstanceProxy.h
@@ -125,6 +125,7 @@ public:
 		, bUpdatingRoot(false)
 		, bBoneCachesInvalidated(false)
 		, bShouldExtractRootMotion(false)
+		, bDeferRootNodeInitialization(false)
 #if WITH_EDITORONLY_DATA
 		, bIsBeingDebugged(false)
 #endif
@@ -148,6 +149,7 @@ public:
 		, bUpdatingRoot(false)
 		, bBoneCachesInvalidated(false)
 		, bShouldExtractRootMotion(false)
+		, bDeferRootNodeInitialization(false)
 #if WITH_EDITORONLY_DATA
 		, bIsBeingDebugged(false)
 #endif
@@ -733,7 +735,7 @@ protected:
 	void GetStateMachineIndexAndDescription(FName InMachineName, int32& OutMachineIndex, const FBakedAnimationStateMachine** OutMachineDescription);
 
 	/** Initialize the root node - split into a separate function for backwards compatibility (initialization order) reasons */
-	void InitializeRootNode();
+	void InitializeRootNode(bool bInDeferRootNodeInitialization = false);
 
 	/** Initialize the specified root node */
 	void InitializeRootNode_WithRoot(FAnimNode_Base* InRootNode);
@@ -937,6 +939,9 @@ private:
 	// Diplicate of bool result of ShouldExtractRootMotion()
 	uint8 bShouldExtractRootMotion : 1;
 
+	/** We can defer initialization until first update */
+	uint8 bDeferRootNodeInitialization : 1;
+
 #if WITH_EDITORONLY_DATA
 	/** Whether this UAnimInstance is currently being debugged in the editor */
 	uint8 bIsBeingDebugged : 1;
