Fix crash when applying a class as an overlay layer that contains layer nodes itself

Dynamic unlinking wasnt happening because the layer nodes are in a different class. This change just unlinks all layer nodes in the overlay class as we wont be using the AnimGraph anyway.

diff --git a/Engine/Source/Runtime/Engine/Classes/Animation/AnimNode_SubInstance.h b/Engine/Source/Runtime/Engine/Classes/Animation/AnimNode_SubInstance.h
--- a/Engine/Source/Runtime/Engine/Classes/Animation/AnimNode_SubInstance.h
+++ b/Engine/Source/Runtime/Engine/Classes/Animation/AnimNode_SubInstance.h
@@ -89,4 +89,6 @@ protected:
 
 	/** Helper function for finding function inputs when linking/unlinking */
 	int32 FindFunctionInputIndex(const FAnimBlueprintFunction& AnimBlueprintFunction, const FName& InInputName);
+
+	friend class UAnimInstance;
 };
diff --git a/Engine/Source/Runtime/Engine/Private/Animation/AnimInstance.cpp b/Engine/Source/Runtime/Engine/Private/Animation/AnimInstance.cpp
--- a/Engine/Source/Runtime/Engine/Private/Animation/AnimInstance.cpp
+++ b/Engine/Source/Runtime/Engine/Private/Animation/AnimInstance.cpp
@@ -2570,6 +2570,16 @@ void UAnimInstance::PerformLayerOverlayOperation(TSubclassOf<UAnimInstance> InCl
 
 		USkeletalMeshComponent* MeshComp = GetSkelMeshComponent();
 
+		auto UnlinkLayerNodesInInstance = [](UAnimInstance* InAnimInstance)
+		{
+			const IAnimClassInterface* NewSubInstanceClass = IAnimClassInterface::GetFromClass(InAnimInstance->GetClass());
+			for(const UStructProperty* LayerNodeProperty : NewSubInstanceClass->GetLayerNodeProperties())
+			{
+				FAnimNode_Layer* SubInstanceLayerNode = LayerNodeProperty->ContainerPtrToValuePtr<FAnimNode_Layer>(InAnimInstance);
+				SubInstanceLayerNode->DynamicUnlink(InAnimInstance);
+			}
+		};
+
 		for(TPair<FName, TArray<FAnimNode_Layer*, TInlineAllocator<4>>> LayerPair : LayerNodesToSet)
 		{
 			if (LayerPair.Key == NAME_None)
@@ -2588,11 +2598,12 @@ void UAnimInstance::PerformLayerOverlayOperation(TSubclassOf<UAnimInstance> InCl
 						if (TargetInstance == nullptr || ClassToSet != TargetInstance->GetClass())
 						{
 							UAnimInstance* NewSubInstance = NewObject<UAnimInstance>(MeshComp, ClassToSet);
+							NewSubInstance->InitializeAnimation();
 
-							LayerNode->SetLayerOverlaySubInstance(this, NewSubInstance);
+							// Unlink any layer nodes in the new sub instance, as they may have been hooked up to self in InitializeAnimation above.
+							UnlinkLayerNodesInInstance(NewSubInstance);
 
-							// Init after we link in the new graph segments, so propagation happens correctly
-							NewSubInstance->InitializeAnimation();
+							LayerNode->SetLayerOverlaySubInstance(this, NewSubInstance);
 
 							// Initialize the correct parts of the sub instance
 							if(LayerNode->LinkedRoot)
@@ -2630,15 +2641,16 @@ void UAnimInstance::PerformLayerOverlayOperation(TSubclassOf<UAnimInstance> InCl
 					{
 						// Create and add one sub-instance for this group
 						UAnimInstance* NewSubInstance = NewObject<UAnimInstance>(MeshComp, ClassToSet);
+						NewSubInstance->InitializeAnimation();
+
+						// Unlink any layer nodes in the new sub instance, as they may have been hooked up to self in InitializeAnimation above.
+						UnlinkLayerNodesInInstance(NewSubInstance);
 
 						for(FAnimNode_Layer* LayerNode : LayerPair.Value)
 						{
 							LayerNode->SetLayerOverlaySubInstance(this, NewSubInstance);
 						}
 
-						// Init after we link in the new graph segments, so propagation happens correctly
-						NewSubInstance->InitializeAnimation();
-
 						FAnimInstanceProxy& Proxy = NewSubInstance->GetProxyOnAnyThread<FAnimInstanceProxy>();
 
 						// Initialize the correct parts of the sub instance
