Network class deprecation purge

diff --git a/Engine/Plugins/Online/OnlineSubsystemUtils/Source/OnlineSubsystemUtils/Private/OnlineSubsystemUtils.cpp b/Engine/Plugins/Online/OnlineSubsystemUtils/Source/OnlineSubsystemUtils/Private/OnlineSubsystemUtils.cpp
--- a/Engine/Plugins/Online/OnlineSubsystemUtils/Source/OnlineSubsystemUtils/Private/OnlineSubsystemUtils.cpp
+++ b/Engine/Plugins/Online/OnlineSubsystemUtils/Source/OnlineSubsystemUtils/Private/OnlineSubsystemUtils.cpp
@@ -218,9 +218,14 @@ int32 GetClientPeerIp(FName InstanceName, const FUniqueNetId& UserId)
 				if (ClientConnection && 
 					ClientConnection->PlayerId.ToString() == UserId.ToString())
 				{
-					PRAGMA_DISABLE_DEPRECATION_WARNINGS
-					PeerIp = ClientConnection->GetAddrAsInt();
-					PRAGMA_ENABLE_DEPRECATION_WARNINGS
+					TSharedPtr<const FInternetAddr> ClientAddr = ClientConnection->GetRemoteAddr();
+					if (ClientAddr.IsValid())
+					{
+						uint32 TempAddr = 0;
+						ClientAddr->GetIp(TempAddr);
+						PeerIp = static_cast<int32>(TempAddr);
+					}
+
 					break;
 				}
 			}
diff --git a/Engine/Source/Runtime/Engine/Classes/Engine/ActorChannel.h b/Engine/Source/Runtime/Engine/Classes/Engine/ActorChannel.h
--- a/Engine/Source/Runtime/Engine/Classes/Engine/ActorChannel.h
+++ b/Engine/Source/Runtime/Engine/Classes/Engine/ActorChannel.h
@@ -152,8 +152,6 @@ public:
 	 * It's expected that InActor is either null (releasing the channel's reference) or
 	 * a valid actor that is not PendingKill or PendingKillPending.
 	 */
-	UE_DEPRECATED(4.23, "SetChannelActor has an additional parameter.")
-	void SetChannelActor(AActor* InActor) { SetChannelActor(InActor, ESetChannelActorFlags::None); }
 	void SetChannelActor(AActor* InActor, ESetChannelActorFlags Flags);
 
 	virtual void NotifyActorChannelOpen(AActor* InActor, FInBunch& InBunch);
@@ -309,8 +307,6 @@ public:
 protected:
 	
 	TSharedRef< FObjectReplicator > & FindOrCreateReplicator(UObject* Obj, bool* bOutCreated=nullptr);
-	UE_DEPRECATED(4.22, "Use FindOrCreateReplicator() which takes a raw UObject pointer.")
-	TSharedRef< FObjectReplicator > & FindOrCreateReplicator( const TWeakObjectPtr<UObject>& Obj);
 
 	bool ObjectHasReplicator(const TWeakObjectPtr<UObject>& Obj) const;	// returns whether we have already created a replicator for this object or not
 
diff --git a/Engine/Source/Runtime/Engine/Classes/Engine/DemoNetDriver.h b/Engine/Source/Runtime/Engine/Classes/Engine/DemoNetDriver.h
--- a/Engine/Source/Runtime/Engine/Classes/Engine/DemoNetDriver.h
+++ b/Engine/Source/Runtime/Engine/Classes/Engine/DemoNetDriver.h
@@ -393,10 +393,6 @@ class ENGINE_API UDemoNetDriver : public UNetDriver
 	/** Total number of frames in the demo */
 	int32 DemoTotalFrames;
 
-	/** True if we are at the end of playing a demo */
-	UE_DEPRECATED(4.23, "bDemoPlaybackDone is no longer used.")
-	bool bDemoPlaybackDone;
-
 	/** True if as have paused all of the channels */
 	bool bChannelsArePaused;
 
@@ -439,17 +435,6 @@ class ENGINE_API UDemoNetDriver : public UNetDriver
 	void SaveCheckpoint();
 	void TickCheckpoint();
 
-	UE_DEPRECATED(4.22, "This method will be made private in future versions.")
-	bool LoadCheckpoint(FArchive* GotoCheckpointArchive, int64 GotoCheckpointSkipExtraTimeInMS)
-	{
-		FGotoResult Result;
-		Result.Result = EStreamingOperationResult::Success;
-		Result.ExtraTimeMS = GotoCheckpointSkipExtraTimeInMS;
-		Result.CheckpointInfo.CheckpointIndex = FReplayCheckpointInfo::NO_CHECKPOINT;
-		Result.CheckpointInfo.CheckpointStartTime = FReplayCheckpointInfo::NO_CHECKPOINT;
-		return LoadCheckpoint(Result);
-	}
-
 private:
 
 	bool LoadCheckpoint(const FGotoResult& GotoResult);
@@ -718,10 +703,6 @@ public:
 	bool ConditionallyProcessPlaybackPackets();
 	void ProcessAllPlaybackPackets();
 
-
-	UE_DEPRECATED(4.22, "This method will be made private in future versions.")
-	bool ReadPacket(FArchive& Archive, uint8* OutReadBuffer, int32& OutBufferSize, const int32 MaxBufferSize);
-
 private:
 
 	// Possible values returned by ReadPacket.
@@ -922,15 +903,6 @@ public:
 		return bHasDeltaCheckpoints;
 	}
 
-	/**
-	 * Called when a new ActorChannel is opened, before the Actor is notified.
-	 *
-	 * @param Channel	The channel associated with the actor.
-	 * @param Actor		The actor that was recently serialized.
-	 */
-	UE_DEPRECATED(4.22, "Use NotifyActorChannelOpen instead")
-	void PreNotifyActorChannelOpen(UActorChannel* Channel, AActor* Actor) { NotifyActorChannelOpen(Channel, Actor);  }
-
 	/**
 	 * Gets the actively recording or playback replay (stream) name.
 	 * Note, this will be empty when not recording or playing back.
diff --git a/Engine/Source/Runtime/Engine/Classes/Engine/NetConnection.h b/Engine/Source/Runtime/Engine/Classes/Engine/NetConnection.h
--- a/Engine/Source/Runtime/Engine/Classes/Engine/NetConnection.h
+++ b/Engine/Source/Runtime/Engine/Classes/Engine/NetConnection.h
@@ -185,16 +185,6 @@ struct FDelayedPacket
 	double SendTime;
 
 public:
-	UE_DEPRECATED(4.21, "Use the constructor that takes PacketTraits for allowing for analytics and flags")
-	FORCEINLINE FDelayedPacket(uint8* InData, int32 InSizeBytes, int32 InSizeBits)
-		: Data()
-		, SizeBits(InSizeBits)
-		, Traits()
-		, SendTime(0.0)
-	{
-		Data.AddUninitialized(InSizeBytes);
-		FMemory::Memcpy(Data.GetData(), InData, InSizeBytes);
-	}
 
 	FORCEINLINE FDelayedPacket(uint8* InData, int32 InSizeBits, FOutPacketTraits& InTraits)
 		: Data()
@@ -371,12 +361,6 @@ public:
 	EClientLoginState::Type	ClientLoginState;
 	uint8					ExpectedClientLoginMsgType;	// Used to determine what the next expected control channel msg type should be from a connecting client
 
-	// CD key authentication
-	UE_DEPRECATED(4.23, "CDKeyHash is deprecated.")
-	FString			CDKeyHash;				// Hash of client's CD key
-	UE_DEPRECATED(4.23, "CDKeyResponse is deprecated.")
-	FString			CDKeyResponse;			// Client's response to CD key challenge
-
 	// Internal.
 	UPROPERTY()
 	double			LastReceiveTime;		// Last time a packet was received, for timeout checking.
@@ -662,15 +646,11 @@ private:
 public:
 
 	/** 
-	 * If true, will resend everything this connection has ever sent, since the connection has been open.
 	 *	This functionality is used during replay checkpoints for example, so we can re-use the existing connection and channels to record
 	 *	a version of each actor and capture all properties that have changed since the actor has been alive...
 	 *	This will also act as if it needs to re-open all the channels, etc.
 	 *   NOTE - This doesn't force all exports to happen again though, it will only export new stuff, so keep that in mind.
 	 */
-	UE_DEPRECATED(4.23, "Use ResendAllDataState instead.")
-	bool bResendAllDataSinceOpen;
-
 	EResendAllDataState ResendAllDataState;
 
 #if !UE_BUILD_SHIPPING
@@ -746,13 +726,6 @@ public:
 	/** Describe the connection. */
 	ENGINE_API virtual FString Describe();
 
-	UE_DEPRECATED(4.21, "Use the method that allows for packet traits for analytics and modification")
-	ENGINE_API virtual void LowLevelSend(void* Data, int32 CountBytes, int32 CountBits)
-	{
-		FOutPacketTraits EmptyTraits;
-		LowLevelSend(Data, CountBits, EmptyTraits);
-	}
-
 	/**
 	 * Sends a byte stream to the remote endpoint using the underlying socket
 	 *
@@ -773,10 +746,6 @@ public:
 	/** Make sure this connection is in a reasonable state. */
 	ENGINE_API virtual void AssertValid();
 
-	/** Send an acknowledgment. */
-	UE_DEPRECATED(4.22, "This method will be removed")
-	ENGINE_API virtual void SendAck( int32 PacketId, bool FirstTime=1);
-
 	/**
 	 * flushes any pending data, bundling it into a packet and sending it via LowLevelSend()
 	 * also handles network simulation settings (simulated lag, packet loss, etc) unless bIgnoreSimulation is true
@@ -797,22 +766,6 @@ public:
 	 */
 	ENGINE_API virtual void HandleClientPlayer( class APlayerController* PC, class UNetConnection* NetConnection );
 
-	/** @return the address of the connection as an integer */
-	UE_DEPRECATED(4.23, "Use GetRemoteAddr as it allows direct access to the RemoteAddr and allows for dynamic address sizing.")
-	virtual int32 GetAddrAsInt(void)
-	{
-		PRAGMA_DISABLE_DEPRECATION_WARNINGS
-		if (RemoteAddr.IsValid())
-		{
-			uint32 OutAddr = 0;
-			// Get the host byte order ip addr
-			RemoteAddr->GetIp(OutAddr);
-			return (int32)OutAddr;
-		}
-		PRAGMA_ENABLE_DEPRECATION_WARNINGS
-		return 0;
-	}
-
 	/** @return the port of the connection as an integer */
 	virtual int32 GetAddrPort(void)
 	{
@@ -823,15 +776,6 @@ public:
 		return 0;
 	}
 
-	/**
-	 * Return the platform specific FInternetAddr type, containing this connections address.
-	 * If nullptr is returned, connection is not added to MappedClientConnections, and can't receive net packets which depend on this.
-	 *
-	 * @return	The platform specific FInternetAddr containing this connections address
-	 */
-	UE_DEPRECATED(4.23, "Use GetRemoteAddr to safely get the FInternetAddr tied to this connection")
-	virtual TSharedPtr<FInternetAddr> GetInternetAddr() { return ConstCastSharedPtr<FInternetAddr>(GetRemoteAddr()); }
-
 	/**
 	 * Return the platform specific FInternetAddr type, containing this connections address.
 	 * If nullptr is returned, connection is not added to MappedClientConnections, and can't receive net packets which depend on this.
@@ -894,12 +838,6 @@ public:
 	 */
 	ENGINE_API virtual void InitConnection(UNetDriver* InDriver, EConnectionState InState, const FURL& InURL, int32 InConnectionSpeed=0, int32 InMaxPacket=0);
 
-	UE_DEPRECATED(4.21, "Analytics providers are now handled in the NetDriver")
-	ENGINE_API virtual void InitHandler(TSharedPtr<IAnalyticsProvider> InProvider)
-	{
-		InitHandler();
-	}
-
 	/**
 	 * Initializes the PacketHandler
 	 */
@@ -968,14 +906,6 @@ public:
 
 	// Functions.
 
-	/** Resend any pending acks. */
-	UE_DEPRECATED(4.22, "This method will be removed.")
-	void PurgeAcks();
-
-	/** Send package map to the remote. */
-	UE_DEPRECATED(4.23, "This method will be removed.")
-	void SendPackageMap() {}
-
 	/** 
 	 * Appends the passed in data to the SendBuffer to be sent when FlushNet is called
 	 * @param Bits Data as bits to be appended to the send buffer
@@ -1027,10 +957,6 @@ public:
 	/** @todo document */
 	class UControlChannel* GetControlChannel();
 
-	/** Create a channel. */
-	UE_DEPRECATED(4.22, "Use CreateChannelByName")
-	ENGINE_API UChannel* CreateChannel( EChannelType Type, bool bOpenedLocally, int32 ChannelIndex=INDEX_NONE );
-
 	/** Create a channel. */
 	ENGINE_API UChannel* CreateChannelByName( const FName& ChName, EChannelCreateFlags CreateFlags, int32 ChannelIndex=INDEX_NONE );
 
diff --git a/Engine/Source/Runtime/Engine/Classes/Engine/NetDriver.h b/Engine/Source/Runtime/Engine/Classes/Engine/NetDriver.h
--- a/Engine/Source/Runtime/Engine/Classes/Engine/NetDriver.h
+++ b/Engine/Source/Runtime/Engine/Classes/Engine/NetDriver.h
@@ -630,13 +630,6 @@ struct FDisconnectedClient
 		, DisconnectTime(InDisconnectTime)
 	{
 	}
-
-	UE_DEPRECATED(4.23, "This object does not update address information, and should have const addresses passed in")
-	FDisconnectedClient(TSharedRef<FInternetAddr>& InAddress, double InDisconnectTime)
-		: Address(InAddress)
-		, DisconnectTime(InDisconnectTime)
-	{
-	}
 };
 
 
@@ -796,10 +789,6 @@ public:
 	UPROPERTY(Config)
 	FName NetDriverName;
 
-	/** The UChannel classes that should be used under this net driver */
-	UE_DEPRECATED(4.22, "Use ChannelDefinitions instead")
-	UClass* ChannelClasses[CHTYPE_MAX];
-
 	/** Used to specify available channel types and their associated UClass */
 	UPROPERTY(Config)
 	TArray<FChannelDefinition> ChannelDefinitions;
@@ -808,10 +797,6 @@ public:
 	UPROPERTY()
 	TMap<FName, FChannelDefinition> ChannelDefinitionMap;
 
-	/** @return true if the specified channel type exists. */
-	UE_DEPRECATED(4.22, "Use IsKnownChannelName")
-	bool IsKnownChannelType(int32 Type);
-
 	/** @return true if the specified channel definition exists. */
 	FORCEINLINE bool IsKnownChannelName(const FName& ChName)
 	{
@@ -832,10 +817,6 @@ private:
 
 public:
 
-	/** Creates a new channel of the specified type. If the type is pooled, it will return a pre-created channel */
-	UE_DEPRECATED(4.22, "Use GetOrCreateChannelByName")
-	UChannel* GetOrCreateChannel(EChannelType ChType);
-
 	/** Creates a new channel of the specified type name. If the type is pooled, it will return a pre-created channel */
 	UChannel* GetOrCreateChannelByName(const FName& ChName);
 
@@ -1264,18 +1245,6 @@ public:
 	/** PostTick actions */
 	ENGINE_API virtual void PostTickFlush();
 
-	UE_DEPRECATED(4.21, "Please use the LowLevelSend that requires packet traits for analytics and packet modifiers.")
-	ENGINE_API virtual void LowLevelSend(FString Address, void* Data, int32 CountBits)
-	{
-		FOutPacketTraits EmptyTraits;
-		PRAGMA_DISABLE_DEPRECATION_WARNINGS
-		LowLevelSend(Address, Data, CountBits, EmptyTraits);
-		PRAGMA_ENABLE_DEPRECATION_WARNINGS
-	}
-
-	UE_DEPRECATED(4.22, "Change arguments to support FInternetAddr instead")
-	ENGINE_API virtual void LowLevelSend(FString Address, void* Data, int32 CountBits, FOutPacketTraits& Traits);
-
 	/**
 	 * Sends a 'connectionless' (not associated with a UNetConection) packet, to the specified address.
 	 * NOTE: Address is an abstract format defined by subclasses. Anything calling this, must use an address supplied by the net driver.
@@ -1549,10 +1518,6 @@ protected:
 	/** Unregister all TickDispatch, TickFlush, PostTickFlush to tick in World */
 	ENGINE_API void UnregisterTickEvents(class UWorld* InWorld);
 
-	UE_DEPRECATED(4.22, "Use InternalCreateChannelByName")
-	/** Subclasses may override this to customize channel creation. Called by GetOrCreateChannel if the pool is exhausted and a new channel must be allocated. */
-	ENGINE_API virtual UChannel* InternalCreateChannel(EChannelType ChType);
-
 	/** Subclasses may override this to customize channel creation. Called by GetOrCreateChannel if the pool is exhausted and a new channel must be allocated. */
 	ENGINE_API virtual UChannel* InternalCreateChannelByName(const FName& ChName);
 
diff --git a/Engine/Source/Runtime/Engine/Classes/Engine/PackageMapClient.h b/Engine/Source/Runtime/Engine/Classes/Engine/PackageMapClient.h
--- a/Engine/Source/Runtime/Engine/Classes/Engine/PackageMapClient.h
+++ b/Engine/Source/Runtime/Engine/Classes/Engine/PackageMapClient.h
@@ -33,7 +33,6 @@
 class UNetConnection;
 class UNetDriver;
 
-PRAGMA_DISABLE_DEPRECATION_WARNINGS	// 4.22, Name, Type
 class ENGINE_API FNetFieldExport
 {
 public:
@@ -41,19 +40,6 @@ public:
 	{
 	}
 
-	UE_DEPRECATED(4.22, "Type is no longer required, please use other constructor.")
-	FNetFieldExport( const uint32 InHandle, const uint32 InCompatibleChecksum, const FString InName, FString InType ) :
-		Handle( InHandle ),
-		CompatibleChecksum( InCompatibleChecksum ),
-		ExportName( *InName ),
-		Name( InName ),
-		Type( InType ),
-		bExported( false ),
-		bDirtyForReplay( true ),
-		bIncompatible( false )
-	{
-	}
-
 	FNetFieldExport( const uint32 InHandle, const uint32 InCompatibleChecksum, const FName& InName ) :
 		Handle( InHandle ),
 		CompatibleChecksum( InCompatibleChecksum ),
@@ -82,10 +68,13 @@ public:
 
 			if (Ar.IsLoading() && Ar.EngineNetVer() < HISTORY_NETEXPORT_SERIALIZATION)
 			{
-				Ar << C.Name;
-				Ar << C.Type;
+				FName TempName;
+				FString TempType;
+
+				Ar << TempName;
+				Ar << TempType;
 
-				C.ExportName = FName(*C.Name);
+				C.ExportName = TempName;
 			}
 			else
 			{
@@ -97,11 +86,6 @@ public:
 				{
 					UPackageMap::StaticSerializeName(Ar, C.ExportName);
 				}
-
-				if (Ar.IsLoading())
-				{
-					C.Name = C.ExportName.ToString();
-				}
 			}
 		}
 
@@ -113,17 +97,12 @@ public:
 	uint32			Handle;
 	uint32			CompatibleChecksum;
 	FName			ExportName;
-	UE_DEPRECATED(4.22, "Name is deprecated.")
-	FString			Name;
-	UE_DEPRECATED(4.22, "Type is deprecated.")
-	FString			Type;
 	bool			bExported;
 	bool			bDirtyForReplay;
 
 	// Transient properties
 	mutable bool	bIncompatible;		// If true, we've already determined that this property isn't compatible. We use this to curb warning spam.
 };
-PRAGMA_ENABLE_DEPRECATION_WARNINGS
 
 class ENGINE_API FNetFieldExportGroup
 {
@@ -295,9 +274,6 @@ public:
 
 	UNetDriver *									Driver;
 
-	UE_DEPRECATED(4.23, "This member will be made private in future versions")
-	TMap< FName, FNetworkGUID >						PendingAsyncPackages;
-
 	ENetworkChecksumMode							NetworkChecksumMode;
 	EAsyncLoadMode									AsyncLoadMode;
 
@@ -307,6 +283,8 @@ private:
 
 	friend class UPackageMapClient;
 
+	TMap<FName, FNetworkGUID> PendingAsyncPackages;
+
 	/** Maps net field export group name to the respective FNetFieldExportGroup */
 	TMap < FString, TSharedPtr< FNetFieldExportGroup > >	NetFieldExportGroupMap;
 
diff --git a/Engine/Source/Runtime/Engine/Private/DataChannel.cpp b/Engine/Source/Runtime/Engine/Private/DataChannel.cpp
--- a/Engine/Source/Runtime/Engine/Private/DataChannel.cpp
+++ b/Engine/Source/Runtime/Engine/Private/DataChannel.cpp
@@ -3845,11 +3845,6 @@ TSharedRef< FObjectReplicator > & UActorChannel::FindOrCreateReplicator( UObject
 	return *ReplicatorRefPtr;
 }
 
-TSharedRef< FObjectReplicator > & UActorChannel::FindOrCreateReplicator(const TWeakObjectPtr<UObject>& Obj)
-{
-	return FindOrCreateReplicator(Obj.Get(), static_cast<bool*>(nullptr));
-}
-
 bool UActorChannel::ObjectHasReplicator(const TWeakObjectPtr<UObject>& Obj) const
 {
 	const TSharedRef<FObjectReplicator>* ReplicatorRefPtr = ReplicationMap.Find(Obj.Get());
diff --git a/Engine/Source/Runtime/Engine/Private/DataReplication.cpp b/Engine/Source/Runtime/Engine/Private/DataReplication.cpp
--- a/Engine/Source/Runtime/Engine/Private/DataReplication.cpp
+++ b/Engine/Source/Runtime/Engine/Private/DataReplication.cpp
@@ -51,7 +51,6 @@ static FAutoConsoleVariableRef CVarSupportsFastArrayDelta(
 
 extern TAutoConsoleVariable<int32> CVarNetEnableDetailedScopeCounters;
 
-PRAGMA_DISABLE_DEPRECATION_WARNINGS
 class FNetSerializeCB : public INetSerializeCB
 {
 private:
@@ -274,9 +273,7 @@ private:
 	FCachedRequestState CachedRequestState;
 	TSharedPtr<FReplicationChangelistMgr> ChangelistMgr;
 };
-PRAGMA_ENABLE_DEPRECATION_WARNINGS
 
-PRAGMA_DISABLE_DEPRECATION_WARNINGS
 FObjectReplicator::FObjectReplicator()
 	: bLastUpdateEmpty(false)
 	, bOpenAckCalled(false)
@@ -295,60 +292,6 @@ FObjectReplicator::~FObjectReplicator()
 {
 	CleanUp();
 }
-PRAGMA_ENABLE_DEPRECATION_WARNINGS
-
-PRAGMA_DISABLE_DEPRECATION_WARNINGS
-bool FObjectReplicator::SerializeCustomDeltaProperty(
-	UNetConnection* Connection,
-	void* Src,
-	UProperty* Property,
-	uint32 ArrayIndex,
-	FNetBitWriter& OutBunch,
-	TSharedPtr<INetDeltaBaseState>& NewFullState,
-	TSharedPtr<INetDeltaBaseState>& OldState)
-{
-	check( NewFullState.IsValid() == false ); // NewState is passed in as nullptr and instantiated within this function if necessary
-
-	CONDITIONAL_SCOPE_CYCLE_COUNTER( STAT_NetSerializeItemDeltaTime, CVarNetEnableDetailedScopeCounters.GetValueOnAnyThread() > 0 );
-
-	UStructProperty * StructProperty = CastChecked< UStructProperty >( Property );
-
-	//------------------------------------------------
-	//	Custom NetDeltaSerialization
-	//------------------------------------------------
-	if ( !ensure( ( StructProperty->Struct->StructFlags & STRUCT_NetDeltaSerializeNative ) != 0 ) )
-	{
-		return false;
-	}
-
-	FNetSerializeCB NetSerializeCB( Connection->Driver );
-
-	FNetDeltaSerializeInfo Parms;
-	Parms.Data = Property->ContainerPtrToValuePtr<void>(Src, ArrayIndex);
-	Parms.Object = reinterpret_cast<UObject*>(Src);
-	Parms.Connection = Connection;
-	Parms.bInternalAck = Connection->InternalAck;
-	Parms.Writer = &OutBunch;
-	Parms.Map = Connection->PackageMap;
-	Parms.OldState = OldState.Get();
-	Parms.NewState = &NewFullState;
-	Parms.NetSerializeCB = &NetSerializeCB;
-	Parms.bIsWritingOnClient = (Connection->Driver && Connection->Driver->GetWorld()) ? Connection->Driver->GetWorld()->IsRecordingClientReplay() : false;
-
-	UScriptStruct::ICppStructOps * CppStructOps = StructProperty->Struct->GetCppStructOps();
-
-	check(CppStructOps); // else should not have STRUCT_NetSerializeNative
-
-	Parms.Struct = StructProperty->Struct;
-
-	if ( Property->ArrayDim != 1 )
-	{
-		OutBunch.SerializeIntPacked( ArrayIndex );
-	}
-
-	return CppStructOps->NetDeltaSerialize(Parms, Parms.Data);
-}
-PRAGMA_ENABLE_DEPRECATION_WARNINGS
 
 bool FObjectReplicator::SendCustomDeltaProperty(UObject* InObject, UProperty* Property, uint32 ArrayIndex, FNetBitWriter& OutBunch, TSharedPtr<INetDeltaBaseState>& NewFullState, TSharedPtr<INetDeltaBaseState>& OldState)
 {
@@ -1746,15 +1689,6 @@ void FObjectReplicator::FRPCPendingLocalCall::CountBytes(FArchive& Ar) const
 	UnmappedGuids.CountBytes(Ar);
 }
 
-
-void FObjectReplicator::Serialize(FArchive& Ar)
-{
-	if (Ar.IsCountingMemory())
-	{
-		CountBytes(Ar);
-	}
-}
-
 void FObjectReplicator::CountBytes(FArchive& Ar) const
 {
 	GRANULAR_NETWORK_MEMORY_TRACKING_INIT(Ar, "FObjectReplicator::CountBytes");
diff --git a/Engine/Source/Runtime/Engine/Private/DemoNetDriver.cpp b/Engine/Source/Runtime/Engine/Private/DemoNetDriver.cpp
--- a/Engine/Source/Runtime/Engine/Private/DemoNetDriver.cpp
+++ b/Engine/Source/Runtime/Engine/Private/DemoNetDriver.cpp
@@ -3358,57 +3358,6 @@ bool UDemoNetDriver::ConditionallyReadDemoFrameIntoPlaybackPackets(FArchive& Ar)
 	return true;
 }
 
-// Deprecated, DO NOT USE.
-bool UDemoNetDriver::ReadPacket(FArchive& Archive, uint8* OutReadBuffer, int32& OutBufferSize, const int32 MaxBufferSize)
-{
-	OutBufferSize = 0;
-
-	Archive << OutBufferSize;
-
-	if (Archive.IsError())
-	{
-		UE_LOG(LogDemo, Error, TEXT("UDemoNetDriver::ReadPacket: Failed to read demo OutBufferSize"));
-		return false;
-	}
-
-	if (OutBufferSize == 0)
-	{
-		return true;		// Done
-	}
-
-	if (OutBufferSize > MaxBufferSize)
-	{
-		UE_LOG(LogDemo, Error, TEXT("UDemoNetDriver::ReadPacket: OutBufferSize > sizeof( ReadBuffer )"));
-		return false;
-	}
-
-	// Read data from file.
-	Archive.Serialize(OutReadBuffer, OutBufferSize);
-
-	if (Archive.IsError())
-	{
-		UE_LOG(LogDemo, Error, TEXT("UDemoNetDriver::ReadPacket: Failed to read demo file packet"));
-		return false;
-	}
-
-#if DEMO_CHECKSUMS == 1
-	{
-		uint32 ServerChecksum = 0;
-		Archive << ServerChecksum;
-
-		const uint32 Checksum = FCrc::MemCrc32(OutReadBuffer, OutBufferSize, 0);
-
-		if (Checksum != ServerChecksum)
-		{
-			UE_LOG(LogDemo, Error, TEXT("UDemoNetDriver::ReadPacket: Checksum != ServerChecksum"));
-			return false;
-		}
-	}
-#endif
-
-	return true;
-}
-
 const UDemoNetDriver::EReadPacketState UDemoNetDriver::ReadPacket(FArchive& Archive, TArray<uint8>& OutBuffer, const EReadPacketMode Mode)
 {
 	const bool bSkipData = (EReadPacketMode::SkipData == Mode);
diff --git a/Engine/Source/Runtime/Engine/Private/NetConnection.cpp b/Engine/Source/Runtime/Engine/Private/NetConnection.cpp
--- a/Engine/Source/Runtime/Engine/Private/NetConnection.cpp
+++ b/Engine/Source/Runtime/Engine/Private/NetConnection.cpp
@@ -211,9 +211,6 @@ UNetConnection::UNetConnection(const FObjectInitializer& ObjectInitializer)
 ,	InitInReliable		( 0 )
 ,	EngineNetworkProtocolVersion( FNetworkVersion::GetEngineNetworkProtocolVersion() )
 ,	GameNetworkProtocolVersion( FNetworkVersion::GetGameNetworkProtocolVersion() )
-PRAGMA_DISABLE_DEPRECATION_WARNINGS
-,	bResendAllDataSinceOpen( false )
-PRAGMA_ENABLE_DEPRECATION_WARNINGS
 ,	ResendAllDataState( EResendAllDataState::None )
 #if !UE_BUILD_SHIPPING
 ,	ReceivedRawPacketDel()
@@ -593,11 +590,6 @@ void UNetConnection::Serialize( FArchive& Ar )
 		GRANULAR_NETWORK_MEMORY_TRACKING_TRACK("ClientResponse", ClientResponse.CountBytes(Ar));
 		GRANULAR_NETWORK_MEMORY_TRACKING_TRACK("RequestURL", RequestURL.CountBytes(Ar));
 
-PRAGMA_DISABLE_DEPRECATION_WARNINGS
-		GRANULAR_NETWORK_MEMORY_TRACKING_TRACK("CDKeyHash", CDKeyHash.CountBytes(Ar));
-		GRANULAR_NETWORK_MEMORY_TRACKING_TRACK("CDKeyResponse", CDKeyResponse.CountBytes(Ar));
-PRAGMA_ENABLE_DEPRECATION_WARNINGS
-
 		GRANULAR_NETWORK_MEMORY_TRACKING_TRACK("SendBuffer", SendBuffer.CountMemory(Ar));
 
 		GRANULAR_NETWORK_MEMORY_TRACKING_TRACK("Channels", Channels.CountBytes(Ar));
@@ -2607,14 +2599,6 @@ TSharedPtr<FObjectReplicator> UNetConnection::CreateReplicatorForNewActorChannel
 	return NewReplicator;
 }
 
-void UNetConnection::PurgeAcks()
-{
-}
-
-void UNetConnection::SendAck(int32 AckPacketId, bool FirstTime/*=1*/)
-{
-}
-
 int32 UNetConnection::SendRawBunch( FOutBunch& Bunch, bool InAllowMerge )
 {
 	ValidateSendBuffer();
@@ -2704,27 +2688,6 @@ int32 UNetConnection::SendRawBunch( FOutBunch& Bunch, bool InAllowMerge )
 	return Bunch.PacketId;
 }
 
-UChannel* UNetConnection::CreateChannel( EChannelType Type, bool bOpenedLocally, int32 ChannelIndex )
-{
-	const EChannelCreateFlags ChannelCreateFlags = (bOpenedLocally ? EChannelCreateFlags::OpenedLocally : EChannelCreateFlags::None);
-	FName ChName = NAME_None;
-
-	switch (Type)
-	{
-	case CHTYPE_Control:
-		ChName = NAME_Control;
-		break;
-	case CHTYPE_Actor:
-		ChName = NAME_Actor;
-		break;
-	case CHTYPE_Voice:
-		ChName = NAME_Voice;
-		break;
-	}
-
-	return CreateChannelByName(ChName, ChannelCreateFlags, ChannelIndex);
-}
-
 UChannel* UNetConnection::CreateChannelByName( const FName& ChName, EChannelCreateFlags CreateFlags, int32 ChIndex )
 {
 	check(Driver->IsKnownChannelName(ChName));
diff --git a/Engine/Source/Runtime/Engine/Private/NetworkDriver.cpp b/Engine/Source/Runtime/Engine/Private/NetworkDriver.cpp
--- a/Engine/Source/Runtime/Engine/Private/NetworkDriver.cpp
+++ b/Engine/Source/Runtime/Engine/Private/NetworkDriver.cpp
@@ -307,12 +307,6 @@ UNetDriver::UNetDriver(const FObjectInitializer& ObjectInitializer)
 ,	PacketLossBurstEndTime(-1.0f)
 ,	OutTotalNotifiedPackets(0u)
 {
-PRAGMA_DISABLE_DEPRECATION_WARNINGS
-	ChannelClasses[CHTYPE_Control]	= UControlChannel::StaticClass();
-	ChannelClasses[CHTYPE_Actor]	= UActorChannel::StaticClass();
-	ChannelClasses[CHTYPE_Voice]	= UVoiceChannel::StaticClass();
-PRAGMA_ENABLE_DEPRECATION_WARNINGS
-
 	UpdateDelayRandomStream.Initialize(FApp::bUseFixedSeed ? GetFName() : NAME_None);
 }
 
@@ -1372,22 +1366,6 @@ void UNetDriver::PostTickFlush()
 	}
 }
 
-PRAGMA_DISABLE_DEPRECATION_WARNINGS
-void UNetDriver::LowLevelSend(FString Address, void* Data, int32 CountBits, FOutPacketTraits& Traits)
-{
-	if (GetSocketSubsystem() != nullptr)
-	{
-		TSharedPtr<FInternetAddr> NetAddress = GetSocketSubsystem()->GetAddressFromString(*Address);
-		if (NetAddress.IsValid())
-		{
-			LowLevelSend(NetAddress, Data, CountBits, Traits);
-			return;
-		}
-	}
-	UE_LOG(LogNet, Warning, TEXT("Could not infer the address to use for LowLevelSend, please use the one that takes an FInternetAddr to avoid this problem"));
-}
-PRAGMA_ENABLE_DEPRECATION_WARNINGS
-
 bool UNetDriver::InitConnectionClass(void)
 {
 	if (NetConnectionClass == NULL && NetConnectionClassName != TEXT(""))
@@ -4658,46 +4636,6 @@ int32 UNetDriver::ServerReplicateActors(float DeltaSeconds)
 #endif // WITH_SERVER_CODE
 }
 
-bool UNetDriver::IsKnownChannelType(int32 Type)
-{
-	FName ChName = NAME_None;
-
-	switch (Type)
-	{
-	case CHTYPE_Control:
-		ChName = NAME_Control;
-		break;
-	case CHTYPE_Actor:
-		ChName = NAME_Actor;
-		break;
-	case CHTYPE_Voice:
-		ChName = NAME_Voice;
-		break;
-	}
-
-	return IsKnownChannelName(ChName);
-}
-
-UChannel* UNetDriver::GetOrCreateChannel(EChannelType ChType)
-{
-	FName ChName = NAME_None;
-
-	switch (ChType)
-	{
-	case CHTYPE_Control:
-		ChName = NAME_Control;
-		break;
-	case CHTYPE_Actor:
-		ChName = NAME_Actor;
-		break;
-	case CHTYPE_Voice:
-		ChName = NAME_Voice;
-		break;
-	}
-
-	return GetOrCreateChannelByName(ChName);
-}
-
 UChannel* UNetDriver::GetOrCreateChannelByName(const FName& ChName)
 {
 	UChannel* RetVal = nullptr;
@@ -4729,26 +4667,6 @@ UChannel* UNetDriver::GetOrCreateChannelByName(const FName& ChName)
 	return RetVal;
 }
 
-UChannel* UNetDriver::InternalCreateChannel(EChannelType Type)
-{
-	FName ChName = NAME_None;
-
-	switch (Type)
-	{
-	case CHTYPE_Control:
-		ChName = NAME_Control;
-		break;
-	case CHTYPE_Actor:
-		ChName = NAME_Actor;
-		break;
-	case CHTYPE_Voice:
-		ChName = NAME_Voice;
-		break;
-	}
-
-	return InternalCreateChannelByName(ChName);
-}
-
 UChannel* UNetDriver::InternalCreateChannelByName(const FName& ChName)
 {
 	return NewObject<UChannel>(GetTransientPackage(), ChannelDefinitionMap[ChName].ChannelClass);
diff --git a/Engine/Source/Runtime/Engine/Private/PackageMapClient.cpp b/Engine/Source/Runtime/Engine/Private/PackageMapClient.cpp
--- a/Engine/Source/Runtime/Engine/Private/PackageMapClient.cpp
+++ b/Engine/Source/Runtime/Engine/Private/PackageMapClient.cpp
@@ -3465,10 +3465,7 @@ PRAGMA_ENABLE_DEPRECATION_WARNINGS
 
 void FNetFieldExport::CountBytes(FArchive& Ar) const
 {
-	PRAGMA_DISABLE_DEPRECATION_WARNINGS
-	Name.CountBytes(Ar);
-	Type.CountBytes(Ar);
-	PRAGMA_ENABLE_DEPRECATION_WARNINGS
+
 }
 
 void FNetFieldExportGroup::CountBytes(FArchive& Ar) const
diff --git a/Engine/Source/Runtime/Engine/Public/Net/DataReplication.h b/Engine/Source/Runtime/Engine/Public/Net/DataReplication.h
--- a/Engine/Source/Runtime/Engine/Public/Net/DataReplication.h
+++ b/Engine/Source/Runtime/Engine/Public/Net/DataReplication.h
@@ -129,15 +129,6 @@ public:
 	/** Takes Data, and compares against shadow state to log differences */
 	bool ValidateAgainstState(const UObject* ObjectState);
 
-	UE_DEPRECATED(4.23, "This method will be removed in future versions. Please use SendCustomDeltaProperty instead.")
-	static bool SerializeCustomDeltaProperty(
-		UNetConnection* Connection,
-		void* Src, UProperty* Property,
-		uint32 ArrayIndex,
-		FNetBitWriter& OutBunch,
-		TSharedPtr<INetDeltaBaseState>& NewFullState,
-		TSharedPtr<INetDeltaBaseState> & OldState);
-
 	//~ Both of these should be private, IMO, but we'll leave them public for now for back compat
 	//~ in case anyone was using SerializeCustomDeltaProperty already.
 
@@ -159,9 +150,6 @@ public:
 	/** Packet was dropped */
 	void ReceivedNak(int32 NakPacketId);
 
-	UE_DEPRECATED(4.23, "Use CountBytes instead")
-	void Serialize(FArchive& Ar);
-
 	void CountBytes(FArchive& Ar) const;
 
 	/** Writes dirty properties to bunch */
