[NO NON-PII COMMIT MESSAGE]

diff --git a/Engine/Source/Runtime/AudioPlatformConfiguration/Public/AudioCompressionSettings.h b/Engine/Source/Runtime/AudioPlatformConfiguration/Public/AudioCompressionSettings.h
--- a/Engine/Source/Runtime/AudioPlatformConfiguration/Public/AudioCompressionSettings.h
+++ b/Engine/Source/Runtime/AudioPlatformConfiguration/Public/AudioCompressionSettings.h
@@ -46,7 +46,7 @@ struct FPlatformAudioCookOverrides
 	// Increment this to force a recook on all Stream Caching assets.
 	// For testing, it's useful to set this to either a negative number or
 	// absurdly large number, to ensure you do not pollute the DDC.
-	static const int32 StreamCachingVersion = 5019;
+	static const int32 StreamCachingVersion = 5021;
 
 	bool bResampleForDevice;
 
diff --git a/Engine/Source/Runtime/Engine/Private/AudioCompressionSettingsUtils.cpp b/Engine/Source/Runtime/Engine/Private/AudioCompressionSettingsUtils.cpp
--- a/Engine/Source/Runtime/Engine/Private/AudioCompressionSettingsUtils.cpp
+++ b/Engine/Source/Runtime/Engine/Private/AudioCompressionSettingsUtils.cpp
@@ -83,7 +83,7 @@ void CacheCurrentPlatformAudioCookOverrides(FPlatformAudioCookOverrides& OutOver
 #elif PLATFORM_WINDOWS
 	const TCHAR* CategoryName = TEXT("/Script/WindowsTargetPlatform.WindowsTargetSettings");
 #elif PLATFORM_PS4
-	const TCHAR* CategoryName = TEXT("/Script/PS4TargetPlatform.PS4TargetSettings");
+	const TCHAR* CategoryName = TEXT("/Script/PS4PlatformEditor.PS4TargetSettings");
 #elif PLATFORM_XBOXONE
 	const TCHAR* CategoryName = TEXT("/Script/XboxOneTargetPlatform.XboxOneTargetSettings");
 #else
diff --git a/Engine/Source/Runtime/Engine/Private/AudioDecompress.cpp b/Engine/Source/Runtime/Engine/Private/AudioDecompress.cpp
--- a/Engine/Source/Runtime/Engine/Private/AudioDecompress.cpp
+++ b/Engine/Source/Runtime/Engine/Private/AudioDecompress.cpp
@@ -14,6 +14,7 @@ IStreamedCompressedInfo::IStreamedCompressedInfo()
 	, SrcBufferDataSize(0)
 	, SrcBufferOffset(0)
 	, AudioDataOffset(0)
+	, AudioDataChunkIndex(0)
 	, SampleRate(0)
 	, TrueSampleCount(0)
 	, CurrentSampleCount(0)
@@ -153,7 +154,20 @@ bool IStreamedCompressedInfo::StreamCompressedInfoInternal(USoundWave* Wave, str
 	bIsStreaming = true;
 	if (FirstChunk)
 	{
-		return ReadCompressedInfo(FirstChunk, ChunkSize, QualityInfo);
+		bool HeaderReadResult = ReadCompressedInfo(FirstChunk, ChunkSize, QualityInfo);
+		
+		// If we've read through the entirety of the zeroth chunk while parsing the header, move onto the next chunk.
+		if (SrcBufferOffset >= ChunkSize)
+		{
+			CurrentChunkIndex++;
+			SrcBufferData = NULL;
+			SrcBufferDataSize = 0;
+
+			AudioDataChunkIndex = CurrentChunkIndex;
+			AudioDataOffset -= ChunkSize;
+		}
+
+		return HeaderReadResult;
 	}
 
 	return false;
@@ -242,7 +256,7 @@ bool IStreamedCompressedInfo::StreamCompressedData(uint8* Destination, bool bLoo
 
 					if (bLooping)
 					{
-						CurrentChunkIndex = 0;
+						CurrentChunkIndex = AudioDataChunkIndex;
 						SrcBufferOffset = AudioDataOffset;
 						CurrentSampleCount = 0;
 
diff --git a/Engine/Source/Runtime/Engine/Public/AudioDecompress.h b/Engine/Source/Runtime/Engine/Public/AudioDecompress.h
--- a/Engine/Source/Runtime/Engine/Public/AudioDecompress.h
+++ b/Engine/Source/Runtime/Engine/Public/AudioDecompress.h
@@ -278,6 +278,8 @@ protected:
 	uint32 SrcBufferOffset;
 	/** Where the actual audio data starts in the current streamed chunk. Accounts for header offset. */
 	uint32 AudioDataOffset;
+	/** The chunk index where the actual audio data starts. */
+	uint32 AudioDataChunkIndex;
 	/** Sample rate of the source file */
 	uint16 SampleRate;
 	/** The total sample count of the source file. */
