[NO NON-PII COMMIT MESSAGE]

diff --git a/Engine/Source/Runtime/AudioMixer/Private/AudioMixerBlueprintLibrary.cpp b/Engine/Source/Runtime/AudioMixer/Private/AudioMixerBlueprintLibrary.cpp
--- a/Engine/Source/Runtime/AudioMixer/Private/AudioMixerBlueprintLibrary.cpp
+++ b/Engine/Source/Runtime/AudioMixer/Private/AudioMixerBlueprintLibrary.cpp
@@ -379,6 +379,14 @@ void UAudioMixerBlueprintLibrary::PrimeSoundForPlayback(USoundWave* SoundWave, c
 	}
 }
 
+void UAudioMixerBlueprintLibrary::PrimeSoundCueForPlayback(USoundCue* SoundCue)
+{
+	if (SoundCue)
+	{
+		SoundCue->PrimeSoundCue();
+	}
+}
+
 float UAudioMixerBlueprintLibrary::TrimAudioCache(float InMegabytesToFree)
 {
 	uint64 NumBytesToFree = (uint64) (((double)InMegabytesToFree) * 1024.0 * 1024.0);
diff --git a/Engine/Source/Runtime/AudioMixer/Public/AudioMixerBlueprintLibrary.h b/Engine/Source/Runtime/AudioMixer/Public/AudioMixerBlueprintLibrary.h
--- a/Engine/Source/Runtime/AudioMixer/Public/AudioMixerBlueprintLibrary.h
+++ b/Engine/Source/Runtime/AudioMixer/Public/AudioMixerBlueprintLibrary.h
@@ -7,6 +7,7 @@
 #include "SubmixEffects/AudioMixerSubmixEffectDynamicsProcessor.h"
 #include "Sound/SoundEffectSource.h"
 #include "SampleBuffer.h"
+#include "Sound/SoundCue.h"
 #include "DSP/SpectrumAnalyzer.h"
 #include "AudioMixerBlueprintLibrary.generated.h"
 
@@ -145,6 +146,10 @@ public:
 	UFUNCTION(BlueprintCallable, Category = "Sound")
 	static void PrimeSoundForPlayback(USoundWave* SoundWave, const FOnSoundLoadComplete OnLoadCompletion);
 
+	/** Begin loading any sounds referenced by a sound cue into the cache so that it can be played immediately. */
+	UFUNCTION(BlueprintCallable, Category = "Sound")
+	static void PrimeSoundCueForPlayback(USoundCue* SoundCue);
+
 	/** Trim memory used by the audio cache. Returns the number of megabytes freed. */
 	UFUNCTION(BlueprintCallable, Category = "Sound")
 	static float TrimAudioCache(float InMegabytesToFree);
diff --git a/Engine/Source/Runtime/Engine/Classes/Sound/SoundCue.h b/Engine/Source/Runtime/Engine/Classes/Sound/SoundCue.h
--- a/Engine/Source/Runtime/Engine/Classes/Sound/SoundCue.h
+++ b/Engine/Source/Runtime/Engine/Classes/Sound/SoundCue.h
@@ -77,6 +77,10 @@ class ENGINE_API USoundCue : public USoundBase
 {
 	GENERATED_UCLASS_BODY()
 
+	/* Makes this sound cue automatically load any sound waves it can play into the cache when it is loaded. */
+	UPROPERTY(EditAnywhere, Category = Caching)
+	uint32 bPrimeOnLoad : 1;
+
 	UPROPERTY()
 	USoundNode* FirstNode;
 
@@ -238,6 +242,9 @@ public:
 	/** Call to cache any values which need to be computed from the sound cue graph. e.g. MaxDistance, Duration, etc. */
 	void CacheAggregateValues();
 
+	/** Call this when stream caching is enabled to prime all SoundWave assets referenced by this Sound Cue. */
+	void PrimeSoundCue();
+
 protected:
 	bool RecursiveFindPathToNode(USoundNode* CurrentNode, const UPTRINT CurrentHash, const UPTRINT NodeHashToFind, TArray<USoundNode*>& OutPath) const;
 
diff --git a/Engine/Source/Runtime/Engine/Classes/Sound/SoundNode.h b/Engine/Source/Runtime/Engine/Classes/Sound/SoundNode.h
--- a/Engine/Source/Runtime/Engine/Classes/Sound/SoundNode.h
+++ b/Engine/Source/Runtime/Engine/Classes/Sound/SoundNode.h
@@ -189,5 +189,15 @@ public:
 	 */
 	static UPTRINT GetNodeWaveInstanceHash(const UPTRINT ParentWaveInstanceHash, const USoundNode* ChildNode, const uint32 ChildIndex);
 	static UPTRINT GetNodeWaveInstanceHash(const UPTRINT ParentWaveInstanceHash, const UPTRINT ChildNodeHash, const uint32 ChildIndex);
+
+
+	/**
+	 * When this is called and stream caching is enabled,
+	 * any wave player sound nodes childed off of this node
+	 * will have their audio loaded into the cache.
+	 *
+	 * @param bRecurse when true, this will cause all children of child nodes to be primed as well.
+	 */
+	void PrimeChildWavePlayers(bool bRecurse);
 };
 
diff --git a/Engine/Source/Runtime/Engine/Private/SoundCue.cpp b/Engine/Source/Runtime/Engine/Private/SoundCue.cpp
--- a/Engine/Source/Runtime/Engine/Private/SoundCue.cpp
+++ b/Engine/Source/Runtime/Engine/Private/SoundCue.cpp
@@ -87,6 +87,11 @@ void USoundCue::CacheAggregateValues()
 	}
 }
 
+void USoundCue::PrimeSoundCue()
+{
+	FirstNode->PrimeChildWavePlayers(true);
+}
+
 void USoundCue::Serialize(FStructuredArchive::FRecord Record)
 {
 	FArchive& UnderlyingArchive = Record.GetUnderlyingArchive();
@@ -153,6 +158,11 @@ void USoundCue::PostLoad()
 	}
 
 	CacheAggregateValues();
+	
+	if (bPrimeOnLoad)
+	{
+		FirstNode->PrimeChildWavePlayers(true);
+	}
 }
 
 bool USoundCue::CanBeClusterRoot() const
diff --git a/Engine/Source/Runtime/Engine/Private/SoundNode.cpp b/Engine/Source/Runtime/Engine/Private/SoundNode.cpp
--- a/Engine/Source/Runtime/Engine/Private/SoundNode.cpp
+++ b/Engine/Source/Runtime/Engine/Private/SoundNode.cpp
@@ -5,6 +5,10 @@
 #include "EngineUtils.h"
 #include "Sound/SoundCue.h"
 #include "Misc/App.h"
+#include "Sound/SoundNodeWavePlayer.h"
+#include "ContentStreaming.h"
+#include "Sound/SoundWave.h"
+#include "AudioCompressionSettingsUtils.h"
 
 /*-----------------------------------------------------------------------------
 	USoundNode implementation.
@@ -86,6 +90,35 @@ UPTRINT USoundNode::GetNodeWaveInstanceHash(const UPTRINT ParentWaveInstanceHash
 #endif // USE_NEW_SOUNDCUE_NODE_HASH
 }
 
+void USoundNode::PrimeChildWavePlayers(bool bRecurse)
+{
+	if (FPlatformCompressionUtilities::IsCurrentPlatformUsingStreamCaching())
+	{
+		// Search child nodes for wave players, then prime their waves.
+		for (USoundNode* ChildNode : ChildNodes)
+		{
+			if (ChildNode)
+			{
+				ChildNode->ConditionalPostLoad();
+				if (bRecurse)
+				{
+					ChildNode->PrimeChildWavePlayers(true);
+				}
+
+				USoundNodeWavePlayer* WavePlayer = Cast<USoundNodeWavePlayer>(ChildNode);
+				if (WavePlayer != nullptr)
+				{
+					USoundWave* SoundWave = WavePlayer->GetSoundWave();
+					if (SoundWave && SoundWave->IsStreaming())
+					{
+						IStreamingManager::Get().GetAudioStreamingManager().RequestChunk(SoundWave, 1, [](EAudioChunkLoadResult) {});
+					}
+				}
+			}
+		}
+	}
+}
+
 void USoundNode::ParseNodes( FAudioDevice* AudioDevice, const UPTRINT NodeWaveInstanceHash, FActiveSound& ActiveSound, const FSoundParseParameters& ParseParams, TArray<FWaveInstance*>& WaveInstances )
 {
 	for( int32 i = 0; i < ChildNodes.Num() && i < GetMaxChildNodes(); ++i )
diff --git a/Engine/Source/Runtime/Engine/Private/SoundNodeDelay.cpp b/Engine/Source/Runtime/Engine/Private/SoundNodeDelay.cpp
--- a/Engine/Source/Runtime/Engine/Private/SoundNodeDelay.cpp
+++ b/Engine/Source/Runtime/Engine/Private/SoundNodeDelay.cpp
@@ -4,6 +4,13 @@
 #include "Sound/SoundNodeDelay.h"
 #include "ActiveSound.h"
 
+static int32 PrimeDelayedSoundNodesCVar = 0;
+FAutoConsoleVariableRef CVarPrimeDelayedSoundNodes(
+	TEXT("au.streamcache.priming.PrimeDelayNodes"),
+	PrimeDelayedSoundNodesCVar,
+	TEXT("When set to 1, sounds will be loaded into the cache automatically when a delay node is hit.\n"),
+	ECVF_Default);
+
 struct FSoundNodeDelayPayload
 {
 	float EndOfDelay;
@@ -33,6 +40,11 @@ void USoundNodeDelay::ParseNodes( FAudioDevice* AudioDevice, const UPTRINT NodeW
 
 		const float ActualDelay = FMath::Max(0.f, DelayMax + ( ( DelayMin - DelayMax ) * RandomStream.FRand() ));
 
+		if (PrimeDelayedSoundNodesCVar != 0)
+		{
+			PrimeChildWavePlayers(true);
+		}
+
 		if (ActualDelay > 0.0f && ParseParams.StartTime >= ActualDelay)
 		{
 			SoundNodeDelayPayload.StartTimeModifier = ActualDelay;
diff --git a/Engine/Source/Runtime/Engine/Private/SoundNodeRandom.cpp b/Engine/Source/Runtime/Engine/Private/SoundNodeRandom.cpp
--- a/Engine/Source/Runtime/Engine/Private/SoundNodeRandom.cpp
+++ b/Engine/Source/Runtime/Engine/Private/SoundNodeRandom.cpp
@@ -17,6 +17,13 @@ FAutoConsoleVariableRef CVarMaxRandomBranches(
 	TEXT("0: No culling, Any other value: The amount of branches we should use as a maximum for any random node."),
 	ECVF_Default);
 
+static int32 PrimeRandomSoundNodesCVar = 0;
+FAutoConsoleVariableRef CVarPrimeRandomSoundNodes(
+	TEXT("au.streamcache.priming.PrimeRandomNodes"),
+	PrimeRandomSoundNodesCVar,
+	TEXT("When set to 1, sounds will be loaded into the cache automatically when a random node is hit.\n"),
+	ECVF_Default);
+
 /*-----------------------------------------------------------------------------
     USoundNodeRandom implementation.
 -----------------------------------------------------------------------------*/
@@ -185,6 +192,11 @@ void USoundNodeRandom::ParseNodes( FAudioDevice* AudioDevice, const UPTRINT Node
 	// Pick a random child node and save the index.
 	if (*RequiresInitialization)
 	{
+		if (PrimeRandomSoundNodesCVar != 0)
+		{
+			PrimeChildWavePlayers(true);
+		}
+
 		NodeIndex = ChooseNodeIndex(ActiveSound);
 		*RequiresInitialization = 0;
 	}
