Created INavigationDataInterface::IsNodeRefValid(), and implemented in RecastNavMesh.  Called by CharacterMovementComponent::PhysNavWalking to detect when the floor drops out.

diff --git a/Engine/Source/Runtime/Engine/Classes/AI/Navigation/NavigationDataInterface.h b/Engine/Source/Runtime/Engine/Classes/AI/Navigation/NavigationDataInterface.h
--- a/Engine/Source/Runtime/Engine/Classes/AI/Navigation/NavigationDataInterface.h
+++ b/Engine/Source/Runtime/Engine/Classes/AI/Navigation/NavigationDataInterface.h
@@ -29,4 +29,10 @@ public:
 	*	@return true if successful, false otherwise
 	*/
 	virtual bool ProjectPoint(const FVector& Point, FNavLocation& OutLocation, const FVector& Extent, FSharedConstNavQueryFilter Filter = nullptr, const UObject* Querier = nullptr) const PURE_VIRTUAL(INavigationDataInterface::ProjectPoint, return false;);
+
+	/** Determines whether the specified NavNodeRef is still valid
+	*   @param NodeRef the NavNodeRef to test for validity
+	*   @return true if valid, false otherwise
+	*/
+	virtual bool IsNodeRefValid(NavNodeRef NodeRef) const PURE_VIRTUAL(INavigationDataInterface::IsNodeRefValid, return true;);
 };
diff --git a/Engine/Source/Runtime/Engine/Private/Components/CharacterMovementComponent.cpp b/Engine/Source/Runtime/Engine/Private/Components/CharacterMovementComponent.cpp
--- a/Engine/Source/Runtime/Engine/Private/Components/CharacterMovementComponent.cpp
+++ b/Engine/Source/Runtime/Engine/Private/Components/CharacterMovementComponent.cpp
@@ -5045,6 +5045,7 @@ void UCharacterMovementComponent::PhysNavWalking(float deltaTime, int32 Iteratio
 
 	const FVector OldLocation = GetActorFeetLocation();
 	const FVector DeltaMove = DesiredMove * deltaTime;
+	const bool bDeltaMoveNearlyZero = DeltaMove.IsNearlyZero();
 
 	FVector AdjustedDest = OldLocation + DeltaMove;
 	FNavLocation DestNavLocation;
@@ -5067,10 +5068,21 @@ void UCharacterMovementComponent::PhysNavWalking(float deltaTime, int32 Iteratio
 		{
 			bSameNavLocation = CachedNavLocation.Location.Equals(OldLocation);
 		}
+
+		if (bDeltaMoveNearlyZero && bSameNavLocation)
+		{
+			if (const INavigationDataInterface* NavData = GetNavData())
+			{
+				if (!NavData->IsNodeRefValid(CachedNavLocation.NodeRef))
+				{
+					CachedNavLocation.NodeRef = INVALID_NAVNODEREF;
+					bSameNavLocation = false;
+				}
+			}
+		}
 	}
-		
 
-	if (DeltaMove.IsNearlyZero() && bSameNavLocation)
+	if (bDeltaMoveNearlyZero && bSameNavLocation)
 	{
 		DestNavLocation = CachedNavLocation;
 		UE_LOG(LogNavMeshMovement, VeryVerbose, TEXT("%s using cached navmesh location! (bProjectNavMeshWalking = %d)"), *GetNameSafe(CharacterOwner), bProjectNavMeshWalking);
diff --git a/Engine/Source/Runtime/NavigationSystem/Private/NavMesh/RecastNavMesh.cpp b/Engine/Source/Runtime/NavigationSystem/Private/NavMesh/RecastNavMesh.cpp
--- a/Engine/Source/Runtime/NavigationSystem/Private/NavMesh/RecastNavMesh.cpp
+++ b/Engine/Source/Runtime/NavigationSystem/Private/NavMesh/RecastNavMesh.cpp
@@ -1095,6 +1095,24 @@ bool ARecastNavMesh::ProjectPoint(const FVector& Point, FNavLocation& OutLocatio
 	return bSuccess;
 }
 
+bool ARecastNavMesh::IsNodeRefValid(NavNodeRef NodeRef) const
+{
+	if (NodeRef == INVALID_NAVNODEREF)
+	{
+		return false;
+	}
+	const dtNavMesh* NavMesh = RecastNavMeshImpl ? RecastNavMeshImpl->GetRecastMesh() : nullptr;
+	if (!NavMesh)
+	{
+		return false;
+	}
+	dtPoly const* Poly = 0;
+	dtMeshTile const* Tile = 0;
+	const dtStatus Status = NavMesh->getTileAndPolyByRef(NodeRef, &Tile, &Poly);
+	const bool bSuccess = dtStatusSucceed(Status);
+	return bSuccess;
+}
+
 void ARecastNavMesh::BatchProjectPoints(TArray<FNavigationProjectionWork>& Workload, const FVector& Extent, FSharedConstNavQueryFilter Filter, const UObject* Querier) const 
 {
 	if (Workload.Num() == 0 || RecastNavMeshImpl == NULL || RecastNavMeshImpl->DetourNavMesh == NULL)
diff --git a/Engine/Source/Runtime/NavigationSystem/Public/NavMesh/RecastNavMesh.h b/Engine/Source/Runtime/NavigationSystem/Public/NavMesh/RecastNavMesh.h
--- a/Engine/Source/Runtime/NavigationSystem/Public/NavMesh/RecastNavMesh.h
+++ b/Engine/Source/Runtime/NavigationSystem/Public/NavMesh/RecastNavMesh.h
@@ -712,7 +712,8 @@ public:
 	virtual bool GetRandomPointInNavigableRadius(const FVector& Origin, float Radius, FNavLocation& OutResult, FSharedConstNavQueryFilter Filter = NULL, const UObject* Querier = NULL) const override;
 
 	virtual bool ProjectPoint(const FVector& Point, FNavLocation& OutLocation, const FVector& Extent, FSharedConstNavQueryFilter Filter = NULL, const UObject* Querier = NULL) const override;
-	
+	virtual bool IsNodeRefValid(NavNodeRef NodeRef) const override;
+
 	/** Project batch of points using shared search extent and filter */
 	virtual void BatchProjectPoints(TArray<FNavigationProjectionWork>& Workload, const FVector& Extent, FSharedConstNavQueryFilter Filter = NULL, const UObject* Querier = NULL) const override;
 
