Don't initialize (and allocate) runtime virtual texture on platforms that don't support it

diff --git a/Engine/Source/Runtime/Engine/Classes/Components/RuntimeVirtualTextureComponent.h b/Engine/Source/Runtime/Engine/Classes/Components/RuntimeVirtualTextureComponent.h
--- a/Engine/Source/Runtime/Engine/Classes/Components/RuntimeVirtualTextureComponent.h
+++ b/Engine/Source/Runtime/Engine/Classes/Components/RuntimeVirtualTextureComponent.h
@@ -47,6 +47,7 @@ public:
 
 protected:
 	//~ Begin UActorComponent Interface
+	virtual bool IsVisible() const override;
 	virtual void CreateRenderState_Concurrent() override;
 	virtual void SendRenderTransform_Concurrent() override;
 	virtual void DestroyRenderState_Concurrent() override;
diff --git a/Engine/Source/Runtime/Engine/Private/Components/RuntimeVirtualTextureComponent.cpp b/Engine/Source/Runtime/Engine/Private/Components/RuntimeVirtualTextureComponent.cpp
--- a/Engine/Source/Runtime/Engine/Private/Components/RuntimeVirtualTextureComponent.cpp
+++ b/Engine/Source/Runtime/Engine/Private/Components/RuntimeVirtualTextureComponent.cpp
@@ -14,6 +14,11 @@ URuntimeVirtualTextureComponent::URuntimeVirtualTextureComponent(const FObjectIn
 	bTickInEditor = true;
 }
 
+bool URuntimeVirtualTextureComponent::IsVisible() const
+{
+	return Super::IsVisible() && UseVirtualTexturing(GetScene()->GetFeatureLevel());
+}
+
 void URuntimeVirtualTextureComponent::CreateRenderState_Concurrent()
 {
 	if (ShouldRender() && VirtualTexture != nullptr)
