Avoid PostSerialize calls in non-editor builds for editor only functionality
Minor packing improvement to FNavigationLinkBase
#rnx

diff --git a/Engine/Plugins/Runtime/GameplayAbilities/Source/GameplayAbilities/Private/AttributeSet.cpp b/Engine/Plugins/Runtime/GameplayAbilities/Source/GameplayAbilities/Private/AttributeSet.cpp
--- a/Engine/Plugins/Runtime/GameplayAbilities/Source/GameplayAbilities/Private/AttributeSet.cpp
+++ b/Engine/Plugins/Runtime/GameplayAbilities/Source/GameplayAbilities/Private/AttributeSet.cpp
@@ -205,6 +205,7 @@ bool FGameplayAttribute::IsGameplayAttributeDataProperty(const UProperty* Proper
 	return false;
 }
 
+#if WITH_EDITORONLY_DATA
 // Fill in missing attribute information
 void FGameplayAttribute::PostSerialize(const FArchive& Ar)
 {
@@ -230,6 +231,7 @@ void FGameplayAttribute::PostSerialize(const FArchive& Ar)
 		}
 	}
 }
+#endif
 
 void FGameplayAttribute::GetAllAttributeProperties(TArray<UProperty*>& OutProperties, FString FilterMetaStr, bool UseEditorOnlyData)
 {
diff --git a/Engine/Plugins/Runtime/GameplayAbilities/Source/GameplayAbilities/Public/AttributeSet.h b/Engine/Plugins/Runtime/GameplayAbilities/Source/GameplayAbilities/Public/AttributeSet.h
--- a/Engine/Plugins/Runtime/GameplayAbilities/Source/GameplayAbilities/Public/AttributeSet.h
+++ b/Engine/Plugins/Runtime/GameplayAbilities/Source/GameplayAbilities/Public/AttributeSet.h
@@ -56,7 +56,7 @@ protected:
 	float CurrentValue;
 };
 
-/** Describes a FGameplayAttributeData or float property inside an attribute set. Using this provides editor UI and gelper functions */
+/** Describes a FGameplayAttributeData or float property inside an attribute set. Using this provides editor UI and helper functions */
 USTRUCT(BlueprintType)
 struct GAMEPLAYABILITIES_API FGameplayAttribute
 {
@@ -137,8 +137,10 @@ struct GAMEPLAYABILITIES_API FGameplayAttribute
 		return AttributeName.IsEmpty() ? *GetNameSafe(Attribute) : AttributeName;
 	}
 
+#if WITH_EDITORONLY_DATA
 	/** Custom serialization */
 	void PostSerialize(const FArchive& Ar);
+#endif
 
 	/** Name of the attribute, usually the same as property name */
 	UPROPERTY(Category = GameplayAttribute, VisibleAnywhere, BlueprintReadOnly)
@@ -157,6 +159,7 @@ private:
 	UStruct* AttributeOwner;
 };
 
+#if WITH_EDITORONLY_DATA
 template<>
 struct TStructOpsTypeTraits< FGameplayAttribute > : public TStructOpsTypeTraitsBase2< FGameplayAttribute >
 {
@@ -165,6 +168,7 @@ struct TStructOpsTypeTraits< FGameplayAttribute > : public TStructOpsTypeTraitsB
 		WithPostSerialize = true,
 	};
 };
+#endif
 
 /**
  * Defines the set of all GameplayAttributes for your game
diff --git a/Engine/Source/Runtime/Engine/Classes/AI/Navigation/NavLinkDefinition.h b/Engine/Source/Runtime/Engine/Classes/AI/Navigation/NavLinkDefinition.h
--- a/Engine/Source/Runtime/Engine/Classes/AI/Navigation/NavLinkDefinition.h
+++ b/Engine/Source/Runtime/Engine/Classes/AI/Navigation/NavLinkDefinition.h
@@ -36,9 +36,6 @@ struct ENGINE_API FNavigationLinkBase
 	UPROPERTY(EditAnywhere, Category=Default, meta=(ClampMin = "0.0", DisplayName="Right Project Height"))
 	float MaxFallDownLength;
 
-	UPROPERTY(EditAnywhere, Category=Default, BlueprintReadWrite)
-	TEnumAsByte<ENavLinkDirection::Type> Direction;
-
 	/** ID passed to navigation data generator */
 	int32 UserId;
 
@@ -103,44 +100,47 @@ struct ENGINE_API FNavigationLinkBase
 	FString Description;
 #endif // WITH_EDITORONLY_DATA
 
+	UPROPERTY(EditAnywhere, Category=Default, BlueprintReadWrite)
+	TEnumAsByte<ENavLinkDirection::Type> Direction;
+
 	UPROPERTY(EditAnywhere, Category=Default, meta=(InlineEditConditionToggle))
-	uint32 bUseSnapHeight : 1;
+	uint8 bUseSnapHeight : 1;
 
 	/** If set, link will try to snap to cheapest area in given radius */
 	UPROPERTY(EditAnywhere, Category = Default)
-	uint32 bSnapToCheapestArea : 1;
+	uint8 bSnapToCheapestArea : 1;
 	
 	/** custom flag, check DescribeCustomFlags for details */
 	UPROPERTY()
-	uint32 bCustomFlag0 : 1;
+	uint8 bCustomFlag0 : 1;
 
 	/** custom flag, check DescribeCustomFlags for details */
 	UPROPERTY()
-	uint32 bCustomFlag1 : 1;
+	uint8 bCustomFlag1 : 1;
 
 	/** custom flag, check DescribeCustomFlags for details */
 	UPROPERTY()
-	uint32 bCustomFlag2 : 1;
+	uint8 bCustomFlag2 : 1;
 
 	/** custom flag, check DescribeCustomFlags for details */
 	UPROPERTY()
-	uint32 bCustomFlag3 : 1;
+	uint8 bCustomFlag3 : 1;
 
 	/** custom flag, check DescribeCustomFlags for details */
 	UPROPERTY()
-	uint32 bCustomFlag4 : 1;
+	uint8 bCustomFlag4 : 1;
 
 	/** custom flag, check DescribeCustomFlags for details */
 	UPROPERTY()
-	uint32 bCustomFlag5 : 1;
+	uint8 bCustomFlag5 : 1;
 
 	/** custom flag, check DescribeCustomFlags for details */
 	UPROPERTY()
-	uint32 bCustomFlag6 : 1;
+	uint8 bCustomFlag6 : 1;
 
 	/** custom flag, check DescribeCustomFlags for details */
 	UPROPERTY()
-	uint32 bCustomFlag7 : 1;
+	uint8 bCustomFlag7 : 1;
 
 	FNavigationLinkBase();
 
@@ -149,7 +149,9 @@ struct ENGINE_API FNavigationLinkBase
 	UClass* GetAreaClass() const;
 	bool HasMetaArea() const;
 
+#if WITH_EDITORONLY_DATA
 	void PostSerialize(const FArchive& Ar);
+#endif
 
 #if WITH_EDITOR
 	/** set up bCustomFlagX properties and expose them for edit
@@ -169,10 +171,12 @@ private:
 template<>
 struct TStructOpsTypeTraits< FNavigationLinkBase > : public TStructOpsTypeTraitsBase2< FNavigationLinkBase >
 {
+#if WITH_EDITORONLY_DATA
 	enum
 	{
 		WithPostSerialize = true,
 	};
+#endif
 };
 
 USTRUCT(BlueprintType)
@@ -221,20 +225,11 @@ struct ENGINE_API FNavigationLink : public FNavigationLinkBase
 
 		return Result;
 	}
-
-	void PostSerialize(const FArchive& Ar)
-	{
-		FNavigationLinkBase::PostSerialize(Ar);
-	}
 };
 
 template<>
-struct TStructOpsTypeTraits< FNavigationLink > : public TStructOpsTypeTraitsBase2< FNavigationLink >
+struct TStructOpsTypeTraits< FNavigationLink > : public TStructOpsTypeTraits< FNavigationLinkBase >
 {
-	enum
-	{
-		WithPostSerialize = true,
-	};
 };
 
 USTRUCT()
@@ -283,20 +278,11 @@ struct ENGINE_API FNavigationSegmentLink : public FNavigationLinkBase
 
 		return Result;
 	}
-
-	void PostSerialize(const FArchive& Ar)
-	{
-		FNavigationLinkBase::PostSerialize(Ar);
-	}
 };
 
 template<>
-struct TStructOpsTypeTraits< FNavigationSegmentLink > : public TStructOpsTypeTraitsBase2< FNavigationSegmentLink >
+struct TStructOpsTypeTraits< FNavigationSegmentLink > : public TStructOpsTypeTraits< FNavigationLinkBase >
 {
-	enum
-	{
-		WithPostSerialize = true,
-	};
 };
 
 /** Class containing definition of a navigation area */
diff --git a/Engine/Source/Runtime/Engine/Classes/PhysicsEngine/ConstraintInstance.h b/Engine/Source/Runtime/Engine/Classes/PhysicsEngine/ConstraintInstance.h
--- a/Engine/Source/Runtime/Engine/Classes/PhysicsEngine/ConstraintInstance.h
+++ b/Engine/Source/Runtime/Engine/Classes/PhysicsEngine/ConstraintInstance.h
@@ -539,7 +539,9 @@ public:
 	void GetUsedMaterials(TArray<UMaterialInterface*>& Materials);
 
 	bool Serialize(FArchive& Ar);
+#if WITH_EDITORONLY_DATA
 	void PostSerialize(const FArchive& Ar);
+#endif
 
 	/** Turn on linear and angular projection */
 	void EnableProjection();
@@ -734,6 +736,8 @@ struct TStructOpsTypeTraits<FConstraintInstance> : public TStructOpsTypeTraitsBa
 	enum 
 	{
 		WithSerializer = true,
+#if WITH_EDITORONLY_DATA
 		WithPostSerialize = true
+#endif
 	};
 };
\ No newline at end of file
diff --git a/Engine/Source/Runtime/Engine/Classes/Sound/SoundAttenuation.h b/Engine/Source/Runtime/Engine/Classes/Sound/SoundAttenuation.h
--- a/Engine/Source/Runtime/Engine/Classes/Sound/SoundAttenuation.h
+++ b/Engine/Source/Runtime/Engine/Classes/Sound/SoundAttenuation.h
@@ -335,7 +335,9 @@ struct ENGINE_API FSoundAttenuationSettings : public FBaseAttenuationSettings
 	}
 
 	bool operator==(const FSoundAttenuationSettings& Other) const;
+#if WITH_EDITORONLY_DATA
 	void PostSerialize(const FArchive& Ar);
+#endif
 
 	virtual void CollectAttenuationShapesForVisualization(TMultiMap<EAttenuationShape::Type, FBaseAttenuationSettings::AttenuationShapeDetails>& ShapeDetailsMap) const override;
 	float GetFocusPriorityScale(const struct FGlobalFocusSettings& FocusSettings, float FocusFactor) const;
@@ -343,9 +345,7 @@ struct ENGINE_API FSoundAttenuationSettings : public FBaseAttenuationSettings
 	float GetFocusDistanceScale(const struct FGlobalFocusSettings& FocusSettings, float FocusFactor) const;
 };
 
-UE_DEPRECATED(4.15, "FAttenuationSettings has been renamed FSoundAttenuationSettings")
-typedef FSoundAttenuationSettings FAttenuationSettings;
-
+#if WITH_EDITORONLY_DATA
 template<>
 struct TStructOpsTypeTraits<FSoundAttenuationSettings> : public TStructOpsTypeTraitsBase2<FSoundAttenuationSettings>
 {
@@ -354,6 +354,7 @@ struct TStructOpsTypeTraits<FSoundAttenuationSettings> : public TStructOpsTypeTr
 		WithPostSerialize = true,
 	};
 };
+#endif
 
 /** 
  * Defines how a sound changes volume with distance to the listener
diff --git a/Engine/Source/Runtime/Engine/Private/AI/Navigation/NavigationModifier.cpp b/Engine/Source/Runtime/Engine/Private/AI/Navigation/NavigationModifier.cpp
--- a/Engine/Source/Runtime/Engine/Private/AI/Navigation/NavigationModifier.cpp
+++ b/Engine/Source/Runtime/Engine/Private/AI/Navigation/NavigationModifier.cpp
@@ -52,6 +52,7 @@ bool FNavigationLinkBase::HasMetaArea() const
 	return AreaClass && AreaClass->GetDefaultObject<UNavAreaBase>()->IsMetaArea();
 }
 
+#if WITH_EDITORONLY_DATA
 void FNavigationLinkBase::PostSerialize(const FArchive& Ar)
 {
 	if (Ar.IsLoading() && Ar.UE4Ver() < VER_UE4_NAVIGATION_AGENT_SELECTOR)
@@ -77,6 +78,7 @@ void FNavigationLinkBase::PostSerialize(const FArchive& Ar)
 
 	// can't initialize at this time, used UClass may not be ready yet
 }
+#endif
 
 #if WITH_EDITOR
 
diff --git a/Engine/Source/Runtime/MovieScene/Private/Channels/MovieSceneFloatChannel.cpp b/Engine/Source/Runtime/MovieScene/Private/Channels/MovieSceneFloatChannel.cpp
--- a/Engine/Source/Runtime/MovieScene/Private/Channels/MovieSceneFloatChannel.cpp
+++ b/Engine/Source/Runtime/MovieScene/Private/Channels/MovieSceneFloatChannel.cpp
@@ -1056,6 +1056,7 @@ bool FMovieSceneFloatChannel::Serialize(FArchive& Ar)
 	return false;
 }
 
+#if WITH_EDITORONLY_DATA
 void FMovieSceneFloatChannel::PostSerialize(const FArchive& Ar)
 {
 	if (Ar.CustomVer(FSequencerObjectVersion::GUID) < FSequencerObjectVersion::ModifyLinearKeysForOldInterp)
@@ -1087,3 +1088,4 @@ void FMovieSceneFloatChannel::PostSerialize(const FArchive& Ar)
 		}
 	}
 }
+#endif
\ No newline at end of file
diff --git a/Engine/Source/Runtime/MovieScene/Public/Channels/MovieSceneFloatChannel.h b/Engine/Source/Runtime/MovieScene/Public/Channels/MovieSceneFloatChannel.h
--- a/Engine/Source/Runtime/MovieScene/Public/Channels/MovieSceneFloatChannel.h
+++ b/Engine/Source/Runtime/MovieScene/Public/Channels/MovieSceneFloatChannel.h
@@ -259,7 +259,9 @@ public:
 public:
 
 	bool Serialize(FArchive& Ar);
+#if WITH_EDITORONLY_DATA
 	void PostSerialize(const FArchive& Ar);
+#endif
 	friend FArchive& operator<<(FArchive& Ar, FMovieSceneFloatChannel& Me)
 	{
 		Me.Serialize(Ar);
@@ -364,7 +366,9 @@ struct TStructOpsTypeTraits<FMovieSceneFloatChannel> : public TStructOpsTypeTrai
 	{ 
 		WithStructuredSerializeFromMismatchedTag = true, 
 	    WithSerializer = true,
+#if WITH_EDITORONLY_DATA
 		WithPostSerialize = true,
+#endif
     };
 };
 
diff --git a/Engine/Source/Runtime/MovieSceneTracks/Private/Channels/MovieSceneEvent.cpp b/Engine/Source/Runtime/MovieSceneTracks/Private/Channels/MovieSceneEvent.cpp
--- a/Engine/Source/Runtime/MovieSceneTracks/Private/Channels/MovieSceneEvent.cpp
+++ b/Engine/Source/Runtime/MovieSceneTracks/Private/Channels/MovieSceneEvent.cpp
@@ -127,9 +127,9 @@ void FMovieSceneEvent::CacheFunctionName()
 #endif // WITH_EDITORONLY_DATA
 
 
+#if WITH_EDITORONLY_DATA
 void FMovieSceneEvent::PostSerialize(const FArchive& Ar)
 {
-#if WITH_EDITORONLY_DATA
 	if (Ar.IsLoading() && !Ar.HasAnyPortFlags(PPF_Duplicate | PPF_DuplicateForPIE))
 	{
 		// Re-cache the function name when loading in-editor in case of renamed function graphs and the like 
@@ -141,12 +141,10 @@ void FMovieSceneEvent::PostSerialize(const FArchive& Ar)
 			CacheFunctionName();
 		}
 	}
-#endif
 }
 
 bool FMovieSceneEvent::Serialize(FArchive& Ar)
 {
-#if WITH_EDITORONLY_DATA
 	if (Ar.IsSaving())
 	{
 		// ---------------------------------------------------------------------------------------
@@ -167,11 +165,11 @@ bool FMovieSceneEvent::Serialize(FArchive& Ar)
 			}
 		}
 	}
-#endif
 
 	// Return false to ensure that the struct receives default serialization
 	return false;
 }
+#endif
 
 bool FMovieSceneEvent::IsValidFunction(UFunction* Function)
 {
diff --git a/Engine/Source/Runtime/MovieSceneTracks/Public/Channels/MovieSceneEvent.h b/Engine/Source/Runtime/MovieSceneTracks/Public/Channels/MovieSceneEvent.h
--- a/Engine/Source/Runtime/MovieSceneTracks/Public/Channels/MovieSceneEvent.h
+++ b/Engine/Source/Runtime/MovieSceneTracks/Public/Channels/MovieSceneEvent.h
@@ -33,6 +33,7 @@ struct MOVIESCENETRACKS_API FMovieSceneEvent
 	GENERATED_BODY()
 
 
+#if WITH_EDITORONLY_DATA
 	/**
 	 * Called after this event has been serialized in order to cache the function pointer if necessary
 	 */
@@ -42,6 +43,7 @@ struct MOVIESCENETRACKS_API FMovieSceneEvent
 	 * Called to perform custom serialization logic for this struct.
 	 */
 	bool Serialize(FArchive& Ar);
+#endif
 
 	/**
 	 * Check whether the specified function is valid for a movie scene event
@@ -128,9 +130,11 @@ private:
 #endif // WITH_EDITORONLY_DATA
 };
 
+#if WITH_EDITORONLY_DATA
 template<>
 struct TStructOpsTypeTraits<FMovieSceneEvent> : TStructOpsTypeTraitsBase2<FMovieSceneEvent>
 {
 	enum { WithSerializer = true, WithPostSerialize = true };
 };
+#endif
 
diff --git a/Engine/Source/Runtime/MovieSceneTracks/Public/Sections/MovieSceneEventSection.h b/Engine/Source/Runtime/MovieSceneTracks/Public/Sections/MovieSceneEventSection.h
--- a/Engine/Source/Runtime/MovieSceneTracks/Public/Sections/MovieSceneEventSection.h
+++ b/Engine/Source/Runtime/MovieSceneTracks/Public/Sections/MovieSceneEventSection.h
@@ -133,10 +133,12 @@ struct MOVIESCENETRACKS_API FMovieSceneEventSectionData : public FMovieSceneChan
 {
 	GENERATED_BODY()
 
-	/**
+#if WITH_EDITORONLY_DATA
+		/**
 	 * Called after this section data has been serialized to upgrade old data
 	 */
 	void PostSerialize(const FArchive& Ar);
+#endif
 
 	/**
 	 * Access a mutable interface for this channel's data
@@ -205,13 +207,13 @@ private:
 #endif
 };
 
-
+#if WITH_EDITORONLY_DATA
 template<>
 struct TStructOpsTypeTraits<FMovieSceneEventSectionData> : public TStructOpsTypeTraitsBase2<FMovieSceneEventSectionData>
 {
 	enum { WithPostSerialize = true };
 };
-
+#endif
 
 /**
  * Implements a section in movie scene event tracks.
