Fixed initialization order
#rnx

diff --git a/Engine/Source/Runtime/Engine/Private/AI/Navigation/NavigationModifier.cpp b/Engine/Source/Runtime/Engine/Private/AI/Navigation/NavigationModifier.cpp
--- a/Engine/Source/Runtime/Engine/Private/AI/Navigation/NavigationModifier.cpp
+++ b/Engine/Source/Runtime/Engine/Private/AI/Navigation/NavigationModifier.cpp
@@ -18,8 +18,8 @@ static const float CONVEX_HULL_POINTS_MIN_DISTANCE_SQ = 4.0f * 4.0f;
 // FNavigationLinkBase
 //----------------------------------------------------------------------//
 FNavigationLinkBase::FNavigationLinkBase() 
-	: LeftProjectHeight(0.0f), MaxFallDownLength(1000.0f), Direction(ENavLinkDirection::BothWays), UserId(0),
-	  SnapRadius(30.f), SnapHeight(50.0f), bUseSnapHeight(false), bSnapToCheapestArea(true),
+	: LeftProjectHeight(0.0f), MaxFallDownLength(1000.0f), UserId(0), SnapRadius(30.f), SnapHeight(50.0f),
+	  Direction(ENavLinkDirection::BothWays), bUseSnapHeight(false), bSnapToCheapestArea(true),
 	  bCustomFlag0(false), bCustomFlag1(false), bCustomFlag2(false), bCustomFlag3(false), bCustomFlag4(false),
 	  bCustomFlag5(false), bCustomFlag6(false), bCustomFlag7(false)
 {
