Preallocate necessary row memory when loading datatable struct data
#rnx

diff --git a/Engine/Source/Runtime/Engine/Private/DataTable.cpp b/Engine/Source/Runtime/Engine/Private/DataTable.cpp
--- a/Engine/Source/Runtime/Engine/Private/DataTable.cpp
+++ b/Engine/Source/Runtime/Engine/Private/DataTable.cpp
@@ -105,6 +105,7 @@ void UDataTable::LoadStructData(FArchive& Ar)
 
 	DATATABLE_CHANGE_SCOPE();
 
+	RowMap.Reserve(NumRows);
 	for (int32 RowIdx = 0; RowIdx < NumRows; RowIdx++)
 	{
 		// Load row name
@@ -203,7 +204,7 @@ void UDataTable::Serialize( FArchive& Ar )
 	}
 #endif	// WITH_EDITORONLY_DATA
 
-	Super::Serialize(Ar); // When loading, this should load our RowStruct!	
+		Super::Serialize(Ar); // When loading, this should load our RowStruct!	
 
 	if (RowStruct && RowStruct->HasAnyFlags(RF_NeedLoad))
 	{
