Update default Canvas fonts in the Engine to use FCoreStyle::GetDefaultFont()

This gives them more comprehensive localization support, matching that available to the UE4 editor and engine.

diff --git a/Engine/Config/BaseEngine.ini b/Engine/Config/BaseEngine.ini
--- a/Engine/Config/BaseEngine.ini
+++ b/Engine/Config/BaseEngine.ini
@@ -73,11 +73,12 @@ bAllowMatureLanguage=false
 GameEngine=/Script/Engine.GameEngine
 EditorEngine=/Script/UnrealEd.EditorEngine
 UnrealEdEngine=/Script/UnrealEd.UnrealEdEngine
-TinyFontName=/Engine/EngineFonts/RobotoTiny.RobotoTiny
-SmallFontName=/Engine/EngineFonts/Roboto.Roboto
-MediumFontName=/Engine/EngineFonts/Roboto.Roboto
-LargeFontName=/Engine/EngineFonts/Roboto.Roboto
-SubtitleFontName=/Engine/EngineFonts/Roboto.Roboto
+; Leaving these unset causes the engine to use FCoreStyle::GetDefaultFont() instead
+;TinyFontName=/Engine/EngineFonts/RobotoTiny.RobotoTiny
+;SmallFontName=/Engine/EngineFonts/Roboto.Roboto
+;MediumFontName=/Engine/EngineFonts/Roboto.Roboto
+;LargeFontName=/Engine/EngineFonts/Roboto.Roboto
+;SubtitleFontName=/Engine/EngineFonts/Roboto.Roboto
 WireframeMaterialName=/Engine/EngineDebugMaterials/WireframeMaterial.WireframeMaterial
 DefaultMaterialName=/Engine/EngineMaterials/WorldGridMaterial.WorldGridMaterial
 DefaultLightFunctionMaterialName=/Engine/EngineMaterials/DefaultLightFunctionMaterial.DefaultLightFunctionMaterial
diff --git a/Engine/Source/Runtime/Engine/Private/UnrealEngine.cpp b/Engine/Source/Runtime/Engine/Private/UnrealEngine.cpp
--- a/Engine/Source/Runtime/Engine/Private/UnrealEngine.cpp
+++ b/Engine/Source/Runtime/Engine/Private/UnrealEngine.cpp
@@ -38,6 +38,7 @@ UnrealEngine.cpp: Implements the UEngine class and helpers.
 #include "Misc/TimeGuard.h"
 #include "Modules/ModuleManager.h"
 #include "UObject/UObjectIterator.h"
+#include "UObject/StrongObjectPtr.h"
 #include "UObject/Package.h"
 #include "UObject/MetaData.h"
 #include "UObject/ObjectMemoryAnalyzer.h"
@@ -2479,27 +2480,64 @@ void UEngine::InitializeObjectReferences()
 	// set the font object pointers, unless on server
 	if (!IsRunningDedicatedServer())
 	{
-		auto ConditionalLoadEngineFont = [](UFont*& FontPtr, const FString& FontName)
+		auto CreateFontObjectFromDefaultFont = [](const FName InObjectName, const FName InLegacyFontName, const int32 InLegacyFontSize) -> UFont*
 		{
-			if (!FontPtr && FontName.Len() > 0)
+			UFont* FontPtr = NewObject<UFont>(GetTransientPackage(), InObjectName, RF_Transient);
+			FontPtr->FontCacheType = EFontCacheType::Runtime;
+			FontPtr->LegacyFontName = InLegacyFontName;
+			FontPtr->LegacyFontSize = InLegacyFontSize;
+			FontPtr->CompositeFont = *FCoreStyle::GetDefaultFont();
+			return FontPtr;
+		};
+
+		TStrongObjectPtr<UFont> DefaultTinyFont;
+		auto GetDefaultTinyFont = [&DefaultTinyFont, &CreateFontObjectFromDefaultFont]() -> UFont*
+		{
+			if (!DefaultTinyFont)
 			{
-				FontPtr = LoadObject<UFont>(nullptr, *FontName, nullptr, LOAD_None, nullptr);
+				DefaultTinyFont.Reset(CreateFontObjectFromDefaultFont(FName("DefaultTinyFont"), FName("Light"), 8));
+			}
+			return DefaultTinyFont.Get();
+		};
+
+		TStrongObjectPtr<UFont> DefaultRegularFont;
+		auto GetDefaultRegularFont = [&DefaultRegularFont, &CreateFontObjectFromDefaultFont]() -> UFont*
+		{
+			if (!DefaultRegularFont)
+			{
+				DefaultRegularFont.Reset(CreateFontObjectFromDefaultFont(FName("DefaultRegularFont"), FName("Regular"), 10));
+			}
+			return DefaultRegularFont.Get();
+		};
+
+		auto ConditionalLoadEngineFont = [](UFont*& FontPtr, const FString& FontName, TFunctionRef<UFont*()> FallbackFontFactory)
+		{
+			if (!FontPtr)
+			{
+				if (FontName.Len() > 0)
+				{
+					FontPtr = LoadObject<UFont>(nullptr, *FontName, nullptr, LOAD_None, nullptr);
+				}
+				if (!FontPtr)
+				{
+					FontPtr = FallbackFontFactory();
+				}
 			}
 		};
 
 		// Standard fonts.
-		ConditionalLoadEngineFont(TinyFont, TinyFontName.ToString());
-		ConditionalLoadEngineFont(SmallFont, SmallFontName.ToString());
-		ConditionalLoadEngineFont(MediumFont, MediumFontName.ToString());
-		ConditionalLoadEngineFont(LargeFont, LargeFontName.ToString());
-		ConditionalLoadEngineFont(SubtitleFont, SubtitleFontName.ToString());
+		ConditionalLoadEngineFont(TinyFont, TinyFontName.ToString(), GetDefaultTinyFont);
+		ConditionalLoadEngineFont(SmallFont, SmallFontName.ToString(), GetDefaultRegularFont);
+		ConditionalLoadEngineFont(MediumFont, MediumFontName.ToString(), GetDefaultRegularFont);
+		ConditionalLoadEngineFont(LargeFont, LargeFontName.ToString(), GetDefaultRegularFont);
+		ConditionalLoadEngineFont(SubtitleFont, SubtitleFontName.ToString(), GetDefaultRegularFont);
 
 		// Additional fonts.
 		AdditionalFonts.Empty(AdditionalFontNames.Num());
 		for (const FString& FontName : AdditionalFontNames)
 		{
 			UFont* NewFont = nullptr;
-			ConditionalLoadEngineFont(NewFont, FontName);
+			ConditionalLoadEngineFont(NewFont, FontName, GetDefaultRegularFont);
 			AdditionalFonts.Add(NewFont);
 		}
 	}
