Do not create RVT resources on unsupported platforms

diff --git a/Engine/Source/Runtime/Engine/Private/VT/RuntimeVirtualTexture.cpp b/Engine/Source/Runtime/Engine/Private/VT/RuntimeVirtualTexture.cpp
--- a/Engine/Source/Runtime/Engine/Private/VT/RuntimeVirtualTexture.cpp
+++ b/Engine/Source/Runtime/Engine/Private/VT/RuntimeVirtualTexture.cpp
@@ -363,8 +363,15 @@ void URuntimeVirtualTexture::Initialize(IVirtualTexture* InProducer, FTransform
 	WorldToUVTransformParameters[0] = VolumeToWorld.GetTranslation();
 	WorldToUVTransformParameters[1] = VolumeToWorld.GetUnitAxis(EAxis::X) * 1.f / VolumeToWorld.GetScale3D().X;
 	WorldToUVTransformParameters[2] = VolumeToWorld.GetUnitAxis(EAxis::Y) * 1.f / VolumeToWorld.GetScale3D().Y;
-
-	InitResource(InProducer, VolumeToWorld);
+	
+	if (UseVirtualTexturing(GMaxRHIFeatureLevel))
+	{
+		InitResource(InProducer, VolumeToWorld);
+	}
+	else
+	{
+		InitNullResource();		
+	}
 }
 
 void URuntimeVirtualTexture::Release()
