Properly account for nativization in changes to allow Python to receive certain blueprint events
#rnx

diff --git a/Engine/Source/Runtime/CoreUObject/Private/UObject/SparseDelegate.cpp b/Engine/Source/Runtime/CoreUObject/Private/UObject/SparseDelegate.cpp
--- a/Engine/Source/Runtime/CoreUObject/Private/UObject/SparseDelegate.cpp
+++ b/Engine/Source/Runtime/CoreUObject/Private/UObject/SparseDelegate.cpp
@@ -57,7 +57,7 @@ FSparseDelegate* FSparseDelegateStorage::ResolveSparseDelegate(const UObject* Ow
 	const UClass* OwningClass = OwningObject->GetClass();
 	while (OwningClass)
 	{
-		if (OwningClass->HasAnyClassFlags(CLASS_Native))
+		if (OwningClass->HasAnyClassFlags(CLASS_Native) && !OwningClass->HasAnyClassFlags(CLASS_CompiledFromBlueprint))
 		{
 			if (size_t* DelegateOffset = SparseDelegateObjectOffsets.Find(TPair<FName, FName>(OwningClass->GetFName(), DelegateName)))
 			{
diff --git a/Engine/Source/Runtime/Engine/Private/Actor.cpp b/Engine/Source/Runtime/Engine/Private/Actor.cpp
--- a/Engine/Source/Runtime/Engine/Private/Actor.cpp
+++ b/Engine/Source/Runtime/Engine/Private/Actor.cpp
@@ -998,7 +998,7 @@ void AActor::TickActor( float DeltaSeconds, ELevelTick TickType, FActorTickFunct
 
 void AActor::Tick( float DeltaSeconds )
 {
-	if (!GetClass()->HasAnyClassFlags(CLASS_Native))
+	if (GetClass()->HasAnyFlags(CLASS_CompiledFromBlueprint) || !GetClass()->HasAnyClassFlags(CLASS_Native))
 	{
 		// Blueprint code outside of the construction script should not run in the editor
 		// Allow tick if we are not a dedicated server, or we allow this tick on dedicated servers
diff --git a/Engine/Source/Runtime/Engine/Private/Components/ActorComponent.cpp b/Engine/Source/Runtime/Engine/Private/Components/ActorComponent.cpp
--- a/Engine/Source/Runtime/Engine/Private/Components/ActorComponent.cpp
+++ b/Engine/Source/Runtime/Engine/Private/Components/ActorComponent.cpp
@@ -822,7 +822,7 @@ void UActorComponent::BeginPlay()
 	check(!bHasBegunPlay);
 	checkSlow(bTickFunctionsRegistered); // If this fails, someone called BeginPlay() without first calling RegisterAllComponentTickFunctions().
 
-	if (!GetClass()->HasAnyClassFlags(CLASS_Native))
+	if (GetClass()->HasAnyFlags(CLASS_CompiledFromBlueprint) || !GetClass()->HasAnyClassFlags(CLASS_Native))
 	{
 		ReceiveBeginPlay();
 	}
@@ -835,7 +835,7 @@ void UActorComponent::EndPlay(const EEndPlayReason::Type EndPlayReason)
 	check(bHasBegunPlay);
 
 	// If we're in the process of being garbage collected it is unsafe to call out to blueprints
-	if (!HasAnyFlags(RF_BeginDestroyed) && !IsUnreachable() && !GetClass()->HasAnyClassFlags(CLASS_Native))
+	if (!HasAnyFlags(RF_BeginDestroyed) && !IsUnreachable() && (GetClass()->HasAnyFlags(CLASS_CompiledFromBlueprint) || !GetClass()->HasAnyClassFlags(CLASS_Native)))
 	{
 		ReceiveEndPlay(EndPlayReason);
 	}
@@ -975,7 +975,7 @@ void UActorComponent::TickComponent(float DeltaTime, enum ELevelTick TickType, F
 {
 	check(bRegistered);
 
-	if (!GetClass()->HasAnyClassFlags(CLASS_Native))
+	if (GetClass()->HasAnyFlags(CLASS_CompiledFromBlueprint) || !GetClass()->HasAnyClassFlags(CLASS_Native))
 	{
 		ReceiveTick(DeltaTime);
 
