Merging 8741278

Added Blueprint functions to obtain the Minimum LOD property on a StaticMesh for a given platform, or for all supported platforms.
- GetMinimumLODForPlatforms
- GetMinimumLODForPlatform

diff --git a/Engine/Source/Runtime/Engine/Classes/Engine/StaticMesh.h b/Engine/Source/Runtime/Engine/Classes/Engine/StaticMesh.h
--- a/Engine/Source/Runtime/Engine/Classes/Engine/StaticMesh.h
+++ b/Engine/Source/Runtime/Engine/Classes/Engine/StaticMesh.h
@@ -641,6 +641,26 @@ class UStaticMesh : public UStreamableRenderAsset, public IInterface_CollisionDa
 	UPROPERTY()
 	FPerPlatformInt MinLOD;
 
+	UFUNCTION(BlueprintPure, Category=StaticMesh)
+	void GetMinimumLODForPlatforms(TMap<FName, int32>& PlatformMinimumLODs) const
+	{
+#if WITH_EDITORONLY_DATA
+		PlatformMinimumLODs = MinLOD.PerPlatform;
+#endif
+	}
+
+	UFUNCTION(BlueprintPure, Category=StaticMesh)
+	int32 GetMinimumLODForPlatform(const FName& PlatformName) const
+	{
+#if WITH_EDITORONLY_DATA
+		if (const int32* Result = MinLOD.PerPlatform.Find(PlatformName))
+		{
+			return *Result;
+		}
+#endif
+		return INDEX_NONE;
+	}
+
 	/** Bias multiplier for Light Propagation Volume lighting */
 	UPROPERTY(EditAnywhere, BlueprintReadOnly, Category=StaticMesh, meta=(UIMin = "0.0", UIMax = "3.0"))
 	float LpvBiasMultiplier;
