CSV Profiler:  CSV actor tick recording no longer includes TickInterval ticks that are still on cooldown.

diff --git a/Engine/Source/Runtime/Engine/Classes/Engine/EngineBaseTypes.h b/Engine/Source/Runtime/Engine/Classes/Engine/EngineBaseTypes.h
--- a/Engine/Source/Runtime/Engine/Classes/Engine/EngineBaseTypes.h
+++ b/Engine/Source/Runtime/Engine/Classes/Engine/EngineBaseTypes.h
@@ -369,6 +369,8 @@ public:
 		return Prerequisites;
 	}
 
+	float GetLastTickGameTime() const { return (InternalData ? InternalData->LastTickGameTimeSeconds : -1.f); }
+
 private:
 	/**
 	 * Queues a tick function for execution from the game thread
diff --git a/Engine/Source/Runtime/Engine/Private/LevelTick.cpp b/Engine/Source/Runtime/Engine/Private/LevelTick.cpp
--- a/Engine/Source/Runtime/Engine/Private/LevelTick.cpp
+++ b/Engine/Source/Runtime/Engine/Private/LevelTick.cpp
@@ -1265,7 +1265,7 @@ static void RecordWorldCountsToCSV(UWorld* World)
 
 			TSortedMap<FName, int32, FDefaultAllocator, FNameFastLess> TickContextToCountMap;
 			int32 EnabledCount;
-			FTickTaskManagerInterface::Get().GetEnabledTickFunctionCounts(World, TickContextToCountMap, EnabledCount, bDetailed);
+			FTickTaskManagerInterface::Get().GetEnabledTickFunctionCounts(World, TickContextToCountMap, EnabledCount, bDetailed, true);
 
 			for (auto It = TickContextToCountMap.CreateConstIterator(); It; ++It)
 			{
diff --git a/Engine/Source/Runtime/Engine/Private/TickTaskManager.cpp b/Engine/Source/Runtime/Engine/Private/TickTaskManager.cpp
--- a/Engine/Source/Runtime/Engine/Private/TickTaskManager.cpp
+++ b/Engine/Source/Runtime/Engine/Private/TickTaskManager.cpp
@@ -1200,7 +1200,7 @@ public:
 		CurrentCount++; 
 	}
 
-	void AddTickFunctionsToMap(TSortedMap<FName, int32, FDefaultAllocator, FNameFastLess>& ClassNameToCountMap, int32& EnabledCount, bool bDetailed)
+	void AddTickFunctionsToMap(TSortedMap<FName, int32, FDefaultAllocator, FNameFastLess>& ClassNameToCountMap, int32& EnabledCount, bool bDetailed, bool bFilterCoolingDown, float CurrentTime, float CurrentUnpausedTime)
 	{
 		// Add ticks from AllEnabledTickFunctions
 		for (TSet<FTickFunction*>::TIterator It(AllEnabledTickFunctions); It; ++It)
@@ -1214,6 +1214,13 @@ public:
 		FTickFunction* TickFunction = AllCoolingDownTickFunctions.Head;
 		while (TickFunction)
 		{
+			// Note: Timestamp check assumes TickFunction->TickGroup has been evaluated this frame
+			if (bFilterCoolingDown && TickFunction->GetLastTickGameTime() != (TickFunction->bTickEvenWhenPaused ? CurrentUnpausedTime : CurrentTime))
+			{				
+				TickFunction = TickFunction->InternalData->Next;
+				continue;
+			}
+
 			AddTickFunctionToMap(ClassNameToCountMap, TickFunction, bDetailed);
 			TickFunction = TickFunction->InternalData->Next;
 			++EnabledCount;
@@ -1695,12 +1702,20 @@ private:
 	}
 
 
-	virtual void GetEnabledTickFunctionCounts(UWorld* InWorld, TSortedMap<FName, int32, FDefaultAllocator, FNameFastLess>& TickContextToCountMap, int32& EnabledCount, bool bDetailed)
+	virtual void GetEnabledTickFunctionCounts(UWorld* InWorld, TSortedMap<FName, int32, FDefaultAllocator, FNameFastLess>& TickContextToCountMap, int32& EnabledCount, bool bDetailed, bool bFilterCoolingDown=false)
 	{
 		check(InWorld);
 		check(InWorld->TickTaskLevel);
 
-		InWorld->TickTaskLevel->AddTickFunctionsToMap(TickContextToCountMap, EnabledCount, bDetailed);
+		if (bFilterCoolingDown && InWorld->TickGroup >= 0 && InWorld->TickGroup < TG_NewlySpawned)
+		{
+			UE_LOG(LogTick, Warning, TEXT("GetEnabledTickFunctionCounts: Filtering results mid-frame. TickFunctions that execute later this frame will not be reported."));
+		}
+
+		const float WorldTimeSeconds = InWorld->GetTimeSeconds();
+		const float WorldUnpausedTimeSeconds = InWorld->GetUnpausedTimeSeconds();
+
+		InWorld->TickTaskLevel->AddTickFunctionsToMap(TickContextToCountMap, EnabledCount, bDetailed, bFilterCoolingDown, WorldTimeSeconds, WorldUnpausedTimeSeconds);
 
 		for (int32 LevelIndex = 0; LevelIndex < InWorld->GetNumLevels(); LevelIndex++)
 		{
@@ -1708,7 +1723,7 @@ private:
 			if (Level->bIsVisible)
 			{
 				check(Level->TickTaskLevel);
-				Level->TickTaskLevel->AddTickFunctionsToMap(TickContextToCountMap, EnabledCount, bDetailed);
+				Level->TickTaskLevel->AddTickFunctionsToMap(TickContextToCountMap, EnabledCount, bDetailed, bFilterCoolingDown, WorldTimeSeconds, WorldUnpausedTimeSeconds);
 			}
 		}
 	}
diff --git a/Engine/Source/Runtime/Engine/Public/TickTaskManagerInterface.h b/Engine/Source/Runtime/Engine/Public/TickTaskManagerInterface.h
--- a/Engine/Source/Runtime/Engine/Public/TickTaskManagerInterface.h
+++ b/Engine/Source/Runtime/Engine/Public/TickTaskManagerInterface.h
@@ -60,7 +60,7 @@ public:
 	virtual void DumpAllTickFunctions(FOutputDevice& Ar, UWorld* InWorld, bool bEnabled, bool bDisabled, bool bGrouped) = 0;
 
 	/** Returns a map of enabled ticks, grouped by 'diagnostic context' string, along with count of enabled ticks */
-	virtual void GetEnabledTickFunctionCounts(UWorld* InWorld, TSortedMap<FName, int32, FDefaultAllocator, FNameFastLess>& TickContextToCountMap, int32& EnabledCount, bool bDetailed) = 0;
+	virtual void GetEnabledTickFunctionCounts(UWorld* InWorld, TSortedMap<FName, int32, FDefaultAllocator, FNameFastLess>& TickContextToCountMap, int32& EnabledCount, bool bDetailed, bool bFilterCoolingDown=false) = 0;
 
 	/**
 	 * Singleton to retrieve the GLOBAL tick task manager
