#ue4 - Fix physics volume updates and overlap events not firing when Pawns placed in a level begin play.

diff --git a/Engine/Source/Runtime/Engine/Private/Pawn.cpp b/Engine/Source/Runtime/Engine/Private/Pawn.cpp
--- a/Engine/Source/Runtime/Engine/Private/Pawn.cpp
+++ b/Engine/Source/Runtime/Engine/Private/Pawn.cpp
@@ -61,6 +61,7 @@ APawn::APawn(const FObjectInitializer& ObjectInitializer)
 	BlendedReplayViewPitch = 0.0f;
 	bCollideWhenPlacing = true;
 	SpawnCollisionHandlingMethod = ESpawnActorCollisionHandlingMethod::AdjustIfPossibleButDontSpawnIfColliding;
+	bGenerateOverlapEventsDuringLevelStreaming = true;
 	bProcessingOutsideWorldBounds = false;
 
 	bUseControllerRotationPitch = false;
