Fix for Material::GetRelevance being called on the render thread

diff --git a/Engine/Source/Runtime/Engine/Private/SkeletalMesh.cpp b/Engine/Source/Runtime/Engine/Private/SkeletalMesh.cpp
--- a/Engine/Source/Runtime/Engine/Private/SkeletalMesh.cpp
+++ b/Engine/Source/Runtime/Engine/Private/SkeletalMesh.cpp
@@ -4777,8 +4777,9 @@ void FSkeletalMeshSceneProxy::UpdateMorphMaterialUsage_GameThread(TArray<UMateri
 			UMaterialInterface* DefaultMaterial = UMaterial::GetDefaultMaterial(MD_Surface);
 			ERHIFeatureLevel::Type InFeatureLevel = GetScene().GetFeatureLevel();
 			FSkeletalMeshSceneProxy* SkelMeshSceneProxy = this;
+			FMaterialRelevance DefaultRelevance = DefaultMaterial->GetRelevance(InFeatureLevel);
 			ENQUEUE_RENDER_COMMAND(UpdateSkelProxyLODSectionElementsCmd)(
-				[InMaterialsToSwap, DefaultMaterial, InFeatureLevel, SkelMeshSceneProxy](FRHICommandList& RHICmdList)
+				[InMaterialsToSwap, DefaultMaterial, DefaultRelevance, InFeatureLevel, SkelMeshSceneProxy](FRHICommandList& RHICmdList)
 				{
 					for( int32 LodIdx=0; LodIdx < SkelMeshSceneProxy->LODSections.Num(); LodIdx++ )
 					{
@@ -4793,7 +4794,7 @@ void FSkeletalMeshSceneProxy::UpdateMorphMaterialUsage_GameThread(TArray<UMateri
 							}
 						}
 					}
-					SkelMeshSceneProxy->MaterialRelevance |= DefaultMaterial->GetRelevance(InFeatureLevel);
+					SkelMeshSceneProxy->MaterialRelevance |= DefaultRelevance;
 				});
 		}
 	}
