Temp fix that globally fixes the UMaterial::GetRelevance crash

diff --git a/Engine/Source/Runtime/Engine/Private/Materials/MaterialShared.cpp b/Engine/Source/Runtime/Engine/Private/Materials/MaterialShared.cpp
--- a/Engine/Source/Runtime/Engine/Private/Materials/MaterialShared.cpp
+++ b/Engine/Source/Runtime/Engine/Private/Materials/MaterialShared.cpp
@@ -856,7 +856,6 @@ bool FMaterial::MaterialUsesSingleLayerWater_RenderThread() const
 
 bool FMaterial::MaterialUsesSingleLayerWater_GameThread() const
 {
-	check(IsInGameThread());
 	FMaterialShaderMap* ShaderMap = GameThreadShaderMap.GetReference();
 	return ShaderMap ? ShaderMap->UsesSingleLayerWaterMaterialOutput() : false;
 }
