Fix for HasAnyExpressionsInMaterialAndFunctionsOfType being too heavy. (temp require IsUseByWater now)

diff --git a/Engine/Source/Runtime/Engine/Private/Materials/Material.cpp b/Engine/Source/Runtime/Engine/Private/Materials/Material.cpp
--- a/Engine/Source/Runtime/Engine/Private/Materials/Material.cpp
+++ b/Engine/Source/Runtime/Engine/Private/Materials/Material.cpp
@@ -157,8 +157,8 @@ int32 FMaterialResource::CompilePropertyAndSetMaterialProperty(EMaterialProperty
 		case MP_OpacityMask:
 			// Force basic opaque surfaces to skip masked/translucent-only attributes.
 			// Some features can force the material to create a masked variant which unintentionally runs this dormant code
-			if (GetMaterialDomain() != MD_Surface || GetBlendMode() != BLEND_Opaque || (GetShadingModels().IsLit() && !GetShadingModels().HasShadingModel(MSM_DefaultLit)) 
-				|| MaterialInterface->GetMaterial()->HasAnyExpressionsInMaterialAndFunctionsOfType<UMaterialExpressionSingleLayerWaterMaterialOutput>())
+			if (GetMaterialDomain() != MD_Surface || GetBlendMode() != BLEND_Opaque || (GetShadingModels().IsLit() && !GetShadingModels().HasShadingModel(MSM_DefaultLit))
+				|| MaterialInterface->GetMaterial()->bUsedWithWater )//MaterialInterface->GetMaterial()->HasAnyExpressionsInMaterialAndFunctionsOfType<UMaterialExpressionSingleLayerWaterMaterialOutput>())
 			{
 				Ret = MaterialInterface->CompileProperty(Compiler, Property);
 			}
@@ -6211,7 +6211,7 @@ bool UMaterial::IsPropertyActive(EMaterialProperty InProperty) const
 #if WITH_EDITOR
 bool UMaterial::IsPropertyActiveInEditor(EMaterialProperty InProperty) const
 {
-	bool bUsesSingleLayerWaterMaterial = HasAnyExpressionsInMaterialAndFunctionsOfType<UMaterialExpressionSingleLayerWaterMaterialOutput>();
+	bool bUsesSingleLayerWaterMaterial = bUsedWithWater; // HasAnyExpressionsInMaterialAndFunctionsOfType<UMaterialExpressionSingleLayerWaterMaterialOutput>();
 	// explicitly DON'T use getters for BlendMode/ShadingModel...these getters may return an optimized value
 	// we want the actual value that's been set by the user in the material editor
 	return IsPropertyActive_Internal(InProperty,
@@ -6230,7 +6230,7 @@ bool UMaterial::IsPropertyActiveInEditor(EMaterialProperty InProperty) const
 
 bool UMaterial::IsPropertyActiveInDerived(EMaterialProperty InProperty, const UMaterialInterface* DerivedMaterial) const
 {
-	bool bUsesSingleLayerWaterMaterial = HasAnyExpressionsInMaterialAndFunctionsOfType<UMaterialExpressionSingleLayerWaterMaterialOutput>();
+	bool bUsesSingleLayerWaterMaterial = DerivedMaterial->GetMaterial()->bUsedWithWater; //HasAnyExpressionsInMaterialAndFunctionsOfType<UMaterialExpressionSingleLayerWaterMaterialOutput>();
 	return IsPropertyActive_Internal(InProperty,
 		MaterialDomain,
 		DerivedMaterial->GetBlendMode(),
diff --git a/Engine/Source/Runtime/Engine/Private/Materials/MaterialShared.cpp b/Engine/Source/Runtime/Engine/Private/Materials/MaterialShared.cpp
--- a/Engine/Source/Runtime/Engine/Private/Materials/MaterialShared.cpp
+++ b/Engine/Source/Runtime/Engine/Private/Materials/MaterialShared.cpp
@@ -856,6 +856,7 @@ bool FMaterial::MaterialUsesSingleLayerWater_RenderThread() const
 
 bool FMaterial::MaterialUsesSingleLayerWater_GameThread() const
 {
+	check(IsInGameThread());
 	FMaterialShaderMap* ShaderMap = GameThreadShaderMap.GetReference();
 	return ShaderMap ? ShaderMap->UsesSingleLayerWaterMaterialOutput() : false;
 }
@@ -1466,7 +1467,9 @@ UMaterialInterface* FMaterialResource::GetMaterialInterface() const
 
 bool FMaterialResource::IsUsingSingleLayerWaterMaterialOutput() const
 {
-	return GetMaterialInterface() ? GetMaterialInterface()->GetMaterial()->HasAnyExpressionsInMaterialAndFunctionsOfType<UMaterialExpressionSingleLayerWaterMaterialOutput>() : false;
+	// Apparently HasAnyExpressionsInMaterialAndFunctionsOfType is too heavy. So checking bUsedWithWater instead.
+	// return GetMaterialInterface() ? GetMaterialInterface()->GetMaterial()->HasAnyExpressionsInMaterialAndFunctionsOfType<UMaterialExpressionSingleLayerWaterMaterialOutput>() : false;
+	return GetMaterialInterface() ? GetMaterialInterface()->GetMaterial()->bUsedWithWater : false;
 }
 
 #if WITH_EDITOR
