Fixed bad format string when logging fatal error in FMaterial::GetShaderPipeline

We were passing more parameters than format string arguments, so we were losing some information that was meant to help with debugging.

diff --git a/Engine/Source/Runtime/Engine/Private/Materials/MaterialShared.cpp b/Engine/Source/Runtime/Engine/Private/Materials/MaterialShared.cpp
--- a/Engine/Source/Runtime/Engine/Private/Materials/MaterialShared.cpp
+++ b/Engine/Source/Runtime/Engine/Private/Materials/MaterialShared.cpp
@@ -2151,7 +2151,6 @@ FShaderPipeline* FMaterial::GetShaderPipeline(class FShaderPipelineType* ShaderP
 		UE_LOG(LogMaterial, Fatal,
 			TEXT("		With VF=%s, Platform=%s\n")
 			TEXT("		MaterialUsageDesc: %s"),
-			ShaderPipelineType->GetName(), *GetFriendlyName(),
 			VertexFactoryType->GetName(), *LegacyShaderPlatformToShaderFormat(ShaderPlatform).ToString(),
 			*MaterialUsage
 			);
