Re applying 8818587 due to wrong merge.

diff --git a/Engine/Source/Runtime/Engine/Private/Materials/MaterialShared.cpp b/Engine/Source/Runtime/Engine/Private/Materials/MaterialShared.cpp
--- a/Engine/Source/Runtime/Engine/Private/Materials/MaterialShared.cpp
+++ b/Engine/Source/Runtime/Engine/Private/Materials/MaterialShared.cpp
@@ -856,7 +856,7 @@ bool FMaterial::MaterialUsesSingleLayerWater_RenderThread() const
 
 bool FMaterial::MaterialUsesSingleLayerWater_GameThread() const
 {
-	check(IsInGameThread());
+	//check(IsInGameThread());
 	FMaterialShaderMap* ShaderMap = GameThreadShaderMap.GetReference();
 	return ShaderMap ? ShaderMap->UsesSingleLayerWaterMaterialOutput() : false;
 }
