Make UActorComponent::IsActive non-virtual. No engine or sample code overrides this, and it's unlikely any game code will either.

diff --git a/Engine/Source/Runtime/Engine/Classes/Components/ActorComponent.h b/Engine/Source/Runtime/Engine/Classes/Components/ActorComponent.h
--- a/Engine/Source/Runtime/Engine/Classes/Components/ActorComponent.h
+++ b/Engine/Source/Runtime/Engine/Classes/Components/ActorComponent.h
@@ -354,7 +354,7 @@ public:
 	 * @return - The active state of the component.
 	 */
 	UFUNCTION(BlueprintCallable, Category="Components|Activation", meta=(UnsafeDuringActorConstruction="true"))
-	virtual bool IsActive() const { return bIsActive; }
+	bool IsActive() const { return bIsActive; }
 
 	/**
 	 * Sets whether the component should be auto activate or not. Only safe during construction scripts.
