Fixing stats not being available in test builds

diff --git a/Engine/Source/Runtime/Engine/Public/Audio/AudioDebug.h b/Engine/Source/Runtime/Engine/Public/Audio/AudioDebug.h
--- a/Engine/Source/Runtime/Engine/Public/Audio/AudioDebug.h
+++ b/Engine/Source/Runtime/Engine/Public/Audio/AudioDebug.h
@@ -4,7 +4,7 @@
 #include "CanvasTypes.h"
 #include "CoreMinimal.h"
 
-#define ENABLE_AUDIO_DEBUG !(UE_BUILD_SHIPPING || UE_BUILD_TEST)
+#define ENABLE_AUDIO_DEBUG !UE_BUILD_SHIPPING
 
 
 #if ENABLE_AUDIO_DEBUG
