[NO NON-PII COMMIT MESSAGE]

diff --git a/Engine/Source/Runtime/Engine/Private/AudioStreaming.cpp b/Engine/Source/Runtime/Engine/Private/AudioStreaming.cpp
--- a/Engine/Source/Runtime/Engine/Private/AudioStreaming.cpp
+++ b/Engine/Source/Runtime/Engine/Private/AudioStreaming.cpp
@@ -139,7 +139,11 @@ bool FStreamingWaveData::Initialize(USoundWave* InSoundWave, FLegacyAudioStreami
 
 	if (!InSoundWave || !InSoundWave->RunningPlatformData->Chunks.Num())
 	{
-		UE_LOG(LogAudio, Error, TEXT("Failed to initialize streaming wave data due to lack of serialized stream chunks. Error during stream cooking."));
+#if WITH_EDITOR
+		UE_LOG(LogAudio, Display, TEXT("Failed to initialize streaming wave data due to lack of serialized stream chunks. Error during stream cooking."));
+#else
+		UE_LOG(LogAudio, Warning, TEXT("Failed to initialize streaming wave data due to lack of serialized stream chunks. Error during stream cooking."));
+#endif
 		return false;
 	}
 
