Correctly serialize FSkeletalMeshLODRenderData when saving in cooked builds

diff --git a/Engine/Source/Runtime/Engine/Private/SkeletalMeshLODRenderData.cpp b/Engine/Source/Runtime/Engine/Private/SkeletalMeshLODRenderData.cpp
--- a/Engine/Source/Runtime/Engine/Private/SkeletalMeshLODRenderData.cpp
+++ b/Engine/Source/Runtime/Engine/Private/SkeletalMeshLODRenderData.cpp
@@ -928,12 +928,25 @@ void FSkeletalMeshLODRenderData::Serialize(FArchive& Ar, UObject* Owner, int32 I
 	FStripDataFlags StripFlags(Ar, ClassDataStripFlags);
 
 	const bool bIsBelowMinLOD = StripFlags.IsClassDataStripped(CDSF_MinLodData);
-	bool bIsLODCookedOut = IsLODCookedOut(Ar.CookingTarget(), OwnerMesh, bIsBelowMinLOD);
-	Ar << bIsLODCookedOut;
+	bool bIsLODCookedOut = false;
+	bool bInlined = false;
 
-	bool bInlined = bIsLODCookedOut || IsLODInlined(Ar.CookingTarget(), OwnerMesh, Idx, bIsBelowMinLOD);
-	Ar << bInlined;
-	bStreamedDataInlined = bInlined;
+	if (Ar.IsSaving() && !Ar.IsCooking())
+	{
+		bInlined = bStreamedDataInlined;
+		bIsLODCookedOut = bIsBelowMinLOD && bInlined;
+		Ar << bIsLODCookedOut;
+		Ar << bInlined;
+	}
+	else
+	{
+		bIsLODCookedOut = IsLODCookedOut(Ar.CookingTarget(), OwnerMesh, bIsBelowMinLOD);
+		Ar << bIsLODCookedOut;
+
+		bInlined = bIsLODCookedOut || IsLODInlined(Ar.CookingTarget(), OwnerMesh, Idx, bIsBelowMinLOD);
+		Ar << bInlined;
+		bStreamedDataInlined = bInlined;
+	}
 
 	// Skeletal mesh buffers are kept in CPU memory after initialization to support merging of skeletal meshes.
 	const bool bForceKeepCPUResources = ShouldForceKeepCPUResources();
