Fix evaluation crash after dynamically linking a layer

CacheBones counter was not incremented at first, meaning all counters get out of sync. Calling the underlying node's Initialize_AnyThread and CacheBones_AnyThread instead fixes this.

diff --git a/Engine/Source/Runtime/Engine/Private/Animation/AnimInstance.cpp b/Engine/Source/Runtime/Engine/Private/Animation/AnimInstance.cpp
--- a/Engine/Source/Runtime/Engine/Private/Animation/AnimInstance.cpp
+++ b/Engine/Source/Runtime/Engine/Private/Animation/AnimInstance.cpp
@@ -2601,11 +2601,14 @@ void UAnimInstance::PerformLayerOverlayOperation(TSubclassOf<UAnimInstance> InCl
 							// Initialize the correct parts of the sub instance
 							if(LayerNode->LinkedRoot)
 							{
-								FAnimInstanceProxy& Proxy = NewSubInstance->GetProxyOnAnyThread<FAnimInstanceProxy>();
-								Proxy.InitializeObjects(NewSubInstance);
-								Proxy.InitializeRootNode_WithRoot(LayerNode->LinkedRoot);
-								Proxy.CacheBones_WithRoot(LayerNode->LinkedRoot);
-								Proxy.ClearObjects();
+								FAnimInstanceProxy& ThisProxy = GetProxyOnAnyThread<FAnimInstanceProxy>();
+								FAnimInstanceProxy& LinkedProxy = NewSubInstance->GetProxyOnAnyThread<FAnimInstanceProxy>();
+								LinkedProxy.InitializeObjects(NewSubInstance);
+								FAnimationInitializeContext InitContext(&ThisProxy);
+								LayerNode->Initialize_AnyThread(InitContext);
+								FAnimationCacheBonesContext CacheBonesContext(&ThisProxy);
+								LayerNode->CacheBones_AnyThread(CacheBonesContext);
+								LinkedProxy.ClearObjects();
 							}
 
 							MeshComp->SubInstances.Add(NewSubInstance);
@@ -2644,17 +2647,20 @@ void UAnimInstance::PerformLayerOverlayOperation(TSubclassOf<UAnimInstance> InCl
 							LayerNode->SetLayerOverlaySubInstance(this, NewSubInstance);
 						}
 
-						FAnimInstanceProxy& Proxy = NewSubInstance->GetProxyOnAnyThread<FAnimInstanceProxy>();
+						FAnimInstanceProxy& ThisProxy = GetProxyOnAnyThread<FAnimInstanceProxy>();
+						FAnimInstanceProxy& LinkedProxy = NewSubInstance->GetProxyOnAnyThread<FAnimInstanceProxy>();
 
 						// Initialize the correct parts of the sub instance
 						for(FAnimNode_Layer* LayerNode : LayerPair.Value)
 						{
 							if(LayerNode->LinkedRoot)
 							{
-								Proxy.InitializeObjects(NewSubInstance);
-								Proxy.InitializeRootNode_WithRoot(LayerNode->LinkedRoot);
-								Proxy.CacheBones_WithRoot(LayerNode->LinkedRoot);
-								Proxy.ClearObjects();
+								LinkedProxy.InitializeObjects(NewSubInstance);
+								FAnimationInitializeContext InitContext(&ThisProxy);
+								LayerNode->Initialize_AnyThread(InitContext);
+								FAnimationCacheBonesContext CacheBonesContext(&ThisProxy);
+								LayerNode->CacheBones_AnyThread(CacheBonesContext);
+								LinkedProxy.ClearObjects();
 							}
 						}
 
