Back out changelist 8858830. It can cause incorrect skeletal mesh data being put in to DDC.

diff --git a/Engine/Source/Runtime/Engine/Private/SkeletalMeshLODRenderData.cpp b/Engine/Source/Runtime/Engine/Private/SkeletalMeshLODRenderData.cpp
--- a/Engine/Source/Runtime/Engine/Private/SkeletalMeshLODRenderData.cpp
+++ b/Engine/Source/Runtime/Engine/Private/SkeletalMeshLODRenderData.cpp
@@ -928,25 +928,12 @@ void FSkeletalMeshLODRenderData::Serialize(FArchive& Ar, UObject* Owner, int32 I
 	FStripDataFlags StripFlags(Ar, ClassDataStripFlags);
 
 	const bool bIsBelowMinLOD = StripFlags.IsClassDataStripped(CDSF_MinLodData);
-	bool bIsLODCookedOut = false;
-	bool bInlined = false;
+	bool bIsLODCookedOut = IsLODCookedOut(Ar.CookingTarget(), OwnerMesh, bIsBelowMinLOD);
+	Ar << bIsLODCookedOut;
 
-	if (Ar.IsSaving() && !Ar.IsCooking())
-	{
-		bInlined = bStreamedDataInlined;
-		bIsLODCookedOut = bIsBelowMinLOD && bInlined;
-		Ar << bIsLODCookedOut;
-		Ar << bInlined;
-	}
-	else
-	{
-		bIsLODCookedOut = IsLODCookedOut(Ar.CookingTarget(), OwnerMesh, bIsBelowMinLOD);
-		Ar << bIsLODCookedOut;
-
-		bInlined = bIsLODCookedOut || IsLODInlined(Ar.CookingTarget(), OwnerMesh, Idx, bIsBelowMinLOD);
-		Ar << bInlined;
-		bStreamedDataInlined = bInlined;
-	}
+	bool bInlined = bIsLODCookedOut || IsLODInlined(Ar.CookingTarget(), OwnerMesh, Idx, bIsBelowMinLOD);
+	Ar << bInlined;
+	bStreamedDataInlined = bInlined;
 
 	// Skeletal mesh buffers are kept in CPU memory after initialization to support merging of skeletal meshes.
 	const bool bForceKeepCPUResources = ShouldForceKeepCPUResources();
diff --git a/Engine/Source/Runtime/Engine/Private/SkeletalMeshRenderData.cpp b/Engine/Source/Runtime/Engine/Private/SkeletalMeshRenderData.cpp
--- a/Engine/Source/Runtime/Engine/Private/SkeletalMeshRenderData.cpp
+++ b/Engine/Source/Runtime/Engine/Private/SkeletalMeshRenderData.cpp
@@ -32,7 +32,7 @@ namespace SkeletalMeshCookStats
 // differences, etc.) replace the version GUID below with a new one.
 // In case of merge conflicts with DDC versions, you MUST generate a new GUID
 // and set this new GUID as the version.                                       
-#define SKELETALMESH_DERIVEDDATA_VER TEXT("51BDAF351B6941548818C2F8B33397D7")
+#define SKELETALMESH_DERIVEDDATA_VER TEXT("4F54F64FD82B40C7A9F4158AF684B64D")
 
 static const FString& GetSkeletalMeshDerivedDataVersion()
 {
