Make USceneComponent::SetRelativeScale3D non-virtual. No engine or sample code overrides this, and it's unlikely any game code will either.

diff --git a/Engine/Source/Runtime/Engine/Classes/Components/SceneComponent.h b/Engine/Source/Runtime/Engine/Classes/Components/SceneComponent.h
--- a/Engine/Source/Runtime/Engine/Classes/Components/SceneComponent.h
+++ b/Engine/Source/Runtime/Engine/Classes/Components/SceneComponent.h
@@ -403,7 +403,7 @@ public:
 
 	/** Set the non-uniform scale of the component relative to its parent */
 	UFUNCTION(BlueprintCallable, Category="Utilities|Transformation")
-	virtual void SetRelativeScale3D(FVector NewScale3D);
+	void SetRelativeScale3D(FVector NewScale3D);
 
 	/**
 	 * Adds a delta to the translation of the component relative to its parent
