Fix for NaNs in pass / view uniform buffers in canvas tile rendering due to scene buffers not being initialized yet.

diff --git a/Engine/Source/Runtime/Engine/Private/SceneView.cpp b/Engine/Source/Runtime/Engine/Private/SceneView.cpp
--- a/Engine/Source/Runtime/Engine/Private/SceneView.cpp
+++ b/Engine/Source/Runtime/Engine/Private/SceneView.cpp
@@ -2140,39 +2140,33 @@ void FSceneView::SetupViewRectUniformBufferParameters(FViewUniformShaderParamete
 	const FViewMatrices& InPrevViewMatrices) const
 {
 	checkfSlow(EffectiveViewRect.Area() > 0, TEXT("Invalid-size EffectiveViewRect passed to CreateUniformBufferParameters [%d * %d]."), EffectiveViewRect.Width(), EffectiveViewRect.Height());
+	ensureMsgf((BufferSize.X > 0 && BufferSize.Y > 0), TEXT("Invalid-size BufferSize passed to CreateUniformBufferParameters [%d * %d]."), BufferSize.X, BufferSize.Y);
 
 	ViewUniformShaderParameters.ViewRectMin = FVector4(EffectiveViewRect.Min.X, EffectiveViewRect.Min.Y, 0.0f, 0.0f);
 	ViewUniformShaderParameters.ViewSizeAndInvSize = FVector4(EffectiveViewRect.Width(), EffectiveViewRect.Height(), 1.0f / float(EffectiveViewRect.Width()), 1.0f / float(EffectiveViewRect.Height()));
 
-	// A buffer size of zero means that we don't have a scene context allocated. This can happen when rendering to a one-off render target (e.g. canvas) prior to
-	// main view rendering. In that case we just skip the scene color target specific parameters, since they wouldn't be relevant to those views anyway.
-	if (BufferSize != FIntPoint::ZeroValue)
-	{
-		ensureMsgf((BufferSize.X > 0 && BufferSize.Y > 0), TEXT("Invalid-size BufferSize passed to CreateUniformBufferParameters [%d * %d]."), BufferSize.X, BufferSize.Y);
-
-		// Calculate the vector used by shaders to convert clip space coordinates to texture space.
-		const float InvBufferSizeX = 1.0f / BufferSize.X;
-		const float InvBufferSizeY = 1.0f / BufferSize.Y;
-		// to bring NDC (-1..1, 1..-1) into 0..1 UV for BufferSize textures
-		const FVector4 ScreenPositionScaleBias(
-			EffectiveViewRect.Width() * InvBufferSizeX / +2.0f,
-			EffectiveViewRect.Height() * InvBufferSizeY / (-2.0f * GProjectionSignY),
-			(EffectiveViewRect.Height() / 2.0f + EffectiveViewRect.Min.Y) * InvBufferSizeY,
-			(EffectiveViewRect.Width() / 2.0f + EffectiveViewRect.Min.X) * InvBufferSizeX);
-
-		ViewUniformShaderParameters.ScreenPositionScaleBias = ScreenPositionScaleBias;
-
-		ViewUniformShaderParameters.BufferSizeAndInvSize = FVector4(BufferSize.X, BufferSize.Y, InvBufferSizeX, InvBufferSizeY);
-		ViewUniformShaderParameters.BufferBilinearUVMinMax = FVector4(
-			InvBufferSizeX * (EffectiveViewRect.Min.X + 0.5),
-			InvBufferSizeY * (EffectiveViewRect.Min.Y + 0.5),
-			InvBufferSizeX * (EffectiveViewRect.Max.X - 0.5),
-			InvBufferSizeY * (EffectiveViewRect.Max.Y - 0.5));
-
-		/* Texture Level-of-Detail Strategies for Real-Time Ray Tracing https://developer.nvidia.com/raytracinggems Equation 20 */
-		float RadFOV = (PI / 180.0f) * FOV;
-		ViewUniformShaderParameters.EyeToPixelSpreadAngle = FPlatformMath::Atan((2.0f * FPlatformMath::Tan(RadFOV * 0.5f)) / BufferSize.Y);
-	}
+	// Calculate the vector used by shaders to convert clip space coordinates to texture space.
+	const float InvBufferSizeX = 1.0f / BufferSize.X;
+	const float InvBufferSizeY = 1.0f / BufferSize.Y;
+	// to bring NDC (-1..1, 1..-1) into 0..1 UV for BufferSize textures
+	const FVector4 ScreenPositionScaleBias(
+		EffectiveViewRect.Width() * InvBufferSizeX / +2.0f,
+		EffectiveViewRect.Height() * InvBufferSizeY / (-2.0f * GProjectionSignY),
+		(EffectiveViewRect.Height() / 2.0f + EffectiveViewRect.Min.Y) * InvBufferSizeY,
+		(EffectiveViewRect.Width() / 2.0f + EffectiveViewRect.Min.X) * InvBufferSizeX);
+
+	ViewUniformShaderParameters.ScreenPositionScaleBias = ScreenPositionScaleBias;
+
+	ViewUniformShaderParameters.BufferSizeAndInvSize = FVector4(BufferSize.X, BufferSize.Y, InvBufferSizeX, InvBufferSizeY);
+	ViewUniformShaderParameters.BufferBilinearUVMinMax = FVector4(
+		InvBufferSizeX * (EffectiveViewRect.Min.X + 0.5),
+		InvBufferSizeY * (EffectiveViewRect.Min.Y + 0.5),
+		InvBufferSizeX * (EffectiveViewRect.Max.X - 0.5),
+		InvBufferSizeY * (EffectiveViewRect.Max.Y - 0.5));
+
+	/* Texture Level-of-Detail Strategies for Real-Time Ray Tracing https://developer.nvidia.com/raytracinggems Equation 20 */
+	float RadFOV = (PI / 180.0f) * FOV;
+	ViewUniformShaderParameters.EyeToPixelSpreadAngle = FPlatformMath::Atan((2.0f * FPlatformMath::Tan(RadFOV * 0.5f)) / BufferSize.Y);
 
 	ViewUniformShaderParameters.MotionBlurNormalizedToPixel = FinalPostProcessSettings.MotionBlurMax * EffectiveViewRect.Width() / 100.0f;
 
diff --git a/Engine/Source/Runtime/Renderer/Private/SceneRendering.cpp b/Engine/Source/Runtime/Renderer/Private/SceneRendering.cpp
--- a/Engine/Source/Runtime/Renderer/Private/SceneRendering.cpp
+++ b/Engine/Source/Runtime/Renderer/Private/SceneRendering.cpp
@@ -1095,10 +1095,15 @@ void FViewInfo::SetupUniformBufferParameters(
 	// Mobile multi-view is not side by side
 	const FIntRect EffectiveViewRect = (bIsMobileMultiViewEnabled) ? FIntRect(0, 0, ViewRect.Width(), ViewRect.Height()) : ViewRect;
 
+	// Scene render targets may not be created yet; avoids NaNs.
+	FIntPoint EffectiveBufferSize = SceneContext.GetBufferSizeXY();
+	EffectiveBufferSize.X = FMath::Max(EffectiveBufferSize.X, 1);
+	EffectiveBufferSize.Y = FMath::Max(EffectiveBufferSize.Y, 1);
+
 	// TODO: We should use a view and previous view uniform buffer to avoid code duplication and keep consistency
 	SetupCommonViewUniformBufferParameters(
 		ViewUniformShaderParameters,
-		SceneContext.GetBufferSizeXY(),
+		EffectiveBufferSize,
 		SceneContext.GetMSAACount(),
 		EffectiveViewRect,
 		InViewMatrices,
diff --git a/Engine/Source/Runtime/Renderer/Private/TranslucentRendering.cpp b/Engine/Source/Runtime/Renderer/Private/TranslucentRendering.cpp
--- a/Engine/Source/Runtime/Renderer/Private/TranslucentRendering.cpp
+++ b/Engine/Source/Runtime/Renderer/Private/TranslucentRendering.cpp
@@ -626,16 +626,20 @@ void CreateTranslucentBasePassUniformBuffer(
 
 		FIntPoint ViewportOffset = View.ViewRect.Min;
 		FIntPoint ViewportExtent = View.ViewRect.Size();
-		FIntPoint BufferSize = SceneRenderTargets.GetBufferSizeXY();
+
+		// Scene render targets might not exist yet; avoids NaNs.
+		FIntPoint EffectiveBufferSize = SceneRenderTargets.GetBufferSizeXY();
+		EffectiveBufferSize.X = FMath::Max(EffectiveBufferSize.X, 1);
+		EffectiveBufferSize.Y = FMath::Max(EffectiveBufferSize.Y, 1);
 
 		if (View.PrevViewInfo.TemporalAAHistory.IsValid())
 		{
 			ViewportOffset = View.PrevViewInfo.TemporalAAHistory.ViewportRect.Min;
 			ViewportExtent = View.PrevViewInfo.TemporalAAHistory.ViewportRect.Size();
-			BufferSize = View.PrevViewInfo.TemporalAAHistory.RT[0]->GetDesc().Extent;
+			EffectiveBufferSize = View.PrevViewInfo.TemporalAAHistory.RT[0]->GetDesc().Extent;
 		}
 
-		FVector2D InvBufferSize(1.0f / float(BufferSize.X), 1.0f / float(BufferSize.Y));
+		FVector2D InvBufferSize(1.0f / float(EffectiveBufferSize.X), 1.0f / float(EffectiveBufferSize.Y));
 
 		FVector4 ScreenPosToPixelValue(
 			ViewportExtent.X * 0.5f * InvBufferSize.X,
@@ -986,4 +990,4 @@ void FDeferredShadingSceneRenderer::RenderTranslucency(FRHICommandListImmediate&
 	}
 
 	checkSlow(RHICmdList.IsOutsideRenderPass());
-}
\ No newline at end of file
+}
