Fixed uninitialized variables.

diff --git a/Engine/Source/Runtime/Engine/Private/Materials/MaterialUniformExpressions.cpp b/Engine/Source/Runtime/Engine/Private/Materials/MaterialUniformExpressions.cpp
--- a/Engine/Source/Runtime/Engine/Private/Materials/MaterialUniformExpressions.cpp
+++ b/Engine/Source/Runtime/Engine/Private/Materials/MaterialUniformExpressions.cpp
@@ -579,7 +579,7 @@ void FUniformExpressionSet::FillUniformBuffer(const FMaterialRenderContext& Mate
 		// Cache 2D texture uniform expressions.
 		for(int32 ExpressionIndex = 0;ExpressionIndex < Uniform2DTextureExpressions.Num();ExpressionIndex++)
 		{
-			const UTexture* Value;
+			const UTexture* Value = nullptr;
 			Uniform2DTextureExpressions[ExpressionIndex]->GetTextureValue(MaterialRenderContext,MaterialRenderContext.Material,Value);
 			if (Value)
 			{
@@ -654,7 +654,7 @@ void FUniformExpressionSet::FillUniformBuffer(const FMaterialRenderContext& Mate
 		// Cache cube texture uniform expressions.
 		for(int32 ExpressionIndex = 0;ExpressionIndex < UniformCubeTextureExpressions.Num();ExpressionIndex++)
 		{
-			const UTexture* Value;
+			const UTexture* Value = nullptr;
 			UniformCubeTextureExpressions[ExpressionIndex]->GetTextureValue(MaterialRenderContext,MaterialRenderContext.Material,Value);
 
 			void** ResourceTableTexturePtr = (void**)((uint8*)BufferCursor + 0 * SHADER_PARAMETER_POINTER_ALIGNMENT);
@@ -731,7 +731,7 @@ void FUniformExpressionSet::FillUniformBuffer(const FMaterialRenderContext& Mate
 		// Cache volume texture uniform expressions.
 		for (int32 ExpressionIndex = 0;ExpressionIndex < UniformVolumeTextureExpressions.Num();ExpressionIndex++)
 		{
-			const UTexture* Value;
+			const UTexture* Value = nullptr;
 			UniformVolumeTextureExpressions[ExpressionIndex]->GetTextureValue(MaterialRenderContext,MaterialRenderContext.Material,Value);
 
 			void** ResourceTableTexturePtr = (void**)((uint8*)BufferCursor + 0 * SHADER_PARAMETER_POINTER_ALIGNMENT);
@@ -852,7 +852,7 @@ void FUniformExpressionSet::FillUniformBuffer(const FMaterialRenderContext& Mate
 			// Check for streaming virtual texture
 			if (!bValidResources)
 			{
-				const UTexture* Texture;
+				const UTexture* Texture = nullptr;
 				UniformVirtualTextureExpressions[ExpressionIndex]->GetTextureValue(MaterialRenderContext, MaterialRenderContext.Material, Texture);
 				if (Texture && Texture->Resource)
 				{
@@ -877,7 +877,7 @@ void FUniformExpressionSet::FillUniformBuffer(const FMaterialRenderContext& Mate
 			// Now check for runtime virtual texture
 			if (!bValidResources)
 			{
-				const URuntimeVirtualTexture* Texture;
+				const URuntimeVirtualTexture* Texture = nullptr;
 				UniformVirtualTextureExpressions[ExpressionIndex]->GetTextureValue(MaterialRenderContext, MaterialRenderContext.Material, Texture);
 				if (Texture != nullptr)
 				{
