Fix for crashing shader complexity view

#rnx

diff --git a/Engine/Source/Runtime/Engine/Private/DebugViewModeMaterialProxy.cpp b/Engine/Source/Runtime/Engine/Private/DebugViewModeMaterialProxy.cpp
--- a/Engine/Source/Runtime/Engine/Private/DebugViewModeMaterialProxy.cpp
+++ b/Engine/Source/Runtime/Engine/Private/DebugViewModeMaterialProxy.cpp
@@ -60,6 +60,7 @@ FDebugViewModeMaterialProxy::FDebugViewModeMaterialProxy(
 				bIsUsedWithSplineMeshes = Resource->IsUsedWithSplineMeshes();
 				bIsUsedWithInstancedStaticMeshes = Resource->IsUsedWithInstancedStaticMeshes();
 				bIsUsedWithAPEXCloth = Resource->IsUsedWithAPEXCloth();
+				bIsUsedWithWater = Resource->IsUsedWithWater();
 			}
 		}
 
diff --git a/Engine/Source/Runtime/Engine/Private/DebugViewModeMaterialProxy.h b/Engine/Source/Runtime/Engine/Private/DebugViewModeMaterialProxy.h
--- a/Engine/Source/Runtime/Engine/Private/DebugViewModeMaterialProxy.h
+++ b/Engine/Source/Runtime/Engine/Private/DebugViewModeMaterialProxy.h
@@ -163,6 +163,7 @@ public:
 	virtual bool IsUsedWithSplineMeshes() const override { return bIsUsedWithSplineMeshes; }
 	virtual bool IsUsedWithInstancedStaticMeshes() const override { return bIsUsedWithInstancedStaticMeshes; }
 	virtual bool IsUsedWithAPEXCloth() const override { return bIsUsedWithAPEXCloth; }
+	virtual bool IsUsedWithWater() const override { return bIsUsedWithWater; }
 
 	virtual EMaterialShaderMapUsage::Type GetMaterialShaderMapUsage() const { return Usage; }
 
@@ -195,6 +196,7 @@ private:
 			uint32 bIsUsedWithSplineMeshes : 1;
 			uint32 bIsUsedWithInstancedStaticMeshes : 1;
 			uint32 bIsUsedWithAPEXCloth : 1;
+			uint32 bIsUsedWithWater : 1;
 		};
 	};
 
