#ue4 - ProjectileMovement: add AddForce() and related functions to allow external applications of forces that accumulate until the next tick, at which time they are applied to acceleration and zeroed. This allows application of a force during the correct stage of integration and when there is substepping involved.

diff --git a/Engine/Source/Runtime/Engine/Classes/GameFramework/ProjectileMovementComponent.h b/Engine/Source/Runtime/Engine/Classes/GameFramework/ProjectileMovementComponent.h
--- a/Engine/Source/Runtime/Engine/Classes/GameFramework/ProjectileMovementComponent.h
+++ b/Engine/Source/Runtime/Engine/Classes/GameFramework/ProjectileMovementComponent.h
@@ -415,8 +415,21 @@ public:
 	/** Allow the projectile to track towards its homing target. */
 	virtual FVector ComputeHomingAcceleration(const FVector& InVelocity, float DeltaTime) const;
 
+	/** Adds a force which is accumulated until next tick, used by ComputeAcceleration() to affect Velocity. */
+	void AddForce(FVector Force);
+
+	/** Returns the sum of pending forces from AddForce(). */
+	FVector GetPendingForce() const { return PendingForce; }
+
+	/** Clears any pending forces from AddForce(). */
+	void ClearPendingForce();
+	
 protected:
 
+	// Double-buffer of pending force so that updates can use the accumulated value and reset the data so other AddForce() calls work correctly.
+	// Also prevents accumulation over frames where the update aborts for whatever reason, and works with substepping movement.
+	FVector PendingForceThisUpdate;
+
 	virtual void TickInterpolation(float DeltaTime);
 	
 	FVector InterpLocationOffset;
@@ -425,6 +438,11 @@ protected:
 	FQuat InterpRotationOffset;
 	FQuat InterpInitialRotationOffset;
 
+private:
+
+	// Pending force for next tick.
+	FVector PendingForce;
+
 public:
 	/** Compute gravity effect given current physics volume, projectile gravity scale, etc. */
 	virtual float GetGravityZ() const override;
diff --git a/Engine/Source/Runtime/Engine/Private/Components/ProjectileMovementComponent.cpp b/Engine/Source/Runtime/Engine/Private/Components/ProjectileMovementComponent.cpp
--- a/Engine/Source/Runtime/Engine/Private/Components/ProjectileMovementComponent.cpp
+++ b/Engine/Source/Runtime/Engine/Private/Components/ProjectileMovementComponent.cpp
@@ -140,6 +140,11 @@ void UProjectileMovementComponent::TickComponent(float DeltaTime, enum ELevelTic
 		TickInterpolation(DeltaTime);
 	}
 
+	// Consume PendingForce and reset to zero.
+	// At this point, any calls to AddForce() will apply to the next frame.
+	PendingForceThisUpdate = PendingForce;
+	ClearPendingForce();
+
 	// skip if don't want component updated when not rendered or updated component can't move
 	if (HasStoppedSimulation() || ShouldSkipUpdate(DeltaTime))
 	{
@@ -466,6 +471,8 @@ FVector UProjectileMovementComponent::ComputeAcceleration(const FVector& InVeloc
 
 	Acceleration.Z += GetGravityZ();
 
+	Acceleration += PendingForceThisUpdate;
+
 	if (bIsHomingProjectile && HomingTargetComponent.IsValid())
 	{
 		Acceleration += ComputeHomingAcceleration(InVelocity, DeltaTime);
@@ -482,6 +489,16 @@ FVector UProjectileMovementComponent::ComputeHomingAcceleration(const FVector& I
 }
 
 
+void UProjectileMovementComponent::AddForce(FVector Force)
+{
+	PendingForce += Force;
+}
+
+void UProjectileMovementComponent::ClearPendingForce()
+{
+	PendingForce = FVector::ZeroVector;
+}
+
 float UProjectileMovementComponent::GetGravityZ() const
 {
 	// TODO: apply buoyancy if in water
