Fix crash loading cooked static meshes.

diff --git a/Engine/Source/Runtime/Engine/Private/StaticMesh.cpp b/Engine/Source/Runtime/Engine/Private/StaticMesh.cpp
--- a/Engine/Source/Runtime/Engine/Private/StaticMesh.cpp
+++ b/Engine/Source/Runtime/Engine/Private/StaticMesh.cpp
@@ -2538,7 +2538,7 @@ void UStaticMesh::InitResources()
 		//&& !bTemporarilyDisableStreaming;
 
 #if (WITH_EDITOR && DO_CHECK)
-	if (bIsStreamable)
+	if (bIsStreamable && !GetOutermost()->bIsCookedForEditor)
 	{
 		for (int32 LODIdx = 0; LODIdx < NumLODs; ++LODIdx)
 		{
