LOD streaming - prevent a level from screen size culling if it meshes with forced LOD level

diff --git a/Engine/Source/Runtime/Engine/Private/Streaming/AsyncTextureStreaming.cpp b/Engine/Source/Runtime/Engine/Private/Streaming/AsyncTextureStreaming.cpp
--- a/Engine/Source/Runtime/Engine/Private/Streaming/AsyncTextureStreaming.cpp
+++ b/Engine/Source/Runtime/Engine/Private/Streaming/AsyncTextureStreaming.cpp
@@ -85,7 +85,8 @@ void FAsyncRenderAssetStreamingData::UpdateBoundSizes_Async(const FRenderAssetSt
 		StaticInstancesView.UpdateBoundSizes_Async(ViewInfos, ViewInfoExtras, LastUpdateTime, Settings);
 
 		// Skip levels that can not contribute to resolution.
-		if (StaticInstancesView.GetMaxLevelRenderAssetScreenSize() > Settings.MinLevelRenderAssetScreenSize)
+		if (StaticInstancesView.GetMaxLevelRenderAssetScreenSize() > Settings.MinLevelRenderAssetScreenSize
+			|| StaticInstancesView.HasAnyComponentWithForcedLOD())
 		{
 			StaticInstancesViewIndices.Add(StaticViewIndex);
 		}
@@ -175,7 +176,8 @@ void FAsyncRenderAssetStreamingData::UpdatePerfectWantedMips_Async(FStreamingRen
 					continue;
 				}
 
-				if (StaticInstancesView.GetMaxLevelRenderAssetScreenSize() < Settings.MinLevelRenderAssetScreenSize)
+				if (StaticInstancesView.GetMaxLevelRenderAssetScreenSize() < Settings.MinLevelRenderAssetScreenSize
+					&& !StaticInstancesView.HasComponentWithForcedLOD(RenderAsset))
 				{
 					bCulled = true;
 					continue;
diff --git a/Engine/Source/Runtime/Engine/Private/Streaming/TextureInstanceState.cpp b/Engine/Source/Runtime/Engine/Private/Streaming/TextureInstanceState.cpp
--- a/Engine/Source/Runtime/Engine/Private/Streaming/TextureInstanceState.cpp
+++ b/Engine/Source/Runtime/Engine/Private/Streaming/TextureInstanceState.cpp
@@ -152,6 +152,19 @@ void FRenderAssetInstanceState::AddElement(const UPrimitiveComponent* InComponen
 	if (HasCompiledElements()) 
 	{
 		CompiledRenderAssetMap.FindOrAdd(Element.RenderAsset).Add(FCompiledElement(Element));
+
+		if (Element.TexelFactor < 0.f && !InAsset->IsA<UTexture>())
+		{
+			int32* CountPtr = CompiledNumForcedLODCompMap.Find(Element.RenderAsset);
+			if (!CountPtr)
+			{
+				CompiledNumForcedLODCompMap.Add(Element.RenderAsset, 1);
+			}
+			else
+			{
+				++*CountPtr;
+			}
+		}
 	}
 }
 
@@ -167,6 +180,14 @@ void FRenderAssetInstanceState::RemoveElement(int32 ElementIndex, int32& NextCom
 	if (HasCompiledElements())
 	{
 		CompiledRenderAssetMap.FindChecked(Element.RenderAsset).RemoveSingleSwap(FCompiledElement(Element), false);
+
+		if (Element.TexelFactor < 0.f
+			&& Element.RenderAsset
+			&& !Element.RenderAsset->IsA<UTexture>()
+			&& !--CompiledNumForcedLODCompMap.FindChecked(Element.RenderAsset))
+		{
+			CompiledNumForcedLODCompMap.Remove(Element.RenderAsset);
+		}
 	}
 
 	// Unlink texutres or meshes
@@ -184,6 +205,7 @@ void FRenderAssetInstanceState::RemoveElement(int32 ElementIndex, int32& NextCom
 			{
 				RenderAssetMap.Remove(Element.RenderAsset);
 				CompiledRenderAssetMap.Remove(Element.RenderAsset);
+				check(!CompiledNumForcedLODCompMap.Find(Element.RenderAsset));
 				Asset = Element.RenderAsset;
 			}
 		}
@@ -583,12 +605,14 @@ uint32 FRenderAssetInstanceState::GetAllocatedSize() const
 		FreeElementIndices.GetAllocatedSize() +
 		RenderAssetMap.GetAllocatedSize() +
 		CompiledRenderAssetMap.GetAllocatedSize() + CompiledElementsSize +
+		CompiledNumForcedLODCompMap.GetAllocatedSize() +
 		ComponentMap.GetAllocatedSize();
 }
 
 int32 FRenderAssetInstanceState::CompileElements()
 {
 	CompiledRenderAssetMap.Empty();
+	CompiledNumForcedLODCompMap.Empty();
 	MaxTexelFactor = 0;
 
 	// First create an entry for all elements, so that there are no reallocs when inserting each compiled elements.
@@ -602,6 +626,7 @@ int32 FRenderAssetInstanceState::CompileElements()
 	for (TMap<const UStreamableRenderAsset*, TArray<FCompiledElement> >::TIterator It(CompiledRenderAssetMap); It; ++It)
 	{
 		const UStreamableRenderAsset* Asset = It.Key();
+		const bool bIsNonTexture = Asset && !Asset->IsA<UTexture>();
 		TArray<FCompiledElement>& CompiledElemements = It.Value();
 
 		int32 CompiledElementCount = 0;
@@ -615,6 +640,20 @@ int32 FRenderAssetInstanceState::CompileElements()
 		for (auto ElementIt = GetElementIterator(Asset); ElementIt; ++ElementIt)
 		{
 			const float TexelFactor = ElementIt.GetTexelFactor();
+
+			if (bIsNonTexture && TexelFactor < 0.f)
+			{
+				int32* CountPtr = CompiledNumForcedLODCompMap.Find(Asset);
+				if (!CountPtr)
+				{
+					CompiledNumForcedLODCompMap.Add(Asset, 1);
+				}
+				else
+				{
+					++*CountPtr;
+				}
+			}
+
 			// No need to care about force load as MaxTexelFactor is used to ignore far away levels.
 			MaxTexelFactor = FMath::Max(TexelFactor, MaxTexelFactor);
 
diff --git a/Engine/Source/Runtime/Engine/Private/Streaming/TextureInstanceView.cpp b/Engine/Source/Runtime/Engine/Private/Streaming/TextureInstanceView.cpp
--- a/Engine/Source/Runtime/Engine/Private/Streaming/TextureInstanceView.cpp
+++ b/Engine/Source/Runtime/Engine/Private/Streaming/TextureInstanceView.cpp
@@ -256,6 +256,7 @@ void FRenderAssetInstanceView::SwapData(FRenderAssetInstanceView* Lfs, FRenderAs
 	check(Lfs->Elements.Num() == Rhs->Elements.Num());
 	check(Lfs->RenderAssetMap.Num() == Rhs->RenderAssetMap.Num());
 	check(Lfs->CompiledRenderAssetMap.Num() == 0 && Rhs->CompiledRenderAssetMap.Num() == 0);
+	check(!Lfs->CompiledNumForcedLODCompMap.Num() && !Rhs->CompiledNumForcedLODCompMap.Num());
 
 	FMemory::Memswap(&Lfs->Bounds4 , &Rhs->Bounds4, sizeof(Lfs->Bounds4));
 	FMemory::Memswap(&Lfs->Elements , &Rhs->Elements, sizeof(Lfs->Elements));
@@ -600,3 +601,13 @@ bool FRenderAssetInstanceAsyncView::HasRenderAssetReferences(const UStreamableRe
 {
 	return View.IsValid() && (bool)View->GetElementIterator(InAsset);
 }
+
+bool FRenderAssetInstanceAsyncView::HasComponentWithForcedLOD(const UStreamableRenderAsset* InAsset) const
+{
+	return View.IsValid() && View->HasComponentWithForcedLOD(InAsset);
+}
+
+bool FRenderAssetInstanceAsyncView::HasAnyComponentWithForcedLOD() const
+{
+	return View.IsValid() && View->HasAnyComponentWithForcedLOD();
+}
diff --git a/Engine/Source/Runtime/Engine/Private/Streaming/TextureInstanceView.h b/Engine/Source/Runtime/Engine/Private/Streaming/TextureInstanceView.h
--- a/Engine/Source/Runtime/Engine/Private/Streaming/TextureInstanceView.h
+++ b/Engine/Source/Runtime/Engine/Private/Streaming/TextureInstanceView.h
@@ -193,6 +193,9 @@ public:
 	// If this has compiled elements, return the array relate to a given texture or mesh.
 	const TArray<FCompiledElement>* GetCompiledElements(const UStreamableRenderAsset* Asset) const { return CompiledRenderAssetMap.Find(Asset); }
 
+	bool HasComponentWithForcedLOD(const UStreamableRenderAsset* Asset) const { return !!CompiledNumForcedLODCompMap.Find(Asset); }
+	bool HasAnyComponentWithForcedLOD() const { return !!CompiledNumForcedLODCompMap.Num(); }
+
 	static TRefCountPtr<const FRenderAssetInstanceView> CreateView(const FRenderAssetInstanceView* RefView);
 	static TRefCountPtr<FRenderAssetInstanceView> CreateViewWithUninitializedBounds(const FRenderAssetInstanceView* RefView);
 	static void SwapData(FRenderAssetInstanceView* Lfs, FRenderAssetInstanceView* Rhs);
@@ -226,6 +229,9 @@ protected:
 	// CompiledTextureMap is used to iterate more quickly on each elements by avoiding the linked list indirections.
 	TMap<const UStreamableRenderAsset*, TArray<FCompiledElement> > CompiledRenderAssetMap;
 
+	// If an asset has components with forced LOD levels, it will appear here after level compilation.
+	TMap<const UStreamableRenderAsset*, int32> CompiledNumForcedLODCompMap;
+
 	/** Max texel factor across all elements. Used for early culling */
 	float MaxTexelFactor;
 
@@ -273,6 +279,10 @@ public:
 
 	bool HasRenderAssetReferences(const UStreamableRenderAsset* InAsset) const;
 
+	bool HasComponentWithForcedLOD(const UStreamableRenderAsset* InAsset) const;
+
+	bool HasAnyComponentWithForcedLOD() const;
+
 	// Release the data now as this is expensive.
 	void OnTaskDone() { BoundsViewInfo.Empty(); }
 
