Suspend rendering thread now aborts texture streaming waiting on DDC.
This is implemeted by testing if the asset streaming was cancelled or if the rendering thread is suspended every "r.Streaming.DDCPendingSleep" milliseconds.
It can be made in sync too, as before, by setting the value to 0 (default 5).

diff --git a/Engine/Source/Runtime/Engine/Private/Streaming/Texture2DStreamIn_DDC.cpp b/Engine/Source/Runtime/Engine/Private/Streaming/Texture2DStreamIn_DDC.cpp
--- a/Engine/Source/Runtime/Engine/Private/Streaming/Texture2DStreamIn_DDC.cpp
+++ b/Engine/Source/Runtime/Engine/Private/Streaming/Texture2DStreamIn_DDC.cpp
@@ -11,6 +11,14 @@ Texture2DStreamIn_DDC.cpp: Stream in helper for 2D textures loading DDC files.
 
 #if WITH_EDITORONLY_DATA
 
+static int32 GStreamingDDCPendingSleep = 5;
+static FAutoConsoleVariableRef CVarStreamingDDCPendingSleep(
+	TEXT("r.Streaming.DDCPendingSleep"),
+	GStreamingDDCPendingSleep,
+	TEXT("Sleep increment while waiting for ddc to be built in milliseconds. 0 to wait synchronously. Allows to abort streaming on cancel and suspend rendering thread. (default=5) \n"),
+	ECVF_Default
+);
+
 void FTexture2DStreamIn_DDC::DoLoadNewMipsFromDDC(const FContext& Context)
 {
 	if (Context.Texture && Context.Resource)
@@ -30,7 +38,37 @@ void FTexture2DStreamIn_DDC::DoLoadNewMipsFromDDC(const FContext& Context)
 			{
 				// The overhead of doing 2 copy of each mip data (from GetSynchronous() and FMemoryReader) in hidden by other texture DDC ops happening at the same time.
 				TArray<uint8> DerivedMipData;
-				if (GetDerivedDataCacheRef().GetSynchronous(*MipMap.DerivedDataKey, DerivedMipData))
+				
+				bool bDDCReady = false; // Whether the DDC is ready to be red.
+				bool bDDCValid= false; // Whether the DDC is read data is good.
+
+				if (GStreamingDDCPendingSleep > 0)
+				{
+					const uint32 AsyncHandle = GetDerivedDataCacheRef().GetAsynchronous(*MipMap.DerivedDataKey);
+					while (!IsCancelled() && !IsAssetStreamingSuspended() && !bDDCReady)
+					{
+						bDDCReady = GetDerivedDataCacheRef().PollAsynchronousCompletion(AsyncHandle);
+						if (!bDDCReady)
+						{
+							FPlatformProcess::Sleep(GStreamingDDCPendingSleep * 0.001f);
+						}
+					}
+					if (bDDCReady)
+					{
+						bDDCValid = GetDerivedDataCacheRef().GetAsynchronousResults(AsyncHandle, DerivedMipData);
+					}
+				}
+				else // sync mode (this can stall in suspend rendering thread and GC)
+				{
+					bDDCReady = true;
+					bDDCValid = GetDerivedDataCacheRef().GetSynchronous(*MipMap.DerivedDataKey, DerivedMipData);
+				}
+
+				if (!bDDCReady)
+				{
+					MarkAsCancelled();
+				}
+				else if (bDDCValid)
 				{
 					FMemoryReader Ar(DerivedMipData, true);
 
