Back out changelist 8932361. Backed up thread pool caused editor freeze when saving

diff --git a/Engine/Source/Runtime/Engine/Private/Streaming/Texture2DStreamIn_DDC.cpp b/Engine/Source/Runtime/Engine/Private/Streaming/Texture2DStreamIn_DDC.cpp
--- a/Engine/Source/Runtime/Engine/Private/Streaming/Texture2DStreamIn_DDC.cpp
+++ b/Engine/Source/Runtime/Engine/Private/Streaming/Texture2DStreamIn_DDC.cpp
@@ -11,14 +11,6 @@ Texture2DStreamIn_DDC.cpp: Stream in helper for 2D textures loading DDC files.
 
 #if WITH_EDITORONLY_DATA
 
-static int32 GStreamingDDCPendingSleep = 5;
-static FAutoConsoleVariableRef CVarStreamingDDCPendingSleep(
-	TEXT("r.Streaming.DDCPendingSleep"),
-	GStreamingDDCPendingSleep,
-	TEXT("Sleep increment while waiting for ddc to be built in milliseconds. 0 to wait synchronously. Allows to abort streaming on cancel and suspend rendering thread. (default=5) \n"),
-	ECVF_Default
-);
-
 void FTexture2DStreamIn_DDC::DoLoadNewMipsFromDDC(const FContext& Context)
 {
 	if (Context.Texture && Context.Resource)
@@ -38,37 +30,7 @@ void FTexture2DStreamIn_DDC::DoLoadNewMipsFromDDC(const FContext& Context)
 			{
 				// The overhead of doing 2 copy of each mip data (from GetSynchronous() and FMemoryReader) in hidden by other texture DDC ops happening at the same time.
 				TArray<uint8> DerivedMipData;
-				
-				bool bDDCReady = false; // Whether the DDC is ready to be red.
-				bool bDDCValid= false; // Whether the DDC is read data is good.
-
-				if (GStreamingDDCPendingSleep > 0)
-				{
-					const uint32 AsyncHandle = GetDerivedDataCacheRef().GetAsynchronous(*MipMap.DerivedDataKey);
-					while (!IsCancelled() && !IsAssetStreamingSuspended() && !bDDCReady)
-					{
-						bDDCReady = GetDerivedDataCacheRef().PollAsynchronousCompletion(AsyncHandle);
-						if (!bDDCReady)
-						{
-							FPlatformProcess::Sleep(GStreamingDDCPendingSleep * 0.001f);
-						}
-					}
-					if (bDDCReady)
-					{
-						bDDCValid = GetDerivedDataCacheRef().GetAsynchronousResults(AsyncHandle, DerivedMipData);
-					}
-				}
-				else // sync mode (this can stall in suspend rendering thread and GC)
-				{
-					bDDCReady = true;
-					bDDCValid = GetDerivedDataCacheRef().GetSynchronous(*MipMap.DerivedDataKey, DerivedMipData);
-				}
-
-				if (!bDDCReady)
-				{
-					MarkAsCancelled();
-				}
-				else if (bDDCValid)
+				if (GetDerivedDataCacheRef().GetSynchronous(*MipMap.DerivedDataKey, DerivedMipData))
 				{
 					FMemoryReader Ar(DerivedMipData, true);
 
