#ue4 - Clear pending forces in ProjectileMovement when simulation is stopped.

diff --git a/Engine/Source/Runtime/Engine/Private/Components/ProjectileMovementComponent.cpp b/Engine/Source/Runtime/Engine/Private/Components/ProjectileMovementComponent.cpp
--- a/Engine/Source/Runtime/Engine/Private/Components/ProjectileMovementComponent.cpp
+++ b/Engine/Source/Runtime/Engine/Private/Components/ProjectileMovementComponent.cpp
@@ -509,6 +509,8 @@ float UProjectileMovementComponent::GetGravityZ() const
 void UProjectileMovementComponent::StopSimulating(const FHitResult& HitResult)
 {
 	Velocity = FVector::ZeroVector;
+	PendingForce = FVector::ZeroVector;
+	PendingForceThisUpdate = FVector::ZeroVector;
 	UpdateComponentVelocity();
 	SetUpdatedComponent(NULL);
 	OnProjectileStop.Broadcast(HitResult);
