- Don't send all EOF updates when rendering materials to textures
- Ensure that parameter collection updates, materials and render target updates are flushed

#rnx

diff --git a/Engine/Source/Runtime/Engine/Classes/Engine/World.h b/Engine/Source/Runtime/Engine/Classes/Engine/World.h
--- a/Engine/Source/Runtime/Engine/Classes/Engine/World.h
+++ b/Engine/Source/Runtime/Engine/Classes/Engine/World.h
@@ -2718,6 +2718,11 @@ public:
 	 */
 	void SendAllEndOfFrameUpdates();
 
+	/**
+	 * Flush any pending parameter collection updates to the render thrad.
+	 */
+	void FlushDeferredParameterCollectionInstanceUpdates();
+
 	/** Do per frame tick behaviors related to the network driver */
 	void TickNetClient( float DeltaSeconds );
 
diff --git a/Engine/Source/Runtime/Engine/Private/KismetRenderingLibrary.cpp b/Engine/Source/Runtime/Engine/Private/KismetRenderingLibrary.cpp
--- a/Engine/Source/Runtime/Engine/Private/KismetRenderingLibrary.cpp
+++ b/Engine/Source/Runtime/Engine/Private/KismetRenderingLibrary.cpp
@@ -113,7 +113,7 @@ void UKismetRenderingLibrary::DrawMaterialToRenderTarget(UObject* WorldContextOb
 	}
 	else
 	{
-		World->SendAllEndOfFrameUpdates();
+		World->FlushDeferredParameterCollectionInstanceUpdates();
 
 		FTextureRenderTargetResource* RenderTargetResource = TextureRenderTarget->GameThread_GetRenderTargetResource();
 
@@ -550,14 +550,17 @@ void UKismetRenderingLibrary::BeginDrawCanvasToRenderTarget(UObject* WorldContex
 	}
 	else
 	{
+		World->FlushDeferredParameterCollectionInstanceUpdates();
+
 		Context.RenderTarget = TextureRenderTarget;
 
 		Canvas = World->GetCanvasForRenderingToTarget();
 
 		Size = FVector2D(TextureRenderTarget->SizeX, TextureRenderTarget->SizeY);
 
+		FTextureRenderTargetResource* RenderTargetResource = TextureRenderTarget->GameThread_GetRenderTargetResource();
 		FCanvas* NewCanvas = new FCanvas(
-			TextureRenderTarget->GameThread_GetRenderTargetResource(), 
+			RenderTargetResource,
 			nullptr, 
 			World,
 			World->FeatureLevel, 
@@ -571,8 +574,10 @@ void UKismetRenderingLibrary::BeginDrawCanvasToRenderTarget(UObject* WorldContex
 		FName RTName = TextureRenderTarget->GetFName();
 		TDrawEvent<FRHICommandList>* DrawEvent = Context.DrawEvent;
 		ENQUEUE_RENDER_COMMAND(BeginDrawEventCommand)(
-			[RTName, DrawEvent](FRHICommandList& RHICmdList)
+			[RTName, DrawEvent, RenderTargetResource](FRHICommandListImmediate& RHICmdList)
 			{
+				RenderTargetResource->FlushDeferredResourceUpdate(RHICmdList);
+
 				BEGIN_DRAW_EVENTF(
 					RHICmdList, 
 					DrawCanvasToTarget, 
diff --git a/Engine/Source/Runtime/Engine/Private/LevelTick.cpp b/Engine/Source/Runtime/Engine/Private/LevelTick.cpp
--- a/Engine/Source/Runtime/Engine/Private/LevelTick.cpp
+++ b/Engine/Source/Runtime/Engine/Private/LevelTick.cpp
@@ -1080,6 +1080,24 @@ void UWorld::SendAllEndOfFrameUpdates()
 	EndSendEndOfFrameUpdatesDrawEvent(SendAllEndOfFrameUpdates);
 }
 
+/**
+ * Flush any pending parameter collection updates to the render thrad.
+ */
+void UWorld::FlushDeferredParameterCollectionInstanceUpdates()
+{
+	if ( bMaterialParameterCollectionInstanceNeedsDeferredUpdate )
+	{
+		for (UMaterialParameterCollectionInstance* ParameterCollectionInstance : ParameterCollectionInstances)
+		{
+			if (ParameterCollectionInstance)
+			{
+				ParameterCollectionInstance->DeferredUpdateRenderState(false);
+			}
+		}
+
+		bMaterialParameterCollectionInstanceNeedsDeferredUpdate = false;
+	}
+}
 
 #if !UE_BUILD_SHIPPING
 static class FFileProfileWrapperExec: private FSelfRegisteringExec
