Handle having sublevels hierarchy when calling UWorld::CleanupWorld: use a global counter to check if a world has already been cleaned up or not.
Note: assumes that UWorld::CleanupWorld is not thread-safe.

diff --git a/Engine/Source/Runtime/Engine/Classes/Engine/World.h b/Engine/Source/Runtime/Engine/Classes/Engine/World.h
--- a/Engine/Source/Runtime/Engine/Classes/Engine/World.h
+++ b/Engine/Source/Runtime/Engine/Classes/Engine/World.h
@@ -1559,10 +1559,11 @@ public:
 	/** Indicates that the world has marked contained objects as pending kill */
 	bool HasMarkedObjectsPendingKill() const { return bMarkedObjectsPendingKill; }
 private:
-	uint32 bCleanedUpWorld:1;
-
 	uint32 bMarkedObjectsPendingKill:1;
 
+	uint32 CleanupWorldTag;
+	static uint32 CleanupWorldGlobalTag;
+
 public:
 #if WITH_EDITORONLY_DATA
 	/** List of DDC async requests we need to wait on before we register components. Game thread only. */
@@ -2426,9 +2427,8 @@ public:
 	 * Cleans up components, streaming data and assorted other intermediate data.
 	 * @param bSessionEnded whether to notify the viewport that the game session has ended.
 	 * @param NewWorld Optional new world that will be loaded after this world is cleaned up. Specify a new world to prevent it and it's sublevels from being GCed during map transitions.
-	 * @param bResetCleanedUpFlag wheter to reset the bCleanedUpWorld flag or not.
 	 */
-	void CleanupWorld(bool bSessionEnded = true, bool bCleanupResources = true, UWorld* NewWorld = nullptr, bool bResetCleanedUpFlag = true);
+	void CleanupWorld(bool bSessionEnded = true, bool bCleanupResources = true, UWorld* NewWorld = nullptr);
 	
 	/**
 	 * Invalidates the cached data used to render the levels' UModel.
@@ -2878,6 +2878,9 @@ public:
 	void BeginTearingDown();
 
 private:
+	/** Internal version of CleanupWorld. */
+	void CleanupWorldInternal(bool bSessionEnded, bool bCleanupResources, UWorld* NewWorld);
+
 	/** Utility function to handle Exec/Console Commands related to the Trace Tags */
 	bool HandleTraceTagCommand( const TCHAR* Cmd, FOutputDevice& Ar );
 
diff --git a/Engine/Source/Runtime/Engine/Private/World.cpp b/Engine/Source/Runtime/Engine/Private/World.cpp
--- a/Engine/Source/Runtime/Engine/Private/World.cpp
+++ b/Engine/Source/Runtime/Engine/Private/World.cpp
@@ -337,6 +337,8 @@ FWorldDelegates::FRefreshLevelScriptActionsEvent FWorldDelegates::RefreshLevelSc
 
 UWorld::FOnWorldInitializedActors FWorldDelegates::OnWorldInitializedActors;
 
+uint32 UWorld::CleanupWorldGlobalTag = 0;
+
 UWorld::UWorld( const FObjectInitializer& ObjectInitializer )
 : UObject(ObjectInitializer)
 , FeatureLevel(GMaxRHIFeatureLevel)
@@ -352,6 +354,7 @@ UWorld::UWorld( const FObjectInitializer& ObjectInitializer )
 ,	TickTaskLevel(FTickTaskManagerInterface::Get().AllocateTickTaskLevel())
 ,	FlushLevelStreamingType(EFlushLevelStreamingType::None)
 ,	NextTravelType(TRAVEL_Relative)
+,	CleanupWorldTag(0)
 {
 	TimerManager = new FTimerManager();
 #if WITH_EDITOR
@@ -4120,22 +4123,25 @@ bool UWorld::IsNavigationRebuilt() const
 	return GetNavigationSystem() == NULL || GetNavigationSystem()->IsNavigationBuilt(GetWorldSettings());
 }
 
-void UWorld::CleanupWorld(bool bSessionEnded, bool bCleanupResources, UWorld* NewWorld, bool bResetCleanedUpFlag)
+void UWorld::CleanupWorld(bool bSessionEnded, bool bCleanupResources, UWorld* NewWorld)
 {
-	UE_LOG(LogWorld, Log, TEXT("UWorld::CleanupWorld for %s, bSessionEnded=%s, bCleanupResources=%s"), *GetName(), bSessionEnded ? TEXT("true") : TEXT("false"), bCleanupResources ? TEXT("true") : TEXT("false"));
-
-	TArray<UWorld*> WorldsToResetCleanedUpFlag;
-
-	check(IsVisibilityRequestPending() == false);
-	
-	check(!bCleanedUpWorld);
-	bCleanedUpWorld = true;
+    CleanupWorldGlobalTag++;
+	CleanupWorldInternal(bSessionEnded, bCleanupResources, NewWorld);
+}
 
-	if (bResetCleanedUpFlag)
+void UWorld::CleanupWorldInternal(bool bSessionEnded, bool bCleanupResources, UWorld* NewWorld)
+{
+	if(CleanupWorldTag == CleanupWorldGlobalTag)
 	{
-		WorldsToResetCleanedUpFlag.Add(this);
+		return;
 	}
 
+	CleanupWorldTag = CleanupWorldGlobalTag;
+
+	UE_LOG(LogWorld, Log, TEXT("UWorld::CleanupWorld for %s, bSessionEnded=%s, bCleanupResources=%s"), *GetName(), bSessionEnded ? TEXT("true") : TEXT("false"), bCleanupResources ? TEXT("true") : TEXT("false"));
+
+	check(IsVisibilityRequestPending() == false);
+	
 	// Wait on current physics scenes if they are processing
 	if(FPhysScene* CurrPhysicsScene = GetPhysicsScene())
 	{
@@ -4232,15 +4238,7 @@ void UWorld::CleanupWorld(bool bSessionEnded, bool bCleanupResources, UWorld* Ne
 	for (int32 LevelIndex = 0; LevelIndex < GetNumLevels(); ++LevelIndex)
 	{
 		UWorld* World = CastChecked<UWorld>(GetLevel(LevelIndex)->GetOuter());
-		if (!World->bCleanedUpWorld)
-		{
-			World->CleanupWorld(bSessionEnded, bCleanupResources, NewWorld, false);
-
-			if (bResetCleanedUpFlag)
-			{
-				WorldsToResetCleanedUpFlag.Add(World);
-			}
-		}
+		World->CleanupWorldInternal(bSessionEnded, bCleanupResources, NewWorld);
 	}
 
 	for (ULevelStreaming* StreamingLevel : GetStreamingLevels())
@@ -4248,15 +4246,7 @@ void UWorld::CleanupWorld(bool bSessionEnded, bool bCleanupResources, UWorld* Ne
 		if (ULevel* Level = StreamingLevel->GetLoadedLevel())
 		{
 			UWorld* World = CastChecked<UWorld>(Level->GetOuter());
-			if (!World->bCleanedUpWorld)
-			{
-				World->CleanupWorld(bSessionEnded, bCleanupResources, NewWorld, false);
-
-				if (bResetCleanedUpFlag)
-				{
-					WorldsToResetCleanedUpFlag.Add(World);
-				}
-			}
+			World->CleanupWorldInternal(bSessionEnded, bCleanupResources, NewWorld);
 		}
 	}
 
@@ -4266,20 +4256,10 @@ void UWorld::CleanupWorld(bool bSessionEnded, bool bCleanupResources, UWorld* Ne
 	{
 		for (const ULevel* Level : DuplicateCollection->GetLevels())
 		{
-			if (!Level)
-			{
-				continue;
-			}
-
-			UWorld* const LevelWorld = CastChecked<UWorld>(Level->GetOuter());
-			if (!LevelWorld->bCleanedUpWorld)
+			if (Level)
 			{
-				LevelWorld->CleanupWorld(bSessionEnded, bCleanupResources, NewWorld, false);
-
-				if (bResetCleanedUpFlag)
-				{
-					WorldsToResetCleanedUpFlag.Add(LevelWorld);
-				}
+				UWorld* const LevelWorld = CastChecked<UWorld>(Level->GetOuter());
+				LevelWorld->CleanupWorldInternal(bSessionEnded, bCleanupResources, NewWorld);
 			}
 		}
 	}
@@ -4288,12 +4268,6 @@ void UWorld::CleanupWorld(bool bSessionEnded, bool bCleanupResources, UWorld* Ne
 
 	FWorldDelegates::OnPostWorldCleanup.Broadcast(this, bSessionEnded, bCleanupResources);
 
-	for (UWorld* WorldToResetCleanedUpFlag: WorldsToResetCleanedUpFlag)
-	{
-		check(WorldToResetCleanedUpFlag->bCleanedUpWorld);
-		WorldToResetCleanedUpFlag->bCleanedUpWorld = false;
-	}
-
 	SubsystemCollection.Deinitialize();
 }
 
