Added texture source image support for G16
Needed for heightmaps

diff --git a/Engine/Source/Developer/Mac/MacTargetPlatform/Private/GenericMacTargetPlatform.h b/Engine/Source/Developer/Mac/MacTargetPlatform/Private/GenericMacTargetPlatform.h
--- a/Engine/Source/Developer/Mac/MacTargetPlatform/Private/GenericMacTargetPlatform.h
+++ b/Engine/Source/Developer/Mac/MacTargetPlatform/Private/GenericMacTargetPlatform.h
@@ -189,6 +189,7 @@ return TSuper::SupportsFeature(Feature);
 		static FName NameBGRA8(TEXT("BGRA8"));
 		static FName NameXGXR8(TEXT("XGXR8"));
 		static FName NameG8(TEXT("G8"));
+		static FName NameG16(TEXT("G16"));
 		static FName NameVU8(TEXT("VU8"));
 		static FName NameRGBA16F(TEXT("RGBA16F"));
 		static FName NameBC6H(TEXT("BC6H"));
@@ -212,6 +213,10 @@ return TSuper::SupportsFeature(Feature);
 			{
 				TextureFormatName = NameRGBA16F;
 			}
+			else if (SourceFormat == TSF_G16)
+			{
+				TextureFormatName = NameG16;
+			}
 			else if (SourceFormat == TSF_G8 || Settings == TC_Grayscale)
 			{
 				TextureFormatName = NameG8;
diff --git a/Engine/Source/Developer/TextureFormatUncompressed/Private/TextureFormatUncompressed.cpp b/Engine/Source/Developer/TextureFormatUncompressed/Private/TextureFormatUncompressed.cpp
--- a/Engine/Source/Developer/TextureFormatUncompressed/Private/TextureFormatUncompressed.cpp
+++ b/Engine/Source/Developer/TextureFormatUncompressed/Private/TextureFormatUncompressed.cpp
@@ -17,6 +17,7 @@ DEFINE_LOG_CATEGORY_STATIC(LogTextureFormatUncompressed, Log, All);
 #define ENUM_SUPPORTED_FORMATS(op) \
 	op(BGRA8) \
 	op(G8) \
+	op(G16) \
 	op(VU8) \
 	op(RGBA16F) \
 	op(XGXR8) \
@@ -87,6 +88,19 @@ class FTextureFormatUncompressed : public ITextureFormat
 
 			return true;
 		}
+		else if (BuildSettings.TextureFormatName == GTextureFormatNameG16)
+		{
+			FImage Image;
+			InImage.CopyTo(Image, ERawImageFormat::G16, EGammaSpace::Linear);
+
+			OutCompressedImage.SizeX = Image.SizeX;
+			OutCompressedImage.SizeY = Image.SizeY;
+			OutCompressedImage.SizeZ = BuildSettings.bVolume ? Image.NumSlices : 1;
+			OutCompressedImage.PixelFormat = PF_G16;
+			OutCompressedImage.RawData = Image.RawData;
+
+			return true;
+		}
 		else if (BuildSettings.TextureFormatName == GTextureFormatNameVU8)
 		{
 			FImage Image;
diff --git a/Engine/Source/Developer/Windows/WindowsTargetPlatform/Private/GenericWindowsTargetPlatform.h b/Engine/Source/Developer/Windows/WindowsTargetPlatform/Private/GenericWindowsTargetPlatform.h
--- a/Engine/Source/Developer/Windows/WindowsTargetPlatform/Private/GenericWindowsTargetPlatform.h
+++ b/Engine/Source/Developer/Windows/WindowsTargetPlatform/Private/GenericWindowsTargetPlatform.h
@@ -398,6 +398,7 @@ public:
 		static FName NameBGRA8(TEXT("BGRA8"));
 		static FName NameXGXR8(TEXT("XGXR8"));
 		static FName NameG8(TEXT("G8"));
+		static FName NameG16(TEXT("G16"));
 		static FName NameVU8(TEXT("VU8"));
 		static FName NameRGBA16F(TEXT("RGBA16F"));
 		static FName NameBC6H(TEXT("BC6H"));
@@ -421,6 +422,10 @@ public:
 			{
 				TextureFormatName = NameRGBA16F;
 			}
+			else if (SourceFormat == TSF_G16)
+			{
+				TextureFormatName = NameG16;
+			}
 			else if (SourceFormat == TSF_G8 || Settings == TC_Grayscale)
 			{
 				TextureFormatName = NameG8;
diff --git a/Engine/Source/Runtime/Engine/Classes/Engine/Texture.h b/Engine/Source/Runtime/Engine/Classes/Engine/Texture.h
--- a/Engine/Source/Runtime/Engine/Classes/Engine/Texture.h
+++ b/Engine/Source/Runtime/Engine/Classes/Engine/Texture.h
@@ -116,6 +116,8 @@ enum ETextureSourceFormat
 	//@todo: Deprecated!
 	TSF_RGBE8,
 
+	TSF_G16,
+
 	TSF_MAX
 };
 
diff --git a/Engine/Source/Runtime/Engine/Private/Texture.cpp b/Engine/Source/Runtime/Engine/Private/Texture.cpp
--- a/Engine/Source/Runtime/Engine/Private/Texture.cpp
+++ b/Engine/Source/Runtime/Engine/Private/Texture.cpp
@@ -814,8 +814,8 @@ void FTextureSource::Compress()
 		IImageWrapperModule& ImageWrapperModule = FModuleManager::LoadModuleChecked<IImageWrapperModule>( FName("ImageWrapper") );
 		TSharedPtr<IImageWrapper> ImageWrapper = ImageWrapperModule.CreateImageWrapper( EImageFormat::PNG );
 		// TODO: TSF_BGRA8 is stored as RGBA, so the R and B channels are swapped in the internal png. Should we fix this?
-		ERGBFormat RawFormat = (Format == TSF_G8) ? ERGBFormat::Gray : ERGBFormat::RGBA;
-		if ( ImageWrapper.IsValid() && ImageWrapper->SetRaw( BulkDataPtr, BulkData.GetBulkDataSize(), SizeX, SizeY, RawFormat, Format == TSF_RGBA16 ? 16 : 8 ) )
+		ERGBFormat RawFormat = (Format == TSF_G8 || Format == TSF_G16) ? ERGBFormat::Gray : ERGBFormat::RGBA;
+		if ( ImageWrapper.IsValid() && ImageWrapper->SetRaw( BulkDataPtr, BulkData.GetBulkDataSize(), SizeX, SizeY, RawFormat, Format == TSF_G16 || Format == TSF_RGBA16 ? 16 : 8 ) )
 		{
 			const TArray<uint8>& CompressedData = ImageWrapper->GetCompressed();
 			if ( CompressedData.Num() > 0 )
@@ -846,7 +846,7 @@ uint8* FTextureSource::LockMip(int32 BlockIndex, int32 LayerIndex, int32 MipInde
 			LockedMipData = (uint8*)BulkData.Lock(LOCK_READ_WRITE);
 			if (bPNGCompressed)
 			{
-				bool bCanPngCompressFormat = (Format == TSF_G8 || Format == TSF_RGBA8 || Format == TSF_BGRA8 || Format == TSF_RGBA16);
+				bool bCanPngCompressFormat = (Format == TSF_G8 || Format == TSF_G16 || Format == TSF_RGBA8 || Format == TSF_BGRA8 || Format == TSF_RGBA16);
 				check(Blocks.Num() == 0 && NumLayers == 1 && NumSlices == 1 && bCanPngCompressFormat);
 				if (MipIndex != 0)
 				{
@@ -861,7 +861,7 @@ uint8* FTextureSource::LockMip(int32 BlockIndex, int32 LayerIndex, int32 MipInde
 					check( ImageWrapper->GetHeight() == SizeY );
 					const TArray<uint8>* RawData = NULL;
 					// TODO: TSF_BGRA8 is stored as RGBA, so the R and B channels are swapped in the internal png. Should we fix this?
-					ERGBFormat RawFormat = (Format == TSF_G8) ? ERGBFormat::Gray : ERGBFormat::RGBA;
+					ERGBFormat RawFormat = (Format == TSF_G8 || Format == TSF_G16) ? ERGBFormat::Gray : ERGBFormat::RGBA;
 					if (ImageWrapper->GetRaw( RawFormat, Format == TSF_RGBA16 ? 16 : 8, RawData ))
 					{
 						if (RawData->Num() > 0)
@@ -925,7 +925,7 @@ bool FTextureSource::GetMipData(TArray64<uint8>& OutMipData, int32 BlockIndex, i
 		void* RawSourceData = BulkData.Lock(LOCK_READ_ONLY);
 		if (bPNGCompressed)
 		{
-			bool bCanPngCompressFormat = (Format == TSF_G8 || Format == TSF_RGBA8 || Format == TSF_BGRA8 || Format == TSF_RGBA16);
+			bool bCanPngCompressFormat = (Format == TSF_G8 || Format == TSF_G16 || Format == TSF_RGBA8 || Format == TSF_BGRA8 || Format == TSF_RGBA16);
 			if (MipIndex == 0 && NumLayers == 1 && NumSlices == 1 && Blocks.Num() == 0 && bCanPngCompressFormat)
 			{
 				if (!ImageWrapperModule) // Optional if called from the gamethread, see FModuleManager::WarnIfItWasntSafeToLoadHere()
@@ -942,7 +942,7 @@ bool FTextureSource::GetMipData(TArray64<uint8>& OutMipData, int32 BlockIndex, i
 					{
 						const TArray<uint8>* RawData = NULL;
 						// TODO: TSF_BGRA8 is stored as RGBA, so the R and B channels are swapped in the internal png. Should we fix this?
-						ERGBFormat RawFormat = (Format == TSF_G8) ? ERGBFormat::Gray : ERGBFormat::RGBA;
+						ERGBFormat RawFormat = (Format == TSF_G8 || Format == TSF_G16) ? ERGBFormat::Gray : ERGBFormat::RGBA;
 						if (ImageWrapper->GetRaw( RawFormat, Format == TSF_RGBA16 ? 16 : 8, RawData ))
 						{
 							OutMipData = *RawData;
@@ -1009,6 +1009,7 @@ int32 FTextureSource::GetBytesPerPixel(ETextureSourceFormat Format)
 	switch (Format)
 	{
 	case TSF_G8:		BytesPerPixel = 1; break;
+	case TSF_G16:		BytesPerPixel = 2; break;
 	case TSF_BGRA8:		BytesPerPixel = 4; break;
 	case TSF_BGRE8:		BytesPerPixel = 4; break;
 	case TSF_RGBA16:	BytesPerPixel = 8; break;
@@ -1100,7 +1101,7 @@ FString FTextureSource::GetIdString() const
 
 bool FTextureSource::CanPNGCompress() const
 {
-	bool bCanPngCompressFormat = (Format == TSF_G8 || Format == TSF_RGBA8 || Format == TSF_BGRA8 || Format == TSF_RGBA16);
+	bool bCanPngCompressFormat = (Format == TSF_G8 || Format == TSF_G16 || Format == TSF_RGBA8 || Format == TSF_BGRA8 || Format == TSF_RGBA16);
 
 	if (!bPNGCompressed &&
 		NumLayers == 1 &&
@@ -1312,6 +1313,7 @@ FName GetDefaultTextureFormatName( const ITargetPlatform* TargetPlatform, const
 	static FName NameBGRA8(TEXT("BGRA8"));
 	static FName NameXGXR8(TEXT("XGXR8"));
 	static FName NameG8(TEXT("G8"));
+	static FName NameG16(TEXT("G16"));
 	static FName NameVU8(TEXT("VU8"));
 	static FName NameRGBA16F(TEXT("RGBA16F"));
 	static FName NameBC6H(TEXT("BC6H"));
@@ -1365,6 +1367,10 @@ FName GetDefaultTextureFormatName( const ITargetPlatform* TargetPlatform, const
 		{
 			TextureFormatName = NameRGBA16F;
 		}
+		else if (SourceFormat == TSF_G16)
+		{
+			TextureFormatName = NameG16;
+		}
 		else if (SourceFormat == TSF_G8 || FormatSettings.CompressionSettings == TC_Grayscale)
 		{
 			TextureFormatName = NameG8;
diff --git a/Engine/Source/Runtime/Engine/Private/TextureDerivedDataTask.cpp b/Engine/Source/Runtime/Engine/Private/TextureDerivedDataTask.cpp
--- a/Engine/Source/Runtime/Engine/Private/TextureDerivedDataTask.cpp
+++ b/Engine/Source/Runtime/Engine/Private/TextureDerivedDataTask.cpp
@@ -59,6 +59,7 @@ void FTextureSourceData::Init(UTexture& InTexture, const FTextureBuildSettings*
 		switch (InTexture.Source.GetFormat(LayerIndex))
 		{
 		case TSF_G8:		LayerData->ImageFormat = ERawImageFormat::G8;		break;
+		case TSF_G16:		LayerData->ImageFormat = ERawImageFormat::G16;		break;
 		case TSF_BGRA8:		LayerData->ImageFormat = ERawImageFormat::BGRA8;	break;
 		case TSF_BGRE8:		LayerData->ImageFormat = ERawImageFormat::BGRE8;	break;
 		case TSF_RGBA16:	LayerData->ImageFormat = ERawImageFormat::RGBA16;	break;
diff --git a/Engine/Source/Runtime/Engine/Private/TextureRenderTarget.cpp b/Engine/Source/Runtime/Engine/Private/TextureRenderTarget.cpp
--- a/Engine/Source/Runtime/Engine/Private/TextureRenderTarget.cpp
+++ b/Engine/Source/Runtime/Engine/Private/TextureRenderTarget.cpp
@@ -75,6 +75,7 @@ bool FTextureRenderTargetResource::IsSupportedFormat( EPixelFormat Format )
 	case PF_FloatRGBA: // for exporting materials to .obj/.mtl
 	case PF_A2B10G10R10: //Pixel inspector for normal buffer
 	case PF_DepthStencil: //Pixel inspector for depth and stencil buffer
+	case PF_G16:// for heightmaps
 		return true;
 	default:
 		return false;
diff --git a/Engine/Source/Runtime/Engine/Private/VT/VirtualTextureDataBuilder.cpp b/Engine/Source/Runtime/Engine/Private/VT/VirtualTextureDataBuilder.cpp
--- a/Engine/Source/Runtime/Engine/Private/VT/VirtualTextureDataBuilder.cpp
+++ b/Engine/Source/Runtime/Engine/Private/VT/VirtualTextureDataBuilder.cpp
@@ -212,6 +212,9 @@ struct FPixelDataRectangle
 		case TSF_G8:
 			ImageWrapper->SetRaw(Data, BytesPerPixel * Width * Height, Width, Height, ERGBFormat::Gray, 8);
 			break;
+		case TSF_G16:
+			ImageWrapper->SetRaw(Data, BytesPerPixel * Width * Height, Width, Height, ERGBFormat::Gray, 16);
+			break;
 		case TSF_BGRA8:
 			ImageWrapper->SetRaw(Data, BytesPerPixel * Width * Height, Width, Height, ERGBFormat::BGRA, 8);
 			break;
@@ -794,16 +797,13 @@ void FVirtualTextureDataBuilder::BuildSourcePixels(const FTextureSourceData& Sou
 		const FTextureBuildSettings& BuildSettingsForLayer = SettingsPerLayer[LayerIndex];
 		FVirtualTextureSourceLayerData& LayerData = SourceLayers[LayerIndex];
 
-		const FName TextureFormatName = BuildSettingsForLayer.TextureFormatName;
-		const bool bIsHdr = BuildSettingsForLayer.bHDRSource || TextureFormatName == "BC6H" || TextureFormatName == "RGBA16F";
-
-		LayerData.FormatName = "BGRA8";
-		LayerData.PixelFormat = PF_B8G8R8A8;
-		LayerData.SourceFormat = TSF_BGRA8;
-		LayerData.ImageFormat = ERawImageFormat::BGRA8;
 		LayerData.GammaSpace = BuildSettingsForLayer.GetGammaSpace();
 		LayerData.bHasAlpha = false;
 
+		const FName TextureFormatName = BuildSettingsForLayer.TextureFormatName;
+		const bool bIsHdr = BuildSettingsForLayer.bHDRSource || TextureFormatName == "BC6H" || TextureFormatName == "RGBA16F";
+		const bool bIsG16 = TextureFormatName == "G16";
+
 		if (bIsHdr)
 		{
 			LayerData.FormatName = "RGBA16F";
@@ -811,6 +811,20 @@ void FVirtualTextureDataBuilder::BuildSourcePixels(const FTextureSourceData& Sou
 			LayerData.SourceFormat = TSF_RGBA16F;
 			LayerData.ImageFormat = ERawImageFormat::RGBA16F;
 		}
+		else if (bIsG16)
+		{
+			LayerData.FormatName = "G16";
+			LayerData.PixelFormat = PF_G16;
+			LayerData.SourceFormat = TSF_G16;
+			LayerData.ImageFormat = ERawImageFormat::G16;
+		}
+		else
+		{
+			LayerData.FormatName = "BGRA8";
+			LayerData.PixelFormat = PF_B8G8R8A8;
+			LayerData.SourceFormat = TSF_BGRA8;
+			LayerData.ImageFormat = ERawImageFormat::BGRA8;
+		}
 	}
 
 	SourceBlocks.AddDefaulted(NumBlocks);
@@ -1206,6 +1220,10 @@ bool FVirtualTextureDataBuilder::DetectAlphaChannel(const FImage &Image)
 		}
 		return false;
 	}
+	else if (Image.Format == ERawImageFormat::G16)
+	{
+		return false;
+	}
 	else
 	{
 		check(false);
diff --git a/Engine/Source/Runtime/ImageCore/Private/ImageCore.cpp b/Engine/Source/Runtime/ImageCore/Private/ImageCore.cpp
--- a/Engine/Source/Runtime/ImageCore/Private/ImageCore.cpp
+++ b/Engine/Source/Runtime/ImageCore/Private/ImageCore.cpp
@@ -71,7 +71,22 @@ static void CopyImage(const FImage& SrcImage, FImage& DestImage)
 				});
 			}
 			break;
-		
+
+		case ERawImageFormat::G16:
+		{
+			uint16* DestLum = DestImage.AsG16();
+			ParallelFor(NumJobs, [NumJobs, DestLum, SrcColors, TexelsPerJob, NumTexels, bDestIsGammaCorrected](int32 JobIndex)
+			{
+				const int32 StartIndex = JobIndex * TexelsPerJob;
+				const int32 EndIndex = FMath::Min(StartIndex + TexelsPerJob, NumTexels);
+				for (int32 TexelIndex = StartIndex; TexelIndex < EndIndex; ++TexelIndex)
+				{
+					DestLum[TexelIndex] = FMath::Clamp(FMath::FloorToInt(SrcColors[TexelIndex].R * 65535.999f), 0, 65535);
+				}
+			});
+		}
+		break;
+
 		case ERawImageFormat::BGRA8:
 			{
 				FColor* DestColors = DestImage.AsBGRA8();
@@ -156,6 +171,16 @@ static void CopyImage(const FImage& SrcImage, FImage& DestImage)
 			}
 			break;
 
+		case ERawImageFormat::G16:
+		{
+			const uint16* SrcLum = SrcImage.AsG16();
+			for (int32 TexelIndex = 0; TexelIndex < NumTexels; ++TexelIndex)
+			{
+				DestColors[TexelIndex] = FLinearColor(SrcLum[TexelIndex] / 65535.0f, SrcLum[TexelIndex] / 65535.0f, SrcLum[TexelIndex] / 65535.0f, 1.0f);
+			}
+		}
+		break;
+
 		case ERawImageFormat::BGRA8:
 			{
 				const FColor* SrcColors = SrcImage.AsBGRA8();
@@ -387,6 +412,10 @@ int32 FImage::GetBytesPerPixel() const
 	case ERawImageFormat::G8:
 		BytesPerPixel = 1;
 		break;
+	
+	case ERawImageFormat::G16:
+		BytesPerPixel = 2;
+		break;
 
 	case ERawImageFormat::BGRA8:
 	case ERawImageFormat::BGRE8:
diff --git a/Engine/Source/Runtime/ImageCore/Public/ImageCore.h b/Engine/Source/Runtime/ImageCore/Public/ImageCore.h
--- a/Engine/Source/Runtime/ImageCore/Public/ImageCore.h
+++ b/Engine/Source/Runtime/ImageCore/Public/ImageCore.h
@@ -23,7 +23,8 @@ namespace ERawImageFormat
 		BGRE8,
 		RGBA16,
 		RGBA16F,
-		RGBA32F
+		RGBA32F,
+		G16
 	};
 };
 
@@ -142,6 +143,12 @@ public:
 		return RawData.GetData();
 	}
 
+	uint16* AsG16()
+	{
+		check(Format == ERawImageFormat::G16);
+		return (uint16*)RawData.GetData();
+	}
+
 	struct FColor* AsBGRA8()
 	{
 		check(Format == ERawImageFormat::BGRA8);
@@ -180,6 +187,12 @@ public:
 		return RawData.GetData();
 	}
 
+	const uint16* AsG16() const
+	{
+		check(Format == ERawImageFormat::G16);
+		return (const uint16*)RawData.GetData();
+	}
+
 	const struct FColor* AsBGRA8() const
 	{
 		check(Format == ERawImageFormat::BGRA8);
