Added ALODActor::GetHLODDistanceOverride()

diff --git a/Engine/Source/Runtime/Engine/Classes/Engine/LODActor.h b/Engine/Source/Runtime/Engine/Classes/Engine/LODActor.h
--- a/Engine/Source/Runtime/Engine/Classes/Engine/LODActor.h
+++ b/Engine/Source/Runtime/Engine/Classes/Engine/LODActor.h
@@ -238,6 +238,9 @@ public:
 	/** Returns instanced static mesh components that represents the imposters found in this LODActor. */
 	TArray<UInstancedStaticMeshComponent*> GetImpostersStaticMeshComponents() const { TArray<UInstancedStaticMeshComponent*> ISM; ImpostersStaticMeshComponents.GenerateValueArray(ISM); return ISM; }
 
+	/** Returns an array of distances that are used to override individual LOD actors min draw distances. */
+	static const TArray<float>& GetHLODDistanceOverride();
+
 private:
 	// Called when CVars are changed to check to see if the maximum HLOD level value has changed
 	static void OnCVarsChanged();
diff --git a/Engine/Source/Runtime/Engine/Private/LODActor.cpp b/Engine/Source/Runtime/Engine/Private/LODActor.cpp
--- a/Engine/Source/Runtime/Engine/Private/LODActor.cpp
+++ b/Engine/Source/Runtime/Engine/Private/LODActor.cpp
@@ -270,29 +270,23 @@ void ALODActor::PostLoad()
 
 #if !WITH_EDITOR
 	// Invalid runtime LOD actor with null static mesh is invalid, look for a possible way to patch this up
-	if (GetStaticMeshComponent()->GetStaticMesh() == nullptr)
+	if (GetStaticMeshComponent() && GetStaticMeshComponent()->GetStaticMesh() == nullptr && GetStaticMeshComponent()->GetLODParentPrimitive())
 	{
-		if (GetStaticMeshComponent() && GetStaticMeshComponent()->GetLODParentPrimitive())
+ 		if (ALODActor* ParentLODActor = Cast<ALODActor>(GetStaticMeshComponent()->GetLODParentPrimitive()->GetOwner()))
 		{
- 			if (ALODActor* ParentLODActor = Cast<ALODActor>(GetStaticMeshComponent()->GetLODParentPrimitive()->GetOwner()))
+			// Make the parent HLOD
+			ParentLODActor->SubActors.Remove(this);
+			ParentLODActor->SubActors.Append(SubActors);
+			for (AActor* Actor : SubActors)
 			{
-				if ( ParentLODActor->GetStaticMeshComponent()->GetStaticMesh() != nullptr )
-				{				
-					// Make the parent HLOD 			
-					ParentLODActor->SubActors.Remove(this);
-					ParentLODActor->SubActors.Append(SubActors);			
-					for (AActor* Actor : SubActors)
-					{
-						if (Actor)
-						{
-							Actor->SetLODParent(ParentLODActor->GetStaticMeshComponent(), ParentLODActor->GetDrawDistance());
-						}
-					}
-  
-					SubActors.Empty();
-					bHasPatchedUpParent = true;
+				if (Actor)
+				{
+					Actor->SetLODParent(ParentLODActor->GetStaticMeshComponent(), ParentLODActor->GetDrawDistance());
 				}
 			}
+  
+			SubActors.Empty();
+			bHasPatchedUpParent = true;
 		}
 	}
 #endif // !WITH_EDITOR
@@ -344,6 +338,13 @@ void ALODActor::SetComponentsMinDrawDistance(float InMinDrawDistance, bool bInMa
 	}
 }
 
+/** Returns an array of distances that are used to override individual LOD actors min draw distances. */
+const TArray<float>& ALODActor::GetHLODDistanceOverride()
+{
+	ParseOverrideDistancesCVar();
+	return HLODDistances;
+}
+
 void ALODActor::UpdateOverrideTransitionDistance()
 {
 	const int32 NumDistances = ALODActor::HLODDistances.Num();
