Make HierarchicalLODIndex protected

diff --git a/Engine/Source/Runtime/Engine/Public/StaticMeshResources.h b/Engine/Source/Runtime/Engine/Public/StaticMeshResources.h
--- a/Engine/Source/Runtime/Engine/Public/StaticMeshResources.h
+++ b/Engine/Source/Runtime/Engine/Public/StaticMeshResources.h
@@ -964,8 +964,6 @@ private:
 
 #if STATICMESH_ENABLE_DEBUG_RENDERING
 	AActor* Owner;
-	/** Hierarchical LOD Index used for rendering */
-	uint8 HierarchicalLODIndex;
 	/** LightMap resolution used for VMI_LightmapDensity */
 	int32 LightMapResolution;
 	/** Body setup for collision debug rendering */
@@ -980,6 +978,10 @@ private:
 	uint32 bDrawMeshCollisionIfComplex : 1;
 	/** Draw mesh collision if used for simple collision */
 	uint32 bDrawMeshCollisionIfSimple : 1;
+
+protected:
+	/** Hierarchical LOD Index used for rendering */
+	uint8 HierarchicalLODIndex;
 #endif
 
 public:
