Renaming layer and sub-instance related APIs

diff --git a/Engine/Config/BaseEngine.ini b/Engine/Config/BaseEngine.ini
--- a/Engine/Config/BaseEngine.ini
+++ b/Engine/Config/BaseEngine.ini
@@ -693,6 +693,34 @@ bUseStreamingPause=false
 
 +FunctionRedirects=(OldName="Actor.GetComponentsByClass", NewName="Actor.K2_GetComponentsByClass")
 
++FunctionRedirects=(OldName="SkeletalMeshComponent.GetSubInstanceByTag",NewName="SkeletalMeshComponent.GetLinkedAnimGraphInstanceByTag")
++FunctionRedirects=(OldName="SkeletalMeshComponent.GetSubInstancesByTag",NewName="SkeletalMeshComponent.GetLinkedAnimGraphInstancesByTag")
++PropertyRedirects=(OldName="SkeletalMeshComponent.GetLinkedAnimGraphInstancesByTag.OutSubInstances",NewName="OutLinkedInstances")
++FunctionRedirects=(OldName="SkeletalMeshComponent.SetSubInstanceClassByTag",NewName="SkeletalMeshComponent.LinkAnimGraphByTag")
++FunctionRedirects=(OldName="SkeletalMeshComponent.SetLayerOverlay",NewName="SkeletalMeshComponent.LinkAnimClassLayers")
++FunctionRedirects=(OldName="SkeletalMeshComponent.ClearLayerOverlay",NewName="SkeletalMeshComponent.UnlinkAnimClassLayers")
++FunctionRedirects=(OldName="SkeletalMeshComponent.GetLayerSubInstanceByGroup",NewName="SkeletalMeshComponent.GetLinkedAnimLayerInstanceByGroup")
++FunctionRedirects=(OldName="SkeletalMeshComponent.GetLayerSubInstanceByClass",NewName="SkeletalMeshComponent.GetLinkedAnimLayerInstanceByClass")
++FunctionRedirects=(OldName="AnimInstance.GetSubInstanceByTag",NewName="AnimInstance.GetLinkedAnimGraphInstanceByTag")
++FunctionRedirects=(OldName="AnimInstance.GetSubInstancesByTag",NewName="AnimInstance.GetLinkedAnimGraphInstancesByTag")
++PropertyRedirects=(OldName="AnimInstance.GetLinkedAnimGraphInstancesByTag.OutSubInstances",NewName="OutLinkedInstances")
++FunctionRedirects=(OldName="AnimInstance.SetSubInstanceClassByTag",NewName="AnimInstance.LinkAnimGraphByTag")
++FunctionRedirects=(OldName="AnimInstance.SetLayerOverlay",NewName="AnimInstance.LinkAnimClassLayers")
++FunctionRedirects=(OldName="AnimInstance.ClearLayerOverlay",NewName="AnimInstance.UnlinkAnimClassLayers")
++FunctionRedirects=(OldName="AnimInstance.GetLayerSubInstanceByGroup",NewName="AnimInstance.GetLinkedAnimLayerInstanceByGroup")
++FunctionRedirects=(OldName="AnimInstance.GetLayerSubInstanceByClass",NewName="AnimInstance.GetLinkedAnimLayerInstanceByClass")
++StructRedirects=(OldName="AnimNode_SubInstance",NewName="/Script/Engine.AnimNode_LinkedAnimGraph")
++StructRedirects=(OldName="AnimNode_SubInput",NewName="/Script/Engine.AnimNode_LinkedInputPose")
++StructRedirects=(OldName="AnimNode_Layer",NewName="/Script/Engine.AnimNode_LinkedAnimLayer")
++ClassRedirects=(OldName="/Script/AnimGraph.AnimGraphNode_SubInstanceBase",NewName="/Script/AnimGraph.AnimGraphNode_LinkedAnimGraphBase")
++ClassRedirects=(OldName="/Script/AnimGraph.AnimGraphNode_SubInstance",NewName="/Script/AnimGraph.AnimGraphNode_LinkedAnimGraph")
++ClassRedirects=(OldName="/Script/AnimGraph.AnimGraphNode_SubInput",NewName="/Script/AnimGraph.AnimGraphNode_LinkedInputPose")
++ClassRedirects=(OldName="/Script/AnimGraph.AnimGraphNode_Layer",NewName="/Script/AnimGraph.AnimGraphNode_LinkedAnimLayer")
++PropertyRedirects=(OldName="PersonaPreviewSceneDescription.SubInstanceTag",NewName="LinkedAnimGraphTag")
++EnumRedirects=(OldName="/Script/Engine.EPreviewAnimationBlueprintApplicationMethod",NewName="/Script/Engine.EPreviewAnimationBlueprintApplicationMethod",ValueChanges=(("OverlayLayer", "LinkedLayers"), ("SubInstance", "LinkedAnimGraph"))
++PropertyRedirects=(OldName="AnimClassData.SubInstanceNodeProperties",NewName="LinkedAnimGraphNodeProperties")
++PropertyRedirects=(OldName="AnimClassData.LayerNodeProperties",NewName="LinkedAnimLayerNodeProperties")
+
 [CoreUObject.Metadata]
 ; Note: UnrealHeaderTool should be rerun after making changes to MetadataRedirects to catch any keys specified in class headers
 +MetadataRedirects=(OldKey="K2Protected", NewKey="BlueprintProtected")
diff --git a/Engine/Source/Editor/AnimGraph/Classes/AnimGraphNode_Layer.h b/Engine/Source/Editor/AnimGraph/Classes/AnimGraphNode_Layer.h
deleted file mode 100644
diff --git a/Engine/Source/Editor/AnimGraph/Classes/AnimGraphNode_LinkedAnimGraph.h b/Engine/Source/Editor/AnimGraph/Classes/AnimGraphNode_LinkedAnimGraph.h
new file mode 100644
--- /dev/null
+++ b/Engine/Source/Editor/AnimGraph/Classes/AnimGraphNode_LinkedAnimGraph.h
@@ -0,0 +1,33 @@
+// Copyright 1998-2019 Epic Games, Inc. All Rights Reserved.
+
+#pragma once
+
+#include "CoreMinimal.h"
+#include "UObject/ObjectMacros.h"
+#include "Misc/Guid.h"
+#include "AnimGraphNode_LinkedAnimGraphBase.h"
+#include "Animation/AnimNode_LinkedAnimGraph.h"
+
+#include "AnimGraphNode_LinkedAnimGraph.generated.h"
+
+UCLASS(MinimalAPI)
+class UAnimGraphNode_LinkedAnimGraph : public UAnimGraphNode_LinkedAnimGraphBase
+{
+	GENERATED_BODY()
+
+public:
+
+	UPROPERTY(EditAnywhere, Category = Settings)
+	FAnimNode_LinkedAnimGraph Node;
+
+	// Begin UAnimGraphNode_CustomProperty
+	virtual FAnimNode_CustomProperty* GetCustomPropertyNode() override { return &Node;  }
+	virtual const FAnimNode_CustomProperty* GetCustomPropertyNode() const override { return &Node; }
+
+	// Begin UAnimGraphNode_LinkedAnimGraphBase
+	virtual FAnimNode_LinkedAnimGraph* GetLinkedAnimGraphNode() override { return &Node; }
+	virtual const FAnimNode_LinkedAnimGraph* GetLinkedAnimGraphNode() const override { return &Node; }
+};
+
+UE_DEPRECATED(4.24, "UAnimGraphNode_SubInstance has been renamed to UAnimGraphNode_LinkedAnimGraph")
+typedef UAnimGraphNode_LinkedAnimGraph UAnimGraphNode_SubInstance;
\ No newline at end of file
diff --git a/Engine/Source/Editor/AnimGraph/Classes/AnimGraphNode_LinkedAnimGraphBase.h b/Engine/Source/Editor/AnimGraph/Classes/AnimGraphNode_LinkedAnimGraphBase.h
new file mode 100644
--- /dev/null
+++ b/Engine/Source/Editor/AnimGraph/Classes/AnimGraphNode_LinkedAnimGraphBase.h
@@ -0,0 +1,70 @@
+// Copyright 1998-2019 Epic Games, Inc. All Rights Reserved.
+
+#pragma once
+
+#include "CoreMinimal.h"
+#include "UObject/ObjectMacros.h"
+#include "Misc/Guid.h"
+#include "AnimGraphNode_CustomProperty.h"
+
+#include "AnimGraphNode_LinkedAnimGraphBase.generated.h"
+
+class FCompilerResultsLog;
+class IDetailLayoutBuilder;
+class IPropertyHandle;
+class SToolTip;
+struct FAnimNode_LinkedAnimGraph;
+
+UCLASS(MinimalAPI, Abstract)
+class UAnimGraphNode_LinkedAnimGraphBase : public UAnimGraphNode_CustomProperty
+{
+	GENERATED_BODY()
+
+public:
+
+	//~ Begin UEdGraphNode Interface.
+	virtual FLinearColor GetNodeTitleColor() const override;
+	virtual FText GetTooltipText() const override;
+	virtual FText GetNodeTitle(ENodeTitleType::Type TitleType) const override;
+	virtual void ValidateAnimNodeDuringCompilation(USkeleton* ForSkeleton, FCompilerResultsLog& MessageLog) override;
+	virtual void ReallocatePinsDuringReconstruction(TArray<UEdGraphPin*>& OldPins) override;
+	virtual void PostEditChangeProperty(FPropertyChangedEvent& PropertyChangedEvent) override;
+	virtual void CustomizeDetails(IDetailLayoutBuilder& DetailBuilder) override;
+	//~ End UEdGraphNode Interface.
+
+	// UAnimGraphNode_Base interface
+	virtual FPoseLinkMappingRecord GetLinkIDLocation(const UScriptStruct* NodeType, UEdGraphPin* SourcePin) override;
+
+	// UAnimGraphNode_CustomProperty interface
+	virtual bool IsStructuralProperty(UProperty* InProperty) const override;
+
+	// Node accessor
+	virtual FAnimNode_LinkedAnimGraph* GetLinkedAnimGraphNode() PURE_VIRTUAL(UAnimGraphNode_LinkedAnimGraphBase::GetLinkedAnimGraphNode, return nullptr;);
+	virtual const FAnimNode_LinkedAnimGraph* GetLinkedAnimGraphNode() const PURE_VIRTUAL(UAnimGraphNode_LinkedAnimGraphBase::GetLinkedAnimGraphNode, return nullptr;);
+
+protected:
+	// Finds out whether there is a loop in the graph formed by linked instances from this node
+	bool HasInstanceLoop();
+	
+	/** Generates widgets for exposing/hiding Pins for this node using the provided detail builder */
+	void GenerateExposedPinsDetails(IDetailLayoutBuilder &DetailBuilder);
+
+	// Finds out whether there is a loop in the graph formed by linked instances from CurrNode, used by HasInstanceLoop. VisitedNodes and NodeStack are required
+	// to track the graph links
+	// VisitedNodes - Node we have searched the links of, so we don't do it twice
+	// NodeStack - The currently considered chain of nodes. If a loop is detected this will contain the chain that causes the loop
+	static bool HasInstanceLoop_Recursive(UAnimGraphNode_LinkedAnimGraphBase* CurrNode, TArray<FGuid>& VisitedNodes, TArray<FGuid>& NodeStack);
+
+	// ----- UI CALLBACKS ----- //
+
+	// Gets path to the currently selected instance class' blueprint
+	virtual FString GetCurrentInstanceBlueprintPath() const;
+	// Filter callback for blueprints (only accept matching skeletons/interfaces)
+	virtual bool OnShouldFilterInstanceBlueprint(const FAssetData& AssetData) const;
+	// Instance blueprint was changed by user
+	void OnSetInstanceBlueprint(const FAssetData& AssetData, IDetailLayoutBuilder* InDetailBuilder);
+	// ----- END UI CALLBACKS ----- //
+};
+
+UE_DEPRECATED(4.24, "UAnimGraphNode_SubInstanceBase has been renamed to UAnimGraphNode_LinkedAnimGraphBase")
+typedef UAnimGraphNode_LinkedAnimGraphBase UAnimGraphNode_SubInstanceBase;
diff --git a/Engine/Source/Editor/AnimGraph/Classes/AnimGraphNode_LinkedAnimLayer.h b/Engine/Source/Editor/AnimGraph/Classes/AnimGraphNode_LinkedAnimLayer.h
new file mode 100644
--- /dev/null
+++ b/Engine/Source/Editor/AnimGraph/Classes/AnimGraphNode_LinkedAnimLayer.h
@@ -0,0 +1,69 @@
+// Copyright 1998-2019 Epic Games, Inc. All Rights Reserved.
+
+#pragma once
+
+#include "CoreMinimal.h"
+#include "UObject/ObjectMacros.h"
+#include "Misc/Guid.h"
+#include "AnimGraphNode_LinkedAnimGraphBase.h"
+#include "Animation/AnimNode_LinkedAnimLayer.h"
+
+#include "AnimGraphNode_LinkedAnimLayer.generated.h"
+
+class FCompilerResultsLog;
+class IDetailLayoutBuilder;
+class IPropertyHandle;
+class SToolTip;
+
+UCLASS(MinimalAPI, Experimental)
+class UAnimGraphNode_LinkedAnimLayer : public UAnimGraphNode_LinkedAnimGraphBase
+{
+	GENERATED_BODY()
+
+public:
+
+	UPROPERTY(EditAnywhere, Category = Settings)
+	FAnimNode_LinkedAnimLayer Node;
+
+	//~ Begin UEdGraphNode Interface.
+	virtual FText GetTooltipText() const override;
+	virtual FText GetNodeTitle(ENodeTitleType::Type TitleType) const override;
+	virtual void ValidateAnimNodeDuringCompilation(USkeleton* ForSkeleton, FCompilerResultsLog& MessageLog) override;
+	virtual UObject* GetJumpTargetForDoubleClick() const override;
+	virtual void JumpToDefinition() const override;
+	virtual bool HasExternalDependencies(TArray<class UStruct*>* OptionalOutput /*= NULL*/) const override;
+	//~ End UEdGraphNode Interface.
+
+	// UAnimGraphNode_Base interface
+	virtual void CustomizeDetails(IDetailLayoutBuilder& DetailBuilder) override;
+
+protected:
+	// ----- UI CALLBACKS ----- //
+	// Handlers for layer combo
+	void GetLayerNames(TArray<TSharedPtr<FString>>& OutStrings, TArray<TSharedPtr<SToolTip>>& OutToolTips, TArray<bool>& OutRestrictedItems);
+	FString GetLayerName() const;
+	void OnLayerChanged(IDetailLayoutBuilder* DetailBuilder);
+	bool HasAvailableLayers() const;
+	bool HasValidNonSelfLayer() const;
+	// ----- END UI CALLBACKS ----- //
+
+	// Begin UAnimGraphNode_CustomProperty
+	virtual FAnimNode_CustomProperty* GetCustomPropertyNode() override { return &Node;  }
+	virtual const FAnimNode_CustomProperty* GetCustomPropertyNode() const override { return &Node; }
+	void GetExposableProperties( TArray<UProperty*>& OutExposableProperties) const override;
+	virtual bool NeedsToSpecifyValidTargetClass() const override { return false; }
+	virtual UClass* GetTargetSkeletonClass() const override;
+
+	// Begin UAnimGraphNode_LinkedAnimGraphBase
+	virtual FAnimNode_LinkedAnimLayer* GetLinkedAnimGraphNode() override { return &Node;  }
+	virtual const FAnimNode_LinkedAnimLayer* GetLinkedAnimGraphNode() const override { return &Node; }
+	virtual bool OnShouldFilterInstanceBlueprint(const FAssetData& AssetData) const override;
+	virtual FString GetCurrentInstanceBlueprintPath() const override;
+	virtual bool IsStructuralProperty(UProperty* InProperty) const override;
+
+	// Helper function to get the interface currently in use by the selected layer
+	TSubclassOf<UInterface> GetInterfaceForLayer() const;
+};
+
+UE_DEPRECATED(4.24, "UAnimGraphNode_Layer has been renamed to UAnimGraphNode_LinkedAnimLayer")
+typedef UAnimGraphNode_LinkedAnimLayer UAnimGraphNode_Layer;
\ No newline at end of file
diff --git a/Engine/Source/Editor/AnimGraph/Classes/AnimGraphNode_LinkedInputPose.h b/Engine/Source/Editor/AnimGraph/Classes/AnimGraphNode_LinkedInputPose.h
new file mode 100644
--- /dev/null
+++ b/Engine/Source/Editor/AnimGraph/Classes/AnimGraphNode_LinkedInputPose.h
@@ -0,0 +1,120 @@
+// Copyright 1998-2019 Epic Games, Inc. All Rights Reserved.
+
+#pragma once
+
+#include "CoreMinimal.h"
+#include "UObject/ObjectMacros.h"
+#include "AnimGraphNode_Base.h"
+#include "Animation/AnimNode_LinkedInputPose.h"
+#include "Engine/MemberReference.h"
+#include "AnimGraphNode_LinkedInputPose.generated.h"
+
+class SEditableTextBox;
+
+/** Required info for reconstructing a manually specified pin */
+USTRUCT()
+struct FAnimBlueprintFunctionPinInfo
+{
+	GENERATED_BODY()
+
+	FAnimBlueprintFunctionPinInfo()
+		: Name(NAME_None)
+	{
+		Type.ResetToDefaults();
+	}
+
+	FAnimBlueprintFunctionPinInfo(const FName& InName, const FEdGraphPinType& InType)
+		: Name(InName)
+		, Type(InType)
+	{
+	}
+
+	/** The name of this parameter */
+	UPROPERTY(EditAnywhere, Category = "Inputs")
+	FName Name;
+
+	/** The type of this parameter */
+	UPROPERTY(EditAnywhere, Category = "Inputs")
+	FEdGraphPinType Type;
+};
+
+UCLASS()
+class ANIMGRAPH_API UAnimGraphNode_LinkedInputPose : public UAnimGraphNode_Base
+{
+	GENERATED_BODY()
+
+public:
+	UPROPERTY(EditAnywhere, Category = "Inputs")
+	FAnimNode_LinkedInputPose Node;
+
+	UPROPERTY(EditAnywhere, Category = "Inputs")
+	TArray<FAnimBlueprintFunctionPinInfo> Inputs;
+
+	/** Reference to the stub function we use to build our parameters */
+	UPROPERTY()
+	FMemberReference FunctionReference;
+
+	/** The index of the input pose, used alongside FunctionReference to build parameters */
+	UPROPERTY()
+	int32 InputPoseIndex;
+
+	UAnimGraphNode_LinkedInputPose();
+
+	/** UObject interface */
+	virtual void PostEditChangeProperty(FPropertyChangedEvent& PropertyChangedEvent) override;
+
+	/** UEdGraphNode interface */
+	virtual FLinearColor GetNodeTitleColor() const override;
+	virtual FText GetTooltipText() const override;
+	virtual FText GetNodeTitle(ENodeTitleType::Type TitleType) const override;
+	virtual bool CanUserDeleteNode() const override;
+	virtual bool CanDuplicateNode() const override;
+	virtual void AllocateDefaultPins() override;
+	virtual void ReallocatePinsDuringReconstruction(TArray<UEdGraphPin*>& OldPins) override;
+	virtual void PostPlacedNewNode() override;
+
+	/** UK2Node interface */
+	virtual bool HasExternalDependencies(TArray<UStruct*>* OptionalOutput) const override;
+
+	/** UAnimGraphNode_Base interface */
+	virtual void CustomizeDetails(IDetailLayoutBuilder& DetailBuilder) override;
+
+	/** Make a name widget for this linked input pose node */
+	TSharedRef<SWidget> MakeNameWidget(IDetailLayoutBuilder& DetailBuilder);
+
+	/** @return whether this node is editable. Non-editable nodes are implemented from function interfaces */
+	bool IsEditable() const { return CanUserDeleteNode(); }
+
+	/** @return the number of parameter inputs this node provides. This is either determined via the Inputs array or via the FunctionReference. */
+	int32 GetNumInputs() const;
+
+	/** Promotes the node from being a part of an interface override to a full function that allows for parameter and result pin additions */
+	void PromoteFromInterfaceOverride();
+
+	/** Conform input pose name according to function */
+	void ConformInputPoseName();
+
+	/** Validate pose index against the function reference (used to determine whether we should exist or not) */
+	bool ValidateAgainstFunctionReference() const;
+
+	/** Helper function for iterating stub function parameters */
+	void IterateFunctionParameters(TFunctionRef<void(const FName&, const FEdGraphPinType&)> InFunc) const;
+
+private:
+	// Helper function for common code in AllocateDefaultPins and ReallocatePinsDuringReconstruction
+	void AllocatePinsInternal();
+
+	/** Create pins from the user-defined Inputs array */
+	void CreateUserDefinedPins();
+
+	/** Create pins from the stub function FunctionReference */
+	void CreatePinsFromStubFunction(const UFunction* Function);
+
+private:
+	/** UI helper functions */
+	void HandleInputPoseArrayChanged();
+	void HandleInputPinArrayChanged();
+};
+
+UE_DEPRECATED(4.24, "UAnimGraphNode_SubInput has been renamed to UAnimGraphNode_LinkedInputPose")
+typedef UAnimGraphNode_LinkedInputPose UAnimGraphNode_SubInput;
\ No newline at end of file
diff --git a/Engine/Source/Editor/AnimGraph/Classes/AnimGraphNode_SubInput.h b/Engine/Source/Editor/AnimGraph/Classes/AnimGraphNode_SubInput.h
deleted file mode 100644
diff --git a/Engine/Source/Editor/AnimGraph/Classes/AnimGraphNode_SubInstance.h b/Engine/Source/Editor/AnimGraph/Classes/AnimGraphNode_SubInstance.h
deleted file mode 100644
diff --git a/Engine/Source/Editor/AnimGraph/Classes/AnimGraphNode_SubInstanceBase.h b/Engine/Source/Editor/AnimGraph/Classes/AnimGraphNode_SubInstanceBase.h
deleted file mode 100644
diff --git a/Engine/Source/Editor/AnimGraph/Classes/AnimationGraphSchema.h b/Engine/Source/Editor/AnimGraph/Classes/AnimationGraphSchema.h
--- a/Engine/Source/Editor/AnimGraph/Classes/AnimationGraphSchema.h
+++ b/Engine/Source/Editor/AnimGraph/Classes/AnimationGraphSchema.h
@@ -99,8 +99,11 @@ class UAnimationGraphSchema : public UEdGraphSchema_K2
 	/** Conforms an anim graph to an interface function */
 	ANIMGRAPH_API static void ConformAnimGraphToInterface(UBlueprint* InBlueprint, UEdGraph& InGraph, UFunction* InFunction);
 
-	/** Find a position for a newly created sub-input */
-	ANIMGRAPH_API static FVector2D GetPositionForNewSubInputNode(UEdGraph& InGraph);
+	UE_DEPRECATED(4.24, "Function renamed, please use GetPositionForNewLinkedInputPoseNode")
+	ANIMGRAPH_API static FVector2D GetPositionForNewSubInputNode(UEdGraph& InGraph) { return GetPositionForNewLinkedInputPoseNode(InGraph); }
+
+	/** Find a position for a newly created linked input pose */
+	ANIMGRAPH_API static FVector2D GetPositionForNewLinkedInputPoseNode(UEdGraph& InGraph);
 };
 
 
diff --git a/Engine/Source/Editor/AnimGraph/Private/AnimBlueprintPinInfoDetails.cpp b/Engine/Source/Editor/AnimGraph/Private/AnimBlueprintPinInfoDetails.cpp
--- a/Engine/Source/Editor/AnimGraph/Private/AnimBlueprintPinInfoDetails.cpp
+++ b/Engine/Source/Editor/AnimGraph/Private/AnimBlueprintPinInfoDetails.cpp
@@ -4,7 +4,7 @@
 #include "SPinTypeSelector.h"
 #include "EdGraphSchema_K2.h"
 #include "PropertyHandle.h"
-#include "AnimGraphNode_SubInput.h"
+#include "AnimGraphNode_LinkedInputPose.h"
 #include "DetailLayoutBuilder.h"
 #include "Widgets/SBoxPanel.h"
 #include "DetailWidgetRow.h"
@@ -28,7 +28,7 @@ void FAnimBlueprintFunctionPinInfoDetails::CustomizeHeader(TSharedRef<IPropertyH
 	}
 
 	// skip if we cant edit this node as it is an interface graph
-	UAnimGraphNode_SubInput* Node = CastChecked<UAnimGraphNode_SubInput>(OuterObjects[0]);
+	UAnimGraphNode_LinkedInputPose* Node = CastChecked<UAnimGraphNode_LinkedInputPose>(OuterObjects[0]);
 	if(!Node->CanUserDeleteNode())
 	{
 		return;
@@ -126,13 +126,13 @@ bool FAnimBlueprintFunctionPinInfoDetails::IsNameValid(const FString& InNewName)
 		{
 			if(Graph->Schema->IsChildOf(UAnimationGraphSchema::StaticClass()))
 			{
-				TArray<UAnimGraphNode_SubInput*> SubInputNodes;
-				Graph->GetNodesOfClass(SubInputNodes);
+				TArray<UAnimGraphNode_LinkedInputPose*> LinkedInputPoseNodes;
+				Graph->GetNodesOfClass(LinkedInputPoseNodes);
 
-				for(UAnimGraphNode_SubInput* SubInput : SubInputNodes)
+				for(UAnimGraphNode_LinkedInputPose* LinkedInputPose : LinkedInputPoseNodes)
 				{
-					const TArray<FAnimBlueprintFunctionPinInfo>& Inputs = SubInput->Inputs;
-					if(SubInput == WeakOuterNode.Get())
+					const TArray<FAnimBlueprintFunctionPinInfo>& Inputs = LinkedInputPose->Inputs;
+					if(LinkedInputPose == WeakOuterNode.Get())
 					{
 						const int32 CurrentIndex = StructPropertyHandle->GetIndexInArray();
 						for(int32 InputIndex = 0; InputIndex < Inputs.Num(); ++InputIndex)
diff --git a/Engine/Source/Editor/AnimGraph/Private/AnimBlueprintPinInfoDetails.h b/Engine/Source/Editor/AnimGraph/Private/AnimBlueprintPinInfoDetails.h
--- a/Engine/Source/Editor/AnimGraph/Private/AnimBlueprintPinInfoDetails.h
+++ b/Engine/Source/Editor/AnimGraph/Private/AnimBlueprintPinInfoDetails.h
@@ -6,7 +6,7 @@
 #include "Types/SlateEnums.h"
 
 struct FEdGraphPinType;
-class UAnimGraphNode_SubInput;
+class UAnimGraphNode_LinkedInputPose;
 class SEditableTextBox;
 
 class FAnimBlueprintFunctionPinInfoDetails : public IPropertyTypeCustomization
@@ -41,5 +41,5 @@ private:
 	TSharedPtr<SEditableTextBox> NameTextBox;
 
 	/** The node we are editing on */
-	TWeakObjectPtr<UAnimGraphNode_SubInput> WeakOuterNode;
+	TWeakObjectPtr<UAnimGraphNode_LinkedInputPose> WeakOuterNode;
 };
diff --git a/Engine/Source/Editor/AnimGraph/Private/AnimGraphDetails.cpp b/Engine/Source/Editor/AnimGraph/Private/AnimGraphDetails.cpp
--- a/Engine/Source/Editor/AnimGraph/Private/AnimGraphDetails.cpp
+++ b/Engine/Source/Editor/AnimGraph/Private/AnimGraphDetails.cpp
@@ -5,7 +5,7 @@
 #include "DetailLayoutBuilder.h"
 #include "DetailCategoryBuilder.h"
 #include "EdGraph/EdGraph.h"
-#include "AnimGraphNode_SubInput.h"
+#include "AnimGraphNode_LinkedInputPose.h"
 #include "Algo/Transform.h"
 #include "DetailWidgetRow.h"
 #include "Widgets/Text/STextBlock.h"
@@ -48,7 +48,7 @@ void FAnimGraphDetails::CustomizeDetails(IDetailLayoutBuilder& DetailLayout)
 
 	if(Objects.Num() > 1 || bIsStateMachine)
 	{
-		IDetailCategoryBuilder& InputsCategory = DetailLayout.EditCategory("Inputs", LOCTEXT("SubInputsCategory", "Inputs"));
+		IDetailCategoryBuilder& InputsCategory = DetailLayout.EditCategory("Inputs", LOCTEXT("LinkedInputPoseInputsCategory", "Inputs"));
 		InputsCategory.SetCategoryVisibility(false);
 		return;
 	}
@@ -60,7 +60,7 @@ void FAnimGraphDetails::CustomizeDetails(IDetailLayoutBuilder& DetailLayout)
 	{
 		FText ReadOnlyWarning = LOCTEXT("ReadOnlyWarning", "This graph's inputs are read-only and cannot be edited");
 
-		IDetailCategoryBuilder& InputsCategory = DetailLayout.EditCategory("Inputs", LOCTEXT("SubInputsCategory", "Inputs"));
+		IDetailCategoryBuilder& InputsCategory = DetailLayout.EditCategory("Inputs", LOCTEXT("LinkedInputPoseInputsCategory", "Inputs"));
 		InputsCategory.SetCategoryVisibility(false);
 
 		IDetailCategoryBuilder& WarningCategoryBuilder = DetailLayout.EditCategory("GraphInputs", LOCTEXT("GraphInputsCategory", "Graph Inputs"));
@@ -129,13 +129,13 @@ void FAnimGraphDetails::CustomizeDetails(IDetailLayoutBuilder& DetailLayout)
 		}
 	}
 
-	IDetailCategoryBuilder& InputsCategory = DetailLayout.EditCategory("Inputs", LOCTEXT("SubInputsCategory", "Inputs"));
+	IDetailCategoryBuilder& InputsCategory = DetailLayout.EditCategory("Inputs", LOCTEXT("LinkedInputPoseInputsCategory", "Inputs"));
 
 	DetailLayoutBuilder = &DetailLayout;
 
 	// Gather inputs, if any
-	TArray<UAnimGraphNode_SubInput*> SubInputs;
-	Graph->GetNodesOfClass<UAnimGraphNode_SubInput>(SubInputs);
+	TArray<UAnimGraphNode_LinkedInputPose*> LinkedInputPoseInputs;
+	Graph->GetNodesOfClass<UAnimGraphNode_LinkedInputPose>(LinkedInputPoseInputs);
 
 	TSharedRef<SHorizontalBox> InputsHeaderContentWidget = SNew(SHorizontalBox);
 	TWeakPtr<SWidget> WeakInputsHeaderWidget = InputsHeaderContentWidget;
@@ -179,26 +179,26 @@ void FAnimGraphDetails::CustomizeDetails(IDetailLayoutBuilder& DetailLayout)
 	];
 	InputsCategory.HeaderContent(InputsHeaderContentWidget);
 
-	if(SubInputs.Num())
+	if(LinkedInputPoseInputs.Num())
 	{
-		for(UAnimGraphNode_SubInput* SubInput : SubInputs)
+		for(UAnimGraphNode_LinkedInputPose* LinkedInputPoseNode : LinkedInputPoseInputs)
 		{
-			auto GetSubInputLabel = [WeakSubInput = TWeakObjectPtr<UAnimGraphNode_SubInput>(SubInput)]()
+			auto GetLinkedInputPoseLabel = [WeakLinkedInputPoseNode = TWeakObjectPtr<UAnimGraphNode_LinkedInputPose>(LinkedInputPoseNode)]()
 			{
-				if(WeakSubInput.IsValid())
+				if(WeakLinkedInputPoseNode.IsValid())
 				{
-					return FText::FromName(WeakSubInput->Node.Name);
+					return FText::FromName(WeakLinkedInputPoseNode->Node.Name);
 				}
 
 				return FText::GetEmpty();
 			};
 
 			TArray<UObject*> ExternalObjects;
-			ExternalObjects.Add(SubInput);
-			if(IDetailPropertyRow* SubInputRow = InputsCategory.AddExternalObjects(ExternalObjects))
+			ExternalObjects.Add(LinkedInputPoseNode);
+			if(IDetailPropertyRow* LinkedInputPoseRow = InputsCategory.AddExternalObjects(ExternalObjects))
 			{
-				SubInputRow->ShouldAutoExpand(true);
-				SubInputRow->CustomWidget()
+				LinkedInputPoseRow->ShouldAutoExpand(true);
+				LinkedInputPoseRow->CustomWidget()
 				.NameContent()
 				[
 					SNew(SBox)
@@ -221,7 +221,7 @@ void FAnimGraphDetails::CustomizeDetails(IDetailLayoutBuilder& DetailLayout)
 						.FillWidth(1.0f)
 						[
 							SNew(STextBlock)
-							.Text_Lambda(GetSubInputLabel)
+							.Text_Lambda(GetLinkedInputPoseLabel)
 							.Font(IDetailLayoutBuilder::GetDetailFont())
 						]
 						+SHorizontalBox::Slot()
@@ -234,7 +234,7 @@ void FAnimGraphDetails::CustomizeDetails(IDetailLayoutBuilder& DetailLayout)
 							.ButtonStyle(FEditorStyle::Get(), "HoverHintOnly")
 							.ForegroundColor(FEditorStyle::GetSlateColor("DefaultForeground"))
 							.ContentPadding(FMargin(2, 2))
-							.OnClicked(this, &FAnimGraphDetails::OnRemoveInputPoseClicked, SubInput)
+							.OnClicked(this, &FAnimGraphDetails::OnRemoveInputPoseClicked, LinkedInputPoseNode)
 							.ToolTipText(LOCTEXT("RemoveInputPoseTooltip", "Remove this input pose"))
 							[
 								SNew(SImage)
@@ -283,10 +283,10 @@ FReply FAnimGraphDetails::OnAddNewInputPoseClicked()
 	bool bIsInterface = AnimBlueprintEditor->GetBlueprintObj()->BlueprintType == BPTYPE_Interface;
 	if(!bIsInterface)
 	{
-		NewNodePosition = UAnimationGraphSchema::GetPositionForNewSubInputNode(*Graph);
+		NewNodePosition = UAnimationGraphSchema::GetPositionForNewLinkedInputPoseNode(*Graph);
 	}
 
-	FEdGraphSchemaAction_K2NewNode::SpawnNode<UAnimGraphNode_SubInput>(Graph, NewNodePosition, EK2NewNodeFlags::None);
+	FEdGraphSchemaAction_K2NewNode::SpawnNode<UAnimGraphNode_LinkedInputPose>(Graph, NewNodePosition, EK2NewNodeFlags::None);
 
 	DetailLayoutBuilder->ForceRefreshDetails();
 	
@@ -298,11 +298,11 @@ EVisibility FAnimGraphDetails::OnGetNewInputPoseTextVisibility(TWeakPtr<SWidget>
 	return WeakInputsHeaderWidget.Pin()->IsHovered() ? EVisibility::Visible : EVisibility::Collapsed;
 }
 
-FReply FAnimGraphDetails::OnRemoveInputPoseClicked(UAnimGraphNode_SubInput* InSubInput)
+FReply FAnimGraphDetails::OnRemoveInputPoseClicked(UAnimGraphNode_LinkedInputPose* InLinkedInputPose)
 {
 	{
-		FScopedTransaction Transaction(LOCTEXT("RemoveInputPose", "Remove Input Pose"));
-		Graph->RemoveNode(InSubInput);
+		FScopedTransaction Transaction(LOCTEXT("RemoveInputPose", "Remove Linked Input Pose"));
+		Graph->RemoveNode(InLinkedInputPose);
 	}
 
 	TSharedPtr<IAnimationBlueprintEditor> AnimBlueprintEditor = AnimBlueprintEditorPtr.Pin();
diff --git a/Engine/Source/Editor/AnimGraph/Private/AnimGraphDetails.h b/Engine/Source/Editor/AnimGraph/Private/AnimGraphDetails.h
--- a/Engine/Source/Editor/AnimGraph/Private/AnimGraphDetails.h
+++ b/Engine/Source/Editor/AnimGraph/Private/AnimGraphDetails.h
@@ -15,7 +15,7 @@ struct EVisibility;
 class SWidget;
 class UAnimGraphNode_Base;
 class UAnimGraphNode_Root;
-class UAnimGraphNode_SubInput;
+class UAnimGraphNode_LinkedInputPose;
 class ITableRow;
 class SComboButton;
 
@@ -41,7 +41,7 @@ private:
 	/** UI handlers */
 	FReply OnAddNewInputPoseClicked();
 	EVisibility OnGetNewInputPoseTextVisibility(TWeakPtr<SWidget> WeakInputsHeaderWidget) const;
-	FReply OnRemoveInputPoseClicked(UAnimGraphNode_SubInput* InSubInput);
+	FReply OnRemoveInputPoseClicked(UAnimGraphNode_LinkedInputPose* InLinkedInputPose);
 	FText OnGetGroupText() const;
 	void OnGroupTextCommitted(const FText& NewText, ETextCommit::Type InTextCommit);
 	void OnGroupSelectionChanged(TSharedPtr<FText> ProposedSelection, ESelectInfo::Type SelectInfo);
diff --git a/Engine/Source/Editor/AnimGraph/Private/AnimGraphNode_CustomProperty.cpp b/Engine/Source/Editor/AnimGraph/Private/AnimGraphNode_CustomProperty.cpp
--- a/Engine/Source/Editor/AnimGraph/Private/AnimGraphNode_CustomProperty.cpp
+++ b/Engine/Source/Editor/AnimGraph/Private/AnimGraphNode_CustomProperty.cpp
@@ -31,7 +31,7 @@ void UAnimGraphNode_CustomProperty::ValidateAnimNodeDuringCompilation(USkeleton*
 		UClass* TargetClass = GetTargetClass();
 		if(!TargetClass)
 		{
-			MessageLog.Error(TEXT("Sub instance node @@ has no valid instance class to spawn."), this);
+			MessageLog.Error(TEXT("Linked graph node @@ has no valid instance class to spawn."), this);
 		}
 	}
 }
diff --git a/Engine/Source/Editor/AnimGraph/Private/AnimGraphNode_Layer.cpp b/Engine/Source/Editor/AnimGraph/Private/AnimGraphNode_Layer.cpp
deleted file mode 100644
diff --git a/Engine/Source/Editor/AnimGraph/Private/AnimGraphNode_LinkedAnimGraphBase.cpp b/Engine/Source/Editor/AnimGraph/Private/AnimGraphNode_LinkedAnimGraphBase.cpp
new file mode 100644
--- /dev/null
+++ b/Engine/Source/Editor/AnimGraph/Private/AnimGraphNode_LinkedAnimGraphBase.cpp
@@ -0,0 +1,450 @@
+// Copyright 1998-2019 Epic Games, Inc. All Rights Reserved.
+
+#include "AnimGraphNode_LinkedAnimGraphBase.h"
+#include "Widgets/Text/STextBlock.h"
+#include "Widgets/Input/SCheckBox.h"
+#include "EdGraphSchema_K2.h"
+#include "Kismet2/CompilerResultsLog.h"
+#include "Kismet2/BlueprintEditorUtils.h"
+#include "PropertyHandle.h"
+#include "DetailLayoutBuilder.h"
+#include "DetailCategoryBuilder.h"
+#include "Animation/AnimNode_LinkedInputPose.h"
+#include "PropertyCustomizationHelpers.h"
+#include "ScopedTransaction.h"
+#include "AnimationGraphSchema.h"
+#include "Widgets/Layout/SBox.h"
+#include "UObject/CoreRedirects.h"
+#include "Animation/AnimNode_LinkedAnimGraph.h"
+
+#define LOCTEXT_NAMESPACE "LinkedAnimGraph"
+
+namespace LinkedAnimGraphGraphNodeConstants
+{
+	FLinearColor TitleColor(0.2f, 0.2f, 0.8f);
+}
+
+FLinearColor UAnimGraphNode_LinkedAnimGraphBase::GetNodeTitleColor() const
+{
+	return LinkedAnimGraphGraphNodeConstants::TitleColor;
+}
+
+FText UAnimGraphNode_LinkedAnimGraphBase::GetTooltipText() const
+{
+	return LOCTEXT("ToolTip", "Runs a linked anim graph in another instance to process animation");
+}
+
+FText UAnimGraphNode_LinkedAnimGraphBase::GetNodeTitle(ENodeTitleType::Type TitleType) const
+{
+	UClass* TargetClass = GetTargetClass();
+	UAnimBlueprint* TargetAnimBlueprint = TargetClass ? CastChecked<UAnimBlueprint>(TargetClass->ClassGeneratedBy) : nullptr;
+
+	const FAnimNode_LinkedAnimGraph& Node = *GetLinkedAnimGraphNode();
+
+	FFormatNamedArguments Args;
+	Args.Add(TEXT("NodeTitle"), LOCTEXT("Title", "Linked Anim Graph"));
+	Args.Add(TEXT("TargetClass"), TargetAnimBlueprint ? FText::FromString(TargetAnimBlueprint->GetName()) : LOCTEXT("ClassNone", "None"));
+
+	if(TitleType == ENodeTitleType::MenuTitle)
+	{
+		return LOCTEXT("NodeTitle", "Linked Anim Graph");
+	}
+	if(TitleType == ENodeTitleType::ListView)
+	{
+		if(Node.Tag != NAME_None)
+		{
+			Args.Add(TEXT("Tag"), FText::FromName(Node.Tag));
+			return FText::Format(LOCTEXT("TitleListFormatTagged", "{NodeTitle} ({Tag}) - {TargetClass}"), Args);
+		}
+		else
+		{
+			return FText::Format(LOCTEXT("TitleListFormat", "{NodeTitle} - {TargetClass}"), Args);
+		}
+	}
+	else
+	{
+
+		if(Node.Tag != NAME_None)
+		{
+			Args.Add(TEXT("Tag"), FText::FromName(Node.Tag));
+			return FText::Format(LOCTEXT("TitleFormatTagged", "{NodeTitle} ({Tag})\n{TargetClass}"), Args);
+		}
+		else
+		{
+			return FText::Format(LOCTEXT("TitleFormat", "{NodeTitle}\n{TargetClass}"), Args);
+		}
+	}
+}
+
+void UAnimGraphNode_LinkedAnimGraphBase::ValidateAnimNodeDuringCompilation(USkeleton* ForSkeleton, FCompilerResultsLog& MessageLog)
+{
+	Super::ValidateAnimNodeDuringCompilation(ForSkeleton, MessageLog);
+
+	UAnimBlueprint* AnimBP = CastChecked<UAnimBlueprint>(GetBlueprint());
+
+	UObject* OriginalNode = MessageLog.FindSourceObject(this);
+
+	if(HasInstanceLoop())
+	{
+		MessageLog.Error(TEXT("Detected loop in linked instance chain starting at @@ inside class @@"), this, AnimBP->GetAnimBlueprintGeneratedClass());
+	}
+
+	// Check for cycles from other linked instance nodes
+	TArray<UEdGraph*> Graphs;
+	AnimBP->GetAllGraphs(Graphs);
+
+	const FAnimNode_LinkedAnimGraph& Node = *GetLinkedAnimGraphNode();
+
+	// Check for duplicate tags in this anim blueprint
+	for(UEdGraph* Graph : Graphs)
+	{
+		TArray<UAnimGraphNode_LinkedAnimGraphBase*> LinkedAnimGraphNodes;
+		Graph->GetNodesOfClass(LinkedAnimGraphNodes);
+
+		for(UAnimGraphNode_LinkedAnimGraphBase* LinkedAnimGraphNode : LinkedAnimGraphNodes)
+		{
+			if(LinkedAnimGraphNode == OriginalNode)
+			{
+				continue;
+			}
+
+			FAnimNode_LinkedAnimGraph& InnerNode = *LinkedAnimGraphNode->GetLinkedAnimGraphNode();
+
+			if(InnerNode.Tag != NAME_None && InnerNode.Tag == Node.Tag)
+			{
+				MessageLog.Error(*FText::Format(LOCTEXT("DuplicateTagErrorFormat", "Node @@ and node @@ both have the same tag '{0}'."), FText::FromName(Node.Tag)).ToString(), this, LinkedAnimGraphNode);
+			}
+		}
+	}
+
+	// Check we don't try to spawn our own blueprint
+	if(GetTargetClass() == AnimBP->GetAnimBlueprintGeneratedClass())
+	{
+		MessageLog.Error(TEXT("Linked instance node @@ targets instance class @@ which it is inside, this would cause a loop."), this, AnimBP->GetAnimBlueprintGeneratedClass());
+	}
+}
+
+void UAnimGraphNode_LinkedAnimGraphBase::ReallocatePinsDuringReconstruction(TArray<UEdGraphPin*>& OldPins)
+{
+	// Grab the SKELETON class here as when we are reconstructed during during BP compilation
+	// the full generated class is not yet present built.
+	UClass* TargetClass = GetTargetSkeletonClass();
+
+	if(!TargetClass)
+	{
+		// Nothing to search for properties
+		return;
+	}
+
+	IAnimClassInterface* AnimClassInterface = IAnimClassInterface::GetFromClass(TargetClass);
+
+	const FAnimNode_LinkedAnimGraph& Node = *GetLinkedAnimGraphNode();
+
+	// Add any pose pins
+	for(const FAnimBlueprintFunction& AnimBlueprintFunction : AnimClassInterface->GetAnimBlueprintFunctions())
+	{
+		if(AnimBlueprintFunction.Name == Node.GetDynamicLinkFunctionName())
+		{
+			for(const FName& PoseName : AnimBlueprintFunction.InputPoseNames)
+			{
+				UEdGraphPin* NewPin = CreatePin(EEdGraphPinDirection::EGPD_Input, UAnimationGraphSchema::MakeLocalSpacePosePin(), PoseName);
+				NewPin->PinFriendlyName = FText::FromName(PoseName);
+				CustomizePinData(NewPin, PoseName, INDEX_NONE);
+			}
+
+			break;
+		}
+	}
+
+	// Call super to add properties
+	Super::ReallocatePinsDuringReconstruction(OldPins);
+}
+
+void UAnimGraphNode_LinkedAnimGraphBase::PostEditChangeProperty(FPropertyChangedEvent& PropertyChangedEvent)
+{
+	Super::PostEditChangeProperty(PropertyChangedEvent);
+
+	bool bRequiresNodeReconstruct = false;
+	UProperty* ChangedProperty = PropertyChangedEvent.Property;
+
+	if(ChangedProperty)
+	{
+		if (IsStructuralProperty(ChangedProperty))
+		{
+			bRequiresNodeReconstruct = true;
+			RebuildExposedProperties();
+		}
+	}
+
+	if(bRequiresNodeReconstruct)
+	{
+		ReconstructNode();
+	}
+}
+
+bool UAnimGraphNode_LinkedAnimGraphBase::HasInstanceLoop()
+{
+	TArray<FGuid> VisitedList;
+	TArray<FGuid> CurrentStack;
+	return HasInstanceLoop_Recursive(this, VisitedList, CurrentStack);
+}
+
+bool UAnimGraphNode_LinkedAnimGraphBase::HasInstanceLoop_Recursive(UAnimGraphNode_LinkedAnimGraphBase* CurrNode, TArray<FGuid>& VisitedNodes, TArray<FGuid>& NodeStack)
+{
+	if(!VisitedNodes.Contains(CurrNode->NodeGuid))
+	{
+		VisitedNodes.Add(CurrNode->NodeGuid);
+		NodeStack.Add(CurrNode->NodeGuid);
+
+		if(UAnimBlueprint* AnimBP = Cast<UAnimBlueprint>(UBlueprint::GetBlueprintFromClass(CurrNode->GetTargetClass())))
+		{
+			// Check for cycles from other linked instance nodes
+			TArray<UEdGraph*> Graphs;
+			AnimBP->GetAllGraphs(Graphs);
+
+			for(UEdGraph* Graph : Graphs)
+			{
+				TArray<UAnimGraphNode_LinkedAnimGraphBase*> LinkedInstanceNodes;
+				Graph->GetNodesOfClass(LinkedInstanceNodes);
+
+				for(UAnimGraphNode_LinkedAnimGraphBase* LinkedInstanceNode : LinkedInstanceNodes)
+				{
+					// If we haven't visited this node, then check it for loops, otherwise if we're pointing to a previously visited node that is in the current instance stack we have a loop
+					if((!VisitedNodes.Contains(LinkedInstanceNode->NodeGuid) && HasInstanceLoop_Recursive(LinkedInstanceNode, VisitedNodes, NodeStack)) || NodeStack.Contains(LinkedInstanceNode->NodeGuid))
+					{
+						return true;
+					}
+				}
+			}
+		}
+	}
+
+	NodeStack.Remove(CurrNode->NodeGuid);
+	return false;
+}
+
+void UAnimGraphNode_LinkedAnimGraphBase::CustomizeDetails(IDetailLayoutBuilder& DetailBuilder)
+{
+	Super::CustomizeDetails(DetailBuilder);
+
+	GenerateExposedPinsDetails(DetailBuilder);
+
+	IDetailCategoryBuilder& CategoryBuilder = DetailBuilder.EditCategory(FName(TEXT("Settings")));
+
+	// Customize InstanceClass
+	{
+		TSharedRef<IPropertyHandle> ClassHandle = DetailBuilder.GetProperty(TEXT("Node.InstanceClass"), GetClass());
+		ClassHandle->MarkHiddenByCustomization();
+
+		FDetailWidgetRow& ClassWidgetRow = CategoryBuilder.AddCustomRow(LOCTEXT("FilterStringInstanceClass", "Instance Class"));
+		ClassWidgetRow.NameContent()
+		[
+			ClassHandle->CreatePropertyNameWidget()
+		]
+		.ValueContent()
+		.MinDesiredWidth(250.0f)
+		[
+			SNew(SObjectPropertyEntryBox)
+			.ObjectPath_UObject(this, &UAnimGraphNode_LinkedAnimGraphBase::GetCurrentInstanceBlueprintPath)
+			.AllowedClass(UAnimBlueprint::StaticClass())
+			.NewAssetFactories(TArray<UFactory*>())
+			.OnShouldFilterAsset(FOnShouldFilterAsset::CreateUObject(this, &UAnimGraphNode_LinkedAnimGraphBase::OnShouldFilterInstanceBlueprint))
+			.OnObjectChanged(FOnSetObject::CreateUObject(this, &UAnimGraphNode_LinkedAnimGraphBase::OnSetInstanceBlueprint, &DetailBuilder))
+		];
+	}
+}
+
+void UAnimGraphNode_LinkedAnimGraphBase::GenerateExposedPinsDetails(IDetailLayoutBuilder &DetailBuilder)
+{
+	// We dont allow multi-select here
+	if(DetailBuilder.GetSelectedObjects().Num() > 1)
+	{
+		DetailBuilder.HideCategory(TEXT("Settings"));
+		return;
+	}
+
+	// We dont allow multi-select here
+	if(DetailBuilder.GetSelectedObjects().Num() > 1)
+	{
+		DetailBuilder.HideCategory(TEXT("Settings"));
+		return;
+	}
+
+	TArray<UProperty*> ExposableProperties;
+	GetExposableProperties(ExposableProperties);
+
+	if(ExposableProperties.Num() > 0)
+	{
+		IDetailCategoryBuilder& CategoryBuilder = DetailBuilder.EditCategory(FName(TEXT("Exposable Properties")));
+
+		FDetailWidgetRow& HeaderWidgetRow = CategoryBuilder.AddCustomRow(LOCTEXT("ExposeAll", "Expose All"));
+		
+		HeaderWidgetRow.NameContent()
+		[
+			SNew(STextBlock)
+			.Text(LOCTEXT("PropertyName", "Name"))
+			.Font(IDetailLayoutBuilder::GetDetailFontBold())
+		];
+
+		HeaderWidgetRow.ValueContent()
+		[
+			SNew(SHorizontalBox)
+			+SHorizontalBox::Slot()
+			.AutoWidth()
+			.VAlign(VAlign_Center)
+			[
+				SNew(STextBlock)
+				.Text(LOCTEXT("ExposeAllPropertyValue", "Expose All"))
+				.Font(IDetailLayoutBuilder::GetDetailFontBold())
+			]
+			+SHorizontalBox::Slot()
+			.FillWidth(1.0f)
+			.HAlign(HAlign_Right)
+			.VAlign(VAlign_Center)
+			[
+				SNew(SCheckBox)
+				.IsChecked_UObject(this, &UAnimGraphNode_CustomProperty::AreAllPropertiesExposed)
+				.OnCheckStateChanged_UObject(this, &UAnimGraphNode_CustomProperty::OnPropertyExposeAllCheckboxChanged)
+			]
+		];
+
+		for(UProperty* Property : ExposableProperties)
+		{
+			FDetailWidgetRow& PropertyWidgetRow = CategoryBuilder.AddCustomRow(FText::FromString(Property->GetName()));
+
+			FName PropertyName = Property->GetFName();
+			FText PropertyTypeText = GetPropertyTypeText(Property);
+
+			FFormatNamedArguments Args;
+			Args.Add(TEXT("PropertyName"), FText::FromName(PropertyName));
+			Args.Add(TEXT("PropertyType"), PropertyTypeText);
+
+			FText TooltipText = FText::Format(LOCTEXT("PropertyTooltipText", "{PropertyName}\nType: {PropertyType}"), Args);
+
+			PropertyWidgetRow.NameContent()
+			[
+				SNew(STextBlock)
+				.Text(FText::FromString(Property->GetName()))
+				.ToolTipText(TooltipText)
+				.Font(IDetailLayoutBuilder::GetDetailFont())
+			];
+
+			PropertyWidgetRow.ValueContent()
+			[
+				SNew(SHorizontalBox)
+				+SHorizontalBox::Slot()
+				.VAlign(VAlign_Center)
+				.AutoWidth()
+				[
+					SNew(STextBlock)
+					.Text(LOCTEXT("ExposePropertyValue", "Expose:"))
+					.Font(IDetailLayoutBuilder::GetDetailFont())
+				]
+				+SHorizontalBox::Slot()
+				.FillWidth(1.0f)
+				.HAlign(HAlign_Right)
+				.VAlign(VAlign_Center)
+				[
+					SNew(SCheckBox)
+					.IsChecked_UObject(this, &UAnimGraphNode_CustomProperty::IsPropertyExposed, PropertyName)
+					.OnCheckStateChanged_UObject(this, &UAnimGraphNode_CustomProperty::OnPropertyExposeCheckboxChanged, PropertyName)
+				]
+			];
+		}
+	}
+}
+
+bool UAnimGraphNode_LinkedAnimGraphBase::IsStructuralProperty(UProperty* InProperty) const
+{
+	return InProperty->GetFName() == GET_MEMBER_NAME_CHECKED(FAnimNode_LinkedAnimGraph, InstanceClass);
+}
+
+FString UAnimGraphNode_LinkedAnimGraphBase::GetCurrentInstanceBlueprintPath() const
+{
+	UClass* InstanceClass = GetTargetClass();
+
+	if(InstanceClass)
+	{
+		UBlueprint* ActualBlueprint = UBlueprint::GetBlueprintFromClass(InstanceClass);
+
+		if(ActualBlueprint)
+		{
+			return ActualBlueprint->GetPathName();
+		}
+	}
+
+	return FString();
+}
+
+bool UAnimGraphNode_LinkedAnimGraphBase::OnShouldFilterInstanceBlueprint(const FAssetData& AssetData) const
+{
+	// Check recursion
+	if(AssetData.IsAssetLoaded() && Cast<UBlueprint>(AssetData.GetAsset()) == GetBlueprint())
+	{
+		return true;
+	}
+
+	// Check skeleton
+	FAssetDataTagMapSharedView::FFindTagResult Result = AssetData.TagsAndValues.FindTag("TargetSkeleton");
+	if (Result.IsSet())
+	{
+		if (UAnimBlueprint* CurrentBlueprint = Cast<UAnimBlueprint>(GetBlueprint()))
+		{
+			FString TargetSkeletonName = FString::Printf(TEXT("%s'%s'"), *CurrentBlueprint->TargetSkeleton->GetClass()->GetName(), *CurrentBlueprint->TargetSkeleton->GetPathName());
+			if(Result.GetValue() != TargetSkeletonName)
+			{
+				return true;
+			}
+		}
+	}
+
+	return false;
+}
+
+void UAnimGraphNode_LinkedAnimGraphBase::OnSetInstanceBlueprint(const FAssetData& AssetData, IDetailLayoutBuilder* InDetailBuilder)
+{
+	FScopedTransaction Transaction(LOCTEXT("SetInstanceBlueprint", "Set Linked Blueprint"));
+
+	Modify();
+	
+	TSharedRef<IPropertyHandle> ClassHandle = InDetailBuilder->GetProperty(TEXT("Node.InstanceClass"), GetClass());
+	if(UAnimBlueprint* Blueprint = Cast<UAnimBlueprint>(AssetData.GetAsset()))
+	{
+		ClassHandle->SetValue(Blueprint->GetAnimBlueprintGeneratedClass());
+	}
+	else
+	{
+		ClassHandle->SetValue((UObject*)nullptr);
+	}
+}
+
+FPoseLinkMappingRecord UAnimGraphNode_LinkedAnimGraphBase::GetLinkIDLocation(const UScriptStruct* NodeType, UEdGraphPin* SourcePin)
+{
+	FPoseLinkMappingRecord Record = Super::GetLinkIDLocation(NodeType, SourcePin);
+	if(Record.IsValid())
+	{
+		return Record;	
+	}
+	else if(SourcePin->LinkedTo.Num() > 0 && SourcePin->Direction == EGPD_Input)
+	{
+		const FAnimNode_LinkedAnimGraph& Node = *GetLinkedAnimGraphNode();
+
+		check(Node.InputPoses.Num() == Node.InputPoseNames.Num());
+
+		// perform name-based logic for input pose pins
+		if (UAnimGraphNode_Base* LinkedNode = Cast<UAnimGraphNode_Base>(FBlueprintEditorUtils::FindFirstCompilerRelevantNode(SourcePin->LinkedTo[0])))
+		{
+			UArrayProperty* ArrayProperty = FindFieldChecked<UArrayProperty>(NodeType, GET_MEMBER_NAME_CHECKED(FAnimNode_LinkedAnimGraph, InputPoses));
+			int32 ArrayIndex = INDEX_NONE;
+			if(Node.InputPoseNames.Find(SourcePin->GetFName(), ArrayIndex))
+			{
+				check(Node.InputPoses.IsValidIndex(ArrayIndex));
+				return FPoseLinkMappingRecord::MakeFromArrayEntry(this, LinkedNode, ArrayProperty, ArrayIndex);
+			}
+		}
+	}
+
+	return FPoseLinkMappingRecord::MakeInvalid();
+}
+
+#undef LOCTEXT_NAMESPACE
diff --git a/Engine/Source/Editor/AnimGraph/Private/AnimGraphNode_LinkedAnimLayer.cpp b/Engine/Source/Editor/AnimGraph/Private/AnimGraphNode_LinkedAnimLayer.cpp
new file mode 100644
--- /dev/null
+++ b/Engine/Source/Editor/AnimGraph/Private/AnimGraphNode_LinkedAnimLayer.cpp
@@ -0,0 +1,589 @@
+// Copyright 1998-2019 Epic Games, Inc. All Rights Reserved.
+
+#include "AnimGraphNode_LinkedAnimLayer.h"
+#include "Widgets/Text/STextBlock.h"
+#include "Widgets/Input/SCheckBox.h"
+#include "EdGraphSchema_K2.h"
+#include "Kismet2/CompilerResultsLog.h"
+#include "Kismet2/BlueprintEditorUtils.h"
+#include "PropertyHandle.h"
+#include "DetailLayoutBuilder.h"
+#include "DetailCategoryBuilder.h"
+#include "Animation/AnimNode_LinkedInputPose.h"
+#include "PropertyCustomizationHelpers.h"
+#include "ScopedTransaction.h"
+#include "AnimationGraphSchema.h"
+#include "Widgets/Layout/SBox.h"
+#include "UObject/CoreRedirects.h"
+#include "Kismet2/KismetEditorUtilities.h"
+#include "AnimationStateGraph.h"
+
+#define LOCTEXT_NAMESPACE "LinkedAnimLayerNode"
+
+FText UAnimGraphNode_LinkedAnimLayer::GetTooltipText() const
+{
+	return LOCTEXT("ToolTip", "Runs another linked animation layer graph to process animation");
+}
+
+FText UAnimGraphNode_LinkedAnimLayer::GetNodeTitle(ENodeTitleType::Type TitleType) const
+{
+	UClass* TargetClass = *Node.Interface;
+	UAnimBlueprint* TargetAnimBlueprint = TargetClass ? CastChecked<UAnimBlueprint>(TargetClass->ClassGeneratedBy) : nullptr;
+
+	FFormatNamedArguments Args;
+	Args.Add(TEXT("NodeTitle"), LOCTEXT("Title", "Linked Anim Layer"));
+	Args.Add(TEXT("TargetClass"), TargetAnimBlueprint ? FText::FromString(TargetAnimBlueprint->GetName()) : LOCTEXT("ClassSelf", "Self"));
+
+	if(TitleType == ENodeTitleType::MenuTitle)
+	{
+		return LOCTEXT("NodeTitle", "Layer");
+	}
+	else if(TitleType == ENodeTitleType::ListView)
+	{
+		Args.Add(TEXT("Layer"), Node.Layer == NAME_None ? LOCTEXT("LayerNone", "None") : FText::FromName(Node.Layer));
+		return FText::Format(LOCTEXT("TitleListFormatOutputPose", "{NodeTitle}: {Layer} - {TargetClass}"), Args);
+	}
+	else
+	{
+		Args.Add(TEXT("Layer"), Node.Layer == NAME_None ? LOCTEXT("LayerNone", "None") : FText::FromName(Node.Layer));
+		return FText::Format(LOCTEXT("TitleFormatOutputPose", "{NodeTitle}: {Layer}\n{TargetClass}"), Args);
+	}
+}
+
+void UAnimGraphNode_LinkedAnimLayer::ValidateAnimNodeDuringCompilation(USkeleton* ForSkeleton, FCompilerResultsLog& MessageLog)
+{
+	Super::ValidateAnimNodeDuringCompilation(ForSkeleton, MessageLog);
+
+	if(Node.Layer == NAME_None)
+	{
+		MessageLog.Error(*LOCTEXT("NoLayerError", "Linked anim layer node @@ does not specify a layer.").ToString(), this);
+	}
+	else
+	{
+		UAnimBlueprint* CurrentBlueprint = Cast<UAnimBlueprint>(GetBlueprint());
+
+		// check layer actually exists in the interface
+		UClass* TargetClass = *Node.Interface;
+		if(TargetClass == nullptr)
+		{
+			// If no interface specified, use this class
+			if (CurrentBlueprint)
+			{
+				TargetClass = *CurrentBlueprint->SkeletonGeneratedClass;
+			}
+		}
+		else
+		{
+			// check we implement this interface
+			bool bImplementsInterface = false;
+
+			if (CurrentBlueprint)
+			{
+				for (FBPInterfaceDescription& InterfaceDesc : CurrentBlueprint->ImplementedInterfaces)
+				{
+					if (InterfaceDesc.Interface.Get() == TargetClass)
+					{
+						bImplementsInterface = true;
+						break;
+					}
+				}
+			}
+
+			if(!bImplementsInterface)
+			{
+				// Its possible we have a left-over interface referenced here that needs clearing now we are a 'self' layer
+				if(GetInterfaceForLayer() == nullptr)
+				{
+					Node.Interface = nullptr;
+
+					// No interface any more, use this class
+					if (CurrentBlueprint)
+					{
+						TargetClass = *CurrentBlueprint->SkeletonGeneratedClass;
+					}
+				}
+				else
+				{
+					MessageLog.Error(*LOCTEXT("MissingInterfaceError", "Linked anim layer node @@ uses interface @@ that this blueprint does not implement.").ToString(), this, Node.Interface.Get());
+				}
+			}
+		}
+
+		if(TargetClass)
+		{
+			bool bFoundFunction = false;
+			IAnimClassInterface* AnimClassInterface = IAnimClassInterface::GetFromClass(TargetClass);
+			for(const FAnimBlueprintFunction& AnimBlueprintFunction : AnimClassInterface->GetAnimBlueprintFunctions())
+			{
+				if(AnimBlueprintFunction.Name == Node.Layer)
+				{
+					bFoundFunction = true;
+				}
+			}
+
+			if(!bFoundFunction)
+			{
+				MessageLog.Error(*FText::Format(LOCTEXT("MissingLayerError", "Linked anim layer node @@ uses invalid layer '{0}'."), FText::FromName(Node.Layer)).ToString(), this);
+			}
+		}
+
+		if(CurrentBlueprint)
+		{
+			UAnimGraphNode_LinkedAnimLayer* OriginalThis = Cast<UAnimGraphNode_LinkedAnimLayer>(MessageLog.FindSourceObject(this));
+
+			// check layer is only used once in this blueprint
+			auto CheckGraph = [this, OriginalThis, &MessageLog](const UEdGraph* InGraph)
+			{
+				TArray<UAnimGraphNode_LinkedAnimLayer*> LayerNodes;
+				InGraph->GetNodesOfClass(LayerNodes);
+				for(const UAnimGraphNode_LinkedAnimLayer* LayerNode : LayerNodes)
+				{
+					if(LayerNode != OriginalThis)
+					{
+						if(LayerNode->Node.Layer == Node.Layer)
+						{
+							MessageLog.Error(*FText::Format(LOCTEXT("DuplicateLayerError", "Linked anim layer node @@ also uses layer '{0}', layers can be used only once in an animation blueprint."), FText::FromName(Node.Layer)).ToString(), this);
+						}
+					}
+				}
+			};
+		
+			TArray<UEdGraph*> Graphs;
+			CurrentBlueprint->GetAllGraphs(Graphs);
+
+			auto ValidateOuterGraph = [this, OriginalThis, &MessageLog](const UEdGraph* InGraph)
+			{
+				static const FName DefaultAnimGraphName("AnimGraph");
+				if (InGraph->Nodes.Contains(OriginalThis))
+				{
+					if (!InGraph->IsA<UAnimationStateGraph>() && InGraph->GetFName() != DefaultAnimGraphName)
+					{
+						MessageLog.Error(*FText::Format(LOCTEXT("NestedLayer", "Linked anim layer node @@ is part of Animation Layer Graph '{0}', layers cannot be nested."), FText::FromName(InGraph->GetFName())).ToString(), this);
+					}
+				}
+			};
+			
+			for(const UEdGraph* Graph : Graphs)
+			{
+				CheckGraph(Graph);
+				ValidateOuterGraph(Graph);				
+			}
+		}
+	}
+}
+
+UObject* UAnimGraphNode_LinkedAnimLayer::GetJumpTargetForDoubleClick() const
+{
+	auto JumpTargetFromClass = [this](UClass* InClass)
+	{
+		UObject* JumpTargetObject = nullptr;
+
+		UAnimBlueprint* TargetAnimBlueprint = InClass ? CastChecked<UAnimBlueprint>(InClass->ClassGeneratedBy) : nullptr;
+		if(TargetAnimBlueprint == nullptr || TargetAnimBlueprint == Cast<UAnimBlueprint>(GetBlueprint()))
+		{
+			// jump to graph in self
+			TArray<UEdGraph*> Graphs;
+			GetBlueprint()->GetAllGraphs(Graphs);
+
+			UEdGraph** FoundGraph = Graphs.FindByPredicate([this](UEdGraph* InGraph){ return InGraph->GetFName() == Node.Layer; });
+			if(FoundGraph)
+			{
+				JumpTargetObject = *FoundGraph;
+			}
+		}
+		else if(TargetAnimBlueprint)
+		{
+			// jump to graph in other BP
+			TArray<UEdGraph*> Graphs;
+			TargetAnimBlueprint->GetAllGraphs(Graphs);
+
+			UEdGraph** FoundGraph = Graphs.FindByPredicate([this](UEdGraph* InGraph){ return InGraph->GetFName() == Node.Layer; });
+			if(FoundGraph)
+			{
+				JumpTargetObject = *FoundGraph;
+			}
+			else
+			{
+				JumpTargetObject = TargetAnimBlueprint;
+			}
+		}
+
+		return JumpTargetObject;
+	};
+
+	// First try a concrete class, if any
+	UObject* JumpTargetObject = JumpTargetFromClass(*Node.InstanceClass);
+	if(JumpTargetObject == nullptr)
+	{
+		// then try the interface
+		JumpTargetObject = JumpTargetFromClass(*Node.Interface);
+	}
+
+	return JumpTargetObject;
+}
+
+void UAnimGraphNode_LinkedAnimLayer::JumpToDefinition() const
+{
+	if (UEdGraph* HyperlinkTarget = Cast<UEdGraph>(GetJumpTargetForDoubleClick()))
+	{
+		FKismetEditorUtilities::BringKismetToFocusAttentionOnObject(HyperlinkTarget);
+	}
+	else
+	{
+		Super::JumpToDefinition();
+	}
+}
+
+bool UAnimGraphNode_LinkedAnimLayer::HasExternalDependencies(TArray<class UStruct*>* OptionalOutput /*= NULL*/) const
+{
+	UClass* InterfaceClassToUse = *Node.Interface;
+
+	// Add our interface class. If that changes we need a recompile
+	if(InterfaceClassToUse && OptionalOutput)
+	{
+		OptionalOutput->AddUnique(InterfaceClassToUse);
+	}
+
+	bool bSuperResult = Super::HasExternalDependencies(OptionalOutput);
+	return InterfaceClassToUse || bSuperResult;
+}
+
+void UAnimGraphNode_LinkedAnimLayer::CustomizeDetails(IDetailLayoutBuilder& DetailBuilder)
+{
+	// We dont allow multi-select here
+	if(DetailBuilder.GetSelectedObjects().Num() > 1)
+	{
+		DetailBuilder.HideCategory(TEXT("Settings"));
+		return;
+	}
+
+	IDetailCategoryBuilder& CategoryBuilder = DetailBuilder.EditCategory(TEXT("Settings"));
+
+	// Hide Tag
+	TSharedRef<IPropertyHandle> TagHandle = DetailBuilder.GetProperty(TEXT("Node.Tag"), GetClass());
+	TagHandle->MarkHiddenByCustomization();
+
+	// Customize Layer
+	{
+		TSharedRef<IPropertyHandle> LayerHandle = DetailBuilder.GetProperty(TEXT("Node.Layer"), GetClass());
+		if(LayerHandle->IsValidHandle())
+		{
+			LayerHandle->SetOnPropertyValueChanged(FSimpleDelegate::CreateUObject(this, &UAnimGraphNode_LinkedAnimLayer::OnLayerChanged, &DetailBuilder));
+		}
+
+		LayerHandle->MarkHiddenByCustomization();
+
+		// Check layers available in this BP
+		FDetailWidgetRow& LayerWidgetRow = CategoryBuilder.AddCustomRow(LOCTEXT("FilterStringLayer", "Layer"));
+		LayerWidgetRow.NameContent()
+		[
+			LayerHandle->CreatePropertyNameWidget()
+		]
+		.ValueContent()
+		.MinDesiredWidth(150.0f)
+		[
+			SNew(SHorizontalBox)
+			+SHorizontalBox::Slot()
+			.VAlign(VAlign_Center)
+			[
+				SNew(SBox)
+				.Visibility_Lambda([this](){ return HasAvailableLayers() ? EVisibility::Visible : EVisibility::Collapsed; })
+				[
+					PropertyCustomizationHelpers::MakePropertyComboBox(
+						LayerHandle, 
+						FOnGetPropertyComboBoxStrings::CreateUObject(this,  &UAnimGraphNode_LinkedAnimLayer::GetLayerNames),
+						FOnGetPropertyComboBoxValue::CreateUObject(this,  &UAnimGraphNode_LinkedAnimLayer::GetLayerName)
+					)
+				]
+			]
+			+SHorizontalBox::Slot()
+			.VAlign(VAlign_Center)
+			[
+				SNew(STextBlock)
+				.Visibility_Lambda([this](){ return !HasAvailableLayers() ? EVisibility::Visible : EVisibility::Collapsed; })
+				.Font(IDetailLayoutBuilder::GetDetailFont())
+				.Text(LOCTEXT("NoLayersWarning", "No available layers."))
+				.ToolTipText(LOCTEXT("NoLayersWarningTooltip", "This Animation Blueprint has no layers to choose from.\nTo add some, either implement an Animation Layer Interface via the Class Settings, or add an animation layer in the My Blueprint tab."))
+			]
+		];
+	}
+
+	GenerateExposedPinsDetails(DetailBuilder);
+	UAnimGraphNode_CustomProperty::CustomizeDetails(DetailBuilder);
+
+	// Customize InstanceClass with unique visibility (identical to parent class apart from this)
+	{
+		TSharedRef<IPropertyHandle> ClassHandle = DetailBuilder.GetProperty(TEXT("Node.InstanceClass"), GetClass());
+		ClassHandle->MarkHiddenByCustomization();
+
+		FDetailWidgetRow& ClassWidgetRow = CategoryBuilder.AddCustomRow(LOCTEXT("FilterStringInstanceClass", "Instance Class"));
+		ClassWidgetRow.NameContent()
+		[
+			ClassHandle->CreatePropertyNameWidget()
+		]
+		.ValueContent()
+		.MinDesiredWidth(250.0f)
+		[
+			SNew(SHorizontalBox)
+			+SHorizontalBox::Slot()
+			.VAlign(VAlign_Center)
+			[
+				SNew(SObjectPropertyEntryBox)
+				.Visibility_Lambda([this](){ return HasValidNonSelfLayer() ? EVisibility::Visible : EVisibility::Collapsed; })
+				.ObjectPath_UObject(this, &UAnimGraphNode_LinkedAnimLayer::GetCurrentInstanceBlueprintPath)
+				.AllowedClass(UAnimBlueprint::StaticClass())
+				.NewAssetFactories(TArray<UFactory*>())
+				.OnShouldFilterAsset(FOnShouldFilterAsset::CreateUObject(this, &UAnimGraphNode_LinkedAnimLayer::OnShouldFilterInstanceBlueprint))
+				.OnObjectChanged(FOnSetObject::CreateUObject(this, &UAnimGraphNode_LinkedAnimLayer::OnSetInstanceBlueprint, &DetailBuilder))
+			]
+			+SHorizontalBox::Slot()
+			.VAlign(VAlign_Center)
+			[
+				SNew(STextBlock)
+				.Visibility_Lambda([this](){ return !HasValidNonSelfLayer() ? EVisibility::Visible : EVisibility::Collapsed; })
+				.Font(IDetailLayoutBuilder::GetDetailFont())
+				.Text(LOCTEXT("SelfLayersWarning", "Uses layer in this Blueprint."))
+				.ToolTipText(LOCTEXT("SelfLayersWarningTooltip", "This linked anim layer node refers to a layer only in this blueprint, so cannot be overriden by an external blueprint implementation.\nChange to use a layer from an implemented interface to allow this override."))
+			]
+		];
+	}
+}
+
+bool UAnimGraphNode_LinkedAnimLayer::OnShouldFilterInstanceBlueprint(const FAssetData& AssetData) const
+{
+	if(Super::OnShouldFilterInstanceBlueprint(AssetData))
+	{
+		return true;
+	}
+
+	if (UAnimBlueprint* CurrentBlueprint = Cast<UAnimBlueprint>(GetBlueprint()))
+	{
+		TArray<TSubclassOf<UInterface>> AnimInterfaces;
+		for(const FBPInterfaceDescription& InterfaceDesc : CurrentBlueprint->ImplementedInterfaces)
+		{
+			if(InterfaceDesc.Interface && InterfaceDesc.Interface->IsChildOf<UAnimLayerInterface>())
+			{
+				if(Node.Layer == NAME_None || InterfaceDesc.Interface->FindFunctionByName(Node.Layer))
+				{
+					AnimInterfaces.Add(InterfaceDesc.Interface);
+				}
+			}
+		}
+
+		// Check interface compatibility
+		if(AnimInterfaces.Num() > 0)
+		{
+			bool bMatchesInterface = false;
+
+			const FString ImplementedInterfaces = AssetData.GetTagValueRef<FString>(FBlueprintTags::ImplementedInterfaces);
+			if(!ImplementedInterfaces.IsEmpty())
+			{
+				FString FullInterface;
+				FString RemainingString;
+				FString InterfacePath;
+				FString CurrentString = *ImplementedInterfaces;
+				while(CurrentString.Split(TEXT(","), &FullInterface, &RemainingString) && !bMatchesInterface)
+				{
+					if (!CurrentString.StartsWith(TEXT("Graphs=(")))
+					{
+						if (FullInterface.Split(TEXT("\""), &CurrentString, &InterfacePath, ESearchCase::CaseSensitive))
+						{
+							// The interface paths in metadata end with "', so remove those
+							InterfacePath.RemoveFromEnd(TEXT("\"'"));
+
+							FCoreRedirectObjectName ResolvedInterfaceName = FCoreRedirects::GetRedirectedName(ECoreRedirectFlags::Type_Class, FCoreRedirectObjectName(InterfacePath));
+
+							// Verify against all interfaces we currently implement
+							for(TSubclassOf<UInterface> AnimInterface : AnimInterfaces)
+							{
+								bMatchesInterface |= ResolvedInterfaceName.ObjectName == AnimInterface->GetFName();
+							}
+						}
+					}
+			
+					CurrentString = RemainingString;
+				}
+			}
+
+			if(!bMatchesInterface)
+			{
+				return true;
+			}
+		}
+		else
+		{
+			// No interfaces, so no compatible BPs
+			return true;
+		}
+	}
+
+	return false;
+}
+
+FString UAnimGraphNode_LinkedAnimLayer::GetCurrentInstanceBlueprintPath() const
+{
+	UClass* InterfaceClass = *Node.InstanceClass;
+
+	if(InterfaceClass)
+	{
+		UBlueprint* ActualBlueprint = UBlueprint::GetBlueprintFromClass(InterfaceClass);
+
+		if(ActualBlueprint)
+		{
+			return ActualBlueprint->GetPathName();
+		}
+	}
+
+	return FString();
+}
+
+void UAnimGraphNode_LinkedAnimLayer::GetExposableProperties(TArray<UProperty*>& OutExposableProperties) const
+{
+	UClass* TargetClass = GetTargetSkeletonClass();
+	if(TargetClass)
+	{
+		// add only sub-input properties
+		IAnimClassInterface* AnimClassInterface = IAnimClassInterface::GetFromClass(TargetClass);
+		for(const FAnimBlueprintFunction& AnimBlueprintFunction : AnimClassInterface->GetAnimBlueprintFunctions())
+		{
+			// Check name matches.
+			if(AnimBlueprintFunction.Name == Node.GetDynamicLinkFunctionName())
+			{
+				for(UProperty* Property : AnimBlueprintFunction.InputProperties)
+				{
+					OutExposableProperties.Add(Property);
+				}
+			}
+		}
+	}
+}
+
+void UAnimGraphNode_LinkedAnimLayer::GetLayerNames(TArray<TSharedPtr<FString>>& OutStrings, TArray<TSharedPtr<SToolTip>>& OutToolTips, TArray<bool>& OutRestrictedItems)
+{
+	// If no interface specified, use this class
+	if (UAnimBlueprint* CurrentBlueprint = Cast<UAnimBlueprint>(GetBlueprint()))
+	{
+		UClass* TargetClass = *CurrentBlueprint->SkeletonGeneratedClass;
+		if(TargetClass)
+		{
+			IAnimClassInterface* AnimClassInterface = IAnimClassInterface::GetFromClass(TargetClass);
+			for(const FAnimBlueprintFunction& AnimBlueprintFunction : AnimClassInterface->GetAnimBlueprintFunctions())
+			{
+				if(AnimBlueprintFunction.Name != UEdGraphSchema_K2::GN_AnimGraph)
+				{
+					OutStrings.Add(MakeShared<FString>(AnimBlueprintFunction.Name.ToString()));
+					OutToolTips.Add(nullptr);
+					OutRestrictedItems.Add(false);
+				}
+			}
+		}
+	}
+}
+
+FString UAnimGraphNode_LinkedAnimLayer::GetLayerName() const
+{
+	return Node.Layer.ToString();
+}
+
+bool UAnimGraphNode_LinkedAnimLayer::IsStructuralProperty(UProperty* InProperty) const
+{
+	return Super::IsStructuralProperty(InProperty) ||
+		InProperty->GetFName() == GET_MEMBER_NAME_CHECKED(FAnimNode_LinkedAnimLayer, Layer);
+}
+
+UClass* UAnimGraphNode_LinkedAnimLayer::GetTargetSkeletonClass() const
+{
+	UClass* SuperTargetSkeletonClass = Super::GetTargetSkeletonClass();
+	if(SuperTargetSkeletonClass == nullptr)
+	{
+		// If no concrete class specified, use this class
+		if (UAnimBlueprint* CurrentBlueprint = Cast<UAnimBlueprint>(GetBlueprint()))
+		{
+			SuperTargetSkeletonClass = *CurrentBlueprint->SkeletonGeneratedClass;
+		}
+	}
+	return SuperTargetSkeletonClass;
+}
+
+TSubclassOf<UInterface> UAnimGraphNode_LinkedAnimLayer::GetInterfaceForLayer() const
+{
+	if (UAnimBlueprint* CurrentBlueprint = Cast<UAnimBlueprint>(GetBlueprint()))
+	{
+		UClass* TargetClass = *CurrentBlueprint->SkeletonGeneratedClass;
+		if(TargetClass)
+		{
+			// Find layer with this name in interfaces
+			for(FBPInterfaceDescription& InterfaceDesc : CurrentBlueprint->ImplementedInterfaces)
+			{
+				for(UEdGraph* InterfaceGraph : InterfaceDesc.Graphs)
+				{
+					if(InterfaceGraph->GetFName() == Node.Layer)
+					{
+						return InterfaceDesc.Interface;
+					}
+				}
+			}
+		}
+	}
+
+	return nullptr;
+}
+
+void UAnimGraphNode_LinkedAnimLayer::OnLayerChanged(IDetailLayoutBuilder* DetailBuilder)
+{
+	OnStructuralPropertyChanged(DetailBuilder);
+
+	// Get the interface for this layer. If null, then we are using a 'self' layer.
+	Node.Interface = GetInterfaceForLayer();
+
+	if(Node.Interface.Get() == nullptr)
+	{
+		// Self layers cannot have override implementations
+		Node.InstanceClass = nullptr;
+	}
+}
+
+bool UAnimGraphNode_LinkedAnimLayer::HasAvailableLayers() const
+{
+	if (UAnimBlueprint* CurrentBlueprint = Cast<UAnimBlueprint>(GetBlueprint()))
+	{
+		UClass* TargetClass = *CurrentBlueprint->SkeletonGeneratedClass;
+		if(TargetClass)
+		{
+			IAnimClassInterface* AnimClassInterface = IAnimClassInterface::GetFromClass(TargetClass);
+			for(const FAnimBlueprintFunction& AnimBlueprintFunction : AnimClassInterface->GetAnimBlueprintFunctions())
+			{
+				if(AnimBlueprintFunction.Name != UEdGraphSchema_K2::GN_AnimGraph)
+				{
+					return true;
+				}
+			}
+		}
+	}
+
+	return false;
+}
+
+bool UAnimGraphNode_LinkedAnimLayer::HasValidNonSelfLayer() const
+{
+	if (UAnimBlueprint* CurrentBlueprint = Cast<UAnimBlueprint>(GetBlueprint()))
+	{
+		if(Node.Interface.Get())
+		{
+			for(const FBPInterfaceDescription& InterfaceDesc : CurrentBlueprint->ImplementedInterfaces)
+			{
+				if(InterfaceDesc.Interface && InterfaceDesc.Interface->IsChildOf<UAnimLayerInterface>())
+				{
+					if(InterfaceDesc.Interface->FindFunctionByName(Node.Layer))
+					{
+						return true;
+					}
+				}
+			}
+		}
+	}
+
+	return false;
+}
+
+#undef LOCTEXT_NAMESPACE
diff --git a/Engine/Source/Editor/AnimGraph/Private/AnimGraphNode_LinkedInputPose.cpp b/Engine/Source/Editor/AnimGraph/Private/AnimGraphNode_LinkedInputPose.cpp
new file mode 100644
--- /dev/null
+++ b/Engine/Source/Editor/AnimGraph/Private/AnimGraphNode_LinkedInputPose.cpp
@@ -0,0 +1,556 @@
+// Copyright 1998-2019 Epic Games, Inc. All Rights Reserved.
+
+#include "AnimGraphNode_LinkedInputPose.h"
+#include "Kismet2/BlueprintEditorUtils.h"
+#include "GraphEditorSettings.h"
+#include "BlueprintActionFilter.h"
+#include "Widgets/Input/SButton.h"
+#include "DetailCategoryBuilder.h"
+#include "IDetailPropertyRow.h"
+#include "DetailWidgetRow.h"
+#include "EditorStyleSet.h"
+#include "DetailLayoutBuilder.h"
+#include "Widgets/Images/SImage.h"
+#include "Widgets/Text/STextBlock.h"
+
+#include "Animation/AnimBlueprint.h"
+#include "IAnimationBlueprintEditor.h"
+#include "AnimationGraphSchema.h"
+#include "Widgets/Input/SEditableTextBox.h"
+#include "Widgets/Layout/SBox.h"
+#include "Widgets/SBoxPanel.h"
+#include "K2Node_CallFunction.h"
+#include "Containers/Ticker.h"
+#include "Widgets/Input/SCheckBox.h"
+#include "Subsystems/AssetEditorSubsystem.h"
+
+#define LOCTEXT_NAMESPACE "LinkedInputPose"
+
+UAnimGraphNode_LinkedInputPose::UAnimGraphNode_LinkedInputPose()
+	: InputPoseIndex(INDEX_NONE)
+{
+}
+
+void UAnimGraphNode_LinkedInputPose::PostEditChangeProperty(FPropertyChangedEvent& PropertyChangedEvent)
+{
+	Super::PostEditChangeProperty(PropertyChangedEvent);
+
+	if(PropertyChangedEvent.Property)
+	{
+		if( PropertyChangedEvent.Property->GetFName() == GET_MEMBER_NAME_CHECKED(UAnimGraphNode_LinkedInputPose, Inputs) ||
+			PropertyChangedEvent.Property->GetFName() == GET_MEMBER_NAME_CHECKED(UAnimGraphNode_LinkedInputPose, Node.Name) ||
+			PropertyChangedEvent.Property->GetFName() == GET_MEMBER_NAME_CHECKED(FAnimBlueprintFunctionPinInfo, Name) ||
+			PropertyChangedEvent.Property->GetFName() == GET_MEMBER_NAME_CHECKED(FAnimBlueprintFunctionPinInfo, Type))
+		{
+			HandleInputPinArrayChanged();
+			ReconstructNode();
+			FBlueprintEditorUtils::MarkBlueprintAsStructurallyModified(GetAnimBlueprint());
+			ReconstructNode();
+		}
+	}
+}
+
+FLinearColor UAnimGraphNode_LinkedInputPose::GetNodeTitleColor() const
+{
+	return GetDefault<UGraphEditorSettings>()->ResultNodeTitleColor;
+}
+
+FText UAnimGraphNode_LinkedInputPose::GetTooltipText() const
+{
+	return LOCTEXT("ToolTip", "Inputs to a sub-animation graph from a parent instance.");
+}
+
+FText UAnimGraphNode_LinkedInputPose::GetNodeTitle(ENodeTitleType::Type TitleType) const
+{
+	FText DefaultTitle = LOCTEXT("Title", "Input Pose");
+
+	if(TitleType != ENodeTitleType::FullTitle)
+	{
+		return DefaultTitle;
+	}
+	else
+	{
+		if(Node.Name != NAME_None)
+		{
+			FFormatNamedArguments Args;
+			Args.Add(TEXT("NodeTitle"), DefaultTitle);
+			Args.Add(TEXT("Name"), FText::FromName(Node.Name));
+			return FText::Format(LOCTEXT("TitleListFormatTagged", "{NodeTitle}\n{Name}"), Args);
+		}
+		else
+		{
+			return DefaultTitle;
+		}
+	}
+}
+
+bool UAnimGraphNode_LinkedInputPose::CanUserDeleteNode() const
+{
+	// Only allow linked input poses to be deleted if their parent graph is mutable
+	// Also allow anim graphs to delete these nodes even theough they are 'read-only'
+	return GetGraph()->bAllowDeletion || GetGraph()->GetFName() == UEdGraphSchema_K2::GN_AnimGraph;
+}
+
+bool UAnimGraphNode_LinkedInputPose::CanDuplicateNode() const
+{
+	return false;
+}
+
+template <typename Predicate>
+static FName CreateUniqueName(const FName& InBaseName, Predicate IsUnique)
+{
+	FName CurrentName = InBaseName;
+	int32 CurrentIndex = 0;
+
+	while (!IsUnique(CurrentName))
+	{
+		FString PossibleName = InBaseName.ToString() + TEXT("_") + FString::FromInt(CurrentIndex++);
+		CurrentName = FName(*PossibleName);
+	}
+
+	return CurrentName;
+}
+
+void UAnimGraphNode_LinkedInputPose::HandleInputPinArrayChanged()
+{
+	TArray<UAnimGraphNode_LinkedInputPose*> LinkedInputPoseNodes;
+	UAnimBlueprint* AnimBlueprint = GetAnimBlueprint();
+
+	for(UEdGraph* Graph : AnimBlueprint->FunctionGraphs)
+	{
+		if(Graph->Schema->IsChildOf(UAnimationGraphSchema::StaticClass()))
+		{
+			// Create a unique name for this new linked input pose
+			Graph->GetNodesOfClass(LinkedInputPoseNodes);
+		}
+	}
+
+	for(FAnimBlueprintFunctionPinInfo& Input : Inputs)
+	{
+		// New names are created empty, so assign a unique name
+		if(Input.Name == NAME_None)
+		{
+			Input.Name = CreateUniqueName(TEXT("InputParam"), [&LinkedInputPoseNodes](FName InName)
+			{
+				for(UAnimGraphNode_LinkedInputPose* LinkedInputPoseNode : LinkedInputPoseNodes)
+				{
+					for(const FAnimBlueprintFunctionPinInfo& Input : LinkedInputPoseNode->Inputs)
+					{
+						if(Input.Name == InName)
+						{
+							return false;
+						}
+					}
+				}
+				return true;
+			});
+
+			if(Input.Type.PinCategory == NAME_None)
+			{
+				IAssetEditorInstance* AssetEditor = GEditor->GetEditorSubsystem<UAssetEditorSubsystem>()->FindEditorForAsset(AnimBlueprint, false);
+				check(AssetEditor->GetEditorName() == "AnimationBlueprintEditor");
+				IAnimationBlueprintEditor* AnimationBlueprintEditor = static_cast<IAnimationBlueprintEditor*>(AssetEditor);
+				Input.Type = AnimationBlueprintEditor->GetLastGraphPinTypeUsed();
+			}
+		}
+	}
+	
+	bool bIsInterface = AnimBlueprint->BlueprintType == BPTYPE_Interface;
+	if(bIsInterface)
+	{
+		UAnimationGraphSchema::AutoArrangeInterfaceGraph(*GetGraph());
+	}
+}
+
+void UAnimGraphNode_LinkedInputPose::AllocatePinsInternal()
+{
+	// use member reference if valid
+	if (UFunction* Function = FunctionReference.ResolveMember<UFunction>(GetBlueprintClassFromNode()))
+	{
+		CreatePinsFromStubFunction(Function);
+	}
+
+	if(IsEditable())
+	{
+		// use user-defined pins
+		CreateUserDefinedPins();
+	}
+}
+
+void UAnimGraphNode_LinkedInputPose::AllocateDefaultPins()
+{
+	Super::AllocateDefaultPins();
+
+	AllocatePinsInternal();
+}
+
+void UAnimGraphNode_LinkedInputPose::ReallocatePinsDuringReconstruction(TArray<UEdGraphPin*>& OldPins)
+{
+	Super::ReallocatePinsDuringReconstruction(OldPins);
+
+	AllocatePinsInternal();
+}
+
+void UAnimGraphNode_LinkedInputPose::CreateUserDefinedPins()
+{
+	for(FAnimBlueprintFunctionPinInfo& PinInfo : Inputs)
+	{
+		UEdGraphPin* NewPin = CreatePin(EEdGraphPinDirection::EGPD_Output, PinInfo.Type, PinInfo.Name);
+		NewPin->PinFriendlyName = FText::FromName(PinInfo.Name);
+	}
+}
+
+void UAnimGraphNode_LinkedInputPose::CreatePinsFromStubFunction(const UFunction* Function)
+{
+	const UEdGraphSchema_K2* K2Schema = GetDefault<UEdGraphSchema_K2>();
+
+	IterateFunctionParameters([this, K2Schema, Function](const FName& InName, const FEdGraphPinType& InPinType)
+	{
+		if(!UAnimationGraphSchema::IsPosePin(InPinType))
+		{
+			UEdGraphPin* Pin = CreatePin(EGPD_Output, InPinType, InName);
+			K2Schema->SetPinAutogeneratedDefaultValueBasedOnType(Pin);
+			
+			UK2Node_CallFunction::GeneratePinTooltipFromFunction(*Pin, Function);
+		}
+	});
+}
+
+void UAnimGraphNode_LinkedInputPose::ConformInputPoseName()
+{
+	IterateFunctionParameters([this](const FName& InName, const FEdGraphPinType& InPinType)
+	{
+		if(UAnimationGraphSchema::IsPosePin(InPinType))
+		{
+			Node.Name = InName;
+		}
+	});
+}
+
+bool UAnimGraphNode_LinkedInputPose::ValidateAgainstFunctionReference() const
+{
+	bool bValid = false;
+
+	IterateFunctionParameters([this, &bValid](const FName& InName, const FEdGraphPinType& InPinType)
+	{
+		bValid = true;
+	});
+
+	return bValid;
+}
+
+void UAnimGraphNode_LinkedInputPose::PostPlacedNewNode()
+{
+	if(IsEditable())
+	{
+		TArray<UAnimGraphNode_LinkedInputPose*> LinkedInputPoseNodes;
+		UAnimBlueprint* AnimBlueprint = CastChecked<UAnimBlueprint>(GetGraph()->GetOuter());
+		for(UEdGraph* Graph : AnimBlueprint->FunctionGraphs)
+		{
+			if(Graph->Schema->IsChildOf(UAnimationGraphSchema::StaticClass()))
+			{
+				// Create a unique name for this new linked input pose
+				Graph->GetNodesOfClass(LinkedInputPoseNodes);
+			}
+		}
+
+		Node.Name = CreateUniqueName(FAnimNode_LinkedInputPose::DefaultInputPoseName, [this, &LinkedInputPoseNodes](const FName& InNameToCheck)
+		{
+			for(UAnimGraphNode_LinkedInputPose* LinkedInputPoseNode : LinkedInputPoseNodes)
+			{
+				if(LinkedInputPoseNode != this && LinkedInputPoseNode->Node.Name == InNameToCheck)
+				{
+					return false;
+				}
+			}
+
+			return true;
+		});
+
+		FTicker::GetCoreTicker().AddTicker(FTickerDelegate::CreateLambda([WeakThis = TWeakObjectPtr<UAnimGraphNode_LinkedInputPose>(this)](float InDeltaTime)
+		{
+			if(UAnimGraphNode_LinkedInputPose* LinkedInputPoseNode = WeakThis.Get())
+			{
+				// refresh the BP editor's details panel in case we are viewing the graph
+				IAssetEditorInstance* AssetEditor = GEditor->GetEditorSubsystem<UAssetEditorSubsystem>()->FindEditorForAsset(LinkedInputPoseNode->GetAnimBlueprint(), false);
+				check(AssetEditor->GetEditorName() == "AnimationBlueprintEditor");
+				IAnimationBlueprintEditor* AnimationBlueprintEditor = static_cast<IAnimationBlueprintEditor*>(AssetEditor);
+				AnimationBlueprintEditor->RefreshInspector();
+			}
+			return false;
+		}));
+	}
+}
+
+class SLinkedInputPoseNodeLabelWidget : public SCompoundWidget
+{
+public:
+	SLATE_BEGIN_ARGS(SLinkedInputPoseNodeLabelWidget) {}
+	SLATE_END_ARGS()
+
+	void Construct(const FArguments& InArgs, TSharedPtr<IPropertyHandle> InNamePropertyHandle, UAnimGraphNode_LinkedInputPose* InLinkedInputPoseNode)
+	{
+		NamePropertyHandle = InNamePropertyHandle;
+		WeakLinkedInputPoseNode = InLinkedInputPoseNode;
+
+		ChildSlot
+		[
+			SAssignNew(NameTextBox, SEditableTextBox)
+			.Font(IDetailLayoutBuilder::GetDetailFont())
+			.Text(this, &SLinkedInputPoseNodeLabelWidget::HandleGetNameText)
+			.OnTextChanged(this, &SLinkedInputPoseNodeLabelWidget::HandleTextChanged)
+			.OnTextCommitted(this, &SLinkedInputPoseNodeLabelWidget::HandleTextCommitted)
+		];
+	}
+
+	FText HandleGetNameText() const
+	{
+		return FText::FromName(WeakLinkedInputPoseNode->Node.Name);
+	}
+
+	bool IsNameValid(const FString& InNewName, FText& OutReason)
+	{
+		if(InNewName.Len() == 0)
+		{
+			OutReason = LOCTEXT("ZeroSizeLinkedInputPoseError", "A name must be specified.");
+			return false;
+		}
+		else if(InNewName.Equals(TEXT("None"), ESearchCase::IgnoreCase))
+		{
+			OutReason = LOCTEXT("LinkedInputPoseInvalidName", "This name is invalid.");
+			return false;
+		}
+		else
+		{
+			UAnimBlueprint* AnimBlueprint = CastChecked<UAnimBlueprint>(WeakLinkedInputPoseNode->GetGraph()->GetOuter());
+			for(UEdGraph* Graph : AnimBlueprint->FunctionGraphs)
+			{
+				if(Graph->Schema->IsChildOf(UAnimationGraphSchema::StaticClass()))
+				{
+					TArray<UAnimGraphNode_LinkedInputPose*> LinkedInputPoseNodes;
+					Graph->GetNodesOfClass(LinkedInputPoseNodes);
+
+					for(UAnimGraphNode_LinkedInputPose* LinkedInputPoseNode : LinkedInputPoseNodes)
+					{
+						if(LinkedInputPoseNode != WeakLinkedInputPoseNode.Get() && LinkedInputPoseNode->Node.Name.ToString().Equals(InNewName, ESearchCase::IgnoreCase))
+						{
+							OutReason = LOCTEXT("DuplicateLinkedInputPoseError", "This linked input pose name already exists in this blueprint.");
+							return false;
+						}
+					}
+				}
+			}
+
+			return true;
+		}
+	}
+
+	void HandleTextChanged(const FText& InNewText)
+	{
+		const FString NewTextAsString = InNewText.ToString();
+	
+		FText Reason;
+		if(!IsNameValid(NewTextAsString, Reason))
+		{
+			NameTextBox->SetError(Reason);
+		}
+		else
+		{
+			NameTextBox->SetError(FText::GetEmpty());
+		}
+	}
+
+	void HandleTextCommitted(const FText& InNewText, ETextCommit::Type InCommitType)
+	{
+		const FString NewTextAsString = InNewText.ToString();
+		FText Reason;
+		if(IsNameValid(NewTextAsString, Reason))
+		{
+			FName NewName = *InNewText.ToString();
+			NamePropertyHandle->SetValue(NewName);
+		}
+
+		NameTextBox->SetError(FText::GetEmpty());
+	}
+
+	TSharedPtr<SEditableTextBox> NameTextBox;
+	TSharedPtr<IPropertyHandle> NamePropertyHandle;
+	TWeakObjectPtr<UAnimGraphNode_LinkedInputPose> WeakLinkedInputPoseNode;
+};
+
+void UAnimGraphNode_LinkedInputPose::CustomizeDetails(IDetailLayoutBuilder& DetailBuilder)
+{
+	IDetailCategoryBuilder& InputsCategoryBuilder = DetailBuilder.EditCategory("Inputs");
+
+	TArray<TWeakObjectPtr<UObject>> OuterObjects;
+	DetailBuilder.GetObjectsBeingCustomized(OuterObjects);
+	if(OuterObjects.Num() != 1)
+	{
+		InputsCategoryBuilder.SetCategoryVisibility(false);
+		return;
+	}
+
+	// skip if we cant edit this node as it is an interface graph
+	UAnimGraphNode_LinkedInputPose* OuterNode = CastChecked<UAnimGraphNode_LinkedInputPose>(OuterObjects[0].Get());
+	if(!OuterNode->CanUserDeleteNode())
+	{
+		FText ReadOnlyWarning = LOCTEXT("ReadOnlyWarning", "This input pose is read-only and cannot be edited");
+
+		InputsCategoryBuilder.SetCategoryVisibility(false);
+
+		IDetailCategoryBuilder& WarningCategoryBuilder = DetailBuilder.EditCategory("InputPose", LOCTEXT("InputPoseCategory", "Input Pose"));
+		WarningCategoryBuilder.AddCustomRow(ReadOnlyWarning)
+			.WholeRowContent()
+			[
+				SNew(STextBlock)
+				.Text(ReadOnlyWarning)
+				.Font(IDetailLayoutBuilder::GetDetailFont())
+			];
+
+		return;
+	}
+
+	TSharedPtr<IPropertyHandle> NamePropertyHandle = DetailBuilder.GetProperty(GET_MEMBER_NAME_CHECKED(UAnimGraphNode_LinkedInputPose, Node.Name), GetClass());
+	InputsCategoryBuilder.AddProperty(NamePropertyHandle)
+	.OverrideResetToDefault(FResetToDefaultOverride::Hide())
+	.CustomWidget()
+	.NameContent()
+	[
+		NamePropertyHandle->CreatePropertyNameWidget()
+	]
+	.ValueContent()
+	[
+		MakeNameWidget(DetailBuilder)
+	];
+
+	InputsCategoryBuilder.AddProperty(GET_MEMBER_NAME_CHECKED(UAnimGraphNode_LinkedInputPose, Inputs), GetClass())
+		.ShouldAutoExpand(true);
+}
+
+TSharedRef<SWidget> UAnimGraphNode_LinkedInputPose::MakeNameWidget(IDetailLayoutBuilder& DetailBuilder)
+{
+	TSharedPtr<IPropertyHandle> NamePropertyHandle = DetailBuilder.GetProperty(GET_MEMBER_NAME_CHECKED(UAnimGraphNode_LinkedInputPose, Node.Name), GetClass());
+	return SNew(SLinkedInputPoseNodeLabelWidget, NamePropertyHandle, this);
+}
+
+bool UAnimGraphNode_LinkedInputPose::HasExternalDependencies(TArray<UStruct*>* OptionalOutput) const
+{
+	const UBlueprint* SourceBlueprint = GetBlueprint();
+
+	UClass* SourceClass = FunctionReference.GetMemberParentClass(GetBlueprintClassFromNode());
+	bool bResult = (SourceClass != nullptr) && (SourceClass->ClassGeneratedBy != SourceBlueprint);
+	if (bResult && OptionalOutput)
+	{
+		OptionalOutput->AddUnique(SourceClass);
+	}
+
+	const bool bSuperResult = Super::HasExternalDependencies(OptionalOutput);
+	return bSuperResult || bResult;
+}
+
+int32 UAnimGraphNode_LinkedInputPose::GetNumInputs() const
+{
+	if (UFunction* Function = FunctionReference.ResolveMember<UFunction>(GetBlueprintClassFromNode()))
+	{
+		// Count the inputs from parameters.
+		int32 NumParameters = 0;
+
+		IterateFunctionParameters([&NumParameters](const FName& InName, const FEdGraphPinType& InPinType)
+		{
+			if(!UAnimationGraphSchema::IsPosePin(InPinType))
+			{
+				NumParameters++;
+			}
+		});
+
+		return NumParameters;
+	}
+	else
+	{
+		return Inputs.Num();
+	}
+}
+
+void UAnimGraphNode_LinkedInputPose::PromoteFromInterfaceOverride()
+{
+	if (UFunction* Function = FunctionReference.ResolveMember<UFunction>(GetBlueprintClassFromNode()))
+	{
+		IterateFunctionParameters([this](const FName& InName, const FEdGraphPinType& InPinType)
+		{
+			if(!UAnimationGraphSchema::IsPosePin(InPinType))
+			{
+				Inputs.Emplace(InName, InPinType);	
+			}
+		});
+
+		// Remove the signature class now, that is not relevant.
+		FunctionReference.SetSelfMember(FunctionReference.GetMemberName());
+		InputPoseIndex = INDEX_NONE;
+	}
+}
+
+void UAnimGraphNode_LinkedInputPose::IterateFunctionParameters(TFunctionRef<void(const FName&, const FEdGraphPinType&)> InFunc) const
+{
+	if (UFunction* Function = FunctionReference.ResolveMember<UFunction>(GetBlueprintClassFromNode()))
+	{
+		const UEdGraphSchema_K2* K2Schema = GetDefault<UEdGraphSchema_K2>();
+
+		// if the generated class is not up to date, use the skeleton's class function to create pins:
+		Function = FBlueprintEditorUtils::GetMostUpToDateFunction(Function);
+
+		// We need to find all parameters AFTER the pose we are representing
+		int32 CurrentPoseIndex = 0;
+		UProperty* PoseParam = nullptr;
+		for (TFieldIterator<UProperty> PropIt(Function); PropIt && (PropIt->PropertyFlags & CPF_Parm); ++PropIt)
+		{
+			UProperty* Param = *PropIt;
+
+			const bool bIsFunctionInput = !Param->HasAnyPropertyFlags(CPF_OutParm) || Param->HasAnyPropertyFlags(CPF_ReferenceParm);
+
+			if (bIsFunctionInput)
+			{
+				FEdGraphPinType PinType;
+				if(K2Schema->ConvertPropertyToPinType(Param, PinType))
+				{
+					if(PoseParam == nullptr)
+					{
+						if(UAnimationGraphSchema::IsPosePin(PinType))
+						{
+							if(CurrentPoseIndex == InputPoseIndex)
+							{
+								PoseParam = Param;
+								InFunc(Param->GetFName(), PinType);
+							}
+							CurrentPoseIndex++;
+						}
+					}
+					else
+					{
+						if(UAnimationGraphSchema::IsPosePin(PinType))
+						{
+							// Found next pose param, so exit
+							break;
+						}
+						else
+						{
+							InFunc(Param->GetFName(), PinType);
+						}
+					}
+				}
+			}
+		}
+	}
+	else
+	{
+		// First call pose
+		InFunc(Node.Name, UAnimationGraphSchema::MakeLocalSpacePosePin());
+
+		// Then each input
+		for(const FAnimBlueprintFunctionPinInfo& PinInfo : Inputs)
+		{
+			InFunc(PinInfo.Name, PinInfo.Type);
+		}
+	}
+}
+
+#undef LOCTEXT_NAMESPACE
diff --git a/Engine/Source/Editor/AnimGraph/Private/AnimGraphNode_SubInput.cpp b/Engine/Source/Editor/AnimGraph/Private/AnimGraphNode_SubInput.cpp
deleted file mode 100644
diff --git a/Engine/Source/Editor/AnimGraph/Private/AnimGraphNode_SubInstanceBase.cpp b/Engine/Source/Editor/AnimGraph/Private/AnimGraphNode_SubInstanceBase.cpp
deleted file mode 100644
diff --git a/Engine/Source/Editor/AnimGraph/Private/AnimationGraphSchema.cpp b/Engine/Source/Editor/AnimGraph/Private/AnimationGraphSchema.cpp
--- a/Engine/Source/Editor/AnimGraph/Private/AnimationGraphSchema.cpp
+++ b/Engine/Source/Editor/AnimGraph/Private/AnimationGraphSchema.cpp
@@ -33,7 +33,7 @@
 #include "K2Node_Knot.h"
 #include "ScopedTransaction.h"
 #include "Animation/AnimMontage.h"
-#include "AnimGraphNode_SubInput.h"
+#include "AnimGraphNode_LinkedInputPose.h"
 
 #define LOCTEXT_NAMESPACE "AnimationGraphSchema"
 
@@ -287,18 +287,18 @@ void UAnimationGraphSchema::CreateFunctionGraphTerminators(UEdGraph& Graph, UCla
 					FEdGraphPinType PinType;
 					if(Schema->ConvertPropertyToPinType(Param, PinType))
 					{
-						// Create sub-input for each pose pin type
+						// Create linked input pose for each pose pin type
 						if(UAnimationGraphSchema::IsPosePin(PinType))
 						{
-							FGraphNodeCreator<UAnimGraphNode_SubInput> SubInputNodeCreator(Graph);
-							UAnimGraphNode_SubInput* SubInputNode = SubInputNodeCreator.CreateNode();
-							SubInputNode->FunctionReference.SetExternalMember(GraphName, Class, GraphGuid);
-							SubInputNode->Node.Name = Param->GetFName();
-							SubInputNode->InputPoseIndex = CurrentPoseIndex;
+							FGraphNodeCreator<UAnimGraphNode_LinkedInputPose> LinkedInputNodeCreator(Graph);
+							UAnimGraphNode_LinkedInputPose* LinkedInputNode = LinkedInputNodeCreator.CreateNode();
+							LinkedInputNode->FunctionReference.SetExternalMember(GraphName, Class, GraphGuid);
+							LinkedInputNode->Node.Name = Param->GetFName();
+							LinkedInputNode->InputPoseIndex = CurrentPoseIndex;
 
-							SubInputNode->ReconstructNode();
-							SetNodeMetaData(SubInputNode, FNodeMetadata::DefaultGraphNode);
-							SubInputNodeCreator.Finalize();
+							LinkedInputNode->ReconstructNode();
+							SetNodeMetaData(LinkedInputNode, FNodeMetadata::DefaultGraphNode);
+							LinkedInputNodeCreator.Finalize();
 
 							CurrentPoseIndex++;
 						}
@@ -636,8 +636,8 @@ void UAnimationGraphSchema::AutoArrangeInterfaceGraph(UEdGraph& Graph)
 	check(RootNodes.Num() == 1);
 	UAnimGraphNode_Root* Root = RootNodes[0];
 				
-	TArray<UAnimGraphNode_SubInput*> SubInputNodes;
-	Graph.GetNodesOfClass<UAnimGraphNode_SubInput>(SubInputNodes);
+	TArray<UAnimGraphNode_LinkedInputPose*> LinkedInputPoseNodes;
+	Graph.GetNodesOfClass<UAnimGraphNode_LinkedInputPose>(LinkedInputPoseNodes);
 
 	FBox2D RootBounds(FVector2D(Root->NodePosX, Root->NodePosY), FVector2D(Root->NodePosX + 130, Root->NodePosY + 200));
 
@@ -645,31 +645,31 @@ void UAnimationGraphSchema::AutoArrangeInterfaceGraph(UEdGraph& Graph)
 	float MaxWidth = 0.0f;
 	const int32 HeightPerProperty = 30;
 
-	for(UAnimGraphNode_SubInput* Node : SubInputNodes)
+	for(UAnimGraphNode_LinkedInputPose* Node : LinkedInputPoseNodes)
 	{
-		FBox2D SubInputBounds(
+		FBox2D LinkedInputPoseBounds(
 			FVector2D(Node->NodePosX, Node->NodePosY),
 			FVector2D(Node->NodePosX + 400, Node->NodePosY + 100 + (Node->GetNumInputs() * HeightPerProperty))
 		);
 
-		FVector2D BoundsSize = SubInputBounds.GetSize();
+		FVector2D BoundsSize = LinkedInputPoseBounds.GetSize();
 		TotalHeight += BoundsSize.Y + 10.0f;
 		MaxWidth = FMath::Max(BoundsSize.X, MaxWidth);
 	}
 
 	float NodeOffset = RootBounds.GetCenter().Y - (TotalHeight * 0.5f);
 	float NodePosX = RootBounds.Min.X - (MaxWidth + 100.0f);
-	for(UAnimGraphNode_SubInput* Node : SubInputNodes)
+	for(UAnimGraphNode_LinkedInputPose* Node : LinkedInputPoseNodes)
 	{
 		Node->NodePosX = NodePosX;
 		Node->NodePosY = NodeOffset;
 
-		FBox2D SubInputBounds(
+		FBox2D LinkedInputPoseBounds(
 			FVector2D(Node->NodePosX, Node->NodePosY),
 			FVector2D(Node->NodePosX + 400, Node->NodePosY + 100 + (Node->GetNumInputs() * HeightPerProperty))
 		);
 
-		NodeOffset += SubInputBounds.GetSize().Y + 10.0f;
+		NodeOffset += LinkedInputPoseBounds.GetSize().Y + 10.0f;
 	}
 }
 
@@ -684,18 +684,18 @@ void UAnimationGraphSchema::ConformAnimGraphToInterface(UBlueprint* InBlueprint,
 		check(RootNodes.Num() == 1);
 		RootNodes[0]->Node.Group = *InFunction->GetMetaDataText(TEXT("Category"), TEXT("UObjectCategory"), InFunction->GetFullGroupName(false)).ToString();
 
-		TArray<UAnimGraphNode_SubInput*> SubInputNodes;
-		InGraph.GetNodesOfClass<UAnimGraphNode_SubInput>(SubInputNodes);
+		TArray<UAnimGraphNode_LinkedInputPose*> LinkedInputPoseNodes;
+		InGraph.GetNodesOfClass<UAnimGraphNode_LinkedInputPose>(LinkedInputPoseNodes);
 
-		for(UAnimGraphNode_SubInput* SubInputNode : SubInputNodes)
+		for(UAnimGraphNode_LinkedInputPose* LinkedInputPoseNode : LinkedInputPoseNodes)
 		{
 			// Sync pose names in case they have changed
-			SubInputNode->ConformInputPoseName();
+			LinkedInputPoseNode->ConformInputPoseName();
 
-			// Clean up any old sub-inputs that no longer exist
-			if(!SubInputNode->ValidateAgainstFunctionReference())
+			// Clean up any old linked input poses that no longer exist
+			if(!LinkedInputPoseNode->ValidateAgainstFunctionReference())
 			{
-				InGraph.RemoveNode(SubInputNode);
+				InGraph.RemoveNode(LinkedInputPoseNode);
 			}
 		}
 
@@ -714,13 +714,13 @@ void UAnimationGraphSchema::ConformAnimGraphToInterface(UBlueprint* InBlueprint,
 				FEdGraphPinType PinType;
 				if(Schema->ConvertPropertyToPinType(Param, PinType))
 				{
-					// Create sub-input for each pose pin type
+					// Create linked input pose for each pose pin type
 					if(UAnimationGraphSchema::IsPosePin(PinType))
 					{
-						UAnimGraphNode_SubInput** MatchingNode = SubInputNodes.FindByPredicate(
-							[CurrentPoseIndex](UAnimGraphNode_SubInput* InSubInput)
+						UAnimGraphNode_LinkedInputPose** MatchingNode = LinkedInputPoseNodes.FindByPredicate(
+							[CurrentPoseIndex](UAnimGraphNode_LinkedInputPose* InLinkedInputPoseNode)
 							{
-								return InSubInput->InputPoseIndex == CurrentPoseIndex; 
+								return InLinkedInputPoseNode->InputPoseIndex == CurrentPoseIndex;
 							});
 
 						if(MatchingNode == nullptr)
@@ -732,19 +732,19 @@ void UAnimationGraphSchema::ConformAnimGraphToInterface(UBlueprint* InBlueprint,
 							FBlueprintEditorUtils::GetFunctionGuidFromClassByFieldName(FBlueprintEditorUtils::GetMostUpToDateClass(InterfaceClass), GraphName, GraphGuid);
 
 							// not found, add this node
-							FGraphNodeCreator<UAnimGraphNode_SubInput> SubInputNodeCreator(InGraph);
-							UAnimGraphNode_SubInput* SubInputNode = SubInputNodeCreator.CreateNode();
-							SubInputNode->FunctionReference.SetExternalMember(GraphName, InterfaceClass, GraphGuid);
-							SubInputNode->Node.Name = Param->GetFName();
-							SubInputNode->InputPoseIndex = CurrentPoseIndex;
-
-							SubInputNode->ReconstructNode();
-							SetNodeMetaData(SubInputNode, FNodeMetadata::DefaultGraphNode);
-							SubInputNodeCreator.Finalize();
-
-							FVector2D NewPosition = GetPositionForNewSubInputNode(InGraph);
-							SubInputNode->NodePosX = NewPosition.X;
-							SubInputNode->NodePosY = NewPosition.Y;
+							FGraphNodeCreator<UAnimGraphNode_LinkedInputPose> LinkedInputPoseNodeCreator(InGraph);
+							UAnimGraphNode_LinkedInputPose* LinkedInputPoseNode = LinkedInputPoseNodeCreator.CreateNode();
+							LinkedInputPoseNode->FunctionReference.SetExternalMember(GraphName, InterfaceClass, GraphGuid);
+							LinkedInputPoseNode->Node.Name = Param->GetFName();
+							LinkedInputPoseNode->InputPoseIndex = CurrentPoseIndex;
+
+							LinkedInputPoseNode->ReconstructNode();
+							SetNodeMetaData(LinkedInputPoseNode, FNodeMetadata::DefaultGraphNode);
+							LinkedInputPoseNodeCreator.Finalize();
+
+							FVector2D NewPosition = GetPositionForNewLinkedInputPoseNode(InGraph);
+							LinkedInputPoseNode->NodePosX = NewPosition.X;
+							LinkedInputPoseNode->NodePosY = NewPosition.Y;
 						}
 
 						CurrentPoseIndex++;
@@ -755,12 +755,12 @@ void UAnimationGraphSchema::ConformAnimGraphToInterface(UBlueprint* InBlueprint,
 	}
 }
 
-FVector2D UAnimationGraphSchema::GetPositionForNewSubInputNode(UEdGraph& InGraph)
+FVector2D UAnimationGraphSchema::GetPositionForNewLinkedInputPoseNode(UEdGraph& InGraph)
 {
-	TArray<UAnimGraphNode_SubInput*> SubInputNodes;
-	InGraph.GetNodesOfClass<UAnimGraphNode_SubInput>(SubInputNodes);
+	TArray<UAnimGraphNode_LinkedInputPose*> LinkedInputPoseNodes;
+	InGraph.GetNodesOfClass<UAnimGraphNode_LinkedInputPose>(LinkedInputPoseNodes);
 
-	if(SubInputNodes.Num() == 0)
+	if(LinkedInputPoseNodes.Num() == 0)
 	{
 		TArray<UAnimGraphNode_Base*> AllNodes;
 		InGraph.GetNodesOfClass<UAnimGraphNode_Base>(AllNodes);
@@ -783,30 +783,30 @@ FVector2D UAnimationGraphSchema::GetPositionForNewSubInputNode(UEdGraph& InGraph
 	{
 		const int32 HeightPerProperty = 30;
 
-		// Some existing sub-inputs. Insert below the bottom-most one.
-		UAnimGraphNode_SubInput* Node = SubInputNodes[0];
+		// Some existing linked input poses. Insert below the bottom-most one.
+		UAnimGraphNode_LinkedInputPose* Node = LinkedInputPoseNodes[0];
 
-		FBox2D BottomMostSubInputBounds(
+		FBox2D BottomMostLinkedInputPoseBounds(
 			FVector2D(Node->NodePosX, Node->NodePosY),
 			FVector2D(Node->NodePosX + 400, Node->NodePosY + 100 + (Node->GetNumInputs() * HeightPerProperty))
 		);
 
-		for(int32 SubInputIndex = 1; SubInputIndex < SubInputNodes.Num(); ++SubInputIndex)
+		for(int32 LinkedInputPoseIndex = 1; LinkedInputPoseIndex < LinkedInputPoseNodes.Num(); ++LinkedInputPoseIndex)
 		{
-			Node = SubInputNodes[SubInputIndex];
+			Node = LinkedInputPoseNodes[LinkedInputPoseIndex];
 
-			FBox2D SubInputBounds(
+			FBox2D LinkedInputPoseBounds(
 				FVector2D(Node->NodePosX, Node->NodePosY),
 				FVector2D(Node->NodePosX + 400, Node->NodePosY + 100 + (Node->GetNumInputs() * HeightPerProperty))
 			);
 
-			if(SubInputBounds.Min.Y > BottomMostSubInputBounds.Min.Y)
+			if(LinkedInputPoseBounds.Min.Y > BottomMostLinkedInputPoseBounds.Min.Y)
 			{
-				BottomMostSubInputBounds = SubInputBounds;
+				BottomMostLinkedInputPoseBounds = LinkedInputPoseBounds;
 			}
 		}
 
-		return FVector2D(BottomMostSubInputBounds.Min.X, BottomMostSubInputBounds.Max.Y + 10.0f);
+		return FVector2D(BottomMostLinkedInputPoseBounds.Min.X, BottomMostLinkedInputPoseBounds.Max.Y + 10.0f);
 	}
 }
 
diff --git a/Engine/Source/Editor/AnimationBlueprintEditor/Private/AnimationBlueprintEditor.cpp b/Engine/Source/Editor/AnimationBlueprintEditor/Private/AnimationBlueprintEditor.cpp
--- a/Engine/Source/Editor/AnimationBlueprintEditor/Private/AnimationBlueprintEditor.cpp
+++ b/Engine/Source/Editor/AnimationBlueprintEditor/Private/AnimationBlueprintEditor.cpp
@@ -1291,13 +1291,13 @@ void FAnimationBlueprintEditor::GetCustomDebugObjects(TArray<FCustomDebugObject>
 		if (PreviewAnimBlueprint)
 		{
 			EPreviewAnimationBlueprintApplicationMethod ApplicationMethod = AnimBlueprint->GetPreviewAnimationBlueprintApplicationMethod();
-			if(ApplicationMethod == EPreviewAnimationBlueprintApplicationMethod::OverlayLayer)
+			if(ApplicationMethod == EPreviewAnimationBlueprintApplicationMethod::LinkedLayers)
 			{
-				PreviewInstance = PreviewInstance->GetLayerSubInstanceByClass(AnimBlueprint->GeneratedClass.Get());
+				PreviewInstance = PreviewInstance->GetLinkedAnimLayerInstanceByClass(AnimBlueprint->GeneratedClass.Get());
 			}
-			else if(ApplicationMethod == EPreviewAnimationBlueprintApplicationMethod::SubInstance)
+			else if(ApplicationMethod == EPreviewAnimationBlueprintApplicationMethod::LinkedAnimGraph)
 			{
-				PreviewInstance = PreviewInstance->GetSubInstanceByTag(AnimBlueprint->GetPreviewAnimationBlueprintTag());
+				PreviewInstance = PreviewInstance->GetLinkedAnimGraphInstanceByTag(AnimBlueprint->GetPreviewAnimationBlueprintTag());
 			}
 		}
 
diff --git a/Engine/Source/Editor/Kismet/Private/BlueprintCompilationManager.cpp b/Engine/Source/Editor/Kismet/Private/BlueprintCompilationManager.cpp
--- a/Engine/Source/Editor/Kismet/Private/BlueprintCompilationManager.cpp
+++ b/Engine/Source/Editor/Kismet/Private/BlueprintCompilationManager.cpp
@@ -2514,7 +2514,7 @@ UClass* FBlueprintCompilationManagerImpl::FastGenerateSkeletonClass(UBlueprint*
 	{
 		CompilerContext.NewClass = Ret;
 		TGuardValue<bool> GuardAssignDelegateSignatureFunction( CompilerContext.bAssignDelegateSignatureFunction, true);
-		TGuardValue<bool> GuardGenerateSubInstanceVariables( CompilerContext.bGenerateSubInstanceVariables, true);
+		TGuardValue<bool> GuardGenerateLinkedAnimGraphVariables( CompilerContext.bGenerateLinkedAnimGraphVariables, true);
 		CompilerContext.CreateClassVariablesFromBlueprint();
 		CompilerContext.NewClass = OriginalNewClass;
 	}
diff --git a/Engine/Source/Editor/KismetCompiler/Private/AnimBlueprintCompiler.cpp b/Engine/Source/Editor/KismetCompiler/Private/AnimBlueprintCompiler.cpp
--- a/Engine/Source/Editor/KismetCompiler/Private/AnimBlueprintCompiler.cpp
+++ b/Engine/Source/Editor/KismetCompiler/Private/AnimBlueprintCompiler.cpp
@@ -44,16 +44,16 @@
 #include "K2Node_AnimGetter.h"
 #include "AnimGraphNode_StateMachine.h"
 #include "Animation/AnimNode_CustomProperty.h"
-#include "Animation/AnimNode_SubInstance.h"
-#include "AnimGraphNode_SubInstance.h"
+#include "Animation/AnimNode_LinkedAnimGraph.h"
+#include "AnimGraphNode_LinkedAnimGraph.h"
 #include "AnimGraphNode_Slot.h"
 #include "AnimationEditorUtils.h"
 
 #include "AnimBlueprintPostCompileValidation.h" 
-#include "AnimGraphNode_SubInput.h"
+#include "AnimGraphNode_LinkedInputPose.h"
 #include "K2Node_FunctionEntry.h"
 #include "K2Node_FunctionResult.h"
-#include "AnimGraphNode_Layer.h"
+#include "AnimGraphNode_LinkedAnimLayer.h"
 
 #define LOCTEXT_NAMESPACE "AnimBlueprintCompiler"
 
@@ -72,7 +72,7 @@ bool FAnimBlueprintCompilerContext::FEffectiveConstantRecord::Apply(UObject* Obj
 	uint8* StructPtr = nullptr;
 	uint8* PropertyPtr = nullptr;
 	
-	if(NodeVariableProperty->Struct->IsChildOf(FAnimNode_SubInstance::StaticStruct()))
+	if(NodeVariableProperty->Struct->IsChildOf(FAnimNode_LinkedAnimGraph::StaticStruct()))
 	{
 		PropertyPtr = ConstantProperty->ContainerPtrToValuePtr<uint8>(Object);
 	}
@@ -193,7 +193,7 @@ void FAnimBlueprintCompilerContext::CreateClassVariablesFromBlueprint()
 {
 	FKismetCompilerContext::CreateClassVariablesFromBlueprint();
 
-	if(bGenerateSubInstanceVariables)
+	if(bGenerateLinkedAnimGraphVariables)
 	{	
 		if(!bIsDerivedAnimBlueprint)
 		{
@@ -208,18 +208,18 @@ void FAnimBlueprintCompilerContext::CreateClassVariablesFromBlueprint()
 						ProcessCustomPropertyNode(CustomPropNode);
 					}
 
-					TArray<UAnimGraphNode_SubInstanceBase*> SubInstanceNodes;
-					InGraph->GetNodesOfClass(SubInstanceNodes);
-					for(UAnimGraphNode_SubInstanceBase* SubInstance : SubInstanceNodes)
+					TArray<UAnimGraphNode_LinkedAnimGraphBase*> LinkedAnimGraphNodes;
+					InGraph->GetNodesOfClass(LinkedAnimGraphNodes);
+					for(UAnimGraphNode_LinkedAnimGraphBase* LinkedAnimGraphNode : LinkedAnimGraphNodes)
 					{
-						ProcessSubInstance(SubInstance, false);
+						ProcessLinkedAnimGraph(LinkedAnimGraphNode, false);
 					}
 
-					TArray<UAnimGraphNode_SubInput*> SubInputNodes;
-					InGraph->GetNodesOfClass(SubInputNodes);
-					for(UAnimGraphNode_SubInput* SubInput : SubInputNodes)
+					TArray<UAnimGraphNode_LinkedInputPose*> LinkedInputPoseNodes;
+					InGraph->GetNodesOfClass(LinkedInputPoseNodes);
+					for(UAnimGraphNode_LinkedInputPose* LinkedInputPoseNode : LinkedInputPoseNodes)
 					{
-						ProcessSubInput(SubInput);
+						ProcessLinkedInputPose(LinkedInputPoseNode);
 					}
 				};
 
@@ -554,14 +554,14 @@ void FAnimBlueprintCompilerContext::ProcessAnimationNode(UAnimGraphNode_Base* Vi
 		// Handle a save/use cached pose linkage
 		ProcessUseCachedPose(UseCachedPose);
 	}
-	else if(UAnimGraphNode_SubInstanceBase* SubInstanceNode = Cast<UAnimGraphNode_SubInstanceBase>(VisualAnimNode))
+	else if(UAnimGraphNode_LinkedAnimGraphBase* LinkedAnimGraphNode = Cast<UAnimGraphNode_LinkedAnimGraphBase>(VisualAnimNode))
 	{
-		ProcessSubInstance(SubInstanceNode, true);
+		ProcessLinkedAnimGraph(LinkedAnimGraphNode, true);
 	}
-	else if (UAnimGraphNode_SubInput* SubInputNode = Cast<UAnimGraphNode_SubInput>(VisualAnimNode))
+	else if (UAnimGraphNode_LinkedInputPose* LinkedInputPoseNode = Cast<UAnimGraphNode_LinkedInputPose>(VisualAnimNode))
 	{
-		// Process sub-input nodes (input)
-		ProcessSubInput(SubInputNode);
+		// Process linked input pose nodes (input)
+		ProcessLinkedInputPose(LinkedInputPoseNode);
 	}
 	else if(UAnimGraphNode_Root* RootNode = Cast<UAnimGraphNode_Root>(VisualAnimNode))
 	{
@@ -642,11 +642,11 @@ void FAnimBlueprintCompilerContext::ProcessAnimationNode(UAnimGraphNode_Base* Vi
 		else
 		{
 			// The property source for our data, either the struct property for an anim node, or the
-			// owning anim instance if using a sub instance node.
+			// owning anim instance if using a linked instance node.
 			UProperty* SourcePinProperty = nullptr;
 			int32 SourceArrayIndex = INDEX_NONE;
 
-			// We have special handling below if we're targeting a subinstance instead of our own instance properties
+			// We have special handling below if we're targeting a linked instance instead of our own instance properties
 			UAnimGraphNode_CustomProperty* CustomPropertyNode = Cast<UAnimGraphNode_CustomProperty>(VisualAnimNode);
 
 			// Does this pin have an associated evaluation handler?
@@ -827,7 +827,7 @@ void FAnimBlueprintCompilerContext::ProcessCustomPropertyNode(UAnimGraphNode_Cus
 				FKismetCompilerUtilities::LinkAddedProperty(NewAnimBlueprintClass, NewProperty);
 
 				// Add mappings to the node
-				if (!bGenerateSubInstanceVariables)
+				if (!bGenerateLinkedAnimGraphVariables)
 				{
 					UClass* InstClass = CustomPropNode->GetTargetSkeletonClass();
 					if (UProperty* FoundProperty = FindField<UProperty>(InstClass, Pin->PinName))
@@ -840,21 +840,21 @@ void FAnimBlueprintCompilerContext::ProcessCustomPropertyNode(UAnimGraphNode_Cus
 	}
 }
 
-void FAnimBlueprintCompilerContext::ProcessSubInstance(UAnimGraphNode_SubInstanceBase* SubInstance, bool bCheckForCycles)
+void FAnimBlueprintCompilerContext::ProcessLinkedAnimGraph(UAnimGraphNode_LinkedAnimGraphBase* LinkedAnimGraph, bool bCheckForCycles)
 {
-	if(!SubInstance)
+	if(!LinkedAnimGraph)
 	{
 		return;
 	}
 
-	FAnimNode_SubInstance& RuntimeNode = *SubInstance->GetSubInstanceNode();
+	FAnimNode_LinkedAnimGraph& RuntimeNode = *LinkedAnimGraph->GetLinkedAnimGraphNode();
 
-	if(!bGenerateSubInstanceVariables)
+	if(!bGenerateLinkedAnimGraphVariables)
 	{
 		RuntimeNode.InputPoses.Empty();
 		RuntimeNode.InputPoseNames.Empty();
 	}
-	for(UEdGraphPin* Pin : SubInstance->Pins)
+	for(UEdGraphPin* Pin : LinkedAnimGraph->Pins)
 	{
 		if(!Pin->bOrphanedPin)
 		{
@@ -862,7 +862,7 @@ void FAnimBlueprintCompilerContext::ProcessSubInstance(UAnimGraphNode_SubInstanc
 			{
 				if(Pin->Direction == EGPD_Input)
 				{
-					if(!bGenerateSubInstanceVariables)
+					if(!bGenerateLinkedAnimGraphVariables)
 					{
 						RuntimeNode.InputPoses.AddDefaulted();
 						RuntimeNode.InputPoseNames.Add(Pin->GetFName());
@@ -879,7 +879,7 @@ void FAnimBlueprintCompilerContext::ProcessSubInstance(UAnimGraphNode_SubInstanc
 		NameToCountMap SlotNameToCountMap;
 		NameToCountMap StateMachineNameToCountMap;
 
-		GetDuplicatedSlotAndStateNames(SubInstance, StateMachineNameToCountMap, SlotNameToCountMap);
+		GetDuplicatedSlotAndStateNames(LinkedAnimGraph, StateMachineNameToCountMap, SlotNameToCountMap);
 
 
 		for(TPair<FName, int32>& Pair : SlotNameToCountMap)
@@ -887,7 +887,7 @@ void FAnimBlueprintCompilerContext::ProcessSubInstance(UAnimGraphNode_SubInstanc
 			if(Pair.Value > 1)
 			{
 				// Duplicated slot node
-				FString CompilerMessage = FString::Printf(TEXT("Slot name \"%s\" found across multiple instances. Slots are not visible outside of instances so duplicates or subinstances may not perform as expected."), *Pair.Key.ToString());
+				FString CompilerMessage = FString::Printf(TEXT("Slot name \"%s\" found across multiple instances. Slots are not visible outside of instances so duplicates or linked instances may not perform as expected."), *Pair.Key.ToString());
 				MessageLog.Warning(*CompilerMessage);
 			}
 		}
@@ -897,22 +897,22 @@ void FAnimBlueprintCompilerContext::ProcessSubInstance(UAnimGraphNode_SubInstanc
 			if(Pair.Value > 1)
 			{
 				// Duplicated slot node
-				FString CompilerMessage = FString::Printf(TEXT("State machine \"%s\" found across multiple instances. States are not visible outside of instances so duplicates or subinstances may not perform as expected."), *Pair.Key.ToString());
+				FString CompilerMessage = FString::Printf(TEXT("State machine \"%s\" found across multiple instances. States are not visible outside of instances so duplicates or linked instances may not perform as expected."), *Pair.Key.ToString());
 				MessageLog.Warning(*CompilerMessage);
 			}
 		}
 	}
 }
 
-void FAnimBlueprintCompilerContext::GetDuplicatedSlotAndStateNames(UAnimGraphNode_SubInstanceBase* InSubInstance, NameToCountMap& OutStateMachineNameToCountMap, NameToCountMap& OutSlotNameToCountMap)
+void FAnimBlueprintCompilerContext::GetDuplicatedSlotAndStateNames(UAnimGraphNode_LinkedAnimGraphBase* InLinkedAnimGraph, NameToCountMap& OutStateMachineNameToCountMap, NameToCountMap& OutSlotNameToCountMap)
 {
-	if(!InSubInstance)
+	if(!InLinkedAnimGraph)
 	{
 		// Nothing to inspect
 		return;
 	}
 
-	if(UClass* InstanceClass = InSubInstance->GetTargetClass())
+	if(UClass* InstanceClass = InLinkedAnimGraph->GetTargetClass())
 	{
 		UBlueprint* ClassBP = UBlueprint::GetBlueprintFromClass(InstanceClass);
 
@@ -923,11 +923,11 @@ void FAnimBlueprintCompilerContext::GetDuplicatedSlotAndStateNames(UAnimGraphNod
 		{
 			TArray<UAnimGraphNode_StateMachine*> StateMachineNodes;
 			TArray<UAnimGraphNode_Slot*> SlotNodes;
-			TArray<UAnimGraphNode_SubInstance*> SubInstanceNodes;
+			TArray<UAnimGraphNode_LinkedAnimGraph*> LinkedAnimGraphNodes;
 
 			Graph->GetNodesOfClass(StateMachineNodes);
 			Graph->GetNodesOfClass(SlotNodes);
-			Graph->GetNodesOfClass(SubInstanceNodes);
+			Graph->GetNodesOfClass(LinkedAnimGraphNodes);
 
 			for(UAnimGraphNode_StateMachine* StateMachineNode : StateMachineNodes)
 			{
@@ -942,9 +942,9 @@ void FAnimBlueprintCompilerContext::GetDuplicatedSlotAndStateNames(UAnimGraphNod
 				Count++;
 			}
 
-			for(UAnimGraphNode_SubInstance* SubInstanceNode : SubInstanceNodes)
+			for(UAnimGraphNode_LinkedAnimGraph* LinkedAnimGraphNode : LinkedAnimGraphNodes)
 			{
-				GetDuplicatedSlotAndStateNames(SubInstanceNode, OutStateMachineNameToCountMap, OutSlotNameToCountMap);
+				GetDuplicatedSlotAndStateNames(LinkedAnimGraphNode, OutStateMachineNameToCountMap, OutSlotNameToCountMap);
 			}
 		}
 	}
@@ -1007,8 +1007,8 @@ void FAnimBlueprintCompilerContext::PruneIsolatedAnimationNodes(const TArray<UAn
 	{
 		UAnimGraphNode_Base* Node = GraphNodes[NodeIndex];
 
-		// We cant prune sub-inputs as even if they are not linked to the root, they are needed for the dynamic link phase at runtime
-		if (!Visitor.VisitedNodes.Contains(Node) && !IsNodePure(Node) && !Node->IsA<UAnimGraphNode_SubInput>())
+		// We cant prune linked input poses as even if they are not linked to the root, they are needed for the dynamic link phase at runtime
+		if (!Visitor.VisitedNodes.Contains(Node) && !IsNodePure(Node) && !Node->IsA<UAnimGraphNode_LinkedInputPose>())
 		{
 			Node->BreakAllNodeLinks();
 			GraphNodes.RemoveAtSwap(NodeIndex);
@@ -1666,12 +1666,12 @@ void FAnimBlueprintCompilerContext::ProcessStateMachine(UAnimGraphNode_StateMach
 		GraphsToCheck.Add(StateNode->GetBoundGraph());
 		StateNode->GetBoundGraph()->GetAllChildrenGraphs(GraphsToCheck);
 
-		TArray<UAnimGraphNode_Layer*> LayerNodes;
+		TArray<UAnimGraphNode_LinkedAnimLayer*> LinkedAnimLayerNodes;
 		TArray<UAnimGraphNode_AssetPlayerBase*> AssetPlayerNodes;
 		for (UEdGraph* ChildGraph : GraphsToCheck)
 		{
 			ChildGraph->GetNodesOfClass(AssetPlayerNodes);
-			ChildGraph->GetNodesOfClass(LayerNodes);
+			ChildGraph->GetNodesOfClass(LinkedAnimLayerNodes);
 		}
 
 		for (UAnimGraphNode_AssetPlayerBase* Node : AssetPlayerNodes)
@@ -1682,7 +1682,7 @@ void FAnimBlueprintCompilerContext::ProcessStateMachine(UAnimGraphNode_StateMach
 			}
 		}
 
-		for (UAnimGraphNode_Layer* Node : LayerNodes)
+		for (UAnimGraphNode_LinkedAnimLayer* Node : LinkedAnimLayerNodes)
 		{
 			if (int32* IndexPtr = NewAnimBlueprintClass->AnimBlueprintDebugData.NodeGuidToIndexMap.Find(Node->NodeGuid))
 			{
@@ -1851,12 +1851,12 @@ void FAnimBlueprintCompilerContext::CopyTermDefaultsToDefaultObject(UObject* Def
 					NewRoot.Name = Cast<UAnimGraphNode_Root>(MessageLog.FindSourceObject(RootNode))->GetGraph()->GetFName();
 					TargetProperty->CopyCompleteValue(DestinationPtr, &NewRoot);
 				}
-				else if(UAnimGraphNode_SubInput* SubInputNode = ExactCast<UAnimGraphNode_SubInput>(VisualAnimNode))
+				else if(UAnimGraphNode_LinkedInputPose* LinkedInputPoseNode = ExactCast<UAnimGraphNode_LinkedInputPose>(VisualAnimNode))
 				{
-					// patch graph name into sub input nodes
-					FAnimNode_SubInput NewSubInput = *reinterpret_cast<FAnimNode_SubInput*>(SourcePtr);
-					NewSubInput.Graph = Cast<UAnimGraphNode_SubInput>(MessageLog.FindSourceObject(SubInputNode))->GetGraph()->GetFName();
-					TargetProperty->CopyCompleteValue(DestinationPtr, &NewSubInput);
+					// patch graph name into linked input pose nodes
+					FAnimNode_LinkedInputPose NewLinkedInputPose = *reinterpret_cast<FAnimNode_LinkedInputPose*>(SourcePtr);
+					NewLinkedInputPose.Graph = Cast<UAnimGraphNode_LinkedInputPose>(MessageLog.FindSourceObject(LinkedInputPoseNode))->GetGraph()->GetFName();
+					TargetProperty->CopyCompleteValue(DestinationPtr, &NewLinkedInputPose);
 				}
 				else
 				{
@@ -2172,38 +2172,38 @@ void FAnimBlueprintCompilerContext::ProcessOneFunctionGraph(UEdGraph* SourceGrap
 	}
 }
 
-void FAnimBlueprintCompilerContext::ProcessSubInput(UAnimGraphNode_SubInput* InSubInput)
+void FAnimBlueprintCompilerContext::ProcessLinkedInputPose(UAnimGraphNode_LinkedInputPose* InLinkedInputPose)
 {
-	InSubInput->IterateFunctionParameters([this, InSubInput](const FName& InName, FEdGraphPinType InPinType)
+	InLinkedInputPose->IterateFunctionParameters([this, InLinkedInputPose](const FName& InName, FEdGraphPinType InPinType)
 	{
 		if(!UAnimationGraphSchema::IsPosePin(InPinType))
 		{
-			// create properties for 'local' sub-input pins
-			UProperty* NewSubInputProperty = CreateVariable(InName, InPinType);
-			if(NewSubInputProperty)
+			// create properties for 'local' linked input pose pins
+			UProperty* NewLinkedInputPoseProperty = CreateVariable(InName, InPinType);
+			if(NewLinkedInputPoseProperty)
 			{
 				if(bIsFullCompile)
 				{
-					UEdGraphPin* Pin = InSubInput->FindPin(InName, EGPD_Output);
+					UEdGraphPin* Pin = InLinkedInputPose->FindPin(InName, EGPD_Output);
 					if(Pin)
 					{
 						// Create new node for property access
-						UK2Node_VariableGet* VariableGetNode = SpawnIntermediateNode<UK2Node_VariableGet>(InSubInput, InSubInput->GetGraph());
-						VariableGetNode->SetFromProperty(NewSubInputProperty, true);
+						UK2Node_VariableGet* VariableGetNode = SpawnIntermediateNode<UK2Node_VariableGet>(InLinkedInputPose, InLinkedInputPose->GetGraph());
+						VariableGetNode->SetFromProperty(NewLinkedInputPoseProperty, true);
 						VariableGetNode->AllocateDefaultPins();
 
 						// Add pin to generated variable association, used for pin watching
 						UEdGraphPin* TrueSourcePin = MessageLog.FindSourcePin(Pin);
 						if (TrueSourcePin)
 						{
-							NewClass->GetDebugData().RegisterClassPropertyAssociation(TrueSourcePin, NewSubInputProperty);
+							NewClass->GetDebugData().RegisterClassPropertyAssociation(TrueSourcePin, NewLinkedInputPoseProperty);
 						}
 
 						// link up to new node - note that this is not a FindPinChecked because if an interface changes without the
 						// implementing class being loaded, then its graphs will not be conformed until AFTER the skeleton class
 						// has been compiled, so the variable cannot be created. This also doesnt matter, as there wont be anything connected
 						// to the pin yet anyways.
-						UEdGraphPin* VariablePin = VariableGetNode->FindPin(NewSubInputProperty->GetFName());
+						UEdGraphPin* VariablePin = VariableGetNode->FindPin(NewLinkedInputPoseProperty->GetFName());
 						if(VariablePin)
 						{
 							TArray<UEdGraphPin*> Links = Pin->LinkedTo;
@@ -2277,8 +2277,8 @@ void FAnimBlueprintCompilerContext::CleanAndSanitizeClass(UBlueprintGeneratedCla
 	NewAnimBlueprintClass->AnimBlueprintFunctions.Empty();
 	NewAnimBlueprintClass->OrderedSavedPoseIndicesMap.Empty();
 	NewAnimBlueprintClass->AnimNodeProperties.Empty();
-	NewAnimBlueprintClass->SubInstanceNodeProperties.Empty();
-	NewAnimBlueprintClass->LayerNodeProperties.Empty();
+	NewAnimBlueprintClass->LinkedAnimGraphNodeProperties.Empty();
+	NewAnimBlueprintClass->LinkedAnimLayerNodeProperties.Empty();
 	NewAnimBlueprintClass->EvaluateGraphExposedInputs.Empty();
 	NewAnimBlueprintClass->GraphAssetPlayerInformation.Empty();
 
@@ -3230,12 +3230,12 @@ void FAnimBlueprintCompilerContext::CreateAnimGraphStubFunctions()
 				}
 			}
 
-			// Find the root and sub-input nodes
+			// Find the root and linked input pose nodes
 			TArray<UAnimGraphNode_Root*> Roots;
 			GraphToUseforSignature->GetNodesOfClass(Roots);
 
-			TArray<UAnimGraphNode_SubInput*> SubInputs;
-			GraphToUseforSignature->GetNodesOfClass(SubInputs);
+			TArray<UAnimGraphNode_LinkedInputPose*> LinkedInputPoseNodes;
+			GraphToUseforSignature->GetNodesOfClass(LinkedInputPoseNodes);
 
 			if(Roots.Num() > 0)
 			{
@@ -3253,17 +3253,17 @@ void FAnimBlueprintCompilerContext::CreateAnimGraphStubFunctions()
 				}
 
 				// Verify no duplicate inputs
-				for(UAnimGraphNode_SubInput* SubInput0 : SubInputs)
+				for(UAnimGraphNode_LinkedInputPose* LinkedInputPoseNode0 : LinkedInputPoseNodes)
 				{
-					for(UAnimGraphNode_SubInput* SubInput1 : SubInputs)	
+					for(UAnimGraphNode_LinkedInputPose* LinkedInputPoseNode1 : LinkedInputPoseNodes)
 					{
-						if(SubInput0 != SubInput1)
+						if(LinkedInputPoseNode0 != LinkedInputPoseNode1)
 						{
-							if(SubInput0->Node.Name == SubInput1->Node.Name)
+							if(LinkedInputPoseNode0->Node.Name == LinkedInputPoseNode1->Node.Name)
 							{
 								MessageLog.Error(
 									*LOCTEXT("DuplicateInputNode_Error", "Found duplicate input node @@ in graph @@").ToString(),
-									SubInput1,
+									LinkedInputPoseNode1,
 									InGraph
 								);
 							}
@@ -3285,24 +3285,24 @@ void FAnimBlueprintCompilerContext::CreateAnimGraphStubFunctions()
 				EntryNode->CustomGeneratedFunctionName = GraphToUseforSignature->GetFName();	// Note that the function generated from this temporary graph is undecorated
 				EntryNode->MetaData.Category = RootNode->Node.Group == NAME_None ? FText::GetEmpty() : FText::FromName(RootNode->Node.Group);
 
-				// Add sub-inputs as parameters
-				for(UAnimGraphNode_SubInput* SubInput : SubInputs)
+				// Add linked input poses as parameters
+				for(UAnimGraphNode_LinkedInputPose* LinkedInputPoseNode : LinkedInputPoseNodes)
 				{
-					// Add user defined pins for each sub-input pose
+					// Add user defined pins for each linked input pose
 					TSharedPtr<FUserPinInfo> PosePinInfo = MakeShared<FUserPinInfo>();
 					PosePinInfo->PinType = UAnimationGraphSchema::MakeLocalSpacePosePin();
-					PosePinInfo->PinName = SubInput->Node.Name;
+					PosePinInfo->PinName = LinkedInputPoseNode->Node.Name;
 					PosePinInfo->DesiredPinDirection = EGPD_Output;
 					EntryNode->UserDefinedPins.Add(PosePinInfo);
 
-					// Add user defined pins for each sub-input parameter
-					for(UEdGraphPin* SubInputPin : SubInput->Pins)
+					// Add user defined pins for each linked input pose parameter
+					for(UEdGraphPin* LinkedInputPoseNodePin : LinkedInputPoseNode->Pins)
 					{
-						if(!SubInputPin->bOrphanedPin && SubInputPin->Direction == EGPD_Output && !UAnimationGraphSchema::IsPosePin(SubInputPin->PinType))
+						if(!LinkedInputPoseNodePin->bOrphanedPin && LinkedInputPoseNodePin->Direction == EGPD_Output && !UAnimationGraphSchema::IsPosePin(LinkedInputPoseNodePin->PinType))
 						{
 							TSharedPtr<FUserPinInfo> ParameterPinInfo = MakeShared<FUserPinInfo>();
-							ParameterPinInfo->PinType = SubInputPin->PinType;
-							ParameterPinInfo->PinName = SubInputPin->PinName;
+							ParameterPinInfo->PinType = LinkedInputPoseNodePin->PinType;
+							ParameterPinInfo->PinName = LinkedInputPoseNodePin->PinName;
 							ParameterPinInfo->DesiredPinDirection = EGPD_Output;
 							EntryNode->UserDefinedPins.Add(ParameterPinInfo);
 						}
diff --git a/Engine/Source/Editor/KismetCompiler/Private/AnimBlueprintCompiler.h b/Engine/Source/Editor/KismetCompiler/Private/AnimBlueprintCompiler.h
--- a/Engine/Source/Editor/KismetCompiler/Private/AnimBlueprintCompiler.h
+++ b/Engine/Source/Editor/KismetCompiler/Private/AnimBlueprintCompiler.h
@@ -13,7 +13,7 @@ class UAnimGraphNode_SaveCachedPose;
 class UAnimGraphNode_StateMachineBase;
 class UAnimGraphNode_StateResult;
 class UAnimGraphNode_CustomProperty;
-class UAnimGraphNode_SubInstance;
+
 class UAnimGraphNode_UseCachedPose;
 class UAnimStateTransitionNode;
 class UK2Node_CallFunction;
@@ -23,9 +23,9 @@ class UK2Node_CallFunction;
 //
 class UAnimGraphNode_SaveCachedPose;
 class UAnimGraphNode_UseCachedPose;
-class UAnimGraphNode_SubInput;
-class UAnimGraphNode_SubInstance;
-class UAnimGraphNode_SubInstanceBase;
+class UAnimGraphNode_LinkedInputPose;
+class UAnimGraphNode_LinkedAnimGraphBase;
+class UAnimGraphNode_LinkedAnimGraph;
 class UAnimGraphNode_Root;
 
 class UStructProperty;
@@ -316,18 +316,18 @@ private:
 	// Compiles one custom property node
 	void ProcessCustomPropertyNode(UAnimGraphNode_CustomProperty* CustomPropNode);
 
-	// Compiles one sub instance node
-	void ProcessSubInstance(UAnimGraphNode_SubInstanceBase* SubInstance, bool bCheckForCycles);
+	// Compiles one linked anim graph node
+	void ProcessLinkedAnimGraph(UAnimGraphNode_LinkedAnimGraphBase* InLinkedAnimGraph, bool bCheckForCycles);
 
-	// Compiles one sub input
-	void ProcessSubInput(UAnimGraphNode_SubInput* InSubInput);
+	// Compiles one linked input pose
+	void ProcessLinkedInputPose(UAnimGraphNode_LinkedInputPose* InLinkedInputPose);
 
 	// Compiles one root node
 	void ProcessRoot(UAnimGraphNode_Root* Root);
 
-	// Traverses subinstance links looking for slot names and state machine names, returning their count in a name map
+	// Traverses linked anim graph links looking for slot names and state machine names, returning their count in a name map
 	typedef TMap<FName, int32> NameToCountMap;
-	void GetDuplicatedSlotAndStateNames(UAnimGraphNode_SubInstanceBase* InSubInstance, NameToCountMap& OutStateMachineNameToCountMap, NameToCountMap& OutSlotNameToCountMap);
+	void GetDuplicatedSlotAndStateNames(UAnimGraphNode_LinkedAnimGraphBase* InLinkedAnimGraph, NameToCountMap& OutStateMachineNameToCountMap, NameToCountMap& OutSlotNameToCountMap);
 
 	// Compiles an entire animation graph
 	void ProcessAllAnimationNodes();
diff --git a/Engine/Source/Editor/KismetCompiler/Private/KismetCompiler.cpp b/Engine/Source/Editor/KismetCompiler/Private/KismetCompiler.cpp
--- a/Engine/Source/Editor/KismetCompiler/Private/KismetCompiler.cpp
+++ b/Engine/Source/Editor/KismetCompiler/Private/KismetCompiler.cpp
@@ -139,7 +139,7 @@ FKismetCompilerContext::FKismetCompilerContext(UBlueprint* SourceSketch, FCompil
 	, OldLinker(nullptr)
 	, TargetClass(nullptr)
 	, bAssignDelegateSignatureFunction(false)
-	, bGenerateSubInstanceVariables(false)
+	, bGenerateLinkedAnimGraphVariables(false)
 {
 	MacroRowMaxHeight = 0;
 
diff --git a/Engine/Source/Editor/KismetCompiler/Public/KismetCompiler.h b/Engine/Source/Editor/KismetCompiler/Public/KismetCompiler.h
--- a/Engine/Source/Editor/KismetCompiler/Public/KismetCompiler.h
+++ b/Engine/Source/Editor/KismetCompiler/Public/KismetCompiler.h
@@ -128,8 +128,8 @@ public:
 	// CreateClassVariablesFromBlueprint:
 	bool bAssignDelegateSignatureFunction;
 
-	// Flag to trigger ProcessSubInstance in CreateClassVariablesFromBlueprint:
-	bool bGenerateSubInstanceVariables;
+	// Flag to trigger ProcessLinkedGraph in CreateClassVariablesFromBlueprint:
+	bool bGenerateLinkedAnimGraphVariables;
 
 	static FSimpleMulticastDelegate OnPreCompile;
 	static FSimpleMulticastDelegate OnPostCompile;
diff --git a/Engine/Source/Editor/Persona/Private/AnimationEditorPreviewScene.cpp b/Engine/Source/Editor/Persona/Private/AnimationEditorPreviewScene.cpp
--- a/Engine/Source/Editor/Persona/Private/AnimationEditorPreviewScene.cpp
+++ b/Engine/Source/Editor/Persona/Private/AnimationEditorPreviewScene.cpp
@@ -80,7 +80,7 @@ FAnimationEditorPreviewScene::FAnimationEditorPreviewScene(const ConstructionVal
 	{
 		PreviewSceneDescription->PreviewAnimationBlueprint = AnimBlueprint->GetPreviewAnimationBlueprint();
 		PreviewSceneDescription->ApplicationMethod = AnimBlueprint->GetPreviewAnimationBlueprintApplicationMethod();
-		PreviewSceneDescription->SubInstanceTag = AnimBlueprint->GetPreviewAnimationBlueprintTag();
+		PreviewSceneDescription->LinkedAnimGraphTag = AnimBlueprint->GetPreviewAnimationBlueprintTag();
 	}
 
 	// create a default additional mesh collection so we dont always have to create an asset to edit additional meshes
@@ -180,13 +180,13 @@ void FAnimationEditorPreviewScene::SetPreviewAnimationBlueprint(UAnimBlueprint*
 	{
 		SkeletalMeshComponent->SetAnimInstanceClass(InAnimBlueprint ? InAnimBlueprint->GeneratedClass : nullptr);
 		EPreviewAnimationBlueprintApplicationMethod ApplicationMethod = InOverlayOrSubAnimBlueprint->GetPreviewAnimationBlueprintApplicationMethod();
-		if(ApplicationMethod == EPreviewAnimationBlueprintApplicationMethod::OverlayLayer)
+		if(ApplicationMethod == EPreviewAnimationBlueprintApplicationMethod::LinkedLayers)
 		{
-			SkeletalMeshComponent->SetLayerOverlay(InOverlayOrSubAnimBlueprint ? InOverlayOrSubAnimBlueprint->GeneratedClass.Get() : nullptr);
+			SkeletalMeshComponent->LinkAnimClassLayers(InOverlayOrSubAnimBlueprint ? InOverlayOrSubAnimBlueprint->GeneratedClass.Get() : nullptr);
 		}
-		else if(ApplicationMethod == EPreviewAnimationBlueprintApplicationMethod::SubInstance)
+		else if(ApplicationMethod == EPreviewAnimationBlueprintApplicationMethod::LinkedAnimGraph)
 		{
-			SkeletalMeshComponent->SetSubInstanceClassByTag(InOverlayOrSubAnimBlueprint->GetPreviewAnimationBlueprintTag(), InOverlayOrSubAnimBlueprint ? InOverlayOrSubAnimBlueprint->GeneratedClass.Get() : nullptr);
+			SkeletalMeshComponent->LinkAnimGraphByTag(InOverlayOrSubAnimBlueprint->GetPreviewAnimationBlueprintTag(), InOverlayOrSubAnimBlueprint ? InOverlayOrSubAnimBlueprint->GeneratedClass.Get() : nullptr);
 		}
 	}
 	else
diff --git a/Engine/Source/Editor/Persona/Private/PersonaPreviewSceneDescription.h b/Engine/Source/Editor/Persona/Private/PersonaPreviewSceneDescription.h
--- a/Engine/Source/Editor/Persona/Private/PersonaPreviewSceneDescription.h
+++ b/Engine/Source/Editor/Persona/Private/PersonaPreviewSceneDescription.h
@@ -43,13 +43,13 @@ public:
 	UPROPERTY(EditAnywhere, Category = "Animation Blueprint", meta=(DisplayThumbnail=true))
 	TSoftObjectPtr<UAnimBlueprint> PreviewAnimationBlueprint;
 
-	/** The method by which a preview animation blueprint is applied, either as an overlay layer, or as a sub-instance */
+	/** The method by which a preview animation blueprint is applied, either as an overlay layer, or as a linked instance */
 	UPROPERTY(EditAnywhere, Category = "Animation Blueprint")
 	EPreviewAnimationBlueprintApplicationMethod ApplicationMethod;
 
-	/** The tag to use when applying a preview animation blueprint via SetSubInstanceClassByTag */
+	/** The tag to use when applying a preview animation blueprint via LinkAnimGraphByTag */
 	UPROPERTY(EditAnywhere, Category = "Animation Blueprint")
-	FName SubInstanceTag;
+	FName LinkedAnimGraphTag;
 
 	UPROPERTY(EditAnywhere, Category = "Additional Meshes")
 	TSoftObjectPtr<UDataAsset> AdditionalMeshes;
diff --git a/Engine/Source/Editor/Persona/Private/PersonaUtils.cpp b/Engine/Source/Editor/Persona/Private/PersonaUtils.cpp
--- a/Engine/Source/Editor/Persona/Private/PersonaUtils.cpp
+++ b/Engine/Source/Editor/Persona/Private/PersonaUtils.cpp
@@ -105,15 +105,15 @@ void SetObjectBeingDebugged(UAnimBlueprint* InAnimBlueprint, UAnimInstance* InAn
 	if (PreviewAnimBlueprint)
 	{
 		EPreviewAnimationBlueprintApplicationMethod ApplicationMethod = InAnimBlueprint->GetPreviewAnimationBlueprintApplicationMethod();
-		if(ApplicationMethod == EPreviewAnimationBlueprintApplicationMethod::OverlayLayer)
+		if(ApplicationMethod == EPreviewAnimationBlueprintApplicationMethod::LinkedLayers)
 		{
-			// Make sure the object being debugged is the overlay instance
-			InAnimBlueprint->SetObjectBeingDebugged(InAnimInstance->GetLayerSubInstanceByClass(InAnimBlueprint->GeneratedClass.Get()));
+			// Make sure the object being debugged is the linked layer instance
+			InAnimBlueprint->SetObjectBeingDebugged(InAnimInstance->GetLinkedAnimLayerInstanceByClass(InAnimBlueprint->GeneratedClass.Get()));
 		}
-		else if(ApplicationMethod == EPreviewAnimationBlueprintApplicationMethod::SubInstance)
+		else if(ApplicationMethod == EPreviewAnimationBlueprintApplicationMethod::LinkedAnimGraph)
 		{
-			// Make sure the object being debugged is the sub instance
-			InAnimBlueprint->SetObjectBeingDebugged(InAnimInstance->GetSubInstanceByTag(InAnimBlueprint->GetPreviewAnimationBlueprintTag()));
+			// Make sure the object being debugged is the linked instance
+			InAnimBlueprint->SetObjectBeingDebugged(InAnimInstance->GetLinkedAnimGraphInstanceByTag(InAnimBlueprint->GetPreviewAnimationBlueprintTag()));
 		}
 	}
 	else
diff --git a/Engine/Source/Editor/Persona/Private/PreviewSceneCustomizations.cpp b/Engine/Source/Editor/Persona/Private/PreviewSceneCustomizations.cpp
--- a/Engine/Source/Editor/Persona/Private/PreviewSceneCustomizations.cpp
+++ b/Engine/Source/Editor/Persona/Private/PreviewSceneCustomizations.cpp
@@ -219,7 +219,7 @@ void FPreviewSceneDescriptionCustomization::CustomizeDetails(IDetailLayoutBuilde
 		// Customize animation blueprint preview
 		TSharedRef<IPropertyHandle> PreviewAnimationBlueprintProperty = DetailBuilder.GetProperty(GET_MEMBER_NAME_CHECKED(UPersonaPreviewSceneDescription, PreviewAnimationBlueprint));
 		TSharedRef<IPropertyHandle> ApplicationMethodProperty = DetailBuilder.GetProperty(GET_MEMBER_NAME_CHECKED(UPersonaPreviewSceneDescription, ApplicationMethod));
-		TSharedRef<IPropertyHandle> SubInstanceTagProperty = DetailBuilder.GetProperty(GET_MEMBER_NAME_CHECKED(UPersonaPreviewSceneDescription, SubInstanceTag));
+		TSharedRef<IPropertyHandle> LinkedAnimGraphTagProperty = DetailBuilder.GetProperty(GET_MEMBER_NAME_CHECKED(UPersonaPreviewSceneDescription, LinkedAnimGraphTag));
 		
 		if (PersonaToolkit.Pin()->GetContext() == UAnimBlueprint::StaticClass()->GetFName())
 		{
@@ -269,33 +269,33 @@ void FPreviewSceneDescriptionCustomization::CustomizeDetails(IDetailLayoutBuilde
 				return PersonaPreviewSceneDescription->PreviewAnimationBlueprint.IsValid();
 			}));
 		
-			SubInstanceTagProperty->SetOnPropertyValueChanged(FSimpleDelegate::CreateLambda([this]()
+			LinkedAnimGraphTagProperty->SetOnPropertyValueChanged(FSimpleDelegate::CreateLambda([this]()
 			{
-				FScopedTransaction Transaction(LOCTEXT("SetAnimationBlueprintTag", "Set Sub Instance Tag"));
+				FScopedTransaction Transaction(LOCTEXT("SetAnimationBlueprintTag", "Set Linked Anim Graph Tag"));
 
 				TSharedPtr<IPersonaToolkit> PinnedPersonaToolkit = PersonaToolkit.Pin();
 				TSharedRef<FAnimationEditorPreviewScene> LocalPreviewScene = StaticCastSharedRef<FAnimationEditorPreviewScene>(PinnedPersonaToolkit->GetPreviewScene());
 				UPersonaPreviewSceneDescription* PersonaPreviewSceneDescription = LocalPreviewScene->GetPreviewSceneDescription();
-				PinnedPersonaToolkit->GetAnimBlueprint()->SetPreviewAnimationBlueprintTag(PersonaPreviewSceneDescription->SubInstanceTag);
+				PinnedPersonaToolkit->GetAnimBlueprint()->SetPreviewAnimationBlueprintTag(PersonaPreviewSceneDescription->LinkedAnimGraphTag);
 				LocalPreviewScene->SetPreviewAnimationBlueprint(PersonaPreviewSceneDescription->PreviewAnimationBlueprint.Get(), PinnedPersonaToolkit->GetAnimBlueprint());
 			}));
 
 			DetailBuilder.EditCategory("Animation Blueprint")
-			.AddProperty(SubInstanceTagProperty)
+			.AddProperty(LinkedAnimGraphTagProperty)
 			.IsEnabled(MakeAttributeLambda([this]()
 			{
 				TSharedPtr<IPersonaToolkit> PinnedPersonaToolkit = PersonaToolkit.Pin();
 				TSharedRef<FAnimationEditorPreviewScene> LocalPreviewScene = StaticCastSharedRef<FAnimationEditorPreviewScene>(PinnedPersonaToolkit->GetPreviewScene());
 				UPersonaPreviewSceneDescription* PersonaPreviewSceneDescription = LocalPreviewScene->GetPreviewSceneDescription();
 				
-				return PersonaPreviewSceneDescription->PreviewAnimationBlueprint.IsValid() && PersonaPreviewSceneDescription->ApplicationMethod == EPreviewAnimationBlueprintApplicationMethod::SubInstance;
+				return PersonaPreviewSceneDescription->PreviewAnimationBlueprint.IsValid() && PersonaPreviewSceneDescription->ApplicationMethod == EPreviewAnimationBlueprintApplicationMethod::LinkedAnimGraph;
 			}));
 		}
 		else
 		{
 			PreviewAnimationBlueprintProperty->MarkHiddenByCustomization();
 			ApplicationMethodProperty->MarkHiddenByCustomization();
-			SubInstanceTagProperty->MarkHiddenByCustomization();
+			LinkedAnimGraphTagProperty->MarkHiddenByCustomization();
 		}
 
 		// set the skeleton to use in our factory as we shouldn't be picking one here
diff --git a/Engine/Source/Editor/UnrealEd/Private/Kismet2/BlueprintEditorUtils.cpp b/Engine/Source/Editor/UnrealEd/Private/Kismet2/BlueprintEditorUtils.cpp
--- a/Engine/Source/Editor/UnrealEd/Private/Kismet2/BlueprintEditorUtils.cpp
+++ b/Engine/Source/Editor/UnrealEd/Private/Kismet2/BlueprintEditorUtils.cpp
@@ -116,7 +116,7 @@
 #include "BaseWidgetBlueprint.h"
 #include "Components/Widget.h"
 #include "UObject/UObjectThreadContext.h"
-#include "AnimGraphNode_SubInput.h"
+#include "AnimGraphNode_LinkedInputPose.h"
 #include "AnimGraphNode_Root.h"
 #include "Subsystems/AssetEditorSubsystem.h"
 
@@ -6244,12 +6244,12 @@ void FBlueprintEditorUtils::PromoteGraphFromInterfaceOverride(UBlueprint* InBlue
 		}
 	}
 
-	// Promote any animation sub-inputs
-	TArray<UAnimGraphNode_SubInput*> SubInputNodes;
-	InInterfaceGraph->GetNodesOfClass(SubInputNodes);
-	for (UAnimGraphNode_SubInput* SubInput : SubInputNodes)
+	// Promote any animation linked input poses
+	TArray<UAnimGraphNode_LinkedInputPose*> LinkedInputPoseNodes;
+	InInterfaceGraph->GetNodesOfClass(LinkedInputPoseNodes);
+	for (UAnimGraphNode_LinkedInputPose* LinkedInputPoseNode : LinkedInputPoseNodes)
 	{
-		SubInput->PromoteFromInterfaceOverride();
+		LinkedInputPoseNode->PromoteFromInterfaceOverride();
 	}
 }
 
diff --git a/Engine/Source/Runtime/Engine/Classes/Animation/AnimBlueprint.h b/Engine/Source/Runtime/Engine/Classes/Animation/AnimBlueprint.h
--- a/Engine/Source/Runtime/Engine/Classes/Animation/AnimBlueprint.h
+++ b/Engine/Source/Runtime/Engine/Classes/Animation/AnimBlueprint.h
@@ -61,11 +61,11 @@ struct FAnimParentNodeAssetOverride
 UENUM()
 enum class EPreviewAnimationBlueprintApplicationMethod : uint8
 {
-	/** Apply the preview animation blueprint using SetLayerOverlay */
-	OverlayLayer,
+	/** Apply the preview animation blueprint using LinkAnimClassLayers */
+	LinkedLayers,
 
-	/** Apply the preview animation blueprint using SetSubInstanceClassByTag */
-	SubInstance,
+	/** Apply the preview animation blueprint using SetLinkedAnimGraphByTag */
+	LinkedAnimGraph,
 };
 
 /**
@@ -217,11 +217,11 @@ private:
 	UPROPERTY(duplicatetransient, AssetRegistrySearchable)
 	TSoftObjectPtr<class UAnimBlueprint> PreviewAnimationBlueprint;
 
-	/** The method by which a preview animation blueprint is applied, either as an overlay layer, or as a sub-instance */
+	/** The method by which a preview animation blueprint is applied, either as an overlay layer, or as a linked instance */
 	UPROPERTY()
 	EPreviewAnimationBlueprintApplicationMethod PreviewAnimationBlueprintApplicationMethod;
 
-	/** The tag to use when applying a preview animation blueprint via SetSubInstanceClassByTag */
+	/** The tag to use when applying a preview animation blueprint via LinkAnimGraphByTag */
 	UPROPERTY()
 	FName PreviewAnimationBlueprintTag;
 #endif // WITH_EDITORONLY_DATA
diff --git a/Engine/Source/Runtime/Engine/Classes/Animation/AnimBlueprintGeneratedClass.h b/Engine/Source/Runtime/Engine/Classes/Animation/AnimBlueprintGeneratedClass.h
--- a/Engine/Source/Runtime/Engine/Classes/Animation/AnimBlueprintGeneratedClass.h
+++ b/Engine/Source/Runtime/Engine/Classes/Animation/AnimBlueprintGeneratedClass.h
@@ -199,8 +199,8 @@ class ENGINE_API UAnimBlueprintGeneratedClass : public UBlueprintGeneratedClass,
 
 	// The arrays of anim nodes; this is transient generated data (created during Link)
 	TArray<UStructProperty*> AnimNodeProperties;
-	TArray<UStructProperty*> SubInstanceNodeProperties;
-	TArray<UStructProperty*> LayerNodeProperties;
+	TArray<UStructProperty*> LinkedAnimGraphNodeProperties;
+	TArray<UStructProperty*> LinkedAnimLayerNodeProperties;
 	TArray<UStructProperty*> PreUpdateNodeProperties;
 	TArray<UStructProperty*> DynamicResetNodeProperties;
 	TArray<UStructProperty*> StateMachineNodeProperties;
@@ -224,8 +224,8 @@ public:
 	virtual USkeleton* GetTargetSkeleton() const override { return TargetSkeleton; }
 	virtual const TArray<FAnimNotifyEvent>& GetAnimNotifies() const override { return AnimNotifies; }
 	virtual const TArray<UStructProperty*>& GetAnimNodeProperties() const override { return AnimNodeProperties; }
-	virtual const TArray<UStructProperty*>& GetSubInstanceNodeProperties() const override { return SubInstanceNodeProperties; }
-	virtual const TArray<UStructProperty*>& GetLayerNodeProperties() const override { return LayerNodeProperties; }
+	virtual const TArray<UStructProperty*>& GetLinkedAnimGraphNodeProperties() const override { return LinkedAnimGraphNodeProperties; }
+	virtual const TArray<UStructProperty*>& GetLinkedAnimLayerNodeProperties() const override { return LinkedAnimLayerNodeProperties; }
 	virtual const TArray<UStructProperty*>& GetPreUpdateNodeProperties() const override { return PreUpdateNodeProperties; }
 	virtual const TArray<UStructProperty*>& GetDynamicResetNodeProperties() const override { return DynamicResetNodeProperties; }
 	virtual const TArray<UStructProperty*>& GetStateMachineNodeProperties() const override { return StateMachineNodeProperties; }
diff --git a/Engine/Source/Runtime/Engine/Classes/Animation/AnimClassData.h b/Engine/Source/Runtime/Engine/Classes/Animation/AnimClassData.h
--- a/Engine/Source/Runtime/Engine/Classes/Animation/AnimClassData.h
+++ b/Engine/Source/Runtime/Engine/Classes/Animation/AnimClassData.h
@@ -42,13 +42,13 @@ public:
 	UPROPERTY()
 	TArray<UStructProperty*> AnimNodeProperties;
 
-	// The array of sub instance nodes
+	// The array of linked anim graph nodes
 	UPROPERTY()
-	TArray<UStructProperty*> SubInstanceNodeProperties;
+	TArray<UStructProperty*> LinkedAnimGraphNodeProperties;
 
-	// The array of layer nodes
+	// The array of linked anim layer nodes
 	UPROPERTY()
-	TArray<UStructProperty*> LayerNodeProperties;
+	TArray<UStructProperty*> LinkedAnimLayerNodeProperties;
 
 	// Array of nodes that need a PreUpdate() call
 	UPROPERTY()
@@ -86,8 +86,8 @@ public:
 	virtual const TArray<FAnimBlueprintFunction>& GetAnimBlueprintFunctions() const override { return AnimBlueprintFunctions; }
 	virtual const TMap<FName, FCachedPoseIndices>& GetOrderedSavedPoseNodeIndicesMap() const override { return OrderedSavedPoseIndicesMap; }
 	virtual const TArray<UStructProperty*>& GetAnimNodeProperties() const override { return AnimNodeProperties; }
-	virtual const TArray<UStructProperty*>& GetSubInstanceNodeProperties() const override { return SubInstanceNodeProperties; }
-	virtual const TArray<UStructProperty*>& GetLayerNodeProperties() const override { return LayerNodeProperties; }
+	virtual const TArray<UStructProperty*>& GetLinkedAnimGraphNodeProperties() const override { return LinkedAnimGraphNodeProperties; }
+	virtual const TArray<UStructProperty*>& GetLinkedAnimLayerNodeProperties() const override { return LinkedAnimLayerNodeProperties; }
 	virtual const TArray<UStructProperty*>& GetPreUpdateNodeProperties() const override { return PreUpdateNodeProperties; }
 	virtual const TArray<UStructProperty*>& GetDynamicResetNodeProperties() const override { return DynamicResetNodeProperties; }
 	virtual const TArray<UStructProperty*>& GetStateMachineNodeProperties() const override { return StateMachineNodeProperties; }
@@ -107,8 +107,8 @@ public:
 		AnimBlueprintFunctions = AnimClass->GetAnimBlueprintFunctions();
 		OrderedSavedPoseIndicesMap = AnimClass->GetOrderedSavedPoseNodeIndicesMap();
 		AnimNodeProperties = AnimClass->GetAnimNodeProperties();
-		SubInstanceNodeProperties = AnimClass->GetSubInstanceNodeProperties();
-		LayerNodeProperties = AnimClass->GetLayerNodeProperties();
+		LinkedAnimGraphNodeProperties = AnimClass->GetLinkedAnimGraphNodeProperties();
+		LinkedAnimLayerNodeProperties = AnimClass->GetLinkedAnimLayerNodeProperties();
 		PreUpdateNodeProperties = AnimClass->GetPreUpdateNodeProperties();
 		DynamicResetNodeProperties = AnimClass->GetDynamicResetNodeProperties();
 		StateMachineNodeProperties = AnimClass->GetStateMachineNodeProperties();
diff --git a/Engine/Source/Runtime/Engine/Classes/Animation/AnimClassInterface.h b/Engine/Source/Runtime/Engine/Classes/Animation/AnimClassInterface.h
--- a/Engine/Source/Runtime/Engine/Classes/Animation/AnimClassInterface.h
+++ b/Engine/Source/Runtime/Engine/Classes/Animation/AnimClassInterface.h
@@ -119,8 +119,12 @@ public:
 	virtual const TArray<FBakedAnimationStateMachine>& GetBakedStateMachines() const = 0;
 	virtual const TArray<FAnimNotifyEvent>& GetAnimNotifies() const = 0;
 	virtual const TArray<UStructProperty*>& GetAnimNodeProperties() const = 0;
-	virtual const TArray<UStructProperty*>& GetSubInstanceNodeProperties() const = 0;
-	virtual const TArray<UStructProperty*>& GetLayerNodeProperties() const = 0;
+	UE_DEPRECATED(4.24, "Function has been renamed, please use GetLinkedAnimGraphNodeProperties")
+	virtual const TArray<UStructProperty*>& GetSubInstanceNodeProperties() const { return GetLinkedAnimGraphNodeProperties(); }
+	virtual const TArray<UStructProperty*>& GetLinkedAnimGraphNodeProperties() const = 0;
+	UE_DEPRECATED(4.24, "Function has been renamed, please use GetLinkedLayerNodeProperties")
+	virtual const TArray<UStructProperty*>& GetLayerNodeProperties() const { return GetLinkedAnimLayerNodeProperties(); }
+	virtual const TArray<UStructProperty*>& GetLinkedAnimLayerNodeProperties() const = 0;
 	virtual const TArray<UStructProperty*>& GetPreUpdateNodeProperties() const = 0;
 	virtual const TArray<UStructProperty*>& GetDynamicResetNodeProperties() const = 0;
 	virtual const TArray<UStructProperty*>& GetStateMachineNodeProperties() const = 0;
diff --git a/Engine/Source/Runtime/Engine/Classes/Animation/AnimInstance.h b/Engine/Source/Runtime/Engine/Classes/Animation/AnimInstance.h
--- a/Engine/Source/Runtime/Engine/Classes/Animation/AnimInstance.h
+++ b/Engine/Source/Runtime/Engine/Classes/Animation/AnimInstance.h
@@ -27,11 +27,11 @@ class UCanvas;
 struct FAnimInstanceProxy;
 struct FAnimNode_AssetPlayerBase;
 struct FAnimNode_StateMachine;
-struct FAnimNode_SubInput;
+struct FAnimNode_LinkedInputPose;
 struct FBakedAnimationStateMachine;
 class FCompilerResultsLog;
 struct FBoneContainer;
-struct FAnimNode_Layer;
+struct FAnimNode_LinkedAnimLayer;
 
 typedef TArray<FTransform> FTransformArrayA2;
 
@@ -699,56 +699,80 @@ public:
 	virtual void OnMontageInstanceStopped(FAnimMontageInstance & StoppedMontageInstance);
 	void ClearMontageInstanceReferences(FAnimMontageInstance& InMontageInstance);
 
+	UE_DEPRECATED(4.24, "Function renamed, please use GetLinkedInputPoseNode")
+	FAnimNode_LinkedInputPose* GetSubInputNode(FName InSubInput = NAME_None, FName InGraph = NAME_None) { return GetLinkedInputPoseNode(InSubInput, InGraph); }
+
 	/** 
-	 * Get a sub-input node by name, given a named graph.
-	 * @param	InSubInput	The name of the sub-input. If this is NAME_None, then we assume that the desired input is FAnimNode_SubInput::DefaultInputPoseName.
-	 * @param	InGraph		The name of the graph in which to find the sub-input. If this is NAME_None, then we assume that the desired graph is "AnimGraph", the default.
+	 * Get a linked input pose node by name, given a named graph.
+	 * @param	InSubInput	The name of the linked input pose. If this is NAME_None, then we assume that the desired input is FAnimNode_LinkedInputPose::DefaultInputPoseName.
+	 * @param	InGraph		The name of the graph in which to find the linked input. If this is NAME_None, then we assume that the desired graph is "AnimGraph", the default.
 	 */
-	FAnimNode_SubInput* GetSubInputNode(FName InSubInput = NAME_None, FName InGraph = NAME_None);
+	FAnimNode_LinkedInputPose* GetLinkedInputPoseNode(FName InSubInput = NAME_None, FName InGraph = NAME_None);
+
+	UE_DEPRECATED(4.24, "Function renamed, please use GetLinkedAnimGraphInstanceByTag")
+	UAnimInstance* GetSubInstanceByTag(FName InTag) const { return GetLinkedAnimGraphInstanceByTag(InTag); }
 
-	/** Runs through all nodes, attempting to find the first sub-instance by name/tag */
-	UFUNCTION(BlueprintPure, Category = "Sub-Instances")
-	UAnimInstance* GetSubInstanceByTag(FName InTag) const;
+	/** Runs through all nodes, attempting to find the first linked instance by name/tag */
+	UFUNCTION(BlueprintPure, Category = "Animation Blueprint Linking")
+	UAnimInstance* GetLinkedAnimGraphInstanceByTag(FName InTag) const;
 
-	/** Runs through all nodes, attempting to find all sub-instance that match the name/tag */
-	UFUNCTION(BlueprintPure, Category = "Sub-Instances")
+	UE_DEPRECATED(4.24, "Function renamed, please use GetLinkedAnimGraphInstancesByTag")
 	void GetSubInstancesByTag(FName InTag, TArray<UAnimInstance*>& OutSubInstances) const;
 
-	/** Runs through all nodes, attempting to find sub-instance by name/tag, then sets the class of each node if the tag matches */
-	UFUNCTION(BlueprintCallable, Category = "Sub-Instances")
-	void SetSubInstanceClassByTag(FName InTag, TSubclassOf<UAnimInstance> InClass);
+	/** Runs through all nodes, attempting to find all linked instances that match the name/tag */
+	UFUNCTION(BlueprintPure, Category = "Animation Blueprint Linking")
+	void GetLinkedAnimGraphInstancesByTag(FName InTag, TArray<UAnimInstance*>& OutLinkedInstances) const;
 
-	/** 
-	 * Runs through all layer nodes, attempting to find layer nodes that are implemented by the specified class, then sets up a sub instance of the class for each.
-	 * Allocates one sub instance to run each of the groups specified in the class, so state is shared. If a layer is not grouped (ie. NAME_None), then state is not shared
-	 * and a separate sub-instance is allocated for each layer node.
+	UE_DEPRECATED(4.24, "Function renamed, please use LinkAnimGraphByTag")
+	void SetSubInstanceClassByTag(FName InTag, TSubclassOf<UAnimInstance> InClass) { LinkAnimGraphByTag(InTag, InClass); }
+
+	/** Runs through all nodes, attempting to find a linked instance by name/tag, then sets the class of each node if the tag matches */
+	UFUNCTION(BlueprintCallable, Category = "Animation Blueprint Linking")
+	void LinkAnimGraphByTag(FName InTag, TSubclassOf<UAnimInstance> InClass);
+
+	UE_DEPRECATED(4.24, "Function renamed, please use LinkAnimClassLayers")
+	void SetLayerOverlay(TSubclassOf<UAnimInstance> InClass) { LinkAnimClassLayers(InClass); }
+
+	/**
+	 * Runs through all layer nodes, attempting to find layer nodes that are implemented by the specified class, then sets up a linked instance of the class for each.
+	 * Allocates one linked instance to run each of the groups specified in the class, so state is shared. If a layer is not grouped (ie. NAME_None), then state is not shared
+	 * and a separate linked instance is allocated for each layer node.
 	 * If InClass is null, then all layers are reset to their defaults.
 	 */
-	UFUNCTION(BlueprintCallable, Category = "Layers")
-	void SetLayerOverlay(TSubclassOf<UAnimInstance> InClass);
+	UFUNCTION(BlueprintCallable, Category = "Animation Blueprint Linking")
+	void LinkAnimClassLayers(TSubclassOf<UAnimInstance> InClass);
 
-	/** 
+	UE_DEPRECATED(4.24, "Function renamed, please use UnlinkAnimClassLayers")
+	void ClearLayerOverlay(TSubclassOf<UAnimInstance> InClass) { UnlinkAnimClassLayers(InClass); }
+
+	/**
 	 * Runs through all layer nodes, attempting to find layer nodes that are currently running the specified class, then resets each to its default value.
 	 * State sharing rules are as with SetLayerOverlay.
 	 * If InClass is null, does nothing.
 	 */
-	UFUNCTION(BlueprintCallable, Category = "Layers")
-	void ClearLayerOverlay(TSubclassOf<UAnimInstance> InClass);
+	UFUNCTION(BlueprintCallable, Category = "Animation Blueprint Linking")
+	void UnlinkAnimClassLayers(TSubclassOf<UAnimInstance> InClass);
+
+	UE_DEPRECATED(4.24, "Function renamed, please use GetLinkedAnimLayerInstanceByGroup")
+	UAnimInstance* GetLayerSubInstanceByGroup(FName InGroup) const { return GetLinkedAnimLayerInstanceByGroup(InGroup); }
+
+	/** Gets the layer linked instance corresponding to the specified group */
+	UFUNCTION(BlueprintPure, Category = "Animation Blueprint Linking")
+	UAnimInstance* GetLinkedAnimLayerInstanceByGroup(FName InGroup) const;
 
-	/** Gets the layer sub instance corresponding to the specified group */
-	UFUNCTION(BlueprintCallable, Category = "Layers")
-	UAnimInstance* GetLayerSubInstanceByGroup(FName InGroup) const;
+	UE_DEPRECATED(4.24, "Function renamed, please use GetLinkedAnimLayerInstanceByClass")
+	UAnimInstance* GetLayerSubInstanceByClass(TSubclassOf<UAnimInstance> InClass) const { return GetLinkedAnimLayerInstanceByClass(InClass); }
 
-	/** Gets the first layer sub instance corresponding to the specified class */
-	UFUNCTION(BlueprintCallable, Category = "Layers")
-	UAnimInstance* GetLayerSubInstanceByClass(TSubclassOf<UAnimInstance> InClass) const;
+	/** Gets the first layer linked instance corresponding to the specified class */
+	UFUNCTION(BlueprintPure, Category = "Animation Blueprint Linking")
+	UAnimInstance* GetLinkedAnimLayerInstanceByClass(TSubclassOf<UAnimInstance> InClass) const;
 
 	/** Sets up initial layer groupings */
 	void InitializeGroupedLayers();
 
 private:
 	/** Helper function to perform layer overlay actions (set, clear) */
-	void PerformLayerOverlayOperation(TSubclassOf<UAnimInstance> InClass, TFunctionRef<UClass*(UClass*, FAnimNode_Layer*)> InClassSelectorFunction);
+	void PerformLinkedLayerOverlayOperation(TSubclassOf<UAnimInstance> InClass, TFunctionRef<UClass*(UClass*, FAnimNode_LinkedAnimLayer*)> InClassSelectorFunction);
 
 protected:
 	/** Map between Active Montages and their FAnimMontageInstance */
@@ -1375,7 +1399,7 @@ protected:
 		return *static_cast<const T*>(AnimInstanceProxy);
 	}
 
-	friend struct FAnimNode_SubInstance;
+	friend struct FAnimNode_LinkedAnimGraph;
 	
 	/** Return whethere this AnimNotifyState should be triggered */
 	virtual bool ShouldTriggerAnimNotifyState(const UAnimNotifyState* AnimNotifyState) const { return true; }
diff --git a/Engine/Source/Runtime/Engine/Classes/Animation/AnimNode_Layer.h b/Engine/Source/Runtime/Engine/Classes/Animation/AnimNode_Layer.h
deleted file mode 100644
diff --git a/Engine/Source/Runtime/Engine/Classes/Animation/AnimNode_LinkedAnimGraph.h b/Engine/Source/Runtime/Engine/Classes/Animation/AnimNode_LinkedAnimGraph.h
new file mode 100644
--- /dev/null
+++ b/Engine/Source/Runtime/Engine/Classes/Animation/AnimNode_LinkedAnimGraph.h
@@ -0,0 +1,97 @@
+// Copyright 1998-2019 Epic Games, Inc. All Rights Reserved.
+
+#pragma once
+
+#include "CoreMinimal.h"
+#include "UObject/ObjectMacros.h"
+#include "Templates/SubclassOf.h"
+#include "Animation/AnimCurveTypes.h"
+#include "Animation/AnimNode_CustomProperty.h"
+#include "Animation/AnimInstance.h"
+#include "Animation/AnimLayerInterface.h"
+
+#include "AnimNode_LinkedAnimGraph.generated.h"
+
+struct FAnimInstanceProxy;
+class UUserDefinedStruct;
+struct FAnimBlueprintFunction;
+class IAnimClassInterface;
+
+USTRUCT(BlueprintInternalUseOnly)
+struct ENGINE_API FAnimNode_LinkedAnimGraph : public FAnimNode_CustomProperty
+{
+	GENERATED_BODY()
+
+public:
+
+	FAnimNode_LinkedAnimGraph();
+
+	/** 
+	 *  Input poses for the node, intentionally not accessible because if there's no input
+	 *  nodes in the target class we don't want to show these as pins
+	 */
+	UPROPERTY()
+	TArray<FPoseLink> InputPoses;
+
+	/** List of input pose names, 1-1 with pose links about, built by the compiler */
+	UPROPERTY()
+	TArray<FName> InputPoseNames;
+
+	/** The class spawned for this linked instance */
+	UPROPERTY(EditAnywhere, Category = Settings)
+	TSubclassOf<UAnimInstance> InstanceClass;
+
+	/** Optional tag used to identify this linked instance */
+	UPROPERTY(EditAnywhere, Category = Settings)
+	FName Tag;
+
+	// The root node of the dynamically-linked graph
+	FAnimNode_Base* LinkedRoot;
+
+	/** Dynamically set the anim class of this linked instance */
+	void SetAnimClass(TSubclassOf<UAnimInstance> InClass, const UAnimInstance* InOwningAnimInstance);
+
+	/** Get the function name we should be linking with when we call DynamicLink/Unlink */
+	virtual FName GetDynamicLinkFunctionName() const;
+
+	/** Get the dynamic link target */
+	virtual UAnimInstance* GetDynamicLinkTarget(UAnimInstance* InOwningAnimInstance) const;
+
+	// FAnimNode_Base interface
+	virtual void Initialize_AnyThread(const FAnimationInitializeContext& Context) override;
+	virtual void CacheBones_AnyThread(const FAnimationCacheBonesContext& Context) override;
+	virtual void Update_AnyThread(const FAnimationUpdateContext& Context) override;
+	virtual void Evaluate_AnyThread(FPoseContext& Output) override;
+	virtual void GatherDebugData(FNodeDebugData& DebugData) override;
+
+protected:
+	virtual void OnInitializeAnimInstance(const FAnimInstanceProxy* InProxy, const UAnimInstance* InAnimInstance) override;
+	virtual bool NeedsOnInitializeAnimInstance() const override { return true; }
+	// End of FAnimNode_Base interface
+
+	// Re-create the linked instances for this node
+	void ReinitializeLinkedAnimInstance(const UAnimInstance* InOwningAnimInstance, UAnimInstance* InNewAnimInstance = nullptr);
+
+	// Shutdown the currently running instance
+	void TeardownInstance();
+
+	// Get Target Class
+	virtual UClass* GetTargetClass() const override 
+	{
+		return *InstanceClass;
+	}
+
+	/** Link up pose links dynamically with linked instance */
+	void DynamicLink(UAnimInstance* InOwningAnimInstance);
+
+	/** Break any pose links dynamically with linked instance */
+	void DynamicUnlink(UAnimInstance* InOwningAnimInstance);
+
+	/** Helper function for finding function inputs when linking/unlinking */
+	int32 FindFunctionInputIndex(const FAnimBlueprintFunction& AnimBlueprintFunction, const FName& InInputName);
+
+	friend class UAnimInstance;
+};
+
+UE_DEPRECATED(4.24, "FAnimNode_SubInstance has been renamed to FAnimNode_LinkedAnimGraph")
+typedef FAnimNode_LinkedAnimGraph FAnimNode_SubInstance;
\ No newline at end of file
diff --git a/Engine/Source/Runtime/Engine/Classes/Animation/AnimNode_LinkedAnimLayer.h b/Engine/Source/Runtime/Engine/Classes/Animation/AnimNode_LinkedAnimLayer.h
new file mode 100644
--- /dev/null
+++ b/Engine/Source/Runtime/Engine/Classes/Animation/AnimNode_LinkedAnimLayer.h
@@ -0,0 +1,57 @@
+// Copyright 1998-2019 Epic Games, Inc. All Rights Reserved.
+
+#pragma once
+
+#include "CoreMinimal.h"
+#include "UObject/ObjectMacros.h"
+#include "Templates/SubclassOf.h"
+#include "Animation/AnimCurveTypes.h"
+#include "Animation/AnimNode_LinkedAnimGraph.h"
+#include "Animation/AnimInstance.h"
+#include "Animation/AnimLayerInterface.h"
+
+#include "AnimNode_LinkedAnimLayer.generated.h"
+
+struct FAnimInstanceProxy;
+class UUserDefinedStruct;
+struct FAnimBlueprintFunction;
+class IAnimClassInterface;
+
+USTRUCT(BlueprintInternalUseOnly)
+struct ENGINE_API FAnimNode_LinkedAnimLayer : public FAnimNode_LinkedAnimGraph
+{
+	GENERATED_BODY()
+
+public:
+	/** 
+	 * Optional interface. If this is set then this node will only accept (both statically and dynamically) anim instances that implement this interface.
+	 * If not set, then this is considered a 'self' layer. This value is set when Layer is changed in the details panel.
+	 */
+	UPROPERTY()
+	TSubclassOf<UAnimLayerInterface> Interface;
+
+	/** The layer in the interface to use */
+	UPROPERTY(EditAnywhere, Category = Settings)
+	FName Layer;
+
+	/** Set the layer's 'overlay' externally managed linked instance. */
+	void SetLinkedLayerInstance(const UAnimInstance* InOwningAnimInstance, UAnimInstance* InNewLinkedInstance);
+
+	/** FAnimNode_Base interface */
+	virtual void OnInitializeAnimInstance(const FAnimInstanceProxy* InProxy, const UAnimInstance* InAnimInstance) override;
+	virtual bool NeedsOnInitializeAnimInstance() const override { return true; }
+
+	/** FAnimNode_LinkedAnimGraph interface */
+	virtual FName GetDynamicLinkFunctionName() const override;
+	virtual UAnimInstance* GetDynamicLinkTarget(UAnimInstance* InOwningAnimInstance) const override;
+	virtual UClass* GetTargetClass() const override 
+	{
+		return *Interface;
+	}
+
+protected:
+	void InitializeSelfLayer(const UAnimInstance* SelfAnimInstance);
+};
+
+UE_DEPRECATED(4.24, "FAnimNode_Layer has been renamed to FAnimNode_LinkedAnimLayer")
+typedef FAnimNode_LinkedAnimLayer FAnimNode_Layer;
\ No newline at end of file
diff --git a/Engine/Source/Runtime/Engine/Classes/Animation/AnimNode_LinkedInputPose.h b/Engine/Source/Runtime/Engine/Classes/Animation/AnimNode_LinkedInputPose.h
new file mode 100644
--- /dev/null
+++ b/Engine/Source/Runtime/Engine/Classes/Animation/AnimNode_LinkedInputPose.h
@@ -0,0 +1,66 @@
+// Copyright 1998-2019 Epic Games, Inc. All Rights Reserved.
+
+#pragma once
+
+#include "CoreMinimal.h"
+#include "UObject/ObjectMacros.h"
+#include "Animation/AnimCurveTypes.h"
+#include "BonePose.h"
+#include "Animation/AnimNodeBase.h"
+#include "AnimNode_LinkedInputPose.generated.h"
+
+USTRUCT()
+struct ENGINE_API FAnimNode_LinkedInputPose : public FAnimNode_Base
+{
+	GENERATED_BODY()
+
+	/** The default name of this input pose */
+	static const FName DefaultInputPoseName;
+
+	FAnimNode_LinkedInputPose()
+		: Name(DefaultInputPoseName)
+		, Graph(NAME_None)
+		, InputProxy(nullptr)
+	{
+	}
+
+	/** The name of this linked input pose node's pose, used to identify the input of this graph. */
+	UPROPERTY(EditAnywhere, Category = "Inputs", meta = (NeverAsPin))
+	FName Name;
+
+	/** The graph that this linked input pose node is in, filled in by the compiler */
+	UPROPERTY()
+	FName Graph;
+
+	/** Input pose, optionally linked dynamically to another graph */
+	UPROPERTY()
+	FPoseLink InputPose;
+
+	/** 
+	 * If this linked input pose is not dynamically linked, this cached data will be populated by the calling 
+	 * linked instance node before this graph is processed.
+	 */
+	FCompactHeapPose CachedInputPose;
+	FBlendedHeapCurve CachedInputCurve;
+
+	// FAnimNode_Base interface
+	virtual void Initialize_AnyThread(const FAnimationInitializeContext& Context) override;
+	virtual void CacheBones_AnyThread(const FAnimationCacheBonesContext& Context) override;
+	virtual void Update_AnyThread(const FAnimationUpdateContext& Context) override;
+	virtual void Evaluate_AnyThread(FPoseContext& Output) override;
+	virtual void GatherDebugData(FNodeDebugData& DebugData) override;
+	// End of FAnimNode_Base interface
+
+	/** Called by linked instance nodes to dynamically link this to an outer graph */
+	void DynamicLink(FAnimInstanceProxy* InInputProxy, FPoseLinkBase* InPoseLink);
+
+	/** Called by linked instance nodes to dynamically unlink this to an outer graph */
+	void DynamicUnlink();
+
+private:
+	/** The proxy to use when getting inputs, set when dynamically linked */
+	FAnimInstanceProxy* InputProxy;
+};
+
+UE_DEPRECATED(4.24, "FAnimNode_SubInput has been renamed to FAnimNode_LinkedInputPose")
+typedef FAnimNode_LinkedInputPose FAnimNode_SubInput;
\ No newline at end of file
diff --git a/Engine/Source/Runtime/Engine/Classes/Animation/AnimNode_SubInput.h b/Engine/Source/Runtime/Engine/Classes/Animation/AnimNode_SubInput.h
deleted file mode 100644
diff --git a/Engine/Source/Runtime/Engine/Classes/Animation/AnimNode_SubInstance.h b/Engine/Source/Runtime/Engine/Classes/Animation/AnimNode_SubInstance.h
deleted file mode 100644
diff --git a/Engine/Source/Runtime/Engine/Classes/Components/SkeletalMeshComponent.h b/Engine/Source/Runtime/Engine/Classes/Components/SkeletalMeshComponent.h
--- a/Engine/Source/Runtime/Engine/Classes/Components/SkeletalMeshComponent.h
+++ b/Engine/Source/Runtime/Engine/Classes/Components/SkeletalMeshComponent.h
@@ -292,6 +292,10 @@ class ENGINE_API USkeletalMeshComponent : public USkinnedMeshComponent, public I
 	friend class FSkinnedMeshComponentRecreateRenderStateContext;
 	friend class FParallelAnimationCompletionTask;
 	friend class USkeletalMesh; 
+	friend class UAnimInstance;
+	friend struct FAnimNode_LinkedAnimGraph;
+	friend struct FAnimNode_LinkedAnimLayer;
+
 	/**
 	 * Animation 
 	 */
@@ -319,9 +323,12 @@ public:
 	UPROPERTY(transient, NonTransactional)
 	UAnimInstance* AnimScriptInstance;
 
-	/** Any running sub anim instances that need to be updates on the game thread */
+#if WITH_EDITORONLY_DATA
+	/** Any running linked anim instances */
+	UE_DEPRECATED(4.24, "Direct access to this property is deprecated and the array is no longer used. Storage is now in LinkedInstances. Please use GetLinkedAnimInstances() instead.")
 	UPROPERTY(transient)
 	TArray<UAnimInstance*> SubInstances;
+#endif
 
 	/** An instance created from the PostPhysicsBlueprint property of the skeletal mesh we're using,
 	 *  Runs after physics has been blended
@@ -375,6 +382,10 @@ public:
 	const TArray<FTransform>& GetCachedComponentSpaceTransforms() const;
 
 private:
+	/** Any running linked anim instances */
+	UPROPERTY(transient)
+	TArray<UAnimInstance*> LinkedInstances;
+
 	// Update Rate
 
 	/** Cached BoneSpaceTransforms for Update Rate optimization. */
@@ -573,19 +584,19 @@ public:
 	uint8 bEnableLineCheckWithBounds:1;
 
 	UPROPERTY(EditAnywhere, BlueprintReadWrite, Category = Clothing)
-    uint8 bUseBendingElements:1;
+	uint8 bUseBendingElements:1;
 	
 	UPROPERTY(EditAnywhere, BlueprintReadWrite, Category = Clothing)
-    uint8 bUseTetrahedralConstraints:1;
+	uint8 bUseTetrahedralConstraints:1;
 	
 	UPROPERTY(EditAnywhere, BlueprintReadWrite, Category = Clothing)
-    uint8 bUseThinShellVolumeConstraints:1;
+	uint8 bUseThinShellVolumeConstraints:1;
 	
 	UPROPERTY(EditAnywhere, BlueprintReadWrite, Category = Clothing)
-    uint8 bUseSelfCollisions:1;
+	uint8 bUseSelfCollisions:1;
 	
 	UPROPERTY(EditAnywhere, BlueprintReadWrite, Category = Clothing)
-    uint8 bUseContinuousCollisionDetection:1;
+	uint8 bUseContinuousCollisionDetection:1;
 
 	/** If true, propagates calls to ApplyAnimationCurvesToComponent for slave components, only needed if slave components do not tick themselves */
 	UPROPERTY(EditAnywhere, BlueprintReadWrite, Category = MasterPoseComponent)
@@ -653,22 +664,22 @@ public:
 	float ClothBlendWeight;
 
 	UPROPERTY(EditAnywhere, BlueprintReadWrite, Category = Clothing)
-    float EdgeStiffness;
-    
+	float EdgeStiffness;
+	
 	UPROPERTY(EditAnywhere, BlueprintReadWrite, Category = Clothing)
 	float BendingStiffness;
 	
 	UPROPERTY(EditAnywhere, BlueprintReadWrite, Category = Clothing)
-    float AreaStiffness;
+	float AreaStiffness;
 	
 	UPROPERTY(EditAnywhere, BlueprintReadWrite, Category = Clothing)
-    float VolumeStiffness;
+	float VolumeStiffness;
 	
 	UPROPERTY(EditAnywhere, BlueprintReadWrite, Category = Clothing)
-    float StrainLimitingStiffness;
+	float StrainLimitingStiffness;
 	
 	UPROPERTY(EditAnywhere, BlueprintReadWrite, Category = Clothing)
-    float ShapeTargetStiffness;
+	float ShapeTargetStiffness;
 
 private:
 
@@ -770,60 +781,92 @@ public:
 	UFUNCTION(BlueprintCallable, Category = "Components|SkeletalMesh", meta = (Keywords = "AnimBlueprint"))
 	UAnimInstance* GetPostProcessInstance() const;
 
-	UE_DEPRECATED(4.23, "This function is deprecated. Please use GetSubInstanceByTag")
-	UAnimInstance* GetSubInstanceByName(FName InTag) const { return GetSubInstanceByTag(InTag); }
+	/** Get the anim instances linked to the main AnimScriptInstance */
+PRAGMA_DISABLE_DEPRECATION_WARNINGS
+	const TArray<UAnimInstance*>& GetLinkedAnimInstances() const { return LinkedInstances; }
+PRAGMA_ENABLE_DEPRECATION_WARNINGS
+
+private:
+PRAGMA_DISABLE_DEPRECATION_WARNINGS
+	TArray<UAnimInstance*>& GetLinkedAnimInstances() { return LinkedInstances; }
+PRAGMA_ENABLE_DEPRECATION_WARNINGS
+
+public:
+	UE_DEPRECATED(4.23, "This function is deprecated. Please use GetLinkedAnimGraphInstanceByTag")
+	UAnimInstance* GetSubInstanceByName(FName InTag) const { return GetLinkedAnimGraphInstanceByTag(InTag); }
+
+	UE_DEPRECATED(4.24, "This function is deprecated. Please use GetLinkedAnimGraphInstanceByTag")
+	UAnimInstance* GetSubInstanceByTag(FName InTag) const { return GetLinkedAnimGraphInstanceByTag(InTag); }
 
 	/**
-	 * Returns the a tagged sub-instance node. If no sub instances are found or none are tagged with the
+	 * Returns the a tagged linked instance node. If no linked instances are found or none are tagged with the
 	 * supplied name, this will return NULL.
 	 */
-	UFUNCTION(BlueprintCallable, Category = "Components|SkeletalMesh|Sub-Instances", meta = (Keywords = "AnimBlueprint"))
-	UAnimInstance* GetSubInstanceByTag(FName InTag) const;
+	UFUNCTION(BlueprintPure, Category = "Components|SkeletalMesh|Animation Blueprint Linking", meta = (Keywords = "AnimBlueprint"))
+	UAnimInstance* GetLinkedAnimGraphInstanceByTag(FName InTag) const;
+
+	UE_DEPRECATED(4.24, "Function renamed, please use GetLinkedAnimGraphInstancesByTag")
+	void GetSubInstancesByTag(FName InTag, TArray<UAnimInstance*>& OutSubInstances) const { GetLinkedAnimGraphInstancesByTag(InTag, OutSubInstances); }
 
 	/**
-	 * Returns all tagged sub-instance nodes that match the tag.
+	 * Returns all tagged linked instance nodes that match the tag.
 	 */
-	UFUNCTION(BlueprintCallable, Category = "Components|SkeletalMesh|Sub-Instances", meta = (Keywords = "AnimBlueprint"))
-	void GetSubInstancesByTag(FName InTag, TArray<UAnimInstance*>& OutSubInstances) const;
+	UFUNCTION(BlueprintPure, Category = "Components|SkeletalMesh|Animation Blueprint Linking", meta = (Keywords = "AnimBlueprint"))
+	void GetLinkedAnimGraphInstancesByTag(FName InTag, TArray<UAnimInstance*>& OutLinkedInstances) const;
+
+	UE_DEPRECATED(4.24, "Function renamed, please use LinkAnimGraphByTag")
+	void SetSubInstanceClassByTag(FName InTag, TSubclassOf<UAnimInstance> InClass) { LinkAnimGraphByTag(InTag, InClass); }
 
-	/** Runs through all nodes, attempting to find sub-instance by name/tag, then sets the class of each node if the tag matches */
-	UFUNCTION(BlueprintCallable, Category = "Components|SkeletalMesh|Sub-Instances", meta = (Keywords = "AnimBlueprint"))
-	void SetSubInstanceClassByTag(FName InTag, TSubclassOf<UAnimInstance> InClass);
+	/** Runs through all nodes, attempting to find linked instance by name/tag, then sets the class of each node if the tag matches */
+	UFUNCTION(BlueprintCallable, Category = "Components|SkeletalMesh|Animation Blueprint Linking", meta = (Keywords = "AnimBlueprint"))
+	void LinkAnimGraphByTag(FName InTag, TSubclassOf<UAnimInstance> InClass);
+
+	UE_DEPRECATED(4.24, "Function renamed, please use LinkAnimClassLayers")
+	void SetLayerOverlay(TSubclassOf<UAnimInstance> InClass) { LinkAnimClassLayers(InClass); }
 
 	/** 
-	 * Runs through all layer nodes, attempting to find layer nodes that are implemented by the specified class, then sets up a sub instance of the class for each.
-	 * Allocates one sub instance to run each of the groups specified in the class, so state is shared. If a layer is not grouped (ie. NAME_None), then state is not shared
-	 * and a separate sub-instance is allocated for each layer node.
+	 * Runs through all layer nodes, attempting to find layer nodes that are implemented by the specified class, then sets up a linked instance of the class for each.
+	 * Allocates one linked instance to run each of the groups specified in the class, so state is shared. If a layer is not grouped (ie. NAME_None), then state is not shared
+	 * and a separate linked instance is allocated for each layer node.
 	 * If InClass is null, then all layers are reset to their defaults.
 	 */
-	UFUNCTION(BlueprintCallable, Category = "Components|SkeletalMesh|Layers")
-	void SetLayerOverlay(TSubclassOf<UAnimInstance> InClass);
+	UFUNCTION(BlueprintCallable, Category = "Components|SkeletalMesh|Animation Blueprint Linking")
+	void LinkAnimClassLayers(TSubclassOf<UAnimInstance> InClass);
+
+	UE_DEPRECATED(4.24, "Function renamed, please use UnlinkAnimClassLayers")
+	void ClearLayerOverlay(TSubclassOf<UAnimInstance> InClass) { UnlinkAnimClassLayers(InClass); }
 
 	/** 
 	 * Runs through all layer nodes, attempting to find layer nodes that are currently running the specified class, then resets each to its default value.
 	 * State sharing rules are as with SetLayerOverlay.
 	 * If InClass is null, does nothing.
 	 */
-	UFUNCTION(BlueprintCallable, Category = "Components|SkeletalMesh|Layers")
-	void ClearLayerOverlay(TSubclassOf<UAnimInstance> InClass);
+	UFUNCTION(BlueprintCallable, Category = "Components|SkeletalMesh|Animation Blueprint Linking")
+	void UnlinkAnimClassLayers(TSubclassOf<UAnimInstance> InClass);
 
-	/** Gets the layer sub instance corresponding to the specified group */
-	UFUNCTION(BlueprintCallable, Category = "Components|SkeletalMesh|Layers")
-	UAnimInstance* GetLayerSubInstanceByGroup(FName InGroup) const;
+	UE_DEPRECATED(4.24, "Function renamed, please use GetLinkedLayerInstanceByGroup")
+	UAnimInstance* GetLayerSubInstanceByGroup(FName InGroup) const { return GetLinkedAnimLayerInstanceByGroup(InGroup); }
 
-	/** Gets the first layer sub instance corresponding to the specified class */
-	UFUNCTION(BlueprintCallable, Category = "Components|SkeletalMesh|Layers")
-	UAnimInstance* GetLayerSubInstanceByClass(TSubclassOf<UAnimInstance> InClass) const;
+	/** Gets the layer linked instance corresponding to the specified group */
+	UFUNCTION(BlueprintPure, Category = "Components|SkeletalMesh|Animation Blueprint Linking")
+	UAnimInstance* GetLinkedAnimLayerInstanceByGroup(FName InGroup) const;
+
+	UE_DEPRECATED(4.24, "Function renamed, please use GetLinkedAnimLayerInstanceByClass")
+	UAnimInstance* GetLayerSubInstanceByClass(TSubclassOf<UAnimInstance> InClass) const { return GetLinkedAnimLayerInstanceByClass(InClass);  }
+
+	/** Gets the first layer linked instance corresponding to the specified class */
+	UFUNCTION(BlueprintPure, Category = "Components|SkeletalMesh|Animation Blueprint Linking")
+	UAnimInstance* GetLinkedAnimLayerInstanceByClass(TSubclassOf<UAnimInstance> InClass) const;
 
 	/** 
 	 * Returns whether there are any valid instances to run, currently this means whether we have
 	 * have an animation instance or a post process instance available to process.
 	 */
-	UFUNCTION(BlueprintCallable, Category = "Components|SkeletalMesh", meta = (Keywords = "AnimBlueprint"))
+	UFUNCTION(BlueprintPure, Category = "Components|SkeletalMesh", meta = (Keywords = "AnimBlueprint"))
 	bool HasValidAnimationInstance() const;
 
 	/**
-	 * Informs any active anim instances (main instance, sub instances, post instance) that a dynamics reset is required
+	 * Informs any active anim instances (main instance, linked instances, post instance) that a dynamics reset is required
 	 * for example if a teleport occurs.
 	 */
 	UFUNCTION(BlueprintCallable, Category = "Components|SkeletalMesh", meta = (Keywords = "Dynamics,Physics", UnsafeDuringActorConstruction = "true"))
@@ -833,7 +876,7 @@ public:
 	UFUNCTION(BlueprintCallable, Category = "Components|Animation", meta = (Keywords = "Animation"))
 	void SetAnimationMode(EAnimationMode::Type InAnimationMode);
 	
-	UFUNCTION(BlueprintCallable, Category = "Components|Animation", meta = (Keywords = "Animation"))
+	UFUNCTION(BlueprintPure, Category = "Components|Animation", meta = (Keywords = "Animation"))
 	EAnimationMode::Type GetAnimationMode() const;
 
 	/* Animation play functions
diff --git a/Engine/Source/Runtime/Engine/Private/AnimBlueprint.cpp b/Engine/Source/Runtime/Engine/Private/AnimBlueprint.cpp
--- a/Engine/Source/Runtime/Engine/Private/AnimBlueprint.cpp
+++ b/Engine/Source/Runtime/Engine/Private/AnimBlueprint.cpp
@@ -270,7 +270,7 @@ EPreviewAnimationBlueprintApplicationMethod UAnimBlueprint::GetPreviewAnimationB
 #if WITH_EDITORONLY_DATA
 	return PreviewAnimationBlueprintApplicationMethod; 
 #else
-	return EPreviewAnimationBlueprintApplicationMethod::OverlayLayer;
+	return EPreviewAnimationBlueprintApplicationMethod::LinkedLayers;
 #endif
 }
 
diff --git a/Engine/Source/Runtime/Engine/Private/Animation/AnimBlueprintGeneratedClass.cpp b/Engine/Source/Runtime/Engine/Private/Animation/AnimBlueprintGeneratedClass.cpp
--- a/Engine/Source/Runtime/Engine/Private/Animation/AnimBlueprintGeneratedClass.cpp
+++ b/Engine/Source/Runtime/Engine/Private/Animation/AnimBlueprintGeneratedClass.cpp
@@ -12,10 +12,10 @@
 #include "Animation/AnimInstance.h"
 #include "UObject/AnimObjectVersion.h"
 #include "UObject/ReleaseObjectVersion.h"
-#include "Animation/AnimNode_SubInstance.h"
+#include "Animation/AnimNode_LinkedAnimGraph.h"
 #include "Animation/AnimNode_Root.h"
-#include "Animation/AnimNode_SubInput.h"
-#include "Animation/AnimNode_Layer.h"
+#include "Animation/AnimNode_LinkedInputPose.h"
+#include "Animation/AnimNode_LinkedAnimLayer.h"
 #include "Animation/AnimNode_AssetPlayerBase.h"
 #include "Animation/AnimNode_StateMachine.h"
 
@@ -255,8 +255,8 @@ void UAnimBlueprintGeneratedClass::Link(FArchive& Ar, bool bRelinkExistingProper
 
 	// @TODO: Shouldn't be necessary to clear these, but currently the class gets linked twice during compilation
 	AnimNodeProperties.Empty();
-	SubInstanceNodeProperties.Empty();
-	LayerNodeProperties.Empty();
+	LinkedAnimGraphNodeProperties.Empty();
+	LinkedAnimLayerNodeProperties.Empty();
 	PreUpdateNodeProperties.Empty();
 	DynamicResetNodeProperties.Empty();
 	StateMachineNodeProperties.Empty();
@@ -274,13 +274,13 @@ void UAnimBlueprintGeneratedClass::Link(FArchive& Ar, bool bRelinkExistingProper
 		{
 			if (StructProp->Struct->IsChildOf(FAnimNode_Base::StaticStruct()))
 			{
-				if(StructProp->Struct == FAnimNode_SubInstance::StaticStruct())
+				if(StructProp->Struct == FAnimNode_LinkedAnimGraph::StaticStruct())
 				{
-					SubInstanceNodeProperties.Add(StructProp);
+					LinkedAnimGraphNodeProperties.Add(StructProp);
 				}
-				else if(StructProp->Struct == FAnimNode_Layer::StaticStruct())
+				else if(StructProp->Struct == FAnimNode_LinkedAnimLayer::StaticStruct())
 				{
-					LayerNodeProperties.Add(StructProp);
+					LinkedAnimLayerNodeProperties.Add(StructProp);
 				}
 				else if(StructProp->Struct == FAnimNode_StateMachine::StaticStruct())
 				{
@@ -469,14 +469,14 @@ void UAnimBlueprintGeneratedClass::LinkFunctionsToDefaultObjectNodes(UObject* De
 				FoundFunction->OutputPoseNodeProperty = StructProperty;
 			}
 		}
-		else if(StructProperty->Struct->IsChildOf(FAnimNode_SubInput::StaticStruct()))
+		else if(StructProperty->Struct->IsChildOf(FAnimNode_LinkedInputPose::StaticStruct()))
 		{
-			FAnimNode_SubInput* SubInputNode = StructProperty->ContainerPtrToValuePtr<FAnimNode_SubInput>(DefaultObject);
-			if(FAnimBlueprintFunction* FoundFunction = AnimBlueprintFunctions.FindByPredicate([SubInputNode](const FAnimBlueprintFunction& InFunction){ return InFunction.Name == SubInputNode->Graph; }))
+			FAnimNode_LinkedInputPose* LinkedInputPoseNode = StructProperty->ContainerPtrToValuePtr<FAnimNode_LinkedInputPose>(DefaultObject);
+			if(FAnimBlueprintFunction* FoundFunction = AnimBlueprintFunctions.FindByPredicate([LinkedInputPoseNode](const FAnimBlueprintFunction& InFunction){ return InFunction.Name == LinkedInputPoseNode->Graph; }))
 			{
 				for(int32 InputIndex = 0; InputIndex < FoundFunction->InputPoseNames.Num(); ++InputIndex)
 				{
-					if(FoundFunction->InputPoseNames[InputIndex] == SubInputNode->Name)
+					if(FoundFunction->InputPoseNames[InputIndex] == LinkedInputPoseNode->Name)
 					{
 						FoundFunction->InputPoseNodeIndices[InputIndex] = AnimNodeIndex;
 						FoundFunction->InputPoseNodeProperties[InputIndex] = StructProperty;
diff --git a/Engine/Source/Runtime/Engine/Private/Animation/AnimInstance.cpp b/Engine/Source/Runtime/Engine/Private/Animation/AnimInstance.cpp
--- a/Engine/Source/Runtime/Engine/Private/Animation/AnimInstance.cpp
+++ b/Engine/Source/Runtime/Engine/Private/Animation/AnimInstance.cpp
@@ -26,9 +26,9 @@
 #include "Rendering/SkeletalMeshRenderData.h"
 #include "Kismet/KismetSystemLibrary.h"
 #include "HAL/LowLevelMemTracker.h"
-#include "Animation/AnimNode_SubInstance.h"
-#include "Animation/AnimNode_SubInput.h"
-#include "Animation/AnimNode_Layer.h"
+#include "Animation/AnimNode_LinkedAnimGraph.h"
+#include "Animation/AnimNode_LinkedInputPose.h"
+#include "Animation/AnimNode_LinkedAnimLayer.h"
 
 /** Anim stats */
 
@@ -595,10 +595,10 @@ void UAnimInstance::ParallelUpdateAnimation()
 {
 	GetProxyOnAnyThread<FAnimInstanceProxy>().UpdateAnimation();
 
-	// Tick asset players for this and any sub instances we have
-	for(UAnimInstance* SubInstance : GetSkelMeshComponent()->SubInstances)
+	// Tick asset players for this and any linked instances we have
+	for(UAnimInstance* LinkedInstance : GetSkelMeshComponent()->GetLinkedAnimInstances())
 	{
-		SubInstance->GetProxyOnAnyThread<FAnimInstanceProxy>().TickAssetPlayerInstances();
+		LinkedInstance->GetProxyOnAnyThread<FAnimInstanceProxy>().TickAssetPlayerInstances();
 	}
 
 	GetProxyOnAnyThread<FAnimInstanceProxy>().TickAssetPlayerInstances();
@@ -2420,18 +2420,18 @@ void UAnimInstance::ClearMontageInstanceReferences(FAnimMontageInstance& InMonta
 	InMontageInstance.MontageSync_StopLeading();
 }
 
-FAnimNode_SubInput* UAnimInstance::GetSubInputNode(FName InSubInputName, FName InGraph)
+FAnimNode_LinkedInputPose* UAnimInstance::GetLinkedInputPoseNode(FName InLinkedInputName, FName InGraph)
 {
 	const FAnimInstanceProxy& Proxy = GetProxyOnAnyThread<FAnimInstanceProxy>();
-	if(InSubInputName == NAME_None && InGraph == NAME_None)
+	if(InLinkedInputName == NAME_None && InGraph == NAME_None)
 	{
-		return Proxy.DefaultSubInstanceInputNode;
+		return Proxy.DefaultLinkedInstanceInputNode;
 	}
 	else if (IAnimClassInterface* AnimBlueprintClass = IAnimClassInterface::GetFromClass(GetClass()))
 	{
-		if(InSubInputName == NAME_None)
+		if(InLinkedInputName == NAME_None)
 		{
-			InSubInputName = FAnimNode_SubInput::DefaultInputPoseName;
+			InLinkedInputName = FAnimNode_LinkedInputPose::DefaultInputPoseName;
 		}
 		if(InGraph == NAME_None)
 		{
@@ -2444,11 +2444,11 @@ FAnimNode_SubInput* UAnimInstance::GetSubInputNode(FName InSubInputName, FName I
 				check(AnimBlueprintFunction.InputPoseNames.Num() == AnimBlueprintFunction.InputPoseNodeProperties.Num());
 				for(int32 InputIndex = 0; InputIndex < AnimBlueprintFunction.InputPoseNames.Num(); ++InputIndex)
 				{
-					if(AnimBlueprintFunction.InputPoseNames[InputIndex] == InSubInputName && AnimBlueprintFunction.InputPoseNodeProperties[InputIndex] != nullptr)
+					if(AnimBlueprintFunction.InputPoseNames[InputIndex] == InLinkedInputName && AnimBlueprintFunction.InputPoseNodeProperties[InputIndex] != nullptr)
 					{
-						FAnimNode_SubInput* SubInput = AnimBlueprintFunction.InputPoseNodeProperties[InputIndex]->ContainerPtrToValuePtr<FAnimNode_SubInput>(this);
-						check(SubInput->Name == InSubInputName);
-						return SubInput;
+						FAnimNode_LinkedInputPose* LinkedInput = AnimBlueprintFunction.InputPoseNodeProperties[InputIndex]->ContainerPtrToValuePtr<FAnimNode_LinkedInputPose>(this);
+						check(LinkedInput->Name == InLinkedInputName);
+						return LinkedInput;
 					}
 				}
 			}
@@ -2458,17 +2458,17 @@ FAnimNode_SubInput* UAnimInstance::GetSubInputNode(FName InSubInputName, FName I
 	return nullptr;
 }
 
-UAnimInstance* UAnimInstance::GetSubInstanceByTag(FName InTag) const
+UAnimInstance* UAnimInstance::GetLinkedAnimGraphInstanceByTag(FName InTag) const
 {
 	if (IAnimClassInterface* AnimBlueprintClass = IAnimClassInterface::GetFromClass(GetClass()))
 	{
-		const TArray<UStructProperty*>& SubInstanceNodeProperties = AnimBlueprintClass->GetSubInstanceNodeProperties();
-		for(UStructProperty* SubInstanceNodeProperty : SubInstanceNodeProperties)
+		const TArray<UStructProperty*>& LinkedAnimGraphNodeProperties = AnimBlueprintClass->GetLinkedAnimGraphNodeProperties();
+		for(UStructProperty* LinkedAnimGraphNodeProperty : LinkedAnimGraphNodeProperties)
 		{
-			const FAnimNode_SubInstance* SubInstance = SubInstanceNodeProperty->ContainerPtrToValuePtr<FAnimNode_SubInstance>(this);
-			if(SubInstance && SubInstance->Tag == InTag)
+			const FAnimNode_LinkedAnimGraph* LinkedAnimGraph = LinkedAnimGraphNodeProperty->ContainerPtrToValuePtr<FAnimNode_LinkedAnimGraph>(this);
+			if(LinkedAnimGraph && LinkedAnimGraph->Tag == InTag)
 			{
-				return SubInstance->GetTargetInstance<UAnimInstance>();
+				return LinkedAnimGraph->GetTargetInstance<UAnimInstance>();
 			}
 		}
 	}
@@ -2476,39 +2476,39 @@ UAnimInstance* UAnimInstance::GetSubInstanceByTag(FName InTag) const
 	return nullptr;
 }
 
-void UAnimInstance::GetSubInstancesByTag(FName InTag, TArray<UAnimInstance*>& OutSubInstances) const
+void UAnimInstance::GetLinkedAnimGraphInstancesByTag(FName InTag, TArray<UAnimInstance*>& OutLinkedInstances) const
 {
 	if (IAnimClassInterface* AnimBlueprintClass = IAnimClassInterface::GetFromClass(GetClass()))
 	{
-		const TArray<UStructProperty*>& SubInstanceNodeProperties = AnimBlueprintClass->GetSubInstanceNodeProperties();
-		for(UStructProperty* SubInstanceNodeProperty : SubInstanceNodeProperties)
+		const TArray<UStructProperty*>& LinkedAnimGraphNodeProperties = AnimBlueprintClass->GetLinkedAnimGraphNodeProperties();
+		for (UStructProperty* LinkedAnimGraphNodeProperty : LinkedAnimGraphNodeProperties)
 		{
-			const FAnimNode_SubInstance* SubInstance = SubInstanceNodeProperty->ContainerPtrToValuePtr<FAnimNode_SubInstance>(this);
-			if(SubInstance && SubInstance->Tag == InTag)
+			const FAnimNode_LinkedAnimGraph* LinkedAnimGraph = LinkedAnimGraphNodeProperty->ContainerPtrToValuePtr<FAnimNode_LinkedAnimGraph>(this);
+			if(LinkedAnimGraph && LinkedAnimGraph->Tag == InTag)
 			{
-				OutSubInstances.Add(SubInstance->GetTargetInstance<UAnimInstance>());
+				OutLinkedInstances.Add(LinkedAnimGraph->GetTargetInstance<UAnimInstance>());
 			}
 		}
 	}
 }
 
-void UAnimInstance::SetSubInstanceClassByTag(FName InTag, TSubclassOf<UAnimInstance> InClass)
+void UAnimInstance::LinkAnimGraphByTag(FName InTag, TSubclassOf<UAnimInstance> InClass)
 {
 	if (IAnimClassInterface* AnimBlueprintClass = IAnimClassInterface::GetFromClass(GetClass()))
 	{
-		const TArray<UStructProperty*>& SubInstanceNodeProperties = AnimBlueprintClass->GetSubInstanceNodeProperties();
-		for(UStructProperty* SubInstanceNodeProperty : SubInstanceNodeProperties)
+		const TArray<UStructProperty*>& LinkedAnimGraphNodeProperties = AnimBlueprintClass->GetLinkedAnimGraphNodeProperties();
+		for (UStructProperty* LinkedAnimGraphNodeProperty : LinkedAnimGraphNodeProperties)
 		{
-			FAnimNode_SubInstance* SubInstance = SubInstanceNodeProperty->ContainerPtrToValuePtr<FAnimNode_SubInstance>(this);
-			if(SubInstance && SubInstance->Tag == InTag)
+			FAnimNode_LinkedAnimGraph* LinkedAnimGraph = LinkedAnimGraphNodeProperty->ContainerPtrToValuePtr<FAnimNode_LinkedAnimGraph>(this);
+			if(LinkedAnimGraph && LinkedAnimGraph->Tag == InTag)
 			{
-				SubInstance->SetAnimClass(InClass, this);
+				LinkedAnimGraph->SetAnimClass(InClass, this);
 			}
 		}
 	}
 }
 
-void UAnimInstance::PerformLayerOverlayOperation(TSubclassOf<UAnimInstance> InClass, TFunctionRef<UClass*(UClass* InClass, FAnimNode_Layer*)> InClassSelectorFunction)
+void UAnimInstance::PerformLinkedLayerOverlayOperation(TSubclassOf<UAnimInstance> InClass, TFunctionRef<UClass*(UClass* InClass, FAnimNode_LinkedAnimLayer*)> InClassSelectorFunction)
 {
 	if (IAnimClassInterface* AnimBlueprintClass = IAnimClassInterface::GetFromClass(GetClass()))
 	{
@@ -2516,26 +2516,25 @@ void UAnimInstance::PerformLayerOverlayOperation(TSubclassOf<UAnimInstance> InCl
 		if(NewClass)
 		{
 			// Verify target skeleton matches at runtime
-			IAnimClassInterface* SubAnimBlueprintClass = IAnimClassInterface::GetFromClass(NewClass);
-			USkeleton* SubSkeleton = SubAnimBlueprintClass->GetTargetSkeleton();
+			IAnimClassInterface* LinkedAnimBlueprintClass = IAnimClassInterface::GetFromClass(NewClass);
+			USkeleton* LinkedSkeleton = LinkedAnimBlueprintClass->GetTargetSkeleton();
 			USkeleton* OuterSkeleton = AnimBlueprintClass->GetTargetSkeleton();
-			if(SubSkeleton != OuterSkeleton)
+			if(LinkedSkeleton != OuterSkeleton)
 			{
-				UE_LOG(LogAnimation, Warning, TEXT("Setting layer overlay: Sub instance class has a mismatched target skeleton. Expected %s, found %s."), OuterSkeleton ? *OuterSkeleton->GetName() : TEXT("null"), SubSkeleton ? *SubSkeleton->GetName() : TEXT("null"));
+				UE_LOG(LogAnimation, Warning, TEXT("Linking layer: Linked instance class has a mismatched target skeleton. Expected %s, found %s."), OuterSkeleton ? *OuterSkeleton->GetName() : TEXT("null"), LinkedSkeleton ? *LinkedSkeleton->GetName() : TEXT("null"));
 				return;
 			}
 		}
 
-		// Make sure we have valid objects as initialization can route back out of sub-instances into this outer graph
+		// Make sure we have valid objects as initialization can route back out of linked instances into this outer graph
 		GetProxyOnAnyThread<FAnimInstanceProxy>().InitializeObjects(this);
 
 		// Map of group name->nodes to run under that group instance
-		TMap<FName, TArray<FAnimNode_Layer*, TInlineAllocator<4>>, TInlineSetAllocator<4>> LayerNodesToSet;
+		TMap<FName, TArray<FAnimNode_LinkedAnimLayer*, TInlineAllocator<4>>, TInlineSetAllocator<4>> LayerNodesToSet;
 
-		const TArray<UStructProperty*>& LayerNodeProperties = AnimBlueprintClass->GetLayerNodeProperties();
-		for(UStructProperty* LayerNodeProperty : LayerNodeProperties)
+		for(UStructProperty* LayerNodeProperty : AnimBlueprintClass->GetLinkedAnimLayerNodeProperties())
 		{
-			FAnimNode_Layer* Layer = LayerNodeProperty->ContainerPtrToValuePtr<FAnimNode_Layer>(this);
+			FAnimNode_LinkedAnimLayer* Layer = LayerNodeProperty->ContainerPtrToValuePtr<FAnimNode_LinkedAnimLayer>(this);
 
 			// If the class is null, then reset to default (which can be null)
 			UClass* ClassToSet = NewClass != nullptr ? NewClass : Layer->InstanceClass.Get();
@@ -2548,7 +2547,7 @@ void UAnimInstance::PerformLayerOverlayOperation(TSubclassOf<UAnimInstance> InCl
 				{
 					if(FoundFunction->bImplemented)
 					{
-						TArray<FAnimNode_Layer*, TInlineAllocator<4>>& LayerNodes = LayerNodesToSet.FindOrAdd(FoundFunction->Group);
+						TArray<FAnimNode_LinkedAnimLayer*, TInlineAllocator<4>>& LayerNodes = LayerNodesToSet.FindOrAdd(FoundFunction->Group);
 						LayerNodes.Add(Layer);
 					}
 				}
@@ -2556,7 +2555,7 @@ void UAnimInstance::PerformLayerOverlayOperation(TSubclassOf<UAnimInstance> InCl
 			else
 			{
 				// Add null classes so we clear the node's instance below
-				TArray<FAnimNode_Layer*, TInlineAllocator<4>>& LayerNodes = LayerNodesToSet.FindOrAdd(NAME_None);
+				TArray<FAnimNode_LinkedAnimLayer*, TInlineAllocator<4>>& LayerNodes = LayerNodesToSet.FindOrAdd(NAME_None);
 				LayerNodes.Add(Layer);
 			}
 		}
@@ -2565,20 +2564,20 @@ void UAnimInstance::PerformLayerOverlayOperation(TSubclassOf<UAnimInstance> InCl
 
 		auto UnlinkLayerNodesInInstance = [](UAnimInstance* InAnimInstance)
 		{
-			const IAnimClassInterface* NewSubInstanceClass = IAnimClassInterface::GetFromClass(InAnimInstance->GetClass());
-			for(const UStructProperty* LayerNodeProperty : NewSubInstanceClass->GetLayerNodeProperties())
+			const IAnimClassInterface* NewLinkedInstanceClass = IAnimClassInterface::GetFromClass(InAnimInstance->GetClass());
+			for(const UStructProperty* LayerNodeProperty : NewLinkedInstanceClass->GetLinkedAnimLayerNodeProperties())
 			{
-				FAnimNode_Layer* SubInstanceLayerNode = LayerNodeProperty->ContainerPtrToValuePtr<FAnimNode_Layer>(InAnimInstance);
-				SubInstanceLayerNode->DynamicUnlink(InAnimInstance);
+				FAnimNode_LinkedAnimLayer* LinkedAnimLayerNode = LayerNodeProperty->ContainerPtrToValuePtr<FAnimNode_LinkedAnimLayer>(InAnimInstance);
+				LinkedAnimLayerNode->DynamicUnlink(InAnimInstance);
 			}
 		};
 
-		for(TPair<FName, TArray<FAnimNode_Layer*, TInlineAllocator<4>>> LayerPair : LayerNodesToSet)
+		for(TPair<FName, TArray<FAnimNode_LinkedAnimLayer*, TInlineAllocator<4>>> LayerPair : LayerNodesToSet)
 		{
 			if (LayerPair.Key == NAME_None)
 			{
 				// Ungrouped path - each layer gets a separate instance
-				for(FAnimNode_Layer* LayerNode : LayerPair.Value)
+				for(FAnimNode_LinkedAnimLayer* LayerNode : LayerPair.Value)
 				{
 					UClass* ClassToSet = InClassSelectorFunction(NewClass, LayerNode);
 
@@ -2590,20 +2589,20 @@ void UAnimInstance::PerformLayerOverlayOperation(TSubclassOf<UAnimInstance> InCl
 						// Skip setting if the class is the same
 						if (TargetInstance == nullptr || ClassToSet != TargetInstance->GetClass())
 						{
-							UAnimInstance* NewSubInstance = NewObject<UAnimInstance>(MeshComp, ClassToSet);
-							NewSubInstance->InitializeAnimation();
+							UAnimInstance* NewLinkedInstance = NewObject<UAnimInstance>(MeshComp, ClassToSet);
+							NewLinkedInstance->InitializeAnimation();
 
-							// Unlink any layer nodes in the new sub instance, as they may have been hooked up to self in InitializeAnimation above.
-							UnlinkLayerNodesInInstance(NewSubInstance);
+							// Unlink any layer nodes in the new linked instance, as they may have been hooked up to self in InitializeAnimation above.
+							UnlinkLayerNodesInInstance(NewLinkedInstance);
 
-							LayerNode->SetLayerOverlaySubInstance(this, NewSubInstance);
+							LayerNode->SetLinkedLayerInstance(this, NewLinkedInstance);
 
-							// Initialize the correct parts of the sub instance
+							// Initialize the correct parts of the linked instance
 							if(LayerNode->LinkedRoot)
 							{
 								FAnimInstanceProxy& ThisProxy = GetProxyOnAnyThread<FAnimInstanceProxy>();
-								FAnimInstanceProxy& LinkedProxy = NewSubInstance->GetProxyOnAnyThread<FAnimInstanceProxy>();
-								LinkedProxy.InitializeObjects(NewSubInstance);
+								FAnimInstanceProxy& LinkedProxy = NewLinkedInstance->GetProxyOnAnyThread<FAnimInstanceProxy>();
+								LinkedProxy.InitializeObjects(NewLinkedInstance);
 								FAnimationInitializeContext InitContext(&ThisProxy);
 								LayerNode->Initialize_AnyThread(InitContext);
 								FAnimationCacheBonesContext CacheBonesContext(&ThisProxy);
@@ -2611,12 +2610,12 @@ void UAnimInstance::PerformLayerOverlayOperation(TSubclassOf<UAnimInstance> InCl
 								LinkedProxy.ClearObjects();
 							}
 
-							MeshComp->SubInstances.Add(NewSubInstance);
+							MeshComp->GetLinkedAnimInstances().Add(NewLinkedInstance);
 						}
 					}
 					else
 					{
-						LayerNode->SetLayerOverlaySubInstance(this, nullptr);
+						LayerNode->SetLinkedLayerInstance(this, nullptr);
 					}
 				}
 			}
@@ -2624,7 +2623,7 @@ void UAnimInstance::PerformLayerOverlayOperation(TSubclassOf<UAnimInstance> InCl
 			{		
 				// Grouped path - each group gets an instance
 				// If the class is null, then reset to default (which can be null)
-				FAnimNode_Layer* FirstLayerNode = LayerPair.Value[0];
+				FAnimNode_LinkedAnimLayer* FirstLayerNode = LayerPair.Value[0];
 				UClass* ClassToSet = InClassSelectorFunction(NewClass, FirstLayerNode);
 
 				// Disallow setting the same class as this instance, which would create infinite recursion
@@ -2635,27 +2634,27 @@ void UAnimInstance::PerformLayerOverlayOperation(TSubclassOf<UAnimInstance> InCl
 					// Skip setting if the class is the same
 					if (TargetInstance == nullptr || ClassToSet != TargetInstance->GetClass())
 					{
-						// Create and add one sub-instance for this group
-						UAnimInstance* NewSubInstance = NewObject<UAnimInstance>(MeshComp, ClassToSet);
-						NewSubInstance->InitializeAnimation();
+						// Create and add one linked instance for this group
+						UAnimInstance* NewLinkedInstance = NewObject<UAnimInstance>(MeshComp, ClassToSet);
+						NewLinkedInstance->InitializeAnimation();
 
-						// Unlink any layer nodes in the new sub instance, as they may have been hooked up to self in InitializeAnimation above.
-						UnlinkLayerNodesInInstance(NewSubInstance);
+						// Unlink any layer nodes in the new linked instance, as they may have been hooked up to self in InitializeAnimation above.
+						UnlinkLayerNodesInInstance(NewLinkedInstance);
 
-						for(FAnimNode_Layer* LayerNode : LayerPair.Value)
+						for(FAnimNode_LinkedAnimLayer* LayerNode : LayerPair.Value)
 						{
-							LayerNode->SetLayerOverlaySubInstance(this, NewSubInstance);
+							LayerNode->SetLinkedLayerInstance(this, NewLinkedInstance);
 						}
 
 						FAnimInstanceProxy& ThisProxy = GetProxyOnAnyThread<FAnimInstanceProxy>();
-						FAnimInstanceProxy& LinkedProxy = NewSubInstance->GetProxyOnAnyThread<FAnimInstanceProxy>();
+						FAnimInstanceProxy& LinkedProxy = NewLinkedInstance->GetProxyOnAnyThread<FAnimInstanceProxy>();
 
-						// Initialize the correct parts of the sub instance
-						for(FAnimNode_Layer* LayerNode : LayerPair.Value)
+						// Initialize the correct parts of the linked instance
+						for(FAnimNode_LinkedAnimLayer* LayerNode : LayerPair.Value)
 						{
 							if(LayerNode->LinkedRoot)
 							{
-								LinkedProxy.InitializeObjects(NewSubInstance);
+								LinkedProxy.InitializeObjects(NewLinkedInstance);
 								FAnimationInitializeContext InitContext(&ThisProxy);
 								LayerNode->Initialize_AnyThread(InitContext);
 								FAnimationCacheBonesContext CacheBonesContext(&ThisProxy);
@@ -2664,15 +2663,15 @@ void UAnimInstance::PerformLayerOverlayOperation(TSubclassOf<UAnimInstance> InCl
 							}
 						}
 
-						MeshComp->SubInstances.Add(NewSubInstance);
+						MeshComp->GetLinkedAnimInstances().Add(NewLinkedInstance);
 					}
 				}
 				else
 				{
 					// Clear the node's instance - we didnt find a class to use
-					for(FAnimNode_Layer* LayerNode : LayerPair.Value)
+					for(FAnimNode_LinkedAnimLayer* LayerNode : LayerPair.Value)
 					{
-						LayerNode->SetLayerOverlaySubInstance(this, nullptr);
+						LayerNode->SetLinkedLayerInstance(this, nullptr);
 					}
 				}
 			}
@@ -2680,9 +2679,9 @@ void UAnimInstance::PerformLayerOverlayOperation(TSubclassOf<UAnimInstance> InCl
 	}
 }
 
-void UAnimInstance::SetLayerOverlay(TSubclassOf<UAnimInstance> InClass)
+void UAnimInstance::LinkAnimClassLayers(TSubclassOf<UAnimInstance> InClass)
 {
-	auto SelectResolvedClassIfValid = [](UClass* InResolvedClass, FAnimNode_Layer* InLayerNode)
+	auto SelectResolvedClassIfValid = [](UClass* InResolvedClass, FAnimNode_LinkedAnimLayer* InLayerNode)
 	{
 		if(InResolvedClass != nullptr)
 		{
@@ -2696,12 +2695,12 @@ void UAnimInstance::SetLayerOverlay(TSubclassOf<UAnimInstance> InClass)
 		}
 	};
 
-	PerformLayerOverlayOperation(InClass, SelectResolvedClassIfValid);
+	PerformLinkedLayerOverlayOperation(InClass, SelectResolvedClassIfValid);
 }
 
-void UAnimInstance::ClearLayerOverlay(TSubclassOf<UAnimInstance> InClass)
+void UAnimInstance::UnlinkAnimClassLayers(TSubclassOf<UAnimInstance> InClass)
 {
-	auto ConditionallySelectDefaultClass = [](UClass* InResolvedClass, FAnimNode_Layer* InLayerNode)
+	auto ConditionallySelectDefaultClass = [](UClass* InResolvedClass, FAnimNode_LinkedAnimLayer* InLayerNode)
 	{
 		if (InResolvedClass != nullptr && InLayerNode->GetTargetInstance<UAnimInstance>()->GetClass() == InResolvedClass)
 		{
@@ -2715,22 +2714,21 @@ void UAnimInstance::ClearLayerOverlay(TSubclassOf<UAnimInstance> InClass)
 		}
 	};
 
-	PerformLayerOverlayOperation(InClass, ConditionallySelectDefaultClass);
+	PerformLinkedLayerOverlayOperation(InClass, ConditionallySelectDefaultClass);
 }
 
 void UAnimInstance::InitializeGroupedLayers()
 {
-	SetLayerOverlay(nullptr);
+	LinkAnimClassLayers(nullptr);
 }
 
-UAnimInstance* UAnimInstance::GetLayerSubInstanceByGroup(FName InGroup) const
+UAnimInstance* UAnimInstance::GetLinkedAnimLayerInstanceByGroup(FName InGroup) const
 {
 	if (IAnimClassInterface* AnimBlueprintClass = IAnimClassInterface::GetFromClass(GetClass()))
 	{
-		const TArray<UStructProperty*>& LayerNodeProperties = AnimBlueprintClass->GetLayerNodeProperties();
-		for(UStructProperty* LayerNodeProperty : LayerNodeProperties)
+		for(UStructProperty* LayerNodeProperty : AnimBlueprintClass->GetLinkedAnimLayerNodeProperties())
 		{
-			const FAnimNode_Layer* Layer = LayerNodeProperty->ContainerPtrToValuePtr<FAnimNode_Layer>(this);
+			const FAnimNode_LinkedAnimLayer* Layer = LayerNodeProperty->ContainerPtrToValuePtr<FAnimNode_LinkedAnimLayer>(this);
 
 			UClass* ClassForGroups;
 			if(UClass* InterfaceClass = Layer->Interface.Get())
@@ -2756,14 +2754,13 @@ UAnimInstance* UAnimInstance::GetLayerSubInstanceByGroup(FName InGroup) const
 	return nullptr;
 }
 
-UAnimInstance* UAnimInstance::GetLayerSubInstanceByClass(TSubclassOf<UAnimInstance> InClass) const
+UAnimInstance* UAnimInstance::GetLinkedAnimLayerInstanceByClass(TSubclassOf<UAnimInstance> InClass) const
 {
 	if (IAnimClassInterface* AnimBlueprintClass = IAnimClassInterface::GetFromClass(GetClass()))
 	{
-		const TArray<UStructProperty*>& LayerNodeProperties = AnimBlueprintClass->GetLayerNodeProperties();
-		for(UStructProperty* LayerNodeProperty : LayerNodeProperties)
+		for(UStructProperty* LayerNodeProperty : AnimBlueprintClass->GetLinkedAnimLayerNodeProperties())
 		{
-			const FAnimNode_Layer* Layer = LayerNodeProperty->ContainerPtrToValuePtr<FAnimNode_Layer>(this);
+			const FAnimNode_LinkedAnimLayer* Layer = LayerNodeProperty->ContainerPtrToValuePtr<FAnimNode_LinkedAnimLayer>(this);
 			UAnimInstance* TargetInstance = Layer->GetTargetInstance<UAnimInstance>();
 			if (TargetInstance && TargetInstance->GetClass() == InClass.Get())
 			{
diff --git a/Engine/Source/Runtime/Engine/Private/Animation/AnimInstanceProxy.cpp b/Engine/Source/Runtime/Engine/Private/Animation/AnimInstanceProxy.cpp
--- a/Engine/Source/Runtime/Engine/Private/Animation/AnimInstanceProxy.cpp
+++ b/Engine/Source/Runtime/Engine/Private/Animation/AnimInstanceProxy.cpp
@@ -11,7 +11,7 @@
 #include "Animation/AnimNode_StateMachine.h"
 #include "Animation/AnimNode_TransitionResult.h"
 #include "Animation/AnimNode_SaveCachedPose.h"
-#include "Animation/AnimNode_SubInput.h"
+#include "Animation/AnimNode_LinkedInputPose.h"
 #include "Engine/Engine.h"
 #include "DrawDebugHelpers.h"
 #include "GameFramework/WorldSettings.h"
@@ -23,7 +23,7 @@
 #define DO_ANIMSTAT_PROCESSING(StatName) DEFINE_STAT(STAT_ ## StatName ## _WorkerThread)
 #include "Animation/AnimMTStats.h"
 #include "Animation/AnimNode_Root.h"
-#include "Animation/AnimNode_Layer.h"
+#include "Animation/AnimNode_LinkedAnimLayer.h"
 #undef DO_ANIMSTAT_PROCESSING
 
 #define LOCTEXT_NAMESPACE "AnimInstance"
@@ -221,7 +221,7 @@ void FAnimInstanceProxy::InitializeRootNode(bool bInDeferRootNodeInitialization)
 			StateMachine->CacheMachineDescription(AnimClassInterface);
 		}
 
-		// Cache default sub-input 
+		// Cache default linked input pose 
 		for(const FAnimBlueprintFunction& AnimBlueprintFunction : AnimClassInterface->GetAnimBlueprintFunctions())
 		{
 			if(AnimBlueprintFunction.Name == NAME_AnimGraph)
@@ -229,9 +229,9 @@ void FAnimInstanceProxy::InitializeRootNode(bool bInDeferRootNodeInitialization)
 				check(AnimBlueprintFunction.InputPoseNames.Num() == AnimBlueprintFunction.InputPoseNodeProperties.Num());
 				for(int32 InputIndex = 0; InputIndex < AnimBlueprintFunction.InputPoseNames.Num(); ++InputIndex)
 				{
-					if(AnimBlueprintFunction.InputPoseNames[InputIndex] == FAnimNode_SubInput::DefaultInputPoseName && AnimBlueprintFunction.InputPoseNodeProperties[InputIndex] != nullptr)
+					if(AnimBlueprintFunction.InputPoseNames[InputIndex] == FAnimNode_LinkedInputPose::DefaultInputPoseName && AnimBlueprintFunction.InputPoseNodeProperties[InputIndex] != nullptr)
 					{
-						DefaultSubInstanceInputNode = AnimBlueprintFunction.InputPoseNodeProperties[InputIndex]->ContainerPtrToValuePtr<FAnimNode_SubInput>(CastChecked<UAnimInstance>(GetAnimInstanceObject()));
+						DefaultLinkedInstanceInputNode = AnimBlueprintFunction.InputPoseNodeProperties[InputIndex]->ContainerPtrToValuePtr<FAnimNode_LinkedInputPose>(CastChecked<UAnimInstance>(GetAnimInstanceObject()));
 						break;
 					}
 				}
@@ -331,7 +331,7 @@ void FAnimInstanceProxy::Uninitialize(UAnimInstance* InAnimInstance)
 	DECLARE_SCOPE_HIERARCHICAL_COUNTER_FUNC()
 
 	MontageEvaluationData.Reset();
-	DefaultSubInstanceInputNode = nullptr;
+	DefaultLinkedInstanceInputNode = nullptr;
 	ResetAnimationCurves();
 	MaterialParametersToClear.Reset();
 }
@@ -1158,9 +1158,9 @@ void FAnimInstanceProxy::RecalcRequiredBones(USkeletalMeshComponent* Component,
 	}
 
 	// If this instance can accept input poses, initialise the input pose container
-	if(DefaultSubInstanceInputNode)
+	if(DefaultLinkedInstanceInputNode)
 	{
-		DefaultSubInstanceInputNode->CachedInputPose.SetBoneContainer(&RequiredBones);
+		DefaultLinkedInstanceInputNode->CachedInputPose.SetBoneContainer(&RequiredBones);
 	}
 
 	// When RequiredBones mapping has changed, AnimNodes need to update their bones caches. 
@@ -1213,7 +1213,7 @@ void FAnimInstanceProxy::UpdateAnimation_WithRoot(FAnimNode_Base* InRootNode, FN
 	if(!bUpdatingRoot)
 	{
 		// Make sure we only update this once the first time we update, as we can re-call this function
-		// from other sub instances with grouped layers
+		// from other linked instances with grouped layers
 		if(FrameCounterForUpdate != GFrameCounter)
 		{
 			SCOPE_CYCLE_COUNTER(STAT_NativeUpdateAnimation);
diff --git a/Engine/Source/Runtime/Engine/Private/Animation/AnimNode_CustomProperty.cpp b/Engine/Source/Runtime/Engine/Private/Animation/AnimNode_CustomProperty.cpp
--- a/Engine/Source/Runtime/Engine/Private/Animation/AnimNode_CustomProperty.cpp
+++ b/Engine/Source/Runtime/Engine/Private/Animation/AnimNode_CustomProperty.cpp
@@ -3,7 +3,7 @@
 #include "Animation/AnimNode_CustomProperty.h"
 #include "Animation/AnimClassInterface.h"
 #include "Animation/AnimInstanceProxy.h"
-#include "Animation/AnimNode_SubInput.h"
+#include "Animation/AnimNode_LinkedInputPose.h"
 
 FAnimNode_CustomProperty::FAnimNode_CustomProperty()
 	: TargetInstance(nullptr)
diff --git a/Engine/Source/Runtime/Engine/Private/Animation/AnimNode_Layer.cpp b/Engine/Source/Runtime/Engine/Private/Animation/AnimNode_Layer.cpp
deleted file mode 100644
diff --git a/Engine/Source/Runtime/Engine/Private/Animation/AnimNode_LinkedAnimGraph.cpp b/Engine/Source/Runtime/Engine/Private/Animation/AnimNode_LinkedAnimGraph.cpp
new file mode 100644
--- /dev/null
+++ b/Engine/Source/Runtime/Engine/Private/Animation/AnimNode_LinkedAnimGraph.cpp
@@ -0,0 +1,351 @@
+// Copyright 1998-2019 Epic Games, Inc. All Rights Reserved.
+
+#include "Animation/AnimNode_LinkedAnimGraph.h"
+#include "Animation/AnimClassInterface.h"
+#include "Animation/AnimInstanceProxy.h"
+#include "Animation/AnimNode_LinkedInputPose.h"
+#include "Animation/AnimNode_Root.h"
+
+FAnimNode_LinkedAnimGraph::FAnimNode_LinkedAnimGraph()
+	: InstanceClass(nullptr)
+	, Tag(NAME_None)
+{
+}
+
+void FAnimNode_LinkedAnimGraph::Initialize_AnyThread(const FAnimationInitializeContext& Context)
+{
+	FAnimNode_Base::Initialize_AnyThread(Context);
+
+	UAnimInstance* InstanceToRun = GetTargetInstance<UAnimInstance>();
+	if(InstanceToRun && LinkedRoot)
+	{
+		FAnimInstanceProxy& Proxy = InstanceToRun->GetProxyOnAnyThread<FAnimInstanceProxy>();
+		Proxy.InitializationCounter.SynchronizeWith(Context.AnimInstanceProxy->InitializationCounter);
+		Proxy.InitializeRootNode_WithRoot(LinkedRoot);
+	}
+	else
+	{
+		// If we have no valid instance (self or otherwise), we need to propagate down the graph to make sure
+		// subsequent nodes get properly initialized
+		for(FPoseLink& InputPose : InputPoses)
+		{
+			InputPose.Initialize(Context);
+		}
+	}
+}
+
+void FAnimNode_LinkedAnimGraph::CacheBones_AnyThread(const FAnimationCacheBonesContext& Context)
+{
+	UAnimInstance* InstanceToRun = GetTargetInstance<UAnimInstance>();
+	if(InstanceToRun && LinkedRoot)
+	{
+		FAnimInstanceProxy& Proxy = InstanceToRun->GetProxyOnAnyThread<FAnimInstanceProxy>();
+		Proxy.CachedBonesCounter.SynchronizeWith(Context.AnimInstanceProxy->CachedBonesCounter);
+
+		// Note not calling Proxy.CacheBones_WithRoot here as it is guarded by
+		// bBoneCachesInvalidated, which is handled at a higher level
+		FAnimationCacheBonesContext LinkedContext(&Proxy);
+		LinkedRoot->CacheBones_AnyThread(LinkedContext);
+	}
+	else
+	{
+		// If we have no valid instance (self or otherwise), we need to propagate down the graph to make sure
+		// subsequent nodes get properly their bones properly cached
+		for(FPoseLink& InputPose : InputPoses)
+		{
+			InputPose.CacheBones(Context);
+		}
+	}
+}
+
+void FAnimNode_LinkedAnimGraph::Update_AnyThread(const FAnimationUpdateContext& Context)
+{
+	GetEvaluateGraphExposedInputs().Execute(Context);
+
+	UAnimInstance* InstanceToRun = GetTargetInstance<UAnimInstance>();
+	if(InstanceToRun && LinkedRoot)
+	{
+		FAnimInstanceProxy& Proxy = InstanceToRun->GetProxyOnAnyThread<FAnimInstanceProxy>();
+		Proxy.UpdateCounter.SynchronizeWith(Context.AnimInstanceProxy->UpdateCounter);
+
+		PropagateInputProperties(Context.AnimInstanceProxy->GetAnimInstanceObject());
+
+		// Only update if we've not had a single-threaded update already
+		if(InstanceToRun->bNeedsUpdate)
+		{
+			Proxy.UpdateAnimation_WithRoot(LinkedRoot, GetDynamicLinkFunctionName());
+		}
+	}
+	else if(InputPoses.Num() > 0)
+	{
+		// If we have no valid instance (self or otherwise), we need to propagate down the graph to make sure
+		// subsequent nodes get properly updated
+		InputPoses[0].Update(Context);
+	}
+}
+
+void FAnimNode_LinkedAnimGraph::Evaluate_AnyThread(FPoseContext& Output)
+{
+	UAnimInstance* InstanceToRun = GetTargetInstance<UAnimInstance>();
+	if(InstanceToRun && LinkedRoot)
+	{
+		FAnimInstanceProxy& Proxy = InstanceToRun->GetProxyOnAnyThread<FAnimInstanceProxy>();
+		Proxy.EvaluationCounter.SynchronizeWith(Output.AnimInstanceProxy->EvaluationCounter);
+		Output.Pose.SetBoneContainer(&Proxy.GetRequiredBones());
+
+		// Create an evaluation context
+		FPoseContext EvaluationContext(&Proxy, Output.ExpectsAdditivePose());
+		EvaluationContext.ResetToRefPose();
+			
+		// Run the anim blueprint
+		Proxy.EvaluateAnimation_WithRoot(EvaluationContext, LinkedRoot);
+
+		// Move the curves
+		Output.Curve.MoveFrom(EvaluationContext.Curve);
+		Output.Pose.MoveBonesFrom(EvaluationContext.Pose);
+	}
+	else if(InputPoses.Num() > 0)
+	{
+		// If we have no valid instance (self or otherwise), we need to propagate down the graph to make sure
+		// subsequent nodes get properly evaluated
+		InputPoses[0].Evaluate(Output);
+	}
+	else
+	{
+		Output.ResetToRefPose();
+	}
+}
+
+void FAnimNode_LinkedAnimGraph::GatherDebugData(FNodeDebugData& DebugData)
+{
+	// Add our entry
+	FString DebugLine = DebugData.GetNodeName(this);
+	DebugLine += FString::Printf(TEXT("Target: %s"), (*InstanceClass) ? *InstanceClass->GetName() : TEXT("None"));
+
+	DebugData.AddDebugItem(DebugLine);
+
+	UAnimInstance* InstanceToRun = GetTargetInstance<UAnimInstance>();
+	// Gather data from the linked instance
+	if(InstanceToRun && LinkedRoot)
+	{
+		FAnimInstanceProxy& Proxy = InstanceToRun->GetProxyOnAnyThread<FAnimInstanceProxy>();
+		Proxy.GatherDebugData_WithRoot(DebugData.BranchFlow(1.0f), LinkedRoot, GetDynamicLinkFunctionName());
+	}
+	else if(InputPoses.Num() > 0)
+	{
+		// If we have no valid instance (self or otherwise), we need to propagate down the graph to make sure
+		// subsequent nodes get their debug data properly collected to reflect relevancy
+		InputPoses[0].GatherDebugData(DebugData);
+	}
+}
+
+void FAnimNode_LinkedAnimGraph::OnInitializeAnimInstance(const FAnimInstanceProxy* InProxy, const UAnimInstance* InAnimInstance)
+{
+	UAnimInstance* InstanceToRun = GetTargetInstance<UAnimInstance>();
+
+	if(*InstanceClass)
+	{
+		ReinitializeLinkedAnimInstance(InAnimInstance);
+	}
+	else if(InstanceToRun)
+	{
+		// We have an instance but no instance class
+		TeardownInstance();
+	}
+}
+
+void FAnimNode_LinkedAnimGraph::TeardownInstance()
+{
+	UAnimInstance* InstanceToRun = GetTargetInstance<UAnimInstance>();
+	if (InstanceToRun)
+	{
+		InstanceToRun->UninitializeAnimation();
+		InstanceToRun = nullptr;
+	}
+}
+
+void FAnimNode_LinkedAnimGraph::ReinitializeLinkedAnimInstance(const UAnimInstance* InOwningAnimInstance, UAnimInstance* InNewAnimInstance)
+{
+	UAnimInstance* InstanceToRun = GetTargetInstance<UAnimInstance>();
+
+	if(*InstanceClass || InNewAnimInstance)
+	{
+		USkeletalMeshComponent* MeshComp = InOwningAnimInstance->GetSkelMeshComponent();
+		check(MeshComp);
+		// Full reinit, kill old instances
+		if(InstanceToRun)
+		{
+			DynamicUnlink(const_cast<UAnimInstance*>(InOwningAnimInstance));
+
+			MeshComp->GetLinkedAnimInstances().Remove(InstanceToRun);
+			// Never delete the owning animation instance
+			if (InstanceToRun != InOwningAnimInstance)
+			{
+				InstanceToRun->MarkPendingKill();
+			}
+			InstanceToRun = nullptr;
+		}
+
+		// Need an instance to run, so create it now
+		InstanceToRun = InNewAnimInstance ? InNewAnimInstance : NewObject<UAnimInstance>(MeshComp, InstanceClass);
+		SetTargetInstance(InstanceToRun);
+
+		// Link before we call InitializeAnimation() so we propgate the call to linked input poses
+		DynamicLink(const_cast<UAnimInstance*>(InOwningAnimInstance));
+
+		if(InNewAnimInstance == nullptr)
+		{
+			// Initialize the new instance
+			InstanceToRun->InitializeAnimation();
+
+			MeshComp->GetLinkedAnimInstances().Add(InstanceToRun);
+		}
+
+		InitializeProperties(InOwningAnimInstance, InstanceToRun->GetClass());
+	}
+	else if(InstanceToRun)
+	{
+		// We have an instance but no instance class
+		TeardownInstance();
+	}
+}
+
+void FAnimNode_LinkedAnimGraph::SetAnimClass(TSubclassOf<UAnimInstance> InClass, const UAnimInstance* InOwningAnimInstance)
+{
+	UClass* NewClass = InClass.Get();
+	if(NewClass)
+	{
+		// Verify target skeleton match at runtime
+		IAnimClassInterface* LinkedAnimBlueprintClass = IAnimClassInterface::GetFromClass(NewClass);
+		IAnimClassInterface* OuterAnimBlueprintClass = IAnimClassInterface::GetFromClass(InOwningAnimInstance->GetClass());
+		USkeleton* LinkedSkeleton = LinkedAnimBlueprintClass->GetTargetSkeleton();
+		USkeleton* OuterSkeleton = OuterAnimBlueprintClass->GetTargetSkeleton();
+		if(LinkedSkeleton != OuterSkeleton)
+		{
+			UE_LOG(LogAnimation, Warning, TEXT("Setting linked anim instance class: Class has a mismatched target skeleton. Expected %s, found %s."), OuterSkeleton ? *OuterSkeleton->GetName() : TEXT("null"), LinkedSkeleton ? *LinkedSkeleton->GetName() : TEXT("null"));
+			return;
+		}
+	}
+
+	// Verified OK, so set it now
+	TSubclassOf<UAnimInstance> OldClass = InstanceClass;
+	InstanceClass = InClass;
+
+	if(InstanceClass != OldClass)
+	{
+		ReinitializeLinkedAnimInstance(InOwningAnimInstance);
+	}
+}
+
+FName FAnimNode_LinkedAnimGraph::GetDynamicLinkFunctionName() const
+{
+	return NAME_AnimGraph;
+}
+
+UAnimInstance* FAnimNode_LinkedAnimGraph::GetDynamicLinkTarget(UAnimInstance* InOwningAnimInstance) const
+{
+	return GetTargetInstance<UAnimInstance>();
+}
+
+void FAnimNode_LinkedAnimGraph::DynamicLink(UAnimInstance* InOwningAnimInstance)
+{
+	UAnimInstance* LinkTargetInstance = GetDynamicLinkTarget(InOwningAnimInstance);
+	if(LinkTargetInstance)
+	{
+		IAnimClassInterface* SubAnimBlueprintClass = IAnimClassInterface::GetFromClass(LinkTargetInstance->GetClass());
+		if(SubAnimBlueprintClass)
+		{
+			FAnimInstanceProxy* NonConstProxy = &InOwningAnimInstance->GetProxyOnAnyThread<FAnimInstanceProxy>();
+
+			// Link input poses
+			for(const FAnimBlueprintFunction& AnimBlueprintFunction : SubAnimBlueprintClass->GetAnimBlueprintFunctions())
+			{
+				const FName FunctionToLink = GetDynamicLinkFunctionName();
+				if(AnimBlueprintFunction.Name == FunctionToLink)
+				{
+					for(int32 InputPoseIndex = 0; InputPoseIndex < AnimBlueprintFunction.InputPoseNames.Num() && InputPoseIndex < InputPoses.Num(); ++InputPoseIndex)
+					{
+						// Make sure we attempt a re-link first, as only this pose link knows its target
+						FAnimationInitializeContext Context(NonConstProxy);
+						InputPoses[InputPoseIndex].AttemptRelink(Context);
+
+						int32 InputPropertyIndex = FindFunctionInputIndex(AnimBlueprintFunction, AnimBlueprintFunction.InputPoseNames[InputPoseIndex]);
+						if(InputPropertyIndex != INDEX_NONE && AnimBlueprintFunction.InputPoseNodeProperties[InputPropertyIndex])
+						{
+							FAnimNode_LinkedInputPose* LinkedInputPoseNode = AnimBlueprintFunction.InputPoseNodeProperties[InputPropertyIndex]->ContainerPtrToValuePtr<FAnimNode_LinkedInputPose>(LinkTargetInstance);
+							check(LinkedInputPoseNode->Name == AnimBlueprintFunction.InputPoseNames[InputPoseIndex]);
+							LinkedInputPoseNode->DynamicLink(NonConstProxy, &InputPoses[InputPoseIndex]);
+						}
+						else
+						{
+							UE_LOG(LogAnimation, Warning, TEXT("Unable to dynamically link input pose %s."), *AnimBlueprintFunction.InputPoseNames[InputPoseIndex].ToString());
+						}
+					}
+
+					if(AnimBlueprintFunction.OutputPoseNodeProperty)
+					{
+						LinkedRoot = AnimBlueprintFunction.OutputPoseNodeProperty->ContainerPtrToValuePtr<FAnimNode_Root>(LinkTargetInstance);
+					}
+					else
+					{
+						UE_LOG(LogAnimation, Warning, TEXT("Unable to dynamically link root %s."), *FunctionToLink.ToString());
+					}
+
+					break;
+				}
+			}
+		}
+	}
+}
+
+void FAnimNode_LinkedAnimGraph::DynamicUnlink(UAnimInstance* InOwningAnimInstance)
+{
+	// unlink root
+	LinkedRoot = nullptr;
+
+	// unlink input poses
+	UAnimInstance* LinkTargetInstance = GetDynamicLinkTarget(InOwningAnimInstance);
+	if(LinkTargetInstance)
+	{
+		IAnimClassInterface* SubAnimBlueprintClass = IAnimClassInterface::GetFromClass(LinkTargetInstance->GetClass());
+		if(SubAnimBlueprintClass)
+		{
+			// Link input poses
+			for(const FAnimBlueprintFunction& AnimBlueprintFunction : SubAnimBlueprintClass->GetAnimBlueprintFunctions())
+			{
+				if(AnimBlueprintFunction.Name == GetDynamicLinkFunctionName())
+				{
+					for(int32 InputPoseIndex = 0; InputPoseIndex < AnimBlueprintFunction.InputPoseNames.Num() && InputPoseIndex < InputPoses.Num(); ++InputPoseIndex)
+					{
+						int32 InputPropertyIndex = FindFunctionInputIndex(AnimBlueprintFunction, AnimBlueprintFunction.InputPoseNames[InputPoseIndex]);
+						if(InputPropertyIndex != INDEX_NONE && AnimBlueprintFunction.InputPoseNodeProperties[InputPropertyIndex])
+						{
+							FAnimNode_LinkedInputPose* LinkedInputPoseNode = AnimBlueprintFunction.InputPoseNodeProperties[InputPropertyIndex]->ContainerPtrToValuePtr<FAnimNode_LinkedInputPose>(LinkTargetInstance);
+							check(LinkedInputPoseNode->Name == AnimBlueprintFunction.InputPoseNames[InputPoseIndex]);
+							LinkedInputPoseNode->DynamicUnlink();
+						}
+						else
+						{
+							UE_LOG(LogAnimation, Warning, TEXT("Unable to dynamically unlink input pose %s."), *AnimBlueprintFunction.InputPoseNames[InputPoseIndex].ToString());
+						}
+					}
+
+					break;
+				}
+			}
+		}
+	}
+}
+
+int32 FAnimNode_LinkedAnimGraph::FindFunctionInputIndex(const FAnimBlueprintFunction& InAnimBlueprintFunction, const FName& InInputName)
+{
+	for(int32 InputPropertyIndex = 0; InputPropertyIndex < InAnimBlueprintFunction.InputPoseNames.Num(); ++InputPropertyIndex)
+	{
+		if(InInputName == InAnimBlueprintFunction.InputPoseNames[InputPropertyIndex])
+		{
+			return InputPropertyIndex;
+		}
+	}
+
+	return INDEX_NONE;
+}
diff --git a/Engine/Source/Runtime/Engine/Private/Animation/AnimNode_LinkedAnimLayer.cpp b/Engine/Source/Runtime/Engine/Private/Animation/AnimNode_LinkedAnimLayer.cpp
new file mode 100644
--- /dev/null
+++ b/Engine/Source/Runtime/Engine/Private/Animation/AnimNode_LinkedAnimLayer.cpp
@@ -0,0 +1,69 @@
+// Copyright 1998-2019 Epic Games, Inc. All Rights Reserved.
+
+#include "Animation/AnimNode_LinkedAnimLayer.h"
+
+FName FAnimNode_LinkedAnimLayer::GetDynamicLinkFunctionName() const
+{
+	return Layer;
+}
+
+UAnimInstance* FAnimNode_LinkedAnimLayer::GetDynamicLinkTarget(UAnimInstance* InOwningAnimInstance) const
+{
+	if(Interface.Get())
+	{
+		return GetTargetInstance<UAnimInstance>();
+	}
+	else
+	{
+		return InOwningAnimInstance;
+	}
+}
+
+void FAnimNode_LinkedAnimLayer::OnInitializeAnimInstance(const FAnimInstanceProxy* InProxy, const UAnimInstance* InAnimInstance)
+{
+	// We only initialize here if we are running a 'self' layer. Layers that use external instances need to be 
+	// initialized by the owning anim instance as they may share linked instances via grouping.
+	if(Interface.Get() == nullptr || InstanceClass.Get() == nullptr)
+	{
+		InitializeSelfLayer(InAnimInstance);
+	}
+}
+
+void FAnimNode_LinkedAnimLayer::InitializeSelfLayer(const UAnimInstance* SelfAnimInstance)
+{
+	UAnimInstance* CurrentTarget = GetTargetInstance<UAnimInstance>();
+
+	USkeletalMeshComponent* MeshComp = SelfAnimInstance->GetSkelMeshComponent();
+	check(MeshComp);
+
+	if (LinkedRoot)
+	{
+		DynamicUnlink(const_cast<UAnimInstance*>(SelfAnimInstance));
+	}
+
+	// Switch from dynamic external to internal, kill old instance
+	if (CurrentTarget && CurrentTarget != SelfAnimInstance)
+	{
+		MeshComp->GetLinkedAnimInstances().Remove(CurrentTarget);
+		CurrentTarget->MarkPendingKill();
+		CurrentTarget = nullptr;
+	}
+
+	SetTargetInstance(const_cast<UAnimInstance*>(SelfAnimInstance));
+
+	// Link before we call InitializeAnimation() so we propgate the call to linked input poses
+	DynamicLink(const_cast<UAnimInstance*>(SelfAnimInstance));
+
+	InitializeProperties(SelfAnimInstance, SelfAnimInstance->GetClass());
+}
+
+void FAnimNode_LinkedAnimLayer::SetLinkedLayerInstance(const UAnimInstance* InOwningAnimInstance, UAnimInstance* InNewLinkedInstance)
+{
+	ReinitializeLinkedAnimInstance(InOwningAnimInstance, InNewLinkedInstance);
+
+	// Reseting to running as a self-layer, in case it is applicable
+	if ((Interface.Get() == nullptr || InstanceClass.Get() == nullptr) && (InNewLinkedInstance == nullptr))
+	{
+		InitializeSelfLayer(InOwningAnimInstance);
+	}
+}
diff --git a/Engine/Source/Runtime/Engine/Private/Animation/AnimNode_LinkedInputPose.cpp b/Engine/Source/Runtime/Engine/Private/Animation/AnimNode_LinkedInputPose.cpp
new file mode 100644
--- /dev/null
+++ b/Engine/Source/Runtime/Engine/Private/Animation/AnimNode_LinkedInputPose.cpp
@@ -0,0 +1,83 @@
+// Copyright 1998-2019 Epic Games, Inc. All Rights Reserved.
+
+#include "Animation/AnimNode_LinkedInputPose.h"
+#include "Animation/AnimInstanceProxy.h"
+
+const FName FAnimNode_LinkedInputPose::DefaultInputPoseName("InPose");
+
+void FAnimNode_LinkedInputPose::Initialize_AnyThread(const FAnimationInitializeContext& Context)
+{
+	if(InputProxy)
+	{
+		FAnimationInitializeContext InputContext(InputProxy);
+		InputPose.Initialize(InputContext);
+	}
+}
+
+void FAnimNode_LinkedInputPose::CacheBones_AnyThread(const FAnimationCacheBonesContext& Context)
+{
+	if(InputProxy)
+	{
+		FAnimationCacheBonesContext InputContext(InputProxy);
+		InputPose.CacheBones(InputContext);
+	}
+}
+
+void FAnimNode_LinkedInputPose::Update_AnyThread(const FAnimationUpdateContext& Context)
+{
+	if(InputProxy)
+	{
+		FAnimationUpdateContext InputContext = Context.WithOtherProxy(InputProxy);
+		InputPose.Update(InputContext);
+	}
+}
+
+void FAnimNode_LinkedInputPose::Evaluate_AnyThread(FPoseContext& Output)
+{
+	if(InputProxy)
+	{
+		Output.Pose.SetBoneContainer(&InputProxy->GetRequiredBones());
+
+		FPoseContext InputContext(InputProxy, Output.ExpectsAdditivePose());
+		InputPose.Evaluate(InputContext);
+
+		Output.Pose.MoveBonesFrom(InputContext.Pose);
+		Output.Curve.MoveFrom(InputContext.Curve);
+	}
+	else if(CachedInputPose.IsValid() && CachedInputCurve.IsValid())
+	{
+		Output.Pose.CopyBonesFrom(CachedInputPose);
+		Output.Curve.CopyFrom(CachedInputCurve);
+	}
+	else
+	{
+		Output.ResetToRefPose();
+	}
+}
+
+void FAnimNode_LinkedInputPose::GatherDebugData(FNodeDebugData& DebugData)
+{
+	FString DebugLine = DebugData.GetNodeName(this);
+	DebugData.AddDebugItem(DebugLine);
+
+	if(InputProxy)
+	{
+		InputPose.GatherDebugData(DebugData);
+	}
+}
+
+void FAnimNode_LinkedInputPose::DynamicLink(FAnimInstanceProxy* InInputProxy, FPoseLinkBase* InPoseLink)
+{
+	check(InputProxy == nullptr);	// Must be unlinked before re-linking
+
+	InputProxy = InInputProxy;
+	InputPose.SetDynamicLinkNode(InPoseLink);
+}
+
+void FAnimNode_LinkedInputPose::DynamicUnlink()
+{
+	check(InputProxy != nullptr);	// Must be linked before unlinking
+
+	InputProxy = nullptr;
+	InputPose.SetDynamicLinkNode(nullptr);
+}
\ No newline at end of file
diff --git a/Engine/Source/Runtime/Engine/Private/Animation/AnimNode_StateMachine.cpp b/Engine/Source/Runtime/Engine/Private/Animation/AnimNode_StateMachine.cpp
--- a/Engine/Source/Runtime/Engine/Private/Animation/AnimNode_StateMachine.cpp
+++ b/Engine/Source/Runtime/Engine/Private/Animation/AnimNode_StateMachine.cpp
@@ -7,7 +7,7 @@
 #include "Animation/AnimNode_AssetPlayerBase.h"
 #include "Animation/AnimNode_Inertialization.h"
 #include "Animation/BlendProfile.h"
-#include "Animation/AnimNode_Layer.h"
+#include "Animation/AnimNode_LinkedAnimLayer.h"
 #include "Animation/AnimInstance.h"
 
 DECLARE_CYCLE_STAT(TEXT("StateMachine SetState"), Stat_StateMachineSetState, STATGROUP_Anim);
@@ -563,7 +563,7 @@ FAnimNode_AssetPlayerBase* FAnimNode_StateMachine::GetRelevantAssetPlayerFromSta
 	for (const int32& LayerIdx : StateInfo.LayerNodeIndices)
 	{
 		// Try and retrieve the actual node object
-		if (FAnimNode_Layer* Layer = Context.AnimInstanceProxy->GetNodeFromIndex<FAnimNode_Layer>(LayerIdx))
+		if (FAnimNode_LinkedAnimLayer* Layer = Context.AnimInstanceProxy->GetNodeFromIndex<FAnimNode_LinkedAnimLayer>(LayerIdx))
 		{
 			// Retrieve the AnimInstance running for this layer
 			if (UAnimInstance* CurrentTarget = Layer->GetTargetInstance<UAnimInstance>())
diff --git a/Engine/Source/Runtime/Engine/Private/Animation/AnimNode_SubInput.cpp b/Engine/Source/Runtime/Engine/Private/Animation/AnimNode_SubInput.cpp
deleted file mode 100644
diff --git a/Engine/Source/Runtime/Engine/Private/Animation/AnimNode_SubInstance.cpp b/Engine/Source/Runtime/Engine/Private/Animation/AnimNode_SubInstance.cpp
deleted file mode 100644
diff --git a/Engine/Source/Runtime/Engine/Private/Components/SkeletalMeshComponent.cpp b/Engine/Source/Runtime/Engine/Private/Components/SkeletalMeshComponent.cpp
--- a/Engine/Source/Runtime/Engine/Private/Components/SkeletalMeshComponent.cpp
+++ b/Engine/Source/Runtime/Engine/Private/Components/SkeletalMeshComponent.cpp
@@ -28,7 +28,7 @@
 #include "Rendering/SkeletalMeshRenderData.h"
 
 #include "Logging/MessageLog.h"
-#include "Animation/AnimNode_SubInput.h"
+#include "Animation/AnimNode_LinkedInputPose.h"
 
 #include "PhysXIncludes.h"
 #include "ClothingSimulationFactoryInterface.h"
@@ -516,9 +516,9 @@ void USkeletalMeshComponent::OnRegister()
 
 	Super::OnRegister();
 
-	// Ensure we have an empty list of subinstances on registration. Ready for the initialization below 
+	// Ensure we have an empty list of linked instances on registration. Ready for the initialization below 
 	// to correctly populate that list.
-	SubInstances.Reset();
+	LinkedInstances.Reset();
 
 	// We force an initialization here because we're in one of two cases.
 	// 1) First register, no spawned instance, need to initialize
@@ -571,11 +571,11 @@ void USkeletalMeshComponent::OnUnregister()
 		AnimScriptInstance->UninitializeAnimation();
 	}
 
-	for(UAnimInstance* SubInstance : SubInstances)
+	for(UAnimInstance* LinkedInstance : LinkedInstances)
 	{
-		SubInstance->UninitializeAnimation();
+		LinkedInstance->UninitializeAnimation();
 	}
-	SubInstances.Reset();
+	LinkedInstances.Reset();
 
 	if(PostProcessAnimInstance)
 	{
@@ -706,8 +706,8 @@ bool USkeletalMeshComponent::InitializeAnimScriptInstance(bool bForceReinit, boo
 
 			if (AnimScriptInstance)
 			{
-				// If we have any sub-instances left we need to clear them out now, we're about to have a new master instance
-				SubInstances.Empty();
+				// If we have any linked instances left we need to clear them out now, we're about to have a new master instance
+				LinkedInstances.Empty();
 
 				AnimScriptInstance->InitializeAnimation(bInDeferRootNodeInitialization);
 				bInitializedMainInstance = true;
@@ -769,7 +769,7 @@ bool USkeletalMeshComponent::InitializeAnimScriptInstance(bool bForceReinit, boo
 			{
 				PostProcessAnimInstance->InitializeAnimation();
 
-				if(FAnimNode_SubInput* InputNode = PostProcessAnimInstance->GetSubInputNode())
+				if(FAnimNode_LinkedInputPose* InputNode = PostProcessAnimInstance->GetLinkedInputPoseNode())
 				{
 					InputNode->CachedInputPose.SetBoneContainer(&PostProcessAnimInstance->GetRequiredBones());
 				}
@@ -824,7 +824,7 @@ void USkeletalMeshComponent::ClearAnimScriptInstance()
 		AnimScriptInstance->EndNotifyStates();
 	}
 	AnimScriptInstance = nullptr;
-	SubInstances.Empty();
+	LinkedInstances.Empty();
 }
 
 
@@ -1014,12 +1014,12 @@ void USkeletalMeshComponent::TickAnimation(float DeltaTime, bool bNeedsValidRoot
 		// We're about to UpdateAnimation, this will potentially queue events that we'll need to dispatch.
 		bNeedsQueuedAnimEventsDispatched = true;
 
-		// We update sub instances first incase we're using either root motion or non-threaded update.
+		// We update linked instances first incase we're using either root motion or non-threaded update.
 		// This ensures that we go through the pre update process and initialize the proxies correctly.
-		for(UAnimInstance* SubInstance : SubInstances)
+		for(UAnimInstance* LinkedInstance : LinkedInstances)
 		{
 			// Sub anim instances are always forced to do a parallel update 
-			SubInstance->UpdateAnimation(DeltaTime * GlobalAnimRateScale, false, UAnimInstance::EUpdateAnimationFlag::ForceParallelUpdate);
+			LinkedInstance->UpdateAnimation(DeltaTime * GlobalAnimRateScale, false, UAnimInstance::EUpdateAnimationFlag::ForceParallelUpdate);
 		}
 
 		if (AnimScriptInstance != nullptr)
@@ -1190,9 +1190,9 @@ void USkeletalMeshComponent::TickPose(float DeltaTime, bool bNeedsValidRootMotio
 			AnimScriptInstance->OnUROSkipTickAnimation();
 		}
 
-		for(UAnimInstance* SubInstance : SubInstances)
+		for(UAnimInstance* LinkedInstance : LinkedInstances)
 		{
-			SubInstance->OnUROSkipTickAnimation();
+			LinkedInstance->OnUROSkipTickAnimation();
 		}
 
 		if(PostProcessAnimInstance)
@@ -1312,9 +1312,9 @@ void USkeletalMeshComponent::ConditionallyDispatchQueuedAnimEvents()
 	{
 		bNeedsQueuedAnimEventsDispatched = false;
 
-		for (UAnimInstance* SubInstance : SubInstances)
+		for (UAnimInstance* LinkedInstance : LinkedInstances)
 		{
-			SubInstance->DispatchQueuedAnimEvents();
+			LinkedInstance->DispatchQueuedAnimEvents();
 		}
 
 		if (AnimScriptInstance)
@@ -1500,9 +1500,9 @@ void USkeletalMeshComponent::RecalcRequiredCurves()
 		AnimScriptInstance->RecalcRequiredCurves(CurveEvalOption);
 	}
 
-	for(UAnimInstance* SubInstance : SubInstances)
+	for(UAnimInstance* LinkedInstance : LinkedInstances)
 	{
-		SubInstance->RecalcRequiredCurves(CurveEvalOption);
+		LinkedInstance->RecalcRequiredCurves(CurveEvalOption);
 	}
 
 	if(PostProcessAnimInstance)
@@ -1586,31 +1586,31 @@ void USkeletalMeshComponent::ComputeRequiredBones(TArray<FBoneIndexType>& OutReq
 	{
 		const TArray<uint8>& EditableBoneVisibilityStates = GetEditableBoneVisibilityStates();
 		check(EditableBoneVisibilityStates.Num() == GetNumComponentSpaceTransforms());
-        
+		
 		if (ensureMsgf(EditableBoneVisibilityStates.Num() >= OutRequiredBones.Num(), 
 			TEXT("Skeletal Mesh asset '%s' has incorrect BoneVisibilityStates. # of BoneVisibilityStatese (%d), # of OutRequiredBones (%d)"), 
 			*SkeletalMesh->GetName(), EditableBoneVisibilityStates.Num(), OutRequiredBones.Num()))
-        {
-            int32 VisibleBoneWriteIndex = 0;
-            for (int32 i = 0; i < OutRequiredBones.Num(); ++i)
-            {
-                FBoneIndexType CurBoneIndex = OutRequiredBones[i];
-
-                // Current bone visible?
-                if (EditableBoneVisibilityStates[CurBoneIndex] == BVS_Visible)
-                {
-                    OutRequiredBones[VisibleBoneWriteIndex++] = CurBoneIndex;
-                }
-            }
-
-            // Remove any trailing junk in the OutRequiredBones array
-            const int32 NumBonesHidden = OutRequiredBones.Num() - VisibleBoneWriteIndex;
-            if (NumBonesHidden > 0)
-            {
-                OutRequiredBones.RemoveAt(VisibleBoneWriteIndex, NumBonesHidden);
-            }
-        }
- 	}
+		{
+			int32 VisibleBoneWriteIndex = 0;
+			for (int32 i = 0; i < OutRequiredBones.Num(); ++i)
+			{
+				FBoneIndexType CurBoneIndex = OutRequiredBones[i];
+
+				// Current bone visible?
+				if (EditableBoneVisibilityStates[CurBoneIndex] == BVS_Visible)
+				{
+					OutRequiredBones[VisibleBoneWriteIndex++] = CurBoneIndex;
+				}
+			}
+
+			// Remove any trailing junk in the OutRequiredBones array
+			const int32 NumBonesHidden = OutRequiredBones.Num() - VisibleBoneWriteIndex;
+			if (NumBonesHidden > 0)
+			{
+				OutRequiredBones.RemoveAt(VisibleBoneWriteIndex, NumBonesHidden);
+			}
+		}
+	}
 
 	// Add in any bones that may be required when mirroring.
 	// JTODO: This is only required if there are mirroring nodes in the tree, but hard to know...
@@ -1702,9 +1702,9 @@ void USkeletalMeshComponent::RecalcRequiredBones(int32 LODIndex)
 		AnimScriptInstance->RecalcRequiredBones();
 	}
 
-	for (UAnimInstance* SubInstance : SubInstances)
+	for (UAnimInstance* LinkedInstance : LinkedInstances)
 	{
-		SubInstance->RecalcRequiredBones();
+		LinkedInstance->RecalcRequiredBones();
 	}
 
 	if (PostProcessAnimInstance)
@@ -1850,7 +1850,7 @@ void USkeletalMeshComponent::EvaluatePostProcessMeshInstance(TArray<FTransform>&
 	if (ShouldEvaluatePostProcessInstance())
 	{
 		// Push the previous pose to any input nodes required
-		if (FAnimNode_SubInput* InputNode = PostProcessAnimInstance->GetSubInputNode())
+		if (FAnimNode_LinkedInputPose* InputNode = PostProcessAnimInstance->GetLinkedInputPoseNode())
 		{
 			if (InOutPose.IsValid())
 			{
@@ -2094,9 +2094,9 @@ void USkeletalMeshComponent::RefreshBoneTransforms(FActorComponentTickFunction*
 			{
 				AnimScriptInstance->PreEvaluateAnimation();
 
-				for (UAnimInstance* SubInstance : SubInstances)
+				for (UAnimInstance* LinkedInstance : LinkedInstances)
 				{
-					SubInstance->PreEvaluateAnimation();
+					LinkedInstance->PreEvaluateAnimation();
 				}
 			}
 
@@ -2279,9 +2279,9 @@ void USkeletalMeshComponent::PostAnimEvaluation(FAnimationEvaluationContext& Eva
 	{
 		EvaluationContext.AnimInstance->PostUpdateAnimation();
 
-		for (UAnimInstance* SubInstance : SubInstances)
+		for (UAnimInstance* LinkedInstance : LinkedInstances)
 		{
-			SubInstance->PostUpdateAnimation();
+			LinkedInstance->PostUpdateAnimation();
 		}
 	}
 
@@ -2327,9 +2327,9 @@ void USkeletalMeshComponent::PostAnimEvaluation(FAnimationEvaluationContext& Eva
 					AnimScriptInstance->OnUROPreInterpolation();
 				}
 
-				for(UAnimInstance* SubInstance : SubInstances)
+				for(UAnimInstance* LinkedInstance : LinkedInstances)
 				{
-					SubInstance->OnUROPreInterpolation();
+					LinkedInstance->OnUROPreInterpolation();
 				}
 
 				if(PostProcessAnimInstance)
@@ -2372,9 +2372,9 @@ void USkeletalMeshComponent::PostAnimEvaluation(FAnimationEvaluationContext& Eva
 
 			// this is same curves, and we don't have to process same for everything. 
 			// we just copy curves from main for the case where GetCurveValue works in that instance
-			for(UAnimInstance* SubInstance : SubInstances)
+			for(UAnimInstance* LinkedInstance : LinkedInstances)
 			{
-				SubInstance->CopyCurveValues(*AnimScriptInstance);
+				LinkedInstance->CopyCurveValues(*AnimScriptInstance);
 			}
 		}
 
@@ -2403,9 +2403,9 @@ void USkeletalMeshComponent::PostAnimEvaluation(FAnimationEvaluationContext& Eva
 			{
 				AnimScriptInstance->PostEvaluateAnimation();
 
-				for (UAnimInstance* SubInstance : SubInstances)
+				for (UAnimInstance* LinkedInstance : LinkedInstances)
 				{
-					SubInstance->PostEvaluateAnimation();
+					LinkedInstance->PostEvaluateAnimation();
 				}
 			}
 
@@ -2507,7 +2507,7 @@ FBoxSphereBounds USkeletalMeshComponent::CalcBounds(const FTransform& LocalToWor
 		{
 			const FVector ComponentLocation = GetComponentLocation();
 			return CachedWorldSpaceBounds.TransformBy(CachedWorldToLocalTransform * LocalToWorld.ToMatrixWithScale()) 
-			        + FBoxSphereBounds(ComponentLocation, FVector(1.0f), 1.0f);
+					+ FBoxSphereBounds(ComponentLocation, FVector(1.0f), 1.0f);
 		}
 		else
 		{
@@ -2725,61 +2725,61 @@ UAnimInstance* USkeletalMeshComponent::GetPostProcessInstance() const
 	return PostProcessAnimInstance;
 }
 
-UAnimInstance* USkeletalMeshComponent::GetSubInstanceByTag(FName InName) const
+UAnimInstance* USkeletalMeshComponent::GetLinkedAnimGraphInstanceByTag(FName InName) const
 {
 	if(AnimScriptInstance)
 	{
-		return AnimScriptInstance->GetSubInstanceByTag(InName);
+		return AnimScriptInstance->GetLinkedAnimGraphInstanceByTag(InName);
 	}
 	return nullptr;
 }
 
-void USkeletalMeshComponent::GetSubInstancesByTag(FName InTag, TArray<UAnimInstance*>& OutSubInstances) const
+void USkeletalMeshComponent::GetLinkedAnimGraphInstancesByTag(FName InTag, TArray<UAnimInstance*>& OutLinkedInstances) const
 {
 	if(AnimScriptInstance)
 	{
-		AnimScriptInstance->GetSubInstancesByTag(InTag, OutSubInstances);
+		AnimScriptInstance->GetLinkedAnimGraphInstancesByTag(InTag, OutLinkedInstances);
 	}
 }
 
-void USkeletalMeshComponent::SetSubInstanceClassByTag(FName InTag, TSubclassOf<UAnimInstance> InClass)
+void USkeletalMeshComponent::LinkAnimGraphByTag(FName InTag, TSubclassOf<UAnimInstance> InClass)
 {
 	if(AnimScriptInstance)
 	{
-		AnimScriptInstance->SetSubInstanceClassByTag(InTag, InClass);
+		AnimScriptInstance->LinkAnimGraphByTag(InTag, InClass);
 	}
 }
 
-void USkeletalMeshComponent::SetLayerOverlay(TSubclassOf<UAnimInstance> InClass)
+void USkeletalMeshComponent::LinkAnimClassLayers(TSubclassOf<UAnimInstance> InClass)
 {
 	if(AnimScriptInstance)
 	{
-		AnimScriptInstance->SetLayerOverlay(InClass);
+		AnimScriptInstance->LinkAnimClassLayers(InClass);
 	}
 }
 
-void USkeletalMeshComponent::ClearLayerOverlay(TSubclassOf<UAnimInstance> InClass)
+void USkeletalMeshComponent::UnlinkAnimClassLayers(TSubclassOf<UAnimInstance> InClass)
 {
 	if(AnimScriptInstance)
 	{
-		AnimScriptInstance->ClearLayerOverlay(InClass);
+		AnimScriptInstance->UnlinkAnimClassLayers(InClass);
 	}
 }
 
-UAnimInstance* USkeletalMeshComponent::GetLayerSubInstanceByGroup(FName InGroup) const
+UAnimInstance* USkeletalMeshComponent::GetLinkedAnimLayerInstanceByGroup(FName InGroup) const
 {
 	if(AnimScriptInstance)
 	{
-		return AnimScriptInstance->GetLayerSubInstanceByGroup(InGroup);
+		return AnimScriptInstance->GetLinkedAnimLayerInstanceByGroup(InGroup);
 	}
 	return nullptr;
 }
 
-UAnimInstance* USkeletalMeshComponent::GetLayerSubInstanceByClass(TSubclassOf<UAnimInstance> InClass) const
+UAnimInstance* USkeletalMeshComponent::GetLinkedAnimLayerInstanceByClass(TSubclassOf<UAnimInstance> InClass) const
 {
 	if(AnimScriptInstance)
 	{
-		return AnimScriptInstance->GetLayerSubInstanceByClass(InClass);
+		return AnimScriptInstance->GetLinkedAnimLayerInstanceByClass(InClass);
 	}
 	return nullptr;
 }
@@ -2796,9 +2796,9 @@ void USkeletalMeshComponent::ResetAnimInstanceDynamics(ETeleportType InTeleportT
 		AnimScriptInstance->ResetDynamics(InTeleportType);
 	}
 
-	for(UAnimInstance* SubInstanceIter : SubInstances)
+	for(UAnimInstance* LinkedInstance : LinkedInstances)
 	{
-		SubInstanceIter->ResetDynamics(InTeleportType);
+		LinkedInstance->ResetDynamics(InTeleportType);
 	}
 
 	if(PostProcessAnimInstance)
@@ -3194,9 +3194,9 @@ FRootMotionMovementParams USkeletalMeshComponent::ConsumeRootMotion_Internal(flo
 	{
 		RootMotion.Accumulate(AnimScriptInstance->ConsumeExtractedRootMotion(InAlpha));
 
-		for(UAnimInstance* SubInstance : SubInstances)
+		for(UAnimInstance* LinkedInstance : LinkedInstances)
 		{
-			RootMotion.Accumulate(SubInstance->ConsumeExtractedRootMotion(InAlpha));
+			RootMotion.Accumulate(LinkedInstance->ConsumeExtractedRootMotion(InAlpha));
 		}
 	}
 
@@ -3478,9 +3478,9 @@ void USkeletalMeshComponent::RefreshMorphTargets()
 		check(!IsRunningParallelEvaluation());
 		AnimScriptInstance->RefreshCurves(this);
 
-		for(UAnimInstance* SubInstance : SubInstances)
+		for(UAnimInstance* LinkedInstance : LinkedInstances)
 		{
-			SubInstance->RefreshCurves(this);
+			LinkedInstance->RefreshCurves(this);
 		}
 		
 		if(PostProcessAnimInstance)
@@ -3558,10 +3558,10 @@ void USkeletalMeshComponent::ParallelDuplicateAndInterpolate(FAnimationEvaluatio
 					AnimScriptInstance->OnUROPreInterpolation_AnyThread(InAnimEvaluationContext);
 				}
 
-				for(UAnimInstance* SubInstance : SubInstances)
+				for(UAnimInstance* LinkedInstance : LinkedInstances)
 				{
-					SubInstance->OnUROPreInterpolation();
-					SubInstance->OnUROPreInterpolation_AnyThread(InAnimEvaluationContext);
+					LinkedInstance->OnUROPreInterpolation();
+					LinkedInstance->OnUROPreInterpolation_AnyThread(InAnimEvaluationContext);
 				}
 
 				if(PostProcessAnimInstance)
diff --git a/Engine/Source/Runtime/Engine/Public/Animation/AnimInstanceProxy.h b/Engine/Source/Runtime/Engine/Public/Animation/AnimInstanceProxy.h
--- a/Engine/Source/Runtime/Engine/Public/Animation/AnimInstanceProxy.h
+++ b/Engine/Source/Runtime/Engine/Public/Animation/AnimInstanceProxy.h
@@ -25,7 +25,7 @@ struct FAnimNode_AssetPlayerBase;
 struct FAnimNode_Base;
 struct FAnimNode_SaveCachedPose;
 struct FAnimNode_StateMachine;
-struct FAnimNode_SubInput;
+struct FAnimNode_LinkedInputPose;
 struct FNodeDebugData;
 struct FPoseContext;
 
@@ -116,7 +116,7 @@ public:
 		, CurrentDeltaSeconds(0.0f)
 		, CurrentTimeDilation(1.0f)
 		, RootNode(nullptr)
-		, DefaultSubInstanceInputNode(nullptr)
+		, DefaultLinkedInstanceInputNode(nullptr)
 		, SyncGroupWriteIndex(0)
 		, RootMotionMode(ERootMotionMode::NoRootMotionExtraction)
 		, FrameCounterForUpdate(0)
@@ -140,7 +140,7 @@ public:
 		, CurrentDeltaSeconds(0.0f)
 		, CurrentTimeDilation(1.0f)
 		, RootNode(nullptr)
-		, DefaultSubInstanceInputNode(nullptr)
+		, DefaultLinkedInstanceInputNode(nullptr)
 		, SyncGroupWriteIndex(0)
 		, RootMotionMode(ERootMotionMode::NoRootMotionExtraction)
 		, FrameCounterForUpdate(0)
@@ -461,7 +461,7 @@ public:
 	friend class UAnimInstance;
 	friend class UAnimSingleNodeInstance;
 	friend class USkeletalMeshComponent;
-	friend struct FAnimNode_SubInstance;
+	friend struct FAnimNode_LinkedAnimGraph;
 	friend struct FAnimationBaseContext;
 
 protected:
@@ -815,8 +815,8 @@ private:
 	/** Anim graph */
 	FAnimNode_Base* RootNode;
 
-	/** Default sub-instance input node if available */
-	FAnimNode_SubInput* DefaultSubInstanceInputNode;
+	/** Default linked instance input node if available */
+	FAnimNode_LinkedInputPose* DefaultLinkedInstanceInputNode;
 
 	/** Map of layer name to saved pose nodes to process after the graph has been updated */
 	TMap<FName, TArray<FAnimNode_SaveCachedPose*>> SavedPoseQueueMap;
