- Add CSV marker for PSC EOF updates

#rnx

diff --git a/Engine/Source/Runtime/Engine/Private/Particles/ParticleComponents.cpp b/Engine/Source/Runtime/Engine/Private/Particles/ParticleComponents.cpp
--- a/Engine/Source/Runtime/Engine/Private/Particles/ParticleComponents.cpp
+++ b/Engine/Source/Runtime/Engine/Private/Particles/ParticleComponents.cpp
@@ -3961,6 +3961,7 @@ void UParticleSystemComponent::SendRenderDynamicData_Concurrent()
 {
 	SCOPE_CYCLE_COUNTER(STAT_ParticleSystemComponent_SendRenderDynamicData_Concurrent);
 	SCOPE_CYCLE_COUNTER(STAT_ParticlesOverview_GT_CNC);
+	CSV_SCOPED_TIMING_STAT_EXCLUSIVE(Effects);
 
 	ForceAsyncWorkCompletion(ENSURE_AND_STALL, false, true);
 	Super::SendRenderDynamicData_Concurrent();
