Build fix: Add missing file from previous submit

diff --git a/Engine/Source/Runtime/Engine/Private/Materials/MaterialExpressions.cpp b/Engine/Source/Runtime/Engine/Private/Materials/MaterialExpressions.cpp
--- a/Engine/Source/Runtime/Engine/Private/Materials/MaterialExpressions.cpp
+++ b/Engine/Source/Runtime/Engine/Private/Materials/MaterialExpressions.cpp
@@ -2075,13 +2075,17 @@ UMaterialExpressionRuntimeVirtualTextureSample::UMaterialExpressionRuntimeVirtua
 #endif
 }
 
-void UMaterialExpressionRuntimeVirtualTextureSample::InitVirtualTextureDependentSettings()
+bool UMaterialExpressionRuntimeVirtualTextureSample::InitVirtualTextureDependentSettings()
 {
+	bool bChanged = false;
 	if (VirtualTexture != nullptr)
 	{
+		bChanged |= MaterialType != VirtualTexture->GetMaterialType();
 		MaterialType = VirtualTexture->GetMaterialType();
+		bChanged |= bSinglePhysicalSpace != VirtualTexture->GetSinglePhysicalSpace();
 		bSinglePhysicalSpace = VirtualTexture->GetSinglePhysicalSpace();
 	}
+	return bChanged;
 }
 
 void UMaterialExpressionRuntimeVirtualTextureSample::InitOutputs()
