Make replicated server time replicate at a faster rate.

diff --git a/Engine/Source/Runtime/Engine/Private/GameStateBase.cpp b/Engine/Source/Runtime/Engine/Private/GameStateBase.cpp
--- a/Engine/Source/Runtime/Engine/Private/GameStateBase.cpp
+++ b/Engine/Source/Runtime/Engine/Private/GameStateBase.cpp
@@ -30,7 +30,7 @@ AGameStateBase::AGameStateBase(const FObjectInitializer& ObjectInitializer)
 	bNetLoadOnClient = false;
 
 	// Default to every few seconds.
-	ServerWorldTimeSecondsUpdateFrequency = 5.f;
+	ServerWorldTimeSecondsUpdateFrequency = 0.1f;
 }
 
 const AGameModeBase* AGameStateBase::GetDefaultGameMode() const
