Finalize HRTF Mode to default to false, allowing users to toggle it on and off in the settings menu

diff --git a/Engine/Source/Runtime/Engine/Private/GameUserSettings.cpp b/Engine/Source/Runtime/Engine/Private/GameUserSettings.cpp
--- a/Engine/Source/Runtime/Engine/Private/GameUserSettings.cpp
+++ b/Engine/Source/Runtime/Engine/Private/GameUserSettings.cpp
@@ -463,7 +463,7 @@ void UGameUserSettings::ApplyNonResolutionSettings()
 
 	IConsoleManager::Get().CallAllConsoleVariableSinks();
 
-	EnableHeadphoneMode(true /*bUseHeadphoneMode*/);
+	EnableHeadphoneMode(bUseHeadphoneMode);
 
 	bool bWithEditor = false;
 #if WITH_EDITOR
