#ue4 - ProjectileMovement ClearPendingForces() now has an option to clear forces being processed during this update as well.

diff --git a/Engine/Source/Runtime/Engine/Classes/GameFramework/ProjectileMovementComponent.h b/Engine/Source/Runtime/Engine/Classes/GameFramework/ProjectileMovementComponent.h
--- a/Engine/Source/Runtime/Engine/Classes/GameFramework/ProjectileMovementComponent.h
+++ b/Engine/Source/Runtime/Engine/Classes/GameFramework/ProjectileMovementComponent.h
@@ -421,8 +421,8 @@ public:
 	/** Returns the sum of pending forces from AddForce(). */
 	FVector GetPendingForce() const { return PendingForce; }
 
-	/** Clears any pending forces from AddForce(). */
-	void ClearPendingForce();
+	/** Clears any pending forces from AddForce(). If bClearImmediateForce is true, clears any force being processed during this update as well. */
+	void ClearPendingForce(bool bClearImmediateForce = false);
 	
 protected:
 
diff --git a/Engine/Source/Runtime/Engine/Private/Components/ProjectileMovementComponent.cpp b/Engine/Source/Runtime/Engine/Private/Components/ProjectileMovementComponent.cpp
--- a/Engine/Source/Runtime/Engine/Private/Components/ProjectileMovementComponent.cpp
+++ b/Engine/Source/Runtime/Engine/Private/Components/ProjectileMovementComponent.cpp
@@ -494,9 +494,13 @@ void UProjectileMovementComponent::AddForce(FVector Force)
 	PendingForce += Force;
 }
 
-void UProjectileMovementComponent::ClearPendingForce()
+void UProjectileMovementComponent::ClearPendingForce(bool bClearImmediateForce)
 {
 	PendingForce = FVector::ZeroVector;
+	if (bClearImmediateForce)
+	{
+		PendingForceThisUpdate = FVector::ZeroVector;
+	}
 }
 
 float UProjectileMovementComponent::GetGravityZ() const
