Make sure file cache async requests are destroyed before handle

diff --git a/Engine/Source/Runtime/Engine/Private/FileCache/FileCache.cpp b/Engine/Source/Runtime/Engine/Private/FileCache/FileCache.cpp
--- a/Engine/Source/Runtime/Engine/Private/FileCache/FileCache.cpp
+++ b/Engine/Source/Runtime/Engine/Private/FileCache/FileCache.cpp
@@ -109,6 +109,8 @@ public:
 	// if InFile is null, will evict all slots
 	bool EvictAll(FFileCacheHandle* InFile = nullptr);
 
+	void FlushCompletedRequests();
+
 	struct FSlotInfo
 	{
 		FFileCacheHandle* Handle;
@@ -359,6 +361,15 @@ bool FFileCache::EvictAll(FFileCacheHandle* InFile)
 	return bAllOK;
 }
 
+void FFileCache::FlushCompletedRequests()
+{
+	while (IAsyncReadRequest* Request = CompletedRequests.Pop())
+	{
+		Request->WaitCompletion();
+		delete Request;
+	}
+}
+
 FFileCacheHandle::~FFileCacheHandle()
 {
 	if (SizeRequestEvent)
@@ -373,6 +384,10 @@ FFileCacheHandle::~FFileCacheHandle()
 
 		const bool result = GetCache().EvictAll(this);
 		check(result);
+
+		// Need to ensure any request created by our async handle is destroyed before destroying the handle
+		GetCache().FlushCompletedRequests();
+
 		delete InnerHandle;
 	}
 }
