Make sure we time out duplicate static guids.
This fixes an issue introduced in CL-6905893

diff --git a/Engine/Source/Runtime/Engine/Private/PackageMapClient.cpp b/Engine/Source/Runtime/Engine/Private/PackageMapClient.cpp
--- a/Engine/Source/Runtime/Engine/Private/PackageMapClient.cpp
+++ b/Engine/Source/Runtime/Engine/Private/PackageMapClient.cpp
@@ -2292,27 +2292,61 @@ void FNetGUIDCache::CleanReferences()
 {
 	const double Time = FPlatformTime::Seconds();
 
+	TMap<TWeakObjectPtr<UObject>, FNetworkGUID> StaticObjectGuids;
+
 	// Mark all static or non valid dynamic guids to timeout after NETWORK_GUID_TIMEOUT seconds
 	// We want to leave them around for a certain amount of time to allow in-flight references to these guids to continue to resolve
 	for (auto It = ObjectLookup.CreateIterator(); It; ++It)
 	{
-		if (It.Value().ReadOnlyTimestamp != 0)
+		const FNetworkGUID& Guid = It.Key();
+		FNetGuidCacheObject& CacheObject = It.Value();
+
+		if (CacheObject.ReadOnlyTimestamp != 0)
 		{
 			// If this guid was suppose to time out, check to see if it has, otherwise ignore it
 			const double NETWORK_GUID_TIMEOUT = 90;
 
-			if (Time - It.Value().ReadOnlyTimestamp > NETWORK_GUID_TIMEOUT)
+			if (Time - CacheObject.ReadOnlyTimestamp > NETWORK_GUID_TIMEOUT)
 			{
 				It.RemoveCurrent();
 			}
-
-			continue;
 		}
-
-		if (!It.Value().Object.IsValid())
+		else if (!CacheObject.Object.IsValid())
 		{
 			// We will leave this guid around for NETWORK_GUID_TIMEOUT seconds to make sure any in-flight guids can be resolved
-			It.Value().ReadOnlyTimestamp = Time;
+			CacheObject.ReadOnlyTimestamp = Time;
+		}
+
+		// Static GUIDs may refer to things on disk that don't get unloaded during travel (Packages, Sublevels, Compiled Blueprints, etc.),
+		// especially if we're traveling to the same map.
+		// The server will forcibly assign a new GUID to certain static objects, but there may be existing requests in flight
+		// already with the old GUID.
+		// So, we'll do a quick sanity check to make sure everything is fixed up.
+		// (Note, even if we get the new GUID later and register it, at worst we'll end up with 2 entries in the ObjectLookup
+		// with different GUIDs, but they'll both point at the same WeakObject, and we can clean them up later).
+		else if (Guid.IsStatic())
+		{
+			FNetworkGUID& FoundGuid = StaticObjectGuids.FindOrAdd(CacheObject.Object);
+
+			// We haven't seen this static object before, so just track it.
+			if (!FoundGuid.IsValid())
+			{
+				FoundGuid = Guid;
+			}
+
+			// We've seen this static object before, but we're seeing it again with a higher guid.
+			// That means this is our newly assigned GUID and we can safely time out the old one.
+			else if (FoundGuid.Value < Guid.Value)
+			{
+				ObjectLookup[FoundGuid].ReadOnlyTimestamp = Time;
+				FoundGuid = Guid;
+			}
+			// We've seen this static object before, but we're seeing it again with a lower guid.
+			// This means this is an older assignment and we can time out this cache object.
+			else
+			{
+				CacheObject.ReadOnlyTimestamp = Time;
+			}
 		}
 	}
 
