Fixed pin write logic for water materiaL.

diff --git a/Engine/Source/Runtime/Engine/Private/Materials/Material.cpp b/Engine/Source/Runtime/Engine/Private/Materials/Material.cpp
--- a/Engine/Source/Runtime/Engine/Private/Materials/Material.cpp
+++ b/Engine/Source/Runtime/Engine/Private/Materials/Material.cpp
@@ -5956,8 +5956,7 @@ static bool IsPropertyActive_Internal(EMaterialProperty InProperty,
 	bool bBlendableOutputAlpha,
 	bool bHasTessellation,
 	bool bHasRefraction,
-	bool bUsesShadingModelFromMaterialExpression,
-	bool bIsSingleLayerWater)
+	bool bUsesShadingModelFromMaterialExpression)
 {
 	if (Domain == MD_PostProcess)
 	{
@@ -6142,10 +6141,10 @@ static bool IsPropertyActive_Internal(EMaterialProperty InProperty,
 		Active = false;
 		break;
 	case MP_Refraction:
-		Active = (bIsTranslucentBlendMode && BlendMode != BLEND_AlphaHoldout && BlendMode != BLEND_Modulate) || bIsSingleLayerWater;
+		Active = (bIsTranslucentBlendMode && BlendMode != BLEND_AlphaHoldout && BlendMode != BLEND_Modulate) || ShadingModels.HasShadingModel(MSM_SingleLayerWater);
 		break;
 	case MP_Opacity:
-		Active = (bIsTranslucentBlendMode && BlendMode != BLEND_Modulate) || bIsSingleLayerWater;
+		Active = (bIsTranslucentBlendMode && BlendMode != BLEND_Modulate) || ShadingModels.HasShadingModel(MSM_SingleLayerWater);
 		if (IsSubsurfaceShadingModel(ShadingModels))
 		{
 			Active = true;
@@ -6212,7 +6211,6 @@ bool UMaterial::IsPropertyActive(EMaterialProperty InProperty) const
 #if WITH_EDITOR
 bool UMaterial::IsPropertyActiveInEditor(EMaterialProperty InProperty) const
 {
-	bool bUsesSingleLayerWaterMaterial = bUsedWithWater; // HasAnyExpressionsInMaterialAndFunctionsOfType<UMaterialExpressionSingleLayerWaterMaterialOutput>();
 	// explicitly DON'T use getters for BlendMode/ShadingModel...these getters may return an optimized value
 	// we want the actual value that's been set by the user in the material editor
 	return IsPropertyActive_Internal(InProperty,
@@ -6224,14 +6222,12 @@ bool UMaterial::IsPropertyActiveInEditor(EMaterialProperty InProperty) const
 		BlendableOutputAlpha,
 		D3D11TessellationMode != MTM_NoTessellation,
 		Refraction.IsConnected(),
-		IsShadingModelFromMaterialExpression(),
-		bUsesSingleLayerWaterMaterial);
+		IsShadingModelFromMaterialExpression());
 }
 #endif // WITH_EDITOR
 
 bool UMaterial::IsPropertyActiveInDerived(EMaterialProperty InProperty, const UMaterialInterface* DerivedMaterial) const
 {
-	bool bUsesSingleLayerWaterMaterial = DerivedMaterial->GetMaterial()->bUsedWithWater; //HasAnyExpressionsInMaterialAndFunctionsOfType<UMaterialExpressionSingleLayerWaterMaterialOutput>();
 	return IsPropertyActive_Internal(InProperty,
 		MaterialDomain,
 		DerivedMaterial->GetBlendMode(),
@@ -6241,8 +6237,7 @@ bool UMaterial::IsPropertyActiveInDerived(EMaterialProperty InProperty, const UM
 		BlendableOutputAlpha,
 		D3D11TessellationMode != MTM_NoTessellation,
 		Refraction.IsConnected(),
-		DerivedMaterial->IsShadingModelFromMaterialExpression(),
-		bUsesSingleLayerWaterMaterial);
+		DerivedMaterial->IsShadingModelFromMaterialExpression());
 }
 
 #if WITH_EDITORONLY_DATA
