Fix warnings: 'uint32': forcing value to bool 'true' or 'false' (performance warning)
#rnx

diff --git a/Engine/Source/Runtime/Engine/Public/MaterialShared.h b/Engine/Source/Runtime/Engine/Public/MaterialShared.h
--- a/Engine/Source/Runtime/Engine/Public/MaterialShared.h
+++ b/Engine/Source/Runtime/Engine/Public/MaterialShared.h
@@ -560,7 +560,7 @@ public:
 
 	ENGINE_API void Serialize(FArchive& Ar);
 
-	ENGINE_API bool IsSceneTextureUsed(ESceneTextureId TexId) const { return UsedSceneTextures & (1 << TexId); }
+	ENGINE_API bool IsSceneTextureUsed(ESceneTextureId TexId) const { return (UsedSceneTextures & (1 << TexId)) != 0; }
 	ENGINE_API void SetIsSceneTextureUsed(ESceneTextureId TexId) { UsedSceneTextures |= (1 << TexId); }
 
 	FUniformExpressionSet UniformExpressionSet;
@@ -1115,7 +1115,7 @@ public:
 	uint32 GetEstimatedNumVirtualTextureLookups() const { return MaterialCompilationOutput.EstimatedNumVirtualTextureLookups; }
 #endif
 	uint32 GetNumVirtualTextureStacks() const { return MaterialCompilationOutput.UniformExpressionSet.VTStacks.Num(); }
-	bool UsesSceneTexture(uint32 TexId) const { return MaterialCompilationOutput.UsedSceneTextures & (1ull << TexId); }
+	bool UsesSceneTexture(uint32 TexId) const { return (MaterialCompilationOutput.UsedSceneTextures & (1ull << TexId)) != 0; }
 
 	bool IsValidForRendering(bool bFailOnInvalid = false) const
 	{
