Reuse UHLODProxy::ExtractStaticMeshComponentsFromLODActor() in ALODActor::HasValidSubActors() to properly process child ALODActors

diff --git a/Engine/Source/Runtime/Engine/Private/LODActor.cpp b/Engine/Source/Runtime/Engine/Private/LODActor.cpp
--- a/Engine/Source/Runtime/Engine/Private/LODActor.cpp
+++ b/Engine/Source/Runtime/Engine/Private/LODActor.cpp
@@ -907,41 +907,19 @@ void ALODActor::DetermineShadowingFlags()
 
 const bool ALODActor::HasValidSubActors() const
 {
-#if WITH_EDITOR
-	FHierarchicalLODUtilitiesModule& Module = FModuleManager::LoadModuleChecked<FHierarchicalLODUtilitiesModule>("HierarchicalLODUtilities");
-	IHierarchicalLODUtilities* Utilities = Module.GetUtilities();
-#endif
-
-	int32 NumMeshes = 0;
+	TArray<UStaticMeshComponent*> Components;
+	UHLODProxy::ExtractStaticMeshComponentsFromLODActor(this, Components);
 
-	// Make sure there is at least one mesh in the subactors
-	for (AActor* SubActor : SubActors)
+	UStaticMeshComponent** ValidComponent = Components.FindByPredicate([&](const UStaticMeshComponent* Component)
 	{
-		if (SubActor)
-		{
-			for (UActorComponent* Comp : SubActor->GetComponents())
-			{
-				if (UStaticMeshComponent* Component = Cast<UStaticMeshComponent>(Comp))
-				{
 #if WITH_EDITOR
-					if (!Component->bHiddenInGame && Component->ShouldGenerateAutoLOD(LODLevel - 1))
-					{
-						++NumMeshes;
-					}
+		return !Component->bHiddenInGame && Component->GetStaticMesh() != nullptr && Component->ShouldGenerateAutoLOD(LODLevel - 1);
 #else
-					++NumMeshes;
+		return true;
 #endif
-				}
-			}
-
-			if (NumMeshes > 0)
-			{
-				break;
-			}
-		}
-	}
+	});
 
-	return NumMeshes > 0;
+	return ValidComponent != nullptr;
 }
 
 const bool ALODActor::HasAnySubActors() const
