Fix weight propgation over sub-instance boundaries

New Update* API now passes through context from a higher level. Old API calls deprecated & fixed up.
We were previously doing the right thing with regards weights on the way *out* of a linked graph (via a linked input pose). This addresses the way *into* a linked graph.

Resolution in AnimInstanceProxy.cpp stomps Dev-Anim changes. Jira to be filed to revisit the consequences to inertial blending.

diff --git a/Engine/Plugins/Animation/LiveLink/Source/LiveLink/Private/LiveLinkInstance.cpp b/Engine/Plugins/Animation/LiveLink/Source/LiveLink/Private/LiveLinkInstance.cpp
--- a/Engine/Plugins/Animation/LiveLink/Source/LiveLink/Private/LiveLinkInstance.cpp
+++ b/Engine/Plugins/Animation/LiveLink/Source/LiveLink/Private/LiveLinkInstance.cpp
@@ -18,12 +18,11 @@ bool FLiveLinkInstanceProxy::Evaluate(FPoseContext& Output)
 	return true;
 }
 
-void FLiveLinkInstanceProxy::UpdateAnimationNode(float DeltaSeconds)
+void FLiveLinkInstanceProxy::UpdateAnimationNode(const FAnimationUpdateContext& InContext)
 {
 	UpdateCounter.Increment();
 
-	FAnimationUpdateContext UpdateContext(this, DeltaSeconds);
-	PoseNode.Update_AnyThread(UpdateContext);
+	PoseNode.Update_AnyThread(InContext);
 	
 	if(ULiveLinkInstance* Instance = Cast<ULiveLinkInstance>(GetAnimInstanceObject()))
 	{
diff --git a/Engine/Plugins/Animation/LiveLink/Source/LiveLink/Public/LiveLinkInstance.h b/Engine/Plugins/Animation/LiveLink/Source/LiveLink/Public/LiveLinkInstance.h
--- a/Engine/Plugins/Animation/LiveLink/Source/LiveLink/Public/LiveLinkInstance.h
+++ b/Engine/Plugins/Animation/LiveLink/Source/LiveLink/Public/LiveLinkInstance.h
@@ -30,7 +30,7 @@ public:
 
 	virtual void Initialize(UAnimInstance* InAnimInstance) override;
 	virtual bool Evaluate(FPoseContext& Output) override;
-	virtual void UpdateAnimationNode(float DeltaSeconds) override;
+	virtual void UpdateAnimationNode(const FAnimationUpdateContext& InContext) override;
 
 	UPROPERTY(EditAnywhere, Category = Settings)
 	FAnimNode_LiveLinkPose PoseNode;
diff --git a/Engine/Source/Runtime/AnimGraphRuntime/Private/AnimSequencerInstanceProxy.cpp b/Engine/Source/Runtime/AnimGraphRuntime/Private/AnimSequencerInstanceProxy.cpp
--- a/Engine/Source/Runtime/AnimGraphRuntime/Private/AnimSequencerInstanceProxy.cpp
+++ b/Engine/Source/Runtime/AnimGraphRuntime/Private/AnimSequencerInstanceProxy.cpp
@@ -18,11 +18,11 @@ bool FAnimSequencerInstanceProxy::Evaluate(FPoseContext& Output)
 	return true;
 }
 
-void FAnimSequencerInstanceProxy::UpdateAnimationNode(float DeltaSeconds)
+void FAnimSequencerInstanceProxy::UpdateAnimationNode(const FAnimationUpdateContext& InContext)
 {
 	UpdateCounter.Increment();
 
-	SequencerRootNode.Update_AnyThread(FAnimationUpdateContext(this, DeltaSeconds));
+	SequencerRootNode.Update_AnyThread(InContext);
 }
 
 void FAnimSequencerInstanceProxy::ConstructNodes()
diff --git a/Engine/Source/Runtime/AnimGraphRuntime/Public/AnimSequencerInstanceProxy.h b/Engine/Source/Runtime/AnimGraphRuntime/Public/AnimSequencerInstanceProxy.h
--- a/Engine/Source/Runtime/AnimGraphRuntime/Public/AnimSequencerInstanceProxy.h
+++ b/Engine/Source/Runtime/AnimGraphRuntime/Public/AnimSequencerInstanceProxy.h
@@ -79,7 +79,7 @@ public:
 	// FAnimInstanceProxy interface
 	virtual void Initialize(UAnimInstance* InAnimInstance) override;
 	virtual bool Evaluate(FPoseContext& Output) override;
-	virtual void UpdateAnimationNode(float DeltaSeconds) override;
+	virtual void UpdateAnimationNode(const FAnimationUpdateContext& InContext) override;
 
 	/** Update an animation sequence player in this instance */
 	void UpdateAnimTrack(UAnimSequenceBase* InAnimSequence, uint32 SequenceId, float InPosition, float Weight, bool bFireNotifies);
diff --git a/Engine/Source/Runtime/Engine/Private/Animation/AnimInstanceProxy.cpp b/Engine/Source/Runtime/Engine/Private/Animation/AnimInstanceProxy.cpp
--- a/Engine/Source/Runtime/Engine/Private/Animation/AnimInstanceProxy.cpp
+++ b/Engine/Source/Runtime/Engine/Private/Animation/AnimInstanceProxy.cpp
@@ -33,12 +33,12 @@ const FName NAME_Evaluate(TEXT("Evaluate"));
 const FName NAME_Update(TEXT("Update"));
 const FName NAME_AnimGraph(TEXT("AnimGraph"));
 
-void FAnimInstanceProxy::UpdateAnimationNode(float DeltaSeconds)
+void FAnimInstanceProxy::UpdateAnimationNode(const FAnimationUpdateContext& InContext)
 {
-	UpdateAnimationNode_WithRoot(DeltaSeconds, RootNode, NAME_AnimGraph);
+	UpdateAnimationNode_WithRoot(InContext, RootNode, NAME_AnimGraph);
 }
 
-void FAnimInstanceProxy::UpdateAnimationNode_WithRoot(float DeltaSeconds, FAnimNode_Base* InRootNode, FName InLayerName)
+void FAnimInstanceProxy::UpdateAnimationNode_WithRoot(const FAnimationUpdateContext& InContext, FAnimNode_Base* InRootNode, FName InLayerName)
 {
 	DECLARE_SCOPE_HIERARCHICAL_COUNTER_FUNC()
 	if(InRootNode != nullptr)
@@ -48,9 +48,7 @@ void FAnimInstanceProxy::UpdateAnimationNode_WithRoot(float DeltaSeconds, FAnimN
 			UpdateCounter.Increment();
 		}
 		
-		FAnimationUpdateSharedContext SharedContext;
-		FAnimationUpdateContext UpdateContext(this, DeltaSeconds, &SharedContext);
-		InRootNode->Update_AnyThread(UpdateContext);
+		InRootNode->Update_AnyThread(InContext);
 
 		// We've updated the graph, now update the fractured saved pose sections
 		TArray<FAnimNode_SaveCachedPose*>& SavedPoseQueue = SavedPoseQueueMap.FindChecked(InLayerName);
@@ -1183,10 +1181,11 @@ void FAnimInstanceProxy::UpdateAnimation()
 	UpdatedNodesThisFrame.Reset();
 #endif
 
-	UpdateAnimation_WithRoot(RootNode, NAME_AnimGraph);
+	FAnimationUpdateContext Context(this, CurrentDeltaSeconds);
+	UpdateAnimation_WithRoot(Context, RootNode, NAME_AnimGraph);
 }
 
-void FAnimInstanceProxy::UpdateAnimation_WithRoot(FAnimNode_Base* InRootNode, FName InLayerName)
+void FAnimInstanceProxy::UpdateAnimation_WithRoot(const FAnimationUpdateContext& InContext, FAnimNode_Base* InRootNode, FName InLayerName)
 {
 	DECLARE_SCOPE_HIERARCHICAL_COUNTER_FUNC()
 
@@ -1233,11 +1232,11 @@ void FAnimInstanceProxy::UpdateAnimation_WithRoot(FAnimNode_Base* InRootNode, FN
 		if(InRootNode == RootNode)
 		{
 			// Call the correct override point if this is the root node
-			UpdateAnimationNode(CurrentDeltaSeconds);
+			UpdateAnimationNode(InContext);
 		}
 		else
 		{
-			UpdateAnimationNode_WithRoot(CurrentDeltaSeconds, InRootNode, InLayerName);
+			UpdateAnimationNode_WithRoot(InContext, InRootNode, InLayerName);
 		}
 	}
 }
diff --git a/Engine/Source/Runtime/Engine/Private/Animation/AnimNode_LinkedAnimGraph.cpp b/Engine/Source/Runtime/Engine/Private/Animation/AnimNode_LinkedAnimGraph.cpp
--- a/Engine/Source/Runtime/Engine/Private/Animation/AnimNode_LinkedAnimGraph.cpp
+++ b/Engine/Source/Runtime/Engine/Private/Animation/AnimNode_LinkedAnimGraph.cpp
@@ -58,29 +58,30 @@ void FAnimNode_LinkedAnimGraph::CacheBones_AnyThread(const FAnimationCacheBonesC
 	}
 }
 
-void FAnimNode_LinkedAnimGraph::Update_AnyThread(const FAnimationUpdateContext& Context)
+void FAnimNode_LinkedAnimGraph::Update_AnyThread(const FAnimationUpdateContext& InContext)
 {
-	GetEvaluateGraphExposedInputs().Execute(Context);
+	GetEvaluateGraphExposedInputs().Execute(InContext);
 
 	UAnimInstance* InstanceToRun = GetTargetInstance<UAnimInstance>();
 	if(InstanceToRun && LinkedRoot)
 	{
 		FAnimInstanceProxy& Proxy = InstanceToRun->GetProxyOnAnyThread<FAnimInstanceProxy>();
-		Proxy.UpdateCounter.SynchronizeWith(Context.AnimInstanceProxy->UpdateCounter);
+		Proxy.UpdateCounter.SynchronizeWith(InContext.AnimInstanceProxy->UpdateCounter);
 
-		PropagateInputProperties(Context.AnimInstanceProxy->GetAnimInstanceObject());
+		PropagateInputProperties(InContext.AnimInstanceProxy->GetAnimInstanceObject());
 
 		// Only update if we've not had a single-threaded update already
 		if(InstanceToRun->bNeedsUpdate)
 		{
-			Proxy.UpdateAnimation_WithRoot(LinkedRoot, GetDynamicLinkFunctionName());
+			FAnimationUpdateContext NewContext = InContext.WithOtherProxy(&Proxy);
+			Proxy.UpdateAnimation_WithRoot(NewContext, LinkedRoot, GetDynamicLinkFunctionName());
 		}
 	}
 	else if(InputPoses.Num() > 0)
 	{
 		// If we have no valid instance (self or otherwise), we need to propagate down the graph to make sure
 		// subsequent nodes get properly updated
-		InputPoses[0].Update(Context);
+		InputPoses[0].Update(InContext);
 	}
 }
 
diff --git a/Engine/Source/Runtime/Engine/Private/Animation/AnimSingleNodeInstanceProxy.cpp b/Engine/Source/Runtime/Engine/Private/Animation/AnimSingleNodeInstanceProxy.cpp
--- a/Engine/Source/Runtime/Engine/Private/Animation/AnimSingleNodeInstanceProxy.cpp
+++ b/Engine/Source/Runtime/Engine/Private/Animation/AnimSingleNodeInstanceProxy.cpp
@@ -59,12 +59,11 @@ void FAnimSingleNodeInstanceProxy::PropagatePreviewCurve(FPoseContext& Output)
 }
 #endif // WITH_EDITORONLY_DATA
 
-void FAnimSingleNodeInstanceProxy::UpdateAnimationNode(float DeltaSeconds)
+void FAnimSingleNodeInstanceProxy::UpdateAnimationNode(const FAnimationUpdateContext& InContext)
 {
 	UpdateCounter.Increment();
 
-	FAnimationUpdateContext UpdateContext(this, DeltaSeconds);
-	SingleNode.Update_AnyThread(UpdateContext);
+	SingleNode.Update_AnyThread(InContext);
 }
 
 void FAnimSingleNodeInstanceProxy::PostUpdate(UAnimInstance* InAnimInstance) const
diff --git a/Engine/Source/Runtime/Engine/Public/Animation/AnimInstanceProxy.h b/Engine/Source/Runtime/Engine/Public/Animation/AnimInstanceProxy.h
--- a/Engine/Source/Runtime/Engine/Public/Animation/AnimInstanceProxy.h
+++ b/Engine/Source/Runtime/Engine/Public/Animation/AnimInstanceProxy.h
@@ -480,11 +480,25 @@ protected:
 	/** Update override point */
 	virtual void Update(float DeltaSeconds) {}
 
+	UE_DEPRECATED(4.24, "Please use the overload that takes an FAnimationUpdateContext")
+	virtual void UpdateAnimationNode(float DeltaSeconds)
+	{
+		FAnimationUpdateContext Context(this, DeltaSeconds);
+		UpdateAnimationNode(Context);
+	}
+
 	/** Updates the anim graph */
-	virtual void UpdateAnimationNode(float DeltaSeconds);
+	virtual void UpdateAnimationNode(const FAnimationUpdateContext& InContext);
+
+	UE_DEPRECATED(4.24, "Please use the overload that takes an FAnimationUpdateContext")
+	virtual void UpdateAnimationNode_WithRoot(float DeltaSeconds, FAnimNode_Base* InRootNode, FName InLayerName) 
+	{
+		FAnimationUpdateContext Context(this, DeltaSeconds);
+		UpdateAnimationNode_WithRoot(Context, InRootNode, InLayerName);
+	}
 
 	/** Updates the anim graph using a specified root node */
-	virtual void UpdateAnimationNode_WithRoot(float DeltaSeconds, FAnimNode_Base* InRootNode, FName InLayerName);
+	virtual void UpdateAnimationNode_WithRoot(const FAnimationUpdateContext& InContext, FAnimNode_Base* InRootNode, FName InLayerName);
 
 	/** Called on the game thread pre-evaluate. */
 	virtual void PreEvaluateAnimation(UAnimInstance* InAnimInstance);
@@ -547,8 +561,15 @@ protected:
 	/** Calls Update(), updates the anim graph, ticks asset players */
 	void UpdateAnimation();
 
+	UE_DEPRECATED(4.24, "Please use the overload that takes an FAnimationUpdateContext")
+	void UpdateAnimation_WithRoot(FAnimNode_Base* InRootNode, FName InLayerName)
+	{
+		FAnimationUpdateContext Context(this, CurrentDeltaSeconds);
+		UpdateAnimation_WithRoot(Context, InRootNode, InLayerName);
+	}
+
 	/** Calls Update(), updates the anim graph from the specified root, ticks asset players */
-	void UpdateAnimation_WithRoot(FAnimNode_Base* InRootNode, FName InLayerName);
+	void UpdateAnimation_WithRoot(const FAnimationUpdateContext& InContext, FAnimNode_Base* InRootNode, FName InLayerName);
 
 	/** Evaluates the anim graph if Evaluate() returns false */
 	void EvaluateAnimation(FPoseContext& Output);
diff --git a/Engine/Source/Runtime/Engine/Public/Animation/AnimSingleNodeInstanceProxy.h b/Engine/Source/Runtime/Engine/Public/Animation/AnimSingleNodeInstanceProxy.h
--- a/Engine/Source/Runtime/Engine/Public/Animation/AnimSingleNodeInstanceProxy.h
+++ b/Engine/Source/Runtime/Engine/Public/Animation/AnimSingleNodeInstanceProxy.h
@@ -80,7 +80,7 @@ public:
 
 	virtual void Initialize(UAnimInstance* InAnimInstance) override;
 	virtual bool Evaluate(FPoseContext& Output) override;
-	virtual void UpdateAnimationNode(float DeltaSeconds) override;
+	virtual void UpdateAnimationNode(const FAnimationUpdateContext& InContext) override;
 	virtual void PostUpdate(UAnimInstance* InAnimInstance) const override;
 	virtual void PreUpdate(UAnimInstance* InAnimInstance, float DeltaSeconds) override;
 	virtual void InitializeObjects(UAnimInstance* InAnimInstance) override;
