Add support for material attributes to MaterialProxyReplace material node

diff --git a/Engine/Source/Runtime/Engine/Classes/Materials/MaterialExpressionMaterialProxyReplace.h b/Engine/Source/Runtime/Engine/Classes/Materials/MaterialExpressionMaterialProxyReplace.h
--- a/Engine/Source/Runtime/Engine/Classes/Materials/MaterialExpressionMaterialProxyReplace.h
+++ b/Engine/Source/Runtime/Engine/Classes/Materials/MaterialExpressionMaterialProxyReplace.h
@@ -24,6 +24,9 @@ class UMaterialExpressionMaterialProxyReplace : public UMaterialExpression
 	// Begin UMaterialExpression Interface
 #if WITH_EDITOR
 	virtual int32 Compile(class FMaterialCompiler* Compiler, int32 OutputIndex) override;
+	virtual bool IsResultMaterialAttributes(int32 OutputIndex) override;
+	virtual uint32 GetInputType(int32 InputIndex) override { return MCT_Unknown; }
+	virtual uint32 GetOutputType(int32 OutputIndex) override { return MCT_Unknown; }
 	virtual void GetCaption(TArray<FString>& OutCaptions) const override;
 #endif
 	// End UMaterialExpression Interface
diff --git a/Engine/Source/Runtime/Engine/Private/Materials/MaterialExpressions.cpp b/Engine/Source/Runtime/Engine/Private/Materials/MaterialExpressions.cpp
--- a/Engine/Source/Runtime/Engine/Private/Materials/MaterialExpressions.cpp
+++ b/Engine/Source/Runtime/Engine/Private/Materials/MaterialExpressions.cpp
@@ -14009,6 +14009,18 @@ int32 UMaterialExpressionMaterialProxyReplace::Compile(class FMaterialCompiler*
 	}
 }
 
+bool UMaterialExpressionMaterialProxyReplace::IsResultMaterialAttributes(int32 OutputIndex)
+{
+	for (FExpressionInput* ExpressionInput : GetInputs())
+	{
+		if (ExpressionInput->GetTracedInput().Expression && !ExpressionInput->Expression->ContainsInputLoop() && ExpressionInput->Expression->IsResultMaterialAttributes(ExpressionInput->OutputIndex))
+		{
+			return true;
+		}
+	}
+	return false;
+}
+
 void UMaterialExpressionMaterialProxyReplace::GetCaption(TArray<FString>& OutCaptions) const
 {
 	OutCaptions.Add(TEXT("MaterialProxyReplace"));
