Made some CSV stats off by default as they were firing too frequently.

diff --git a/Engine/Source/Runtime/Engine/Private/Collision/SceneQuery.cpp b/Engine/Source/Runtime/Engine/Private/Collision/SceneQuery.cpp
--- a/Engine/Source/Runtime/Engine/Private/Collision/SceneQuery.cpp
+++ b/Engine/Source/Runtime/Engine/Private/Collision/SceneQuery.cpp
@@ -30,6 +30,8 @@ float DebugLineLifetime = 2.f;
 #include "Physics/Experimental/ChaosInterfaceWrapper.h"
 #endif
 
+CSV_DEFINE_CATEGORY(SceneQuery, false);
+
 enum class ESingleMultiOrTest : uint8
 {
 	Single,
@@ -440,7 +442,7 @@ bool FGenericPhysicsInterface::RaycastTest(const UWorld* World, const FVector St
 {
 	SCOPE_CYCLE_COUNTER(STAT_Collision_SceneQueryTotal);
 	SCOPE_CYCLE_COUNTER(STAT_Collision_RaycastAny);
-	CSV_SCOPED_TIMING_STAT_EXCLUSIVE(RaycastTest);
+	CSV_SCOPED_TIMING_STAT(SceneQuery, RaycastTest);
 
 	using TCastTraits = TSQTraits<FHitRaycast, ESweepOrRay::Raycast, ESingleMultiOrTest::Test>;
 	FHitResult DummyHit(NoInit);
@@ -451,7 +453,7 @@ bool FGenericPhysicsInterface::RaycastSingle(const UWorld* World, struct FHitRes
 {
 	SCOPE_CYCLE_COUNTER(STAT_Collision_SceneQueryTotal);
 	SCOPE_CYCLE_COUNTER(STAT_Collision_RaycastSingle);
-	CSV_SCOPED_TIMING_STAT_EXCLUSIVE(RaycastSingle);
+	CSV_SCOPED_TIMING_STAT(SceneQuery, RaycastSingle);
 
 	using TCastTraits = TSQTraits<FHitRaycast, ESweepOrRay::Raycast, ESingleMultiOrTest::Single>;
 	return TSceneCastCommon<TCastTraits>(World, OutHit, FRaycastSQAdditionalInputs(), Start, End, TraceChannel, Params, ResponseParams, ObjectParams);
@@ -463,7 +465,7 @@ bool FGenericPhysicsInterface::RaycastMulti(const UWorld* World, TArray<struct F
 {
 	SCOPE_CYCLE_COUNTER(STAT_Collision_SceneQueryTotal);
 	SCOPE_CYCLE_COUNTER(STAT_Collision_RaycastMultiple);
-	CSV_SCOPED_TIMING_STAT_EXCLUSIVE(RaycastMultiple);
+	CSV_SCOPED_TIMING_STAT(SceneQuery, RaycastMultiple);
 
 	using TCastTraits = TSQTraits<FHitRaycast, ESweepOrRay::Raycast, ESingleMultiOrTest::Multi>;
 	return TSceneCastCommon<TCastTraits> (World, OutHits, FRaycastSQAdditionalInputs(), Start, End, TraceChannel, Params, ResponseParams, ObjectParams);
@@ -476,7 +478,7 @@ bool FGenericPhysicsInterface::GeomSweepTest(const UWorld* World, const struct F
 {
 	SCOPE_CYCLE_COUNTER(STAT_Collision_SceneQueryTotal);
 	SCOPE_CYCLE_COUNTER(STAT_Collision_GeomSweepAny);
-	CSV_SCOPED_TIMING_STAT_EXCLUSIVE(GeomSweepTest);
+	CSV_SCOPED_TIMING_STAT(SceneQuery, GeomSweepTest);
 
 	using TCastTraits = TSQTraits<FHitSweep, ESweepOrRay::Sweep, ESingleMultiOrTest::Test>;
 	FHitResult DummyHit(NoInit);
@@ -487,7 +489,7 @@ bool FGenericPhysicsInterface::GeomSweepSingle(const UWorld* World, const struct
 {
 	SCOPE_CYCLE_COUNTER(STAT_Collision_SceneQueryTotal);
 	SCOPE_CYCLE_COUNTER(STAT_Collision_GeomSweepSingle);
-	CSV_SCOPED_TIMING_STAT_EXCLUSIVE(GeomSweepSingle);
+	CSV_SCOPED_TIMING_STAT(SceneQuery, GeomSweepSingle);
 
 	using TCastTraits = TSQTraits<FHitSweep, ESweepOrRay::Sweep, ESingleMultiOrTest::Single>;
 	return TSceneCastCommon<TCastTraits>(World, OutHit, FGeomSQAdditionalInputs(CollisionShape, Rot), Start, End, TraceChannel, Params, ResponseParams, ObjectParams);
@@ -498,7 +500,7 @@ bool FGenericPhysicsInterface::GeomSweepMulti(const UWorld* World, const FPhysic
 {
 	SCOPE_CYCLE_COUNTER(STAT_Collision_SceneQueryTotal);
 	SCOPE_CYCLE_COUNTER(STAT_Collision_GeomSweepMultiple);
-	CSV_SCOPED_TIMING_STAT_EXCLUSIVE(GeomSweepMultiple);
+	CSV_SCOPED_TIMING_STAT(SceneQuery, GeomSweepMultiple);
 
 	using TCastTraits = TSQTraits<FHitSweep, ESweepOrRay::Sweep, ESingleMultiOrTest::Multi>;
 	return TSceneCastCommon<TCastTraits>(World, OutHits, FGeomCollectionSQAdditionalInputs(InGeom, InGeomRot), Start, End, TraceChannel, Params, ResponseParams, ObjectParams);
@@ -509,7 +511,7 @@ bool FGenericPhysicsInterface::GeomSweepMulti(const UWorld* World, const FCollis
 {
 	SCOPE_CYCLE_COUNTER(STAT_Collision_SceneQueryTotal);
 	SCOPE_CYCLE_COUNTER(STAT_Collision_GeomSweepMultiple);
-	CSV_SCOPED_TIMING_STAT_EXCLUSIVE(GeomSweepMultiple);
+	CSV_SCOPED_TIMING_STAT(SceneQuery, GeomSweepMultiple);
 
 	using TCastTraits = TSQTraits<FHitSweep, ESweepOrRay::Sweep, ESingleMultiOrTest::Multi>;
 	return TSceneCastCommon<TCastTraits>(World, OutHits, FGeomSQAdditionalInputs(InGeom, InGeomRot), Start, End, TraceChannel, Params, ResponseParams, ObjectParams);
@@ -617,7 +619,7 @@ bool FGenericPhysicsInterface::GeomOverlapBlockingTest(const UWorld* World, cons
 {
 	SCOPE_CYCLE_COUNTER(STAT_Collision_SceneQueryTotal);
 	SCOPE_CYCLE_COUNTER(STAT_Collision_GeomOverlapBlocking);
-	CSV_SCOPED_TIMING_STAT_EXCLUSIVE(GeomOverlapBlocking);
+	CSV_SCOPED_TIMING_STAT(SceneQuery, GeomOverlapBlocking);
 
 	TArray<FOverlapResult> Overlaps;	//needed only for template shared code
 	FTransform GeomTransform(Rot, Pos);
@@ -633,7 +635,7 @@ bool FGenericPhysicsInterface::GeomOverlapAnyTest(const UWorld* World, const str
 {
 	SCOPE_CYCLE_COUNTER(STAT_Collision_SceneQueryTotal);
 	SCOPE_CYCLE_COUNTER(STAT_Collision_GeomOverlapAny);
-	CSV_SCOPED_TIMING_STAT_EXCLUSIVE(GeomOverlapAny);
+	CSV_SCOPED_TIMING_STAT(SceneQuery, GeomOverlapAny);
 
 	TArray<FOverlapResult> Overlaps;	//needed only for template shared code
 	FTransform GeomTransform(Rot, Pos);
@@ -650,7 +652,7 @@ bool FGenericPhysicsInterface::GeomOverlapMulti(const UWorld* World, const FPhys
 {
 	SCOPE_CYCLE_COUNTER(STAT_Collision_SceneQueryTotal);
 	SCOPE_CYCLE_COUNTER(STAT_Collision_GeomOverlapMultiple);
-	CSV_SCOPED_TIMING_STAT_EXCLUSIVE(GeomOverlapMultiple);
+	CSV_SCOPED_TIMING_STAT(SceneQuery, GeomOverlapMultiple);
 
 #if WITH_PHYSX && !WITH_CHAOS_NEEDS_TO_BE_FIXED
 	FTransform GeomTransform(InRotation, InPosition);
@@ -665,7 +667,7 @@ bool FGenericPhysicsInterface::GeomOverlapMulti(const UWorld* World, const FColl
 {
 	SCOPE_CYCLE_COUNTER(STAT_Collision_SceneQueryTotal);
 	SCOPE_CYCLE_COUNTER(STAT_Collision_GeomOverlapMultiple);
-	CSV_SCOPED_TIMING_STAT_EXCLUSIVE(GeomOverlapMultiple);
+	CSV_SCOPED_TIMING_STAT(SceneQuery, GeomOverlapMultiple);
 
 	FTransform GeomTransform(InRotation, InPosition);
 #if WITH_PHYSX && !WITH_CHAOS_NEEDS_TO_BE_FIXED
diff --git a/Engine/Source/Runtime/Engine/Private/Components/PrimitiveComponent.cpp b/Engine/Source/Runtime/Engine/Private/Components/PrimitiveComponent.cpp
--- a/Engine/Source/Runtime/Engine/Private/Components/PrimitiveComponent.cpp
+++ b/Engine/Source/Runtime/Engine/Private/Components/PrimitiveComponent.cpp
@@ -45,6 +45,8 @@
 
 #define LOCTEXT_NAMESPACE "PrimitiveComponent"
 
+CSV_DEFINE_CATEGORY(PrimitiveComponent, false);
+
 //////////////////////////////////////////////////////////////////////////
 // Globals
 
@@ -1952,7 +1954,7 @@ FCollisionShape UPrimitiveComponent::GetCollisionShape(float Inflation) const
 bool UPrimitiveComponent::MoveComponentImpl( const FVector& Delta, const FQuat& NewRotationQuat, bool bSweep, FHitResult* OutHit, EMoveComponentFlags MoveFlags, ETeleportType Teleport)
 {
 	SCOPE_CYCLE_COUNTER(STAT_MoveComponentTime);
-	CSV_SCOPED_TIMING_STAT_EXCLUSIVE(MoveComponentTime);
+	CSV_SCOPED_TIMING_STAT(PrimitiveComponent, MoveComponentTime);
 
 #if !(UE_BUILD_SHIPPING || UE_BUILD_TEST) && PERF_MOVECOMPONENT_STATS
 	FScopedMoveCompTimer MoveTimer(this, Delta);
diff --git a/Engine/Source/Runtime/NavigationSystem/Private/NavigationSystem.cpp b/Engine/Source/Runtime/NavigationSystem/Private/NavigationSystem.cpp
--- a/Engine/Source/Runtime/NavigationSystem/Private/NavigationSystem.cpp
+++ b/Engine/Source/Runtime/NavigationSystem/Private/NavigationSystem.cpp
@@ -117,6 +117,8 @@ DEFINE_STAT(STAT_Navigation_PathVisibilityOptimisation);
 DEFINE_STAT(STAT_Navigation_ObservedPathsCount);
 DEFINE_STAT(STAT_Navigation_RecastMemory);
 
+CSV_DEFINE_CATEGORY(NavigationSystem, false);
+
 //----------------------------------------------------------------------//
 // consts
 //----------------------------------------------------------------------//
@@ -1056,7 +1058,7 @@ void UNavigationSystemV1::SetNavigationAutoUpdateEnabled(bool bNewEnable, UNavig
 FPathFindingResult UNavigationSystemV1::FindPathSync(const FNavAgentProperties& AgentProperties, FPathFindingQuery Query, EPathFindingMode::Type Mode)
 {
 	SCOPE_CYCLE_COUNTER(STAT_Navigation_PathfindingSync);
-	CSV_SCOPED_TIMING_STAT_EXCLUSIVE(PathfindingSync);
+	CSV_SCOPED_TIMING_STAT(NavigationSystem, PathfindingSync);
 
 	if (Query.NavData.IsValid() == false)
 	{
@@ -1082,7 +1084,7 @@ FPathFindingResult UNavigationSystemV1::FindPathSync(const FNavAgentProperties&
 FPathFindingResult UNavigationSystemV1::FindPathSync(FPathFindingQuery Query, EPathFindingMode::Type Mode)
 {
 	SCOPE_CYCLE_COUNTER(STAT_Navigation_PathfindingSync);
-	CSV_SCOPED_TIMING_STAT_EXCLUSIVE(PathfindingSync);
+	CSV_SCOPED_TIMING_STAT(NavigationSystem, PathfindingSync);
 
 	if (Query.NavData.IsValid() == false)
 	{
@@ -1108,7 +1110,7 @@ FPathFindingResult UNavigationSystemV1::FindPathSync(FPathFindingQuery Query, EP
 bool UNavigationSystemV1::TestPathSync(FPathFindingQuery Query, EPathFindingMode::Type Mode, int32* NumVisitedNodes) const
 {
 	SCOPE_CYCLE_COUNTER(STAT_Navigation_PathfindingSync);
-	CSV_SCOPED_TIMING_STAT_EXCLUSIVE(PathfindingSync);
+	CSV_SCOPED_TIMING_STAT(NavigationSystem, PathfindingSync);
 
 	if (Query.NavData.IsValid() == false)
 	{
@@ -1197,7 +1199,7 @@ void UNavigationSystemV1::TriggerAsyncQueries(TArray<FAsyncPathFindingQuery>& Pa
 
 static void AsyncQueryDone(FAsyncPathFindingQuery Query)
 {
-	CSV_SCOPED_TIMING_STAT_EXCLUSIVE(AsyncNavQueryFinished);
+	CSV_SCOPED_TIMING_STAT(NavigationSystem, AsyncNavQueryFinished);
 
 	Query.OnDoneDelegate.ExecuteIfBound(Query.QueryID, Query.Result.Result, Query.Result.Path);
 }
@@ -1205,7 +1207,7 @@ static void AsyncQueryDone(FAsyncPathFindingQuery Query)
 void UNavigationSystemV1::PerformAsyncQueries(TArray<FAsyncPathFindingQuery> PathFindingQueries)
 {
 	SCOPE_CYCLE_COUNTER(STAT_Navigation_PathfindingAsync);
-	CSV_SCOPED_TIMING_STAT_EXCLUSIVE(PathfindingAsync);
+	CSV_SCOPED_TIMING_STAT(NavigationSystem, PathfindingAsync);
 
 	if (PathFindingQueries.Num() == 0)
 	{
