- Add BP statics for setting mod control mixes explicitly or mod control mix values/AT/RT by address-space filter
    - Required moving OSC Address matching logic to engine so we can leverige pattern matching in modulation (and potentially elsewhere)
- Add BP static convenience function for creating mix channels
- Fix thread safety issue in OnUpdateProxy (Was performing look-up on proxy maps in ModulationImpl on game thread, but those resources are to be strictly accessed on Audio Thread).
- Tweak mod value attack times to lerp a bit by default (100ms)

diff --git a/Engine/Plugins/Runtime/AudioModulation/Source/AudioModulation/Private/AudioModulationImpl.cpp b/Engine/Plugins/Runtime/AudioModulation/Source/AudioModulation/Private/AudioModulationImpl.cpp
--- a/Engine/Plugins/Runtime/AudioModulation/Source/AudioModulation/Private/AudioModulationImpl.cpp
+++ b/Engine/Plugins/Runtime/AudioModulation/Source/AudioModulation/Private/AudioModulationImpl.cpp
@@ -485,7 +485,7 @@ namespace AudioModulation
 
 	void FAudioModulationImpl::DeactivateBus(const FBusId InBusId, const ISoundModulatable* InSound)
 	{
-		auto Deactivate = [this, InBusId, InSound]()
+		RunCommandOnAudioThread([this, InBusId, InSound]()
 		{
 			check(IsInAudioThread());
 
@@ -516,14 +516,12 @@ namespace AudioModulation
 					}
 				}
 			}
-		};
-
-		IsInAudioThread() ? Deactivate() : FAudioThread::RunCommandOnAudioThread(Deactivate);
+		});
 	}
 
 	void FAudioModulationImpl::DeactivateLFO(const FLFOId InLFOId, const ISoundModulatable* InSound)
 	{
-		auto Deactivate = [this, InLFOId, InSound]()
+		RunCommandOnAudioThread([this, InLFOId, InSound]()
 		{
 			check(IsInAudioThread());
 
@@ -544,9 +542,7 @@ namespace AudioModulation
 					}
 				}
 			}
-		};
-
-		IsInAudioThread() ? Deactivate() : FAudioThread::RunCommandOnAudioThread(Deactivate);
+		});
 	}
 
 	bool FAudioModulationImpl::IsBusActive(const FBusId InBusId) const
@@ -620,30 +616,88 @@ namespace AudioModulation
 #endif // !UE_BUILD_SHIPPING
 	}
 
-	void FAudioModulationImpl::UpdateModulator(const USoundModulatorBase& InModulator)
+	void FAudioModulationImpl::RunCommandOnAudioThread(TFunction<void()> Cmd)
 	{
-		if (const USoundBusModulatorBase* BusModulator = Cast<USoundBusModulatorBase>(&InModulator))
+		if (IsInAudioThread())
 		{
-			if (FModulatorLFOProxy* LFOProxy = ActiveLFOs.Find(static_cast<FLFOId>(BusModulator->GetUniqueID())))
+			Cmd();
+		}
+		else
+		{
+			FAudioThread::RunCommandOnAudioThread(Cmd);
+		}
+	}
+
+	void FAudioModulationImpl::UpdateMix(const USoundControlBusMix& InMix, const TArray<FSoundControlBusMixChannel>& InChannels)
+	{
+		const FBusMixId MixId = static_cast<FBusMixId>(InMix.GetUniqueID());
+		RunCommandOnAudioThread([this, MixId, InChannels]()
+		{
+			if (FModulatorBusMixProxy* BusMixes = ActiveBusMixes.Find(MixId))
 			{
-				LFOProxy->OnUpdateProxy(InModulator);
+				BusMixes->SetMix(InChannels);
 			}
+		});
+	}
+
+	void FAudioModulationImpl::UpdateMixByFilter(
+		const USoundControlBusMix&					InMix,
+		const FString&								InAddressFilter,
+		const TSubclassOf<USoundControlBusBase>&	InBusClass,
+		const FSoundModulationValue&				InValue)
+	{
+		const FString	AddressFilter	= InAddressFilter;
+		const uint32	ClassId			= InBusClass ? InBusClass->GetUniqueID() : USoundControlBusBase::StaticClass()->GetUniqueID();
+		const FBusMixId MixId			= static_cast<FBusMixId>(InMix.GetUniqueID());
+
+		RunCommandOnAudioThread([this, ClassId, MixId, AddressFilter, InValue]()
+		{
+			if (FModulatorBusMixProxy* MixProxy = ActiveBusMixes.Find(MixId))
+			{
+				MixProxy->SetMixByFilter(AddressFilter, ClassId, InValue);
+			}
+		});
+	}
+
+	void FAudioModulationImpl::UpdateModulator(const USoundModulatorBase& InModulator)
+	{
+		if (const USoundBusModulatorLFO* LFO = Cast<USoundBusModulatorLFO>(&InModulator))
+		{
+			const FLFOId LFOId = static_cast<FLFOId>(InModulator.GetUniqueID());
+			const FModulatorLFOProxy UpdateProxy(*LFO);
+			RunCommandOnAudioThread([this, UpdateProxy, LFOId]()
+			{
+				if (FModulatorLFOProxy* LFOProxy = ActiveLFOs.Find(static_cast<FLFOId>(LFOId)))
+				{
+					LFOProxy->OnUpdateProxy(UpdateProxy);
+				}
+			});
 		}
 
 		if (const USoundControlBusBase* Bus = Cast<USoundControlBusBase>(&InModulator))
 		{
-			if (FControlBusProxy* BusProxy = ActiveBuses.Find(static_cast<FBusId>(Bus->GetUniqueID())))
+			const FBusId BusId = static_cast<FBusId>(InModulator.GetUniqueID());
+			const FControlBusProxy UpdateProxy(*Bus);
+			RunCommandOnAudioThread([this, UpdateProxy, BusId]()
 			{
-				BusProxy->OnUpdateProxy(InModulator);
-			}
+				if (FControlBusProxy* BusProxy = ActiveBuses.Find(static_cast<FBusId>(BusId)))
+				{
+					BusProxy->OnUpdateProxy(UpdateProxy);
+				}
+			});
 		}
 
 		if (const USoundControlBusMix* Mix = Cast<USoundControlBusMix>(&InModulator))
 		{
-			if (FModulatorBusMixProxy* BusMixes = ActiveBusMixes.Find(static_cast<FBusMixId>(Mix->GetUniqueID())))
+			const FBusMixId BusMixId = static_cast<FBusMixId>(InModulator.GetUniqueID());
+			const FModulatorBusMixProxy UpdateProxy(*Mix);
+			RunCommandOnAudioThread([this, UpdateProxy, BusMixId]()
 			{
-				BusMixes->OnUpdateProxy(InModulator);
-			}
+				if (FModulatorBusMixProxy* BusMixProxy = ActiveBusMixes.Find(static_cast<FBusMixId>(BusMixId)))
+				{
+					BusMixProxy->OnUpdateProxy(UpdateProxy);
+				}
+			});
 		}
 	}
 } // namespace AudioModulation
diff --git a/Engine/Plugins/Runtime/AudioModulation/Source/AudioModulation/Private/AudioModulationInternal.h b/Engine/Plugins/Runtime/AudioModulation/Source/AudioModulation/Private/AudioModulationInternal.h
--- a/Engine/Plugins/Runtime/AudioModulation/Source/AudioModulation/Private/AudioModulationInternal.h
+++ b/Engine/Plugins/Runtime/AudioModulation/Source/AudioModulation/Private/AudioModulationInternal.h
@@ -29,6 +29,7 @@ namespace AudioModulation
 
 #if WITH_EDITOR
 		void OnEditPluginSettings(const USoundModulationPluginSourceSettingsBase& Settings);
+		void OnExitPIE(bool bSimulating);
 #endif // WITH_EDITOR
 
 		void OnInitSound(ISoundModulatable& InSound, const USoundModulationPluginSourceSettingsBase& Settings);
@@ -55,10 +56,13 @@ namespace AudioModulation
 		void ProcessControls(const uint32 InSourceId, FSoundModulationControls& OutControls);
 		void ProcessModulators(const float Elapsed);
 
+		void UpdateMix(const USoundControlBusMix& InMix, const TArray<FSoundControlBusMixChannel>& InChannels);
+		void UpdateMixByFilter(const USoundControlBusMix& InMix, const FString& InAddressFilter, const TSubclassOf<USoundControlBusBase>& InClassFilter, const FSoundModulationValue& InValue);
 		void UpdateModulator(const USoundModulatorBase& InModulator);
 
 	private:
 		float CalculateModulationValue(FModulationPatchProxy& Proxy) const;
+		void RunCommandOnAudioThread(TFunction<void()> Cmd);
 
 		BusMixProxyMap ActiveBusMixes;
 		BusProxyMap    ActiveBuses;
@@ -96,6 +100,7 @@ namespace AudioModulation
 
 #if WITH_EDITOR
 		void OnEditPluginSettings(const USoundModulationPluginSourceSettingsBase& Settings) { }
+		void OnExitPIE(bool bSimulating) { }
 #endif // WITH_EDITOR
 
 		void OnInitSound(ISoundModulatable* InSound, const USoundModulationPluginSourceSettingsBase& Settings) { }
@@ -120,6 +125,10 @@ namespace AudioModulation
 		void ProcessControls(const uint32 InSourceId, FSoundModulationControls& OutControls) { }
 		void ProcessModulators(const float Elapsed) { }
 
+		void UpdateMix(const USoundControlBusMix& InMix, const TArray<FSoundControlBusMixChannel>& InChannels) { }
+		void UpdateMixByFilter(const USoundControlBusMix& InMix, const FString& InAddressFilter, const TSubclassOf<USoundControlBusBase>& InClassFilter, const FSoundModulationValue& InValue) { }
+		void UpdateModulator(const USoundModulatorBase& InModulator) { }
+
 		void SetBusDefault(const USoundControlBusBase& InBus, const float Value) { }
 		void SetBusMin(const USoundControlBusBase& InBus, const float Value) { }
 		void SetBusMax(const USoundControlBusBase& InBus, const float Value) { }
diff --git a/Engine/Plugins/Runtime/AudioModulation/Source/AudioModulation/Private/AudioModulationStatics.cpp b/Engine/Plugins/Runtime/AudioModulation/Source/AudioModulation/Private/AudioModulationStatics.cpp
--- a/Engine/Plugins/Runtime/AudioModulation/Source/AudioModulation/Private/AudioModulationStatics.cpp
+++ b/Engine/Plugins/Runtime/AudioModulation/Source/AudioModulation/Private/AudioModulationStatics.cpp
@@ -1,5 +1,4 @@
 // Copyright 1998-2019 Epic Games, Inc. All Rights Reserved.
-
 #include "AudioModulationStatics.h"
 
 #include "AudioDevice.h"
@@ -9,6 +8,7 @@
 
 #include "Engine/Engine.h"
 #include "SoundControlBus.h"
+#include "SoundControlBusMix.h"
 
 #define LOCTEXT_NAMESPACE "AudioModulationStatics"
 
@@ -24,8 +24,9 @@ namespace
 			return nullptr;
 		}
 
-		T* NewBus = NewObject<T>(nullptr, Name, RF_Transient);
+		T* NewBus = NewObject<T>(GetTransientPackage(), Name);
 		NewBus->DefaultValue = DefaultValue;
+		NewBus->Address = Name.ToString();
 
 		if (Activate)
 		{
@@ -155,7 +156,7 @@ USoundBusModulatorLFO* UAudioModulationStatics::CreateLFO(const UObject* WorldCo
 		return nullptr;
 	}
 
-	USoundBusModulatorLFO* NewLFO = NewObject<USoundBusModulatorLFO>(nullptr, Name, RF_Transient);
+	USoundBusModulatorLFO* NewLFO = NewObject<USoundBusModulatorLFO>(GetTransientPackage(), Name);
 	NewLFO->Amplitude = Amplitude;
 	NewLFO->Frequency = Frequency;
 	NewLFO->Offset    = Offset;
@@ -171,7 +172,15 @@ USoundBusModulatorLFO* UAudioModulationStatics::CreateLFO(const UObject* WorldCo
 	return NewLFO;
 }
 
-USoundControlBusMix* UAudioModulationStatics::CreateBusMix(const UObject* WorldContextObject, FName Name, TArray<USoundControlBusBase*> Buses, float TargetValue, bool Activate)
+FSoundControlBusMixChannel UAudioModulationStatics::CreateBusMixChannel(const UObject* WorldContextObject, USoundControlBusBase* Bus, float Value, float AttackTime, float ReleaseTime)
+{
+	FSoundControlBusMixChannel MixChannel;
+	MixChannel.Bus = Bus;
+	MixChannel.Value = FSoundModulationValue(Value, AttackTime, ReleaseTime);
+	return MixChannel;
+}
+
+USoundControlBusMix* UAudioModulationStatics::CreateBusMix(const UObject* WorldContextObject, FName Name, TArray<FSoundControlBusMixChannel> Channels, bool Activate)
 {
 	UWorld* World = GetAudioWorld(WorldContextObject);
 	if (!World)
@@ -179,17 +188,17 @@ USoundControlBusMix* UAudioModulationStatics::CreateBusMix(const UObject* WorldC
 		return nullptr;
 	}
 
-	USoundControlBusMix* NewBusMix = NewObject<USoundControlBusMix>(nullptr, Name, RF_Transient);
-	for (USoundControlBusBase* Bus : Buses)
+	USoundControlBusMix* NewBusMix = NewObject<USoundControlBusMix>(GetTransientPackage(), Name);
+	for (FSoundControlBusMixChannel& Channel : Channels)
 	{
-		if (Bus)
+		if (Channel.Bus)
 		{
-			NewBusMix->Channels.Emplace_GetRef(Bus, TargetValue);
+			NewBusMix->Channels.Emplace_GetRef(Channel);
 		}
 		else
 		{
 			UE_LOG(LogAudioModulation, Warning,
-				TEXT("USoundControlBusMix '%s' was created but bus provided is null. Channel not added."),
+				TEXT("USoundControlBusMix '%s' was created but bus provided is null. Channel not added to mix."),
 				*Name.ToString());
 		}
 	}
@@ -246,6 +255,38 @@ void UAudioModulationStatics::DeactivateBusModulator(const UObject* WorldContext
 	}
 }
 
+void UAudioModulationStatics::UpdateMix(const UObject* WorldContextObject, USoundControlBusMix* Mix, TArray<FSoundControlBusMixChannel> Channels)
+{
+	if (Mix)
+	{
+		UWorld* World = GetAudioWorld(WorldContextObject);
+		if (AudioModulation::FAudioModulationImpl* ModulationImpl = GetModulationImpl(World))
+		{
+			ModulationImpl->UpdateMix(*Mix, Channels);
+		}
+	}
+}
+
+void UAudioModulationStatics::UpdateMixByFilter(
+	const UObject*						WorldContextObject,
+	USoundControlBusMix*				Mix,
+	FString								AddressFilter,
+	TSubclassOf<USoundControlBusBase>	BusClassFilter,
+	float								Value,
+	float								AttackTime,
+	float								ReleaseTime)
+{
+	if (Mix)
+	{
+		UWorld* World = GetAudioWorld(WorldContextObject);
+		if (AudioModulation::FAudioModulationImpl* ModulationImpl = GetModulationImpl(World))
+		{
+			FSoundModulationValue ModValue(Value, AttackTime, ReleaseTime);
+			ModulationImpl->UpdateMixByFilter(*Mix, AddressFilter, BusClassFilter, ModValue);
+		}
+	}
+}
+
 void UAudioModulationStatics::UpdateModulator(const UObject* WorldContextObject, USoundModulatorBase* Modulator)
 {
 	if (Modulator)
diff --git a/Engine/Plugins/Runtime/AudioModulation/Source/AudioModulation/Private/AudioModulationStatics.h b/Engine/Plugins/Runtime/AudioModulation/Source/AudioModulation/Private/AudioModulationStatics.h
--- a/Engine/Plugins/Runtime/AudioModulation/Source/AudioModulation/Private/AudioModulationStatics.h
+++ b/Engine/Plugins/Runtime/AudioModulation/Source/AudioModulation/Private/AudioModulationStatics.h
@@ -6,11 +6,15 @@
 #include "UObject/Object.h"
 #include "UObject/ObjectMacros.h"
 
+#include "SoundControlBusMix.h"
+
 #include "AudioModulationStatics.generated.h"
 
+// Forward Declarations
 class USoundControlBus;
 class USoundControlBusMix;
 
+
 namespace AudioModulation
 {
 	class FAudioModulationImpl;
@@ -56,7 +60,7 @@ public:
 	 * @param DefaultValue - Default value for created bus
 	 * @param Activate - Whether or not to activate bus on creation.
 	 */
-	UFUNCTION(BlueprintCallable, Category = "Audio", DisplayName = "Create Volume Control Bus", meta = (WorldContext = "WorldContextObject", Keywords = "create bus modulation volume modulator"))
+	UFUNCTION(BlueprintCallable, Category = "Audio", DisplayName = "Create Volume Control Bus", meta = (WorldContext = "WorldContextObject", Keywords = "make create bus modulation volume modulator"))
 	static USoundVolumeControlBus* CreateVolumeBus(const UObject* WorldContextObject, FName Name, float DefaultValue, bool Activate);
 
 	/** Creates a pitch modulation bus with the provided default value.
@@ -64,7 +68,7 @@ public:
 	 * @param DefaultValue - Default value for created bus
 	 * @param Activate - Whether or not to activate bus on creation.
 	 */
-	UFUNCTION(BlueprintCallable, Category = "Audio", DisplayName = "Create Pitch Control Bus", meta = (WorldContext = "WorldContextObject", Keywords = "create bus modulation pitch modulator"))
+	UFUNCTION(BlueprintCallable, Category = "Audio", DisplayName = "Create Pitch Control Bus", meta = (WorldContext = "WorldContextObject", Keywords = "make create bus modulation pitch modulator"))
 	static USoundPitchControlBus* CreatePitchBus(const UObject* WorldContextObject, FName Name, float DefaultValue, bool Activate);
 
 	/** Creates a high-pass filter (HPF) modulation bus with the provided default value.
@@ -72,7 +76,7 @@ public:
 	 * @param DefaultValue - Default value for created bus
 	 * @param Activate - Whether or not to activate bus on creation.
 	 */
-	UFUNCTION(BlueprintCallable, Category = "Audio", DisplayName = "Create HPF Control Bus", meta = (WorldContext = "WorldContextObject", Keywords = "create bus modulation HPF modulator"))
+	UFUNCTION(BlueprintCallable, Category = "Audio", DisplayName = "Create HPF Control Bus", meta = (WorldContext = "WorldContextObject", Keywords = "make create bus modulation HPF modulator"))
 	static USoundHPFControlBus* CreateHPFBus(const UObject* WorldContextObject, FName Name, float DefaultValue, bool Activate);
 
 	/** Creates a low-pass filter (LPF) modulation bus with the provided default value.
@@ -80,26 +84,34 @@ public:
 	 * @param DefaultValue - Default value for created bus
 	 * @param Activate - Whether or not to activate bus on creation.
 	 */
-	UFUNCTION(BlueprintCallable, Category = "Audio", DisplayName = "Create LPF Control Bus", meta = (WorldContext = "WorldContextObject", Keywords = "create bus modulation LPF modulator"))
+	UFUNCTION(BlueprintCallable, Category = "Audio", DisplayName = "Create LPF Control Bus", meta = (WorldContext = "WorldContextObject", Keywords = "make create bus modulation LPF modulator"))
 	static USoundLPFControlBus* CreateLPFBus(const UObject* WorldContextObject, FName Name, float DefaultValue, bool Activate);
 
-	/** Creates a modulation LFO.
+	/** Creates an LFO modulator.
 	 * @param Name - Name of LFO
 	 * @param Amplitude - Amplitude of new LFO.
 	 * @param Frequency - Frequency of new LFO.
 	 * @param Offset - Offset of new LFO.
 	 * @param Activate - Whether or not to activate lfo on creation.
 	 */
-	UFUNCTION(BlueprintCallable, Category = "Audio", DisplayName = "Create Control Bus LFO", meta = (WorldContext = "WorldContextObject", Keywords = "create lfo modulation modulator"))
+	UFUNCTION(BlueprintCallable, Category = "Audio", DisplayName = "Create Control Bus LFO", meta = (WorldContext = "WorldContextObject", Keywords = "make create lfo modulation modulator"))
 	static USoundBusModulatorLFO* CreateLFO(const UObject* WorldContextObject, FName Name, float Amplitude, float Frequency, float Offset, bool Activate);
 
+	/** Creates a channel used to mix a control bus.
+	 * @param Bus - Bus channel is in charge of applying mix value to.
+	 * @param Channels - Value for added bus channel to target when mix is active.
+	 * @param Attack/ReleaseTime - Time in seconds for channel to mix in/out.
+	 */
+	UFUNCTION(BlueprintCallable, Category = "Audio", DisplayName = "Create Control Bus Mix Channel", meta = (AdvancedDisplay = "3", WorldContext = "WorldContextObject", Keywords = "make create control bus mix modulation modulator channel"))
+	static FSoundControlBusMixChannel CreateBusMixChannel(const UObject* WorldContextObject, USoundControlBusBase* Bus, float Value, float AttackTime = 0.1f, float ReleaseTime = 0.1f);
+
 	/** Creates a modulation bus mix and adds a bus channel set to the provided target value
-	 * @param Bus - Bus to add to mix
-	 * @param TargetValue - value for added bus channel to target
+	 * @param Name - Name of mix.
+	 * @param Channels - Channels mix is responsible for.
 	 * @param Activate - Whether or not to activate mix on creation.
 	 */
-	UFUNCTION(BlueprintCallable, Category = "Audio", DisplayName = "Create Control Bus Mix", meta = (WorldContext = "WorldContextObject", Keywords = "create bus mix modulation modulator"))
-	static USoundControlBusMix* CreateBusMix(const UObject* WorldContextObject, FName Name, TArray<USoundControlBusBase*> Buses, float TargetValue, bool Activate);
+	UFUNCTION(BlueprintCallable, Category = "Audio", DisplayName = "Create Control Bus Mix", meta = (WorldContext = "WorldContextObject", Keywords = "make create control bus mix modulation modulator"))
+	static USoundControlBusMix* CreateBusMix(const UObject* WorldContextObject, FName Name, TArray<FSoundControlBusMixChannel> Channels, bool Activate);
 
 	/** Deactivates a bus. Does nothing if an instance of the provided bus is already inactive
 	 * @param Bus - Scope of modulator
@@ -120,6 +132,31 @@ public:
 	UFUNCTION(BlueprintCallable, Category = "Audio", DisplayName = "Deactivate Control Bus Modulator", meta = (WorldContext = "WorldContextObject", Keywords = "deactivate bus modulation modulator"))
 	static void DeactivateBusModulator(const UObject* WorldContextObject, USoundBusModulatorBase* Modulator);
 
+	/** Updates a mix with the provided channel data.
+	 * @param Mix - Mix to update
+	 * @param Channels - Channels to set.  If channel's bus is not referenced by mix, channel's update request is ignored.
+	 */
+	UFUNCTION(BlueprintCallable, Category = "Audio", DisplayName = "Update Control Bus Mix", meta = (WorldContext = "WorldContextObject", Keywords = "update bus control modulation modulator mix channel"))
+	static void UpdateMix(const UObject* WorldContextObject, USoundControlBusMix* Mix, TArray<FSoundControlBusMixChannel> Channels);
+
+	/** Updates filtered channels of a given class to a provided target value.
+	 * @param Mix - Mix to modify
+	 * @param AddressFilter - Address filter to apply to provided mix's channels.
+	 * @param BusClass - Filters buses by subclass.
+	 * @param Value - Target value to mix filtered channels to.
+	 * @param AttackTime - If positive, updates the attack time for the resulting bus channels found matching the provided filter.
+	 * @param ReleaseTime - If positive, updates the release time for the resulting bus channels found matching the provided filter.
+	 */
+	UFUNCTION(BlueprintCallable, Category = "Audio", DisplayName = "Update Control Bus Mix By Filter", meta = (AdvancedDisplay = "4", WorldContext = "WorldContextObject", Keywords = "update bus control class modulation modulator mix channel value filter"))
+	static void UpdateMixByFilter(
+		const UObject*						WorldContextObject,
+		USoundControlBusMix*				Mix,
+		FString								AddressFilter,
+		TSubclassOf<USoundControlBusBase>	BusClassFilter,
+		float								Value,
+		float								AttackTime	= -1.0f,
+		float								ReleaseTime = -1.0f);
+
 	/** Commits updates from a UObject definition of a modulator (e.g. Bus, Bus Mix, LFO) to active instance (does not apply if modulator has not been activated).
 	 * @param Modulator - Modulator to update
 	 */
diff --git a/Engine/Plugins/Runtime/AudioModulation/Source/AudioModulation/Private/SoundControlBus.cpp b/Engine/Plugins/Runtime/AudioModulation/Source/AudioModulation/Private/SoundControlBus.cpp
--- a/Engine/Plugins/Runtime/AudioModulation/Source/AudioModulation/Private/SoundControlBus.cpp
+++ b/Engine/Plugins/Runtime/AudioModulation/Source/AudioModulation/Private/SoundControlBus.cpp
@@ -12,12 +12,46 @@
 
 USoundControlBusBase::USoundControlBusBase(const FObjectInitializer& ObjectInitializer)
 	: Super(ObjectInitializer)
+	, bOverrideAddress(false)
 	, DefaultValue(1.0f)
 	, Min(0.0f)
 	, Max(1.0f)
 {
 }
 
+#if WITH_EDITOR
+void USoundControlBusBase::PostEditChangeProperty(FPropertyChangedEvent& PropertyChangedEvent)
+{
+	if (UProperty* Property = PropertyChangedEvent.Property)
+	{
+		if (Property->GetFName() == GET_MEMBER_NAME_CHECKED(USoundControlBusBase, bOverrideAddress) && !bOverrideAddress)
+		{
+			Address = GetName();
+		}
+	}
+
+	Super::PostEditChangeProperty(PropertyChangedEvent);
+}
+
+void USoundControlBusBase::PostInitProperties()
+{
+	if (!bOverrideAddress)
+	{
+		Address = GetName();
+	}
+
+	Super::PostInitProperties();
+}
+
+void USoundControlBusBase::PostRename(UObject* OldOuter, const FName OldName)
+{
+	if (!bOverrideAddress)
+	{
+		Address = GetName();
+	}
+}
+#endif // WITH_EDITOR
+
 USoundVolumeControlBus::USoundVolumeControlBus(const FObjectInitializer& ObjectInitializer)
 	: Super(ObjectInitializer)
 {
@@ -195,18 +229,11 @@ namespace AudioModulation
 		}
 	}
 
-	void FControlBusProxy::OnUpdateProxy(const USoundModulatorBase& InModulatorArchetype)
+	void FControlBusProxy::OnUpdateProxy(const FControlBusProxy& InBusProxy)
 	{
-		const FControlBusProxy CopyProxy(*CastChecked<USoundControlBusBase>(&InModulatorArchetype));
-		auto UpdateProxy = [this, CopyProxy]()
-		{
-			// LFOIds cannot be updated as this would swap sounds subscribing to LFOs if set to AutoActivate
-			DefaultValue = FMath::Clamp(CopyProxy.DefaultValue, CopyProxy.Range.GetMin(), CopyProxy.Range.GetMax());
-			Operator = CopyProxy.Operator;
-			Range = CopyProxy.Range;
-		};
-
-		IsInAudioThread() ? UpdateProxy() : FAudioThread::RunCommandOnAudioThread(UpdateProxy);
+		DefaultValue = FMath::Clamp(InBusProxy.DefaultValue, InBusProxy.Range.GetMin(), InBusProxy.Range.GetMax());
+		Operator = InBusProxy.Operator;
+		Range = InBusProxy.Range;
 	}
 
 	void FControlBusProxy::Reset()
diff --git a/Engine/Plugins/Runtime/AudioModulation/Source/AudioModulation/Private/SoundControlBus.h b/Engine/Plugins/Runtime/AudioModulation/Source/AudioModulation/Private/SoundControlBus.h
--- a/Engine/Plugins/Runtime/AudioModulation/Source/AudioModulation/Private/SoundControlBus.h
+++ b/Engine/Plugins/Runtime/AudioModulation/Source/AudioModulation/Private/SoundControlBus.h
@@ -33,6 +33,16 @@ class USoundControlBusBase : public USoundModulatorBase
 	GENERATED_UCLASS_BODY()
 
 public:
+#if WITH_EDITORONLY_DATA
+	/** If true, Address field is used in place of object name for address used when applying mix changes using filtering. */
+	UPROPERTY(EditAnywhere, Category = General, BlueprintReadWrite)
+	bool bOverrideAddress;
+#endif // WITH_EDITORONLY_DATA
+
+	/** Address to use when applying mix changes. */
+	UPROPERTY(EditAnywhere, Category = General, BlueprintReadWrite, meta = (EditCondition = "bOverrideAddress"))
+	FString Address;
+
 	/** Default value of modulator (when no mix is applied). */
 	UPROPERTY(EditAnywhere, Category = Modulation, BlueprintReadWrite, meta = (UIMin = "0", UIMax = "1"))
 	float DefaultValue;
@@ -48,8 +58,13 @@ public:
 	UPROPERTY(EditAnywhere, Category = Modulation, BlueprintReadWrite)
 	TArray<USoundBusModulatorBase*> Modulators;
 
-	virtual void BeginDestroy() override;
+#if WITH_EDITOR
+	virtual void PostEditChangeProperty(FPropertyChangedEvent& PropertyChangedEvent) override;
+	virtual void PostInitProperties() override;
+	virtual void PostRename(UObject* OldOuter, const FName OldName) override;
+#endif // WITH_EDITOR
 
+	virtual void BeginDestroy() override;
 	virtual ESoundModulatorOperator GetOperator() const { return ESoundModulatorOperator::Multiply; }
 };
 
@@ -91,7 +106,7 @@ namespace AudioModulation
 		FControlBusProxy();
 		FControlBusProxy(const USoundControlBusBase& Bus);
 
-		void OnUpdateProxy(const USoundModulatorBase& InModulatorArchetype) override;
+		void OnUpdateProxy(const FControlBusProxy& InBusProxy);
 
 		float GetDefaultValue() const;
 		const TArray<FLFOId>& GetLFOIds() const;
diff --git a/Engine/Plugins/Runtime/AudioModulation/Source/AudioModulation/Private/SoundControlBusMix.cpp b/Engine/Plugins/Runtime/AudioModulation/Source/AudioModulation/Private/SoundControlBusMix.cpp
--- a/Engine/Plugins/Runtime/AudioModulation/Source/AudioModulation/Private/SoundControlBusMix.cpp
+++ b/Engine/Plugins/Runtime/AudioModulation/Source/AudioModulation/Private/SoundControlBusMix.cpp
@@ -1,6 +1,7 @@
 // Copyright 1998-2019 Epic Games, Inc. All Rights Reserved.
 #include "SoundControlBusMix.h"
 
+#include "Audio/AudioAddressPattern.h"
 #include "AudioDevice.h"
 #include "AudioModulation.h"
 #include "AudioModulationInternal.h"
@@ -55,6 +56,8 @@ namespace AudioModulation
 {
 	FModulatorBusMixChannelProxy::FModulatorBusMixChannelProxy(const FSoundControlBusMixChannel& Channel)
 		: TModulatorProxyBase<FBusId>(Channel.Bus->GetName(), Channel.Bus->GetUniqueID())
+		, Address(Channel.Bus->Address)
+		, ClassId(Channel.Bus->GetClass()->GetUniqueID())
 		, Value(Channel.Value)
 	{
 	}
@@ -98,17 +101,57 @@ namespace AudioModulation
 		return TModulatorProxyRefBase<FBusMixId>::CanDestroy();
 	}
 
-	void FModulatorBusMixProxy::OnUpdateProxy(const USoundModulatorBase& InModulatorArchetype)
+	void FModulatorBusMixProxy::OnUpdateProxy(const FModulatorBusMixProxy& InBusMixProxy)
 	{
-		const FModulatorBusMixProxy CopyProxy(*CastChecked<USoundControlBusMix>(&InModulatorArchetype));
-		auto UpdateProxy = [this, CopyProxy]()
+		Channels = InBusMixProxy.Channels;
+	}
+
+	void FModulatorBusMixProxy::SetMix(const TArray<FSoundControlBusMixChannel>& InChannels)
+	{
+		for (const FSoundControlBusMixChannel& NewChannel : InChannels)
+		{
+			if(NewChannel.Bus)
+			{
+				FBusId BusId = static_cast<FBusId>(NewChannel.Bus->GetUniqueID());
+				if (FModulatorBusMixChannelProxy* ChannelProxy = Channels.Find(BusId))
+				{
+					ChannelProxy->Value.TargetValue = NewChannel.Value.TargetValue;
+					ChannelProxy->Value.AttackTime = NewChannel.Value.AttackTime;
+					ChannelProxy->Value.ReleaseTime = NewChannel.Value.ReleaseTime;
+				}
+			}
+		}
+	}
+
+	void FModulatorBusMixProxy::SetMixByFilter(const FString& InAddressFilter, uint32 InFilterClassId, const FSoundModulationValue& InValue)
+	{
+		static const uint32 BaseClassId = USoundControlBusBase::StaticClass()->GetUniqueID();
+
+		for (TPair<FBusId, FModulatorBusMixChannelProxy>& IdProxyPair : Channels)
 		{
-			check(IsInAudioThread());
+			FModulatorBusMixChannelProxy& ChannelProxy = IdProxyPair.Value;
+			if (InFilterClassId != BaseClassId && ChannelProxy.ClassId != InFilterClassId)
+			{
+				continue;
+			}
+
+			if (!FAudioAddressPattern::PartsMatch(InAddressFilter, ChannelProxy.Address))
+			{
+				continue;
+			}
 
-			Channels = CopyProxy.Channels;
-		};
+			ChannelProxy.Value.TargetValue = InValue.TargetValue;
 
-		IsInAudioThread() ? UpdateProxy() : FAudioThread::RunCommandOnAudioThread(UpdateProxy);
+			if (InValue.AttackTime >= 0.0f)
+			{
+				ChannelProxy.Value.AttackTime = InValue.AttackTime;
+			}
+
+			if (InValue.ReleaseTime >= 0.0f)
+			{
+				ChannelProxy.Value.ReleaseTime = InValue.ReleaseTime;
+			}
+		}
 	}
 
 	void FModulatorBusMixProxy::SetEnabled()
diff --git a/Engine/Plugins/Runtime/AudioModulation/Source/AudioModulation/Private/SoundControlBusMix.h b/Engine/Plugins/Runtime/AudioModulation/Source/AudioModulation/Private/SoundControlBusMix.h
--- a/Engine/Plugins/Runtime/AudioModulation/Source/AudioModulation/Private/SoundControlBusMix.h
+++ b/Engine/Plugins/Runtime/AudioModulation/Source/AudioModulation/Private/SoundControlBusMix.h
@@ -6,11 +6,16 @@
 #include "UObject/Object.h"
 
 #include "SoundControlBus.h"
+#include "SoundModulatorBase.h"
 #include "SoundModulationValue.h"
 
 #include "SoundControlBusMix.generated.h"
 
 
+// Forward Declarations
+class USoundControlBusBase;
+
+
 USTRUCT(BlueprintType)
 struct FSoundControlBusMixChannel
 {
@@ -47,6 +52,8 @@ namespace AudioModulation
 	struct FModulatorBusMixChannelProxy : public TModulatorProxyBase<FBusId>
 	{
 		FModulatorBusMixChannelProxy(const FSoundControlBusMixChannel& Channel);
+		FString Address;
+		uint32 ClassId;
 		FSoundModulationValue Value;
 	};
 
@@ -64,10 +71,14 @@ namespace AudioModulation
 		virtual ~FModulatorBusMixProxy() = default;
 
 		virtual bool CanDestroy() const override;
-		virtual void OnUpdateProxy(const USoundModulatorBase& InModulatorArchetype) override;
+
+		virtual void OnUpdateProxy(const FModulatorBusMixProxy& InBusMixProxy);
 
 		void SetEnabled();
+		void SetMix(const TArray<FSoundControlBusMixChannel>& InChannels);
+		void SetMixByFilter(const FString& InAddressFilter, uint32 InFilterClassId, const FSoundModulationValue& InValue);
 		void SetStopping();
+
 		void Update(const float Elapsed, BusProxyMap& ProxyMap);
 
 		TMap<FBusId, FModulatorBusMixChannelProxy> Channels;
diff --git a/Engine/Plugins/Runtime/AudioModulation/Source/AudioModulation/Private/SoundModulationValue.cpp b/Engine/Plugins/Runtime/AudioModulation/Source/AudioModulation/Private/SoundModulationValue.cpp
--- a/Engine/Plugins/Runtime/AudioModulation/Source/AudioModulation/Private/SoundModulationValue.cpp
+++ b/Engine/Plugins/Runtime/AudioModulation/Source/AudioModulation/Private/SoundModulationValue.cpp
@@ -11,8 +11,8 @@ namespace AudioModulation
 
 FSoundModulationValue::FSoundModulationValue()
 	: TargetValue(1.0f)
-	, AttackTime(0.0f)
-	, ReleaseTime(0.0f)
+	, AttackTime(0.1f)
+	, ReleaseTime(0.1f)
 	, Value(1.0f)
 {
 }
diff --git a/Engine/Plugins/Runtime/AudioModulation/Source/AudioModulation/Private/SoundModulatorBase.cpp b/Engine/Plugins/Runtime/AudioModulation/Source/AudioModulation/Private/SoundModulatorBase.cpp
--- a/Engine/Plugins/Runtime/AudioModulation/Source/AudioModulation/Private/SoundModulatorBase.cpp
+++ b/Engine/Plugins/Runtime/AudioModulation/Source/AudioModulation/Private/SoundModulatorBase.cpp
@@ -3,6 +3,7 @@
 
 #include "AudioDevice.h"
 #include "AudioDeviceManager.h"
+#include "IAudioExtensionPlugin.h"
 
 #include "AudioModulation.h"
 #include "AudioModulationInternal.h"
@@ -14,6 +15,27 @@ USoundModulatorBase::USoundModulatorBase(const FObjectInitializer& ObjectInitial
 {
 }
 
+bool USoundModulatorBase::CanAutoActivate(const ISoundModulatable* InSound) const
+{
+	if (bAutoActivate)
+	{
+		if (InSound)
+		{
+			return true;
+		}
+	}
+	else
+	{
+		if (!InSound || InSound->IsPreviewSound())
+		{
+			return true;
+		}
+	}
+
+	return false;
+}
+
+
 #if WITH_EDITOR
 void USoundModulatorBase::PostEditChangeProperty(struct FPropertyChangedEvent& PropertyChangedEvent)
 {
diff --git a/Engine/Plugins/Runtime/AudioModulation/Source/AudioModulation/Private/SoundModulatorBase.h b/Engine/Plugins/Runtime/AudioModulation/Source/AudioModulation/Private/SoundModulatorBase.h
--- a/Engine/Plugins/Runtime/AudioModulation/Source/AudioModulation/Private/SoundModulatorBase.h
+++ b/Engine/Plugins/Runtime/AudioModulation/Source/AudioModulation/Private/SoundModulatorBase.h
@@ -6,10 +6,10 @@
 
 #include "SoundModulatorBase.generated.h"
 
-namespace AudioModulation
-{
-	class FAudioModulationImpl;
-}
+
+// Forward Declarations
+class ISoundModulatable;
+
 
 /**
  * Base class for all modulators
@@ -26,25 +26,7 @@ public:
 	UPROPERTY(EditAnywhere, Category = General, BlueprintReadWrite)
 	uint8 bAutoActivate : 1;
 
-	bool CanAutoActivate(const ISoundModulatable* InSound) const
-	{
-		if (bAutoActivate)
-		{
-			if (InSound)
-			{
-				return true;
-			}
-		}
-		else
-		{
-			if (!InSound || InSound->IsPreviewSound())
-			{
-				return true;
-			}
-		}
-		
-		return false;
-	}
+	bool CanAutoActivate(const ISoundModulatable* InSound) const;
 
 #if WITH_EDITOR
 	virtual void PostEditChangeProperty(FPropertyChangedEvent& PropertyChangedEvent) override;
@@ -125,8 +107,6 @@ namespace AudioModulation
 			return !bAutoActivate || (bAutoActivate && RefSounds.Num() == 0);
 		}
 
-		virtual void OnUpdateProxy(const USoundModulatorBase& InModulatorArchetype) { }
-
 		int32 OnReleaseSound(const ISoundModulatable& Sound)
 		{
 			if (bAutoActivate)
diff --git a/Engine/Plugins/Runtime/AudioModulation/Source/AudioModulation/Private/SoundModulatorLFO.cpp b/Engine/Plugins/Runtime/AudioModulation/Source/AudioModulation/Private/SoundModulatorLFO.cpp
--- a/Engine/Plugins/Runtime/AudioModulation/Source/AudioModulation/Private/SoundModulatorLFO.cpp
+++ b/Engine/Plugins/Runtime/AudioModulation/Source/AudioModulation/Private/SoundModulatorLFO.cpp
@@ -74,19 +74,11 @@ namespace AudioModulation
 		return Value;
 	}
 
-	void FModulatorLFOProxy::OnUpdateProxy(const USoundModulatorBase& InModulatorArchetype)
+	void FModulatorLFOProxy::OnUpdateProxy(const FModulatorLFOProxy& InLFOProxy)
 	{
-		const FModulatorLFOProxy CopyProxy(*CastChecked<USoundBusModulatorLFO>(&InModulatorArchetype));
-		auto UpdateProxy = [this, CopyProxy]()
-		{
-			check(IsInAudioThread());
-
-			LFO = CopyProxy.LFO;
-			Offset = CopyProxy.Offset;
-			Value = CopyProxy.Value;
-		};
-
-		IsInAudioThread() ? UpdateProxy() : FAudioThread::RunCommandOnAudioThread(UpdateProxy);
+		LFO = InLFOProxy.LFO;
+		Offset = InLFOProxy.Offset;
+		Value = InLFOProxy.Value;
 	}
 
 	void FModulatorLFOProxy::Update(float InElapsed)
diff --git a/Engine/Plugins/Runtime/AudioModulation/Source/AudioModulation/Private/SoundModulatorLFO.h b/Engine/Plugins/Runtime/AudioModulation/Source/AudioModulation/Private/SoundModulatorLFO.h
--- a/Engine/Plugins/Runtime/AudioModulation/Source/AudioModulation/Private/SoundModulatorLFO.h
+++ b/Engine/Plugins/Runtime/AudioModulation/Source/AudioModulation/Private/SoundModulatorLFO.h
@@ -62,7 +62,7 @@ namespace AudioModulation
 		FModulatorLFOProxy();
 		FModulatorLFOProxy(const USoundBusModulatorLFO& InLFO);
 
-		void OnUpdateProxy(const USoundModulatorBase& InModulatorArchetype) override;
+		void OnUpdateProxy(const FModulatorLFOProxy& InLFOProxy);
 
 		float GetValue() const;
 		void Update(float InElapsed);
diff --git a/Engine/Source/Runtime/Engine/Private/Audio/AudioAddressPattern.cpp b/Engine/Source/Runtime/Engine/Private/Audio/AudioAddressPattern.cpp
new file mode 100644
--- /dev/null
+++ b/Engine/Source/Runtime/Engine/Private/Audio/AudioAddressPattern.cpp
@@ -0,0 +1,420 @@
+// Copyright 1998-2019 Epic Games, Inc. All Rights Reserved.
+// Pattern-Matching based on OSC 1.0 Protocol
+#include "Audio/AudioAddressPattern.h"
+
+namespace
+{
+	const TArray<TCHAR> AddressInvalidChars = { ' ', '#', };
+	const TArray<TCHAR> AddressPatternChars = { ',', '*', '?', '[', ']', '{', '}' };
+} // namespace <>
+
+bool FAudioAddressPattern::BracePatternMatches(const FString& Pattern, int32 PatternStartIndex, int32 PatternEndIndex, const FString& Part, int32& PartIter)
+{
+	// If these checks are hit, function is being called prior to pattern
+	// validation or pattern validation is failing to validate properly.
+	check(PatternEndIndex - PatternStartIndex > 1);
+	check(Pattern[PatternStartIndex] == '{');
+	check(Pattern[PatternEndIndex] == '}');
+
+	// 1. Empty pattern, so always true
+	if (Pattern.Len() == 2)
+	{
+		return true;
+	}
+
+	// Increment start index to match first slot in strings to match
+	int32 PatternIter = PatternStartIndex + 1;
+	bool bMatches = false;
+	for (int32 PartIndex = PartIter; PartIndex < Part.Len() && PatternIter < PatternEndIndex;)
+	{
+		// End of matching string.
+		if (Pattern[PatternIter] == ',')
+		{
+			// Check to see if succeeded and return state if so.
+			if (bMatches)
+			{
+				PartIter = PartIndex;
+				return true;
+			}
+			// If not, reset part search and 
+			// continue to next Pattern string.
+			PartIndex = PartIter;
+		}
+		// Mark as matching if we're on the path to a match.
+		else if (Pattern[PatternIter] == Part[PartIndex])
+		{
+			bMatches = true;
+			++PartIndex;
+		}
+		// Reset part iterator and fast forward to next pattern string to match
+		else
+		{
+			bMatches = false;
+			PartIndex = PartIter;
+			while (Pattern[PatternIter] != ',' && PatternIter < PatternEndIndex)
+			{
+				++PatternIter;
+			}
+		}
+		++PatternIter;
+	}
+
+	return false;
+}
+
+bool FAudioAddressPattern::BracketPatternMatches(const FString& Pattern, int32 PatternStartIndex, int32 PatternEndIndex, TCHAR MatchChar)
+{
+	// If these checks are hit, function is being called prior to pattern
+	// validation or pattern validation is failing to validate properly.
+	check(PatternEndIndex - PatternStartIndex > 1);
+	check(Pattern[PatternStartIndex] == '[');
+	check(Pattern[PatternEndIndex] == ']');
+
+	// 1. Empty pattern, so always true
+	if (Pattern.Len() == 2)
+	{
+		return true;
+	}
+
+	// Prime start index and cache whether or not negation of pattern is requested.
+	PatternStartIndex++;
+	bool bNegate = false;
+	if (Pattern[1] == '!')
+	{
+		// Of non-standard form [!]: OSC 1.0 Standard is non-descript about this
+		// scenario, so treat as if this case negates any additional character.
+		if (Pattern.Len() == 3)
+		{
+			return false;
+		}
+
+		PatternStartIndex++;
+		bNegate = true;
+	}
+
+	// 2. Special case where in the form [a-z] or [!a-z]
+	if (Pattern.Len() == PatternStartIndex + 4)
+	{
+		if (Pattern[PatternStartIndex + 1] == '-')
+		{
+			if (MatchChar >= Pattern[PatternStartIndex] && MatchChar <= Pattern[PatternStartIndex + 2])
+			{
+				return !bNegate;
+			}
+			else if (MatchChar >= Pattern[PatternStartIndex + 2] && MatchChar <= Pattern[PatternStartIndex])
+			{
+				return !bNegate;
+			}
+
+			return bNegate;
+		}
+	}
+
+	// 4. Form is either [!abcdef] or [abcdef]
+	for (int32 i = PatternStartIndex; i < PatternEndIndex; ++i)
+	{
+		if (MatchChar == Pattern[i])
+		{
+			return !bNegate;
+		}
+	}
+
+	return bNegate;
+}
+
+int32 FAudioAddressPattern::FindPatternTerminatorIndex(const FString& Pattern, int32 PatternIter, TCHAR Terminator)
+{
+	int32 EndIndex = PatternIter + 1;
+	for (;Pattern[EndIndex] != Terminator; ++EndIndex)
+	{
+		// Should never hit as this function should never be called on invalid pattern
+		check(EndIndex < Pattern.Len());
+	}
+	return EndIndex;
+}
+
+bool FAudioAddressPattern::IsValidPatternPart(const FString& Part)
+{
+	bool bInBrackets = false;
+	bool bInBraces = false;
+
+	for (TCHAR Char : Part.GetCharArray())
+	{
+		if (Char == '/')
+		{
+			return false;
+		}
+
+		if (AddressInvalidChars.Contains(Char))
+		{
+			return false;
+		}
+
+		switch (Char)
+		{
+		case ',':
+		{
+			if (!bInBraces)
+			{
+				return false;
+			}
+		}
+		break;
+
+		case '?':
+		case '*':
+		{
+			if (bInBraces || bInBrackets)
+			{
+				return false;
+			}
+		}
+		break;
+
+		case '[':
+		{
+			// Nested bracket/brace case
+			if (bInBraces || bInBrackets)
+			{
+				return false;
+			}
+			bInBrackets = true;
+		}
+		break;
+
+		case ']':
+		{
+			// Mixed bracket/brace case
+			if (bInBraces)
+			{
+				return false;
+			}
+			bInBrackets = false;
+		}
+		break;
+
+		case '{':
+		{
+			// Nested bracket/brace case
+			if (bInBraces || bInBrackets)
+			{
+				return false;
+			}
+			bInBraces = true;
+		}
+		break;
+
+		case '}':
+		{
+			// Mixed bracket/brace case
+			if (bInBrackets)
+			{
+				return false;
+			}
+			bInBraces = false;
+		}
+		break;
+		}
+	}
+
+	// Missing close braces/brackets
+	return !bInBraces && !bInBrackets;
+}
+
+bool FAudioAddressPattern::IsValidPattern(const TArray<FString>& InContainers, const FString& InMethod)
+{
+	if (InContainers.Num() == 0 || InMethod.IsEmpty())
+	{
+		return false;
+
+	}
+
+	for (const FString& Container : InContainers)
+	{
+		if (!IsValidPatternPart(Container))
+		{
+			return false;
+		}
+	}
+
+	return IsValidPatternPart(InMethod);
+}
+
+bool FAudioAddressPattern::IsValidPath(const FString& Path, bool bInvalidateSeparator)
+{
+	if (Path.IsEmpty() || Path == TEXT("/"))
+	{
+		return false;
+	}
+
+	for (TCHAR Char : Path.GetCharArray())
+	{
+		if (Char == '/')
+		{
+			if (bInvalidateSeparator)
+			{
+				return false;
+			}
+			continue;
+		}
+
+		if (AddressInvalidChars.Contains(Char))
+		{
+			return false;
+		}
+
+		if (AddressPatternChars.Contains(Char))
+		{
+			return false;
+		}
+	}
+
+	return true;
+}
+
+bool FAudioAddressPattern::PartsMatch(const FString& Pattern, const FString& Part)
+{
+	int32 PartIter = 0;
+	for (int32 PatternIter = 0; PatternIter < Pattern.Len() && PartIter < Part.Len();)
+	{
+		// If pattern is at end but part still has characters to process, match fails
+		if (PatternIter >= Pattern.Len() && PartIter < Part.Len())
+		{
+			return false;
+		}
+
+		// If part is at end but pattern still has rules to process, match fails
+		if (PatternIter < Pattern.Len() && PartIter >= Part.Len())
+		{
+			return false;
+		}
+
+		switch (Pattern[PatternIter])
+		{
+		case '?':
+		{
+			// ignore case as any ol' input path character will do
+			PatternIter++;
+			PartIter++;
+		}
+		break;
+
+		case '*':
+		{
+			if (!WildPatternMatches(Pattern, PatternIter, Part, PartIter))
+			{
+				return false;
+			}
+			PatternIter++;
+			PartIter++;
+		}
+		break;
+
+		case '[':
+		{
+			int32 PatternEndIndex = FindPatternTerminatorIndex(Pattern, PatternIter, ']');
+			if (!BracketPatternMatches(Pattern, PatternIter, PatternEndIndex, Part[PartIter]))
+			{
+				return false;
+			}
+			PatternIter = PatternEndIndex + 1;
+			PartIter++;
+		}
+		break;
+
+		case '{':
+		{
+			int32 PatternEndIndex = FindPatternTerminatorIndex(Pattern, PatternIter, '}');
+			if (!BracePatternMatches(Pattern, PatternIter, PatternEndIndex, Part, PartIter))
+			{
+				return false;
+			}
+			PatternIter = PatternEndIndex + 1;
+			PartIter++;
+		}
+		break;
+
+		// All other valid characters direct test
+		default:
+		{
+			if (Pattern[PatternIter] != Part[PartIter])
+			{
+				return false;
+			}
+			PatternIter++;
+			PartIter++;
+		}
+		break;
+		}
+	}
+	return true;
+}
+
+bool FAudioAddressPattern::WildPatternMatches(const FString& Pattern, int32& PatternIter, const FString& Part, int32& PartIter)
+{
+	// Filter out/bump pattern index for case of having sequential wildcards
+	for (; PatternIter + 1 < Pattern.Len(); ++PatternIter)
+	{
+		if (Pattern[PatternIter + 1] != '*')
+		{
+			break;
+		}
+	}
+
+	// Increment to start of next thing to match next test.  Have to always run following test
+	// here as if it fails, the part iter can be incremented, effectively allowing the lead character to be
+	// "consumed" by the wildcard and the following test has to be run again until search is exhausted
+	// or following test is fulfilled.
+	++PatternIter;
+
+	// Reached end of pattern, which is a '*' so no more parsing on this part char required 
+	// (all path part beyond is considered valid)
+	if (PatternIter == Pattern.Len())
+	{
+		return true;
+	}
+
+	const bool bWildLeadsBrackets = Pattern[PatternIter] == '[';
+	const bool bWildLeadsBraces = Pattern[PatternIter] == '{';
+	const TCHAR Terminator = bWildLeadsBrackets ? ']' : '}';
+
+	int32 PatternEndIndex = 0;
+	if (bWildLeadsBrackets || bWildLeadsBraces)
+	{
+		PatternEndIndex = FindPatternTerminatorIndex(Pattern, PatternIter, Terminator);
+	}
+
+	// Continue bumping PartIter until either at end of Part or Part char matches next pattern
+	for (;PartIter < Part.Len(); ++PartIter)
+	{
+		if (bWildLeadsBrackets)
+		{
+			if (BracketPatternMatches(Pattern, PatternIter, PatternEndIndex, Part[PartIter]))
+			{
+				PatternIter = PatternEndIndex;
+				return true;
+			}
+		}
+		else if (bWildLeadsBraces)
+		{
+			// Brace match updates OutPartIter.  Only pass along and update end of brace section
+			// if it succeeds to avoid corrupting initial state for subsequent tests after failing.
+			int32 OutPartIter = PartIter;
+			if (BracePatternMatches(Pattern, PatternIter, PatternEndIndex, Part, OutPartIter))
+			{
+				PatternIter = PatternEndIndex;
+				PartIter = OutPartIter;
+				return true;
+			}
+		}
+		else
+		{
+			if (Pattern[PatternIter] == Part[PartIter] || Pattern[PatternIter] == '?')
+			{
+				return true;
+			}
+		}
+	}
+
+	return false;
+}
diff --git a/Engine/Source/Runtime/Engine/Public/Audio/AudioAddressPattern.h b/Engine/Source/Runtime/Engine/Public/Audio/AudioAddressPattern.h
new file mode 100644
--- /dev/null
+++ b/Engine/Source/Runtime/Engine/Public/Audio/AudioAddressPattern.h
@@ -0,0 +1,17 @@
+// Copyright 1998-2019 Epic Games, Inc. All Rights Reserved.
+// Pattern-Matching that fulfills the OSC 1.0 Protocol on address pattern matching
+#include "CoreMinimal.h"
+
+class ENGINE_API FAudioAddressPattern
+{
+	static int32 FindPatternTerminatorIndex(const FString& Pattern, int32 PatternIter, TCHAR Terminator);
+
+public:
+	static bool BracePatternMatches(const FString& Pattern, int32 PatternStartIndex, int32 PatternEndIndex, const FString& Part, int32& PartIter);
+	static bool BracketPatternMatches(const FString& Pattern, int32 PatternStartIndex, int32 PatternEndIndex, TCHAR MatchChar);
+	static bool IsValidPatternPart(const FString& Part);
+	static bool IsValidPattern(const TArray<FString>& InContainers, const FString& InMethod);
+	static bool IsValidPath(const FString& Path, bool bInvalidateSeparator);
+	static bool PartsMatch(const FString& Pattern, const FString& Part);
+	static bool WildPatternMatches(const FString& Pattern, int32& PatternIter, const FString& Part, int32& PartIter);
+};
\ No newline at end of file
