Possible fix for Save Cached Pose crash (fix taking address of structure that is allocated by TArray)

diff --git a/Engine/Source/Runtime/Engine/Classes/Animation/AnimNode_SaveCachedPose.h b/Engine/Source/Runtime/Engine/Classes/Animation/AnimNode_SaveCachedPose.h
--- a/Engine/Source/Runtime/Engine/Classes/Animation/AnimNode_SaveCachedPose.h
+++ b/Engine/Source/Runtime/Engine/Classes/Animation/AnimNode_SaveCachedPose.h
@@ -31,7 +31,7 @@ protected:
 	struct FCachedUpdateContext
 	{
 		FAnimationUpdateContext Context;
-		FAnimationUpdateSharedContext SharedContext;
+		TSharedPtr<FAnimationUpdateSharedContext> SharedContext;
 	};
 
 	TArray<FCachedUpdateContext> CachedUpdateContexts;
diff --git a/Engine/Source/Runtime/Engine/Private/Animation/AnimNode_SaveCachedPose.cpp b/Engine/Source/Runtime/Engine/Private/Animation/AnimNode_SaveCachedPose.cpp
--- a/Engine/Source/Runtime/Engine/Private/Animation/AnimNode_SaveCachedPose.cpp
+++ b/Engine/Source/Runtime/Engine/Private/Animation/AnimNode_SaveCachedPose.cpp
@@ -45,11 +45,13 @@ void FAnimNode_SaveCachedPose::Update_AnyThread(const FAnimationUpdateContext& C
 	// Make a minimal copy of the shared context for cached updates
 	if (FAnimationUpdateSharedContext* SharedContext = Context.GetSharedContext())
 	{
-		CachedUpdate.SharedContext.CopyForCachedUpdate(*SharedContext);
+		CachedUpdate.SharedContext = MakeShared<FAnimationUpdateSharedContext>();
+		CachedUpdate.SharedContext->CopyForCachedUpdate(*SharedContext);
+
 	}
 
 	// Store this context for the post update
-	CachedUpdate.Context = Context.WithOtherSharedContext(&CachedUpdate.SharedContext);
+	CachedUpdate.Context = Context.WithOtherSharedContext(CachedUpdate.SharedContext.Get());
 }
 
 void FAnimNode_SaveCachedPose::Evaluate_AnyThread(FPoseContext& Output)
