Restore animation update shared context after merge

diff --git a/Engine/Source/Runtime/Engine/Private/Animation/AnimInstanceProxy.cpp b/Engine/Source/Runtime/Engine/Private/Animation/AnimInstanceProxy.cpp
--- a/Engine/Source/Runtime/Engine/Private/Animation/AnimInstanceProxy.cpp
+++ b/Engine/Source/Runtime/Engine/Private/Animation/AnimInstanceProxy.cpp
@@ -1181,7 +1181,8 @@ void FAnimInstanceProxy::UpdateAnimation()
 	UpdatedNodesThisFrame.Reset();
 #endif
 
-	FAnimationUpdateContext Context(this, CurrentDeltaSeconds);
+	FAnimationUpdateSharedContext SharedContext;
+	FAnimationUpdateContext Context(this, CurrentDeltaSeconds, &SharedContext);
 	UpdateAnimation_WithRoot(Context, RootNode, NAME_AnimGraph);
 }
 
