Ensure debug viewmode would only ever be compiled on PC.
Will eventually be supported on console through "r.ForceDebugViewModes"

diff --git a/Engine/Source/Runtime/Engine/Private/DebugViewModeHelpers.cpp b/Engine/Source/Runtime/Engine/Private/DebugViewModeHelpers.cpp
--- a/Engine/Source/Runtime/Engine/Private/DebugViewModeHelpers.cpp
+++ b/Engine/Source/Runtime/Engine/Private/DebugViewModeHelpers.cpp
@@ -24,11 +24,7 @@ static bool PlatformSupportsDebugViewShaders(EShaderPlatform Platform)
 
 bool AllowDebugViewVSDSHS(EShaderPlatform Platform)
 {
-#if WITH_EDITOR
-	return true; 
-#else
-	return false;
-#endif
+	return IsPCPlatform(Platform); 
 }
 
 bool AllowDebugViewShaderMode(EDebugViewShaderMode ShaderMode, EShaderPlatform Platform, ERHIFeatureLevel::Type FeatureLevel)
@@ -44,7 +40,7 @@ bool AllowDebugViewShaderMode(EDebugViewShaderMode ShaderMode, EShaderPlatform P
 	case DVSM_None:
 		return false;
 	case DVSM_ShaderComplexity:
-		return true;
+		return IsPCPlatform(Platform);
 	case DVSM_ShaderComplexityContainedQuadOverhead:
 	case DVSM_ShaderComplexityBleedingQuadOverhead:
 	case DVSM_QuadComplexity:
