Adding ability to disable binaural audio if sound starts playing below a new radius.

diff --git a/Engine/Source/Editor/DetailCustomizations/Private/AttenuationSettingsCustomizations.cpp b/Engine/Source/Editor/DetailCustomizations/Private/AttenuationSettingsCustomizations.cpp
--- a/Engine/Source/Editor/DetailCustomizations/Private/AttenuationSettingsCustomizations.cpp
+++ b/Engine/Source/Editor/DetailCustomizations/Private/AttenuationSettingsCustomizations.cpp
@@ -318,6 +318,9 @@ void FSoundAttenuationSettingsCustomization::CustomizeChildren(TSharedRef<IPrope
 	LayoutBuilder.AddPropertyToCategory(PropertyHandles.FindChecked(GET_MEMBER_NAME_CHECKED(FSoundAttenuationSettings, SpatializationAlgorithm)))
 		.EditCondition(GetIsSpatializationEnabledAttribute(), nullptr);
 
+	LayoutBuilder.AddPropertyToCategory(PropertyHandles.FindChecked(GET_MEMBER_NAME_CHECKED(FSoundAttenuationSettings, BinauralRadius)))
+		.EditCondition(GetIsSpatializationEnabledAttribute(), nullptr);
+
 	LayoutBuilder.AddPropertyToCategory(PropertyHandles.FindChecked(GET_MEMBER_NAME_CHECKED(FSoundAttenuationSettings, OmniRadius)))
 		.EditCondition(GetIsSpatializationEnabledAttribute(), nullptr);
 
@@ -459,7 +462,7 @@ void FSoundAttenuationSettingsCustomization::CustomizeChildren(TSharedRef<IPrope
 			.EditCondition(IsAttenuationOverriddenAttribute(), nullptr);
 	}
 
-	if (PropertyHandles.Num() != 51)
+	if (PropertyHandles.Num() != 52)
 	{
 		FString PropertyList;
 		for (auto It(PropertyHandles.CreateConstIterator()); It; ++It)
diff --git a/Engine/Source/Runtime/Engine/Classes/Sound/SoundAttenuation.h b/Engine/Source/Runtime/Engine/Classes/Sound/SoundAttenuation.h
--- a/Engine/Source/Runtime/Engine/Classes/Sound/SoundAttenuation.h
+++ b/Engine/Source/Runtime/Engine/Classes/Sound/SoundAttenuation.h
@@ -122,6 +122,10 @@ struct ENGINE_API FSoundAttenuationSettings : public FBaseAttenuationSettings
 	UPROPERTY(EditAnywhere, BlueprintReadWrite, Category = AttenuationSpatialization, meta = (ClampMin = "0", EditCondition = "bSpatialize", DisplayName = "Spatialization Method"))
 	TEnumAsByte<enum ESoundSpatializationAlgorithm> SpatializationAlgorithm;
 
+	/** What min radius to use to swap to non-binaural audio when a sound starts playing. */
+	UPROPERTY(EditAnywhere, BlueprintReadWrite, Category = AttenuationSpatialization, meta = (ClampMin = "0", EditCondition = "bSpatialize"))
+	float BinauralRadius;
+
 	/** What method to use to map distance values to frequency absorption values. */
 	UPROPERTY(EditAnywhere, BlueprintReadWrite, Category = AttenuationAirAbsorption)
 	EAirAbsorptionMethod AbsorptionMethod;
diff --git a/Engine/Source/Runtime/Engine/Private/ActiveSound.cpp b/Engine/Source/Runtime/Engine/Private/ActiveSound.cpp
--- a/Engine/Source/Runtime/Engine/Private/ActiveSound.cpp
+++ b/Engine/Source/Runtime/Engine/Private/ActiveSound.cpp
@@ -1694,17 +1694,25 @@ void FActiveSound::UpdateAttenuation(float DeltaTime, FSoundParseParameters& Par
 	ParseParams.bApplyNormalizationToStereoSounds = Settings->bApplyNormalizationToStereoSounds;
 	ParseParams.bUseSpatialization |= Settings->bSpatialize;
 
-	if (Settings->SpatializationAlgorithm == ESoundSpatializationAlgorithm::SPATIALIZATION_Default && AudioDevice->IsHRTFEnabledForAll())
-	{
-		ParseParams.SpatializationMethod = ESoundSpatializationAlgorithm::SPATIALIZATION_HRTF;
-	}
-	else if (Settings->SpatializationAlgorithm == ESoundSpatializationAlgorithm::SPATIALIZATION_HRTF && AudioDevice->IsHRTFDisabled())
+	// Check the binaural radius to determine if we're going to HRTF spatialize
+	if (ListenerData.ListenerToSoundDistance < Settings->BinauralRadius)
 	{
 		ParseParams.SpatializationMethod = ESoundSpatializationAlgorithm::SPATIALIZATION_Default;
 	}
 	else
 	{
-		ParseParams.SpatializationMethod = Settings->SpatializationAlgorithm;
+		if (Settings->SpatializationAlgorithm == ESoundSpatializationAlgorithm::SPATIALIZATION_Default && AudioDevice->IsHRTFEnabledForAll())
+		{
+			ParseParams.SpatializationMethod = ESoundSpatializationAlgorithm::SPATIALIZATION_HRTF;
+		}
+		else if (Settings->SpatializationAlgorithm == ESoundSpatializationAlgorithm::SPATIALIZATION_HRTF && AudioDevice->IsHRTFDisabled())
+		{
+			ParseParams.SpatializationMethod = ESoundSpatializationAlgorithm::SPATIALIZATION_Default;
+		}
+		else
+		{
+			ParseParams.SpatializationMethod = Settings->SpatializationAlgorithm;
+		}
 	}
 
 	// If not overriding from a node, set focus data
