Fixed: InstancedStaticMeshes disappear in Shipping builds on mobile when affected by dynamic point lights

diff --git a/Engine/Source/Runtime/Engine/Private/InstancedStaticMesh.cpp b/Engine/Source/Runtime/Engine/Private/InstancedStaticMesh.cpp
--- a/Engine/Source/Runtime/Engine/Private/InstancedStaticMesh.cpp
+++ b/Engine/Source/Runtime/Engine/Private/InstancedStaticMesh.cpp
@@ -751,7 +751,6 @@ void FInstancedStaticMeshSceneProxy::GetDynamicMeshElements(const TArray<const F
 {
 	QUICK_SCOPE_CYCLE_COUNTER(STAT_InstancedStaticMeshSceneProxy_GetMeshElements);
 
-#if !(UE_BUILD_SHIPPING || UE_BUILD_TEST)
 	const bool bSelectionRenderEnabled = GIsEditor && ViewFamily.EngineShowFlags.Selection;
 
 	// If the first pass rendered selected instances only, we need to render the deselected instances in a second pass
@@ -815,7 +814,6 @@ void FInstancedStaticMeshSceneProxy::GetDynamicMeshElements(const TArray<const F
 			}
 		}
 	}
-#endif
 }
 
 int32 FInstancedStaticMeshSceneProxy::GetNumMeshBatches() const
