Zero initializing new struct var

diff --git a/Engine/Source/Runtime/Engine/Classes/Sound/SoundAttenuation.h b/Engine/Source/Runtime/Engine/Classes/Sound/SoundAttenuation.h
--- a/Engine/Source/Runtime/Engine/Classes/Sound/SoundAttenuation.h
+++ b/Engine/Source/Runtime/Engine/Classes/Sound/SoundAttenuation.h
@@ -294,6 +294,7 @@ struct ENGINE_API FSoundAttenuationSettings : public FBaseAttenuationSettings
 		, bApplyNormalizationToStereoSounds(false)
 		, bEnableLogFrequencyScaling(false)
 		, SpatializationAlgorithm(ESoundSpatializationAlgorithm::SPATIALIZATION_Default)
+		, BinauralRadius(0.0f)
 		, AbsorptionMethod(EAirAbsorptionMethod::Linear)
 		, OcclusionTraceChannel(ECC_Visibility)
 		, ReverbSendMethod(EReverbSendMethod::Linear)
