- Enabled fading out of distance field ray-traced cascade in the same way as other cascades, added CVar to disable 'r.DFFarTransitionScale' defaults to 1.
- Fixed incorrect default value in distance field shadow tracing that caused a line at the end of the shadow range.

diff --git a/Engine/Shaders/Private/DistanceFieldShadowing.usf b/Engine/Shaders/Private/DistanceFieldShadowing.usf
--- a/Engine/Shaders/Private/DistanceFieldShadowing.usf
+++ b/Engine/Shaders/Private/DistanceFieldShadowing.usf
@@ -519,7 +519,7 @@ void DistanceFieldShadowingCS(
 #endif
 #endif // USE_CULLING
 
-	float Result = 0;
+	float Result = 1.0;
 
 #define COMPUTE_SHADOWING 1
 #if COMPUTE_SHADOWING
diff --git a/Engine/Source/Runtime/Engine/Private/Components/DirectionalLightComponent.cpp b/Engine/Source/Runtime/Engine/Private/Components/DirectionalLightComponent.cpp
--- a/Engine/Source/Runtime/Engine/Private/Components/DirectionalLightComponent.cpp
+++ b/Engine/Source/Runtime/Engine/Private/Components/DirectionalLightComponent.cpp
@@ -65,6 +65,15 @@ static TAutoConsoleVariable<float> CVarFarShadowDistanceOverride(
 	TEXT("Overriding far shadow distance for all directional lighst"),
 	ECVF_RenderThreadSafe | ECVF_Scalability);
 
+static TAutoConsoleVariable<float> CVarRtdfFarTransitionScale(
+	TEXT("r.DFFarTransitionScale"),
+	1.0f,
+	TEXT("Use to modify the length of the far transition (fade out) of the distance field shadows.\n")
+	TEXT("1.0: (default) Calculate in the same way as other cascades.")
+	TEXT("0.0: Disable fade out."),
+	ECVF_RenderThreadSafe | ECVF_Scalability);
+
+
 /**
  * The scene info for a directional light.
  */
@@ -758,7 +767,15 @@ private:
 		// Add the fade region to the end of the subfrustum, if this is not the last cascade.
 		if ((int32)ShadowSplitIndex < (int32)NumTotalCascades - 1)
 		{
-			SplitFar += FadeExtension;			
+			SplitFar += FadeExtension;
+		}
+		else if (bIsRayTracedCascade)
+		{
+			FadeExtension *= FMath::Clamp(CVarRtdfFarTransitionScale.GetValueOnAnyThread(), 0.0f, 1.0f);
+			// For the ray-traced cascade, we want to fade out to avoid a hard line.
+			// Since there is no further cascade, extending the far makes less sensse as it affects performance by extending the shadow range.
+			// Thus, we instead move the fade plane closer.
+			FadePlane -= FadeExtension;
 		}
 
 		if(OutCascadeSettings)
