Added header guard
#rnx

diff --git a/Engine/Source/Runtime/Engine/Public/Audio/AudioAddressPattern.h b/Engine/Source/Runtime/Engine/Public/Audio/AudioAddressPattern.h
--- a/Engine/Source/Runtime/Engine/Public/Audio/AudioAddressPattern.h
+++ b/Engine/Source/Runtime/Engine/Public/Audio/AudioAddressPattern.h
@@ -1,5 +1,7 @@
 // Copyright 1998-2019 Epic Games, Inc. All Rights Reserved.
 // Pattern-Matching that fulfills the OSC 1.0 Protocol on address pattern matching
+#pragma once
+
 #include "CoreMinimal.h"
 
 class ENGINE_API FAudioAddressPattern
