Fix graph fragments not getting correctly initialized/cached when unlinking an anim layer

We were correctly initializing & caching when the layer was linked, but not when it was unlinked

diff --git a/Engine/Source/Runtime/Engine/Private/Animation/AnimInstance.cpp b/Engine/Source/Runtime/Engine/Private/Animation/AnimInstance.cpp
--- a/Engine/Source/Runtime/Engine/Private/Animation/AnimInstance.cpp
+++ b/Engine/Source/Runtime/Engine/Private/Animation/AnimInstance.cpp
@@ -2572,6 +2572,16 @@ void UAnimInstance::PerformLinkedLayerOverlayOperation(TSubclassOf<UAnimInstance
 			}
 		};
 
+		auto InitializeAndCacheBonesForLinkedRoot = [](FAnimNode_LinkedAnimLayer* InLayerNode, FAnimInstanceProxy& InThisProxy, UAnimInstance* InLinkedInstance, FAnimInstanceProxy& InLinkedProxy)
+		{
+			InLinkedProxy.InitializeObjects(InLinkedInstance);
+			FAnimationInitializeContext InitContext(&InThisProxy);
+			InLayerNode->Initialize_AnyThread(InitContext);
+			FAnimationCacheBonesContext CacheBonesContext(&InThisProxy);
+			InLayerNode->CacheBones_AnyThread(CacheBonesContext);
+			InLinkedProxy.ClearObjects();
+		};
+
 		for(TPair<FName, TArray<FAnimNode_LinkedAnimLayer*, TInlineAllocator<4>>> LayerPair : LayerNodesToSet)
 		{
 			if (LayerPair.Key == NAME_None)
@@ -2602,12 +2612,7 @@ void UAnimInstance::PerformLinkedLayerOverlayOperation(TSubclassOf<UAnimInstance
 							{
 								FAnimInstanceProxy& ThisProxy = GetProxyOnAnyThread<FAnimInstanceProxy>();
 								FAnimInstanceProxy& LinkedProxy = NewLinkedInstance->GetProxyOnAnyThread<FAnimInstanceProxy>();
-								LinkedProxy.InitializeObjects(NewLinkedInstance);
-								FAnimationInitializeContext InitContext(&ThisProxy);
-								LayerNode->Initialize_AnyThread(InitContext);
-								FAnimationCacheBonesContext CacheBonesContext(&ThisProxy);
-								LayerNode->CacheBones_AnyThread(CacheBonesContext);
-								LinkedProxy.ClearObjects();
+								InitializeAndCacheBonesForLinkedRoot(LayerNode, ThisProxy, NewLinkedInstance, LinkedProxy);
 							}
 
 							MeshComp->GetLinkedAnimInstances().Add(NewLinkedInstance);
@@ -2616,6 +2621,14 @@ void UAnimInstance::PerformLinkedLayerOverlayOperation(TSubclassOf<UAnimInstance
 					else
 					{
 						LayerNode->SetLinkedLayerInstance(this, nullptr);
+
+						UAnimInstance* LinkedInstance = LayerNode->GetTargetInstance<UAnimInstance>();
+						if(LayerNode->LinkedRoot && LinkedInstance)
+						{
+							FAnimInstanceProxy& ThisProxy = GetProxyOnAnyThread<FAnimInstanceProxy>();
+							FAnimInstanceProxy& LinkedProxy = LinkedInstance->GetProxyOnAnyThread<FAnimInstanceProxy>();
+							InitializeAndCacheBonesForLinkedRoot(LayerNode, ThisProxy, LinkedInstance, LinkedProxy);
+						}
 					}
 				}
 			}
@@ -2654,12 +2667,7 @@ void UAnimInstance::PerformLinkedLayerOverlayOperation(TSubclassOf<UAnimInstance
 						{
 							if(LayerNode->LinkedRoot)
 							{
-								LinkedProxy.InitializeObjects(NewLinkedInstance);
-								FAnimationInitializeContext InitContext(&ThisProxy);
-								LayerNode->Initialize_AnyThread(InitContext);
-								FAnimationCacheBonesContext CacheBonesContext(&ThisProxy);
-								LayerNode->CacheBones_AnyThread(CacheBonesContext);
-								LinkedProxy.ClearObjects();
+								InitializeAndCacheBonesForLinkedRoot(LayerNode, ThisProxy, NewLinkedInstance, LinkedProxy);
 							}
 						}
 
@@ -2672,6 +2680,14 @@ void UAnimInstance::PerformLinkedLayerOverlayOperation(TSubclassOf<UAnimInstance
 					for(FAnimNode_LinkedAnimLayer* LayerNode : LayerPair.Value)
 					{
 						LayerNode->SetLinkedLayerInstance(this, nullptr);
+
+						UAnimInstance* LinkedInstance = LayerNode->GetTargetInstance<UAnimInstance>();
+						if(LayerNode->LinkedRoot && LinkedInstance)
+						{
+							FAnimInstanceProxy& ThisProxy = GetProxyOnAnyThread<FAnimInstanceProxy>();
+							FAnimInstanceProxy& LinkedProxy = LinkedInstance->GetProxyOnAnyThread<FAnimInstanceProxy>();
+							InitializeAndCacheBonesForLinkedRoot(LayerNode, ThisProxy, LinkedInstance, LinkedProxy);
+						}
 					}
 				}
 			}
