#ue4 - Expose all suitable SplineComponent::Get*AtSplineInputKey() functions to blueprints. Using a spline input key is more efficient than continually querying by location, and we already exposed the FindInputKeyClosestToWorldLocation() function.

diff --git a/Engine/Source/Runtime/Engine/Classes/Components/SplineComponent.h b/Engine/Source/Runtime/Engine/Classes/Components/SplineComponent.h
--- a/Engine/Source/Runtime/Engine/Classes/Components/SplineComponent.h
+++ b/Engine/Source/Runtime/Engine/Classes/Components/SplineComponent.h
@@ -327,39 +327,50 @@ public:
 	void UpdateSpline();
 
 	/** Get location along spline at the provided input key value */
+	UFUNCTION(BlueprintCallable, Category = Spline)
 	FVector GetLocationAtSplineInputKey(float InKey, ESplineCoordinateSpace::Type CoordinateSpace) const;
 
 	/** Get tangent along spline at the provided input key value */
+	UFUNCTION(BlueprintCallable, Category = Spline)
 	FVector GetTangentAtSplineInputKey(float InKey, ESplineCoordinateSpace::Type CoordinateSpace) const;
 
 	/** Get unit direction along spline at the provided input key value */
+	UFUNCTION(BlueprintCallable, Category = Spline)
 	FVector GetDirectionAtSplineInputKey(float InKey, ESplineCoordinateSpace::Type CoordinateSpace) const;
 
 	/** Get rotator corresponding to rotation along spline at the provided input key value */
+	UFUNCTION(BlueprintCallable, Category = Spline)
 	FRotator GetRotationAtSplineInputKey(float InKey, ESplineCoordinateSpace::Type CoordinateSpace) const;
 
 	/** Get quaternion corresponding to rotation along spline at the provided input key value */
 	FQuat GetQuaternionAtSplineInputKey(float InKey, ESplineCoordinateSpace::Type CoordinateSpace) const;
 
 	/** Get up vector at the provided input key value */
+	UFUNCTION(BlueprintCallable, Category = Spline)
 	FVector GetUpVectorAtSplineInputKey(float InKey, ESplineCoordinateSpace::Type CoordinateSpace) const;
 
 	/** Get up vector at the provided input key value */
+	UFUNCTION(BlueprintCallable, Category = Spline)
 	FVector GetRightVectorAtSplineInputKey(float InKey, ESplineCoordinateSpace::Type CoordinateSpace) const;
 
 	/** Get transform at the provided input key value */
+	UFUNCTION(BlueprintCallable, Category = Spline)
 	FTransform GetTransformAtSplineInputKey(float InKey, ESplineCoordinateSpace::Type CoordinateSpace, bool bUseScale = false) const;
 
 	/** Get roll in degrees at the provided input key value */
+	UFUNCTION(BlueprintCallable, Category = Spline)
 	float GetRollAtSplineInputKey(float InKey, ESplineCoordinateSpace::Type CoordinateSpace) const;
 
 	/** Get scale at the provided input key value */
+	UFUNCTION(BlueprintCallable, Category = Spline)
 	FVector GetScaleAtSplineInputKey(float InKey) const;
 
 	/** Get a metadata property float value along the spline at spline input key */
+	UFUNCTION(BlueprintCallable, Category = Spline)
 	float GetFloatPropertyAtSplineInputKey(float InKey, FName PropertyName) const;
 
 	/** Get a metadata property vector value along the spline at spline input key */
+	UFUNCTION(BlueprintCallable, Category = Spline)
 	FVector GetVectorPropertyAtSplineInputKey(float InKey, FName PropertyName) const;
 
 	/** Specify unselected spline component segment color in the editor */
