Don't generate fast path Blueprint component instancing data on load when using cooked data with the editor.

diff --git a/Engine/Source/Runtime/Engine/Private/BlueprintGeneratedClass.cpp b/Engine/Source/Runtime/Engine/Private/BlueprintGeneratedClass.cpp
--- a/Engine/Source/Runtime/Engine/Private/BlueprintGeneratedClass.cpp
+++ b/Engine/Source/Runtime/Engine/Private/BlueprintGeneratedClass.cpp
@@ -357,9 +357,14 @@ void UBlueprintGeneratedClass::SerializeDefaultObject(UObject* Object, FStructur
 			}
 		};
 
+#if WITH_EDITOR
+		const bool bShouldUseCookedComponentInstancingData = bHasCookedComponentInstancingData && !GIsEditor;
+#else
+		const bool bShouldUseCookedComponentInstancingData = bHasCookedComponentInstancingData;
+#endif
 		// Generate "fast path" instancing data for inherited SCS node templates. This data may also be used to support inherited SCS component default value overrides
 		// in a nativized, cooked build, in which this Blueprint class inherits from a nativized Blueprint parent. See CheckAndApplyComponentTemplateOverrides() below.
-		if (InheritableComponentHandler && (bHasCookedComponentInstancingData || bHasNativizedParent))
+		if (InheritableComponentHandler && (bShouldUseCookedComponentInstancingData || bHasNativizedParent))
 		{
 			for (auto RecordIt = InheritableComponentHandler->CreateRecordIterator(); RecordIt; ++RecordIt)
 			{
@@ -367,7 +372,7 @@ void UBlueprintGeneratedClass::SerializeDefaultObject(UObject* Object, FStructur
 			}
 		}
 
-		if (bHasCookedComponentInstancingData)
+		if (bShouldUseCookedComponentInstancingData)
 		{
 			// Generate "fast path" instancing data for SCS node templates owned by this Blueprint class.
 			if (SimpleConstructionScript)
