Changed AsyncCreateDDC streaming path so that it susport rendering thread suspension as well as cancellation.
Added "r.Streaming.UseAsyncRequestsForDDC" to control when to use the new path or the previous one.
To prevent abandoned DDC request handle from leaking, added "r.Streaming.AbandonedDDCHandlePurgeFrequency" to clean up those requests every now and then.

diff --git a/Engine/Source/Runtime/Engine/Private/Streaming/StreamingManagerTexture.cpp b/Engine/Source/Runtime/Engine/Private/Streaming/StreamingManagerTexture.cpp
--- a/Engine/Source/Runtime/Engine/Private/Streaming/StreamingManagerTexture.cpp
+++ b/Engine/Source/Runtime/Engine/Private/Streaming/StreamingManagerTexture.cpp
@@ -161,6 +161,11 @@ FRenderAssetStreamingManager::~FRenderAssetStreamingManager()
 
 void FRenderAssetStreamingManager::OnPreGarbageCollect()
 {
+#if WITH_EDITORONLY_DATA
+	void PurgeAbandonedDDCHandles();
+	PurgeAbandonedDDCHandles();
+#endif
+
 	FScopeLock ScopeLock(&CriticalSection);
 	QUICK_SCOPE_CYCLE_COUNTER(STAT_FRenderAssetStreamingManager_OnPreGarbageCollect);
 
diff --git a/Engine/Source/Runtime/Engine/Private/Streaming/Texture2DStreamIn_DDC.cpp b/Engine/Source/Runtime/Engine/Private/Streaming/Texture2DStreamIn_DDC.cpp
--- a/Engine/Source/Runtime/Engine/Private/Streaming/Texture2DStreamIn_DDC.cpp
+++ b/Engine/Source/Runtime/Engine/Private/Streaming/Texture2DStreamIn_DDC.cpp
@@ -5,12 +5,176 @@ Texture2DStreamIn_DDC.cpp: Stream in helper for 2D textures loading DDC files.
 =============================================================================*/
 
 #include "Streaming/Texture2DStreamIn_DDC.h"
+#include "Streaming/TextureStreamingHelpers.h"
 #include "RenderUtils.h"
 #include "DerivedDataCacheInterface.h"
 #include "Serialization/MemoryReader.h"
 
 #if WITH_EDITORONLY_DATA
 
+int32 GStreamingUseAsyncRequestsForDDC = 0;
+static FAutoConsoleVariableRef CVarStreamingDDCPendingSleep(
+	TEXT("r.Streaming.UseAsyncRequestsForDDC"),
+	GStreamingUseAsyncRequestsForDDC,
+	TEXT("Whether to use async DDC requets in order to react quickly to cancel and suspend rendering requests (default=0)"),
+	ECVF_Default
+);
+
+int32 GStreamingAbandonedDDCHandlePurgeFrequency = 150;
+static FAutoConsoleVariableRef CVarStreamingAbandonedDDCHandlePurgeFrequency(
+	TEXT("r.Streaming.AbandonedDDCHandlePurgeFrequency"),
+	GStreamingAbandonedDDCHandlePurgeFrequency,
+	TEXT("The number of abandonned handle at which a purge will be triggered (default=150)"),
+	ECVF_Default
+);
+
+namespace
+{
+	class FAbandonedDDCHandleManager
+	{
+	public:
+		void Add(uint32 InHandle);
+		void Purge();
+	private:
+		TArray<uint32> Handles;	
+		FCriticalSection CS;
+		uint32 TotalAdd = 0;
+	};
+
+
+	void FAbandonedDDCHandleManager::Add(uint32 InHandle)
+	{
+		check(InHandle);
+
+		bool bPurge = false;
+		{
+			FScopeLock ScopeLock(&CS);
+			Handles.Add(InHandle);
+
+			++TotalAdd;
+			bPurge = (TotalAdd % GStreamingAbandonedDDCHandlePurgeFrequency == 0);
+		}
+
+		if (bPurge)
+		{
+			Purge();
+		}
+	}
+
+	void FAbandonedDDCHandleManager::Purge()
+	{
+		TArray<uint32> TempHandles;	
+		{
+			FScopeLock ScopeLock(&CS);
+			FMemory::Memswap(&TempHandles, &Handles, sizeof(TempHandles));
+		}
+
+		FDerivedDataCacheInterface& DDCInterface = GetDerivedDataCacheRef();
+		TArray<uint8> Data;
+
+		for (int32 Index = 0; Index < TempHandles.Num(); ++Index)
+		{
+			uint32 Handle = TempHandles[Index];
+			if (DDCInterface.PollAsynchronousCompletion(Handle))
+			{
+				DDCInterface.GetAsynchronousResults(Handle, Data);
+				Data.Reset();
+
+				TempHandles.RemoveAtSwap(Index);
+				--Index;
+			}
+		}
+
+		if (TempHandles.Num())
+		{
+			FScopeLock ScopeLock(&CS);
+			Handles.Append(TempHandles);
+		}
+	}
+
+	FAbandonedDDCHandleManager GAbandonedDDCHandleManager;
+}
+
+void PurgeAbandonedDDCHandles()
+{
+	GAbandonedDDCHandleManager.Purge();
+}
+
+FTexture2DStreamIn_DDC::FTexture2DStreamIn_DDC(UTexture2D* InTexture, int32 InRequestedMips)
+	: FTexture2DStreamIn(InTexture, InRequestedMips)
+	, bDDCIsInvalid(false)
+{
+	FMemory::Memzero(DDCHandles, sizeof(DDCHandles));
+}
+
+FTexture2DStreamIn_DDC::~FTexture2DStreamIn_DDC()
+{
+	// On cancellation, we don't wait for DDC requests to complete before releasing the object.
+	// This prevents GC from being stalled when texture are deleted.
+	for (uint32& AbandonedHandle : DDCHandles)
+	{
+		if (AbandonedHandle)
+		{
+			GAbandonedDDCHandleManager.Add(AbandonedHandle);
+			AbandonedHandle = 0;
+		}
+	}
+}
+
+void FTexture2DStreamIn_DDC::DoCreateAsyncDDCRequests(const FContext& Context)
+{
+	if (Context.Texture && Context.Resource)
+	{
+		const TIndirectArray<FTexture2DMipMap>& OwnerMips = Context.Texture->GetPlatformMips();
+		const int32 CurrentFirstMip = Context.Resource->GetCurrentFirstMip();
+		const FTexture2DRHIRef Texture2DRHI = Context.Resource->GetTexture2DRHI();
+
+		for (int32 MipIndex = PendingFirstMip; MipIndex < CurrentFirstMip && !IsCancelled(); ++MipIndex)
+		{
+			const FTexture2DMipMap& MipMap = OwnerMips[MipIndex];
+			if (!MipMap.DerivedDataKey.IsEmpty())
+			{
+				check(!DDCHandles[MipIndex]);
+				DDCHandles[MipIndex] = GetDerivedDataCacheRef().GetAsynchronous(*MipMap.DerivedDataKey);
+
+#if !UE_BUILD_SHIPPING
+				// On some platforms the IO is too fast to test cancelation requests timing issues.
+				if (FRenderAssetStreamingSettings::ExtraIOLatency > 0 && TaskSynchronization.GetValue() == 0)
+				{
+					FPlatformProcess::Sleep(FRenderAssetStreamingSettings::ExtraIOLatency * .001f); // Slow down the streaming.
+				}
+#endif
+			}
+			else
+			{
+				UE_LOG(LogTexture, Error, TEXT("DDC key missing."));
+				MarkAsCancelled();
+			}
+		}
+	}
+}
+
+bool FTexture2DStreamIn_DDC::DoPoolDDCRequests(const FContext& Context) 
+{
+	if (Context.Texture && Context.Resource)
+	{
+		const int32 CurrentFirstMip = Context.Resource->GetCurrentFirstMip();
+		for (int32 MipIndex = PendingFirstMip; MipIndex < CurrentFirstMip && !IsCancelled(); ++MipIndex)
+		{
+			const uint32 Handle = DDCHandles[MipIndex];
+			if (Handle && !GetDerivedDataCacheRef().PollAsynchronousCompletion(Handle))
+			{
+				return false;
+			}
+		}
+	}
+	else
+	{
+		MarkAsCancelled();
+	}
+	return true;
+}
+
 void FTexture2DStreamIn_DDC::DoLoadNewMipsFromDDC(const FContext& Context)
 {
 	if (Context.Texture && Context.Resource)
@@ -22,7 +186,6 @@ void FTexture2DStreamIn_DDC::DoLoadNewMipsFromDDC(const FContext& Context)
 		for (int32 MipIndex = PendingFirstMip; MipIndex < CurrentFirstMip && !IsCancelled(); ++MipIndex)
 		{
 			const FTexture2DMipMap& MipMap = OwnerMips[MipIndex];
-			const int32 ExpectedMipSize = CalcTextureMipMapSize(MipMap.SizeX, MipMap.SizeY, Texture2DRHI->GetFormat(), 0);
 
 			check(MipData[MipIndex]);
 
@@ -30,12 +193,27 @@ void FTexture2DStreamIn_DDC::DoLoadNewMipsFromDDC(const FContext& Context)
 			{
 				// The overhead of doing 2 copy of each mip data (from GetSynchronous() and FMemoryReader) in hidden by other texture DDC ops happening at the same time.
 				TArray<uint8> DerivedMipData;
-				if (GetDerivedDataCacheRef().GetSynchronous(*MipMap.DerivedDataKey, DerivedMipData))
+				bool bDDCValid = true;
+
+				const uint32 Handle = DDCHandles[MipIndex];
+				if (Handle)
 				{
+					bDDCValid = GetDerivedDataCacheRef().GetAsynchronousResults(Handle, DerivedMipData);
+					DDCHandles[MipIndex] = 0;
+				}
+				else
+				{
+					bDDCValid = GetDerivedDataCacheRef().GetSynchronous(*MipMap.DerivedDataKey, DerivedMipData);
+				}
+
+				if (bDDCValid)				
+				{
+					const int32 ExpectedMipSize = CalcTextureMipMapSize(MipMap.SizeX, MipMap.SizeY, Texture2DRHI->GetFormat(), 0);
 					FMemoryReader Ar(DerivedMipData, true);
 
 					int32 MipSize = 0;
 					Ar << MipSize;
+
 					if (MipSize == ExpectedMipSize)
 					{
 						Ar.Serialize(MipData[MipIndex], MipSize);
diff --git a/Engine/Source/Runtime/Engine/Private/Streaming/Texture2DStreamIn_DDC.h b/Engine/Source/Runtime/Engine/Private/Streaming/Texture2DStreamIn_DDC.h
--- a/Engine/Source/Runtime/Engine/Private/Streaming/Texture2DStreamIn_DDC.h
+++ b/Engine/Source/Runtime/Engine/Private/Streaming/Texture2DStreamIn_DDC.h
@@ -11,21 +11,23 @@ Texture2DStreamIn_DDC.h: Stream in helper for 2D textures loading DDC files.
 
 #if WITH_EDITORONLY_DATA
 
+extern int32 GStreamingUseAsyncRequestsForDDC;
+
 class FTexture2DStreamIn_DDC : public FTexture2DStreamIn
 {
 public:
 
-	FTexture2DStreamIn_DDC(UTexture2D* InTexture, int32 InRequestedMips)
-	: FTexture2DStreamIn(InTexture, InRequestedMips)
-	, bDDCIsInvalid(false)
-	{
-	}
+	FTexture2DStreamIn_DDC(UTexture2D* InTexture, int32 InRequestedMips);
+	~FTexture2DStreamIn_DDC();
 
 	/** Returns whether DDC of this texture needs to be regenerated.  */
 	bool DDCIsInvalid() const override { return bDDCIsInvalid; }
 
 protected:
 
+	// StreamIn_Default : Locked mips of the intermediate textures, used as disk load destination.
+	uint32 DDCHandles[MAX_TEXTURE_MIP_COUNT];
+
 	// Whether the DDC data was compatible or not.
 	bool bDDCIsInvalid;
 
@@ -33,6 +35,12 @@ protected:
 	// ********* Helpers **********
 	// ****************************
 
+	// Create DDC load requests (into DDCHandles)
+	void DoCreateAsyncDDCRequests(const FContext& Context);
+
+	// Create DDC load requests (into DDCHandles)
+	bool DoPoolDDCRequests(const FContext& Context);
+
 	// Load from DDC into MipData
 	void DoLoadNewMipsFromDDC(const FContext& Context);
 };
diff --git a/Engine/Source/Runtime/Engine/Private/Streaming/Texture2DStreamIn_DDC_AsyncCreate.cpp b/Engine/Source/Runtime/Engine/Private/Streaming/Texture2DStreamIn_DDC_AsyncCreate.cpp
--- a/Engine/Source/Runtime/Engine/Private/Streaming/Texture2DStreamIn_DDC_AsyncCreate.cpp
+++ b/Engine/Source/Runtime/Engine/Private/Streaming/Texture2DStreamIn_DDC_AsyncCreate.cpp
@@ -12,13 +12,47 @@ Texture2DStreamIn_DDC_AsyncCreate.h: Load texture 2D mips from the DDC using asy
 FTexture2DStreamIn_DDC_AsyncCreate::FTexture2DStreamIn_DDC_AsyncCreate(UTexture2D* InTexture, int32 InRequestedMips)
 	: FTexture2DStreamIn_DDC(InTexture, InRequestedMips) 
 {
-	PushTask(FContext(InTexture, TT_None), TT_Async, SRA_UPDATE_CALLBACK(AllocateAndLoadMips), TT_None, nullptr);
+	if (GStreamingUseAsyncRequestsForDDC)
+	{
+		PushTask(FContext(InTexture, TT_None), TT_Async, SRA_UPDATE_CALLBACK(AsyncDDC), TT_None, nullptr);
+	}
+	else
+	{
+		PushTask(FContext(InTexture, TT_None), TT_Async, SRA_UPDATE_CALLBACK(AllocateAndLoadMips), TT_None, nullptr);
+	}
 }
 
 // ****************************
 // ******* Update Steps *******
 // ****************************
 
+void FTexture2DStreamIn_DDC_AsyncCreate::AsyncDDC(const FContext& Context)
+{
+	DECLARE_SCOPE_CYCLE_COUNTER(TEXT("FTexture2DStreamIn_DDC_AsyncCreate::AsyncDDC"), STAT_Texture2DStreamInDDCAsyncCreate_AsyncDDC, STATGROUP_StreamingDetails);
+	check(Context.CurrentThread == TT_Async);
+
+	DoCreateAsyncDDCRequests(Context);
+
+	bDeferExecution = !IsCancelled(); // Set as pending.
+	PushTask(Context, TT_Async, SRA_UPDATE_CALLBACK(PollDDC), TT_Render, SRA_UPDATE_CALLBACK(Cancel));
+}
+
+void FTexture2DStreamIn_DDC_AsyncCreate::PollDDC(const FContext& Context)
+{
+	DECLARE_SCOPE_CYCLE_COUNTER(TEXT("FTexture2DStreamIn_DDC_AsyncCreate::PollDDC"), STAT_Texture2DStreamInDDCAsyncCreate_PollDDC, STATGROUP_StreamingDetails);
+	check(Context.CurrentThread == TT_Async);
+
+	if (DoPoolDDCRequests(Context))
+	{
+		PushTask(Context, TT_Async, SRA_UPDATE_CALLBACK(AllocateAndLoadMips), TT_Render, SRA_UPDATE_CALLBACK(Cancel));
+	}
+	else
+	{
+		bDeferExecution = !IsCancelled(); // Set as pending.
+		PushTask(Context, TT_Async, SRA_UPDATE_CALLBACK(PollDDC), TT_Render, SRA_UPDATE_CALLBACK(Cancel));
+	}
+}
+
 void FTexture2DStreamIn_DDC_AsyncCreate::AllocateAndLoadMips(const FContext& Context)
 {
 	DECLARE_SCOPE_CYCLE_COUNTER(TEXT("FTexture2DStreamIn_DDC_AsyncCreate::AllocateAndLoadMips"), STAT_Texture2DStreamInDDCAsyncCreate_AllocateAndLoadMips, STATGROUP_StreamingDetails);
diff --git a/Engine/Source/Runtime/Engine/Private/Streaming/Texture2DStreamIn_DDC_AsyncCreate.h b/Engine/Source/Runtime/Engine/Private/Streaming/Texture2DStreamIn_DDC_AsyncCreate.h
--- a/Engine/Source/Runtime/Engine/Private/Streaming/Texture2DStreamIn_DDC_AsyncCreate.h
+++ b/Engine/Source/Runtime/Engine/Private/Streaming/Texture2DStreamIn_DDC_AsyncCreate.h
@@ -23,6 +23,10 @@ protected:
 	// ******* Update Steps *******
 	// ****************************
 
+	// Create DDC requests for each mips. (AsyncThread)
+	void AsyncDDC(const FContext& Context);
+	// Poll DDC requests completion for each mips. (AsyncThread)
+	void PollDDC(const FContext& Context);
 	// Allocate the MipData (AsyncThread)
 	void AllocateAndLoadMips(const FContext& Context);
 	// Create load requests into each locked mips. (AsyncThread)
