Anim: Skip serializing the source reference mesh if the containing package is cooked for the editor.

diff --git a/Engine/Source/Runtime/Engine/Classes/Animation/Skeleton.h b/Engine/Source/Runtime/Engine/Classes/Animation/Skeleton.h
--- a/Engine/Source/Runtime/Engine/Classes/Animation/Skeleton.h
+++ b/Engine/Source/Runtime/Engine/Classes/Animation/Skeleton.h
@@ -124,11 +124,10 @@ struct FReferencePose
 	 * Serializes the bones
 	 *
 	 * @param Ar - The archive to serialize into.
-	 * @param Rect - The bone container to serialize.
-	 *
-	 * @return Reference to the Archive after serialization.
+	 * @param P - The FReferencePose to serialize
+	 * @param Outer - The object containing this instance. Used to determine if we're loading cooked data.
 	 */
-	friend FArchive& operator<<(FArchive& Ar, FReferencePose & P);
+	friend void SerializeReferencePose(FArchive& Ar, FReferencePose& P, UObject* Outer);
 };
 
 USTRUCT()
diff --git a/Engine/Source/Runtime/Engine/Private/Animation/Skeleton.cpp b/Engine/Source/Runtime/Engine/Private/Animation/Skeleton.cpp
--- a/Engine/Source/Runtime/Engine/Private/Animation/Skeleton.cpp
+++ b/Engine/Source/Runtime/Engine/Private/Animation/Skeleton.cpp
@@ -50,7 +50,7 @@ const FName USkeleton::AnimTrackCurveMappingName = FName(TEXT("AnimationTrackCur
 const FName FAnimSlotGroup::DefaultGroupName = FName(TEXT("DefaultGroup"));
 const FName FAnimSlotGroup::DefaultSlotName = FName(TEXT("DefaultSlot"));
 
-FArchive& operator<<(FArchive& Ar, FReferencePose & P)
+void SerializeReferencePose(FArchive& Ar, FReferencePose& P, UObject* Outer)
 {
 	Ar.UsingCustomVersion(FAnimPhysObjectVersion::GUID);
 
@@ -58,7 +58,7 @@ FArchive& operator<<(FArchive& Ar, FReferencePose & P)
 	Ar << P.ReferencePose;
 #if WITH_EDITORONLY_DATA
 	//TODO: we should use strip flags but we need to rev the serialization version
-	if (!Ar.IsCooking())
+	if (!Ar.IsCooking() && (!Ar.IsLoading() || !Outer->GetOutermost()->bIsCookedForEditor))
 	{
 		if (Ar.IsLoading() && Ar.CustomVer(FAnimPhysObjectVersion::GUID) < FAnimPhysObjectVersion::ChangeRetargetSourceReferenceToSoftObjectPtr)
 		{
@@ -80,7 +80,6 @@ FArchive& operator<<(FArchive& Ar, FReferencePose & P)
 		}
 	}
 #endif
-	return Ar;
 }
 
 const TCHAR* SkipPrefix(const FString& InName)
@@ -188,7 +187,7 @@ void USkeleton::Serialize( FArchive& Ar )
 			for (int32 Index=0; Index<NumOfRetargetSources; ++Index)
 			{
 				Ar << RetargetSourceName;
-				Ar << RetargetSource;
+				SerializeReferencePose(Ar, RetargetSource, this);
 
 				AnimRetargetSources.Add(RetargetSourceName, RetargetSource);
 			}
@@ -201,7 +200,7 @@ void USkeleton::Serialize( FArchive& Ar )
 			for (auto Iter = AnimRetargetSources.CreateIterator(); Iter; ++Iter)
 			{
 				Ar << Iter.Key();
-				Ar << Iter.Value();
+				SerializeReferencePose(Ar, Iter.Value(), this);
 			}
 		}
 	}
@@ -211,7 +210,7 @@ void USkeleton::Serialize( FArchive& Ar )
 		for (auto Iter = AnimRetargetSources.CreateIterator(); Iter; ++Iter)
 		{
 			Ar << Iter.Key();
-			Ar << Iter.Value();
+			SerializeReferencePose(Ar, Iter.Value(), this);
 		}
 	}
 
