Fixed possible invalid access when building texture data.

diff --git a/Engine/Source/Runtime/Engine/Private/TextureDerivedDataTask.cpp b/Engine/Source/Runtime/Engine/Private/TextureDerivedDataTask.cpp
--- a/Engine/Source/Runtime/Engine/Private/TextureDerivedDataTask.cpp
+++ b/Engine/Source/Runtime/Engine/Private/TextureDerivedDataTask.cpp
@@ -450,7 +450,8 @@ void FTextureCacheDerivedDataWorker::DoWork()
 			CompositeTextureData.GetAsyncSourceMips(ImageWrapper);
 		}
 
-		if (TextureData.Blocks[0].MipsPerLayer[0].Num() && (!CompositeTextureData.IsValid() || CompositeTextureData.Blocks[0].MipsPerLayer[0].Num()))
+		if (TextureData.Blocks.Num() && TextureData.Blocks[0].MipsPerLayer.Num() && TextureData.Blocks[0].MipsPerLayer[0].Num() && 
+			(!CompositeTextureData.IsValid() || (CompositeTextureData.Blocks.Num() && CompositeTextureData.Blocks[0].MipsPerLayer.Num() && CompositeTextureData.Blocks[0].MipsPerLayer[0].Num())))
 		{
 			BuildTexture();
 			bSucceeded = true;
