Added Mute/Solo actions to editors and sub menus, extended soundcue stats to display mute/solo state. Temp Icons.

diff --git a/Engine/Source/Editor/AudioEditor/Private/AssetTypeActions/AssetTypeActions_SoundBase.cpp b/Engine/Source/Editor/AudioEditor/Private/AssetTypeActions/AssetTypeActions_SoundBase.cpp
--- a/Engine/Source/Editor/AudioEditor/Private/AssetTypeActions/AssetTypeActions_SoundBase.cpp
+++ b/Engine/Source/Editor/AudioEditor/Private/AssetTypeActions/AssetTypeActions_SoundBase.cpp
@@ -1,7 +1,11 @@
 // Copyright 1998-2019 Epic Games, Inc. All Rights Reserved.
 
 #include "AssetTypeActions/AssetTypeActions_SoundBase.h"
+#include "AudioDeviceManager.h"
+#include "Audio/AudioDebug.h"
 #include "Sound/SoundBase.h"
+#include "Sound/SoundCue.h"
+#include "Sound/SoundWave.h"
 #include "ToolMenus.h"
 #include "Editor.h"
 #include "EditorStyleSet.h"
@@ -42,6 +46,33 @@ void FAssetTypeActions_SoundBase::GetActions(const TArray<UObject*>& InObjects,
 			FCanExecuteAction()
 			)
 		);
+
+	Section.AddMenuEntry(
+		LOCTEXT("Sound_MuteSound", "Mute"),
+		LOCTEXT("Sound_MuteSoundTooltip", "Mutes the selected sounds."),
+		FSlateIcon(FEditorStyle::GetStyleSetName(), "MediaAsset.AssetActions.Mute.Small"),
+		FUIAction(
+			FExecuteAction::CreateSP(this, &FAssetTypeActions_SoundBase::ExecuteMuteSound, Sounds),
+			FCanExecuteAction::CreateSP(this, &FAssetTypeActions_SoundBase::CanExecuteMuteCommand, Sounds),
+			FIsActionChecked::CreateSP(this, &FAssetTypeActions_SoundBase::IsActionCheckedMute, Sounds)
+		),
+		NAME_None, 
+		EUserInterfaceActionType::ToggleButton
+	);
+
+	Section.AddMenuEntry(
+		LOCTEXT("Sound_SoloSound", "Solo"),
+		LOCTEXT("Sound_SoloSoundTooltip", "Solos the selected sounds."),
+		FSlateIcon(FEditorStyle::GetStyleSetName(), "MediaAsset.AssetActions.Solo.Small"),
+		FUIAction(
+			FExecuteAction::CreateSP(this, &FAssetTypeActions_SoundBase::ExecuteSoloSound, Sounds),
+			FCanExecuteAction::CreateSP(this, &FAssetTypeActions_SoundBase::CanExecuteSoloCommand, Sounds),
+			FIsActionChecked::CreateSP(this, &FAssetTypeActions_SoundBase::IsActionCheckedSolo, Sounds)
+		),
+		NAME_None,
+		EUserInterfaceActionType::ToggleButton
+	);
+
 }
 PRAGMA_DISABLE_DEPRECATION_WARNINGS
 void FAssetTypeActions_SoundBase::AssetsActivated( const TArray<UObject*>& InObjects, EAssetTypeActivationMethod::Type ActivationType )
@@ -261,4 +292,148 @@ bool FAssetTypeActions_SoundBase::CanExecutePlayCommand(TArray<TWeakObjectPtr<US
 	return Objects.Num() == 1;
 }
 
+void FAssetTypeActions_SoundBase::ExecuteMuteSound(TArray<TWeakObjectPtr<USoundBase>> Objects) const
+{	
+	FAudioDebugger& Debugger = GEditor->GetAudioDeviceManager()->GetDebugger();	
+		
+	// In a selection that consists of some already muted, toggle everything in the same direction,
+	// to avoid AB problem.
+
+	bool bAnyMuted = IsActionCheckedMute(Objects);
+	for (TWeakObjectPtr<USoundBase> SoundBase : Objects)
+	{
+		if (USoundCue* Cue = Cast<USoundCue>(SoundBase.Get()))
+		{											
+			Debugger.SetMuteSoundCue(Cue->GetFName(), !bAnyMuted);
+		}
+		else if(USoundWave* Wave = Cast<USoundWave>(SoundBase.Get()))
+		{
+			Debugger.SetMuteSoundWave(Wave->GetFName(), !bAnyMuted);
+		}		
+	}	
+}
+
+void FAssetTypeActions_SoundBase::ExecuteSoloSound(TArray<TWeakObjectPtr<USoundBase>> Objects) const
+{
+	FAudioDebugger& Debugger = GEditor->GetAudioDeviceManager()->GetDebugger();
+	
+	// In a selection that consists of some already soloed, toggle everything in the same direction,
+	// to avoid AB problem.
+		
+	bool bAnySoloed = IsActionCheckedSolo(Objects);
+	for (TWeakObjectPtr<USoundBase> SoundBase : Objects)
+	{
+		if (USoundCue* Cue = Cast<USoundCue>(SoundBase.Get()))
+		{
+			Debugger.SetSoloSoundCue(Cue->GetFName(), !bAnySoloed);
+		}
+		else if (USoundWave* Wave = Cast<USoundWave>(SoundBase.Get()))
+		{
+			Debugger.SetSoloSoundCue(Wave->GetFName(), !bAnySoloed);
+		}
+	}
+}
+
+bool FAssetTypeActions_SoundBase::IsActionCheckedMute(TArray<TWeakObjectPtr<USoundBase>> Objects) const
+{
+	// If *any* of the selection are muted, show the tick box as ticked.
+	FAudioDebugger& Debugger = GEditor->GetAudioDeviceManager()->GetDebugger();
+	for (TWeakObjectPtr<USoundBase> SoundBase : Objects)
+	{
+		if (USoundCue* Cue = Cast<USoundCue>(SoundBase.Get()))
+		{
+			if (Debugger.IsMuteSoundCue(Cue->GetFName()))
+			{
+				return true;
+			}
+		}
+		else if (USoundWave* Wave = Cast<USoundWave>(SoundBase.Get()))
+		{
+			if (Debugger.IsMuteSoundWave(Wave->GetFName()))
+			{
+				return true;
+			}
+		}
+	}
+
+	return false;
+}
+
+bool FAssetTypeActions_SoundBase::IsActionCheckedSolo(TArray<TWeakObjectPtr<USoundBase>> Objects) const
+{
+	// If *any* of the selection are solod, show the tick box as ticked.
+	FAudioDebugger& Debugger = GEditor->GetAudioDeviceManager()->GetDebugger();
+	for (TWeakObjectPtr<USoundBase> SoundBase : Objects)
+	{
+		if (USoundCue* Cue = Cast<USoundCue>(SoundBase.Get()))
+		{
+			if (Debugger.IsSoloSoundCue(Cue->GetFName()))
+			{
+				return true;
+			}
+		}
+		else if (USoundWave* Wave = Cast<USoundWave>(SoundBase.Get()))
+		{
+			if (Debugger.IsSoloSoundWave(Wave->GetFName()))
+			{
+				return true;
+			}
+		}
+	}
+
+	return false;
+}
+
+bool FAssetTypeActions_SoundBase::CanExecuteMuteCommand(TArray<TWeakObjectPtr<USoundBase>> Objects) const
+{
+	// Allow muting if we're not Soloing.
+	FAudioDebugger& Debugger = GEditor->GetAudioDeviceManager()->GetDebugger();
+	for (TWeakObjectPtr<USoundBase> SoundBase : Objects)
+	{		
+		if (USoundCue* Cue = Cast<USoundCue>(SoundBase.Get()))
+		{
+			if (Debugger.IsSoloSoundCue(Cue->GetFName()))
+			{
+				return false;
+			}
+		}
+		else if (USoundWave* Wave = Cast<USoundWave>(SoundBase.Get()))
+		{
+			if (Debugger.IsSoloSoundWave(Wave->GetFName()))
+			{
+				return false;
+			}
+		}
+	}
+
+	// Ok.
+	return true;
+}
+
+bool FAssetTypeActions_SoundBase::CanExecuteSoloCommand(TArray<TWeakObjectPtr<USoundBase>> Objects) const
+{	
+	// Allow Soloing if we're not Muting.
+	FAudioDebugger& Debugger = GEditor->GetAudioDeviceManager()->GetDebugger();
+	for (TWeakObjectPtr<USoundBase> SoundBase : Objects)
+	{		
+		if (USoundCue* Cue = Cast<USoundCue>(SoundBase.Get()))
+		{
+			if (Debugger.IsMuteSoundCue(Cue->GetFName()))
+			{
+				return false;
+			}
+		}
+		else if (USoundWave* Wave = Cast<USoundWave>(SoundBase.Get()))
+		{
+			if (Debugger.IsMuteSoundWave(Wave->GetFName()))
+			{
+				return false;
+			}
+		}
+	}
+
+	// Ok.
+	return true;
+}
+
 #undef LOCTEXT_NAMESPACE
diff --git a/Engine/Source/Editor/AudioEditor/Private/AssetTypeActions/AssetTypeActions_SoundClass.cpp b/Engine/Source/Editor/AudioEditor/Private/AssetTypeActions/AssetTypeActions_SoundClass.cpp
--- a/Engine/Source/Editor/AudioEditor/Private/AssetTypeActions/AssetTypeActions_SoundClass.cpp
+++ b/Engine/Source/Editor/AudioEditor/Private/AssetTypeActions/AssetTypeActions_SoundClass.cpp
@@ -3,6 +3,9 @@
 #include "AssetTypeActions/AssetTypeActions_SoundClass.h"
 #include "Sound/SoundClass.h"
 #include "AudioEditorModule.h"
+#include "Audio/AudioDebug.h"
+#include "AudioDeviceManager.h"
+#include "Editor.h"
 
 #define LOCTEXT_NAMESPACE "AssetTypeActions"
 
@@ -35,4 +38,135 @@ void FAssetTypeActions_SoundClass::OpenAssetEditor( const TArray<UObject*>& InOb
 	}
 }
 
+
+void FAssetTypeActions_SoundClass::GetActions(const TArray<UObject*>& InObjects, FMenuBuilder& MenuBuilder)
+{
+	auto Sounds = GetTypedWeakObjectPtrs<USoundClass>(InObjects);
+
+	MenuBuilder.AddMenuEntry(
+		LOCTEXT("Sound_MuteSound", "Mute"),
+		LOCTEXT("Sound_MuteSoundTooltip", "Mutes anything using this SoundClass"),
+		FSlateIcon(FEditorStyle::GetStyleSetName(), "MediaAsset.AssetActions.Mute.Small"),
+		FUIAction(
+			FExecuteAction::CreateSP(this, &FAssetTypeActions_SoundClass::ExecuteMute, Sounds),
+			FCanExecuteAction::CreateSP(this, &FAssetTypeActions_SoundClass::CanExecuteMuteCommand, Sounds),
+			FIsActionChecked::CreateSP(this, &FAssetTypeActions_SoundClass::IsActionCheckedMute, Sounds)
+		),
+		NAME_None,
+		EUserInterfaceActionType::ToggleButton
+	);
+
+	MenuBuilder.AddMenuEntry(
+		LOCTEXT("Sound_SoloSound", "Solo"),
+		LOCTEXT("Sound_SoloSoundTooltip", "Mutes anything using this SoundClass"),
+		FSlateIcon(FEditorStyle::GetStyleSetName(), "MediaAsset.AssetActions.Solo.Small"),
+		FUIAction(
+			FExecuteAction::CreateSP(this, &FAssetTypeActions_SoundClass::ExecuteSolo, Sounds),
+			FCanExecuteAction::CreateSP(this, &FAssetTypeActions_SoundClass::CanExecuteSoloCommand, Sounds),
+			FIsActionChecked::CreateSP(this, &FAssetTypeActions_SoundClass::IsActionCheckedSolo, Sounds)
+		),
+		NAME_None,
+		EUserInterfaceActionType::ToggleButton
+	);
+}
+
+void FAssetTypeActions_SoundClass::ExecuteMute(TArray<TWeakObjectPtr<USoundClass>> Objects) const
+{
+	FAudioDebugger& Debugger = GEditor->GetAudioDeviceManager()->GetDebugger();
+	for (TWeakObjectPtr<USoundClass> i : Objects)
+	{		
+		if (USoundClass* Class = Cast<USoundClass>(i.Get()))
+		{
+			Debugger.ToggleMuteSoundClass(Class->GetFName());
+		}
+	}
+}
+
+void FAssetTypeActions_SoundClass::ExecuteSolo(TArray<TWeakObjectPtr<USoundClass>> Objects) const
+{
+	FAudioDebugger& Debugger = GEditor->GetAudioDeviceManager()->GetDebugger();
+	for (TWeakObjectPtr<USoundClass> i : Objects)
+	{
+		if (USoundClass* Class = Cast<USoundClass>(i.Get()))
+		{
+			Debugger.ToggleSoloSoundClass(Class->GetFName());
+		}
+	}
+}
+
+bool FAssetTypeActions_SoundClass::IsActionCheckedMute(TArray<TWeakObjectPtr<USoundClass>> Objects) const
+{
+	FAudioDebugger& Debugger = GEditor->GetAudioDeviceManager()->GetDebugger();
+	for (TWeakObjectPtr<USoundClass> i : Objects)
+	{
+		if (USoundClass* Class = Cast<USoundClass>(i.Get()))
+		{
+			if (Debugger.IsMuteSoundClass(Class->GetFName()))
+			{
+				return true;
+			}
+		}
+	}
+	return false;
+}
+
+bool FAssetTypeActions_SoundClass::IsActionCheckedSolo(TArray<TWeakObjectPtr<USoundClass>> Objects) const
+{
+	FAudioDebugger& Debugger = GEditor->GetAudioDeviceManager()->GetDebugger();
+
+	for (TWeakObjectPtr<USoundClass> i : Objects)
+	{
+		if (USoundClass* Class = Cast<USoundClass>(i.Get()))
+		{
+			if (Debugger.IsSoloSoundClass(Class->GetFName()))
+			{
+				return true;
+			}
+		}
+	}
+
+	return false;
+}
+
+bool FAssetTypeActions_SoundClass::CanExecuteMuteCommand(TArray<TWeakObjectPtr<USoundClass>> Objects) const
+{
+	FAudioDebugger& Debugger = GEditor->GetAudioDeviceManager()->GetDebugger();
+
+	// Allow muting if we're not Soloing.
+	for (TWeakObjectPtr<USoundClass> i : Objects)
+	{
+		if (USoundClass* Class = Cast<USoundClass>(i.Get()))
+		{
+			if (Debugger.IsSoloSoundClass(Class->GetFName()))
+			{
+				return false;
+			}
+		}
+	}
+
+	// Ok.
+	return true;
+}
+
+bool FAssetTypeActions_SoundClass::CanExecuteSoloCommand(TArray<TWeakObjectPtr<USoundClass>> Objects) const
+{
+	FAudioDebugger& Debugger = GEditor->GetAudioDeviceManager()->GetDebugger();
+
+	// Allow Soloing if we're not Muting.
+	for (TWeakObjectPtr<USoundClass> i : Objects)
+	{
+		if (USoundClass* Class = Cast<USoundClass>(i.Get()))
+		{
+			if (Debugger.IsMuteSoundClass(Class->GetFName()))
+			{
+				return false;
+			}
+		}
+	}
+
+	// Ok.
+	return true;
+}
+
+
 #undef LOCTEXT_NAMESPACE
diff --git a/Engine/Source/Editor/AudioEditor/Private/AssetTypeActions/AssetTypeActions_SoundClass.h b/Engine/Source/Editor/AudioEditor/Private/AssetTypeActions/AssetTypeActions_SoundClass.h
--- a/Engine/Source/Editor/AudioEditor/Private/AssetTypeActions/AssetTypeActions_SoundClass.h
+++ b/Engine/Source/Editor/AudioEditor/Private/AssetTypeActions/AssetTypeActions_SoundClass.h
@@ -16,4 +16,25 @@ public:
 	virtual UClass* GetSupportedClass() const override;
 	virtual void OpenAssetEditor( const TArray<UObject*>& InObjects, TSharedPtr<class IToolkitHost> EditWithinLevelEditor = TSharedPtr<IToolkitHost>() ) override;
 	virtual uint32 GetCategories() override { return EAssetTypeCategories::Sounds; }
+	virtual bool HasActions(const TArray<UObject*>& InObjects) const override { return true; }
+	virtual void GetActions(const TArray<UObject*>& InObjects, FMenuBuilder& MenuBuilder) override;
+private:
+	
+	/** Handler for Mute is selected  */
+	void ExecuteMute(TArray<TWeakObjectPtr<USoundClass>> Objects) const;
+
+	/** Handler for Solo is selected  */
+	void ExecuteSolo(TArray<TWeakObjectPtr<USoundClass>> Objects) const;
+
+	/** Returns true if the mute state is set.  */
+	bool IsActionCheckedMute(TArray<TWeakObjectPtr<USoundClass>> Objects) const;
+
+	/** Returns true if the solo state is set.  */
+	bool IsActionCheckedSolo(TArray<TWeakObjectPtr<USoundClass>> Objects) const;
+
+	/** Returns true if its possible to mute a sound */
+	bool CanExecuteMuteCommand(TArray<TWeakObjectPtr<USoundClass>> Objects) const;
+
+	/** Returns true if its possible to solo a sound */
+	bool CanExecuteSoloCommand(TArray<TWeakObjectPtr<USoundClass>> Objects) const;
 };
diff --git a/Engine/Source/Editor/AudioEditor/Private/SoundClassEditor.cpp b/Engine/Source/Editor/AudioEditor/Private/SoundClassEditor.cpp
--- a/Engine/Source/Editor/AudioEditor/Private/SoundClassEditor.cpp
+++ b/Engine/Source/Editor/AudioEditor/Private/SoundClassEditor.cpp
@@ -27,6 +27,9 @@
 #include "Factories/SoundClassFactory.h"
 #include "UObject/Package.h"
 #include "UObject/UObjectIterator.h"
+#include "SoundClassEditorCommands.h"
+#include "Audio/AudioDebug.h"
+#include "AudioDeviceManager.h"
 
 #define LOCTEXT_NAMESPACE "SoundClassEditor"
 DEFINE_LOG_CATEGORY_STATIC( LogSoundClassEditor, Log, All );
@@ -35,6 +38,31 @@ const FName FSoundClassEditor::GraphCanvasTabId( TEXT( "SoundClassEditor_GraphCa
 const FName FSoundClassEditor::PropertiesTabId( TEXT( "SoundClassEditor_Properties" ) );
 
 
+void FSoundClassEditor::ExtendToolbar()
+{
+	TSharedPtr<FExtender> ToolbarExtender = MakeShareable(new FExtender);
+
+	ToolbarExtender->AddToolBarExtension(
+		"Asset",
+		EExtensionHook::After,
+		GetToolkitCommands(),
+		FToolBarExtensionDelegate::CreateLambda([](FToolBarBuilder& ToolbarBuilder)
+		{
+			ToolbarBuilder.BeginSection("Audition");
+			{
+				ToolbarBuilder.AddToolBarButton(FSoundClassEditorCommands::Get().ToggleMute);
+				ToolbarBuilder.AddToolBarButton(FSoundClassEditorCommands::Get().ToggleSolo);
+			}
+			ToolbarBuilder.EndSection();
+		})
+	);
+
+	AddToolbarExtender(ToolbarExtender);
+
+	IAudioEditorModule* AudioEditorModule = &FModuleManager::LoadModuleChecked<IAudioEditorModule>("AudioEditor");
+	AddToolbarExtender(AudioEditorModule->GetSoundClassToolBarExtensibilityManager()->GetAllExtenders(GetToolkitCommands(), GetEditingObjects()));
+}
+
 //////////////////////////////////////////////////////////////////////////
 // FSoundClassEditor
 
@@ -77,15 +105,10 @@ void FSoundClassEditor::InitSoundClassEditor( const EToolkitMode::Type Mode, con
 	SoundClass->SetFlags(RF_Transactional);
 
 	GEditor->RegisterForUndo(this);
-
-	ToolkitCommands->MapAction(
-		FGenericCommands::Get().Undo,
-		FExecuteAction::CreateSP( this, &FSoundClassEditor::UndoGraphAction ));
-
-	ToolkitCommands->MapAction(
-		FGenericCommands::Get().Redo,
-		FExecuteAction::CreateSP( this, &FSoundClassEditor::RedoGraphAction ));
-
+	
+	FSoundClassEditorCommands::Register();
+	BindCommands();
+	
 	if( !SoundClass->SoundClassGraph )
 	{
 		USoundClassGraph* SoundClassGraph = CastChecked<USoundClassGraph>(FBlueprintEditorUtils::CreateNewGraph(SoundClass, NAME_None, USoundClassGraph::StaticClass(), USoundClassGraphSchema::StaticClass()));
@@ -133,7 +156,13 @@ void FSoundClassEditor::InitSoundClassEditor( const EToolkitMode::Type Mode, con
 
 	IAudioEditorModule* AudioEditorModule = &FModuleManager::LoadModuleChecked<IAudioEditorModule>( "AudioEditor" );
 	AddMenuExtender(AudioEditorModule->GetSoundClassMenuExtensibilityManager()->GetAllExtenders(GetToolkitCommands(), GetEditingObjects()));
-	AddToolbarExtender(AudioEditorModule->GetSoundClassToolBarExtensibilityManager()->GetAllExtenders(GetToolkitCommands(), GetEditingObjects()));
+
+	ExtendToolbar();
+	RegenerateMenusAndToolbars();
+			
+	ensure(GEditor);
+	ensure(GEditor->GetAudioDeviceManager());
+	Debugger = &GEditor->GetAudioDeviceManager()->GetDebugger();
 
 	GraphEditor->SelectAllNodes();
 	for (UObject* SelectedNode : GraphEditor->GetSelectedNodes())
@@ -149,8 +178,34 @@ void FSoundClassEditor::InitSoundClassEditor( const EToolkitMode::Type Mode, con
 	}
 }
 
+void FSoundClassEditor::BindCommands()
+{
+	ToolkitCommands->MapAction(
+		FGenericCommands::Get().Undo,
+		FExecuteAction::CreateSP(this, &FSoundClassEditor::UndoGraphAction));
+
+	ToolkitCommands->MapAction(
+		FGenericCommands::Get().Redo,
+		FExecuteAction::CreateSP(this, &FSoundClassEditor::RedoGraphAction));
+
+	const FSoundClassEditorCommands& Commands = FSoundClassEditorCommands::Get();
+	
+	ToolkitCommands->MapAction(
+		Commands.ToggleSolo,
+		FExecuteAction::CreateSP(this, &FSoundClassEditor::ToggleSolo),
+		FCanExecuteAction::CreateSP(this, &FSoundClassEditor::CanExcuteToggleSolo),
+		FIsActionChecked::CreateSP(this, &FSoundClassEditor::IsSoloToggled));
+
+	ToolkitCommands->MapAction(
+		Commands.ToggleMute,
+		FExecuteAction::CreateSP(this, &FSoundClassEditor::ToggleMute),
+		FCanExecuteAction::CreateSP(this, &FSoundClassEditor::CanExcuteToggleMute),
+		FIsActionChecked::CreateSP(this, &FSoundClassEditor::IsMuteToggled));
+}
+
 FSoundClassEditor::FSoundClassEditor()
 	: SoundClass(nullptr)
+	, Debugger(nullptr)
 {
 }
 
@@ -337,6 +392,38 @@ void FSoundClassEditor::RedoGraphAction()
 	GEditor->RedoTransaction();
 }
 
+void FSoundClassEditor::ToggleSolo()
+{
+	Debugger->ToggleSoloSoundClass(SoundClass->GetFName());
+}
+
+bool FSoundClassEditor::CanExcuteToggleSolo() const
+{
+	// Enable solo if we are not Muted	
+	return !Debugger->IsMuteSoundClass(SoundClass->GetFName());
+}
+
+bool FSoundClassEditor::IsSoloToggled() const
+{
+	return Debugger->IsSoloSoundClass(SoundClass->GetFName());
+}
+
+void FSoundClassEditor::ToggleMute()
+{
+	Debugger->ToggleMuteSoundClass(SoundClass->GetFName());
+}
+
+bool FSoundClassEditor::CanExcuteToggleMute() const
+{
+	// Enable mute if we are not Soloed
+	return !Debugger->IsSoloSoundClass(SoundClass->GetFName());
+}
+
+bool FSoundClassEditor::IsMuteToggled() const
+{	
+	return Debugger->IsMuteSoundClass(SoundClass->GetFName());
+}
+
 void FSoundClassEditor::CreateSoundClass(UEdGraphPin* FromPin, const FVector2D& Location, const FString& Name)
 {
 	// If we have a valid name
diff --git a/Engine/Source/Editor/AudioEditor/Private/SoundClassEditor.h b/Engine/Source/Editor/AudioEditor/Private/SoundClassEditor.h
--- a/Engine/Source/Editor/AudioEditor/Private/SoundClassEditor.h
+++ b/Engine/Source/Editor/AudioEditor/Private/SoundClassEditor.h
@@ -92,6 +92,16 @@ private:
 	/** Called to redo the last undone action */
 	void RedoGraphAction();
 
+	/** Toggle solo (Soloing means we just hear this) */
+	void ToggleSolo();
+	bool CanExcuteToggleSolo() const;
+	bool IsSoloToggled() const;
+
+	/** Toggle mute (Muting means this will be explicitly silienced)*/
+	void ToggleMute();
+	bool CanExcuteToggleMute() const;
+	bool IsMuteToggled() const;
+
 private:
 	/** The SoundClass asset being inspected */
 	USoundClass* SoundClass;
@@ -108,7 +118,16 @@ private:
 	/** Command list for this editor */
 	TSharedPtr<FUICommandList> GraphEditorCommands;
 
+	/** Cache the audio debugger instance */
+	class FAudioDebugger* Debugger;
+
 	/**	The tab ids for all the tabs used */
 	static const FName GraphCanvasTabId;
 	static const FName PropertiesTabId;
+
+	/** Bind to commands. */
+	void BindCommands();
+
+	/** Helper function to grow the toolbar and add custom buttons */
+	void ExtendToolbar();
 };
diff --git a/Engine/Source/Editor/AudioEditor/Private/SoundClassEditorCommands.cpp b/Engine/Source/Editor/AudioEditor/Private/SoundClassEditorCommands.cpp
new file mode 100644
--- /dev/null
+++ b/Engine/Source/Editor/AudioEditor/Private/SoundClassEditorCommands.cpp
@@ -0,0 +1,13 @@
+// Copyright 1998-2019 Epic Games, Inc. All Rights Reserved.
+
+#include "SoundClassEditorCommands.h"
+
+#define LOCTEXT_NAMESPACE "SoundClassEditorCommands"
+
+void FSoundClassEditorCommands::RegisterCommands()
+{
+	UI_COMMAND(ToggleSolo, "Solo", "Toggles Soloing this cue", EUserInterfaceActionType::ToggleButton, FInputChord(EKeys::S));
+	UI_COMMAND(ToggleMute, "Mute", "Toggles Muting this cue", EUserInterfaceActionType::ToggleButton, FInputChord(EKeys::M));
+}
+
+#undef LOCTEXT_NAMESPACE
diff --git a/Engine/Source/Editor/AudioEditor/Private/SoundCueEditor.cpp b/Engine/Source/Editor/AudioEditor/Private/SoundCueEditor.cpp
--- a/Engine/Source/Editor/AudioEditor/Private/SoundCueEditor.cpp
+++ b/Engine/Source/Editor/AudioEditor/Private/SoundCueEditor.cpp
@@ -30,6 +30,8 @@
 #include "Sound/SoundNodeDialoguePlayer.h"
 #include "SSoundCuePalette.h"
 #include "HAL/PlatformApplicationMisc.h"
+#include "AudioDeviceManager.h"
+#include "Audio/AudioDebug.h"
 
 #define LOCTEXT_NAMESPACE "SoundCueEditor"
 
@@ -39,6 +41,7 @@ const FName FSoundCueEditor::PaletteTabId( TEXT( "SoundCueEditor_Palette" ) );
 
 FSoundCueEditor::FSoundCueEditor()
 	: SoundCue(nullptr)
+	, Debugger(nullptr)
 {
 }
 
@@ -146,6 +149,10 @@ void FSoundCueEditor::InitSoundCueEditor(const EToolkitMode::Type Mode, const TS
 	ExtendToolbar();
 	RegenerateMenusAndToolbars();
 
+	check(GEditor->GetAudioDeviceManager());
+	Debugger = &GEditor->GetAudioDeviceManager()->GetDebugger();
+	check(Debugger);
+
 	// @todo toolkit world centric editing
 	/*if(IsWorldCentricAssetEditor())
 	{
@@ -281,6 +288,14 @@ void FSoundCueEditor::ExtendToolbar()
 {
 		static void FillToolbar(FToolBarBuilder& ToolbarBuilder)
 		{
+			ToolbarBuilder.BeginSection("Debug");
+			{
+				ToolbarBuilder.AddToolBarButton(FSoundCueGraphEditorCommands::Get().ToggleSolo);
+
+				ToolbarBuilder.AddToolBarButton(FSoundCueGraphEditorCommands::Get().ToggleMute);
+			}
+			ToolbarBuilder.EndSection();
+
 			ToolbarBuilder.BeginSection("Toolbar");
 			{
 				ToolbarBuilder.AddToolBarButton(FSoundCueGraphEditorCommands::Get().PlayCue);
@@ -336,6 +351,18 @@ void FSoundCueEditor::BindGraphCommands()
 	ToolkitCommands->MapAction(
 		FGenericCommands::Get().Redo,
 		FExecuteAction::CreateSP( this, &FSoundCueEditor::RedoGraphAction ));
+
+	ToolkitCommands->MapAction(
+		Commands.ToggleSolo,
+		FExecuteAction::CreateSP(this, &FSoundCueEditor::ToggleSolo),
+		FCanExecuteAction::CreateSP(this, &FSoundCueEditor::CanExcuteToggleSolo),
+		FIsActionChecked::CreateSP(this, &FSoundCueEditor::IsSoloToggled));
+		
+	ToolkitCommands->MapAction(
+		Commands.ToggleMute,
+		FExecuteAction::CreateSP(this, &FSoundCueEditor::ToggleMute),
+		FCanExecuteAction::CreateSP(this, &FSoundCueEditor::CanExcuteToggleMute),
+		FIsActionChecked::CreateSP(this, &FSoundCueEditor::IsMuteToggled));
 }
 
 void FSoundCueEditor::PlayCue()
@@ -392,6 +419,38 @@ void FSoundCueEditor::TogglePlayback()
 	}
 }
 
+void FSoundCueEditor::ToggleSolo()
+{
+	Debugger->ToggleSoloSoundCue(SoundCue->GetFName());
+}
+
+bool FSoundCueEditor::CanExcuteToggleSolo() const
+{
+	// Allow Solo if Mute is not Toggle on
+	return !Debugger->IsMuteSoundCue(SoundCue->GetFName());
+}
+
+bool FSoundCueEditor::IsSoloToggled() const
+{
+	return Debugger->IsSoloSoundCue(SoundCue->GetFName());
+}
+
+void FSoundCueEditor::ToggleMute()
+{
+	Debugger->ToggleMuteSoundCue(SoundCue->GetFName());	
+}
+
+bool FSoundCueEditor::CanExcuteToggleMute() const
+{
+	// Allow Mute if Solo is not Toggle on
+	return !Debugger->IsSoloSoundCue(SoundCue->GetFName());
+}
+
+bool FSoundCueEditor::IsMuteToggled() const
+{
+	return Debugger->IsMuteSoundCue(SoundCue->GetFName());
+}
+
 void FSoundCueEditor::PlaySingleNode(UEdGraphNode* Node)
 {
 	USoundCueGraphNode* SoundGraphNode = Cast<USoundCueGraphNode>(Node);
diff --git a/Engine/Source/Editor/AudioEditor/Private/SoundCueEditor.h b/Engine/Source/Editor/AudioEditor/Private/SoundCueEditor.h
--- a/Engine/Source/Editor/AudioEditor/Private/SoundCueEditor.h
+++ b/Engine/Source/Editor/AudioEditor/Private/SoundCueEditor.h
@@ -156,6 +156,17 @@ private:
 	void Stop();
 	/** Either play the cue or stop currently playing sound */
 	void TogglePlayback();
+
+	/** Toggle solo explain soloing*/
+	void ToggleSolo();
+	bool CanExcuteToggleSolo() const;
+	bool IsSoloToggled() const;
+	
+	/** Toggle mute */
+	void ToggleMute();
+	bool CanExcuteToggleMute() const;
+	bool IsMuteToggled() const;
+	
 	/** Plays a single specified node */
 	void PlaySingleNode(UEdGraphNode* Node);
 
@@ -199,6 +210,9 @@ private:
 	/** Command list for this editor */
 	TSharedPtr<FUICommandList> GraphEditorCommands;
 
+	/** Cache of the Audio debugger instance */
+	class FAudioDebugger* Debugger;
+
 	/**	The tab ids for all the tabs used */
 	static const FName GraphCanvasTabId;
 	static const FName PropertiesTabId;
diff --git a/Engine/Source/Editor/AudioEditor/Private/SoundCueGraphEditorCommands.cpp b/Engine/Source/Editor/AudioEditor/Private/SoundCueGraphEditorCommands.cpp
--- a/Engine/Source/Editor/AudioEditor/Private/SoundCueGraphEditorCommands.cpp
+++ b/Engine/Source/Editor/AudioEditor/Private/SoundCueGraphEditorCommands.cpp
@@ -14,6 +14,9 @@ void FSoundCueGraphEditorCommands::RegisterCommands()
 	UI_COMMAND(BrowserSync, "Sync to Browser", "Selects the SoundWave in the content browser", EUserInterfaceActionType::Button, FInputChord());
 	UI_COMMAND(AddInput, "Add Input", "Adds an input to the node", EUserInterfaceActionType::Button, FInputChord());
 	UI_COMMAND(DeleteInput, "Delete Input", "Removes an input from the node", EUserInterfaceActionType::Button, FInputChord());
+
+	UI_COMMAND(ToggleSolo, "Solo", "Toggles Soloing this cue", EUserInterfaceActionType::ToggleButton, FInputChord(EKeys::S));
+	UI_COMMAND(ToggleMute, "Mute", "Toggles Muting this cue", EUserInterfaceActionType::ToggleButton, FInputChord(EKeys::M));
 }
 
 #undef LOCTEXT_NAMESPACE
diff --git a/Engine/Source/Editor/AudioEditor/Public/AssetTypeActions/AssetTypeActions_SoundBase.h b/Engine/Source/Editor/AudioEditor/Public/AssetTypeActions/AssetTypeActions_SoundBase.h
--- a/Engine/Source/Editor/AudioEditor/Public/AssetTypeActions/AssetTypeActions_SoundBase.h
+++ b/Engine/Source/Editor/AudioEditor/Public/AssetTypeActions/AssetTypeActions_SoundBase.h
@@ -46,5 +46,23 @@ private:
 
 	/** Returns true if only one sound is selected to play */
 	bool CanExecutePlayCommand(TArray<TWeakObjectPtr<USoundBase>> Objects) const;
+
+	/** Handler for Mute is selected  */
+	void ExecuteMuteSound(TArray<TWeakObjectPtr<USoundBase>> Objects) const;
+	
+	/** Handler for Solo is selected  */
+	void ExecuteSoloSound(TArray<TWeakObjectPtr<USoundBase>> Objects) const;
+
+	/** Returns true if the mute state is set.  */
+	bool IsActionCheckedMute(TArray<TWeakObjectPtr<USoundBase>> Objects) const;
+
+	/** Returns true if the solo state is set.  */
+	bool IsActionCheckedSolo(TArray<TWeakObjectPtr<USoundBase>> Objects) const;
+
+	/** Returns true if its possible to mute a sound */
+	bool CanExecuteMuteCommand(TArray<TWeakObjectPtr<USoundBase>> Objects) const;
+	
+	/** Returns true if its possible to solo a sound */
+	bool CanExecuteSoloCommand(TArray<TWeakObjectPtr<USoundBase>> Objects) const;
 };
 PRAGMA_ENABLE_DEPRECATION_WARNINGS
diff --git a/Engine/Source/Editor/AudioEditor/Public/SoundClassEditorCommands.h b/Engine/Source/Editor/AudioEditor/Public/SoundClassEditorCommands.h
new file mode 100644
--- /dev/null
+++ b/Engine/Source/Editor/AudioEditor/Public/SoundClassEditorCommands.h
@@ -0,0 +1,30 @@
+// Copyright 1998-2019 Epic Games, Inc. All Rights Reserved.
+
+#pragma once
+
+#include "CoreMinimal.h"
+#include "Framework/Commands/Commands.h"
+#include "EditorStyleSet.h"
+
+/*-----------------------------------------------------------------------------
+   FSoundCueGraphEditorCommands
+-----------------------------------------------------------------------------*/
+
+class FSoundClassEditorCommands : public TCommands<FSoundClassEditorCommands>
+{
+public:
+	/** Constructor */
+	FSoundClassEditorCommands()
+		: TCommands<FSoundClassEditorCommands>("SoundClassEditor", NSLOCTEXT("Contexts", "SoundClassEditor", "SoundClass Editor"), NAME_None, FEditorStyle::GetStyleSetName())
+	{
+	}
+	
+	/** Plays the SoundCue or stops the currently playing cue/node */
+	TSharedPtr<FUICommandInfo> ToggleSolo;
+
+	/** Plays the SoundCue or stops the currently playing cue/node */
+	TSharedPtr<FUICommandInfo> ToggleMute;
+
+	/** Initialize commands */
+	virtual void RegisterCommands() override;
+};
diff --git a/Engine/Source/Editor/AudioEditor/Public/SoundCueGraphEditorCommands.h b/Engine/Source/Editor/AudioEditor/Public/SoundCueGraphEditorCommands.h
--- a/Engine/Source/Editor/AudioEditor/Public/SoundCueGraphEditorCommands.h
+++ b/Engine/Source/Editor/AudioEditor/Public/SoundCueGraphEditorCommands.h
@@ -43,6 +43,12 @@ public:
 	/** Removes an input from the node */
 	TSharedPtr<FUICommandInfo> DeleteInput;
 
+	/** Plays the SoundCue or stops the currently playing cue/node */
+	TSharedPtr<FUICommandInfo> ToggleSolo;
+
+	/** Plays the SoundCue or stops the currently playing cue/node */
+	TSharedPtr<FUICommandInfo> ToggleMute;
+
 	/** Initialize commands */
 	virtual void RegisterCommands() override;
 };
diff --git a/Engine/Source/Editor/EditorStyle/Private/SlateEditorStyle.cpp b/Engine/Source/Editor/EditorStyle/Private/SlateEditorStyle.cpp
--- a/Engine/Source/Editor/EditorStyle/Private/SlateEditorStyle.cpp
+++ b/Engine/Source/Editor/EditorStyle/Private/SlateEditorStyle.cpp
@@ -6592,11 +6592,15 @@ void FSlateEditorStyle::FStyle::SetupContentBrowserStyle()
 		Set( "MediaAsset.AssetActions.Play.Small", new IMAGE_BRUSH( "Icons/icon_SCueEd_PlayCue_16x", Icon16x16 ) );
 		Set( "MediaAsset.AssetActions.Stop.Small", new IMAGE_BRUSH( "Icons/icon_SCueEd_Stop_16x", Icon16x16 ) );
 		Set( "MediaAsset.AssetActions.Pause.Small", new IMAGE_BRUSH( "Icons/icon_SCueEd_Pause_16x", Icon16x16 ) );
+		Set( "MediaAsset.AssetActions.Solo.Small", new IMAGE_BRUSH( "Icons/icon_SCueEd_Solo_16x", Icon16x16));
+		Set( "MediaAsset.AssetActions.Mute.Small", new IMAGE_BRUSH( "Icons/icon_SCueEd_Mute_16x", Icon16x16));
 
 		Set("MediaAsset.AssetActions.Play.Large", new IMAGE_BRUSH("Icons/icon_SCueEd_PlayCue_40x", Icon40x40));
 		Set("MediaAsset.AssetActions.Stop.Large", new IMAGE_BRUSH("Icons/icon_SCueEd_Stop_40x", Icon40x40));
 		Set("MediaAsset.AssetActions.Pause.Large", new IMAGE_BRUSH("Icons/icon_SCueEd_Pause_40x", Icon40x40));
-
+		Set("MediaAsset.AssetActions.Solo.Large", new IMAGE_BRUSH("Icons/icon_SCueEd_Solo_40x", Icon40x40));
+		Set("MediaAsset.AssetActions.Mute.Large", new IMAGE_BRUSH("Icons/icon_SCueEd_Mute_40x", Icon40x40));
+				
 		// Misc
 		Set( "ContentBrowser.ThumbnailShadow", new BOX_BRUSH( "ContentBrowser/ThumbnailShadow" , FMargin( 4.0f / 64.0f ) ) );
 		Set( "ContentBrowser.ColumnViewAssetIcon", new IMAGE_BRUSH( "Icons/doc_16x", Icon16x16 ) );
@@ -6989,11 +6993,15 @@ void FSlateEditorStyle::FStyle::SetupToolkitStyles()
 		Set("MaterialInstanceEditor.StackBody_Highlighted", new BOX_BRUSH("Common/TextBox_Special_Active", FMargin(8.0f / 32.0f)));
 
 	}
-
-
+	
 	// Sound Class Editor
 	{
 		Set( "SoundClassEditor.Tabs.Properties", new IMAGE_BRUSH( "/Icons/icon_tab_SelectionDetails_16x", Icon16x16 ) );
+
+		Set("SoundClassEditor.ToggleSolo", new IMAGE_BRUSH("Icons/icon_SCueEd_Solo_40x", Icon40x40));
+		Set("SoundClassEditor.ToggleSolo.Small", new IMAGE_BRUSH("Icons/icon_SCueEd_Solo_40x", Icon20x20));
+		Set("SoundClassEditor.ToggleMute", new IMAGE_BRUSH("Icons/icon_SCueEd_Mute_40x", Icon40x40));
+		Set("SoundClassEditor.ToggleMute.Small", new IMAGE_BRUSH("Icons/icon_SCueEd_Mute_40x", Icon20x20));
 	}
 
 	// Font Editor
@@ -7031,6 +7039,11 @@ void FSlateEditorStyle::FStyle::SetupToolkitStyles()
 		Set( "SoundCueGraphEditor.PlayNode.Small", new IMAGE_BRUSH( "Icons/icon_SCueEd_PlayNode_40x", Icon20x20 ) );
 		Set( "SoundCueGraphEditor.StopCueNode", new IMAGE_BRUSH( "Icons/icon_SCueEd_Stop_40x", Icon40x40 ) );
 		Set( "SoundCueGraphEditor.StopCueNode.Small", new IMAGE_BRUSH( "Icons/icon_SCueEd_Stop_40x", Icon20x20 ) );
+
+		Set("SoundCueGraphEditor.ToggleSolo", new IMAGE_BRUSH("Icons/icon_SCueEd_Solo_40x", Icon40x40));
+		Set("SoundCueGraphEditor.ToggleSolo.Small", new IMAGE_BRUSH("Icons/icon_SCueEd_Solo_40x", Icon20x20));
+		Set("SoundCueGraphEditor.ToggleMute", new IMAGE_BRUSH("Icons/icon_SCueEd_Mute_40x", Icon40x40));
+		Set("SoundCueGraphEditor.ToggleMute.Small", new IMAGE_BRUSH("Icons/icon_SCueEd_Mute_40x", Icon20x20));
 	}
 
 	// Static Mesh Editor
diff --git a/Engine/Source/Runtime/AudioMixer/Private/AudioMixerSource.cpp b/Engine/Source/Runtime/AudioMixer/Private/AudioMixerSource.cpp
--- a/Engine/Source/Runtime/AudioMixer/Private/AudioMixerSource.cpp
+++ b/Engine/Source/Runtime/AudioMixer/Private/AudioMixerSource.cpp
@@ -1,14 +1,14 @@
 // Copyright 1998-2019 Epic Games, Inc. All Rights Reserved.
 
 #include "AudioMixerSource.h"
+#include "ActiveSound.h"
 #include "AudioMixerSourceBuffer.h"
 #include "AudioMixerDevice.h"
 #include "AudioMixerSourceVoice.h"
-#include "ActiveSound.h"
-#include "IAudioExtensionPlugin.h"
-#include "Sound/AudioSettings.h"
 #include "ContentStreaming.h"
+#include "IAudioExtensionPlugin.h"
 #include "ProfilingDebugging/CsvProfiler.h"
+#include "Sound/AudioSettings.h"
 
 // Link to "Audio" profiling category
 CSV_DECLARE_CATEGORY_MODULE_EXTERN(AUDIOMIXERCORE_API, Audio);
@@ -295,7 +295,7 @@ namespace Audio
 
 			bool bIsDebug = false;
 			FString WaveInstanceName = WaveInstance->GetName(); //-V595
-			FString TestName = GEngine->GetAudioDeviceManager()->GetAudioMixerDebugSoundName();
+			FString TestName = GEngine->GetAudioDeviceManager()->GetDebugger().GetAudioMixerDebugSoundName();
 			if (WaveInstanceName.Contains(TestName))
 			{
 				bDebugMode = true;
diff --git a/Engine/Source/Runtime/Engine/Private/Audio.cpp b/Engine/Source/Runtime/Engine/Private/Audio.cpp
--- a/Engine/Source/Runtime/Engine/Private/Audio.cpp
+++ b/Engine/Source/Runtime/Engine/Private/Audio.cpp
@@ -5,24 +5,24 @@
 =============================================================================*/
 
 #include "Audio.h"
-#include "Misc/Paths.h"
-#include "UObject/UObjectHash.h"
-#include "UObject/UObjectIterator.h"
-#include "Components/AudioComponent.h"
-#include "ContentStreaming.h"
-#include "DrawDebugHelpers.h"
+#include "ActiveSound.h"
+#include "AnalyticsEventAttribute.h"
 #include "Audio/AudioDebug.h"
 #include "AudioDevice.h"
+#include "AudioPluginUtilities.h"
 #include "AudioThread.h"
+#include "Components/AudioComponent.h"
+#include "ContentStreaming.h"
+#include "DrawDebugHelpers.h"
+#include "EngineAnalytics.h"
+#include "Interfaces/IAnalyticsProvider.h"
+#include "Misc/Paths.h"
 #include "Sound/SoundBase.h"
 #include "Sound/SoundCue.h"
-#include "Sound/SoundWave.h"
-#include "ActiveSound.h"
 #include "Sound/SoundNodeWavePlayer.h"
-#include "EngineAnalytics.h"
-#include "AnalyticsEventAttribute.h"
-#include "Interfaces/IAnalyticsProvider.h"
-#include "AudioPluginUtilities.h"
+#include "Sound/SoundWave.h"
+#include "UObject/UObjectHash.h"
+#include "UObject/UObjectIterator.h"
 
 DEFINE_LOG_CATEGORY(LogAudio);
 
@@ -448,73 +448,62 @@ float FSoundSource::GetDebugVolume(const float InVolume)
 {
 	float OutVolume = InVolume;
 
-#if !UE_BUILD_SHIPPING
+#if ENABLE_AUDIO_DEBUG
 
-	if (OutVolume != 0.0f)
+	// Bail if we don't have a device manager.
+	if (!GEngine || !GEngine->GetAudioDeviceManager() || !WaveInstance || !DebugInfo.IsValid() )
 	{
-		// Check for solo sound class debugging. Mute all sounds that don't substring match their sound class name to the debug solo'd sound class
-		const FString& DebugSoloSoundName = GEngine->GetAudioDeviceManager()->GetDebugSoloSoundWave();
-		if (DebugSoloSoundName != TEXT(""))
-		{
-			bool bMute = true;
-			FString WaveInstanceName = WaveInstance->GetName();
-			if (WaveInstanceName.Contains(DebugSoloSoundName))
-			{
-				bMute = false;
-			}
-			if (bMute)
-			{
-				OutVolume = 0.0f;
-			}
-		}
+		return OutVolume;
 	}
 
+	// Solos/Mutes (dev only).
+	FAudioDebugger& Debugger = GEngine->GetAudioDeviceManager()->GetDebugger();	
+	FDebugInfo Info;
+				
+	// SoundWave Solo/Mutes.
 	if (OutVolume != 0.0f)
 	{
-		// Check for solo sound class debugging. Mute all sounds that don't substring match their sound class name to the debug solo'd sound class
-		const FString& DebugSoloSoundCue = GEngine->GetAudioDeviceManager()->GetDebugSoloSoundCue();
-		if (DebugSoloSoundCue != TEXT(""))
+		Debugger.QuerySoloMuteSoundWave(WaveInstance->GetName(), Info.bIsSoloed, Info.bIsMuted, Info.MuteSoloReason);
+		if (Info.bIsMuted)
 		{
-			bool bMute = true;
-			USoundBase* Sound = WaveInstance->ActiveSound->GetSound();
-			if (Sound->IsA<USoundCue>())
-			{
-				FString SoundCueName = Sound->GetName();
-				if (SoundCueName.Contains(DebugSoloSoundCue))
-				{
-					bMute = false;
-				}
-			}
+			OutVolume = 0.0f;
+		}
+	}
 
-			if (bMute)
+	// SoundCues mutes/solos											
+	if (OutVolume != 0.0f && WaveInstance->ActiveSound)
+	{						
+		if (USoundCue* SoundCue = Cast<USoundCue>(WaveInstance->ActiveSound->GetSound()))
+		{
+			Debugger.QuerySoloMuteSoundCue(SoundCue->GetName(), Info.bIsSoloed, Info.bIsMuted, Info.MuteSoloReason);
+			if (Info.bIsMuted)
 			{
 				OutVolume = 0.0f;
 			}
 		}
 	}
 
-	if (OutVolume != 0.0f)
+	// SoundClass mutes/solos.
+	if (OutVolume != 0.0f && WaveInstance->SoundClass)
 	{
-		const FString& DebugSoloSoundClassName = GEngine->GetAudioDeviceManager()->GetDebugSoloSoundClass();
-		if (DebugSoloSoundClassName != TEXT(""))
+		FString SoundClassName;
+		WaveInstance->SoundClass->GetName(SoundClassName);
+		Debugger.QuerySoloMuteSoundClass(SoundClassName, Info.bIsSoloed, Info.bIsMuted, Info.MuteSoloReason);
+		if (Info.bIsMuted)
 		{
-			bool bMute = true;
-			if (WaveInstance->SoundClass)
-			{
-				FString SoundClassName;
-				WaveInstance->SoundClass->GetName(SoundClassName);
-				if (SoundClassName.Contains(DebugSoloSoundClassName))
-				{
-					bMute = false;
-				}
-			}
-			if (bMute)
-			{
-				OutVolume = 0.0f;
-			}
+			OutVolume = 0.0f;
 		}
 	}
-#endif
+
+	// Update State. 
+	FScopeLock Lock(&DebugInfo->CS);
+	{
+		DebugInfo->bIsMuted = Info.bIsMuted;
+		DebugInfo->bIsSoloed = Info.bIsSoloed;
+		DebugInfo->MuteSoloReason = MoveTemp(Info.MuteSoloReason);
+	}
+
+#endif //ENABLE_AUDIO_DEBUG
 
 	return OutVolume;
 }
@@ -588,6 +577,10 @@ void FSoundSource::InitCommon()
 	// Reset pause state
 	bIsPausedByGame = false;
 	bIsManuallyPaused = false;
+	
+#if ENABLE_AUDIO_DEBUG
+	DebugInfo = MakeShared<FDebugInfo, ESPMode::ThreadSafe>();
+#endif //ENABLE_AUDIO_DEBUG
 }
 
 void FSoundSource::UpdateCommon()
diff --git a/Engine/Source/Runtime/Engine/Private/AudioDebug.cpp b/Engine/Source/Runtime/Engine/Private/AudioDebug.cpp
--- a/Engine/Source/Runtime/Engine/Private/AudioDebug.cpp
+++ b/Engine/Source/Runtime/Engine/Private/AudioDebug.cpp
@@ -25,6 +25,9 @@
 #include "Sound/SoundWave.h"
 #include "UnrealEngine.h"
 
+#ifndef ENABLE_AUDIO_DEBUG
+#error "Please define ENABLE_AUDIO_DEBUG"
+#endif //ENABLE_AUDIO_DEBUG
 
 #if ENABLE_AUDIO_DEBUG
 
@@ -108,6 +111,7 @@ namespace
 
 		struct FStatWaveInstanceInfo
 		{
+			TSharedPtr<FSoundSource::FDebugInfo, ESPMode::ThreadSafe> DebugInfo;
 			FString Description;
 			float Volume;
 			int32 InstanceIndex;
@@ -165,8 +169,18 @@ namespace
 			GEngine->GetAudioDeviceManager()->GetDebugger().DumpActiveSounds();
 		}
 	}
+	
+	void HandleClearMutesAndSolos(UWorld* World)
+	{
+		if (GEngine)
+		{
+			GEngine->GetAudioDeviceManager()->GetDebugger().ClearMutesAndSolos();
+		}
+	}
 
 	FAutoConsoleCommandWithWorld DumpActiveSounds(TEXT("Audio.DumpActiveSounds"), TEXT("Outputs data about all the currently active sounds."), FConsoleCommandWithWorldDelegate::CreateStatic(&HandleDumpActiveSounds), ECVF_Cheat);
+	FAutoConsoleCommandWithWorld ClearMutesAndSolos(TEXT("Audio.ClearMutesAndSolos"), TEXT("Clears any solo-ing/mute-ing sounds"), FConsoleCommandWithWorldDelegate::CreateStatic(&HandleClearMutesAndSolos), ECVF_Cheat);
+
 } // namespace <>
 
 
@@ -188,6 +202,132 @@ void FAudioDebugger::ToggleVisualizeDebug3dEnabled()
 	bVisualize3dDebug = !bVisualize3dDebug;
 }
 
+void FAudioDebugger::QuerySoloMuteSoundClass(const FString& Name, bool& bOutIsSoloed, bool& bOutIsMuted, FString& OutReason) const
+{
+	GetDebugSoloMuteStateX(Name, DebugNames.SoloSoundClass, DebugNames.MuteSoundClass, bOutIsSoloed, bOutIsMuted, OutReason);
+}
+
+void FAudioDebugger::QuerySoloMuteSoundWave(const FString& Name, bool& bOutIsSoloed, bool& bOutIsMuted, FString& OutReason) const
+{
+	GetDebugSoloMuteStateX(Name, DebugNames.SoloSoundWave, DebugNames.MuteSoundWave, bOutIsSoloed, bOutIsMuted, OutReason);
+}
+
+void FAudioDebugger::QuerySoloMuteSoundCue(const FString& Name, bool& bOutIsSoloed, bool& bOutIsMuted, FString& OutReason) const
+{
+	GetDebugSoloMuteStateX(Name, DebugNames.SoloSoundCue, DebugNames.MuteSoundCue, bOutIsSoloed, bOutIsMuted, OutReason);
+}
+
+void FAudioDebugger::SetNameArray(FName InName, TArray<FName>& InNameArray, bool bOnOff)
+{
+	ExecuteCmdOnAudioThread([InName, &InNameArray, bOnOff]
+	{
+		if (bOnOff)
+		{
+			InNameArray.AddUnique(InName);
+		}
+		else
+		{
+			InNameArray.Remove(InName);
+		}
+	});
+}
+
+void FAudioDebugger::ToggleNameArray(FName InName, TArray<FName>& InNameArray )
+{
+	ExecuteCmdOnAudioThread([InName, &InNameArray]
+	{ 
+		// On already?
+		int32 IndexOf = InNameArray.IndexOfByKey(InName);
+		if (IndexOf != INDEX_NONE)
+		{
+			InNameArray.RemoveAtSwap(IndexOf);
+		}			
+		else // Add it.
+		{
+			InNameArray.Add(InName);
+		}
+	});
+}
+
+void FAudioDebugger::SetAudioMixerDebugSound(const TCHAR* SoundName)
+{
+	ExecuteCmdOnAudioThread([Name = FString(SoundName), this]{ DebugNames.DebugAudioMixerSoundName = Name; });
+}
+
+void FAudioDebugger::SetAudioDebugSound(const TCHAR* SoundName)
+{
+	ExecuteCmdOnAudioThread([Name = FString(SoundName), this]{ 
+		DebugNames.DebugSoundName = Name;
+		DebugNames.bDebugSoundName = DebugNames.DebugSoundName != TEXT("");
+	});	
+}
+
+const FString& FAudioDebugger::GetAudioMixerDebugSoundName() const
+{
+	check(IsInAudioThread());
+	return DebugNames.DebugAudioMixerSoundName;
+}
+
+bool FAudioDebugger::GetAudioDebugSound(FString& OutDebugSound)
+{
+	check(IsInAudioThread());
+	if (DebugNames.bDebugSoundName)
+	{
+		OutDebugSound = DebugNames.DebugSoundName;
+		return true;
+	}
+	return false;
+}
+
+void FAudioDebugger::ExecuteCmdOnAudioThread(TFunction<void()> Cmd)
+{
+	// If not on audio thread, queue it.
+	if (!IsInAudioThread())
+	{
+		FAudioThread::RunCommandOnAudioThread(Cmd);
+		return;
+	}
+
+	// Otherwise, do it inline.
+	Cmd();
+}
+
+void FAudioDebugger::GetDebugSoloMuteStateX(
+	const FString& Name, 
+	const TArray<FName>& Solos,
+	const TArray<FName>& Mutes,	
+	bool& bOutIsSoloed, 
+	bool& bOutIsMuted, 
+	FString& OutReason) const
+{
+	check(IsInAudioThread());
+	
+	// Allow for partial matches of the name.
+	auto MatchesPartOfName = [&Name](FName i) -> bool { return Name.Contains(i.GetPlainNameString()); };
+
+	// Solo active?
+	if (Solos.Num() > 0)
+	{
+		if (Solos.ContainsByPredicate(MatchesPartOfName))
+		{
+			bOutIsSoloed = true;
+			OutReason = FString::Printf(TEXT("Sound is soloed explicitly."));
+		}
+		else
+		{
+			// Something else is soloed (record the first item in the solo list for debug reason).
+			bOutIsMuted = true;
+			OutReason = FString::Printf(TEXT("Sound is muted due to [%s] being soloed"), *Solos[0].ToString() );
+		}
+	}
+	// Are we explicitly muted?
+	else if (Mutes.ContainsByPredicate(MatchesPartOfName))
+	{
+		bOutIsMuted = true;
+		OutReason = FString::Printf(TEXT("Sound is explicitly muted"));
+	}
+}
+
 void FAudioDebugger::DrawDebugInfo(const FSoundSource& SoundSource)
 {
 #if ENABLE_DRAW_DEBUG
@@ -549,8 +689,19 @@ int32 FAudioDebugger::RenderStatCues(UWorld* World, FViewport* Viewport, FCanvas
 		{
 			if (WaveInstanceInfo.Volume >= MinDisplayVolume)
 			{
-				const FString TheString = FString::Printf(TEXT("%4i. %s %s"), ActiveSoundCount++, *StatSoundInfo.SoundName, *StatSoundInfo.SoundClassName.ToString());
-				Canvas->DrawShadowedString(X, Y, *TheString, GetStatsFont(), FColor::White);
+				FColor Color = FColor::White;
+				FString MuteSoloReason;
+				
+				if (FSoundSource::FDebugInfo* DebugInfo = WaveInstanceInfo.DebugInfo.Get())
+				{
+					// Color code same as icons. Orange (mute), Solo (red), Normal (white).					
+					FScopeLock Lock(&DebugInfo->CS);
+					Color = DebugInfo->bIsMuted ? FColor::Orange : DebugInfo->bIsSoloed ? FColor::Red : FColor::White;
+					MuteSoloReason = !DebugInfo->MuteSoloReason.IsEmpty() ? FString::Printf(TEXT(" - %s"), *DebugInfo->MuteSoloReason) : TEXT("");
+				}
+
+				const FString TheString = FString::Printf(TEXT("%4i. %s %s %s"), ActiveSoundCount++, *StatSoundInfo.SoundName, *StatSoundInfo.SoundClassName.ToString(), *MuteSoloReason );
+				Canvas->DrawShadowedString(X, Y, *TheString, GetStatsFont(), Color);
 				Y += FontHeight;
 				break;
 			}
@@ -1168,6 +1319,8 @@ void FAudioDebugger::SendUpdateResultsToGameThread(const FAudioDevice& AudioDevi
 	const uint8 RequestedStats = Stats_AudioThread->RequestedStats;
 	TMap<FActiveSound*, int32> ActiveSoundToInfoIndex;
 
+	FAudioDebugger& Debugger = GEngine->GetAudioDeviceManager()->GetDebugger();
+	
 	const bool bDebug = (RequestedStats & ERequestedAudioStats::DebugSounds) != 0;
 
 	static const uint8 SoundMask = ERequestedAudioStats::Sounds | ERequestedAudioStats::SoundCues | ERequestedAudioStats::SoundWaves;
@@ -1176,7 +1329,7 @@ void FAudioDebugger::SendUpdateResultsToGameThread(const FAudioDevice& AudioDevi
 		for (FActiveSound* ActiveSound : AudioDevice.GetActiveSounds())
 		{
 			if (USoundBase* SoundBase = ActiveSound->GetSound())
-			{
+			{				
 				if (!bDebug || ActiveSound->GetSound()->bDebug)
 				{
 					ActiveSoundToInfoIndex.Add(ActiveSound, StatSoundInfos.AddDefaulted());
@@ -1187,6 +1340,7 @@ void FAudioDebugger::SendUpdateResultsToGameThread(const FAudioDevice& AudioDevi
 					if (USoundClass* SoundClass = ActiveSound->GetSoundClass())
 					{
 						StatSoundInfo.SoundClassName = SoundClass->GetFName();
+
 					}
 					else
 					{
@@ -1194,16 +1348,15 @@ void FAudioDebugger::SendUpdateResultsToGameThread(const FAudioDevice& AudioDevi
 					}
 					StatSoundInfo.Transform = ActiveSound->Transform;
 					StatSoundInfo.AudioComponentID = ActiveSound->GetAudioComponentID();
-
+										
 					if (bDebug && ActiveSound->GetSound()->bDebug)
 					{
 						ActiveSound->CollectAttenuationShapesForVisualization(StatSoundInfo.ShapeDetailsMap);
-					}
+					}										
 				}
 			}
 		}
 
-
 		// Iterate through all wave instances.
 		const TArray<FWaveInstance*>& WaveInstances = AudioDevice.GetActiveWaveInstances();
 		auto WaveInstanceSourceMap = AudioDevice.GetWaveInstanceSourceMap();
@@ -1220,6 +1373,7 @@ void FAudioDebugger::SendUpdateResultsToGameThread(const FAudioDevice& AudioDevi
 				WaveInstanceInfo.InstanceIndex = InstanceIndex;
 				WaveInstanceInfo.WaveInstanceName = *WaveInstance->GetName();
 				WaveInstanceInfo.bPlayWhenSilent = WaveInstance->ActiveSound->IsPlayWhenSilent() ? 1 : 0;
+				WaveInstanceInfo.DebugInfo = Source ? Source->DebugInfo : WaveInstanceInfo.DebugInfo;
 				StatSoundInfos[*SoundInfoIndex].WaveInstanceInfos.Add(MoveTemp(WaveInstanceInfo));
 			}
 		}
@@ -1278,4 +1432,15 @@ void FAudioDebugger::UpdateAudibleInactiveSounds(const uint32 FirstActiveIndex,
 	SET_DWORD_STAT(STAT_AudibleWavesDroppedDueToPriority, AudibleInactiveSounds);
 #endif
 }
+
+void FAudioDebugger::ClearMutesAndSolos()
+{
+	DebugNames.MuteSoundClass.Empty();
+	DebugNames.MuteSoundCue.Empty();
+	DebugNames.MuteSoundWave.Empty();
+	DebugNames.SoloSoundClass.Empty();
+	DebugNames.SoloSoundCue.Empty();
+	DebugNames.SoloSoundWave.Empty();
+}
+
 #endif // ENABLE_AUDIO_DEBUG
diff --git a/Engine/Source/Runtime/Engine/Private/AudioDevice.cpp b/Engine/Source/Runtime/Engine/Private/AudioDevice.cpp
--- a/Engine/Source/Runtime/Engine/Private/AudioDevice.cpp
+++ b/Engine/Source/Runtime/Engine/Private/AudioDevice.cpp
@@ -1336,7 +1336,7 @@ bool FAudioDevice::HandleAudioSoloSoundClass(const TCHAR* Cmd, FOutputDevice& Ar
 	FAudioDeviceManager* DeviceManager = GEngine->GetAudioDeviceManager();
 	if (DeviceManager)
 	{
-		DeviceManager->SetDebugSoloSoundClass(Cmd);
+		DeviceManager->GetDebugger().ToggleSoloSoundClass(Cmd);
 	}
 	return true;
 }
@@ -1346,7 +1346,7 @@ bool FAudioDevice::HandleAudioSoloSoundWave(const TCHAR* Cmd, FOutputDevice& Ar)
 	FAudioDeviceManager* DeviceManager = GEngine->GetAudioDeviceManager();
 	if (DeviceManager)
 	{
-		DeviceManager->SetDebugSoloSoundWave(Cmd);
+		DeviceManager->GetDebugger().ToggleSoloSoundWave(Cmd);
 	}
 	return true;
 }
@@ -1356,7 +1356,7 @@ bool FAudioDevice::HandleAudioSoloSoundCue(const TCHAR* Cmd, FOutputDevice& Ar)
 	FAudioDeviceManager* DeviceManager = GEngine->GetAudioDeviceManager();
 	if (DeviceManager)
 	{
-		DeviceManager->SetDebugSoloSoundCue(Cmd);
+		DeviceManager->GetDebugger().ToggleSoloSoundCue(Cmd);
 	}
 	return true;
 }
@@ -1366,7 +1366,7 @@ bool FAudioDevice::HandleAudioMixerDebugSound(const TCHAR* Cmd, FOutputDevice& A
 	FAudioDeviceManager* DeviceManager = GEngine->GetAudioDeviceManager();
 	if (DeviceManager)
 	{
-		DeviceManager->SetAudioMixerDebugSound(Cmd);
+		DeviceManager->GetDebugger().SetAudioMixerDebugSound(Cmd);
 	}
 	return true;
 }
@@ -1376,7 +1376,7 @@ bool FAudioDevice::HandleAudioDebugSound(const TCHAR* Cmd, FOutputDevice& Ar)
 	FAudioDeviceManager* DeviceManager = GEngine->GetAudioDeviceManager();
 	if (DeviceManager)
 	{
-		DeviceManager->SetAudioDebugSound(Cmd);
+		DeviceManager->GetDebugger().SetAudioDebugSound(Cmd);
 	}
 	return true;
 }
@@ -4297,7 +4297,7 @@ void FAudioDevice::AddNewActiveSoundInternal(const FActiveSound& NewActiveSound,
 #if !UE_BUILD_SHIPPING
 	FAudioDeviceManager* AudioDeviceManager = GEngine->GetAudioDeviceManager();
 	FString DebugSound;
-	if (AudioDeviceManager->GetAudioDebugSound(DebugSound))
+	if (AudioDeviceManager->GetDebugger().GetAudioDebugSound(DebugSound))
 	{
 		// Reject the new sound if it doesn't have the debug sound name substring
 		FString SoundName;
diff --git a/Engine/Source/Runtime/Engine/Private/AudioDeviceManager.cpp b/Engine/Source/Runtime/Engine/Private/AudioDeviceManager.cpp
--- a/Engine/Source/Runtime/Engine/Private/AudioDeviceManager.cpp
+++ b/Engine/Source/Runtime/Engine/Private/AudioDeviceManager.cpp
@@ -80,12 +80,17 @@ FAudioDeviceManager::FAudioDeviceManager()
 	, bOnlyToggleAudioMixerOnce(false)
 	, bToggledAudioMixer(false)
 {
+
+#if ENABLE_AUDIO_DEBUG
+
 	// Check for a command line debug sound argument.
 	FString DebugSound;
 	if (FParse::Value(FCommandLine::Get(), TEXT("DebugSound="), DebugSound))
 	{
-		SetAudioDebugSound(*DebugSound);
+		GetDebugger().SetAudioDebugSound(*DebugSound);
 	}
+
+#endif //ENABLE_AUDIO_DEBUG
 }
 
 FAudioDeviceManager::~FAudioDeviceManager()
@@ -996,101 +1001,6 @@ void FAudioDeviceManager::ToggleVisualize3dDebug()
 #endif // ENABLE_AUDIO_DEBUG
 }
 
-void FAudioDeviceManager::SetDebugSoloSoundClass(const TCHAR* SoundClassName)
-{
-	if (!IsInAudioThread())
-	{
-		DECLARE_CYCLE_STAT(TEXT("FAudioThreadTask.SetDebugSoloSoundClass"), STAT_SetDebugSoloSoundClass, STATGROUP_AudioThreadCommands);
-
-		FAudioDeviceManager* AudioDeviceManager = this;
-		FAudioThread::RunCommandOnAudioThread([AudioDeviceManager, SoundClassName]()
-		{
-			AudioDeviceManager->SetDebugSoloSoundClass(SoundClassName);
-
-		}, GET_STATID(STAT_SetDebugSoloSoundClass));
-		return;
-	}
-
-	DebugNames.DebugSoloSoundClass = SoundClassName;
-}
-
-const FString& FAudioDeviceManager::GetDebugSoloSoundClass() const
-{
-	return DebugNames.DebugSoloSoundClass;
-}
-
-void FAudioDeviceManager::SetDebugSoloSoundWave(const TCHAR* SoundWave)
-{
-	if (!IsInAudioThread())
-	{
-		DECLARE_CYCLE_STAT(TEXT("FAudioThreadTask.SetDebugSoloSoundWave"), STAT_SetDebugSoloSoundWave, STATGROUP_AudioThreadCommands);
-
-		FAudioDeviceManager* AudioDeviceManager = this;
-		FAudioThread::RunCommandOnAudioThread([AudioDeviceManager, SoundWave]()
-		{
-			AudioDeviceManager->SetDebugSoloSoundWave(SoundWave);
-
-		}, GET_STATID(STAT_SetDebugSoloSoundWave));
-		return;
-	}
-
-	DebugNames.DebugSoloSoundWave = SoundWave;
-}
-
-const FString& FAudioDeviceManager::GetDebugSoloSoundWave() const
-{
-	return DebugNames.DebugSoloSoundWave;
-}
-
-void FAudioDeviceManager::SetDebugSoloSoundCue(const TCHAR* SoundCue)
-{
-	if (!IsInAudioThread())
-	{
-		DECLARE_CYCLE_STAT(TEXT("FAudioThreadTask.SetDebugSoloSoundCue"), STAT_SetDebugSoloSoundCue, STATGROUP_AudioThreadCommands);
-
-		FAudioDeviceManager* AudioDeviceManager = this;
-		FAudioThread::RunCommandOnAudioThread([AudioDeviceManager, SoundCue]()
-		{
-			AudioDeviceManager->SetDebugSoloSoundCue(SoundCue);
-
-		}, GET_STATID(STAT_SetDebugSoloSoundCue));
-		return;
-	}
-
-	DebugNames.DebugSoloSoundCue = SoundCue;
-}
-
-const FString& FAudioDeviceManager::GetDebugSoloSoundCue() const
-{
-	return DebugNames.DebugSoloSoundCue;
-}
-
-void FAudioDeviceManager::SetAudioMixerDebugSound(const TCHAR* SoundName)
-{
-	DebugNames.DebugAudioMixerSoundName = SoundName;
-}
-
-void FAudioDeviceManager::SetAudioDebugSound(const TCHAR* SoundName)
-{
-	DebugNames.DebugSoundName = SoundName;
-	DebugNames.bDebugSoundName = DebugNames.DebugSoundName != TEXT("");
-}
-
-const FString& FAudioDeviceManager::GetAudioMixerDebugSoundName() const
-{
-	return DebugNames.DebugAudioMixerSoundName;
-}
-
-bool FAudioDeviceManager::GetAudioDebugSound(FString& OutDebugSound)
-{
-	if (DebugNames.bDebugSoundName)
-	{
-		OutDebugSound = DebugNames.DebugSoundName;
-		return true;
-	}
-	return false;
-}
-
 float FAudioDeviceManager::GetDynamicSoundVolume(ESoundType SoundType, const FName& SoundName) const
 {
 	check(IsInAudioThread());
diff --git a/Engine/Source/Runtime/Engine/Public/Audio.h b/Engine/Source/Runtime/Engine/Public/Audio.h
--- a/Engine/Source/Runtime/Engine/Public/Audio.h
+++ b/Engine/Source/Runtime/Engine/Public/Audio.h
@@ -775,6 +775,13 @@ protected:
 
 	friend class FAudioDevice;
 	friend struct FActiveSound;
+
+#if ENABLE_AUDIO_DEBUG
+public:
+	struct FDebugInfo;
+	TSharedPtr<FDebugInfo, ESPMode::ThreadSafe> DebugInfo;
+	friend struct FDebugInfo;
+#endif //ENABLE_AUDIO_DEBUG
 };
 
 /*-----------------------------------------------------------------------------
diff --git a/Engine/Source/Runtime/Engine/Public/Audio/AudioDebug.h b/Engine/Source/Runtime/Engine/Public/Audio/AudioDebug.h
--- a/Engine/Source/Runtime/Engine/Public/Audio/AudioDebug.h
+++ b/Engine/Source/Runtime/Engine/Public/Audio/AudioDebug.h
@@ -4,8 +4,7 @@
 #include "CanvasTypes.h"
 #include "CoreMinimal.h"
 
-#define ENABLE_AUDIO_DEBUG !UE_BUILD_SHIPPING
-
+#include "AudioDefines.h"		// For ENABLE_AUDIO_DEBUG
 
 #if ENABLE_AUDIO_DEBUG
 
@@ -16,12 +15,46 @@ struct FWaveInstance;
 
 class FSoundSource;
 
+/** Struct containing the debug state of a SoundSource */
+struct ENGINE_API FSoundSource::FDebugInfo
+{
+	/** True if this sound has been soloed. */
+	bool bIsSoloed = false;
+
+	/** True if this sound has been muted . */
+	bool bIsMuted = false;
+
+	/** Reason why this sound is mute/soloed. */
+	FString MuteSoloReason;
+
+	/** Basic CS so we can pass this around safely. */
+	FCriticalSection CS;
+};
 
 class ENGINE_API FAudioDebugger
 {
 public:
 	FAudioDebugger();
 
+	/** Struct which contains debug names for run-time debugging of sounds. */
+	struct FDebugNames
+	{
+		TArray<FName> SoloSoundClass;
+		TArray<FName> SoloSoundWave;
+		TArray<FName> SoloSoundCue;
+		TArray<FName> MuteSoundClass;
+		TArray<FName> MuteSoundWave;
+		TArray<FName> MuteSoundCue;
+
+		FString DebugAudioMixerSoundName;
+		FString DebugSoundName;
+		bool bDebugSoundName;
+
+		FDebugNames()
+			: bDebugSoundName(false)
+		{}
+	};
+
 	static void DrawDebugInfo(const FSoundSource& SoundSource);
 	static void DrawDebugInfo(const FActiveSound& ActiveSound, const TArray<FWaveInstance*>& ThisSoundsWaveInstances, const float DeltaTime);
 	static void DrawDebugInfo(const FAudioVirtualLoop& VirtualLoop);
@@ -42,14 +75,59 @@ public:
 	static bool ToggleStatWaves(UWorld* World, FCommonViewportClient* ViewportClient, const TCHAR* Stream);
 	static void UpdateAudibleInactiveSounds(const uint32 FirstIndex, const TArray<FWaveInstance*>& WaveInstances);
 
+	void ClearMutesAndSolos();
 	void DumpActiveSounds() const;
+
 	bool IsVisualizeDebug3dEnabled() const;
 	void ToggleVisualizeDebug3dEnabled();
+		
+	// Evaluate Mute/Solos
+	void QuerySoloMuteSoundClass(const FString& Name, bool& bOutIsSoloed, bool& bOutIsMuted, FString& OutReason ) const;
+	void QuerySoloMuteSoundWave(const FString& Name, bool& bOutIsSoloed, bool& bOutIsMuted, FString& OutReason) const;
+	void QuerySoloMuteSoundCue(const FString& Name, bool& bOutIsSoloed, bool& bOutIsMuted, FString& OutReason ) const;
+   
+	// Is Mute/Solos. (only audio thread).
+	bool IsSoloSoundClass(FName InName) const { return DebugNames.SoloSoundClass.Contains(InName); }
+	bool IsSoloSoundWave(FName InName) const  { return DebugNames.SoloSoundWave.Contains(InName); }
+	bool IsSoloSoundCue(FName InName) const   { return DebugNames.SoloSoundCue.Contains(InName); }
+	bool IsMuteSoundClass(FName InName) const { return DebugNames.MuteSoundClass.Contains(InName); }
+	bool IsMuteSoundWave(FName InName) const  { return DebugNames.MuteSoundWave.Contains(InName); }
+	bool IsMuteSoundCue(FName InName) const   { return DebugNames.MuteSoundCue.Contains(InName); }
+		
+	// Mute/Solos toggles. (any thread).
+	void ToggleSoloSoundClass(FName InName) { ToggleNameArray(InName, DebugNames.SoloSoundClass); }
+	void ToggleSoloSoundWave(FName InName)  { ToggleNameArray(InName, DebugNames.SoloSoundWave); }
+	void ToggleSoloSoundCue(FName InName)   { ToggleNameArray(InName, DebugNames.SoloSoundCue); }
+	void ToggleMuteSoundClass(FName InName) { ToggleNameArray(InName, DebugNames.MuteSoundClass); }
+	void ToggleMuteSoundWave(FName InName)  { ToggleNameArray(InName, DebugNames.MuteSoundWave); }
+	void ToggleMuteSoundCue(FName InName)   { ToggleNameArray(InName, DebugNames.MuteSoundCue); }
 
+	// Set Mute/Solo. (any thread).
+	void SetMuteSoundCue(FName InName, bool bInOnOff)  { SetNameArray(InName, DebugNames.MuteSoundCue, bInOnOff); }
+	void SetMuteSoundWave(FName InName, bool bInOnOff) { SetNameArray(InName, DebugNames.MuteSoundWave, bInOnOff); }
+	void SetSoloSoundCue(FName InName, bool bInOnOff)  { SetNameArray(InName, DebugNames.SoloSoundCue, bInOnOff); }	
+	void SetSoloSoundWave(FName InName, bool bInOnOff) { SetNameArray(InName, DebugNames.SoloSoundWave, bInOnOff); }
+		   	
+	void SetAudioMixerDebugSound(const TCHAR* SoundName);
+	void SetAudioDebugSound(const TCHAR* SoundName);
+
+	const FString& GetAudioMixerDebugSoundName() const;
+	bool GetAudioDebugSound(FString& OutDebugSound);	  	
 private:
+	void SetNameArray(FName InName, TArray<FName>& InNameArray, bool bOnOff);
+	void ToggleNameArray(FName InName, TArray<FName>& NameArray);
+	void ExecuteCmdOnAudioThread(TFunction<void()> Cmd);
+	
+	void GetDebugSoloMuteStateX(
+		const FString& Name, const TArray<FName>& Solos, const TArray<FName>& Mutes, 
+		bool& bOutIsSoloed, bool& bOutIsMuted, FString& OutReason) const;
+
 	static bool ToggleStats(UWorld* World, const uint8 StatToToggle);
 	void ToggleStats(const uint32 AudioDeviceHandle, const uint8 StatsToToggle);
 
+	/** Instance of the debug names struct. */
+	FDebugNames DebugNames;
+
 	/** Whether or not 3d debug visualization is enabled. */
 	uint8 bVisualize3dDebug : 1;
 };
diff --git a/Engine/Source/Runtime/Engine/Public/AudioDefines.h b/Engine/Source/Runtime/Engine/Public/AudioDefines.h
--- a/Engine/Source/Runtime/Engine/Public/AudioDefines.h
+++ b/Engine/Source/Runtime/Engine/Public/AudioDefines.h
@@ -49,3 +49,7 @@
 #define MIN_FILTER_BANDWIDTH			0.1f
 #define MAX_FILTER_BANDWIDTH			2.0f
 
+/**
+ * Debugger is Available on non-shipping builds
+ */
+#define ENABLE_AUDIO_DEBUG !UE_BUILD_SHIPPING
\ No newline at end of file
diff --git a/Engine/Source/Runtime/Engine/Public/AudioDeviceManager.h b/Engine/Source/Runtime/Engine/Public/AudioDeviceManager.h
--- a/Engine/Source/Runtime/Engine/Public/AudioDeviceManager.h
+++ b/Engine/Source/Runtime/Engine/Public/AudioDeviceManager.h
@@ -184,33 +184,6 @@ public:
 	/** Toggles 3d visualization of 3d sounds on/off */
 	void ToggleVisualize3dDebug();
 
-	/** Debug solos the given sound class name. Sounds that play with this sound class will be solo'd */
-	void SetDebugSoloSoundClass(const TCHAR* SoundClassName);
-
-	/** Debug solos the given sound name. Sounds that play with this name will be solo'd */
-	void SetDebugSoloSoundWave(const TCHAR* SoundClassName);
-
-	void SetDebugSoloSoundCue(const TCHAR* SoundClassName);
-
-	/** Enables debug logging and soloing of sounds that substring match the given sound name. Only works if built with AUDIO_MIXER_ENABLE_DEBUG_MODE is turned on. */
-	void SetAudioMixerDebugSound(const TCHAR* SoundName);
-	void SetAudioDebugSound(const TCHAR* SoundName);
-
-	/** Gets the solo sound class used for debugging sounds */
-	const FString& GetDebugSoloSoundClass() const;
-
-	/** Gets the solo sound name used for debugging sound waves */
-	const FString& GetDebugSoloSoundWave() const;
-
-	/** Gets the solo sound name used for debugging sound cues */
-	const FString& GetDebugSoloSoundCue() const;
-
-	/** Gets the audio mixer debug sound name. Returns empty string if AUDIO_MIXER_ENABLE_DEBUG_MODE is not enabled. */
-	const FString& GetAudioMixerDebugSoundName() const;
-
-	/** Gets the debug sound string, returns true if one has been set. */
-	bool GetAudioDebugSound(FString& OutDebugSound);
-
 	/** Reset all sound cue trims */
 	void ResetAllDynamicSoundVolumes();
 
@@ -234,29 +207,14 @@ public:
 
 	/** Returns all the audio devices managed by device manager. */
 	TArray<FAudioDevice*>& GetAudioDevices() { return Devices; }
-
-
+	
 private:
+
 #if ENABLE_AUDIO_DEBUG
 	/** Instance of audio debugger shared across audio devices */
 	FAudioDebugger AudioDebugger;
 #endif // ENABLE_AUDIO_DEBUG
 
-	/** Struct which contains debug names for run-time debugging of sounds. */
-	struct FDebugNames
-	{
-		FString DebugSoloSoundClass;
-		FString DebugSoloSoundWave;
-		FString DebugSoloSoundCue;
-		FString DebugAudioMixerSoundName;
-		FString DebugSoundName;
-		bool bDebugSoundName;
-
-		FDebugNames()
-			: bDebugSoundName(false)
-		{}
-	};
-
 	/** Returns index of the given handle */
 	uint32 GetIndex(uint32 Handle) const;
 
@@ -317,9 +275,6 @@ private:
 	/** Which audio device is currently active */
 	Audio::FDeviceId ActiveAudioDeviceHandle;
 
-	/** Instance of the debug names struct. */
-	FDebugNames DebugNames;
-
 	/** Dynamic volume map */
 	TMap<TTuple<ESoundType, FName>, float> DynamicSoundVolumes;
 
