Fix GC crash at world transition time if the current world has a post process mid that references textures from the same world (not assets)

diff --git a/Engine/Source/Runtime/Engine/Classes/Engine/LocalPlayer.h b/Engine/Source/Runtime/Engine/Classes/Engine/LocalPlayer.h
--- a/Engine/Source/Runtime/Engine/Classes/Engine/LocalPlayer.h
+++ b/Engine/Source/Runtime/Engine/Classes/Engine/LocalPlayer.h
@@ -480,6 +480,11 @@ public:
 	 * @return	true if this player is not using splitscreen, or is the first player in the split-screen layout.
 	 */
 	bool IsPrimaryPlayer() const;
+	 
+	/**
+	 * Clear cached view state.  Suitable for calling when cleaning up the world but the view state has some references objects (usually mids) owned by the world (thus preventing GC) 
+	 */
+	void CleanupViewState();
 
 	/** Locked view state needs access to GetViewPoint. */
 	friend class FLockedViewState;
diff --git a/Engine/Source/Runtime/Engine/Private/LocalPlayer.cpp b/Engine/Source/Runtime/Engine/Private/LocalPlayer.cpp
--- a/Engine/Source/Runtime/Engine/Private/LocalPlayer.cpp
+++ b/Engine/Source/Runtime/Engine/Private/LocalPlayer.cpp
@@ -1669,4 +1669,11 @@ bool ULocalPlayer::IsPrimaryPlayer() const
 	return (this == PrimaryPlayer);
 }
 
+void ULocalPlayer::CleanupViewState()
+{
+	if (FSceneViewStateInterface* Ref = ViewState.GetReference())
+	{
+		Ref->ClearMIDPool();
+	}
+}
 
diff --git a/Engine/Source/Runtime/Engine/Private/UnrealEngine.cpp b/Engine/Source/Runtime/Engine/Private/UnrealEngine.cpp
--- a/Engine/Source/Runtime/Engine/Private/UnrealEngine.cpp
+++ b/Engine/Source/Runtime/Engine/Private/UnrealEngine.cpp
@@ -12321,6 +12321,8 @@ bool UEngine::LoadMap( FWorldContext& WorldContext, FURL URL, class UPendingNetG
 				}
 				// reset split join info so we'll send one after loading the new map if necessary
 				Player->bSentSplitJoin = false;
+				// When loading maps, clear out mids that are referenced as they may prevent the world from shutting down cleanly and the local player will not be cleaned up until later
+				Player->CleanupViewState(); 
 			}
 		}
 
