fix for non-unity compile error

diff --git a/Engine/Source/Editor/AudioEditor/Private/AssetTypeActions/AssetTypeActions_SoundClass.h b/Engine/Source/Editor/AudioEditor/Private/AssetTypeActions/AssetTypeActions_SoundClass.h
--- a/Engine/Source/Editor/AudioEditor/Private/AssetTypeActions/AssetTypeActions_SoundClass.h
+++ b/Engine/Source/Editor/AudioEditor/Private/AssetTypeActions/AssetTypeActions_SoundClass.h
@@ -5,6 +5,7 @@
 #include "CoreMinimal.h"
 #include "Toolkits/IToolkitHost.h"
 #include "AssetTypeActions_Base.h"
+#include "Sound/SoundClass.h"
 
 class FAssetTypeActions_SoundClass : public FAssetTypeActions_Base
 {
diff --git a/Engine/Source/Runtime/Engine/Public/Audio.h b/Engine/Source/Runtime/Engine/Public/Audio.h
--- a/Engine/Source/Runtime/Engine/Public/Audio.h
+++ b/Engine/Source/Runtime/Engine/Public/Audio.h
@@ -778,7 +778,23 @@ protected:
 
 #if ENABLE_AUDIO_DEBUG
 public:
-	struct FDebugInfo;
+
+	/** Struct containing the debug state of a SoundSource */
+	struct ENGINE_API FDebugInfo
+	{
+		/** True if this sound has been soloed. */
+		bool bIsSoloed = false;
+
+		/** True if this sound has been muted . */
+		bool bIsMuted = false;
+
+		/** Reason why this sound is mute/soloed. */
+		FString MuteSoloReason;
+
+		/** Basic CS so we can pass this around safely. */
+		FCriticalSection CS;
+	};
+	
 	TSharedPtr<FDebugInfo, ESPMode::ThreadSafe> DebugInfo;
 	friend struct FDebugInfo;
 #endif //ENABLE_AUDIO_DEBUG
diff --git a/Engine/Source/Runtime/Engine/Public/Audio/AudioDebug.h b/Engine/Source/Runtime/Engine/Public/Audio/AudioDebug.h
--- a/Engine/Source/Runtime/Engine/Public/Audio/AudioDebug.h
+++ b/Engine/Source/Runtime/Engine/Public/Audio/AudioDebug.h
@@ -5,6 +5,7 @@
 #include "CoreMinimal.h"
 
 #include "AudioDefines.h"		// For ENABLE_AUDIO_DEBUG
+#include "Audio.h"
 
 #if ENABLE_AUDIO_DEBUG
 
@@ -15,22 +16,6 @@ struct FWaveInstance;
 
 class FSoundSource;
 
-/** Struct containing the debug state of a SoundSource */
-struct ENGINE_API FSoundSource::FDebugInfo
-{
-	/** True if this sound has been soloed. */
-	bool bIsSoloed = false;
-
-	/** True if this sound has been muted . */
-	bool bIsMuted = false;
-
-	/** Reason why this sound is mute/soloed. */
-	FString MuteSoloReason;
-
-	/** Basic CS so we can pass this around safely. */
-	FCriticalSection CS;
-};
-
 class ENGINE_API FAudioDebugger
 {
 public:
