Made UAssetManager run the localization cache update after a late delivery asset registry is merged in. Stops log spam from async loading accessing the cache while it is dirty.

diff --git a/Engine/Source/Runtime/Engine/Private/AssetManager.cpp b/Engine/Source/Runtime/Engine/Private/AssetManager.cpp
--- a/Engine/Source/Runtime/Engine/Private/AssetManager.cpp
+++ b/Engine/Source/Runtime/Engine/Private/AssetManager.cpp
@@ -17,6 +17,7 @@
 #include "HAL/PlatformFilemanager.h"
 #include "IPlatformFilePak.h"
 #include "Stats/StatsMisc.h"
+#include "Internationalization/PackageLocalizationManager.h"
 
 #if WITH_EDITOR
 #include "Editor.h"
@@ -1938,6 +1939,7 @@ bool UAssetManager::OnAssetRegistryAvailableAfterInitialization(FName InName, FA
 		if (bLoaded)
 		{
 			LocalAssetRegistry.AppendState(OutNewState);
+			FPackageLocalizationManager::Get().ConditionalUpdateCache();
 
 			TArray<FAssetData> NewAssetData;
 			bool bRebuildReferenceList = false;
