FEnginePackageLocalizationCache will no longer be marked as dirty after handling an add/remove/rename event if that event didn't actually modify the cache contents

diff --git a/Engine/Source/Runtime/CoreUObject/Private/Internationalization/PackageLocalizationCache.cpp b/Engine/Source/Runtime/CoreUObject/Private/Internationalization/PackageLocalizationCache.cpp
--- a/Engine/Source/Runtime/CoreUObject/Private/Internationalization/PackageLocalizationCache.cpp
+++ b/Engine/Source/Runtime/CoreUObject/Private/Internationalization/PackageLocalizationCache.cpp
@@ -92,11 +92,11 @@ void FPackageLocalizationCultureCache::RemoveRootSourcePath(const FString& InRoo
 	}
 }
 
-void FPackageLocalizationCultureCache::AddPackage(const FString& InPackageName)
+bool FPackageLocalizationCultureCache::AddPackage(const FString& InPackageName)
 {
 	if (!FPackageName::IsLocalizedPackage(InPackageName))
 	{
-		return;
+		return false;
 	}
 
 	FScopeLock Lock(&LocalizedPackagesCS);
@@ -112,13 +112,15 @@ void FPackageLocalizationCultureCache::AddPackage(const FString& InPackageName)
 
 				TArray<FName>& PrioritizedLocalizedPackageNames = SourcePackagesToLocalizedPackages.FindOrAdd(SourcePackageName);
 				PrioritizedLocalizedPackageNames.AddUnique(*InPackageName);
-				return;
+				return true;
 			}
 		}
 	}
+
+	return false;
 }
 
-void FPackageLocalizationCultureCache::RemovePackage(const FString& InPackageName)
+bool FPackageLocalizationCultureCache::RemovePackage(const FString& InPackageName)
 {
 	FScopeLock Lock(&LocalizedPackagesCS);
 
@@ -137,14 +139,16 @@ void FPackageLocalizationCultureCache::RemovePackage(const FString& InPackageNam
 				{
 					SourcePackagesToLocalizedPackages.Remove(SourcePackageToLocalizedPackagesPair.Key);
 				}
-				return;
+				return true;
 			}
 		}
 	}
 	else
 	{
-		SourcePackagesToLocalizedPackages.Remove(*InPackageName);
+		return SourcePackagesToLocalizedPackages.Remove(*InPackageName) > 0;
 	}
+
+	return false;
 }
 
 void FPackageLocalizationCultureCache::Empty()
diff --git a/Engine/Source/Runtime/CoreUObject/Public/Internationalization/PackageLocalizationCache.h b/Engine/Source/Runtime/CoreUObject/Public/Internationalization/PackageLocalizationCache.h
--- a/Engine/Source/Runtime/CoreUObject/Public/Internationalization/PackageLocalizationCache.h
+++ b/Engine/Source/Runtime/CoreUObject/Public/Internationalization/PackageLocalizationCache.h
@@ -42,15 +42,17 @@ public:
 	 * Add a package (potentially source or localized) to this cache.
 	 *
 	 * @param InPackageName		The name of the package to add.
+	 * @return TRUE if the supplied package is a localized asset and was added to the culture cache, FALSE otherwise
 	 */
-	void AddPackage(const FString& InPackageName);
+	bool AddPackage(const FString& InPackageName);
 
 	/**
 	 * Remove a package (potentially source or localized) from this cache.
 	 *
 	 * @param InPackageName		The name of the package to remove.
+	 * @return TRUE if the supplied package is a localized asset and was removed from the culture cache, FALSE otherwise
 	 */
-	void RemovePackage(const FString& InPackageName);
+	bool RemovePackage(const FString& InPackageName);
 
 	/**
 	 * Restore this cache to an empty state.
diff --git a/Engine/Source/Runtime/Engine/Private/Internationalization/EnginePackageLocalizationCache.cpp b/Engine/Source/Runtime/Engine/Private/Internationalization/EnginePackageLocalizationCache.cpp
--- a/Engine/Source/Runtime/Engine/Private/Internationalization/EnginePackageLocalizationCache.cpp
+++ b/Engine/Source/Runtime/Engine/Private/Internationalization/EnginePackageLocalizationCache.cpp
@@ -117,10 +117,8 @@ void FEnginePackageLocalizationCache::HandleAssetAdded(const FAssetData& InAsset
 
 	for (auto& CultureCachePair : AllCultureCaches)
 	{
-		CultureCachePair.Value->AddPackage(InAssetData.PackageName.ToString());
+		bPackageNameToAssetGroupDirty |= CultureCachePair.Value->AddPackage(InAssetData.PackageName.ToString());
 	}
-
-	bPackageNameToAssetGroupDirty = true;
 }
 
 void FEnginePackageLocalizationCache::HandleAssetRemoved(const FAssetData& InAssetData)
@@ -129,10 +127,8 @@ void FEnginePackageLocalizationCache::HandleAssetRemoved(const FAssetData& InAss
 
 	for (auto& CultureCachePair : AllCultureCaches)
 	{
-		CultureCachePair.Value->RemovePackage(InAssetData.PackageName.ToString());
+		bPackageNameToAssetGroupDirty |= CultureCachePair.Value->RemovePackage(InAssetData.PackageName.ToString());
 	}
-
-	bPackageNameToAssetGroupDirty = true;
 }
 
 void FEnginePackageLocalizationCache::HandleAssetRenamed(const FAssetData& InAssetData, const FString& InOldObjectPath)
@@ -143,9 +139,7 @@ void FEnginePackageLocalizationCache::HandleAssetRenamed(const FAssetData& InAss
 
 	for (auto& CultureCachePair : AllCultureCaches)
 	{
-		CultureCachePair.Value->RemovePackage(OldPackagePath);
-		CultureCachePair.Value->AddPackage(InAssetData.PackageName.ToString());
+		bPackageNameToAssetGroupDirty |= CultureCachePair.Value->RemovePackage(OldPackagePath);
+		bPackageNameToAssetGroupDirty |= CultureCachePair.Value->AddPackage(InAssetData.PackageName.ToString());
 	}
-
-	bPackageNameToAssetGroupDirty = true;
 }
