#ue4 - Don't allow a character move to combine when the movement mode changes over the course of that move. Reduces latency of sending that move (by using bForceNoCombine) and prevents some issues when combining incompatible movement modes.

diff --git a/Engine/Source/Runtime/Engine/Private/Components/CharacterMovementComponent.cpp b/Engine/Source/Runtime/Engine/Private/Components/CharacterMovementComponent.cpp
--- a/Engine/Source/Runtime/Engine/Private/Components/CharacterMovementComponent.cpp
+++ b/Engine/Source/Runtime/Engine/Private/Components/CharacterMovementComponent.cpp
@@ -10677,6 +10677,12 @@ void FSavedMove_Character::PostUpdate(ACharacter* Character, FSavedMove_Characte
 		{
 			bForceNoCombine = true;
 		}
+
+		// Don't combine if this move caused us to change movement modes during the move.
+		if (StartPackedMovementMode != EndPackedMovementMode)
+		{
+			bForceNoCombine = true;
+		}
 	}
 	else if (PostUpdateMode == PostUpdate_Replay)
 	{
@@ -10907,6 +10913,11 @@ bool FSavedMove_Character::CanCombineWith(const FSavedMovePtr& NewMovePtr, AChar
 		}
 	}
 
+	if (StartPackedMovementMode != NewMove->StartPackedMovementMode)
+	{
+		return false;
+	}
+
 	if (EndPackedMovementMode != NewMove->StartPackedMovementMode)
 	{
 		return false;
