Put virtual texture upload staging buffer in cpu accessible memory

diff --git a/Engine/Source/Runtime/Engine/Private/VT/VirtualTextureUploadCache.cpp b/Engine/Source/Runtime/Engine/Private/VT/VirtualTextureUploadCache.cpp
--- a/Engine/Source/Runtime/Engine/Private/VT/VirtualTextureUploadCache.cpp
+++ b/Engine/Source/Runtime/Engine/Private/VT/VirtualTextureUploadCache.cpp
@@ -209,7 +209,7 @@ FVTUploadTileHandle FVirtualTextureUploadCache::PrepareTileForUpload(FVTUploadTi
 		if (ALLOW_COPY_FROM_BUFFER)
 		{
 			FRHIResourceCreateInfo CreateInfo;
-			NewEntry.RHIStagingBuffer = RHICreateStructuredBuffer(FormatInfo.BlockBytes, MemorySize, BUF_ShaderResource | BUF_Static, CreateInfo);
+			NewEntry.RHIStagingBuffer = RHICreateStructuredBuffer(FormatInfo.BlockBytes, MemorySize, BUF_ShaderResource | BUF_Static | BUF_KeepCPUAccessible, CreateInfo);
 
 			// Here we bypass 'normal' RHI operations in order to get a persistent pointer to GPU memory, on supported platforms
 			// This should be encapsulated into a proper RHI method at some point
