shelf from sebastien
mandatory for switch to not try to sample local lights

diff --git a/Engine/Shaders/Private/ForwardLightingCommon.ush b/Engine/Shaders/Private/ForwardLightingCommon.ush
--- a/Engine/Shaders/Private/ForwardLightingCommon.ush
+++ b/Engine/Shaders/Private/ForwardLightingCommon.ush
@@ -164,6 +164,8 @@ float3 GetForwardDirectLighting(uint GridIndex, float3 WorldPosition, float3 Cam
 		}
 	}
 
+#if !FORWARD_SIMPLE_SINGLE_LAYER_WATER
+
 	const FCulledLightsGridData CulledLightsGrid = GetCulledLightsGrid(GridIndex, EyeIndex);
 
 	
@@ -214,6 +216,8 @@ float3 GetForwardDirectLighting(uint GridIndex, float3 WorldPosition, float3 Cam
 		}
 	}
 
+#endif
+
 	// For debugging
 	//DirectLighting = CulledLightsGrid.NumLocalLights / (float)ForwardLightData.MaxCulledLightsPerCell;
 	return DirectLighting;
diff --git a/Engine/Source/Runtime/Engine/Classes/Materials/Material.h b/Engine/Source/Runtime/Engine/Classes/Materials/Material.h
--- a/Engine/Source/Runtime/Engine/Classes/Materials/Material.h
+++ b/Engine/Source/Runtime/Engine/Classes/Materials/Material.h
@@ -1796,6 +1796,12 @@ public:
 	bool HasFlippedCoordinates();
 #endif //WITH_EDITORONLY_DATA
 
+	// This is to have switch use the simple single layer water shading similar to mobile: no dynamic lights, only sun and sky, no distortion, no colored transmittance on background, no custom depth read.
+	static bool SingleLayerWaterForcesForwardSimpleShading(EShaderPlatform ShaderPlatform) 
+	{
+		return IsSwitchPlatform(ShaderPlatform) && IsForwardShadingEnabled(ShaderPlatform);
+	}
+
 private:
 #if WITH_EDITOR
 	static FMaterialCompilationFinished MaterialCompilationFinishedEvent;
diff --git a/Engine/Source/Runtime/Engine/Private/Materials/HLSLMaterialTranslator.h b/Engine/Source/Runtime/Engine/Private/Materials/HLSLMaterialTranslator.h
--- a/Engine/Source/Runtime/Engine/Private/Materials/HLSLMaterialTranslator.h
+++ b/Engine/Source/Runtime/Engine/Private/Materials/HLSLMaterialTranslator.h
@@ -1523,6 +1523,8 @@ ResourcesString = TEXT("");
 				NumSetMaterials++;
 			}
 
+			OutEnvironment.SetDefine(TEXT("FORWARD_SIMPLE_SINGLE_LAYER_WATER"), UMaterial::SingleLayerWaterForcesForwardSimpleShading(InPlatform) ? TEXT("1") : TEXT("0"));
+
 			if (NumSetMaterials == 1)
 			{
 				OutEnvironment.SetDefine(TEXT("MATERIAL_SINGLE_SHADINGMODEL"), TEXT("1"));
