shelf from rolando
performance improvement for splines

diff --git a/Engine/Shaders/Private/LocalVertexFactory.ush b/Engine/Shaders/Private/LocalVertexFactory.ush
--- a/Engine/Shaders/Private/LocalVertexFactory.ush
+++ b/Engine/Shaders/Private/LocalVertexFactory.ush
@@ -56,31 +56,24 @@ float4 InstancingFadeOutParams;
 #endif	// USE_INSTANCING
 
 #if USE_SPLINEDEFORM
-float3 SplineStartPos;
-float3 SplineStartTangent;
-float SplineStartRoll;
-float2 SplineStartScale;
-float2 SplineStartOffset;
-
-float3 SplineEndPos;
-float3 SplineEndTangent;
-float SplineEndRoll;
-float2 SplineEndScale;
-float2 SplineEndOffset;
-
-float3 SplineUpDir;
-#if ES2_PROFILE
-float SmoothInterpRollScale;
-#else
-bool SmoothInterpRollScale;
-#endif
-
-float SplineMeshMinZ;
-float SplineMeshScaleZ;
-
-float3 SplineMeshDir;
-float3 SplineMeshX;
-float3 SplineMeshY;
+float4 SplineParams[10];
+#define SplineStartPos			SplineParams[0].xyz
+#define SplineStartRoll			SplineParams[0].w
+#define SplineStartTangent		SplineParams[1].xyz
+#define SplineEndRoll			SplineParams[1].w
+#define SplineStartScale		SplineParams[2].xy
+#define SplineStartOffset		SplineParams[2].zw
+#define SplineEndPos			SplineParams[3].xyz
+#define SmoothInterpRollScale	SplineParams[3].w
+#define SplineEndTangent		SplineParams[4].xyz
+#define SplineMeshMinZ			SplineParams[4].w
+#define SplineEndScale			SplineParams[5].xy
+#define SplineEndOffset			SplineParams[5].zw
+#define SplineUpDir				SplineParams[6].xyz
+#define SplineMeshScaleZ		SplineParams[6].w
+#define SplineMeshDir			SplineParams[7].xyz
+#define SplineMeshX				SplineParams[8].xyz
+#define SplineMeshY				SplineParams[9].xyz
 #endif	// USE_SPLINEDEFORM
 
 #ifndef MANUAL_VERTEX_FETCH
diff --git a/Engine/Source/Runtime/Engine/Private/Components/SplineMeshComponent.cpp b/Engine/Source/Runtime/Engine/Private/Components/SplineMeshComponent.cpp
--- a/Engine/Source/Runtime/Engine/Private/Components/SplineMeshComponent.cpp
+++ b/Engine/Source/Runtime/Engine/Private/Components/SplineMeshComponent.cpp
@@ -34,27 +34,7 @@ static FAutoConsoleVariableRef CVarNoRecreateSplineMeshProxy(
 
 void FSplineMeshVertexFactoryShaderParameters::Bind(const FShaderParameterMap& ParameterMap)
 {
-	SplineStartPosParam.Bind(ParameterMap, TEXT("SplineStartPos"), SPF_Mandatory);
-	SplineStartTangentParam.Bind(ParameterMap, TEXT("SplineStartTangent"), SPF_Mandatory);
-	SplineStartRollParam.Bind(ParameterMap, TEXT("SplineStartRoll"), SPF_Mandatory);
-	SplineStartScaleParam.Bind(ParameterMap, TEXT("SplineStartScale"), SPF_Mandatory);
-	SplineStartOffsetParam.Bind(ParameterMap, TEXT("SplineStartOffset"), SPF_Mandatory);
-
-	SplineEndPosParam.Bind(ParameterMap, TEXT("SplineEndPos"), SPF_Mandatory);
-	SplineEndTangentParam.Bind(ParameterMap, TEXT("SplineEndTangent"), SPF_Mandatory);
-	SplineEndRollParam.Bind(ParameterMap, TEXT("SplineEndRoll"), SPF_Mandatory);
-	SplineEndScaleParam.Bind(ParameterMap, TEXT("SplineEndScale"), SPF_Mandatory);
-	SplineEndOffsetParam.Bind(ParameterMap, TEXT("SplineEndOffset"), SPF_Mandatory);
-
-	SplineUpDirParam.Bind(ParameterMap, TEXT("SplineUpDir"), SPF_Mandatory);
-	SmoothInterpRollScaleParam.Bind(ParameterMap, TEXT("SmoothInterpRollScale"), SPF_Mandatory);
-
-	SplineMeshMinZParam.Bind(ParameterMap, TEXT("SplineMeshMinZ"), SPF_Mandatory);
-	SplineMeshScaleZParam.Bind(ParameterMap, TEXT("SplineMeshScaleZ"), SPF_Mandatory);
-
-	SplineMeshDirParam.Bind(ParameterMap, TEXT("SplineMeshDir"), SPF_Mandatory);
-	SplineMeshXParam.Bind(ParameterMap, TEXT("SplineMeshX"), SPF_Mandatory);
-	SplineMeshYParam.Bind(ParameterMap, TEXT("SplineMeshY"), SPF_Mandatory);
+	SplineMeshParams.Bind(ParameterMap, TEXT("SplineParams"), SPF_Mandatory);
 }
 
 void FSplineMeshVertexFactoryShaderParameters::GetElementShaderBindings(
@@ -85,34 +65,27 @@ void FSplineMeshVertexFactoryShaderParameters::GetElementShaderBindings(
 	FSplineMeshSceneProxy* SplineProxy = BatchElement.SplineMeshSceneProxy;
 	FSplineMeshParams& SplineParams = SplineProxy->SplineParams;
 
-	ShaderBindings.Add(SplineStartPosParam, SplineParams.StartPos);
-	ShaderBindings.Add(SplineStartTangentParam, SplineParams.StartTangent);
-	ShaderBindings.Add(SplineStartRollParam, SplineParams.StartRoll);
-	ShaderBindings.Add(SplineStartScaleParam, SplineParams.StartScale);
-	ShaderBindings.Add(SplineStartOffsetParam, SplineParams.StartOffset);
-
-	ShaderBindings.Add(SplineEndPosParam, SplineParams.EndPos);
-	ShaderBindings.Add(SplineEndTangentParam, SplineParams.EndTangent);
-	ShaderBindings.Add(SplineEndRollParam, SplineParams.EndRoll);
-	ShaderBindings.Add(SplineEndScaleParam, SplineParams.EndScale);
-	ShaderBindings.Add(SplineEndOffsetParam, SplineParams.EndOffset);
-
-	ShaderBindings.Add(SplineUpDirParam, SplineProxy->SplineUpDir);
-	ShaderBindings.Add(SmoothInterpRollScaleParam, (int32)SplineProxy->bSmoothInterpRollScale);
-
-	ShaderBindings.Add(SplineMeshMinZParam, SplineProxy->SplineMeshMinZ);
-	ShaderBindings.Add(SplineMeshScaleZParam, SplineProxy->SplineMeshScaleZ);
+	FVector4 ParamData[10];
+	ParamData[0] = FVector4(SplineParams.StartPos, SplineParams.StartRoll);
+	ParamData[1] = FVector4(SplineParams.StartTangent, SplineParams.EndRoll);
+	ParamData[2] = FVector4(SplineParams.StartScale, SplineParams.StartOffset);
+	ParamData[3] = FVector4(SplineParams.EndPos, (float)(int32)SplineProxy->bSmoothInterpRollScale);
+	ParamData[4] = FVector4(SplineParams.EndTangent, SplineProxy->SplineMeshMinZ);
+	ParamData[5] = FVector4(SplineParams.EndScale, SplineParams.EndOffset);
+	ParamData[6] = FVector4(SplineProxy->SplineUpDir, SplineProxy->SplineMeshScaleZ);
 
 	FVector DirMask(0, 0, 0);
 	DirMask = FVector::ZeroVector;
 	DirMask[SplineProxy->ForwardAxis] = 1;
-	ShaderBindings.Add(SplineMeshDirParam, DirMask);
+	ParamData[7] = FVector4(DirMask, 0);
 	DirMask = FVector::ZeroVector;
 	DirMask[(SplineProxy->ForwardAxis + 1) % 3] = 1;
-	ShaderBindings.Add(SplineMeshXParam, DirMask);
+	ParamData[8] = FVector4(DirMask, 0);
 	DirMask = FVector::ZeroVector;
 	DirMask[(SplineProxy->ForwardAxis + 2) % 3] = 1;
-	ShaderBindings.Add(SplineMeshYParam, DirMask);
+	ParamData[9] = FVector4(DirMask, 0);
+
+	ShaderBindings.Add(SplineMeshParams, ParamData);
 }
 
 //////////////////////////////////////////////////////////////////////////
diff --git a/Engine/Source/Runtime/Engine/Private/SplineMeshSceneProxy.h b/Engine/Source/Runtime/Engine/Private/SplineMeshSceneProxy.h
--- a/Engine/Source/Runtime/Engine/Private/SplineMeshSceneProxy.h
+++ b/Engine/Source/Runtime/Engine/Private/SplineMeshSceneProxy.h
@@ -85,27 +85,7 @@ class FSplineMeshVertexFactoryShaderParameters : public FVertexFactoryShaderPara
 
 	void Serialize(FArchive& Ar) override
 	{
-		Ar << SplineStartPosParam;
-		Ar << SplineStartTangentParam;
-		Ar << SplineStartRollParam;
-		Ar << SplineStartScaleParam;
-		Ar << SplineStartOffsetParam;
-
-		Ar << SplineEndPosParam;
-		Ar << SplineEndTangentParam;
-		Ar << SplineEndRollParam;
-		Ar << SplineEndScaleParam;
-		Ar << SplineEndOffsetParam;
-
-		Ar << SplineUpDirParam;
-		Ar << SmoothInterpRollScaleParam;
-
-		Ar << SplineMeshMinZParam;
-		Ar << SplineMeshScaleZParam;
-
-		Ar << SplineMeshDirParam;
-		Ar << SplineMeshXParam;
-		Ar << SplineMeshYParam;
+		Ar << SplineMeshParams;
 	}
 
 	virtual uint32 GetSize() const override
@@ -114,27 +94,7 @@ class FSplineMeshVertexFactoryShaderParameters : public FVertexFactoryShaderPara
 	}
 
 private:
-	FShaderParameter SplineStartPosParam;
-	FShaderParameter SplineStartTangentParam;
-	FShaderParameter SplineStartRollParam;
-	FShaderParameter SplineStartScaleParam;
-	FShaderParameter SplineStartOffsetParam;
-
-	FShaderParameter SplineEndPosParam;
-	FShaderParameter SplineEndTangentParam;
-	FShaderParameter SplineEndRollParam;
-	FShaderParameter SplineEndScaleParam;
-	FShaderParameter SplineEndOffsetParam;
-
-	FShaderParameter SplineUpDirParam;
-	FShaderParameter SmoothInterpRollScaleParam;
-
-	FShaderParameter SplineMeshMinZParam;
-	FShaderParameter SplineMeshScaleZParam;
-
-	FShaderParameter SplineMeshDirParam;
-	FShaderParameter SplineMeshXParam;
-	FShaderParameter SplineMeshYParam;
+	FShaderParameter SplineMeshParams;
 };
 
 //////////////////////////////////////////////////////////////////////////
