Fix compile

diff --git a/Engine/Source/Runtime/Engine/Classes/Materials/Material.h b/Engine/Source/Runtime/Engine/Classes/Materials/Material.h
--- a/Engine/Source/Runtime/Engine/Classes/Materials/Material.h
+++ b/Engine/Source/Runtime/Engine/Classes/Materials/Material.h
@@ -1797,10 +1797,7 @@ public:
 #endif //WITH_EDITORONLY_DATA
 
 	// This is to have switch use the simple single layer water shading similar to mobile: no dynamic lights, only sun and sky, no distortion, no colored transmittance on background, no custom depth read.
-	static bool SingleLayerWaterForcesForwardSimpleShading(EShaderPlatform ShaderPlatform) 
-	{
-		return IsSwitchPlatform(ShaderPlatform) && IsForwardShadingEnabled(ShaderPlatform);
-	}
+	static bool SingleLayerWaterForcesForwardSimpleShading(EShaderPlatform ShaderPlatform);
 
 private:
 #if WITH_EDITOR
diff --git a/Engine/Source/Runtime/Engine/Private/Materials/Material.cpp b/Engine/Source/Runtime/Engine/Private/Materials/Material.cpp
--- a/Engine/Source/Runtime/Engine/Private/Materials/Material.cpp
+++ b/Engine/Source/Runtime/Engine/Private/Materials/Material.cpp
@@ -62,6 +62,7 @@
 #include "Materials/MaterialExpressionComment.h"
 #include "UObject/EditorObjectVersion.h"
 #include "UObject/ReleaseObjectVersion.h"
+#include "RenderCore/Public/RenderUtils.h"
 
 #if WITH_EDITOR
 #include "Logging/TokenizedMessage.h"
@@ -6316,6 +6317,12 @@ bool UMaterial::HasFlippedCoordinates()
 }
 #endif //WITH_EDITORONLY_DATA
 
+// This is to have switch use the simple single layer water shading similar to mobile: no dynamic lights, only sun and sky, no distortion, no colored transmittance on background, no custom depth read.
+bool UMaterial::SingleLayerWaterForcesForwardSimpleShading(EShaderPlatform ShaderPlatform)
+{
+	return IsSwitchPlatform(ShaderPlatform) && IsForwardShadingEnabled(ShaderPlatform);
+}
+
 void UMaterial::GetLightingGuidChain(bool bIncludeTextures, TArray<FGuid>& OutGuids) const
 {
 #if WITH_EDITORONLY_DATA
