Fixed MATERIAL_SHADINGMODEL_SINGLELAYERWATER to affect only FORWARD_SIMPLE_SINGLE_LAYER_WATER
Replaced SingleLayerWaterForcesForwardSimpleShading with a Switch check inside HLSLMaterialTranslator.h

diff --git a/Engine/Shaders/Private/ForwardLightingCommon.ush b/Engine/Shaders/Private/ForwardLightingCommon.ush
--- a/Engine/Shaders/Private/ForwardLightingCommon.ush
+++ b/Engine/Shaders/Private/ForwardLightingCommon.ush
@@ -164,7 +164,7 @@ float3 GetForwardDirectLighting(uint GridIndex, float3 WorldPosition, float3 Cam
 		}
 	}
 
-#if !FORWARD_SIMPLE_SINGLE_LAYER_WATER
+#if !(MATERIAL_SHADINGMODEL_SINGLELAYERWATER && FORWARD_SIMPLE_SINGLE_LAYER_WATER)
 
 	const FCulledLightsGridData CulledLightsGrid = GetCulledLightsGrid(GridIndex, EyeIndex);
 
diff --git a/Engine/Source/Runtime/Engine/Classes/Materials/Material.h b/Engine/Source/Runtime/Engine/Classes/Materials/Material.h
--- a/Engine/Source/Runtime/Engine/Classes/Materials/Material.h
+++ b/Engine/Source/Runtime/Engine/Classes/Materials/Material.h
@@ -1796,9 +1796,6 @@ public:
 	bool HasFlippedCoordinates();
 #endif //WITH_EDITORONLY_DATA
 
-	// This is to have switch use the simple single layer water shading similar to mobile: no dynamic lights, only sun and sky, no distortion, no colored transmittance on background, no custom depth read.
-	static bool SingleLayerWaterForcesForwardSimpleShading(EShaderPlatform ShaderPlatform);
-
 private:
 #if WITH_EDITOR
 	static FMaterialCompilationFinished MaterialCompilationFinishedEvent;
diff --git a/Engine/Source/Runtime/Engine/Private/Materials/HLSLMaterialTranslator.h b/Engine/Source/Runtime/Engine/Private/Materials/HLSLMaterialTranslator.h
--- a/Engine/Source/Runtime/Engine/Private/Materials/HLSLMaterialTranslator.h
+++ b/Engine/Source/Runtime/Engine/Private/Materials/HLSLMaterialTranslator.h
@@ -1523,7 +1523,11 @@ ResourcesString = TEXT("");
 				NumSetMaterials++;
 			}
 
-			OutEnvironment.SetDefine(TEXT("FORWARD_SIMPLE_SINGLE_LAYER_WATER"), UMaterial::SingleLayerWaterForcesForwardSimpleShading(InPlatform) ? TEXT("1") : TEXT("0"));
+			// This is to have switch use the simple single layer water shading similar to mobile: no dynamic lights, only sun and sky, no distortion, no colored transmittance on background, no custom depth read.
+			if (IsSwitchPlatform(Platform))
+			{
+				OutEnvironment.SetDefine(TEXT("FORWARD_SIMPLE_SINGLE_LAYER_WATER"), TEXT("1"));
+			}
 
 			if (NumSetMaterials == 1)
 			{
