CSVProfiler: Skip recording tick counts if TickGroups are paused.  Prevents log warning from GetEnabledTickFunctionCounts

#rnx

diff --git a/Engine/Source/Runtime/Engine/Private/LevelTick.cpp b/Engine/Source/Runtime/Engine/Private/LevelTick.cpp
--- a/Engine/Source/Runtime/Engine/Private/LevelTick.cpp
+++ b/Engine/Source/Runtime/Engine/Private/LevelTick.cpp
@@ -1270,11 +1270,11 @@ extern TMap<FName, int32> CSVActorClassNameToCountMap;
 extern int32 CSVActorTotalCount;
 
 /** Add additional stats to CSV profile for tick and actor counts */
-static void RecordWorldCountsToCSV(UWorld* World)
+static void RecordWorldCountsToCSV(UWorld* World, bool bDoingActorTicks)
 {
 	if(FCsvProfiler::Get()->IsCapturing())
 	{
-		if (CVarRecordTickCountsToCSV.GetValueOnGameThread())
+		if (bDoingActorTicks && CVarRecordTickCountsToCSV.GetValueOnGameThread())
 		{
 			QUICK_SCOPE_CYCLE_COUNTER(STAT_RecordTickCountsToCSV);
 			CSV_SCOPED_TIMING_STAT_EXCLUSIVE(RecordTickCountsToCSV);
@@ -1757,7 +1757,7 @@ void UWorld::Tick( ELevelTick TickType, float DeltaSeconds )
 #endif
 
 #if (CSV_PROFILER && !UE_BUILD_SHIPPING)
-	RecordWorldCountsToCSV(this);
+	RecordWorldCountsToCSV(this, bDoingActorTicks);
 #endif // (CSV_PROFILER && !UE_BUILD_SHIPPING)
 
 #if !(UE_BUILD_SHIPPING || UE_BUILD_TEST)
diff --git a/Engine/Source/Runtime/Engine/Private/TickTaskManager.cpp b/Engine/Source/Runtime/Engine/Private/TickTaskManager.cpp
--- a/Engine/Source/Runtime/Engine/Private/TickTaskManager.cpp
+++ b/Engine/Source/Runtime/Engine/Private/TickTaskManager.cpp
@@ -1709,7 +1709,7 @@ private:
 
 		if (bFilterCoolingDown && InWorld->TickGroup >= 0 && InWorld->TickGroup < TG_NewlySpawned)
 		{
-			UE_LOG(LogTick, Warning, TEXT("GetEnabledTickFunctionCounts: Filtering results mid-frame. TickFunctions that execute later this frame will not be reported."));
+			UE_LOG(LogTick, Warning, TEXT("GetEnabledTickFunctionCounts is filtering results before TickGroups have finished. TickFunctions with a cooldown interval may be missing."));
 		}
 
 		const float WorldTimeSeconds = InWorld->GetTimeSeconds();
