remove now defunct function

diff --git a/Engine/Source/Runtime/Engine/Private/Materials/Material.cpp b/Engine/Source/Runtime/Engine/Private/Materials/Material.cpp
--- a/Engine/Source/Runtime/Engine/Private/Materials/Material.cpp
+++ b/Engine/Source/Runtime/Engine/Private/Materials/Material.cpp
@@ -6317,12 +6317,6 @@ bool UMaterial::HasFlippedCoordinates()
 }
 #endif //WITH_EDITORONLY_DATA
 
-// This is to have switch use the simple single layer water shading similar to mobile: no dynamic lights, only sun and sky, no distortion, no colored transmittance on background, no custom depth read.
-bool UMaterial::SingleLayerWaterForcesForwardSimpleShading(EShaderPlatform ShaderPlatform)
-{
-	return IsSwitchPlatform(ShaderPlatform) && IsForwardShadingEnabled(ShaderPlatform);
-}
-
 void UMaterial::GetLightingGuidChain(bool bIncludeTextures, TArray<FGuid>& OutGuids) const
 {
 #if WITH_EDITORONLY_DATA
