Adjusting assert to catch UTexture recycled by GC.

diff --git a/Engine/Source/Runtime/Engine/Private/Materials/MaterialUniformExpressions.cpp b/Engine/Source/Runtime/Engine/Private/Materials/MaterialUniformExpressions.cpp
--- a/Engine/Source/Runtime/Engine/Private/Materials/MaterialUniformExpressions.cpp
+++ b/Engine/Source/Runtime/Engine/Private/Materials/MaterialUniformExpressions.cpp
@@ -593,6 +593,15 @@ void FUniformExpressionSet::FillUniformBuffer(const FMaterialRenderContext& Mate
 					ensureMsgf(false, TEXT("Texture not valid! UE-23902! Parameter (%s)"), TextureParameter ? *TextureParameter->GetParameterName().ToString() : TEXT("non-parameter"));
 				}
 
+				// Trying to track down a dangling pointer bug.
+				checkf(
+					Value->IsA<UTexture>(),
+					TEXT("Expecting a UTexture! Name(%s), Type(%s), TextureParameter(%s), Expression(%d), Material(%s)"),
+					*Value->GetName(), *Value->GetClass()->GetName(),
+					*(TextureParameter ? TextureParameter->GetParameterName() : FName()).ToString(),
+					ExpressionIndex,
+					*MaterialRenderContext.Material.GetFriendlyName());
+
 				// Do not allow external textures to be applied to normal texture samplers
 				if (Value->GetMaterialType() == MCT_TextureExternal)
 				{
@@ -621,9 +630,6 @@ void FUniformExpressionSet::FillUniformBuffer(const FMaterialRenderContext& Mate
 			// Going to leave the check for now though, to hopefully avoid any unexpected problems
 			if (Value && Value->Resource && Value->TextureReference.TextureReferenceRHI && (Value->GetMaterialType() & ValidTextureTypes) != 0u)
 			{
-				//@todo-rco: Help track down a invalid values
-				checkf(Value->IsA(UTexture::StaticClass()), TEXT("Expecting a UTexture! Value='%s' class='%s'"), *Value->GetName(), *Value->GetClass()->GetName());
-
 				*ResourceTableTexturePtr = Value->TextureReference.TextureReferenceRHI;
 				FSamplerStateRHIRef* SamplerSource = &Value->Resource->SamplerStateRHI;
 
