Adding an event/virtual funtion call to be able to symetrically destroy render thread resources when removing a primitive from the scene.

diff --git a/Engine/Source/Runtime/Engine/Public/PrimitiveSceneProxy.h b/Engine/Source/Runtime/Engine/Public/PrimitiveSceneProxy.h
--- a/Engine/Source/Runtime/Engine/Public/PrimitiveSceneProxy.h
+++ b/Engine/Source/Runtime/Engine/Public/PrimitiveSceneProxy.h
@@ -310,6 +310,13 @@ public:
 	 */
 	virtual void CreateRenderThreadResources() {}
 
+	/**
+	 *	Called when the rendering thread removes the proxy from the scene.
+	 *	This function allows for removing renderer-side resources.
+	 *	Called in the rendering thread.
+	 */
+	virtual void DestroyRenderThreadResources() {}
+
 	/**
 	 * Called by the rendering thread to notify the proxy when a light is no longer
 	 * associated with the proxy, so that it can clean up any cached resources.
diff --git a/Engine/Source/Runtime/Renderer/Private/RendererScene.cpp b/Engine/Source/Runtime/Renderer/Private/RendererScene.cpp
--- a/Engine/Source/Runtime/Renderer/Private/RendererScene.cpp
+++ b/Engine/Source/Runtime/Renderer/Private/RendererScene.cpp
@@ -1433,6 +1433,7 @@ void FScene::RemovePrimitive( UPrimitiveComponent* Primitive )
 		ENQUEUE_RENDER_COMMAND(FRemovePrimitiveCommand)(
 			[Scene, PrimitiveSceneInfo, AttachmentCounter](FRHICommandList&)
 			{
+				PrimitiveSceneInfo->Proxy->DestroyRenderThreadResources();
 				Scene->RemovePrimitiveSceneInfo_RenderThread(PrimitiveSceneInfo);
 				AttachmentCounter->Decrement();
 			});
