Now handling detail mode changes via a CVar sink rather than having each component check on their next tick.
This avoids the situation where certain parameters etc were cleared as part of the reset on a detail mode chagne.

diff --git a/Engine/Source/Runtime/Engine/Classes/Particles/ParticleSystemComponent.h b/Engine/Source/Runtime/Engine/Classes/Particles/ParticleSystemComponent.h
--- a/Engine/Source/Runtime/Engine/Classes/Particles/ParticleSystemComponent.h
+++ b/Engine/Source/Runtime/Engine/Classes/Particles/ParticleSystemComponent.h
@@ -663,11 +663,6 @@ public:
 	/** LOD updating... */
 	float AccumLODDistanceCheckTime;
 
-private:
-	/** Remember the global detail mode when we last checked the emitters. */
-	uint32 LastCheckedDetailMode;
-
-public:
 	/** The view relevance flags for each LODLevel. */
 	TArray<FMaterialRelevance> CachedViewRelevanceFlags;
 
diff --git a/Engine/Source/Runtime/Engine/Private/Particles/ParticleComponents.cpp b/Engine/Source/Runtime/Engine/Private/Particles/ParticleComponents.cpp
--- a/Engine/Source/Runtime/Engine/Private/Particles/ParticleComponents.cpp
+++ b/Engine/Source/Runtime/Engine/Private/Particles/ParticleComponents.cpp
@@ -3377,7 +3377,6 @@ UParticleSystemComponent::UParticleSystemComponent(const FObjectInitializer& Obj
 #if WITH_EDITORONLY_DATA
 	EditorDetailMode = -1;
 #endif // WITH_EDITORONLY_DATA
-	LastCheckedDetailMode = -1;
 	SetCollisionProfileName(UCollisionProfile::NoCollision_ProfileName);
 	SetGenerateOverlapEvents(false);
 
@@ -3600,6 +3599,24 @@ bool UParticleSystemComponent::CanConsiderInvisible()const
 	return false;
 }
 
+void DetailModeSink()
+{
+	int32 DetailMode = GetCachedScalabilityCVars().DetailMode;
+	static int32 CachedDetailMode = DetailMode;
+
+	if (CachedDetailMode != DetailMode)
+	{
+		CachedDetailMode = DetailMode;
+
+		for (TObjectIterator<UParticleSystemComponent> It; It; ++It)
+		{
+			It->ForceReset();
+		}
+	}
+}
+
+static FAutoConsoleVariableSink CVarDetailModeSink(FConsoleCommandDelegate::CreateStatic(&DetailModeSink));
+
 void UParticleSystemComponent::ForceReset()
 {
 #if WITH_EDITOR
@@ -5101,30 +5118,6 @@ void UParticleSystemComponent::TickComponent(float DeltaTime, enum ELevelTick Ti
 		}
 		return;
 	} 
-	// See if DetailMode has changed since the last time we checked
-	else if (bWarmingUp == false && LastCheckedDetailMode != DetailModeCVar)
-	{
-		if (!bRequiresReset && LastCheckedDetailMode!=-1)	// only reset if the component isn't new and detail mode has really changed
-		{
-			check(DetailModeCVar < NUM_DETAILMODE_FLAGS);
-			SCOPE_CYCLE_COUNTER(STAT_UParticleSystemComponent_CheckForReset);
-			for (int32 EmitterIndex = 0; EmitterIndex < EmitterInstances.Num(); ++EmitterIndex)
-			{
-				UParticleEmitter *Emitter = Template->Emitters[EmitterIndex];
-				// [op] this should exist if the emitter hasn't been cooked out; need to reset always in this case,
-				//   not only when the detail mode doesn't match, since we need to turn emitters back on as well
-				//   and can't check for instance, because it may not have been created if detail mode is mismatched
-				if(Emitter)
-				{
-					bRequiresReset = true;
-					break;
-				}
-			}
-		}
-
-		// Save the detail mode we've checked
-		LastCheckedDetailMode = DetailModeCVar;
-	}
 	
 	if (bRequiresReset)
 	{
