#ue4 - Character moves should not combine on the frame where jump input begins or ends.

diff --git a/Engine/Source/Runtime/Engine/Private/Components/CharacterMovementComponent.cpp b/Engine/Source/Runtime/Engine/Private/Components/CharacterMovementComponent.cpp
--- a/Engine/Source/Runtime/Engine/Private/Components/CharacterMovementComponent.cpp
+++ b/Engine/Source/Runtime/Engine/Private/Components/CharacterMovementComponent.cpp
@@ -10683,6 +10683,12 @@ void FSavedMove_Character::PostUpdate(ACharacter* Character, FSavedMove_Characte
 		{
 			bForceNoCombine = true;
 		}
+
+		// Don't combine when jump input just began or ended during the move.
+		if (bPressedJump != CharacterOwner->bPressedJump)
+		{
+			bForceNoCombine = true;
+		}
 	}
 	else if (PostUpdateMode == PostUpdate_Replay)
 	{
