Added hlsl translator error when attempting to use landscape layers weight from vertex shader on mobile

diff --git a/Engine/Source/Runtime/Engine/Private/Materials/HLSLMaterialTranslator.h b/Engine/Source/Runtime/Engine/Private/Materials/HLSLMaterialTranslator.h
--- a/Engine/Source/Runtime/Engine/Private/Materials/HLSLMaterialTranslator.h
+++ b/Engine/Source/Runtime/Engine/Private/Materials/HLSLMaterialTranslator.h
@@ -5366,6 +5366,11 @@ protected:
 
 	virtual int32 StaticTerrainLayerWeight(FName ParameterName,int32 Default) override
 	{
+		if (GetFeatureLevel() <= ERHIFeatureLevel::ES3_1 && ShaderFrequency != SF_Pixel)
+		{
+			return Errorf(TEXT("Landscape layer weights are only available in the pixel shader."));
+		}
+		
 		// Look up the weight-map index for this static parameter.
 		int32 WeightmapIndex = INDEX_NONE;
 		bool bFoundParameter = false;
