Fix bit depth used when compiling G16 textures
Fixes potential read out of bounds memory and image corruption issues

diff --git a/Engine/Source/Runtime/Engine/Private/Texture.cpp b/Engine/Source/Runtime/Engine/Private/Texture.cpp
--- a/Engine/Source/Runtime/Engine/Private/Texture.cpp
+++ b/Engine/Source/Runtime/Engine/Private/Texture.cpp
@@ -815,7 +815,7 @@ void FTextureSource::Compress()
 		TSharedPtr<IImageWrapper> ImageWrapper = ImageWrapperModule.CreateImageWrapper( EImageFormat::PNG );
 		// TODO: TSF_BGRA8 is stored as RGBA, so the R and B channels are swapped in the internal png. Should we fix this?
 		ERGBFormat RawFormat = (Format == TSF_G8 || Format == TSF_G16) ? ERGBFormat::Gray : ERGBFormat::RGBA;
-		if ( ImageWrapper.IsValid() && ImageWrapper->SetRaw( BulkDataPtr, BulkData.GetBulkDataSize(), SizeX, SizeY, RawFormat, Format == TSF_G16 || Format == TSF_RGBA16 ? 16 : 8 ) )
+		if ( ImageWrapper.IsValid() && ImageWrapper->SetRaw( BulkDataPtr, BulkData.GetBulkDataSize(), SizeX, SizeY, RawFormat, (Format == TSF_G16 || Format == TSF_RGBA16) ? 16 : 8 ) )
 		{
 			const TArray<uint8>& CompressedData = ImageWrapper->GetCompressed();
 			if ( CompressedData.Num() > 0 )
@@ -862,7 +862,7 @@ uint8* FTextureSource::LockMip(int32 BlockIndex, int32 LayerIndex, int32 MipInde
 					const TArray<uint8>* RawData = NULL;
 					// TODO: TSF_BGRA8 is stored as RGBA, so the R and B channels are swapped in the internal png. Should we fix this?
 					ERGBFormat RawFormat = (Format == TSF_G8 || Format == TSF_G16) ? ERGBFormat::Gray : ERGBFormat::RGBA;
-					if (ImageWrapper->GetRaw( RawFormat, Format == TSF_RGBA16 ? 16 : 8, RawData ))
+					if (ImageWrapper->GetRaw( RawFormat, (Format == TSF_G16 || Format == TSF_RGBA16) ? 16 : 8, RawData ))
 					{
 						if (RawData->Num() > 0)
 						{
@@ -943,7 +943,7 @@ bool FTextureSource::GetMipData(TArray64<uint8>& OutMipData, int32 BlockIndex, i
 						const TArray<uint8>* RawData = NULL;
 						// TODO: TSF_BGRA8 is stored as RGBA, so the R and B channels are swapped in the internal png. Should we fix this?
 						ERGBFormat RawFormat = (Format == TSF_G8 || Format == TSF_G16) ? ERGBFormat::Gray : ERGBFormat::RGBA;
-						if (ImageWrapper->GetRaw( RawFormat, Format == TSF_RGBA16 ? 16 : 8, RawData ))
+						if (ImageWrapper->GetRaw( RawFormat, (Format == TSF_G16 || Format == TSF_RGBA16) ? 16 : 8, RawData ))
 						{
 							OutMipData = *RawData;
 							bSuccess = true;
