Submitting shelf CL 9199061: Disable local lights sampling on all consoles for the sake of performance.

diff --git a/Engine/Shaders/Private/Definitions.usf b/Engine/Shaders/Private/Definitions.usf
--- a/Engine/Shaders/Private/Definitions.usf
+++ b/Engine/Shaders/Private/Definitions.usf
@@ -154,3 +154,7 @@
 #ifndef USING_VERTEX_SHADER_LAYER
 #define USING_VERTEX_SHADER_LAYER	0
 #endif
+
+#ifndef DISABLE_FORWARD_LOCAL_LIGHTS
+#define DISABLE_FORWARD_LOCAL_LIGHTS					0
+#endif
diff --git a/Engine/Shaders/Private/ForwardLightingCommon.ush b/Engine/Shaders/Private/ForwardLightingCommon.ush
--- a/Engine/Shaders/Private/ForwardLightingCommon.ush
+++ b/Engine/Shaders/Private/ForwardLightingCommon.ush
@@ -164,7 +164,7 @@ float3 GetForwardDirectLighting(uint GridIndex, float3 WorldPosition, float3 Cam
 		}
 	}
 
-#if !(MATERIAL_SHADINGMODEL_SINGLELAYERWATER && FORWARD_SIMPLE_SINGLE_LAYER_WATER)
+#if !DISABLE_FORWARD_LOCAL_LIGHTS
 
 	const FCulledLightsGridData CulledLightsGrid = GetCulledLightsGrid(GridIndex, EyeIndex);
 
diff --git a/Engine/Source/Runtime/Engine/Private/Materials/HLSLMaterialTranslator.h b/Engine/Source/Runtime/Engine/Private/Materials/HLSLMaterialTranslator.h
--- a/Engine/Source/Runtime/Engine/Private/Materials/HLSLMaterialTranslator.h
+++ b/Engine/Source/Runtime/Engine/Private/Materials/HLSLMaterialTranslator.h
@@ -1523,10 +1523,9 @@ ResourcesString = TEXT("");
 				NumSetMaterials++;
 			}
 
-			// This is to have switch use the simple single layer water shading similar to mobile: no dynamic lights, only sun and sky, no distortion, no colored transmittance on background, no custom depth read.
-			if (IsSwitchPlatform(Platform))
+			if(ShadingModels.HasShadingModel(MSM_SingleLayerWater) && (IsSwitchPlatform(Platform) || IsPS4Platform(Platform) || Platform == SP_XBOXONE_D3D12))
 			{
-				OutEnvironment.SetDefine(TEXT("FORWARD_SIMPLE_SINGLE_LAYER_WATER"), TEXT("1"));
+				OutEnvironment.SetDefine(TEXT("DISABLE_FORWARD_LOCAL_LIGHTS"), TEXT("1"));
 			}
 
 			if (NumSetMaterials == 1)
