Fixed Subsurface Profile shading model in cooked builds when using FromMaterialExpression.

#rnx

diff --git a/Engine/Source/Runtime/Engine/Private/Materials/Material.cpp b/Engine/Source/Runtime/Engine/Private/Materials/Material.cpp
--- a/Engine/Source/Runtime/Engine/Private/Materials/Material.cpp
+++ b/Engine/Source/Runtime/Engine/Private/Materials/Material.cpp
@@ -3070,8 +3070,10 @@ void UMaterial::UpdateMaterialShaderCacheAndTextureReferences()
 
 void UMaterial::CacheResourceShadersForRendering(bool bRegenerateId)
 {
+#if WITH_EDITOR
 	// Always rebuild the shading model field on recompile
 	RebuildShadingModelField();
+#endif //WITH_EDITOR
 
 	if (bRegenerateId)
 	{
@@ -4640,7 +4642,6 @@ void UMaterial::UpdateExpressionParameterName(UMaterialExpression* Expression)
 		}
 	}
 }
-#endif // WITH_EDITOR
 
 void UMaterial::RebuildShadingModelField()
 {
@@ -4681,6 +4682,7 @@ void UMaterial::RebuildShadingModelField()
 	UsedShadingModels = GetShadingModelFieldString(ShadingModels, FShadingModelToStringDelegate::CreateLambda(ShadingModelToStringLambda), " | ");
 #endif
 }
+#endif // WITH_EDITOR
 
 bool UMaterial::GetExpressionParameterName(const UMaterialExpression* Expression, FName& OutName)
 {
