Added a method to add table on composite object, used when creating a composite tables at runtime.
Be aware this method is slow, can cause hitches at during gameplay.

#rnx

diff --git a/Engine/Source/Runtime/Engine/Classes/Engine/CompositeCurveTable.h b/Engine/Source/Runtime/Engine/Classes/Engine/CompositeCurveTable.h
--- a/Engine/Source/Runtime/Engine/Classes/Engine/CompositeCurveTable.h
+++ b/Engine/Source/Runtime/Engine/Classes/Engine/CompositeCurveTable.h
@@ -33,6 +33,10 @@ class UCompositeCurveTable
 
 	ENGINE_API virtual void EmptyTable() override;
 
+	// Support for runtime modification of parent tables, 
+	// Be aware this can be slow and can cause hitches during gameplay
+	ENGINE_API void AppendParentTables(const TArray<UCurveTable*>& NewTables);
+
 protected:
 
 	// Searches the parent tables to see if there are any loops.
diff --git a/Engine/Source/Runtime/Engine/Classes/Engine/CompositeDataTable.h b/Engine/Source/Runtime/Engine/Classes/Engine/CompositeDataTable.h
--- a/Engine/Source/Runtime/Engine/Classes/Engine/CompositeDataTable.h
+++ b/Engine/Source/Runtime/Engine/Classes/Engine/CompositeDataTable.h
@@ -57,6 +57,10 @@ class UCompositeDataTable
 	virtual void PostEditChangeProperty(FPropertyChangedEvent& PropertyChangedEvent) override;
 #endif // WITH_EDITOR
 
+	// Support for runtime modification of parent tables
+	// Be aware this can be slow and can cause hitches during gameplay
+	ENGINE_API void AppendParentTables(const TArray<UDataTable*>& NewTables);
+
 protected:
 
 	// Searches the parent tables to see if there are any loops.
diff --git a/Engine/Source/Runtime/Engine/Private/CompositeCurveTable.cpp b/Engine/Source/Runtime/Engine/Private/CompositeCurveTable.cpp
--- a/Engine/Source/Runtime/Engine/Private/CompositeCurveTable.cpp
+++ b/Engine/Source/Runtime/Engine/Private/CompositeCurveTable.cpp
@@ -230,6 +230,12 @@ void UCompositeCurveTable::EmptyTable()
 	Super::EmptyTable();
 }
 
+void UCompositeCurveTable::AppendParentTables(const TArray<UCurveTable*>& NewTables)
+{
+	ParentTables.Append(NewTables);
+	OnParentTablesUpdated();
+}
+
 const UCompositeCurveTable* UCompositeCurveTable::FindLoops(TArray<const UCompositeCurveTable*> AlreadySeenTables) const
 {
 	AlreadySeenTables.Add(this);
diff --git a/Engine/Source/Runtime/Engine/Private/CompositeDataTable.cpp b/Engine/Source/Runtime/Engine/Private/CompositeDataTable.cpp
--- a/Engine/Source/Runtime/Engine/Private/CompositeDataTable.cpp
+++ b/Engine/Source/Runtime/Engine/Private/CompositeDataTable.cpp
@@ -258,6 +258,12 @@ void UCompositeDataTable::PostEditChangeProperty(FPropertyChangedEvent& Property
 }
 #endif // WITH_EDITOR
 
+void UCompositeDataTable::AppendParentTables(const TArray<UDataTable*>& NewTables)
+{
+	ParentTables.Append(NewTables);
+	OnParentTablesUpdated(EPropertyChangeType::ValueSet);
+}
+
 void UCompositeDataTable::OnParentTablesUpdated(EPropertyChangeType::Type ChangeType)
 {
 	for (UDataTable* Table : OldParentTables)
