Make ClearRenderTarget2D use fast clear if possible.

diff --git a/Engine/Source/Runtime/Engine/Private/KismetRenderingLibrary.cpp b/Engine/Source/Runtime/Engine/Private/KismetRenderingLibrary.cpp
--- a/Engine/Source/Runtime/Engine/Private/KismetRenderingLibrary.cpp
+++ b/Engine/Source/Runtime/Engine/Private/KismetRenderingLibrary.cpp
@@ -51,10 +51,17 @@ void UKismetRenderingLibrary::ClearRenderTarget2D(UObject* WorldContextObject, U
 		ENQUEUE_RENDER_COMMAND(ClearRTCommand)(
 			[RenderTargetResource, ClearColor](FRHICommandList& RHICmdList)
 			{
-				FRHIRenderPassInfo RPInfo(RenderTargetResource->GetRenderTargetTexture(), ERenderTargetActions::DontLoad_Store);
+				// Use hardware fast clear if the clear colors match.
+				const FTexture2DRHIRef& Texture = RenderTargetResource->GetRenderTargetTexture();
+				const bool bFastClear = Texture->GetClearColor() == ClearColor;
+
+				FRHIRenderPassInfo RPInfo(RenderTargetResource->GetRenderTargetTexture(), bFastClear ? ERenderTargetActions::Clear_Store : ERenderTargetActions::DontLoad_Store);
 				TransitionRenderPassTargets(RHICmdList, RPInfo);
 				RHICmdList.BeginRenderPass(RPInfo, TEXT("ClearRT"));
-				DrawClearQuad(RHICmdList, ClearColor);
+				if (!bFastClear)
+				{
+					DrawClearQuad(RHICmdList, ClearColor);
+				}
 				RHICmdList.EndRenderPass();
 			});
 	}
