Remove RHI thread stall from DrawTileMesh.

diff --git a/Engine/Source/Runtime/Engine/Private/PrimitiveUniformShaderParameters.cpp b/Engine/Source/Runtime/Engine/Private/PrimitiveUniformShaderParameters.cpp
--- a/Engine/Source/Runtime/Engine/Private/PrimitiveUniformShaderParameters.cpp
+++ b/Engine/Source/Runtime/Engine/Private/PrimitiveUniformShaderParameters.cpp
@@ -10,32 +10,33 @@ void FSinglePrimitiveStructuredBuffer::InitRHI()
 	{
 		FRHIResourceCreateInfo CreateInfo;
 		PrimitiveSceneDataBufferRHI = RHICreateStructuredBuffer(sizeof(FVector4), FPrimitiveSceneShaderData::PrimitiveDataStrideInFloat4s * sizeof(FVector4), BUF_Static | BUF_ShaderResource, CreateInfo);
-
-		void* LockedData = RHILockStructuredBuffer(PrimitiveSceneDataBufferRHI, 0, FPrimitiveSceneShaderData::PrimitiveDataStrideInFloat4s * sizeof(FVector4), RLM_WriteOnly);
-
-		FPlatformMemory::Memcpy(LockedData, PrimitiveSceneData.Data, FPrimitiveSceneShaderData::PrimitiveDataStrideInFloat4s * sizeof(FVector4));
-
-		RHIUnlockStructuredBuffer(PrimitiveSceneDataBufferRHI);
-
 		PrimitiveSceneDataBufferSRV = RHICreateShaderResourceView(PrimitiveSceneDataBufferRHI);
+
+		LightmapSceneDataBufferRHI = RHICreateStructuredBuffer(sizeof(FVector4), FLightmapSceneShaderData::LightmapDataStrideInFloat4s * sizeof(FVector4), BUF_Static | BUF_ShaderResource, CreateInfo);
+		LightmapSceneDataBufferSRV = RHICreateShaderResourceView(LightmapSceneDataBufferRHI);
 	}
 
+	UploadToGPU();
+}
+
+void FSinglePrimitiveStructuredBuffer::UploadToGPU()
+{
 	if (RHISupportsComputeShaders(GMaxRHIShaderPlatform))
 	{
-		FRHIResourceCreateInfo CreateInfo;
-		LightmapSceneDataBufferRHI = RHICreateStructuredBuffer(sizeof(FVector4), FLightmapSceneShaderData::LightmapDataStrideInFloat4s * sizeof(FVector4), BUF_Static | BUF_ShaderResource, CreateInfo);
+		void* LockedData = nullptr;
 
-		void* LockedData = RHILockStructuredBuffer(LightmapSceneDataBufferRHI, 0, FLightmapSceneShaderData::LightmapDataStrideInFloat4s * sizeof(FVector4), RLM_WriteOnly);
+		LockedData = RHILockStructuredBuffer(PrimitiveSceneDataBufferRHI, 0, FPrimitiveSceneShaderData::PrimitiveDataStrideInFloat4s * sizeof(FVector4), RLM_WriteOnly);
+		FPlatformMemory::Memcpy(LockedData, PrimitiveSceneData.Data, FPrimitiveSceneShaderData::PrimitiveDataStrideInFloat4s * sizeof(FVector4));
+		RHIUnlockStructuredBuffer(PrimitiveSceneDataBufferRHI);
 
+		LockedData = RHILockStructuredBuffer(LightmapSceneDataBufferRHI, 0, FLightmapSceneShaderData::LightmapDataStrideInFloat4s * sizeof(FVector4), RLM_WriteOnly);
 		FPlatformMemory::Memcpy(LockedData, LightmapSceneData.Data, FLightmapSceneShaderData::LightmapDataStrideInFloat4s * sizeof(FVector4));
-
 		RHIUnlockStructuredBuffer(LightmapSceneDataBufferRHI);
-
-		LightmapSceneDataBufferSRV = RHICreateShaderResourceView(LightmapSceneDataBufferRHI);
 	}
 }
 
 TGlobalResource<FSinglePrimitiveStructuredBuffer> GIdentityPrimitiveBuffer;
+TGlobalResource<FSinglePrimitiveStructuredBuffer> GTilePrimitiveBuffer;
 
 FPrimitiveSceneShaderData::FPrimitiveSceneShaderData(const FPrimitiveSceneProxy* RESTRICT Proxy)
 {
diff --git a/Engine/Source/Runtime/Engine/Public/PrimitiveUniformShaderParameters.h b/Engine/Source/Runtime/Engine/Public/PrimitiveUniformShaderParameters.h
--- a/Engine/Source/Runtime/Engine/Public/PrimitiveUniformShaderParameters.h
+++ b/Engine/Source/Runtime/Engine/Public/PrimitiveUniformShaderParameters.h
@@ -240,6 +240,8 @@ public:
 		LightmapSceneDataBufferSRV.SafeRelease();
 	}
 
+	ENGINE_API void UploadToGPU();
+
 	FPrimitiveSceneShaderData PrimitiveSceneData;
 	FLightmapSceneShaderData LightmapSceneData;
 
@@ -255,3 +257,4 @@ public:
 * This is used when the VF is used for rendering outside normal mesh passes, where there is no valid scene.
 */
 extern ENGINE_API TGlobalResource<FSinglePrimitiveStructuredBuffer> GIdentityPrimitiveBuffer;
+extern ENGINE_API TGlobalResource<FSinglePrimitiveStructuredBuffer> GTilePrimitiveBuffer;
diff --git a/Engine/Source/Runtime/Renderer/Private/Renderer.cpp b/Engine/Source/Runtime/Renderer/Private/Renderer.cpp
--- a/Engine/Source/Runtime/Renderer/Private/Renderer.cpp
+++ b/Engine/Source/Runtime/Renderer/Private/Renderer.cpp
@@ -120,7 +120,7 @@ void FRendererModule::DrawTileMesh(FRHICommandListImmediate& RHICmdList, FMeshPa
 		extern FForwardLightingViewResources* GetMinimalDummyForwardLightingResources();
 		View.ForwardLightingResources = GetMinimalDummyForwardLightingResources();
 
-		FSinglePrimitiveStructuredBuffer SinglePrimitiveStructuredBuffer;
+		FSinglePrimitiveStructuredBuffer& SinglePrimitiveStructuredBuffer = GTilePrimitiveBuffer;
 
 		if (Mesh.VertexFactory->GetPrimitiveIdStreamIndex(EVertexInputStreamType::PositionOnly) >= 0)
 		{
@@ -147,7 +147,7 @@ void FRendererModule::DrawTileMesh(FRHICommandListImmediate& RHICmdList, FMeshPa
 				GetPrecomputedLightingParameters(FeatureLevel, LightmapParams, Mesh.LCI);
 				SinglePrimitiveStructuredBuffer.LightmapSceneData = FLightmapSceneShaderData(LightmapParams);
 
-				SinglePrimitiveStructuredBuffer.InitResource();
+				SinglePrimitiveStructuredBuffer.UploadToGPU();
 				View.PrimitiveSceneDataOverrideSRV = SinglePrimitiveStructuredBuffer.PrimitiveSceneDataBufferSRV;
 				View.LightmapSceneDataOverrideSRV = SinglePrimitiveStructuredBuffer.LightmapSceneDataBufferSRV;
 			}
@@ -341,8 +341,6 @@ void FRendererModule::DrawTileMesh(FRHICommandListImmediate& RHICmdList, FMeshPa
 				}
 			}
 		}
-
-		SinglePrimitiveStructuredBuffer.ReleaseResource();
 	}
 }
 
