Expose new getters and setters on UActorComponent, USceneComponent, and AActor for replicated properties.

diff --git a/Engine/Source/Runtime/Engine/Classes/Components/ActorComponent.h b/Engine/Source/Runtime/Engine/Classes/Components/ActorComponent.h
--- a/Engine/Source/Runtime/Engine/Classes/Components/ActorComponent.h
+++ b/Engine/Source/Runtime/Engine/Classes/Components/ActorComponent.h
@@ -915,6 +915,48 @@ private:
 	friend class FComponentReregisterContextBase;
 	friend class FComponentRecreateRenderStateContext;
 	friend struct FActorComponentTickFunction;
+
+	//~ Begin Methods for Replicated Members.
+protected:
+
+	/**
+	 * Sets the value of bReplicates without causing other side effects to this instance.
+	 * This should only be called during component construction.
+	 *
+	 * This method exists only to allow code to directly modify bReplicates to maintain existing
+	 * behavior.
+	 */
+	void SetIsReplicatedByDefault(const bool bNewReplicates);
+
+	/**
+	 * Gets the property name for bReplicates.
+	 *
+	 * This exists so subclasses don't need to have direct access to the bReplicates property so it
+	 * can be made private later.
+	 */
+	static const FName GetReplicatesPropertyName()
+	{
+		return GET_MEMBER_NAME_CHECKED(UActorComponent, bReplicates);
+	}
+
+public:
+
+	/**
+	 * Sets the value of bIsActive without causing other side effects to this instance.
+	 *
+	 * Activate, Deactivate, and SetActive are preferred in most cases because they respect virtual behavior.
+	 */
+	void SetActiveFlag(const bool bNewIsActive);
+
+	//~ End Methods for Replicated Members.
+
+protected:
+
+	/** Convenience method for testing whether or not our owner is still being constructed / initialized. */
+	bool OwnerNeedsInitialization() const;
+
+	/** Convenience method for testing whether or not we are still be constructed / initialized. */
+	bool NeedsInitialization() const;
 };
 
 //////////////////////////////////////////////////////////////////////////
diff --git a/Engine/Source/Runtime/Engine/Classes/Components/SceneComponent.h b/Engine/Source/Runtime/Engine/Classes/Components/SceneComponent.h
--- a/Engine/Source/Runtime/Engine/Classes/Components/SceneComponent.h
+++ b/Engine/Source/Runtime/Engine/Classes/Components/SceneComponent.h
@@ -834,7 +834,7 @@ public:
 	UFUNCTION(BlueprintCallable, Category="Rendering")
 	void ToggleVisibility(bool bPropagateToChildren = false)
 	{
-		SetVisibility(!bVisible, bPropagateToChildren);
+		SetVisibility(!GetVisibleFlag(), bPropagateToChildren);
 	}
 
 	/** Changes the value of bHiddenInGame, if false this will disable Visibility during gameplay */
@@ -851,7 +851,7 @@ public:
 	virtual bool ShouldCreateRenderState() const override { return true; }
 	virtual void UpdateComponentToWorld(EUpdateTransformFlags UpdateTransformFlags = EUpdateTransformFlags::None, ETeleportType Teleport = ETeleportType::None) override final
 	{
-		UpdateComponentToWorldWithParent(GetAttachParent(), GetAttachSocketName(), UpdateTransformFlags, RelativeRotationCache.RotatorToQuat(RelativeRotation), Teleport);
+		UpdateComponentToWorldWithParent(GetAttachParent(), GetAttachSocketName(), UpdateTransformFlags, RelativeRotationCache.RotatorToQuat(GetRelativeRotation()), Teleport);
 	}
 	virtual void DestroyComponent(bool bPromoteChildren = false) override;
 	virtual void OnComponentDestroyed(bool bDestroyingHierarchy) override;
@@ -1075,7 +1075,7 @@ protected:
 		Parent = Parent ? Parent : GetAttachParent();
 		if (Parent)
 		{
-			const bool bGeneral = bAbsoluteLocation || bAbsoluteRotation || bAbsoluteScale;
+			const bool bGeneral = IsUsingAbsoluteLocation() || IsUsingAbsoluteRotation() || IsUsingAbsoluteScale();
 			if (!bGeneral)
 			{
 				return NewRelativeTransform * Parent->GetSocketTransform(SocketName);
@@ -1249,6 +1249,264 @@ private:
 	friend class FScopedMovementUpdate;
 	friend class FScopedPreventAttachedComponentMove;
 	friend struct FDirectAttachChildrenAccessor;
+
+	//~ Begin Methods for Replicated Members.	
+private:
+
+	/**
+	 * Sets the value of AttachParent without causing other side effects to this instance.
+	 * Other systems may leverage this to get notifications for when the value is changed.
+	 */
+	void SetAttachParent(USceneComponent* NewAttachParent);
+	
+	/**
+	 * Sets the value of AttachSocketName without causing other side effects to this instance.
+	 * Other systems may leverage this to get notifications for when the value is changed.
+	 */
+	void SetAttachSocketName(FName NewSocketName);
+	
+	/**
+	 * Called when AttachChildren is modified.
+	 * Other systems may leverage this to get notifications for when the value is changed.
+	 */
+	void ModifiedAttachChildren();
+
+public:
+
+	/**
+	 * Gets the property name for RelativeLocation.
+	 *
+	 * This exists so subclasses don't need to have direct access to the RelativeLocation property so it
+	 * can be made private later.
+	 */
+	static const FName GetRelativeLocationPropertyName()
+	{
+		return GET_MEMBER_NAME_CHECKED(USceneComponent, RelativeLocation);
+	}
+	
+	/**
+	 * Gets the literal value of RelativeLocation.
+	 * Note, this may be an absolute location if this is a root component (not attached to anything) or
+	 * when IsUsingAbsoluteLocation returns true.
+	 *
+	 * This exists so subclasses don't need to have direct access to the RelativeLocation property so it
+	 * can be made private later.
+	 */
+	FVector GetRelativeLocation() const
+	{
+		return RelativeLocation;
+	}
+	
+	/**
+	 * Gets a refence to RelativeLocation with the expectation that it will be modified.
+	 * Note, this may be an absolute location if this is a root component (not attached to anything) or
+	 * when IsUsingAbsoluteLocation returns true.
+	 *
+	 * This exists so subclasses don't need to have direct access to the RelativeLocation property so it
+	 * can be made private later.
+	 *
+	 * You should not use this method. The standard SetRelativeLocation variants should be used.
+	 */
+	FVector& GetRelativeLocation_DirectMutable();
+
+	/**
+	 * Sets the value of RelativeLocation without causing other side effects to this instance.
+	 *
+	 * You should not use this method. The standard SetRelativeLocation variants should be used.
+	 */
+	void SetRelativeLocation_Direct(const FVector NewRelativeLocation);
+
+	/**
+	 * Gets the property name for RelativeRotation.
+	 *
+	 * This exists so subclasses don't need to have direct access to the RelativeRotation property so it
+	 * can be made private later.
+	 */
+	static const FName GetRelativeRotationPropertyName()
+	{
+		return GET_MEMBER_NAME_CHECKED(USceneComponent, RelativeRotation);
+	}
+	
+	/**
+	 * Gets the literal value of RelativeRotation.
+	 * Note, this may be an absolute rotation if this is a root component (not attached to anything) or
+	 * when GetAbsoluteRotation returns true.
+	 *
+	 * This exists so subclasses don't need to have direct access to the RelativeRotation property so it
+	 * can be made private later.
+	 */
+	FRotator GetRelativeRotation() const
+	{
+		return RelativeRotation;
+	}
+	
+	/**
+	 * Gets a refence to RelativeRotation with the expectation that it will be modified.
+	 * Note, this may be an absolute rotation if this is a root component (not attached to anything) or
+	 * when GetAbsoluteRotation returns true.
+	 *
+	 * This exists so subclasses don't need to have direct access to the RelativeRotation property so it
+	 * can be made private later.
+	 *
+	 * You should not use this method. The standard SetRelativeRotation variants should be used.
+	 */
+	FRotator& GetRelativeRotation_DirectMutable();
+
+	/**
+	 * Sets the value of RelativeRotation without causing other side effects to this instance.
+	 *
+	 * You should not use this method. The standard SetRelativeRotation variants should be used.
+	 */
+	void SetRelativeRotation_Direct(const FRotator NewRelativeRotation);
+
+	/**
+	 * Gets the property name for RelativeScale3D.
+	 *
+	 * This exists so subclasses don't need to have direct access to the RelativeScale3D property so it
+	 * can be made private later.
+	 */
+	static const FName GetRelativeScale3DPropertyName()
+	{
+		return GET_MEMBER_NAME_CHECKED(USceneComponent, RelativeScale3D);
+	}
+	
+	/**
+	 * Gets the literal value of RelativeScale3D.
+	 * Note, this may be an absolute scale if this is a root component (not attached to anything) or
+	 * when GetAbsoluteScale3D returns true.
+	 *
+	 * This exists so subclasses don't need to have direct access to the RelativeScale3D property so it
+	 * can be made private later.
+	 */
+	FVector GetRelativeScale3D() const
+	{
+		return RelativeScale3D;
+	}
+	
+	/**
+	 * Gets a refence to RelativeRotation with the expectation that it will be modified.
+	 * Note, this may be an absolute scale if this is a root component (not attached to anything) or
+	 * when GetAbsoluteScale3D returns true.
+	 *
+	 * This exists so subclasses don't need to have direct access to the RelativeScale3D property so it
+	 * can be made private later.
+	 *
+	 * You should not use this method. The standard SetRelativeScale3D variants should be used.
+	 */
+	FVector& GetRelativeScale3D_DirectMutable();
+
+	/**
+	 * Sets the value of RelativeScale3D without causing other side effects to this instance.
+	 *
+	 * You should not use this method. The standard SetRelativeScale3D variants should be used.
+	 */
+	void SetRelativeScale3D_Direct(const FVector NewRelativeScale3D);
+
+	/**
+	 * Gets the property name for bAbsoluteLocation.
+	 *
+	 * This exists so subclasses don't need to have direct access to the bAbsoluteLocation property so it
+	 * can be made private later.
+	 */
+	static const FName GetAbsoluteLocationPropertyName()
+	{
+		return GET_MEMBER_NAME_CHECKED(USceneComponent, bAbsoluteLocation);
+	}
+
+	/**
+	 * Gets the literal value of bAbsoluteLocation.
+	 *
+	 * This exists so subclasses don't need to have direct access to the bAbsoluteLocation property so it
+	 * can be made private later.
+	 */
+	bool IsUsingAbsoluteLocation() const
+	{
+		return bAbsoluteLocation;
+	}
+	
+	/** Sets the value of bAbsoluteLocation without causing other side effects to this instance. */
+	void SetUsingAbsoluteLocation(const bool bInAbsoluteLocation);
+
+	/**
+	 * Gets the property name for bAbsoluteRotation.
+	 *
+	 * This exists so subclasses don't need to have direct access to the bAbsoluteRotation property so it
+	 * can be made private later.
+	 */
+	static const FName GetAbsoluteRotationPropertyName()
+	{
+		return GET_MEMBER_NAME_CHECKED(USceneComponent, bAbsoluteRotation);
+	}
+
+	/**
+	 * Gets the literal value of bAbsoluteRotation.
+	 *
+	 * This exists so subclasses don't need to have direct access to the bAbsoluteRotation property so it
+	 * can be made private later.
+	 */
+	bool IsUsingAbsoluteRotation() const
+	{
+		return bAbsoluteRotation;
+	}
+	
+	/** Sets the value of bAbsoluteRotation without causing other side effects to this instance. */
+	void SetUsingAbsoluteRotation(const bool bInAbsoluteRotation);
+
+	/**
+	 * Gets the property name for bAbsoluteScale.
+	 * This exists so subclasses don't need to have direct access to the bAbsoluteScale property so it
+	 * can be made private later.
+	 */
+	static const FName GetAbsoluteScalePropertyName()
+	{
+		return GET_MEMBER_NAME_CHECKED(USceneComponent, bAbsoluteScale);
+	}
+
+	/**
+	 * Gets the literal value of bAbsoluteScale.
+	 *
+	 * This exists so subclasses don't need to have direct access to the bReplicates property so it
+	 * can be made private later.
+	 '*/
+	bool IsUsingAbsoluteScale() const
+	{
+		return bAbsoluteScale;
+	}
+	
+	/** Sets the value of bAbsoluteScale without causing other side effects to this instance. */
+	void SetUsingAbsoluteScale(const bool bInAbsoluteRotation);
+
+	/**
+	 * Gets the property name for bVisible.
+	 * This exists so subclasses don't need to have direct access to the bVisible property so it
+	 * can be made private later.
+	 */
+	static const FName GetVisiblePropertyName()
+	{
+		return GET_MEMBER_NAME_CHECKED(USceneComponent, bVisible);
+	}
+
+	/**
+	 * Gets the literal value of bVisible.
+	 *
+	 * This exists so subclasses don't need to have direct access to the bVisible property so it
+	 * can be made private later.
+	 *
+	 * IsVisible and IsVisibleInEditor are preferred in most cases because they respect virtual behavior.
+	 */
+	bool GetVisibleFlag() const
+	{
+		return bVisible;
+	}
+	
+	/**
+	 * Sets the value of bVisible without causing other side effects to this instance.
+	 *
+	 * ToggleVisible and SetVisibility are preferred in most cases because they respect virtual behavior and side effects.
+	 */
+	void SetVisibleFlag(const bool bInVisible);
+	
+	//~ End Methods for Replicated Members.
 };
 
 /** 
@@ -1305,22 +1563,22 @@ FORCEINLINE_DEBUGGABLE bool USceneComponent::MoveComponent(const FVector& Delta,
 
 FORCEINLINE_DEBUGGABLE void USceneComponent::SetRelativeLocation(FVector NewLocation, bool bSweep, FHitResult* OutSweepHitResult, ETeleportType Teleport)
 {
-	SetRelativeLocationAndRotation(NewLocation, RelativeRotationCache.RotatorToQuat(RelativeRotation), bSweep, OutSweepHitResult, Teleport);
+	SetRelativeLocationAndRotation(NewLocation, RelativeRotationCache.RotatorToQuat(GetRelativeRotation()), bSweep, OutSweepHitResult, Teleport);
 }
 
 FORCEINLINE_DEBUGGABLE void USceneComponent::SetRelativeRotation(const FQuat& NewRotation, bool bSweep, FHitResult* OutSweepHitResult, ETeleportType Teleport)
 {
-	SetRelativeLocationAndRotation(RelativeLocation, NewRotation, bSweep, OutSweepHitResult, Teleport);
+	SetRelativeLocationAndRotation(GetRelativeLocation(), NewRotation, bSweep, OutSweepHitResult, Teleport);
 }
 
 FORCEINLINE_DEBUGGABLE void USceneComponent::AddRelativeLocation(FVector DeltaLocation, bool bSweep, FHitResult* OutSweepHitResult, ETeleportType Teleport)
 {
-	SetRelativeLocationAndRotation(RelativeLocation + DeltaLocation, RelativeRotationCache.RotatorToQuat(RelativeRotation), bSweep, OutSweepHitResult, Teleport);
+	SetRelativeLocationAndRotation(GetRelativeLocation() + DeltaLocation, RelativeRotationCache.RotatorToQuat(GetRelativeRotation()), bSweep, OutSweepHitResult, Teleport);
 }
 
 FORCEINLINE_DEBUGGABLE void USceneComponent::AddRelativeRotation(FRotator DeltaRotation, bool bSweep, FHitResult* OutSweepHitResult, ETeleportType Teleport)
 {
-	SetRelativeRotation(RelativeRotation + DeltaRotation, bSweep, OutSweepHitResult, Teleport);
+	SetRelativeRotation(GetRelativeRotation() + DeltaRotation, bSweep, OutSweepHitResult, Teleport);
 }
 
 //////////////////////////////////////////////////////////////////////////
@@ -1363,39 +1621,7 @@ public:
 	 * Init scoped behavior for a given Component.
 	 * Note that null is perfectly acceptable here (does nothing) as a simple way to toggle behavior at runtime without weird conditional compilation.
 	 */
-	FScopedPreventAttachedComponentMove(USceneComponent* Component)
-	: Owner(Component)
-	{
-		if (Component)
-		{
-			// Save old flags
-			bSavedAbsoluteLocation = Component->bAbsoluteLocation;
-			bSavedAbsoluteRotation = Component->bAbsoluteRotation;
-			bSavedAbsoluteScale = Component->bAbsoluteScale;
-			bSavedNonAbsoluteComponent = !(bSavedAbsoluteLocation && bSavedAbsoluteRotation && bSavedAbsoluteScale);
-		
-			// Use absolute (stay in world space no matter what parent does)
-			Component->bAbsoluteLocation = true;
-			Component->bAbsoluteRotation = true;
-			Component->bAbsoluteScale = true;
-
-			if (bSavedNonAbsoluteComponent && Component->GetAttachParent())
-			{
-				// Make RelativeLocation etc relative to the world.
-				Component->ConditionalUpdateComponentToWorld();
-				Component->RelativeLocation = Component->GetComponentLocation();
-				Component->RelativeRotation = Component->GetComponentRotation();
-				Component->RelativeScale3D = Component->GetComponentScale();
-			}
-		}
-		else
-		{
-			bSavedAbsoluteLocation = false;
-			bSavedAbsoluteRotation = false;
-			bSavedAbsoluteScale = false;
-		}
-	}
-
+	FScopedPreventAttachedComponentMove(USceneComponent* Component);
 	~FScopedPreventAttachedComponentMove();
 
 private:
diff --git a/Engine/Source/Runtime/Engine/Classes/Engine/EngineTypes.h b/Engine/Source/Runtime/Engine/Classes/Engine/EngineTypes.h
--- a/Engine/Source/Runtime/Engine/Classes/Engine/EngineTypes.h
+++ b/Engine/Source/Runtime/Engine/Classes/Engine/EngineTypes.h
@@ -3127,7 +3127,7 @@ struct ENGINE_API FRepMovement
 		bRepPhysics = true;
 	}
 
-	void CopyTo(struct FRigidBodyState& RBState, const AActor* const Actor = nullptr)
+	void CopyTo(struct FRigidBodyState& RBState, const AActor* const Actor = nullptr) const
 	{
 		RBState.Position = RebaseOntoLocalOrigin(Location, Actor);
 		RBState.Quaternion = Rotation.Quaternion();
diff --git a/Engine/Source/Runtime/Engine/Classes/GameFramework/Actor.h b/Engine/Source/Runtime/Engine/Classes/GameFramework/Actor.h
--- a/Engine/Source/Runtime/Engine/Classes/GameFramework/Actor.h
+++ b/Engine/Source/Runtime/Engine/Classes/GameFramework/Actor.h
@@ -195,14 +195,18 @@ public:
 	UFUNCTION()
 	virtual void OnRep_ReplicateMovement();
 
-	UE_DEPRECATED(4.20, "Use GetTearOff() or TearOff() functions. This property will become private.")
+private:
+
 	UPROPERTY(Replicated)
 	uint8 bTearOff:1; 
 
+public:
+
 	/** If true, this actor is no longer replicated to new clients, and is "torn off" (becomes a ROLE_Authority) on clients to which it was being replicated. */
-	PRAGMA_DISABLE_DEPRECATION_WARNINGS
-	bool GetTearOff() const { return bTearOff; }
-	PRAGMA_ENABLE_DEPRECATION_WARNINGS
+	bool GetTearOff() const
+	{
+		return bTearOff;
+	}
 
 	/** Networking - Server - TearOff this actor to stop replication to clients. Will set bTearOff to true. */
 	UFUNCTION(BlueprintCallable, Category=Replication)
@@ -1992,9 +1996,6 @@ public:
 	/** Swaps Role and RemoteRole if client */
 	void ExchangeNetRoles(bool bRemoteOwner);
 
-	UE_DEPRECATED(4.22, "Renamed to SwapRoles")
-	void SwapRolesForReplay() { SwapRoles(); }
-
 	/** Calls this to swap the Role and RemoteRole.  Only call this if you know what you're doing! */
 	void SwapRoles();
 
@@ -3160,6 +3161,137 @@ private:
 	{
 		return (RootComponent != nullptr) ? RootComponent->GetRightVector() : FVector::RightVector;
 	}
+	
+	//~ Begin Methods for Replicated Members.
+public:
+
+	/**
+	 * Gets the property name for bHidden.
+	 * This exists so subclasses don't need to have direct access to the bHidden property so it
+	 * can be made private later.
+	 */
+	static const FName GetHiddenPropertyName()
+	{
+		return GET_MEMBER_NAME_CHECKED(AActor, bHidden);
+	}
+
+	/**
+	 * Gets the literal value of bHidden.
+	 *
+	 * This exists so subclasses don't need to have direct access to the bHidden property so it
+	 * can be made private later.
+	 */
+	bool IsHidden() const
+	{
+		return bHidden;
+	}
+
+	/**
+	 * Sets the value of bHidden without causing other side effects to this instance.
+	 *
+	 * SetActorHiddenInGame is preferred preferred in most cases because it respects virtual behavior.
+	 */
+	void SetHidden(const bool bInHidden);
+
+	/**
+	 * Gets the property name for bReplicateMovement.
+	 * This exists so subclasses don't need to have direct access to the bReplicateMovement property so it
+	 * can be made private later.
+	 */
+	static const FName GetReplicateMovementPropertyName()
+	{
+		return GET_MEMBER_NAME_CHECKED(AActor, bReplicateMovement);
+	}
+
+	/**
+	 * Gets the literal value of bReplicateMovement.
+	 *
+	 * This exists so subclasses don't need to have direct access to the bReplicateMovement property so it
+	 * can be made private later.
+	 */
+	bool IsReplicatingMovement() const
+	{
+		return bReplicateMovement;
+	}
+
+	/** Sets the value of bReplicateMovement without causing other side effects to this instance. */
+	void SetReplicatingMovement(bool bInReplicateMovement);
+
+	/**
+	 * Gets the property name for bCanBeDamaged.
+	 * This exists so subclasses don't need to have direct access to the bCanBeDamaged property so it
+	 * can be made private later.
+	 */
+	static const FName GetCanBeDamagedPropertyName()
+	{
+		return GET_MEMBER_NAME_CHECKED(AActor, bCanBeDamaged);
+	}
+
+	/**
+	 * Gets the literal value of bCanBeDamaged.
+	 *
+	 * This exists so subclasses don't need to have direct access to the bCanBeDamaged property so it
+	 * can be made private later.
+	 */
+	bool CanBeDamaged() const
+	{
+		return bCanBeDamaged;
+	}
+
+	/** Sets the value of bCanBeDamaged without causing other side effects to this instance. */
+	void SetCanBeDamaged(bool bInCanBeDamaged);
+
+	/**
+	 * Gets the property name for Role.
+	 * This exists so subclasses don't need to have direct access to the Role property so it
+	 * can be made private later.
+	 */
+	static const FName GetRolePropertyName()
+	{
+		return GET_MEMBER_NAME_CHECKED(AActor, Role);
+	}
+
+	/**
+	 * Sets the value of Role without causing other side effects to this instance.
+	 */
+	void SetRole(ENetRole InRole);
+	
+	/**
+	 * Gets the literal value of ReplicatedMovement.
+	 *
+	 * This exists so subclasses don't need to have direct access to the Role property so it
+	 * can be made private later.
+	 */
+	const FRepMovement& GetReplicatedMovement() const
+	{
+		return ReplicatedMovement;
+	}
+
+	/**
+	 * Gets a reference to ReplicatedMovement with the expectation that it will be modified.
+	 *
+	 * This exists so subclasses don't need to have direct access to the ReplicatedMovement property
+	 * so it can be made private later.
+	 */
+	FRepMovement& GetReplicatedMovement_Mutable();
+
+	/** Sets the value of ReplicatedMovement without causing other side effects to this instance. */
+	void SetReplicatedMovement(const FRepMovement& InReplicatedMovement);
+
+	/**
+	 * Gets the property name for Instigator.
+	 * This exists so subclasses don't need to have direct access to the Instigator property so it
+	 * can be made private later.
+	 */
+	static const FName GetInstigatorPropertyName()
+	{
+		return GET_MEMBER_NAME_CHECKED(AActor, Instigator);
+	}
+
+	/** Sets the value of Instigator without causing other side effects to this instance. */
+	void SetInstigator(APawn* InInstigator);
+
+	//~ End Methods for Replicated Members.
 };
 
 /** Internal struct used by level code to mark actors as destroyed */
@@ -3316,7 +3448,7 @@ FORCEINLINE_DEBUGGABLE bool AActor::GetActorEnableCollision() const
 
 FORCEINLINE_DEBUGGABLE bool AActor::HasAuthority() const
 {
-	return (Role == ROLE_Authority);
+	return (GetLocalRole() == ROLE_Authority);
 }
 
 FORCEINLINE_DEBUGGABLE AActor* AActor::GetOwner() const
diff --git a/Engine/Source/Runtime/Engine/Private/Actor.cpp b/Engine/Source/Runtime/Engine/Private/Actor.cpp
--- a/Engine/Source/Runtime/Engine/Private/Actor.cpp
+++ b/Engine/Source/Runtime/Engine/Private/Actor.cpp
@@ -99,7 +99,7 @@ void AActor::InitializeDefaults()
 
 	CustomTimeDilation = 1.0f;
 
-	Role = ROLE_Authority;
+	SetRole(ROLE_Authority);
 	RemoteRole = ROLE_None;
 	bReplicates = false;
 	NetPriority = 1.0f;
@@ -125,7 +125,7 @@ void AActor::InitializeDefaults()
 	bActorEnableCollision = true;
 	bActorSeamlessTraveled = false;
 	bBlockInput = false;
-	bCanBeDamaged = true;
+	SetCanBeDamaged(true);
 	bFindCameraComponentWhenViewTarget = true;
 	bAllowReceiveTickEventOnDedicatedServer = true;
 	bRelevantForNetworkReplays = true;
@@ -725,9 +725,9 @@ void AActor::PostLoadSubobjects(FObjectInstancingGraph* OuterInstanceGraph)
 	FVector OldScale;
 	if (bHadRoot)
 	{
-		OldRotation = OldRoot->RelativeRotation;
-		OldTranslation = OldRoot->RelativeLocation;
-		OldScale = OldRoot->RelativeScale3D;
+		OldRotation = OldRoot->GetRelativeRotation();
+		OldTranslation = OldRoot->GetRelativeLocation();
+		OldScale = OldRoot->GetRelativeScale3D();
 	}
 
 	Super::PostLoadSubobjects(OuterInstanceGraph);
@@ -737,9 +737,9 @@ void AActor::PostLoadSubobjects(FObjectInstancingGraph* OuterInstanceGraph)
 	if (RootComponent && bHadRoot && OldRoot != RootComponent && OldRoot->IsIn(this))
 	{
 		UE_LOG(LogActor, Log, TEXT("Root component has changed, relocating new root component to old position %s->%s"), *OldRoot->GetFullName(), *GetRootComponent()->GetFullName());
-		GetRootComponent()->RelativeRotation = OldRotation;
-		GetRootComponent()->RelativeLocation = OldTranslation;
-		GetRootComponent()->RelativeScale3D = OldScale;
+		GetRootComponent()->SetRelativeRotation_Direct(OldRotation);
+		GetRootComponent()->SetRelativeLocation_Direct(OldTranslation);
+		GetRootComponent()->SetRelativeScale3D_Direct(OldScale);
 		
 		// Migrate any attachment to the new root
 		if (OldRoot->GetAttachParent())
@@ -759,9 +759,9 @@ void AActor::PostLoadSubobjects(FObjectInstancingGraph* OuterInstanceGraph)
 		}
 
 		// Reset the transform on the old component
-		OldRoot->RelativeRotation = FRotator::ZeroRotator;
-		OldRoot->RelativeLocation = FVector::ZeroVector;
-		OldRoot->RelativeScale3D = FVector(1.0f, 1.0f, 1.0f);
+		OldRoot->SetRelativeRotation_Direct(FRotator::ZeroRotator);
+		OldRoot->SetRelativeLocation_Direct(FVector::ZeroVector);
+		OldRoot->SetRelativeScale3D_Direct(FVector(1.0f, 1.0f, 1.0f));
 	}
 }
 
@@ -969,7 +969,7 @@ UNetConnection* AActor::GetNetConnection() const
 UPlayer* AActor::GetNetOwningPlayer()
 {
 	// We can only replicate RPCs to the owning player
-	if (Role == ROLE_Authority)
+	if (GetLocalRole() == ROLE_Authority)
 	{
 		if (Owner)
 		{
@@ -1049,7 +1049,7 @@ bool AActor::ShouldTickIfViewportsOnly() const
 	return false;
 }
 
-void AActor::PreReplication( IRepChangedPropertyTracker & ChangedPropertyTracker )
+void AActor::PreReplication(IRepChangedPropertyTracker & ChangedPropertyTracker)
 {
 	// Attachment replication gets filled in by GatherCurrentMovement(), but in the case of a detached root we need to trigger remote detachment.
 	AttachmentReplication.AttachParent = nullptr;
@@ -1057,10 +1057,10 @@ void AActor::PreReplication( IRepChangedPropertyTracker & ChangedPropertyTracker
 
 	GatherCurrentMovement();
 
-	DOREPLIFETIME_ACTIVE_OVERRIDE( AActor, ReplicatedMovement, bReplicateMovement );
+	DOREPLIFETIME_ACTIVE_OVERRIDE(AActor, ReplicatedMovement, bReplicateMovement);
 
 	// Don't need to replicate AttachmentReplication if the root component replicates, because it already handles it.
-	DOREPLIFETIME_ACTIVE_OVERRIDE( AActor, AttachmentReplication, RootComponent && !RootComponent->GetIsReplicated() );
+	DOREPLIFETIME_ACTIVE_OVERRIDE(AActor, AttachmentReplication, RootComponent && !RootComponent->GetIsReplicated());
 
 	UBlueprintGeneratedClass* BPClass = Cast<UBlueprintGeneratedClass>(GetClass());
 	if (BPClass != nullptr)
@@ -1080,7 +1080,8 @@ void AActor::CallPreReplication(UNetDriver* NetDriver)
 
 	// PreReplication is only called on the server, except when we're recording a Client Replay.
 	// In that case we call PreReplication on the locally controlled Character as well.
-	if ((Role == ROLE_Authority) || ((Role == ROLE_AutonomousProxy) && GetWorld()->IsRecordingClientReplay()))
+	const ENetRole LocalRole = GetLocalRole();
+	if ((LocalRole == ROLE_Authority) || ((LocalRole == ROLE_AutonomousProxy) && GetWorld()->IsRecordingClientReplay()))
 	{
 		PreReplication(*ActorChangedPropertyTracker);
 	}
@@ -1296,7 +1297,7 @@ bool AActor::CheckStillInWorld()
 	}
 
 	// Only authority or non-networked actors should be destroyed, otherwise misprediction can destroy something the server is intending to keep alive.
-	if (!(HasAuthority() || Role == ROLE_None))
+	if (!(HasAuthority() || GetLocalRole() == ROLE_None))
 	{
 		return true;
 	}
@@ -1592,6 +1593,7 @@ void AActor::SetOwner(AActor* NewOwner)
 		}
 
 		Owner = NewOwner;
+		// MARK_PROPERTY_DIRTY_BY_NAME(AActor, Owner, this);
 
 		if (Owner != nullptr)
 		{
@@ -1681,9 +1683,9 @@ void AActor::OnRep_AttachmentReplication()
 
 			if (AttachParentComponent)
 			{
-				RootComponent->RelativeLocation = AttachmentReplication.LocationOffset;
-				RootComponent->RelativeRotation = AttachmentReplication.RotationOffset;
-				RootComponent->RelativeScale3D = AttachmentReplication.RelativeScale3D;
+				RootComponent->SetRelativeLocation_Direct(AttachmentReplication.LocationOffset);
+				RootComponent->SetRelativeRotation_Direct(AttachmentReplication.RotationOffset);
+				RootComponent->SetRelativeScale3D_Direct(AttachmentReplication.RelativeScale3D);
 
 				// If we're already attached to the correct Parent and Socket, then the update must be position only.
 				// AttachToComponent would early out in this case.
@@ -1765,10 +1767,11 @@ void AActor::DetachRootComponentFromParent(bool bMaintainWorldPosition)
 	{
 PRAGMA_DISABLE_DEPRECATION_WARNINGS
 		RootComponent->DetachFromParent(bMaintainWorldPosition);
-PRAGMA_ENABLE_DEPRECATION_WARNINGS
+PRAGMA_ENABLE_DEPRECATION_WARNINGS		
 
 		// Clear AttachmentReplication struct
 		AttachmentReplication = FRepAttachment();
+		// MARK_PROPERTY_DIRTY_FROM_NAME(AActor, AttachmentReplication, this);
 	}
 }
 
@@ -1961,7 +1964,7 @@ bool AActor::IsRelevancyOwnerFor(const AActor* ReplicatedActor, const AActor* Ac
 
 void AActor::ForceNetUpdate()
 {
-	if (Role == ROLE_Authority)
+	if (GetLocalRole() == ROLE_Authority)
 	{
 		// ForceNetUpdate on the game net driver only if we are the authority...
 		UNetDriver* NetDriver = GetNetDriver();
@@ -2222,9 +2225,8 @@ void AActor::TearOff()
 
 	if (NetMode == NM_ListenServer || NetMode == NM_DedicatedServer)
 	{
-		PRAGMA_DISABLE_DEPRECATION_WARNINGS
 		bTearOff = true;
-		PRAGMA_ENABLE_DEPRECATION_WARNINGS
+		// MARK_PROPERTY_DIRTY_FROM_NAME(AActor, bTearOff, this);
 		if (UNetDriver* NetDriver = GetNetDriver())
 		{
 			NetDriver->NotifyActorTearOff(this);
@@ -2242,7 +2244,7 @@ void AActor::Reset()
 void AActor::FellOutOfWorld(const UDamageType& dmgType)
 {
 	// Only authority or non-networked actors should be destroyed, otherwise misprediction can destroy something the server is intending to keep alive.
-	if (HasAuthority() || Role == ROLE_None)
+	if (HasAuthority() || GetLocalRole() == ROLE_None)
 	{
 		DisableComponentsSimulatePhysics();
 		SetActorHiddenInGame(true);
@@ -2253,7 +2255,7 @@ void AActor::FellOutOfWorld(const UDamageType& dmgType)
 
 void AActor::MakeNoise(float Loudness, APawn* NoiseInstigator, FVector NoiseLocation, float MaxRange, FName Tag)
 {
-	NoiseInstigator = NoiseInstigator ? NoiseInstigator : Instigator;
+	NoiseInstigator = NoiseInstigator ? NoiseInstigator : GetInstigator();
 	if ((GetNetMode() != NM_Client) && NoiseInstigator)
 	{
 		AActor::MakeNoiseDelegate.Execute(this, Loudness, NoiseInstigator
@@ -2452,7 +2454,7 @@ void AActor::DisplayDebug(UCanvas* Canvas, const FDebugDisplayInfo& DebugDisplay
 		if( GetNetMode() != NM_Standalone )
 		{
 			// networking attributes
-			T = FString::Printf(TEXT("ROLE: %i RemoteRole: %i NetNode: %i"), (int32)Role, (int32)RemoteRole, (int32)GetNetMode());
+			T = FString::Printf(TEXT("ROLE: %i RemoteRole: %i NetNode: %i"), (int32)GetLocalRole(), (int32)RemoteRole, (int32)GetNetMode());
 
 			if( GetTearOff() )
 			{
@@ -2503,7 +2505,7 @@ void AActor::DisplayDebug(UCanvas* Canvas, const FDebugDisplayInfo& DebugDisplay
 		DisplayDebugManager.DrawString(T);
 	}
 	DisplayDebugManager.DrawString(FString::Printf(TEXT(" Instigator: %s Owner: %s"), 
-		(Instigator ? *Instigator->GetName() : TEXT("None")), (Owner ? *Owner->GetName() : TEXT("None"))));
+		*GetNameSafe(GetInstigator()), *GetNameSafe(Owner)));
 
 	static FName NAME_Animation(TEXT("Animation"));
 	static FName NAME_Bones = FName(TEXT("Bones"));
@@ -2565,7 +2567,7 @@ void AActor::CalcCamera(float DeltaTime, FMinimalViewInfo& OutResult)
 
 		for (UCameraComponent* CameraComponent : Cameras)
 		{
-			if (CameraComponent->bIsActive)
+			if (CameraComponent->IsActive())
 			{
 				CameraComponent->GetCameraView(DeltaTime, OutResult);
 				return;
@@ -2586,7 +2588,7 @@ bool AActor::HasActiveCameraComponent() const
 			const UCameraComponent* CameraComponent = Cast<const UCameraComponent>(Component);
 			if (CameraComponent)
 			{
-				if (CameraComponent->bIsActive)
+				if (CameraComponent->IsActive())
 				{
 					return true;
 				}
@@ -2606,7 +2608,7 @@ bool AActor::HasActivePawnControlCameraComponent() const
 			const UCameraComponent* CameraComponent = Cast<const UCameraComponent>(Component);
 			if (CameraComponent)
 			{
-				if (CameraComponent->bIsActive && CameraComponent->bUsePawnControlRotation)
+				if (CameraComponent->IsActive() && CameraComponent->bUsePawnControlRotation)
 				{
 					return true;
 				}
@@ -2951,14 +2953,14 @@ void AActor::PostSpawnInitialize(FTransform const& UserSpawnTransform, AActor* I
 	CreationTime = (World ? World->GetTimeSeconds() : 0.f);
 
 	// Set network role.
-	check(Role == ROLE_Authority);
+	check(GetLocalRole() == ROLE_Authority);
 	ExchangeNetRoles(bRemoteOwned);
 
 	// Set owner.
 	SetOwner(InOwner);
 
 	// Set instigator
-	Instigator = InInstigator;
+	SetInstigator(InInstigator);
 
 	// Set the actor's world transform if it has a native rootcomponent.
 	USceneComponent* const SceneRootComponent = FixupNativeActorComponents(this);
@@ -3001,7 +3003,7 @@ void AActor::PostSpawnInitialize(FTransform const& UserSpawnTransform, AActor* I
 		{
 			// In the "normal" case we do respect any non-default transform value that the root component may have received from the archetype
 			// that's owned by the native CDO, so the final transform might not always necessarily equate to the passed-in UserSpawnTransform.
-			const FTransform RootTransform(SceneRootComponent->RelativeRotation, SceneRootComponent->RelativeLocation, SceneRootComponent->RelativeScale3D);
+			const FTransform RootTransform(SceneRootComponent->GetRelativeRotation(), SceneRootComponent->GetRelativeLocation(), SceneRootComponent->GetRelativeScale3D());
 			const FTransform FinalRootComponentTransform = RootTransform * UserSpawnTransform;
 			SceneRootComponent->SetWorldTransform(FinalRootComponentTransform);
 		}
@@ -3217,7 +3219,7 @@ void AActor::PostActorConstruction()
 
 void AActor::SetReplicates(bool bInReplicates)
 { 
-	if (Role == ROLE_Authority)
+	if (GetLocalRole() == ROLE_Authority)
 	{
 		const bool bNewlyReplicates = (bReplicates == false && bInReplicates == true);
 	
@@ -3239,6 +3241,8 @@ void AActor::SetReplicates(bool bInReplicates)
 				ForcePropertyCompare();
 			}
 		}
+
+		// MARK_PROPERTY_DIRTY_FROM_NAME(AActor, RemoteRole, this);
 	}
 	else
 	{
@@ -3248,7 +3252,7 @@ void AActor::SetReplicates(bool bInReplicates)
 
 void AActor::SetReplicateMovement(bool bInReplicateMovement)
 {
-	bReplicateMovement = bInReplicateMovement;
+	SetReplicatingMovement(bInReplicateMovement);
 }
 
 void AActor::SetAutonomousProxy(const bool bInAutonomousProxy, const bool bAllowForcePropertyCompare)
@@ -3260,7 +3264,9 @@ void AActor::SetAutonomousProxy(const bool bInAutonomousProxy, const bool bAllow
 		RemoteRole = (bInAutonomousProxy ? ROLE_AutonomousProxy : ROLE_SimulatedProxy);
 
 		if (bAllowForcePropertyCompare && RemoteRole != OldRemoteRole)
-		{ 
+		{
+			// MARK_PROPERTY_DIRTY_FROM_NAME(AActor, RemoteRole, this);
+
 			// We have to do this so the role change above will replicate (turn off shadow state sharing for a frame)
 			// This is because RemoteRole is special since it will change between connections, so we have to special case
 			ForcePropertyCompare();
@@ -3275,6 +3281,8 @@ void AActor::SetAutonomousProxy(const bool bInAutonomousProxy, const bool bAllow
 void AActor::CopyRemoteRoleFrom(const AActor* CopyFromActor)
 {
 	RemoteRole = CopyFromActor->GetRemoteRole();
+	// MARK_PROPERTY_DIRTY_FROM_NAME(AActor, RemoteRole, this);
+
 	if (RemoteRole != ROLE_None)
 	{
 		GetWorld()->AddNetworkActor(this);
@@ -3319,6 +3327,10 @@ void AActor::ExchangeNetRoles(bool bRemoteOwned)
 void AActor::SwapRoles()
 {
 	Swap(Role, RemoteRole);
+
+	// MARK_PROPERTY_DIRTY_FROM_NAME(AActor, RemoteRole, this);
+	// MARK_PROPERTY_DIRTY_FROM_NAME(AActor, Role, this);
+
 	ForcePropertyCompare();
 }
 
@@ -3812,16 +3824,16 @@ FVector AActor::GetActorRelativeScale3D() const
 {
 	if (RootComponent)
 	{
-		return RootComponent->RelativeScale3D;
+		return RootComponent->GetRelativeScale3D();
 	}
 	return FVector(1,1,1);
 }
 
 void AActor::SetActorHiddenInGame( bool bNewHidden )
 {
-	if (bHidden != bNewHidden)
+	if (IsHidden() != bNewHidden)
 	{
-		bHidden = bNewHidden;
+		SetHidden(bNewHidden);
 		MarkComponentsRenderStateDirty();
 	}
 }
@@ -3989,8 +4001,10 @@ int32 AActor::GetFunctionCallspace( UFunction* Function, FFrame* Stack )
 		return FunctionCallspace::Local;
 	}
 
+	const ENetRole LocalRole = GetLocalRole();
+
 	// If we are on a client and function is 'skip on client', absorb it
-	FunctionCallspace::Type Callspace = (Role < ROLE_Authority) && Function->HasAllFunctionFlags(FUNC_BlueprintAuthorityOnly) ? FunctionCallspace::Absorbed : FunctionCallspace::Local;
+	FunctionCallspace::Type Callspace = (LocalRole < ROLE_Authority) && Function->HasAllFunctionFlags(FUNC_BlueprintAuthorityOnly) ? FunctionCallspace::Absorbed : FunctionCallspace::Local;
 	
 	if (IsPendingKill())
 	{
@@ -4025,7 +4039,7 @@ int32 AActor::GetFunctionCallspace( UFunction* Function, FFrame* Stack )
 	// Quick reject 2. Has to be a network game to continue
 	if (NetMode == NM_Standalone)
 	{
-		if (Role < ROLE_Authority && (Function->FunctionFlags & FUNC_NetServer))
+		if (LocalRole < ROLE_Authority && (Function->FunctionFlags & FUNC_NetServer))
 		{
 			// Don't let clients call server functions (in edge cases where NetMode is standalone (net driver is null))
 			DEBUG_CALLSPACE(TEXT("GetFunctionCallspace No Authority Server Call Absorbed: %s"), *Function->GetName());
@@ -4096,7 +4110,7 @@ int32 AActor::GetFunctionCallspace( UFunction* Function, FFrame* Stack )
 	}
 
 	// Check if the actor can potentially call remote functions	
-	if (Role == ROLE_Authority)
+	if (LocalRole == ROLE_Authority)
 	{
 		UNetConnection* NetConnection = GetNetConnection();
 		if (NetConnection == nullptr)
@@ -4650,7 +4664,7 @@ void AActor::SetLifeSpan( float InLifespan )
 	// Store the new value
 	InitialLifeSpan = InLifespan;
 	// Initialize a timer for the actors lifespan if there is one. Otherwise clear any existing timer
-	if ((Role == ROLE_Authority || GetTearOff()) && !IsPendingKill())
+	if ((GetLocalRole() == ROLE_Authority || GetTearOff()) && !IsPendingKill())
 	{
 		if( InLifespan > 0.0f)
 		{
@@ -4939,4 +4953,45 @@ void AActor::SetLODParent(UPrimitiveComponent* InLODParent, float InParentDrawDi
 	}
 }
 
+void AActor::SetHidden(bool bInHidden)
+{
+	bHidden = bInHidden;
+	// MARK_PROPERTY_DIRTY_FROM_NAME(AActor, bHidden, this);
+}
+
+void AActor::SetReplicatingMovement(bool bInReplicateMovement)
+{
+	bReplicateMovement = bInReplicateMovement;
+	// MARK_PROPERTY_DIRTY_FROM_NAME(AActor, bReplicateMovement, this);
+}
+
+void AActor::SetCanBeDamaged(bool bInCanBeDamaged)
+{
+	bCanBeDamaged = bInCanBeDamaged;
+	// MARK_PROPERTY_DIRTY_FROM_NAME(AActor, bCanBeDamaged, this);
+}
+
+void AActor::SetRole(ENetRole InRole)
+{
+	Role = InRole;
+	// MARK_PROPERTY_DIRTY_FROM_NAME(AActor, Role, this);
+}
+
+FRepMovement& AActor::GetReplicatedMovement_Mutable()
+{
+	// MARK_PROPERTY_DIRTY_FROM_NAME(AActor, ReplicatedMovement, this);
+	return ReplicatedMovement;
+}
+
+void AActor::SetReplicatedMovement(const FRepMovement& InReplicatedMovement)
+{
+	GetReplicatedMovement_Mutable() = InReplicatedMovement;
+}
+
+void AActor::SetInstigator(APawn* InInstigator)
+{
+	Instigator = InInstigator;
+	// MARK_PROPERTY_DIRTY_FROM_NAME(AActor, Instigator, this);
+}
+
 #undef LOCTEXT_NAMESPACE
diff --git a/Engine/Source/Runtime/Engine/Private/ActorReplication.cpp b/Engine/Source/Runtime/Engine/Private/ActorReplication.cpp
--- a/Engine/Source/Runtime/Engine/Private/ActorReplication.cpp
+++ b/Engine/Source/Runtime/Engine/Private/ActorReplication.cpp
@@ -42,12 +42,12 @@ float AActor::GetNetPriority(const FVector& ViewPos, const FVector& ViewDir, AAc
 		return Owner->GetNetPriority(ViewPos, ViewDir, Viewer, ViewTarget, InChannel, Time, bLowBandwidth);
 	}
 
-	if (ViewTarget && (this == ViewTarget || Instigator == ViewTarget))
+	if (ViewTarget && (this == ViewTarget || GetInstigator() == ViewTarget))
 	{
 		// If we're the view target or owned by the view target, use a high priority
 		Time *= 4.f;
 	}
-	else if (!bHidden && GetRootComponent() != NULL)
+	else if (!IsHidden() && GetRootComponent() != NULL)
 	{
 		// If this actor has a location, adjust priority based on location
 		FVector Dir = GetActorLocation() - ViewPos;
@@ -82,12 +82,12 @@ float AActor::GetNetPriority(const FVector& ViewPos, const FVector& ViewDir, AAc
 
 float AActor::GetReplayPriority(const FVector& ViewPos, const FVector& ViewDir, class AActor* Viewer, AActor* ViewTarget, UActorChannel* const InChannel, float Time)
 {
-	if (ViewTarget && (this == ViewTarget || Instigator == ViewTarget))
+	if (ViewTarget && (this == ViewTarget || GetInstigator() == ViewTarget))
 	{
 		// If we're the view target or owned by the view target, use a high priority
 		Time *= 10.0f;
 	}
-	else if (!bHidden && GetRootComponent() != NULL)
+	else if (!IsHidden() && GetRootComponent() != NULL)
 	{
 		// If this actor has a location, adjust priority based on location
 		FVector Dir = GetActorLocation() - ViewPos;
@@ -128,10 +128,10 @@ bool AActor::GetNetDormancy(const FVector& ViewPos, const FVector& ViewDir, AAct
 
 void AActor::PreNetReceive()
 {
-	ActorReplication::SavedbHidden = bHidden;
+	ActorReplication::SavedbHidden = IsHidden();
 	ActorReplication::SavedOwner = Owner;
-	ActorReplication::SavedbRepPhysics = ReplicatedMovement.bRepPhysics;
-	ActorReplication::SavedRole = Role;
+	ActorReplication::SavedbRepPhysics = GetReplicatedMovement().bRepPhysics;
+	ActorReplication::SavedRole = GetLocalRole();
 }
 
 void AActor::PostNetReceive()
@@ -140,12 +140,12 @@ void AActor::PostNetReceive()
 	{
 		// Initially we need to sync the state regardless of whether bRepPhysics has "changed" since it may not currently match IsSimulatingPhysics().
 		SyncReplicatedPhysicsSimulation();
-		ActorReplication::SavedbRepPhysics = ReplicatedMovement.bRepPhysics;
+		ActorReplication::SavedbRepPhysics = GetReplicatedMovement().bRepPhysics;
 		bNetCheckedInitialPhysicsState = true;
 	}
 
-	ExchangeB( bHidden, ActorReplication::SavedbHidden );
-	Exchange ( Owner, ActorReplication::SavedOwner );
+	ExchangeB(bHidden, ActorReplication::SavedbHidden);
+	Exchange(Owner, ActorReplication::SavedOwner);
 
 	if (bHidden != ActorReplication::SavedbHidden)
 	{
@@ -173,27 +173,29 @@ void AActor::OnRep_ReplicatedMovement()
 	// Since ReplicatedMovement and AttachmentReplication are REPNOTIFY_Always (and OnRep_AttachmentReplication may call OnRep_ReplicatedMovement directly),
 	// this check is needed since this can still be called on actors for which bReplicateMovement is false - for example, during fast-forward in replay playback.
 	// When this happens, the values in ReplicatedMovement aren't valid, and must be ignored.
-	if (!bReplicateMovement)
+	if (!IsReplicatingMovement())
 	{
 		return;
 	}
 
+	const FRepMovement& LocalRepMovement = GetReplicatedMovement();
+
 #if !(UE_BUILD_SHIPPING || UE_BUILD_TEST)
 		if (CVarDrawDebugRepMovement->GetInt() > 0)
 		{
-			DrawDebugCapsule(GetWorld(), ReplicatedMovement.Location, GetSimpleCollisionHalfHeight(), GetSimpleCollisionRadius(), ReplicatedMovement.Rotation.Quaternion(), FColor(100, 255, 100), true, 1.f);
+			DrawDebugCapsule(GetWorld(), LocalRepMovement.Location, GetSimpleCollisionHalfHeight(), GetSimpleCollisionRadius(), LocalRepMovement.Rotation.Quaternion(), FColor(100, 255, 100), true, 1.f);
 		}
 #endif
 
 	if (RootComponent)
 	{
-		if (ActorReplication::SavedbRepPhysics != ReplicatedMovement.bRepPhysics)
+		if (ActorReplication::SavedbRepPhysics != LocalRepMovement.bRepPhysics)
 		{
 			// Turn on/off physics sim to match server.
 			SyncReplicatedPhysicsSimulation();
 		}
 
-		if (ReplicatedMovement.bRepPhysics)
+		if (LocalRepMovement.bRepPhysics)
 		{
 			// Sync physics state
 			checkSlow(RootComponent->IsSimulatingPhysics());
@@ -209,20 +211,20 @@ void AActor::OnRep_ReplicatedMovement()
 			// Attachment trumps global position updates, see GatherCurrentMovement().
 			if (!RootComponent->GetAttachParent())
 			{
-				if (Role == ROLE_SimulatedProxy)
+				if (GetLocalRole() == ROLE_SimulatedProxy)
 				{
 #if ENABLE_NAN_DIAGNOSTIC
-					if (ReplicatedMovement.Location.ContainsNaN())
+					if (LocalRepMovement.Location.ContainsNaN())
 					{
 						logOrEnsureNanError(TEXT("AActor::OnRep_ReplicatedMovement found NaN in ReplicatedMovement.Location"));
 					}
-					if (ReplicatedMovement.Rotation.ContainsNaN())
+					if (LocalRepMovement.Rotation.ContainsNaN())
 					{
 						logOrEnsureNanError(TEXT("AActor::OnRep_ReplicatedMovement found NaN in ReplicatedMovement.Rotation"));
 					}
 #endif
 
-					PostNetReceiveVelocity(ReplicatedMovement.LinearVelocity);
+					PostNetReceiveVelocity(LocalRepMovement.LinearVelocity);
 					PostNetReceiveLocationAndRotation();
 				}
 			}
@@ -232,11 +234,12 @@ void AActor::OnRep_ReplicatedMovement()
 
 void AActor::PostNetReceiveLocationAndRotation()
 {
-	FVector NewLocation = FRepMovement::RebaseOntoLocalOrigin(ReplicatedMovement.Location, this);
+	const FRepMovement& LocalRepMovement = GetReplicatedMovement();
+	FVector NewLocation = FRepMovement::RebaseOntoLocalOrigin(LocalRepMovement.Location, this);
 
-	if( RootComponent && RootComponent->IsRegistered() && (NewLocation != GetActorLocation() || ReplicatedMovement.Rotation != GetActorRotation()) )
+	if( RootComponent && RootComponent->IsRegistered() && (NewLocation != GetActorLocation() || LocalRepMovement.Rotation != GetActorRotation()) )
 	{
-		SetActorLocationAndRotation(NewLocation, ReplicatedMovement.Rotation, /*bSweep=*/ false);
+		SetActorLocationAndRotation(NewLocation, LocalRepMovement.Rotation, /*bSweep=*/ false);
 	}
 }
 
@@ -250,7 +253,7 @@ void AActor::PostNetReceivePhysicState()
 	if (RootPrimComp)
 	{
 		FRigidBodyState NewState;
-		ReplicatedMovement.CopyTo(NewState, this);
+		GetReplicatedMovement().CopyTo(NewState, this);
 
 		FVector DeltaPos(FVector::ZeroVector);
 		RootPrimComp->SetRigidBodyReplicatedTarget(NewState);
@@ -259,14 +262,16 @@ void AActor::PostNetReceivePhysicState()
 
 void AActor::SyncReplicatedPhysicsSimulation()
 {
-	if (bReplicateMovement && RootComponent && (RootComponent->IsSimulatingPhysics() != ReplicatedMovement.bRepPhysics))
+	const FRepMovement& LocalRepMovement = GetReplicatedMovement();
+
+	if (IsReplicatingMovement() && RootComponent && (RootComponent->IsSimulatingPhysics() != LocalRepMovement.bRepPhysics))
 	{
 		UPrimitiveComponent* RootPrimComp = Cast<UPrimitiveComponent>(RootComponent);
 		if (RootPrimComp)
 		{
-			RootPrimComp->SetSimulatePhysics(ReplicatedMovement.bRepPhysics);
+			RootPrimComp->SetSimulatePhysics(LocalRepMovement.bRepPhysics);
 
-			if(!ReplicatedMovement.bRepPhysics)
+			if(!LocalRepMovement.bRepPhysics)
 			{
 				if (UWorld* World = GetWorld())
 				{
@@ -290,23 +295,23 @@ bool AActor::IsWithinNetRelevancyDistance(const FVector& SrcLocation) const
 
 bool AActor::IsNetRelevantFor(const AActor* RealViewer, const AActor* ViewTarget, const FVector& SrcLocation) const
 {
-	if (bAlwaysRelevant || IsOwnedBy(ViewTarget) || IsOwnedBy(RealViewer) || this == ViewTarget || ViewTarget == Instigator)
+	if (bAlwaysRelevant || IsOwnedBy(ViewTarget) || IsOwnedBy(RealViewer) || this == ViewTarget || ViewTarget == GetInstigator())
 	{
 		return true;
 	}
-	else if ( bNetUseOwnerRelevancy && Owner)
+	else if (bNetUseOwnerRelevancy && Owner)
 	{
 		return Owner->IsNetRelevantFor(RealViewer, ViewTarget, SrcLocation);
 	}
-	else if ( bOnlyRelevantToOwner )
+	else if (bOnlyRelevantToOwner)
 	{
 		return false;
 	}
-	else if ( RootComponent && RootComponent->GetAttachParent() && RootComponent->GetAttachParent()->GetOwner() && (Cast<USkeletalMeshComponent>(RootComponent->GetAttachParent()) || (RootComponent->GetAttachParent()->GetOwner() == Owner)) )
+	else if (RootComponent && RootComponent->GetAttachParent() && RootComponent->GetAttachParent()->GetOwner() && (Cast<USkeletalMeshComponent>(RootComponent->GetAttachParent()) || (RootComponent->GetAttachParent()->GetOwner() == Owner)))
 	{
 		return RootComponent->GetAttachParent()->GetOwner()->IsNetRelevantFor(RealViewer, ViewTarget, SrcLocation);
 	}
-	else if( bHidden && (!RootComponent || !RootComponent->IsCollisionEnabled()) )
+	else if(IsHidden() && (!RootComponent || !RootComponent->IsCollisionEnabled()))
 	{
 		return false;
 	}
@@ -328,8 +333,14 @@ bool AActor::IsReplayRelevantFor(const AActor* RealViewer, const AActor* ViewTar
 
 void AActor::GatherCurrentMovement()
 {
-	if (bReplicateMovement || (RootComponent && RootComponent->GetAttachParent()))
+	if (IsReplicatingMovement() || (RootComponent && RootComponent->GetAttachParent()))
 	{
+		bool bWasAttachmentModified = false;
+		bool bWasRepMovementModified = false;
+
+		AActor* OldAttachParent = AttachmentReplication.AttachParent;
+		USceneComponent* OldAttachComponent = AttachmentReplication.AttachComponent; 
+	
 		AttachmentReplication.AttachParent = nullptr;
 		AttachmentReplication.AttachComponent = nullptr;
 
@@ -343,6 +354,8 @@ void AActor::GatherCurrentMovement()
 			// Don't replicate movement if we're welded to another parent actor.
 			// Their replication will affect our position indirectly since we are attached.
 			ReplicatedMovement.bRepPhysics = !RootPrimComp->IsWelded();
+			
+			bWasRepMovementModified = true;
 		}
 		else if (RootComponent != nullptr)
 		{
@@ -354,11 +367,14 @@ void AActor::GatherCurrentMovement()
 				AttachmentReplication.AttachParent = RootComponent->GetAttachParent()->GetAttachmentRootActor();
 				if (AttachmentReplication.AttachParent != nullptr)
 				{
-					AttachmentReplication.LocationOffset = RootComponent->RelativeLocation;
-					AttachmentReplication.RotationOffset = RootComponent->RelativeRotation;
-					AttachmentReplication.RelativeScale3D = RootComponent->RelativeScale3D;
+					AttachmentReplication.LocationOffset = RootComponent->GetRelativeLocation();
+					AttachmentReplication.RotationOffset = RootComponent->GetRelativeRotation();
+					AttachmentReplication.RelativeScale3D = RootComponent->GetRelativeScale3D();
 					AttachmentReplication.AttachComponent = RootComponent->GetAttachParent();
 					AttachmentReplication.AttachSocket = RootComponent->GetAttachSocketName();
+
+					// Technically, the values might have stayed the same, but we'll just assume they've changed.
+					bWasAttachmentModified = true;
 				}
 			}
 			else
@@ -367,6 +383,22 @@ void AActor::GatherCurrentMovement()
 				ReplicatedMovement.Rotation = RootComponent->GetComponentRotation();
 				ReplicatedMovement.LinearVelocity = GetVelocity();
 				ReplicatedMovement.AngularVelocity = FVector::ZeroVector;
+
+				// Technically, the values might have stayed the same, but we'll just assume they've changed.
+				bWasRepMovementModified = true;
+			}
+
+			if (bWasRepMovementModified ||
+				ReplicatedMovement.bRepPhysics != false)
+			{
+				// MARK_PROPERTY_DIRTY_FROM_NAME(AActor, ReplicatedMovement, this);
+			}
+
+			if (bWasAttachmentModified ||
+				OldAttachParent != AttachmentReplication.AttachParent ||
+				OldAttachComponent != AttachmentReplication.AttachComponent)
+			{
+				// MARK_PROPERTY_DIRTY_FROM_NAME(AActor, AttachmentReplication, this);
 			}
 
 			ReplicatedMovement.bRepPhysics = false;
diff --git a/Engine/Source/Runtime/Engine/Private/Components/ActorComponent.cpp b/Engine/Source/Runtime/Engine/Private/Components/ActorComponent.cpp
--- a/Engine/Source/Runtime/Engine/Private/Components/ActorComponent.cpp
+++ b/Engine/Source/Runtime/Engine/Private/Components/ActorComponent.cpp
@@ -1568,7 +1568,7 @@ void UActorComponent::Activate(bool bReset)
 	if (bReset || ShouldActivate()==true)
 	{
 		SetComponentTickEnabled(true);
-		bIsActive = true;
+		SetActiveFlag(true);
 
 		OnComponentActivated.Broadcast(this, bReset);
 	}
@@ -1579,7 +1579,7 @@ void UActorComponent::Deactivate()
 	if (ShouldActivate()==false)
 	{
 		SetComponentTickEnabled(false);
-		bIsActive = false;
+		SetActiveFlag(false);
 
 		OnComponentDeactivated.Broadcast(this);
 	}
@@ -1588,7 +1588,7 @@ void UActorComponent::Deactivate()
 bool UActorComponent::ShouldActivate() const
 {
 	// if not active, should activate
-	return !bIsActive;
+	return !IsActive();
 }
 
 void UActorComponent::SetActive(bool bNewActive, bool bReset)
@@ -1620,7 +1620,7 @@ void UActorComponent::SetAutoActivate(bool bNewAutoActivate)
 
 void UActorComponent::ToggleActive()
 {
-	SetActive(!bIsActive);
+	SetActive(!IsActive());
 }
 
 void UActorComponent::SetTickableWhenPaused(bool bTickableWhenPaused)
@@ -1697,13 +1697,16 @@ bool UActorComponent::IsSupportedForNetworking() const
 	return GetIsReplicated() || IsNameStableForNetworking();
 }
 
-void UActorComponent::SetIsReplicated(bool ShouldReplicate)
+void UActorComponent::SetIsReplicated(bool bShouldReplicate)
 {
-	if (bReplicates != ShouldReplicate)
+	if (GetIsReplicated() != bShouldReplicate)
 	{
+		ensureMsgf(!NeedsInitialization(), TEXT("SetReplicatedByDefault is preferred during Component Construction."));
+
 		if (GetComponentClassCanReplicate())
 		{
-			bReplicates = ShouldReplicate;
+			bReplicates = bShouldReplicate;
+			// MARK_PROPERTY_DIRTY_FROM_NAME(UActorComponent, bReplicates, this);
 
 			if (AActor* MyOwner = GetOwner())
 			{
@@ -1739,7 +1742,7 @@ bool UActorComponent::GetComponentClassCanReplicate() const
 ENetRole UActorComponent::GetOwnerRole() const
 {
 	AActor* MyOwner = GetOwner();
-	return (MyOwner ? MyOwner->Role.GetValue() : ROLE_None);
+	return (MyOwner ? MyOwner->GetLocalRole() : ROLE_None);
 }
 
 void UActorComponent::GetLifetimeReplicatedProps( TArray< FLifetimeProperty > & OutLifetimeProps ) const
@@ -1756,7 +1759,7 @@ void UActorComponent::GetLifetimeReplicatedProps( TArray< FLifetimeProperty > &
 
 void UActorComponent::OnRep_IsActive()
 {
-	SetComponentTickEnabled(bIsActive);
+	SetComponentTickEnabled(IsActive());
 }
 
 bool UActorComponent::IsEditableWhenInherited() const
@@ -1895,4 +1898,36 @@ AActor* UActorComponent::GetActorOwnerNoninline() const
 	return GetTypedOuter<AActor>();
 }
 
+void UActorComponent::SetIsReplicatedByDefault(const bool bNewReplicates)
+{
+	// Don't bother checking parent here.
+	if (LIKELY(NeedsInitialization()))
+	{
+		bReplicates = bNewReplicates;
+		// MARK_PROPERTY_DIRTY_FROM_NAME(UActorComponent, bReplicates, this);
+	}
+	else
+	{
+		ensureMsgf(false, TEXT("SetIsReplicatedByDefault should only be called during Component Construction. Class=%s"), *GetPathNameSafe(GetClass()));
+		SetIsReplicated(bNewReplicates);
+	}
+}
+
+void UActorComponent::SetActiveFlag(const bool bNewIsActive)
+{
+	bIsActive = bNewIsActive;
+	// MARK_PROPERTY_DIRTY_FROM_NAME(UActorComponent, bIsActive, this);
+}
+
+bool UActorComponent::OwnerNeedsInitialization() const
+{
+	AActor* Owner = GetOwner();
+	return Owner && Owner->HasAnyFlags(RF_NeedInitialization);
+}
+
+bool UActorComponent::NeedsInitialization() const
+{
+	return HasAnyFlags(RF_NeedInitialization);
+}
+
 #undef LOCTEXT_NAMESPACE
diff --git a/Engine/Source/Runtime/Engine/Private/Components/SceneComponent.cpp b/Engine/Source/Runtime/Engine/Private/Components/SceneComponent.cpp
--- a/Engine/Source/Runtime/Engine/Private/Components/SceneComponent.cpp
+++ b/Engine/Source/Runtime/Engine/Private/Components/SceneComponent.cpp
@@ -71,9 +71,9 @@ USceneComponent::USceneComponent(const FObjectInitializer& ObjectInitializer /*=
 	, CachedLevelCollection(nullptr)
 {
 	Mobility = EComponentMobility::Movable;
-	RelativeScale3D = FVector(1.0f, 1.0f, 1.0f);
+	SetRelativeScale3D_Direct(FVector(1.0f, 1.0f, 1.0f));
 	// default behavior is visible
-	bVisible = true;
+	SetVisibleFlag(true);
 	bAutoActivate = false;
 	bShouldBeAttached = AttachParent != nullptr;
 }
@@ -505,7 +505,7 @@ void USceneComponent::PostEditChangeProperty(FPropertyChangedEvent& PropertyChan
 	{
 		UpdateAttachedIsEditorOnly(this);
 	}
-	if (PropertyName == GET_MEMBER_NAME_CHECKED(USceneComponent, bVisible))
+	if (PropertyName == USceneComponent::GetVisiblePropertyName())
 	{
 		OnVisibilityChanged();
 	}
@@ -550,17 +550,17 @@ FTransform USceneComponent::CalcNewComponentToWorld_GeneralCase(const FTransform
 	{
 		const FTransform ParentToWorld = Parent->GetSocketTransform(SocketName);
 		FTransform NewCompToWorld = NewRelativeTransform * ParentToWorld;
-		if(bAbsoluteLocation)
+		if(IsUsingAbsoluteLocation())
 		{
 			NewCompToWorld.CopyTranslation(NewRelativeTransform);
 		}
 
-		if(bAbsoluteRotation)
+		if(IsUsingAbsoluteRotation())
 		{
 			NewCompToWorld.CopyRotation(NewRelativeTransform);
 		}
 
-		if(bAbsoluteScale)
+		if(IsUsingAbsoluteScale())
 		{
 			NewCompToWorld.CopyScale3D(NewRelativeTransform);
 		}
@@ -609,7 +609,7 @@ void USceneComponent::UpdateComponentToWorldWithParent(USceneComponent* Parent,F
 	{
 		//QUICK_SCOPE_CYCLE_COUNTER(STAT_USceneComponent_UpdateComponentToWorldWithParent_XForm);
 		// Calculate the new ComponentToWorld transform
-		const FTransform RelativeTransform(RelativeRotationQuat, RelativeLocation, RelativeScale3D);
+		const FTransform RelativeTransform(RelativeRotationQuat, GetRelativeLocation(), GetRelativeScale3D());
 #if ENABLE_NAN_DIAGNOSTIC
 		if (!RelativeTransform.IsValid())
 		{
@@ -656,8 +656,8 @@ void USceneComponent::OnRegister()
 		if (AttachToComponent(GetAttachParent(), FAttachmentTransformRules::KeepRelativeTransform, GetAttachSocketName()) == false)
 		{
 			// Failed to attach, we need to clear AttachParent so we don't think we're actually attached when we're not.
-			AttachParent = nullptr;
-			AttachSocketName = NAME_None;
+			SetAttachParent(nullptr);
+			SetAttachSocketName(NAME_None);
 			bShouldBeAttached = false;
 			bShouldSnapLocationWhenAttached = false;
 			bShouldSnapRotationWhenAttached = false;
@@ -681,7 +681,7 @@ void USceneComponent::OnRegister()
 		SpriteComponent = NewObject<UBillboardComponent>(GetOwner(), NAME_None, RF_Transactional | RF_Transient | RF_TextExportTransient);
 
 		SpriteComponent->Sprite = LoadObject<UTexture2D>(nullptr, TEXT("/Engine/EditorResources/EmptyActor.EmptyActor"));
-		SpriteComponent->RelativeScale3D = FVector(0.5f, 0.5f, 0.5f);
+		SpriteComponent->SetRelativeScale3D_Direct(FVector(0.5f, 0.5f, 0.5f));
 		SpriteComponent->Mobility = EComponentMobility::Movable;
 		SpriteComponent->AlwaysLoadOnClient = false;
 		SpriteComponent->SetIsVisualizationComponent(true);
@@ -972,6 +972,8 @@ void USceneComponent::DestroyComponent(bool bPromoteChildren/*= false*/)
 
 						Index = FMath::Clamp<int32>(Index, 0, AttachSiblings.Num());
 						CachedAttachParent->AttachChildren.Insert(ChildToPromote, Index);
+
+						CachedAttachParent->ModifiedAttachChildren();
 					}
 				}
 			}
@@ -1209,6 +1211,13 @@ void USceneComponent::UpdateBounds()
 
 void USceneComponent::SetRelativeLocationAndRotation(FVector NewLocation, const FQuat& NewRotation, bool bSweep, FHitResult* OutSweepHitResult, ETeleportType Teleport)
 {
+	if (UNLIKELY(NeedsInitialization() || OwnerNeedsInitialization()))
+	{
+		SetRelativeLocation_Direct(NewLocation);
+		SetRelativeRotation_Direct(RelativeRotationCache.QuatToRotator(NewRotation));
+		return;
+	}
+
 	ConditionalUpdateComponentToWorld();
 	
 #if ENABLE_NAN_DIAGNOSTIC
@@ -1241,12 +1250,12 @@ void USceneComponent::SetRelativeLocationAndRotation(FVector NewLocation, const
 // (b) conversions are expensive.
 void USceneComponent::SetRelativeLocationAndRotation(FVector NewLocation, FRotator NewRotation, bool bSweep, FHitResult* OutSweepHitResult, ETeleportType Teleport)
 {
-	if (NewLocation != RelativeLocation)
+	if (NewLocation != GetRelativeLocation())
 	{
 		// It's possible that NewRotation == RelativeRotation, so check the cache for a Rotator->Quat conversion.
 		SetRelativeLocationAndRotation(NewLocation, RelativeRotationCache.RotatorToQuat_ReadOnly(NewRotation), bSweep, OutSweepHitResult, Teleport);
 	}
-	else if (!NewRotation.Equals(RelativeRotation, SCENECOMPONENT_ROTATOR_TOLERANCE))
+	else if (!NewRotation.Equals(GetRelativeRotation(), SCENECOMPONENT_ROTATOR_TOLERANCE))
 	{
 		// We know the rotations are different, don't bother with the cache.
 		SetRelativeLocationAndRotation(NewLocation, NewRotation.Quaternion(), bSweep, OutSweepHitResult, Teleport);
@@ -1255,55 +1264,55 @@ void USceneComponent::SetRelativeLocationAndRotation(FVector NewLocation, FRotat
 
 void USceneComponent::SetRelativeRotationExact(FRotator NewRotation, bool bSweep, FHitResult* OutSweepHitResult, ETeleportType Teleport)
 {
-	if (!NewRotation.Equals(RelativeRotation, SCENECOMPONENT_ROTATOR_TOLERANCE))
+	if (!NewRotation.Equals(GetRelativeRotation(), SCENECOMPONENT_ROTATOR_TOLERANCE))
 	{
 		// We know the rotations are different, don't bother with the cache.
 		const FQuat NewQuat = NewRotation.Quaternion();
-     	SetRelativeLocationAndRotation(RelativeLocation, NewQuat, bSweep, OutSweepHitResult, Teleport);		
+		SetRelativeLocationAndRotation(GetRelativeLocation(), NewQuat, bSweep, OutSweepHitResult, Teleport);
 	}
-	RelativeRotation = NewRotation;
+	SetRelativeRotation_Direct(NewRotation);
 }
 
 void USceneComponent::SetRelativeRotation(FRotator NewRotation, bool bSweep, FHitResult* OutSweepHitResult, ETeleportType Teleport)
 {
-	if (!NewRotation.Equals(RelativeRotation, SCENECOMPONENT_ROTATOR_TOLERANCE))
+	if (!NewRotation.Equals(GetRelativeRotation(), SCENECOMPONENT_ROTATOR_TOLERANCE))
 	{
 		// We know the rotations are different, don't bother with the cache.
-		SetRelativeLocationAndRotation(RelativeLocation, NewRotation.Quaternion(), bSweep, OutSweepHitResult, Teleport);
+		SetRelativeLocationAndRotation(GetRelativeLocation(), NewRotation.Quaternion(), bSweep, OutSweepHitResult, Teleport);
 	}
 }
 
 void USceneComponent::AddRelativeRotation(const FQuat& DeltaRotation, bool bSweep, FHitResult* OutSweepHitResult, ETeleportType Teleport)
 {
-	const FQuat CurRelRotQuat = RelativeRotationCache.RotatorToQuat(RelativeRotation);
+	const FQuat CurRelRotQuat = RelativeRotationCache.RotatorToQuat(GetRelativeRotation());
 	const FQuat NewRelRotQuat = DeltaRotation * CurRelRotQuat;
-	SetRelativeLocationAndRotation(RelativeLocation, NewRelRotQuat, bSweep, OutSweepHitResult, Teleport);
+	SetRelativeLocationAndRotation(GetRelativeLocation(), NewRelRotQuat, bSweep, OutSweepHitResult, Teleport);
 }
 
 void USceneComponent::AddLocalOffset(FVector DeltaLocation, bool bSweep, FHitResult* OutSweepHitResult, ETeleportType Teleport)
 {
-	const FQuat CurRelRotQuat = RelativeRotationCache.RotatorToQuat(RelativeRotation);
+	const FQuat CurRelRotQuat = RelativeRotationCache.RotatorToQuat(GetRelativeRotation());
 	const FVector LocalOffset = CurRelRotQuat.RotateVector(DeltaLocation);
-	SetRelativeLocationAndRotation(RelativeLocation + LocalOffset, CurRelRotQuat, bSweep, OutSweepHitResult, Teleport);
+	SetRelativeLocationAndRotation(GetRelativeLocation() + LocalOffset, CurRelRotQuat, bSweep, OutSweepHitResult, Teleport);
 }
 
 void USceneComponent::AddLocalRotation(FRotator DeltaRotation, bool bSweep, FHitResult* OutSweepHitResult, ETeleportType Teleport)
 {
-	const FQuat CurRelRotQuat = RelativeRotationCache.RotatorToQuat(RelativeRotation);
+	const FQuat CurRelRotQuat = RelativeRotationCache.RotatorToQuat(GetRelativeRotation());
 	const FQuat NewRelRotQuat = CurRelRotQuat * DeltaRotation.Quaternion();
-	SetRelativeLocationAndRotation(RelativeLocation, NewRelRotQuat, bSweep, OutSweepHitResult, Teleport);
+	SetRelativeLocationAndRotation(GetRelativeLocation(), NewRelRotQuat, bSweep, OutSweepHitResult, Teleport);
 }
 
 void USceneComponent::AddLocalRotation(const FQuat& DeltaRotation, bool bSweep, FHitResult* OutSweepHitResult, ETeleportType Teleport)
 {
-	const FQuat CurRelRotQuat = RelativeRotationCache.RotatorToQuat(RelativeRotation);
+	const FQuat CurRelRotQuat = RelativeRotationCache.RotatorToQuat(GetRelativeRotation());
 	const FQuat NewRelRotQuat = CurRelRotQuat * DeltaRotation;
-	SetRelativeLocationAndRotation(RelativeLocation, NewRelRotQuat, bSweep, OutSweepHitResult, Teleport);
+	SetRelativeLocationAndRotation(GetRelativeLocation(), NewRelRotQuat, bSweep, OutSweepHitResult, Teleport);
 }
 
 void USceneComponent::AddLocalTransform(const FTransform& DeltaTransform, bool bSweep, FHitResult* OutSweepHitResult, ETeleportType Teleport)
 {
-	const FTransform RelativeTransform( RelativeRotationCache.RotatorToQuat(RelativeRotation), RelativeLocation, FVector(1,1,1) ); // don't use scaling, so it matches how AddLocalRotation/Offset work
+	const FTransform RelativeTransform(RelativeRotationCache.RotatorToQuat(GetRelativeRotation()), GetRelativeLocation(), FVector(1,1,1) ); // don't use scaling, so it matches how AddLocalRotation/Offset work
 	const FTransform NewRelTransform = DeltaTransform * RelativeTransform;
 	SetRelativeTransform(NewRelTransform, bSweep, OutSweepHitResult, Teleport);
 }
@@ -1336,14 +1345,22 @@ void USceneComponent::AddWorldTransform(const FTransform& DeltaTransform, bool b
 
 void USceneComponent::SetRelativeScale3D(FVector NewScale3D)
 {
-	if( NewScale3D != RelativeScale3D )
+	if (NewScale3D != GetRelativeScale3D())
 	{
 		if (NewScale3D.ContainsNaN())
 		{
 			UE_LOG(LogBlueprint, Warning, TEXT("SetRelativeScale3D : Invalid Scale entered (%s). Resetting to 1.f."), *NewScale3D.ToString());
 			NewScale3D = FVector(1.f);
 		}
-		RelativeScale3D = NewScale3D;		
+
+		SetRelativeScale3D_Direct(NewScale3D);
+
+		if (UNLIKELY(NeedsInitialization() || OwnerNeedsInitialization()))
+		{
+			// If we're in the component or actor constructor, don't do anything else.
+			return;
+		}
+
 		UpdateComponentToWorld();
 		
 		if (IsRegistered())
@@ -1376,7 +1393,7 @@ void USceneComponent::SetRelativeTransform(const FTransform& NewTransform, bool
 
 FTransform USceneComponent::GetRelativeTransform() const
 {
-	const FTransform RelativeTransform( RelativeRotationCache.RotatorToQuat(RelativeRotation), RelativeLocation, RelativeScale3D );
+	const FTransform RelativeTransform(RelativeRotationCache.RotatorToQuat(GetRelativeRotation()), GetRelativeLocation(), GetRelativeScale3D());
 	return RelativeTransform;
 }
 
@@ -1385,7 +1402,7 @@ void USceneComponent::SetWorldLocation(FVector NewLocation, bool bSweep, FHitRes
 	FVector NewRelLocation = NewLocation;
 
 	// If attached to something, transform into local space
-	if (GetAttachParent() != nullptr && !bAbsoluteLocation)
+	if (GetAttachParent() != nullptr && !IsUsingAbsoluteLocation())
 	{
 		FTransform ParentToWorld = GetAttachParent()->GetSocketTransform(GetAttachSocketName());
 		NewRelLocation = ParentToWorld.InverseTransformPosition(NewLocation);
@@ -1399,12 +1416,12 @@ FQuat USceneComponent::GetRelativeRotationFromWorld(const FQuat & NewRotation)
 	FQuat NewRelRotation = NewRotation;
 
 	// If already attached to something, transform into local space
-	if (GetAttachParent() != nullptr && !bAbsoluteRotation)
+	if (GetAttachParent() != nullptr && !IsUsingAbsoluteRotation())
 	{
 		const FTransform  ParentToWorld = GetAttachParent()->GetSocketTransform(GetAttachSocketName());
 		// in order to support mirroring, you'll have to use FTransform.GetRelativeTransform
 		// because negative SCALE should flip the rotation
-		if (FTransform::AnyHasNegativeScale(RelativeScale3D, ParentToWorld.GetScale3D()))
+		if (FTransform::AnyHasNegativeScale(GetRelativeScale3D(), ParentToWorld.GetScale3D()))
 		{
 			FTransform NewTransform = GetComponentTransform();
 			// set new desired rotation
@@ -1448,7 +1465,7 @@ void USceneComponent::SetWorldScale3D(FVector NewScale)
 	FVector NewRelScale = NewScale;
 
 	// If attached to something, transform into local space
-	if(GetAttachParent() != nullptr && !bAbsoluteScale)
+	if(GetAttachParent() != nullptr && !IsUsingAbsoluteScale())
 	{
 		FTransform ParentToWorld = GetAttachParent()->GetSocketTransform(GetAttachSocketName());
 		NewRelScale = NewScale * ParentToWorld.GetSafeScaleReciprocal(ParentToWorld.GetScale3D());
@@ -1466,17 +1483,17 @@ void USceneComponent::SetWorldTransform(const FTransform& NewTransform, bool bSw
 		FTransform RelativeTM = NewTransform.GetRelativeTransform(ParentToWorld);
 
 		// Absolute location, rotation, and scale use the world transform directly.
-		if (bAbsoluteLocation)
+		if (IsUsingAbsoluteLocation())
 		{
 			RelativeTM.CopyTranslation(NewTransform);
 		}
 
-		if (bAbsoluteRotation)
+		if (IsUsingAbsoluteRotation())
 		{
 			RelativeTM.CopyRotation(NewTransform);
 		}
 
-		if (bAbsoluteScale)
+		if (IsUsingAbsoluteScale())
 		{
 			RelativeTM.CopyScale3D(NewTransform);
 		}
@@ -1510,12 +1527,12 @@ void USceneComponent::SetWorldLocationAndRotation(FVector NewLocation, const FQu
 	{
 		FTransform ParentToWorld = GetAttachParent()->GetSocketTransform(GetAttachSocketName());
 
-		if (!bAbsoluteLocation)
+		if (!IsUsingAbsoluteLocation())
 		{
 			NewLocation = ParentToWorld.InverseTransformPosition(NewLocation);
 		}
 
-		if (!bAbsoluteRotation)
+		if (!IsUsingAbsoluteRotation())
 		{
 			// Quat multiplication works reverse way, make sure you do Parent(-1) * World = Local, not World*Parent(-) = Local (the way matrix does)
 			FQuat NewRelQuat = ParentToWorld.GetRotation().Inverse() * NewRotation;
@@ -1533,30 +1550,30 @@ void USceneComponent::SetWorldLocationAndRotationNoPhysics(const FVector& NewLoc
 	{
 		const FTransform ParentToWorld = GetAttachParent()->GetSocketTransform(GetAttachSocketName());
 
-		if(bAbsoluteLocation)
+		if(IsUsingAbsoluteLocation())
 		{
-			RelativeLocation = NewLocation;
+			SetRelativeLocation_Direct(NewLocation);
 		}
 		else
 		{
-			RelativeLocation = ParentToWorld.InverseTransformPosition(NewLocation);
+			SetRelativeLocation_Direct(ParentToWorld.InverseTransformPosition(NewLocation));
 		}
 
-		if(bAbsoluteRotation)
+		if(IsUsingAbsoluteRotation())
 		{
-			RelativeRotation = NewRotation;
+			SetRelativeRotation_Direct(NewRotation);
 		}
 		else
 		{
 			// Quat multiplication works reverse way, make sure you do Parent(-1) * World = Local, not World*Parent(-) = Local (the way matrix does)
 			const FQuat NewRelQuat = ParentToWorld.GetRotation().Inverse() * NewRotation.Quaternion();
-			RelativeRotation = RelativeRotationCache.QuatToRotator(NewRelQuat);
+			SetRelativeRotation_Direct(RelativeRotationCache.QuatToRotator(NewRelQuat));
 		}
 	}
 	else
 	{
-		RelativeLocation = NewLocation;
-		RelativeRotation = NewRotation;
+		SetRelativeLocation_Direct(NewLocation);
+		SetRelativeRotation_Direct(NewRotation);
 	}
 
 	UpdateComponentToWorld(EUpdateTransformFlags::SkipPhysicsUpdate);
@@ -1564,9 +1581,9 @@ void USceneComponent::SetWorldLocationAndRotationNoPhysics(const FVector& NewLoc
 
 void USceneComponent::SetAbsolute(bool bNewAbsoluteLocation, bool bNewAbsoluteRotation, bool bNewAbsoluteScale)
 {
-	bAbsoluteLocation = bNewAbsoluteLocation;
-	bAbsoluteRotation = bNewAbsoluteRotation;
-	bAbsoluteScale = bNewAbsoluteScale;
+	SetUsingAbsoluteLocation(bNewAbsoluteLocation);
+	SetUsingAbsoluteRotation(bNewAbsoluteRotation);
+	SetUsingAbsoluteScale(bNewAbsoluteScale);
 
 	UpdateComponentToWorld();
 }
@@ -1713,8 +1730,8 @@ void USceneComponent::SetupAttachment(class USceneComponent* InParent, FName InS
 			{
 				if (ensureMsgf(AttachParent == nullptr || !AttachParent->AttachChildren.Contains(this), TEXT("SetupAttachment cannot be used once a component has already had AttachTo used to connect it to a parent.")))
 				{
-					AttachParent = InParent;
-					AttachSocketName = InSocketName;
+					SetAttachParent(InParent);
+					SetAttachSocketName(InSocketName);
 					bShouldBeAttached = AttachParent != nullptr;
 				}
 			}
@@ -1917,18 +1934,19 @@ bool USceneComponent::AttachToComponent(USceneComponent* Parent, const FAttachme
 						//User tried to attach but physically based so detach. However, if they provided relative coordinates we should still get the correct position
 						if (AttachmentRules.LocationRule == EAttachmentRule::KeepRelative || AttachmentRules.RotationRule == EAttachmentRule::KeepRelative || AttachmentRules.ScaleRule == EAttachmentRule::KeepRelative)
 						{
-							UpdateComponentToWorldWithParent(Parent, SocketName, EUpdateTransformFlags::None, RelativeRotationCache.RotatorToQuat(RelativeRotation));
+							UpdateComponentToWorldWithParent(Parent, SocketName, EUpdateTransformFlags::None, RelativeRotationCache.RotatorToQuat(GetRelativeRotation()));
+
 							if (AttachmentRules.LocationRule == EAttachmentRule::KeepRelative)
 							{
-								RelativeLocation = GetComponentLocation();
+								SetRelativeLocation_Direct(GetComponentLocation());
 							}
 							if (AttachmentRules.RotationRule == EAttachmentRule::KeepRelative)
 							{
-								RelativeRotation = GetComponentRotation();
+								SetRelativeRotation_Direct(GetComponentRotation());
 							}
 							if (AttachmentRules.ScaleRule == EAttachmentRule::KeepRelative)
 							{
-								RelativeScale3D = GetComponentScale();
+								SetRelativeScale3D_Direct(GetComponentScale());
 							}
 							if (IsRegistered())
 							{
@@ -1946,8 +1964,8 @@ bool USceneComponent::AttachToComponent(USceneComponent* Parent, const FAttachme
 		PrimaryComponentTick.AddPrerequisite(Parent, Parent->PrimaryComponentTick); // force us to tick after the parent does
 
 		// Save pointer from child to parent
-		AttachParent = Parent;
-		AttachSocketName = SocketName;
+		SetAttachParent(Parent);
+		SetAttachSocketName(SocketName);
 		bShouldBeAttached = AttachParent != nullptr;
 
 		bShouldSnapLocationWhenAttached = AttachmentRules.LocationRule == EAttachmentRule::SnapToTarget;
@@ -1964,6 +1982,8 @@ bool USceneComponent::AttachToComponent(USceneComponent* Parent, const FAttachme
 		{
 			Parent->AttachChildren.Add(this);
 		}
+
+		Parent->ModifiedAttachChildren();
 		AddToCluster(Parent, true);
 
 		if (Parent->IsNetSimulating() && !IsNetSimulating())
@@ -1995,17 +2015,17 @@ bool USceneComponent::AttachToComponent(USceneComponent* Parent, const FAttachme
 			// dont do anything, keep relative position the same
 			break;
 		case EAttachmentRule::KeepWorld:
-			if (bAbsoluteLocation)
+			if (IsUsingAbsoluteLocation())
 			{
-				RelativeLocation = GetComponentTransform().GetTranslation();
+				SetRelativeLocation_Direct(GetComponentTransform().GetTranslation());
 			}
 			else
 			{
-				RelativeLocation = RelativeTM.GetTranslation();
+				SetRelativeLocation_Direct(RelativeTM.GetTranslation());
 			}
 			break;
 		case EAttachmentRule::SnapToTarget:
-			RelativeLocation = FVector::ZeroVector;
+			SetRelativeLocation_Direct(FVector::ZeroVector);
 			break;
 		}
 
@@ -2015,17 +2035,17 @@ bool USceneComponent::AttachToComponent(USceneComponent* Parent, const FAttachme
 			// dont do anything, keep relative rotation the same
 			break;
 		case EAttachmentRule::KeepWorld:
-			if (bAbsoluteRotation)
+			if (IsUsingAbsoluteRotation())
 			{
-				RelativeRotation = GetComponentRotation();
+				SetRelativeRotation_Direct(GetComponentRotation());
 			}
 			else
 			{
-				RelativeRotation = RelativeRotationCache.QuatToRotator(RelativeTM.GetRotation());
+				SetRelativeRotation_Direct(RelativeRotationCache.QuatToRotator(RelativeTM.GetRotation()));
 			}
 			break;
 		case EAttachmentRule::SnapToTarget:
-			RelativeRotation = FRotator::ZeroRotator;
+			SetRelativeRotation_Direct(FRotator::ZeroRotator);
 			break;
 		}
 
@@ -2035,17 +2055,17 @@ bool USceneComponent::AttachToComponent(USceneComponent* Parent, const FAttachme
 			// dont do anything, keep relative scale the same
 			break;
 		case EAttachmentRule::KeepWorld:
-			if (bAbsoluteScale)
+			if (IsUsingAbsoluteScale())
 			{
-				RelativeScale3D = GetComponentTransform().GetScale3D();
+				SetRelativeScale3D_Direct(GetComponentTransform().GetScale3D());
 			}
 			else
 			{
-				RelativeScale3D = RelativeTM.GetScale3D();
+				SetRelativeScale3D_Direct(RelativeTM.GetScale3D());
 			}
 			break;
 		case EAttachmentRule::SnapToTarget:
-			RelativeScale3D = FVector(1.0f, 1.0f, 1.0f);
+			SetRelativeScale3D_Direct(FVector(1.0f, 1.0f, 1.0f));
 			break;
 		}
 
@@ -2140,6 +2160,7 @@ void USceneComponent::DetachFromComponent(const FDetachmentTransformRules& Detac
 		GetAttachParent()->AttachChildren.Remove(this);
 		GetAttachParent()->ClientAttachedChildren.Remove(this);
 		GetAttachParent()->OnChildDetached(this);
+		GetAttachParent()->ModifiedAttachChildren();
 
 #if WITH_EDITOR
 		if(GEngine)
@@ -2150,8 +2171,8 @@ void USceneComponent::DetachFromComponent(const FDetachmentTransformRules& Detac
 			}
 		}
 #endif
-		AttachParent = nullptr;
-		AttachSocketName = NAME_None;
+		SetAttachParent(nullptr);
+		SetAttachSocketName(NAME_None);
 		bShouldBeAttached = 0;
 
 		bShouldSnapLocationWhenAttached = false;
@@ -2165,7 +2186,7 @@ void USceneComponent::DetachFromComponent(const FDetachmentTransformRules& Detac
 		case EDetachmentRule::KeepRelative:
 			break;
 		case EDetachmentRule::KeepWorld:
-			RelativeLocation = GetComponentTransform().GetTranslation(); // or GetComponentLocation, but worried about custom location...
+			SetRelativeLocation_Direct(GetComponentTransform().GetTranslation()); // or GetComponentLocation, but worried about custom location...
 			break;
 		}
 
@@ -2174,7 +2195,7 @@ void USceneComponent::DetachFromComponent(const FDetachmentTransformRules& Detac
 		case EDetachmentRule::KeepRelative:
 			break;
 		case EDetachmentRule::KeepWorld:
-			RelativeRotation = GetComponentRotation();
+			SetRelativeRotation_Direct(GetComponentRotation());
 			break;
 		}
 
@@ -2183,7 +2204,7 @@ void USceneComponent::DetachFromComponent(const FDetachmentTransformRules& Detac
 		case EDetachmentRule::KeepRelative:
 			break;
 		case EDetachmentRule::KeepWorld:
-			RelativeScale3D = GetComponentScale();
+			SetRelativeScale3D_Direct(GetComponentScale());
 			break;
 		}
 
@@ -2236,7 +2257,7 @@ FSceneComponentInstanceData::FSceneComponentInstanceData(const USceneComponent*
 		USceneComponent* SceneComponent = AttachedChildren[i];
 		if (SceneComponent && SceneComponent->GetOwner() == SourceOwner && !SceneComponent->IsCreatedByConstructionScript() && !SceneComponent->HasAnyFlags(RF_DefaultSubObject))
 		{
-			AttachedInstanceComponents.Emplace(SceneComponent, FTransform(SceneComponent->RelativeRotation, SceneComponent->RelativeLocation, SceneComponent->RelativeScale3D));
+			AttachedInstanceComponents.Emplace(SceneComponent, FTransform(SceneComponent->GetRelativeRotation(), SceneComponent->GetRelativeLocation(), SceneComponent->GetRelativeScale3D()));
 		}
 	}
 }
@@ -2264,9 +2285,9 @@ void FSceneComponentInstanceData::ApplyToComponent(UActorComponent* Component, c
 		// and so the rebuilt component should not take back attachment ownership
 		if (ChildComponent && (ChildComponent->GetAttachParent() == nullptr || ChildComponent->GetAttachParent()->IsPendingKill()))
 		{
-			ChildComponent->RelativeLocation = ChildComponentPair.Value.GetLocation();
-			ChildComponent->RelativeRotation = ChildComponentPair.Value.GetRotation().Rotator();
-			ChildComponent->RelativeScale3D = ChildComponentPair.Value.GetScale3D();
+			ChildComponent->SetRelativeLocation_Direct(ChildComponentPair.Value.GetLocation());
+			ChildComponent->SetRelativeRotation_Direct(ChildComponentPair.Value.GetRotation().Rotator());
+			ChildComponent->SetRelativeScale3D_Direct(ChildComponentPair.Value.GetScale3D());
 			ChildComponent->AttachToComponent(SceneComponent, FAttachmentTransformRules::KeepRelativeTransform);
 		}
 	}
@@ -2343,7 +2364,7 @@ void USceneComponent::UpdateChildTransforms(EUpdateTransformFlags UpdateTransfor
 					}
 
 					// Don't update the child if it uses a completely absolute (world-relative) scheme.
-					if (ChildComp->bAbsoluteLocation && ChildComp->bAbsoluteRotation && ChildComp->bAbsoluteScale)
+					if (ChildComp->IsUsingAbsoluteLocation() && ChildComp->IsUsingAbsoluteRotation() && ChildComp->IsUsingAbsoluteScale())
 					{
 						continue;
 					}
@@ -2359,9 +2380,7 @@ void USceneComponent::PostInterpChange(UProperty* PropertyThatChanged)
 {
 	Super::PostInterpChange(PropertyThatChanged);
 
-	static FName NAME_RelativeScale3D(GET_MEMBER_NAME_CHECKED(USceneComponent, RelativeScale3D));
-
-	if (PropertyThatChanged->GetFName() == NAME_RelativeScale3D)
+	if (PropertyThatChanged->GetFName() == GetRelativeScale3DPropertyName())
 	{
 		UpdateComponentToWorld();
 	}
@@ -2696,29 +2715,29 @@ bool USceneComponent::InternalSetWorldLocationAndRotation(FVector NewLocation, c
 		FTransform const ParentToWorld = GetAttachParent()->GetSocketTransform(GetAttachSocketName());
 		// in order to support mirroring, you'll have to use FTransform.GetrelativeTransform
 		// because negative scale should flip the rotation
-		if (FTransform::AnyHasNegativeScale(RelativeScale3D, ParentToWorld.GetScale3D()))
+		if (FTransform::AnyHasNegativeScale(GetRelativeScale3D(), ParentToWorld.GetScale3D()))
 		{
-			FTransform const WorldTransform = FTransform(RotationQuat, NewLocation, RelativeScale3D*ParentToWorld.GetScale3D());
+			FTransform const WorldTransform = FTransform(RotationQuat, NewLocation, GetRelativeScale3D() * ParentToWorld.GetScale3D());
 			FTransform const RelativeTransform = WorldTransform.GetRelativeTransform(ParentToWorld);
 
-			if (!bAbsoluteLocation)
+			if (!IsUsingAbsoluteLocation())
 			{
 				NewLocation = RelativeTransform.GetLocation();
 			}
 
-			if (!bAbsoluteRotation)
+			if (!IsUsingAbsoluteRotation())
 			{
 				NewRotationQuat = RelativeTransform.GetRotation();
 			}
 		}
 		else
 		{
-			if (!bAbsoluteLocation)
+			if (!IsUsingAbsoluteLocation())
 			{
 				NewLocation = ParentToWorld.InverseTransformPosition(NewLocation);
 			}
 
-			if (!bAbsoluteRotation)
+			if (!IsUsingAbsoluteRotation())
 			{
 				// Quat multiplication works reverse way, make sure you do Parent(-1) * World = Local, not World*Parent(-) = Local (the way matrix does)
 				NewRotationQuat = ParentToWorld.GetRotation().Inverse() * NewRotationQuat;
@@ -2727,11 +2746,11 @@ bool USceneComponent::InternalSetWorldLocationAndRotation(FVector NewLocation, c
 	}
 
 	const FRotator NewRelativeRotation = RelativeRotationCache.QuatToRotator_ReadOnly(NewRotationQuat);
-	bool bDiffLocation = !NewLocation.Equals(RelativeLocation);
-	bool bDiffRotation = !NewRelativeRotation.Equals(RelativeRotation);
+	bool bDiffLocation = !NewLocation.Equals(GetRelativeLocation());
+	bool bDiffRotation = !NewRelativeRotation.Equals(GetRelativeRotation());
 	if (bDiffLocation || bDiffRotation)
 	{
-		RelativeLocation = NewLocation;
+		SetRelativeLocation_Direct(NewLocation);
 
 		// Here it is important to compute the quaternion from the rotator and not the opposite.
 		// In some cases, similar quaternions generate the same rotator, which create issues.
@@ -2742,7 +2761,7 @@ bool USceneComponent::InternalSetWorldLocationAndRotation(FVector NewLocation, c
 		// we must enforce the quaternion to come from the rotator (as in load)
 		if (bDiffRotation)
 		{
-			RelativeRotation = NewRelativeRotation;
+			SetRelativeRotation_Direct(NewRelativeRotation);
 			RelativeRotationCache.RotatorToQuat(NewRelativeRotation);
 		}
 
@@ -2827,7 +2846,7 @@ bool USceneComponent::MoveComponent(const FVector& Delta, const FRotator& NewRot
 {
 	if (GetAttachParent() == nullptr)
 	{
-		if (Delta.IsZero() && NewRotation.Equals(RelativeRotation, SCENECOMPONENT_ROTATOR_TOLERANCE))
+		if (Delta.IsZero() && NewRotation.Equals(GetRelativeRotation(), SCENECOMPONENT_ROTATOR_TOLERANCE))
 		{
 			if (Hit)
 			{
@@ -2891,7 +2910,7 @@ bool USceneComponent::MoveComponentImpl(const FVector& Delta, const FQuat& NewRo
 bool USceneComponent::IsVisibleInEditor() const
 {
 	// in editor, we only check bVisible
-	return bVisible;
+	return GetVisibleFlag();
 }
 
 bool USceneComponent::ShouldRender() const
@@ -2926,8 +2945,8 @@ bool USceneComponent::ShouldRender() const
 #endif
 	const bool bInGameWorld = World && World->UsesGameHiddenFlags();
 
-	const bool bShowInGame = IsVisible() && (!Owner || !Owner->bHidden);
-	return ((bInGameWorld && bShowInGame) || (!bInGameWorld && bShowInEditor)) && bVisible == true;
+	const bool bShowInGame = IsVisible() && (!Owner || !Owner->IsHidden());
+	return ((bInGameWorld && bShowInGame) || (!bInGameWorld && bShowInEditor)) && GetVisibleFlag() == true;
 }
 
 bool USceneComponent::CanEverRender() const
@@ -2954,7 +2973,7 @@ bool USceneComponent::CanEverRender() const
 	UWorld *World = GetWorld();
 	const bool bInGameWorld = World && World->UsesGameHiddenFlags();
 
-	const bool bShowInGame = (!Owner || !Owner->bHidden);
+	const bool bShowInGame = (!Owner || !Owner->IsHidden());
 	return ((bInGameWorld && bShowInGame) || (!bInGameWorld && bShowInEditor));
 }
 
@@ -2972,7 +2991,7 @@ bool USceneComponent::IsVisible() const
 		return false;
 	}
 	
-	return ( bVisible && (!CachedLevelCollection || CachedLevelCollection->IsVisible()) ); 
+	return (GetVisibleFlag() && (!CachedLevelCollection || CachedLevelCollection->IsVisible())); 
 }
 
 void USceneComponent::OnVisibilityChanged()
@@ -2983,10 +3002,10 @@ void USceneComponent::OnVisibilityChanged()
 void USceneComponent::SetVisibility(const bool bNewVisibility, const USceneComponent::EVisibilityPropagation PropagateToChildren)
 {
 	bool bRecurseChildren = (PropagateToChildren == EVisibilityPropagation::Propagate);
-	if ( bNewVisibility != bVisible )
+	if ( bNewVisibility != GetVisibleFlag() )
 	{
 		bRecurseChildren = bRecurseChildren || (PropagateToChildren == EVisibilityPropagation::DirtyOnly);
-		bVisible = bNewVisibility;
+		SetVisibleFlag(bNewVisibility);
 		OnVisibilityChanged();
 	}
 
@@ -3072,21 +3091,19 @@ void USceneComponent::ApplyWorldOffset(const FVector& InOffset, bool bWorldShift
 	// Calculate current ComponentToWorld transform
 	// We do this because at level load/duplication ComponentToWorld is uninitialized
 	{
-		const FTransform RelativeTransform(RelativeRotationCache.RotatorToQuat(RelativeRotation), RelativeLocation, RelativeScale3D);
-		ComponentToWorld = CalcNewComponentToWorld(RelativeTransform);
+		ComponentToWorld = CalcNewComponentToWorld(GetRelativeTransform());
 	}
 
 	// Update bounds
 	Bounds.Origin+= InOffset;
 
 	// Update component location
-	if (GetAttachParent() == nullptr || bAbsoluteLocation)
+	if (GetAttachParent() == nullptr || IsUsingAbsoluteLocation())
 	{
-		RelativeLocation = GetComponentLocation() + InOffset;
+		SetRelativeLocation_Direct(GetComponentLocation() + InOffset);
 		
 		// Calculate the new ComponentToWorld transform
-		const FTransform RelativeTransform(RelativeRotationCache.RotatorToQuat(RelativeRotation), RelativeLocation, RelativeScale3D);
-		ComponentToWorld = CalcNewComponentToWorld(RelativeTransform);
+		ComponentToWorld = CalcNewComponentToWorld(GetRelativeTransform());
 	}
 
 	// Physics move is skipped if physics state is not created or physics scene supports origin shifting
@@ -3140,6 +3157,10 @@ void USceneComponent::OnRep_AttachSocketName()
 
 void USceneComponent::OnRep_AttachChildren()
 {
+	// Don't worry about marking AttachChildren dirty here.
+	// This should only be called after we've received a value from the network,
+	// and that will handle dirtying for us.
+
 	// Because replication of AttachChildren is not atomic with AttachParent of the corresponding component it is
 	// entirely possible to get duplicates in the AttachChildren array.  So we have to extract them and the later entry
 	// is always the duplicate
@@ -3209,7 +3230,7 @@ void USceneComponent::PostNetReceive()
 		if (ParentRoot != this)
 		{
 			bNetUpdateAttachment = true;
-			AttachParent = ParentRoot;
+			SetAttachParent(ParentRoot);
 		}
 	}
 }
@@ -3224,11 +3245,11 @@ void USceneComponent::PostRepNotifies()
 		// Note: This is a local fix for JIRA UE-43355.
 		if (bShouldSnapLocationWhenAttached)
 		{
-			RelativeLocation = FVector::ZeroVector;
+			SetRelativeLocation_Direct(FVector::ZeroVector);
 		}
 		if (bShouldSnapRotationWhenAttached)
 		{
-			RelativeRotation = FRotator::ZeroRotator;
+			SetRelativeRotation_Direct(FRotator::ZeroRotator);
 		}
 
 		// Check if this is a detach
@@ -3331,6 +3352,39 @@ bool USceneComponent::CanEditChange( const UProperty* Property ) const
  * FScopedPreventAttachedComponentMove implementation
  */
 
+FScopedPreventAttachedComponentMove::FScopedPreventAttachedComponentMove(USceneComponent* Component)
+	: Owner(Component)
+{
+	if (Owner)
+	{
+		// Save old flags
+		bSavedAbsoluteLocation = Owner->IsUsingAbsoluteLocation();
+		bSavedAbsoluteRotation = Owner->IsUsingAbsoluteRotation();
+		bSavedAbsoluteScale = Owner->IsUsingAbsoluteScale();
+		bSavedNonAbsoluteComponent = !(bSavedAbsoluteLocation && bSavedAbsoluteRotation && bSavedAbsoluteScale);
+
+		// Use absolute (stay in world space no matter what parent does)
+		Owner->bAbsoluteLocation = true;
+		Owner->bAbsoluteRotation = true;
+		Owner->bAbsoluteScale = true;
+
+		if (bSavedNonAbsoluteComponent && Owner->GetAttachParent())
+		{
+			// Make RelativeLocation etc relative to the world.
+			Component->ConditionalUpdateComponentToWorld();
+			Owner->RelativeLocation = Owner->GetComponentLocation();
+			Owner->RelativeRotation = Owner->GetComponentRotation();
+			Owner->RelativeScale3D = Owner->GetComponentScale();
+		}
+	}
+	else
+	{
+		bSavedAbsoluteLocation = false;
+		bSavedAbsoluteRotation = false;
+		bSavedAbsoluteScale = false;
+	}
+}
+
 FScopedPreventAttachedComponentMove::~FScopedPreventAttachedComponentMove()
 {
 	if (Owner)
@@ -3347,15 +3401,15 @@ FScopedPreventAttachedComponentMove::~FScopedPreventAttachedComponentMove()
 
 			if (!bSavedAbsoluteLocation)
 			{
-				Owner->RelativeLocation = ChildRelativeTM.GetTranslation();
+				Owner->SetRelativeLocation_Direct(ChildRelativeTM.GetTranslation());
 			}
 			if (!bSavedAbsoluteRotation)
 			{
-				Owner->RelativeRotation = Owner->RelativeRotationCache.QuatToRotator(ChildRelativeTM.GetRotation());
+				Owner->SetRelativeRotation_Direct(Owner->RelativeRotationCache.QuatToRotator(ChildRelativeTM.GetRotation()));
 			}
 			if (!bSavedAbsoluteScale)
 			{
-				Owner->RelativeScale3D = ChildRelativeTM.GetScale3D();
+				Owner->SetRelativeScale3D_Direct(ChildRelativeTM.GetScale3D());
 			}
 		}
 	}
@@ -3382,9 +3436,9 @@ FScopedMovementUpdate::FScopedMovementUpdate( class USceneComponent* Component,
 	{
 		OuterDeferredScope = Component->GetCurrentScopedMovement();
 		InitialTransform = Component->GetComponentToWorld();
-		InitialRelativeLocation = Component->RelativeLocation;
-		InitialRelativeRotation = Component->RelativeRotation;
-		InitialRelativeScale = Component->RelativeScale3D;
+		InitialRelativeLocation = Component->GetRelativeLocation();
+		InitialRelativeRotation = Component->GetRelativeRotation();
+		InitialRelativeScale = Component->GetRelativeScale3D();
 
 		if (ScopeBehavior == EScopedUpdate::ImmediateUpdates)
 		{
@@ -3765,4 +3819,78 @@ void USceneComponent::K2_AddWorldTransform(const FTransform& DeltaTransform, boo
 	AddWorldTransform(DeltaTransform, bSweep, (bSweep ? &SweepHitResult : nullptr), TeleportFlagToEnum(bTeleport));
 }
 
+void USceneComponent::SetVisibleFlag(const bool bNewVisible)
+{
+	bVisible = bNewVisible;
+	// MARK_PROPERTY_DIRTY_FROM_NAME(USceneComponent, bVisible, this);
+}
+
+FVector& USceneComponent::GetRelativeLocation_DirectMutable()
+{
+	// MARK_PROPERTY_DIRTY_FROM_NAME(USceneComponent, RelativeLocation, this);
+	return RelativeLocation;
+}
+
+void USceneComponent::SetRelativeLocation_Direct(const FVector NewRelativeLocation)
+{
+	GetRelativeLocation_DirectMutable() = NewRelativeLocation;
+}
+
+FRotator& USceneComponent::GetRelativeRotation_DirectMutable()
+{
+	// MARK_PROPERTY_DIRTY_FROM_NAME(USceneComponent, RelativeRotation, this);
+	return RelativeRotation;
+}
+
+void USceneComponent::SetRelativeRotation_Direct(const FRotator NewRelativeRotation)
+{
+	GetRelativeRotation_DirectMutable() = NewRelativeRotation;
+}
+
+FVector& USceneComponent::GetRelativeScale3D_DirectMutable()
+{
+	// MARK_PROPERTY_DIRTY_FROM_NAME(USceneComponent, RelativeScale3D, this);
+	return RelativeScale3D;
+}
+
+void USceneComponent::SetRelativeScale3D_Direct(const FVector NewRelativeScale3D)
+{
+	GetRelativeScale3D_DirectMutable() = NewRelativeScale3D;
+}
+
+void USceneComponent::SetUsingAbsoluteLocation(bool bInAbsoluteLocation)
+{
+	bAbsoluteLocation = bInAbsoluteLocation;
+	// MARK_PROPERTY_DIRTY_FROM_NAME(USceneComponent, bAbsoluteLocation, this);
+}
+
+void USceneComponent::SetUsingAbsoluteRotation(bool bInAbsoluteRotation)
+{
+	bAbsoluteRotation = bInAbsoluteRotation;
+	// MARK_PROPERTY_DIRTY_FROM_NAME(USceneComponent, bAbsoluteRotation, this);
+}
+
+void USceneComponent::SetUsingAbsoluteScale(bool bInAbsoluteScale)
+{
+	bAbsoluteScale = bInAbsoluteScale;
+	// MARK_PROPERTY_DIRTY_FROM_NAME(USceneComponent, bAbsoluteScale, this);
+}
+
+void USceneComponent::SetAttachParent(USceneComponent* NewAttachParent)
+{
+	AttachParent = NewAttachParent;
+	// MARK_PROPERTY_DIRTY_FROM_NAME(USceneComponent, AttachParent, this);
+}
+
+void USceneComponent::SetAttachSocketName(FName NewSocketName)
+{
+	AttachSocketName = NewSocketName;
+	// MARK_PROPERTY_DIRTY_FROM_NAME(USceneComponent, AttachSocketName, this);
+}
+
+void USceneComponent::ModifiedAttachChildren()
+{
+	// MARK_PROPERTY_DIRTY_FROM_NAME(USceneComponent, AttachChildren, this);
+}
+
 #undef LOCTEXT_NAMESPACE
