[NO NON-PII COMMIT MESSAGE]

diff --git a/Engine/Source/Runtime/Engine/Classes/Sound/SoundNode.h b/Engine/Source/Runtime/Engine/Classes/Sound/SoundNode.h
--- a/Engine/Source/Runtime/Engine/Classes/Sound/SoundNode.h
+++ b/Engine/Source/Runtime/Engine/Classes/Sound/SoundNode.h
@@ -198,6 +198,6 @@ public:
 	 *
 	 * @param bRecurse when true, this will cause all children of child nodes to be primed as well.
 	 */
-	void PrimeChildWavePlayers(bool bRecurse);
+	virtual void PrimeChildWavePlayers(bool bRecurse);
 };
 
diff --git a/Engine/Source/Runtime/Engine/Classes/Sound/SoundNodeQualityLevel.h b/Engine/Source/Runtime/Engine/Classes/Sound/SoundNodeQualityLevel.h
--- a/Engine/Source/Runtime/Engine/Classes/Sound/SoundNodeQualityLevel.h
+++ b/Engine/Source/Runtime/Engine/Classes/Sound/SoundNodeQualityLevel.h
@@ -34,6 +34,8 @@ public:
 	virtual void ParseNodes( FAudioDevice* AudioDevice, const UPTRINT NodeWaveInstanceHash, FActiveSound& ActiveSound, const FSoundParseParameters& ParseParams, TArray<FWaveInstance*>& WaveInstances ) override;
 	virtual int32 GetMaxChildNodes() const override;
 	virtual int32 GetMinChildNodes() const override;
+	virtual void PrimeChildWavePlayers(bool bRecurse) override;
+
 #if WITH_EDITOR
 	virtual FText GetInputPinName(int32 PinIndex) const override;
 #endif
@@ -42,5 +44,6 @@ public:
 #if WITH_EDITOR
 	void ReconcileNode(bool bReconstructNode);
 #endif
+
 };
 
diff --git a/Engine/Source/Runtime/Engine/Private/SoundNodeQualityLevel.cpp b/Engine/Source/Runtime/Engine/Private/SoundNodeQualityLevel.cpp
--- a/Engine/Source/Runtime/Engine/Private/SoundNodeQualityLevel.cpp
+++ b/Engine/Source/Runtime/Engine/Private/SoundNodeQualityLevel.cpp
@@ -47,6 +47,23 @@ FText USoundNodeQualityLevel::GetInputPinName(int32 PinIndex) const
 }
 #endif
 
+void USoundNodeQualityLevel::PrimeChildWavePlayers(bool bRecurse)
+{
+	// If we're able to retrieve a valid cached quality level for this sound cue,
+	// only prime that quality level.
+	int32 QualityLevel = USoundCue::GetCachedQualityLevel();
+
+	if (GIsEditor && QualityLevel < 0)
+	{
+		QualityLevel = GetDefault<ULevelEditorPlaySettings>()->PlayInEditorSoundQualityLevel;
+	}
+
+	if (ChildNodes.IsValidIndex(QualityLevel) && ChildNodes[QualityLevel])
+	{
+		ChildNodes[QualityLevel]->PrimeChildWavePlayers(bRecurse);
+	}
+}
+
 int32 USoundNodeQualityLevel::GetMaxChildNodes() const
 {
 	return GetDefault<UAudioSettings>()->QualityLevels.Num();
