Fix compile for non-editor builds. #jira none

diff --git a/Engine/Source/Runtime/Engine/Private/SoundNodeQualityLevel.cpp b/Engine/Source/Runtime/Engine/Private/SoundNodeQualityLevel.cpp
--- a/Engine/Source/Runtime/Engine/Private/SoundNodeQualityLevel.cpp
+++ b/Engine/Source/Runtime/Engine/Private/SoundNodeQualityLevel.cpp
@@ -53,10 +53,12 @@ void USoundNodeQualityLevel::PrimeChildWavePlayers(bool bRecurse)
 	// only prime that quality level.
 	int32 QualityLevel = USoundCue::GetCachedQualityLevel();
 
+#if WITH_EDITOR
 	if (GIsEditor && QualityLevel < 0)
 	{
 		QualityLevel = GetDefault<ULevelEditorPlaySettings>()->PlayInEditorSoundQualityLevel;
 	}
+#endif
 
 	if (ChildNodes.IsValidIndex(QualityLevel) && ChildNodes[QualityLevel])
 	{
