fix for trying to play a sound that failed to cook in the editor. Also includes fix for issue where cook settings for the editor could be used in a cook on the fly server.

diff --git a/Engine/Source/Runtime/Engine/Private/AudioDerivedData.cpp b/Engine/Source/Runtime/Engine/Private/AudioDerivedData.cpp
--- a/Engine/Source/Runtime/Engine/Private/AudioDerivedData.cpp
+++ b/Engine/Source/Runtime/Engine/Private/AudioDerivedData.cpp
@@ -140,6 +140,28 @@ static void GetStreamedAudioDerivedDataKey(
 	GetStreamedAudioDerivedDataKeyFromSuffix(KeySuffix, OutKey);
 }
 
+static ITargetPlatform* GetRunningTargetPlatform(ITargetPlatformManagerModule* TPM)
+{
+	ITargetPlatform* CurrentPlatform = NULL;
+	const TArray<ITargetPlatform*>& Platforms = TPM->GetActiveTargetPlatforms();
+
+	check(Platforms.Num());
+
+	CurrentPlatform = Platforms[0];
+
+	for (int32 Index = 1; Index < Platforms.Num(); Index++)
+	{
+		if (Platforms[Index]->IsRunningPlatform())
+		{
+			CurrentPlatform = Platforms[Index];
+			break;
+		}
+	}
+
+	check(CurrentPlatform != NULL);
+	return CurrentPlatform;
+}
+
 /**
  * Gets Wave format for a SoundWave on the current running platform
  * @param SoundWave - The SoundWave to get format for.
@@ -150,23 +172,8 @@ static FName GetWaveFormatForRunningPlatform(USoundWave& SoundWave)
 	ITargetPlatformManagerModule* TPM = GetTargetPlatformManager();
 	if (TPM)
 	{
-		ITargetPlatform* CurrentPlatform = NULL;
-		const TArray<ITargetPlatform*>& Platforms = TPM->GetActiveTargetPlatforms();
-
-		check(Platforms.Num());
+		ITargetPlatform* CurrentPlatform = GetRunningTargetPlatform(TPM);
 
-		CurrentPlatform = Platforms[0];
-
-		for (int32 Index = 1; Index < Platforms.Num(); Index++)
-		{
-			if (Platforms[Index]->IsRunningPlatform())
-			{
-				CurrentPlatform = Platforms[Index];
-				break;
-			}
-		}
-
-		check(CurrentPlatform != NULL);
 
 		return CurrentPlatform->GetWaveFormat(&SoundWave);
 	}
@@ -174,6 +181,22 @@ static FName GetWaveFormatForRunningPlatform(USoundWave& SoundWave)
 	return NAME_None;
 }
 
+static const FPlatformAudioCookOverrides* GetCookOverridesForRunningPlatform()
+{
+	ITargetPlatformManagerModule* TPM = GetTargetPlatformManager();
+	if (TPM)
+	{
+		ITargetPlatform* CurrentPlatform = GetRunningTargetPlatform(TPM);
+
+
+		return CurrentPlatform->GetAudioCompressionSettings();
+	}
+	else
+	{
+		return nullptr;
+	}
+}
+
 /**
  * Stores derived data in the DDC.
  * After this returns, all bulk data from streaming chunks will be sent separately to the DDC and the BulkData for those chunks removed.
@@ -753,6 +776,12 @@ int32 FStreamedAudioPlatformData::DeserializeChunkFromDDC(TArray<uint8> TempData
 
 int32 FStreamedAudioPlatformData::GetChunkFromDDC(int32 ChunkIndex, uint8** OutChunkData, bool bMakeSureChunkIsLoaded /* = false */)
 {
+	if (Chunks.Num() == 0)
+	{
+		UE_LOG(LogAudioDerivedData, Warning, TEXT("No streamed audio chunks found!"));
+		return 0;
+	}
+
 	// if bMakeSureChunkIsLoaded is true, we don't actually know the size of the chunk's bulk data,
 	// so it will need to be allocated in GetCopy.
 	check(!bMakeSureChunkIsLoaded || (OutChunkData && (*OutChunkData == nullptr)));
@@ -1223,7 +1252,14 @@ static void CookSurroundWave( USoundWave* SoundWave, FName FormatName, const IAu
 			{
 				CompressionQualityModifier = CompressionOverrides->CompressionQualityModifier;
 			}
-			float ModifiedCompressionQuality = (float)SoundWave->CompressionQuality * CompressionQualityModifier;
+
+			float ModifiedCompressionQuality = (float)SoundWave->CompressionQuality;
+
+			if (!FMath::IsNearlyEqual(CompressionQualityModifier, 1.0f))
+			{
+				ModifiedCompressionQuality = (float)SoundWave->CompressionQuality * CompressionQualityModifier;
+			}
+			
 			if (ModifiedCompressionQuality >= 1.0f)
 			{
 				QualityInfo.Quality = FMath::FloorToInt(ModifiedCompressionQuality);
@@ -1415,7 +1451,7 @@ void USoundWave::CachePlatformData(bool bAsyncCache)
 {
 	FString DerivedDataKey;
 	FName AudioFormat = GetWaveFormatForRunningPlatform(*this);
-	const FPlatformAudioCookOverrides* CompressionOverrides = GetPlatformCompressionOverridesForCurrentPlatform();
+	const FPlatformAudioCookOverrides* CompressionOverrides = GetCookOverridesForRunningPlatform();
 	GetStreamedAudioDerivedDataKey(*this, AudioFormat, CompressionOverrides, DerivedDataKey);
 
 	if (RunningPlatformData == NULL || RunningPlatformData->DerivedDataKey != DerivedDataKey)
@@ -1587,7 +1623,7 @@ void USoundWave::FinishCachePlatformData()
 	{
 		FString DerivedDataKey;
 		FName AudioFormat = GetWaveFormatForRunningPlatform(*this);
-		const FPlatformAudioCookOverrides* CompressionOverrides = GetPlatformCompressionOverridesForCurrentPlatform();
+		const FPlatformAudioCookOverrides* CompressionOverrides = GetCookOverridesForRunningPlatform();
 		GetStreamedAudioDerivedDataKey(*this, AudioFormat, CompressionOverrides, DerivedDataKey);
 
 		check(RunningPlatformData->DerivedDataKey == DerivedDataKey);
diff --git a/Engine/Source/Runtime/Engine/Private/SoundWave.cpp b/Engine/Source/Runtime/Engine/Private/SoundWave.cpp
--- a/Engine/Source/Runtime/Engine/Private/SoundWave.cpp
+++ b/Engine/Source/Runtime/Engine/Private/SoundWave.cpp
@@ -47,6 +47,14 @@ FAutoConsoleVariableRef CVarLoadIntoCacheOnPostLoad(
 	TEXT("0: Honor the Play When Silent flag, 1: stop all silent non-procedural sources."),
 	ECVF_Default);
 
+static int32 ForceNonStreamingInEditorCVar = 0;
+FAutoConsoleVariableRef CVarForceNonStreamingInEditor(
+	TEXT("au.editor.ForceAudioNonStreaming"),
+	LoadIntoCacheOnPostLoadCVar,
+	TEXT("When set to 1, forces any audio played to be non-streaming May force a DDC miss.\n")
+	TEXT("0: Honor the Play When Silent flag, 1: stop all silent non-procedural sources."),
+	ECVF_Default);
+
 #if !UE_BUILD_SHIPPING
 static void DumpBakedAnalysisData(const TArray<FString>& Args)
 {
@@ -1820,6 +1828,11 @@ bool USoundWave::IsStreaming(const FPlatformAudioCookOverrides* Overrides /* = n
 	// TODO: add in check on whether it's part of a streaming SoundGroup.
 	if (!Overrides)
 	{
+		if (GIsEditor && ForceNonStreamingInEditorCVar != 0)
+		{
+			return false;
+		}
+
 		Overrides = GetPlatformCompressionOverridesForCurrentPlatform();
 	}
 
@@ -1990,7 +2003,7 @@ bool USoundWave::GetChunkData(int32 ChunkIndex, uint8** OutChunkData, bool bMake
 		ForceRebuildPlatformData();
 		if (RunningPlatformData->GetChunkFromDDC(ChunkIndex, OutChunkData, bMakeSureChunkIsLoaded) == 0)
 		{
-			UE_LOG(LogAudio, Display, TEXT("Failed to build sound %s."), *GetPathName());
+			UE_LOG(LogAudio, Warning, TEXT("Failed to build sound %s."), *GetPathName());
 		}
 		else
 		{
