Fix issue with SoundWaves played via SoundCues set to PlayWhenSilent not observing flag

diff --git a/Engine/Source/Runtime/Engine/Private/SoundWave.cpp b/Engine/Source/Runtime/Engine/Private/SoundWave.cpp
--- a/Engine/Source/Runtime/Engine/Private/SoundWave.cpp
+++ b/Engine/Source/Runtime/Engine/Private/SoundWave.cpp
@@ -1749,7 +1749,7 @@ void USoundWave::Parse(FAudioDevice* AudioDevice, const UPTRINT NodeWaveInstance
 	// When the BypassVirtualizeWhenSilent cvar is enabled, we should only honor bVirtualizeWhenSilent for procedural sounds:
 	const bool bHasSubtitles = ActiveSound.bHandleSubtitles && (ActiveSound.bHasExternalSubtitles || (Subtitles.Num() > 0));
 	const bool bStarted = WaveInstanceVolume > KINDA_SMALL_NUMBER || ActiveSound.ComponentVolumeFader.IsFadingIn();
-	const bool bCanPlayWhenSilent = IsPlayWhenSilent() && (!BypassPlayWhenSilentCVar || bProcedural);
+	const bool bCanPlayWhenSilent = ActiveSound.IsPlayWhenSilent() && (!BypassPlayWhenSilentCVar || bProcedural);
 	if (bStarted || bCanPlayWhenSilent || bHasSubtitles)
 	{
 		WaveInstances.Add(WaveInstance);
