Deprecating ARRAY_COUNT and changing it to UE_ARRAY_COUNT.

diff --git a/Engine/Source/Runtime/CrashReportCore/Private/IOS/IOSErrorReport.cpp b/Engine/Source/Runtime/CrashReportCore/Private/IOS/IOSErrorReport.cpp
--- a/Engine/Source/Runtime/CrashReportCore/Private/IOS/IOSErrorReport.cpp
+++ b/Engine/Source/Runtime/CrashReportCore/Private/IOS/IOSErrorReport.cpp
@@ -42,7 +42,7 @@ FString FIOSErrorReport::FindCrashedAppPath() const
 		CFRelease(CFString);
 		
 		static const TCHAR AppPathLineStart[] = TEXT("AppPath=");
-		static const int AppPathIdLength = ARRAY_COUNT(AppPathLineStart) - 1;
+		static const int AppPathIdLength = UE_ARRAY_COUNT(AppPathLineStart) - 1;
 		int32 AppPathStart = FileData.Find(AppPathLineStart);
 		if(AppPathStart >= 0)
 		{
diff --git a/Engine/Source/Runtime/CrashReportCore/Private/Mac/MacErrorReport.cpp b/Engine/Source/Runtime/CrashReportCore/Private/Mac/MacErrorReport.cpp
--- a/Engine/Source/Runtime/CrashReportCore/Private/Mac/MacErrorReport.cpp
+++ b/Engine/Source/Runtime/CrashReportCore/Private/Mac/MacErrorReport.cpp
@@ -52,7 +52,7 @@ FString FMacErrorReport::FindCrashedAppPath() const
 		CFRelease(CFString);
 		
 		static const TCHAR AppPathLineStart[] = TEXT("AppPath=");
-		static const int AppPathIdLength = ARRAY_COUNT(AppPathLineStart) - 1;
+		static const int AppPathIdLength = UE_ARRAY_COUNT(AppPathLineStart) - 1;
 		int32 AppPathStart = FileData.Find(AppPathLineStart);
 		if(AppPathStart >= 0)
 		{
diff --git a/Engine/Source/Runtime/D3D12RHI/Private/D3D12CommandContext.cpp b/Engine/Source/Runtime/D3D12RHI/Private/D3D12CommandContext.cpp
--- a/Engine/Source/Runtime/D3D12RHI/Private/D3D12CommandContext.cpp
+++ b/Engine/Source/Runtime/D3D12RHI/Private/D3D12CommandContext.cpp
@@ -115,9 +115,9 @@ FD3D12CommandContext::FD3D12CommandContext(FD3D12Device* InParent, FD3D12SubAllo
 {
 	FMemory::Memzero(DirtyUniformBuffers);
 	FMemory::Memzero(BoundUniformBuffers);
-	for (int i = 0; i < ARRAY_COUNT(BoundUniformBufferRefs); i++)
+	for (int i = 0; i < UE_ARRAY_COUNT(BoundUniformBufferRefs); i++)
 	{
-		for (int j = 0; j < ARRAY_COUNT(BoundUniformBufferRefs[i]); j++)
+		for (int j = 0; j < UE_ARRAY_COUNT(BoundUniformBufferRefs[i]); j++)
 		{
 			BoundUniformBufferRefs[i][j] = NULL;
 		}
@@ -366,9 +366,9 @@ void FD3D12CommandContext::ClearState()
 	FMemory::Memzero(BoundUniformBuffers, sizeof(BoundUniformBuffers));
 	FMemory::Memzero(DirtyUniformBuffers, sizeof(DirtyUniformBuffers));
 
-	for (int i = 0; i < ARRAY_COUNT(BoundUniformBufferRefs); i++)
+	for (int i = 0; i < UE_ARRAY_COUNT(BoundUniformBufferRefs); i++)
 	{
-		for (int j = 0; j < ARRAY_COUNT(BoundUniformBufferRefs[i]); j++)
+		for (int j = 0; j < UE_ARRAY_COUNT(BoundUniformBufferRefs[i]); j++)
 		{
 			BoundUniformBufferRefs[i][j] = NULL;
 		}
diff --git a/Engine/Source/Runtime/D3D12RHI/Private/D3D12DescriptorCache.cpp b/Engine/Source/Runtime/D3D12RHI/Private/D3D12DescriptorCache.cpp
--- a/Engine/Source/Runtime/D3D12RHI/Private/D3D12DescriptorCache.cpp
+++ b/Engine/Source/Runtime/D3D12RHI/Private/D3D12DescriptorCache.cpp
@@ -254,7 +254,7 @@ bool FD3D12DescriptorCache::SetDescriptorHeaps()
 	if (bHeapChanged)
 	{
 		ID3D12DescriptorHeap* /*const*/ ppHeaps[] = { pCurrentViewHeap, pCurrentSamplerHeap };
-		CmdContext->CommandListHandle->SetDescriptorHeaps(ARRAY_COUNT(ppHeaps), ppHeaps);
+		CmdContext->CommandListHandle->SetDescriptorHeaps(UE_ARRAY_COUNT(ppHeaps), ppHeaps);
 
 		pPreviousViewHeap = pCurrentViewHeap;
 		pPreviousSamplerHeap = pCurrentSamplerHeap;
@@ -478,7 +478,7 @@ void FD3D12DescriptorCache::SetSamplers(const FD3D12RootSignature* RootSignature
 
 	// Check to see if the sampler configuration is already in the sampler heap
 	FD3D12SamplerArrayDesc Desc = {};
-	if (SlotsNeeded <= ARRAY_COUNT(Desc.SamplerID))
+	if (SlotsNeeded <= UE_ARRAY_COUNT(Desc.SamplerID))
 	{
 		Desc.Count = SlotsNeeded;
 
@@ -491,7 +491,7 @@ void FD3D12DescriptorCache::SetSamplers(const FD3D12RootSignature* RootSignature
 		}
 
 		// The hash uses all of the bits
-		for (uint32 SlotIndex = SlotsNeeded; SlotIndex < ARRAY_COUNT(Desc.SamplerID); SlotIndex++)
+		for (uint32 SlotIndex = SlotsNeeded; SlotIndex < UE_ARRAY_COUNT(Desc.SamplerID); SlotIndex++)
 		{
 			Desc.SamplerID[SlotIndex] = 0;
 		}
@@ -538,7 +538,7 @@ void FD3D12DescriptorCache::SetSamplers(const FD3D12RootSignature* RootSignature
 			FD3D12_DESCRIPTOR_HEAP_TYPE_SAMPLER);
 
 		// Remember the locations of the samplers in the sampler map
-		if (SlotsNeeded <= ARRAY_COUNT(Desc.SamplerID))
+		if (SlotsNeeded <= UE_ARRAY_COUNT(Desc.SamplerID))
 		{
 			UniqueTables.Add(FD3D12UniqueSamplerTable(Desc, SrcDescriptors));
 
diff --git a/Engine/Source/Runtime/D3D12RHI/Private/D3D12DescriptorCache.h b/Engine/Source/Runtime/D3D12RHI/Private/D3D12DescriptorCache.h
--- a/Engine/Source/Runtime/D3D12RHI/Private/D3D12DescriptorCache.h
+++ b/Engine/Source/Runtime/D3D12RHI/Private/D3D12DescriptorCache.h
@@ -87,8 +87,8 @@ struct FD3D12SamplerArrayDesc
 	uint16 SamplerID[16];
 	inline bool operator==(const FD3D12SamplerArrayDesc& rhs) const
 	{
-		check(Count <= ARRAY_COUNT(SamplerID));
-		check(rhs.Count <= ARRAY_COUNT(rhs.SamplerID));
+		check(Count <= UE_ARRAY_COUNT(SamplerID));
+		check(rhs.Count <= UE_ARRAY_COUNT(rhs.SamplerID));
 
 		if (Count != rhs.Count)
 		{
diff --git a/Engine/Source/Runtime/D3D12RHI/Private/D3D12RayTracing.cpp b/Engine/Source/Runtime/D3D12RHI/Private/D3D12RayTracing.cpp
--- a/Engine/Source/Runtime/D3D12RHI/Private/D3D12RayTracing.cpp
+++ b/Engine/Source/Runtime/D3D12RHI/Private/D3D12RayTracing.cpp
@@ -3036,11 +3036,11 @@ static void SetRayTracingShaderResources(
 	ResourceBinderType& Binder)
 {
 	SetRayTracingShaderResources(CommandContext, Shader,
-		ARRAY_COUNT(ResourceBindings.Textures), ResourceBindings.Textures,
-		ARRAY_COUNT(ResourceBindings.SRVs), ResourceBindings.SRVs,
-		ARRAY_COUNT(ResourceBindings.UniformBuffers), ResourceBindings.UniformBuffers,
-		ARRAY_COUNT(ResourceBindings.Samplers), ResourceBindings.Samplers,
-		ARRAY_COUNT(ResourceBindings.UAVs), ResourceBindings.UAVs,
+		UE_ARRAY_COUNT(ResourceBindings.Textures), ResourceBindings.Textures,
+		UE_ARRAY_COUNT(ResourceBindings.SRVs), ResourceBindings.SRVs,
+		UE_ARRAY_COUNT(ResourceBindings.UniformBuffers), ResourceBindings.UniformBuffers,
+		UE_ARRAY_COUNT(ResourceBindings.Samplers), ResourceBindings.Samplers,
+		UE_ARRAY_COUNT(ResourceBindings.UAVs), ResourceBindings.UAVs,
 		0, nullptr, // loose parameters
 		DescriptorCache, Binder);
 }
diff --git a/Engine/Source/Runtime/D3D12RHI/Private/D3D12RootSignature.cpp b/Engine/Source/Runtime/D3D12RHI/Private/D3D12RootSignature.cpp
--- a/Engine/Source/Runtime/D3D12RHI/Private/D3D12RootSignature.cpp
+++ b/Engine/Source/Runtime/D3D12RHI/Private/D3D12RootSignature.cpp
@@ -147,12 +147,12 @@ FD3D12RootSignatureDesc::FD3D12RootSignatureDesc(const FD3D12QuantizedBoundShade
 	const uint32 RootDescriptorTableCost = QBSS.RootSignatureType == RS_RayTracingLocal ? RootDescriptorTableCostLocal : RootDescriptorTableCostGlobal;
 
 	// For each root parameter type...
-	for (uint32 RootParameterTypeIndex = 0; RootParameterTypeIndex < ARRAY_COUNT(RootParameterTypePriorityOrder); RootParameterTypeIndex++)
+	for (uint32 RootParameterTypeIndex = 0; RootParameterTypeIndex < UE_ARRAY_COUNT(RootParameterTypePriorityOrder); RootParameterTypeIndex++)
 	{
 		const D3D12_ROOT_PARAMETER_TYPE& RootParameterType = RootParameterTypePriorityOrder[RootParameterTypeIndex];
 
 		// ... and each shader stage visibility ...
-		for (uint32 ShaderVisibilityIndex = 0; ShaderVisibilityIndex < ARRAY_COUNT(ShaderVisibilityPriorityOrder); ShaderVisibilityIndex++)
+		for (uint32 ShaderVisibilityIndex = 0; ShaderVisibilityIndex < UE_ARRAY_COUNT(ShaderVisibilityPriorityOrder); ShaderVisibilityIndex++)
 		{
 			const EShaderVisibility& Visibility = ShaderVisibilityPriorityOrder[ShaderVisibilityIndex];
 			const FShaderRegisterCounts& Shader = QBSS.RegisterCounts[Visibility];
@@ -238,7 +238,7 @@ FD3D12RootSignatureDesc::FD3D12RootSignatureDesc(const FD3D12QuantizedBoundShade
 			Flags |= D3D12_ROOT_SIGNATURE_FLAG_ALLOW_INPUT_ASSEMBLER_INPUT_LAYOUT;
 		}
 
-		for (uint32 ShaderVisibilityIndex = 0; ShaderVisibilityIndex < ARRAY_COUNT(ShaderVisibilityPriorityOrder); ShaderVisibilityIndex++)
+		for (uint32 ShaderVisibilityIndex = 0; ShaderVisibilityIndex < UE_ARRAY_COUNT(ShaderVisibilityPriorityOrder); ShaderVisibilityIndex++)
 		{
 			const EShaderVisibility& Visibility = ShaderVisibilityPriorityOrder[ShaderVisibilityIndex];
 			const FShaderRegisterCounts& Shader = QBSS.RegisterCounts[Visibility];
diff --git a/Engine/Source/Runtime/D3D12RHI/Private/D3D12RootSignature.h b/Engine/Source/Runtime/D3D12RHI/Private/D3D12RootSignature.h
--- a/Engine/Source/Runtime/D3D12RHI/Private/D3D12RootSignature.h
+++ b/Engine/Source/Runtime/D3D12RHI/Private/D3D12RootSignature.h
@@ -201,7 +201,7 @@ public:
 	inline uint32 MaxUAVCount(uint32 ShaderStage) const { check(ShaderStage != SF_NumFrequencies); return Stage[ShaderStage].MaxUAVCount; }
 	inline CBVSlotMask CBVRegisterMask(uint32 ShaderStage) const { check(ShaderStage != SF_NumFrequencies); return Stage[ShaderStage].CBVRegisterMask; }
 
-	uint32 GetBindSlotOffsetInBytes(uint8 BindSlotIndex) const { check(BindSlotIndex < ARRAY_COUNT(BindSlotOffsetsInDWORDs)); return 4 * BindSlotOffsetsInDWORDs[BindSlotIndex]; }
+	uint32 GetBindSlotOffsetInBytes(uint8 BindSlotIndex) const { check(BindSlotIndex < UE_ARRAY_COUNT(BindSlotOffsetsInDWORDs)); return 4 * BindSlotOffsetsInDWORDs[BindSlotIndex]; }
 	uint32 GetTotalRootSignatureSizeInBytes() const { return 4 * TotalRootSignatureSizeInDWORDs; }
 
 private:
diff --git a/Engine/Source/Runtime/D3D12RHI/Private/D3D12StateCache.cpp b/Engine/Source/Runtime/D3D12RHI/Private/D3D12StateCache.cpp
--- a/Engine/Source/Runtime/D3D12RHI/Private/D3D12StateCache.cpp
+++ b/Engine/Source/Runtime/D3D12RHI/Private/D3D12StateCache.cpp
@@ -286,7 +286,7 @@ void FD3D12StateCacheBase::SetViewport(const D3D12_VIEWPORT& Viewport)
 
 void FD3D12StateCacheBase::SetViewports(uint32 Count, const D3D12_VIEWPORT* const Viewports)
 {
-	check(Count < ARRAY_COUNT(PipelineState.Graphics.CurrentViewport));
+	check(Count < UE_ARRAY_COUNT(PipelineState.Graphics.CurrentViewport));
 	if ((PipelineState.Graphics.CurrentNumberOfViewports != Count || FMemory::Memcmp(&PipelineState.Graphics.CurrentViewport[0], Viewports, sizeof(D3D12_VIEWPORT) * Count)) || GD3D12SkipStateCaching)
 	{
 		FMemory::Memcpy(&PipelineState.Graphics.CurrentViewport[0], Viewports, sizeof(D3D12_VIEWPORT) * Count);
@@ -318,7 +318,7 @@ void FD3D12StateCacheBase::SetScissorRect(const D3D12_RECT& ScissorRect)
 
 void FD3D12StateCacheBase::SetScissorRects(uint32 Count, const D3D12_RECT* const ScissorRects)
 {
-	check(Count < ARRAY_COUNT(PipelineState.Graphics.CurrentScissorRects));
+	check(Count < UE_ARRAY_COUNT(PipelineState.Graphics.CurrentScissorRects));
 
 	for (uint32 Rect = 0; Rect < Count; ++Rect)
 	{
@@ -956,7 +956,7 @@ bool FD3D12StateCacheBase::AssertResourceStates(ED3D12PipelineType PipelineType)
 
 		// RTV
 		{
-			const uint32 numRTVs = ARRAY_COUNT(PipelineState.Graphics.RenderTargetArray);
+			const uint32 numRTVs = UE_ARRAY_COUNT(PipelineState.Graphics.RenderTargetArray);
 			for (uint32 i = 0; i < numRTVs; i++)
 			{
 				FD3D12RenderTargetView* pCurrentView = PipelineState.Graphics.RenderTargetArray[i];
diff --git a/Engine/Source/Runtime/D3D12RHI/Private/D3D12StateCachePrivate.h b/Engine/Source/Runtime/D3D12RHI/Private/D3D12StateCachePrivate.h
--- a/Engine/Source/Runtime/D3D12RHI/Private/D3D12StateCachePrivate.h
+++ b/Engine/Source/Runtime/D3D12RHI/Private/D3D12StateCachePrivate.h
@@ -126,7 +126,7 @@ struct FD3D12ResourceCache
 	// Mark a specific shader stage as dirty.
 	inline void Dirty(EShaderFrequency ShaderFrequency, const ResourceSlotMask& SlotMask = -1)
 	{
-		checkSlow(ShaderFrequency < ARRAY_COUNT(DirtySlotMask));
+		checkSlow(ShaderFrequency < UE_ARRAY_COUNT(DirtySlotMask));
 		DirtySlotMask[ShaderFrequency] |= SlotMask;
 	}
 
diff --git a/Engine/Source/Runtime/D3D12RHI/Private/Windows/WindowsD3D12Device.cpp b/Engine/Source/Runtime/D3D12RHI/Private/Windows/WindowsD3D12Device.cpp
--- a/Engine/Source/Runtime/D3D12RHI/Private/Windows/WindowsD3D12Device.cpp
+++ b/Engine/Source/Runtime/D3D12RHI/Private/Windows/WindowsD3D12Device.cpp
@@ -142,7 +142,7 @@ static bool SafeTestD3D12CreateDevice(IDXGIAdapter* Adapter, D3D_FEATURE_LEVEL M
 			D3D_FEATURE_LEVEL MaxFeatureLevel = MinFeatureLevel;
 			D3D12_FEATURE_DATA_FEATURE_LEVELS FeatureLevelCaps = {};
 			FeatureLevelCaps.pFeatureLevelsRequested = FeatureLevels;
-			FeatureLevelCaps.NumFeatureLevels = ARRAY_COUNT(FeatureLevels);	
+			FeatureLevelCaps.NumFeatureLevels = UE_ARRAY_COUNT(FeatureLevels);	
 			if (SUCCEEDED(pDevice->CheckFeatureSupport(D3D12_FEATURE_FEATURE_LEVELS, &FeatureLevelCaps, sizeof(FeatureLevelCaps))))
 			{
 				MaxFeatureLevel = FeatureLevelCaps.MaxSupportedFeatureLevel;
diff --git a/Engine/Source/Runtime/Engine/Private/Animation/AnimInstanceProxy.cpp b/Engine/Source/Runtime/Engine/Private/Animation/AnimInstanceProxy.cpp
--- a/Engine/Source/Runtime/Engine/Private/Animation/AnimInstanceProxy.cpp
+++ b/Engine/Source/Runtime/Engine/Private/Animation/AnimInstanceProxy.cpp
@@ -62,7 +62,7 @@ void FAnimInstanceProxy::UpdateAnimationNode_WithRoot(const FAnimationUpdateCont
 void FAnimInstanceProxy::AddReferencedObjects(UAnimInstance* InAnimInstance, FReferenceCollector& Collector)
 {
 	DECLARE_SCOPE_HIERARCHICAL_COUNTER_FUNC()
-	for (int32 Index = 0; Index < ARRAY_COUNT(UngroupedActivePlayerArrays); ++Index)
+	for (int32 Index = 0; Index < UE_ARRAY_COUNT(UngroupedActivePlayerArrays); ++Index)
 	{
 		TArray<FAnimTickRecord>& UngroupedPlayers = UngroupedActivePlayerArrays[Index];
 		for (FAnimTickRecord& TickRecord : UngroupedPlayers)
diff --git a/Engine/Source/Runtime/Engine/Private/Collision/CollisionConversionsPhysx.cpp b/Engine/Source/Runtime/Engine/Private/Collision/CollisionConversionsPhysx.cpp
--- a/Engine/Source/Runtime/Engine/Private/Collision/CollisionConversionsPhysx.cpp
+++ b/Engine/Source/Runtime/Engine/Private/Collision/CollisionConversionsPhysx.cpp
@@ -445,8 +445,8 @@ static bool FindOverlappedTriangleNormal_Internal(const UWorld* World, const PxG
 			bool bOverflow = false;
 
 			const int32 NumTrisHit = bIsTriMesh ?
-				PxMeshQuery::findOverlapTriangleMesh(Geom, QueryTM, PTriMeshGeom, PShapeWorldPose, HitTris, ARRAY_COUNT(HitTris), 0, bOverflow) :
-				PxMeshQuery::findOverlapHeightField(Geom, QueryTM, PHeightfieldGeom, PShapeWorldPose, HitTris, ARRAY_COUNT(HitTris), 0, bOverflow);
+				PxMeshQuery::findOverlapTriangleMesh(Geom, QueryTM, PTriMeshGeom, PShapeWorldPose, HitTris, UE_ARRAY_COUNT(HitTris), 0, bOverflow) :
+				PxMeshQuery::findOverlapHeightField(Geom, QueryTM, PHeightfieldGeom, PShapeWorldPose, HitTris, UE_ARRAY_COUNT(HitTris), 0, bOverflow);
 
 			if (NumTrisHit > 0)
 			{
diff --git a/Engine/Source/Runtime/Engine/Private/Collision/PhysicsFiltering.cpp b/Engine/Source/Runtime/Engine/Private/Collision/PhysicsFiltering.cpp
--- a/Engine/Source/Runtime/Engine/Private/Collision/PhysicsFiltering.cpp
+++ b/Engine/Source/Runtime/Engine/Private/Collision/PhysicsFiltering.cpp
@@ -11,7 +11,7 @@ FPhysicsFilterBuilder::FPhysicsFilterBuilder(TEnumAsByte<enum ECollisionChannel>
 	, TouchingBits(0)
 	, Word3(0)
 {
-	for (int32 i = 0; i < ARRAY_COUNT(ResponseToChannels.EnumArray); i++)
+	for (int32 i = 0; i < UE_ARRAY_COUNT(ResponseToChannels.EnumArray); i++)
 	{
 		if (ResponseToChannels.EnumArray[i] == ECR_Block)
 		{
diff --git a/Engine/Source/Runtime/Engine/Private/Components/LightComponent.cpp b/Engine/Source/Runtime/Engine/Private/Components/LightComponent.cpp
--- a/Engine/Source/Runtime/Engine/Private/Components/LightComponent.cpp
+++ b/Engine/Source/Runtime/Engine/Private/Components/LightComponent.cpp
@@ -1483,7 +1483,7 @@ void ULightComponent::ReassignStationaryLightChannels(UWorld* TargetWorld, bool
 		}
 
 		// Use the lowest free channel
-		for (int32 ChannelIndex = 0; ChannelIndex < ARRAY_COUNT(bChannelUsed); ChannelIndex++)
+		for (int32 ChannelIndex = 0; ChannelIndex < UE_ARRAY_COUNT(bChannelUsed); ChannelIndex++)
 		{
 			if (!bChannelUsed[ChannelIndex])
 			{
diff --git a/Engine/Source/Runtime/Engine/Private/Components/StaticMeshComponent.cpp b/Engine/Source/Runtime/Engine/Private/Components/StaticMeshComponent.cpp
--- a/Engine/Source/Runtime/Engine/Private/Components/StaticMeshComponent.cpp
+++ b/Engine/Source/Runtime/Engine/Private/Components/StaticMeshComponent.cpp
@@ -1480,7 +1480,7 @@ void UStaticMeshComponent::ImportCustomProperties(const TCHAR* SourceText, FFeed
 		if (FParse::Value(SourceText, LODString, LODIndex))
 		{
 			TmpStr = FString::Printf(TEXT("%d"), LODIndex);
-			SourceText += TmpStr.Len() + (ARRAY_COUNT(LODString) - 1); // without the zero terminator
+			SourceText += TmpStr.Len() + (UE_ARRAY_COUNT(LODString) - 1); // without the zero terminator
 
 			// See if we need to add a new element to the LODData array
 			if (LODIndex > MaxLODIndex)
diff --git a/Engine/Source/Runtime/Engine/Private/GPUSkinCache.cpp b/Engine/Source/Runtime/Engine/Private/GPUSkinCache.cpp
--- a/Engine/Source/Runtime/Engine/Private/GPUSkinCache.cpp
+++ b/Engine/Source/Runtime/Engine/Private/GPUSkinCache.cpp
@@ -632,7 +632,7 @@ void FGPUSkinCache::TransitionToWriteable(FRHICommandList& RHICmdList)
 {
 	int32 BufferIndex = InternalUpdateCount % GPUSKINCACHE_FRAMES;
 	FRHIUnorderedAccessView* OutUAVs[] = { SkinCacheBuffer[BufferIndex].UAV };
-	RHICmdList.TransitionResources(EResourceTransitionAccess::ERWNoBarrier, EResourceTransitionPipeline::EGfxToCompute, OutUAVs, ARRAY_COUNT(OutUAVs));
+	RHICmdList.TransitionResources(EResourceTransitionAccess::ERWNoBarrier, EResourceTransitionPipeline::EGfxToCompute, OutUAVs, UE_ARRAY_COUNT(OutUAVs));
 }
 
 void FGPUSkinCache::TransitionAllToWriteable(FRHICommandList& RHICmdList)
@@ -646,7 +646,7 @@ void FGPUSkinCache::TransitionAllToWriteable(FRHICommandList& RHICmdList)
 			OutUAVs[Index] = SkinCacheBuffer[Index].UAV;
 		}
 
-		RHICmdList.TransitionResources(EResourceTransitionAccess::ERWNoBarrier, EResourceTransitionPipeline::EGfxToCompute, OutUAVs, ARRAY_COUNT(OutUAVs));
+		RHICmdList.TransitionResources(EResourceTransitionAccess::ERWNoBarrier, EResourceTransitionPipeline::EGfxToCompute, OutUAVs, UE_ARRAY_COUNT(OutUAVs));
 	}
 }
 #endif
diff --git a/Engine/Source/Runtime/Engine/Private/GameViewportClient.cpp b/Engine/Source/Runtime/Engine/Private/GameViewportClient.cpp
--- a/Engine/Source/Runtime/Engine/Private/GameViewportClient.cpp
+++ b/Engine/Source/Runtime/Engine/Private/GameViewportClient.cpp
@@ -3446,8 +3446,8 @@ bool UGameViewportClient::HandleDisplayCommand( const TCHAR* Cmd, FOutputDevice&
 {
 	TCHAR ObjectName[256];
 	TCHAR PropStr[256];
-	if ( FParse::Token(Cmd, ObjectName, ARRAY_COUNT(ObjectName), true) &&
-		FParse::Token(Cmd, PropStr, ARRAY_COUNT(PropStr), true) )
+	if ( FParse::Token(Cmd, ObjectName, UE_ARRAY_COUNT(ObjectName), true) &&
+		FParse::Token(Cmd, PropStr, UE_ARRAY_COUNT(PropStr), true) )
 	{
 		UObject* Obj = FindObject<UObject>(ANY_PACKAGE, ObjectName);
 		if (Obj != nullptr)
@@ -3475,7 +3475,7 @@ bool UGameViewportClient::HandleDisplayAllCommand( const TCHAR* Cmd, FOutputDevi
 {
 	TCHAR ClassName[256];
 	TCHAR PropStr[256];
-	if (FParse::Token(Cmd, ClassName, ARRAY_COUNT(ClassName), true))
+	if (FParse::Token(Cmd, ClassName, UE_ARRAY_COUNT(ClassName), true))
 	{
 		bool bValidClassToken = true;
 		UClass* WithinClass = nullptr;
@@ -3500,7 +3500,7 @@ bool UGameViewportClient::HandleDisplayAllCommand( const TCHAR* Cmd, FOutputDevi
 		}
 		if (bValidClassToken)
 		{
-			FParse::Token(Cmd, PropStr, ARRAY_COUNT(PropStr), true);
+			FParse::Token(Cmd, PropStr, UE_ARRAY_COUNT(PropStr), true);
 			UClass* Cls = FindObject<UClass>(ANY_PACKAGE, ClassName);
 			if (Cls != nullptr)
 			{
@@ -3536,7 +3536,7 @@ bool UGameViewportClient::HandleDisplayAllCommand( const TCHAR* Cmd, FOutputDevi
 bool UGameViewportClient::HandleDisplayAllLocationCommand( const TCHAR* Cmd, FOutputDevice& Ar )
 {
 	TCHAR ClassName[256];
-	if (FParse::Token(Cmd, ClassName, ARRAY_COUNT(ClassName), true))
+	if (FParse::Token(Cmd, ClassName, UE_ARRAY_COUNT(ClassName), true))
 	{
 		UClass* Cls = FindObject<UClass>(ANY_PACKAGE, ClassName);
 		if (Cls != nullptr)
@@ -3568,7 +3568,7 @@ bool UGameViewportClient::HandleDisplayAllLocationCommand( const TCHAR* Cmd, FOu
 bool UGameViewportClient::HandleDisplayAllRotationCommand( const TCHAR* Cmd, FOutputDevice& Ar )
 {
 	TCHAR ClassName[256];
-	if (FParse::Token(Cmd, ClassName, ARRAY_COUNT(ClassName), true))
+	if (FParse::Token(Cmd, ClassName, UE_ARRAY_COUNT(ClassName), true))
 	{
 		UClass* Cls = FindObject<UClass>(ANY_PACKAGE, ClassName);
 		if (Cls != nullptr)
@@ -3610,7 +3610,7 @@ bool UGameViewportClient::HandleGetAllLocationCommand(const TCHAR* Cmd, FOutputD
 	TCHAR ClassName[256];
 	UClass* Class;
 
-	if (FParse::Token(Cmd, ClassName, ARRAY_COUNT(ClassName), 1) &&
+	if (FParse::Token(Cmd, ClassName, UE_ARRAY_COUNT(ClassName), 1) &&
 		(Class = FindObject<UClass>(ANY_PACKAGE, ClassName)) != NULL)
 	{
 		bool bShowPendingKills = FParse::Command(&Cmd, TEXT("SHOWPENDINGKILLS"));
diff --git a/Engine/Source/Runtime/Engine/Private/GeomTools.cpp b/Engine/Source/Runtime/Engine/Private/GeomTools.cpp
--- a/Engine/Source/Runtime/Engine/Private/GeomTools.cpp
+++ b/Engine/Source/Runtime/Engine/Private/GeomTools.cpp
@@ -27,7 +27,7 @@ static FClipSMVertex GetVert(const UStaticMesh* StaticMesh, int32 VertIndex)
 	{
 		Result.UVs[UVIndex] = LODModel.VertexBuffers.StaticMeshVertexBuffer.GetVertexUV(VertIndex,UVIndex);
 	}
-	for(int32 UVIndex = NumUVs;UVIndex < ARRAY_COUNT(Result.UVs);UVIndex++)
+	for(int32 UVIndex = NumUVs;UVIndex < UE_ARRAY_COUNT(Result.UVs);UVIndex++)
 	{
 		Result.UVs[UVIndex] = FVector2D::ZeroVector;
 	}
@@ -380,7 +380,7 @@ static bool VertsAreEqual(const FClipSMVertex& A,const FClipSMVertex& B)
 		return false;
 	}
 
-	for(int32 i=0; i<ARRAY_COUNT(A.UVs); i++)
+	for(int32 i=0; i<UE_ARRAY_COUNT(A.UVs); i++)
 	{
 		if( !A.UVs[i].Equals(B.UVs[i], 1.0f / 1024.0f) )
 		{
diff --git a/Engine/Source/Runtime/Engine/Private/GestureRecognizer.cpp b/Engine/Source/Runtime/Engine/Private/GestureRecognizer.cpp
--- a/Engine/Source/Runtime/Engine/Private/GestureRecognizer.cpp
+++ b/Engine/Source/Runtime/Engine/Private/GestureRecognizer.cpp
@@ -23,7 +23,7 @@ void FGestureRecognizer::DetectGestures(const FVector (&Touches)[EKeys::NUM_TOUC
 	if (PreviousTouchCount != 0 || TouchCount != 0)
 	{
 		// place new anchor points
-		for (int32 Index = 0; Index < ARRAY_COUNT(AnchorPoints); Index++)
+		for (int32 Index = 0; Index < UE_ARRAY_COUNT(AnchorPoints); Index++)
 		{
 			if (PreviousTouchCount < Index + 1 && TouchCount >= Index + 1)
 			{
diff --git a/Engine/Source/Runtime/Engine/Private/InputDelegateBinding.cpp b/Engine/Source/Runtime/Engine/Private/InputDelegateBinding.cpp
--- a/Engine/Source/Runtime/Engine/Private/InputDelegateBinding.cpp
+++ b/Engine/Source/Runtime/Engine/Private/InputDelegateBinding.cpp
@@ -35,7 +35,7 @@ void UInputDelegateBinding::BindInputDelegates(const UClass* InClass, UInputComp
 	{
 		BindInputDelegates(InClass->GetSuperClass(), InputComponent);
 
-		for (int32 Index = 0; Index < ARRAY_COUNT(InputBindingClasses); ++Index)
+		for (int32 Index = 0; Index < UE_ARRAY_COUNT(InputBindingClasses); ++Index)
 		{
 			UInputDelegateBinding* BindingObject = CastChecked<UInputDelegateBinding>(
 				UBlueprintGeneratedClass::GetDynamicBindingObject(InClass, InputBindingClasses[Index])
diff --git a/Engine/Source/Runtime/Engine/Private/Level.cpp b/Engine/Source/Runtime/Engine/Private/Level.cpp
--- a/Engine/Source/Runtime/Engine/Private/Level.cpp
+++ b/Engine/Source/Runtime/Engine/Private/Level.cpp
@@ -596,7 +596,7 @@ void ULevel::PostLoad()
 	}
 
 	// Fixup deprecated stuff in levels simplification settings
-	for (int32 Index = 0; Index < ARRAY_COUNT(LevelSimplification); ++Index)
+	for (int32 Index = 0; Index < UE_ARRAY_COUNT(LevelSimplification); ++Index)
 	{
 		LevelSimplification[Index].PostLoadDeprecated();
 	}
diff --git a/Engine/Source/Runtime/Engine/Private/LightMap.cpp b/Engine/Source/Runtime/Engine/Private/LightMap.cpp
--- a/Engine/Source/Runtime/Engine/Private/LightMap.cpp
+++ b/Engine/Source/Runtime/Engine/Private/LightMap.cpp
@@ -2756,7 +2756,7 @@ int32 FLightMap2D::EncodeShadowTexture(ULevel* LightingScenario, struct FLightMa
 							ExtrapolatedFilterableComponents[i] = 0;
 						}
 
-						for (int32 NeighborIndex = 0; NeighborIndex < ARRAY_COUNT(Neighbors); NeighborIndex++)
+						for (int32 NeighborIndex = 0; NeighborIndex < UE_ARRAY_COUNT(Neighbors); NeighborIndex++)
 						{
 							if (static_cast<int32>(DestY) + Neighbors[NeighborIndex].Y >= 0
 								&& DestY + Neighbors[NeighborIndex].Y < MipSizeY
@@ -2975,7 +2975,7 @@ void FLightMap2D::Serialize(FArchive& Ar)
 
 	if (RenderCustomVersion >= FRenderingObjectVersion::LightmapHasShadowmapData)
 	{
-		for (int Channel = 0; Channel < ARRAY_COUNT(bShadowChannelValid); Channel++)
+		for (int Channel = 0; Channel < UE_ARRAY_COUNT(bShadowChannelValid); Channel++)
 		{
 			Ar << bShadowChannelValid[Channel];
 		}
@@ -3257,7 +3257,7 @@ bool FQuantizedLightmapData::HasNonZeroData() const
 
 			if (bHasSkyShadowing)
 			{
-				for (int32 Index = 0; Index < ARRAY_COUNT(LightmapSample.SkyOcclusion); Index++)
+				for (int32 Index = 0; Index < UE_ARRAY_COUNT(LightmapSample.SkyOcclusion); Index++)
 				{
 					if (LightmapSample.SkyOcclusion[Index] != 0)
 					{
diff --git a/Engine/Source/Runtime/Engine/Private/LocalPlayer.cpp b/Engine/Source/Runtime/Engine/Private/LocalPlayer.cpp
--- a/Engine/Source/Runtime/Engine/Private/LocalPlayer.cpp
+++ b/Engine/Source/Runtime/Engine/Private/LocalPlayer.cpp
@@ -1298,7 +1298,7 @@ bool ULocalPlayer::HandleListPawnComponentsCommand( const TCHAR* Cmd, FOutputDev
 bool ULocalPlayer::HandleExecCommand( const TCHAR* Cmd, FOutputDevice& Ar )
 {
 	TCHAR Filename[512];
-	if( FParse::Token( Cmd, Filename, ARRAY_COUNT(Filename), 0 ) )
+	if( FParse::Token( Cmd, Filename, UE_ARRAY_COUNT(Filename), 0 ) )
 	{
 		ExecMacro( Filename, Ar );
 	}
diff --git a/Engine/Source/Runtime/Engine/Private/Materials/Material.cpp b/Engine/Source/Runtime/Engine/Private/Materials/Material.cpp
--- a/Engine/Source/Runtime/Engine/Private/Materials/Material.cpp
+++ b/Engine/Source/Runtime/Engine/Private/Materials/Material.cpp
@@ -988,7 +988,7 @@ void UMaterial::GetUsedTextures(TArray<UTexture*>& OutTextures, EMaterialQuality
 					&CurrentResource->GetUniformVolumeTextureExpressions(),
 					&CurrentResource->GetUniformVirtualTextureExpressions()
 				};
-				for (int32 TypeIndex = 0; TypeIndex < ARRAY_COUNT(ExpressionsByType); TypeIndex++)
+				for (int32 TypeIndex = 0; TypeIndex < UE_ARRAY_COUNT(ExpressionsByType); TypeIndex++)
 				{
 					// Iterate over each of the material's texture expressions.
 					for (FMaterialUniformExpressionTexture* Expression : *ExpressionsByType[TypeIndex])
@@ -1009,7 +1009,7 @@ void UMaterial::GetUsedTextures(TArray<UTexture*>& OutTextures, EMaterialQuality
 				{
 					&CurrentResource->GetUniformScalarParameterExpressions()
 				};
-				for (int32 TypeIndex = 0; TypeIndex < ARRAY_COUNT(AtlasExpressions); TypeIndex++)
+				for (int32 TypeIndex = 0; TypeIndex < UE_ARRAY_COUNT(AtlasExpressions); TypeIndex++)
 				{
 					// Iterate over each of the material's texture expressions.
 					for (FMaterialUniformExpression* Expression : *AtlasExpressions[TypeIndex])
@@ -1057,7 +1057,7 @@ void UMaterial::GetUsedTexturesAndIndices(TArray<UTexture*>& OutTextures, TArray
 			// Try to prevent resizing since this would be expensive.
 			OutIndices.Empty(ExpressionsByType[0]->Num() + ExpressionsByType[1]->Num() + ExpressionsByType[2]->Num() + ExpressionsByType[3]->Num());
 
-			for (int32 TypeIndex = 0; TypeIndex < ARRAY_COUNT(ExpressionsByType); TypeIndex++)
+			for (int32 TypeIndex = 0; TypeIndex < UE_ARRAY_COUNT(ExpressionsByType); TypeIndex++)
 			{
 				// Iterate over each of the material's texture expressions.
 				for (FMaterialUniformExpressionTexture* Expression : *ExpressionsByType[TypeIndex])
@@ -1116,7 +1116,7 @@ void UMaterial::LogMaterialsAndTextures(FOutputDevice& Ar, int32 Indent) const
 						&MaterialResource->GetUniformVolumeTextureExpressions(),
 						&MaterialResource->GetUniformVirtualTextureExpressions(),
 					};
-					for (int32 TypeIndex = 0; TypeIndex < ARRAY_COUNT(ExpressionsByType); TypeIndex++)
+					for (int32 TypeIndex = 0; TypeIndex < UE_ARRAY_COUNT(ExpressionsByType); TypeIndex++)
 					{
 						for (FMaterialUniformExpressionTexture* Expression : *ExpressionsByType[TypeIndex])
 						{
@@ -1173,7 +1173,7 @@ void UMaterial::OverrideTexture(const UTexture* InTextureToOverride, UTexture* O
 			&Resource->GetUniformVolumeTextureExpressions(),
 			&Resource->GetUniformVirtualTextureExpressions()
 		};
-		for(int32 TypeIndex = 0;TypeIndex < ARRAY_COUNT(ExpressionsByType);TypeIndex++)
+		for(int32 TypeIndex = 0;TypeIndex < UE_ARRAY_COUNT(ExpressionsByType);TypeIndex++)
 		{
 			// Iterate over each of the material's texture expressions.
 			for (FMaterialUniformExpressionTexture* Expression : *ExpressionsByType[TypeIndex])
diff --git a/Engine/Source/Runtime/Engine/Private/Materials/MaterialExpressions.cpp b/Engine/Source/Runtime/Engine/Private/Materials/MaterialExpressions.cpp
--- a/Engine/Source/Runtime/Engine/Private/Materials/MaterialExpressions.cpp
+++ b/Engine/Source/Runtime/Engine/Private/Materials/MaterialExpressions.cpp
@@ -8043,7 +8043,7 @@ UMaterialExpressionQualitySwitch::UMaterialExpressionQualitySwitch(const FObject
 int32 UMaterialExpressionQualitySwitch::Compile(class FMaterialCompiler* Compiler, int32 OutputIndex)
 {
 	const EMaterialQualityLevel::Type QualityLevelToCompile = Compiler->GetQualityLevel();
-	check(QualityLevelToCompile < ARRAY_COUNT(Inputs));
+	check(QualityLevelToCompile < UE_ARRAY_COUNT(Inputs));
 	FExpressionInput QualityInput = Inputs[QualityLevelToCompile].GetTracedInput();
 	FExpressionInput DefaultTraced = Default.GetTracedInput();
 
@@ -8071,7 +8071,7 @@ const TArray<FExpressionInput*> UMaterialExpressionQualitySwitch::GetInputs()
 
 	OutInputs.Add(&Default);
 
-	for (int32 InputIndex = 0; InputIndex < ARRAY_COUNT(Inputs); InputIndex++)
+	for (int32 InputIndex = 0; InputIndex < UE_ARRAY_COUNT(Inputs); InputIndex++)
 	{
 		OutInputs.Add(&Inputs[InputIndex]);
 	}
@@ -8154,7 +8154,7 @@ UMaterialExpressionFeatureLevelSwitch::UMaterialExpressionFeatureLevelSwitch(con
 int32 UMaterialExpressionFeatureLevelSwitch::Compile(class FMaterialCompiler* Compiler, int32 OutputIndex)
 {
 	const ERHIFeatureLevel::Type FeatureLevelToCompile = Compiler->GetFeatureLevel();
-	check(FeatureLevelToCompile < ARRAY_COUNT(Inputs));
+	check(FeatureLevelToCompile < UE_ARRAY_COUNT(Inputs));
 	FExpressionInput& FeatureInput = Inputs[FeatureLevelToCompile];
 
 	if (!Default.GetTracedInput().Expression)
@@ -8181,7 +8181,7 @@ const TArray<FExpressionInput*> UMaterialExpressionFeatureLevelSwitch::GetInputs
 
 	OutInputs.Add(&Default);
 
-	for (int32 InputIndex = 0; InputIndex < ARRAY_COUNT(Inputs); InputIndex++)
+	for (int32 InputIndex = 0; InputIndex < UE_ARRAY_COUNT(Inputs); InputIndex++)
 	{
 		OutInputs.Add(&Inputs[InputIndex]);
 	}
@@ -8296,7 +8296,7 @@ int32 UMaterialExpressionShadingPathSwitch::Compile(class FMaterialCompiler* Com
 		ShadingPathToCompile = ERHIShadingPath::Mobile;
 	}
 
-	check(ShadingPathToCompile < ARRAY_COUNT(Inputs));
+	check(ShadingPathToCompile < UE_ARRAY_COUNT(Inputs));
 	FExpressionInput ShadingPathInput = Inputs[ShadingPathToCompile].GetTracedInput();
 	FExpressionInput DefaultTraced = Default.GetTracedInput();
 
@@ -8324,7 +8324,7 @@ const TArray<FExpressionInput*> UMaterialExpressionShadingPathSwitch::GetInputs(
 
 	OutInputs.Add(&Default);
 
-	for (int32 InputIndex = 0; InputIndex < ARRAY_COUNT(Inputs); InputIndex++)
+	for (int32 InputIndex = 0; InputIndex < UE_ARRAY_COUNT(Inputs); InputIndex++)
 	{
 		OutInputs.Add(&Inputs[InputIndex]);
 	}
diff --git a/Engine/Source/Runtime/Engine/Private/Materials/MaterialInstance.cpp b/Engine/Source/Runtime/Engine/Private/Materials/MaterialInstance.cpp
--- a/Engine/Source/Runtime/Engine/Private/Materials/MaterialInstance.cpp
+++ b/Engine/Source/Runtime/Engine/Private/Materials/MaterialInstance.cpp
@@ -1110,7 +1110,7 @@ void UMaterialInstance::GetTextureExpressionValues(const FMaterialResource* Mate
 		OutIndices->Empty(ExpressionsByType[0]->Num() + ExpressionsByType[1]->Num() + ExpressionsByType[2]->Num() + ExpressionsByType[3]->Num() + ExpressionsByType[4]->Num());
 	}
 
-	for(int32 TypeIndex = 0;TypeIndex < ARRAY_COUNT(ExpressionsByType);TypeIndex++)
+	for(int32 TypeIndex = 0;TypeIndex < UE_ARRAY_COUNT(ExpressionsByType);TypeIndex++)
 	{
 		const TArray<TRefCountPtr<FMaterialUniformExpressionTexture> >& Expressions = *ExpressionsByType[TypeIndex];
 
@@ -1148,7 +1148,7 @@ void UMaterialInstance::GetAtlasTextureValues(const FMaterialResource* MaterialR
 	{
 		&MaterialResource->GetUniformScalarParameterExpressions()
 	};
-	for (int32 TypeIndex = 0; TypeIndex < ARRAY_COUNT(AtlasExpressions); TypeIndex++)
+	for (int32 TypeIndex = 0; TypeIndex < UE_ARRAY_COUNT(AtlasExpressions); TypeIndex++)
 	{
 		// Iterate over each of the material's scalar expressions.
 		for (FMaterialUniformExpression* Expression : *AtlasExpressions[TypeIndex])
@@ -1467,7 +1467,7 @@ void UMaterialInstance::OverrideTexture(const UTexture* InTextureToOverride, UTe
 		ExpressionsByType[4] = &SourceMaterialResource->GetUniformVirtualTextureExpressions();
 	}
 		
-	for(int32 TypeIndex = 0;TypeIndex < ARRAY_COUNT(ExpressionsByType);TypeIndex++)
+	for(int32 TypeIndex = 0;TypeIndex < UE_ARRAY_COUNT(ExpressionsByType);TypeIndex++)
 	{
 		const TArray<TRefCountPtr<FMaterialUniformExpressionTexture> >& Expressions = *ExpressionsByType[TypeIndex];
 
@@ -4889,7 +4889,7 @@ void UMaterialInstance::CopyMaterialUniformParametersInternal(UMaterialInterface
 						&MaterialResource->GetUniformVirtualTextureExpressions()
 					};
 
-					for (int32 TypeIndex = 0; TypeIndex < ARRAY_COUNT(TextureExpressions); TypeIndex++)
+					for (int32 TypeIndex = 0; TypeIndex < UE_ARRAY_COUNT(TextureExpressions); TypeIndex++)
 					{
 						for (FMaterialUniformExpressionTexture* TextureExpression : *TextureExpressions[TypeIndex])
 						{
diff --git a/Engine/Source/Runtime/Engine/Private/Materials/MaterialShared.cpp b/Engine/Source/Runtime/Engine/Private/Materials/MaterialShared.cpp
--- a/Engine/Source/Runtime/Engine/Private/Materials/MaterialShared.cpp
+++ b/Engine/Source/Runtime/Engine/Private/Materials/MaterialShared.cpp
@@ -90,17 +90,17 @@ FName MaterialQualityLevelNames[] =
 	FName(TEXT("Num"))
 };
 
-static_assert(ARRAY_COUNT(MaterialQualityLevelNames) == EMaterialQualityLevel::Num + 1, "Missing entry from material quality level names.");
+static_assert(UE_ARRAY_COUNT(MaterialQualityLevelNames) == EMaterialQualityLevel::Num + 1, "Missing entry from material quality level names.");
 
 void GetMaterialQualityLevelName(EMaterialQualityLevel::Type InQualityLevel, FString& OutName)
 {
-	check(InQualityLevel < ARRAY_COUNT(MaterialQualityLevelNames));
+	check(InQualityLevel < UE_ARRAY_COUNT(MaterialQualityLevelNames));
 	MaterialQualityLevelNames[(int32)InQualityLevel].ToString(OutName);
 }
 
 FName GetMaterialQualityLevelFName(EMaterialQualityLevel::Type InQualityLevel)
 {
-	check(InQualityLevel < ARRAY_COUNT(MaterialQualityLevelNames));
+	check(InQualityLevel < UE_ARRAY_COUNT(MaterialQualityLevelNames));
 	return MaterialQualityLevelNames[(int32)InQualityLevel];
 }
 
@@ -194,7 +194,7 @@ static inline SIZE_T AddShaderSize(FShader* Shader, TSet<FShaderResourceId>& Uni
 int32 FMaterialCompiler::Errorf(const TCHAR* Format,...)
 {
 	TCHAR	ErrorText[2048];
-	GET_VARARGS( ErrorText, ARRAY_COUNT(ErrorText), ARRAY_COUNT(ErrorText)-1, Format, Format );
+	GET_VARARGS( ErrorText, UE_ARRAY_COUNT(ErrorText), UE_ARRAY_COUNT(ErrorText)-1, Format, Format );
 	return Error(ErrorText);
 }
 
@@ -3960,7 +3960,7 @@ void SetShaderMapsOnMaterialResources_RenderThread(FRHICommandListImmediate& RHI
 
 	// Iterate through all loaded material render proxies and recache their uniform expressions if needed
 	// This search does not scale well, but is only used when uploading async shader compile results
-	for (int32 FeatureLevelIndex = 0; FeatureLevelIndex < ARRAY_COUNT(bUpdateFeatureLevel); ++FeatureLevelIndex)
+	for (int32 FeatureLevelIndex = 0; FeatureLevelIndex < UE_ARRAY_COUNT(bUpdateFeatureLevel); ++FeatureLevelIndex)
 	{
 		if (bUpdateFeatureLevel[FeatureLevelIndex])
 		{
diff --git a/Engine/Source/Runtime/Engine/Private/NetConnection.cpp b/Engine/Source/Runtime/Engine/Private/NetConnection.cpp
--- a/Engine/Source/Runtime/Engine/Private/NetConnection.cpp
+++ b/Engine/Source/Runtime/Engine/Private/NetConnection.cpp
@@ -1391,7 +1391,7 @@ void UNetConnection::FlushNet(bool bIgnoreSimulation)
 		}
 #endif
 		// Update stuff.
-		const int32 Index = OutPacketId & (ARRAY_COUNT(OutLagPacketId)-1);
+		const int32 Index = OutPacketId & (UE_ARRAY_COUNT(OutLagPacketId)-1);
 
 		// Remember the actual time this packet was sent out, so we can compute ping when the ack comes back
 		OutLagPacketId[Index]			= OutPacketId;
@@ -1687,7 +1687,7 @@ bool UNetConnection::ReadPacketInfo(FBitReader& Reader)
 
 	// Update ping
 	// At this time we have updated OutAckPacketId to the latest received ack.
-	const int32 Index = OutAckPacketId & (ARRAY_COUNT(OutLagPacketId)-1);
+	const int32 Index = OutAckPacketId & (UE_ARRAY_COUNT(OutLagPacketId)-1);
 
 	if ( OutLagPacketId[Index] == OutAckPacketId )
 	{
diff --git a/Engine/Source/Runtime/Engine/Private/NetworkDriver.cpp b/Engine/Source/Runtime/Engine/Private/NetworkDriver.cpp
--- a/Engine/Source/Runtime/Engine/Private/NetworkDriver.cpp
+++ b/Engine/Source/Runtime/Engine/Private/NetworkDriver.cpp
@@ -899,7 +899,7 @@ void UNetDriver::TickFlush(float DeltaSeconds)
 					PerfCounters->Set(TEXT("MinPing"), MinPing, IPerfCounters::Flags::Transient);
 
 					// update buckets
-					for (int BucketIdx = 0; BucketIdx < ARRAY_COUNT(Buckets); ++BucketIdx)
+					for (int BucketIdx = 0; BucketIdx < UE_ARRAY_COUNT(Buckets); ++BucketIdx)
 					{
 						PerfCountersIncrement(FString::Printf(TEXT("PingBucketInt%d"), BucketIdx), Buckets[BucketIdx], 0, IPerfCounters::Flags::Transient);
 					}
@@ -5692,7 +5692,7 @@ static bool NetDriverExec(UWorld* InWorld, const TCHAR* Cmd, FOutputDevice& Ar)
 		TCHAR TokenStr[128];
 
 		// Route the command to a specific beacon if a name is specified or all of them otherwise
-		if (FParse::Token(Cmd, TokenStr, ARRAY_COUNT(TokenStr), true))
+		if (FParse::Token(Cmd, TokenStr, UE_ARRAY_COUNT(TokenStr), true))
 		{
 			NamedDriver = GEngine->FindNamedNetDriver(InWorld, FName(TokenStr));
 			if (NamedDriver != NULL)
diff --git a/Engine/Source/Runtime/Engine/Private/PacketHandlers/StatelessConnectHandlerComponent.cpp b/Engine/Source/Runtime/Engine/Private/PacketHandlers/StatelessConnectHandlerComponent.cpp
--- a/Engine/Source/Runtime/Engine/Private/PacketHandlers/StatelessConnectHandlerComponent.cpp
+++ b/Engine/Source/Runtime/Engine/Private/PacketHandlers/StatelessConnectHandlerComponent.cpp
@@ -292,8 +292,8 @@ void StatelessConnectHandlerComponent::NotifyHandshakeBegin()
 			InitialPacket.WriteBit(bRestartHandshake);
 			InitialPacket.WriteBit(SecretIdPad);
 
-			FMemory::Memzero(PacketSizeFiller, ARRAY_COUNT(PacketSizeFiller));
-			InitialPacket.Serialize(PacketSizeFiller, ARRAY_COUNT(PacketSizeFiller));
+			FMemory::Memzero(PacketSizeFiller, UE_ARRAY_COUNT(PacketSizeFiller));
+			InitialPacket.Serialize(PacketSizeFiller, UE_ARRAY_COUNT(PacketSizeFiller));
 
 
 
@@ -355,13 +355,13 @@ void StatelessConnectHandlerComponent::SendConnectChallenge(TSharedPtr<const FIn
 		ChallengePacket.WriteBit(ActiveSecret);
 
 		ChallengePacket << Timestamp;
-		ChallengePacket.Serialize(Cookie, ARRAY_COUNT(Cookie));
+		ChallengePacket.Serialize(Cookie, UE_ARRAY_COUNT(Cookie));
 
 #if !UE_BUILD_SHIPPING
 		FDDoSDetection* DDoS = Handler->GetDDoS();
 
 		UE_CLOG((DDoS == nullptr || !DDoS->CheckLogRestrictions()), LogHandshake, Log,
-				TEXT("SendConnectChallenge. Timestamp: %f, Cookie: %s" ), Timestamp, *FString::FromBlob(Cookie, ARRAY_COUNT(Cookie)));
+				TEXT("SendConnectChallenge. Timestamp: %f, Cookie: %s" ), Timestamp, *FString::FromBlob(Cookie, UE_ARRAY_COUNT(Cookie)));
 #endif
 
 		CapHandshakePacket(ChallengePacket);
@@ -475,7 +475,7 @@ void StatelessConnectHandlerComponent::SendChallengeResponse(uint8 InSecretId, f
 		LastServerSequence = *CurSequence & (MAX_PACKETID - 1);
 		LastClientSequence = *(CurSequence + 1) & (MAX_PACKETID - 1);
 
-		FMemory::Memcpy(LastCookie, InCookie, ARRAY_COUNT(LastCookie));
+		FMemory::Memcpy(LastCookie, InCookie, UE_ARRAY_COUNT(LastCookie));
 	}
 	else
 	{
@@ -666,7 +666,7 @@ void StatelessConnectHandlerComponent::InitFromConnectionless(StatelessConnectHa
 	// Store the cookie/address used for the handshake, to enable server ack-retries
 	LastChallengeSuccessAddress = InConnectionlessHandler->LastChallengeSuccessAddress;
 
-	FMemory::Memcpy(AuthorisedCookie, InConnectionlessHandler->AuthorisedCookie, ARRAY_COUNT(AuthorisedCookie));
+	FMemory::Memcpy(AuthorisedCookie, InConnectionlessHandler->AuthorisedCookie, UE_ARRAY_COUNT(AuthorisedCookie));
 }
 
 void StatelessConnectHandlerComponent::Incoming(FBitReader& Packet)
@@ -731,7 +731,7 @@ void StatelessConnectHandlerComponent::Incoming(FBitReader& Packet)
 							}
 
 							// Save the final authorized cookie
-							FMemory::Memcpy(AuthorisedCookie, Cookie, ARRAY_COUNT(AuthorisedCookie));
+							FMemory::Memcpy(AuthorisedCookie, Cookie, UE_ARRAY_COUNT(AuthorisedCookie));
 						}
 
 						// Now finish initializing the handler - flushing the queued packet buffer in the process.
@@ -896,7 +896,7 @@ void StatelessConnectHandlerComponent::IncomingConnectionless(const TSharedPtr<c
 						{
 							if (bRestartHandshake)
 							{
-								FMemory::Memcpy(AuthorisedCookie, OrigCookie, ARRAY_COUNT(AuthorisedCookie));
+								FMemory::Memcpy(AuthorisedCookie, OrigCookie, UE_ARRAY_COUNT(AuthorisedCookie));
 							}
 							else
 							{
@@ -905,7 +905,7 @@ void StatelessConnectHandlerComponent::IncomingConnectionless(const TSharedPtr<c
 								LastServerSequence = *CurSequence & (MAX_PACKETID - 1);
 								LastClientSequence = *(CurSequence + 1) & (MAX_PACKETID - 1);
 
-								FMemory::Memcpy(AuthorisedCookie, Cookie, ARRAY_COUNT(AuthorisedCookie));
+								FMemory::Memcpy(AuthorisedCookie, Cookie, UE_ARRAY_COUNT(AuthorisedCookie));
 							}
 
 							bRestartedHandshake = bRestartHandshake;
diff --git a/Engine/Source/Runtime/Engine/Private/Particles/SubUVAnimation.cpp b/Engine/Source/Runtime/Engine/Private/Particles/SubUVAnimation.cpp
--- a/Engine/Source/Runtime/Engine/Private/Particles/SubUVAnimation.cpp
+++ b/Engine/Source/Runtime/Engine/Private/Particles/SubUVAnimation.cpp
@@ -483,7 +483,7 @@ int32 ComputeNeighborCount(int32 X, int32 Y, int32 GlobalX, int32 GlobalY, int32
 {
 	int32 NeighborCount = 0;
 
-	for (int32 NeighborIndex = 0; NeighborIndex < ARRAY_COUNT(GNeighbors); NeighborIndex++)
+	for (int32 NeighborIndex = 0; NeighborIndex < UE_ARRAY_COUNT(GNeighbors); NeighborIndex++)
 	{
 		int32 NeighborX = X + GNeighbors[NeighborIndex].X;
 		int32 NeighborY = Y + GNeighbors[NeighborIndex].Y;
diff --git a/Engine/Source/Runtime/Engine/Private/PhysicsEngine/BodyInstance.cpp b/Engine/Source/Runtime/Engine/Private/PhysicsEngine/BodyInstance.cpp
--- a/Engine/Source/Runtime/Engine/Private/PhysicsEngine/BodyInstance.cpp
+++ b/Engine/Source/Runtime/Engine/Private/PhysicsEngine/BodyInstance.cpp
@@ -222,12 +222,12 @@ bool FCollisionResponse::AddReponseToArray(ECollisionChannel Channel, ECollision
 
 void FCollisionResponse::UpdateArrayFromResponseContainer()
 {
-	ResponseArray.Empty(ARRAY_COUNT(ResponseToChannels.EnumArray));
+	ResponseArray.Empty(UE_ARRAY_COUNT(ResponseToChannels.EnumArray));
 
 	const FCollisionResponseContainer& DefaultResponse = FCollisionResponseContainer::GetDefaultResponseContainer();
 	const UCollisionProfile* CollisionProfile = UCollisionProfile::Get();
 
-	for (int32 i = 0; i < ARRAY_COUNT(ResponseToChannels.EnumArray); i++)
+	for (int32 i = 0; i < UE_ARRAY_COUNT(ResponseToChannels.EnumArray); i++)
 	{
 		// if not same as default
 		if (ResponseToChannels.EnumArray[i] != DefaultResponse.EnumArray[i])
diff --git a/Engine/Source/Runtime/Engine/Private/PhysicsEngine/Experimental/PhysicsInterfaceImmediatePhysX.cpp b/Engine/Source/Runtime/Engine/Private/PhysicsEngine/Experimental/PhysicsInterfaceImmediatePhysX.cpp
--- a/Engine/Source/Runtime/Engine/Private/PhysicsEngine/Experimental/PhysicsInterfaceImmediatePhysX.cpp
+++ b/Engine/Source/Runtime/Engine/Private/PhysicsEngine/Experimental/PhysicsInterfaceImmediatePhysX.cpp
@@ -1604,7 +1604,7 @@ bool FPhysicsInterface_ImmediatePhysX::LineTrace_Geom(FHitResult& OutHit, const
 						const bool bShapeIsComplex = ShapeRef.Shape.Geometry->getType() == PxGeometryType::Enum::eTRIANGLEMESH;
 						if((bTraceComplex && bShapeIsComplex) || (!bTraceComplex && !bShapeIsComplex))
 						{
-							const int32 ArraySize = ARRAY_COUNT(PHits);
+							const int32 ArraySize = UE_ARRAY_COUNT(PHits);
 							// #PHYS2 This may not work with shared shapes (GetTransform requires getActor to return non-nullptr) verify
 							PxTransform ShapeTransform = ShapeRef.Shape.LocalTM;
 							const PxI32 NumHits = PxGeometryQuery::raycast(U2PVector(InStart), U2PVector(Delta / DeltaMag), *ShapeRef.Shape.Geometry, ShapeTransform, DeltaMag, PHitFlags, ArraySize, PHits);
diff --git a/Engine/Source/Runtime/Engine/Private/PhysicsEngine/PhysUtils.cpp b/Engine/Source/Runtime/Engine/Private/PhysicsEngine/PhysUtils.cpp
--- a/Engine/Source/Runtime/Engine/Private/PhysicsEngine/PhysUtils.cpp
+++ b/Engine/Source/Runtime/Engine/Private/PhysicsEngine/PhysUtils.cpp
@@ -194,7 +194,7 @@ void FRigidBodyContactInfo::SwapOrder()
 /** Set the status of a particular channel in the structure. */
 bool FCollisionResponseContainer::SetResponse(ECollisionChannel Channel, ECollisionResponse NewResponse)
 {
-	if (Channel < ARRAY_COUNT(EnumArray))
+	if (Channel < UE_ARRAY_COUNT(EnumArray))
 	{
 		uint8& CurrentResponse = EnumArray[Channel];
 		if (CurrentResponse != NewResponse)
@@ -210,7 +210,7 @@ bool FCollisionResponseContainer::SetResponse(ECollisionChannel Channel, ECollis
 bool FCollisionResponseContainer::SetAllChannels(ECollisionResponse NewResponse)
 {
 	bool bHasChanged = false;
-	for(int32 i=0; i<ARRAY_COUNT(EnumArray); i++)
+	for(int32 i=0; i<UE_ARRAY_COUNT(EnumArray); i++)
 	{
 		uint8& CurrentResponse = EnumArray[i];
 		if (CurrentResponse != NewResponse)
@@ -225,7 +225,7 @@ bool FCollisionResponseContainer::SetAllChannels(ECollisionResponse NewResponse)
 bool FCollisionResponseContainer::ReplaceChannels(ECollisionResponse OldResponse, ECollisionResponse NewResponse)
 {
 	bool bHasChanged = false;
-	for (int32 i = 0; i < ARRAY_COUNT(EnumArray); i++)
+	for (int32 i = 0; i < UE_ARRAY_COUNT(EnumArray); i++)
 	{
 		uint8& CurrentResponse = EnumArray[i];
 		if(CurrentResponse == OldResponse)
@@ -240,7 +240,7 @@ bool FCollisionResponseContainer::ReplaceChannels(ECollisionResponse OldResponse
 FCollisionResponseContainer FCollisionResponseContainer::CreateMinContainer(const FCollisionResponseContainer& A, const FCollisionResponseContainer& B)
 {
 	FCollisionResponseContainer Result;
-	for(int32 i=0; i<ARRAY_COUNT(Result.EnumArray); i++)
+	for(int32 i=0; i<UE_ARRAY_COUNT(Result.EnumArray); i++)
 	{
 		Result.EnumArray[i] = FMath::Min(A.EnumArray[i], B.EnumArray[i]);
 	}
@@ -302,7 +302,7 @@ bool FPhysScene::ExecPxVis(const TCHAR* Cmd, FOutputDevice* Ar)
 	if ( FParse::Command(&Cmd,TEXT("PHYSX_CLEAR_ALL")) )
 	{
 		Ar->Logf(TEXT("Clearing all PhysX Debug Flags."));
-		for (int32 i = 0; i < ARRAY_COUNT(Flags); i++)
+		for (int32 i = 0; i < UE_ARRAY_COUNT(Flags); i++)
 		{
 			PScene->setVisualizationParameter(Flags[i].Flag, 0.0f);
 			bFoundFlag = true;
@@ -310,7 +310,7 @@ bool FPhysScene::ExecPxVis(const TCHAR* Cmd, FOutputDevice* Ar)
 	}
 	else
 	{
-		for (int32 i = 0; i < ARRAY_COUNT(Flags); i++)
+		for (int32 i = 0; i < UE_ARRAY_COUNT(Flags); i++)
 		{
 			// Parse out the command sent in and set only those flags
 			if (FParse::Command(&Cmd, Flags[i].Name))
diff --git a/Engine/Source/Runtime/Engine/Private/PhysicsEngine/PhysicsInterfacePhysX.cpp b/Engine/Source/Runtime/Engine/Private/PhysicsEngine/PhysicsInterfacePhysX.cpp
--- a/Engine/Source/Runtime/Engine/Private/PhysicsEngine/PhysicsInterfacePhysX.cpp
+++ b/Engine/Source/Runtime/Engine/Private/PhysicsEngine/PhysicsInterfacePhysX.cpp
@@ -2469,7 +2469,7 @@ bool FPhysicsInterface_PhysX::LineTrace_Geom(FHitResult& OutHit, const FBodyInst
 						const bool bShapeIsSimple = (ShapeFilter.word3 & EPDF_SimpleCollision) != 0;
 						if((bTraceComplex && bShapeIsComplex) || (!bTraceComplex && bShapeIsSimple))
 						{
-							const int32 ArraySize = ARRAY_COUNT(PHits);
+							const int32 ArraySize = UE_ARRAY_COUNT(PHits);
 							
 							PxTransform ShapeTransform = U2PTransform(GetTransform(ShapeRef));
 							const PxI32 NumHits = PxGeometryQuery::raycast(U2PVector(InStart), U2PVector(Delta / DeltaMag), PShape->getGeometry().any(), ShapeTransform, DeltaMag, PHitFlags, ArraySize, PHits);
diff --git a/Engine/Source/Runtime/Engine/Private/PhysicsEngine/PhysicsInterfaceUtils.cpp b/Engine/Source/Runtime/Engine/Private/PhysicsEngine/PhysicsInterfaceUtils.cpp
--- a/Engine/Source/Runtime/Engine/Private/PhysicsEngine/PhysicsInterfaceUtils.cpp
+++ b/Engine/Source/Runtime/Engine/Private/PhysicsEngine/PhysicsInterfaceUtils.cpp
@@ -70,7 +70,7 @@ FCollisionFilterData CreateTraceQueryFilterData(const uint8 MyChannel, const boo
 	}
 
 	// word1 encodes 'what i block', word2 encodes 'what i touch'
-	for (int32 i = 0; i<ARRAY_COUNT(InCollisionResponseContainer.EnumArray); i++)
+	for (int32 i = 0; i<UE_ARRAY_COUNT(InCollisionResponseContainer.EnumArray); i++)
 	{
 		if (InCollisionResponseContainer.EnumArray[i] == ECR_Block)
 		{
diff --git a/Engine/Source/Runtime/Engine/Private/PlayerState.cpp b/Engine/Source/Runtime/Engine/Private/PlayerState.cpp
--- a/Engine/Source/Runtime/Engine/Private/PlayerState.cpp
+++ b/Engine/Source/Runtime/Engine/Private/PlayerState.cpp
@@ -59,7 +59,7 @@ void APlayerState::UpdatePing(float InPing)
 		//	(misses the latest ping update, but averages a full 4 seconds data)
 		RecalculateAvgPing();
 
-		CurPingBucket = (CurPingBucket + 1) % ARRAY_COUNT(PingBucket);
+		CurPingBucket = (CurPingBucket + 1) % UE_ARRAY_COUNT(PingBucket);
 		CurPingBucketTimestamp = CurTime;
 
 
@@ -100,7 +100,7 @@ void APlayerState::RecalculateAvgPing()
 	int32 Sum = 0;
 	int32 Count = 0;
 
-	for (uint8 i=0; i<ARRAY_COUNT(PingBucket); i++)
+	for (uint8 i=0; i<UE_ARRAY_COUNT(PingBucket); i++)
 	{
 		Sum += PingBucket[i].PingSum;
 		Count += PingBucket[i].PingCount;
@@ -120,12 +120,12 @@ void APlayerState::RecalculateAvgPing()
 	PingBucketV2[CurPingBucket].AvgPingV2 = NumValidValues > 0 ? SumV2 / NumValidValues : MAX_flt;
 
 	float AvgSumV2 = 0.0f;
-	for (int32 i = 0; i < ARRAY_COUNT(PingBucketV2); i++)
+	for (int32 i = 0; i < UE_ARRAY_COUNT(PingBucketV2); i++)
 	{
 		AvgSumV2 += PingBucketV2[i].AvgPingV2;
 	}
 
-	ExactPingV2 = AvgSumV2 / ARRAY_COUNT(PingBucketV2);
+	ExactPingV2 = AvgSumV2 / UE_ARRAY_COUNT(PingBucketV2);
 
 	// Calculate the average, and divide it by 4 to optimize replication
 	ExactPing = (Count > 0 ? ((float)Sum / (float)Count) : 0.f);
diff --git a/Engine/Source/Runtime/Engine/Private/Polygon.cpp b/Engine/Source/Runtime/Engine/Private/Polygon.cpp
--- a/Engine/Source/Runtime/Engine/Private/Polygon.cpp
+++ b/Engine/Source/Runtime/Engine/Private/Polygon.cpp
@@ -745,7 +745,7 @@ int32 FPoly::Triangulate( ABrush* InOwnerBrush, TArray<FPoly>& OutTriangles )
 		vtx.TangentY = FVector::ZeroVector;
 		vtx.TangentZ = FVector::ZeroVector;
 		vtx.Color = FColor(0, 0, 0);
-		for( int32 uvIndex=0; uvIndex<ARRAY_COUNT(vtx.UVs); ++uvIndex )
+		for( int32 uvIndex=0; uvIndex<UE_ARRAY_COUNT(vtx.UVs); ++uvIndex )
 		{
 			vtx.UVs[uvIndex] = FVector2D(0.f, 0.f);
 		}
diff --git a/Engine/Source/Runtime/Engine/Private/PrecomputedVolumetricLightmap.cpp b/Engine/Source/Runtime/Engine/Private/PrecomputedVolumetricLightmap.cpp
--- a/Engine/Source/Runtime/Engine/Private/PrecomputedVolumetricLightmap.cpp
+++ b/Engine/Source/Runtime/Engine/Private/PrecomputedVolumetricLightmap.cpp
@@ -73,7 +73,7 @@ struct FVolumetricLightmapBrickTextureSet
 		SkyBentNormal.Format = BrickData.SkyBentNormal.Format;
 		DirectionalLightShadowing.Format = BrickData.DirectionalLightShadowing.Format;
 
-		for (int32 i = 0; i < ARRAY_COUNT(SHCoefficients); i++)
+		for (int32 i = 0; i < UE_ARRAY_COUNT(SHCoefficients); i++)
 		{
 			SHCoefficients[i].Format = BrickData.SHCoefficients[i].Format;
 		}
@@ -81,7 +81,7 @@ struct FVolumetricLightmapBrickTextureSet
 		AmbientVector.CreateTargetTexture(BrickDataDimensions);
 		AmbientVector.CreateUAV();
 
-		for (int32 i = 0; i < ARRAY_COUNT(SHCoefficients); i++)
+		for (int32 i = 0; i < UE_ARRAY_COUNT(SHCoefficients); i++)
 		{
 			SHCoefficients[i].CreateTargetTexture(BrickDataDimensions);
 			SHCoefficients[i].CreateUAV();
@@ -100,7 +100,7 @@ struct FVolumetricLightmapBrickTextureSet
 	void Release()
 	{
 		AmbientVector.Texture.SafeRelease();
-		for (int32 i = 0; i < ARRAY_COUNT(SHCoefficients); i++)
+		for (int32 i = 0; i < UE_ARRAY_COUNT(SHCoefficients); i++)
 		{
 			SHCoefficients[i].Texture.SafeRelease();
 		}
@@ -108,7 +108,7 @@ struct FVolumetricLightmapBrickTextureSet
 		DirectionalLightShadowing.Texture.SafeRelease();
 
 		AmbientVector.UAV.SafeRelease();
-		for (int32 i = 0; i < ARRAY_COUNT(SHCoefficients); i++)
+		for (int32 i = 0; i < UE_ARRAY_COUNT(SHCoefficients); i++)
 		{
 			SHCoefficients[i].UAV.SafeRelease();
 		}
@@ -212,7 +212,7 @@ FArchive& operator<<(FArchive& Ar,FPrecomputedVolumetricLightmapData& Volume)
 
 	Ar << Volume.BrickData.AmbientVector;
 
-	for (int32 i = 0; i < ARRAY_COUNT(Volume.BrickData.SHCoefficients); i++)
+	for (int32 i = 0; i < UE_ARRAY_COUNT(Volume.BrickData.SHCoefficients); i++)
 	{
 		Ar << Volume.BrickData.SHCoefficients[i];
 	}
@@ -283,7 +283,7 @@ int32 FVolumetricLightmapBrickData::GetMinimumVoxelSize() const
 	check(AmbientVector.Format != PF_Unknown);
 	int32 VoxelSize = GPixelFormats[AmbientVector.Format].BlockBytes;
 
-	for (int32 i = 0; i < ARRAY_COUNT(SHCoefficients); i++)
+	for (int32 i = 0; i < UE_ARRAY_COUNT(SHCoefficients); i++)
 	{
 		VoxelSize += GPixelFormats[SHCoefficients[i].Format].BlockBytes;
 	}
@@ -340,7 +340,7 @@ ENGINE_API void FPrecomputedVolumetricLightmapData::InitRHI()
 		{
 			BrickData.AmbientVector.CreateTexture(BrickDataDimensions);
 
-			for (int32 i = 0; i < ARRAY_COUNT(BrickData.SHCoefficients); i++)
+			for (int32 i = 0; i < UE_ARRAY_COUNT(BrickData.SHCoefficients); i++)
 			{
 				BrickData.SHCoefficients[i].CreateTexture(BrickDataDimensions);
 			}
@@ -610,7 +610,7 @@ ENGINE_API void FPrecomputedVolumetricLightmapData::AddToSceneData(FPrecomputedV
 
 	SceneData->BrickDataDimensions = GVolumetricLightmapBrickAtlas.TextureSet.BrickDataDimensions;
 	SceneData->BrickData.AmbientVector = GVolumetricLightmapBrickAtlas.TextureSet.AmbientVector;
-	for (int32 i = 0; i < ARRAY_COUNT(SceneData->BrickData.SHCoefficients); i++)
+	for (int32 i = 0; i < UE_ARRAY_COUNT(SceneData->BrickData.SHCoefficients); i++)
 	{
 		SceneData->BrickData.SHCoefficients[i] = GVolumetricLightmapBrickAtlas.TextureSet.SHCoefficients[i];
 	}
@@ -691,7 +691,7 @@ ENGINE_API void FPrecomputedVolumetricLightmapData::RemoveFromSceneData(FPrecomp
 
 	SceneData->BrickDataDimensions = GVolumetricLightmapBrickAtlas.TextureSet.BrickDataDimensions;
 	SceneData->BrickData.AmbientVector = GVolumetricLightmapBrickAtlas.TextureSet.AmbientVector;
-	for (int32 i = 0; i < ARRAY_COUNT(SceneData->BrickData.SHCoefficients); i++)
+	for (int32 i = 0; i < UE_ARRAY_COUNT(SceneData->BrickData.SHCoefficients); i++)
 	{
 		SceneData->BrickData.SHCoefficients[i] = GVolumetricLightmapBrickAtlas.TextureSet.SHCoefficients[i];
 	}
@@ -909,7 +909,7 @@ void CopyDataIntoAtlas(FRHICommandList& RHICmdList, ERHIFeatureLevel::Type Featu
 		FComputeShaderUtils::Dispatch(RHICmdList, *ComputeShader, Parameters, FIntVector(NumBricks, 1, 1));
 	}
 
-	for (int32 i = 0; i < ARRAY_COUNT(SrcData.SHCoefficients); i++)
+	for (int32 i = 0; i < UE_ARRAY_COUNT(SrcData.SHCoefficients); i++)
 	{
 		TShaderMap<FGlobalShaderType>* GlobalShaderMap = GetGlobalShaderMap(FeatureLevel);
 
@@ -942,7 +942,7 @@ void FVolumetricLightmapBrickAtlas::Insert(int32 Index, FPrecomputedVolumetricLi
 	else
 	{
 		check(TextureSet.AmbientVector.Format == Data->BrickData.AmbientVector.Format);
-		for (int32 i = 0; i < ARRAY_COUNT(TextureSet.SHCoefficients); i++)
+		for (int32 i = 0; i < UE_ARRAY_COUNT(TextureSet.SHCoefficients); i++)
 		{
 			check(TextureSet.SHCoefficients[i].Format == Data->BrickData.SHCoefficients[i].Format);
 		}
diff --git a/Engine/Source/Runtime/Engine/Private/Scene.cpp b/Engine/Source/Runtime/Engine/Private/Scene.cpp
--- a/Engine/Source/Runtime/Engine/Private/Scene.cpp
+++ b/Engine/Source/Runtime/Engine/Private/Scene.cpp
@@ -938,7 +938,7 @@ FPostProcessSettings::FPostProcessSettings(const FPostProcessSettings& Settings)
 	, WeightedBlendables(Settings.WeightedBlendables)
 	//, Blendables_DEPRECATED(Settings.Blendables_DEPRECATED)
 {
-	for (int32 i = 0; i < ARRAY_COUNT(LensFlareTints); i++)
+	for (int32 i = 0; i < UE_ARRAY_COUNT(LensFlareTints); i++)
 		LensFlareTints[i] = Settings.LensFlareTints[i];
 }
 	
diff --git a/Engine/Source/Runtime/Engine/Private/ShaderCompiler/ShaderCompilerXGE.cpp b/Engine/Source/Runtime/Engine/Private/ShaderCompiler/ShaderCompilerXGE.cpp
--- a/Engine/Source/Runtime/Engine/Private/ShaderCompiler/ShaderCompilerXGE.cpp
+++ b/Engine/Source/Runtime/Engine/Private/ShaderCompiler/ShaderCompilerXGE.cpp
@@ -109,7 +109,7 @@ bool FShaderCompileXGEThreadRunnable_XmlInterface::IsSupported()
 	{
 		IPlatformFile& PlatformFile = FPlatformFileManager::Get().GetPlatformFile();
 
-		for (int PathIndex = 0; PathIndex < ARRAY_COUNT(Paths); PathIndex++)
+		for (int PathIndex = 0; PathIndex < UE_ARRAY_COUNT(Paths); PathIndex++)
 		{
 			if (PlatformFile.FileExists(Paths[PathIndex]))
 			{
@@ -349,7 +349,7 @@ static void WriteScriptFileFooter(FArchive* ScriptFile)
 		"\t</Project>\r\n"
 		"</BuildSet>\r\n";
 
-	ScriptFile->Serialize((void*)HeaderFooter, sizeof(ANSICHAR) * (ARRAY_COUNT(HeaderFooter) - 1));
+	ScriptFile->Serialize((void*)HeaderFooter, sizeof(ANSICHAR) * (UE_ARRAY_COUNT(HeaderFooter) - 1));
 }
 
 void FShaderCompileXGEThreadRunnable_XmlInterface::GatherResultsFromXGE()
diff --git a/Engine/Source/Runtime/Engine/Private/ShadowMap.cpp b/Engine/Source/Runtime/Engine/Private/ShadowMap.cpp
--- a/Engine/Source/Runtime/Engine/Private/ShadowMap.cpp
+++ b/Engine/Source/Runtime/Engine/Private/ShadowMap.cpp
@@ -503,7 +503,7 @@ FShadowMap2D::FShadowMap2D() :
 	CoordinateScale(FVector2D(0, 0)),
 	CoordinateBias(FVector2D(0, 0))
 {
-	for (int Channel = 0; Channel < ARRAY_COUNT(bChannelValid); Channel++)
+	for (int Channel = 0; Channel < UE_ARRAY_COUNT(bChannelValid); Channel++)
 	{
 		bChannelValid[Channel] = false;
 	}
@@ -514,7 +514,7 @@ FShadowMap2D::FShadowMap2D(const TMap<ULightComponent*,FShadowMapData2D*>& Shado
 	CoordinateScale(FVector2D(0, 0)),
 	CoordinateBias(FVector2D(0, 0))
 {
-	for (int Channel = 0; Channel < ARRAY_COUNT(bChannelValid); Channel++)
+	for (int Channel = 0; Channel < UE_ARRAY_COUNT(bChannelValid); Channel++)
 	{
 		bChannelValid[Channel] = false;
 	}
@@ -531,7 +531,7 @@ FShadowMap2D::FShadowMap2D(TArray<FGuid> LightGuids)
 	, CoordinateScale(FVector2D(0, 0))
 	, CoordinateBias(FVector2D(0, 0))
 {
-	for (int Channel = 0; Channel < ARRAY_COUNT(bChannelValid); Channel++)
+	for (int Channel = 0; Channel < UE_ARRAY_COUNT(bChannelValid); Channel++)
 	{
 		bChannelValid[Channel] = false;
 	}
@@ -561,7 +561,7 @@ void FShadowMap2D::Serialize(FArchive& Ar)
 
 	Ar << CoordinateScale << CoordinateBias;
 
-	for (int Channel = 0; Channel < ARRAY_COUNT(bChannelValid); Channel++)
+	for (int Channel = 0; Channel < UE_ARRAY_COUNT(bChannelValid); Channel++)
 	{
 		Ar << bChannelValid[Channel];
 	}
@@ -1108,7 +1108,7 @@ int32 FShadowMap2D::EncodeSingleTexture(ULevel* LightingScenario, FShadowMapPend
 							ExtrapolatedFilterableComponents[i] = 0;
 						}
 
-						for (int32 NeighborIndex = 0; NeighborIndex < ARRAY_COUNT(Neighbors); NeighborIndex++)
+						for (int32 NeighborIndex = 0; NeighborIndex < UE_ARRAY_COUNT(Neighbors); NeighborIndex++)
 						{
 							if (static_cast<int32>(DestY) + Neighbors[NeighborIndex].Y >= 0 
 								&& DestY + Neighbors[NeighborIndex].Y < MipSizeY
diff --git a/Engine/Source/Runtime/Engine/Private/StaticMeshLight.cpp b/Engine/Source/Runtime/Engine/Private/StaticMeshLight.cpp
--- a/Engine/Source/Runtime/Engine/Private/StaticMeshLight.cpp
+++ b/Engine/Source/Runtime/Engine/Private/StaticMeshLight.cpp
@@ -38,7 +38,7 @@ static void GetStaticLightingVertex(
 	OutVertex.WorldTangentY = LocalToWorld.TransformVector(VertexBuffer.VertexTangentY(VertexIndex)).GetSafeNormal();
 	OutVertex.WorldTangentZ = LocalToWorldInverseTranspose.TransformVector(VertexBuffer.VertexTangentZ(VertexIndex)).GetSafeNormal();
 
-	checkSlow(VertexBuffer.GetNumTexCoords() <= ARRAY_COUNT(OutVertex.TextureCoordinates));
+	checkSlow(VertexBuffer.GetNumTexCoords() <= UE_ARRAY_COUNT(OutVertex.TextureCoordinates));
 	for(uint32 LightmapTextureCoordinateIndex = 0;LightmapTextureCoordinateIndex < VertexBuffer.GetNumTexCoords();LightmapTextureCoordinateIndex++)
 	{
 		OutVertex.TextureCoordinates[LightmapTextureCoordinateIndex] = VertexBuffer.GetVertexUV(VertexIndex,LightmapTextureCoordinateIndex);
diff --git a/Engine/Source/Runtime/Engine/Private/SystemSettings.cpp b/Engine/Source/Runtime/Engine/Private/SystemSettings.cpp
--- a/Engine/Source/Runtime/Engine/Private/SystemSettings.cpp
+++ b/Engine/Source/Runtime/Engine/Private/SystemSettings.cpp
@@ -42,11 +42,11 @@ static inline const TCHAR* GetSectionName(bool bIsEditor)
 		static int32 SystemSettingsStrLen = FCString::Strlen(ReturnedOverrideName);
 		if (FCString::Strnicmp(*OverrideSubName, TEXT("SystemSettings"), SystemSettingsStrLen) == 0)
 		{
-			FCString::Strcpy(ReturnedOverrideName + SystemSettingsStrLen, ARRAY_COUNT(ReturnedOverrideName), (*OverrideSubName + SystemSettingsStrLen));
+			FCString::Strcpy(ReturnedOverrideName + SystemSettingsStrLen, UE_ARRAY_COUNT(ReturnedOverrideName), (*OverrideSubName + SystemSettingsStrLen));
 		}
 		else
 		{
-			FCString::Strcpy(ReturnedOverrideName + SystemSettingsStrLen, ARRAY_COUNT(ReturnedOverrideName), *OverrideSubName);
+			FCString::Strcpy(ReturnedOverrideName + SystemSettingsStrLen, UE_ARRAY_COUNT(ReturnedOverrideName), *OverrideSubName);
 		}
 		return ReturnedOverrideName;
 	}
diff --git a/Engine/Source/Runtime/Engine/Private/Texture2D.cpp b/Engine/Source/Runtime/Engine/Private/Texture2D.cpp
--- a/Engine/Source/Runtime/Engine/Private/Texture2D.cpp
+++ b/Engine/Source/Runtime/Engine/Private/Texture2D.cpp
@@ -1465,7 +1465,7 @@ FTexture2DResource::FTexture2DResource( UTexture2D* InOwner, int32 InitialMipCou
 	bSRGB = InOwner->SRGB;
 
 	check(InitialMipCount>0);
-	check(ARRAY_COUNT(MipData)>=GMaxTextureMipCount);
+	check(UE_ARRAY_COUNT(MipData)>=GMaxTextureMipCount);
 
 	// Keep track of first miplevel to use.
 	CurrentFirstMip = InOwner->GetNumMips() - InitialMipCount;
@@ -1498,7 +1498,7 @@ FTexture2DResource::~FTexture2DResource()
 		});
 
 	// Make sure we're not leaking memory if InitRHI has never been called.
-	for( int32 MipIndex=0; MipIndex<ARRAY_COUNT(MipData); MipIndex++ )
+	for( int32 MipIndex=0; MipIndex<UE_ARRAY_COUNT(MipData); MipIndex++ )
 	{
 		// free any mip data that was copied 
 		if( MipData[MipIndex] )
diff --git a/Engine/Source/Runtime/Engine/Private/TextureCube.cpp b/Engine/Source/Runtime/Engine/Private/TextureCube.cpp
--- a/Engine/Source/Runtime/Engine/Private/TextureCube.cpp
+++ b/Engine/Source/Runtime/Engine/Private/TextureCube.cpp
@@ -139,7 +139,7 @@ public:
 		//Initialize the MipData array
 		for ( int32 FaceIndex=0;FaceIndex<6; FaceIndex++)
 		{
-			for( int32 MipIndex=0; MipIndex<ARRAY_COUNT(MipData[FaceIndex]); MipIndex++ )
+			for( int32 MipIndex=0; MipIndex<UE_ARRAY_COUNT(MipData[FaceIndex]); MipIndex++ )
 			{
 				MipData[FaceIndex][MipIndex] = NULL;
 			}
@@ -183,7 +183,7 @@ public:
 		// Make sure we're not leaking memory if InitRHI has never been called.
 		for (int32 i=0; i<6; i++)
 		{
-			for( int32 MipIndex=0; MipIndex<ARRAY_COUNT(MipData[i]); MipIndex++ )
+			for( int32 MipIndex=0; MipIndex<UE_ARRAY_COUNT(MipData[i]); MipIndex++ )
 			{
 				// free any mip data that was copied 
 				if( MipData[i][MipIndex] )
diff --git a/Engine/Source/Runtime/Engine/Private/UnrealEngine.cpp b/Engine/Source/Runtime/Engine/Private/UnrealEngine.cpp
--- a/Engine/Source/Runtime/Engine/Private/UnrealEngine.cpp
+++ b/Engine/Source/Runtime/Engine/Private/UnrealEngine.cpp
@@ -1411,7 +1411,7 @@ void UEngine::Init(IEngineLoop* InEngineLoop)
 #if !UE_BUILD_SHIPPING
 	// Check for overrides to the default map on the command line
 	TCHAR MapName[512];
-	if ( FParse::Value(FCommandLine::Get(), TEXT("DEFAULTMAP="), MapName, ARRAY_COUNT(MapName)) )
+	if ( FParse::Value(FCommandLine::Get(), TEXT("DEFAULTMAP="), MapName, UE_ARRAY_COUNT(MapName)) )
 	{
 		UE_LOG(LogEngine, Log, TEXT("Overriding default map to %s"), MapName);
 
@@ -3746,7 +3746,7 @@ static void BufferOverflowFunction(SIZE_T BufferSize, const ANSICHAR* Buffer)
 	ANSICHAR LocalBuffer[32];
 	LocalBuffer[0] = LocalBuffer[31] = 0; //if BufferSize is 0 then there's nothing to print out!
 
-	BufferSize = FMath::Min<SIZE_T>(BufferSize, ARRAY_COUNT(LocalBuffer)-1);
+	BufferSize = FMath::Min<SIZE_T>(BufferSize, UE_ARRAY_COUNT(LocalBuffer)-1);
 
 	for( uint32 i = 0; i < BufferSize; i++ ) 
 	{
@@ -7600,7 +7600,7 @@ bool UEngine::HandleObjCommand( const TCHAR* Cmd, FOutputDevice& Ar )
 
 		if ( !ParseObject<UClass>( Cmd, TEXT("CLASS="), Cls, ANY_PACKAGE ) || !ParseObject(Cmd,TEXT("NAME="), Cls, Obj, ANY_PACKAGE) )
 		{
-			if ( FParse::Token(Cmd,ObjectName,ARRAY_COUNT(ObjectName), 1) )
+			if ( FParse::Token(Cmd,ObjectName,UE_ARRAY_COUNT(ObjectName), 1) )
 			{
 				Obj = FindObject<UObject>(ANY_PACKAGE,ObjectName);
 			}
@@ -7956,7 +7956,7 @@ bool UEngine::HandleGetIniCommand(const TCHAR* Cmd, FOutputDevice& Ar)
 	TCHAR IniPlusSectionName[256];
 	TCHAR KeyName[256];
 
-	if (FParse::Token(Cmd, IniPlusSectionName, ARRAY_COUNT(IniPlusSectionName), true))
+	if (FParse::Token(Cmd, IniPlusSectionName, UE_ARRAY_COUNT(IniPlusSectionName), true))
 	{
 		FString IniPlusSection = IniPlusSectionName;
 		int32 IniDelim = IniPlusSection.Find(TEXT(":"));
@@ -8002,7 +8002,7 @@ bool UEngine::HandleGetIniCommand(const TCHAR* Cmd, FOutputDevice& Ar)
 
 		if (!IniName.IsEmpty() && !SectionName.IsEmpty())
 		{
-			if (FParse::Token(Cmd, KeyName, ARRAY_COUNT(KeyName), true))
+			if (FParse::Token(Cmd, KeyName, UE_ARRAY_COUNT(KeyName), true))
 			{
 				TArray<FString> Values;
 
@@ -8386,7 +8386,7 @@ bool UEngine::PerformError(const TCHAR* Cmd, FOutputDevice& Ar)
 		// stack overflow test - this case should be caught by /GS (Buffer Overflow Check) compile option
 		ANSICHAR SrcBuffer[] = "12345678901234567890123456789012345678901234567890";
 		SetCrashType(ECrashType::Debug);
-		BufferOverflowFunction(ARRAY_COUNT(SrcBuffer), SrcBuffer);
+		BufferOverflowFunction(UE_ARRAY_COUNT(SrcBuffer), SrcBuffer);
 		return true;
 	}
 	else if (FParse::Command(&Cmd, TEXT("CRTINVALID")))
@@ -10423,7 +10423,7 @@ FString appGetStartupMap(const TCHAR* CommandLine)
 #endif // UE_BUILD_SHIPPING
 
 	const TCHAR* Tmp = CommandLine ? CommandLine : TEXT("");
-	if (!FParse::Token(Tmp, Parm, ARRAY_COUNT(Parm), 0) || Parm[0] == '-')
+	if (!FParse::Token(Tmp, Parm, UE_ARRAY_COUNT(Parm), 0) || Parm[0] == '-')
 	{
 		const UGameMapsSettings* GameMapsSettings = GetDefault<UGameMapsSettings>();
 		FCString::Strcpy(Parm, *(GameMapsSettings->GetGameDefaultMap() + GameMapsSettings->LocalMapOptions));
@@ -11484,12 +11484,12 @@ void UEngine::SpawnServerActors(UWorld *World)
 	{
 		TCHAR Str[2048];
 		const TCHAR* Ptr = * FullServerActors[i];
-		if( FParse::Token( Ptr, Str, ARRAY_COUNT(Str), 1 ) )
+		if( FParse::Token( Ptr, Str, UE_ARRAY_COUNT(Str), 1 ) )
 		{
 			UE_LOG(LogNet, Log, TEXT("Spawning: %s"), Str );
 			UClass* HelperClass = StaticLoadClass( AActor::StaticClass(), NULL, Str, NULL, LOAD_None, NULL );
 			AActor* Actor = World->SpawnActor( HelperClass );
-			while( Actor && FParse::Token(Ptr,Str,ARRAY_COUNT(Str),1) )
+			while( Actor && FParse::Token(Ptr,Str,UE_ARRAY_COUNT(Str),1) )
 			{
 				TCHAR* Value = FCString::Strchr(Str,'=');
 				if( Value )
diff --git a/Engine/Source/Runtime/Engine/Private/UserInterface/Canvas.cpp b/Engine/Source/Runtime/Engine/Private/UserInterface/Canvas.cpp
--- a/Engine/Source/Runtime/Engine/Private/UserInterface/Canvas.cpp
+++ b/Engine/Source/Runtime/Engine/Private/UserInterface/Canvas.cpp
@@ -1551,7 +1551,7 @@ void UCanvas::ClippedStrLen( const UFont* Font, float ScaleX, float ScaleY, int3
 void VARARGS UCanvas::WrappedStrLenf( const UFont* Font, float ScaleX, float ScaleY, int32& XL, int32& YL, const TCHAR* Fmt, ... ) 
 {
 	TCHAR Text[4096];
-	GET_VARARGS( Text, ARRAY_COUNT(Text), ARRAY_COUNT(Text)-1, Fmt, Fmt );
+	GET_VARARGS( Text, UE_ARRAY_COUNT(Text), UE_ARRAY_COUNT(Text)-1, Fmt, Fmt );
 
 	FFontRenderInfo Info;
 	WrappedPrint( false, 0.0f, 0.0f, XL, YL, Font, ScaleX, ScaleY, false, false, Text, Info ); 
diff --git a/Engine/Source/Runtime/Engine/Private/UserInterface/PlayerInput.cpp b/Engine/Source/Runtime/Engine/Private/UserInterface/PlayerInput.cpp
--- a/Engine/Source/Runtime/Engine/Private/UserInterface/PlayerInput.cpp
+++ b/Engine/Source/Runtime/Engine/Private/UserInterface/PlayerInput.cpp
@@ -1402,7 +1402,7 @@ void UPlayerInput::ClearSmoothing()
 float UPlayerInput::SmoothMouse(float aMouse, uint8& SampleCount, int32 Index)
 {
 	check(Index >= 0);
-	check(Index < ARRAY_COUNT(ZeroTime));
+	check(Index < UE_ARRAY_COUNT(ZeroTime));
 
 	UWorld* World = GetWorld();
 	if (World)
@@ -2004,7 +2004,7 @@ bool UPlayerInput::Exec(UWorld* InWorld, const TCHAR* Str,FOutputDevice& Ar)
 {
 	TCHAR Temp[256];
 
-	if( FParse::Command( &Str, TEXT("KEYBINDING") ) && FParse::Token( Str, Temp, ARRAY_COUNT(Temp), 0 ) )
+	if( FParse::Command( &Str, TEXT("KEYBINDING") ) && FParse::Token( Str, Temp, UE_ARRAY_COUNT(Temp), 0 ) )
 	{
 		FKey const Key(Temp);
 		if (Key.IsValid())
@@ -2021,7 +2021,7 @@ bool UPlayerInput::Exec(UWorld* InWorld, const TCHAR* Str,FOutputDevice& Ar)
 
 		return 1;
 	}
-	else if( !bExecutingBindCommand && FParse::Token( Str, Temp, ARRAY_COUNT(Temp), 0 ) )
+	else if( !bExecutingBindCommand && FParse::Token( Str, Temp, UE_ARRAY_COUNT(Temp), 0 ) )
 	{
 		FKey const Key(Temp);
 		if(Key.IsValid())
diff --git a/Engine/Source/Runtime/Engine/Public/LightMap.h b/Engine/Source/Runtime/Engine/Public/LightMap.h
--- a/Engine/Source/Runtime/Engine/Public/LightMap.h
+++ b/Engine/Source/Runtime/Engine/Public/LightMap.h
@@ -213,7 +213,7 @@ public:
 
 	void GetReferencedTextures(TArray<UTexture2D*>& OutTextures) const
 	{
-		for (int32 BasisIndex = 0; BasisIndex < ARRAY_COUNT(Textures); BasisIndex++)
+		for (int32 BasisIndex = 0; BasisIndex < UE_ARRAY_COUNT(Textures); BasisIndex++)
 		{
 			if (Textures[BasisIndex])
 			{
diff --git a/Engine/Source/Runtime/Engine/Public/Net/DataChannel.h b/Engine/Source/Runtime/Engine/Public/Net/DataChannel.h
--- a/Engine/Source/Runtime/Engine/Public/Net/DataChannel.h
+++ b/Engine/Source/Runtime/Engine/Public/Net/DataChannel.h
@@ -50,7 +50,7 @@ private:
 		static bool bInitialized = false;
 		if (!bInitialized)
 		{
-			for (int32 i = 0; i < ARRAY_COUNT(Names); i++)
+			for (int32 i = 0; i < UE_ARRAY_COUNT(Names); i++)
 			{
 				Names[i] = TEXT("");
 			}
diff --git a/Engine/Source/Runtime/Engine/Public/SceneManagement.h b/Engine/Source/Runtime/Engine/Public/SceneManagement.h
--- a/Engine/Source/Runtime/Engine/Public/SceneManagement.h
+++ b/Engine/Source/Runtime/Engine/Public/SceneManagement.h
@@ -614,7 +614,7 @@ public:
 		InvUniformPenumbraSize(FVector4(0, 0, 0, 0)),
 		Type(SMIT_None)
 	{
-		for (int Channel = 0; Channel < ARRAY_COUNT(bChannelValid); Channel++)
+		for (int Channel = 0; Channel < UE_ARRAY_COUNT(bChannelValid); Channel++)
 		{
 			bChannelValid[Channel] = false;
 		}
diff --git a/Engine/Source/Runtime/Experimental/Chaos/Private/Chaos/PBDCollisionConstraint.cpp b/Engine/Source/Runtime/Experimental/Chaos/Private/Chaos/PBDCollisionConstraint.cpp
--- a/Engine/Source/Runtime/Experimental/Chaos/Private/Chaos/PBDCollisionConstraint.cpp
+++ b/Engine/Source/Runtime/Experimental/Chaos/Private/Chaos/PBDCollisionConstraint.cpp
@@ -178,7 +178,7 @@ struct FStatHelper2
 template <>
 struct FStatHelper2<true>
 {
-	int32 BucketCount[ARRAY_COUNT(BucketSizes2)];
+	int32 BucketCount[UE_ARRAY_COUNT(BucketSizes2)];
 	int32 MaxCount;
 
 	FStatHelper2()
@@ -189,7 +189,7 @@ struct FStatHelper2<true>
 
 	void Record(int32 Count)
 	{
-		for (int32 BucketIdx = 1; BucketIdx < ARRAY_COUNT(BucketSizes2); ++BucketIdx)
+		for (int32 BucketIdx = 1; BucketIdx < UE_ARRAY_COUNT(BucketSizes2); ++BucketIdx)
 		{
 			if (Count >= BucketSizes2[BucketIdx - 1] && Count < BucketSizes2[BucketIdx])
 			{
@@ -216,10 +216,10 @@ struct FStatHelper2<true>
 		}
 
 		const float CountPerChar = MaxBucketCount / 20.f;
-		for (int32 Idx = 1; Idx < ARRAY_COUNT(BucketSizes2); ++Idx)
+		for (int32 Idx = 1; Idx < UE_ARRAY_COUNT(BucketSizes2); ++Idx)
 		{
 			int32 NumChars = BucketCount[Idx] / CountPerChar;
-			if (Idx < ARRAY_COUNT(BucketSizes2) - 1)
+			if (Idx < UE_ARRAY_COUNT(BucketSizes2) - 1)
 			{
 				OutLog += FString::Printf(TEXT("\t[%4d - %4d) (%4d) |"), BucketSizes2[Idx - 1], BucketSizes2[Idx], BucketCount[Idx]);
 			}
diff --git a/Engine/Source/Runtime/HeadMountedDisplay/Private/XRMotionControllerBase.cpp b/Engine/Source/Runtime/HeadMountedDisplay/Private/XRMotionControllerBase.cpp
--- a/Engine/Source/Runtime/HeadMountedDisplay/Private/XRMotionControllerBase.cpp
+++ b/Engine/Source/Runtime/HeadMountedDisplay/Private/XRMotionControllerBase.cpp
@@ -89,7 +89,7 @@ ETrackingStatus FXRMotionControllerBase::GetControllerTrackingStatus(const int32
 
 void FXRMotionControllerBase::EnumerateSources(TArray<FMotionControllerSource>& SourcesOut) const
 {
-	const int32 HandsCount = ARRAY_COUNT(XRMotionControllerBase_Impl::LegacyHandMappings);
+	const int32 HandsCount = UE_ARRAY_COUNT(XRMotionControllerBase_Impl::LegacyHandMappings);
 	ensure(HandsCount == (int32)EControllerHand::ControllerHand_Count);
 
 	for (int32 HandIndex = 0; HandIndex < HandsCount; ++HandIndex)
@@ -100,7 +100,7 @@ void FXRMotionControllerBase::EnumerateSources(TArray<FMotionControllerSource>&
 
 bool FXRMotionControllerBase::GetHandEnumForSourceName(const FName Source, EControllerHand& OutHand)
 {
-	const int32 HandsCount = ARRAY_COUNT(XRMotionControllerBase_Impl::LegacyHandMappings);
+	const int32 HandsCount = UE_ARRAY_COUNT(XRMotionControllerBase_Impl::LegacyHandMappings);
 	ensure(HandsCount == (int32)EControllerHand::ControllerHand_Count);
 
 	bool bFound = false;
diff --git a/Engine/Source/Runtime/ImageWrapper/Private/Formats/ExrImageWrapper.cpp b/Engine/Source/Runtime/ImageWrapper/Private/Formats/ExrImageWrapper.cpp
--- a/Engine/Source/Runtime/ImageWrapper/Private/Formats/ExrImageWrapper.cpp
+++ b/Engine/Source/Runtime/ImageWrapper/Private/Formats/ExrImageWrapper.cpp
@@ -270,7 +270,7 @@ const char* FExrImageWrapper::GetRawChannelName(int ChannelIndex) const
 		break;
 		case ERGBFormat::Gray:
 		{
-			check(ChannelIndex < ARRAY_COUNT(GrayChannelNames));
+			check(ChannelIndex < UE_ARRAY_COUNT(GrayChannelNames));
 			ChannelNames = GrayChannelNames;
 		}
 		break;
diff --git a/Engine/Source/Runtime/InstallBundleManager/Public/InstallBundleManagerInterface.h b/Engine/Source/Runtime/InstallBundleManager/Public/InstallBundleManagerInterface.h
--- a/Engine/Source/Runtime/InstallBundleManager/Public/InstallBundleManagerInterface.h
+++ b/Engine/Source/Runtime/InstallBundleManager/Public/InstallBundleManagerInterface.h
@@ -42,7 +42,7 @@ inline const TCHAR* LexToString(EInstallBundleModuleInitResult Result)
 		TEXT("NoInternetConnectionError"),
 		TEXT("ConfigurationError"),
 	};
-	static_assert(static_cast<UnderType>(EInstallBundleModuleInitResult::Count) == ARRAY_COUNT(Strings), "");
+	static_assert(static_cast<UnderType>(EInstallBundleModuleInitResult::Count) == UE_ARRAY_COUNT(Strings), "");
 
 	return Strings[static_cast<UnderType>(Result)];
 }
@@ -72,7 +72,7 @@ inline const TCHAR* LexToString(EInstallBundleResult Result)
 		TEXT("UserCancelledError"),
 		TEXT("InitializationError"),
 	};
-	static_assert(static_cast<UnderType>(EInstallBundleResult::Count) == ARRAY_COUNT(Strings), "");
+	static_assert(static_cast<UnderType>(EInstallBundleResult::Count) == UE_ARRAY_COUNT(Strings), "");
 
 	return Strings[static_cast<UnderType>(Result)];
 }
@@ -141,7 +141,7 @@ inline const TCHAR* LexToString(EInstallBundleStatus Status)
 		TEXT("Installed"),
 	};
 
-	static_assert(static_cast<UnderType>(EInstallBundleStatus::Count) == ARRAY_COUNT(Strings), "");
+	static_assert(static_cast<UnderType>(EInstallBundleStatus::Count) == UE_ARRAY_COUNT(Strings), "");
 
 	return Strings[static_cast<UnderType>(Status)];
 }
@@ -219,7 +219,7 @@ inline const TCHAR* LexToString(EInstallBundleContentState State)
 		TEXT("NeedsUpdate"),
 		TEXT("UpToDate"),
 	};
-	static_assert(static_cast<UnderType>(EInstallBundleContentState::Count) == ARRAY_COUNT(Strings), "");
+	static_assert(static_cast<UnderType>(EInstallBundleContentState::Count) == UE_ARRAY_COUNT(Strings), "");
 
 	return Strings[static_cast<UnderType>(State)];
 }
diff --git a/Engine/Source/Runtime/Landscape/Private/LandscapeEdit.cpp b/Engine/Source/Runtime/Landscape/Private/LandscapeEdit.cpp
--- a/Engine/Source/Runtime/Landscape/Private/LandscapeEdit.cpp
+++ b/Engine/Source/Runtime/Landscape/Private/LandscapeEdit.cpp
@@ -3083,15 +3083,15 @@ bool ALandscapeProxy::ExportToRawMesh(int32 InExportLOD, FMeshDescription& OutRa
 
 		//We need to not duplicate the vertex position, so we use the FIndexAndZ to achieve fast result
 		TArray<FIndexAndZ> VertIndexAndZ;
-		VertIndexAndZ.Reserve(ComponentSizeQuadsLOD*ComponentSizeQuadsLOD*ARRAY_COUNT(QuadPattern));
+		VertIndexAndZ.Reserve(ComponentSizeQuadsLOD*ComponentSizeQuadsLOD*UE_ARRAY_COUNT(QuadPattern));
 		int32 CurrentIndex = 0;
 		TMap<int32, FVector> IndexToPosition;
-		IndexToPosition.Reserve(ComponentSizeQuadsLOD*ComponentSizeQuadsLOD*ARRAY_COUNT(QuadPattern));
+		IndexToPosition.Reserve(ComponentSizeQuadsLOD*ComponentSizeQuadsLOD*UE_ARRAY_COUNT(QuadPattern));
 		for (int32 y = 0; y < ComponentSizeQuadsLOD; y++)
 		{
 			for (int32 x = 0; x < ComponentSizeQuadsLOD; x++)
 			{
-				for (int32 i = 0; i < ARRAY_COUNT(QuadPattern); i++)
+				for (int32 i = 0; i < UE_ARRAY_COUNT(QuadPattern); i++)
 				{
 					int32 VertexX = x + QuadPattern[i].X;
 					int32 VertexY = y + QuadPattern[i].Y;
@@ -3140,11 +3140,11 @@ bool ALandscapeProxy::ExportToRawMesh(int32 InExportLOD, FMeshDescription& OutRa
 		{
 			for (int32 x = 0; x < ComponentSizeQuadsLOD; x++)
 			{
-				FVector Positions[ARRAY_COUNT(QuadPattern)];
+				FVector Positions[UE_ARRAY_COUNT(QuadPattern)];
 				bool bProcess = bIgnoreBounds;
 
 				// Fill positions
-				for (int32 i = 0; i < ARRAY_COUNT(QuadPattern); i++)
+				for (int32 i = 0; i < UE_ARRAY_COUNT(QuadPattern); i++)
 				{
 					int32 VertexX = x + QuadPattern[i].X;
 					int32 VertexY = y + QuadPattern[i].Y;
@@ -3161,11 +3161,11 @@ bool ALandscapeProxy::ExportToRawMesh(int32 InExportLOD, FMeshDescription& OutRa
 				{
 					//Fill the vertexID we need
 					TArray<FVertexID> VertexIDs;
-					VertexIDs.Reserve(ARRAY_COUNT(QuadPattern));
+					VertexIDs.Reserve(UE_ARRAY_COUNT(QuadPattern));
 					TArray<FVertexInstanceID> VertexInstanceIDs;
-					VertexInstanceIDs.Reserve(ARRAY_COUNT(QuadPattern));
+					VertexInstanceIDs.Reserve(UE_ARRAY_COUNT(QuadPattern));
 					// Fill positions
-					for (int32 i = 0; i < ARRAY_COUNT(QuadPattern); i++)
+					for (int32 i = 0; i < UE_ARRAY_COUNT(QuadPattern); i++)
 					{
 						int32 DuplicateLowestIndex = FindPreviousIndex(CurrentIndex);
 						FVertexID VertexID;
@@ -3205,7 +3205,7 @@ bool ALandscapeProxy::ExportToRawMesh(int32 InExportLOD, FMeshDescription& OutRa
 							VertexInstanceIDs.Add(OutRawMesh.CreateVertexInstance(VertexIDs[5]));
 
 							// Fill other vertex data
-							for (int32 i = 0; i < ARRAY_COUNT(QuadPattern); i++)
+							for (int32 i = 0; i < UE_ARRAY_COUNT(QuadPattern); i++)
 							{
 								int32 VertexX = x + QuadPattern[i].X;
 								int32 VertexY = y + QuadPattern[i].Y;
@@ -3247,7 +3247,7 @@ bool ALandscapeProxy::ExportToRawMesh(int32 InExportLOD, FMeshDescription& OutRa
 				}
 				else
 				{
-					CurrentIndex += ARRAY_COUNT(QuadPattern);
+					CurrentIndex += UE_ARRAY_COUNT(QuadPattern);
 				}
 			}
 		}
diff --git a/Engine/Source/Runtime/Landscape/Private/LandscapeRender.cpp b/Engine/Source/Runtime/Landscape/Private/LandscapeRender.cpp
--- a/Engine/Source/Runtime/Landscape/Private/LandscapeRender.cpp
+++ b/Engine/Source/Runtime/Landscape/Private/LandscapeRender.cpp
@@ -4273,7 +4273,7 @@ public:
 		if (bIsLayerThumbnail || bDisableTessellation)
 		{
 			FSHA1 Hash;
-			Hash.Update(OutId.BasePropertyOverridesHash.Hash, ARRAY_COUNT(OutId.BasePropertyOverridesHash.Hash));
+			Hash.Update(OutId.BasePropertyOverridesHash.Hash, UE_ARRAY_COUNT(OutId.BasePropertyOverridesHash.Hash));
 
 			const FString HashString = TEXT("bOverride_TessellationMode");
 			Hash.UpdateWithString(*HashString, HashString.Len());
diff --git a/Engine/Source/Runtime/Landscape/Public/LandscapeRender.h b/Engine/Source/Runtime/Landscape/Public/LandscapeRender.h
--- a/Engine/Source/Runtime/Landscape/Public/LandscapeRender.h
+++ b/Engine/Source/Runtime/Landscape/Public/LandscapeRender.h
@@ -1012,7 +1012,7 @@ public:
 			FName(TEXT("Color15"))
 		};
 
-		for (int32 i = 0; i < ARRAY_COUNT(ColorNames) && i < LayerColors.Num(); i++)
+		for (int32 i = 0; i < UE_ARRAY_COUNT(ColorNames) && i < LayerColors.Num(); i++)
 		{
 			if (ParameterInfo.Name == ColorNames[i])
 			{
diff --git a/Engine/Source/Runtime/Launch/Private/Android/LaunchAndroid.cpp b/Engine/Source/Runtime/Launch/Private/Android/LaunchAndroid.cpp
--- a/Engine/Source/Runtime/Launch/Private/Android/LaunchAndroid.cpp
+++ b/Engine/Source/Runtime/Launch/Private/Android/LaunchAndroid.cpp
@@ -308,7 +308,7 @@ static void InitCommandLine()
 	if(CommandLineFile)
 	{
 		char CommandLine[CMD_LINE_MAX];
-		fgets(CommandLine, ARRAY_COUNT(CommandLine) - 1, CommandLineFile);
+		fgets(CommandLine, UE_ARRAY_COUNT(CommandLine) - 1, CommandLineFile);
 
 		fclose(CommandLineFile);
 
diff --git a/Engine/Source/Runtime/Launch/Private/Lumin/LaunchLumin.cpp b/Engine/Source/Runtime/Launch/Private/Lumin/LaunchLumin.cpp
--- a/Engine/Source/Runtime/Launch/Private/Lumin/LaunchLumin.cpp
+++ b/Engine/Source/Runtime/Launch/Private/Lumin/LaunchLumin.cpp
@@ -55,7 +55,7 @@ static void InitCommandLine()
 	{
 		char CommandLine[CMD_LINE_MAX];
 		FMemory::Memzero(CommandLine, CMD_LINE_MAX * sizeof(char));
-		fgets(CommandLine, ARRAY_COUNT(CommandLine) - 1, CommandLineFile);
+		fgets(CommandLine, UE_ARRAY_COUNT(CommandLine) - 1, CommandLineFile);
 		
 		fclose(CommandLineFile);
 		
diff --git a/Engine/Source/Runtime/MovieScene/Private/Channels/MovieSceneFloatChannel.cpp b/Engine/Source/Runtime/MovieScene/Private/Channels/MovieSceneFloatChannel.cpp
--- a/Engine/Source/Runtime/MovieScene/Private/Channels/MovieSceneFloatChannel.cpp
+++ b/Engine/Source/Runtime/MovieScene/Private/Channels/MovieSceneFloatChannel.cpp
@@ -841,9 +841,9 @@ void FMovieSceneFloatChannel::RefineCurvePoints(FFrameRate InTickResolution, dou
 		{
 			bool bSegmentIsLinear = true;
 
-			TTuple<double, double> Evaluated[ARRAY_COUNT(InterpTimes)];
+			TTuple<double, double> Evaluated[UE_ARRAY_COUNT(InterpTimes)];
 
-			for (int32 InterpIndex = 0; InterpIndex < ARRAY_COUNT(InterpTimes); ++InterpIndex)
+			for (int32 InterpIndex = 0; InterpIndex < UE_ARRAY_COUNT(InterpTimes); ++InterpIndex)
 			{
 				double& EvalTime  = Evaluated[InterpIndex].Get<0>();
 
@@ -864,7 +864,7 @@ void FMovieSceneFloatChannel::RefineCurvePoints(FFrameRate InTickResolution, dou
 			if (!bSegmentIsLinear)
 			{
 				// Add the point
-				InOutPoints.Insert(Evaluated, ARRAY_COUNT(Evaluated), Index+1);
+				InOutPoints.Insert(Evaluated, UE_ARRAY_COUNT(Evaluated), Index+1);
 				--Index;
 			}
 		}
diff --git a/Engine/Source/Runtime/MovieScene/Private/Tests/MovieSceneCompilerTests.cpp b/Engine/Source/Runtime/MovieScene/Private/Tests/MovieSceneCompilerTests.cpp
--- a/Engine/Source/Runtime/MovieScene/Private/Tests/MovieSceneCompilerTests.cpp
+++ b/Engine/Source/Runtime/MovieScene/Private/Tests/MovieSceneCompilerTests.cpp
@@ -260,7 +260,7 @@ bool FMovieSceneCompilerEmptySpaceOnTheFlyTest::RunTest(const FString& Parameter
 	{
 		Tests.Emplace();
 		Tests.Last().CompileRanges   = { TRange<FFrameNumber>(0, 70) };
-		Tests.Last().ExpectedResults = TArray<FResult>(&ExpectedResults[0], ARRAY_COUNT(ExpectedResults));
+		Tests.Last().ExpectedResults = TArray<FResult>(&ExpectedResults[0], UE_ARRAY_COUNT(ExpectedResults));
 	}
 
 	UTestMovieSceneSequence* Sequence = NewObject<UTestMovieSceneSequence>(GetTransientPackage());
diff --git a/Engine/Source/Runtime/MovieSceneTracks/Private/Evaluation/MovieScene3DTransformTemplate.cpp b/Engine/Source/Runtime/MovieSceneTracks/Private/Evaluation/MovieScene3DTransformTemplate.cpp
--- a/Engine/Source/Runtime/MovieSceneTracks/Private/Evaluation/MovieScene3DTransformTemplate.cpp
+++ b/Engine/Source/Runtime/MovieSceneTracks/Private/Evaluation/MovieScene3DTransformTemplate.cpp
@@ -184,7 +184,7 @@ MovieScene::TMultiChannelValue<float, 9> FMovieSceneComponentTransformSectionTem
 		Components[4] = Rotation.Y;
 		Components[5] = Rotation.Z;
 
-		for (int32 Index = 0; Index < ARRAY_COUNT(Components); ++Index)
+		for (int32 Index = 0; Index < UE_ARRAY_COUNT(Components); ++Index)
 		{
 			if (TransformValue.IsSet(Index))
 			{
diff --git a/Engine/Source/Runtime/MovieSceneTracks/Private/MovieSceneTracksModule.cpp b/Engine/Source/Runtime/MovieSceneTracks/Private/MovieSceneTracksModule.cpp
--- a/Engine/Source/Runtime/MovieSceneTracks/Private/MovieSceneTracksModule.cpp
+++ b/Engine/Source/Runtime/MovieSceneTracks/Private/MovieSceneTracksModule.cpp
@@ -37,7 +37,7 @@ FName IMovieSceneTracksModule::GetEvaluationGroupName(EBuiltInEvaluationGroup::T
 		"SpawnObjects",
 		"PostEvaluation",
 	};
-	check(InEvalGroup >= 0 && InEvalGroup < ARRAY_COUNT(Names));
+	check(InEvalGroup >= 0 && InEvalGroup < UE_ARRAY_COUNT(Names));
 	return Names[InEvalGroup];
 }
 
diff --git a/Engine/Source/Runtime/MovieSceneTracks/Private/Sections/MovieSceneVectorSection.cpp b/Engine/Source/Runtime/MovieSceneTracks/Private/Sections/MovieSceneVectorSection.cpp
--- a/Engine/Source/Runtime/MovieSceneTracks/Private/Sections/MovieSceneVectorSection.cpp
+++ b/Engine/Source/Runtime/MovieSceneTracks/Private/Sections/MovieSceneVectorSection.cpp
@@ -233,7 +233,7 @@ void UMovieSceneVectorSection::RecreateChannelProxy()
 {
 	FMovieSceneChannelProxyData Channels;
 
-	check(ChannelsUsed <= ARRAY_COUNT(Curves));
+	check(ChannelsUsed <= UE_ARRAY_COUNT(Curves));
 
 #if WITH_EDITOR
 
diff --git a/Engine/Source/Runtime/NavigationSystem/Private/NavMesh/NavMeshRenderingComponent.cpp b/Engine/Source/Runtime/NavigationSystem/Private/NavMesh/NavMeshRenderingComponent.cpp
--- a/Engine/Source/Runtime/NavigationSystem/Private/NavMesh/NavMeshRenderingComponent.cpp
+++ b/Engine/Source/Runtime/NavigationSystem/Private/NavMesh/NavMeshRenderingComponent.cpp
@@ -280,7 +280,7 @@ void FNavMeshSceneProxyData::Serialize(FArchive& Ar)
 	}
 
 	TArray<FDebugRenderSceneProxy::FDebugLine>* LineArraysToSerialize[] = { &ThickLineItems, &TileEdgeLines, &NavMeshEdgeLines, &NavLinkLines, &ClusterLinkLines };
-	for (int32 ArrIdx = 0; ArrIdx < ARRAY_COUNT(LineArraysToSerialize); ArrIdx++)
+	for (int32 ArrIdx = 0; ArrIdx < UE_ARRAY_COUNT(LineArraysToSerialize); ArrIdx++)
 	{
 		int32 NumItems = LineArraysToSerialize[ArrIdx]->Num();
 		Ar << NumItems;
diff --git a/Engine/Source/Runtime/NavigationSystem/Private/NavMesh/RecastNavMeshGenerator.cpp b/Engine/Source/Runtime/NavigationSystem/Private/NavMesh/RecastNavMeshGenerator.cpp
--- a/Engine/Source/Runtime/NavigationSystem/Private/NavMesh/RecastNavMeshGenerator.cpp
+++ b/Engine/Source/Runtime/NavigationSystem/Private/NavMesh/RecastNavMeshGenerator.cpp
@@ -772,7 +772,7 @@ void ExportRigidBodyBoxElements(const FKAggregateGeom& AggGeom, TNavStatArray<fl
 			ElemTM.TransformPosition(FVector( Extent.X,  Extent.Y, -Extent.Z))
 		};
 
-		for (int32 iv = 0; iv < ARRAY_COUNT(UnrealVerts); iv++)
+		for (int32 iv = 0; iv < UE_ARRAY_COUNT(UnrealVerts); iv++)
 		{
 			UnrealBounds += UnrealVerts[iv];
 
diff --git a/Engine/Source/Runtime/NavigationSystem/Private/NavigationSystem.cpp b/Engine/Source/Runtime/NavigationSystem/Private/NavigationSystem.cpp
--- a/Engine/Source/Runtime/NavigationSystem/Private/NavigationSystem.cpp
+++ b/Engine/Source/Runtime/NavigationSystem/Private/NavigationSystem.cpp
@@ -2267,7 +2267,7 @@ void UNavigationSystemV1::DescribeFilterFlags(const TArray<FString>& FlagsDesc)
 	UStruct* Structs[] = { StructProp1->Struct, StructProp2->Struct };
 	const FString CustomNameMeta = TEXT("DisplayName");
 
-	for (int32 StructIndex = 0; StructIndex < ARRAY_COUNT(Structs); StructIndex++)
+	for (int32 StructIndex = 0; StructIndex < UE_ARRAY_COUNT(Structs); StructIndex++)
 	{
 		for (int32 FlagIndex = 0; FlagIndex < MaxFlags; FlagIndex++)
 		{
diff --git a/Engine/Source/Runtime/Navmesh/Private/DebugUtils/RecastDump.cpp b/Engine/Source/Runtime/Navmesh/Private/DebugUtils/RecastDump.cpp
--- a/Engine/Source/Runtime/Navmesh/Private/DebugUtils/RecastDump.cpp
+++ b/Engine/Source/Runtime/Navmesh/Private/DebugUtils/RecastDump.cpp
@@ -35,7 +35,7 @@ static void ioprintf(duFileIO* io, const char* format, ...)
 	char line[256];
 	va_list ap;
 	va_start(ap, format);
-	const int n = FCStringAnsi::GetVarArgs(line, ARRAY_COUNT(line), format, ap);
+	const int n = FCStringAnsi::GetVarArgs(line, UE_ARRAY_COUNT(line), format, ap);
 	va_end(ap);
 	if (n > 0)
 		io->write(line, sizeof(char)*n);
diff --git a/Engine/Source/Runtime/NetworkFile/Private/NetworkPlatformFile.cpp b/Engine/Source/Runtime/NetworkFile/Private/NetworkPlatformFile.cpp
--- a/Engine/Source/Runtime/NetworkFile/Private/NetworkPlatformFile.cpp
+++ b/Engine/Source/Runtime/NetworkFile/Private/NetworkPlatformFile.cpp
@@ -959,7 +959,7 @@ public:
 				while (RemainingData)
 				{
 					// read next chunk from archive
-					uint32 LocalSize = FPlatformMath::Min<uint32>(ARRAY_COUNT(Buffer), RemainingData);
+					uint32 LocalSize = FPlatformMath::Min<uint32>(UE_ARRAY_COUNT(Buffer), RemainingData);
 					FileArchive->Serialize(Buffer, LocalSize);
 					// write it out
 					if (!FileHandle->Write(Buffer, LocalSize))
diff --git a/Engine/Source/Runtime/Online/BuildPatchServices/Private/BuildPatchServicesModule.cpp b/Engine/Source/Runtime/Online/BuildPatchServices/Private/BuildPatchServicesModule.cpp
--- a/Engine/Source/Runtime/Online/BuildPatchServices/Private/BuildPatchServicesModule.cpp
+++ b/Engine/Source/Runtime/Online/BuildPatchServices/Private/BuildPatchServicesModule.cpp
@@ -44,8 +44,8 @@ void FBuildPatchServicesModule::StartupModule()
 #if UE_BUILD_DEBUG
 	TSet<FString> NoDupes;
 	bool bWasDupe = false;
-	check(ARRAY_COUNT(InstallErrorPrefixes::ErrorTypeStrings) == (uint64)EBuildPatchInstallError::NumInstallErrors);
-	for (int32 Idx = 0; Idx < ARRAY_COUNT(InstallErrorPrefixes::ErrorTypeStrings); ++Idx)
+	check(UE_ARRAY_COUNT(InstallErrorPrefixes::ErrorTypeStrings) == (uint64)EBuildPatchInstallError::NumInstallErrors);
+	for (int32 Idx = 0; Idx < UE_ARRAY_COUNT(InstallErrorPrefixes::ErrorTypeStrings); ++Idx)
 	{
 		NoDupes.Add(FString(InstallErrorPrefixes::ErrorTypeStrings[Idx]), &bWasDupe);
 		check(bWasDupe == false);
diff --git a/Engine/Source/Runtime/Online/BuildPatchServices/Private/Installer/CloudChunkSource.h b/Engine/Source/Runtime/Online/BuildPatchServices/Private/Installer/CloudChunkSource.h
--- a/Engine/Source/Runtime/Online/BuildPatchServices/Private/Installer/CloudChunkSource.h
+++ b/Engine/Source/Runtime/Online/BuildPatchServices/Private/Installer/CloudChunkSource.h
@@ -75,12 +75,12 @@ namespace BuildPatchServices
 			, TcpZeroWindowMinimumSeconds(20.0f)
 		{
 			const float RetryFloats[] = {0.5f, 1.0f, 1.0f, 3.0f, 3.0f, 10.0f, 10.0f, 20.0f, 20.0f, 30.0f};
-			RetryDelayTimes.Empty(ARRAY_COUNT(RetryFloats));
-			RetryDelayTimes.Append(RetryFloats, ARRAY_COUNT(RetryFloats));
+			RetryDelayTimes.Empty(UE_ARRAY_COUNT(RetryFloats));
+			RetryDelayTimes.Append(RetryFloats, UE_ARRAY_COUNT(RetryFloats));
 			const float HealthFloats[] = {0.0f, 0.0f, 0.9f, 0.99f, 1.0f};
-			check((int32)EBuildPatchDownloadHealth::NUM_Values == ARRAY_COUNT(HealthFloats));
-			HealthPercentages.Empty(ARRAY_COUNT(HealthFloats));
-			HealthPercentages.Append(HealthFloats, ARRAY_COUNT(HealthFloats));
+			check((int32)EBuildPatchDownloadHealth::NUM_Values == UE_ARRAY_COUNT(HealthFloats));
+			HealthPercentages.Empty(UE_ARRAY_COUNT(HealthFloats));
+			HealthPercentages.Append(HealthFloats, UE_ARRAY_COUNT(HealthFloats));
 		}
 	};
 
diff --git a/Engine/Source/Runtime/Online/HTTP/Private/Curl/CurlHttp.cpp b/Engine/Source/Runtime/Online/HTTP/Private/Curl/CurlHttp.cpp
--- a/Engine/Source/Runtime/Online/HTTP/Private/Curl/CurlHttp.cpp
+++ b/Engine/Source/Runtime/Online/HTTP/Private/Curl/CurlHttp.cpp
@@ -247,7 +247,7 @@ FString FCurlHttpRequest::CombineHeaderKeyValue(const FString& HeaderKey, const
 {
 	FString Combined;
 	const TCHAR Separator[] = TEXT(": ");
-	constexpr const int32 SeparatorLength = ARRAY_COUNT(Separator) - 1;
+	constexpr const int32 SeparatorLength = UE_ARRAY_COUNT(Separator) - 1;
 	Combined.Reserve(HeaderKey.Len() + SeparatorLength + HeaderValue.Len());
 	Combined.Append(HeaderKey);
 	Combined.AppendChars(Separator, SeparatorLength);
diff --git a/Engine/Source/Runtime/Online/HTTP/Private/Curl/CurlHttpManager.cpp b/Engine/Source/Runtime/Online/HTTP/Private/Curl/CurlHttpManager.cpp
--- a/Engine/Source/Runtime/Online/HTTP/Private/Curl/CurlHttpManager.cpp
+++ b/Engine/Source/Runtime/Online/HTTP/Private/Curl/CurlHttpManager.cpp
@@ -256,7 +256,7 @@ void FCurlHttpManager::InitCurl()
 	}
 
 	TCHAR Home[256] = TEXT("");
-	if (FParse::Value(FCommandLine::Get(), TEXT("MULTIHOMEHTTP="), Home, ARRAY_COUNT(Home)))
+	if (FParse::Value(FCommandLine::Get(), TEXT("MULTIHOMEHTTP="), Home, UE_ARRAY_COUNT(Home)))
 	{
 		ISocketSubsystem* SocketSubsystem = ISocketSubsystem::Get(PLATFORM_SOCKETSUBSYSTEM);
 		if (SocketSubsystem)
diff --git a/Engine/Source/Runtime/Online/HTTP/Private/GenericPlatform/GenericPlatformHttp.cpp b/Engine/Source/Runtime/Online/HTTP/Private/GenericPlatform/GenericPlatformHttp.cpp
--- a/Engine/Source/Runtime/Online/HTTP/Private/GenericPlatform/GenericPlatformHttp.cpp
+++ b/Engine/Source/Runtime/Online/HTTP/Private/GenericPlatform/GenericPlatformHttp.cpp
@@ -67,10 +67,10 @@ static bool IsAllowedChar(UTF8CHAR LookupChar)
 
 	if (!bTableFilled)
 	{
-		for (int32 Idx = 0; Idx < ARRAY_COUNT(AllowedChars) - 1; ++Idx)	// -1 to avoid trailing 0
+		for (int32 Idx = 0; Idx < UE_ARRAY_COUNT(AllowedChars) - 1; ++Idx)	// -1 to avoid trailing 0
 		{
 			uint8 AllowedCharIdx = static_cast<uint8>(AllowedChars[Idx]);
-			check(AllowedCharIdx < ARRAY_COUNT(AllowedTable));
+			check(AllowedCharIdx < UE_ARRAY_COUNT(AllowedTable));
 			AllowedTable[AllowedCharIdx] = true;
 		}
 
@@ -133,7 +133,7 @@ FString FGenericPlatformHttp::UrlDecode(const FString &EncodedString)
 
 					ANSICHAR Buffer[8] = { 0 };
 					ANSICHAR* BufferPtr = Buffer;
-					const int32 Len = ARRAY_COUNT(Buffer);
+					const int32 Len = UE_ARRAY_COUNT(Buffer);
 					const int32 WrittenChars = FTCHARToUTF8_Convert::Utf8FromCodepoint(Value, BufferPtr, Len);
 
 					Data.Append(Buffer, WrittenChars);
diff --git a/Engine/Source/Runtime/Online/HTTP/Private/GenericPlatform/HttpRequestPayload.cpp b/Engine/Source/Runtime/Online/HTTP/Private/GenericPlatform/HttpRequestPayload.cpp
--- a/Engine/Source/Runtime/Online/HTTP/Private/GenericPlatform/HttpRequestPayload.cpp
+++ b/Engine/Source/Runtime/Online/HTTP/Private/GenericPlatform/HttpRequestPayload.cpp
@@ -13,10 +13,10 @@ bool FGenericPlatformHttp::IsURLEncoded(const TArray<uint8>& Payload)
 
 	if (!bTableFilled)
 	{
-		for (int32 Idx = 0; Idx < ARRAY_COUNT(AllowedChars) - 1; ++Idx)	// -1 to avoid trailing 0
+		for (int32 Idx = 0; Idx < UE_ARRAY_COUNT(AllowedChars) - 1; ++Idx)	// -1 to avoid trailing 0
 		{
 			uint8 AllowedCharIdx = static_cast<uint8>(AllowedChars[Idx]);
-			check(AllowedCharIdx < ARRAY_COUNT(AllowedTable));
+			check(AllowedCharIdx < UE_ARRAY_COUNT(AllowedTable));
 			AllowedTable[AllowedCharIdx] = true;
 		}
 
diff --git a/Engine/Source/Runtime/Online/HTTPServer/Private/HttpConnectionRequestReadContext.cpp b/Engine/Source/Runtime/Online/HTTPServer/Private/HttpConnectionRequestReadContext.cpp
--- a/Engine/Source/Runtime/Online/HTTPServer/Private/HttpConnectionRequestReadContext.cpp
+++ b/Engine/Source/Runtime/Online/HTTPServer/Private/HttpConnectionRequestReadContext.cpp
@@ -346,7 +346,7 @@ FString FHttpConnectionRequestReadContext::UrlDecode(const FString &EncodedStrin
 
 					ANSICHAR Buffer[8] = { 0 };
 					ANSICHAR* BufferPtr = Buffer;
-					const int32 Len = ARRAY_COUNT(Buffer);
+					const int32 Len = UE_ARRAY_COUNT(Buffer);
 					const int32 WrittenChars = FTCHARToUTF8_Convert::Utf8FromCodepoint(Value, BufferPtr, Len);
 
 					Data.Append(Buffer, WrittenChars);
diff --git a/Engine/Source/Runtime/Online/SSL/Private/SslError.cpp b/Engine/Source/Runtime/Online/SSL/Private/SslError.cpp
--- a/Engine/Source/Runtime/Online/SSL/Private/SslError.cpp
+++ b/Engine/Source/Runtime/Online/SSL/Private/SslError.cpp
@@ -22,8 +22,8 @@ FString GetSslErrorString()
 	if (Error != SSL_ERROR_NONE)
 	{
 		char AnsiErrorBuffer[256];
-		ERR_error_string_n(Error, AnsiErrorBuffer, ARRAY_COUNT(AnsiErrorBuffer) - 1);
-		AnsiErrorBuffer[ARRAY_COUNT(AnsiErrorBuffer) - 1] = '\0';
+		ERR_error_string_n(Error, AnsiErrorBuffer, UE_ARRAY_COUNT(AnsiErrorBuffer) - 1);
+		AnsiErrorBuffer[UE_ARRAY_COUNT(AnsiErrorBuffer) - 1] = '\0';
 
 		SslErrorString = ANSI_TO_TCHAR(AnsiErrorBuffer);
 	}
diff --git a/Engine/Source/Runtime/Online/SSL/Private/Windows/WindowsPlatformSslCertificateManager.cpp b/Engine/Source/Runtime/Online/SSL/Private/Windows/WindowsPlatformSslCertificateManager.cpp
--- a/Engine/Source/Runtime/Online/SSL/Private/Windows/WindowsPlatformSslCertificateManager.cpp
+++ b/Engine/Source/Runtime/Online/SSL/Private/Windows/WindowsPlatformSslCertificateManager.cpp
@@ -44,14 +44,14 @@ void FWindowsPlatformSslCertificateManager::BuildRootCertificateArray()
 				else
 				{
 					TCHAR Name[128];
-					CertGetNameString(CertContext, CERT_NAME_SIMPLE_DISPLAY_TYPE, 0, nullptr, Name, ARRAY_COUNT(Name));
+					CertGetNameString(CertContext, CERT_NAME_SIMPLE_DISPLAY_TYPE, 0, nullptr, Name, UE_ARRAY_COUNT(Name));
 					UE_LOG(LogSsl, Log, TEXT("Unable to convert certificate: name:%s error:%s"), Name, *GetSslErrorString());
 				}
 			}
 			else
 			{
 				TCHAR Name[128];
-				CertGetNameString(CertContext, CERT_NAME_SIMPLE_DISPLAY_TYPE, 0, nullptr, Name, ARRAY_COUNT(Name));
+				CertGetNameString(CertContext, CERT_NAME_SIMPLE_DISPLAY_TYPE, 0, nullptr, Name, UE_ARRAY_COUNT(Name));
 				UE_LOG(LogSsl, Log, TEXT("Unhandled certificate encoding: name:%s encodingType:0x%08x"), Name, CertContext->dwCertEncodingType);
 			}
 		}
