Deprecating ARRAY_COUNT and changing it to UE_ARRAY_COUNT.

diff --git a/Engine/Source/Runtime/Engine/Private/GPUSkinCache.cpp b/Engine/Source/Runtime/Engine/Private/GPUSkinCache.cpp
--- a/Engine/Source/Runtime/Engine/Private/GPUSkinCache.cpp
+++ b/Engine/Source/Runtime/Engine/Private/GPUSkinCache.cpp
@@ -632,7 +632,7 @@ void FGPUSkinCache::TransitionToWriteable(FRHICommandList& RHICmdList)
 {
 	int32 BufferIndex = InternalUpdateCount % GPUSKINCACHE_FRAMES;
 	FRHIUnorderedAccessView* OutUAVs[] = { SkinCacheBuffer[BufferIndex].UAV };
-	RHICmdList.TransitionResources(EResourceTransitionAccess::ERWNoBarrier, EResourceTransitionPipeline::EGfxToCompute, OutUAVs, ARRAY_COUNT(OutUAVs));
+	RHICmdList.TransitionResources(EResourceTransitionAccess::ERWNoBarrier, EResourceTransitionPipeline::EGfxToCompute, OutUAVs, UE_ARRAY_COUNT(OutUAVs));
 }
 
 void FGPUSkinCache::TransitionAllToWriteable(FRHICommandList& RHICmdList)
@@ -646,7 +646,7 @@ void FGPUSkinCache::TransitionAllToWriteable(FRHICommandList& RHICmdList)
 			OutUAVs[Index] = SkinCacheBuffer[Index].UAV;
 		}
 
-		RHICmdList.TransitionResources(EResourceTransitionAccess::ERWNoBarrier, EResourceTransitionPipeline::EGfxToCompute, OutUAVs, ARRAY_COUNT(OutUAVs));
+		RHICmdList.TransitionResources(EResourceTransitionAccess::ERWNoBarrier, EResourceTransitionPipeline::EGfxToCompute, OutUAVs, UE_ARRAY_COUNT(OutUAVs));
 	}
 }
 #endif
diff --git a/Engine/Source/Runtime/Engine/Private/GameViewportClient.cpp b/Engine/Source/Runtime/Engine/Private/GameViewportClient.cpp
--- a/Engine/Source/Runtime/Engine/Private/GameViewportClient.cpp
+++ b/Engine/Source/Runtime/Engine/Private/GameViewportClient.cpp
@@ -3442,8 +3442,8 @@ bool UGameViewportClient::HandleDisplayCommand( const TCHAR* Cmd, FOutputDevice&
 {
 	TCHAR ObjectName[256];
 	TCHAR PropStr[256];
-	if ( FParse::Token(Cmd, ObjectName, ARRAY_COUNT(ObjectName), true) &&
-		FParse::Token(Cmd, PropStr, ARRAY_COUNT(PropStr), true) )
+	if ( FParse::Token(Cmd, ObjectName, UE_ARRAY_COUNT(ObjectName), true) &&
+		FParse::Token(Cmd, PropStr, UE_ARRAY_COUNT(PropStr), true) )
 	{
 		UObject* Obj = FindObject<UObject>(ANY_PACKAGE, ObjectName);
 		if (Obj != nullptr)
@@ -3471,7 +3471,7 @@ bool UGameViewportClient::HandleDisplayAllCommand( const TCHAR* Cmd, FOutputDevi
 {
 	TCHAR ClassName[256];
 	TCHAR PropStr[256];
-	if (FParse::Token(Cmd, ClassName, ARRAY_COUNT(ClassName), true))
+	if (FParse::Token(Cmd, ClassName, UE_ARRAY_COUNT(ClassName), true))
 	{
 		bool bValidClassToken = true;
 		UClass* WithinClass = nullptr;
@@ -3496,7 +3496,7 @@ bool UGameViewportClient::HandleDisplayAllCommand( const TCHAR* Cmd, FOutputDevi
 		}
 		if (bValidClassToken)
 		{
-			FParse::Token(Cmd, PropStr, ARRAY_COUNT(PropStr), true);
+			FParse::Token(Cmd, PropStr, UE_ARRAY_COUNT(PropStr), true);
 			UClass* Cls = FindObject<UClass>(ANY_PACKAGE, ClassName);
 			if (Cls != nullptr)
 			{
@@ -3532,7 +3532,7 @@ bool UGameViewportClient::HandleDisplayAllCommand( const TCHAR* Cmd, FOutputDevi
 bool UGameViewportClient::HandleDisplayAllLocationCommand( const TCHAR* Cmd, FOutputDevice& Ar )
 {
 	TCHAR ClassName[256];
-	if (FParse::Token(Cmd, ClassName, ARRAY_COUNT(ClassName), true))
+	if (FParse::Token(Cmd, ClassName, UE_ARRAY_COUNT(ClassName), true))
 	{
 		UClass* Cls = FindObject<UClass>(ANY_PACKAGE, ClassName);
 		if (Cls != nullptr)
@@ -3564,7 +3564,7 @@ bool UGameViewportClient::HandleDisplayAllLocationCommand( const TCHAR* Cmd, FOu
 bool UGameViewportClient::HandleDisplayAllRotationCommand( const TCHAR* Cmd, FOutputDevice& Ar )
 {
 	TCHAR ClassName[256];
-	if (FParse::Token(Cmd, ClassName, ARRAY_COUNT(ClassName), true))
+	if (FParse::Token(Cmd, ClassName, UE_ARRAY_COUNT(ClassName), true))
 	{
 		UClass* Cls = FindObject<UClass>(ANY_PACKAGE, ClassName);
 		if (Cls != nullptr)
@@ -3606,7 +3606,7 @@ bool UGameViewportClient::HandleGetAllLocationCommand(const TCHAR* Cmd, FOutputD
 	TCHAR ClassName[256];
 	UClass* Class;
 
-	if (FParse::Token(Cmd, ClassName, ARRAY_COUNT(ClassName), 1) &&
+	if (FParse::Token(Cmd, ClassName, UE_ARRAY_COUNT(ClassName), 1) &&
 		(Class = FindObject<UClass>(ANY_PACKAGE, ClassName)) != NULL)
 	{
 		bool bShowPendingKills = FParse::Command(&Cmd, TEXT("SHOWPENDINGKILLS"));
diff --git a/Engine/Source/Runtime/Engine/Private/GeomTools.cpp b/Engine/Source/Runtime/Engine/Private/GeomTools.cpp
--- a/Engine/Source/Runtime/Engine/Private/GeomTools.cpp
+++ b/Engine/Source/Runtime/Engine/Private/GeomTools.cpp
@@ -27,7 +27,7 @@ static FClipSMVertex GetVert(const UStaticMesh* StaticMesh, int32 VertIndex)
 	{
 		Result.UVs[UVIndex] = LODModel.VertexBuffers.StaticMeshVertexBuffer.GetVertexUV(VertIndex,UVIndex);
 	}
-	for(int32 UVIndex = NumUVs;UVIndex < ARRAY_COUNT(Result.UVs);UVIndex++)
+	for(int32 UVIndex = NumUVs;UVIndex < UE_ARRAY_COUNT(Result.UVs);UVIndex++)
 	{
 		Result.UVs[UVIndex] = FVector2D::ZeroVector;
 	}
@@ -380,7 +380,7 @@ static bool VertsAreEqual(const FClipSMVertex& A,const FClipSMVertex& B)
 		return false;
 	}
 
-	for(int32 i=0; i<ARRAY_COUNT(A.UVs); i++)
+	for(int32 i=0; i<UE_ARRAY_COUNT(A.UVs); i++)
 	{
 		if( !A.UVs[i].Equals(B.UVs[i], 1.0f / 1024.0f) )
 		{
diff --git a/Engine/Source/Runtime/Engine/Private/GestureRecognizer.cpp b/Engine/Source/Runtime/Engine/Private/GestureRecognizer.cpp
--- a/Engine/Source/Runtime/Engine/Private/GestureRecognizer.cpp
+++ b/Engine/Source/Runtime/Engine/Private/GestureRecognizer.cpp
@@ -23,7 +23,7 @@ void FGestureRecognizer::DetectGestures(const FVector (&Touches)[EKeys::NUM_TOUC
 	if (PreviousTouchCount != 0 || TouchCount != 0)
 	{
 		// place new anchor points
-		for (int32 Index = 0; Index < ARRAY_COUNT(AnchorPoints); Index++)
+		for (int32 Index = 0; Index < UE_ARRAY_COUNT(AnchorPoints); Index++)
 		{
 			if (PreviousTouchCount < Index + 1 && TouchCount >= Index + 1)
 			{
diff --git a/Engine/Source/Runtime/Engine/Private/InputDelegateBinding.cpp b/Engine/Source/Runtime/Engine/Private/InputDelegateBinding.cpp
--- a/Engine/Source/Runtime/Engine/Private/InputDelegateBinding.cpp
+++ b/Engine/Source/Runtime/Engine/Private/InputDelegateBinding.cpp
@@ -35,7 +35,7 @@ void UInputDelegateBinding::BindInputDelegates(const UClass* InClass, UInputComp
 	{
 		BindInputDelegates(InClass->GetSuperClass(), InputComponent);
 
-		for (int32 Index = 0; Index < ARRAY_COUNT(InputBindingClasses); ++Index)
+		for (int32 Index = 0; Index < UE_ARRAY_COUNT(InputBindingClasses); ++Index)
 		{
 			UInputDelegateBinding* BindingObject = CastChecked<UInputDelegateBinding>(
 				UBlueprintGeneratedClass::GetDynamicBindingObject(InClass, InputBindingClasses[Index])
diff --git a/Engine/Source/Runtime/Engine/Private/Level.cpp b/Engine/Source/Runtime/Engine/Private/Level.cpp
--- a/Engine/Source/Runtime/Engine/Private/Level.cpp
+++ b/Engine/Source/Runtime/Engine/Private/Level.cpp
@@ -596,7 +596,7 @@ void ULevel::PostLoad()
 	}
 
 	// Fixup deprecated stuff in levels simplification settings
-	for (int32 Index = 0; Index < ARRAY_COUNT(LevelSimplification); ++Index)
+	for (int32 Index = 0; Index < UE_ARRAY_COUNT(LevelSimplification); ++Index)
 	{
 		LevelSimplification[Index].PostLoadDeprecated();
 	}
diff --git a/Engine/Source/Runtime/Engine/Private/LightMap.cpp b/Engine/Source/Runtime/Engine/Private/LightMap.cpp
--- a/Engine/Source/Runtime/Engine/Private/LightMap.cpp
+++ b/Engine/Source/Runtime/Engine/Private/LightMap.cpp
@@ -2756,7 +2756,7 @@ int32 FLightMap2D::EncodeShadowTexture(ULevel* LightingScenario, struct FLightMa
 							ExtrapolatedFilterableComponents[i] = 0;
 						}
 
-						for (int32 NeighborIndex = 0; NeighborIndex < ARRAY_COUNT(Neighbors); NeighborIndex++)
+						for (int32 NeighborIndex = 0; NeighborIndex < UE_ARRAY_COUNT(Neighbors); NeighborIndex++)
 						{
 							if (static_cast<int32>(DestY) + Neighbors[NeighborIndex].Y >= 0
 								&& DestY + Neighbors[NeighborIndex].Y < MipSizeY
@@ -2975,7 +2975,7 @@ void FLightMap2D::Serialize(FArchive& Ar)
 
 	if (RenderCustomVersion >= FRenderingObjectVersion::LightmapHasShadowmapData)
 	{
-		for (int Channel = 0; Channel < ARRAY_COUNT(bShadowChannelValid); Channel++)
+		for (int Channel = 0; Channel < UE_ARRAY_COUNT(bShadowChannelValid); Channel++)
 		{
 			Ar << bShadowChannelValid[Channel];
 		}
@@ -3257,7 +3257,7 @@ bool FQuantizedLightmapData::HasNonZeroData() const
 
 			if (bHasSkyShadowing)
 			{
-				for (int32 Index = 0; Index < ARRAY_COUNT(LightmapSample.SkyOcclusion); Index++)
+				for (int32 Index = 0; Index < UE_ARRAY_COUNT(LightmapSample.SkyOcclusion); Index++)
 				{
 					if (LightmapSample.SkyOcclusion[Index] != 0)
 					{
diff --git a/Engine/Source/Runtime/Engine/Private/LocalPlayer.cpp b/Engine/Source/Runtime/Engine/Private/LocalPlayer.cpp
--- a/Engine/Source/Runtime/Engine/Private/LocalPlayer.cpp
+++ b/Engine/Source/Runtime/Engine/Private/LocalPlayer.cpp
@@ -1298,7 +1298,7 @@ bool ULocalPlayer::HandleListPawnComponentsCommand( const TCHAR* Cmd, FOutputDev
 bool ULocalPlayer::HandleExecCommand( const TCHAR* Cmd, FOutputDevice& Ar )
 {
 	TCHAR Filename[512];
-	if( FParse::Token( Cmd, Filename, ARRAY_COUNT(Filename), 0 ) )
+	if( FParse::Token( Cmd, Filename, UE_ARRAY_COUNT(Filename), 0 ) )
 	{
 		ExecMacro( Filename, Ar );
 	}
diff --git a/Engine/Source/Runtime/Engine/Private/Materials/Material.cpp b/Engine/Source/Runtime/Engine/Private/Materials/Material.cpp
--- a/Engine/Source/Runtime/Engine/Private/Materials/Material.cpp
+++ b/Engine/Source/Runtime/Engine/Private/Materials/Material.cpp
@@ -988,7 +988,7 @@ void UMaterial::GetUsedTextures(TArray<UTexture*>& OutTextures, EMaterialQuality
 					&CurrentResource->GetUniformVolumeTextureExpressions(),
 					&CurrentResource->GetUniformVirtualTextureExpressions()
 				};
-				for (int32 TypeIndex = 0; TypeIndex < ARRAY_COUNT(ExpressionsByType); TypeIndex++)
+				for (int32 TypeIndex = 0; TypeIndex < UE_ARRAY_COUNT(ExpressionsByType); TypeIndex++)
 				{
 					// Iterate over each of the material's texture expressions.
 					for (FMaterialUniformExpressionTexture* Expression : *ExpressionsByType[TypeIndex])
@@ -1009,7 +1009,7 @@ void UMaterial::GetUsedTextures(TArray<UTexture*>& OutTextures, EMaterialQuality
 				{
 					&CurrentResource->GetUniformScalarParameterExpressions()
 				};
-				for (int32 TypeIndex = 0; TypeIndex < ARRAY_COUNT(AtlasExpressions); TypeIndex++)
+				for (int32 TypeIndex = 0; TypeIndex < UE_ARRAY_COUNT(AtlasExpressions); TypeIndex++)
 				{
 					// Iterate over each of the material's texture expressions.
 					for (FMaterialUniformExpression* Expression : *AtlasExpressions[TypeIndex])
@@ -1057,7 +1057,7 @@ void UMaterial::GetUsedTexturesAndIndices(TArray<UTexture*>& OutTextures, TArray
 			// Try to prevent resizing since this would be expensive.
 			OutIndices.Empty(ExpressionsByType[0]->Num() + ExpressionsByType[1]->Num() + ExpressionsByType[2]->Num() + ExpressionsByType[3]->Num());
 
-			for (int32 TypeIndex = 0; TypeIndex < ARRAY_COUNT(ExpressionsByType); TypeIndex++)
+			for (int32 TypeIndex = 0; TypeIndex < UE_ARRAY_COUNT(ExpressionsByType); TypeIndex++)
 			{
 				// Iterate over each of the material's texture expressions.
 				for (FMaterialUniformExpressionTexture* Expression : *ExpressionsByType[TypeIndex])
@@ -1116,7 +1116,7 @@ void UMaterial::LogMaterialsAndTextures(FOutputDevice& Ar, int32 Indent) const
 						&MaterialResource->GetUniformVolumeTextureExpressions(),
 						&MaterialResource->GetUniformVirtualTextureExpressions(),
 					};
-					for (int32 TypeIndex = 0; TypeIndex < ARRAY_COUNT(ExpressionsByType); TypeIndex++)
+					for (int32 TypeIndex = 0; TypeIndex < UE_ARRAY_COUNT(ExpressionsByType); TypeIndex++)
 					{
 						for (FMaterialUniformExpressionTexture* Expression : *ExpressionsByType[TypeIndex])
 						{
@@ -1173,7 +1173,7 @@ void UMaterial::OverrideTexture(const UTexture* InTextureToOverride, UTexture* O
 			&Resource->GetUniformVolumeTextureExpressions(),
 			&Resource->GetUniformVirtualTextureExpressions()
 		};
-		for(int32 TypeIndex = 0;TypeIndex < ARRAY_COUNT(ExpressionsByType);TypeIndex++)
+		for(int32 TypeIndex = 0;TypeIndex < UE_ARRAY_COUNT(ExpressionsByType);TypeIndex++)
 		{
 			// Iterate over each of the material's texture expressions.
 			for (FMaterialUniformExpressionTexture* Expression : *ExpressionsByType[TypeIndex])
diff --git a/Engine/Source/Runtime/Engine/Private/Materials/MaterialExpressions.cpp b/Engine/Source/Runtime/Engine/Private/Materials/MaterialExpressions.cpp
--- a/Engine/Source/Runtime/Engine/Private/Materials/MaterialExpressions.cpp
+++ b/Engine/Source/Runtime/Engine/Private/Materials/MaterialExpressions.cpp
@@ -8043,7 +8043,7 @@ UMaterialExpressionQualitySwitch::UMaterialExpressionQualitySwitch(const FObject
 int32 UMaterialExpressionQualitySwitch::Compile(class FMaterialCompiler* Compiler, int32 OutputIndex)
 {
 	const EMaterialQualityLevel::Type QualityLevelToCompile = Compiler->GetQualityLevel();
-	check(QualityLevelToCompile < ARRAY_COUNT(Inputs));
+	check(QualityLevelToCompile < UE_ARRAY_COUNT(Inputs));
 	FExpressionInput QualityInput = Inputs[QualityLevelToCompile].GetTracedInput();
 	FExpressionInput DefaultTraced = Default.GetTracedInput();
 
@@ -8071,7 +8071,7 @@ const TArray<FExpressionInput*> UMaterialExpressionQualitySwitch::GetInputs()
 
 	OutInputs.Add(&Default);
 
-	for (int32 InputIndex = 0; InputIndex < ARRAY_COUNT(Inputs); InputIndex++)
+	for (int32 InputIndex = 0; InputIndex < UE_ARRAY_COUNT(Inputs); InputIndex++)
 	{
 		OutInputs.Add(&Inputs[InputIndex]);
 	}
@@ -8154,7 +8154,7 @@ UMaterialExpressionFeatureLevelSwitch::UMaterialExpressionFeatureLevelSwitch(con
 int32 UMaterialExpressionFeatureLevelSwitch::Compile(class FMaterialCompiler* Compiler, int32 OutputIndex)
 {
 	const ERHIFeatureLevel::Type FeatureLevelToCompile = Compiler->GetFeatureLevel();
-	check(FeatureLevelToCompile < ARRAY_COUNT(Inputs));
+	check(FeatureLevelToCompile < UE_ARRAY_COUNT(Inputs));
 	FExpressionInput& FeatureInput = Inputs[FeatureLevelToCompile];
 
 	if (!Default.GetTracedInput().Expression)
@@ -8181,7 +8181,7 @@ const TArray<FExpressionInput*> UMaterialExpressionFeatureLevelSwitch::GetInputs
 
 	OutInputs.Add(&Default);
 
-	for (int32 InputIndex = 0; InputIndex < ARRAY_COUNT(Inputs); InputIndex++)
+	for (int32 InputIndex = 0; InputIndex < UE_ARRAY_COUNT(Inputs); InputIndex++)
 	{
 		OutInputs.Add(&Inputs[InputIndex]);
 	}
@@ -8296,7 +8296,7 @@ int32 UMaterialExpressionShadingPathSwitch::Compile(class FMaterialCompiler* Com
 		ShadingPathToCompile = ERHIShadingPath::Mobile;
 	}
 
-	check(ShadingPathToCompile < ARRAY_COUNT(Inputs));
+	check(ShadingPathToCompile < UE_ARRAY_COUNT(Inputs));
 	FExpressionInput ShadingPathInput = Inputs[ShadingPathToCompile].GetTracedInput();
 	FExpressionInput DefaultTraced = Default.GetTracedInput();
 
@@ -8324,7 +8324,7 @@ const TArray<FExpressionInput*> UMaterialExpressionShadingPathSwitch::GetInputs(
 
 	OutInputs.Add(&Default);
 
-	for (int32 InputIndex = 0; InputIndex < ARRAY_COUNT(Inputs); InputIndex++)
+	for (int32 InputIndex = 0; InputIndex < UE_ARRAY_COUNT(Inputs); InputIndex++)
 	{
 		OutInputs.Add(&Inputs[InputIndex]);
 	}
diff --git a/Engine/Source/Runtime/Engine/Private/Materials/MaterialInstance.cpp b/Engine/Source/Runtime/Engine/Private/Materials/MaterialInstance.cpp
--- a/Engine/Source/Runtime/Engine/Private/Materials/MaterialInstance.cpp
+++ b/Engine/Source/Runtime/Engine/Private/Materials/MaterialInstance.cpp
@@ -1110,7 +1110,7 @@ void UMaterialInstance::GetTextureExpressionValues(const FMaterialResource* Mate
 		OutIndices->Empty(ExpressionsByType[0]->Num() + ExpressionsByType[1]->Num() + ExpressionsByType[2]->Num() + ExpressionsByType[3]->Num() + ExpressionsByType[4]->Num());
 	}
 
-	for(int32 TypeIndex = 0;TypeIndex < ARRAY_COUNT(ExpressionsByType);TypeIndex++)
+	for(int32 TypeIndex = 0;TypeIndex < UE_ARRAY_COUNT(ExpressionsByType);TypeIndex++)
 	{
 		const TArray<TRefCountPtr<FMaterialUniformExpressionTexture> >& Expressions = *ExpressionsByType[TypeIndex];
 
@@ -1148,7 +1148,7 @@ void UMaterialInstance::GetAtlasTextureValues(const FMaterialResource* MaterialR
 	{
 		&MaterialResource->GetUniformScalarParameterExpressions()
 	};
-	for (int32 TypeIndex = 0; TypeIndex < ARRAY_COUNT(AtlasExpressions); TypeIndex++)
+	for (int32 TypeIndex = 0; TypeIndex < UE_ARRAY_COUNT(AtlasExpressions); TypeIndex++)
 	{
 		// Iterate over each of the material's scalar expressions.
 		for (FMaterialUniformExpression* Expression : *AtlasExpressions[TypeIndex])
@@ -1467,7 +1467,7 @@ void UMaterialInstance::OverrideTexture(const UTexture* InTextureToOverride, UTe
 		ExpressionsByType[4] = &SourceMaterialResource->GetUniformVirtualTextureExpressions();
 	}
 		
-	for(int32 TypeIndex = 0;TypeIndex < ARRAY_COUNT(ExpressionsByType);TypeIndex++)
+	for(int32 TypeIndex = 0;TypeIndex < UE_ARRAY_COUNT(ExpressionsByType);TypeIndex++)
 	{
 		const TArray<TRefCountPtr<FMaterialUniformExpressionTexture> >& Expressions = *ExpressionsByType[TypeIndex];
 
@@ -4889,7 +4889,7 @@ void UMaterialInstance::CopyMaterialUniformParametersInternal(UMaterialInterface
 						&MaterialResource->GetUniformVirtualTextureExpressions()
 					};
 
-					for (int32 TypeIndex = 0; TypeIndex < ARRAY_COUNT(TextureExpressions); TypeIndex++)
+					for (int32 TypeIndex = 0; TypeIndex < UE_ARRAY_COUNT(TextureExpressions); TypeIndex++)
 					{
 						for (FMaterialUniformExpressionTexture* TextureExpression : *TextureExpressions[TypeIndex])
 						{
diff --git a/Engine/Source/Runtime/Engine/Private/Materials/MaterialShared.cpp b/Engine/Source/Runtime/Engine/Private/Materials/MaterialShared.cpp
--- a/Engine/Source/Runtime/Engine/Private/Materials/MaterialShared.cpp
+++ b/Engine/Source/Runtime/Engine/Private/Materials/MaterialShared.cpp
@@ -90,17 +90,17 @@ FName MaterialQualityLevelNames[] =
 	FName(TEXT("Num"))
 };
 
-static_assert(ARRAY_COUNT(MaterialQualityLevelNames) == EMaterialQualityLevel::Num + 1, "Missing entry from material quality level names.");
+static_assert(UE_ARRAY_COUNT(MaterialQualityLevelNames) == EMaterialQualityLevel::Num + 1, "Missing entry from material quality level names.");
 
 void GetMaterialQualityLevelName(EMaterialQualityLevel::Type InQualityLevel, FString& OutName)
 {
-	check(InQualityLevel < ARRAY_COUNT(MaterialQualityLevelNames));
+	check(InQualityLevel < UE_ARRAY_COUNT(MaterialQualityLevelNames));
 	MaterialQualityLevelNames[(int32)InQualityLevel].ToString(OutName);
 }
 
 FName GetMaterialQualityLevelFName(EMaterialQualityLevel::Type InQualityLevel)
 {
-	check(InQualityLevel < ARRAY_COUNT(MaterialQualityLevelNames));
+	check(InQualityLevel < UE_ARRAY_COUNT(MaterialQualityLevelNames));
 	return MaterialQualityLevelNames[(int32)InQualityLevel];
 }
 
@@ -194,7 +194,7 @@ static inline SIZE_T AddShaderSize(FShader* Shader, TSet<FShaderResourceId>& Uni
 int32 FMaterialCompiler::Errorf(const TCHAR* Format,...)
 {
 	TCHAR	ErrorText[2048];
-	GET_VARARGS( ErrorText, ARRAY_COUNT(ErrorText), ARRAY_COUNT(ErrorText)-1, Format, Format );
+	GET_VARARGS( ErrorText, UE_ARRAY_COUNT(ErrorText), UE_ARRAY_COUNT(ErrorText)-1, Format, Format );
 	return Error(ErrorText);
 }
 
@@ -3960,7 +3960,7 @@ void SetShaderMapsOnMaterialResources_RenderThread(FRHICommandListImmediate& RHI
 
 	// Iterate through all loaded material render proxies and recache their uniform expressions if needed
 	// This search does not scale well, but is only used when uploading async shader compile results
-	for (int32 FeatureLevelIndex = 0; FeatureLevelIndex < ARRAY_COUNT(bUpdateFeatureLevel); ++FeatureLevelIndex)
+	for (int32 FeatureLevelIndex = 0; FeatureLevelIndex < UE_ARRAY_COUNT(bUpdateFeatureLevel); ++FeatureLevelIndex)
 	{
 		if (bUpdateFeatureLevel[FeatureLevelIndex])
 		{
diff --git a/Engine/Source/Runtime/Engine/Private/NetConnection.cpp b/Engine/Source/Runtime/Engine/Private/NetConnection.cpp
--- a/Engine/Source/Runtime/Engine/Private/NetConnection.cpp
+++ b/Engine/Source/Runtime/Engine/Private/NetConnection.cpp
@@ -1391,7 +1391,7 @@ void UNetConnection::FlushNet(bool bIgnoreSimulation)
 		}
 #endif
 		// Update stuff.
-		const int32 Index = OutPacketId & (ARRAY_COUNT(OutLagPacketId)-1);
+		const int32 Index = OutPacketId & (UE_ARRAY_COUNT(OutLagPacketId)-1);
 
 		// Remember the actual time this packet was sent out, so we can compute ping when the ack comes back
 		OutLagPacketId[Index]			= OutPacketId;
@@ -1687,7 +1687,7 @@ bool UNetConnection::ReadPacketInfo(FBitReader& Reader)
 
 	// Update ping
 	// At this time we have updated OutAckPacketId to the latest received ack.
-	const int32 Index = OutAckPacketId & (ARRAY_COUNT(OutLagPacketId)-1);
+	const int32 Index = OutAckPacketId & (UE_ARRAY_COUNT(OutLagPacketId)-1);
 
 	if ( OutLagPacketId[Index] == OutAckPacketId )
 	{
diff --git a/Engine/Source/Runtime/Engine/Private/NetworkDriver.cpp b/Engine/Source/Runtime/Engine/Private/NetworkDriver.cpp
--- a/Engine/Source/Runtime/Engine/Private/NetworkDriver.cpp
+++ b/Engine/Source/Runtime/Engine/Private/NetworkDriver.cpp
@@ -899,7 +899,7 @@ void UNetDriver::TickFlush(float DeltaSeconds)
 					PerfCounters->Set(TEXT("MinPing"), MinPing, IPerfCounters::Flags::Transient);
 
 					// update buckets
-					for (int BucketIdx = 0; BucketIdx < ARRAY_COUNT(Buckets); ++BucketIdx)
+					for (int BucketIdx = 0; BucketIdx < UE_ARRAY_COUNT(Buckets); ++BucketIdx)
 					{
 						PerfCountersIncrement(FString::Printf(TEXT("PingBucketInt%d"), BucketIdx), Buckets[BucketIdx], 0, IPerfCounters::Flags::Transient);
 					}
@@ -5692,7 +5692,7 @@ static bool NetDriverExec(UWorld* InWorld, const TCHAR* Cmd, FOutputDevice& Ar)
 		TCHAR TokenStr[128];
 
 		// Route the command to a specific beacon if a name is specified or all of them otherwise
-		if (FParse::Token(Cmd, TokenStr, ARRAY_COUNT(TokenStr), true))
+		if (FParse::Token(Cmd, TokenStr, UE_ARRAY_COUNT(TokenStr), true))
 		{
 			NamedDriver = GEngine->FindNamedNetDriver(InWorld, FName(TokenStr));
 			if (NamedDriver != NULL)
diff --git a/Engine/Source/Runtime/Engine/Private/PacketHandlers/StatelessConnectHandlerComponent.cpp b/Engine/Source/Runtime/Engine/Private/PacketHandlers/StatelessConnectHandlerComponent.cpp
--- a/Engine/Source/Runtime/Engine/Private/PacketHandlers/StatelessConnectHandlerComponent.cpp
+++ b/Engine/Source/Runtime/Engine/Private/PacketHandlers/StatelessConnectHandlerComponent.cpp
@@ -292,8 +292,8 @@ void StatelessConnectHandlerComponent::NotifyHandshakeBegin()
 			InitialPacket.WriteBit(bRestartHandshake);
 			InitialPacket.WriteBit(SecretIdPad);
 
-			FMemory::Memzero(PacketSizeFiller, ARRAY_COUNT(PacketSizeFiller));
-			InitialPacket.Serialize(PacketSizeFiller, ARRAY_COUNT(PacketSizeFiller));
+			FMemory::Memzero(PacketSizeFiller, UE_ARRAY_COUNT(PacketSizeFiller));
+			InitialPacket.Serialize(PacketSizeFiller, UE_ARRAY_COUNT(PacketSizeFiller));
 
 
 
@@ -355,13 +355,13 @@ void StatelessConnectHandlerComponent::SendConnectChallenge(TSharedPtr<const FIn
 		ChallengePacket.WriteBit(ActiveSecret);
 
 		ChallengePacket << Timestamp;
-		ChallengePacket.Serialize(Cookie, ARRAY_COUNT(Cookie));
+		ChallengePacket.Serialize(Cookie, UE_ARRAY_COUNT(Cookie));
 
 #if !UE_BUILD_SHIPPING
 		FDDoSDetection* DDoS = Handler->GetDDoS();
 
 		UE_CLOG((DDoS == nullptr || !DDoS->CheckLogRestrictions()), LogHandshake, Log,
-				TEXT("SendConnectChallenge. Timestamp: %f, Cookie: %s" ), Timestamp, *FString::FromBlob(Cookie, ARRAY_COUNT(Cookie)));
+				TEXT("SendConnectChallenge. Timestamp: %f, Cookie: %s" ), Timestamp, *FString::FromBlob(Cookie, UE_ARRAY_COUNT(Cookie)));
 #endif
 
 		CapHandshakePacket(ChallengePacket);
@@ -475,7 +475,7 @@ void StatelessConnectHandlerComponent::SendChallengeResponse(uint8 InSecretId, f
 		LastServerSequence = *CurSequence & (MAX_PACKETID - 1);
 		LastClientSequence = *(CurSequence + 1) & (MAX_PACKETID - 1);
 
-		FMemory::Memcpy(LastCookie, InCookie, ARRAY_COUNT(LastCookie));
+		FMemory::Memcpy(LastCookie, InCookie, UE_ARRAY_COUNT(LastCookie));
 	}
 	else
 	{
@@ -666,7 +666,7 @@ void StatelessConnectHandlerComponent::InitFromConnectionless(StatelessConnectHa
 	// Store the cookie/address used for the handshake, to enable server ack-retries
 	LastChallengeSuccessAddress = InConnectionlessHandler->LastChallengeSuccessAddress;
 
-	FMemory::Memcpy(AuthorisedCookie, InConnectionlessHandler->AuthorisedCookie, ARRAY_COUNT(AuthorisedCookie));
+	FMemory::Memcpy(AuthorisedCookie, InConnectionlessHandler->AuthorisedCookie, UE_ARRAY_COUNT(AuthorisedCookie));
 }
 
 void StatelessConnectHandlerComponent::Incoming(FBitReader& Packet)
@@ -731,7 +731,7 @@ void StatelessConnectHandlerComponent::Incoming(FBitReader& Packet)
 							}
 
 							// Save the final authorized cookie
-							FMemory::Memcpy(AuthorisedCookie, Cookie, ARRAY_COUNT(AuthorisedCookie));
+							FMemory::Memcpy(AuthorisedCookie, Cookie, UE_ARRAY_COUNT(AuthorisedCookie));
 						}
 
 						// Now finish initializing the handler - flushing the queued packet buffer in the process.
@@ -896,7 +896,7 @@ void StatelessConnectHandlerComponent::IncomingConnectionless(const TSharedPtr<c
 						{
 							if (bRestartHandshake)
 							{
-								FMemory::Memcpy(AuthorisedCookie, OrigCookie, ARRAY_COUNT(AuthorisedCookie));
+								FMemory::Memcpy(AuthorisedCookie, OrigCookie, UE_ARRAY_COUNT(AuthorisedCookie));
 							}
 							else
 							{
@@ -905,7 +905,7 @@ void StatelessConnectHandlerComponent::IncomingConnectionless(const TSharedPtr<c
 								LastServerSequence = *CurSequence & (MAX_PACKETID - 1);
 								LastClientSequence = *(CurSequence + 1) & (MAX_PACKETID - 1);
 
-								FMemory::Memcpy(AuthorisedCookie, Cookie, ARRAY_COUNT(AuthorisedCookie));
+								FMemory::Memcpy(AuthorisedCookie, Cookie, UE_ARRAY_COUNT(AuthorisedCookie));
 							}
 
 							bRestartedHandshake = bRestartHandshake;
diff --git a/Engine/Source/Runtime/Engine/Private/Particles/SubUVAnimation.cpp b/Engine/Source/Runtime/Engine/Private/Particles/SubUVAnimation.cpp
--- a/Engine/Source/Runtime/Engine/Private/Particles/SubUVAnimation.cpp
+++ b/Engine/Source/Runtime/Engine/Private/Particles/SubUVAnimation.cpp
@@ -483,7 +483,7 @@ int32 ComputeNeighborCount(int32 X, int32 Y, int32 GlobalX, int32 GlobalY, int32
 {
 	int32 NeighborCount = 0;
 
-	for (int32 NeighborIndex = 0; NeighborIndex < ARRAY_COUNT(GNeighbors); NeighborIndex++)
+	for (int32 NeighborIndex = 0; NeighborIndex < UE_ARRAY_COUNT(GNeighbors); NeighborIndex++)
 	{
 		int32 NeighborX = X + GNeighbors[NeighborIndex].X;
 		int32 NeighborY = Y + GNeighbors[NeighborIndex].Y;
diff --git a/Engine/Source/Runtime/Engine/Private/PhysicsEngine/BodyInstance.cpp b/Engine/Source/Runtime/Engine/Private/PhysicsEngine/BodyInstance.cpp
--- a/Engine/Source/Runtime/Engine/Private/PhysicsEngine/BodyInstance.cpp
+++ b/Engine/Source/Runtime/Engine/Private/PhysicsEngine/BodyInstance.cpp
@@ -222,12 +222,12 @@ bool FCollisionResponse::AddReponseToArray(ECollisionChannel Channel, ECollision
 
 void FCollisionResponse::UpdateArrayFromResponseContainer()
 {
-	ResponseArray.Empty(ARRAY_COUNT(ResponseToChannels.EnumArray));
+	ResponseArray.Empty(UE_ARRAY_COUNT(ResponseToChannels.EnumArray));
 
 	const FCollisionResponseContainer& DefaultResponse = FCollisionResponseContainer::GetDefaultResponseContainer();
 	const UCollisionProfile* CollisionProfile = UCollisionProfile::Get();
 
-	for (int32 i = 0; i < ARRAY_COUNT(ResponseToChannels.EnumArray); i++)
+	for (int32 i = 0; i < UE_ARRAY_COUNT(ResponseToChannels.EnumArray); i++)
 	{
 		// if not same as default
 		if (ResponseToChannels.EnumArray[i] != DefaultResponse.EnumArray[i])
diff --git a/Engine/Source/Runtime/Engine/Private/PhysicsEngine/PhysUtils.cpp b/Engine/Source/Runtime/Engine/Private/PhysicsEngine/PhysUtils.cpp
--- a/Engine/Source/Runtime/Engine/Private/PhysicsEngine/PhysUtils.cpp
+++ b/Engine/Source/Runtime/Engine/Private/PhysicsEngine/PhysUtils.cpp
@@ -194,7 +194,7 @@ void FRigidBodyContactInfo::SwapOrder()
 /** Set the status of a particular channel in the structure. */
 bool FCollisionResponseContainer::SetResponse(ECollisionChannel Channel, ECollisionResponse NewResponse)
 {
-	if (Channel < ARRAY_COUNT(EnumArray))
+	if (Channel < UE_ARRAY_COUNT(EnumArray))
 	{
 		uint8& CurrentResponse = EnumArray[Channel];
 		if (CurrentResponse != NewResponse)
@@ -210,7 +210,7 @@ bool FCollisionResponseContainer::SetResponse(ECollisionChannel Channel, ECollis
 bool FCollisionResponseContainer::SetAllChannels(ECollisionResponse NewResponse)
 {
 	bool bHasChanged = false;
-	for(int32 i=0; i<ARRAY_COUNT(EnumArray); i++)
+	for(int32 i=0; i<UE_ARRAY_COUNT(EnumArray); i++)
 	{
 		uint8& CurrentResponse = EnumArray[i];
 		if (CurrentResponse != NewResponse)
@@ -225,7 +225,7 @@ bool FCollisionResponseContainer::SetAllChannels(ECollisionResponse NewResponse)
 bool FCollisionResponseContainer::ReplaceChannels(ECollisionResponse OldResponse, ECollisionResponse NewResponse)
 {
 	bool bHasChanged = false;
-	for (int32 i = 0; i < ARRAY_COUNT(EnumArray); i++)
+	for (int32 i = 0; i < UE_ARRAY_COUNT(EnumArray); i++)
 	{
 		uint8& CurrentResponse = EnumArray[i];
 		if(CurrentResponse == OldResponse)
@@ -240,7 +240,7 @@ bool FCollisionResponseContainer::ReplaceChannels(ECollisionResponse OldResponse
 FCollisionResponseContainer FCollisionResponseContainer::CreateMinContainer(const FCollisionResponseContainer& A, const FCollisionResponseContainer& B)
 {
 	FCollisionResponseContainer Result;
-	for(int32 i=0; i<ARRAY_COUNT(Result.EnumArray); i++)
+	for(int32 i=0; i<UE_ARRAY_COUNT(Result.EnumArray); i++)
 	{
 		Result.EnumArray[i] = FMath::Min(A.EnumArray[i], B.EnumArray[i]);
 	}
@@ -302,7 +302,7 @@ bool FPhysScene::ExecPxVis(const TCHAR* Cmd, FOutputDevice* Ar)
 	if ( FParse::Command(&Cmd,TEXT("PHYSX_CLEAR_ALL")) )
 	{
 		Ar->Logf(TEXT("Clearing all PhysX Debug Flags."));
-		for (int32 i = 0; i < ARRAY_COUNT(Flags); i++)
+		for (int32 i = 0; i < UE_ARRAY_COUNT(Flags); i++)
 		{
 			PScene->setVisualizationParameter(Flags[i].Flag, 0.0f);
 			bFoundFlag = true;
@@ -310,7 +310,7 @@ bool FPhysScene::ExecPxVis(const TCHAR* Cmd, FOutputDevice* Ar)
 	}
 	else
 	{
-		for (int32 i = 0; i < ARRAY_COUNT(Flags); i++)
+		for (int32 i = 0; i < UE_ARRAY_COUNT(Flags); i++)
 		{
 			// Parse out the command sent in and set only those flags
 			if (FParse::Command(&Cmd, Flags[i].Name))
diff --git a/Engine/Source/Runtime/Engine/Private/PhysicsEngine/PhysicsInterfaceUtils.cpp b/Engine/Source/Runtime/Engine/Private/PhysicsEngine/PhysicsInterfaceUtils.cpp
--- a/Engine/Source/Runtime/Engine/Private/PhysicsEngine/PhysicsInterfaceUtils.cpp
+++ b/Engine/Source/Runtime/Engine/Private/PhysicsEngine/PhysicsInterfaceUtils.cpp
@@ -70,7 +70,7 @@ FCollisionFilterData CreateTraceQueryFilterData(const uint8 MyChannel, const boo
 	}
 
 	// word1 encodes 'what i block', word2 encodes 'what i touch'
-	for (int32 i = 0; i<ARRAY_COUNT(InCollisionResponseContainer.EnumArray); i++)
+	for (int32 i = 0; i<UE_ARRAY_COUNT(InCollisionResponseContainer.EnumArray); i++)
 	{
 		if (InCollisionResponseContainer.EnumArray[i] == ECR_Block)
 		{
diff --git a/Engine/Source/Runtime/Engine/Private/PlayerState.cpp b/Engine/Source/Runtime/Engine/Private/PlayerState.cpp
--- a/Engine/Source/Runtime/Engine/Private/PlayerState.cpp
+++ b/Engine/Source/Runtime/Engine/Private/PlayerState.cpp
@@ -59,7 +59,7 @@ void APlayerState::UpdatePing(float InPing)
 		//	(misses the latest ping update, but averages a full 4 seconds data)
 		RecalculateAvgPing();
 
-		CurPingBucket = (CurPingBucket + 1) % ARRAY_COUNT(PingBucket);
+		CurPingBucket = (CurPingBucket + 1) % UE_ARRAY_COUNT(PingBucket);
 		CurPingBucketTimestamp = CurTime;
 
 
@@ -100,7 +100,7 @@ void APlayerState::RecalculateAvgPing()
 	int32 Sum = 0;
 	int32 Count = 0;
 
-	for (uint8 i=0; i<ARRAY_COUNT(PingBucket); i++)
+	for (uint8 i=0; i<UE_ARRAY_COUNT(PingBucket); i++)
 	{
 		Sum += PingBucket[i].PingSum;
 		Count += PingBucket[i].PingCount;
@@ -120,12 +120,12 @@ void APlayerState::RecalculateAvgPing()
 	PingBucketV2[CurPingBucket].AvgPingV2 = NumValidValues > 0 ? SumV2 / NumValidValues : MAX_flt;
 
 	float AvgSumV2 = 0.0f;
-	for (int32 i = 0; i < ARRAY_COUNT(PingBucketV2); i++)
+	for (int32 i = 0; i < UE_ARRAY_COUNT(PingBucketV2); i++)
 	{
 		AvgSumV2 += PingBucketV2[i].AvgPingV2;
 	}
 
-	ExactPingV2 = AvgSumV2 / ARRAY_COUNT(PingBucketV2);
+	ExactPingV2 = AvgSumV2 / UE_ARRAY_COUNT(PingBucketV2);
 
 	// Calculate the average, and divide it by 4 to optimize replication
 	ExactPing = (Count > 0 ? ((float)Sum / (float)Count) : 0.f);
diff --git a/Engine/Source/Runtime/Engine/Private/Polygon.cpp b/Engine/Source/Runtime/Engine/Private/Polygon.cpp
--- a/Engine/Source/Runtime/Engine/Private/Polygon.cpp
+++ b/Engine/Source/Runtime/Engine/Private/Polygon.cpp
@@ -745,7 +745,7 @@ int32 FPoly::Triangulate( ABrush* InOwnerBrush, TArray<FPoly>& OutTriangles )
 		vtx.TangentY = FVector::ZeroVector;
 		vtx.TangentZ = FVector::ZeroVector;
 		vtx.Color = FColor(0, 0, 0);
-		for( int32 uvIndex=0; uvIndex<ARRAY_COUNT(vtx.UVs); ++uvIndex )
+		for( int32 uvIndex=0; uvIndex<UE_ARRAY_COUNT(vtx.UVs); ++uvIndex )
 		{
 			vtx.UVs[uvIndex] = FVector2D(0.f, 0.f);
 		}
diff --git a/Engine/Source/Runtime/Engine/Private/PrecomputedVolumetricLightmap.cpp b/Engine/Source/Runtime/Engine/Private/PrecomputedVolumetricLightmap.cpp
--- a/Engine/Source/Runtime/Engine/Private/PrecomputedVolumetricLightmap.cpp
+++ b/Engine/Source/Runtime/Engine/Private/PrecomputedVolumetricLightmap.cpp
@@ -73,7 +73,7 @@ struct FVolumetricLightmapBrickTextureSet
 		SkyBentNormal.Format = BrickData.SkyBentNormal.Format;
 		DirectionalLightShadowing.Format = BrickData.DirectionalLightShadowing.Format;
 
-		for (int32 i = 0; i < ARRAY_COUNT(SHCoefficients); i++)
+		for (int32 i = 0; i < UE_ARRAY_COUNT(SHCoefficients); i++)
 		{
 			SHCoefficients[i].Format = BrickData.SHCoefficients[i].Format;
 		}
@@ -81,7 +81,7 @@ struct FVolumetricLightmapBrickTextureSet
 		AmbientVector.CreateTargetTexture(BrickDataDimensions);
 		AmbientVector.CreateUAV();
 
-		for (int32 i = 0; i < ARRAY_COUNT(SHCoefficients); i++)
+		for (int32 i = 0; i < UE_ARRAY_COUNT(SHCoefficients); i++)
 		{
 			SHCoefficients[i].CreateTargetTexture(BrickDataDimensions);
 			SHCoefficients[i].CreateUAV();
@@ -100,7 +100,7 @@ struct FVolumetricLightmapBrickTextureSet
 	void Release()
 	{
 		AmbientVector.Texture.SafeRelease();
-		for (int32 i = 0; i < ARRAY_COUNT(SHCoefficients); i++)
+		for (int32 i = 0; i < UE_ARRAY_COUNT(SHCoefficients); i++)
 		{
 			SHCoefficients[i].Texture.SafeRelease();
 		}
@@ -108,7 +108,7 @@ struct FVolumetricLightmapBrickTextureSet
 		DirectionalLightShadowing.Texture.SafeRelease();
 
 		AmbientVector.UAV.SafeRelease();
-		for (int32 i = 0; i < ARRAY_COUNT(SHCoefficients); i++)
+		for (int32 i = 0; i < UE_ARRAY_COUNT(SHCoefficients); i++)
 		{
 			SHCoefficients[i].UAV.SafeRelease();
 		}
@@ -212,7 +212,7 @@ FArchive& operator<<(FArchive& Ar,FPrecomputedVolumetricLightmapData& Volume)
 
 	Ar << Volume.BrickData.AmbientVector;
 
-	for (int32 i = 0; i < ARRAY_COUNT(Volume.BrickData.SHCoefficients); i++)
+	for (int32 i = 0; i < UE_ARRAY_COUNT(Volume.BrickData.SHCoefficients); i++)
 	{
 		Ar << Volume.BrickData.SHCoefficients[i];
 	}
@@ -283,7 +283,7 @@ int32 FVolumetricLightmapBrickData::GetMinimumVoxelSize() const
 	check(AmbientVector.Format != PF_Unknown);
 	int32 VoxelSize = GPixelFormats[AmbientVector.Format].BlockBytes;
 
-	for (int32 i = 0; i < ARRAY_COUNT(SHCoefficients); i++)
+	for (int32 i = 0; i < UE_ARRAY_COUNT(SHCoefficients); i++)
 	{
 		VoxelSize += GPixelFormats[SHCoefficients[i].Format].BlockBytes;
 	}
@@ -340,7 +340,7 @@ ENGINE_API void FPrecomputedVolumetricLightmapData::InitRHI()
 		{
 			BrickData.AmbientVector.CreateTexture(BrickDataDimensions);
 
-			for (int32 i = 0; i < ARRAY_COUNT(BrickData.SHCoefficients); i++)
+			for (int32 i = 0; i < UE_ARRAY_COUNT(BrickData.SHCoefficients); i++)
 			{
 				BrickData.SHCoefficients[i].CreateTexture(BrickDataDimensions);
 			}
@@ -610,7 +610,7 @@ ENGINE_API void FPrecomputedVolumetricLightmapData::AddToSceneData(FPrecomputedV
 
 	SceneData->BrickDataDimensions = GVolumetricLightmapBrickAtlas.TextureSet.BrickDataDimensions;
 	SceneData->BrickData.AmbientVector = GVolumetricLightmapBrickAtlas.TextureSet.AmbientVector;
-	for (int32 i = 0; i < ARRAY_COUNT(SceneData->BrickData.SHCoefficients); i++)
+	for (int32 i = 0; i < UE_ARRAY_COUNT(SceneData->BrickData.SHCoefficients); i++)
 	{
 		SceneData->BrickData.SHCoefficients[i] = GVolumetricLightmapBrickAtlas.TextureSet.SHCoefficients[i];
 	}
@@ -691,7 +691,7 @@ ENGINE_API void FPrecomputedVolumetricLightmapData::RemoveFromSceneData(FPrecomp
 
 	SceneData->BrickDataDimensions = GVolumetricLightmapBrickAtlas.TextureSet.BrickDataDimensions;
 	SceneData->BrickData.AmbientVector = GVolumetricLightmapBrickAtlas.TextureSet.AmbientVector;
-	for (int32 i = 0; i < ARRAY_COUNT(SceneData->BrickData.SHCoefficients); i++)
+	for (int32 i = 0; i < UE_ARRAY_COUNT(SceneData->BrickData.SHCoefficients); i++)
 	{
 		SceneData->BrickData.SHCoefficients[i] = GVolumetricLightmapBrickAtlas.TextureSet.SHCoefficients[i];
 	}
@@ -909,7 +909,7 @@ void CopyDataIntoAtlas(FRHICommandList& RHICmdList, ERHIFeatureLevel::Type Featu
 		FComputeShaderUtils::Dispatch(RHICmdList, *ComputeShader, Parameters, FIntVector(NumBricks, 1, 1));
 	}
 
-	for (int32 i = 0; i < ARRAY_COUNT(SrcData.SHCoefficients); i++)
+	for (int32 i = 0; i < UE_ARRAY_COUNT(SrcData.SHCoefficients); i++)
 	{
 		TShaderMap<FGlobalShaderType>* GlobalShaderMap = GetGlobalShaderMap(FeatureLevel);
 
@@ -942,7 +942,7 @@ void FVolumetricLightmapBrickAtlas::Insert(int32 Index, FPrecomputedVolumetricLi
 	else
 	{
 		check(TextureSet.AmbientVector.Format == Data->BrickData.AmbientVector.Format);
-		for (int32 i = 0; i < ARRAY_COUNT(TextureSet.SHCoefficients); i++)
+		for (int32 i = 0; i < UE_ARRAY_COUNT(TextureSet.SHCoefficients); i++)
 		{
 			check(TextureSet.SHCoefficients[i].Format == Data->BrickData.SHCoefficients[i].Format);
 		}
diff --git a/Engine/Source/Runtime/Engine/Private/Scene.cpp b/Engine/Source/Runtime/Engine/Private/Scene.cpp
--- a/Engine/Source/Runtime/Engine/Private/Scene.cpp
+++ b/Engine/Source/Runtime/Engine/Private/Scene.cpp
@@ -938,7 +938,7 @@ FPostProcessSettings::FPostProcessSettings(const FPostProcessSettings& Settings)
 	, WeightedBlendables(Settings.WeightedBlendables)
 	//, Blendables_DEPRECATED(Settings.Blendables_DEPRECATED)
 {
-	for (int32 i = 0; i < ARRAY_COUNT(LensFlareTints); i++)
+	for (int32 i = 0; i < UE_ARRAY_COUNT(LensFlareTints); i++)
 		LensFlareTints[i] = Settings.LensFlareTints[i];
 }
 	
diff --git a/Engine/Source/Runtime/Engine/Private/ShaderCompiler/ShaderCompilerXGE.cpp b/Engine/Source/Runtime/Engine/Private/ShaderCompiler/ShaderCompilerXGE.cpp
--- a/Engine/Source/Runtime/Engine/Private/ShaderCompiler/ShaderCompilerXGE.cpp
+++ b/Engine/Source/Runtime/Engine/Private/ShaderCompiler/ShaderCompilerXGE.cpp
@@ -109,7 +109,7 @@ bool FShaderCompileXGEThreadRunnable_XmlInterface::IsSupported()
 	{
 		IPlatformFile& PlatformFile = FPlatformFileManager::Get().GetPlatformFile();
 
-		for (int PathIndex = 0; PathIndex < ARRAY_COUNT(Paths); PathIndex++)
+		for (int PathIndex = 0; PathIndex < UE_ARRAY_COUNT(Paths); PathIndex++)
 		{
 			if (PlatformFile.FileExists(Paths[PathIndex]))
 			{
@@ -349,7 +349,7 @@ static void WriteScriptFileFooter(FArchive* ScriptFile)
 		"\t</Project>\r\n"
 		"</BuildSet>\r\n";
 
-	ScriptFile->Serialize((void*)HeaderFooter, sizeof(ANSICHAR) * (ARRAY_COUNT(HeaderFooter) - 1));
+	ScriptFile->Serialize((void*)HeaderFooter, sizeof(ANSICHAR) * (UE_ARRAY_COUNT(HeaderFooter) - 1));
 }
 
 void FShaderCompileXGEThreadRunnable_XmlInterface::GatherResultsFromXGE()
diff --git a/Engine/Source/Runtime/Engine/Private/ShadowMap.cpp b/Engine/Source/Runtime/Engine/Private/ShadowMap.cpp
--- a/Engine/Source/Runtime/Engine/Private/ShadowMap.cpp
+++ b/Engine/Source/Runtime/Engine/Private/ShadowMap.cpp
@@ -503,7 +503,7 @@ FShadowMap2D::FShadowMap2D() :
 	CoordinateScale(FVector2D(0, 0)),
 	CoordinateBias(FVector2D(0, 0))
 {
-	for (int Channel = 0; Channel < ARRAY_COUNT(bChannelValid); Channel++)
+	for (int Channel = 0; Channel < UE_ARRAY_COUNT(bChannelValid); Channel++)
 	{
 		bChannelValid[Channel] = false;
 	}
@@ -514,7 +514,7 @@ FShadowMap2D::FShadowMap2D(const TMap<ULightComponent*,FShadowMapData2D*>& Shado
 	CoordinateScale(FVector2D(0, 0)),
 	CoordinateBias(FVector2D(0, 0))
 {
-	for (int Channel = 0; Channel < ARRAY_COUNT(bChannelValid); Channel++)
+	for (int Channel = 0; Channel < UE_ARRAY_COUNT(bChannelValid); Channel++)
 	{
 		bChannelValid[Channel] = false;
 	}
@@ -531,7 +531,7 @@ FShadowMap2D::FShadowMap2D(TArray<FGuid> LightGuids)
 	, CoordinateScale(FVector2D(0, 0))
 	, CoordinateBias(FVector2D(0, 0))
 {
-	for (int Channel = 0; Channel < ARRAY_COUNT(bChannelValid); Channel++)
+	for (int Channel = 0; Channel < UE_ARRAY_COUNT(bChannelValid); Channel++)
 	{
 		bChannelValid[Channel] = false;
 	}
@@ -561,7 +561,7 @@ void FShadowMap2D::Serialize(FArchive& Ar)
 
 	Ar << CoordinateScale << CoordinateBias;
 
-	for (int Channel = 0; Channel < ARRAY_COUNT(bChannelValid); Channel++)
+	for (int Channel = 0; Channel < UE_ARRAY_COUNT(bChannelValid); Channel++)
 	{
 		Ar << bChannelValid[Channel];
 	}
@@ -1108,7 +1108,7 @@ int32 FShadowMap2D::EncodeSingleTexture(ULevel* LightingScenario, FShadowMapPend
 							ExtrapolatedFilterableComponents[i] = 0;
 						}
 
-						for (int32 NeighborIndex = 0; NeighborIndex < ARRAY_COUNT(Neighbors); NeighborIndex++)
+						for (int32 NeighborIndex = 0; NeighborIndex < UE_ARRAY_COUNT(Neighbors); NeighborIndex++)
 						{
 							if (static_cast<int32>(DestY) + Neighbors[NeighborIndex].Y >= 0 
 								&& DestY + Neighbors[NeighborIndex].Y < MipSizeY
diff --git a/Engine/Source/Runtime/Engine/Private/StaticMeshLight.cpp b/Engine/Source/Runtime/Engine/Private/StaticMeshLight.cpp
--- a/Engine/Source/Runtime/Engine/Private/StaticMeshLight.cpp
+++ b/Engine/Source/Runtime/Engine/Private/StaticMeshLight.cpp
@@ -38,7 +38,7 @@ static void GetStaticLightingVertex(
 	OutVertex.WorldTangentY = LocalToWorld.TransformVector(VertexBuffer.VertexTangentY(VertexIndex)).GetSafeNormal();
 	OutVertex.WorldTangentZ = LocalToWorldInverseTranspose.TransformVector(VertexBuffer.VertexTangentZ(VertexIndex)).GetSafeNormal();
 
-	checkSlow(VertexBuffer.GetNumTexCoords() <= ARRAY_COUNT(OutVertex.TextureCoordinates));
+	checkSlow(VertexBuffer.GetNumTexCoords() <= UE_ARRAY_COUNT(OutVertex.TextureCoordinates));
 	for(uint32 LightmapTextureCoordinateIndex = 0;LightmapTextureCoordinateIndex < VertexBuffer.GetNumTexCoords();LightmapTextureCoordinateIndex++)
 	{
 		OutVertex.TextureCoordinates[LightmapTextureCoordinateIndex] = VertexBuffer.GetVertexUV(VertexIndex,LightmapTextureCoordinateIndex);
diff --git a/Engine/Source/Runtime/Engine/Private/SystemSettings.cpp b/Engine/Source/Runtime/Engine/Private/SystemSettings.cpp
--- a/Engine/Source/Runtime/Engine/Private/SystemSettings.cpp
+++ b/Engine/Source/Runtime/Engine/Private/SystemSettings.cpp
@@ -42,11 +42,11 @@ static inline const TCHAR* GetSectionName(bool bIsEditor)
 		static int32 SystemSettingsStrLen = FCString::Strlen(ReturnedOverrideName);
 		if (FCString::Strnicmp(*OverrideSubName, TEXT("SystemSettings"), SystemSettingsStrLen) == 0)
 		{
-			FCString::Strcpy(ReturnedOverrideName + SystemSettingsStrLen, ARRAY_COUNT(ReturnedOverrideName), (*OverrideSubName + SystemSettingsStrLen));
+			FCString::Strcpy(ReturnedOverrideName + SystemSettingsStrLen, UE_ARRAY_COUNT(ReturnedOverrideName), (*OverrideSubName + SystemSettingsStrLen));
 		}
 		else
 		{
-			FCString::Strcpy(ReturnedOverrideName + SystemSettingsStrLen, ARRAY_COUNT(ReturnedOverrideName), *OverrideSubName);
+			FCString::Strcpy(ReturnedOverrideName + SystemSettingsStrLen, UE_ARRAY_COUNT(ReturnedOverrideName), *OverrideSubName);
 		}
 		return ReturnedOverrideName;
 	}
diff --git a/Engine/Source/Runtime/Engine/Private/Texture2D.cpp b/Engine/Source/Runtime/Engine/Private/Texture2D.cpp
--- a/Engine/Source/Runtime/Engine/Private/Texture2D.cpp
+++ b/Engine/Source/Runtime/Engine/Private/Texture2D.cpp
@@ -1465,7 +1465,7 @@ FTexture2DResource::FTexture2DResource( UTexture2D* InOwner, int32 InitialMipCou
 	bSRGB = InOwner->SRGB;
 
 	check(InitialMipCount>0);
-	check(ARRAY_COUNT(MipData)>=GMaxTextureMipCount);
+	check(UE_ARRAY_COUNT(MipData)>=GMaxTextureMipCount);
 
 	// Keep track of first miplevel to use.
 	CurrentFirstMip = InOwner->GetNumMips() - InitialMipCount;
@@ -1498,7 +1498,7 @@ FTexture2DResource::~FTexture2DResource()
 		});
 
 	// Make sure we're not leaking memory if InitRHI has never been called.
-	for( int32 MipIndex=0; MipIndex<ARRAY_COUNT(MipData); MipIndex++ )
+	for( int32 MipIndex=0; MipIndex<UE_ARRAY_COUNT(MipData); MipIndex++ )
 	{
 		// free any mip data that was copied 
 		if( MipData[MipIndex] )
diff --git a/Engine/Source/Runtime/Engine/Private/TextureCube.cpp b/Engine/Source/Runtime/Engine/Private/TextureCube.cpp
--- a/Engine/Source/Runtime/Engine/Private/TextureCube.cpp
+++ b/Engine/Source/Runtime/Engine/Private/TextureCube.cpp
@@ -139,7 +139,7 @@ public:
 		//Initialize the MipData array
 		for ( int32 FaceIndex=0;FaceIndex<6; FaceIndex++)
 		{
-			for( int32 MipIndex=0; MipIndex<ARRAY_COUNT(MipData[FaceIndex]); MipIndex++ )
+			for( int32 MipIndex=0; MipIndex<UE_ARRAY_COUNT(MipData[FaceIndex]); MipIndex++ )
 			{
 				MipData[FaceIndex][MipIndex] = NULL;
 			}
@@ -183,7 +183,7 @@ public:
 		// Make sure we're not leaking memory if InitRHI has never been called.
 		for (int32 i=0; i<6; i++)
 		{
-			for( int32 MipIndex=0; MipIndex<ARRAY_COUNT(MipData[i]); MipIndex++ )
+			for( int32 MipIndex=0; MipIndex<UE_ARRAY_COUNT(MipData[i]); MipIndex++ )
 			{
 				// free any mip data that was copied 
 				if( MipData[i][MipIndex] )
diff --git a/Engine/Source/Runtime/Engine/Private/UnrealEngine.cpp b/Engine/Source/Runtime/Engine/Private/UnrealEngine.cpp
--- a/Engine/Source/Runtime/Engine/Private/UnrealEngine.cpp
+++ b/Engine/Source/Runtime/Engine/Private/UnrealEngine.cpp
@@ -1411,7 +1411,7 @@ void UEngine::Init(IEngineLoop* InEngineLoop)
 #if !UE_BUILD_SHIPPING
 	// Check for overrides to the default map on the command line
 	TCHAR MapName[512];
-	if ( FParse::Value(FCommandLine::Get(), TEXT("DEFAULTMAP="), MapName, ARRAY_COUNT(MapName)) )
+	if ( FParse::Value(FCommandLine::Get(), TEXT("DEFAULTMAP="), MapName, UE_ARRAY_COUNT(MapName)) )
 	{
 		UE_LOG(LogEngine, Log, TEXT("Overriding default map to %s"), MapName);
 
@@ -3746,7 +3746,7 @@ static void BufferOverflowFunction(SIZE_T BufferSize, const ANSICHAR* Buffer)
 	ANSICHAR LocalBuffer[32];
 	LocalBuffer[0] = LocalBuffer[31] = 0; //if BufferSize is 0 then there's nothing to print out!
 
-	BufferSize = FMath::Min<SIZE_T>(BufferSize, ARRAY_COUNT(LocalBuffer)-1);
+	BufferSize = FMath::Min<SIZE_T>(BufferSize, UE_ARRAY_COUNT(LocalBuffer)-1);
 
 	for( uint32 i = 0; i < BufferSize; i++ ) 
 	{
@@ -7600,7 +7600,7 @@ bool UEngine::HandleObjCommand( const TCHAR* Cmd, FOutputDevice& Ar )
 
 		if ( !ParseObject<UClass>( Cmd, TEXT("CLASS="), Cls, ANY_PACKAGE ) || !ParseObject(Cmd,TEXT("NAME="), Cls, Obj, ANY_PACKAGE) )
 		{
-			if ( FParse::Token(Cmd,ObjectName,ARRAY_COUNT(ObjectName), 1) )
+			if ( FParse::Token(Cmd,ObjectName,UE_ARRAY_COUNT(ObjectName), 1) )
 			{
 				Obj = FindObject<UObject>(ANY_PACKAGE,ObjectName);
 			}
@@ -7956,7 +7956,7 @@ bool UEngine::HandleGetIniCommand(const TCHAR* Cmd, FOutputDevice& Ar)
 	TCHAR IniPlusSectionName[256];
 	TCHAR KeyName[256];
 
-	if (FParse::Token(Cmd, IniPlusSectionName, ARRAY_COUNT(IniPlusSectionName), true))
+	if (FParse::Token(Cmd, IniPlusSectionName, UE_ARRAY_COUNT(IniPlusSectionName), true))
 	{
 		FString IniPlusSection = IniPlusSectionName;
 		int32 IniDelim = IniPlusSection.Find(TEXT(":"));
@@ -8002,7 +8002,7 @@ bool UEngine::HandleGetIniCommand(const TCHAR* Cmd, FOutputDevice& Ar)
 
 		if (!IniName.IsEmpty() && !SectionName.IsEmpty())
 		{
-			if (FParse::Token(Cmd, KeyName, ARRAY_COUNT(KeyName), true))
+			if (FParse::Token(Cmd, KeyName, UE_ARRAY_COUNT(KeyName), true))
 			{
 				TArray<FString> Values;
 
@@ -8386,7 +8386,7 @@ bool UEngine::PerformError(const TCHAR* Cmd, FOutputDevice& Ar)
 		// stack overflow test - this case should be caught by /GS (Buffer Overflow Check) compile option
 		ANSICHAR SrcBuffer[] = "12345678901234567890123456789012345678901234567890";
 		SetCrashType(ECrashType::Debug);
-		BufferOverflowFunction(ARRAY_COUNT(SrcBuffer), SrcBuffer);
+		BufferOverflowFunction(UE_ARRAY_COUNT(SrcBuffer), SrcBuffer);
 		return true;
 	}
 	else if (FParse::Command(&Cmd, TEXT("CRTINVALID")))
@@ -10423,7 +10423,7 @@ FString appGetStartupMap(const TCHAR* CommandLine)
 #endif // UE_BUILD_SHIPPING
 
 	const TCHAR* Tmp = CommandLine ? CommandLine : TEXT("");
-	if (!FParse::Token(Tmp, Parm, ARRAY_COUNT(Parm), 0) || Parm[0] == '-')
+	if (!FParse::Token(Tmp, Parm, UE_ARRAY_COUNT(Parm), 0) || Parm[0] == '-')
 	{
 		const UGameMapsSettings* GameMapsSettings = GetDefault<UGameMapsSettings>();
 		FCString::Strcpy(Parm, *(GameMapsSettings->GetGameDefaultMap() + GameMapsSettings->LocalMapOptions));
@@ -11484,12 +11484,12 @@ void UEngine::SpawnServerActors(UWorld *World)
 	{
 		TCHAR Str[2048];
 		const TCHAR* Ptr = * FullServerActors[i];
-		if( FParse::Token( Ptr, Str, ARRAY_COUNT(Str), 1 ) )
+		if( FParse::Token( Ptr, Str, UE_ARRAY_COUNT(Str), 1 ) )
 		{
 			UE_LOG(LogNet, Log, TEXT("Spawning: %s"), Str );
 			UClass* HelperClass = StaticLoadClass( AActor::StaticClass(), NULL, Str, NULL, LOAD_None, NULL );
 			AActor* Actor = World->SpawnActor( HelperClass );
-			while( Actor && FParse::Token(Ptr,Str,ARRAY_COUNT(Str),1) )
+			while( Actor && FParse::Token(Ptr,Str,UE_ARRAY_COUNT(Str),1) )
 			{
 				TCHAR* Value = FCString::Strchr(Str,'=');
 				if( Value )
diff --git a/Engine/Source/Runtime/Engine/Private/UserInterface/Canvas.cpp b/Engine/Source/Runtime/Engine/Private/UserInterface/Canvas.cpp
--- a/Engine/Source/Runtime/Engine/Private/UserInterface/Canvas.cpp
+++ b/Engine/Source/Runtime/Engine/Private/UserInterface/Canvas.cpp
@@ -1551,7 +1551,7 @@ void UCanvas::ClippedStrLen( const UFont* Font, float ScaleX, float ScaleY, int3
 void VARARGS UCanvas::WrappedStrLenf( const UFont* Font, float ScaleX, float ScaleY, int32& XL, int32& YL, const TCHAR* Fmt, ... ) 
 {
 	TCHAR Text[4096];
-	GET_VARARGS( Text, ARRAY_COUNT(Text), ARRAY_COUNT(Text)-1, Fmt, Fmt );
+	GET_VARARGS( Text, UE_ARRAY_COUNT(Text), UE_ARRAY_COUNT(Text)-1, Fmt, Fmt );
 
 	FFontRenderInfo Info;
 	WrappedPrint( false, 0.0f, 0.0f, XL, YL, Font, ScaleX, ScaleY, false, false, Text, Info ); 
diff --git a/Engine/Source/Runtime/Engine/Private/UserInterface/PlayerInput.cpp b/Engine/Source/Runtime/Engine/Private/UserInterface/PlayerInput.cpp
--- a/Engine/Source/Runtime/Engine/Private/UserInterface/PlayerInput.cpp
+++ b/Engine/Source/Runtime/Engine/Private/UserInterface/PlayerInput.cpp
@@ -1402,7 +1402,7 @@ void UPlayerInput::ClearSmoothing()
 float UPlayerInput::SmoothMouse(float aMouse, uint8& SampleCount, int32 Index)
 {
 	check(Index >= 0);
-	check(Index < ARRAY_COUNT(ZeroTime));
+	check(Index < UE_ARRAY_COUNT(ZeroTime));
 
 	UWorld* World = GetWorld();
 	if (World)
@@ -2004,7 +2004,7 @@ bool UPlayerInput::Exec(UWorld* InWorld, const TCHAR* Str,FOutputDevice& Ar)
 {
 	TCHAR Temp[256];
 
-	if( FParse::Command( &Str, TEXT("KEYBINDING") ) && FParse::Token( Str, Temp, ARRAY_COUNT(Temp), 0 ) )
+	if( FParse::Command( &Str, TEXT("KEYBINDING") ) && FParse::Token( Str, Temp, UE_ARRAY_COUNT(Temp), 0 ) )
 	{
 		FKey const Key(Temp);
 		if (Key.IsValid())
@@ -2021,7 +2021,7 @@ bool UPlayerInput::Exec(UWorld* InWorld, const TCHAR* Str,FOutputDevice& Ar)
 
 		return 1;
 	}
-	else if( !bExecutingBindCommand && FParse::Token( Str, Temp, ARRAY_COUNT(Temp), 0 ) )
+	else if( !bExecutingBindCommand && FParse::Token( Str, Temp, UE_ARRAY_COUNT(Temp), 0 ) )
 	{
 		FKey const Key(Temp);
 		if(Key.IsValid())
diff --git a/Engine/Source/Runtime/Engine/Public/LightMap.h b/Engine/Source/Runtime/Engine/Public/LightMap.h
--- a/Engine/Source/Runtime/Engine/Public/LightMap.h
+++ b/Engine/Source/Runtime/Engine/Public/LightMap.h
@@ -213,7 +213,7 @@ public:
 
 	void GetReferencedTextures(TArray<UTexture2D*>& OutTextures) const
 	{
-		for (int32 BasisIndex = 0; BasisIndex < ARRAY_COUNT(Textures); BasisIndex++)
+		for (int32 BasisIndex = 0; BasisIndex < UE_ARRAY_COUNT(Textures); BasisIndex++)
 		{
 			if (Textures[BasisIndex])
 			{
diff --git a/Engine/Source/Runtime/Engine/Public/Net/DataChannel.h b/Engine/Source/Runtime/Engine/Public/Net/DataChannel.h
--- a/Engine/Source/Runtime/Engine/Public/Net/DataChannel.h
+++ b/Engine/Source/Runtime/Engine/Public/Net/DataChannel.h
@@ -50,7 +50,7 @@ private:
 		static bool bInitialized = false;
 		if (!bInitialized)
 		{
-			for (int32 i = 0; i < ARRAY_COUNT(Names); i++)
+			for (int32 i = 0; i < UE_ARRAY_COUNT(Names); i++)
 			{
 				Names[i] = TEXT("");
 			}
diff --git a/Engine/Source/Runtime/Engine/Public/SceneManagement.h b/Engine/Source/Runtime/Engine/Public/SceneManagement.h
--- a/Engine/Source/Runtime/Engine/Public/SceneManagement.h
+++ b/Engine/Source/Runtime/Engine/Public/SceneManagement.h
@@ -616,7 +616,7 @@ public:
 		InvUniformPenumbraSize(FVector4(0, 0, 0, 0)),
 		Type(SMIT_None)
 	{
-		for (int Channel = 0; Channel < ARRAY_COUNT(bChannelValid); Channel++)
+		for (int Channel = 0; Channel < UE_ARRAY_COUNT(bChannelValid); Channel++)
 		{
 			bChannelValid[Channel] = false;
 		}
diff --git a/Engine/Source/Runtime/Experimental/Chaos/Private/Chaos/PBDCollisionConstraint.cpp b/Engine/Source/Runtime/Experimental/Chaos/Private/Chaos/PBDCollisionConstraint.cpp
--- a/Engine/Source/Runtime/Experimental/Chaos/Private/Chaos/PBDCollisionConstraint.cpp
+++ b/Engine/Source/Runtime/Experimental/Chaos/Private/Chaos/PBDCollisionConstraint.cpp
@@ -178,7 +178,7 @@ struct FStatHelper2
 template <>
 struct FStatHelper2<true>
 {
-	int32 BucketCount[ARRAY_COUNT(BucketSizes2)];
+	int32 BucketCount[UE_ARRAY_COUNT(BucketSizes2)];
 	int32 MaxCount;
 
 	FStatHelper2()
@@ -189,7 +189,7 @@ struct FStatHelper2<true>
 
 	void Record(int32 Count)
 	{
-		for (int32 BucketIdx = 1; BucketIdx < ARRAY_COUNT(BucketSizes2); ++BucketIdx)
+		for (int32 BucketIdx = 1; BucketIdx < UE_ARRAY_COUNT(BucketSizes2); ++BucketIdx)
 		{
 			if (Count >= BucketSizes2[BucketIdx - 1] && Count < BucketSizes2[BucketIdx])
 			{
@@ -216,10 +216,10 @@ struct FStatHelper2<true>
 		}
 
 		const float CountPerChar = MaxBucketCount / 20.f;
-		for (int32 Idx = 1; Idx < ARRAY_COUNT(BucketSizes2); ++Idx)
+		for (int32 Idx = 1; Idx < UE_ARRAY_COUNT(BucketSizes2); ++Idx)
 		{
 			int32 NumChars = BucketCount[Idx] / CountPerChar;
-			if (Idx < ARRAY_COUNT(BucketSizes2) - 1)
+			if (Idx < UE_ARRAY_COUNT(BucketSizes2) - 1)
 			{
 				OutLog += FString::Printf(TEXT("\t[%4d - %4d) (%4d) |"), BucketSizes2[Idx - 1], BucketSizes2[Idx], BucketCount[Idx]);
 			}
diff --git a/Engine/Source/Runtime/HeadMountedDisplay/Private/XRMotionControllerBase.cpp b/Engine/Source/Runtime/HeadMountedDisplay/Private/XRMotionControllerBase.cpp
--- a/Engine/Source/Runtime/HeadMountedDisplay/Private/XRMotionControllerBase.cpp
+++ b/Engine/Source/Runtime/HeadMountedDisplay/Private/XRMotionControllerBase.cpp
@@ -89,7 +89,7 @@ ETrackingStatus FXRMotionControllerBase::GetControllerTrackingStatus(const int32
 
 void FXRMotionControllerBase::EnumerateSources(TArray<FMotionControllerSource>& SourcesOut) const
 {
-	const int32 HandsCount = ARRAY_COUNT(XRMotionControllerBase_Impl::LegacyHandMappings);
+	const int32 HandsCount = UE_ARRAY_COUNT(XRMotionControllerBase_Impl::LegacyHandMappings);
 	ensure(HandsCount == (int32)EControllerHand::ControllerHand_Count);
 
 	for (int32 HandIndex = 0; HandIndex < HandsCount; ++HandIndex)
@@ -100,7 +100,7 @@ void FXRMotionControllerBase::EnumerateSources(TArray<FMotionControllerSource>&
 
 bool FXRMotionControllerBase::GetHandEnumForSourceName(const FName Source, EControllerHand& OutHand)
 {
-	const int32 HandsCount = ARRAY_COUNT(XRMotionControllerBase_Impl::LegacyHandMappings);
+	const int32 HandsCount = UE_ARRAY_COUNT(XRMotionControllerBase_Impl::LegacyHandMappings);
 	ensure(HandsCount == (int32)EControllerHand::ControllerHand_Count);
 
 	bool bFound = false;
diff --git a/Engine/Source/Runtime/ImageWrapper/Private/Formats/ExrImageWrapper.cpp b/Engine/Source/Runtime/ImageWrapper/Private/Formats/ExrImageWrapper.cpp
--- a/Engine/Source/Runtime/ImageWrapper/Private/Formats/ExrImageWrapper.cpp
+++ b/Engine/Source/Runtime/ImageWrapper/Private/Formats/ExrImageWrapper.cpp
@@ -270,7 +270,7 @@ const char* FExrImageWrapper::GetRawChannelName(int ChannelIndex) const
 		break;
 		case ERGBFormat::Gray:
 		{
-			check(ChannelIndex < ARRAY_COUNT(GrayChannelNames));
+			check(ChannelIndex < UE_ARRAY_COUNT(GrayChannelNames));
 			ChannelNames = GrayChannelNames;
 		}
 		break;
diff --git a/Engine/Source/Runtime/InstallBundleManager/Public/InstallBundleManagerInterface.h b/Engine/Source/Runtime/InstallBundleManager/Public/InstallBundleManagerInterface.h
--- a/Engine/Source/Runtime/InstallBundleManager/Public/InstallBundleManagerInterface.h
+++ b/Engine/Source/Runtime/InstallBundleManager/Public/InstallBundleManagerInterface.h
@@ -42,7 +42,7 @@ inline const TCHAR* LexToString(EInstallBundleModuleInitResult Result)
 		TEXT("NoInternetConnectionError"),
 		TEXT("ConfigurationError"),
 	};
-	static_assert(static_cast<UnderType>(EInstallBundleModuleInitResult::Count) == ARRAY_COUNT(Strings), "");
+	static_assert(static_cast<UnderType>(EInstallBundleModuleInitResult::Count) == UE_ARRAY_COUNT(Strings), "");
 
 	return Strings[static_cast<UnderType>(Result)];
 }
@@ -72,7 +72,7 @@ inline const TCHAR* LexToString(EInstallBundleResult Result)
 		TEXT("UserCancelledError"),
 		TEXT("InitializationError"),
 	};
-	static_assert(static_cast<UnderType>(EInstallBundleResult::Count) == ARRAY_COUNT(Strings), "");
+	static_assert(static_cast<UnderType>(EInstallBundleResult::Count) == UE_ARRAY_COUNT(Strings), "");
 
 	return Strings[static_cast<UnderType>(Result)];
 }
@@ -141,7 +141,7 @@ inline const TCHAR* LexToString(EInstallBundleStatus Status)
 		TEXT("Installed"),
 	};
 
-	static_assert(static_cast<UnderType>(EInstallBundleStatus::Count) == ARRAY_COUNT(Strings), "");
+	static_assert(static_cast<UnderType>(EInstallBundleStatus::Count) == UE_ARRAY_COUNT(Strings), "");
 
 	return Strings[static_cast<UnderType>(Status)];
 }
@@ -219,7 +219,7 @@ inline const TCHAR* LexToString(EInstallBundleContentState State)
 		TEXT("NeedsUpdate"),
 		TEXT("UpToDate"),
 	};
-	static_assert(static_cast<UnderType>(EInstallBundleContentState::Count) == ARRAY_COUNT(Strings), "");
+	static_assert(static_cast<UnderType>(EInstallBundleContentState::Count) == UE_ARRAY_COUNT(Strings), "");
 
 	return Strings[static_cast<UnderType>(State)];
 }
diff --git a/Engine/Source/Runtime/Landscape/Private/LandscapeEdit.cpp b/Engine/Source/Runtime/Landscape/Private/LandscapeEdit.cpp
--- a/Engine/Source/Runtime/Landscape/Private/LandscapeEdit.cpp
+++ b/Engine/Source/Runtime/Landscape/Private/LandscapeEdit.cpp
@@ -3083,15 +3083,15 @@ bool ALandscapeProxy::ExportToRawMesh(int32 InExportLOD, FMeshDescription& OutRa
 
 		//We need to not duplicate the vertex position, so we use the FIndexAndZ to achieve fast result
 		TArray<FIndexAndZ> VertIndexAndZ;
-		VertIndexAndZ.Reserve(ComponentSizeQuadsLOD*ComponentSizeQuadsLOD*ARRAY_COUNT(QuadPattern));
+		VertIndexAndZ.Reserve(ComponentSizeQuadsLOD*ComponentSizeQuadsLOD*UE_ARRAY_COUNT(QuadPattern));
 		int32 CurrentIndex = 0;
 		TMap<int32, FVector> IndexToPosition;
-		IndexToPosition.Reserve(ComponentSizeQuadsLOD*ComponentSizeQuadsLOD*ARRAY_COUNT(QuadPattern));
+		IndexToPosition.Reserve(ComponentSizeQuadsLOD*ComponentSizeQuadsLOD*UE_ARRAY_COUNT(QuadPattern));
 		for (int32 y = 0; y < ComponentSizeQuadsLOD; y++)
 		{
 			for (int32 x = 0; x < ComponentSizeQuadsLOD; x++)
 			{
-				for (int32 i = 0; i < ARRAY_COUNT(QuadPattern); i++)
+				for (int32 i = 0; i < UE_ARRAY_COUNT(QuadPattern); i++)
 				{
 					int32 VertexX = x + QuadPattern[i].X;
 					int32 VertexY = y + QuadPattern[i].Y;
@@ -3140,11 +3140,11 @@ bool ALandscapeProxy::ExportToRawMesh(int32 InExportLOD, FMeshDescription& OutRa
 		{
 			for (int32 x = 0; x < ComponentSizeQuadsLOD; x++)
 			{
-				FVector Positions[ARRAY_COUNT(QuadPattern)];
+				FVector Positions[UE_ARRAY_COUNT(QuadPattern)];
 				bool bProcess = bIgnoreBounds;
 
 				// Fill positions
-				for (int32 i = 0; i < ARRAY_COUNT(QuadPattern); i++)
+				for (int32 i = 0; i < UE_ARRAY_COUNT(QuadPattern); i++)
 				{
 					int32 VertexX = x + QuadPattern[i].X;
 					int32 VertexY = y + QuadPattern[i].Y;
@@ -3161,11 +3161,11 @@ bool ALandscapeProxy::ExportToRawMesh(int32 InExportLOD, FMeshDescription& OutRa
 				{
 					//Fill the vertexID we need
 					TArray<FVertexID> VertexIDs;
-					VertexIDs.Reserve(ARRAY_COUNT(QuadPattern));
+					VertexIDs.Reserve(UE_ARRAY_COUNT(QuadPattern));
 					TArray<FVertexInstanceID> VertexInstanceIDs;
-					VertexInstanceIDs.Reserve(ARRAY_COUNT(QuadPattern));
+					VertexInstanceIDs.Reserve(UE_ARRAY_COUNT(QuadPattern));
 					// Fill positions
-					for (int32 i = 0; i < ARRAY_COUNT(QuadPattern); i++)
+					for (int32 i = 0; i < UE_ARRAY_COUNT(QuadPattern); i++)
 					{
 						int32 DuplicateLowestIndex = FindPreviousIndex(CurrentIndex);
 						FVertexID VertexID;
@@ -3205,7 +3205,7 @@ bool ALandscapeProxy::ExportToRawMesh(int32 InExportLOD, FMeshDescription& OutRa
 							VertexInstanceIDs.Add(OutRawMesh.CreateVertexInstance(VertexIDs[5]));
 
 							// Fill other vertex data
-							for (int32 i = 0; i < ARRAY_COUNT(QuadPattern); i++)
+							for (int32 i = 0; i < UE_ARRAY_COUNT(QuadPattern); i++)
 							{
 								int32 VertexX = x + QuadPattern[i].X;
 								int32 VertexY = y + QuadPattern[i].Y;
@@ -3247,7 +3247,7 @@ bool ALandscapeProxy::ExportToRawMesh(int32 InExportLOD, FMeshDescription& OutRa
 				}
 				else
 				{
-					CurrentIndex += ARRAY_COUNT(QuadPattern);
+					CurrentIndex += UE_ARRAY_COUNT(QuadPattern);
 				}
 			}
 		}
diff --git a/Engine/Source/Runtime/Landscape/Private/LandscapeRender.cpp b/Engine/Source/Runtime/Landscape/Private/LandscapeRender.cpp
--- a/Engine/Source/Runtime/Landscape/Private/LandscapeRender.cpp
+++ b/Engine/Source/Runtime/Landscape/Private/LandscapeRender.cpp
@@ -4273,7 +4273,7 @@ public:
 		if (bIsLayerThumbnail || bDisableTessellation)
 		{
 			FSHA1 Hash;
-			Hash.Update(OutId.BasePropertyOverridesHash.Hash, ARRAY_COUNT(OutId.BasePropertyOverridesHash.Hash));
+			Hash.Update(OutId.BasePropertyOverridesHash.Hash, UE_ARRAY_COUNT(OutId.BasePropertyOverridesHash.Hash));
 
 			const FString HashString = TEXT("bOverride_TessellationMode");
 			Hash.UpdateWithString(*HashString, HashString.Len());
diff --git a/Engine/Source/Runtime/Landscape/Public/LandscapeRender.h b/Engine/Source/Runtime/Landscape/Public/LandscapeRender.h
--- a/Engine/Source/Runtime/Landscape/Public/LandscapeRender.h
+++ b/Engine/Source/Runtime/Landscape/Public/LandscapeRender.h
@@ -1012,7 +1012,7 @@ public:
 			FName(TEXT("Color15"))
 		};
 
-		for (int32 i = 0; i < ARRAY_COUNT(ColorNames) && i < LayerColors.Num(); i++)
+		for (int32 i = 0; i < UE_ARRAY_COUNT(ColorNames) && i < LayerColors.Num(); i++)
 		{
 			if (ParameterInfo.Name == ColorNames[i])
 			{
diff --git a/Engine/Source/Runtime/MovieScene/Private/Channels/MovieSceneFloatChannel.cpp b/Engine/Source/Runtime/MovieScene/Private/Channels/MovieSceneFloatChannel.cpp
--- a/Engine/Source/Runtime/MovieScene/Private/Channels/MovieSceneFloatChannel.cpp
+++ b/Engine/Source/Runtime/MovieScene/Private/Channels/MovieSceneFloatChannel.cpp
@@ -841,9 +841,9 @@ void FMovieSceneFloatChannel::RefineCurvePoints(FFrameRate InTickResolution, dou
 		{
 			bool bSegmentIsLinear = true;
 
-			TTuple<double, double> Evaluated[ARRAY_COUNT(InterpTimes)];
+			TTuple<double, double> Evaluated[UE_ARRAY_COUNT(InterpTimes)];
 
-			for (int32 InterpIndex = 0; InterpIndex < ARRAY_COUNT(InterpTimes); ++InterpIndex)
+			for (int32 InterpIndex = 0; InterpIndex < UE_ARRAY_COUNT(InterpTimes); ++InterpIndex)
 			{
 				double& EvalTime  = Evaluated[InterpIndex].Get<0>();
 
@@ -864,7 +864,7 @@ void FMovieSceneFloatChannel::RefineCurvePoints(FFrameRate InTickResolution, dou
 			if (!bSegmentIsLinear)
 			{
 				// Add the point
-				InOutPoints.Insert(Evaluated, ARRAY_COUNT(Evaluated), Index+1);
+				InOutPoints.Insert(Evaluated, UE_ARRAY_COUNT(Evaluated), Index+1);
 				--Index;
 			}
 		}
diff --git a/Engine/Source/Runtime/MovieScene/Private/Tests/MovieSceneCompilerTests.cpp b/Engine/Source/Runtime/MovieScene/Private/Tests/MovieSceneCompilerTests.cpp
--- a/Engine/Source/Runtime/MovieScene/Private/Tests/MovieSceneCompilerTests.cpp
+++ b/Engine/Source/Runtime/MovieScene/Private/Tests/MovieSceneCompilerTests.cpp
@@ -260,7 +260,7 @@ bool FMovieSceneCompilerEmptySpaceOnTheFlyTest::RunTest(const FString& Parameter
 	{
 		Tests.Emplace();
 		Tests.Last().CompileRanges   = { TRange<FFrameNumber>(0, 70) };
-		Tests.Last().ExpectedResults = TArray<FResult>(&ExpectedResults[0], ARRAY_COUNT(ExpectedResults));
+		Tests.Last().ExpectedResults = TArray<FResult>(&ExpectedResults[0], UE_ARRAY_COUNT(ExpectedResults));
 	}
 
 	UTestMovieSceneSequence* Sequence = NewObject<UTestMovieSceneSequence>(GetTransientPackage());
diff --git a/Engine/Source/Runtime/MovieSceneTracks/Private/Evaluation/MovieScene3DTransformTemplate.cpp b/Engine/Source/Runtime/MovieSceneTracks/Private/Evaluation/MovieScene3DTransformTemplate.cpp
--- a/Engine/Source/Runtime/MovieSceneTracks/Private/Evaluation/MovieScene3DTransformTemplate.cpp
+++ b/Engine/Source/Runtime/MovieSceneTracks/Private/Evaluation/MovieScene3DTransformTemplate.cpp
@@ -184,7 +184,7 @@ MovieScene::TMultiChannelValue<float, 9> FMovieSceneComponentTransformSectionTem
 		Components[4] = Rotation.Y;
 		Components[5] = Rotation.Z;
 
-		for (int32 Index = 0; Index < ARRAY_COUNT(Components); ++Index)
+		for (int32 Index = 0; Index < UE_ARRAY_COUNT(Components); ++Index)
 		{
 			if (TransformValue.IsSet(Index))
 			{
diff --git a/Engine/Source/Runtime/MovieSceneTracks/Private/MovieSceneTracksModule.cpp b/Engine/Source/Runtime/MovieSceneTracks/Private/MovieSceneTracksModule.cpp
--- a/Engine/Source/Runtime/MovieSceneTracks/Private/MovieSceneTracksModule.cpp
+++ b/Engine/Source/Runtime/MovieSceneTracks/Private/MovieSceneTracksModule.cpp
@@ -37,7 +37,7 @@ FName IMovieSceneTracksModule::GetEvaluationGroupName(EBuiltInEvaluationGroup::T
 		"SpawnObjects",
 		"PostEvaluation",
 	};
-	check(InEvalGroup >= 0 && InEvalGroup < ARRAY_COUNT(Names));
+	check(InEvalGroup >= 0 && InEvalGroup < UE_ARRAY_COUNT(Names));
 	return Names[InEvalGroup];
 }
 
diff --git a/Engine/Source/Runtime/MovieSceneTracks/Private/Sections/MovieSceneVectorSection.cpp b/Engine/Source/Runtime/MovieSceneTracks/Private/Sections/MovieSceneVectorSection.cpp
--- a/Engine/Source/Runtime/MovieSceneTracks/Private/Sections/MovieSceneVectorSection.cpp
+++ b/Engine/Source/Runtime/MovieSceneTracks/Private/Sections/MovieSceneVectorSection.cpp
@@ -233,7 +233,7 @@ void UMovieSceneVectorSection::RecreateChannelProxy()
 {
 	FMovieSceneChannelProxyData Channels;
 
-	check(ChannelsUsed <= ARRAY_COUNT(Curves));
+	check(ChannelsUsed <= UE_ARRAY_COUNT(Curves));
 
 #if WITH_EDITOR
 
diff --git a/Engine/Source/Runtime/NavigationSystem/Private/NavMesh/NavMeshRenderingComponent.cpp b/Engine/Source/Runtime/NavigationSystem/Private/NavMesh/NavMeshRenderingComponent.cpp
--- a/Engine/Source/Runtime/NavigationSystem/Private/NavMesh/NavMeshRenderingComponent.cpp
+++ b/Engine/Source/Runtime/NavigationSystem/Private/NavMesh/NavMeshRenderingComponent.cpp
@@ -280,7 +280,7 @@ void FNavMeshSceneProxyData::Serialize(FArchive& Ar)
 	}
 
 	TArray<FDebugRenderSceneProxy::FDebugLine>* LineArraysToSerialize[] = { &ThickLineItems, &TileEdgeLines, &NavMeshEdgeLines, &NavLinkLines, &ClusterLinkLines };
-	for (int32 ArrIdx = 0; ArrIdx < ARRAY_COUNT(LineArraysToSerialize); ArrIdx++)
+	for (int32 ArrIdx = 0; ArrIdx < UE_ARRAY_COUNT(LineArraysToSerialize); ArrIdx++)
 	{
 		int32 NumItems = LineArraysToSerialize[ArrIdx]->Num();
 		Ar << NumItems;
diff --git a/Engine/Source/Runtime/NavigationSystem/Private/NavMesh/RecastNavMeshGenerator.cpp b/Engine/Source/Runtime/NavigationSystem/Private/NavMesh/RecastNavMeshGenerator.cpp
--- a/Engine/Source/Runtime/NavigationSystem/Private/NavMesh/RecastNavMeshGenerator.cpp
+++ b/Engine/Source/Runtime/NavigationSystem/Private/NavMesh/RecastNavMeshGenerator.cpp
@@ -772,7 +772,7 @@ void ExportRigidBodyBoxElements(const FKAggregateGeom& AggGeom, TNavStatArray<fl
 			ElemTM.TransformPosition(FVector( Extent.X,  Extent.Y, -Extent.Z))
 		};
 
-		for (int32 iv = 0; iv < ARRAY_COUNT(UnrealVerts); iv++)
+		for (int32 iv = 0; iv < UE_ARRAY_COUNT(UnrealVerts); iv++)
 		{
 			UnrealBounds += UnrealVerts[iv];
 
diff --git a/Engine/Source/Runtime/NavigationSystem/Private/NavigationSystem.cpp b/Engine/Source/Runtime/NavigationSystem/Private/NavigationSystem.cpp
--- a/Engine/Source/Runtime/NavigationSystem/Private/NavigationSystem.cpp
+++ b/Engine/Source/Runtime/NavigationSystem/Private/NavigationSystem.cpp
@@ -2267,7 +2267,7 @@ void UNavigationSystemV1::DescribeFilterFlags(const TArray<FString>& FlagsDesc)
 	UStruct* Structs[] = { StructProp1->Struct, StructProp2->Struct };
 	const FString CustomNameMeta = TEXT("DisplayName");
 
-	for (int32 StructIndex = 0; StructIndex < ARRAY_COUNT(Structs); StructIndex++)
+	for (int32 StructIndex = 0; StructIndex < UE_ARRAY_COUNT(Structs); StructIndex++)
 	{
 		for (int32 FlagIndex = 0; FlagIndex < MaxFlags; FlagIndex++)
 		{
diff --git a/Engine/Source/Runtime/Navmesh/Private/DebugUtils/RecastDump.cpp b/Engine/Source/Runtime/Navmesh/Private/DebugUtils/RecastDump.cpp
--- a/Engine/Source/Runtime/Navmesh/Private/DebugUtils/RecastDump.cpp
+++ b/Engine/Source/Runtime/Navmesh/Private/DebugUtils/RecastDump.cpp
@@ -35,7 +35,7 @@ static void ioprintf(duFileIO* io, const char* format, ...)
 	char line[256];
 	va_list ap;
 	va_start(ap, format);
-	const int n = FCStringAnsi::GetVarArgs(line, ARRAY_COUNT(line), format, ap);
+	const int n = FCStringAnsi::GetVarArgs(line, UE_ARRAY_COUNT(line), format, ap);
 	va_end(ap);
 	if (n > 0)
 		io->write(line, sizeof(char)*n);
diff --git a/Engine/Source/Runtime/NetworkFile/Private/NetworkPlatformFile.cpp b/Engine/Source/Runtime/NetworkFile/Private/NetworkPlatformFile.cpp
--- a/Engine/Source/Runtime/NetworkFile/Private/NetworkPlatformFile.cpp
+++ b/Engine/Source/Runtime/NetworkFile/Private/NetworkPlatformFile.cpp
@@ -961,7 +961,7 @@ public:
 				while (RemainingData)
 				{
 					// read next chunk from archive
-					uint32 LocalSize = FPlatformMath::Min<uint32>(ARRAY_COUNT(Buffer), RemainingData);
+					uint32 LocalSize = FPlatformMath::Min<uint32>(UE_ARRAY_COUNT(Buffer), RemainingData);
 					FileArchive->Serialize(Buffer, LocalSize);
 					// write it out
 					if (!FileHandle->Write(Buffer, LocalSize))
diff --git a/Engine/Source/Runtime/Online/BuildPatchServices/Private/BuildPatchServicesModule.cpp b/Engine/Source/Runtime/Online/BuildPatchServices/Private/BuildPatchServicesModule.cpp
--- a/Engine/Source/Runtime/Online/BuildPatchServices/Private/BuildPatchServicesModule.cpp
+++ b/Engine/Source/Runtime/Online/BuildPatchServices/Private/BuildPatchServicesModule.cpp
@@ -44,8 +44,8 @@ void FBuildPatchServicesModule::StartupModule()
 #if UE_BUILD_DEBUG
 	TSet<FString> NoDupes;
 	bool bWasDupe = false;
-	check(ARRAY_COUNT(InstallErrorPrefixes::ErrorTypeStrings) == (uint64)EBuildPatchInstallError::NumInstallErrors);
-	for (int32 Idx = 0; Idx < ARRAY_COUNT(InstallErrorPrefixes::ErrorTypeStrings); ++Idx)
+	check(UE_ARRAY_COUNT(InstallErrorPrefixes::ErrorTypeStrings) == (uint64)EBuildPatchInstallError::NumInstallErrors);
+	for (int32 Idx = 0; Idx < UE_ARRAY_COUNT(InstallErrorPrefixes::ErrorTypeStrings); ++Idx)
 	{
 		NoDupes.Add(FString(InstallErrorPrefixes::ErrorTypeStrings[Idx]), &bWasDupe);
 		check(bWasDupe == false);
diff --git a/Engine/Source/Runtime/Online/BuildPatchServices/Private/Installer/CloudChunkSource.h b/Engine/Source/Runtime/Online/BuildPatchServices/Private/Installer/CloudChunkSource.h
--- a/Engine/Source/Runtime/Online/BuildPatchServices/Private/Installer/CloudChunkSource.h
+++ b/Engine/Source/Runtime/Online/BuildPatchServices/Private/Installer/CloudChunkSource.h
@@ -75,12 +75,12 @@ namespace BuildPatchServices
 			, TcpZeroWindowMinimumSeconds(20.0f)
 		{
 			const float RetryFloats[] = {0.5f, 1.0f, 1.0f, 3.0f, 3.0f, 10.0f, 10.0f, 20.0f, 20.0f, 30.0f};
-			RetryDelayTimes.Empty(ARRAY_COUNT(RetryFloats));
-			RetryDelayTimes.Append(RetryFloats, ARRAY_COUNT(RetryFloats));
+			RetryDelayTimes.Empty(UE_ARRAY_COUNT(RetryFloats));
+			RetryDelayTimes.Append(RetryFloats, UE_ARRAY_COUNT(RetryFloats));
 			const float HealthFloats[] = {0.0f, 0.0f, 0.9f, 0.99f, 1.0f};
-			check((int32)EBuildPatchDownloadHealth::NUM_Values == ARRAY_COUNT(HealthFloats));
-			HealthPercentages.Empty(ARRAY_COUNT(HealthFloats));
-			HealthPercentages.Append(HealthFloats, ARRAY_COUNT(HealthFloats));
+			check((int32)EBuildPatchDownloadHealth::NUM_Values == UE_ARRAY_COUNT(HealthFloats));
+			HealthPercentages.Empty(UE_ARRAY_COUNT(HealthFloats));
+			HealthPercentages.Append(HealthFloats, UE_ARRAY_COUNT(HealthFloats));
 		}
 	};
 
diff --git a/Engine/Source/Runtime/Online/HTTP/Private/Curl/CurlHttp.cpp b/Engine/Source/Runtime/Online/HTTP/Private/Curl/CurlHttp.cpp
--- a/Engine/Source/Runtime/Online/HTTP/Private/Curl/CurlHttp.cpp
+++ b/Engine/Source/Runtime/Online/HTTP/Private/Curl/CurlHttp.cpp
@@ -247,7 +247,7 @@ FString FCurlHttpRequest::CombineHeaderKeyValue(const FString& HeaderKey, const
 {
 	FString Combined;
 	const TCHAR Separator[] = TEXT(": ");
-	constexpr const int32 SeparatorLength = ARRAY_COUNT(Separator) - 1;
+	constexpr const int32 SeparatorLength = UE_ARRAY_COUNT(Separator) - 1;
 	Combined.Reserve(HeaderKey.Len() + SeparatorLength + HeaderValue.Len());
 	Combined.Append(HeaderKey);
 	Combined.AppendChars(Separator, SeparatorLength);
diff --git a/Engine/Source/Runtime/Online/HTTP/Private/Curl/CurlHttpManager.cpp b/Engine/Source/Runtime/Online/HTTP/Private/Curl/CurlHttpManager.cpp
--- a/Engine/Source/Runtime/Online/HTTP/Private/Curl/CurlHttpManager.cpp
+++ b/Engine/Source/Runtime/Online/HTTP/Private/Curl/CurlHttpManager.cpp
@@ -256,7 +256,7 @@ void FCurlHttpManager::InitCurl()
 	}
 
 	TCHAR Home[256] = TEXT("");
-	if (FParse::Value(FCommandLine::Get(), TEXT("MULTIHOMEHTTP="), Home, ARRAY_COUNT(Home)))
+	if (FParse::Value(FCommandLine::Get(), TEXT("MULTIHOMEHTTP="), Home, UE_ARRAY_COUNT(Home)))
 	{
 		ISocketSubsystem* SocketSubsystem = ISocketSubsystem::Get(PLATFORM_SOCKETSUBSYSTEM);
 		if (SocketSubsystem)
diff --git a/Engine/Source/Runtime/Online/HTTP/Private/GenericPlatform/GenericPlatformHttp.cpp b/Engine/Source/Runtime/Online/HTTP/Private/GenericPlatform/GenericPlatformHttp.cpp
--- a/Engine/Source/Runtime/Online/HTTP/Private/GenericPlatform/GenericPlatformHttp.cpp
+++ b/Engine/Source/Runtime/Online/HTTP/Private/GenericPlatform/GenericPlatformHttp.cpp
@@ -67,10 +67,10 @@ static bool IsAllowedChar(UTF8CHAR LookupChar)
 
 	if (!bTableFilled)
 	{
-		for (int32 Idx = 0; Idx < ARRAY_COUNT(AllowedChars) - 1; ++Idx)	// -1 to avoid trailing 0
+		for (int32 Idx = 0; Idx < UE_ARRAY_COUNT(AllowedChars) - 1; ++Idx)	// -1 to avoid trailing 0
 		{
 			uint8 AllowedCharIdx = static_cast<uint8>(AllowedChars[Idx]);
-			check(AllowedCharIdx < ARRAY_COUNT(AllowedTable));
+			check(AllowedCharIdx < UE_ARRAY_COUNT(AllowedTable));
 			AllowedTable[AllowedCharIdx] = true;
 		}
 
@@ -133,7 +133,7 @@ FString FGenericPlatformHttp::UrlDecode(const FString &EncodedString)
 
 					ANSICHAR Buffer[8] = { 0 };
 					ANSICHAR* BufferPtr = Buffer;
-					const int32 Len = ARRAY_COUNT(Buffer);
+					const int32 Len = UE_ARRAY_COUNT(Buffer);
 					const int32 WrittenChars = FTCHARToUTF8_Convert::Utf8FromCodepoint(Value, BufferPtr, Len);
 
 					Data.Append(Buffer, WrittenChars);
diff --git a/Engine/Source/Runtime/Online/HTTP/Private/GenericPlatform/HttpRequestPayload.cpp b/Engine/Source/Runtime/Online/HTTP/Private/GenericPlatform/HttpRequestPayload.cpp
--- a/Engine/Source/Runtime/Online/HTTP/Private/GenericPlatform/HttpRequestPayload.cpp
+++ b/Engine/Source/Runtime/Online/HTTP/Private/GenericPlatform/HttpRequestPayload.cpp
@@ -13,10 +13,10 @@ bool FGenericPlatformHttp::IsURLEncoded(const TArray<uint8>& Payload)
 
 	if (!bTableFilled)
 	{
-		for (int32 Idx = 0; Idx < ARRAY_COUNT(AllowedChars) - 1; ++Idx)	// -1 to avoid trailing 0
+		for (int32 Idx = 0; Idx < UE_ARRAY_COUNT(AllowedChars) - 1; ++Idx)	// -1 to avoid trailing 0
 		{
 			uint8 AllowedCharIdx = static_cast<uint8>(AllowedChars[Idx]);
-			check(AllowedCharIdx < ARRAY_COUNT(AllowedTable));
+			check(AllowedCharIdx < UE_ARRAY_COUNT(AllowedTable));
 			AllowedTable[AllowedCharIdx] = true;
 		}
 
diff --git a/Engine/Source/Runtime/Online/HTTPServer/Private/HttpConnectionRequestReadContext.cpp b/Engine/Source/Runtime/Online/HTTPServer/Private/HttpConnectionRequestReadContext.cpp
--- a/Engine/Source/Runtime/Online/HTTPServer/Private/HttpConnectionRequestReadContext.cpp
+++ b/Engine/Source/Runtime/Online/HTTPServer/Private/HttpConnectionRequestReadContext.cpp
@@ -346,7 +346,7 @@ FString FHttpConnectionRequestReadContext::UrlDecode(const FString &EncodedStrin
 
 					ANSICHAR Buffer[8] = { 0 };
 					ANSICHAR* BufferPtr = Buffer;
-					const int32 Len = ARRAY_COUNT(Buffer);
+					const int32 Len = UE_ARRAY_COUNT(Buffer);
 					const int32 WrittenChars = FTCHARToUTF8_Convert::Utf8FromCodepoint(Value, BufferPtr, Len);
 
 					Data.Append(Buffer, WrittenChars);
diff --git a/Engine/Source/Runtime/Online/SSL/Private/SslError.cpp b/Engine/Source/Runtime/Online/SSL/Private/SslError.cpp
--- a/Engine/Source/Runtime/Online/SSL/Private/SslError.cpp
+++ b/Engine/Source/Runtime/Online/SSL/Private/SslError.cpp
@@ -22,8 +22,8 @@ FString GetSslErrorString()
 	if (Error != SSL_ERROR_NONE)
 	{
 		char AnsiErrorBuffer[256];
-		ERR_error_string_n(Error, AnsiErrorBuffer, ARRAY_COUNT(AnsiErrorBuffer) - 1);
-		AnsiErrorBuffer[ARRAY_COUNT(AnsiErrorBuffer) - 1] = '\0';
+		ERR_error_string_n(Error, AnsiErrorBuffer, UE_ARRAY_COUNT(AnsiErrorBuffer) - 1);
+		AnsiErrorBuffer[UE_ARRAY_COUNT(AnsiErrorBuffer) - 1] = '\0';
 
 		SslErrorString = ANSI_TO_TCHAR(AnsiErrorBuffer);
 	}
diff --git a/Engine/Source/Runtime/Online/SSL/Private/Windows/WindowsPlatformSslCertificateManager.cpp b/Engine/Source/Runtime/Online/SSL/Private/Windows/WindowsPlatformSslCertificateManager.cpp
--- a/Engine/Source/Runtime/Online/SSL/Private/Windows/WindowsPlatformSslCertificateManager.cpp
+++ b/Engine/Source/Runtime/Online/SSL/Private/Windows/WindowsPlatformSslCertificateManager.cpp
@@ -44,14 +44,14 @@ void FWindowsPlatformSslCertificateManager::BuildRootCertificateArray()
 				else
 				{
 					TCHAR Name[128];
-					CertGetNameString(CertContext, CERT_NAME_SIMPLE_DISPLAY_TYPE, 0, nullptr, Name, ARRAY_COUNT(Name));
+					CertGetNameString(CertContext, CERT_NAME_SIMPLE_DISPLAY_TYPE, 0, nullptr, Name, UE_ARRAY_COUNT(Name));
 					UE_LOG(LogSsl, Log, TEXT("Unable to convert certificate: name:%s error:%s"), Name, *GetSslErrorString());
 				}
 			}
 			else
 			{
 				TCHAR Name[128];
-				CertGetNameString(CertContext, CERT_NAME_SIMPLE_DISPLAY_TYPE, 0, nullptr, Name, ARRAY_COUNT(Name));
+				CertGetNameString(CertContext, CERT_NAME_SIMPLE_DISPLAY_TYPE, 0, nullptr, Name, UE_ARRAY_COUNT(Name));
 				UE_LOG(LogSsl, Log, TEXT("Unhandled certificate encoding: name:%s encodingType:0x%08x"), Name, CertContext->dwCertEncodingType);
 			}
 		}
diff --git a/Engine/Source/Runtime/OpenGLDrv/Private/Android/AndroidEGL.cpp b/Engine/Source/Runtime/OpenGLDrv/Private/Android/AndroidEGL.cpp
--- a/Engine/Source/Runtime/OpenGLDrv/Private/Android/AndroidEGL.cpp
+++ b/Engine/Source/Runtime/OpenGLDrv/Private/Android/AndroidEGL.cpp
@@ -70,7 +70,7 @@ void VerifyEGLResult(EGLint ErrorCode, const TCHAR* Msg1, const TCHAR* Msg2, con
 			TEXT("UNKNOWN EGL ERROR")
 		};
 
-		uint32 ErrorIndex = FMath::Min<uint32>(ErrorCode - EGL_SUCCESS, ARRAY_COUNT(EGLErrorStrings) - 1);
+		uint32 ErrorIndex = FMath::Min<uint32>(ErrorCode - EGL_SUCCESS, UE_ARRAY_COUNT(EGLErrorStrings) - 1);
 		UE_LOG(LogRHI, Warning, TEXT("%s(%u): %s%s failed with error %s (0x%x)"),
 			Filename, Line, Msg1, Msg2, EGLErrorStrings[ErrorIndex], ErrorCode);
 		check(0);
diff --git a/Engine/Source/Runtime/OpenGLDrv/Private/OpenGLShaders.cpp b/Engine/Source/Runtime/OpenGLDrv/Private/OpenGLShaders.cpp
--- a/Engine/Source/Runtime/OpenGLDrv/Private/OpenGLShaders.cpp
+++ b/Engine/Source/Runtime/OpenGLDrv/Private/OpenGLShaders.cpp
@@ -5158,7 +5158,7 @@ void FOpenGLProgramBinaryCache::CompilePendingShaders(const FOpenGLLinkedProgram
 {
 	if (CachePtr)
 	{
-		for (int32 StageIdx = 0; StageIdx < ARRAY_COUNT(Config.Shaders); ++StageIdx)
+		for (int32 StageIdx = 0; StageIdx < UE_ARRAY_COUNT(Config.Shaders); ++StageIdx)
 		{
 			GLuint ShaderResource = Config.Shaders[StageIdx].Resource;
 			FPendingShaderCode* PendingShaderCodePtr = CachePtr->ShadersPendingCompilation.Find(ShaderResource);
diff --git a/Engine/Source/Runtime/OpenGLDrv/Private/OpenGLUtil.cpp b/Engine/Source/Runtime/OpenGLDrv/Private/OpenGLUtil.cpp
--- a/Engine/Source/Runtime/OpenGLDrv/Private/OpenGLUtil.cpp
+++ b/Engine/Source/Runtime/OpenGLDrv/Private/OpenGLUtil.cpp
@@ -25,7 +25,7 @@ void VerifyOpenGLResult(GLenum ErrorCode, const TCHAR* Msg1, const TCHAR* Msg2,
 			TEXT("UNKNOWN ERROR")
 		};
 
-		uint32 ErrorIndex = FMath::Min<uint32>(ErrorCode - GL_INVALID_ENUM, ARRAY_COUNT(ErrorStrings) - 1);
+		uint32 ErrorIndex = FMath::Min<uint32>(ErrorCode - GL_INVALID_ENUM, UE_ARRAY_COUNT(ErrorStrings) - 1);
 		UE_LOG(LogRHI,Fatal,TEXT("%s(%u): %s%s failed with error %s (0x%x)"),
 			Filename,Line,Msg1,Msg2,ErrorStrings[ErrorIndex],ErrorCode);
 	}
diff --git a/Engine/Source/Runtime/OpenGLDrv/Public/OpenGLDrv.h b/Engine/Source/Runtime/OpenGLDrv/Public/OpenGLDrv.h
--- a/Engine/Source/Runtime/OpenGLDrv/Public/OpenGLDrv.h
+++ b/Engine/Source/Runtime/OpenGLDrv/Public/OpenGLDrv.h
@@ -1170,7 +1170,7 @@ private:
 	{
 		if (FOpenGL::NeedsVertexAttribRemapTable())
 		{
-			check(VertexAttributeIndex < ARRAY_COUNT(PendingState.BoundShaderState->GetVertexShader()->Bindings.VertexAttributeRemap));
+			check(VertexAttributeIndex < UE_ARRAY_COUNT(PendingState.BoundShaderState->GetVertexShader()->Bindings.VertexAttributeRemap));
 			VertexAttributeIndex = PendingState.BoundShaderState->GetVertexShader()->Bindings.VertexAttributeRemap[VertexAttributeIndex];
 		}
 		check(VertexAttributeIndex < NUM_OPENGL_VERTEX_STREAMS); // check that this attribute has remaped correctly.
@@ -1181,7 +1181,7 @@ private:
 	{
 		if (FOpenGL::NeedsVertexAttribRemapTable())
 		{
-			check(VertexAttributeIndex < ARRAY_COUNT(Bindings.VertexAttributeRemap));
+			check(VertexAttributeIndex < UE_ARRAY_COUNT(Bindings.VertexAttributeRemap));
 			VertexAttributeIndex = Bindings.VertexAttributeRemap[VertexAttributeIndex];
 		}
 		check(VertexAttributeIndex < NUM_OPENGL_VERTEX_STREAMS); // check that this attribute has remaped correctly.
diff --git a/Engine/Source/Runtime/OpenGLDrv/Public/OpenGLShaderResources.h b/Engine/Source/Runtime/OpenGLDrv/Public/OpenGLShaderResources.h
--- a/Engine/Source/Runtime/OpenGLDrv/Public/OpenGLShaderResources.h
+++ b/Engine/Source/Runtime/OpenGLDrv/Public/OpenGLShaderResources.h
@@ -208,7 +208,7 @@ inline FArchive& operator<<(FArchive& Ar, FOpenGLShaderBindings& Bindings)
 	Ar << Bindings.NumUniformBuffers;
 	Ar << Bindings.NumUAVs;
 	Ar << Bindings.bFlattenUB;
-	for (uint32 i = 0; i < ARRAY_COUNT(Bindings.VertexAttributeRemap); i++)
+	for (uint32 i = 0; i < UE_ARRAY_COUNT(Bindings.VertexAttributeRemap); i++)
 	{
 		Ar << Bindings.VertexAttributeRemap[i];
 	}
diff --git a/Engine/Source/Runtime/PakFile/Public/IPlatformFilePak.h b/Engine/Source/Runtime/PakFile/Public/IPlatformFilePak.h
--- a/Engine/Source/Runtime/PakFile/Public/IPlatformFilePak.h
+++ b/Engine/Source/Runtime/PakFile/Public/IPlatformFilePak.h
@@ -2477,7 +2477,7 @@ struct FPakSignatureFile
 	{
 		ChunkHashes = InChunkHashes;
 		DecryptedHash = ComputeCurrentMasterHash();
-		FRSA::EncryptPrivate(TArrayView<const uint8>(DecryptedHash.Hash, ARRAY_COUNT(FSHAHash::Hash)), EncryptedHash, InKey);
+		FRSA::EncryptPrivate(TArrayView<const uint8>(DecryptedHash.Hash, UE_ARRAY_COUNT(FSHAHash::Hash)), EncryptedHash, InKey);
 	}
 
 	/**
@@ -2514,7 +2514,7 @@ struct FPakSignatureFile
 		{
 			TArray<uint8> Decrypted;
 			int32 BytesDecrypted = FRSA::DecryptPublic(EncryptedHash, Decrypted, InKey);
-			if (BytesDecrypted == ARRAY_COUNT(FSHAHash::Hash))
+			if (BytesDecrypted == UE_ARRAY_COUNT(FSHAHash::Hash))
 			{
 				FSHAHash CurrentHash = ComputeCurrentMasterHash();
 				if (FMemory::Memcmp(Decrypted.GetData(), CurrentHash.Hash, Decrypted.Num()) == 0)
diff --git a/Engine/Source/Runtime/Portal/LauncherCheck/Private/LauncherCheckModule.cpp b/Engine/Source/Runtime/Portal/LauncherCheck/Private/LauncherCheckModule.cpp
--- a/Engine/Source/Runtime/Portal/LauncherCheck/Private/LauncherCheckModule.cpp
+++ b/Engine/Source/Runtime/Portal/LauncherCheck/Private/LauncherCheckModule.cpp
@@ -104,7 +104,7 @@ private:
 		FString EncodedExePath;
 
 		static const TCHAR* Delims[] = { TEXT("/") };
-		static const int32 NumDelims = ARRAY_COUNT(Delims);
+		static const int32 NumDelims = UE_ARRAY_COUNT(Delims);
 
 		// Get the path to the executable, and encode it ('cos spaces, colons, etc break things)
 		FString ExePath = FPlatformProcess::BaseDir();
diff --git a/Engine/Source/Runtime/RHI/Private/RHI.cpp b/Engine/Source/Runtime/RHI/Private/RHI.cpp
--- a/Engine/Source/Runtime/RHI/Private/RHI.cpp
+++ b/Engine/Source/Runtime/RHI/Private/RHI.cpp
@@ -582,11 +582,11 @@ FName FeatureLevelNames[] =
 	FName(TEXT("SM5")),
 };
 
-static_assert(ARRAY_COUNT(FeatureLevelNames) == ERHIFeatureLevel::Num, "Missing entry from feature level names.");
+static_assert(UE_ARRAY_COUNT(FeatureLevelNames) == ERHIFeatureLevel::Num, "Missing entry from feature level names.");
 
 RHI_API bool GetFeatureLevelFromName(FName Name, ERHIFeatureLevel::Type& OutFeatureLevel)
 {
-	for (int32 NameIndex = 0; NameIndex < ARRAY_COUNT(FeatureLevelNames); NameIndex++)
+	for (int32 NameIndex = 0; NameIndex < UE_ARRAY_COUNT(FeatureLevelNames); NameIndex++)
 	{
 		if (FeatureLevelNames[NameIndex] == Name)
 		{
@@ -601,8 +601,8 @@ RHI_API bool GetFeatureLevelFromName(FName Name, ERHIFeatureLevel::Type& OutFeat
 
 RHI_API void GetFeatureLevelName(ERHIFeatureLevel::Type InFeatureLevel, FString& OutName)
 {
-	check(InFeatureLevel < ARRAY_COUNT(FeatureLevelNames));
-	if (InFeatureLevel < ARRAY_COUNT(FeatureLevelNames))
+	check(InFeatureLevel < UE_ARRAY_COUNT(FeatureLevelNames));
+	if (InFeatureLevel < UE_ARRAY_COUNT(FeatureLevelNames))
 	{
 		FeatureLevelNames[(int32)InFeatureLevel].ToString(OutName);
 	}
@@ -615,8 +615,8 @@ RHI_API void GetFeatureLevelName(ERHIFeatureLevel::Type InFeatureLevel, FString&
 static FName InvalidFeatureLevelName(TEXT("InvalidFeatureLevel"));
 RHI_API void GetFeatureLevelName(ERHIFeatureLevel::Type InFeatureLevel, FName& OutName)
 {
-	check(InFeatureLevel < ARRAY_COUNT(FeatureLevelNames));
-	if (InFeatureLevel < ARRAY_COUNT(FeatureLevelNames))
+	check(InFeatureLevel < UE_ARRAY_COUNT(FeatureLevelNames));
+	if (InFeatureLevel < UE_ARRAY_COUNT(FeatureLevelNames))
 	{
 		OutName = FeatureLevelNames[(int32)InFeatureLevel];
 	}
@@ -634,11 +634,11 @@ FName ShadingPathNames[] =
 	FName(TEXT("Mobile")),
 };
 
-static_assert(ARRAY_COUNT(ShadingPathNames) == ERHIShadingPath::Num, "Missing entry from shading path names.");
+static_assert(UE_ARRAY_COUNT(ShadingPathNames) == ERHIShadingPath::Num, "Missing entry from shading path names.");
 
 RHI_API bool GetShadingPathFromName(FName Name, ERHIShadingPath::Type& OutShadingPath)
 {
-	for (int32 NameIndex = 0; NameIndex < ARRAY_COUNT(ShadingPathNames); NameIndex++)
+	for (int32 NameIndex = 0; NameIndex < UE_ARRAY_COUNT(ShadingPathNames); NameIndex++)
 	{
 		if (ShadingPathNames[NameIndex] == Name)
 		{
@@ -653,8 +653,8 @@ RHI_API bool GetShadingPathFromName(FName Name, ERHIShadingPath::Type& OutShadin
 
 RHI_API void GetShadingPathName(ERHIShadingPath::Type InShadingPath, FString& OutName)
 {
-	check(InShadingPath < ARRAY_COUNT(ShadingPathNames));
-	if (InShadingPath < ARRAY_COUNT(ShadingPathNames))
+	check(InShadingPath < UE_ARRAY_COUNT(ShadingPathNames));
+	if (InShadingPath < UE_ARRAY_COUNT(ShadingPathNames))
 	{
 		ShadingPathNames[(int32)InShadingPath].ToString(OutName);
 	}
@@ -667,8 +667,8 @@ RHI_API void GetShadingPathName(ERHIShadingPath::Type InShadingPath, FString& Ou
 static FName InvalidShadingPathName(TEXT("InvalidShadingPath"));
 RHI_API void GetShadingPathName(ERHIShadingPath::Type InShadingPath, FName& OutName)
 {
-	check(InShadingPath < ARRAY_COUNT(ShadingPathNames));
-	if (InShadingPath < ARRAY_COUNT(ShadingPathNames))
+	check(InShadingPath < UE_ARRAY_COUNT(ShadingPathNames));
+	if (InShadingPath < UE_ARRAY_COUNT(ShadingPathNames))
 	{
 		OutName = ShadingPathNames[(int32)InShadingPath];
 	}
diff --git a/Engine/Source/Runtime/Renderer/Private/CapsuleShadowRendering.cpp b/Engine/Source/Runtime/Renderer/Private/CapsuleShadowRendering.cpp
--- a/Engine/Source/Runtime/Renderer/Private/CapsuleShadowRendering.cpp
+++ b/Engine/Source/Runtime/Renderer/Private/CapsuleShadowRendering.cpp
@@ -158,7 +158,7 @@ public:
 
 		FRHIUnorderedAccessView* OutUAVs[1];
 		OutUAVs[0] = ComputedLightDirectionDataBuffer.UAV;
-		RHICmdList.TransitionResources(EResourceTransitionAccess::EWritable, EResourceTransitionPipeline::EComputeToCompute, OutUAVs, ARRAY_COUNT(OutUAVs));
+		RHICmdList.TransitionResources(EResourceTransitionAccess::EWritable, EResourceTransitionPipeline::EComputeToCompute, OutUAVs, UE_ARRAY_COUNT(OutUAVs));
 
 		SetShaderValue(RHICmdList, ShaderRHI, NumLightDirectionData, NumLightDirectionDataValue);
 		SetSRVParameter(RHICmdList, ShaderRHI, LightDirectionData, LightDirectionDataSRV);
@@ -178,7 +178,7 @@ public:
 
 		FRHIUnorderedAccessView* OutUAVs[1];
 		OutUAVs[0] = ComputedLightDirectionDataBuffer.UAV;
-		RHICmdList.TransitionResources(EResourceTransitionAccess::EReadable, EResourceTransitionPipeline::EComputeToCompute, OutUAVs, ARRAY_COUNT(OutUAVs) );
+		RHICmdList.TransitionResources(EResourceTransitionAccess::EReadable, EResourceTransitionPipeline::EComputeToCompute, OutUAVs, UE_ARRAY_COUNT(OutUAVs) );
 	}
 
 	// FShader interface.
@@ -341,7 +341,7 @@ public:
 			OutUAVs[1] = TileIntersectionCountsBuffer->UAV;
 		}
 		
-		RHICmdList.TransitionResources(EResourceTransitionAccess::ERWBarrier, EResourceTransitionPipeline::EComputeToCompute, OutUAVs, TileIntersectionCountsBuffer ? ARRAY_COUNT(OutUAVs) : 1);
+		RHICmdList.TransitionResources(EResourceTransitionAccess::ERWBarrier, EResourceTransitionPipeline::EComputeToCompute, OutUAVs, TileIntersectionCountsBuffer ? UE_ARRAY_COUNT(OutUAVs) : 1);
 
 		if (ShadowingType == ShapeShadow_MovableSkylightTiledCulling)
 		{
@@ -448,7 +448,7 @@ public:
 			OutUAVs[1] = TileIntersectionCountsBuffer->UAV;
 		}
 		
-		RHICmdList.TransitionResources(EResourceTransitionAccess::EReadable, EResourceTransitionPipeline::EComputeToCompute, OutUAVs, TileIntersectionCountsBuffer ? ARRAY_COUNT(OutUAVs) : 1);
+		RHICmdList.TransitionResources(EResourceTransitionAccess::EReadable, EResourceTransitionPipeline::EComputeToCompute, OutUAVs, TileIntersectionCountsBuffer ? UE_ARRAY_COUNT(OutUAVs) : 1);
 	}
 
 	// FShader interface.
diff --git a/Engine/Source/Runtime/Renderer/Private/CompositionLighting/PostProcessDeferredDecals.cpp b/Engine/Source/Runtime/Renderer/Private/CompositionLighting/PostProcessDeferredDecals.cpp
--- a/Engine/Source/Runtime/Renderer/Private/CompositionLighting/PostProcessDeferredDecals.cpp
+++ b/Engine/Source/Runtime/Renderer/Private/CompositionLighting/PostProcessDeferredDecals.cpp
@@ -287,7 +287,7 @@ bool RenderPreStencil(FRenderingCompositePassContext& Context, const FMatrix& Co
 	Context.RHICmdList.SetStreamSource(0, GetUnitCubeVertexBuffer(), 0);
 
 	// Render decal mask
-	Context.RHICmdList.DrawIndexedPrimitive(GetUnitCubeIndexBuffer(), 0, 0, 8, 0, ARRAY_COUNT(GCubeIndices) / 3, 1);
+	Context.RHICmdList.DrawIndexedPrimitive(GetUnitCubeIndexBuffer(), 0, 0, 8, 0, UE_ARRAY_COUNT(GCubeIndices) / 3, 1);
 
 	return true;
 }
@@ -756,7 +756,7 @@ void FRCPassPostProcessDeferredDecals::Process(FRenderingCompositePassContext& C
 					FDecalRendering::SetShader(RHICmdList, GraphicsPSOInit, View, DecalData, CurrentStage, FrustumComponentToClip);
 					RHICmdList.SetStencilRef(StencilRef);
 
-					RHICmdList.DrawIndexedPrimitive(GetUnitCubeIndexBuffer(), 0, 0, 8, 0, ARRAY_COUNT(GCubeIndices) / 3, 1);
+					RHICmdList.DrawIndexedPrimitive(GetUnitCubeIndexBuffer(), 0, 0, 8, 0, UE_ARRAY_COUNT(GCubeIndices) / 3, 1);
 					RenderTargetManager.bGufferADirty |= (RenderTargetManager.TargetsToResolve[FDecalRenderTargetManager::GBufferAIndex] != nullptr);
 				}
 
diff --git a/Engine/Source/Runtime/Renderer/Private/ConvertToUniformMesh.cpp b/Engine/Source/Runtime/Renderer/Private/ConvertToUniformMesh.cpp
--- a/Engine/Source/Runtime/Renderer/Private/ConvertToUniformMesh.cpp
+++ b/Engine/Source/Runtime/Renderer/Private/ConvertToUniformMesh.cpp
@@ -363,7 +363,7 @@ public:
 
 		FRHIUnorderedAccessView* UniformMeshUAVs[1];
 		UniformMeshUAVs[0] = Scene->DistanceFieldSceneData.SurfelBuffers->Surfels.UAV;
-		RHICmdList.TransitionResources(EResourceTransitionAccess::ERWBarrier, EResourceTransitionPipeline::EComputeToCompute, UniformMeshUAVs, ARRAY_COUNT(UniformMeshUAVs));
+		RHICmdList.TransitionResources(EResourceTransitionAccess::ERWBarrier, EResourceTransitionPipeline::EComputeToCompute, UniformMeshUAVs, UE_ARRAY_COUNT(UniformMeshUAVs));
 
 		SurfelBufferParameters.Set(RHICmdList, ShaderRHI, *Scene->DistanceFieldSceneData.SurfelBuffers, *Scene->DistanceFieldSceneData.InstancedSurfelBuffers);
 		
@@ -381,7 +381,7 @@ public:
 		const FScene* Scene = (const FScene*)View.Family->Scene;
 		FRHIUnorderedAccessView* UniformMeshUAVs[1];
 		UniformMeshUAVs[0] = Scene->DistanceFieldSceneData.SurfelBuffers->Surfels.UAV;
-		RHICmdList.TransitionResources(EResourceTransitionAccess::EReadable, EResourceTransitionPipeline::EComputeToCompute, UniformMeshUAVs, ARRAY_COUNT(UniformMeshUAVs));
+		RHICmdList.TransitionResources(EResourceTransitionAccess::EReadable, EResourceTransitionPipeline::EComputeToCompute, UniformMeshUAVs, UE_ARRAY_COUNT(UniformMeshUAVs));
 	}
 
 	virtual bool Serialize(FArchive& Ar) override
diff --git a/Engine/Source/Runtime/Renderer/Private/CustomDepthRendering.cpp b/Engine/Source/Runtime/Renderer/Private/CustomDepthRendering.cpp
--- a/Engine/Source/Runtime/Renderer/Private/CustomDepthRendering.cpp
+++ b/Engine/Source/Runtime/Renderer/Private/CustomDepthRendering.cpp
@@ -88,7 +88,7 @@ void FCustomDepthPassMeshProcessor::AddMeshBatch(const FMeshBatch& RESTRICT Mesh
 				TStaticDepthStencilState<true, CF_DepthNearOrEqual, true, CF_Always, SO_Keep, SO_Replace, SO_Replace, false, CF_Always, SO_Keep, SO_Keep, SO_Keep, 255, 64>::GetRHI(),
 				TStaticDepthStencilState<true, CF_DepthNearOrEqual, true, CF_Always, SO_Keep, SO_Replace, SO_Replace, false, CF_Always, SO_Keep, SO_Keep, SO_Keep, 255, 128>::GetRHI()
 			};
-			checkSlow(EStencilMask::SM_Count == ARRAY_COUNT(StencilStates));
+			checkSlow(EStencilMask::SM_Count == UE_ARRAY_COUNT(StencilStates));
 
 			PassDrawRenderState.SetDepthStencilState(StencilStates[(int32)PrimitiveSceneProxy->GetStencilWriteMask()]);
 			PassDrawRenderState.SetStencilRef(CustomDepthStencilValue);
diff --git a/Engine/Source/Runtime/Renderer/Private/DistanceFieldAmbientOcclusion.cpp b/Engine/Source/Runtime/Renderer/Private/DistanceFieldAmbientOcclusion.cpp
--- a/Engine/Source/Runtime/Renderer/Private/DistanceFieldAmbientOcclusion.cpp
+++ b/Engine/Source/Runtime/Renderer/Private/DistanceFieldAmbientOcclusion.cpp
@@ -205,18 +205,18 @@ float TemporalHalton2( int32 Index, int32 Base )
 
 void GetSpacedVectors(uint32 FrameNumber, TArray<FVector, TInlineAllocator<9> >& OutVectors)
 {
-	OutVectors.Empty(ARRAY_COUNT(SpacedVectors9));
+	OutVectors.Empty(UE_ARRAY_COUNT(SpacedVectors9));
 
 	if (GAOSampleSet == 0)
 	{
-		for (int32 i = 0; i < ARRAY_COUNT(SpacedVectors9); i++)
+		for (int32 i = 0; i < UE_ARRAY_COUNT(SpacedVectors9); i++)
 		{
 			OutVectors.Add(SpacedVectors9[i]);
 		}
 	}
 	else
 	{
-		for (int32 i = 0; i < ARRAY_COUNT(RelaxedSpacedVectors9); i++)
+		for (int32 i = 0; i < UE_ARRAY_COUNT(RelaxedSpacedVectors9); i++)
 		{
 			OutVectors.Add(RelaxedSpacedVectors9[i]);
 		}
@@ -240,7 +240,7 @@ void GetSpacedVectors(uint32 FrameNumber, TArray<FVector, TInlineAllocator<9> >&
 }
 
 // Cone half angle derived from each cone covering an equal solid angle
-float GAOConeHalfAngle = FMath::Acos(1 - 1.0f / (float)ARRAY_COUNT(SpacedVectors9));
+float GAOConeHalfAngle = FMath::Acos(1 - 1.0f / (float)UE_ARRAY_COUNT(SpacedVectors9));
 
 // Number of AO sample positions along each cone
 // Must match shader code
@@ -497,14 +497,14 @@ void GenerateBestSpacedVectors()
 	{
 		bGenerated = true;
 
-		FVector OriginalSpacedVectors9[ARRAY_COUNT(SpacedVectors9)];
+		FVector OriginalSpacedVectors9[UE_ARRAY_COUNT(SpacedVectors9)];
 
-		for (int32 i = 0; i < ARRAY_COUNT(OriginalSpacedVectors9); i++)
+		for (int32 i = 0; i < UE_ARRAY_COUNT(OriginalSpacedVectors9); i++)
 		{
 			OriginalSpacedVectors9[i] = SpacedVectors9[i];
 		}
 
-		float CosHalfAngle = 1 - 1.0f / (float)ARRAY_COUNT(OriginalSpacedVectors9);
+		float CosHalfAngle = 1 - 1.0f / (float)UE_ARRAY_COUNT(OriginalSpacedVectors9);
 		// Used to prevent self-shadowing on a plane
 		float AngleBias = .03f;
 		float MinAngle = FMath::Acos(CosHalfAngle) + AngleBias;
@@ -513,11 +513,11 @@ void GenerateBestSpacedVectors()
 		// Relaxation iterations by repulsion
 		for (int32 Iteration = 0; Iteration < 10000; Iteration++)
 		{
-			for (int32 i = 0; i < ARRAY_COUNT(OriginalSpacedVectors9); i++)
+			for (int32 i = 0; i < UE_ARRAY_COUNT(OriginalSpacedVectors9); i++)
 			{
 				FVector Force(0.0f, 0.0f, 0.0f);
 
-				for (int32 j = 0; j < ARRAY_COUNT(OriginalSpacedVectors9); j++)
+				for (int32 j = 0; j < UE_ARRAY_COUNT(OriginalSpacedVectors9); j++)
 				{
 					if (i != j)
 					{
@@ -539,7 +539,7 @@ void GenerateBestSpacedVectors()
 			}
 		}
 
-		for (int32 i = 0; i < ARRAY_COUNT(OriginalSpacedVectors9); i++)
+		for (int32 i = 0; i < UE_ARRAY_COUNT(OriginalSpacedVectors9); i++)
 		{
 			UE_LOG(LogDistanceField, Log, TEXT("FVector(%f, %f, %f),"), OriginalSpacedVectors9[i].X, OriginalSpacedVectors9[i].Y, OriginalSpacedVectors9[i].Z);
 		}
@@ -557,7 +557,7 @@ void GenerateBestSpacedVectors()
 		// HemisphereSolidAngle = 2 * PI
 		// ConeSolidAngle = 2 * PI * (1 - cos(ConeHalfAngle))
 		// cos(ConeHalfAngle) = 1 - 1 / NumCones
-		float CosHalfAngle = 1 - 1.0f / (float)ARRAY_COUNT(BestSpacedVectors9);
+		float CosHalfAngle = 1 - 1.0f / (float)UE_ARRAY_COUNT(BestSpacedVectors9);
 		// Prevent self-intersection in sample set
 		float MinAngle = FMath::Acos(CosHalfAngle);
 		float MinZ = FMath::Sin(MinAngle);
@@ -566,9 +566,9 @@ void GenerateBestSpacedVectors()
 		// Super slow random brute force search
 		for (int i = 0; i < 1000000; i++)
 		{
-			FVector CandidateSpacedVectors[ARRAY_COUNT(BestSpacedVectors9)];
+			FVector CandidateSpacedVectors[UE_ARRAY_COUNT(BestSpacedVectors9)];
 
-			for (int j = 0; j < ARRAY_COUNT(CandidateSpacedVectors); j++)
+			for (int j = 0; j < UE_ARRAY_COUNT(CandidateSpacedVectors); j++)
 			{
 				FVector NewSample;
 
@@ -598,7 +598,7 @@ void GenerateBestSpacedVectors()
 
 				bool bIntersects = false;
 
-				for (int j = 0; j < ARRAY_COUNT(CandidateSpacedVectors); j++)
+				for (int j = 0; j < UE_ARRAY_COUNT(CandidateSpacedVectors); j++)
 				{
 					if (FVector::DotProduct(CandidateSpacedVectors[j], NewSample) > CosHalfAngle)
 					{
@@ -614,7 +614,7 @@ void GenerateBestSpacedVectors()
 			{
 				BestCoverage = Coverage;
 
-				for (int j = 0; j < ARRAY_COUNT(CandidateSpacedVectors); j++)
+				for (int j = 0; j < UE_ARRAY_COUNT(CandidateSpacedVectors); j++)
 				{
 					BestSpacedVectors9[j] = CandidateSpacedVectors[j];
 				}
