Fix for incorrect sample rate for media playback.

diff --git a/Engine/Source/Runtime/Engine/Private/SoundWave.cpp b/Engine/Source/Runtime/Engine/Private/SoundWave.cpp
--- a/Engine/Source/Runtime/Engine/Private/SoundWave.cpp
+++ b/Engine/Source/Runtime/Engine/Private/SoundWave.cpp
@@ -1950,6 +1950,11 @@ void USoundWave::UpdatePlatformData()
 
 float USoundWave::GetSampleRateForCurrentPlatform()
 {
+	if (bProcedural)
+	{
+		return SampleRate;
+	}
+
 #if WITH_EDITOR
 	float SampleRateOverride = FPlatformCompressionUtilities::GetTargetSampleRateForPlatform(SampleRateQuality);
 	return (SampleRateOverride > 0) ? FMath::Min(SampleRateOverride, (float) SampleRate) : SampleRate;
