Add a callback on UWorld for when world post processing volumes are about to be processed.  This way systems that have many post processing volumes can do some pre-processing and avoid complex per-post process volume calculations.

diff --git a/Engine/Source/Runtime/Engine/Classes/Engine/World.h b/Engine/Source/Runtime/Engine/Classes/Engine/World.h
--- a/Engine/Source/Runtime/Engine/Classes/Engine/World.h
+++ b/Engine/Source/Runtime/Engine/Classes/Engine/World.h
@@ -58,6 +58,7 @@ class UNetDriver;
 class UPrimitiveComponent;
 class UTexture2D;
 class FPhysScene_Chaos;
+class FSceneView;
 struct FUniqueNetIdRepl;
 struct FEncryptionKeyResponse;
 
@@ -1467,12 +1468,19 @@ public:
 	}
 #endif
 
+	/** Called when the world computes how post process volumes contribute to the scene. */
+	DECLARE_EVENT_OneParam(UWorld, FOnBeginPostProcessSettings, FVector);
+	FOnBeginPostProcessSettings OnBeginPostProcessSettings;
+
 	/** Inserts a post process volume into the world in priority order */
 	void InsertPostProcessVolume(IInterface_PostProcessVolume* InVolume);
 
 	/** Removes a post process volume from the world */
 	void RemovePostProcessVolume(IInterface_PostProcessVolume* InVolume);
 
+	/** Called when a scene view for this world needs the worlds post process settings computed */
+	void AddPostProcessingSettings(FVector ViewLocation, FSceneView* SceneView);
+
 	/** An array of post processing volumes, sorted in ascending order of priority.					*/
 	TArray< IInterface_PostProcessVolume * > PostProcessVolumes;
 
diff --git a/Engine/Source/Runtime/Engine/Private/SceneView.cpp b/Engine/Source/Runtime/Engine/Private/SceneView.cpp
--- a/Engine/Source/Runtime/Engine/Private/SceneView.cpp
+++ b/Engine/Source/Runtime/Engine/Private/SceneView.cpp
@@ -1626,51 +1626,7 @@ UBlendableInterface::UBlendableInterface(const FObjectInitializer& ObjectInitial
 {
 }
 
-void DoPostProcessVolume(IInterface_PostProcessVolume* Volume, FVector ViewLocation, FSceneView* SceneView)
-{
-	const FPostProcessVolumeProperties VolumeProperties = Volume->GetProperties();
-	if (!VolumeProperties.bIsEnabled)
-	{
-		return;
-	}
-
-	float DistanceToPoint = 0.0f;
-	float LocalWeight = FMath::Clamp(VolumeProperties.BlendWeight, 0.0f, 1.0f);
-
-	if (!VolumeProperties.bIsUnbound)
-	{
-		float SquaredBlendRadius = VolumeProperties.BlendRadius * VolumeProperties.BlendRadius;
-		Volume->EncompassesPoint(ViewLocation, 0.0f, &DistanceToPoint);
 
-		if (DistanceToPoint >= 0)
-		{
-			if (DistanceToPoint > VolumeProperties.BlendRadius)
-			{
-				// outside
-				LocalWeight = 0.0f;
-			}
-			else
-			{
-				// to avoid div by 0
-				if (VolumeProperties.BlendRadius >= 1.0f)
-				{
-					LocalWeight *= 1.0f - DistanceToPoint / VolumeProperties.BlendRadius;
-
-					check(LocalWeight >= 0 && LocalWeight <= 1.0f);
-				}
-			}
-		}
-		else
-		{
-			LocalWeight = 0;
-		}
-	}
-
-	if (LocalWeight > 0)
-	{
-		SceneView->OverridePostProcessSettings(*VolumeProperties.Settings, LocalWeight);
-	}
-}
 
 void FSceneView::StartFinalPostprocessSettings(FVector InViewLocation)
 {
@@ -1749,10 +1705,7 @@ void FSceneView::StartFinalPostprocessSettings(FVector InViewLocation)
 	// Some views have no world (e.g. material preview)
 	if (World != nullptr)
 	{
-		for (auto VolumeIt = World->PostProcessVolumes.CreateIterator(); VolumeIt; ++VolumeIt)
-		{
-			DoPostProcessVolume(*VolumeIt, InViewLocation, this);
-		}
+		World->AddPostProcessingSettings(InViewLocation, this);
 	}
 }
 
diff --git a/Engine/Source/Runtime/Engine/Private/World.cpp b/Engine/Source/Runtime/Engine/Private/World.cpp
--- a/Engine/Source/Runtime/Engine/Private/World.cpp
+++ b/Engine/Source/Runtime/Engine/Private/World.cpp
@@ -7510,6 +7510,63 @@ void UWorld::InitializeSubsystems()
 	}
 }
 
+static void DoPostProcessVolume(IInterface_PostProcessVolume* Volume, FVector ViewLocation, FSceneView* SceneView)
+{
+	const FPostProcessVolumeProperties VolumeProperties = Volume->GetProperties();
+	if (!VolumeProperties.bIsEnabled)
+	{
+		return;
+	}
+
+	float DistanceToPoint = 0.0f;
+	float LocalWeight = FMath::Clamp(VolumeProperties.BlendWeight, 0.0f, 1.0f);
+
+	if (!VolumeProperties.bIsUnbound)
+	{
+		float SquaredBlendRadius = VolumeProperties.BlendRadius * VolumeProperties.BlendRadius;
+		Volume->EncompassesPoint(ViewLocation, 0.0f, &DistanceToPoint);
+
+		if (DistanceToPoint >= 0)
+		{
+			if (DistanceToPoint > VolumeProperties.BlendRadius)
+			{
+				// outside
+				LocalWeight = 0.0f;
+			}
+			else
+			{
+				// to avoid div by 0
+				if (VolumeProperties.BlendRadius >= 1.0f)
+				{
+					LocalWeight *= 1.0f - DistanceToPoint / VolumeProperties.BlendRadius;
+
+					check(LocalWeight >= 0 && LocalWeight <= 1.0f);
+				}
+			}
+		}
+		else
+		{
+			LocalWeight = 0;
+		}
+	}
+
+	if (LocalWeight > 0)
+	{
+		SceneView->OverridePostProcessSettings(*VolumeProperties.Settings, LocalWeight);
+	}
+}
+
+void UWorld::AddPostProcessingSettings(FVector ViewLocation, FSceneView* SceneView)
+{
+	OnBeginPostProcessSettings.Broadcast(ViewLocation);
+
+	for (IInterface_PostProcessVolume* PPVolume : PostProcessVolumes)
+	{
+		DoPostProcessVolume(PPVolume, ViewLocation, SceneView);
+	}
+}
+
+
 /**
 * Dump visible actors in current world.
 */
