Update Engine and Plugins to use the new Getters / Setters for UActorComponent, USceneComponent, and AActor replicated properties

diff --git a/Engine/Source/Runtime/Engine/Private/Components/CharacterMovementComponent.cpp b/Engine/Source/Runtime/Engine/Private/Components/CharacterMovementComponent.cpp
--- a/Engine/Source/Runtime/Engine/Private/Components/CharacterMovementComponent.cpp
+++ b/Engine/Source/Runtime/Engine/Private/Components/CharacterMovementComponent.cpp
@@ -1281,7 +1281,7 @@ void UCharacterMovementComponent::TickComponent(float DeltaTime, enum ELevelTick
 
 	// See if we fell out of the world.
 	const bool bIsSimulatingPhysics = UpdatedComponent->IsSimulatingPhysics();
-	if (CharacterOwner->Role == ROLE_Authority && (!bCheatFlying || bIsSimulatingPhysics) && !CharacterOwner->CheckStillInWorld())
+	if (CharacterOwner->GetLocalRole() == ROLE_Authority && (!bCheatFlying || bIsSimulatingPhysics) && !CharacterOwner->CheckStillInWorld())
 	{
 		return;
 	}
@@ -1290,7 +1290,7 @@ void UCharacterMovementComponent::TickComponent(float DeltaTime, enum ELevelTick
 	if (bIsSimulatingPhysics)
 	{
 		// Update camera to ensure client gets updates even when physics move him far away from point where simulation started
-		if (CharacterOwner->Role == ROLE_AutonomousProxy && IsNetMode(NM_Client))
+		if (CharacterOwner->GetLocalRole() == ROLE_AutonomousProxy && IsNetMode(NM_Client))
 		{
 			MarkForClientCameraUpdate();
 		}
@@ -1301,12 +1301,12 @@ void UCharacterMovementComponent::TickComponent(float DeltaTime, enum ELevelTick
 
 	AvoidanceLockTimer -= DeltaTime;
 
-	if (CharacterOwner->Role > ROLE_SimulatedProxy)
+	if (CharacterOwner->GetLocalRole() > ROLE_SimulatedProxy)
 	{
 		SCOPE_CYCLE_COUNTER(STAT_CharacterMovementNonSimulated);
 
 		// If we are a client we might have received an update from the server.
-		const bool bIsClient = (CharacterOwner->Role == ROLE_AutonomousProxy && IsNetMode(NM_Client));
+		const bool bIsClient = (CharacterOwner->GetLocalRole() == ROLE_AutonomousProxy && IsNetMode(NM_Client));
 		if (bIsClient)
 		{
 			FNetworkPredictionData_Client_Character* ClientData = GetPredictionData_Client_Character();
@@ -1331,7 +1331,7 @@ void UCharacterMovementComponent::TickComponent(float DeltaTime, enum ELevelTick
 				AnalogInputModifier = ComputeAnalogInputModifier();
 			}
 
-			if (CharacterOwner->Role == ROLE_Authority)
+			if (CharacterOwner->GetLocalRole() == ROLE_Authority)
 			{
 				PerformMovement(DeltaTime);
 			}
@@ -1355,7 +1355,7 @@ void UCharacterMovementComponent::TickComponent(float DeltaTime, enum ELevelTick
 			}
 		}
 	}
-	else if (CharacterOwner->Role == ROLE_SimulatedProxy)
+	else if (CharacterOwner->GetLocalRole() == ROLE_SimulatedProxy)
 	{
 		if (bShrinkProxyCapsule)
 		{
@@ -1399,7 +1399,7 @@ void UCharacterMovementComponent::PostPhysicsTickComponent(float DeltaTime, FCha
 
 void UCharacterMovementComponent::AdjustProxyCapsuleSize()
 {
-	if (bShrinkProxyCapsule && CharacterOwner && CharacterOwner->Role == ROLE_SimulatedProxy)
+	if (bShrinkProxyCapsule && CharacterOwner && CharacterOwner->GetLocalRole() == ROLE_SimulatedProxy)
 	{
 		bShrinkProxyCapsule = false;
 
@@ -1583,10 +1583,10 @@ void UCharacterMovementComponent::SimulatedTick(float DeltaSeconds)
 			ApplyNetworkMovementMode(GetCharacterOwner()->GetReplicatedMovementMode());
 		}
 
-		if (CharacterOwner->bReplicateMovement && UpdatedComponent)
+		if (CharacterOwner->IsReplicatingMovement() && UpdatedComponent)
 		{
 			USkeletalMeshComponent* Mesh = CharacterOwner->GetMesh();
-			const FVector SavedMeshRelativeLocation = Mesh ? Mesh->RelativeLocation : FVector::ZeroVector; 
+			const FVector SavedMeshRelativeLocation = Mesh ? Mesh->GetRelativeLocation() : FVector::ZeroVector; 
 			const FQuat SavedCapsuleRotation = UpdatedComponent->GetComponentQuat();
 			const bool bPreventMeshMovement = !bNetworkSmoothingComplete;
 
@@ -1721,13 +1721,15 @@ void UCharacterMovementComponent::SimulateMovement(float DeltaSeconds)
 		return;
 	}
 
-	const bool bIsSimulatedProxy = (CharacterOwner->Role == ROLE_SimulatedProxy);
+	const bool bIsSimulatedProxy = (CharacterOwner->GetLocalRole() == ROLE_SimulatedProxy);
+
+	const FRepMovement& ConstRepMovement = CharacterOwner->GetReplicatedMovement();
 
 	// Workaround for replication not being updated initially
 	if (bIsSimulatedProxy &&
-		CharacterOwner->ReplicatedMovement.Location.IsZero() &&
-		CharacterOwner->ReplicatedMovement.Rotation.IsZero() &&
-		CharacterOwner->ReplicatedMovement.LinearVelocity.IsZero())
+		ConstRepMovement.Location.IsZero() &&
+		ConstRepMovement.Rotation.IsZero() &&
+		ConstRepMovement.LinearVelocity.IsZero())
 	{
 		return;
 	}
@@ -1786,7 +1788,7 @@ void UCharacterMovementComponent::SimulateMovement(float DeltaSeconds)
 		}
 
 		const bool bSimGravityDisabled = (bIsSimulatedProxy && CharacterOwner->bSimGravityDisabled);
-		const bool bZeroReplicatedGroundVelocity = (bIsSimulatedProxy && IsMovingOnGround() && CharacterOwner->ReplicatedMovement.LinearVelocity.IsZero());
+		const bool bZeroReplicatedGroundVelocity = (bIsSimulatedProxy && IsMovingOnGround() && ConstRepMovement.LinearVelocity.IsZero());
 		
 		// bSimGravityDisabled means velocity was zero when replicated and we were stuck in something. Avoid external changes in velocity as well.
 		// Being in ground movement with zero velocity, we cannot simulate proxy velocities safely because we might not get any further updates from the server.
@@ -2300,7 +2302,7 @@ void UCharacterMovementComponent::PerformMovement(float DeltaSeconds)
 				}
 
 				// For local human clients, save off root motion data so it can be used by movement networking code.
-				if( CharacterOwner->IsLocallyControlled() && (CharacterOwner->Role == ROLE_AutonomousProxy) && CharacterOwner->IsPlayingNetworkedRootMotionMontage() )
+				if( CharacterOwner->IsLocallyControlled() && (CharacterOwner->GetLocalRole() == ROLE_AutonomousProxy) && CharacterOwner->IsPlayingNetworkedRootMotionMontage() )
 				{
 					CharacterOwner->ClientRootMotionParams = RootMotionParams;
 				}
@@ -2313,7 +2315,7 @@ void UCharacterMovementComponent::PerformMovement(float DeltaSeconds)
 			}
 
 			// For local human clients, save off root motion data so it can be used by movement networking code.
-			if( CharacterOwner->IsLocallyControlled() && (CharacterOwner->Role == ROLE_AutonomousProxy) )
+			if( CharacterOwner->IsLocallyControlled() && (CharacterOwner->GetLocalRole() == ROLE_AutonomousProxy) )
 			{
 				CharacterOwner->SavedRootMotion = CurrentRootMotion;
 			}
@@ -2606,7 +2608,7 @@ void UCharacterMovementComponent::Crouch(bool bClientSimulation)
 		return;
 	}
 
-	if (bClientSimulation && CharacterOwner->Role == ROLE_SimulatedProxy)
+	if (bClientSimulation && CharacterOwner->GetLocalRole() == ROLE_SimulatedProxy)
 	{
 		// restore collision size before crouching
 		ACharacter* DefaultCharacter = CharacterOwner->GetClass()->GetDefaultObject<ACharacter>();
@@ -2664,7 +2666,7 @@ void UCharacterMovementComponent::Crouch(bool bClientSimulation)
 	CharacterOwner->OnStartCrouch( HalfHeightAdjust, ScaledHalfHeightAdjust );
 
 	// Don't smooth this change in mesh position
-	if (bClientSimulation && CharacterOwner->Role == ROLE_SimulatedProxy)
+	if (bClientSimulation && CharacterOwner->GetLocalRole() == ROLE_SimulatedProxy)
 	{
 		FNetworkPredictionData_Client_Character* ClientData = GetPredictionData_Client_Character();
 		if (ClientData && ClientData->MeshTranslationOffset.Z != 0.f)
@@ -2809,7 +2811,7 @@ void UCharacterMovementComponent::UnCrouch(bool bClientSimulation)
 	CharacterOwner->OnEndCrouch( HalfHeightAdjust, ScaledHalfHeightAdjust );
 
 	// Don't smooth this change in mesh position
-	if (bClientSimulation && CharacterOwner->Role == ROLE_SimulatedProxy)
+	if (bClientSimulation && CharacterOwner->GetLocalRole() == ROLE_SimulatedProxy)
 	{
 		FNetworkPredictionData_Client_Character* ClientData = GetPredictionData_Client_Character();
 		if (ClientData && ClientData->MeshTranslationOffset.Z != 0.f)
@@ -2939,7 +2941,7 @@ FVector UCharacterMovementComponent::GetPenetrationAdjustment(const FHitResult&
 
 	if (CharacterOwner)
 	{
-		const bool bIsProxy = (CharacterOwner->Role == ROLE_SimulatedProxy);
+		const bool bIsProxy = (CharacterOwner->GetLocalRole() == ROLE_SimulatedProxy);
 		float MaxDistance = bIsProxy ? MaxDepenetrationWithGeometryAsProxy : MaxDepenetrationWithGeometry;
 		const AActor* HitActor = Hit.GetActor();
 		if (Cast<APawn>(HitActor))
@@ -3174,7 +3176,7 @@ bool UCharacterMovementComponent::IsCrouching() const
 void UCharacterMovementComponent::CalcVelocity(float DeltaTime, float Friction, bool bFluid, float BrakingDeceleration)
 {
 	// Do not update velocity when using root motion or when SimulatedProxy and not simulating root motion - SimulatedProxy are repped their Velocity
-	if (!HasValidData() || HasAnimRootMotion() || DeltaTime < MIN_TICK_TIME || (CharacterOwner && CharacterOwner->Role == ROLE_SimulatedProxy && !bWasSimulatingRootMotion))
+	if (!HasValidData() || HasAnimRootMotion() || DeltaTime < MIN_TICK_TIME || (CharacterOwner && CharacterOwner->GetLocalRole() == ROLE_SimulatedProxy && !bWasSimulatingRootMotion))
 	{
 		return;
 	}
@@ -3410,7 +3412,7 @@ void UCharacterMovementComponent::CalcAvoidanceVelocity(float DeltaTime)
 		return;
 	}
 
-	if (GetCharacterOwner()->Role != ROLE_Authority)
+	if (GetCharacterOwner()->GetLocalRole() != ROLE_Authority)
 	{
 		return;
 	}
@@ -4783,7 +4785,7 @@ void UCharacterMovementComponent::PhysWalking(float deltaTime, int32 Iterations)
 		return;
 	}
 
-	if (!CharacterOwner || (!CharacterOwner->Controller && !bRunPhysicsWithNoController && !HasAnimRootMotion() && !CurrentRootMotion.HasOverrideVelocity() && (CharacterOwner->Role != ROLE_SimulatedProxy)))
+	if (!CharacterOwner || (!CharacterOwner->Controller && !bRunPhysicsWithNoController && !HasAnimRootMotion() && !CurrentRootMotion.HasOverrideVelocity() && (CharacterOwner->GetLocalRole() != ROLE_SimulatedProxy)))
 	{
 		Acceleration = FVector::ZeroVector;
 		Velocity = FVector::ZeroVector;
@@ -4804,7 +4806,7 @@ void UCharacterMovementComponent::PhysWalking(float deltaTime, int32 Iterations)
 	float remainingTime = deltaTime;
 
 	// Perform the move
-	while ( (remainingTime >= MIN_TICK_TIME) && (Iterations < MaxSimulationIterations) && CharacterOwner && (CharacterOwner->Controller || bRunPhysicsWithNoController || HasAnimRootMotion() || CurrentRootMotion.HasOverrideVelocity() || (CharacterOwner->Role == ROLE_SimulatedProxy)) )
+	while ( (remainingTime >= MIN_TICK_TIME) && (Iterations < MaxSimulationIterations) && CharacterOwner && (CharacterOwner->Controller || bRunPhysicsWithNoController || HasAnimRootMotion() || CurrentRootMotion.HasOverrideVelocity() || (CharacterOwner->GetLocalRole() == ROLE_SimulatedProxy)) )
 	{
 		Iterations++;
 		bJustTeleported = false;
@@ -6002,7 +6004,7 @@ bool UCharacterMovementComponent::IsWithinEdgeTolerance(const FVector& CapsuleLo
 
 void UCharacterMovementComponent::ComputeFloorDist(const FVector& CapsuleLocation, float LineDistance, float SweepDistance, FFindFloorResult& OutFloorResult, float SweepRadius, const FHitResult* DownwardSweepResult) const
 {
-	UE_LOG(LogCharacterMovement, VeryVerbose, TEXT("[Role:%d] ComputeFloorDist: %s at location %s"), (int32)CharacterOwner->Role, *GetNameSafe(CharacterOwner), *CapsuleLocation.ToString());
+	UE_LOG(LogCharacterMovement, VeryVerbose, TEXT("[Role:%d] ComputeFloorDist: %s at location %s"), (int32)CharacterOwner->GetLocalRole(), *GetNameSafe(CharacterOwner), *CapsuleLocation.ToString());
 	OutFloorResult.Clear();
 
 	float PawnRadius, PawnHalfHeight;
@@ -6155,7 +6157,7 @@ void UCharacterMovementComponent::FindFloor(const FVector& CapsuleLocation, FFin
 		return;
 	}
 
-	UE_LOG(LogCharacterMovement, VeryVerbose, TEXT("[Role:%d] FindFloor: %s at location %s"), (int32)CharacterOwner->Role, *GetNameSafe(CharacterOwner), *CapsuleLocation.ToString());
+	UE_LOG(LogCharacterMovement, VeryVerbose, TEXT("[Role:%d] FindFloor: %s at location %s"), (int32)CharacterOwner->GetLocalRole(), *GetNameSafe(CharacterOwner), *CapsuleLocation.ToString());
 	check(CharacterOwner->GetCapsuleComponent());
 
 	// Increase height check slightly if walking, to prevent floor height adjustment from later invalidating the floor result.
@@ -7070,7 +7072,7 @@ void UCharacterMovementComponent::SmoothCorrection(const FVector& OldLocation, c
 	checkSlow(GetNetMode() != NM_Standalone);
 
 	// Only client proxies or remote clients on a listen server should run this code.
-	const bool bIsSimulatedProxy = (CharacterOwner->Role == ROLE_SimulatedProxy);
+	const bool bIsSimulatedProxy = (CharacterOwner->GetLocalRole() == ROLE_SimulatedProxy);
 	const bool bIsRemoteAutoProxy = (CharacterOwner->GetRemoteRole() == ROLE_AutonomousProxy);
 	ensure(bIsSimulatedProxy || bIsRemoteAutoProxy);
 
@@ -7267,7 +7269,7 @@ void UCharacterMovementComponent::SmoothClientPosition(float DeltaSeconds)
 	checkSlow(GetNetMode() != NM_Standalone);
 
 	// Only client proxies or remote clients on a listen server should run this code.
-	const bool bIsSimulatedProxy = (CharacterOwner->Role == ROLE_SimulatedProxy);
+	const bool bIsSimulatedProxy = (CharacterOwner->GetLocalRole() == ROLE_SimulatedProxy);
 	const bool bIsRemoteAutoProxy = (CharacterOwner->GetRemoteRole() == ROLE_AutonomousProxy);
 	if (!ensure(bIsSimulatedProxy || bIsRemoteAutoProxy))
 	{
@@ -7484,8 +7486,8 @@ void UCharacterMovementComponent::SmoothClientPosition_Interpolate(float DeltaSe
 
 					if ( Mesh )
 					{
-						Mesh->RelativeLocation = CharacterOwner->GetBaseTranslationOffset();
-						Mesh->RelativeRotation = CharacterOwner->GetBaseRotationOffset().Rotator();
+						Mesh->SetRelativeLocation_Direct(CharacterOwner->GetBaseTranslationOffset());
+						Mesh->SetRelativeRotation_Direct(CharacterOwner->GetBaseRotationOffset().Rotator());
 					}
 
 					ClientData->MeshTranslationOffset = Location;
@@ -7540,8 +7542,8 @@ void UCharacterMovementComponent::SmoothClientPosition_Interpolate(float DeltaSe
 
 				if ( Mesh )
 				{
-					Mesh->RelativeLocation = CharacterOwner->GetBaseTranslationOffset();
-					Mesh->RelativeRotation = CharacterOwner->GetBaseRotationOffset().Rotator();
+					Mesh->SetRelativeLocation_Direct(CharacterOwner->GetBaseTranslationOffset());
+					Mesh->SetRelativeRotation_Direct(CharacterOwner->GetBaseRotationOffset().Rotator());
 				}
 
 				ClientData->MeshTranslationOffset	= ReplaySample.Location;
@@ -7666,15 +7668,15 @@ void UCharacterMovementComponent::SmoothClientPosition_UpdateVisuals()
 			const FVector NewRelLocation = ClientData->MeshRotationOffset.UnrotateVector(ClientData->MeshTranslationOffset) + CharacterOwner->GetBaseTranslationOffset();
 			if (!UpdatedComponent->GetComponentQuat().Equals(ClientData->MeshRotationOffset, 1e-6f))
 			{
-				const FVector OldLocation = Mesh->RelativeLocation;
-				const FRotator OldRotation = UpdatedComponent->RelativeRotation;
-				Mesh->RelativeLocation = NewRelLocation;
+				const FVector OldLocation = Mesh->GetRelativeLocation();
+				const FRotator OldRotation = UpdatedComponent->GetRelativeRotation();
+				Mesh->SetRelativeLocation_Direct(NewRelLocation);
 				UpdatedComponent->SetWorldRotation(ClientData->MeshRotationOffset);
 
 				// If we did not move from SetWorldRotation, we need to at least call SetRelativeLocation since we were relying on the UpdatedComponent to update the transform of the mesh
-				if (UpdatedComponent->RelativeRotation == OldRotation)
+				if (UpdatedComponent->GetRelativeRotation() == OldRotation)
 				{
-					Mesh->RelativeLocation = OldLocation;
+					Mesh->SetRelativeLocation_Direct(OldLocation);
 					Mesh->SetRelativeLocation(NewRelLocation, false, nullptr, GetTeleportType());
 				}
 			}
@@ -7806,7 +7808,7 @@ bool UCharacterMovementComponent::ForcePositionUpdate(float DeltaTime)
 		return false;
 	}
 
-	check(CharacterOwner->Role == ROLE_Authority);
+	check(CharacterOwner->GetLocalRole() == ROLE_Authority);
 	check(CharacterOwner->GetRemoteRole() == ROLE_AutonomousProxy);
 
 	// TODO: this is basically copied from MoveAutonmous, should consolidate. Or consider moving CheckJumpInput inside PerformMovement().
@@ -7869,7 +7871,7 @@ FNetworkPredictionData_Client_Character* UCharacterMovementComponent::GetPredict
 {
 	// Should only be called on client or listen server (for remote clients) in network games
 	checkSlow(CharacterOwner != NULL);
-	checkSlow(CharacterOwner->Role < ROLE_Authority || (CharacterOwner->GetRemoteRole() == ROLE_AutonomousProxy && GetNetMode() == NM_ListenServer));
+	checkSlow(CharacterOwner->GetLocalRole() < ROLE_Authority || (CharacterOwner->GetRemoteRole() == ROLE_AutonomousProxy && GetNetMode() == NM_ListenServer));
 	checkSlow(GetNetMode() == NM_Client || GetNetMode() == NM_ListenServer);
 
 	if (ClientPredictionData == nullptr)
@@ -7886,7 +7888,7 @@ FNetworkPredictionData_Server_Character* UCharacterMovementComponent::GetPredict
 {
 	// Should only be called on server in network games
 	checkSlow(CharacterOwner != NULL);
-	checkSlow(CharacterOwner->Role == ROLE_Authority);
+	checkSlow(CharacterOwner->GetLocalRole() == ROLE_Authority);
 	checkSlow(GetNetMode() < NM_Client);
 
 	if (ServerPredictionData == nullptr)
@@ -8098,7 +8100,7 @@ void UCharacterMovementComponent::ReplicateMoveToServer(float DeltaTime, const F
 	NewMove->PostUpdate(CharacterOwner, FSavedMove_Character::PostUpdate_Record);
 
 	// Add NewMove to the list
-	if (CharacterOwner->bReplicateMovement)
+	if (CharacterOwner->IsReplicatingMovement())
 	{
 		check(NewMove == NewMovePtr.Get());
 		ClientData->SavedMoves.Push(NewMovePtr);
@@ -8376,7 +8378,7 @@ void UCharacterMovementComponent::ProcessClientTimeStampForTimeDiscrepancy(float
 {
 	// Should only be called on server in network games
 	check(CharacterOwner != NULL);
-	check(CharacterOwner->Role == ROLE_Authority);
+	check(CharacterOwner->GetLocalRole() == ROLE_Authority);
 	checkSlow(GetNetMode() < NM_Client);
 
 	// Movement time discrepancy detection and resolution (potentially caused by client speed hacks, time manipulation)
@@ -9846,7 +9848,7 @@ void UCharacterMovementComponent::ApplyWorldOffset(const FVector& InOffset, bool
 	OldBaseLocation += InOffset;
 	LastUpdateLocation += InOffset;
 
-	if (CharacterOwner != nullptr && CharacterOwner->Role == ROLE_AutonomousProxy)
+	if (CharacterOwner != nullptr && CharacterOwner->GetLocalRole() == ROLE_AutonomousProxy)
 	{
 		FNetworkPredictionData_Client_Character* ClientData = GetPredictionData_Client_Character();
 		if (ClientData != nullptr)
@@ -9909,7 +9911,7 @@ void UCharacterMovementComponent::TickCharacterPose(float DeltaTime)
 			{
 				FAnimMontageInstance* RootMotionMontageInstance = CharacterOwner->GetRootMotionAnimMontageInstance();
 				UE_LOG(LogRootMotion, Log, TEXT("UCharacterMovementComponent::TickCharacterPose Role: %s, RootMotionMontage: %s, MontagePos: %f, DeltaTime: %f, ExtractedRootMotion: %s, AccumulatedRootMotion: %s")
-					, *UEnum::GetValueAsString(TEXT("Engine.ENetRole"), CharacterOwner->Role)
+					, *UEnum::GetValueAsString(TEXT("Engine.ENetRole"), CharacterOwner->GetLocalRole())
 					, *GetNameSafe(RootMotionMontageInstance ? RootMotionMontageInstance->Montage : NULL)
 					, RootMotionMontageInstance ? RootMotionMontageInstance->GetPosition() : -1.f
 					, DeltaTime
@@ -9942,7 +9944,7 @@ uint16 UCharacterMovementComponent::ApplyRootMotionSource(FRootMotionSource* Sou
 		{
 			if (CharacterOwner)
 			{
-				if (CharacterOwner->Role == ROLE_AutonomousProxy)
+				if (CharacterOwner->GetLocalRole() == ROLE_AutonomousProxy)
 				{
 					// Autonomous defaults to local timestamp
 					FNetworkPredictionData_Client_Character* ClientData = GetPredictionData_Client_Character();
@@ -9951,7 +9953,7 @@ uint16 UCharacterMovementComponent::ApplyRootMotionSource(FRootMotionSource* Sou
 						SourcePtr->StartTime = ClientData->CurrentTimeStamp;
 					}
 				}
-				else if (CharacterOwner->Role == ROLE_Authority && !IsNetMode(NM_Client))
+				else if (CharacterOwner->GetLocalRole() == ROLE_Authority && !IsNetMode(NM_Client))
 				{
 					// Authority defaults to current client time stamp, meaning it'll start next tick if not corrected
 					FNetworkPredictionData_Server_Character* ServerData = GetPredictionData_Server_Character();
@@ -10589,8 +10591,8 @@ void FSavedMove_Character::SetInitialPosition(ACharacter* Character)
 	{
 		StartAttachParent = UpdatedComponent->GetAttachParent();
 		StartAttachSocketName = UpdatedComponent->GetAttachSocketName();
-		StartAttachRelativeLocation = UpdatedComponent->RelativeLocation;
-		StartAttachRelativeRotation = UpdatedComponent->RelativeRotation;
+		StartAttachRelativeLocation = UpdatedComponent->GetRelativeLocation();
+		StartAttachRelativeRotation = UpdatedComponent->GetRelativeRotation();
 	}
 
 	StartControlRotation = Character->GetControlRotation().Clamp();
@@ -10640,8 +10642,8 @@ void FSavedMove_Character::PostUpdate(ACharacter* Character, FSavedMove_Characte
 		{
 			EndAttachParent = UpdatedComponent->GetAttachParent();
 			EndAttachSocketName = UpdatedComponent->GetAttachSocketName();
-			EndAttachRelativeLocation = UpdatedComponent->RelativeLocation;
-			EndAttachRelativeRotation = UpdatedComponent->RelativeRotation;
+			EndAttachRelativeLocation = UpdatedComponent->GetRelativeLocation();
+			EndAttachRelativeRotation = UpdatedComponent->GetRelativeRotation();
 		}
 
 		SavedControlRotation = Character->GetControlRotation().Clamp();
@@ -11108,7 +11110,7 @@ void UCharacterMovementComponent::UpdateFromCompressedFlags(uint8 Flags)
 	bWantsToCrouch = ((Flags & FSavedMove_Character::FLAG_WantsToCrouch) != 0);
 
 	// Detect change in jump press on the server
-	if (CharacterOwner->Role == ROLE_Authority) 
+	if (CharacterOwner->GetLocalRole() == ROLE_Authority) 
 	{
 		const bool bIsPressingJump = CharacterOwner->bPressedJump;
 		if (bIsPressingJump && !bWasPressingJump)
@@ -11125,7 +11127,7 @@ void UCharacterMovementComponent::UpdateFromCompressedFlags(uint8 Flags)
 void UCharacterMovementComponent::FlushServerMoves()
 {
 	// Send pendingMove to server if this character is replicating movement
-	if (CharacterOwner && CharacterOwner->bReplicateMovement && (CharacterOwner->Role == ROLE_AutonomousProxy))
+	if (CharacterOwner && CharacterOwner->IsReplicatingMovement() && (CharacterOwner->GetLocalRole() == ROLE_AutonomousProxy))
 	{
 		FNetworkPredictionData_Client_Character* ClientData = GetPredictionData_Client_Character();
 		if (!ClientData)
diff --git a/Engine/Source/Runtime/Engine/Private/Components/DecalComponent.cpp b/Engine/Source/Runtime/Engine/Private/Components/DecalComponent.cpp
--- a/Engine/Source/Runtime/Engine/Private/Components/DecalComponent.cpp
+++ b/Engine/Source/Runtime/Engine/Private/Components/DecalComponent.cpp
@@ -18,8 +18,8 @@ static TAutoConsoleVariable<float> CVarDecalFadeDurationScale(
 	);
 
 FDeferredDecalProxy::FDeferredDecalProxy(const UDecalComponent* InComponent)
-	: DrawInGame( InComponent->bVisible && !InComponent->bHiddenInGame )
-	, DrawInEditor( InComponent->bVisible )
+	: DrawInGame(InComponent->GetVisibleFlag() && !InComponent->bHiddenInGame)
+	, DrawInEditor(InComponent->GetVisibleFlag())
 	, InvFadeDuration(-1.0f)
 	, InvFadeInDuration(1.0f)
 	, FadeStartDelayNormalized(1.0f)
@@ -55,7 +55,7 @@ FDeferredDecalProxy::FDeferredDecalProxy(const UDecalComponent* InComponent)
 	
 	if ( InComponent->GetOwner() )
 	{
-		DrawInGame &= !( InComponent->GetOwner()->bHidden );
+		DrawInGame &= !(InComponent->GetOwner()->IsHidden());
 #if WITH_EDITOR
 		DrawInEditor &= !InComponent->GetOwner()->IsHiddenEd();
 #endif
diff --git a/Engine/Source/Runtime/Engine/Private/Components/ForceFeedbackComponent.cpp b/Engine/Source/Runtime/Engine/Private/Components/ForceFeedbackComponent.cpp
--- a/Engine/Source/Runtime/Engine/Private/Components/ForceFeedbackComponent.cpp
+++ b/Engine/Source/Runtime/Engine/Private/Components/ForceFeedbackComponent.cpp
@@ -195,7 +195,7 @@ void UForceFeedbackComponent::OnRegister()
 
 bool UForceFeedbackComponent::IsReadyForOwnerToAutoDestroy() const
 {
-	return !bIsActive;
+	return !IsActive();
 }
 
 const UObject* UForceFeedbackComponent::AdditionalStatObject() const
@@ -234,7 +234,7 @@ void UForceFeedbackComponent::Deactivate()
 	{
 		Stop();
 
-		if (!bIsActive)
+		if (!IsActive())
 		{
 			OnComponentDeactivated.Broadcast(this);
 		}
@@ -263,7 +263,7 @@ void UForceFeedbackComponent::Play(const float StartTime)
 {
 	UWorld* World = GetWorld();
 
-	if (bIsActive)
+	if (IsActive())
 	{
 		// If this is an auto destroy component we need to prevent it from being auto-destroyed since we're really just restarting it
 		const bool bCurrentAutoDestroy = bAutoDestroy;
@@ -274,7 +274,7 @@ void UForceFeedbackComponent::Play(const float StartTime)
 
 	if (ForceFeedbackEffect && World)
 	{
-		bIsActive = true;
+		SetActiveFlag(true);
 		PlayTime = StartTime;
 		FForceFeedbackManager::Get(World, true)->AddActiveComponent(this);
 	}
@@ -282,7 +282,7 @@ void UForceFeedbackComponent::Play(const float StartTime)
 
 void UForceFeedbackComponent::Stop()
 {
-	if (bIsActive)
+	if (IsActive())
 	{
 		StopInternal(true);
 	}
@@ -291,7 +291,7 @@ void UForceFeedbackComponent::Stop()
 void UForceFeedbackComponent::StopInternal(const bool bRemoveFromManager)
 {
 	// Set this to immediately be inactive
-	bIsActive = false;
+	SetActiveFlag(false);
 	PlayTime = 0.f;
 
 	if (bRemoveFromManager)
diff --git a/Engine/Source/Runtime/Engine/Private/Components/HeightFogComponent.cpp b/Engine/Source/Runtime/Engine/Private/Components/HeightFogComponent.cpp
--- a/Engine/Source/Runtime/Engine/Private/Components/HeightFogComponent.cpp
+++ b/Engine/Source/Runtime/Engine/Private/Components/HeightFogComponent.cpp
@@ -315,7 +315,7 @@ AExponentialHeightFog::AExponentialHeightFog(const FObjectInitializer& ObjectIni
 	Component = CreateDefaultSubobject<UExponentialHeightFogComponent>(TEXT("HeightFogComponent0"));
 	RootComponent = Component;
 
-	bHidden = false;
+	SetHidden(false);
 
 #if WITH_EDITORONLY_DATA
 	if (!IsRunningCommandlet() && (GetSpriteComponent() != NULL))
@@ -336,7 +336,7 @@ AExponentialHeightFog::AExponentialHeightFog(const FObjectInitializer& ObjectIni
 		static FConstructorStatics ConstructorStatics;
 
 		GetSpriteComponent()->Sprite = ConstructorStatics.FogTextureObject.Get();
-		GetSpriteComponent()->RelativeScale3D = FVector(0.5f, 0.5f, 0.5f);
+		GetSpriteComponent()->SetRelativeScale3D_Direct(FVector(0.5f, 0.5f, 0.5f));
 		GetSpriteComponent()->SpriteInfo.Category = ConstructorStatics.ID_Fog;
 		GetSpriteComponent()->SpriteInfo.DisplayName = ConstructorStatics.NAME_Fog;
 		GetSpriteComponent()->SetupAttachment(Component);
@@ -348,7 +348,7 @@ void AExponentialHeightFog::PostInitializeComponents()
 {
 	Super::PostInitializeComponents();
 
-	bEnabled = Component->bVisible;
+	bEnabled = Component->GetVisibleFlag();
 }
 
 void AExponentialHeightFog::GetLifetimeReplicatedProps( TArray< FLifetimeProperty > & OutLifetimeProps ) const
diff --git a/Engine/Source/Runtime/Engine/Private/Components/InterpToMovementComponent.cpp b/Engine/Source/Runtime/Engine/Private/Components/InterpToMovementComponent.cpp
--- a/Engine/Source/Runtime/Engine/Private/Components/InterpToMovementComponent.cpp
+++ b/Engine/Source/Runtime/Engine/Private/Components/InterpToMovementComponent.cpp
@@ -85,7 +85,7 @@ void UInterpToMovementComponent::TickComponent(float DeltaTime, enum ELevelTick
 	{
 		WaitPos = UpdatedComponent->GetComponentLocation(); //-V595
 	}
-	while (RemainingTime >= MIN_TICK_TIME && (Iterations < MaxSimulationIterations) && !ActorOwner->IsPendingKill() && UpdatedComponent && bIsActive)
+	while (RemainingTime >= MIN_TICK_TIME && (Iterations < MaxSimulationIterations) && !ActorOwner->IsPendingKill() && UpdatedComponent && IsActive())
 	{
 		Iterations++;
 
@@ -116,7 +116,7 @@ void UInterpToMovementComponent::TickComponent(float DeltaTime, enum ELevelTick
 		}
 		//DrawDebugPoint(GetWorld(), UpdatedComponent->GetComponentLocation(), 16, FColor::White,true,5.0f);
 		// If we hit a trigger that destroyed us, abort.
-		if (ActorOwner->IsPendingKill() || !UpdatedComponent || !bIsActive)
+		if (ActorOwner->IsPendingKill() || !UpdatedComponent || !IsActive())
 		{
 			return;
 		}
diff --git a/Engine/Source/Runtime/Engine/Private/Components/LightComponent.cpp b/Engine/Source/Runtime/Engine/Private/Components/LightComponent.cpp
--- a/Engine/Source/Runtime/Engine/Private/Components/LightComponent.cpp
+++ b/Engine/Source/Runtime/Engine/Private/Components/LightComponent.cpp
@@ -680,7 +680,7 @@ void ULightComponent::PostEditChangeProperty(FPropertyChangedEvent& PropertyChan
 		PropertyName != GET_MEMBER_NAME_STRING_CHECKED(ULightComponent, BloomTint) &&
 		PropertyName != GET_MEMBER_NAME_STRING_CHECKED(ULightComponent, bUseRayTracedDistanceFieldShadows) &&
 		PropertyName != GET_MEMBER_NAME_STRING_CHECKED(ULightComponent, RayStartOffsetDepthScale) &&
-		PropertyName != GET_MEMBER_NAME_STRING_CHECKED(ULightComponent, bVisible) &&
+		PropertyName != USceneComponent::GetVisiblePropertyName().ToString() &&
 		PropertyName != GET_MEMBER_NAME_STRING_CHECKED(ULightComponent, LightingChannels) &&
 		PropertyName != GET_MEMBER_NAME_STRING_CHECKED(ULightComponent, VolumetricScatteringIntensity) &&
 		PropertyName != GET_MEMBER_NAME_STRING_CHECKED(ULightComponent, bCastVolumetricShadow) &&
diff --git a/Engine/Source/Runtime/Engine/Private/Components/PortalComponent.cpp b/Engine/Source/Runtime/Engine/Private/Components/PortalComponent.cpp
--- a/Engine/Source/Runtime/Engine/Private/Components/PortalComponent.cpp
+++ b/Engine/Source/Runtime/Engine/Private/Components/PortalComponent.cpp
@@ -17,7 +17,7 @@ ALightmassPortal::ALightmassPortal(const FObjectInitializer& ObjectInitializer)
 	: Super(ObjectInitializer)
 {
 	PortalComponent = CreateDefaultSubobject<ULightmassPortalComponent>(TEXT("PortalComponent"));
-	PortalComponent->RelativeScale3D = FVector(10, 100, 100);
+	PortalComponent->SetRelativeScale3D_Direct(FVector(10, 100, 100));
 	RootComponent = PortalComponent;
 	UBoxComponent* DrawInfluenceBox = CreateDefaultSubobject<UBoxComponent>(TEXT("DrawBox0"));
 	DrawInfluenceBox->SetupAttachment(GetPortalComponent());
@@ -42,9 +42,9 @@ ALightmassPortal::ALightmassPortal(const FObjectInitializer& ObjectInitializer)
 		static FConstructorStatics ConstructorStatics;
 
 		SpriteComponent->Sprite = ConstructorStatics.DecalTexture.Get();
-		SpriteComponent->RelativeScale3D = FVector(0.5f, 0.5f, 0.5f);
+		SpriteComponent->SetRelativeScale3D_Direct(FVector(0.5f, 0.5f, 0.5f));
 		SpriteComponent->bHiddenInGame = true;
-		SpriteComponent->bAbsoluteScale = true;
+		SpriteComponent->SetUsingAbsoluteScale(true);
 		SpriteComponent->SetCollisionProfileName(UCollisionProfile::NoCollision_ProfileName);
 		SpriteComponent->bIsScreenSizeScaled = true;
 		SpriteComponent->SetupAttachment(PortalComponent);
@@ -57,7 +57,7 @@ void ALightmassPortal::PostEditMove(bool bFinished)
 {
 	Super::PostEditMove(bFinished);
 
-	PortalComponent->RelativeScale3D.X = 10;
+	PortalComponent->GetRelativeScale3D_DirectMutable().X = 10;
 	PortalComponent->MarkRenderStateDirty();
 }
 #endif // WITH_EDITOR
diff --git a/Engine/Source/Runtime/Engine/Private/Components/PrimitiveComponent.cpp b/Engine/Source/Runtime/Engine/Private/Components/PrimitiveComponent.cpp
--- a/Engine/Source/Runtime/Engine/Private/Components/PrimitiveComponent.cpp
+++ b/Engine/Source/Runtime/Engine/Private/Components/PrimitiveComponent.cpp
@@ -2896,7 +2896,7 @@ bool UPrimitiveComponent::AreAllCollideableDescendantsRelative(bool bAllowCached
 			if (CurrentComp)
 			{
 				// Is the component not using relative position?
-				if (CurrentComp->bAbsoluteLocation || CurrentComp->bAbsoluteRotation)
+				if (CurrentComp->IsUsingAbsoluteLocation() || CurrentComp->IsUsingAbsoluteRotation())
 				{
 					// Can we possibly collide with the component?
 					UPrimitiveComponent* const CurrentPrimitive = Cast<UPrimitiveComponent>(CurrentComp);
diff --git a/Engine/Source/Runtime/Engine/Private/Components/ProjectileMovementComponent.cpp b/Engine/Source/Runtime/Engine/Private/Components/ProjectileMovementComponent.cpp
--- a/Engine/Source/Runtime/Engine/Private/Components/ProjectileMovementComponent.cpp
+++ b/Engine/Source/Runtime/Engine/Private/Components/ProjectileMovementComponent.cpp
@@ -188,7 +188,7 @@ void UProjectileMovementComponent::TickComponent(float DeltaTime, enum ELevelTic
 		
 		// Logging
 		UE_LOG(LogProjectileMovement, Verbose, TEXT("Projectile %s: (Role: %d, Iteration %d, step %.3f, [%.3f / %.3f] cur/total) sim (Pos %s, Vel %s)"),
-			*GetNameSafe(ActorOwner), (int32)ActorOwner->Role, LoopCount, TimeTick, FMath::Max(0.f, DeltaTime - InitialTimeRemaining), DeltaTime,
+			*GetNameSafe(ActorOwner), (int32)ActorOwner->GetLocalRole(), LoopCount, TimeTick, FMath::Max(0.f, DeltaTime - InitialTimeRemaining), DeltaTime,
 			*UpdatedComponent->GetComponentLocation().ToString(), *Velocity.ToString());
 
 		// Initial move state
@@ -239,7 +239,7 @@ void UProjectileMovementComponent::TickComponent(float DeltaTime, enum ELevelTic
 
 			// Logging
 			UE_LOG(LogProjectileMovement, VeryVerbose, TEXT("Projectile %s: (Role: %d, Iteration %d, step %.3f) no hit (Pos %s, Vel %s)"),
-				*GetNameSafe(ActorOwner), (int32)ActorOwner->Role, LoopCount, TimeTick, *UpdatedComponent->GetComponentLocation().ToString(), *Velocity.ToString());
+				*GetNameSafe(ActorOwner), (int32)ActorOwner->GetLocalRole(), LoopCount, TimeTick, *UpdatedComponent->GetComponentLocation().ToString(), *Velocity.ToString());
 		}
 		else
 		{
@@ -252,7 +252,7 @@ void UProjectileMovementComponent::TickComponent(float DeltaTime, enum ELevelTic
 
 			// Logging
 			UE_CLOG(UpdatedComponent != nullptr, LogProjectileMovement, VeryVerbose, TEXT("Projectile %s: (Role: %d, Iteration %d, step %.3f) new hit at t=%.3f: (Pos %s, Vel %s)"),
-				*GetNameSafe(ActorOwner), (int32)ActorOwner->Role, LoopCount, TimeTick, Hit.Time, *UpdatedComponent->GetComponentLocation().ToString(), *Velocity.ToString());
+				*GetNameSafe(ActorOwner), (int32)ActorOwner->GetLocalRole(), LoopCount, TimeTick, Hit.Time, *UpdatedComponent->GetComponentLocation().ToString(), *Velocity.ToString());
 
 			// Handle blocking hit
 			NumImpacts++;
@@ -282,7 +282,7 @@ void UProjectileMovementComponent::TickComponent(float DeltaTime, enum ELevelTic
 			
 			// Logging
 			UE_CLOG(UpdatedComponent != nullptr, LogProjectileMovement, VeryVerbose, TEXT("Projectile %s: (Role: %d, Iteration %d, step %.3f) deflect at t=%.3f: (Pos %s, Vel %s)"),
-				*GetNameSafe(ActorOwner), (int32)ActorOwner->Role, Iterations, TimeTick, Hit.Time, *UpdatedComponent->GetComponentLocation().ToString(), *Velocity.ToString());
+				*GetNameSafe(ActorOwner), (int32)ActorOwner->GetLocalRole(), Iterations, TimeTick, Hit.Time, *UpdatedComponent->GetComponentLocation().ToString(), *Velocity.ToString());
 			
 			// Add unprocessed time after impact
 			if (SubTickTimeRemaining >= MIN_TICK_TIME)
@@ -296,7 +296,7 @@ void UProjectileMovementComponent::TickComponent(float DeltaTime, enum ELevelTic
 
 					// Logging
 					UE_LOG(LogProjectileMovement, Verbose, TEXT("Projectile %s: (Role: %d, Iteration %d, step %.3f) allowing extra iteration after bounce %u (t=%.3f, adding %.3f secs)"),
-						*GetNameSafe(ActorOwner), (int32)ActorOwner->Role, LoopCount, TimeTick, NumBounces, Hit.Time, SubTickTimeRemaining);
+						*GetNameSafe(ActorOwner), (int32)ActorOwner->GetLocalRole(), LoopCount, TimeTick, NumBounces, Hit.Time, SubTickTimeRemaining);
 				}
 			}
 		}
@@ -708,8 +708,8 @@ void UProjectileMovementComponent::SetInterpolatedComponent(USceneComponent* Com
 	{
 		ResetInterpolation();
 		InterpolatedComponentPtr = Component;
-		InterpInitialLocationOffset = Component->RelativeLocation;
-		InterpInitialRotationOffset = Component->RelativeRotation.Quaternion();
+		InterpInitialLocationOffset = Component->GetRelativeLocation();
+		InterpInitialRotationOffset = Component->GetRelativeRotation().Quaternion();
 		bInterpolationComplete = false;
 	}
 	else
@@ -740,7 +740,7 @@ void UProjectileMovementComponent::MoveInterpolationTarget(const FVector& NewLoc
 		if (USceneComponent* InterpComponent = GetInterpolatedComponent())
 		{
 			// Avoid moving the child, it will interpolate later
-			const FRotator InterpRelativeRotation = InterpComponent->RelativeRotation;
+			const FRotator InterpRelativeRotation = InterpComponent->GetRelativeRotation();
 			FScopedPreventAttachedComponentMove ScopedChildNoMove(InterpComponent);
 			
 			// Update interp offset
@@ -768,8 +768,8 @@ void UProjectileMovementComponent::MoveInterpolationTarget(const FVector& NewLoc
 			{
 				// If not interpolating rotation, we should allow the component to rotate.
 				// The absolute flag will get restored by the scoped move.
-				InterpComponent->bAbsoluteRotation = false;
-				InterpComponent->RelativeRotation = InterpRelativeRotation;
+				InterpComponent->SetUsingAbsoluteRotation(false);
+				InterpComponent->SetRelativeRotation_Direct(InterpRelativeRotation);
 				InterpRotationOffset = FQuat::Identity;
 			}
 
diff --git a/Engine/Source/Runtime/Engine/Private/Components/ReflectionCaptureComponent.cpp b/Engine/Source/Runtime/Engine/Private/Components/ReflectionCaptureComponent.cpp
--- a/Engine/Source/Runtime/Engine/Private/Components/ReflectionCaptureComponent.cpp
+++ b/Engine/Source/Runtime/Engine/Private/Components/ReflectionCaptureComponent.cpp
@@ -145,9 +145,9 @@ AReflectionCapture::AReflectionCapture(const FObjectInitializer& ObjectInitializ
 		static FConstructorStatics ConstructorStatics;
 
 		SpriteComponent->Sprite = ConstructorStatics.DecalTexture.Get();
-		SpriteComponent->RelativeScale3D = FVector(0.5f, 0.5f, 0.5f);
+		SpriteComponent->SetRelativeScale3D_Direct(FVector(0.5f, 0.5f, 0.5f));
 		SpriteComponent->bHiddenInGame = true;
-		SpriteComponent->bAbsoluteScale = true;
+		SpriteComponent->SetUsingAbsoluteScale(true);
 		SpriteComponent->SetCollisionProfileName(UCollisionProfile::NoCollision_ProfileName);
 		SpriteComponent->bIsScreenSizeScaled = true;
 	}
@@ -169,9 +169,9 @@ AReflectionCapture::AReflectionCapture(const FObjectInitializer& ObjectInitializ
 		static FConstructorStatics ConstructorStatics;
 
 		CaptureOffsetComponent->Sprite = ConstructorStatics.DecalTexture.Get();
-		CaptureOffsetComponent->RelativeScale3D = FVector(0.2f, 0.2f, 0.2f);
+		CaptureOffsetComponent->SetRelativeScale3D_Direct(FVector(0.2f, 0.2f, 0.2f));
 		CaptureOffsetComponent->bHiddenInGame = true;
-		CaptureOffsetComponent->bAbsoluteScale = true;
+		CaptureOffsetComponent->SetUsingAbsoluteScale(true);
 		CaptureOffsetComponent->SetCollisionProfileName(UCollisionProfile::NoCollision_ProfileName);
 		CaptureOffsetComponent->bIsScreenSizeScaled = true;
 	}
@@ -241,7 +241,7 @@ ABoxReflectionCapture::ABoxReflectionCapture(const FObjectInitializer& ObjectIni
 	: Super(ObjectInitializer.SetDefaultSubobjectClass<UBoxReflectionCaptureComponent>(TEXT("NewReflectionComponent")))
 {
 	UBoxReflectionCaptureComponent* BoxComponent = CastChecked<UBoxReflectionCaptureComponent>(GetCaptureComponent());
-	BoxComponent->RelativeScale3D = FVector(1000, 1000, 400);
+	BoxComponent->SetRelativeScale3D_Direct(FVector(1000, 1000, 400));
 	RootComponent = BoxComponent;
 #if WITH_EDITORONLY_DATA
 	if (GetSpriteComponent())
@@ -275,7 +275,7 @@ APlaneReflectionCapture::APlaneReflectionCapture(const FObjectInitializer& Objec
 	: Super(ObjectInitializer.SetDefaultSubobjectClass<UPlaneReflectionCaptureComponent>(TEXT("NewReflectionComponent")))
 {
 	UPlaneReflectionCaptureComponent* PlaneComponent = CastChecked<UPlaneReflectionCaptureComponent>(GetCaptureComponent());
-	PlaneComponent->RelativeScale3D = FVector(1, 1000, 1000);
+	PlaneComponent->SetRelativeScale3D_Direct(FVector(1, 1000, 1000));
 	RootComponent = PlaneComponent;
 #if WITH_EDITORONLY_DATA
 	if (GetSpriteComponent())
@@ -290,7 +290,7 @@ APlaneReflectionCapture::APlaneReflectionCapture(const FObjectInitializer& Objec
 	UDrawSphereComponent* DrawInfluenceRadius = CreateDefaultSubobject<UDrawSphereComponent>(TEXT("DrawRadius0"));
 	DrawInfluenceRadius->SetupAttachment(GetCaptureComponent());
 	DrawInfluenceRadius->bDrawOnlyIfSelected = true;
-	DrawInfluenceRadius->bAbsoluteScale = true;
+	DrawInfluenceRadius->SetUsingAbsoluteScale(true);
 	DrawInfluenceRadius->bUseEditorCompositing = true;
 	DrawInfluenceRadius->SetCollisionProfileName(UCollisionProfile::NoCollision_ProfileName);
 	PlaneComponent->PreviewInfluenceRadius = DrawInfluenceRadius;
@@ -948,7 +948,7 @@ void UReflectionCaptureComponent::UpdatePreviewShape()
 {
 	if (CaptureOffsetComponent)
 	{
-		CaptureOffsetComponent->RelativeLocation = CaptureOffset / GetComponentTransform().GetScale3D();
+		CaptureOffsetComponent->SetRelativeLocation_Direct(CaptureOffset / GetComponentTransform().GetScale3D());
 	}
 }
 
@@ -1027,7 +1027,7 @@ void UReflectionCaptureComponent::FinishDestroy()
 
 void UReflectionCaptureComponent::MarkDirtyForRecaptureOrUpload() 
 { 
-	if (bVisible)
+	if (GetVisibleFlag())
 	{
 		ReflectionCapturesToUpdate.AddUnique(this);
 		bNeedsRecaptureOrUpload = true; 
@@ -1036,7 +1036,7 @@ void UReflectionCaptureComponent::MarkDirtyForRecaptureOrUpload()
 
 void UReflectionCaptureComponent::MarkDirtyForRecapture() 
 { 
-	if (bVisible)
+	if (GetVisibleFlag())
 	{
 		MarkPackageDirty();
 		MapBuildDataId = FGuid::NewGuid();
diff --git a/Engine/Source/Runtime/Engine/Private/Components/SceneCaptureComponent.cpp b/Engine/Source/Runtime/Engine/Private/Components/SceneCaptureComponent.cpp
--- a/Engine/Source/Runtime/Engine/Private/Components/SceneCaptureComponent.cpp
+++ b/Engine/Source/Runtime/Engine/Private/Components/SceneCaptureComponent.cpp
@@ -732,9 +732,9 @@ APlanarReflection::APlanarReflection(const FObjectInitializer& ObjectInitializer
 		static FConstructorStatics ConstructorStatics;
 
 		SpriteComponent->Sprite = ConstructorStatics.DecalTexture.Get();
-		SpriteComponent->RelativeScale3D = FVector(0.5f, 0.5f, 0.5f);
+		SpriteComponent->SetRelativeScale3D_Direct(FVector(0.5f, 0.5f, 0.5f));
 		SpriteComponent->bHiddenInGame = true;
-		SpriteComponent->bAbsoluteScale = true;
+		SpriteComponent->SetUsingAbsoluteScale(true);
 		SpriteComponent->SetCollisionProfileName(UCollisionProfile::NoCollision_ProfileName);
 		SpriteComponent->bIsScreenSizeScaled = true;
 	}
@@ -833,7 +833,7 @@ void UPlanarReflectionComponent::OnRegister()
 	if (ProxyMeshComponent)
 	{
 		ProxyMeshComponent->SetMaterial(0, CaptureMaterial);
-		ProxyMeshComponent->bVisible = bShowPreviewPlane;
+		ProxyMeshComponent->SetVisibleFlag(bShowPreviewPlane);
 		ProxyMeshComponent->SetRelativeScale3D(FVector(4, 4, 1));
 	}
 #endif
@@ -911,7 +911,7 @@ void UPlanarReflectionComponent::PostEditChangeProperty(FPropertyChangedEvent& P
 
 	if (ProxyMeshComponent)
 	{
-		ProxyMeshComponent->bVisible = bShowPreviewPlane;
+		ProxyMeshComponent->SetVisibleFlag(bShowPreviewPlane);
 		ProxyMeshComponent->MarkRenderStateDirty();
 	}
 }
diff --git a/Engine/Source/Runtime/Engine/Private/Components/SkyAtmosphereComponent.cpp b/Engine/Source/Runtime/Engine/Private/Components/SkyAtmosphereComponent.cpp
--- a/Engine/Source/Runtime/Engine/Private/Components/SkyAtmosphereComponent.cpp
+++ b/Engine/Source/Runtime/Engine/Private/Components/SkyAtmosphereComponent.cpp
@@ -110,7 +110,7 @@ void USkyAtmosphereComponent::CreateRenderState_Concurrent()
 		bHidden = true;
 	}
 
-	if (bVisible && !bHidden &&
+	if (GetVisibleFlag() && !bHidden &&
 		ShouldComponentAddToScene() && ShouldRender() && IsRegistered() && (GetOuter() == NULL || !GetOuter()->HasAnyFlags(RF_ClassDefaultObject)))
 	{
 		// Create the scene proxy.
@@ -158,7 +158,7 @@ void USkyAtmosphereComponent::UpdateStaticLightingGUIDs()
 void USkyAtmosphereComponent::CheckForErrors()
 {
 	AActor* Owner = GetOwner();
-	if (Owner && bVisible)
+	if (Owner && GetVisibleFlag())
 	{
 		UWorld* ThisWorld = Owner->GetWorld();
 		bool bMultipleFound = false;
@@ -172,7 +172,7 @@ void USkyAtmosphereComponent::CheckForErrors()
 
 				if (Component != this
 					&& !Component->IsPendingKill()
-					&& Component->bVisible
+					&& Component->GetVisibleFlag()
 					&& Component->GetOwner()
 					&& ThisWorld->ContainsActor(Component->GetOwner())
 					&& !Component->GetOwner()->IsPendingKill())
@@ -186,7 +186,7 @@ void USkyAtmosphereComponent::CheckForErrors()
 				UAtmosphericFogComponent* Component = *ComponentIt;
 
 				if (!Component->IsPendingKill()
-					&& Component->bVisible
+					&& Component->GetVisibleFlag()
 					&& Component->GetOwner()
 					&& ThisWorld->ContainsActor(Component->GetOwner())
 					&& !Component->GetOwner()->IsPendingKill())
@@ -342,7 +342,7 @@ ASkyAtmosphere::ASkyAtmosphere(const FObjectInitializer& ObjectInitializer)
 		if (GetSpriteComponent())
 		{
 			GetSpriteComponent()->Sprite = ConstructorStatics.SkyAtmosphereTextureObject.Get();
-			GetSpriteComponent()->RelativeScale3D = FVector(0.5f, 0.5f, 0.5f);
+			GetSpriteComponent()->SetRelativeScale3D(FVector(0.5f, 0.5f, 0.5f));
 			GetSpriteComponent()->SpriteInfo.Category = ConstructorStatics.ID_SkyAtmosphere;
 			GetSpriteComponent()->SpriteInfo.DisplayName = ConstructorStatics.NAME_SkyAtmosphere;
 			GetSpriteComponent()->SetupAttachment(SkyAtmosphereComponent);
@@ -363,7 +363,7 @@ ASkyAtmosphere::ASkyAtmosphere(const FObjectInitializer& ObjectInitializer)
 #endif // WITH_EDITORONLY_DATA
 
 	PrimaryActorTick.bCanEverTick = true;
-	bHidden = false;
+	SetHidden(false);
 }
 
 
diff --git a/Engine/Source/Runtime/Engine/Private/Components/SkyLightComponent.cpp b/Engine/Source/Runtime/Engine/Private/Components/SkyLightComponent.cpp
--- a/Engine/Source/Runtime/Engine/Private/Components/SkyLightComponent.cpp
+++ b/Engine/Source/Runtime/Engine/Private/Components/SkyLightComponent.cpp
@@ -257,7 +257,7 @@ FSkyLightSceneProxy* USkyLightComponent::CreateSceneProxy() const
 
 void USkyLightComponent::SetCaptureIsDirty()
 { 
-	if (bVisible && bAffectsWorld)
+	if (GetVisibleFlag() && bAffectsWorld)
 	{
 		FScopeLock Lock(&SkyCapturesToUpdateLock);
 
@@ -314,7 +314,7 @@ void USkyLightComponent::SanitizeCubemapSize()
 
 void USkyLightComponent::SetBlendDestinationCaptureIsDirty()
 { 
-	if (bVisible && bAffectsWorld && BlendDestinationCubemap)
+	if (GetVisibleFlag() && bAffectsWorld && BlendDestinationCubemap)
 	{
 		SkyCapturesToUpdateBlendDestinations.AddUnique(this); 
 
@@ -343,7 +343,7 @@ void USkyLightComponent::CreateRenderState_Concurrent()
 
 	const bool bIsValid = SourceType != SLS_SpecifiedCubemap || Cubemap != NULL;
 
-	if (bAffectsWorld && bVisible && !bHidden && bIsValid)
+	if (bAffectsWorld && GetVisibleFlag() && !bHidden && bIsValid)
 	{
 		// Create the light's scene proxy.
 		SceneProxy = CreateSceneProxy();
@@ -377,7 +377,7 @@ void USkyLightComponent::PostLoad()
 	SanitizeCubemapSize();
 
 	// All components are queued for update on creation by default, remove if needed
-	if (!bVisible || HasAnyFlags(RF_ClassDefaultObject | RF_ArchetypeObject))
+	if (!GetVisibleFlag() || HasAnyFlags(RF_ClassDefaultObject | RF_ArchetypeObject))
 	{
 		FScopeLock Lock(&SkyCapturesToUpdateLock);
 		SkyCapturesToUpdate.Remove(this);
@@ -524,7 +524,7 @@ void USkyLightComponent::CheckForErrors()
 {
 	AActor* Owner = GetOwner();
 
-	if (Owner && bVisible && bAffectsWorld)
+	if (Owner && GetVisibleFlag() && bAffectsWorld)
 	{
 		UWorld* ThisWorld = Owner->GetWorld();
 		bool bMultipleFound = false;
@@ -537,7 +537,7 @@ void USkyLightComponent::CheckForErrors()
 
 				if (Component != this 
 					&& !Component->IsPendingKill()
-					&& Component->bVisible
+					&& Component->GetVisibleFlag()
 					&& Component->bAffectsWorld
 					&& Component->GetOwner() 
 					&& ThisWorld->ContainsActor(Component->GetOwner())
@@ -906,7 +906,7 @@ void USkyLightComponent::OnVisibilityChanged()
 {
 	Super::OnVisibilityChanged();
 
-	if (bVisible && !bHasEverCaptured)
+	if (GetVisibleFlag() && !bHasEverCaptured)
 	{
 		// Capture if we are being enabled for the first time
 		SetCaptureIsDirty();
diff --git a/Engine/Source/Runtime/Engine/Private/Components/StereoLayerComponent.cpp b/Engine/Source/Runtime/Engine/Private/Components/StereoLayerComponent.cpp
--- a/Engine/Source/Runtime/Engine/Private/Components/StereoLayerComponent.cpp
+++ b/Engine/Source/Runtime/Engine/Private/Components/StereoLayerComponent.cpp
@@ -67,12 +67,12 @@ void UStereoLayerComponent::TickComponent(float DeltaTime, enum ELevelTick TickT
 	}
 	
 	// If the transform changed dirty the layer and push the new transform
-	if (!bIsDirty && (bLastVisible != bVisible || FMemory::Memcmp(&LastTransform, &Transform, sizeof(Transform)) != 0))
+	if (!bIsDirty && (bLastVisible != GetVisibleFlag() || FMemory::Memcmp(&LastTransform, &Transform, sizeof(Transform)) != 0))
 	{
 		bIsDirty = true;
 	}
 
-	bool bCurrVisible = bVisible;
+	bool bCurrVisible = GetVisibleFlag();
 	if (!Texture || !Texture->Resource)
 	{
 		bCurrVisible = false;
diff --git a/Engine/Source/Runtime/Engine/Private/Components/TextRenderComponent.cpp b/Engine/Source/Runtime/Engine/Private/Components/TextRenderComponent.cpp
--- a/Engine/Source/Runtime/Engine/Private/Components/TextRenderComponent.cpp
+++ b/Engine/Source/Runtime/Engine/Private/Components/TextRenderComponent.cpp
@@ -50,10 +50,10 @@ ATextRenderActor::ATextRenderActor(const FObjectInitializer& ObjectInitializer)
 		static FConstructorStatics ConstructorStatics;
 
 		SpriteComponent->Sprite = ConstructorStatics.TextRenderTexture.Get();
-		SpriteComponent->RelativeScale3D = FVector(0.5f, 0.5f, 0.5f);
+		SpriteComponent->SetRelativeScale3D_Direct(FVector(0.5f, 0.5f, 0.5f));
 		SpriteComponent->SetupAttachment(TextRender);
 		SpriteComponent->bIsScreenSizeScaled = true;
-		SpriteComponent->bAbsoluteScale = true;
+		SpriteComponent->SetUsingAbsoluteScale(true);
 		SpriteComponent->bReceivesDecals = false;
 	}
 #endif
diff --git a/Engine/Source/Runtime/Engine/Private/Controller.cpp b/Engine/Source/Runtime/Engine/Private/Controller.cpp
--- a/Engine/Source/Runtime/Engine/Private/Controller.cpp
+++ b/Engine/Source/Runtime/Engine/Private/Controller.cpp
@@ -35,7 +35,7 @@ AController::AController(const FObjectInitializer& ObjectInitializer)
 {
 
 	PrimaryActorTick.bCanEverTick = true;
-	bHidden = true;
+	SetHidden(true);
 #if WITH_EDITORONLY_DATA
 	bHiddenEd = true;
 #endif // WITH_EDITORONLY_DATA
@@ -44,7 +44,7 @@ AController::AController(const FObjectInitializer& ObjectInitializer)
 	TransformComponent = CreateDefaultSubobject<USceneComponent>(TEXT("TransformComponent0"));
 	RootComponent = TransformComponent;
 
-	bCanBeDamaged = false;
+	SetCanBeDamaged(false);
 	bAttachToPawn = false;
 	bIsPlayerController = false;
 	bCanPossessWithoutAuthority = false;
@@ -53,7 +53,7 @@ AController::AController(const FObjectInitializer& ObjectInitializer)
 	{
 		// We attach the RootComponent to the pawn for location updates,
 		// but we want to drive rotation with ControlRotation regardless of attachment state.
-		RootComponent->bAbsoluteRotation = true;
+		RootComponent->SetUsingAbsoluteRotation(true);
 	}
 }
 
@@ -72,13 +72,13 @@ bool AController::IsLocalController() const
 		return true;
 	}
 	
-	if (NetMode == NM_Client && Role == ROLE_AutonomousProxy)
+	if (NetMode == NM_Client && GetLocalRole() == ROLE_AutonomousProxy)
 	{
 		// Networked client in control.
 		return true;
 	}
 
-	if (GetRemoteRole() != ROLE_AutonomousProxy && Role == ROLE_Authority)
+	if (GetRemoteRole() != ROLE_AutonomousProxy && GetLocalRole() == ROLE_Authority)
 	{
 		// Local authority in control.
 		return true;
@@ -117,7 +117,7 @@ void AController::SetControlRotation(const FRotator& NewRotation)
 	{
 		ControlRotation = NewRotation;
 
-		if (RootComponent && RootComponent->bAbsoluteRotation)
+		if (RootComponent && RootComponent->IsUsingAbsoluteRotation())
 		{
 			RootComponent->SetWorldRotation(GetControlRotation());
 		}
@@ -267,7 +267,7 @@ void AController::PostInitializeComponents()
 
 		// Since we avoid updating rotation in SetControlRotation() if it hasn't changed,
 		// we should make sure that the initial RootComponent rotation matches it if ControlRotation was set directly.
-		if (RootComponent && RootComponent->bAbsoluteRotation)
+		if (RootComponent && RootComponent->IsUsingAbsoluteRotation())
 		{
 			RootComponent->SetWorldRotation(GetControlRotation());
 		}
@@ -476,7 +476,7 @@ void AController::OnRep_PlayerState()
 
 void AController::Destroyed()
 {
-	if (Role == ROLE_Authority && PlayerState != NULL)
+	if (GetLocalRole() == ROLE_Authority && PlayerState != NULL)
 	{
 		// if we are a player, log out
 		AGameModeBase* const GameMode = GetWorld()->GetAuthGameMode();
@@ -526,7 +526,7 @@ void AController::InitPlayerState()
 		{
 			FActorSpawnParameters SpawnInfo;
 			SpawnInfo.Owner = this;
-			SpawnInfo.Instigator = Instigator;
+			SpawnInfo.Instigator = GetInstigator();
 			SpawnInfo.SpawnCollisionHandlingOverride = ESpawnActorCollisionHandlingMethod::AlwaysSpawn;
 			SpawnInfo.ObjectFlags |= RF_Transient;	// We never want player states to save into a map
 
diff --git a/Engine/Source/Runtime/Engine/Private/DataChannel.cpp b/Engine/Source/Runtime/Engine/Private/DataChannel.cpp
--- a/Engine/Source/Runtime/Engine/Private/DataChannel.cpp
+++ b/Engine/Source/Runtime/Engine/Private/DataChannel.cpp
@@ -1992,7 +1992,7 @@ bool UActorChannel::CleanUp(const bool bForDestroy, EChannelCloseReason CloseRea
 			{
 				if (!bTornOff)
 				{
-					Actor->Role = ROLE_Authority;
+					Actor->SetRole(ROLE_Authority);
 					Actor->SetReplicates(false);
 					bTornOff = true;
 					if (Actor->GetWorld() != NULL && !IsEngineExitRequested())
@@ -2937,7 +2937,7 @@ int64 UActorChannel::ReplicateActor()
 	FScopedRoleDowngrade ScopedRoleDowngrade( Actor, RepFlags );
 
 	RepFlags.bNetSimulated	= (Actor->GetRemoteRole() == ROLE_SimulatedProxy);
-	RepFlags.bRepPhysics	= Actor->ReplicatedMovement.bRepPhysics;
+	RepFlags.bRepPhysics	= Actor->GetReplicatedMovement().bRepPhysics;
 	RepFlags.bReplay		= bReplay;
 	//RepFlags.bNetInitial	= RepFlags.bNetInitial;
 
@@ -3090,7 +3090,7 @@ FString UActorChannel::Describe()
 	}
 	else
 	{
-		return FString::Printf(TEXT("[UActorChannel] Actor: %s, Role: %i, RemoteRole: %i %s"), *Actor->GetFullName(), ( int32 )Actor->Role, ( int32 )Actor->GetRemoteRole(), *UChannel::Describe());
+		return FString::Printf(TEXT("[UActorChannel] Actor: %s, Role: %i, RemoteRole: %i %s"), *Actor->GetFullName(), ( int32 )Actor->GetLocalRole(), ( int32 )Actor->GetRemoteRole(), *UChannel::Describe());
 	}
 }
 
diff --git a/Engine/Source/Runtime/Engine/Private/DebugCameraController.cpp b/Engine/Source/Runtime/Engine/Private/DebugCameraController.cpp
--- a/Engine/Source/Runtime/Engine/Private/DebugCameraController.cpp
+++ b/Engine/Source/Runtime/Engine/Private/DebugCameraController.cpp
@@ -43,7 +43,7 @@ ADebugCameraController::ADebugCameraController(const FObjectInitializer& ObjectI
 
 	bIsFrozenRendering = false;
 	DrawFrustum = nullptr;
-	bHidden = false;
+	SetHidden(false);
 #if WITH_EDITORONLY_DATA
 	bHiddenEd = false;
 #endif // WITH_EDITORONLY_DATA
@@ -399,7 +399,7 @@ void ADebugCameraController::PostInitializeComponents()
 	}
 	FActorSpawnParameters SpawnInfo;
 	SpawnInfo.Owner = this;
-	SpawnInfo.Instigator = Instigator;
+	SpawnInfo.Instigator = GetInstigator();
 	SpawnInfo.SpawnCollisionHandlingOverride = ESpawnActorCollisionHandlingMethod::AlwaysSpawn;
 	SpawnInfo.ObjectFlags |= RF_Transient;	// We never want these to save into a map
 	MyHUD = GetWorld()->SpawnActor<ADebugCameraHUD>( SpawnInfo );
diff --git a/Engine/Source/Runtime/Engine/Private/DebugCameraHUD.cpp b/Engine/Source/Runtime/Engine/Private/DebugCameraHUD.cpp
--- a/Engine/Source/Runtime/Engine/Private/DebugCameraHUD.cpp
+++ b/Engine/Source/Runtime/Engine/Private/DebugCameraHUD.cpp
@@ -30,7 +30,7 @@ static TAutoConsoleVariable<int32> CVarDebugCameraTraceComplex(
 ADebugCameraHUD::ADebugCameraHUD(const FObjectInitializer& ObjectInitializer)
 	: Super(ObjectInitializer)
 {
-	bHidden = false;
+	SetHidden(false);
 }
 
 bool ADebugCameraHUD::DisplayMaterials( float X, float& Y, float DY, UMeshComponent* MeshComp )
diff --git a/Engine/Source/Runtime/Engine/Private/DecalActor.cpp b/Engine/Source/Runtime/Engine/Private/DecalActor.cpp
--- a/Engine/Source/Runtime/Engine/Private/DecalActor.cpp
+++ b/Engine/Source/Runtime/Engine/Private/DecalActor.cpp
@@ -25,7 +25,7 @@ ADecalActor::ADecalActor(const FObjectInitializer& ObjectInitializer)
 	: Super(ObjectInitializer)
 {
 	Decal = CreateDefaultSubobject<UDecalComponent>(TEXT("NewDecalComponent"));
-	Decal->RelativeRotation = FRotator(-90, 0, 0);
+	Decal->SetRelativeRotation(FRotator(-90, 0, 0));
 	Decal->bDestroyOwnerAfterFade = true;
 
 	RootComponent = Decal;
@@ -59,25 +59,25 @@ ADecalActor::ADecalActor(const FObjectInitializer& ObjectInitializer)
 			ArrowComponent->SpriteInfo.Category = ConstructorStatics.ID_Decals;
 			ArrowComponent->SpriteInfo.DisplayName = ConstructorStatics.NAME_Decals;
 			ArrowComponent->SetupAttachment(Decal);
-			ArrowComponent->bAbsoluteScale = true;
+			ArrowComponent->SetUsingAbsoluteScale(true);
 			ArrowComponent->bIsScreenSizeScaled = true;
 		}
 
 		if (SpriteComponent)
 		{
 			SpriteComponent->Sprite = ConstructorStatics.DecalTexture.Get();
-			SpriteComponent->RelativeScale3D = FVector(0.5f, 0.5f, 0.5f);
+			SpriteComponent->SetRelativeScale3D(FVector(0.5f, 0.5f, 0.5f));
 			SpriteComponent->SpriteInfo.Category = ConstructorStatics.ID_Decals;
 			SpriteComponent->SpriteInfo.DisplayName = ConstructorStatics.NAME_Decals;
 			SpriteComponent->SetupAttachment(Decal);
 			SpriteComponent->bIsScreenSizeScaled = true;
-			SpriteComponent->bAbsoluteScale = true;
+			SpriteComponent->SetUsingAbsoluteScale(true);
 			SpriteComponent->bReceivesDecals = false;
 		}
 	}
 #endif // WITH_EDITORONLY_DATA
 
-	bCanBeDamaged = false;
+	SetCanBeDamaged(false);
 }
 
 #if WITH_EDITOR
@@ -162,7 +162,7 @@ void ADecalActor::Serialize(FArchive& Ar)
 		if(Decal)
 		{
 			// before Super::Serialize(Ar);
-			Decal->RelativeScale3D = FVector(128.0f, 256.0f, 256.0f);
+			Decal->SetRelativeScale3D_Direct(FVector(128.0f, 256.0f, 256.0f));
 		}
 	}
 
diff --git a/Engine/Source/Runtime/Engine/Private/DefaultPawn.cpp b/Engine/Source/Runtime/Engine/Private/DefaultPawn.cpp
--- a/Engine/Source/Runtime/Engine/Private/DefaultPawn.cpp
+++ b/Engine/Source/Runtime/Engine/Private/DefaultPawn.cpp
@@ -21,7 +21,7 @@ FName ADefaultPawn::MeshComponentName(TEXT("MeshComponent0"));
 ADefaultPawn::ADefaultPawn(const FObjectInitializer& ObjectInitializer)
 	: Super(ObjectInitializer)
 {
-	bCanBeDamaged = true;
+	SetCanBeDamaged(true);
 
 	SetRemoteRoleForBackwardsCompat(ROLE_SimulatedProxy);
 	bReplicates = true;
diff --git a/Engine/Source/Runtime/Engine/Private/DemoNetDriver.cpp b/Engine/Source/Runtime/Engine/Private/DemoNetDriver.cpp
--- a/Engine/Source/Runtime/Engine/Private/DemoNetDriver.cpp
+++ b/Engine/Source/Runtime/Engine/Private/DemoNetDriver.cpp
@@ -4366,14 +4366,14 @@ void UDemoNetDriver::RespawnNecessaryNetStartupActors(TArray<AActor*>& SpawnedAc
 
 			if (RepLayout.IsValid() && ReceivingRepState && bSanityCheckReferences)
 			{
-				const ENetRole SavedRole = Actor->Role;
+				const ENetRole SavedRole = Actor->GetLocalRole();
 
 				FRepObjectDataBuffer ActorData(Actor);
 				FConstRepShadowDataBuffer ShadowData(ReceivingRepState->StaticBuffer.GetData());
 
 				RepLayout->DiffStableProperties(&ReceivingRepState->RepNotifies, nullptr, ActorData, ShadowData);
 
-				Actor->Role = SavedRole;
+				Actor->SetRole(SavedRole);
 			}
 
 			check(Actor->GetRemoteRole() != ROLE_Authority);
@@ -4382,7 +4382,7 @@ void UDemoNetDriver::RespawnNecessaryNetStartupActors(TArray<AActor*>& SpawnedAc
 
 			UGameplayStatics::FinishSpawningActor(Actor, SpawnTransform);
 
-			if (Actor->Role == ROLE_Authority)
+			if (Actor->GetLocalRole() == ROLE_Authority)
 			{
 				Actor->SwapRoles();
 			}
@@ -5648,7 +5648,7 @@ void UDemoNetDriver::RestoreConnectionPostScrub(APlayerController* PC, UNetConne
 	check(NetConnection != nullptr);
 	check(PC != nullptr);
 
-	PC->Role = ROLE_AutonomousProxy;
+	PC->SetRole(ROLE_AutonomousProxy);
 	PC->NetConnection = NetConnection;
 	NetConnection->LastReceiveTime = Time;
 	NetConnection->LastReceiveRealtime = FPlatformTime::Seconds();
@@ -5954,7 +5954,7 @@ TSharedPtr<FObjectReplicator> UDemoNetConnection::CreateReplicatorForNewActorCha
 
 		// Need to swap roles for the startup actor since in the CDO they aren't swapped, and the CDO just
 		// overwrote the actor state.
-		if (Actor && (Actor->Role == ROLE_Authority))
+		if (Actor && (Actor->GetLocalRole() == ROLE_Authority))
 		{
 			Actor->SwapRoles();
 		}
diff --git a/Engine/Source/Runtime/Engine/Private/GameFramework/RootMotionSource.cpp b/Engine/Source/Runtime/Engine/Private/GameFramework/RootMotionSource.cpp
--- a/Engine/Source/Runtime/Engine/Private/GameFramework/RootMotionSource.cpp
+++ b/Engine/Source/Runtime/Engine/Private/GameFramework/RootMotionSource.cpp
@@ -32,7 +32,7 @@ void RootMotionSourceDebug::PrintOnScreen(const ACharacter& InCharacter, const F
 		const FString AdjustedDebugString = FString::Printf(TEXT("[%d] [%s] %s"), (uint64)GFrameCounter, *InCharacter.GetName(), *InString);
 
 		// If on the server, replicate this message to everyone.
-		if (!InCharacter.IsLocallyControlled() && (InCharacter.Role == ROLE_Authority))
+		if (!InCharacter.IsLocallyControlled() && (InCharacter.GetLocalRole() == ROLE_Authority))
 		{
 			for (FConstPlayerControllerIterator Iterator = InCharacter.GetWorld()->GetPlayerControllerIterator(); Iterator; ++Iterator)
 			{
@@ -1197,7 +1197,7 @@ void FRootMotionSource_JumpForce::PrepareRootMotion
 			DrawDebugCapsule(Character.GetWorld(), DestinationLocation + LocDiff, Character.GetSimpleCollisionHalfHeight(), Character.GetSimpleCollisionRadius(), FQuat::Identity, FColor::White, true, DebugLifetime);
 
 			UE_LOG(LogRootMotion, VeryVerbose, TEXT("RootMotionJumpForce %s %s preparing from %f to %f from (%s) to (%s) resulting force %s"), 
-				Character.Role == ROLE_AutonomousProxy ? TEXT("AUTONOMOUS") : TEXT("AUTHORITY"),
+				Character.GetLocalRole() == ROLE_AutonomousProxy ? TEXT("AUTONOMOUS") : TEXT("AUTHORITY"),
 				Character.bClientUpdating ? TEXT("UPD") : TEXT("NOR"),
 				GetTime(), GetTime() + SimulationTime, 
 				*CurrentLocation.ToString(), *CurrentTargetLocation.ToString(), 
@@ -1475,7 +1475,7 @@ void FRootMotionSourceGroup::PrepareRootMotion(float DeltaTime, const ACharacter
 							if (bRootMotionHasNotStarted && bRootMotionHasValidStartTime)
 							{
 								float CharacterMovementTime = -1.f;
-								if (Character.Role == ROLE_AutonomousProxy)
+								if (Character.GetLocalRole() == ROLE_AutonomousProxy)
 								{
 									const FNetworkPredictionData_Client_Character* ClientData = MoveComponent.HasPredictionData_Client() ? static_cast<FNetworkPredictionData_Client_Character*>(MoveComponent.GetPredictionData_Client()) : nullptr;
 									if (ClientData)
@@ -1495,7 +1495,7 @@ void FRootMotionSourceGroup::PrepareRootMotion(float DeltaTime, const ACharacter
 										}
 									}
 								}
-								else if (Character.Role == ROLE_Authority)
+								else if (Character.GetLocalRole() == ROLE_Authority)
 								{
 									const FNetworkPredictionData_Server_Character* ServerData = MoveComponent.HasPredictionData_Server() ? static_cast<FNetworkPredictionData_Server_Character*>(MoveComponent.GetPredictionData_Server()) : nullptr;
 									if (ServerData)
diff --git a/Engine/Source/Runtime/Engine/Private/GameFramework/SpringArmComponent.cpp b/Engine/Source/Runtime/Engine/Private/GameFramework/SpringArmComponent.cpp
--- a/Engine/Source/Runtime/Engine/Private/GameFramework/SpringArmComponent.cpp
+++ b/Engine/Source/Runtime/Engine/Private/GameFramework/SpringArmComponent.cpp
@@ -64,21 +64,22 @@ FRotator USpringArmComponent::GetTargetRotation() const
 	}
 
 	// If inheriting rotation, check options for which components to inherit
-	if (!bAbsoluteRotation)
+	if (!IsUsingAbsoluteRotation())
 	{
+		const FRotator LocalRelativeRotation = GetRelativeRotation();
 		if (!bInheritPitch)
 		{
-			DesiredRot.Pitch = RelativeRotation.Pitch;
+			DesiredRot.Pitch = LocalRelativeRotation.Pitch;
 		}
 
 		if (!bInheritYaw)
 		{
-			DesiredRot.Yaw = RelativeRotation.Yaw;
+			DesiredRot.Yaw = LocalRelativeRotation.Yaw;
 		}
 
 		if (!bInheritRoll)
 		{
-			DesiredRot.Roll = RelativeRotation.Roll;
+			DesiredRot.Roll = LocalRelativeRotation.Roll;
 		}
 	}
 
diff --git a/Engine/Source/Runtime/Engine/Private/GameInstance.cpp b/Engine/Source/Runtime/Engine/Private/GameInstance.cpp
--- a/Engine/Source/Runtime/Engine/Private/GameInstance.cpp
+++ b/Engine/Source/Runtime/Engine/Private/GameInstance.cpp
@@ -742,7 +742,7 @@ bool UGameInstance::RemoveLocalPlayer(ULocalPlayer* ExistingPlayer)
 	// FIXME: Notify server we want to leave the game if this is an online game
 	if (ExistingPlayer->PlayerController != nullptr)
 	{
-		bool bShouldRemovePlayer = (ExistingPlayer->PlayerController->Role == ROLE_Authority);
+		bool bShouldRemovePlayer = (ExistingPlayer->PlayerController->GetLocalRole() == ROLE_Authority);
 
 		// FIXME: Do this all inside PlayerRemoved?
 		ExistingPlayer->PlayerController->CleanupGameViewport();
diff --git a/Engine/Source/Runtime/Engine/Private/GameModeBase.cpp b/Engine/Source/Runtime/Engine/Private/GameModeBase.cpp
--- a/Engine/Source/Runtime/Engine/Private/GameModeBase.cpp
+++ b/Engine/Source/Runtime/Engine/Private/GameModeBase.cpp
@@ -68,7 +68,7 @@ void AGameModeBase::InitGame(const FString& MapName, const FString& Options, FSt
 	OptionsString = Options;
 
 	FActorSpawnParameters SpawnInfo;
-	SpawnInfo.Instigator = Instigator;
+	SpawnInfo.Instigator = GetInstigator();
 	SpawnInfo.ObjectFlags |= RF_Transient;	// We never want to save game sessions into a map
 	GameSession = World->SpawnActor<AGameSession>(GetGameSessionClass(), SpawnInfo);
 	GameSession->InitOptions(Options);
@@ -99,7 +99,7 @@ void AGameModeBase::PreInitializeComponents()
 	Super::PreInitializeComponents();
 
 	FActorSpawnParameters SpawnInfo;
-	SpawnInfo.Instigator = Instigator;
+	SpawnInfo.Instigator = GetInstigator();
 	SpawnInfo.ObjectFlags |= RF_Transient;	// We never want to save game states or network managers into a map									
 											
 	// Fallback to default GameState if none was specified.
@@ -694,7 +694,7 @@ APlayerController* AGameModeBase::SpawnReplayPlayerController(ENetRole InRemoteR
 APlayerController* AGameModeBase::SpawnPlayerControllerCommon(ENetRole InRemoteRole, FVector const& SpawnLocation, FRotator const& SpawnRotation, TSubclassOf<APlayerController> InPlayerControllerClass)
 {
 	FActorSpawnParameters SpawnInfo;
-	SpawnInfo.Instigator = Instigator;
+	SpawnInfo.Instigator = GetInstigator();
 	SpawnInfo.ObjectFlags |= RF_Transient;	// We never want to save player controllers into a map
 	SpawnInfo.bDeferConstruction = true;
 	APlayerController* NewPC = GetWorld()->SpawnActor<APlayerController>(InPlayerControllerClass, SpawnLocation, SpawnRotation, SpawnInfo);
@@ -1162,7 +1162,7 @@ APawn* AGameModeBase::SpawnDefaultPawnFor_Implementation(AController* NewPlayer,
 APawn* AGameModeBase::SpawnDefaultPawnAtTransform_Implementation(AController* NewPlayer, const FTransform& SpawnTransform)
 {
 	FActorSpawnParameters SpawnInfo;
-	SpawnInfo.Instigator = Instigator;
+	SpawnInfo.Instigator = GetInstigator();
 	SpawnInfo.ObjectFlags |= RF_Transient;	// We never want to save default player pawns into a map
 	UClass* PawnClass = GetDefaultPawnClassForController(NewPlayer);
 	APawn* ResultPawn = GetWorld()->SpawnActor<APawn>(PawnClass, SpawnTransform, SpawnInfo);
diff --git a/Engine/Source/Runtime/Engine/Private/GameState.cpp b/Engine/Source/Runtime/Engine/Private/GameState.cpp
--- a/Engine/Source/Runtime/Engine/Private/GameState.cpp
+++ b/Engine/Source/Runtime/Engine/Private/GameState.cpp
@@ -57,7 +57,7 @@ void AGameState::PostInitializeComponents()
 
 void AGameState::HandleMatchIsWaitingToStart()
 {
-	if (Role != ROLE_Authority)
+	if (GetLocalRole() != ROLE_Authority)
 	{
 		// Server handles this in AGameMode::HandleMatchIsWaitingToStart
 		GetWorldSettings()->NotifyBeginPlay();
@@ -68,7 +68,7 @@ void AGameState::HandleMatchIsWaitingToStart()
 
 void AGameState::HandleMatchHasStarted()
 {
-	if (Role != ROLE_Authority)
+	if (GetLocalRole() != ROLE_Authority)
 	{
 		// Server handles this in AGameMode::HandleMatchHasStarted
 		GetWorldSettings()->NotifyMatchStarted();
@@ -127,7 +127,7 @@ bool AGameState::HasMatchEnded() const
 
 void AGameState::SetMatchState(FName NewState)
 {
-	if (Role == ROLE_Authority)
+	if (GetLocalRole() == ROLE_Authority)
 	{
 		UE_LOG(LogGameState, Log, TEXT("Match State Changed from %s to %s"), *MatchState.ToString(), *NewState.ToString());
 
diff --git a/Engine/Source/Runtime/Engine/Private/GameStateBase.cpp b/Engine/Source/Runtime/Engine/Private/GameStateBase.cpp
--- a/Engine/Source/Runtime/Engine/Private/GameStateBase.cpp
+++ b/Engine/Source/Runtime/Engine/Private/GameStateBase.cpp
@@ -21,7 +21,7 @@ AGameStateBase::AGameStateBase(const FObjectInitializer& ObjectInitializer)
 	SetRemoteRoleForBackwardsCompat(ROLE_SimulatedProxy);
 	bReplicates = true;
 	bAlwaysRelevant = true;
-	bReplicateMovement = false;
+	SetReplicatingMovement(false);
 
 	// Note: this is very important to set to false. Though all replication infos are spawned at run time, during seamless travel
 	// they are held on to and brought over into the new world. In ULevel::InitializeNetworkActors, these PlayerStates may be treated as map/startup actors
@@ -51,7 +51,7 @@ void AGameStateBase::PostInitializeComponents()
 	World->SetGameState(this);
 
 	FTimerManager& TimerManager = GetWorldTimerManager();
-	if (World->IsGameWorld() && Role == ROLE_Authority)
+	if (World->IsGameWorld() && GetLocalRole() == ROLE_Authority)
 	{
 		UpdateServerTimeSeconds();
 		if (ServerWorldTimeSecondsUpdateFrequency > 0.f)
@@ -164,7 +164,7 @@ void AGameStateBase::OnRep_ReplicatedWorldTimeSeconds()
 
 void AGameStateBase::OnRep_ReplicatedHasBegunPlay()
 {
-	if (bReplicatedHasBegunPlay && Role != ROLE_Authority)
+	if (bReplicatedHasBegunPlay && GetLocalRole() != ROLE_Authority)
 	{
 		GetWorldSettings()->NotifyBeginPlay();
 		GetWorldSettings()->NotifyMatchStarted();
diff --git a/Engine/Source/Runtime/Engine/Private/GameViewportClient.cpp b/Engine/Source/Runtime/Engine/Private/GameViewportClient.cpp
--- a/Engine/Source/Runtime/Engine/Private/GameViewportClient.cpp
+++ b/Engine/Source/Runtime/Engine/Private/GameViewportClient.cpp
@@ -3008,12 +3008,14 @@ bool UGameViewportClient::HandleShowLayerCommand( const TCHAR* Cmd, FOutputDevic
 
 			if (Actor->Layers.Contains(LayerFName))
 			{
+				const bool bHiddenLocal = Actor->IsHidden();
+
 				// look for always toggle, or a set when it's unset, etc
-				if ((SetMode == -1) || (SetMode == 0 && !Actor->bHidden) || (SetMode != 0 && Actor->bHidden))
+				if ((SetMode == -1) || (SetMode == 0 && !bHiddenLocal) || (SetMode != 0 && bHiddenLocal))
 				{
 					NumActorsToggled++;
 					// Note: overriding existing hidden property, ideally this would be something orthogonal
-					Actor->bHidden = !Actor->bHidden;
+					Actor->SetHidden(!bHiddenLocal);
 
 					Actor->MarkComponentsRenderStateDirty();
 				}
diff --git a/Engine/Source/Runtime/Engine/Private/GameplayStatics.cpp b/Engine/Source/Runtime/Engine/Private/GameplayStatics.cpp
--- a/Engine/Source/Runtime/Engine/Private/GameplayStatics.cpp
+++ b/Engine/Source/Runtime/Engine/Private/GameplayStatics.cpp
@@ -528,7 +528,7 @@ bool UGameplayStatics::ApplyRadialDamageWithFalloff(const UObject* WorldContextO
 		AActor* const OverlapActor = Overlap.GetActor();
 
 		if (OverlapActor &&
-			OverlapActor->bCanBeDamaged &&
+			OverlapActor->CanBeDamaged() &&
 			OverlapActor != DamageCauser &&
 			Overlap.Component.IsValid())
 		{
@@ -658,7 +658,7 @@ class AActor* UGameplayStatics::BeginDeferredActorSpawnFromClass(const UObject*
 		{
 			if (AActor* ContextActor = Cast<AActor>(MutableWorldContextObject))
 			{
-				AutoInstigator = ContextActor->Instigator;
+				AutoInstigator = ContextActor->GetInstigator();
 			}
 		}
 
@@ -996,12 +996,13 @@ UParticleSystemComponent* UGameplayStatics::InternalSpawnEmitterAtLocation(UWorl
 
 	if (UParticleSystemComponent* PSC = CreateParticleSystem(EmitterTemplate, World, World->GetWorldSettings(), bAutoDestroy, PoolingMethod))
 	{
-		PSC->bAbsoluteLocation = true;
-		PSC->bAbsoluteRotation = true;
-		PSC->bAbsoluteScale = true;
-		PSC->RelativeLocation = SpawnLocation;
-		PSC->RelativeRotation = SpawnRotation;
-		PSC->RelativeScale3D = SpawnScale;
+
+		PSC->SetUsingAbsoluteLocation(true);
+		PSC->SetUsingAbsoluteRotation(true);
+		PSC->SetUsingAbsoluteScale(true);
+		PSC->SetRelativeLocation_Direct(SpawnLocation);
+		PSC->SetRelativeRotation_Direct(SpawnRotation);
+		PSC->SetRelativeScale3D_Direct(SpawnScale);
 
 		PSC->RegisterComponentWithWorld(World);
 		if (bAutoActivateSystem)
@@ -1081,25 +1082,25 @@ UParticleSystemComponent* UGameplayStatics::SpawnEmitterAttached(UParticleSystem
 						const FTransform ParentToWorld = AttachToComponent->GetSocketTransform(AttachPointName);
 						const FTransform ComponentToWorld(Rotation, Location, Scale);
 						const FTransform RelativeTM = ComponentToWorld.GetRelativeTransform(ParentToWorld);
-						PSC->RelativeLocation = RelativeTM.GetLocation();
-						PSC->RelativeRotation = RelativeTM.GetRotation().Rotator();
-						PSC->RelativeScale3D = RelativeTM.GetScale3D();
+						PSC->SetRelativeLocation_Direct(RelativeTM.GetLocation());
+						PSC->SetRelativeRotation_Direct(RelativeTM.GetRotation().Rotator());
+						PSC->SetRelativeScale3D_Direct(RelativeTM.GetScale3D());
 					}
 					else
 					{
-						PSC->RelativeLocation = Location;
-						PSC->RelativeRotation = Rotation;
+						PSC->SetRelativeLocation_Direct(Location);
+						PSC->SetRelativeRotation_Direct(Rotation);
 
 						if (LocationType == EAttachLocation::SnapToTarget)
 						{
 							// SnapToTarget indicates we "keep world scale", this indicates we we want the inverse of the parent-to-world scale 
 							// to calculate world scale at Scale 1, and then apply the passed in Scale
 							const FTransform ParentToWorld = AttachToComponent->GetSocketTransform(AttachPointName);
-							PSC->RelativeScale3D = Scale * ParentToWorld.GetSafeScaleReciprocal(ParentToWorld.GetScale3D());
+							PSC->SetRelativeScale3D_Direct(Scale * ParentToWorld.GetSafeScaleReciprocal(ParentToWorld.GetScale3D()));
 						}
 						else
 						{
-							PSC->RelativeScale3D = Scale;
+							PSC->SetRelativeScale3D_Direct(Scale);
 						}
 					}
 
@@ -1788,7 +1789,7 @@ UDecalComponent* CreateDecalComponent(class UMaterialInterface* DecalMaterial, F
 	DecalComp->bAllowAnyoneToDestroyMe = true;
 	DecalComp->SetDecalMaterial(DecalMaterial);
 	DecalComp->DecalSize = DecalSize;
-	DecalComp->bAbsoluteScale = true;
+	DecalComp->SetUsingAbsoluteScale(true);
 	DecalComp->RegisterComponentWithWorld(World);
 
 	if (LifeSpan > 0.f)
diff --git a/Engine/Source/Runtime/Engine/Private/HUD.cpp b/Engine/Source/Runtime/Engine/Private/HUD.cpp
--- a/Engine/Source/Runtime/Engine/Private/HUD.cpp
+++ b/Engine/Source/Runtime/Engine/Private/HUD.cpp
@@ -65,13 +65,13 @@ AHUD::AHUD(const FObjectInitializer& ObjectInitializer)
 {
 	PrimaryActorTick.TickGroup = TG_DuringPhysics;
 	PrimaryActorTick.bCanEverTick = true;
-	bHidden = true;
+	SetHidden(true);
 	bReplicates = false;
 
 	bLostFocusPaused = false;
 	bShowHUD = true;
 
-	bCanBeDamaged = false;
+	SetCanBeDamaged(false);
 	bEnableDebugTextShadow = false;
 }
 
@@ -455,7 +455,7 @@ void AHUD::ShowDebugInfo(float& YL, float& YPos)
 #endif
 			ShowDebugTargetActor->DisplayDebug(DebugCanvas, DisplayInfo, YL, YPos);
 
-			if (!bShowDebugForReticleTarget && ShowDebugTargetActor->Role == ROLE_SimulatedProxy)
+			if (!bShowDebugForReticleTarget && ShowDebugTargetActor->GetLocalRole() == ROLE_SimulatedProxy)
 			{
 				PlayerOwner->DisplayDebug(DebugCanvas, DisplayInfo, YL, YPos);
 			}
diff --git a/Engine/Source/Runtime/Engine/Private/Info.cpp b/Engine/Source/Runtime/Engine/Private/Info.cpp
--- a/Engine/Source/Runtime/Engine/Private/Info.cpp
+++ b/Engine/Source/Runtime/Engine/Private/Info.cpp
@@ -40,9 +40,9 @@ AInfo::AInfo(const FObjectInitializer& ObjectInitializer)
 	bAllowTickBeforeBeginPlay = true;
 	bReplicates = false;
 	NetUpdateFrequency = 10.0f;
-	bHidden = true;
-	bReplicateMovement = false;
-	bCanBeDamaged = false;
+	SetHidden(true);
+	SetReplicatingMovement(false);
+	SetCanBeDamaged(false);
 }
 
 #if WITH_EDITORONLY_DATA
diff --git a/Engine/Source/Runtime/Engine/Private/Interpolation.cpp b/Engine/Source/Runtime/Engine/Private/Interpolation.cpp
--- a/Engine/Source/Runtime/Engine/Private/Interpolation.cpp
+++ b/Engine/Source/Runtime/Engine/Private/Interpolation.cpp
@@ -386,7 +386,7 @@ AMatineeActor::AMatineeActor(const FObjectInitializer& ObjectInitializer)
 		static FConstructorStatics ConstructorStatics;
 
 		SpriteComponent->Sprite = ConstructorStatics.SceneManagerObject.Get();
-		SpriteComponent->RelativeScale3D = FVector(0.5f, 0.5f, 0.5f);
+		SpriteComponent->SetRelativeScale3D(FVector(0.5f, 0.5f, 0.5f));
 		SpriteComponent->SpriteInfo.Category = ConstructorStatics.ID_Matinee;
 		SpriteComponent->SpriteInfo.DisplayName = ConstructorStatics.NAME_Matinee;
 		SpriteComponent->SetupAttachment(RootComponent);
@@ -1577,7 +1577,7 @@ void AMatineeActor::SaveActorVisibility( AActor* Actor )
 			if ( !SavedVisibility )
 			{
 				// Save both bHidden and bHiddenEdTemporary to make it work properly in the editor
-				uint8 bSaveHidden = (Actor->bHidden ? 1 : 0) | (Actor->IsTemporarilyHiddenInEditor() ? 2 : 0);
+				uint8 bSaveHidden = (Actor->IsHidden() ? 1 : 0) | (Actor->IsTemporarilyHiddenInEditor() ? 2 : 0);
 				SavedActorVisibilities.Add( Actor, bSaveHidden );
 			}
 		}
@@ -1645,8 +1645,8 @@ void AMatineeActor::SaveActorTransforms( AActor* Actor )
 		if ( !SavedTransform && Actor->GetRootComponent() )
 		{
 			FSavedTransform NewSavedTransform;
-			NewSavedTransform.Translation = Actor->GetRootComponent()->RelativeLocation;
-			NewSavedTransform.Rotation = Actor->GetRootComponent()->RelativeRotation;
+			NewSavedTransform.Translation = Actor->GetRootComponent()->GetRelativeLocation();
+			NewSavedTransform.Rotation = Actor->GetRootComponent()->GetRelativeRotation();
 
 			SavedActorTransforms.Add( Actor, NewSavedTransform );
 		}
@@ -1667,7 +1667,7 @@ void AMatineeActor::RestoreActorTransforms()
 			FSavedTransform& SavedTransform = It.Value();
 
 			// only update actor position/rotation if the track changed its position/rotation
-			if( SavedActor->GetRootComponent() && (SavedActor->GetRootComponent()->RelativeLocation != SavedTransform.Translation || SavedActor->GetRootComponent()->RelativeRotation != SavedTransform.Rotation) )
+			if( SavedActor->GetRootComponent() && (SavedActor->GetRootComponent()->GetRelativeLocation() != SavedTransform.Translation || SavedActor->GetRootComponent()->GetRelativeRotation() != SavedTransform.Rotation) )
 			{
 				SavedActor->GetRootComponent()->SetRelativeLocationAndRotation(SavedTransform.Translation, SavedTransform.Rotation);
 			}
@@ -1692,7 +1692,7 @@ void AMatineeActor::RestoreActorVisibilities()
 			bool bSavedHiddenEditor = ( It.Value() & 2 ) ? true : false;
 
 			// only update actor if something has actually changed
-			if( SavedActor->bHidden != bSavedHidden )
+			if( SavedActor->IsHidden() != bSavedHidden )
 			{
 				SavedActor->SetActorHiddenInGame( bSavedHidden );
 			}
@@ -3740,9 +3740,9 @@ void UInterpTrackMove::UpdateKeyframe(int32 KeyIndex, UInterpTrackInst* TrInst)
 	}
 	*/
 
-	FVector RelativeSpaceEuler = Actor->GetRootComponent()->RelativeRotation.Euler();
+	FVector RelativeSpaceEuler = Actor->GetRootComponent()->GetRelativeRotation().Euler();
 
-	PosTrack.Points[KeyIndex].OutVal = Actor->GetRootComponent()->RelativeLocation;
+	PosTrack.Points[KeyIndex].OutVal = Actor->GetRootComponent()->GetRelativeLocation();
 	
 	// peek at an adjacent key frame to attempt to keep rotation continuous
 	if( EulerTrack.Points.Num() > 1 )
@@ -3838,9 +3838,9 @@ void UInterpTrackMove::UpdateChildKeyframe( UInterpTrack* ChildTrack, int32 KeyI
 	}
 
 	// New position of the actor
-	FVector NewPos = Actor->GetRootComponent()->RelativeLocation;
+	FVector NewPos = Actor->GetRootComponent()->GetRelativeLocation();
 	// New rotation of the actor
-	FVector NewRot =  Actor->GetRootComponent()->RelativeRotation.Euler();
+	FVector NewRot =  Actor->GetRootComponent()->GetRelativeRotation().Euler();
 
 	// Now determine what value should be updated in the float track.
 	switch( MoveAxis )
@@ -6728,7 +6728,7 @@ void UInterpTrackDirector::UpdateTrack(float NewPosition, UInterpTrackInst* TrIn
 	{
 		AMatineeActor* MatineeActor = CastChecked<AMatineeActor>( GrInst->GetOuter() );
 		// server is authoritative on viewtarget changes
-		if (PC->Role == ROLE_Authority || MatineeActor->bClientSideOnly || bSimulateCameraCutsOnClients)
+		if (PC->GetLocalRole() == ROLE_Authority || MatineeActor->bClientSideOnly || bSimulateCameraCutsOnClients)
 		{
 			float CutTime, CutTransitionTime;
 			FName ViewGroupName = GetViewedGroupName(NewPosition, CutTime, CutTransitionTime);
@@ -8252,7 +8252,7 @@ void UInterpTrackInstSound::TermTrackInst(UInterpTrack* Track)
 	{
 		// If we are currently playing, and want to keep the sound playing, 
 		// just flag it as 'auto destroy', and it will destroy itself when it finishes.
-		if(PlayAudioComp->bIsActive && SoundTrack->bContinueSoundOnMatineeEnd)
+		if(PlayAudioComp->IsActive() && SoundTrack->bContinueSoundOnMatineeEnd)
 		{
 			PlayAudioComp->bAutoDestroy = true;
 		}
@@ -9017,7 +9017,7 @@ void UInterpTrackInstToggle::SaveActorState(UInterpTrack* Track)
 	}
 	else if( LightActor != NULL )
 	{
-		bSavedActiveState = LightActor->GetLightComponent()->bVisible;
+		bSavedActiveState = LightActor->GetLightComponent()->GetVisibleFlag();
 	}
 }
 
@@ -9418,7 +9418,7 @@ void UInterpTrackVisibility::UpdateTrack(float NewPosition, UInterpTrackInst* Tr
 						else if( VisibilityKey.Action == EVTA_Toggle )
 						{
 							// Toggle the actor's visibility
-							HideActor( Actor, !Actor->bHidden );
+							HideActor( Actor, !Actor->IsHidden() );
 						}
 						if (!MatineeActor->bClientSideOnly && VisibilityKey.ActiveCondition == EVTC_Always)
 						{
diff --git a/Engine/Source/Runtime/Engine/Private/KismetSystemLibrary.cpp b/Engine/Source/Runtime/Engine/Private/KismetSystemLibrary.cpp
--- a/Engine/Source/Runtime/Engine/Private/KismetSystemLibrary.cpp
+++ b/Engine/Source/Runtime/Engine/Private/KismetSystemLibrary.cpp
@@ -2133,8 +2133,8 @@ void UKismetSystemLibrary::MoveComponentTo(USceneComponent* Component, FVector T
 		FLatentActionManager& LatentActionManager = World->GetLatentActionManager();
 		FInterpolateComponentToAction* Action = LatentActionManager.FindExistingAction<FInterpolateComponentToAction>(LatentInfo.CallbackTarget, LatentInfo.UUID);
 
-		const FVector ComponentLocation = (Component != NULL) ? Component->RelativeLocation : FVector::ZeroVector;
-		const FRotator ComponentRotation = (Component != NULL) ? Component->RelativeRotation : FRotator::ZeroRotator;
+		const FVector ComponentLocation = (Component != NULL) ? Component->GetRelativeLocation() : FVector::ZeroVector;
+		const FRotator ComponentRotation = (Component != NULL) ? Component->GetRelativeRotation() : FRotator::ZeroRotator;
 
 		// If not currently running
 		if (Action == NULL)
diff --git a/Engine/Source/Runtime/Engine/Private/LODActor.cpp b/Engine/Source/Runtime/Engine/Private/LODActor.cpp
--- a/Engine/Source/Runtime/Engine/Private/LODActor.cpp
+++ b/Engine/Source/Runtime/Engine/Private/LODActor.cpp
@@ -174,7 +174,7 @@ ALODActor::ALODActor(const FObjectInitializer& ObjectInitializer)
 	, LODDrawDistance(5000)
 	, bHasActorTriedToRegisterComponents(false)
 {
-	bCanBeDamaged = false;
+	SetCanBeDamaged(false);
 
 	PrimaryActorTick.bCanEverTick = true;
 	PrimaryActorTick.bStartWithTickEnabled = false;
diff --git a/Engine/Source/Runtime/Engine/Private/LevelActor.cpp b/Engine/Source/Runtime/Engine/Private/LevelActor.cpp
--- a/Engine/Source/Runtime/Engine/Private/LevelActor.cpp
+++ b/Engine/Source/Runtime/Engine/Private/LevelActor.cpp
@@ -428,7 +428,7 @@ AActor* UWorld::SpawnActor( UClass* Class, FTransform const* UserTransformPtr, c
 		// might necessarily be exactly the passed-in UserTransform.
 		FTransform const FinalRootComponentTransform =
 			TemplateRootComponent
-			? FTransform(TemplateRootComponent->RelativeRotation, TemplateRootComponent->RelativeLocation, TemplateRootComponent->RelativeScale3D) * UserTransform
+			? FTransform(TemplateRootComponent->GetRelativeRotation(), TemplateRootComponent->GetRelativeLocation(), TemplateRootComponent->GetRelativeScale3D()) * UserTransform
 			: UserTransform;
 
 		FVector const FinalRootLocation = FinalRootComponentTransform.GetLocation();
@@ -576,10 +576,10 @@ bool UWorld::DestroyActor( AActor* ThisActor, bool bNetForce, bool bShouldModify
 	{
 		// Note, for Standalone games, Actors should have Authority == ROLE_Authority.
 		// In that sense, they'll be treated as Network Actors here.
-		const bool bIsNetworkedActor = ThisActor->Role != ROLE_None;
+		const bool bIsNetworkedActor = ThisActor->GetLocalRole() != ROLE_None;
 
 		// Can't kill if wrong role.
-		const bool bCanDestroyNetworkActor = ThisActor->Role == ROLE_Authority || bNetForce || ThisActor->bNetTemporary;
+		const bool bCanDestroyNetworkActor = ThisActor->GetLocalRole() == ROLE_Authority || bNetForce || ThisActor->bNetTemporary;
 		if (bIsNetworkedActor && !bCanDestroyNetworkActor)
 		{
 			return false;
@@ -767,7 +767,7 @@ APlayerController* UWorld::SpawnPlayActor(UPlayer* NewPlayer, ENetRole RemoteRol
 
 		// Possess the newly-spawned player.
 		NewPlayerController->NetPlayerIndex = InNetPlayerIndex;
-		NewPlayerController->Role = ROLE_Authority;
+		NewPlayerController->SetRole(ROLE_Authority);
 		NewPlayerController->SetReplicates(RemoteRole != ROLE_None);
 		if (RemoteRole == ROLE_AutonomousProxy)
 		{
diff --git a/Engine/Source/Runtime/Engine/Private/LevelBounds.cpp b/Engine/Source/Runtime/Engine/Private/LevelBounds.cpp
--- a/Engine/Source/Runtime/Engine/Private/LevelBounds.cpp
+++ b/Engine/Source/Runtime/Engine/Private/LevelBounds.cpp
@@ -15,7 +15,7 @@ ALevelBounds::ALevelBounds(const FObjectInitializer& ObjectInitializer)
 	UBoxComponent* BoxComponent = CreateDefaultSubobject<UBoxComponent>(TEXT("BoxComponent0"));
 	RootComponent = BoxComponent;
 	RootComponent->Mobility = EComponentMobility::Movable;
-	RootComponent->RelativeScale3D = DefaultLevelSize;
+	RootComponent->SetRelativeScale3D(DefaultLevelSize);
 	RootComponent->SetCanEverAffectNavigation(false);
 
 	bAutoUpdateBounds = true;
@@ -26,7 +26,7 @@ ALevelBounds::ALevelBounds(const FObjectInitializer& ObjectInitializer)
 	BoxComponent->SetCollisionProfileName(UCollisionProfile::NoCollision_ProfileName);
 	BoxComponent->InitBoxExtent(FVector(0.5f, 0.5f, 0.5f));
 	
-	bCanBeDamaged = false;
+	SetCanBeDamaged(false);
 	
 #if WITH_EDITOR
 	bLevelBoundsDirty = true;
diff --git a/Engine/Source/Runtime/Engine/Private/LevelUtils.cpp b/Engine/Source/Runtime/Engine/Private/LevelUtils.cpp
--- a/Engine/Source/Runtime/Engine/Private/LevelUtils.cpp
+++ b/Engine/Source/Runtime/Engine/Private/LevelUtils.cpp
@@ -408,8 +408,8 @@ void FLevelUtils::ApplyLevelTransform(const FLevelUtils::FApplyLevelTransformPar
 					// Don't want to transform children they should stay relative to their parents.
 					if (RootComponent && RootComponent->GetAttachParent() == nullptr)
 					{
-						RootComponent->RelativeLocation = TransformParams.LevelTransform.TransformPosition(RootComponent->RelativeLocation);
-						RootComponent->RelativeRotation = TransformParams.LevelTransform.TransformRotation(RootComponent->RelativeRotation.Quaternion()).Rotator();
+						RootComponent->SetRelativeLocation_Direct(TransformParams.LevelTransform.TransformPosition(RootComponent->GetRelativeLocation()));
+						RootComponent->SetRelativeRotation_Direct(TransformParams.LevelTransform.TransformRotation(RootComponent->GetRelativeRotation().Quaternion()).Rotator());
 					}
 				}
 			}
@@ -426,7 +426,7 @@ void FLevelUtils::ApplyLevelTransform(const FLevelUtils::FApplyLevelTransformPar
 					// Don't want to transform children they should stay relative to their parents.
 					if (RootComponent && RootComponent->GetAttachParent() == nullptr)
 					{
-						RootComponent->SetRelativeLocationAndRotation(TransformParams.LevelTransform.TransformPosition(RootComponent->RelativeLocation), TransformParams.LevelTransform.TransformRotation(RootComponent->RelativeRotation.Quaternion()));
+						RootComponent->SetRelativeLocationAndRotation(TransformParams.LevelTransform.TransformPosition(RootComponent->GetRelativeLocation()), TransformParams.LevelTransform.TransformRotation(RootComponent->GetRelativeRotation().Quaternion()));
 					}
 				}
 			}
diff --git a/Engine/Source/Runtime/Engine/Private/Light.cpp b/Engine/Source/Runtime/Engine/Private/Light.cpp
--- a/Engine/Source/Runtime/Engine/Private/Light.cpp
+++ b/Engine/Source/Runtime/Engine/Private/Light.cpp
@@ -104,7 +104,7 @@ void ALight::SetEnabled(bool bSetEnabled)
 
 bool ALight::IsEnabled() const
 {
-	return LightComponent ? LightComponent->bVisible : false;
+	return LightComponent ? LightComponent->GetVisibleFlag() : false;
 }
 
 void ALight::ToggleEnabled()
@@ -237,7 +237,7 @@ ADirectionalLight::ADirectionalLight(const FObjectInitializer& ObjectInitializer
 #endif
 	DirectionalLightComponent = CastChecked<UDirectionalLightComponent>(GetLightComponent());
 	DirectionalLightComponent->Mobility = EComponentMobility::Stationary;
-	DirectionalLightComponent->RelativeRotation = FRotator(-46.0f, 0.0f, 0.0f);
+	DirectionalLightComponent->SetRelativeRotation(FRotator(-46.0f, 0.0f, 0.0f));
 	// Make directional light icons big since they tend to be important
 	// This is the root component so its scale affects all other components
 	DirectionalLightComponent->SetRelativeScale3D(FVector(2.5f, 2.5f, 2.5f));
diff --git a/Engine/Source/Runtime/Engine/Private/ModelLight.cpp b/Engine/Source/Runtime/Engine/Private/ModelLight.cpp
--- a/Engine/Source/Runtime/Engine/Private/ModelLight.cpp
+++ b/Engine/Source/Runtime/Engine/Private/ModelLight.cpp
@@ -735,7 +735,7 @@ void UModel::GroupAllNodes(ULevel* Level, const TArray<ULightComponentBase*>& Li
 			ULightComponent* Light = Cast<ULightComponent>(LightBase);
 
 			// Only add enabled lights and lights that can potentially be enabled at runtime (toggleable)
-			if (Light && (Light->bVisible || (!Light->HasStaticLighting() && Light->AffectsPrimitive(Level->ModelComponents[ComponentIndex]))))
+			if (Light && (Light->GetVisibleFlag() || (!Light->HasStaticLighting() && Light->AffectsPrimitive(Level->ModelComponents[ComponentIndex]))))
 			{
 				RelevantLights->Add(Light);
 			}
diff --git a/Engine/Source/Runtime/Engine/Private/NavigationObjectBase.cpp b/Engine/Source/Runtime/Engine/Private/NavigationObjectBase.cpp
--- a/Engine/Source/Runtime/Engine/Private/NavigationObjectBase.cpp
+++ b/Engine/Source/Runtime/Engine/Private/NavigationObjectBase.cpp
@@ -55,12 +55,12 @@ ANavigationObjectBase::ANavigationObjectBase(const FObjectInitializer& ObjectIni
 		if (GoodSprite)
 		{
 			GoodSprite->Sprite = ConstructorStatics.NavigationTextureObject.Get();
-			GoodSprite->RelativeScale3D = FVector(0.5f, 0.5f, 0.5f);
+			GoodSprite->SetRelativeScale3D(FVector(0.5f, 0.5f, 0.5f));
 			GoodSprite->bHiddenInGame = true;
 			GoodSprite->SpriteInfo.Category = ConstructorStatics.ID_Navigation;
 			GoodSprite->SpriteInfo.DisplayName = ConstructorStatics.NAME_Navigation;
 			GoodSprite->SetupAttachment(CapsuleComponent);
-			GoodSprite->bAbsoluteScale = true;
+			GoodSprite->SetUsingAbsoluteScale(true);
 			GoodSprite->bIsScreenSizeScaled = true;
 		}
 
@@ -70,7 +70,7 @@ ANavigationObjectBase::ANavigationObjectBase(const FObjectInitializer& ObjectIni
 			BadSprite->bHiddenInGame = true;
 			BadSprite->SpriteInfo.Category = ConstructorStatics.ID_Navigation;
 			BadSprite->SpriteInfo.DisplayName = ConstructorStatics.NAME_Navigation;
-			BadSprite->bAbsoluteScale = true;
+			BadSprite->SetUsingAbsoluteScale(true);
 			BadSprite->SetupAttachment(CapsuleComponent);
 			BadSprite->bIsScreenSizeScaled = true;
 		}
diff --git a/Engine/Source/Runtime/Engine/Private/NetConnection.cpp b/Engine/Source/Runtime/Engine/Private/NetConnection.cpp
--- a/Engine/Source/Runtime/Engine/Private/NetConnection.cpp
+++ b/Engine/Source/Runtime/Engine/Private/NetConnection.cpp
@@ -3154,7 +3154,7 @@ void UNetConnection::HandleClientPlayer( APlayerController *PC, UNetConnection*
 	check(LocalPlayer);
 	if( LocalPlayer->PlayerController && LocalPlayer->PlayerController->GetLevel() == PC->GetLevel())
 	{
-		if (LocalPlayer->PlayerController->Role == ROLE_Authority)
+		if (LocalPlayer->PlayerController->GetLocalRole() == ROLE_Authority)
 		{
 			// local placeholder PC while waiting for connection to be established
 			LocalPlayer->PlayerController->GetWorld()->DestroyActor(LocalPlayer->PlayerController);
@@ -3175,7 +3175,7 @@ void UNetConnection::HandleClientPlayer( APlayerController *PC, UNetConnection*
 	LocalPlayer->CurrentNetSpeed = CurrentNetSpeed;
 
 	// Init the new playerpawn.
-	PC->Role = ROLE_AutonomousProxy;
+	PC->SetRole(ROLE_AutonomousProxy);
 	PC->NetConnection = NetConnection;
 	PC->SetPlayer(LocalPlayer);
 	UE_LOG(LogNet, Verbose, TEXT("%s setplayer %s"),*PC->GetName(),*LocalPlayer->GetName());
@@ -3240,7 +3240,7 @@ void UChildConnection::HandleClientPlayer(APlayerController* PC, UNetConnection*
 		check( PC->GetWorld() );
 		for (TActorIterator<APlayerController> It(PC->GetWorld()); It; ++It)
 		{
-			if (It->Role < ROLE_Authority)
+			if (It->GetLocalRole() < ROLE_Authority)
 			{
 				UE_LOG(LogNet, Log, TEXT(" - %s"), *It->GetFullName());
 			}
@@ -3257,7 +3257,7 @@ void UChildConnection::HandleClientPlayer(APlayerController* PC, UNetConnection*
 	check(NewPlayer);
 	if (NewPlayer->PlayerController != NULL)
 	{
-		if (NewPlayer->PlayerController->Role == ROLE_Authority)
+		if (NewPlayer->PlayerController->GetLocalRole() == ROLE_Authority)
 		{
 			// local placeholder PC while waiting for connection to be established
 			NewPlayer->PlayerController->GetWorld()->DestroyActor(NewPlayer->PlayerController);
@@ -3278,7 +3278,7 @@ void UChildConnection::HandleClientPlayer(APlayerController* PC, UNetConnection*
 	NewPlayer->CurrentNetSpeed = CurrentNetSpeed;
 
 	// Init the new playerpawn.
-	PC->Role = ROLE_AutonomousProxy;
+	PC->SetRole(ROLE_AutonomousProxy);
 	PC->NetConnection = NetConnection;
 	PC->SetPlayer(NewPlayer);
 	UE_LOG(LogNet, Verbose, TEXT("%s setplayer %s"), *PC->GetName(), *NewPlayer->GetName());
diff --git a/Engine/Source/Runtime/Engine/Private/Note.cpp b/Engine/Source/Runtime/Engine/Private/Note.cpp
--- a/Engine/Source/Runtime/Engine/Private/Note.cpp
+++ b/Engine/Source/Runtime/Engine/Private/Note.cpp
@@ -50,7 +50,7 @@ ANote::ANote(const FObjectInitializer& ObjectInitializer)
 		if (SpriteComponent)
 		{
 			SpriteComponent->Sprite = ConstructorStatics.NoteTextureObject.Get();
-			SpriteComponent->RelativeScale3D = FVector(0.5f, 0.5f, 0.5f);
+			SpriteComponent->SetRelativeScale3D(FVector(0.5f, 0.5f, 0.5f));
 			SpriteComponent->SpriteInfo.Category = ConstructorStatics.ID_Notes;
 			SpriteComponent->SpriteInfo.DisplayName = ConstructorStatics.NAME_Notes;
 			SpriteComponent->SetupAttachment(RootComponent);
@@ -59,6 +59,6 @@ ANote::ANote(const FObjectInitializer& ObjectInitializer)
 	}
 #endif // WITH_EDITORONLY_DATA
 
-	bHidden = true;
-	bCanBeDamaged = false;
+	SetHidden(true);
+	SetCanBeDamaged(false);
 }
diff --git a/Engine/Source/Runtime/Engine/Private/PainCausingVolume.cpp b/Engine/Source/Runtime/Engine/Private/PainCausingVolume.cpp
--- a/Engine/Source/Runtime/Engine/Private/PainCausingVolume.cpp
+++ b/Engine/Source/Runtime/Engine/Private/PainCausingVolume.cpp
@@ -38,7 +38,7 @@ void APainCausingVolume::Reset()
 void APainCausingVolume::ActorEnteredVolume(AActor* Other)
 {
 	Super::ActorEnteredVolume(Other);
-	if ( bPainCausing && bEntryPain && Other->bCanBeDamaged )
+	if ( bPainCausing && bEntryPain && Other->CanBeDamaged() )
 	{
 		CausePainTo(Other);
 	}
@@ -59,7 +59,7 @@ void APainCausingVolume::PainTimer()
 
 		for (AActor* const A : TouchingActors)
 		{
-			if (A && A->bCanBeDamaged && !A->IsPendingKill())
+			if (A && A->CanBeDamaged() && !A->IsPendingKill())
 			{
 				// @todo physicsVolume This won't work for other actor. Need to fix it properly
 				APawn* PawnA = Cast<APawn>(A);
diff --git a/Engine/Source/Runtime/Engine/Private/Particles/Emitter.cpp b/Engine/Source/Runtime/Engine/Private/Particles/Emitter.cpp
--- a/Engine/Source/Runtime/Engine/Private/Particles/Emitter.cpp
+++ b/Engine/Source/Runtime/Engine/Private/Particles/Emitter.cpp
@@ -48,7 +48,7 @@ AEmitter::AEmitter(const FObjectInitializer& ObjectInitializer)
 		if (SpriteComponent)
 		{
 			SpriteComponent->Sprite = ConstructorStatics.SpriteTextureObject.Get();
-			SpriteComponent->RelativeScale3D = FVector(0.5f, 0.5f, 0.5f);
+			SpriteComponent->SetRelativeScale3D_Direct(FVector(0.5f, 0.5f, 0.5f));
 			SpriteComponent->bHiddenInGame = true;
 			SpriteComponent->bIsScreenSizeScaled = true;
 			SpriteComponent->SpriteInfo.Category = ConstructorStatics.ID_Effects;
@@ -67,7 +67,7 @@ AEmitter::AEmitter(const FObjectInitializer& ObjectInitializer)
 			ArrowComponent->SpriteInfo.Category = ConstructorStatics.ID_Effects;
 			ArrowComponent->SpriteInfo.DisplayName = ConstructorStatics.NAME_Effects;
 			ArrowComponent->SetupAttachment(ParticleSystemComponent);
-			ArrowComponent->bAbsoluteScale = true;
+			ArrowComponent->SetUsingAbsoluteScale(true);
 		}
 	}
 #endif // WITH_EDITORONLY_DATA
@@ -221,7 +221,7 @@ void AEmitter::ToggleActive()
 	if (ParticleSystemComponent)
 	{
 		ParticleSystemComponent->ToggleActive();
-		bCurrentlyActive = ParticleSystemComponent->bIsActive;
+		bCurrentlyActive = ParticleSystemComponent->IsActive();
 	}
 }
 
