Update Engine and Plugins to use the new Getters / Setters for UActorComponent, USceneComponent, and AActor replicated properties

diff --git a/Engine/Source/Runtime/Engine/Private/Particles/ParticleComponents.cpp b/Engine/Source/Runtime/Engine/Private/Particles/ParticleComponents.cpp
--- a/Engine/Source/Runtime/Engine/Private/Particles/ParticleComponents.cpp
+++ b/Engine/Source/Runtime/Engine/Private/Particles/ParticleComponents.cpp
@@ -3627,7 +3627,7 @@ void UParticleSystemComponent::ForceReset()
 	}
 #endif
 
-	bool bOldActive = bIsActive;
+	bool bOldActive = IsActive();
 	ResetParticles(true);
 	if (bOldActive)
 	{
@@ -3866,9 +3866,9 @@ void UParticleSystemComponent::OnRegister()
 			}
 		}
 
-		SavedAutoAttachRelativeLocation = RelativeLocation;
-		SavedAutoAttachRelativeRotation = RelativeRotation;
-		SavedAutoAttachRelativeScale3D = RelativeScale3D;
+		SavedAutoAttachRelativeLocation = GetRelativeLocation();
+		SavedAutoAttachRelativeRotation = GetRelativeRotation();
+		SavedAutoAttachRelativeScale3D = GetRelativeScale3D();
 	}
 
 	if (ShouldBeTickManaged())
@@ -3879,7 +3879,7 @@ void UParticleSystemComponent::OnRegister()
 	Super::OnRegister();
 
 	// If we were active before but are not now, activate us
-	if (bWasActive && !bIsActive)
+	if (bWasActive && !IsActive())
 	{
 		Activate(true);
 	}
@@ -3902,7 +3902,7 @@ void UParticleSystemComponent::OnUnregister()
 		TEXT("OnUnregister %s Component=0x%p Scene=0x%p FXSystem=0x%p"),
 		Template != NULL ? *Template->GetName() : TEXT("NULL"), this, GetWorld()->Scene, FXSystem);
 
-	bWasActive = bIsActive && !bWasDeactivated;
+	bWasActive = IsActive() && !bWasDeactivated;
 
 	check(GetWorld());
 	SetComponentTickEnabled(false);
@@ -3962,7 +3962,7 @@ void UParticleSystemComponent::SendRenderTransform_Concurrent()
 	SCOPE_CYCLE_COUNTER(STAT_ParticlesOverview_GT_CNC);
 
 	ForceAsyncWorkCompletion(ENSURE_AND_STALL, false, true);
-	if (bIsActive)
+	if (IsActive())
 	{
 		if (bSkipUpdateDynamicDataDuringTick == false)
 		{
@@ -3994,7 +3994,7 @@ void UParticleSystemComponent::SendRenderDynamicData_Concurrent()
 		// check to see if this PSC is active.  When you attach a PSC it gets added to the DataManager
 		// even if it might be bIsActive = false  (e.g. attach and later in the frame activate it)
 		// or also for PSCs that are attached to a SkelComp which is being attached and reattached but the PSC itself is not active!
-		if (bIsActive)
+		if (IsActive())
 		{
 			UpdateDynamicData();
 		}
@@ -4567,7 +4567,7 @@ void UParticleSystemComponent::OrientZAxisTowardCamera()
 		
 		// Adjust our rotation
 		const FRotator AdjustmentAngle(PointTo);
-		RelativeRotation += AdjustmentAngle;
+		GetRelativeRotation_DirectMutable() += AdjustmentAngle;
 
 		// Mark the component transform as dirty if the rotation has changed.
 		bIsTransformDirty |= !AdjustmentAngle.IsZero();
@@ -4934,7 +4934,7 @@ void FGameThreadDispatchBatchedAsyncTasks::DoTask(ENamedThreads::Type CurrentThr
 
 bool UParticleSystemComponent::IsReadyForOwnerToAutoDestroy() const
 {
-	return (!bIsActive && bWasCompleted);
+	return (!IsActive() && bWasCompleted);
 }
 
 void UParticleSystemComponent::SetComponentTickEnabled(bool bEnabled)
@@ -5077,7 +5077,7 @@ void UParticleSystemComponent::TickComponent(float DeltaTime, enum ELevelTick Ti
 	bool bDisallowAsync = false;
 
 	// Bail out if inactive and not AutoActivate
-	if ((bIsActive == false) && (bAutoActivate == false))
+	if ((IsActive() == false) && (bAutoActivate == false))
 	{
 		// Disable our tick here, will be enabled when activating
 		SetComponentTickEnabled(false);
@@ -5109,9 +5109,9 @@ void UParticleSystemComponent::TickComponent(float DeltaTime, enum ELevelTick Ti
 	// System settings may have been lowered. Support late deactivation.
 	int32 DetailModeCVar = GetCurrentDetailMode();
 	const bool bDetailModeAllowsRendering	= DetailMode <= DetailModeCVar;
-	if ( bDetailModeAllowsRendering == false )
+	if (bDetailModeAllowsRendering == false)
 	{
-		if ( bIsActive )
+		if (IsActive())
 		{
 			DeactivateSystem();
 			Super::MarkRenderDynamicDataDirty();
@@ -5435,7 +5435,7 @@ void UParticleSystemComponent::FinalizeTickComponent()
 	float CurrTime = GetWorld()->GetTimeSeconds();
 
 	//Are we still significant?
-	if ((bIsActive && !bWasDeactivated) && bIsManagingSignificance && NumSignificantEmitters == 0 && CurrTime >= LastSignificantTime + Template->InsignificanceDelay)
+	if ((IsActive() && !bWasDeactivated) && bIsManagingSignificance && NumSignificantEmitters == 0 && CurrTime >= LastSignificantTime + Template->InsignificanceDelay)
 	{
 		OnSignificanceChanged(false, true);
 	}
@@ -5701,7 +5701,7 @@ void UParticleSystemComponent::ResetParticles(bool bEmptyInstances)
 	UWorld* OwningWorld = GetWorld();
 
 	//Also consider this deactivation.
-	if (bIsActive)
+	if (IsActive())
 	{
 		OnSystemPreActivationChange.Broadcast(this, false);
 	}
@@ -5726,7 +5726,7 @@ void UParticleSystemComponent::ResetParticles(bool bEmptyInstances)
 	}
 
 	// Set the system as inactive
-	bIsActive	= false;
+	SetActiveFlag(false);
 
 	// Remove instances if we're not running gameplay.ww
 	if (!bIsGameWorld || bEmptyInstances)
@@ -5792,7 +5792,7 @@ void UParticleSystemComponent::SetTemplate(class UParticleSystem* NewTemplate)
 		bool bIsTemplate = IsTemplate();
 		bWasCompleted = false;
 		// remember if we were active and therefore should restart after setting up the new template
-		bWasActive = bIsActive && !bWasDeactivated; 
+		bWasActive = IsActive() && !bWasDeactivated; 
 		bool bResetInstances = false;
 		if (NewTemplate != Template)
 		{
@@ -5931,9 +5931,9 @@ void UParticleSystemComponent::ActivateSystem(bool bFlagAsJustAttached)
 				{
 					bDidAutoAttach = bWasAutoAttached;
 					CancelAutoAttachment(true, World);
-					SavedAutoAttachRelativeLocation = RelativeLocation;
-					SavedAutoAttachRelativeRotation = RelativeRotation;
-					SavedAutoAttachRelativeScale3D = RelativeScale3D;
+					SavedAutoAttachRelativeLocation = GetRelativeLocation();
+					SavedAutoAttachRelativeRotation = GetRelativeRotation();
+					SavedAutoAttachRelativeScale3D = GetRelativeScale3D();
 					AttachToComponent(NewParent, FAttachmentTransformRules(AutoAttachLocationRule, AutoAttachRotationRule, AutoAttachScaleRule, false), AutoAttachSocketName);
 				}
 
@@ -5948,7 +5948,7 @@ void UParticleSystemComponent::ActivateSystem(bool bFlagAsJustAttached)
 
 		AccumTickTime = 0.0;
 
-		if (!bIsActive)
+		if (!IsActive())
 		{
 			LastSignificantTime = World->GetTimeSeconds();
 			RequiredSignificance = EParticleSignificanceLevel::Low;
@@ -5972,7 +5972,7 @@ void UParticleSystemComponent::ActivateSystem(bool bFlagAsJustAttached)
 		bool bNeedToUpdateTransform = bWasDeactivated;
 		bWasCompleted = false;
 		bWasDeactivated = false;
-		bIsActive = true;
+		SetActiveFlag(true);
 		bWasActive = false; // Set to false now, it may get set to true when it's deactivated due to unregister
 		SetComponentTickEnabled(true);
 
@@ -6031,7 +6031,7 @@ void UParticleSystemComponent::ActivateSystem(bool bFlagAsJustAttached)
 			{
 				if (DesiredLODLevel != LODLevel)
 				{
-					bIsActive = true;
+					SetActiveFlag(true);
 					SetComponentTickEnabled(true);
 				}
 				SetLODLevel(DesiredLODLevel);
@@ -6093,7 +6093,7 @@ void UParticleSystemComponent::Complete()
 
 	// When system is done - destroy all subemitters etc. We don't need them any more.
 	ResetParticles();
-	bIsActive = false;
+	SetActiveFlag(false);
 	SetComponentTickEnabled(false);
 
 	if (PoolingMethod == EPSCPoolMethod::AutoRelease)
@@ -6132,7 +6132,7 @@ void UParticleSystemComponent::DeactivateSystem()
 		TEXT("DeactivateSystem @ %fs %s"), World->TimeSeconds,
 		Template != NULL ? *Template->GetName() : TEXT("NULL"));
 
-	if (bIsActive)
+	if (IsActive())
 	{
 		OnSystemPreActivationChange.Broadcast(this, false);
 	}
@@ -6192,9 +6192,9 @@ void UParticleSystemComponent::CancelAutoAttachment(bool bDetachFromParent, cons
 		if (bDidAutoAttach)
 		{
 			// Restore relative transform from before attachment. Actual transform will be updated as part of DetachFromParent().
-			RelativeLocation = SavedAutoAttachRelativeLocation;
-			RelativeRotation = SavedAutoAttachRelativeRotation;
-			RelativeScale3D = SavedAutoAttachRelativeScale3D;
+			SetRelativeLocation_Direct(SavedAutoAttachRelativeLocation);
+			SetRelativeRotation_Direct(SavedAutoAttachRelativeRotation);
+			SetRelativeScale3D_Direct(SavedAutoAttachRelativeScale3D);
 			bDidAutoAttach = false;
 		}
 
@@ -6250,7 +6250,7 @@ void UParticleSystemComponent::Activate(bool bReset)
 		{
 			ActivateSystem(bReset);
 
-			if (bIsActive)
+			if (IsActive())
 			{
 				OnComponentActivated.Broadcast(this, bReset);
 			}
@@ -6611,7 +6611,7 @@ void UParticleSystemComponent::InitializeSystem()
 		if (IsRegistered())
 		{
 			AccumTickTime = 0.0;
-			if ((bIsActive == false) && (bAutoActivate == true) && (bWasDeactivated == false))
+			if ((IsActive() == false) && (bAutoActivate == true) && (bWasDeactivated == false))
 			{
 				SetActive(true);
 			}
diff --git a/Engine/Source/Runtime/Engine/Private/Particles/ParticleSystemRender.cpp b/Engine/Source/Runtime/Engine/Private/Particles/ParticleSystemRender.cpp
--- a/Engine/Source/Runtime/Engine/Private/Particles/ParticleSystemRender.cpp
+++ b/Engine/Source/Runtime/Engine/Private/Particles/ParticleSystemRender.cpp
@@ -7325,13 +7325,13 @@ FPrimitiveSceneProxy* UParticleSystemComponent::CreateSceneProxy()
 
 	//@fixme EmitterInstances.Num() check should be here to avoid proxies for dead emitters but there are some edge cases where it happens for emitters that have just activated...
 	//@fixme Get non-instanced path working in ES2!
-	if ((bIsActive == true)/** && (EmitterInstances.Num() > 0)*/ && Template)
+	if ((IsActive() == true)/** && (EmitterInstances.Num() > 0)*/ && Template)
 	{
 		FInGameScopedCycleCounter InGameCycleCounter(GetWorld(), EInGamePerfTrackers::VFXSignificance, EInGamePerfTrackerThreads::GameThread, bIsManagingSignificance);
 
 		UE_LOG(LogParticles,Verbose,
 			TEXT("CreateSceneProxy @ %fs %s bIsActive=%d"), GetWorld()->TimeSeconds,
-			Template != NULL ? *Template->GetName() : TEXT("NULL"), bIsActive);
+			Template != NULL ? *Template->GetName() : TEXT("NULL"), IsActive());
 
 		if (EmitterInstances.Num() > 0)
 		{
diff --git a/Engine/Source/Runtime/Engine/Private/Pawn.cpp b/Engine/Source/Runtime/Engine/Private/Pawn.cpp
--- a/Engine/Source/Runtime/Engine/Private/Pawn.cpp
+++ b/Engine/Source/Runtime/Engine/Private/Pawn.cpp
@@ -49,13 +49,13 @@ APawn::APawn(const FObjectInitializer& ObjectInitializer)
 	{
 		AIControllerClass = ((APawn*)APawn::StaticClass()->GetDefaultObject())->AIControllerClass;
 	}
-	bCanBeDamaged = true;
+	SetCanBeDamaged(true);
 	
 	SetRemoteRoleForBackwardsCompat(ROLE_SimulatedProxy);
 	bReplicates = true;
 	NetPriority = 3.0f;
 	NetUpdateFrequency = 100.f;
-	bReplicateMovement = true;
+	SetReplicatingMovement(true);
 	BaseEyeHeight = 64.0f;
 	AllowedYawError = 10.99f;
 	BlendedReplayViewPitch = 0.0f;
@@ -70,16 +70,16 @@ APawn::APawn(const FObjectInitializer& ObjectInitializer)
 
 	bInputEnabled = true;
 
-	ReplicatedMovement.LocationQuantizationLevel = EVectorQuantization::RoundTwoDecimals;
+	GetReplicatedMovement_Mutable().LocationQuantizationLevel = EVectorQuantization::RoundTwoDecimals;
 }
 
 void APawn::PreInitializeComponents()
 {
 	Super::PreInitializeComponents();
 
-	if (Instigator == nullptr)
+	if (GetInstigator() == nullptr)
 	{
-		Instigator = this;
+		SetInstigator(this);
 	}
 
 	if (AutoPossessPlayer != EAutoReceiveInput::Disabled && GetNetMode() != NM_Client )
@@ -287,7 +287,7 @@ FRotator APawn::GetViewRotation() const
 	{
 		return Controller->GetControlRotation();
 	}
-	else if (Role < ROLE_Authority)
+	else if (GetLocalRole() < ROLE_Authority)
 	{
 		// check if being spectated
 		for( FConstPlayerControllerIterator Iterator = GetWorld()->GetPlayerControllerIterator(); Iterator; ++Iterator )
@@ -313,7 +313,7 @@ void APawn::SpawnDefaultController()
 	if (AIControllerClass != nullptr)
 	{
 		FActorSpawnParameters SpawnInfo;
-		SpawnInfo.Instigator = Instigator;
+		SpawnInfo.Instigator = GetInstigator();
 		SpawnInfo.SpawnCollisionHandlingOverride = ESpawnActorCollisionHandlingMethod::AlwaysSpawn;
 		SpawnInfo.OverrideLevel = GetLevel();
 		SpawnInfo.ObjectFlags |= RF_Transient;	// We never want to save AI controllers into a map
@@ -329,7 +329,7 @@ void APawn::SpawnDefaultController()
 
 void APawn::TurnOff()
 {
-	if (Role == ROLE_Authority)
+	if (GetLocalRole() == ROLE_Authority)
 	{
 		SetReplicates(true);
 	}
@@ -408,7 +408,7 @@ void APawn::EndPlay(const EEndPlayReason::Type EndPlayReason)
 
 bool APawn::ShouldTakeDamage(float Damage, FDamageEvent const& DamageEvent, AController* EventInstigator, AActor* DamageCauser) const
 {
-	if ( (Role < ROLE_Authority) || !bCanBeDamaged || !GetWorld()->GetAuthGameMode() || (Damage == 0.f) )
+	if ((GetLocalRole() < ROLE_Authority) || !CanBeDamaged() || !GetWorld()->GetAuthGameMode() || (Damage == 0.f))
 	{
 		return false;
 	}
@@ -566,7 +566,7 @@ const AActor* APawn::GetNetOwner() const
 
 class UPlayer* APawn::GetNetOwningPlayer()
 {
-	if (Role == ROLE_Authority)
+	if (GetLocalRole() == ROLE_Authority)
 	{
 		if (Controller)
 		{
@@ -844,7 +844,7 @@ void APawn::OutsideWorldBounds()
 	{
 		bProcessingOutsideWorldBounds = true;
 		// AI pawns on the server just destroy
-		if (Role == ROLE_Authority && Cast<APlayerController>(Controller) == nullptr)
+		if (GetLocalRole() == ROLE_Authority && Cast<APlayerController>(Controller) == nullptr)
 		{
 			Destroy();
 		}
@@ -1002,7 +1002,7 @@ void APawn::LaunchPawn(FVector LaunchVelocity, bool bXYOverride, bool bZOverride
 
 void APawn::PostNetReceiveVelocity(const FVector& NewVelocity)
 {
-	if (Role == ROLE_SimulatedProxy)
+	if (GetLocalRole() == ROLE_SimulatedProxy)
 	{
 		UMovementComponent* const MoveComponent = GetMovementComponent();
 		if ( MoveComponent )
@@ -1014,27 +1014,30 @@ void APawn::PostNetReceiveVelocity(const FVector& NewVelocity)
 
 void APawn::PostNetReceiveLocationAndRotation()
 {
+	const FRepMovement& ConstRepMovement = GetReplicatedMovement();
+
 	// always consider Location as changed if we were spawned this tick as in that case our replicated Location was set as part of spawning, before PreNetReceive()
-	if( (FRepMovement::RebaseOntoLocalOrigin(ReplicatedMovement.Location, this) == GetActorLocation() 
-		&& ReplicatedMovement.Rotation == GetActorRotation()) && (CreationTime != GetWorld()->TimeSeconds) )
+	if( (FRepMovement::RebaseOntoLocalOrigin(ConstRepMovement.Location, this) == GetActorLocation()
+		&& ConstRepMovement.Rotation == GetActorRotation()) && (CreationTime != GetWorld()->TimeSeconds) )
 	{
 		return;
 	}
 
-	if( Role == ROLE_SimulatedProxy )
+	if (GetLocalRole() == ROLE_SimulatedProxy)
 	{
 		// Correction to make sure pawn doesn't penetrate floor after replication rounding
-		ReplicatedMovement.Location.Z += 0.01f;
+		FRepMovement& MutableRepMovement = GetReplicatedMovement_Mutable();
+		MutableRepMovement.Location.Z += 0.01f;
 
 		const FVector OldLocation = GetActorLocation();
 		const FQuat OldRotation = GetActorQuat();
-		const FVector NewLocation = FRepMovement::RebaseOntoLocalOrigin(ReplicatedMovement.Location, this);
-		SetActorLocationAndRotation(NewLocation, ReplicatedMovement.Rotation, /*bSweep=*/ false);
+		const FVector NewLocation = FRepMovement::RebaseOntoLocalOrigin(MutableRepMovement.Location, this);
+		SetActorLocationAndRotation(NewLocation, MutableRepMovement.Rotation, /*bSweep=*/ false);
 
 		INetworkPredictionInterface* PredictionInterface = Cast<INetworkPredictionInterface>(GetMovementComponent());
 		if (PredictionInterface)
 		{
-			PredictionInterface->SmoothCorrection(OldLocation, OldRotation, NewLocation, ReplicatedMovement.Rotation.Quaternion());
+			PredictionInterface->SmoothCorrection(OldLocation, OldRotation, NewLocation, MutableRepMovement.Rotation.Quaternion());
 		}
 	}
 }
@@ -1056,12 +1059,12 @@ bool APawn::IsBasedOnActor(const AActor* Other) const
 bool APawn::IsNetRelevantFor(const AActor* RealViewer, const AActor* ViewTarget, const FVector& SrcLocation) const
 {
 	CA_SUPPRESS(6011);
-	if (bAlwaysRelevant || RealViewer == Controller || IsOwnedBy(ViewTarget) || IsOwnedBy(RealViewer) || this == ViewTarget || ViewTarget == Instigator
+	if (bAlwaysRelevant || RealViewer == Controller || IsOwnedBy(ViewTarget) || IsOwnedBy(RealViewer) || this == ViewTarget || ViewTarget == GetInstigator()
 		|| IsBasedOnActor(ViewTarget) || (ViewTarget && ViewTarget->IsBasedOnActor(this)))
 	{
 		return true;
 	}
-	else if( (bHidden || bOnlyRelevantToOwner) && (!GetRootComponent() || !GetRootComponent()->IsCollisionEnabled()) ) 
+	else if ((IsHidden() || bOnlyRelevantToOwner) && (!GetRootComponent() || !GetRootComponent()->IsCollisionEnabled())) 
 	{
 		return false;
 	}
@@ -1093,7 +1096,7 @@ void APawn::PreReplication( IRepChangedPropertyTracker & ChangedPropertyTracker
 {
 	Super::PreReplication( ChangedPropertyTracker );
 
-	if (Role == ROLE_Authority && GetController())
+	if (GetLocalRole() == ROLE_Authority && GetController())
 	{
 		SetRemoteViewPitch(GetController()->GetControlRotation().Pitch);
 	}
diff --git a/Engine/Source/Runtime/Engine/Private/PhysicsEngine/PhysicsConstraintActor.cpp b/Engine/Source/Runtime/Engine/Private/PhysicsEngine/PhysicsConstraintActor.cpp
--- a/Engine/Source/Runtime/Engine/Private/PhysicsEngine/PhysicsConstraintActor.cpp
+++ b/Engine/Source/Runtime/Engine/Private/PhysicsEngine/PhysicsConstraintActor.cpp
@@ -24,7 +24,7 @@ APhysicsConstraintActor::APhysicsConstraintActor(const FObjectInitializer& Objec
 
 	ConstraintComp = CreateDefaultSubobject<UPhysicsConstraintComponent>(TEXT("MyConstraintComp"));
 	RootComponent = ConstraintComp;
-	bHidden = true;
+	SetHidden(true);
 }
 
 void APhysicsConstraintActor::PostLoad()
diff --git a/Engine/Source/Runtime/Engine/Private/PhysicsEngine/PhysicsReplication.cpp b/Engine/Source/Runtime/Engine/Private/PhysicsEngine/PhysicsReplication.cpp
--- a/Engine/Source/Runtime/Engine/Private/PhysicsEngine/PhysicsReplication.cpp
+++ b/Engine/Source/Runtime/Engine/Private/PhysicsEngine/PhysicsReplication.cpp
@@ -383,8 +383,9 @@ void FPhysicsReplication::OnTick(float DeltaSeconds, TMap<TWeakObjectPtr<UPrimit
 				bool bUpdated = false;
 				if (AActor* OwningActor = PrimComp->GetOwner())
 				{
-					const bool bIsSimulated = OwningActor->Role == ROLE_SimulatedProxy;
-					const bool bIsReplicatedAutonomous = OwningActor->Role == ROLE_AutonomousProxy && PrimComp->bReplicatePhysicsToAutonomousProxy;
+					const ENetRole OwnerRole = OwningActor->GetLocalRole();
+					const bool bIsSimulated = OwnerRole == ROLE_SimulatedProxy;
+					const bool bIsReplicatedAutonomous = OwnerRole == ROLE_AutonomousProxy && PrimComp->bReplicatePhysicsToAutonomousProxy;
 					if (bIsSimulated || bIsReplicatedAutonomous)
 					{
 						// Get the ping of the guy who owns this thing. If nobody is,
diff --git a/Engine/Source/Runtime/Engine/Private/PhysicsEngine/PhysicsSpring.cpp b/Engine/Source/Runtime/Engine/Private/PhysicsEngine/PhysicsSpring.cpp
--- a/Engine/Source/Runtime/Engine/Private/PhysicsEngine/PhysicsSpring.cpp
+++ b/Engine/Source/Runtime/Engine/Private/PhysicsEngine/PhysicsSpring.cpp
@@ -101,7 +101,7 @@ void UPhysicsSpringComponent::TickComponent(float DeltaTime, enum ELevelTick Tic
 
 	if (UPrimitiveComponent* BasePrimComp = Cast<UPrimitiveComponent>(GetAttachParent()))
 	{
-		if (bIsActive)
+		if (IsActive())
 		{
 			const FVector SpringStart = GetComponentTransform().GetLocation();
 			const FVector SpringDesiredEnd = SpringPositionFromLength(SpringLengthAtRest);
diff --git a/Engine/Source/Runtime/Engine/Private/PhysicsEngine/PhysicsThruster.cpp b/Engine/Source/Runtime/Engine/Private/PhysicsEngine/PhysicsThruster.cpp
--- a/Engine/Source/Runtime/Engine/Private/PhysicsEngine/PhysicsThruster.cpp
+++ b/Engine/Source/Runtime/Engine/Private/PhysicsEngine/PhysicsThruster.cpp
@@ -43,7 +43,7 @@ void UPhysicsThrusterComponent::TickComponent(float DeltaTime, enum ELevelTick T
 	Super::TickComponent(DeltaTime, TickType, ThisTickFunction);
 
 	// Applied force to the base, so if we don't have one, do nothing.
-	if( bIsActive && GetAttachParent())
+	if(IsActive() && GetAttachParent())
 	{
 		FVector WorldForce = ThrustStrength * GetComponentTransform().TransformVectorNoScale( FVector(-1.f,0.f,0.f) );
 
diff --git a/Engine/Source/Runtime/Engine/Private/PhysicsEngine/RadialForceComponent.cpp b/Engine/Source/Runtime/Engine/Private/PhysicsEngine/RadialForceComponent.cpp
--- a/Engine/Source/Runtime/Engine/Private/PhysicsEngine/RadialForceComponent.cpp
+++ b/Engine/Source/Runtime/Engine/Private/PhysicsEngine/RadialForceComponent.cpp
@@ -37,7 +37,7 @@ void URadialForceComponent::TickComponent(float DeltaTime, enum ELevelTick TickT
 {
 	Super::TickComponent(DeltaTime, TickType, ThisTickFunction);
 
-	if(bIsActive)
+	if(IsActive())
 	{
 		const FVector Origin = GetComponentLocation();
 
@@ -244,9 +244,7 @@ ARadialForceActor::ARadialForceActor(const FObjectInitializer& ObjectInitializer
 #endif // WITH_EDITORONLY_DATA
 		}
 
-		SpriteComponent->RelativeScale3D.X = 0.5f;
-		SpriteComponent->RelativeScale3D.Y = 0.5f;
-		SpriteComponent->RelativeScale3D.Z = 0.5f;
+		SpriteComponent->SetRelativeScale3D(FVector(0.5f, 0.5f, 0.5f));
 		SpriteComponent->SetupAttachment(ForceComponent);
 		SpriteComponent->bIsScreenSizeScaled = true;
 	}
diff --git a/Engine/Source/Runtime/Engine/Private/PhysicsVolume.cpp b/Engine/Source/Runtime/Engine/Private/PhysicsVolume.cpp
--- a/Engine/Source/Runtime/Engine/Private/PhysicsVolume.cpp
+++ b/Engine/Source/Runtime/Engine/Private/PhysicsVolume.cpp
@@ -16,7 +16,7 @@ APhysicsVolume::APhysicsVolume(const FObjectInitializer& ObjectInitializer)
 	TerminalVelocity = UPhysicsSettings::Get()->DefaultTerminalVelocity;
 	bAlwaysRelevant = true;
 	NetUpdateFrequency = 0.1f;
-	bReplicateMovement = false;
+	SetReplicatingMovement(false);
 }
 
 #if WITH_EDITOR
diff --git a/Engine/Source/Runtime/Engine/Private/PlayerCameraManager.cpp b/Engine/Source/Runtime/Engine/Private/PlayerCameraManager.cpp
--- a/Engine/Source/Runtime/Engine/Private/PlayerCameraManager.cpp
+++ b/Engine/Source/Runtime/Engine/Private/PlayerCameraManager.cpp
@@ -40,7 +40,7 @@ APlayerCameraManager::APlayerCameraManager(const FObjectInitializer& ObjectIniti
 	DefaultAspectRatio = 1.33333f;
 	bDefaultConstrainAspectRatio = false;
 	DefaultOrthoWidth = 512.0f;
-	bHidden = true;
+	SetHidden(true);
 	bReplicates = false;
 	FreeCamDistance = 256.0f;
 	bDebugClientSideCamera = false;
@@ -52,7 +52,7 @@ APlayerCameraManager::APlayerCameraManager(const FObjectInitializer& ObjectIniti
 	ViewRollMax = 89.9f;
 	bUseClientSideCameraUpdates = true;
 	CameraStyle = NAME_Default;
-	bCanBeDamaged = false;
+	SetCanBeDamaged(false);
 	TimeSinceLastServerUpdateCamera = 0.0f;
 
 	// create dummy transform component
@@ -749,7 +749,7 @@ void APlayerCameraManager::PostInitializeComponents()
 	// spawn the temp CameraActor used for updating CameraAnims
 	FActorSpawnParameters SpawnInfo;
 	SpawnInfo.Owner = this;
-	SpawnInfo.Instigator = Instigator;
+	SpawnInfo.Instigator = GetInstigator();
 	SpawnInfo.SpawnCollisionHandlingOverride = ESpawnActorCollisionHandlingMethod::AlwaysSpawn;
 	SpawnInfo.ObjectFlags |= RF_Transient;	// We never want to save these temp actors into a map
 	AnimCameraActor = GetWorld()->SpawnActor<ACameraActor>(SpawnInfo);
@@ -1222,7 +1222,7 @@ AEmitterCameraLensEffectBase* APlayerCameraManager::AddCameraLensEffect(TSubclas
 			// spawn with viewtarget as the owner so bOnlyOwnerSee works as intended
 			FActorSpawnParameters SpawnInfo;
 			SpawnInfo.Owner = PCOwner->GetViewTarget();
-			SpawnInfo.Instigator = Instigator;
+			SpawnInfo.Instigator = GetInstigator();
 			SpawnInfo.SpawnCollisionHandlingOverride = ESpawnActorCollisionHandlingMethod::AlwaysSpawn;
 			SpawnInfo.ObjectFlags |= RF_Transient;	// We never want to save these into a map
 			
diff --git a/Engine/Source/Runtime/Engine/Private/PlayerController.cpp b/Engine/Source/Runtime/Engine/Private/PlayerController.cpp
--- a/Engine/Source/Runtime/Engine/Private/PlayerController.cpp
+++ b/Engine/Source/Runtime/Engine/Private/PlayerController.cpp
@@ -182,7 +182,7 @@ APlayerController::APlayerController(const FObjectInitializer& ObjectInitializer
 	if (RootComponent)
 	{
 		// We want to drive rotation with ControlRotation regardless of attachment state.
-		RootComponent->bAbsoluteRotation = true;
+		RootComponent->SetUsingAbsoluteRotation(true);
 	}
 }
 
@@ -709,7 +709,7 @@ void APlayerController::InitInputSystem()
 
 	// add the player to any matinees running so that it gets in on any cinematics already running, etc
 	// (already done on server in PostLogin())
-	if (Role < ROLE_Authority)
+	if (GetLocalRole() < ROLE_Authority)
 	{
 		TArray<AMatineeActor*> AllMatineeActors;
 		World->GetMatineeActors(AllMatineeActors);
@@ -793,7 +793,7 @@ void APlayerController::ClientRestart_Implementation(APawn* NewPawn)
 	GetPawn()->Controller = this;
 	GetPawn()->PawnClientRestart();
 	
-	if (Role < ROLE_Authority)
+	if (GetLocalRole() < ROLE_Authority)
 	{
 		ChangeState( NAME_Playing );
 		if (bAutoManageActiveCameraTarget)
@@ -1123,7 +1123,7 @@ void APlayerController::SpawnDefaultHUD()
 	UE_LOG(LogPlayerController, Verbose, TEXT("SpawnDefaultHUD"));
 	FActorSpawnParameters SpawnInfo;
 	SpawnInfo.Owner = this;
-	SpawnInfo.Instigator = Instigator;
+	SpawnInfo.Instigator = GetInstigator();
 	SpawnInfo.ObjectFlags |= RF_Transient;	// We never want to save HUDs into a map
 	MyHUD = GetWorld()->SpawnActor<AHUD>( SpawnInfo );
 }
@@ -1285,7 +1285,7 @@ void APlayerController::OnUnPossess()
 {
 	if (GetPawn() != NULL)
 	{
-		if (Role == ROLE_Authority)
+		if (GetLocalRole() == ROLE_Authority)
 		{
 			GetPawn()->SetReplicates(true);
 		}
@@ -1312,7 +1312,7 @@ void APlayerController::ClientSetHUD_Implementation(TSubclassOf<AHUD> NewHUDClas
 
 	FActorSpawnParameters SpawnInfo;
 	SpawnInfo.Owner = this;
-	SpawnInfo.Instigator = Instigator;
+	SpawnInfo.Instigator = GetInstigator();
 	SpawnInfo.ObjectFlags |= RF_Transient;	// We never want to save HUDs into a map
 
 	MyHUD = GetWorld()->SpawnActor<AHUD>(NewHUDClass, SpawnInfo );
@@ -1541,7 +1541,7 @@ void APlayerController::Destroyed()
 	if (GetPawn() != NULL)
 	{
 		// Handle players leaving the game
-		if (Player == NULL && Role == ROLE_Authority)
+		if (Player == NULL && GetLocalRole() == ROLE_Authority)
 		{
 			PawnLeavingGame();
 		}
@@ -2241,7 +2241,7 @@ void APlayerController::PlayerTick( float DeltaTime )
 		
 		if ( IsInState(NAME_Inactive) )
 		{
-			if (Role < ROLE_Authority)
+			if (GetLocalRole() < ROLE_Authority)
 			{
 				SafeServerCheckClientPossession();
 			}
@@ -2250,7 +2250,7 @@ void APlayerController::PlayerTick( float DeltaTime )
 		}
 		else if ( IsInState(NAME_Spectating) )
 		{
-			if (Role < ROLE_Authority)
+			if (GetLocalRole() < ROLE_Authority)
 			{
 				SafeServerUpdateSpectatorState();
 			}
@@ -2664,7 +2664,7 @@ void APlayerController::SpawnPlayerCameraManager()
 	// If no archetype specified, spawn an Engine.PlayerCameraManager.  NOTE all games should specify an archetype.
 	FActorSpawnParameters SpawnInfo;
 	SpawnInfo.Owner = this;
-	SpawnInfo.Instigator = Instigator;
+	SpawnInfo.Instigator = GetInstigator();
 	SpawnInfo.ObjectFlags |= RF_Transient;	// We never want to save camera managers into a map
 	if (PlayerCameraManagerClass != NULL)
 	{
@@ -4480,7 +4480,7 @@ void APlayerController::SetPlayer( UPlayer* InPlayer )
 	UpdateStateInputComponents();
 
 #if ENABLE_VISUAL_LOG
-	if (Role == ROLE_Authority && FVisualLogger::Get().IsRecordingOnServer())
+	if (GetLocalRole() == ROLE_Authority && FVisualLogger::Get().IsRecordingOnServer())
 	{
 		OnServerStartedVisualLogger(true);
 	}
@@ -4600,7 +4600,7 @@ void APlayerController::TickActor( float DeltaSeconds, ELevelTick TickType, FAct
 		// force physics update for clients that aren't sending movement updates in a timely manner 
 		// this prevents cheats associated with artificially induced ping spikes
 		// skip updates if pawn lost autonomous proxy role (e.g. TurnOff() call)
-		if (GetPawn() && !GetPawn()->IsPendingKill() && GetPawn()->GetRemoteRole() == ROLE_AutonomousProxy && GetPawn()->bReplicateMovement)
+		if (GetPawn() && !GetPawn()->IsPendingKill() && GetPawn()->GetRemoteRole() == ROLE_AutonomousProxy && GetPawn()->IsReplicatingMovement())
 		{
 			UMovementComponent* PawnMovement = GetPawn()->GetMovementComponent();
 			INetworkPredictionInterface* NetworkPredictionInterface = Cast<INetworkPredictionInterface>(PawnMovement);
@@ -4696,7 +4696,7 @@ void APlayerController::TickActor( float DeltaSeconds, ELevelTick TickType, FAct
 			}
 		}
 	}
-	else if (Role > ROLE_SimulatedProxy)
+	else if (GetLocalRole() > ROLE_SimulatedProxy)
 	{
 		// Process PlayerTick with input.
 		if (!PlayerInput && (Player == nullptr || Cast<ULocalPlayer>( Player ) != nullptr))
@@ -4847,7 +4847,7 @@ void APlayerController::EndPlayingState()
 
 void APlayerController::BeginSpectatingState()
 {
-	if (GetPawn() != NULL && Role == ROLE_Authority && ShouldKeepCurrentPawnUponSpectating() == false)
+	if (GetPawn() != NULL && GetLocalRole() == ROLE_Authority && ShouldKeepCurrentPawnUponSpectating() == false)
 	{
 		UnPossess();
 	}
diff --git a/Engine/Source/Runtime/Engine/Private/PlayerStart.cpp b/Engine/Source/Runtime/Engine/Private/PlayerStart.cpp
--- a/Engine/Source/Runtime/Engine/Private/PlayerStart.cpp
+++ b/Engine/Source/Runtime/Engine/Private/PlayerStart.cpp
@@ -40,7 +40,7 @@ APlayerStart::APlayerStart(const FObjectInitializer& ObjectInitializer)
 		if (GetGoodSprite())
 		{
 			GetGoodSprite()->Sprite = ConstructorStatics.PlayerStartTextureObject.Get();
-			GetGoodSprite()->RelativeScale3D = FVector(0.5f, 0.5f, 0.5f);
+			GetGoodSprite()->SetRelativeScale3D(FVector(0.5f, 0.5f, 0.5f));
 			GetGoodSprite()->SpriteInfo.Category = ConstructorStatics.ID_PlayerStart;
 			GetGoodSprite()->SpriteInfo.DisplayName = ConstructorStatics.NAME_PlayerStart;
 		}
diff --git a/Engine/Source/Runtime/Engine/Private/PlayerState.cpp b/Engine/Source/Runtime/Engine/Private/PlayerState.cpp
--- a/Engine/Source/Runtime/Engine/Private/PlayerState.cpp
+++ b/Engine/Source/Runtime/Engine/Private/PlayerState.cpp
@@ -20,7 +20,7 @@ APlayerState::APlayerState(const FObjectInitializer& ObjectInitializer)
 	SetRemoteRoleForBackwardsCompat(ROLE_SimulatedProxy);
 	bReplicates = true;
 	bAlwaysRelevant = true;
-	bReplicateMovement = false;
+	SetReplicatingMovement(false);
 	NetUpdateFrequency = 1;
 
 	// Note: this is very important to set to false. Though all replication infos are spawned at run time, during seamless travel
@@ -193,7 +193,7 @@ void APlayerState::PostInitializeComponents()
 		GameStateBase->AddPlayerState(this);
 	}
 
-	if ( Role < ROLE_Authority )
+	if (GetLocalRole() < ROLE_Authority)
 	{
 		return;
 	}
@@ -421,7 +421,7 @@ void APlayerState::UnregisterPlayerWithSession()
 APlayerState* APlayerState::Duplicate()
 {
 	FActorSpawnParameters SpawnInfo;
-	SpawnInfo.Instigator = Instigator;
+	SpawnInfo.Instigator = GetInstigator();
 	SpawnInfo.SpawnCollisionHandlingOverride = ESpawnActorCollisionHandlingMethod::AlwaysSpawn;
 	SpawnInfo.ObjectFlags |= RF_Transient;	// We never want to save player states into a map
 	APlayerState* NewPlayerState = GetWorld()->SpawnActor<APlayerState>(GetClass(), SpawnInfo );
diff --git a/Engine/Source/Runtime/Engine/Private/PrimitiveSceneProxy.cpp b/Engine/Source/Runtime/Engine/Private/PrimitiveSceneProxy.cpp
--- a/Engine/Source/Runtime/Engine/Private/PrimitiveSceneProxy.cpp
+++ b/Engine/Source/Runtime/Engine/Private/PrimitiveSceneProxy.cpp
@@ -89,7 +89,7 @@ FPrimitiveSceneProxy::FPrimitiveSceneProxy(const UPrimitiveComponent* InComponen
 ,	LightmapType(InComponent->LightmapType)
 ,	StatId()
 ,	DrawInGame(InComponent->IsVisible())
-,	DrawInEditor(InComponent->bVisible)
+,	DrawInEditor(InComponent->GetVisibleFlag())
 ,	bReceivesDecals(InComponent->bReceivesDecals)
 ,	bOnlyOwnerSee(InComponent->bOnlyOwnerSee)
 ,	bOwnerNoSee(InComponent->bOwnerNoSee)
@@ -195,7 +195,7 @@ FPrimitiveSceneProxy::FPrimitiveSceneProxy(const UPrimitiveComponent* InComponen
 	
 	if(InComponent->GetOwner())
 	{
-		DrawInGame &= !(InComponent->GetOwner()->bHidden);
+		DrawInGame &= !(InComponent->GetOwner()->IsHidden());
 		#if WITH_EDITOR
 			DrawInEditor &= !InComponent->GetOwner()->IsHiddenEd();
 		#endif
diff --git a/Engine/Source/Runtime/Engine/Private/SCS_Node.cpp b/Engine/Source/Runtime/Engine/Private/SCS_Node.cpp
--- a/Engine/Source/Runtime/Engine/Private/SCS_Node.cpp
+++ b/Engine/Source/Runtime/Engine/Private/SCS_Node.cpp
@@ -128,7 +128,7 @@ UActorComponent* USCS_Node::ExecuteNodeOnActor(AActor* Actor, USceneComponent* P
 					// does NOT occur when this component is instanced as part of dynamically spawning a Blueprint class in a cooked build (i.e.
 					// 'bIsDefaultTransform' will be 'false' in that situation). In order to maintain the same behavior between a nativized and
 					// non-nativized cooked build, if this ever changes, we would also need to update the code in AActor::PostSpawnInitialize().
-					WorldTransform.SetScale3D(NewSceneComp->RelativeScale3D);
+					WorldTransform.SetScale3D(NewSceneComp->GetRelativeScale3D());
 				}
 
 				if (RootRelativeRotationCache)
diff --git a/Engine/Source/Runtime/Engine/Private/SimpleConstructionScript.cpp b/Engine/Source/Runtime/Engine/Private/SimpleConstructionScript.cpp
--- a/Engine/Source/Runtime/Engine/Private/SimpleConstructionScript.cpp
+++ b/Engine/Source/Runtime/Engine/Private/SimpleConstructionScript.cpp
@@ -253,12 +253,14 @@ void USimpleConstructionScript::PostLoad()
 						if(Node->ParentComponentOrVariableName == NAME_None)
 						{
 							// Note that we have to check for nullptr here, because it may be an ActorComponent type
-							USceneComponent* SceneComponentTemplate = Cast<USceneComponent>(Node->ComponentTemplate);
-							if(SceneComponentTemplate != nullptr
-								&& SceneComponentTemplate->RelativeScale3D != FVector(1.0f, 1.0f, 1.0f))
+							if (USceneComponent* SceneComponentTemplate = Cast<USceneComponent>(Node->ComponentTemplate))
 							{
-								UE_LOG(LogBlueprint, Warning, TEXT("%s: Found non-native root component custom scale for %s (%s) saved prior to being usable; reverting to default scale."), *BPGeneratedClass->GetName(), *Node->GetVariableName().ToString(), *SceneComponentTemplate->RelativeScale3D.ToString());
-								SceneComponentTemplate->RelativeScale3D = FVector(1.0f, 1.0f, 1.0f);
+								const FVector ComponentRelativeScale3D = SceneComponentTemplate->GetRelativeScale3D();
+								if (ComponentRelativeScale3D != FVector(1.0f, 1.0f, 1.0f))
+								{
+									UE_LOG(LogBlueprint, Warning, TEXT("%s: Found non-native root component custom scale for %s (%s) saved prior to being usable; reverting to default scale."), *BPGeneratedClass->GetName(), *Node->GetVariableName().ToString(), *ComponentRelativeScale3D.ToString());
+									SceneComponentTemplate->SetRelativeScale3D_Direct(FVector(1.0f, 1.0f, 1.0f));
+								}
 							}
 
 							// Done - no need to fix up any other nodes.
diff --git a/Engine/Source/Runtime/Engine/Private/SpotLight.cpp b/Engine/Source/Runtime/Engine/Private/SpotLight.cpp
--- a/Engine/Source/Runtime/Engine/Private/SpotLight.cpp
+++ b/Engine/Source/Runtime/Engine/Private/SpotLight.cpp
@@ -22,7 +22,7 @@ ASpotLight::ASpotLight(const FObjectInitializer& ObjectInitializer)
 
 	SpotLightComponent = CastChecked<USpotLightComponent>(GetLightComponent());
 	SpotLightComponent->Mobility = EComponentMobility::Stationary;
-	SpotLightComponent->RelativeRotation = FRotator(-90, 0, 0);
+	SpotLightComponent->SetRelativeRotation(FRotator(-90, 0, 0));
 
 	RootComponent = SpotLightComponent;
 
diff --git a/Engine/Source/Runtime/Engine/Private/StaticMeshActor.cpp b/Engine/Source/Runtime/Engine/Private/StaticMeshActor.cpp
--- a/Engine/Source/Runtime/Engine/Private/StaticMeshActor.cpp
+++ b/Engine/Source/Runtime/Engine/Private/StaticMeshActor.cpp
@@ -25,7 +25,7 @@ FName AStaticMeshActor::StaticMeshComponentName(TEXT("StaticMeshComponent0"));
 AStaticMeshActor::AStaticMeshActor(const FObjectInitializer& ObjectInitializer)
 	: Super(ObjectInitializer)
 {
-	bCanBeDamaged = false;
+	SetCanBeDamaged(false);
 
 	StaticMeshComponent = CreateDefaultSubobject<UStaticMeshComponent>(StaticMeshComponentName);
 	StaticMeshComponent->SetCollisionProfileName(UCollisionProfile::BlockAll_ProfileName);
@@ -47,7 +47,7 @@ void AStaticMeshActor::BeginPlay()
 	//
 	// This is a short term fix until we find a better play for SetReplicates to be called in AActor.
 
-	if (Role == ROLE_Authority && bStaticMeshReplicateMovement)
+	if (GetLocalRole() == ROLE_Authority && bStaticMeshReplicateMovement)
 	{
 		bReplicates = false;
 		SetRemoteRoleForBackwardsCompat(ROLE_SimulatedProxy);
@@ -100,8 +100,8 @@ void AStaticMeshActor::LoadedFromAnotherClass(const FName& OldClassName)
 			StaticMeshComponent->SetCollisionProfileName(UCollisionProfile::PhysicsActor_ProfileName);
 			StaticMeshComponent->BodyInstance.bSimulatePhysics = true;
 
-			bCanBeDamaged = true;
-			bReplicateMovement = true;
+			SetCanBeDamaged(true);
+			SetReplicatingMovement(true);
 			SetRemoteRoleForBackwardsCompat(ROLE_SimulatedProxy);
 			bReplicates = true;
 		}
@@ -118,8 +118,8 @@ void AStaticMeshActor::PostEditChangeChainProperty(FPropertyChangedChainEvent& P
 		// If we want to replicate movement, set RemoteRole to match
 		if (PropertyChangedEvent.Property->GetFName() == FName(TEXT("bStaticMeshReplicateMovement")))
 		{
-			bReplicateMovement = bStaticMeshReplicateMovement;
-			SetReplicates(bReplicateMovement);
+			SetReplicatingMovement(bStaticMeshReplicateMovement);
+			SetReplicates(IsReplicatingMovement());
 		}
 	}
 
@@ -128,7 +128,7 @@ void AStaticMeshActor::PostEditChangeChainProperty(FPropertyChangedChainEvent& P
 	static FName SimulatePhysics_NAME(TEXT("bSimulatePhysics"));
 	if(TailPropName == SimulatePhysics_NAME)
 	{
-		bCanBeDamaged = StaticMeshComponent->BodyInstance.bSimulatePhysics;
+		SetCanBeDamaged(StaticMeshComponent->BodyInstance.bSimulatePhysics);
 	}
 }
 
@@ -189,7 +189,7 @@ void AStaticMeshActor::CheckForErrors()
 			AStaticMeshActor *A = Cast<AStaticMeshActor>(Overlaps[OverlapIdx].GetActor());
 			if ( A && (A != this) && (A->GetActorLocation() - GetActorLocation()).IsNearlyZero() && A->StaticMeshComponent
 				&& (A->StaticMeshComponent->GetStaticMesh() == StaticMeshComponent->GetStaticMesh()) && (A->GetActorRotation() == GetActorRotation())
-				&& (A->StaticMeshComponent->RelativeScale3D - StaticMeshComponent->RelativeScale3D).IsNearlyZero() )
+				&& (A->StaticMeshComponent->GetRelativeScale3D() == StaticMeshComponent->GetRelativeScale3D()) )
 			{
 				FFormatNamedArguments Arguments;
 				Arguments.Add(TEXT("ActorName0"), FText::FromString(GetName()));
diff --git a/Engine/Source/Runtime/Engine/Private/TargetPoint.cpp b/Engine/Source/Runtime/Engine/Private/TargetPoint.cpp
--- a/Engine/Source/Runtime/Engine/Private/TargetPoint.cpp
+++ b/Engine/Source/Runtime/Engine/Private/TargetPoint.cpp
@@ -39,7 +39,7 @@ ATargetPoint::ATargetPoint(const FObjectInitializer& ObjectInitializer)
 		if (SpriteComponent)
 		{
 			SpriteComponent->Sprite = ConstructorStatics.TargetIconObject.Get();
-			SpriteComponent->RelativeScale3D = FVector(0.35f, 0.35f, 0.35f);
+			SpriteComponent->SetRelativeScale3D_Direct(FVector(0.35f, 0.35f, 0.35f));
 			SpriteComponent->SpriteInfo.Category = ConstructorStatics.ID_TargetPoint;
 			SpriteComponent->SpriteInfo.DisplayName = ConstructorStatics.NAME_TargetPoint;
 			SpriteComponent->bIsScreenSizeScaled = true;
@@ -60,13 +60,15 @@ ATargetPoint::ATargetPoint(const FObjectInitializer& ObjectInitializer)
 
 			// Counteract the scaled down parent so that the arrow is large enough to see.
 			if (SpriteComponent)
-				ArrowComponent->SetRelativeScale3D((FVector::OneVector / SpriteComponent->RelativeScale3D));
+			{
+				ArrowComponent->SetRelativeScale3D((FVector::OneVector / SpriteComponent->GetRelativeScale3D()));
+			}
 		}
 	}
 #endif // WITH_EDITORONLY_DATA
 
-	bHidden = true;
-	bCanBeDamaged = false;
+	SetHidden(true);
+	SetCanBeDamaged(false);
 }
 
 #if WITH_EDITORONLY_DATA
diff --git a/Engine/Source/Runtime/Engine/Private/TriggerBase.cpp b/Engine/Source/Runtime/Engine/Private/TriggerBase.cpp
--- a/Engine/Source/Runtime/Engine/Private/TriggerBase.cpp
+++ b/Engine/Source/Runtime/Engine/Private/TriggerBase.cpp
@@ -9,8 +9,8 @@
 ATriggerBase::ATriggerBase(const FObjectInitializer& ObjectInitializer)
 	: Super(ObjectInitializer)
 {
-	bHidden = true;
-	bCanBeDamaged = false;
+	SetHidden(true);
+	SetCanBeDamaged(false);
 
 	// ATriggerBase is requesting UShapeComponent which is abstract, however it is responsibility
 	// of a derived class to override this type with ObjectInitializer.SetDefaultSubobjectClass.
@@ -41,7 +41,7 @@ ATriggerBase::ATriggerBase(const FObjectInitializer& ObjectInitializer)
 		static FConstructorStatics ConstructorStatics;
 
 		SpriteComponent->Sprite = ConstructorStatics.TriggerTextureObject.Get();
-		SpriteComponent->RelativeScale3D = FVector(0.5f, 0.5f, 0.5f);
+		SpriteComponent->SetRelativeScale3D(FVector(0.5f, 0.5f, 0.5f));
 		SpriteComponent->bHiddenInGame = false;
 		SpriteComponent->SpriteInfo.Category = ConstructorStatics.ID_Triggers;
 		SpriteComponent->SpriteInfo.DisplayName = ConstructorStatics.NAME_Triggers;
diff --git a/Engine/Source/Runtime/Engine/Private/VectorFieldVolume.cpp b/Engine/Source/Runtime/Engine/Private/VectorFieldVolume.cpp
--- a/Engine/Source/Runtime/Engine/Private/VectorFieldVolume.cpp
+++ b/Engine/Source/Runtime/Engine/Private/VectorFieldVolume.cpp
@@ -36,7 +36,7 @@ AVectorFieldVolume::AVectorFieldVolume(const FObjectInitializer& ObjectInitializ
 		SpriteComponent->bIsScreenSizeScaled = true;
 		SpriteComponent->SpriteInfo.Category = ConstructorStatics.ID_Effects;
 		SpriteComponent->SpriteInfo.DisplayName = ConstructorStatics.NAME_Effects;
-		SpriteComponent->bAbsoluteScale = true;
+		SpriteComponent->SetUsingAbsoluteScale(true);
 		SpriteComponent->SetupAttachment(VectorFieldComponent);
 		SpriteComponent->bReceivesDecals = false;
 	}
diff --git a/Engine/Source/Runtime/Engine/Private/Volume.cpp b/Engine/Source/Runtime/Engine/Private/Volume.cpp
--- a/Engine/Source/Runtime/Engine/Private/Volume.cpp
+++ b/Engine/Source/Runtime/Engine/Private/Volume.cpp
@@ -23,12 +23,12 @@ AVolume::AVolume(const FObjectInitializer& ObjectInitializer)
 	static FName CollisionProfileName(TEXT("OverlapAll"));
 	GetBrushComponent()->SetCollisionProfileName(CollisionProfileName);
 	GetBrushComponent()->SetGenerateOverlapEvents(true);
-	bReplicateMovement = false;
+	SetReplicatingMovement(false);
 #if WITH_EDITORONLY_DATA
 	bActorLabelEditable = true;
 #endif // WITH_EDITORONLY_DATA
 
-	bCanBeDamaged = false;
+	SetCanBeDamaged(false);
 }
 
 #if WITH_EDITOR
diff --git a/Engine/Source/Runtime/Engine/Private/WindDirectionalSource.cpp b/Engine/Source/Runtime/Engine/Private/WindDirectionalSource.cpp
--- a/Engine/Source/Runtime/Engine/Private/WindDirectionalSource.cpp
+++ b/Engine/Source/Runtime/Engine/Private/WindDirectionalSource.cpp
@@ -48,7 +48,7 @@ AWindDirectionalSource::AWindDirectionalSource(const FObjectInitializer& ObjectI
 		if (UBillboardComponent* SpriteComp = GetSpriteComponent())
 		{
 			SpriteComp->Sprite = ConstructorStatics.SpriteTexture.Get();
-			SpriteComp->RelativeScale3D = FVector(0.5f, 0.5f, 0.5f);
+			SpriteComp->SetRelativeScale3D(FVector(0.5f, 0.5f, 0.5f));
 			SpriteComp->SpriteInfo.Category = ConstructorStatics.ID_Wind;
 			SpriteComp->SpriteInfo.DisplayName = ConstructorStatics.NAME_Wind;
 			SpriteComp->SetupAttachment(Component);
diff --git a/Engine/Source/Runtime/Engine/Private/WorldSettings.cpp b/Engine/Source/Runtime/Engine/Private/WorldSettings.cpp
--- a/Engine/Source/Runtime/Engine/Private/WorldSettings.cpp
+++ b/Engine/Source/Runtime/Engine/Private/WorldSettings.cpp
@@ -97,7 +97,7 @@ PRAGMA_ENABLE_DEPRECATION_WARNINGS
 	MatineeTimeDilation = 1.0f;
 	DemoPlayTimeDilation = 1.0f;
 	PackedLightAndShadowMapTextureSize = 1024;
-	bHidden = false;
+	SetHidden(false);
 
 	DefaultColorScale = FVector(1.0f, 1.0f, 1.0f);
 	DefaultMaxDistanceFieldOcclusionDistance = 600;
@@ -180,7 +180,7 @@ void AWorldSettings::PreInitializeComponents()
 				{
 					FActorSpawnParameters SpawnParameters;
 					SpawnParameters.Owner = this;
-					SpawnParameters.Instigator = Instigator;
+					SpawnParameters.Instigator = GetInstigator();
 					SpawnParameters.ObjectFlags |= RF_Transient;	// We never want to save particle event managers into a map
 					World->MyParticleEventManager = World->SpawnActor<AParticleEventManager>(ParticleEventManagerClass, SpawnParameters);
 				}
diff --git a/Engine/Source/Runtime/Foliage/Private/FoliageComponent.cpp b/Engine/Source/Runtime/Foliage/Private/FoliageComponent.cpp
--- a/Engine/Source/Runtime/Foliage/Private/FoliageComponent.cpp
+++ b/Engine/Source/Runtime/Foliage/Private/FoliageComponent.cpp
@@ -127,7 +127,7 @@ void AInteractiveFoliageActor::SetupCollisionCylinder()
 	if (GetStaticMeshComponent()->GetStaticMesh())
 	{
 		const FBoxSphereBounds MeshBounds = GetStaticMeshComponent()->GetStaticMesh()->GetBounds();
-		const FVector Scale3D = GetStaticMeshComponent()->RelativeScale3D;
+		const FVector Scale3D = GetStaticMeshComponent()->GetRelativeScale3D();
 		// Set the cylinder's radius based off of the static mesh's bounds radius
 		// CollisionRadius is in world space so apply the actor's scale
 		CapsuleComponent->SetCapsuleSize(MeshBounds.SphereRadius * .7f * FMath::Max(Scale3D.X, Scale3D.Y), MeshBounds.BoxExtent.Z * Scale3D.Z);
diff --git a/Engine/Source/Runtime/Foliage/Private/FoliageInstanceBase.cpp b/Engine/Source/Runtime/Foliage/Private/FoliageInstanceBase.cpp
--- a/Engine/Source/Runtime/Foliage/Private/FoliageInstanceBase.cpp
+++ b/Engine/Source/Runtime/Foliage/Private/FoliageInstanceBase.cpp
@@ -34,9 +34,9 @@ void FFoliageInstanceBaseInfo::UpdateLocationFromComponent(UActorComponent* InCo
 			const USceneComponent* RootComponent = Owner->GetRootComponent();
 			if (RootComponent)
 			{
-				CachedLocation = RootComponent->RelativeLocation;
-				CachedRotation = RootComponent->RelativeRotation;
-				CachedDrawScale = RootComponent->RelativeScale3D;
+				CachedLocation = RootComponent->GetRelativeLocation();
+				CachedRotation = RootComponent->GetRelativeRotation();
+				CachedDrawScale = RootComponent->GetRelativeScale3D();
 			}
 		}
 	}
diff --git a/Engine/Source/Runtime/Foliage/Private/InteractiveFoliageActor.cpp b/Engine/Source/Runtime/Foliage/Private/InteractiveFoliageActor.cpp
--- a/Engine/Source/Runtime/Foliage/Private/InteractiveFoliageActor.cpp
+++ b/Engine/Source/Runtime/Foliage/Private/InteractiveFoliageActor.cpp
@@ -23,7 +23,7 @@ AInteractiveFoliageActor::AInteractiveFoliageActor(const FObjectInitializer& Obj
 	RootComponent = CapsuleComponent;
 
 	PrimaryActorTick.bCanEverTick = true;
-	bCanBeDamaged = true;
+	SetCanBeDamaged(true);
 	bCollideWhenPlacing = true;
 	SpawnCollisionHandlingMethod = ESpawnActorCollisionHandlingMethod::AdjustIfPossibleButDontSpawnIfColliding;
 	FoliageDamageImpulseScale = 20.0f;
diff --git a/Engine/Source/Runtime/GameplayTasks/Private/GameplayTasksComponent.cpp b/Engine/Source/Runtime/GameplayTasks/Private/GameplayTasksComponent.cpp
--- a/Engine/Source/Runtime/GameplayTasks/Private/GameplayTasksComponent.cpp
+++ b/Engine/Source/Runtime/GameplayTasks/Private/GameplayTasksComponent.cpp
@@ -51,7 +51,7 @@ UGameplayTasksComponent::UGameplayTasksComponent(const FObjectInitializer& Objec
 	PrimaryComponentTick.bStartWithTickEnabled = false;
 	PrimaryComponentTick.bCanEverTick = true;
 
-	bReplicates = true;
+	SetIsReplicatedByDefault(true);
 	bInEventProcessingInProgress = false;
 	TopActivePriority = 0;
 }
@@ -250,7 +250,7 @@ bool UGameplayTasksComponent::GetShouldTick() const
 
 void UGameplayTasksComponent::RequestTicking()
 {
-	if (bIsActive == false)
+	if (IsActive() == false)
 	{
 		SetActive(true);
 	}
@@ -259,7 +259,7 @@ void UGameplayTasksComponent::RequestTicking()
 void UGameplayTasksComponent::UpdateShouldTick()
 {
 	const bool bShouldTick = GetShouldTick();	
-	if (bIsActive != bShouldTick)
+	if (IsActive() != bShouldTick)
 	{
 		SetActive(bShouldTick);
 	}
diff --git a/Engine/Source/Runtime/GameplayTasks/Private/Tasks/GameplayTask_SpawnActor.cpp b/Engine/Source/Runtime/GameplayTasks/Private/Tasks/GameplayTask_SpawnActor.cpp
--- a/Engine/Source/Runtime/GameplayTasks/Private/Tasks/GameplayTask_SpawnActor.cpp
+++ b/Engine/Source/Runtime/GameplayTasks/Private/Tasks/GameplayTask_SpawnActor.cpp
@@ -17,7 +17,7 @@ UGameplayTask_SpawnActor* UGameplayTask_SpawnActor::SpawnActor(TScriptInterface<
 		AActor* TaskOwnerActor = TaskOwner->GetGameplayTaskOwner(nullptr);
 		if (TaskOwnerActor)
 		{
-			bCanSpawn = (TaskOwnerActor->Role == ROLE_Authority);
+			bCanSpawn = (TaskOwnerActor->GetLocalRole() == ROLE_Authority);
 		}
 		else
 		{
diff --git a/Engine/Source/Runtime/HeadMountedDisplay/Private/MotionControllerComponent.cpp b/Engine/Source/Runtime/HeadMountedDisplay/Private/MotionControllerComponent.cpp
--- a/Engine/Source/Runtime/HeadMountedDisplay/Private/MotionControllerComponent.cpp
+++ b/Engine/Source/Runtime/HeadMountedDisplay/Private/MotionControllerComponent.cpp
@@ -126,7 +126,7 @@ void UMotionControllerComponent::TickComponent(float DeltaTime, enum ELevelTick
 {
 	Super::TickComponent(DeltaTime, TickType, ThisTickFunction);
 
-	if (bIsActive)
+	if (IsActive())
 	{
 		FVector Position;
 		FRotator Orientation;
@@ -469,7 +469,7 @@ void UMotionControllerComponent::RefreshDisplayComponent(const bool bForceDestro
 				}
 
 				DisplayComponent->SetHiddenInGame(bHiddenInGame);
-				DisplayComponent->SetVisibility(bVisible);
+				DisplayComponent->SetVisibility(GetVisibleFlag());
 			}
 		}
 		else if (DisplayComponent)
diff --git a/Engine/Source/Runtime/HeadMountedDisplay/Private/XRTrackingSystemBase.cpp b/Engine/Source/Runtime/HeadMountedDisplay/Private/XRTrackingSystemBase.cpp
--- a/Engine/Source/Runtime/HeadMountedDisplay/Private/XRTrackingSystemBase.cpp
+++ b/Engine/Source/Runtime/HeadMountedDisplay/Private/XRTrackingSystemBase.cpp
@@ -164,7 +164,7 @@ FTransform FXRTrackingSystemBase::ComputeTrackingToWorldTransform(FWorldContext&
 						// emulates AActor::CalcCamera(); the PlayerCameraManager uses ViewTarget->CalcCamera() for 
 						// the HMD view's basis (regardless of whether bLockToHmd is set), CalcCamera() just finds 
 						// the first active cam component and chooses that
-						if (Cam->bIsActive)
+						if (Cam->IsActive())
 						{
 							PlayerCamera = Cam;
 							break;
@@ -178,17 +178,17 @@ FTransform FXRTrackingSystemBase::ComputeTrackingToWorldTransform(FWorldContext&
 						{
 							TrackingToWorld = ViewParent->GetComponentTransform();
 
-							if (PlayerCamera->bAbsoluteLocation)
+							if (PlayerCamera->IsUsingAbsoluteLocation())
 							{
 								TrackingToWorld.SetLocation(FVector::ZeroVector);
 							}
 
-							if (PlayerCamera->bAbsoluteRotation)
+							if (PlayerCamera->IsUsingAbsoluteRotation())
 							{
 								TrackingToWorld.SetRotation(FQuat::Identity);
 							}
 
-							if (PlayerCamera->bAbsoluteScale)
+							if (PlayerCamera->IsUsingAbsoluteScale())
 							{
 								TrackingToWorld.SetScale3D(FVector::OneVector);
 							}
diff --git a/Engine/Source/Runtime/Landscape/Classes/LandscapeGizmoActiveActor.h b/Engine/Source/Runtime/Landscape/Classes/LandscapeGizmoActiveActor.h
--- a/Engine/Source/Runtime/Landscape/Classes/LandscapeGizmoActiveActor.h
+++ b/Engine/Source/Runtime/Landscape/Classes/LandscapeGizmoActiveActor.h
@@ -169,25 +169,25 @@ public:
 	// @todo document
 	float GetWidth() const
 	{
-		return Width * GetRootComponent()->RelativeScale3D.X;
+		return Width * GetRootComponent()->GetRelativeScale3D().X;
 	}
 
 	// @todo document
 	float GetHeight() const
 	{
-		return Height * GetRootComponent()->RelativeScale3D.Y;
+		return Height * GetRootComponent()->GetRelativeScale3D().Y;
 	}
 
 	// @todo document
 	float GetLength() const
 	{
-		return LengthZ * GetRootComponent()->RelativeScale3D.Z;
+		return LengthZ * GetRootComponent()->GetRelativeScale3D().Z;
 	}
 
 	// @todo document
 	void SetLength(float WorldLength)
 	{
-		LengthZ = WorldLength / GetRootComponent()->RelativeScale3D.Z;
+		LengthZ = WorldLength / GetRootComponent()->GetRelativeScale3D().Z;
 	}
 
 	static const int32 DataTexSize = 128;
diff --git a/Engine/Source/Runtime/Landscape/Private/Landscape.cpp b/Engine/Source/Runtime/Landscape/Private/Landscape.cpp
--- a/Engine/Source/Runtime/Landscape/Private/Landscape.cpp
+++ b/Engine/Source/Runtime/Landscape/Private/Landscape.cpp
@@ -701,8 +701,9 @@ void ULandscapeComponent::PostLoad()
 		UpdatedSharedPropertiesFromActor();	
 
 		// check SectionBaseX/Y are correct
-		int32 CheckSectionBaseX = FMath::RoundToInt(RelativeLocation.X) + LandscapeProxy->LandscapeSectionOffset.X;
-		int32 CheckSectionBaseY = FMath::RoundToInt(RelativeLocation.Y) + LandscapeProxy->LandscapeSectionOffset.Y;
+		const FVector LocalRelativeLocation = GetRelativeLocation();
+		int32 CheckSectionBaseX = FMath::RoundToInt(LocalRelativeLocation.X) + LandscapeProxy->LandscapeSectionOffset.X;
+		int32 CheckSectionBaseY = FMath::RoundToInt(LocalRelativeLocation.Y) + LandscapeProxy->LandscapeSectionOffset.Y;
 		if (CheckSectionBaseX != SectionBaseX ||
 			CheckSectionBaseY != SectionBaseY)
 		{
@@ -717,15 +718,18 @@ void ULandscapeComponent::PostLoad()
 	if (GIsEditor && !HasAnyFlags(RF_ClassDefaultObject))
 	{
 		// This is to ensure that component relative location is exact section base offset value
+		FVector LocalRelativeLocation = GetRelativeLocation();
 		float CheckRelativeLocationX = float(SectionBaseX - LandscapeProxy->LandscapeSectionOffset.X);
 		float CheckRelativeLocationY = float(SectionBaseY - LandscapeProxy->LandscapeSectionOffset.Y);
-		if (CheckRelativeLocationX != RelativeLocation.X || 
-			CheckRelativeLocationY != RelativeLocation.Y)
+		if (CheckRelativeLocationX != LocalRelativeLocation.X ||
+			CheckRelativeLocationY != LocalRelativeLocation.Y)
 		{
 			UE_LOG(LogLandscape, Warning, TEXT("LandscapeComponent RelativeLocation disagrees with its section base, attempted automated fix: '%s', %f,%f vs %f,%f."),
-				*GetFullName(), RelativeLocation.X, RelativeLocation.Y, CheckRelativeLocationX, CheckRelativeLocationY);
-			RelativeLocation.X = CheckRelativeLocationX;
-			RelativeLocation.Y = CheckRelativeLocationY;
+				*GetFullName(), LocalRelativeLocation.X, LocalRelativeLocation.Y, CheckRelativeLocationX, CheckRelativeLocationY);
+			LocalRelativeLocation.X = CheckRelativeLocationX;
+			LocalRelativeLocation.Y = CheckRelativeLocationY;
+
+			SetRelativeLocation_Direct(LocalRelativeLocation);
 		}
 
 		// Remove standalone flags from data textures to ensure data is unloaded in the editor when reverting an unsaved level.
@@ -1011,16 +1015,16 @@ ALandscapeProxy::ALandscapeProxy(const FObjectInitializer& ObjectInitializer)
 
 	bReplicates = false;
 	NetUpdateFrequency = 10.0f;
-	bHidden = false;
-	bReplicateMovement = false;
-	bCanBeDamaged = false;
+	SetHidden(false);
+	SetReplicatingMovement(false);
+	SetCanBeDamaged(false);
 	// by default we want to see the Landscape shadows even in the far shadow cascades
 	bCastFarShadow = true;
 	bAffectDistanceFieldLighting = true;
 
 	USceneComponent* SceneComponent = CreateDefaultSubobject<USceneComponent>(TEXT("RootComponent0"));
 	RootComponent = SceneComponent;
-	RootComponent->RelativeScale3D = FVector(128.0f, 128.0f, 256.0f); // Old default scale, preserved for compatibility. See ULandscapeEditorObject::NewLandscape_Scale
+	RootComponent->SetRelativeScale3D(FVector(128.0f, 128.0f, 256.0f)); // Old default scale, preserved for compatibility. See ULandscapeEditorObject::NewLandscape_Scale
 	RootComponent->Mobility = EComponentMobility::Static;
 	LandscapeSectionOffset = FIntPoint::ZeroValue;
 
@@ -2929,7 +2933,7 @@ void ULandscapeInfo::RegisterActor(ALandscapeProxy* Proxy, bool bMapCheck)
 			ComponentSizeQuads = Proxy->ComponentSizeQuads;
 			ComponentNumSubsections = Proxy->NumSubsections;
 			SubsectionSizeQuads = Proxy->SubsectionSizeQuads;
-			DrawScale = Proxy->GetRootComponent() != nullptr ? Proxy->GetRootComponent()->RelativeScale3D : FVector(100.0f);
+			DrawScale = Proxy->GetRootComponent() != nullptr ? Proxy->GetRootComponent()->GetRelativeScale3D() : FVector(100.0f);
 		}
 
 		// check that passed actor matches all shared parameters
@@ -2938,10 +2942,10 @@ void ULandscapeInfo::RegisterActor(ALandscapeProxy* Proxy, bool bMapCheck)
 		check(ComponentNumSubsections == Proxy->NumSubsections);
 		check(SubsectionSizeQuads == Proxy->SubsectionSizeQuads);
 
-		if (Proxy->GetRootComponent() != nullptr && !DrawScale.Equals(Proxy->GetRootComponent()->RelativeScale3D))
+		if (Proxy->GetRootComponent() != nullptr && !DrawScale.Equals(Proxy->GetRootComponent()->GetRelativeScale3D()))
 		{
 			UE_LOG(LogLandscape, Warning, TEXT("Landscape proxy (%s) scale (%s) does not match to main actor scale (%s)."),
-				*Proxy->GetName(), *Proxy->GetRootComponent()->RelativeScale3D.ToCompactString(), *DrawScale.ToCompactString());
+				*Proxy->GetName(), *Proxy->GetRootComponent()->GetRelativeScale3D().ToCompactString(), *DrawScale.ToCompactString());
 		}
 
 		// register
@@ -3323,7 +3327,7 @@ ALandscapeProxy::~ALandscapeProxy()
 ALandscapeMeshProxyActor::ALandscapeMeshProxyActor(const FObjectInitializer& ObjectInitializer)
 : Super(ObjectInitializer)
 {
-	bCanBeDamaged = false;
+	SetCanBeDamaged(false);
 
 	LandscapeMeshProxyComponent = CreateDefaultSubobject<ULandscapeMeshProxyComponent>(TEXT("LandscapeMeshProxyComponent0"));
 	LandscapeMeshProxyComponent->SetCollisionProfileName(UCollisionProfile::NoCollision_ProfileName);
diff --git a/Engine/Source/Runtime/Landscape/Private/LandscapeCollision.cpp b/Engine/Source/Runtime/Landscape/Private/LandscapeCollision.cpp
--- a/Engine/Source/Runtime/Landscape/Private/LandscapeCollision.cpp
+++ b/Engine/Source/Runtime/Landscape/Private/LandscapeCollision.cpp
@@ -612,7 +612,7 @@ bool ULandscapeHeightfieldCollisionComponent::CookCollisionData(const FName& For
 	UPhysicalMaterial* DefMaterial = Proxy->DefaultPhysMaterial ? Proxy->DefaultPhysMaterial : GEngine->DefaultPhysMaterial;
 
 	// GetComponentTransform() might not be initialized at this point, so use landscape transform
-	const FVector LandscapeScale = Proxy->GetRootComponent()->RelativeScale3D;
+	const FVector LandscapeScale = Proxy->GetRootComponent()->GetRelativeScale3D();
 	const bool bIsMirrored = (LandscapeScale.X*LandscapeScale.Y*LandscapeScale.Z) < 0.f;
 
 	const bool bGenerateSimpleCollision = SimpleCollisionSizeQuads > 0 && !bUseDefMaterial;
@@ -1692,15 +1692,17 @@ void ULandscapeHeightfieldCollisionComponent::PostLoad()
 		if (ensure(LandscapeProxy) && GIsEditor)
 		{
 			// This is to ensure that component relative location is exact section base offset value
+			FVector LocalRelativeLocation = GetRelativeLocation();
 			float CheckRelativeLocationX = float(SectionBaseX - LandscapeProxy->LandscapeSectionOffset.X);
 			float CheckRelativeLocationY = float(SectionBaseY - LandscapeProxy->LandscapeSectionOffset.Y);
-			if (CheckRelativeLocationX != RelativeLocation.X || 
-				CheckRelativeLocationY != RelativeLocation.Y)
+			if (CheckRelativeLocationX != LocalRelativeLocation.X ||
+				CheckRelativeLocationY != LocalRelativeLocation.Y)
 			{
 				UE_LOG(LogLandscape, Warning, TEXT("ULandscapeHeightfieldCollisionComponent RelativeLocation disagrees with its section base, attempted automated fix: '%s', %f,%f vs %f,%f."),
-					*GetFullName(), RelativeLocation.X, RelativeLocation.Y, CheckRelativeLocationX, CheckRelativeLocationY);
-				RelativeLocation.X = CheckRelativeLocationX;
-				RelativeLocation.Y = CheckRelativeLocationY;
+					*GetFullName(), LocalRelativeLocation.X, LocalRelativeLocation.Y, CheckRelativeLocationX, CheckRelativeLocationY);
+				LocalRelativeLocation.X = CheckRelativeLocationX;
+				LocalRelativeLocation.Y = CheckRelativeLocationY;
+				SetRelativeLocation_Direct(LocalRelativeLocation);
 			}
 		}
 
diff --git a/Engine/Source/Runtime/Landscape/Private/LandscapeEdit.cpp b/Engine/Source/Runtime/Landscape/Private/LandscapeEdit.cpp
--- a/Engine/Source/Runtime/Landscape/Private/LandscapeEdit.cpp
+++ b/Engine/Source/Runtime/Landscape/Private/LandscapeEdit.cpp
@@ -1084,7 +1084,7 @@ void ULandscapeComponent::RecreateCollisionComponent(bool bUseSimpleCollision)
 		CollisionComp = NewObject<ULandscapeHeightfieldCollisionComponent>(Proxy, NAME_None, RF_Transactional);
 	}
 
-	CollisionComp->SetRelativeLocation(RelativeLocation);
+	CollisionComp->SetRelativeLocation(GetRelativeLocation());
 	CollisionComp->SetupAttachment(Proxy->GetRootComponent(), NAME_None);
 	Proxy->CollisionComponents.Add(CollisionComp);
 
@@ -1995,7 +1995,7 @@ float ULandscapeComponent::GetLayerWeightAtLocation(const FVector& InLocation, U
 	// TODO: Root landscape isn't always loaded, would Proxy suffice?
 	if (ALandscape* Landscape = GetLandscapeActor())
 	{
-		const FVector DrawScale = Landscape->GetRootComponent()->RelativeScale3D;
+		const FVector DrawScale = Landscape->GetRootComponent()->GetRelativeScale3D();
 		float TestX = (InLocation.X - Landscape->GetActorLocation().X) / DrawScale.X - (float)GetSectionBase().X;
 		float TestY = (InLocation.Y - Landscape->GetActorLocation().Y) / DrawScale.Y - (float)GetSectionBase().Y;
 
@@ -2326,7 +2326,7 @@ LANDSCAPE_API void ALandscapeProxy::Import(const FGuid& InGuid, int32 InMinX, in
 		}
 	}
 
-	const FVector DrawScale3D = GetRootComponent()->RelativeScale3D;
+	const FVector DrawScale3D = GetRootComponent()->GetRelativeScale3D();
 
 	// layer to import data (Final or 1st layer)
 	const FGuid FinalLayerGuid = FGuid();
@@ -3602,7 +3602,7 @@ void ALandscapeProxy::EditorApplyScale(const FVector& DeltaScale, const FVector*
 	ModifiedScale.X = ModifiedScale.Y = (FMath::Abs(DeltaScale.X) > FMath::Abs(DeltaScale.Y)) ? DeltaScale.X : DeltaScale.Y;
 
 	// Correct for attempts to scale to 0 on any axis
-	FVector CurrentScale = GetRootComponent()->RelativeScale3D;
+	FVector CurrentScale = GetRootComponent()->GetRelativeScale3D();
 	if (AActor::bUsePercentageBasedScaling)
 	{
 		if (ModifiedScale.X == -1)
@@ -4023,17 +4023,18 @@ void ALandscapeProxy::PostEditChangeProperty(FPropertyChangedEvent& PropertyChan
 		// RelativeScale3D isn't even a property of ALandscapeProxy, it's a property of the root component
 		if (RootComponent)
 		{
-			FVector ModifiedScale = RootComponent->RelativeScale3D;
+			const FVector OriginalScale = RootComponent->GetRelativeScale3D();
+			FVector ModifiedScale = OriginalScale;
 
 			// Lock X and Y scaling to the same value
 			if (SubPropertyName == FName("Y"))
 			{
-				ModifiedScale.X = FMath::Abs(RootComponent->RelativeScale3D.Y)*FMath::Sign(ModifiedScale.X);
+				ModifiedScale.X = FMath::Abs(OriginalScale.Y)*FMath::Sign(ModifiedScale.X);
 			}
 			else
 			{
 				// There's no "if name == X" here so that if we can't tell which has changed out of X and Y, we just use X
-				ModifiedScale.Y = FMath::Abs(RootComponent->RelativeScale3D.X)*FMath::Sign(ModifiedScale.Y);
+				ModifiedScale.Y = FMath::Abs(OriginalScale.X)*FMath::Sign(ModifiedScale.Y);
 			}
 
 			ULandscapeInfo* Info = GetLandscapeInfo();
diff --git a/Engine/Source/Runtime/Landscape/Private/LandscapeEditLayers.cpp b/Engine/Source/Runtime/Landscape/Private/LandscapeEditLayers.cpp
--- a/Engine/Source/Runtime/Landscape/Private/LandscapeEditLayers.cpp
+++ b/Engine/Source/Runtime/Landscape/Private/LandscapeEditLayers.cpp
@@ -2992,7 +2992,7 @@ int32 ALandscape::RegenerateLayersHeightmaps(const TArray<ULandscapeComponent*>&
 		}
 
 		ShaderParams.GenerateNormals = true;
-		ShaderParams.GridSize = GetRootComponent()->RelativeScale3D;
+		ShaderParams.GridSize = GetRootComponent()->GetRelativeScale3D();
 
 		// Broadcast Event of the Full Render
 		if ((HeightmapUpdateModes & Update_Heightmap_All) == Update_Heightmap_All)
diff --git a/Engine/Source/Runtime/Landscape/Private/LandscapeGizmoActor.cpp b/Engine/Source/Runtime/Landscape/Private/LandscapeGizmoActor.cpp
--- a/Engine/Source/Runtime/Landscape/Private/LandscapeGizmoActor.cpp
+++ b/Engine/Source/Runtime/Landscape/Private/LandscapeGizmoActor.cpp
@@ -434,7 +434,7 @@ ALandscapeGizmoActor::ALandscapeGizmoActor(const FObjectInitializer& ObjectIniti
 		static FConstructorStatics ConstructorStatics;
 
 		SpriteComponent->Sprite = ConstructorStatics.DecalActorIconTexture.Get();
-		SpriteComponent->RelativeScale3D = FVector(0.5f, 0.5f, 0.5f);
+		SpriteComponent->SetRelativeScale3D(FVector(0.5f, 0.5f, 0.5f));
 		SpriteComponent->bHiddenInGame = true;
 		SpriteComponent->SpriteInfo.Category = ConstructorStatics.ID_Misc;
 		SpriteComponent->SpriteInfo.DisplayName = ConstructorStatics.NAME_Misc;
@@ -469,9 +469,9 @@ void ALandscapeGizmoActor::Duplicate(ALandscapeGizmoActor* Gizmo)
 	Gizmo->SetActorLocation( GetActorLocation(), false );
 	Gizmo->SetActorRotation( GetActorRotation() );
 
-	if( Gizmo->GetRootComponent() != NULL && GetRootComponent() != NULL )
+	if (Gizmo->GetRootComponent() != NULL && GetRootComponent() != NULL)
 	{
-		Gizmo->GetRootComponent()->SetRelativeScale3D( GetRootComponent()->RelativeScale3D );
+		Gizmo->GetRootComponent()->SetRelativeScale3D(GetRootComponent()->GetRelativeScale3D());
 	}
 
 	Gizmo->MinRelativeZ = MinRelativeZ;
@@ -710,9 +710,11 @@ void ALandscapeGizmoActiveActor::FitToSelection()
 		TargetLandscapeInfo->GetSelectedExtent(MinX, MinY, MaxX, MaxY);
 		if (MinX != MAX_int32)
 		{
+			const FVector LocalScale3D = GetRootComponent()->GetRelativeScale3D();
+
 			float ScaleXY = TargetLandscapeInfo->DrawScale.X;
-			Width = ScaleXY * (MaxX - MinX + 1) / (GetRootComponent()->RelativeScale3D.X);
-			Height = ScaleXY * (MaxY - MinY + 1) / (GetRootComponent()->RelativeScale3D.Y);
+			Width = ScaleXY * (MaxX - MinX + 1) / (LocalScale3D.X);
+			Height = ScaleXY * (MaxY - MinY + 1) / (LocalScale3D.Y);
 			float NewLengthZ;
 			FVector NewLocation = TargetLandscapeInfo->GetLandscapeCenterPos(NewLengthZ, MinX, MinY, MaxX, MaxY);
 			SetLength(NewLengthZ);
diff --git a/Engine/Source/Runtime/Landscape/Private/LandscapeRender.cpp b/Engine/Source/Runtime/Landscape/Private/LandscapeRender.cpp
--- a/Engine/Source/Runtime/Landscape/Private/LandscapeRender.cpp
+++ b/Engine/Source/Runtime/Landscape/Private/LandscapeRender.cpp
@@ -4690,7 +4690,7 @@ void ULandscapeComponent::GetStreamingRenderAssetInfo(FStreamingTextureLevelCont
 	if (Proxy)
 	{
 		LocalStreamingDistanceMultiplier = FMath::Max(0.0f, Proxy->StreamingDistanceMultiplier);
-		TexelFactor = 0.75f * LocalStreamingDistanceMultiplier * ComponentSizeQuads * FMath::Abs(Proxy->GetRootComponent()->RelativeScale3D.X);
+		TexelFactor = 0.75f * LocalStreamingDistanceMultiplier * ComponentSizeQuads * FMath::Abs(Proxy->GetRootComponent()->GetRelativeScale3D().X);
 	}
 
 	ERHIFeatureLevel::Type FeatureLevel = LevelContext.GetFeatureLevel();
diff --git a/Engine/Source/Runtime/Landscape/Private/LandscapeSplines.cpp b/Engine/Source/Runtime/Landscape/Private/LandscapeSplines.cpp
--- a/Engine/Source/Runtime/Landscape/Private/LandscapeSplines.cpp
+++ b/Engine/Source/Runtime/Landscape/Private/LandscapeSplines.cpp
@@ -990,7 +990,7 @@ ULandscapeSplinesComponent* ULandscapeSplinesComponent::GetStreamingSplinesCompo
 			{
 				ComponentLandscapeProxy->Modify();
 				ComponentLandscapeProxy->SplineComponent = NewObject<ULandscapeSplinesComponent>(ComponentLandscapeProxy, NAME_None, RF_Transactional);
-				ComponentLandscapeProxy->SplineComponent->RelativeScale3D = RelativeScale3D;
+				ComponentLandscapeProxy->SplineComponent->SetRelativeScale3D_Direct(GetRelativeScale3D());
 				ComponentLandscapeProxy->SplineComponent->AttachToComponent(ComponentLandscapeProxy->GetRootComponent(), FAttachmentTransformRules::KeepRelativeTransform);
 			}
 			if (ComponentLandscapeProxy->SplineComponent)
@@ -1018,7 +1018,7 @@ ULandscapeSplinesComponent* ULandscapeSplinesComponent::GetStreamingSplinesCompo
 			{
 				Proxy->Modify();
 				Proxy->SplineComponent = NewObject<ULandscapeSplinesComponent>(Proxy, NAME_None, RF_Transactional);
-				Proxy->SplineComponent->RelativeScale3D = RelativeScale3D;
+				Proxy->SplineComponent->SetRelativeScale3D_Direct(GetRelativeScale3D());
 				Proxy->SplineComponent->AttachToComponent(Proxy->GetRootComponent(), FAttachmentTransformRules::KeepRelativeTransform);
 			}
 			return Proxy->SplineComponent;
@@ -1918,9 +1918,9 @@ void ULandscapeSplineControlPoint::UpdateSplinePoints(bool bUpdateCollision, boo
 			MeshLocation = RelativeTransform.TransformPosition(MeshLocation);
 		}
 
-		if (MeshComponent->RelativeLocation != MeshLocation ||
-			MeshComponent->RelativeRotation != MeshRotation ||
-			MeshComponent->RelativeScale3D != MeshScale)
+		if (MeshComponent->GetRelativeLocation() != MeshLocation ||
+			MeshComponent->GetRelativeRotation() != MeshRotation ||
+			MeshComponent->GetRelativeScale3D() != MeshScale)
 		{
 			MeshComponent->Modify();
 			MeshComponent->SetRelativeTransform(FTransform(MeshRotation, MeshLocation, MeshScale));
@@ -2992,9 +2992,9 @@ void ULandscapeSplineSegment::UpdateSplinePoints(bool bUpdateCollision, bool bUp
 			}
 
 			// Set Mesh component's location to half way between the start and end points. Improves the bounds and allows LDMaxDrawDistance to work
-			MeshComponent->RelativeLocation = (MeshComponent->SplineParams.StartPos + MeshComponent->SplineParams.EndPos) / 2;
-			MeshComponent->SplineParams.StartPos -= MeshComponent->RelativeLocation;
-			MeshComponent->SplineParams.EndPos -= MeshComponent->RelativeLocation;
+			MeshComponent->SetRelativeLocation_Direct((MeshComponent->SplineParams.StartPos + MeshComponent->SplineParams.EndPos) / 2);
+			MeshComponent->SplineParams.StartPos -= MeshComponent->GetRelativeLocation();
+			MeshComponent->SplineParams.EndPos -= MeshComponent->GetRelativeLocation();
 
 			if (MeshComponent->LDMaxDrawDistance != LDMaxDrawDistance)
 			{
diff --git a/Engine/Source/Runtime/LevelSequence/Private/LevelSequenceActor.cpp b/Engine/Source/Runtime/LevelSequence/Private/LevelSequenceActor.cpp
--- a/Engine/Source/Runtime/LevelSequence/Private/LevelSequenceActor.cpp
+++ b/Engine/Source/Runtime/LevelSequence/Private/LevelSequenceActor.cpp
@@ -44,7 +44,7 @@ ALevelSequenceActor::ALevelSequenceActor(const FObjectInitializer& Init)
 			SpriteComponent->Sprite = ConstructorStatics.DecalTexture.Get();
 			SpriteComponent->SetupAttachment(RootComponent);
 			SpriteComponent->bIsScreenSizeScaled = true;
-			SpriteComponent->bAbsoluteScale = true;
+			SpriteComponent->SetUsingAbsoluteScale(true);
 			SpriteComponent->bReceivesDecals = false;
 			SpriteComponent->bHiddenInGame = true;
 		}
diff --git a/Engine/Source/Runtime/LevelSequence/Private/LevelSequenceActorSpawner.cpp b/Engine/Source/Runtime/LevelSequence/Private/LevelSequenceActorSpawner.cpp
--- a/Engine/Source/Runtime/LevelSequence/Private/LevelSequenceActorSpawner.cpp
+++ b/Engine/Source/Runtime/LevelSequence/Private/LevelSequenceActorSpawner.cpp
@@ -110,8 +110,8 @@ UObject* FLevelSequenceActorSpawner::SpawnObject(FMovieSceneSpawnable& Spawnable
 
 	if (USceneComponent* RootComponent = ObjectTemplate->GetRootComponent())
 	{
-		SpawnTransform.SetTranslation(RootComponent->RelativeLocation);
-		SpawnTransform.SetRotation(RootComponent->RelativeRotation.Quaternion());
+		SpawnTransform.SetTranslation(RootComponent->GetRelativeLocation());
+		SpawnTransform.SetRotation(RootComponent->GetRelativeRotation().Quaternion());
 	}
 	else
 	{
diff --git a/Engine/Source/Runtime/MovieScene/Private/Evaluation/MovieSceneMotionVectorSimulation.cpp b/Engine/Source/Runtime/MovieScene/Private/Evaluation/MovieSceneMotionVectorSimulation.cpp
--- a/Engine/Source/Runtime/MovieScene/Private/Evaluation/MovieSceneMotionVectorSimulation.cpp
+++ b/Engine/Source/Runtime/MovieScene/Private/Evaluation/MovieSceneMotionVectorSimulation.cpp
@@ -95,15 +95,15 @@ void FMovieSceneMotionVectorSimulation::Apply(IMovieScenePlayer& Player)
 		if (ParentComp)
 		{
 			FTransform ParentTransform = ParentComp->GetSocketTransform(AttachSocket, RTS_World);
-			if (!Component->bAbsoluteLocation)
+			if (!Component->IsUsingAbsoluteLocation())
 			{
 				ParentToWorld.SetTranslation(ParentTransform.GetTranslation());
 			}
-			if (!Component->bAbsoluteRotation)
+			if (!Component->IsUsingAbsoluteRotation())
 			{
 				ParentToWorld.SetRotation(ParentTransform.GetRotation());
 			}
-			if (!Component->bAbsoluteScale)
+			if (!Component->IsUsingAbsoluteScale())
 			{
 				ParentToWorld.SetScale3D(ParentTransform.GetScale3D());
 			}
@@ -167,15 +167,15 @@ void FMovieSceneMotionVectorSimulation::ApplySimulatedTransforms(USceneComponent
 		FTransform SocketTransform = AttachSocketName == NAME_None ? FTransform::Identity : GetSocketTransform(InComponent, AttachSocketName);
 		FTransform ParentToWorld = SocketTransform * InPreviousTransform;
 
-		if (Child->bAbsoluteLocation)
+		if (Child->IsUsingAbsoluteLocation())
 		{
 			ParentToWorld.SetTranslation(FVector::ZeroVector);
 		}
-		if (Child->bAbsoluteRotation)
+		if (Child->IsUsingAbsoluteRotation())
 		{
 			ParentToWorld.SetRotation(FQuat::Identity);
 		}
-		if (Child->bAbsoluteScale)
+		if (Child->IsUsingAbsoluteScale())
 		{
 			ParentToWorld.SetScale3D(FVector(1.f, 1.f, 1.f));
 		}
diff --git a/Engine/Source/Runtime/MovieScene/Private/MovieSceneCommonHelpers.cpp b/Engine/Source/Runtime/MovieScene/Private/MovieSceneCommonHelpers.cpp
--- a/Engine/Source/Runtime/MovieScene/Private/MovieSceneCommonHelpers.cpp
+++ b/Engine/Source/Runtime/MovieScene/Private/MovieSceneCommonHelpers.cpp
@@ -177,7 +177,7 @@ UCameraComponent* MovieSceneHelpers::CameraComponentFromActor(const AActor* InAc
 	// If there's a camera component that's active, return that one
 	for (UCameraComponent* CameraComponent : CameraComponents)
 	{
-		if (CameraComponent->bIsActive)
+		if (CameraComponent->IsActive())
 		{
 			return CameraComponent;
 		}
diff --git a/Engine/Source/Runtime/MovieSceneTracks/Private/Evaluation/MovieScene3DTransformTemplate.cpp b/Engine/Source/Runtime/MovieSceneTracks/Private/Evaluation/MovieScene3DTransformTemplate.cpp
--- a/Engine/Source/Runtime/MovieSceneTracks/Private/Evaluation/MovieScene3DTransformTemplate.cpp
+++ b/Engine/Source/Runtime/MovieSceneTracks/Private/Evaluation/MovieScene3DTransformTemplate.cpp
@@ -74,7 +74,7 @@ struct FComponentTransformActuator : TMovieSceneBlendingActuator<F3DTransformTra
 	{
 		if (USceneComponent* SceneComponent = MovieSceneHelpers::SceneComponentFromRuntimeObject(InObject))
 		{
-			return F3DTransformTrackToken(SceneComponent->RelativeLocation, SceneComponent->RelativeRotation, SceneComponent->RelativeScale3D);
+			return F3DTransformTrackToken(SceneComponent->GetRelativeLocation(), SceneComponent->GetRelativeRotation(), SceneComponent->GetRelativeScale3D());
 		}
 
 		return F3DTransformTrackToken();
@@ -129,7 +129,7 @@ struct FSimulatedComponentTransformActuator : TMovieSceneBlendingActuator<F3DTra
 	{
 		if (USceneComponent* SceneComponent = MovieSceneHelpers::SceneComponentFromRuntimeObject(InObject))
 		{
-			return F3DTransformTrackToken(SceneComponent->RelativeLocation, SceneComponent->RelativeRotation, SceneComponent->RelativeScale3D);
+			return F3DTransformTrackToken(SceneComponent->GetRelativeLocation(), SceneComponent->GetRelativeRotation(), SceneComponent->GetRelativeScale3D());
 		}
 
 		return F3DTransformTrackToken();
diff --git a/Engine/Source/Runtime/MovieSceneTracks/Private/Evaluation/MovieSceneTemplateCommon.cpp b/Engine/Source/Runtime/MovieSceneTracks/Private/Evaluation/MovieSceneTemplateCommon.cpp
--- a/Engine/Source/Runtime/MovieSceneTracks/Private/Evaluation/MovieSceneTemplateCommon.cpp
+++ b/Engine/Source/Runtime/MovieSceneTracks/Private/Evaluation/MovieSceneTemplateCommon.cpp
@@ -50,8 +50,8 @@ void F3DTransformTrackToken::Apply(USceneComponent& SceneComponent, float DeltaT
 	SceneComponent.SetRelativeScale3D(Scale);
 
 	// Force the location and rotation values to avoid Rot->Quat->Rot conversions
-	SceneComponent.RelativeLocation = Translation;
-	SceneComponent.RelativeRotation = Rotation;
+	SceneComponent.SetRelativeLocation_Direct(Translation);
+	SceneComponent.SetRelativeRotation_Direct(Rotation);
 
 	if (DeltaTime > 0.f)
 	{
diff --git a/Engine/Source/Runtime/MovieSceneTracks/Private/Evaluation/MovieSceneVisibilityTemplate.cpp b/Engine/Source/Runtime/MovieSceneTracks/Private/Evaluation/MovieSceneVisibilityTemplate.cpp
--- a/Engine/Source/Runtime/MovieSceneTracks/Private/Evaluation/MovieSceneVisibilityTemplate.cpp
+++ b/Engine/Source/Runtime/MovieSceneTracks/Private/Evaluation/MovieSceneVisibilityTemplate.cpp
@@ -60,7 +60,7 @@ struct FTemporarilyHiddenInGameTokenProducer : IMovieScenePreAnimatedTokenProduc
 #else
 				false;
 #endif
-			return FTemporarilyHiddenInGamePreAnimatedToken(Actor->bHidden, bInTemporarilyHiddenInGame);
+			return FTemporarilyHiddenInGamePreAnimatedToken(Actor->IsHidden(), bInTemporarilyHiddenInGame);
 		}
 		else if (InObject.IsA(USceneComponent::StaticClass()))
 		{
diff --git a/Engine/Source/Runtime/MovieSceneTracks/Private/Sections/MovieScene3DTransformSection.cpp b/Engine/Source/Runtime/MovieSceneTracks/Private/Sections/MovieScene3DTransformSection.cpp
--- a/Engine/Source/Runtime/MovieSceneTracks/Private/Sections/MovieScene3DTransformSection.cpp
+++ b/Engine/Source/Runtime/MovieSceneTracks/Private/Sections/MovieScene3DTransformSection.cpp
@@ -163,7 +163,7 @@ struct F3DTransformChannelEditorData
 		{
 			if (USceneComponent* RootComponent = Actor->GetRootComponent())
 			{
-				return RootComponent->RelativeRotation;
+				return RootComponent->GetRelativeRotation();
 			}
 		}
 
diff --git a/Engine/Source/Runtime/MovieSceneTracks/Public/Sections/MovieScene3DTransformSection.h b/Engine/Source/Runtime/MovieSceneTracks/Public/Sections/MovieScene3DTransformSection.h
--- a/Engine/Source/Runtime/MovieSceneTracks/Public/Sections/MovieScene3DTransformSection.h
+++ b/Engine/Source/Runtime/MovieSceneTracks/Public/Sections/MovieScene3DTransformSection.h
@@ -37,9 +37,9 @@ struct FTransformData
 	FVector Scale;
 
 	FTransformData()
-		: Translation( ForceInitToZero )
-		, Rotation( ForceInitToZero )
-		, Scale( ForceInitToZero )
+		: Translation(ForceInitToZero)
+		, Rotation(ForceInitToZero)
+		, Scale(ForceInitToZero)
 	{}
 
 	/**
@@ -48,10 +48,10 @@ struct FTransformData
 	*
 	* @param InComponent	The component to build from
 	*/
-	FTransformData( const USceneComponent* InComponent )
-		: Translation( InComponent->RelativeLocation )
-		, Rotation( InComponent->RelativeRotation )
-		, Scale( InComponent->RelativeScale3D )
+	FTransformData(const USceneComponent* InComponent)
+		: Translation(InComponent->GetRelativeLocation())
+		, Rotation(InComponent->GetRelativeRotation())
+		, Scale(InComponent->GetRelativeScale3D())
 	{}
 };
 
diff --git a/Engine/Source/Runtime/NavigationSystem/Private/NavMesh/NavMeshBoundsVolume.cpp b/Engine/Source/Runtime/NavigationSystem/Private/NavMesh/NavMeshBoundsVolume.cpp
--- a/Engine/Source/Runtime/NavigationSystem/Private/NavMesh/NavMeshBoundsVolume.cpp
+++ b/Engine/Source/Runtime/NavigationSystem/Private/NavMesh/NavMeshBoundsVolume.cpp
@@ -31,9 +31,9 @@ void ANavMeshBoundsVolume::PostEditChangeProperty( struct FPropertyChangedEvent&
 
 		if (PropName == GET_MEMBER_NAME_CHECKED(ABrush, BrushBuilder)
 			|| MemberName == GET_MEMBER_NAME_CHECKED(ANavMeshBoundsVolume, SupportedAgents)
-			|| MemberName == GET_MEMBER_NAME_CHECKED(USceneComponent, RelativeLocation)
-			|| MemberName == GET_MEMBER_NAME_CHECKED(USceneComponent, RelativeRotation)
-			|| MemberName == GET_MEMBER_NAME_CHECKED(USceneComponent, RelativeScale3D))
+			|| MemberName == USceneComponent::GetRelativeLocationPropertyName()
+			|| MemberName == USceneComponent::GetRelativeRotationPropertyName()
+			|| MemberName == USceneComponent::GetRelativeScale3DPropertyName())
 		{
 			NavSys->OnNavigationBoundsUpdated(this);
 		}
@@ -57,7 +57,7 @@ void ANavMeshBoundsVolume::PostRegisterAllComponents()
 	Super::PostRegisterAllComponents();
 	
 	UNavigationSystemV1* NavSys = FNavigationSystem::GetCurrent<UNavigationSystemV1>(GetWorld());
-	if (NavSys && Role == ROLE_Authority)
+	if (NavSys && GetLocalRole() == ROLE_Authority)
 	{
 		NavSys->OnNavigationBoundsAdded(this);
 	}
@@ -68,7 +68,7 @@ void ANavMeshBoundsVolume::PostUnregisterAllComponents()
 	Super::PostUnregisterAllComponents();
 	
 	UNavigationSystemV1* NavSys = FNavigationSystem::GetCurrent<UNavigationSystemV1>(GetWorld());
-	if (NavSys && Role == ROLE_Authority)
+	if (NavSys && GetLocalRole() == ROLE_Authority)
 	{
 		NavSys->OnNavigationBoundsRemoved(this);
 	}
diff --git a/Engine/Source/Runtime/NavigationSystem/Private/NavMesh/RecastNavMesh.cpp b/Engine/Source/Runtime/NavigationSystem/Private/NavMesh/RecastNavMesh.cpp
--- a/Engine/Source/Runtime/NavigationSystem/Private/NavMesh/RecastNavMesh.cpp
+++ b/Engine/Source/Runtime/NavigationSystem/Private/NavMesh/RecastNavMesh.cpp
@@ -354,7 +354,7 @@ void ARecastNavMesh::UpdateNavMeshDrawing()
 {
 #if !UE_BUILD_SHIPPING
 	UNavMeshRenderingComponent* NavMeshRenderComp = Cast<UNavMeshRenderingComponent>(RenderingComp);
-	if (NavMeshRenderComp != nullptr && NavMeshRenderComp->bVisible && (NavMeshRenderComp->IsForcingUpdate() || UNavMeshRenderingComponent::IsNavigationShowFlagSet(GetWorld())))
+	if (NavMeshRenderComp != nullptr && NavMeshRenderComp->GetVisibleFlag() && (NavMeshRenderComp->IsForcingUpdate() || UNavMeshRenderingComponent::IsNavigationShowFlagSet(GetWorld())))
 	{
 		RenderingComp->MarkRenderStateDirty();
 	}
diff --git a/Engine/Source/Runtime/NavigationSystem/Private/NavSystemConfigOverride.cpp b/Engine/Source/Runtime/NavigationSystem/Private/NavSystemConfigOverride.cpp
--- a/Engine/Source/Runtime/NavigationSystem/Private/NavSystemConfigOverride.cpp
+++ b/Engine/Source/Runtime/NavigationSystem/Private/NavSystemConfigOverride.cpp
@@ -45,7 +45,7 @@ ANavSystemConfigOverride::ANavSystemConfigOverride(const FObjectInitializer& Obj
 		if (SpriteComponent)
 		{
 			SpriteComponent->Sprite = ConstructorStatics.NoteTextureObject.Get();
-			SpriteComponent->RelativeScale3D = FVector(0.5f, 0.5f, 0.5f);
+			SpriteComponent->SetRelativeScale3D(FVector(0.5f, 0.5f, 0.5f));
 			SpriteComponent->SpriteInfo.Category = ConstructorStatics.ID_Notes;
 			SpriteComponent->SpriteInfo.DisplayName = ConstructorStatics.NAME_Notes;
 			SpriteComponent->SetupAttachment(RootComponent);
@@ -54,8 +54,8 @@ ANavSystemConfigOverride::ANavSystemConfigOverride(const FObjectInitializer& Obj
 	}
 #endif // WITH_EDITORONLY_DATA
 
-	bHidden = true;
-	bCanBeDamaged = false;
+	SetHidden(true);
+	SetCanBeDamaged(false);
 
 	bNetLoadOnClient = false;
 }
diff --git a/Engine/Source/Runtime/NavigationSystem/Private/NavigationTestingActor.cpp b/Engine/Source/Runtime/NavigationSystem/Private/NavigationTestingActor.cpp
--- a/Engine/Source/Runtime/NavigationSystem/Private/NavigationTestingActor.cpp
+++ b/Engine/Source/Runtime/NavigationSystem/Private/NavigationTestingActor.cpp
@@ -142,7 +142,7 @@ void ANavigationTestingActor::PostEditChangeProperty(FPropertyChangedEvent& Prop
 		}
 		else if (ChangedPropName == NAME_ShouldBeVisibleInGame)
 		{
-			bHidden = !bShouldBeVisibleInGame;
+			SetHidden(!bShouldBeVisibleInGame);
 		}
 		else if (ChangedCategory == TEXT("Debug"))
 		{
@@ -250,7 +250,7 @@ void ANavigationTestingActor::PostLoad()
 	}
 #endif
 
-	bHidden = !bShouldBeVisibleInGame;
+	SetHidden(!bShouldBeVisibleInGame);
 }
 
 void ANavigationTestingActor::TickMe()
diff --git a/Engine/Source/Runtime/UMG/Private/Components/WidgetComponent.cpp b/Engine/Source/Runtime/UMG/Private/Components/WidgetComponent.cpp
--- a/Engine/Source/Runtime/UMG/Private/Components/WidgetComponent.cpp
+++ b/Engine/Source/Runtime/UMG/Private/Components/WidgetComponent.cpp
@@ -604,7 +604,7 @@ UWidgetComponent::UWidgetComponent( const FObjectInitializer& PCIP )
 	PrimaryComponentTick.bCanEverTick = true;
 	bTickInEditor = true;
 
-	RelativeRotation = FRotator::ZeroRotator;
+	SetRelativeRotation(FRotator::ZeroRotator);
 
 	BodyInstance.SetCollisionProfileName(FName(TEXT("UI")));
 
@@ -1086,7 +1086,7 @@ void UWidgetComponent::TickComponent(float DeltaTime, enum ELevelTick TickType,
 				ULocalPlayer* TargetPlayer = GetOwnerPlayer();
 				APlayerController* PlayerController = TargetPlayer ? TargetPlayer->PlayerController : nullptr;
 
-				if ( TargetPlayer && PlayerController && IsVisible() && !(GetOwner()->bHidden))
+				if ( TargetPlayer && PlayerController && IsVisible() && !(GetOwner()->IsHidden()))
 				{
 					if ( !bAddedToScreen )
 					{
