Add a scalability option for TileCount of runtime virtual texture.
Also modified logic for runtime virtual texture size with non unity aspect ratio so that we snap to the nearest power of 2 correctly.

diff --git a/Engine/Source/Runtime/Engine/Classes/VT/RuntimeVirtualTexture.h b/Engine/Source/Runtime/Engine/Classes/VT/RuntimeVirtualTexture.h
--- a/Engine/Source/Runtime/Engine/Classes/VT/RuntimeVirtualTexture.h
+++ b/Engine/Source/Runtime/Engine/Classes/VT/RuntimeVirtualTexture.h
@@ -40,6 +40,10 @@ protected:
 	UPROPERTY(EditAnywhere, BlueprintReadOnly, AdvancedDisplay, Category = Layout, meta = (DisplayName = "Enable private page table"))
 	bool bPrivateSpace = true;
 
+	/** Enable device scalability settings to modify the TileCount of the virtual texture. */
+	UPROPERTY(EditAnywhere, BlueprintReadOnly, AdvancedDisplay, Category = Layout)
+	bool bEnableScalability = false;
+
 	/** Size of virtual texture along the largest axis. (Actual values increase in powers of 2) */
 	UPROPERTY()
 	int32 Size_DEPRECATED = -1;
@@ -82,7 +86,8 @@ public:
 	/** Public getter for virtual texture size */
 	int32 GetSize() const { return GetTileCount() * GetTileSize(); }
 	/** Public getter for virtual texture tile count */
-	int32 GetTileCount() const { return 1 << FMath::Clamp(TileCount, 0, 12); }
+	int32 GetTileCount() const { return GetTileCount(TileCount); }
+	static int32 GetTileCount(int32 InTileCount) { return 1 << FMath::Clamp(InTileCount, 0, 12); }
 	/** Public getter for virtual texture tile size */
 	int32 GetTileSize() const { return 1 << FMath::Clamp(TileSize + 6, 6, 10); }
 	/** Public getter for virtual texture tile border size */
diff --git a/Engine/Source/Runtime/Engine/Private/VT/RuntimeVirtualTexture.cpp b/Engine/Source/Runtime/Engine/Private/VT/RuntimeVirtualTexture.cpp
--- a/Engine/Source/Runtime/Engine/Private/VT/RuntimeVirtualTexture.cpp
+++ b/Engine/Source/Runtime/Engine/Private/VT/RuntimeVirtualTexture.cpp
@@ -10,6 +10,16 @@
 #include "VT/UploadingVirtualTexture.h"
 #include "VT/VirtualTextureLevelRedirector.h"
 
+
+/** Device scalability option for runtime virtual texture size. */
+static TAutoConsoleVariable<int32> CVarVTTileCountBias(
+	TEXT("r.VT.RVT.TileCountBias"),
+	0,
+	TEXT("Bias to apply to Runtime Virtual Texture size."),
+	ECVF_RenderThreadSafe
+);
+
+
 namespace
 {
 	/** Null producer to use as placeholder when no producer has been set on a URuntimeVirtualTexture */
@@ -218,23 +228,30 @@ void URuntimeVirtualTexture::GetProducerDescription(FVTProducerDescription& OutD
 	OutDesc.WidthInBlocks = 1;
 	OutDesc.HeightInBlocks = 1;
 
-	// Set width and height to best match the aspect ratio
+	// Apply TileCount modifier here to allow size scalability option
+	const int32 TileCountBias = CVarVTTileCountBias.GetValueOnAnyThread();
+	const int32 MaxSizeInTiles = GetTileCount(TileCount + TileCountBias);
+
+	// Set width and height to best match the runtime virtual texture volume's aspect ratio
 	const FVector VolumeSize = VolumeToWorld.GetScale3D();
-	const float AspectRatio = VolumeSize.X / VolumeSize.Y;
-	int32 Width, Height;
-	if (AspectRatio >= 1.f)
+	const float VolumeSizeX = FMath::Max(FMath::Abs(VolumeSize.X), 0.0001f);
+	const float VolumeSizeY = FMath::Max(FMath::Abs(VolumeSize.Y), 0.0001f);
+	const float AspectRatioLog2 = FMath::Log2(VolumeSizeX / VolumeSizeY);
+
+	uint32 WidthInTiles, HeightInTiles;
+	if (AspectRatioLog2 >= 0.f)
 	{
-		Width = GetSize();
-		Height = FMath::Max((int32)FMath::RoundUpToPowerOfTwo(Width / AspectRatio), GetTileSize());
+		WidthInTiles = MaxSizeInTiles;
+		HeightInTiles = FMath::Max(WidthInTiles >> FMath::RoundToInt(AspectRatioLog2), 1u);
 	}
 	else
 	{
-		Height = GetSize();
-		Width = FMath::Max((int32)FMath::RoundUpToPowerOfTwo(Height * AspectRatio), GetTileSize());
+		HeightInTiles = MaxSizeInTiles;
+		WidthInTiles = FMath::Max(HeightInTiles >> FMath::RoundToInt(-AspectRatioLog2), 1u);
 	}
 
-	OutDesc.BlockWidthInTiles = Width / GetTileSize();
-	OutDesc.BlockHeightInTiles = Height / GetTileSize();
+	OutDesc.BlockWidthInTiles = WidthInTiles;
+	OutDesc.BlockHeightInTiles = HeightInTiles;
 	OutDesc.MaxLevel = FMath::Max((int32)FMath::CeilLogTwo(FMath::Max(OutDesc.BlockWidthInTiles, OutDesc.BlockHeightInTiles)) - GetRemoveLowMips(), 0);
 
 	OutDesc.NumTextureLayers = GetLayerCount();
