Fix for low lod being used by road mesh

diff --git a/Engine/Source/Runtime/Engine/Private/StaticMeshRender.cpp b/Engine/Source/Runtime/Engine/Private/StaticMeshRender.cpp
--- a/Engine/Source/Runtime/Engine/Private/StaticMeshRender.cpp
+++ b/Engine/Source/Runtime/Engine/Private/StaticMeshRender.cpp
@@ -953,15 +953,14 @@ void FStaticMeshSceneProxy::DrawStaticElements(FStaticPrimitiveDrawInterface* PD
 
 				for (int32 BatchIndex = 0; BatchIndex < NumBatches; BatchIndex++)
 				{
-					FMeshBatch MeshBatch;
+					FMeshBatch BaseMeshBatch;
 
-					if (GetMeshElement(LODIndex, BatchIndex, SectionIndex, PrimitiveDPG, bIsMeshElementSelected, true, MeshBatch))
+					if (GetMeshElement(LODIndex, BatchIndex, SectionIndex, PrimitiveDPG, bIsMeshElementSelected, true, BaseMeshBatch))
 					{
-						PDI->DrawMesh(MeshBatch, FLT_MAX);
-
 						if (NumRuntimeVirtualTextureTypes > 0)
 						{
 							// Runtime virtual texture mesh elements.
+							FMeshBatch MeshBatch(BaseMeshBatch);
 							SetupMeshBatchForRuntimeVirtualTexture(MeshBatch);
 							for (ERuntimeVirtualTextureMaterialType MaterialType : RuntimeVirtualTextureMaterialTypes)
 							{
@@ -969,6 +968,9 @@ void FStaticMeshSceneProxy::DrawStaticElements(FStaticPrimitiveDrawInterface* PD
 								PDI->DrawMesh(MeshBatch, FLT_MAX);
 							}
 						}
+						{
+							PDI->DrawMesh(BaseMeshBatch, FLT_MAX);
+						}
 					}
 				}
 			}
@@ -1070,15 +1072,6 @@ void FStaticMeshSceneProxy::DrawStaticElements(FStaticPrimitiveDrawInterface* PD
 						FMeshBatch BaseMeshBatch;
 						if (GetMeshElement(LODIndex, BatchIndex, SectionIndex, PrimitiveDPG, bIsMeshElementSelected, true, BaseMeshBatch))
 						{
-							{
-								// Standard mesh elements.
-								// If we have submitted an optimized shadow-only mesh, remaining mesh elements must not cast shadows.
-								FMeshBatch MeshBatch(BaseMeshBatch);
-								MeshBatch.CastShadow &= !bUseUnifiedMeshForShadow;
-								MeshBatch.bUseAsOccluder &= !bUseUnifiedMeshForDepth;
-								MeshBatch.bUseForDepthPass &= !bUseUnifiedMeshForDepth;
-								PDI->DrawMesh(MeshBatch, ScreenSize);
-							}
 							if (NumRuntimeVirtualTextureTypes > 0)
 							{
 								// Runtime virtual texture mesh elements.
@@ -1088,9 +1081,19 @@ void FStaticMeshSceneProxy::DrawStaticElements(FStaticPrimitiveDrawInterface* PD
 								for (ERuntimeVirtualTextureMaterialType MaterialType : RuntimeVirtualTextureMaterialTypes)
 								{
 									MeshBatch.RuntimeVirtualTextureMaterialType = (uint32)MaterialType;
-									PDI->DrawMesh(MeshBatch, FLT_MAX);
+									PDI->DrawMesh(MeshBatch, ScreenSize);
 								}
 							}
+
+							{
+								// Standard mesh elements.
+								// If we have submitted an optimized shadow-only mesh, remaining mesh elements must not cast shadows.
+								FMeshBatch MeshBatch(BaseMeshBatch);
+								MeshBatch.CastShadow &= !bUseUnifiedMeshForShadow;
+								MeshBatch.bUseAsOccluder &= !bUseUnifiedMeshForDepth;
+								MeshBatch.bUseForDepthPass &= !bUseUnifiedMeshForDepth;
+								PDI->DrawMesh(MeshBatch, ScreenSize);
+							}
 						}
 					}
 				}
