Support loading cooked static mesh LOD data in editor

diff --git a/Engine/Source/Runtime/Engine/Private/StaticMesh.cpp b/Engine/Source/Runtime/Engine/Private/StaticMesh.cpp
--- a/Engine/Source/Runtime/Engine/Private/StaticMesh.cpp
+++ b/Engine/Source/Runtime/Engine/Private/StaticMesh.cpp
@@ -625,6 +625,8 @@ void FStaticMeshLODResources::Serialize(FArchive& Ar, UObject* Owner, int32 Inde
 	if (!StripFlags.IsDataStrippedForServer() && !bIsLODCookedOut)
 	{
 		FStaticMeshBuffersSize TmpBuffersSize;
+		TArray<uint8> TmpBuff;
+
 		if (bInlined)
 		{
 			SerializeBuffers(Ar, OwnerStaticMesh, ClassDataStripFlags, TmpBuffersSize);
@@ -640,7 +642,6 @@ void FStaticMeshLODResources::Serialize(FArchive& Ar, UObject* Owner, int32 Inde
 				const int32 OptionalLODIdx = GetPlatformMinLODIdx(Ar.CookingTarget(), OwnerStaticMesh) - Index;
 				const bool bDiscardBulkData = OptionalLODIdx > MaxNumOptionalLODs;
 
-				TArray<uint8> TmpBuff;
 				if (!bDiscardBulkData)
 				{
 					FMemoryWriter MemWriter(TmpBuff, true);
@@ -677,6 +678,15 @@ void FStaticMeshLODResources::Serialize(FArchive& Ar, UObject* Owner, int32 Inde
 				OffsetInFile = static_cast<uint32>(Tmp);
 				check(TmpBulkData.GetBulkDataSize() >= 0);
 				BulkDataSize = static_cast<uint32>(TmpBulkData.GetBulkDataSize());
+
+				// Streaming CPU data in editor build isn't supported yet because tools and utils need access
+				if (bUsingCookedData && BulkDataSize > 0)
+				{
+					TmpBuff.Empty(BulkDataSize);
+					TmpBuff.AddUninitialized(BulkDataSize);
+					void* Dest = TmpBuff.GetData();
+					TmpBulkData.GetCopy(&Dest);
+				}
 			}
 
 			SerializeAvailabilityInfo(Ar);
@@ -688,6 +698,14 @@ void FStaticMeshLODResources::Serialize(FArchive& Ar, UObject* Owner, int32 Inde
 			{
 				ClearAvailabilityInfo();
 			}
+
+			if (Ar.IsLoading() && bUsingCookedData && BulkDataSize > 0)
+			{
+				ClearAvailabilityInfo();
+				FMemoryReader MemReader(TmpBuff, true);
+				MemReader.SetByteSwapping(Ar.IsByteSwapping());
+				SerializeBuffers(MemReader, OwnerStaticMesh, ClassDataStripFlags, TmpBuffersSize);
+			}
 		}
 	}
 }
@@ -2553,7 +2571,7 @@ void UStaticMesh::InitResources()
 	{
 		LinkStreaming();
 	}
-	
+
 #if	STATS
 	UStaticMesh* This = this;
 	ENQUEUE_RENDER_COMMAND(UpdateMemoryStats)(
@@ -4881,7 +4899,24 @@ int32 UStaticMesh::CalcNumOptionalMips() const
 #if !WITH_EDITOR
 	return MinLOD.Default;
 #else
-	return 0;
+	int32 NumOptionalLODs = 0;
+	if (RenderData)
+	{
+		const TIndirectArray<FStaticMeshLODResources>& LODResources = RenderData->LODResources;
+		for (int32 Idx = 0; Idx < LODResources.Num(); ++Idx)
+		{
+			const FStaticMeshLODResources& Resource = LODResources[Idx];
+			if (Resource.bIsOptionalLOD)
+			{
+				++NumOptionalLODs;
+			}
+			else
+			{
+				break;
+			}
+		}
+	}
+	return NumOptionalLODs;
 #endif
 }
 
@@ -4899,7 +4934,6 @@ int32 UStaticMesh::CalcCumulativeLODSize(int32 NumLODs) const
 
 bool UStaticMesh::GetMipDataFilename(const int32 MipIndex, FString& OutBulkDataFilename) const
 {
-#if !WITH_EDITOR
 	// TODO: this is slow. Should cache the name once per mesh
 	FString PackageName = GetOutermost()->FileName.ToString();
 	// Handle name redirection and localization
@@ -4915,9 +4949,6 @@ bool UStaticMesh::GetMipDataFilename(const int32 MipIndex, FString& OutBulkDataF
 	OutBulkDataFilename = FPaths::ChangeExtension(OutBulkDataFilename, MipIndex < MinLOD.Default ? TEXT(".uptnl") : TEXT(".ubulk"));
 	check(MipIndex < MinLOD.Default || IFileManager::Get().FileExists(*OutBulkDataFilename));
 	return true;
-#else
-	return false;
-#endif
 }
 
 bool UStaticMesh::IsReadyForStreaming() const
