Fix for excessive anim initialization cost (50ms+ hitch)

A previous crash fix introduced pathological repeated calls to Initialize() and CacheBones(), where each time a linked layer was default-initialized, it would repeat calls to initialize/cachebones in all downstream nodes (until it met a cached pose guard).
This fix addresses this by 'cutting off' the tree propgation of Initialize() and CacheBones() at linked input poses. It moves the calls down the blend tree to the 'outer' linked anim graph node, while still ensuring that all branches of the linked anim graph node are taken at Initialize() and CacheBones() time.
This fix also introduces a defered path for linked graph fragment initiailzation, which was previously always called on the game thread.

Also:

- Removes duplicate work performed when initializing Linked Anim Graph nodes.
- Removes needless array allocations when initializing pose links.
- Minor optimizations to FAnimNode_LinkedAnimGraph::DynamicLink, FAnimNode_LinkedAnimGraph::DynamicUnlink and FAnimNode_CustomProperty::InitializeProperties (pulled out invariants).

diff --git a/Engine/Source/Runtime/Engine/Classes/Animation/AnimInstance.h b/Engine/Source/Runtime/Engine/Classes/Animation/AnimInstance.h
--- a/Engine/Source/Runtime/Engine/Classes/Animation/AnimInstance.h
+++ b/Engine/Source/Runtime/Engine/Classes/Animation/AnimInstance.h
@@ -768,11 +768,11 @@ public:
 	UAnimInstance* GetLinkedAnimLayerInstanceByClass(TSubclassOf<UAnimInstance> InClass) const;
 
 	/** Sets up initial layer groupings */
-	void InitializeGroupedLayers();
+	void InitializeGroupedLayers(bool bInDeferSubGraphInitialization);
 
 private:
 	/** Helper function to perform layer overlay actions (set, clear) */
-	void PerformLinkedLayerOverlayOperation(TSubclassOf<UAnimInstance> InClass, TFunctionRef<UClass*(UClass*, FAnimNode_LinkedAnimLayer*)> InClassSelectorFunction);
+	void PerformLinkedLayerOverlayOperation(TSubclassOf<UAnimInstance> InClass, TFunctionRef<UClass*(UClass*, FAnimNode_LinkedAnimLayer*)> InClassSelectorFunction, bool bInDeferSubGraphInitialization = false);
 
 protected:
 	/** Map between Active Montages and their FAnimMontageInstance */
diff --git a/Engine/Source/Runtime/Engine/Classes/Animation/AnimNode_LinkedAnimGraph.h b/Engine/Source/Runtime/Engine/Classes/Animation/AnimNode_LinkedAnimGraph.h
--- a/Engine/Source/Runtime/Engine/Classes/Animation/AnimNode_LinkedAnimGraph.h
+++ b/Engine/Source/Runtime/Engine/Classes/Animation/AnimNode_LinkedAnimGraph.h
@@ -64,6 +64,12 @@ public:
 	virtual void Evaluate_AnyThread(FPoseContext& Output) override;
 	virtual void GatherDebugData(FNodeDebugData& DebugData) override;
 
+	// Initializes only the sub-graph that this node is linked to
+	void InitializeSubGraph_AnyThread(const FAnimationInitializeContext& Context);
+
+	// Caches bones only for the sub graph that this node is linked to
+	void CacheBonesSubGraph_AnyThread(const FAnimationCacheBonesContext& Context);
+
 protected:
 	virtual void OnInitializeAnimInstance(const FAnimInstanceProxy* InProxy, const UAnimInstance* InAnimInstance) override;
 	virtual bool NeedsOnInitializeAnimInstance() const override { return true; }
diff --git a/Engine/Source/Runtime/Engine/Classes/Animation/AnimNode_LinkedInputPose.h b/Engine/Source/Runtime/Engine/Classes/Animation/AnimNode_LinkedInputPose.h
--- a/Engine/Source/Runtime/Engine/Classes/Animation/AnimNode_LinkedInputPose.h
+++ b/Engine/Source/Runtime/Engine/Classes/Animation/AnimNode_LinkedInputPose.h
@@ -44,8 +44,10 @@ struct ENGINE_API FAnimNode_LinkedInputPose : public FAnimNode_Base
 	FBlendedHeapCurve CachedInputCurve;
 
 	// FAnimNode_Base interface
+#if ENABLE_ANIMGRAPH_TRAVERSAL_DEBUG
 	virtual void Initialize_AnyThread(const FAnimationInitializeContext& Context) override;
 	virtual void CacheBones_AnyThread(const FAnimationCacheBonesContext& Context) override;
+#endif
 	virtual void Update_AnyThread(const FAnimationUpdateContext& Context) override;
 	virtual void Evaluate_AnyThread(FPoseContext& Output) override;
 	virtual void GatherDebugData(FNodeDebugData& DebugData) override;
diff --git a/Engine/Source/Runtime/Engine/Private/Animation/AnimInstance.cpp b/Engine/Source/Runtime/Engine/Private/Animation/AnimInstance.cpp
--- a/Engine/Source/Runtime/Engine/Private/Animation/AnimInstance.cpp
+++ b/Engine/Source/Runtime/Engine/Private/Animation/AnimInstance.cpp
@@ -218,7 +218,7 @@ void UAnimInstance::InitializeAnimation(bool bInDeferRootNodeInitialization)
 	// we can bind rules & events now the graph has been initialized
 	GetProxyOnGameThread<FAnimInstanceProxy>().BindNativeDelegates();
 
-	InitializeGroupedLayers();
+	InitializeGroupedLayers(bInDeferRootNodeInitialization);
 }
 
 void UAnimInstance::UninitializeAnimation()
@@ -2508,7 +2508,7 @@ void UAnimInstance::LinkAnimGraphByTag(FName InTag, TSubclassOf<UAnimInstance> I
 	}
 }
 
-void UAnimInstance::PerformLinkedLayerOverlayOperation(TSubclassOf<UAnimInstance> InClass, TFunctionRef<UClass*(UClass* InClass, FAnimNode_LinkedAnimLayer*)> InClassSelectorFunction)
+void UAnimInstance::PerformLinkedLayerOverlayOperation(TSubclassOf<UAnimInstance> InClass, TFunctionRef<UClass*(UClass* InClass, FAnimNode_LinkedAnimLayer*)> InClassSelectorFunction, bool bInDeferSubGraphInitialization)
 {
 	if (IAnimClassInterface* AnimBlueprintClass = IAnimClassInterface::GetFromClass(GetClass()))
 	{
@@ -2576,9 +2576,9 @@ void UAnimInstance::PerformLinkedLayerOverlayOperation(TSubclassOf<UAnimInstance
 		{
 			InLinkedProxy.InitializeObjects(InLinkedInstance);
 			FAnimationInitializeContext InitContext(&InThisProxy);
-			InLayerNode->Initialize_AnyThread(InitContext);
+			InLayerNode->InitializeSubGraph_AnyThread(InitContext);
 			FAnimationCacheBonesContext CacheBonesContext(&InThisProxy);
-			InLayerNode->CacheBones_AnyThread(CacheBonesContext);
+			InLayerNode->CacheBonesSubGraph_AnyThread(CacheBonesContext);
 			InLinkedProxy.ClearObjects();
 		};
 
@@ -2607,12 +2607,15 @@ void UAnimInstance::PerformLinkedLayerOverlayOperation(TSubclassOf<UAnimInstance
 
 							LayerNode->SetLinkedLayerInstance(this, NewLinkedInstance);
 
-							// Initialize the correct parts of the linked instance
-							if(LayerNode->LinkedRoot)
+							if(!bInDeferSubGraphInitialization)
 							{
-								FAnimInstanceProxy& ThisProxy = GetProxyOnAnyThread<FAnimInstanceProxy>();
-								FAnimInstanceProxy& LinkedProxy = NewLinkedInstance->GetProxyOnAnyThread<FAnimInstanceProxy>();
-								InitializeAndCacheBonesForLinkedRoot(LayerNode, ThisProxy, NewLinkedInstance, LinkedProxy);
+								// Initialize the correct parts of the linked instance
+								if(LayerNode->LinkedRoot)
+								{
+									FAnimInstanceProxy& ThisProxy = GetProxyOnAnyThread<FAnimInstanceProxy>();
+									FAnimInstanceProxy& LinkedProxy = NewLinkedInstance->GetProxyOnAnyThread<FAnimInstanceProxy>();
+									InitializeAndCacheBonesForLinkedRoot(LayerNode, ThisProxy, NewLinkedInstance, LinkedProxy);
+								}
 							}
 
 							MeshComp->GetLinkedAnimInstances().Add(NewLinkedInstance);
@@ -2622,12 +2625,15 @@ void UAnimInstance::PerformLinkedLayerOverlayOperation(TSubclassOf<UAnimInstance
 					{
 						LayerNode->SetLinkedLayerInstance(this, nullptr);
 
-						UAnimInstance* LinkedInstance = LayerNode->GetTargetInstance<UAnimInstance>();
-						if(LayerNode->LinkedRoot && LinkedInstance)
+						if(!bInDeferSubGraphInitialization)
 						{
-							FAnimInstanceProxy& ThisProxy = GetProxyOnAnyThread<FAnimInstanceProxy>();
-							FAnimInstanceProxy& LinkedProxy = LinkedInstance->GetProxyOnAnyThread<FAnimInstanceProxy>();
-							InitializeAndCacheBonesForLinkedRoot(LayerNode, ThisProxy, LinkedInstance, LinkedProxy);
+							UAnimInstance* LinkedInstance = LayerNode->GetTargetInstance<UAnimInstance>();
+							if(LayerNode->LinkedRoot && LinkedInstance)
+							{
+								FAnimInstanceProxy& ThisProxy = GetProxyOnAnyThread<FAnimInstanceProxy>();
+								FAnimInstanceProxy& LinkedProxy = LinkedInstance->GetProxyOnAnyThread<FAnimInstanceProxy>();
+								InitializeAndCacheBonesForLinkedRoot(LayerNode, ThisProxy, LinkedInstance, LinkedProxy);
+							}
 						}
 					}
 				}
@@ -2659,15 +2665,18 @@ void UAnimInstance::PerformLinkedLayerOverlayOperation(TSubclassOf<UAnimInstance
 							LayerNode->SetLinkedLayerInstance(this, NewLinkedInstance);
 						}
 
-						FAnimInstanceProxy& ThisProxy = GetProxyOnAnyThread<FAnimInstanceProxy>();
-						FAnimInstanceProxy& LinkedProxy = NewLinkedInstance->GetProxyOnAnyThread<FAnimInstanceProxy>();
-
-						// Initialize the correct parts of the linked instance
-						for(FAnimNode_LinkedAnimLayer* LayerNode : LayerPair.Value)
+						if(!bInDeferSubGraphInitialization)
 						{
-							if(LayerNode->LinkedRoot)
+							FAnimInstanceProxy& ThisProxy = GetProxyOnAnyThread<FAnimInstanceProxy>();
+							FAnimInstanceProxy& LinkedProxy = NewLinkedInstance->GetProxyOnAnyThread<FAnimInstanceProxy>();
+
+							// Initialize the correct parts of the linked instance
+							for(FAnimNode_LinkedAnimLayer* LayerNode : LayerPair.Value)
 							{
-								InitializeAndCacheBonesForLinkedRoot(LayerNode, ThisProxy, NewLinkedInstance, LinkedProxy);
+								if(LayerNode->LinkedRoot)
+								{
+									InitializeAndCacheBonesForLinkedRoot(LayerNode, ThisProxy, NewLinkedInstance, LinkedProxy);
+								}
 							}
 						}
 
@@ -2681,12 +2690,15 @@ void UAnimInstance::PerformLinkedLayerOverlayOperation(TSubclassOf<UAnimInstance
 					{
 						LayerNode->SetLinkedLayerInstance(this, nullptr);
 
-						UAnimInstance* LinkedInstance = LayerNode->GetTargetInstance<UAnimInstance>();
-						if(LayerNode->LinkedRoot && LinkedInstance)
+						if(!bInDeferSubGraphInitialization)
 						{
-							FAnimInstanceProxy& ThisProxy = GetProxyOnAnyThread<FAnimInstanceProxy>();
-							FAnimInstanceProxy& LinkedProxy = LinkedInstance->GetProxyOnAnyThread<FAnimInstanceProxy>();
-							InitializeAndCacheBonesForLinkedRoot(LayerNode, ThisProxy, LinkedInstance, LinkedProxy);
+							UAnimInstance* LinkedInstance = LayerNode->GetTargetInstance<UAnimInstance>();
+							if(LayerNode->LinkedRoot && LinkedInstance)
+							{
+								FAnimInstanceProxy& ThisProxy = GetProxyOnAnyThread<FAnimInstanceProxy>();
+								FAnimInstanceProxy& LinkedProxy = LinkedInstance->GetProxyOnAnyThread<FAnimInstanceProxy>();
+								InitializeAndCacheBonesForLinkedRoot(LayerNode, ThisProxy, LinkedInstance, LinkedProxy);
+							}
 						}
 					}
 				}
@@ -2733,9 +2745,23 @@ void UAnimInstance::UnlinkAnimClassLayers(TSubclassOf<UAnimInstance> InClass)
 	PerformLinkedLayerOverlayOperation(InClass, ConditionallySelectDefaultClass);
 }
 
-void UAnimInstance::InitializeGroupedLayers()
+void UAnimInstance::InitializeGroupedLayers(bool bInDeferSubGraphInitialization)
 {
-	LinkAnimClassLayers(nullptr);
+	auto SelectResolvedClassIfValid = [](UClass* InResolvedClass, FAnimNode_LinkedAnimLayer* InLayerNode)
+	{
+		if(InResolvedClass != nullptr)
+		{
+			// If we have a valid resolved class, use that as an overlay
+			return InResolvedClass;
+		}
+		else
+		{
+			// Otherwise use the default (which can be null)
+			return InLayerNode->InstanceClass.Get();
+		}
+	};
+
+	PerformLinkedLayerOverlayOperation(nullptr, SelectResolvedClassIfValid, bInDeferSubGraphInitialization);
 }
 
 UAnimInstance* UAnimInstance::GetLinkedAnimLayerInstanceByGroup(FName InGroup) const
diff --git a/Engine/Source/Runtime/Engine/Private/Animation/AnimInstanceProxy.cpp b/Engine/Source/Runtime/Engine/Private/Animation/AnimInstanceProxy.cpp
--- a/Engine/Source/Runtime/Engine/Private/Animation/AnimInstanceProxy.cpp
+++ b/Engine/Source/Runtime/Engine/Private/Animation/AnimInstanceProxy.cpp
@@ -1198,6 +1198,25 @@ void FAnimInstanceProxy::UpdateAnimation_WithRoot(const FAnimationUpdateContext&
 		if(bDeferRootNodeInitialization)
 		{
 			InitializeRootNode_WithRoot(RootNode);
+
+			if(AnimClassInterface)
+			{
+				// Initialize linked sub graphs
+				for(UStructProperty* LayerNodeProperty : AnimClassInterface->GetLinkedAnimLayerNodeProperties())
+				{
+					if(FAnimNode_LinkedAnimLayer* LayerNode = LayerNodeProperty->ContainerPtrToValuePtr<FAnimNode_LinkedAnimLayer>(AnimInstanceObject))
+					{
+						if(UAnimInstance* LinkedInstance = LayerNode->GetTargetInstance<UAnimInstance>())
+						{
+							FAnimationInitializeContext InitContext(this);
+							LayerNode->InitializeSubGraph_AnyThread(InitContext);
+							FAnimationCacheBonesContext CacheBonesContext(this);
+							LayerNode->CacheBonesSubGraph_AnyThread(CacheBonesContext);
+						}
+					}
+				}
+			}
+
 			bDeferRootNodeInitialization = false;
 		}
 
diff --git a/Engine/Source/Runtime/Engine/Private/Animation/AnimNodeBase.cpp b/Engine/Source/Runtime/Engine/Private/Animation/AnimNodeBase.cpp
--- a/Engine/Source/Runtime/Engine/Private/Animation/AnimNodeBase.cpp
+++ b/Engine/Source/Runtime/Engine/Private/Animation/AnimNodeBase.cpp
@@ -192,7 +192,7 @@ void FPoseLinkBase::AttemptRelink(const FAnimationBaseContext& Context)
 		check(AnimBlueprintClass);
 
 		// adding ensure. We had a crash on here
-		const TArray<UStructProperty*> AnimNodeProperties = AnimBlueprintClass->GetAnimNodeProperties();
+		const TArray<UStructProperty*>& AnimNodeProperties = AnimBlueprintClass->GetAnimNodeProperties();
 		if ( ensure(AnimNodeProperties.IsValidIndex(LinkID)) )
 		{
 			UProperty* LinkedProperty = AnimNodeProperties[LinkID];
diff --git a/Engine/Source/Runtime/Engine/Private/Animation/AnimNode_CustomProperty.cpp b/Engine/Source/Runtime/Engine/Private/Animation/AnimNode_CustomProperty.cpp
--- a/Engine/Source/Runtime/Engine/Private/Animation/AnimNode_CustomProperty.cpp
+++ b/Engine/Source/Runtime/Engine/Private/Animation/AnimNode_CustomProperty.cpp
@@ -62,6 +62,8 @@ void FAnimNode_CustomProperty::InitializeProperties(const UObject* InSourceInsta
 {
 	if(InTargetClass)
 	{
+		UClass* SourceClass = InSourceInstance->GetClass();
+
 		// Build property lists
 		SourceProperties.Reset(SourcePropertyNames.Num());
 		DestProperties.Reset(SourcePropertyNames.Num());
@@ -70,10 +72,8 @@ void FAnimNode_CustomProperty::InitializeProperties(const UObject* InSourceInsta
 
 		for(int32 Idx = 0; Idx < SourcePropertyNames.Num(); ++Idx)
 		{
-			FName& SourceName = SourcePropertyNames[Idx];
-			FName& DestName = DestPropertyNames[Idx];
-
-			UClass* SourceClass = InSourceInstance->GetClass();
+			const FName& SourceName = SourcePropertyNames[Idx];
+			const FName& DestName = DestPropertyNames[Idx];
 
 			UProperty* SourceProperty = FindField<UProperty>(SourceClass, SourceName);
 			UProperty* DestProperty = FindField<UProperty>(InTargetClass, DestName);
diff --git a/Engine/Source/Runtime/Engine/Private/Animation/AnimNode_LinkedAnimGraph.cpp b/Engine/Source/Runtime/Engine/Private/Animation/AnimNode_LinkedAnimGraph.cpp
--- a/Engine/Source/Runtime/Engine/Private/Animation/AnimNode_LinkedAnimGraph.cpp
+++ b/Engine/Source/Runtime/Engine/Private/Animation/AnimNode_LinkedAnimGraph.cpp
@@ -12,10 +12,8 @@ FAnimNode_LinkedAnimGraph::FAnimNode_LinkedAnimGraph()
 {
 }
 
-void FAnimNode_LinkedAnimGraph::Initialize_AnyThread(const FAnimationInitializeContext& Context)
+void FAnimNode_LinkedAnimGraph::InitializeSubGraph_AnyThread(const FAnimationInitializeContext& Context)
 {
-	FAnimNode_Base::Initialize_AnyThread(Context);
-
 	UAnimInstance* InstanceToRun = GetTargetInstance<UAnimInstance>();
 	if(InstanceToRun && LinkedRoot)
 	{
@@ -23,18 +21,22 @@ void FAnimNode_LinkedAnimGraph::Initialize_AnyThread(const FAnimationInitializeC
 		Proxy.InitializationCounter.SynchronizeWith(Context.AnimInstanceProxy->InitializationCounter);
 		Proxy.InitializeRootNode_WithRoot(LinkedRoot);
 	}
-	else
+}
+
+void FAnimNode_LinkedAnimGraph::Initialize_AnyThread(const FAnimationInitializeContext& Context)
+{
+	FAnimNode_Base::Initialize_AnyThread(Context);
+
+	InitializeSubGraph_AnyThread(Context);
+
+	// Make sure we propagate down all input poses, as they may not all be linked in the linked graph
+	for(FPoseLink& InputPose : InputPoses)
 	{
-		// If we have no valid instance (self or otherwise), we need to propagate down the graph to make sure
-		// subsequent nodes get properly initialized
-		for(FPoseLink& InputPose : InputPoses)
-		{
-			InputPose.Initialize(Context);
-		}
+		InputPose.Initialize(Context);
 	}
 }
 
-void FAnimNode_LinkedAnimGraph::CacheBones_AnyThread(const FAnimationCacheBonesContext& Context)
+void FAnimNode_LinkedAnimGraph::CacheBonesSubGraph_AnyThread(const FAnimationCacheBonesContext& Context)
 {
 	UAnimInstance* InstanceToRun = GetTargetInstance<UAnimInstance>();
 	if(InstanceToRun && LinkedRoot)
@@ -47,14 +49,16 @@ void FAnimNode_LinkedAnimGraph::CacheBones_AnyThread(const FAnimationCacheBonesC
 		FAnimationCacheBonesContext LinkedContext(&Proxy);
 		LinkedRoot->CacheBones_AnyThread(LinkedContext);
 	}
-	else
+}
+
+void FAnimNode_LinkedAnimGraph::CacheBones_AnyThread(const FAnimationCacheBonesContext& Context)
+{
+	CacheBonesSubGraph_AnyThread(Context);
+
+	// Make sure we propagate down all input poses, as they may not all be linked in the linked graph
+	for(FPoseLink& InputPose : InputPoses)
 	{
-		// If we have no valid instance (self or otherwise), we need to propagate down the graph to make sure
-		// subsequent nodes get properly their bones properly cached
-		for(FPoseLink& InputPose : InputPoses)
-		{
-			InputPose.CacheBones(Context);
-		}
+		InputPose.CacheBones(Context);
 	}
 }
 
@@ -257,11 +261,11 @@ void FAnimNode_LinkedAnimGraph::DynamicLink(UAnimInstance* InOwningAnimInstance)
 		if(SubAnimBlueprintClass)
 		{
 			FAnimInstanceProxy* NonConstProxy = &InOwningAnimInstance->GetProxyOnAnyThread<FAnimInstanceProxy>();
+			const FName FunctionToLink = GetDynamicLinkFunctionName();
 
 			// Link input poses
 			for(const FAnimBlueprintFunction& AnimBlueprintFunction : SubAnimBlueprintClass->GetAnimBlueprintFunctions())
 			{
-				const FName FunctionToLink = GetDynamicLinkFunctionName();
 				if(AnimBlueprintFunction.Name == FunctionToLink)
 				{
 					for(int32 InputPoseIndex = 0; InputPoseIndex < AnimBlueprintFunction.InputPoseNames.Num() && InputPoseIndex < InputPoses.Num(); ++InputPoseIndex)
@@ -311,10 +315,12 @@ void FAnimNode_LinkedAnimGraph::DynamicUnlink(UAnimInstance* InOwningAnimInstanc
 		IAnimClassInterface* SubAnimBlueprintClass = IAnimClassInterface::GetFromClass(LinkTargetInstance->GetClass());
 		if(SubAnimBlueprintClass)
 		{
+			const FName FunctionToLink = GetDynamicLinkFunctionName();
+
 			// Link input poses
 			for(const FAnimBlueprintFunction& AnimBlueprintFunction : SubAnimBlueprintClass->GetAnimBlueprintFunctions())
 			{
-				if(AnimBlueprintFunction.Name == GetDynamicLinkFunctionName())
+				if(AnimBlueprintFunction.Name == FunctionToLink)
 				{
 					for(int32 InputPoseIndex = 0; InputPoseIndex < AnimBlueprintFunction.InputPoseNames.Num() && InputPoseIndex < InputPoses.Num(); ++InputPoseIndex)
 					{
diff --git a/Engine/Source/Runtime/Engine/Private/Animation/AnimNode_LinkedAnimLayer.cpp b/Engine/Source/Runtime/Engine/Private/Animation/AnimNode_LinkedAnimLayer.cpp
--- a/Engine/Source/Runtime/Engine/Private/Animation/AnimNode_LinkedAnimLayer.cpp
+++ b/Engine/Source/Runtime/Engine/Private/Animation/AnimNode_LinkedAnimLayer.cpp
@@ -59,13 +59,15 @@ void FAnimNode_LinkedAnimLayer::InitializeSelfLayer(const UAnimInstance* SelfAni
 
 void FAnimNode_LinkedAnimLayer::SetLinkedLayerInstance(const UAnimInstance* InOwningAnimInstance, UAnimInstance* InNewLinkedInstance)
 {
-	ReinitializeLinkedAnimInstance(InOwningAnimInstance, InNewLinkedInstance);
-
 	// Reseting to running as a self-layer, in case it is applicable
 	if ((Interface.Get() == nullptr || InstanceClass.Get() == nullptr) && (InNewLinkedInstance == nullptr))
 	{
 		InitializeSelfLayer(InOwningAnimInstance);
 	}
+	else
+	{
+		ReinitializeLinkedAnimInstance(InOwningAnimInstance, InNewLinkedInstance);
+	}
 
 #if WITH_EDITOR
 	OnInstanceChangedEvent.Broadcast();
diff --git a/Engine/Source/Runtime/Engine/Private/Animation/AnimNode_LinkedInputPose.cpp b/Engine/Source/Runtime/Engine/Private/Animation/AnimNode_LinkedInputPose.cpp
--- a/Engine/Source/Runtime/Engine/Private/Animation/AnimNode_LinkedInputPose.cpp
+++ b/Engine/Source/Runtime/Engine/Private/Animation/AnimNode_LinkedInputPose.cpp
@@ -5,23 +5,24 @@
 
 const FName FAnimNode_LinkedInputPose::DefaultInputPoseName("InPose");
 
+// Note not calling through Initialize or CacheBones here.
+// This is handled in the owning LinkedAnimGraph node. This is because not all input poses may be linked in a
+// particular linked graph, so to avoid mismatches in initialization and bone references we make sure that all
+// branches of the tree are taken when initializing and caching bones.
+
+#if ENABLE_ANIMGRAPH_TRAVERSAL_DEBUG
 void FAnimNode_LinkedInputPose::Initialize_AnyThread(const FAnimationInitializeContext& Context)
 {
-	if(InputProxy)
-	{
-		FAnimationInitializeContext InputContext(InputProxy);
-		InputPose.Initialize(InputContext);
-	}
+	// Make sure to sync input pose debug counters as we still use this pose link in Update/Evaluate etc.
+	InputPose.InitializationCounter.SynchronizeWith(Context.AnimInstanceProxy->GetInitializationCounter());
 }
 
 void FAnimNode_LinkedInputPose::CacheBones_AnyThread(const FAnimationCacheBonesContext& Context)
 {
-	if(InputProxy)
-	{
-		FAnimationCacheBonesContext InputContext(InputProxy);
-		InputPose.CacheBones(InputContext);
-	}
+	// Make sure to sync input pose debug counters as we still use this pose link in Update/Evaluate etc.
+	InputPose.CachedBonesCounter.SynchronizeWith(Context.AnimInstanceProxy->GetCachedBonesCounter());
 }
+#endif
 
 void FAnimNode_LinkedInputPose::Update_AnyThread(const FAnimationUpdateContext& Context)
 {
@@ -68,7 +69,7 @@ void FAnimNode_LinkedInputPose::GatherDebugData(FNodeDebugData& DebugData)
 
 void FAnimNode_LinkedInputPose::DynamicLink(FAnimInstanceProxy* InInputProxy, FPoseLinkBase* InPoseLink)
 {
-	check(InputProxy == nullptr);	// Must be unlinked before re-linking
+	check(InputProxy == nullptr);			// Must be unlinked before re-linking
 
 	InputProxy = InInputProxy;
 	InputPose.SetDynamicLinkNode(InPoseLink);
@@ -76,7 +77,7 @@ void FAnimNode_LinkedInputPose::DynamicLink(FAnimInstanceProxy* InInputProxy, FP
 
 void FAnimNode_LinkedInputPose::DynamicUnlink()
 {
-	check(InputProxy != nullptr);	// Must be linked before unlinking
+	check(InputProxy != nullptr);			// Must be linked before unlinking
 
 	InputProxy = nullptr;
 	InputPose.SetDynamicLinkNode(nullptr);
