[NO NON-PII COMMIT MESSAGE]

diff --git a/Engine/Source/Runtime/Engine/Private/AudioDerivedData.cpp b/Engine/Source/Runtime/Engine/Private/AudioDerivedData.cpp
--- a/Engine/Source/Runtime/Engine/Private/AudioDerivedData.cpp
+++ b/Engine/Source/Runtime/Engine/Private/AudioDerivedData.cpp
@@ -454,7 +454,7 @@ class FStreamedAudioCacheDerivedDataWorker : public FNonAbandonableTask
 			}
 			else
 			{
-				UE_LOG(LogAudio, Warning, TEXT("Failed to retrieve compressed data for format %s and soundwave %s"),
+				UE_LOG(LogAudio, Display, TEXT("Failed to retrieve compressed data for format %s and soundwave %s."),
 					   *AudioFormatName.GetPlainNameString(),
 					   *SoundWave.GetPathName()
 					);
