Fix for issue with cascade detail mode sink.

diff --git a/Engine/Source/Runtime/Engine/Private/Particles/ParticleComponents.cpp b/Engine/Source/Runtime/Engine/Private/Particles/ParticleComponents.cpp
--- a/Engine/Source/Runtime/Engine/Private/Particles/ParticleComponents.cpp
+++ b/Engine/Source/Runtime/Engine/Private/Particles/ParticleComponents.cpp
@@ -3601,16 +3601,20 @@ bool UParticleSystemComponent::CanConsiderInvisible()const
 
 void DetailModeSink()
 {
-	int32 DetailMode = GetCachedScalabilityCVars().DetailMode;
-	static int32 CachedDetailMode = DetailMode;
+	//This Cvar sink can happen before the one which primes the cached scalability cvars so we must grab this ourselves.
+	IConsoleManager& ConsoleMan = IConsoleManager::Get();
+	static const auto DetailMode = ConsoleMan.FindTConsoleVariableDataInt(TEXT("r.DetailMode"));
+	int32 NewDetailMode = DetailMode->GetValueOnGameThread();
+	static int32 CachedDetailMode = NewDetailMode;
 
-	if (CachedDetailMode != DetailMode)
+	if (CachedDetailMode != NewDetailMode)
 	{
-		CachedDetailMode = DetailMode;
+		CachedDetailMode = NewDetailMode;
 
 		for (TObjectIterator<UParticleSystemComponent> It; It; ++It)
 		{
-			It->ForceReset();
+			//We must also reset on next tick rather than immediately as the cached cvar values are read internally to determin detail mode.
+			It->ResetNextTick();
 		}
 	}
 }
