Fix wrong Engine.h include

trivial

diff --git a/Engine/Source/Runtime/Engine/Private/Net/BandwidthTestActor.cpp b/Engine/Source/Runtime/Engine/Private/Net/BandwidthTestActor.cpp
--- a/Engine/Source/Runtime/Engine/Private/Net/BandwidthTestActor.cpp
+++ b/Engine/Source/Runtime/Engine/Private/Net/BandwidthTestActor.cpp
@@ -4,7 +4,7 @@
 
 #include "Net/UnrealNetwork.h"
 #include "EngineUtils.h"
-#include "Engine.h"
+#include "Engine/Engine.h"
 
 //-----------------------------------------------------------------------------
 //
diff --git a/Engine/Source/Runtime/Engine/Public/Net/BandwidthTestActor.h b/Engine/Source/Runtime/Engine/Public/Net/BandwidthTestActor.h
--- a/Engine/Source/Runtime/Engine/Public/Net/BandwidthTestActor.h
+++ b/Engine/Source/Runtime/Engine/Public/Net/BandwidthTestActor.h
@@ -3,6 +3,8 @@
 #pragma once
 
 #include "CoreMinimal.h"
+#include "GameFramework/Actor.h"
+
 #include "BandwidthTestActor.generated.h"
 
 //-----------------------------------------------------------------------------
