Shelve from Ola.Olsson:
Added 'r.DFDistanceScale' to enable scaling the distance field cut-off range (DistanceFieldShadowDistance *= r.DFDistanceScale)

diff --git a/Engine/Source/Runtime/Engine/Private/Components/DirectionalLightComponent.cpp b/Engine/Source/Runtime/Engine/Private/Components/DirectionalLightComponent.cpp
--- a/Engine/Source/Runtime/Engine/Private/Components/DirectionalLightComponent.cpp
+++ b/Engine/Source/Runtime/Engine/Private/Components/DirectionalLightComponent.cpp
@@ -71,7 +71,17 @@ static TAutoConsoleVariable<float> CVarRtdfFarTransitionScale(
 	TEXT("Use to modify the length of the far transition (fade out) of the distance field shadows.\n")
 	TEXT("1.0: (default) Calculate in the same way as other cascades.")
 	TEXT("0.0: Disable fade out."),
-	ECVF_RenderThreadSafe | ECVF_Scalability);
+	ECVF_RenderThreadSafe);
+
+static TAutoConsoleVariable<float> CVarRTDFDistanceScale(
+	TEXT("r.DFDistanceScale"),
+	1.0f,
+	TEXT("Factor to scale directional light property 'DistanceField Shadows Distance', clamped to [0.0001, 10000].\n")
+	TEXT("I.e.: DistanceFieldShadowsDistance *= r.DFDistanceScale.\n")	
+	TEXT("[0.0001,1): shorter distance\n")
+	TEXT(" 1: normal (default)\n")
+	TEXT("(1,10000]: larger distance.)"),
+	ECVF_RenderThreadSafe);
 
 
 /**
@@ -480,7 +490,7 @@ private:
 			return 0;
 		}
 
-		return DistanceFieldShadowDistance;
+		return DistanceFieldShadowDistance * FMath::Clamp(CVarRTDFDistanceScale.GetValueOnRenderThread(), 0.0001f, 10000.0f);
 	}
 
 	float GetShadowTransitionScale() const
@@ -640,7 +650,7 @@ private:
 			{
 				// there is only one distance field shadow cascade
 				check(SplitIndex == NumNearCascades + 1);
-				return DistanceFieldShadowDistance;
+				return GetDistanceFieldShadowDistance();
 			}
 			else
 			{
