PR #6192: BugFix: Fix WorldAsset name in LevelStreaming for duplicated levels (Contributed by [REDACTED_PERSON])

diff --git a/Engine/Source/Runtime/Engine/Private/LevelStreaming.cpp b/Engine/Source/Runtime/Engine/Private/LevelStreaming.cpp
--- a/Engine/Source/Runtime/Engine/Private/LevelStreaming.cpp
+++ b/Engine/Source/Runtime/Engine/Private/LevelStreaming.cpp
@@ -1418,8 +1418,9 @@ FName ULevelStreaming::GetLoadedLevelPackageName() const
 
 void ULevelStreaming::SetWorldAssetByPackageName(FName InPackageName)
 {
+	// Need to strip PIE prefix from object name, only the package has it
 	const FString TargetWorldPackageName = InPackageName.ToString();
-	const FString TargetWorldObjectName = FPackageName::GetLongPackageAssetName(TargetWorldPackageName);
+	const FString TargetWorldObjectName = UWorld::RemovePIEPrefix(FPackageName::GetLongPackageAssetName(TargetWorldPackageName));
 	TSoftObjectPtr<UWorld> NewWorld;
 	NewWorld = FString::Printf(TEXT("%s.%s"), *TargetWorldPackageName, *TargetWorldObjectName);
 	SetWorldAsset(NewWorld);
