Fixed display name and tooltip for Divide_TimespanFloat so it's now correct

diff --git a/Engine/Source/Runtime/Engine/Classes/Kismet/KismetMathLibrary.h b/Engine/Source/Runtime/Engine/Classes/Kismet/KismetMathLibrary.h
--- a/Engine/Source/Runtime/Engine/Classes/Kismet/KismetMathLibrary.h
+++ b/Engine/Source/Runtime/Engine/Classes/Kismet/KismetMathLibrary.h
@@ -2844,8 +2844,8 @@ class ENGINE_API UKismetMathLibrary : public UBlueprintFunctionLibrary
 	UFUNCTION(BlueprintPure, meta=(DisplayName="Timespan * float", CompactNodeTitle="*", Keywords="* multiply"), Category="Math|Timespan")
 	static FTimespan Multiply_TimespanFloat( FTimespan A, float Scalar );
 
-	/** Scalar division (A * s) */
-	UFUNCTION(BlueprintPure, meta=(DisplayName="Timespan * float", CompactNodeTitle="/", Keywords="/ divide"), Category="Math|Timespan")
+	/** Scalar division (A / s) */
+	UFUNCTION(BlueprintPure, meta=(DisplayName="Timespan / float", CompactNodeTitle="/", Keywords="/ divide"), Category="Math|Timespan")
 	static FTimespan Divide_TimespanFloat( FTimespan A, float Scalar );
 
 	/** Returns true if the values are equal (A == B) */
