Mark functions whose implementation is entirely with editoronly data as such rather than just compiling out their contents

diff --git a/Engine/Source/Runtime/Engine/Classes/Components/StaticMeshComponent.h b/Engine/Source/Runtime/Engine/Classes/Components/StaticMeshComponent.h
--- a/Engine/Source/Runtime/Engine/Classes/Components/StaticMeshComponent.h
+++ b/Engine/Source/Runtime/Engine/Classes/Components/StaticMeshComponent.h
@@ -404,7 +404,9 @@ public:
 	virtual void PostEditChangeProperty(FPropertyChangedEvent& PropertyChangedEvent) override;
 	virtual bool CanEditChange(const UProperty* InProperty) const override;
 #endif // WITH_EDITOR
+#if WITH_EDITORONLY_DATA
 	virtual void PreSave(const class ITargetPlatform* TargetPlatform) override;
+#endif // WITH_EDITORONLY_DATA
 	virtual void PostLoad() override;
 	virtual bool IsPostLoadThreadSafe() const override;
 	virtual bool AreNativePropertiesIdenticalTo( UObject* Other ) const override;
@@ -555,6 +557,7 @@ public:
 		int32& VertexLightMapMemoryUsage, int32& VertexShadowMapMemoryUsage,
 		int32& StaticLightingResolution, bool& bIsUsingTextureMapping, bool& bHasLightmapTexCoords) const;
 
+#if WITH_EDITORONLY_DATA
 	/**
 	 * Determines whether any of the component's LODs require override vertex color fixups
 	 *
@@ -588,9 +591,11 @@ public:
 	 * Removes instance vertex colors from all LODs
 	 */
 	void RemoveInstanceVertexColors();
+#endif
 
 	void UpdatePreCulledData(int32 LODIndex, const TArray<uint32>& PreCulledData, const TArray<int32>& NumTrianglesPerSection);
 
+#if WITH_EDITORONLY_DATA
 	/**
 	*	Sets the value of the SectionIndexPreview flag and reattaches the component as necessary.
 	*	@param	InSectionIndexPreview		New value of SectionIndexPreview.
@@ -602,7 +607,8 @@ public:
 	*	@param	InMaterialIndexPreview		New value of MaterialIndexPreview.
 	*/
 	void SetMaterialPreview(int32 InMaterialIndexPreview);
-	
+#endif
+
 	/** Sets the BodyInstance to use the mesh's body setup for external collision information*/
 	void UpdateCollisionFromStaticMesh();
 
@@ -616,8 +622,10 @@ private:
 	/** Initializes the resources used by the static mesh component. */
 	void InitResources();
 
+#if WITH_EDITOR
 	/** Update the vertex override colors */
 	void PrivateFixupOverrideColors();
+#endif
 protected:
 
 	/** Whether the component type supports static lighting. */
diff --git a/Engine/Source/Runtime/Engine/Private/Components/StaticMeshComponent.cpp b/Engine/Source/Runtime/Engine/Private/Components/StaticMeshComponent.cpp
--- a/Engine/Source/Runtime/Engine/Private/Components/StaticMeshComponent.cpp
+++ b/Engine/Source/Runtime/Engine/Private/Components/StaticMeshComponent.cpp
@@ -128,7 +128,9 @@ bool FStaticMeshComponentInstanceData::ApplyVertexColorData(UStaticMeshComponent
 				// that had its own vert color overrides; so before we apply
 				// the instance's color data, we need to clear the old
 				// vert colors (so we can properly call InitFromColorArray())
+#if WITH_EDITORONLY_DATA
 				StaticMeshComponent->RemoveInstanceVertexColorsFromLOD(LODIndex);
+#endif
 				// may not be null at the start (could have been initialized 
 				// from a  component template with vert coloring), but should
 				// be null at this point, after RemoveInstanceVertexColorsFromLOD()
@@ -376,13 +378,14 @@ bool UStaticMeshComponent::AreNativePropertiesIdenticalTo( UObject* Other ) cons
 	return bNativePropertiesAreIdentical;
 }
 
+#if WITH_EDITORONLY_DATA
 void UStaticMeshComponent::PreSave(const class ITargetPlatform* TargetPlatform)
 {
 	Super::PreSave(TargetPlatform);
-#if WITH_EDITORONLY_DATA
+
 	CachePaintedDataIfNecessary();
-#endif // WITH_EDITORONLY_DATA
 }
+#endif // WITH_EDITORONLY_DATA
 
 #if WITH_EDITOR
 void UStaticMeshComponent::CheckForErrors()
@@ -983,10 +986,9 @@ FTransform UStaticMeshComponent::GetSocketTransform(FName InSocketName, ERelativ
 	return Super::GetSocketTransform(InSocketName, TransformSpace);
 }
 
-
+#if WITH_EDITORONLY_DATA
 bool UStaticMeshComponent::RequiresOverrideVertexColorsFixup()
 {
-#if WITH_EDITORONLY_DATA
 	UStaticMesh* Mesh = GetStaticMesh();
 	if (!Mesh)
 	{
@@ -1027,36 +1029,28 @@ bool UStaticMeshComponent::RequiresOverrideVertexColorsFixup()
 	}
 
 	return true;
-#else
-	return false;
-#endif // WITH_EDITORONLY_DATA
 }
 
 void UStaticMeshComponent::SetSectionPreview(int32 InSectionIndexPreview)
 {
-#if WITH_EDITORONLY_DATA
 	if (SectionIndexPreview != InSectionIndexPreview)
 	{
 		SectionIndexPreview = InSectionIndexPreview;
 		MarkRenderStateDirty();
 	}
-#endif
 }
 
 void UStaticMeshComponent::SetMaterialPreview(int32 InMaterialIndexPreview)
 {
-#if WITH_EDITORONLY_DATA
 	if (MaterialIndexPreview != InMaterialIndexPreview)
 	{
 		MaterialIndexPreview = InMaterialIndexPreview;
 		MarkRenderStateDirty();
 	}
-#endif
 }
 
 void UStaticMeshComponent::RemoveInstanceVertexColorsFromLOD( int32 LODToRemoveColorsFrom )
 {
-#if WITH_EDITORONLY_DATA
 	if (GetStaticMesh() && LODToRemoveColorsFrom < GetStaticMesh()->GetNumLODs() && LODToRemoveColorsFrom < LODData.Num())
 	{
 		FStaticMeshComponentLODInfo& CurrentLODInfo = LODData[LODToRemoveColorsFrom];
@@ -1065,22 +1059,18 @@ void UStaticMeshComponent::RemoveInstanceVertexColorsFromLOD( int32 LODToRemoveC
 		CurrentLODInfo.PaintedVertices.Empty();
 		StaticMeshDerivedDataKey = GetStaticMesh()->RenderData->DerivedDataKey;
 	}
-#endif
 }
 
 void UStaticMeshComponent::RemoveInstanceVertexColors()
 {
-#if WITH_EDITORONLY_DATA
 	for ( int32 i=0; i < GetStaticMesh()->GetNumLODs(); i++ )
 	{
 		RemoveInstanceVertexColorsFromLOD( i );
 	}
-#endif
 }
 
 void UStaticMeshComponent::CopyInstanceVertexColorsIfCompatible( UStaticMeshComponent* SourceComponent )
 {
-#if WITH_EDITORONLY_DATA
 	// The static mesh assets have to match, currently.
 	if (( GetStaticMesh()->GetPathName() == SourceComponent->GetStaticMesh()->GetPathName() ) &&
 		( SourceComponent->LODData.Num() != 0 ))
@@ -1149,12 +1139,10 @@ void UStaticMeshComponent::CopyInstanceVertexColorsIfCompatible( UStaticMeshComp
 
 		MarkRenderStateDirty();
 	}
-#endif
 }
 
 void UStaticMeshComponent::CachePaintedDataIfNecessary()
 {
-#if WITH_EDITORONLY_DATA
 	// Only cache the vertex positions if we're in the editor
 	if ( GIsEditor && GetStaticMesh() )
 	{
@@ -1207,12 +1195,10 @@ void UStaticMeshComponent::CachePaintedDataIfNecessary()
 			}
 		}
 	}
-#endif // WITH_EDITORONLY_DATA
 }
 
 bool UStaticMeshComponent::FixupOverrideColorsIfNecessary( bool bRebuildingStaticMesh )
 {
-#if WITH_EDITORONLY_DATA
 
 	// Detect if there is a version mismatch between the source mesh and the component. If so, the component's LODs potentially
 	// need to have their override colors updated to match changes in the source mesh.
@@ -1234,10 +1220,10 @@ bool UStaticMeshComponent::FixupOverrideColorsIfNecessary( bool bRebuildingStati
 
 		return true;
 	}
-#endif // WITH_EDITORONLY_DATA
 
 	return false;
 }
+#endif // WITH_EDITORONLY_DATA
 
 void UStaticMeshComponent::InitResources()
 {
@@ -1252,9 +1238,9 @@ void UStaticMeshComponent::InitResources()
 	}
 }
 
+#if WITH_EDITOR
 void UStaticMeshComponent::PrivateFixupOverrideColors()
 {
-#if WITH_EDITOR
 	if (!GetStaticMesh() || !GetStaticMesh()->RenderData)
 	{
 		return;
@@ -1349,9 +1335,8 @@ void UStaticMeshComponent::PrivateFixupOverrideColors()
 	{
 		StaticMeshDerivedDataKey = GetStaticMesh()->RenderData->DerivedDataKey;
 	}
-
-#endif // WITH_EDITOR
 }
+#endif // WITH_EDITOR
 
 float GKeepPreCulledIndicesThreshold = .95f;
 
@@ -1497,11 +1482,13 @@ void UStaticMeshComponent::ImportCustomProperties(const TCHAR* SourceText, FFeed
 		{
 			SourceText = VertColorStr;
 
+#if WITH_EDITORONLY_DATA
 			// this component could have been constructed from a template that
 			// had its own vert color overrides; so before we apply the
 			// custom color data, we need to clear the old vert colors (so
 			// we can properly call ImportText())
 			RemoveInstanceVertexColorsFromLOD(LODIndex);
+#endif
 
 			// may not be null at the start (could have been initialized 
 			// from a blueprint component template with vert coloring), but 
@@ -1697,6 +1684,7 @@ void UStaticMeshComponent::PostLoad()
 
 	Super::PostLoad();
 
+#if WITH_EDITORONLY_DATA
 	if ( GetStaticMesh() )
 	{
 		CachePaintedDataIfNecessary();
@@ -1705,7 +1693,6 @@ void UStaticMeshComponent::PostLoad()
 
 		if (FixupOverrideColorsIfNecessary())
 		{
-#if WITH_EDITORONLY_DATA
 
 			AActor* Owner = GetOwner();
 
@@ -1720,9 +1707,9 @@ void UStaticMeshComponent::PostLoad()
 					Level->FixupOverrideVertexColorsCount++;
 				}
 			}
-#endif
 		}
 	}
+#endif
 
 	// Empty after potential editor fix-up when we don't care about re-saving, e.g. game or client
 	if (!GIsEditor && !IsRunningCommandlet())
diff --git a/Engine/Source/Runtime/Engine/Private/MeshVertexPainter/MeshVertexPainter.cpp b/Engine/Source/Runtime/Engine/Private/MeshVertexPainter/MeshVertexPainter.cpp
--- a/Engine/Source/Runtime/Engine/Private/MeshVertexPainter/MeshVertexPainter.cpp
+++ b/Engine/Source/Runtime/Engine/Private/MeshVertexPainter/MeshVertexPainter.cpp
@@ -21,7 +21,9 @@ void FMeshVertexPainter::PaintVerticesSingleColor(UStaticMeshComponent* StaticMe
 	uint32 LODIndex = 0;
 	for (FStaticMeshComponentLODInfo& LODInfo : StaticMeshComponent->LODData)
 	{
+#if WITH_EDITORONLY_DATA
 		StaticMeshComponent->RemoveInstanceVertexColorsFromLOD(LODIndex);
+#endif
 		check(LODInfo.OverrideVertexColors == nullptr);
 
 		const FStaticMeshLODResources& LODModel = StaticMeshComponent->GetStaticMesh()->RenderData->LODResources[LODIndex];
@@ -36,7 +38,9 @@ void FMeshVertexPainter::PaintVerticesSingleColor(UStaticMeshComponent* StaticMe
 		LODIndex++;
 	}
 
+#if WITH_EDITORONLY_DATA
 	StaticMeshComponent->CachePaintedDataIfNecessary();
+#endif
 	StaticMeshComponent->MarkRenderStateDirty();
 	StaticMeshComponent->bDisallowMeshPaintPerInstance = true;
 }
@@ -60,7 +64,9 @@ void FMeshVertexPainter::PaintVerticesLerpAlongAxis(UStaticMeshComponent* Static
 	uint32 LODIndex = 0;
 	for (FStaticMeshComponentLODInfo& LODInfo : StaticMeshComponent->LODData)
 	{
+#if WITH_EDITORONLY_DATA
 		StaticMeshComponent->RemoveInstanceVertexColorsFromLOD(LODIndex);
+#endif
 		check(LODInfo.OverrideVertexColors == nullptr);
 
 		const FStaticMeshLODResources& LODModel = StaticMeshComponent->GetStaticMesh()->RenderData->LODResources[LODIndex];
@@ -85,7 +91,9 @@ void FMeshVertexPainter::PaintVerticesLerpAlongAxis(UStaticMeshComponent* Static
 		LODIndex++;
 	}
 
+#if WITH_EDITORONLY_DATA
 	StaticMeshComponent->CachePaintedDataIfNecessary();
+#endif
 	StaticMeshComponent->MarkRenderStateDirty();
 	StaticMeshComponent->bDisallowMeshPaintPerInstance = true;
 }
@@ -100,7 +108,9 @@ void FMeshVertexPainter::RemovePaintedVertices(UStaticMeshComponent* StaticMeshC
 	uint32 LODIndex = 0;
 	for (FStaticMeshComponentLODInfo& LODInfo : StaticMeshComponent->LODData)
 	{
+#if WITH_EDITORONLY_DATA
 		StaticMeshComponent->RemoveInstanceVertexColorsFromLOD(LODIndex);
+#endif
 		check(LODInfo.OverrideVertexColors == nullptr);
 		LODIndex++;
 	}
