PR #6176: Add a GetComponentsByInterface function for Actor (Contributed by [REDACTED_PERSON])

diff --git a/Engine/Source/Runtime/Engine/Classes/GameFramework/Actor.h b/Engine/Source/Runtime/Engine/Classes/GameFramework/Actor.h
--- a/Engine/Source/Runtime/Engine/Classes/GameFramework/Actor.h
+++ b/Engine/Source/Runtime/Engine/Classes/GameFramework/Actor.h
@@ -2798,6 +2798,10 @@ public:
 	UFUNCTION(BlueprintCallable, Category = "Actor", meta = (ComponentClass = "ActorComponent"), meta = (DeterminesOutputType = "ComponentClass"))
 	TArray<UActorComponent*> GetComponentsByTag(TSubclassOf<UActorComponent> ComponentClass, FName Tag) const;
 
+	/** Gets all the components that implements the given interface. */
+	UFUNCTION(BlueprintCallable, Category = "Actor")
+	TArray<UActorComponent*> GetComponentsByInterface(TSubclassOf<UInterface> Interface) const;
+
 	/** Templatized version of FindComponentByClass that handles casting for you */
 	template<class T>
 	T* FindComponentByClass() const
diff --git a/Engine/Source/Runtime/Engine/Private/Actor.cpp b/Engine/Source/Runtime/Engine/Private/Actor.cpp
--- a/Engine/Source/Runtime/Engine/Private/Actor.cpp
+++ b/Engine/Source/Runtime/Engine/Private/Actor.cpp
@@ -2842,6 +2842,24 @@ TArray<UActorComponent*> AActor::GetComponentsByTag(TSubclassOf<UActorComponent>
 	return ComponentsByTag;
 }
 
+TArray<UActorComponent*> AActor::GetComponentsByInterface(TSubclassOf<UInterface> Interface) const
+{
+	TArray<UActorComponent*> Components;
+
+	if (Interface)
+	{
+		for (UActorComponent* Component : GetComponents())
+		{
+			if (Component && Component->GetClass()->ImplementsInterface(Interface))
+			{
+				Components.Add(Component);
+			}
+		}
+	}
+
+	return Components;
+}
+
 void AActor::DisableComponentsSimulatePhysics()
 {
 	for (UActorComponent* Component : GetComponents())
