ChildActor doesn't need to be attached to ChildActorComponent as part of OnRegister, it is properly done in CreateChild.

diff --git a/Engine/Source/Runtime/Engine/Private/Components/ChildActorComponent.cpp b/Engine/Source/Runtime/Engine/Private/Components/ChildActorComponent.cpp
--- a/Engine/Source/Runtime/Engine/Private/Components/ChildActorComponent.cpp
+++ b/Engine/Source/Runtime/Engine/Private/Components/ChildActorComponent.cpp
@@ -40,16 +40,6 @@ void UChildActorComponent::OnRegister()
 		}
 		else
 		{
-			USceneComponent* ChildRoot = ChildActor->GetRootComponent();
-			if (ChildRoot && ChildRoot->GetAttachParent() != this)
-			{
-				// attach new actor to this component
-				// we can't attach in CreateChildActor since it has intermediate Mobility set up
-				// causing spam with inconsistent mobility set up
-				// so moving Attach to happen in Register
-				ChildRoot->AttachToComponent(this, FAttachmentTransformRules::SnapToTargetNotIncludingScale);
-			}
-
 			// Ensure the components replication is correctly initialized
 			SetIsReplicated(ChildActor->GetIsReplicated());
 		}
