Pull calls to RemoveInstanceVertexColorsFromLOD out of editor only in light of incoming changes
#rnx

diff --git a/Engine/Source/Runtime/Engine/Classes/Components/StaticMeshComponent.h b/Engine/Source/Runtime/Engine/Classes/Components/StaticMeshComponent.h
--- a/Engine/Source/Runtime/Engine/Classes/Components/StaticMeshComponent.h
+++ b/Engine/Source/Runtime/Engine/Classes/Components/StaticMeshComponent.h
@@ -580,6 +580,7 @@ public:
 	 * @param SourceComponent The component to copy vertex colors from
 	 */
 	void CopyInstanceVertexColorsIfCompatible( UStaticMeshComponent* SourceComponent );
+#endif
 
 	/**
 	 * Removes instance vertex colors from the specified LOD
@@ -587,6 +588,7 @@ public:
 	 */
 	void RemoveInstanceVertexColorsFromLOD( int32 LODToRemoveColorsFrom );
 
+#if WITH_EDITORONLY_DATA
 	/**
 	 * Removes instance vertex colors from all LODs
 	 */
diff --git a/Engine/Source/Runtime/Engine/Private/Components/StaticMeshComponent.cpp b/Engine/Source/Runtime/Engine/Private/Components/StaticMeshComponent.cpp
--- a/Engine/Source/Runtime/Engine/Private/Components/StaticMeshComponent.cpp
+++ b/Engine/Source/Runtime/Engine/Private/Components/StaticMeshComponent.cpp
@@ -128,9 +128,7 @@ bool FStaticMeshComponentInstanceData::ApplyVertexColorData(UStaticMeshComponent
 				// that had its own vert color overrides; so before we apply
 				// the instance's color data, we need to clear the old
 				// vert colors (so we can properly call InitFromColorArray())
-#if WITH_EDITORONLY_DATA
 				StaticMeshComponent->RemoveInstanceVertexColorsFromLOD(LODIndex);
-#endif
 				// may not be null at the start (could have been initialized 
 				// from a  component template with vert coloring), but should
 				// be null at this point, after RemoveInstanceVertexColorsFromLOD()
@@ -1487,13 +1485,11 @@ void UStaticMeshComponent::ImportCustomProperties(const TCHAR* SourceText, FFeed
 		{
 			SourceText = VertColorStr;
 
-#if WITH_EDITORONLY_DATA
 			// this component could have been constructed from a template that
 			// had its own vert color overrides; so before we apply the
 			// custom color data, we need to clear the old vert colors (so
 			// we can properly call ImportText())
 			RemoveInstanceVertexColorsFromLOD(LODIndex);
-#endif
 
 			// may not be null at the start (could have been initialized 
 			// from a blueprint component template with vert coloring), but 
diff --git a/Engine/Source/Runtime/Engine/Private/MeshVertexPainter/MeshVertexPainter.cpp b/Engine/Source/Runtime/Engine/Private/MeshVertexPainter/MeshVertexPainter.cpp
--- a/Engine/Source/Runtime/Engine/Private/MeshVertexPainter/MeshVertexPainter.cpp
+++ b/Engine/Source/Runtime/Engine/Private/MeshVertexPainter/MeshVertexPainter.cpp
@@ -21,9 +21,7 @@ void FMeshVertexPainter::PaintVerticesSingleColor(UStaticMeshComponent* StaticMe
 	uint32 LODIndex = 0;
 	for (FStaticMeshComponentLODInfo& LODInfo : StaticMeshComponent->LODData)
 	{
-#if WITH_EDITORONLY_DATA
 		StaticMeshComponent->RemoveInstanceVertexColorsFromLOD(LODIndex);
-#endif
 		check(LODInfo.OverrideVertexColors == nullptr);
 
 		const FStaticMeshLODResources& LODModel = StaticMeshComponent->GetStaticMesh()->RenderData->LODResources[LODIndex];
@@ -64,9 +62,7 @@ void FMeshVertexPainter::PaintVerticesLerpAlongAxis(UStaticMeshComponent* Static
 	uint32 LODIndex = 0;
 	for (FStaticMeshComponentLODInfo& LODInfo : StaticMeshComponent->LODData)
 	{
-#if WITH_EDITORONLY_DATA
 		StaticMeshComponent->RemoveInstanceVertexColorsFromLOD(LODIndex);
-#endif
 		check(LODInfo.OverrideVertexColors == nullptr);
 
 		const FStaticMeshLODResources& LODModel = StaticMeshComponent->GetStaticMesh()->RenderData->LODResources[LODIndex];
@@ -108,9 +104,7 @@ void FMeshVertexPainter::RemovePaintedVertices(UStaticMeshComponent* StaticMeshC
 	uint32 LODIndex = 0;
 	for (FStaticMeshComponentLODInfo& LODInfo : StaticMeshComponent->LODData)
 	{
-#if WITH_EDITORONLY_DATA
 		StaticMeshComponent->RemoveInstanceVertexColorsFromLOD(LODIndex);
-#endif
 		check(LODInfo.OverrideVertexColors == nullptr);
 		LODIndex++;
 	}
