CVarUnloadPrimaryAssetsWithType was incorrectly calling LoadPrimaryAssetsWithType, now calls Unload variant

diff --git a/Engine/Source/Runtime/Engine/Private/AssetManager.cpp b/Engine/Source/Runtime/Engine/Private/AssetManager.cpp
--- a/Engine/Source/Runtime/Engine/Private/AssetManager.cpp
+++ b/Engine/Source/Runtime/Engine/Private/AssetManager.cpp
@@ -3823,7 +3823,7 @@ static FAutoConsoleCommandWithWorldAndArgs CVarLoadPrimaryAssetsWithType(
 // Cheat command to unload all assets of a given type
 static FAutoConsoleCommandWithWorldAndArgs CVarUnloadPrimaryAssetsWithType(
 	TEXT("AssetManager.UnloadPrimaryAssetsWithType"),
-	TEXT("Loads all assets of a given type"),
+	TEXT("Unloads all assets of a given type"),
 	FConsoleCommandWithWorldAndArgsDelegate::CreateStatic(
 		[](const TArray<FString>& Params, UWorld* World)
 {
@@ -3834,13 +3834,13 @@ static FAutoConsoleCommandWithWorldAndArgs CVarUnloadPrimaryAssetsWithType(
 
 	for (const FString& Param : Params)
 	{
-		const FPrimaryAssetType TypeToLoad(*Param);
+		const FPrimaryAssetType TypeToUnload(*Param);
 
 		FPrimaryAssetTypeInfo Info;
-		if (UAssetManager::Get().GetPrimaryAssetTypeInfo(TypeToLoad, /*out*/ Info))
+		if (UAssetManager::Get().GetPrimaryAssetTypeInfo(TypeToUnload, /*out*/ Info))
 		{
-			UE_LOG(LogAssetManager, Log, TEXT("UnloadPrimaryAssetsWithType(%s)"), *Param);
-			UAssetManager::Get().LoadPrimaryAssetsWithType(TypeToLoad);
+			int32 NumUnloaded = UAssetManager::Get().UnloadPrimaryAssetsWithType(TypeToUnload);
+			UE_LOG(LogAssetManager, Log, TEXT("UnloadPrimaryAssetsWithType(%s): Unloaded %d assets"), *Param, NumUnloaded);
 		}
 		else
 		{
