Make it a fatal error to try to spawn an Actor with the same name as an existing Actor in the level

diff --git a/Engine/Source/Runtime/Engine/Private/LevelActor.cpp b/Engine/Source/Runtime/Engine/Private/LevelActor.cpp
--- a/Engine/Source/Runtime/Engine/Private/LevelActor.cpp
+++ b/Engine/Source/Runtime/Engine/Private/LevelActor.cpp
@@ -384,9 +384,17 @@ AActor* UWorld::SpawnActor( UClass* Class, FTransform const* UserTransformPtr, c
 		// Use class's default actor as a template.
 		Template = Class->GetDefaultObject<AActor>();
 	}
-	else if (NewActorName.IsNone() && !Template->HasAnyFlags(RF_ClassDefaultObject))
+	else if (!Template->HasAnyFlags(RF_ClassDefaultObject))
 	{
-		NewActorName = MakeUniqueObjectName(LevelToSpawnIn, Template->GetClass(), *Template->GetFName().GetPlainNameString());
+		if (NewActorName.IsNone())
+		{
+			NewActorName = MakeUniqueObjectName(LevelToSpawnIn, Template->GetClass(), *Template->GetFName().GetPlainNameString());
+		}
+		else if (StaticFindObjectFast(nullptr, LevelToSpawnIn, NewActorName))
+		{
+			UE_LOG(LogSpawn, Fatal, TEXT("An actor of name '%s' already exists in level '%s'."), *NewActorName.ToString(), *LevelToSpawnIn->GetFullName());
+			return nullptr;
+		}
 	}
 	check(Template);
 
