PR #4242: UE-52482: Reset InitialSeed instead of regenerating (Contributed by [REDACTED_PERSON])

diff --git a/Engine/Source/Runtime/Engine/Private/ActorConstruction.cpp b/Engine/Source/Runtime/Engine/Private/ActorConstruction.cpp
--- a/Engine/Source/Runtime/Engine/Private/ActorConstruction.cpp
+++ b/Engine/Source/Runtime/Engine/Private/ActorConstruction.cpp
@@ -42,6 +42,7 @@ DECLARE_CYCLE_STAT(TEXT("InstanceActorComponent"), STAT_InstanceActorComponent,
 //////////////////////////////////////////////////////////////////////////
 // AActor Blueprint Stuff
 
+#if WITH_EDITOR
 static TArray<FRandomStream*> FindRandomStreams(AActor* InActor)
 {
 	check(InActor);
@@ -59,13 +60,12 @@ static TArray<FRandomStream*> FindRandomStreams(AActor* InActor)
 	return OutStreams;
 }
 
-#if WITH_EDITOR
 void AActor::SeedAllRandomStreams()
 {
 	TArray<FRandomStream*> Streams = FindRandomStreams(this);
 	for(int32 i=0; i<Streams.Num(); i++)
 	{
-		Streams[i]->GenerateNewSeed();
+		Streams[i]->Reset();
 	}
 }
 #endif //WITH_EDITOR
