Remove unused function IsGeneratedClassCompileType

diff --git a/Engine/Source/Runtime/Engine/Classes/Engine/Blueprint.h b/Engine/Source/Runtime/Engine/Classes/Engine/Blueprint.h
--- a/Engine/Source/Runtime/Engine/Classes/Engine/Blueprint.h
+++ b/Engine/Source/Runtime/Engine/Classes/Engine/Blueprint.h
@@ -179,12 +179,6 @@ public:
 			|| (CompileType == EKismetCompileType::Cpp);
 	}
 
-	/** Whether or not this compile type should operate on the generated class of the blueprint, as opposed to just the skeleton */
-	bool IsGeneratedClassCompileType() const
-	{
-		return (CompileType != EKismetCompileType::SkeletonOnly);
-	}
-
 	FKismetCompilerOptions()
 		: CompileType(EKismetCompileType::Full)
 		, bSaveIntermediateProducts(false)
