Fix a typo in UFUNCTION() metadata for SpawnForceFeedbackAtLocation().

- Original CL: 3206776

diff --git a/Engine/Source/Runtime/Engine/Classes/Kismet/GameplayStatics.h b/Engine/Source/Runtime/Engine/Classes/Kismet/GameplayStatics.h
--- a/Engine/Source/Runtime/Engine/Classes/Kismet/GameplayStatics.h
+++ b/Engine/Source/Runtime/Engine/Classes/Kismet/GameplayStatics.h
@@ -671,7 +671,7 @@ public:
 	 * @param bAutoDestroy - Whether the returned force feedback component will be automatically cleaned up when the feedback patern finishes (by completing or stopping) or whether it can be reactivated
 	 * @return Force Feedback Component to manipulate the playing feedback effect with
 	 */
-	UFUNCTION(BlueprintCallable, Category="ForceFeedback", meta=(WorldContext="WorldContextObject", AdvancedDisplay = "3", UnsafeDuringActorConsturction = "true", Keywords = "play"))
+	UFUNCTION(BlueprintCallable, Category="ForceFeedback", meta=(WorldContext="WorldContextObject", AdvancedDisplay = "3", UnsafeDuringActorConstruction = "true", Keywords = "play"))
 	static UForceFeedbackComponent* SpawnForceFeedbackAtLocation(const UObject* WorldContextObject, UForceFeedbackEffect* ForceFeedbackEffect, FVector Location, FRotator Rotation = FRotator::ZeroRotator, bool bLooping = false, float IntensityMultiplier = 1.f, float StartTime = 0.f, UForceFeedbackAttenuation* AttenuationSettings = nullptr, bool bAutoDestroy = true);
 
 	/** Plays a force feedback effect attached to and following the specified component. This is a fire and forget effect. Replication is also not handled at this point.
