PR #3080: Fix for GetNextViewablePlayer() (Contributed by [REDACTED_PERSON])

diff --git a/Engine/Source/Runtime/Engine/Classes/Camera/PlayerCameraManager.h b/Engine/Source/Runtime/Engine/Classes/Camera/PlayerCameraManager.h
--- a/Engine/Source/Runtime/Engine/Classes/Camera/PlayerCameraManager.h
+++ b/Engine/Source/Runtime/Engine/Classes/Camera/PlayerCameraManager.h
@@ -78,10 +78,9 @@ protected:
 	UPROPERTY(EditAnywhere, BlueprintReadWrite, Category=TViewTarget)
 	class APlayerState* PlayerState;
 
-	//@TODO: All for GetNextViewablePlayer
-	friend class APlayerController;
-
 public:
+	class APlayerState* GetPlayerState() const { return PlayerState; }
+	
 	void SetNewTarget(AActor* NewTarget);
 
 	class APawn* GetTargetPawn() const;
diff --git a/Engine/Source/Runtime/Engine/Private/PlayerController.cpp b/Engine/Source/Runtime/Engine/Private/PlayerController.cpp
--- a/Engine/Source/Runtime/Engine/Private/PlayerController.cpp
+++ b/Engine/Source/Runtime/Engine/Private/PlayerController.cpp
@@ -2905,8 +2905,6 @@ void APlayerController::ServerViewPrevPlayer_Implementation()
 
 APlayerState* APlayerController::GetNextViewablePlayer(int32 dir)
 {
-	int32 CurrentIndex = -1;
-
 	UWorld* World = GetWorld();
 	AGameModeBase* GameMode = World->GetAuthGameMode();
 	AGameStateBase* GameState = World->GetGameState();
@@ -2917,44 +2915,33 @@ APlayerState* APlayerController::GetNextViewablePlayer(int32 dir)
 		return nullptr;
 	}
 
-	if (PlayerCameraManager && PlayerCameraManager->ViewTarget.PlayerState)
+	APlayerState* NextPlayerState = (PlayerCameraManager ? PlayerCameraManager->ViewTarget.GetPlayerState() : nullptr);
+	
+	// If we don't have a NextPlayerState, use our own.
+	// This will allow us to attempt to find another player to view or, if all else fails, makes sure we have a playerstate set for next time.
+	int32 NextIndex = (NextPlayerState ? GameState->PlayerArray.Find(NextPlayerState) : GameState->PlayerArray.Find(PlayerState));
+
+	// Cycle through the player states until we find a valid one.
+	for (int32 i = 0; i < GameState->PlayerArray.Num(); ++i)
 	{
-		// Find index of current viewtarget's PlayerState
-		for ( int32 i=0; i<GameState->PlayerArray.Num(); i++ )
+		NextIndex = ((NextIndex == 0) && (dir < 0)) ? (GameState->PlayerArray.Num() - 1) : ((NextIndex == (GameState->PlayerArray.Num() - 1)) && (dir > 0)) ? 0 : NextIndex += dir;
+		NextPlayerState = GameState->PlayerArray[NextIndex];
+
+		// Make sure we're not trying to view our own player state.
+		if (NextPlayerState != PlayerState)
 		{
-			if (PlayerCameraManager->ViewTarget.PlayerState == GameState->PlayerArray[i])
+			AController* NextController = Cast<AController>(NextPlayerState->GetOwner());
+
+			// Check they have a pawn & the game mode is ok with us spectating them.
+			if (NextController && NextController->GetPawn() && GameMode->CanSpectate(this, NextPlayerState))
 			{
-				CurrentIndex = i;
 				break;
 			}
 		}
 	}
 
-	// Find next valid viewtarget in appropriate direction
-	int32 NewIndex;
-	for ( NewIndex=CurrentIndex+dir; (NewIndex>=0)&&(NewIndex<GameState->PlayerArray.Num()); NewIndex=NewIndex+dir )
-	{
-		APlayerState* const NextPlayerState = GameState->PlayerArray[NewIndex];
-		AController* NextController = (NextPlayerState ? Cast<AController>(NextPlayerState->GetOwner()) : nullptr);
-		if ( NextController && NextController->GetPawn() != nullptr && GameMode->CanSpectate(this, NextPlayerState) )
-		{
-			return NextPlayerState;
-		}
-	}
-
-	// wrap around
-	CurrentIndex = (NewIndex < 0) ? GameState->PlayerArray.Num() : -1;
-	for ( NewIndex=CurrentIndex+dir; (NewIndex>=0)&&(NewIndex<GameState->PlayerArray.Num()); NewIndex=NewIndex+dir )
-	{
-		APlayerState* const NextPlayerState = GameState->PlayerArray[NewIndex];
-		AController* NextController = (NextPlayerState ? Cast<AController>(NextPlayerState->GetOwner()) : nullptr);
-		if ( NextController && NextController->GetPawn() != nullptr && GameMode->CanSpectate(this, NextPlayerState) )
-		{
-			return NextPlayerState;
-		}
-	}
-
-	return nullptr;
+	// If we've failed to find another player to view, we'll be back to our original view target playerstate.
+	return NextPlayerState;
 }
 
 void APlayerController::ViewAPlayer(int32 dir)
