FURL: Fix parsing string retrieved from FURL::ToString that have the default (empty) host, but a non-default port

diff --git a/Engine/Source/Runtime/Engine/Private/URL.cpp b/Engine/Source/Runtime/Engine/Private/URL.cpp
--- a/Engine/Source/Runtime/Engine/Private/URL.cpp
+++ b/Engine/Source/Runtime/Engine/Private/URL.cpp
@@ -288,7 +288,7 @@ FURL::FURL( FURL* Base, const TCHAR* TextURL, ETravelType Type )
 
 			// Square bracket indicates an IPv6 address, but IPv6 addresses can contain dots also
 			// They also typically have multiple colons in them as well.
-			if (bIsHostnameWithDot || SquareBracket || (Colon && Colon != LastColonLocation) || bHasMultipleColons)
+			if (bIsHostnameWithDot || SquareBracket || (Colon && Colon == LastColonLocation) || bHasMultipleColons)
 			{
 				TCHAR* ss = URL;
 				// Clear out the URL such that all that should be left is the map
@@ -685,6 +685,7 @@ static bool URLSerializationTests(UWorld* InWorld, const TCHAR* Cmd, FOutputDevi
 		TestCases.Push(FURLTestCase(TEXT("http://[::ffff:192.168.0.1]:8080"), TEXT("::ffff:192.168.0.1"), TEXT("http"), 8080));
 		TestCases.Push(FURLTestCase(TEXT("https:[2001:db8:85a3::8a2e:370:7334]:443"), TEXT("2001:db8:85a3::8a2e:370:7334"), TEXT("https"), 443));
 		TestCases.Push(FURLTestCase(TEXT("steam.76561197993275299:20/"), TEXT("steam.76561197993275299"), TEXT(""), 20));
+		TestCases.Push(FURLTestCase(TEXT("unreal::44750/"), TEXT(""), TEXT("unreal"), 44750));
 
 		bool bAllCasesPassed = true;
 		for (const auto& TestCase : TestCases)
