Fixed crash in ray tracing when when r.DeferSkeletalDynamicDataUpdateUntilGDME is set to 1

diff --git a/Engine/Source/Runtime/Engine/Private/SkeletalRenderGPUSkin.cpp b/Engine/Source/Runtime/Engine/Private/SkeletalRenderGPUSkin.cpp
--- a/Engine/Source/Runtime/Engine/Private/SkeletalRenderGPUSkin.cpp
+++ b/Engine/Source/Runtime/Engine/Private/SkeletalRenderGPUSkin.cpp
@@ -333,6 +333,7 @@ void FSkeletalMeshObjectGPUSkin::UpdateDynamicData_RenderThread(FGPUSkinCache* G
 
 #if RHI_RAYTRACING
 	bRequireRecreatingRayTracingGeometry = (DynamicData == nullptr || // Newly created
+		RayTracingGeometry.Initializer.PositionVertexBuffer == nullptr ||
 		(DynamicData != nullptr && DynamicData->LODIndex != InDynamicData->LODIndex)); // LOD level changed
 #endif
 
