Copying //UE4/Dev-Editor to Dev-Main (//UE4/Dev-Main) for 4.24

diff --git a/Engine/Source/Editor/UnrealEd/Private/FbxMeshUtils.cpp b/Engine/Source/Editor/UnrealEd/Private/FbxMeshUtils.cpp
--- a/Engine/Source/Editor/UnrealEd/Private/FbxMeshUtils.cpp
+++ b/Engine/Source/Editor/UnrealEd/Private/FbxMeshUtils.cpp
@@ -35,6 +35,7 @@
 
 #include "Assets/ClothingAsset.h"
 #include "SkinWeightsUtilities.h"
+#include "LODUtilities.h"
 
 DEFINE_LOG_CATEGORY_STATIC(LogExportMeshUtils, Log, All);
 
@@ -280,10 +281,13 @@ namespace FbxMeshUtils
 			return false;
 		}
 
+		FScopedSkeletalMeshPostEditChange ScopePostEditChange(SelectedSkelMesh);
+
 		//If the imported LOD already exist, we will need to reimport all the skin weight profiles
 		bool bMustReimportAlternateSkinWeightProfile = false;
 
 		// Get a list of all the clothing assets affecting this LOD so we can re-apply later
+		TArray<ClothingAssetUtils::FClothingAssetMeshBinding> ClothingBindings;
 		TArray<UClothingAssetBase*> ClothingAssetsInUse;
 		TArray<int32> ClothingAssetSectionIndices;
 		TArray<int32> ClothingAssetInternalLodIndices;
@@ -292,47 +296,16 @@ namespace FbxMeshUtils
 		if(ImportedResource && ImportedResource->LODModels.IsValidIndex(LODLevel))
 		{
 			bMustReimportAlternateSkinWeightProfile = true;
-			FSkeletalMeshLODModel& LodModel = ImportedResource->LODModels[LODLevel];
-
-			const int32 NumSections = LodModel.Sections.Num();
-
-			for(int32 SectionIndex = 0; SectionIndex < NumSections; ++SectionIndex)
-			{
-				FSkelMeshSection& Section = LodModel.Sections[SectionIndex];
-
-				if(Section.HasClothingData())
-				{
-					UClothingAssetBase* AssetInUse = SelectedSkelMesh->GetSectionClothingAsset(LODLevel, SectionIndex);
-					ClothingAssetsInUse.Add(AssetInUse);
-					ClothingAssetSectionIndices.Add(SectionIndex);
-					ClothingAssetInternalLodIndices.Add(Section.ClothingData.AssetLodIndex);
-				}
-			}
-		}
-
-		// Remove our clothing assets while we import this LOD
-		for(UClothingAssetBase* ClothingAsset : ClothingAssetsInUse)
-		{
-			ClothingAsset->UnbindFromSkeletalMesh(SelectedSkelMesh, LODLevel);
+			FLODUtilities::UnbindClothingAndBackup(SelectedSkelMesh, ClothingBindings, LODLevel);
 		}
 
 		//Lambda to call to re-apply the clothing
-		auto ReapplyClothing = [&SelectedSkelMesh, &ClothingAssetsInUse, &ClothingAssetSectionIndices, &ClothingAssetInternalLodIndices, &ImportedResource, &LODLevel]()
+		auto ReapplyClothing = [&SelectedSkelMesh, &ClothingBindings, &ImportedResource, &LODLevel]()
 		{
-			// Re-apply our clothing assets
-			int32 NumClothingAssetsToApply = ClothingAssetsInUse.Num();
 			if (ImportedResource && ImportedResource->LODModels.IsValidIndex(LODLevel))
 			{
-				FSkeletalMeshLODModel& LodModel = ImportedResource->LODModels[LODLevel];
-				for (int32 AssetIndex = 0; AssetIndex < NumClothingAssetsToApply; ++AssetIndex)
-				{
-					// Only if the same section exists
-					if (LodModel.Sections.IsValidIndex(ClothingAssetSectionIndices[AssetIndex]))
-					{
-						UClothingAssetBase* AssetToApply = ClothingAssetsInUse[AssetIndex];
-						AssetToApply->BindToSkeletalMesh(SelectedSkelMesh, LODLevel, ClothingAssetSectionIndices[AssetIndex], ClothingAssetInternalLodIndices[AssetIndex]);
-					}
-				}
+				// Re-apply our clothing assets
+				FLODUtilities::RestoreClothingFromBackup(SelectedSkelMesh, ClothingBindings, LODLevel);
 			}
 		};
 
@@ -395,6 +368,24 @@ namespace FbxMeshUtils
 			TArray<FString> LODStrings;
 			TArray<FbxNode*>* MeshObject = NULL;;
 
+			//Set the build options if the BuildDat is not available so it is the same option we use to import the LOD
+			if (ImportedResource && ImportedResource->LODModels.IsValidIndex(LODLevel) && !ImportedResource->LODModels[LODLevel].RawSkeletalMeshBulkData.IsBuildDataAvailable())
+			{
+				FSkeletalMeshLODInfo* LODInfo = SelectedSkelMesh->GetLODInfo(LODLevel);
+				if (LODInfo)
+				{
+					LODInfo->BuildSettings.bBuildAdjacencyBuffer = true;
+					LODInfo->BuildSettings.bRecomputeNormals = !ImportOptions->ShouldImportNormals();
+					LODInfo->BuildSettings.bRecomputeTangents = !ImportOptions->ShouldImportTangents();
+					LODInfo->BuildSettings.bUseMikkTSpace = (ImportOptions->NormalGenerationMethod == EFBXNormalGenerationMethod::MikkTSpace) && (!ImportOptions->ShouldImportNormals() || !ImportOptions->ShouldImportTangents());
+					LODInfo->BuildSettings.bComputeWeightedNormals = ImportOptions->bComputeWeightedNormals;
+					LODInfo->BuildSettings.bRemoveDegenerates = ImportOptions->bRemoveDegenerates;
+					LODInfo->BuildSettings.ThresholdPosition = ImportOptions->OverlappingThresholds.ThresholdPosition;
+					LODInfo->BuildSettings.ThresholdTangentNormal = ImportOptions->OverlappingThresholds.ThresholdTangentNormal;
+					LODInfo->BuildSettings.ThresholdUV = ImportOptions->OverlappingThresholds.ThresholdUV;
+				}
+			}
+
 			// Populate the mesh array
 			FFbxImporter->FillFbxSkelMeshArrayInScene(FFbxImporter->Scene->GetRootNode(), MeshArray, false, ImportOptions->bImportAsSkeletalGeometry || ImportOptions->bImportAsSkeletalSkinning, ImportOptions->bImportScene);
 
@@ -524,12 +515,8 @@ namespace FbxMeshUtils
 
 				TempSkelMesh = (USkeletalMesh*)FFbxImporter->ImportSkeletalMesh( ImportSkeletalMeshArgs );
 				// Add the new imported LOD to the existing model (check skeleton compatibility)
-				if( TempSkelMesh  && FFbxImporter->ImportSkeletalMeshLOD(TempSkelMesh, SelectedSkelMesh, SelectedLOD, true, nullptr, TempAssetImportData))
+				if( TempSkelMesh  && FFbxImporter->ImportSkeletalMeshLOD(TempSkelMesh, SelectedSkelMesh, SelectedLOD, TempAssetImportData))
 				{
-					TComponentReregisterContext<USkinnedMeshComponent> ReregisterContext;
-					SelectedSkelMesh->ReleaseResources();
-					SelectedSkelMesh->ReleaseResourcesFence.Wait();
-					
 					//Update the import data for this lod
 					UnFbx::FFbxImporter::UpdateSkeletalMeshImportData(SelectedSkelMesh, nullptr, SelectedLOD, &ImportMaterialOriginalNameData, &ImportMeshLodData);
 
@@ -555,13 +542,9 @@ namespace FbxMeshUtils
 					{
 						//We cannot use anymore the FFbxImporter after the cleanup
 						CleanUpScene();
-						FSkinWeightsUtilities::ReimportAlternateSkinWeight(SelectedSkelMesh, SelectedLOD, false);
+						FSkinWeightsUtilities::ReimportAlternateSkinWeight(SelectedSkelMesh, SelectedLOD);
 					}
 
-					SelectedSkelMesh->PostEditChange();
-					SelectedSkelMesh->InitResources();
-					SelectedSkelMesh->MarkPackageDirty();
-
 					// Notification of success
 					FNotificationInfo NotificationInfo(FText::GetEmpty());
 					NotificationInfo.Text = FText::Format(NSLOCTEXT("UnrealEd", "LODImportSuccessful", "Mesh for LOD {0} imported successfully!"), FText::AsNumber(SelectedLOD));
@@ -633,7 +616,7 @@ namespace FbxMeshUtils
 		return ChosenFilname;
 	}
 
-	bool ImportMeshLODDialog(class UObject* SelectedMesh, int32 LODLevel)
+	bool ImportMeshLODDialog(class UObject* SelectedMesh, int32 LODLevel, bool bNotifyCB /*= true*/)
 	{
 		if(!SelectedMesh)
 		{
@@ -732,7 +715,7 @@ namespace FbxMeshUtils
 			FMessageDialog::Open(EAppMsgType::Ok, FText::Format(LOCTEXT("LODImport_Failure", "Failed to import LOD{0}"), FText::AsNumber(LODLevel)));
 		}
 
-		if (bImportSuccess)
+		if (bImportSuccess && bNotifyCB)
 		{
 			if (SkeletalMesh)
 			{
diff --git a/Engine/Source/Editor/UnrealEd/Private/GeomFitUtils.cpp b/Engine/Source/Editor/UnrealEd/Private/GeomFitUtils.cpp
--- a/Engine/Source/Editor/UnrealEd/Private/GeomFitUtils.cpp
+++ b/Engine/Source/Editor/UnrealEd/Private/GeomFitUtils.cpp
@@ -20,8 +20,7 @@
 #include "PhysicsEngine/SphylElem.h"
 #include "PhysicsEngine/BodySetup.h"
 #include "Engine/StaticMesh.h"
-#include "MeshDescription.h"
-#include "MeshAttributes.h"
+#include "StaticMeshAttributes.h"
 #include "Settings/EditorExperimentalSettings.h"
 
 #define LOCAL_EPS (0.01f)
@@ -190,12 +189,13 @@ int32 GenerateKDopAsSimpleCollision(UStaticMesh* StaticMesh, const TArray<FVecto
 
 /* ******************************** BOX ******************************** */
 
-static void CalcBoundingBox(const FMeshDescription& MeshDescription, FVector& Center, FVector& Extents, FVector& LimitVec)
+static void CalcBoundingBox(const FMeshDescription* MeshDescription, FVector& Center, FVector& Extents, FVector& LimitVec)
 {
-	TVertexAttributesConstRef<FVector> VertexPositions = MeshDescription.VertexAttributes().GetAttributesRef<FVector>(MeshAttribute::Vertex::Position);
+	FStaticMeshConstAttributes Attributes(*MeshDescription);
+	TVertexAttributesConstRef<FVector> VertexPositions = Attributes.GetVertexPositions();
 
 	FBox Box(ForceInit);
-	for (const FVertexID VertexID : MeshDescription.Vertices().GetElementIDs())
+	for (const FVertexID VertexID : MeshDescription->Vertices().GetElementIDs())
 	{
 		Box += VertexPositions[VertexID] * LimitVec;
 	}
@@ -207,8 +207,9 @@ void ComputeBoundingBox(UStaticMesh* StaticMesh, FVector& Center, FVector& Exten
 	// Calculate bounding Box.
 	
 	FMeshDescription* MeshDescription = StaticMesh->GetMeshDescription(0);
+	check(MeshDescription);
 	FVector unitVec = FVector(1.f);
-	CalcBoundingBox(*MeshDescription, Center, Extents, unitVec);
+	CalcBoundingBox(MeshDescription, Center, Extents, unitVec);
 }
 
 int32 GenerateBoxAsSimpleCollision(UStaticMesh* StaticMesh)
@@ -224,7 +225,8 @@ int32 GenerateBoxAsSimpleCollision(UStaticMesh* StaticMesh)
 	FVector unitVec = bs->BuildScale3D;
 	FVector Center, Extents;
 	FMeshDescription* MeshDescription = StaticMesh->GetMeshDescription(0);
-	CalcBoundingBox(*MeshDescription, Center, Extents, unitVec);
+	check(MeshDescription);
+	CalcBoundingBox(MeshDescription, Center, Extents, unitVec);
 	bs->Modify();
 
 	// Create new GUID
@@ -255,19 +257,21 @@ int32 GenerateBoxAsSimpleCollision(UStaticMesh* StaticMesh)
 
 // This algorithm taken from Ritter, 1990
 // This one seems to do well with asymmetric input.
-static void CalcBoundingSphere(const FMeshDescription& MeshDescription, FSphere& sphere, FVector& LimitVec)
+static void CalcBoundingSphere(const FMeshDescription* MeshDescription, FSphere& sphere, FVector& LimitVec)
 {
-	if (MeshDescription.Vertices().Num() == 0)
+	if (MeshDescription->Vertices().Num() == 0)
 		return;
 
 	FBox Box;
 	FVector MinIx[3];
 	FVector MaxIx[3];
 
-	TVertexAttributesConstRef<FVector> VertexPositions = MeshDescription.VertexAttributes().GetAttributesRef<FVector>(MeshAttribute::Vertex::Position);
+	FStaticMeshConstAttributes Attributes(*MeshDescription);
+
+	TVertexAttributesConstRef<FVector> VertexPositions = Attributes.GetVertexPositions();
 
 	bool bFirstVertex = true;
-	for (const FVertexID VertexID : MeshDescription.Vertices().GetElementIDs())
+	for (const FVertexID VertexID : MeshDescription->Vertices().GetElementIDs())
 	{
 		FVector p = VertexPositions[VertexID] * LimitVec;
 		if (bFirstVertex)
@@ -348,7 +352,7 @@ static void CalcBoundingSphere(const FMeshDescription& MeshDescription, FSphere&
 	float r2 = FMath::Square(r);
 
 	// Now check each point lies within this sphere. If not - expand it a bit.
-	for (const FVertexID VertexID : MeshDescription.Vertices().GetElementIDs())
+	for (const FVertexID VertexID : MeshDescription->Vertices().GetElementIDs())
 	{
 		const FVector cToP = (VertexPositions[VertexID] * LimitVec) - sphere.Center;
 
@@ -371,7 +375,7 @@ static void CalcBoundingSphere(const FMeshDescription& MeshDescription, FSphere&
 
 // This is the one thats already used by unreal.
 // Seems to do better with more symmetric input...
-static void CalcBoundingSphere2(const FMeshDescription& MeshDescription, FSphere& sphere, FVector& LimitVec)
+static void CalcBoundingSphere2(const FMeshDescription* MeshDescription, FSphere& sphere, FVector& LimitVec)
 {
 	FVector Center, Extents;
 	CalcBoundingBox(MeshDescription, Center, Extents, LimitVec);
@@ -379,9 +383,10 @@ static void CalcBoundingSphere2(const FMeshDescription& MeshDescription, FSphere
 	sphere.Center = Center;
 	sphere.W = 0.0f;
 
-	TVertexAttributesConstRef<FVector> VertexPositions = MeshDescription.VertexAttributes().GetAttributesRef<FVector>(MeshAttribute::Vertex::Position);
+	FStaticMeshConstAttributes Attributes(*MeshDescription);
+	TVertexAttributesConstRef<FVector> VertexPositions = Attributes.GetVertexPositions();
 
-	for (const FVertexID VertexID : MeshDescription.Vertices().GetElementIDs())
+	for (const FVertexID VertexID : MeshDescription->Vertices().GetElementIDs())
 	{
 		float Dist = FVector::DistSquared(VertexPositions[VertexID] * LimitVec, sphere.Center);
 		if (Dist > sphere.W)
@@ -405,8 +410,9 @@ int32 GenerateSphereAsSimpleCollision(UStaticMesh* StaticMesh)
 
 	// Calculate bounding sphere.
 	FMeshDescription* MeshDescription = StaticMesh->GetMeshDescription(0);
-	CalcBoundingSphere(*MeshDescription, bSphere, unitVec);
-	CalcBoundingSphere2(*MeshDescription, bSphere2, unitVec);
+	check(MeshDescription);
+	CalcBoundingSphere(MeshDescription, bSphere, unitVec);
+	CalcBoundingSphere2(MeshDescription, bSphere2, unitVec);
 
 	if(bSphere.W < bSphere2.W)
 		bestSphere = bSphere;
@@ -442,9 +448,9 @@ int32 GenerateSphereAsSimpleCollision(UStaticMesh* StaticMesh)
 
 /* ******************************** SPHYL ******************************** */
 
-static void CalcBoundingSphyl(const FMeshDescription& MeshDescription, FSphere& sphere, float& length, FRotator& rotation, FVector& LimitVec)
+static void CalcBoundingSphyl(const FMeshDescription* MeshDescription, FSphere& sphere, float& length, FRotator& rotation, FVector& LimitVec)
 {
-	if (MeshDescription.Vertices().Num() == 0)
+	if (MeshDescription->Vertices().Num() == 0)
 		return;
 
 	FVector Center, Extents;
@@ -475,10 +481,11 @@ static void CalcBoundingSphyl(const FMeshDescription& MeshDescription, FSphere&
 	float r = Extents.GetMax();
 	float r2 = FMath::Square(r);
 	
-	TVertexAttributesConstRef<FVector> VertexPositions = MeshDescription.VertexAttributes().GetAttributesRef<FVector>(MeshAttribute::Vertex::Position);
+	FStaticMeshConstAttributes Attributes(*MeshDescription);
+	TVertexAttributesConstRef<FVector> VertexPositions = Attributes.GetVertexPositions();
 
 	// Now check each point lies within this the radius. If not - expand it a bit.
-	for (const FVertexID VertexID : MeshDescription.Vertices().GetElementIDs())
+	for (const FVertexID VertexID : MeshDescription->Vertices().GetElementIDs())
 	{
 		FVector cToP = (VertexPositions[VertexID] * LimitVec) - sphere.Center;
 		cToP = rotation.UnrotateVector(cToP);
@@ -499,7 +506,7 @@ static void CalcBoundingSphyl(const FMeshDescription& MeshDescription, FSphere&
 	float hl = FMath::Max(0.0f, Extent - r);
 
 	// Now check each point lies within the length. If not - expand it a bit.
-	for (const FVertexID VertexID : MeshDescription.Vertices().GetElementIDs())
+	for (const FVertexID VertexID : MeshDescription->Vertices().GetElementIDs())
 	{
 		FVector cToP = (VertexPositions[VertexID] * LimitVec) - sphere.Center;
 		cToP = rotation.UnrotateVector(cToP);
@@ -546,7 +553,8 @@ int32 GenerateSphylAsSimpleCollision(UStaticMesh* StaticMesh)
 
 	// Calculate bounding box.
 	FMeshDescription* MeshDescription = StaticMesh->GetMeshDescription(0);
-	CalcBoundingSphyl(*MeshDescription, sphere, length, rotation, unitVec);
+	check(MeshDescription);
+	CalcBoundingSphyl(MeshDescription, sphere, length, rotation, unitVec);
 
 	// Dont use if radius is zero.
 	if (sphere.W <= 0.f)
diff --git a/Engine/Source/Editor/UnrealEd/Private/SkinWeightsUtilities.cpp b/Engine/Source/Editor/UnrealEd/Private/SkinWeightsUtilities.cpp
--- a/Engine/Source/Editor/UnrealEd/Private/SkinWeightsUtilities.cpp
+++ b/Engine/Source/Editor/UnrealEd/Private/SkinWeightsUtilities.cpp
@@ -34,7 +34,7 @@
 
 DEFINE_LOG_CATEGORY_STATIC(LogSkinWeightsUtilities, Log, All);
 
-bool FSkinWeightsUtilities::ImportAlternateSkinWeight(USkeletalMesh* SkeletalMesh, FString Path, int32 TargetLODIndex, const FName& ProfileName, bool bReregisterComponent)
+bool FSkinWeightsUtilities::ImportAlternateSkinWeight(USkeletalMesh* SkeletalMesh, FString Path, int32 TargetLODIndex, const FName& ProfileName)
 {
 	check(SkeletalMesh);
 	check(SkeletalMesh->GetLODInfo(TargetLODIndex));
@@ -53,6 +53,8 @@ bool FSkinWeightsUtilities::ImportAlternateSkinWeight(USkeletalMesh* SkeletalMes
 		UE_LOG(LogSkinWeightsUtilities, Error, TEXT("Path containing Skin Weight Profile data does not exist (%s)."), *Path);
 		return false;
 	}
+	FScopedSkeletalMeshPostEditChange ScopePostEditChange(SkeletalMesh);
+
 	UnFbx::FBXImportOptions ImportOptions;
 	//Import the alternate fbx into a temporary skeletal mesh using the same import options
 	UFbxFactory* FbxFactory = NewObject<UFbxFactory>(UFbxFactory::StaticClass());
@@ -200,21 +202,7 @@ bool FSkinWeightsUtilities::ImportAlternateSkinWeight(USkeletalMesh* SkeletalMes
 				bool bUseMikkTSpace = ImportOptions.NormalGenerationMethod == EFBXNormalGenerationMethod::MikkTSpace;
 
 				TArray<FRawSkinWeight>& SkinWeights = ProfileData.SkinWeights;
-				if (bReregisterComponent)
-				{
-					TComponentReregisterContext<USkinnedMeshComponent> ReregisterContext;
-					SkeletalMesh->ReleaseResources();
-					SkeletalMesh->ReleaseResourcesFence.Wait();
-
-					bResult = FLODUtilities::UpdateAlternateSkinWeights(SkeletalMesh, ProfileName, TmpSkeletalMesh, TargetLODIndex, SrcLodIndex, OverlappingThresholds, ShouldImportNormals, ShouldImportTangents, bUseMikkTSpace, ImportOptions.bComputeWeightedNormals);
-
-					SkeletalMesh->PostEditChange();
-					SkeletalMesh->InitResources();
-				}
-				else
-				{
-					bResult = FLODUtilities::UpdateAlternateSkinWeights(SkeletalMesh, ProfileName, TmpSkeletalMesh, TargetLODIndex, SrcLodIndex, OverlappingThresholds, ShouldImportNormals, ShouldImportTangents, bUseMikkTSpace, ImportOptions.bComputeWeightedNormals);
-				}
+				bResult = FLODUtilities::UpdateAlternateSkinWeights(SkeletalMesh, ProfileName, TmpSkeletalMesh, TargetLODIndex, SrcLodIndex, OverlappingThresholds, ShouldImportNormals, ShouldImportTangents, bUseMikkTSpace, ImportOptions.bComputeWeightedNormals);
 				
 				if (!bResult)
 				{
@@ -252,57 +240,45 @@ bool FSkinWeightsUtilities::ImportAlternateSkinWeight(USkeletalMesh* SkeletalMes
 	return bResult;
 }
 
-bool FSkinWeightsUtilities::ReimportAlternateSkinWeight(USkeletalMesh* SkeletalMesh, int32 TargetLODIndex, bool bReregisterComponent)
+bool FSkinWeightsUtilities::ReimportAlternateSkinWeight(USkeletalMesh* SkeletalMesh, int32 TargetLODIndex)
 {
 	bool bResult = false;
 
-	//Bulk work of the function, we use a lambda because of the re-register component option.
-	auto DoWork = [&SkeletalMesh, &TargetLODIndex, &bResult]()
+	const TArray<FSkinWeightProfileInfo>& SkinWeightProfiles = SkeletalMesh->GetSkinWeightProfiles();
+	if (SkinWeightProfiles.Num() <= 0)
 	{
-		const TArray<FSkinWeightProfileInfo>& SkinWeightProfiles = SkeletalMesh->GetSkinWeightProfiles();
-		for (int32 ProfileIndex = 0; ProfileIndex < SkinWeightProfiles.Num(); ++ProfileIndex)
-		{
-			const FSkinWeightProfileInfo& ProfileInfo = SkinWeightProfiles[ProfileIndex];
+		return bResult;
+	}
 
-			const FString* PathNamePtr = ProfileInfo.PerLODSourceFiles.Find(TargetLODIndex);
-			//Skip profile that do not have data for TargetLODIndex
-			if (!PathNamePtr)
-			{
-				continue;
-			}
+	FScopedSkeletalMeshPostEditChange ScopePostEditChange(SkeletalMesh);
 
-			const FString& PathName = *PathNamePtr;
+	for (int32 ProfileIndex = 0; ProfileIndex < SkinWeightProfiles.Num(); ++ProfileIndex)
+	{
+		const FSkinWeightProfileInfo& ProfileInfo = SkinWeightProfiles[ProfileIndex];
 
-			if (FPaths::FileExists(PathName))
-			{
-				bResult |= FSkinWeightsUtilities::ImportAlternateSkinWeight(SkeletalMesh, PathName, TargetLODIndex, ProfileInfo.Name, false);
-			}
-			else
-			{
-				const FString PickedFileName = FSkinWeightsUtilities::PickSkinWeightFBXPath(TargetLODIndex);
-				if (!PickedFileName.IsEmpty() && FPaths::FileExists(PickedFileName))
-				{
-					bResult |= FSkinWeightsUtilities::ImportAlternateSkinWeight(SkeletalMesh, PickedFileName, TargetLODIndex, ProfileInfo.Name, false);
-				}
-			}
+		const FString* PathNamePtr = ProfileInfo.PerLODSourceFiles.Find(TargetLODIndex);
+		//Skip profile that do not have data for TargetLODIndex
+		if (!PathNamePtr)
+		{
+			continue;
 		}
-	};
 
-	if (bReregisterComponent)
-	{
-		TComponentReregisterContext<USkinnedMeshComponent> ReregisterContext;
-		SkeletalMesh->ReleaseResources();
-		SkeletalMesh->ReleaseResourcesFence.Wait();
-
-		DoWork();
+		const FString& PathName = *PathNamePtr;
 
-		SkeletalMesh->PostEditChange();
-		SkeletalMesh->InitResources();
-	}
-	else
-	{
-		DoWork();
+		if (FPaths::FileExists(PathName))
+		{
+			bResult |= FSkinWeightsUtilities::ImportAlternateSkinWeight(SkeletalMesh, PathName, TargetLODIndex, ProfileInfo.Name);
+		}
+		else
+		{
+			const FString PickedFileName = FSkinWeightsUtilities::PickSkinWeightFBXPath(TargetLODIndex);
+			if (!PickedFileName.IsEmpty() && FPaths::FileExists(PickedFileName))
+			{
+				bResult |= FSkinWeightsUtilities::ImportAlternateSkinWeight(SkeletalMesh, PickedFileName, TargetLODIndex, ProfileInfo.Name);
+			}
+		}
 	}
+
 	
 	if (bResult)
 	{
diff --git a/Engine/Source/Editor/UnrealEd/Private/StaticMeshEdit.cpp b/Engine/Source/Editor/UnrealEd/Private/StaticMeshEdit.cpp
--- a/Engine/Source/Editor/UnrealEd/Private/StaticMeshEdit.cpp
+++ b/Engine/Source/Editor/UnrealEd/Private/StaticMeshEdit.cpp
@@ -13,6 +13,7 @@
 #include "EditorFramework/ThumbnailInfo.h"
 #include "Engine/MeshMerging.h"
 #include "Engine/StaticMesh.h"
+#include "StaticMeshAttributes.h"
 #include "Engine/StaticMeshSocket.h"
 #include "Engine/Polys.h"
 #include "Editor.h"
@@ -31,8 +32,6 @@
 #include "Interfaces/ITargetPlatformManagerModule.h"
 #include "Interfaces/ITargetPlatform.h"
 #include "Settings/EditorExperimentalSettings.h"
-#include "MeshDescription.h"
-#include "MeshAttributes.h"
 
 #include "Modules/ModuleManager.h"
 #include "IMeshReductionManagerModule.h"
@@ -595,9 +594,8 @@ UStaticMesh* CreateStaticMesh(FMeshDescription& RawMesh,TArray<FStaticMaterial>&
 	auto StaticMesh = NewObject<UStaticMesh>(InOuter, InName, RF_Public | RF_Standalone);
 
 	// Add one LOD for the base mesh
-	FStaticMeshSourceModel& SrcModel = StaticMesh->AddSourceModel();
-	FMeshDescription* MeshDescription = StaticMesh->CreateMeshDescription(0);
-	*MeshDescription = RawMesh;
+	StaticMesh->AddSourceModel();
+	FMeshDescription* MeshDescription = StaticMesh->CreateMeshDescription(0, RawMesh);
 	StaticMesh->CommitMeshDescription(0);
 	StaticMesh->StaticMaterials = Materials;
 
@@ -805,10 +803,9 @@ UStaticMesh* CreateStaticMeshFromBrush(UObject* Outer, FName Name, ABrush* Brush
 	UStaticMesh* StaticMesh = NewObject<UStaticMesh>(Outer, Name, RF_Public | RF_Standalone);
 
 	// Add one LOD for the base mesh
-	FStaticMeshSourceModel& SrcModel = StaticMesh->AddSourceModel();
+	StaticMesh->AddSourceModel();
 	const int32 LodIndex = StaticMesh->GetNumSourceModels() - 1;
 	FMeshDescription* MeshDescription = StaticMesh->CreateMeshDescription(LodIndex);
-	UStaticMesh::RegisterMeshAttributes(*MeshDescription);
 
 	// Fill out the mesh description and materials from the brush geometry
 	TArray<FStaticMaterial> Materials;
@@ -1490,9 +1487,7 @@ void RestoreExistingMeshData(ExistingStaticMeshData* ExistingMeshDataPtr, UStati
 		FStaticMeshSourceModel& SrcModel = NewMesh->AddSourceModel();
 		if (ExistingMeshDataPtr->ExistingLODData[i].ExistingMeshDescription.IsValid())
 		{
-			FMeshDescription* MeshDescription = NewMesh->CreateMeshDescription(i);
-			*MeshDescription = MoveTemp(*ExistingMeshDataPtr->ExistingLODData[i].ExistingMeshDescription);
-			ExistingMeshDataPtr->ExistingLODData[i].ExistingMeshDescription.Reset();
+			FMeshDescription* MeshDescription = NewMesh->CreateMeshDescription(i, MoveTemp(*ExistingMeshDataPtr->ExistingLODData[i].ExistingMeshDescription));
 			NewMesh->CommitMeshDescription(i);
 		}
 		SrcModel.BuildSettings = ExistingMeshDataPtr->ExistingLODData[i].ExistingBuildSettings;
diff --git a/Engine/Source/Editor/UnrealEd/Public/FbxImporter.h b/Engine/Source/Editor/UnrealEd/Public/FbxImporter.h
--- a/Engine/Source/Editor/UnrealEd/Public/FbxImporter.h
+++ b/Engine/Source/Editor/UnrealEd/Public/FbxImporter.h
@@ -756,9 +756,9 @@ public:
 	void AddStaticMeshSourceModelGeneratedLOD(UStaticMesh* StaticMesh, int32 LODIndex);
 
 	/**
-	* Return the node that match the staticmesh name. Return nullptr in case there is no match
+	* Return the node that match the mesh name. Return nullptr in case there is no match
 	*/
-	FbxNode* GetMeshNodesFromName(UStaticMesh* StaticMesh, TArray<FbxNode*>& FbxMeshArray);
+	FbxNode* GetMeshNodesFromName(const FString& ReimportMeshName, TArray<FbxNode*>& FbxMeshArray);
 
 	/**
 	 * re-import Unreal static mesh from updated Fbx file
@@ -917,11 +917,9 @@ public:
 	 * @param InSkeletalMesh - LOD mesh object
 	 * @param BaseSkeletalMesh - base mesh object
 	 * @param DesiredLOD - LOD level
-	 * @param bNeedToReregister - if true, re-register this skeletal mesh to shut down the skeletal mesh component that is previewing this mesh. 
-									But you can set this to false when in the first loading before rendering this mesh for a performance issue 
 	   @param ReregisterAssociatedComponents - if NULL, just re-registers all SkinnedMeshComponents but if you set the specific components, will only re-registers those components
 	 */
-	UNREALED_API bool ImportSkeletalMeshLOD(USkeletalMesh* InSkeletalMesh, USkeletalMesh* BaseSkeletalMesh, int32 DesiredLOD, bool bNeedToReregister = true, TArray<UActorComponent*>* ReregisterAssociatedComponents = NULL, UFbxSkeletalMeshImportData* TemplateImportData = nullptr);
+	UNREALED_API bool ImportSkeletalMeshLOD(USkeletalMesh* InSkeletalMesh, USkeletalMesh* BaseSkeletalMesh, int32 DesiredLOD, UFbxSkeletalMeshImportData* TemplateImportData = nullptr);
 
 	/**
 	 * Empties the FBX scene, releasing its memory.
diff --git a/Engine/Source/Editor/UnrealEd/Public/FbxMeshUtils.h b/Engine/Source/Editor/UnrealEd/Public/FbxMeshUtils.h
--- a/Engine/Source/Editor/UnrealEd/Public/FbxMeshUtils.h
+++ b/Engine/Source/Editor/UnrealEd/Public/FbxMeshUtils.h
@@ -39,7 +39,7 @@ namespace FbxMeshUtils
 	 * @param LODLevel			The level of the lod to import
 	 * @return					Whether or not the import succeeded
 	 */
-	UNREALED_API bool ImportMeshLODDialog( UObject* Mesh, int32 LODLevel);
+	UNREALED_API bool ImportMeshLODDialog( UObject* Mesh, int32 LODLevel, bool bNotifyCB = true);
 
 	/**
 	 * Sets import option before importing
diff --git a/Engine/Source/Editor/UnrealEd/Public/SkelImport.h b/Engine/Source/Editor/UnrealEd/Public/SkelImport.h
--- a/Engine/Source/Editor/UnrealEd/Public/SkelImport.h
+++ b/Engine/Source/Editor/UnrealEd/Public/SkelImport.h
@@ -22,18 +22,20 @@ class FSkeletalMeshLODModel;
 
 struct ExistingMeshLodSectionData
 {
-	ExistingMeshLodSectionData(FName InImportedMaterialSlotName, bool InbCastShadow, bool InbRecomputeTangents, int32 InGenerateUpTo, bool InbDisabled)
+	ExistingMeshLodSectionData(FName InImportedMaterialSlotName, bool InbCastShadow, bool InbRecomputeTangents, int32 InGenerateUpTo, bool InbDisabled, bool InbBoneChunkedSection)
 	: ImportedMaterialSlotName(InImportedMaterialSlotName)
 	, bCastShadow(InbCastShadow)
 	, bRecomputeTangents(InbRecomputeTangents)
 	, GenerateUpTo(InGenerateUpTo)
 	, bDisabled(InbDisabled)
+	, bBoneChunkedSection(InbBoneChunkedSection)
 	{}
 	FName ImportedMaterialSlotName;
 	bool bCastShadow;
 	bool bRecomputeTangents;
 	int32 GenerateUpTo;
 	bool bDisabled;
+	bool bBoneChunkedSection;
 };
 
 struct ExistingSkelMeshData
@@ -41,6 +43,7 @@ struct ExistingSkelMeshData
 	TArray<USkeletalMeshSocket*>			ExistingSockets;
 	TArray<FReductionBaseSkeletalMeshBulkData*> ExistingOriginalReductionSourceMeshData;
 	TIndirectArray<FSkeletalMeshLODModel>	ExistingLODModels;
+	FSkeletalMeshLODInfo					ExistingBaseLODInfo;
 	TArray<FSkeletalMeshLODInfo>			ExistingLODInfo;
 	FReferenceSkeleton						ExistingRefSkeleton;
 	TArray<FSkeletalMaterial>				ExistingMaterials;
diff --git a/Engine/Source/Editor/UnrealEd/Public/SkinWeightsUtilities.h b/Engine/Source/Editor/UnrealEd/Public/SkinWeightsUtilities.h
--- a/Engine/Source/Editor/UnrealEd/Public/SkinWeightsUtilities.h
+++ b/Engine/Source/Editor/UnrealEd/Public/SkinWeightsUtilities.h
@@ -27,7 +27,7 @@ public:
 	* @Param ProfileName - The name of the profile to associate the imported skin weight
 	* @param bReregisterComponent - if true the component using the skeletal mesh will all be re register.
 	*/
-	static bool ImportAlternateSkinWeight(USkeletalMesh* SkeletalMesh, FString Path, int32 TargetLODIndex, const FName& ProfileName, bool bReregisterComponent = true);
+	static bool ImportAlternateSkinWeight(USkeletalMesh* SkeletalMesh, FString Path, int32 TargetLODIndex, const FName& ProfileName);
 	
 	/**
 	* This function reimport all skin weights profile for a specified LOD. Return true if the weights are successfully updates.
@@ -37,7 +37,7 @@ public:
 	* @Param TargetLODIndex - The LODIndex to imp[ort the skin weight
 	* @param bReregisterComponent - if true the component using the skeletal mesh will all be re register.
 	*/
-	static bool ReimportAlternateSkinWeight(USkeletalMesh* SkeletalMesh, int32 TargetLODIndex, bool bReregisterComponent = true);
+	static bool ReimportAlternateSkinWeight(USkeletalMesh* SkeletalMesh, int32 TargetLODIndex);
 
 	static bool RemoveSkinnedWeightProfileData(USkeletalMesh* SkeletalMesh, const FName& ProfileName, int32 LODIndex);
 
diff --git a/Engine/Source/Editor/WorldBrowser/Private/Tiles/WorldTileCollectionModel.cpp b/Engine/Source/Editor/WorldBrowser/Private/Tiles/WorldTileCollectionModel.cpp
--- a/Engine/Source/Editor/WorldBrowser/Private/Tiles/WorldTileCollectionModel.cpp
+++ b/Engine/Source/Editor/WorldBrowser/Private/Tiles/WorldTileCollectionModel.cpp
@@ -23,10 +23,7 @@
 #include "AssetRegistryModule.h"
 #include "MeshUtilities.h"
 
-#include "MeshDescription.h"
-#include "MeshAttributes.h"
-#include "MeshAttributeArray.h"
-
+#include "StaticMeshAttributes.h"
 #include "Materials/Material.h"
 #include "MaterialUtilities.h"
 #include "LandscapeLayerInfoObject.h"
@@ -2158,12 +2155,13 @@ bool FWorldTileCollectionModel::GenerateLODLevels(FLevelModelList InLevelList, i
 				//Set the Imported version before calling the build
 				StaticMesh->ImportVersion = EImportStaticMeshVersion::LastVersion;
 			}
-			FMeshDescription& LandscapeRawMesh = *(StaticMesh->CreateMeshDescription(0));
+			FMeshDescription* LandscapeRawMesh = StaticMesh->CreateMeshDescription(0);
+			FStaticMeshAttributes Attributes(*LandscapeRawMesh);
 		
-			Landscape->ExportToRawMesh(LandscapeLOD, LandscapeRawMesh);
+			Landscape->ExportToRawMesh(LandscapeLOD, *LandscapeRawMesh);
 		
-			TVertexAttributesRef<FVector> VertexPositions = LandscapeRawMesh.VertexAttributes().GetAttributesRef<FVector>(MeshAttribute::Vertex::Position);
-			for (const FVertexID& VertexID : LandscapeRawMesh.Vertices().GetElementIDs())
+			TVertexAttributesRef<FVector> VertexPositions = Attributes.GetVertexPositions();
+			for (const FVertexID& VertexID : LandscapeRawMesh->Vertices().GetElementIDs())
 			{
 				VertexPositions[VertexID] -= LandscapeWorldLocation;
 			}
diff --git a/Engine/Source/Runtime/ClothingSystemRuntimeInterface/Public/ClothingAssetBase.h b/Engine/Source/Runtime/ClothingSystemRuntimeInterface/Public/ClothingAssetBase.h
--- a/Engine/Source/Runtime/ClothingSystemRuntimeInterface/Public/ClothingAssetBase.h
+++ b/Engine/Source/Runtime/ClothingSystemRuntimeInterface/Public/ClothingAssetBase.h
@@ -23,7 +23,7 @@ public:
 	* @param InSubmeshIdx Submesh in this asset to replace section with
 	* @param InAssetLodIndex Internal clothing LOD to use
 	*/
-	virtual bool BindToSkeletalMesh(USkeletalMesh* InSkelMesh, int32 InMeshLodIndex, int32 InSectionIndex, int32 InAssetLodIndex, bool bCallPostEditChange = true)
+	virtual bool BindToSkeletalMesh(USkeletalMesh* InSkelMesh, int32 InMeshLodIndex, int32 InSectionIndex, int32 InAssetLodIndex)
 	PURE_VIRTUAL(UClothingAssetBase::BindToSkeletalMesh, return false;);
 
 	/**
diff --git a/Engine/Source/Runtime/ClothingSystemRuntimeNv/Public/Assets/ClothingAsset.h b/Engine/Source/Runtime/ClothingSystemRuntimeNv/Public/Assets/ClothingAsset.h
--- a/Engine/Source/Runtime/ClothingSystemRuntimeNv/Public/Assets/ClothingAsset.h
+++ b/Engine/Source/Runtime/ClothingSystemRuntimeNv/Public/Assets/ClothingAsset.h
@@ -526,7 +526,7 @@ public:
 	// UClothingAssetBase Interface ////////////////////////////////////////////
 	virtual void RefreshBoneMapping(USkeletalMesh* InSkelMesh) override;
 #if WITH_EDITOR
-	virtual bool BindToSkeletalMesh(USkeletalMesh* InSkelMesh, int32 InMeshLodIndex, int32 InSectionIndex, int32 InAssetLodIndex, bool bCallPostEditChange=true) override;
+	virtual bool BindToSkeletalMesh(USkeletalMesh* InSkelMesh, int32 InMeshLodIndex, int32 InSectionIndex, int32 InAssetLodIndex) override;
 	virtual void UnbindFromSkeletalMesh(USkeletalMesh* InSkelMesh) override;
 	virtual void UnbindFromSkeletalMesh(USkeletalMesh* InSkelMesh, int32 InMeshLodIndex) override;
 	virtual void InvalidateCachedData() override;
diff --git a/Engine/Source/Runtime/Core/Public/UObject/EditorObjectVersion.h b/Engine/Source/Runtime/Core/Public/UObject/EditorObjectVersion.h
--- a/Engine/Source/Runtime/Core/Public/UObject/EditorObjectVersion.h
+++ b/Engine/Source/Runtime/Core/Public/UObject/EditorObjectVersion.h
@@ -83,6 +83,8 @@ struct CORE_API FEditorObjectVersion
 		MeshDescriptionTriangles,
 		//Add weighted area and angle when computing the normals
 		ComputeWeightedNormals,
+		// SkeletalMesh now can be rebuild in editor, no more need to re-import
+		SkeletalMeshBuildRefactor,
 		// -----<new versions can be added above this line>-------------------------------------------------
 		VersionPlusOne,
 		LatestVersion = VersionPlusOne - 1
diff --git a/Engine/Source/Runtime/Engine/Classes/Engine/EngineTypes.h b/Engine/Source/Runtime/Engine/Classes/Engine/EngineTypes.h
--- a/Engine/Source/Runtime/Engine/Classes/Engine/EngineTypes.h
+++ b/Engine/Source/Runtime/Engine/Classes/Engine/EngineTypes.h
@@ -2675,6 +2675,96 @@ struct FMeshBuildSettings
 	}
 };
 
+/**
+ * Settings applied when building a mesh.
+ */
+USTRUCT()
+struct FSkeletalMeshBuildSettings
+{
+	GENERATED_BODY()
+
+	/** If true, normals in the raw mesh are ignored and recomputed. */
+	UPROPERTY(EditAnywhere, Category=BuildSettings)
+	uint8 bRecomputeNormals:1;
+
+	/** If true, tangents in the raw mesh are ignored and recomputed. */
+	UPROPERTY(EditAnywhere, Category=BuildSettings)
+	uint8 bRecomputeTangents:1;
+	
+	/** If true, degenerate triangles will be removed. */
+	UPROPERTY(EditAnywhere, Category=BuildSettings)
+	uint8 bUseMikkTSpace:1;
+	
+	/** If true, we will use the surface area and the corner angle of the triangle as a ratio when computing the normals. */
+	UPROPERTY(EditAnywhere, Category = BuildSettings)
+	uint8 bComputeWeightedNormals : 1;
+
+	/** If true, degenerate triangles will be removed. */
+	UPROPERTY(EditAnywhere, Category=BuildSettings)
+	uint8 bRemoveDegenerates:1;
+	
+	/** If true, Tangents will be stored at 16 bit vs 8 bit precision. */
+	UPROPERTY(EditAnywhere, Category = BuildSettings)
+	uint8 bUseHighPrecisionTangentBasis:1;
+
+	/** If true, UVs will be stored at full floating point precision. */
+	UPROPERTY(EditAnywhere, Category=BuildSettings)
+	uint8 bUseFullPrecisionUVs:1;
+	
+	/** Required for PNT tessellation but can be slow. Recommend disabling for larger meshes. */
+	UPROPERTY(EditAnywhere, Category=BuildSettings)
+	uint8 bBuildAdjacencyBuffer:1;
+
+	/** Threshold use to decide if two vertex position are equal. */
+	UPROPERTY(EditAnywhere, Category = BuildSettings)
+	float ThresholdPosition;
+
+	/** Threshold use to decide if two normal, tangents or bi-normals are equal. */
+	UPROPERTY(EditAnywhere, Category = BuildSettings)
+	float ThresholdTangentNormal;
+
+	/** Threshold use to decide if two UVs are equal. */
+	UPROPERTY(EditAnywhere, Category = BuildSettings)
+	float ThresholdUV;
+
+	/** Default settings. */
+	FSkeletalMeshBuildSettings()
+		: bRecomputeNormals(true)
+		, bRecomputeTangents(true)
+		, bUseMikkTSpace(true)
+		, bComputeWeightedNormals(false)
+		, bRemoveDegenerates(true)
+		, bUseHighPrecisionTangentBasis(false)
+		, bUseFullPrecisionUVs(false)
+		, bBuildAdjacencyBuffer(true)
+		, ThresholdPosition(0.00002)
+		, ThresholdTangentNormal(0.00002)
+		, ThresholdUV(0.0009765625)
+	{}
+
+	/** Equality operator. */
+	bool operator==(const FSkeletalMeshBuildSettings& Other) const
+	{
+		return bRecomputeNormals == Other.bRecomputeNormals
+			&& bRecomputeTangents == Other.bRecomputeTangents
+			&& bUseMikkTSpace == Other.bUseMikkTSpace
+			&& bComputeWeightedNormals == Other.bComputeWeightedNormals
+			&& bRemoveDegenerates == Other.bRemoveDegenerates
+			&& bUseHighPrecisionTangentBasis == Other.bUseHighPrecisionTangentBasis
+			&& bUseFullPrecisionUVs == Other.bUseFullPrecisionUVs
+			&& bBuildAdjacencyBuffer == Other.bBuildAdjacencyBuffer
+			&& ThresholdPosition == Other.ThresholdPosition
+			&& ThresholdTangentNormal == Other.ThresholdTangentNormal
+			&& ThresholdUV == Other.ThresholdUV;
+	}
+
+	/** Inequality. */
+	bool operator!=(const FSkeletalMeshBuildSettings& Other) const
+	{
+		return !(*this == Other);
+	}
+};
+
 /** Event used by AActor::TakeDamage and related functions */
 USTRUCT(BlueprintType)
 struct ENGINE_API FDamageEvent
diff --git a/Engine/Source/Runtime/Engine/Classes/Engine/SkeletalMeshLODSettings.h b/Engine/Source/Runtime/Engine/Classes/Engine/SkeletalMeshLODSettings.h
--- a/Engine/Source/Runtime/Engine/Classes/Engine/SkeletalMeshLODSettings.h
+++ b/Engine/Source/Runtime/Engine/Classes/Engine/SkeletalMeshLODSettings.h
@@ -5,6 +5,7 @@
 #include "CoreMinimal.h"
 #include "UObject/ObjectMacros.h"
 #include "UObject/Object.h"
+#include "Engine/EngineTypes.h"
 #include "SkeletalMeshReductionSettings.h"
 #include "DataAsset.h"
 #include "PerPlatformProperties.h"
@@ -63,6 +64,9 @@ struct FSkeletalMeshLODGroupSettings
 	/** Get Skeletal mesh optimizations setting structure for the given LOD level */
 	ENGINE_API FSkeletalMeshOptimizationSettings GetReductionSettings() const;
 
+	/** Get Skeletal mesh build setting structure for the given LOD level */
+	ENGINE_API FSkeletalMeshBuildSettings GetBuildSettings() const;
+
 	/** Get Skeletal mesh optimizations setting structure for the given LOD level */
 	UE_DEPRECATED(4.20, "Please use GetReductionSettings instead")
 	FSkeletalMeshOptimizationSettings GetSettings() const
@@ -107,6 +111,10 @@ struct FSkeletalMeshLODGroupSettings
 	/** The optimization settings to use for the respective LOD level */
 	UPROPERTY(EditAnywhere, Category = Reduction)
 	FSkeletalMeshOptimizationSettings ReductionSettings;
+
+	/** The build settings to use for the respective LOD level */
+	UPROPERTY(EditAnywhere, Category = Build)
+	FSkeletalMeshBuildSettings BuildSettings;
 };
 
 UCLASS(config = Engine, defaultconfig, BlueprintType, MinimalAPI)
diff --git a/Engine/Source/Runtime/Engine/Private/Components/SceneCaptureComponent.cpp b/Engine/Source/Runtime/Engine/Private/Components/SceneCaptureComponent.cpp
--- a/Engine/Source/Runtime/Engine/Private/Components/SceneCaptureComponent.cpp
+++ b/Engine/Source/Runtime/Engine/Private/Components/SceneCaptureComponent.cpp
@@ -7,6 +7,7 @@
 #include "Components/SceneCaptureComponent.h"
 #include "Misc/ScopeLock.h"
 #include "UObject/RenderingObjectVersion.h"
+#include "UObject/EditorObjectVersion.h"
 #include "UObject/ConstructorHelpers.h"
 #include "GameFramework/Actor.h"
 #include "RenderingThread.h"
@@ -746,15 +747,15 @@ void APlanarReflection::PostLoad()
 	Super::PostLoad();
 
 	if (GetLinkerCustomVersion(FEditorObjectVersion::GUID) < FEditorObjectVersion::ChangeSceneCaptureRootComponent)
-{
-		if (PlanarReflectionComponent)
 	{
+		if (PlanarReflectionComponent)
+		{
 			PRAGMA_DISABLE_DEPRECATION_WARNINGS
 			PlanarReflectionComponent->bShowPreviewPlane = bShowPreviewPlane_DEPRECATED;
 			PRAGMA_ENABLE_DEPRECATION_WARNINGS
 		}
-			}
-		}
+	}
+}
 
 void APlanarReflection::OnInterpToggle(bool bEnable)
 {
diff --git a/Engine/Source/Runtime/Engine/Private/SkeletalMeshLODModel.cpp b/Engine/Source/Runtime/Engine/Private/SkeletalMeshLODModel.cpp
--- a/Engine/Source/Runtime/Engine/Private/SkeletalMeshLODModel.cpp
+++ b/Engine/Source/Runtime/Engine/Private/SkeletalMeshLODModel.cpp
@@ -442,12 +442,43 @@ FArchive& operator<<(FArchive& Ar, FSkelMeshSection& S)
 		{
 			S.GenerateUpToLodIndex = -1;
 		}
+
+		if (Ar.CustomVer(FEditorObjectVersion::GUID) >= FEditorObjectVersion::SkeletalMeshBuildRefactor)
+		{
+			Ar << S.OriginalDataSectionIndex;
+			Ar << S.ChunkedParentSectionIndex;
+		}
+		else if (Ar.IsLoading())
+		{
+			S.OriginalDataSectionIndex = INDEX_NONE;
+			S.ChunkedParentSectionIndex = INDEX_NONE;
+		}
 		return Ar;
 	}
 
 	return Ar;
 }
 
+// Serialization.
+FArchive& operator<<(FArchive& Ar, FSkelMeshSourceSectionUserData& S)
+{
+	FStripDataFlags StripFlags(Ar);
+	// When data is cooked we do not serialize anything
+	//This is for editor only editing
+	if (StripFlags.IsEditorDataStripped())
+	{
+		return Ar;
+	}
+
+	Ar << S.bRecomputeTangent;
+	Ar << S.bCastShadow;
+	Ar << S.bDisabled;
+	Ar << S.GenerateUpToLodIndex;
+	Ar << S.CorrespondClothAssetIndex;
+	Ar << S.ClothingData;
+
+	return Ar;
+}
 
 //////////////////////////////////////////////////////////////////////////
 
@@ -555,12 +586,19 @@ void FSkeletalMeshLODModel::Serialize(FArchive& Ar, UObject* Owner, int32 Idx)
 
 	Ar.UsingCustomVersion(FSkeletalMeshCustomVersion::GUID);
 	Ar.UsingCustomVersion(FFortniteMainBranchObjectVersion::GUID);
+	Ar.UsingCustomVersion(FEditorObjectVersion::GUID);
 
 	if (StripFlags.IsDataStrippedForServer())
 	{
 		TArray<FSkelMeshSection> TempSections;
 		Ar << TempSections;
 
+		if (Ar.CustomVer(FEditorObjectVersion::GUID) >= FEditorObjectVersion::SkeletalMeshBuildRefactor)
+		{
+			TMap<int32, FSkelMeshSourceSectionUserData> TempUserSectionsData;
+			Ar << TempUserSectionsData;
+		}
+
 		// For old content, load as a multi-size container
 		if (Ar.IsLoading() && Ar.CustomVer(FSkeletalMeshCustomVersion::GUID) < FSkeletalMeshCustomVersion::SplitModelAndRenderData)
 		{
@@ -579,6 +617,12 @@ void FSkeletalMeshLODModel::Serialize(FArchive& Ar, UObject* Owner, int32 Idx)
 	else
 	{
 		Ar << Sections;
+		
+		if (!StripFlags.IsEditorDataStripped() && Ar.CustomVer(FEditorObjectVersion::GUID) >= FEditorObjectVersion::SkeletalMeshBuildRefactor)
+		{
+			//Editor builds only
+			Ar << UserSectionsData;
+		}
 
 		// For old content, load as a multi-size container, but convert into regular array
 		if (Ar.IsLoading() && Ar.CustomVer(FSkeletalMeshCustomVersion::GUID) < FSkeletalMeshCustomVersion::SplitModelAndRenderData)
@@ -887,7 +931,7 @@ void FSkeletalMeshLODModel::GetNonClothVertices(TArray<FSoftSkinVertex>& OutVert
 	for (int32 SectionIndex = 0; SectionIndex < NumSections; SectionIndex++)
 	{
 		const FSkelMeshSection& Section = Sections[SectionIndex];
-		if (Section.HasClothingData())
+		if (Section.ClothingData.AssetGuid.IsValid())
 		{
 			continue;
 		}
@@ -909,4 +953,74 @@ bool FSkeletalMeshLODModel::DoSectionsNeedExtraBoneInfluences() const
 	return false;
 }
 
+void FSkeletalMeshLODModel::SyncronizeUserSectionsDataArray(bool bResetNonUsedSection /*= false*/)
+{
+	int32 SectionNum = Sections.Num();
+	for (int32 SectionIndex = 0; SectionIndex < SectionNum; ++SectionIndex)
+	{
+		FSkelMeshSection& Section = Sections[SectionIndex];
+		FSkelMeshSourceSectionUserData& SectionUserData = UserSectionsData.FindOrAdd(Section.OriginalDataSectionIndex);
+		Section.bCastShadow					= SectionUserData.bCastShadow;
+		Section.bRecomputeTangent			= SectionUserData.bRecomputeTangent;
+		Section.bDisabled					= SectionUserData.bDisabled;
+		Section.GenerateUpToLodIndex		= SectionUserData.GenerateUpToLodIndex;
+		Section.CorrespondClothAssetIndex	= SectionUserData.CorrespondClothAssetIndex;
+		Section.ClothingData.AssetGuid		= SectionUserData.ClothingData.AssetGuid;
+		Section.ClothingData.AssetLodIndex	= SectionUserData.ClothingData.AssetLodIndex;
+	}
+
+	//Reset normally happen when we re-import a skeletalmesh, we never want to reset this when we build the skeletalmesh (reduce can remove section, but we need to keep the original section data)
+	if (bResetNonUsedSection)
+	{
+		//Make sure we have the correct amount of UserSectionData we delete all the entries and recreate them with the previously sync Sections
+		UserSectionsData.Reset();
+		for (int32 SectionIndex = 0; SectionIndex < SectionNum; ++SectionIndex)
+		{
+			FSkelMeshSection& Section = Sections[SectionIndex];
+			//We only need parent section, no need to iterate bone chunked sections
+			if (Section.ChunkedParentSectionIndex != INDEX_NONE)
+			{
+				continue;
+			}
+			FSkelMeshSourceSectionUserData& SectionUserData = UserSectionsData.FindOrAdd(Section.OriginalDataSectionIndex);
+			SectionUserData.bCastShadow = Section.bCastShadow;
+			SectionUserData.bRecomputeTangent = Section.bRecomputeTangent;
+			SectionUserData.bDisabled = Section.bDisabled;
+			SectionUserData.GenerateUpToLodIndex = Section.GenerateUpToLodIndex;
+			SectionUserData.CorrespondClothAssetIndex = Section.CorrespondClothAssetIndex;
+			SectionUserData.ClothingData.AssetGuid = Section.ClothingData.AssetGuid;
+			SectionUserData.ClothingData.AssetLodIndex = Section.ClothingData.AssetLodIndex;
+		}
+	}
+}
+
+FString FSkeletalMeshLODModel::GetLODModelDeriveDataKey() const
+{
+	FString KeySuffix = TEXT("LODMODEL");
+
+	TArray<uint8> ByteData;
+	FMemoryWriter Ar(ByteData, true);
+
+	//Add the bulk data ID (if someone modify the original imported data, this ID will change)
+	FString BulkDatIDString = RawSkeletalMeshBulkData.GetIdString();
+	Ar << BulkDatIDString;
+	int32 UserSectionCount = UserSectionsData.Num();
+	Ar << UserSectionCount;
+	for (auto Kvp : UserSectionsData)
+	{
+		Ar << Kvp.Key;
+		Ar << Kvp.Value;
+	}
+
+	FSHA1 Sha;
+	Sha.Update(ByteData.GetData(), ByteData.Num() * ByteData.GetTypeSize());
+	Sha.Final();
+	// Retrieve the hash and use it to construct a pseudo-GUID.
+	uint32 Hash[5];
+	Sha.GetHash((uint8*)Hash);
+	KeySuffix += FGuid(Hash[0] ^ Hash[4], Hash[1], Hash[2], Hash[3]).ToString(EGuidFormats::Digits);
+
+	return KeySuffix;
+}
+
 #endif // WITH_EDITOR
diff --git a/Engine/Source/Runtime/Engine/Private/SkeletalMeshLODSettings.cpp b/Engine/Source/Runtime/Engine/Private/SkeletalMeshLODSettings.cpp
--- a/Engine/Source/Runtime/Engine/Private/SkeletalMeshLODSettings.cpp
+++ b/Engine/Source/Runtime/Engine/Private/SkeletalMeshLODSettings.cpp
@@ -6,6 +6,7 @@
 #include "Animation/Skeleton.h"
 #include "Animation/AnimSequence.h"
 #include "UObject/FortniteMainBranchObjectVersion.h"
+#include "UObject/EditorObjectVersion.h"
 
 DEFINE_LOG_CATEGORY_STATIC(LogSkeletalMeshLODSettings, Warning, All)
 
@@ -47,6 +48,7 @@ bool USkeletalMeshLODSettings::SetLODSettingsToMesh(USkeletalMesh* InMesh, int32
 		FSkeletalMeshLODInfo* LODInfo = InMesh->GetLODInfo(LODIndex);
 		const FSkeletalMeshLODGroupSettings& Setting = LODGroups[LODIndex];
 		LODInfo->ReductionSettings = Setting.ReductionSettings;
+		LODInfo->BuildSettings = Setting.BuildSettings;
 		LODInfo->ScreenSize = Setting.ScreenSize;
 		LODInfo->LODHysteresis = Setting.LODHysteresis;
 		LODInfo->WeightOfPrioritization = Setting.WeightOfPrioritization;
@@ -206,6 +208,7 @@ int32 USkeletalMeshLODSettings::SetLODSettingsFromMesh(USkeletalMesh* InMesh)
 			FSkeletalMeshLODInfo* LODInfo = InMesh->GetLODInfo(Index);
 			FSkeletalMeshLODGroupSettings& Setting = LODGroups[Index];
 			Setting.ReductionSettings = LODInfo->ReductionSettings;
+			Setting.BuildSettings = LODInfo->BuildSettings;
 			Setting.ScreenSize = LODInfo->ScreenSize;
 			Setting.LODHysteresis = LODInfo->LODHysteresis;
 			// copy mesh setting to shared setting
@@ -259,6 +262,7 @@ void USkeletalMeshLODSettings::Serialize(FArchive& Ar)
 	Super::Serialize(Ar);
 
 	Ar.UsingCustomVersion(FFortniteMainBranchObjectVersion::GUID);
+	Ar.UsingCustomVersion(FEditorObjectVersion::GUID);
 
 	if (Ar.CustomVer(FFortniteMainBranchObjectVersion::GUID) < FFortniteMainBranchObjectVersion::ConvertReductionSettingOptions)
 	{
@@ -275,6 +279,27 @@ void USkeletalMeshLODSettings::Serialize(FArchive& Ar)
 			}
 		}
 	}
+
+	if (Ar.CustomVer(FEditorObjectVersion::GUID) < FEditorObjectVersion::SkeletalMeshBuildRefactor)
+	{
+		//Since we use the bake source model when loading old asset we do not want to change anything
+		for (int32 Index = 0; Index < LODGroups.Num(); ++Index)
+		{
+			FSkeletalMeshBuildSettings& BuildSettings = LODGroups[Index].BuildSettings;
+			BuildSettings.bRecomputeNormals = false;
+			BuildSettings.bRecomputeTangents = false;
+			BuildSettings.bUseMikkTSpace = false;
+			BuildSettings.bRemoveDegenerates = false;
+			BuildSettings.ThresholdPosition = 0.0f;
+			BuildSettings.ThresholdTangentNormal = 0.0f;
+			BuildSettings.ThresholdUV = 0.0f;
+
+			//The user will have to make sure those settings are ok since before the modification precisions
+			//were one value for the whole skeletal mesh LODs. We set them to false by default
+			BuildSettings.bUseFullPrecisionUVs = GVertexElementTypeSupport.IsSupported(VET_Half2) ? false : true;
+			BuildSettings.bUseHighPrecisionTangentBasis = false;
+		}
+	}
 }
 
 /////////////////////////////////////////////////////////////
@@ -285,6 +310,14 @@ FSkeletalMeshOptimizationSettings FSkeletalMeshLODGroupSettings::GetReductionSet
 	return ReductionSettings;
 }
 
+/////////////////////////////////////////////////////////////
+// FSkeletalMeshBuildSettings 
+/////////////////////////////////////////////////////////////
+FSkeletalMeshBuildSettings FSkeletalMeshLODGroupSettings::GetBuildSettings() const
+{
+	return BuildSettings;
+}
+
 const float FSkeletalMeshLODGroupSettings::GetScreenSize() const
 {
 	return ScreenSize.Default;
diff --git a/Engine/Source/Runtime/Engine/Private/SkeletalMeshMerge.cpp b/Engine/Source/Runtime/Engine/Private/SkeletalMeshMerge.cpp
--- a/Engine/Source/Runtime/Engine/Private/SkeletalMeshMerge.cpp
+++ b/Engine/Source/Runtime/Engine/Private/SkeletalMeshMerge.cpp
@@ -163,8 +163,6 @@ bool FSkeletalMeshMerge::FinalizeMesh()
 	// If things are going ok so far...
 	if (Result)
 	{
-		// force 16 bit UVs if supported on hardware
-		MergeMesh->bUseFullPrecisionUVs = GVertexElementTypeSupport.IsSupported(VET_Half2) ? false : true;
 
 		// Array of per-lod number of UV sets
 		TArray<uint32> PerLODNumUVSets;
@@ -194,7 +192,9 @@ bool FSkeletalMeshMerge::FinalizeMesh()
 		MergeMesh->AllocateResourceForRendering();
 		for (int32 LODIdx = 0; LODIdx < MaxNumLODs; LODIdx++)
 		{
-			if (!MergeMesh->bUseFullPrecisionUVs)
+			const FSkeletalMeshLODInfo* LODInfoPtr = MergeMesh->GetLODInfo(LODIdx);
+			bool bUseFullPrecisionUVs = LODInfoPtr ? LODInfoPtr->BuildSettings.bUseFullPrecisionUVs : GVertexElementTypeSupport.IsSupported(VET_Half2) ? false : true;
+			if (!bUseFullPrecisionUVs)
 			{
 				if (PerLODExtraBoneInfluences[LODIdx])
 				{
@@ -303,8 +303,7 @@ void FSkeletalMeshMerge::GenerateNewSectionArray( TArray<FNewSectionInfo>& NewSe
 				// get the material for this section
 				int32 MaterialIndex = Section.MaterialIndex;
 				// use the remapping of material indices for all LODs besides the base LOD 
-				if( LODIdx > 0 && 
-					SrcLODInfo.LODMaterialMap.IsValidIndex(Section.MaterialIndex) )
+				if(SrcLODInfo.LODMaterialMap.IsValidIndex(Section.MaterialIndex) && SrcMesh->Materials.IsValidIndex(SrcLODInfo.LODMaterialMap[Section.MaterialIndex]))
 				{
 					MaterialIndex = FMath::Clamp<int32>( SrcLODInfo.LODMaterialMap[Section.MaterialIndex], 0, SrcMesh->Materials.Num() );
 				}
@@ -564,6 +563,9 @@ void FSkeletalMeshMerge::GenerateLODModel( int32 LODIdx )
 				}
 			}
 #endif
+			MergeLODInfo.BuildSettings.bUseFullPrecisionUVs |= SrcLODInfo.BuildSettings.bUseFullPrecisionUVs;
+			MergeLODInfo.BuildSettings.bUseHighPrecisionTangentBasis |= SrcLODInfo.BuildSettings.bUseHighPrecisionTangentBasis;
+
 			MergeLODInfo.LODHysteresis = FMath::Min<float>(MergeLODInfo.LODHysteresis,SrcLODInfo.LODHysteresis);
 
 			// get the source skel LOD model from this merge entry
@@ -730,9 +732,9 @@ void FSkeletalMeshMerge::GenerateLODModel( int32 LODIdx )
 	// sort required bone array in strictly increasing order
 	MergeLODData.RequiredBones.Sort();
 	MergeMesh->RefSkeleton.EnsureParentsExistAndSort(MergeLODData.ActiveBoneIndices);
-
+	
 	// copy the new vertices and indices to the vertex buffer for the new model
-	MergeLODData.StaticVertexBuffers.StaticMeshVertexBuffer.SetUseFullPrecisionUVs(MergeMesh->bUseFullPrecisionUVs);
+	MergeLODData.StaticVertexBuffers.StaticMeshVertexBuffer.SetUseFullPrecisionUVs(MergeLODInfo.BuildSettings.bUseFullPrecisionUVs);
 
 	MergeLODData.StaticVertexBuffers.PositionVertexBuffer.Init(MergedVertexBuffer.Num(), bNeedsCPUAccess);
 	MergeLODData.StaticVertexBuffers.StaticMeshVertexBuffer.Init(MergedVertexBuffer.Num(), TotalNumUVs, bNeedsCPUAccess);
diff --git a/Engine/Source/Runtime/Engine/Private/SkeletalMeshModel.cpp b/Engine/Source/Runtime/Engine/Private/SkeletalMeshModel.cpp
--- a/Engine/Source/Runtime/Engine/Private/SkeletalMeshModel.cpp
+++ b/Engine/Source/Runtime/Engine/Private/SkeletalMeshModel.cpp
@@ -89,6 +89,31 @@ FString FSkeletalMeshModel::GetIdString() const
 	return GuidString;
 }
 
+void FSkeletalMeshModel::SyncronizeLODUserSectionsData()
+{
+	for (int32 LODIndex = 0; LODIndex < LODModels.Num(); ++LODIndex)
+	{
+		FSkeletalMeshLODModel& Model = LODModels[LODIndex];
+		Model.SyncronizeUserSectionsDataArray();
+	}
+}
+
+FString FSkeletalMeshModel::GetLODModelIdString() const
+{
+	FSHA1 Sha;
+	for (int32 LODIndex = 0; LODIndex < LODModels.Num(); ++LODIndex)
+	{
+		const FSkeletalMeshLODModel& Model = LODModels[LODIndex];
+		Model.BuildStringID = Model.GetLODModelDeriveDataKey();
+		TArray<TCHAR> IDArray = Model.BuildStringID.GetCharArray();
+		Sha.Update((uint8*)IDArray.GetData(), IDArray.Num() * IDArray.GetTypeSize());
+	}
+	Sha.Final();
+	// Retrieve the hash and use it to construct a pseudo-GUID. 
+	uint32 Hash[5];
+	Sha.GetHash((uint8*)Hash);
+	return FGuid(Hash[0] ^ Hash[4], Hash[1], Hash[2], Hash[3]).ToString(EGuidFormats::Digits);
+}
 
 void FSkeletalMeshModel::GetResourceSizeEx(FResourceSizeEx& CumulativeResourceSize)
 {
diff --git a/Engine/Source/Runtime/Engine/Private/SkeletalMeshRenderData.cpp b/Engine/Source/Runtime/Engine/Private/SkeletalMeshRenderData.cpp
--- a/Engine/Source/Runtime/Engine/Private/SkeletalMeshRenderData.cpp
+++ b/Engine/Source/Runtime/Engine/Private/SkeletalMeshRenderData.cpp
@@ -15,6 +15,8 @@
 #include "Interfaces/ITargetPlatformManagerModule.h"
 #include "Interfaces/ITargetPlatform.h"
 #include "PlatformInfo.h"
+#include "IMeshBuilderModule.h"
+#include "EngineUtils.h"
 
 
 #if ENABLE_COOK_STATS
@@ -28,11 +30,36 @@ namespace SkeletalMeshCookStats
 }
 #endif
 
+//Serialize the LODInfo and append the result to the KeySuffix to build the LODInfo part of the DDC KEY
+//Note: this serializer is only used to build the mesh DDC key, no versioning is required
+static void SerializeLODInfoForDDC(USkeletalMesh* SkeletalMesh, FString& KeySuffix)
+{
+	TArray<FSkeletalMeshLODInfo>& LODInfos = SkeletalMesh->GetLODInfoArray();
+	const bool bIs16BitfloatBufferSupported = GVertexElementTypeSupport.IsSupported(VET_Half2);
+	for (int32 LODIndex = 0; LODIndex < SkeletalMesh->GetLODNum(); ++LODIndex)
+	{
+		check(LODInfos.IsValidIndex(LODIndex));
+		FSkeletalMeshLODInfo& LODInfo = LODInfos[LODIndex];
+		bool bValidLODSettings = false;
+		if (SkeletalMesh->LODSettings != nullptr)
+		{
+			const int32 NumSettings = FMath::Min(SkeletalMesh->LODSettings->GetNumberOfSettings(), SkeletalMesh->GetLODNum());
+			if (LODIndex < NumSettings)
+			{
+				bValidLODSettings = true;
+			}
+		}
+		const FSkeletalMeshLODGroupSettings* SkeletalMeshLODGroupSettings = bValidLODSettings ? &SkeletalMesh->LODSettings->GetSettingsForLODLevel(LODIndex) : nullptr;
+		LODInfo.BuildGUID = LODInfo.ComputeDeriveDataCacheKey(SkeletalMeshLODGroupSettings);
+		KeySuffix += LODInfo.BuildGUID.ToString(EGuidFormats::Digits);
+	}
+}
+
 // If skeletal mesh derived data needs to be rebuilt (new format, serialization
 // differences, etc.) replace the version GUID below with a new one.
 // In case of merge conflicts with DDC versions, you MUST generate a new GUID
 // and set this new GUID as the version.                                       
-#define SKELETALMESH_DERIVEDDATA_VER TEXT("4F54F64FD82B40C7A9F4158AF684B64D")
+#define SKELETALMESH_DERIVEDDATA_VER TEXT("51BDAF351B6941548818C2F8B33397D7")
 
 static const FString& GetSkeletalMeshDerivedDataVersion()
 {
@@ -48,8 +75,33 @@ static FString BuildSkeletalMeshDerivedDataKey(USkeletalMesh* SkelMesh)
 {
 	FString KeySuffix(TEXT(""));
 
-	KeySuffix += SkelMesh->GetImportedModel()->GetIdString();
-	KeySuffix += (SkelMesh->bUseFullPrecisionUVs || !GVertexElementTypeSupport.IsSupported(VET_Half2)) ? "1" : "0";
+	if (SkelMesh->UseLegacyMeshDerivedDataKey )
+	{
+		//Old asset will have the same LOD settings for bUseFullPrecisionUVs. We can use the LOD 0
+		const FSkeletalMeshLODInfo* BaseLODInfo = SkelMesh->GetLODInfo(0);
+		bool bUseFullPrecisionUVs = BaseLODInfo ? BaseLODInfo->BuildSettings.bUseFullPrecisionUVs : false;
+		KeySuffix += SkelMesh->GetImportedModel()->GetIdString();
+		KeySuffix += (bUseFullPrecisionUVs || !GVertexElementTypeSupport.IsSupported(VET_Half2)) ? "1" : "0";
+	}
+	else
+	{
+		FString tmpDebugString;
+		//Synchronize the user data that are part of the key
+		SkelMesh->GetImportedModel()->SyncronizeLODUserSectionsData();
+		tmpDebugString = SkelMesh->GetImportedModel()->GetIdString();
+		KeySuffix += tmpDebugString;
+		tmpDebugString = SkelMesh->GetImportedModel()->GetLODModelIdString();
+		KeySuffix += tmpDebugString;
+		
+		//Add the max gpu bone per section
+		const uint32 MaxGPUSkinBones = FGPUBaseSkinVertexFactory::GetMaxGPUSkinBones();
+		KeySuffix += FString::FromInt((int32)MaxGPUSkinBones);
+
+		tmpDebugString = TEXT("");
+		SerializeLODInfoForDDC(SkelMesh, tmpDebugString);
+		KeySuffix += tmpDebugString;
+	}
+
 	KeySuffix += SkelMesh->bHasVertexColors ? "1" : "0";
 	KeySuffix += SkelMesh->VertexColorGuid.ToString(EGuidFormats::Digits);
 
@@ -99,7 +151,34 @@ void FSkeletalMeshRenderData::Cache(USkeletalMesh* Owner)
 		if (GetDerivedDataCacheRef().GetSynchronous(*DerivedDataKey, DerivedData))
 		{
 			COOK_STAT(Timer.AddHit(DerivedData.Num()));
+			
+
 			FMemoryReader Ar(DerivedData, /*bIsPersistent=*/ true);
+
+			//With skeletal mesh build refactor we serialize the LODModel sections into the DDC
+			//We need to store those so we do not have to rerun the reduction to make them up to date
+			//with the serialize renderdata. This allow to use DDC when changing the reduction settings.
+			//The old workflow has to reduce the LODModel before getting the render data DDC.
+			if (!Owner->UseLegacyMeshDerivedDataKey)
+			{
+				FSkeletalMeshModel* SkelMeshModel = Owner->GetImportedModel();
+				check(SkelMeshModel);
+				//Serialize the LODModel sections since they are dependent on the reduction
+				for (int32 LODIndex = 0; LODIndex < SkelMeshModel->LODModels.Num(); LODIndex++)
+				{
+					FSkeletalMeshLODModel* LODModel = &(SkelMeshModel->LODModels[LODIndex]);
+					int32 SectionNum = 0;
+					Ar << SectionNum;
+					LODModel->Sections.Reset(SectionNum);
+					LODModel->Sections.AddDefaulted(SectionNum);
+					for (int32 SectionIndex = 0; SectionIndex < LODModel->Sections.Num(); ++SectionIndex)
+					{
+						Ar << LODModel->Sections[SectionIndex];
+					}
+					LODModel->SyncronizeUserSectionsDataArray();
+				}
+			}
+
 			Serialize(Ar, Owner);
 			for (int32 LODIndex = 0; LODIndex < LODRenderData.Num(); ++LODIndex)
 			{
@@ -129,14 +208,77 @@ void FSkeletalMeshRenderData::Cache(USkeletalMesh* Owner)
 
 			for (int32 LODIndex = 0; LODIndex < SkelMeshModel->LODModels.Num(); LODIndex++)
 			{
-				FSkeletalMeshLODModel& LODModel = SkelMeshModel->LODModels[LODIndex];
+				FSkeletalMeshLODModel* LODModel = &(SkelMeshModel->LODModels[LODIndex]);
+				FSkeletalMeshLODInfo* LODInfo = Owner->GetLODInfo(LODIndex);
+				check(LODInfo);
+				//Build the source model before the render data, if we are a purely generated LOD we do not need to be build
+				IMeshBuilderModule& MeshBuilderModule = FModuleManager::Get().LoadModuleChecked<IMeshBuilderModule>("MeshBuilder");
+				if (!LODModel->RawSkeletalMeshBulkData.IsEmpty() && LODModel->RawSkeletalMeshBulkData.IsBuildDataAvailable())
+				{
+					MeshBuilderModule.BuildSkeletalMesh(Owner, LODIndex, true);
+					LODModel = &(SkelMeshModel->LODModels[LODIndex]);
+				}
+				else
+				{
+					//We need to synchronize when we are generated mesh or if we have load an old asset that was not re-imported
+					LODModel->SyncronizeUserSectionsDataArray();
+				}
+
 				FSkeletalMeshLODRenderData* LODData = new FSkeletalMeshLODRenderData();
 				LODRenderData.Add(LODData);
-
-				LODData->BuildFromLODModel(&LODModel, VertexBufferBuildFlags);
+				
+				//Get the UVs and tangents precision build settings flag specific for this LOD index
+				{
+					bool bUseFullPrecisionUVs = LODInfo->BuildSettings.bUseFullPrecisionUVs;
+					bool bUseHighPrecisionTangentBasis = LODInfo->BuildSettings.bUseHighPrecisionTangentBasis;
+					bool bBuildAdjacencyBuffer = LODInfo->BuildSettings.bBuildAdjacencyBuffer;
+					if (Owner->LODSettings != nullptr)
+					{
+						const int32 NumSettings = Owner->LODSettings->GetNumberOfSettings();
+						if (LODIndex < NumSettings)
+						{
+							const FSkeletalMeshLODGroupSettings& SkeletalMeshLODGroupSettings = Owner->LODSettings->GetSettingsForLODLevel(LODIndex);
+							bUseFullPrecisionUVs = SkeletalMeshLODGroupSettings.GetBuildSettings().bUseFullPrecisionUVs;
+							bUseHighPrecisionTangentBasis = SkeletalMeshLODGroupSettings.GetBuildSettings().bUseHighPrecisionTangentBasis;
+							bBuildAdjacencyBuffer = SkeletalMeshLODGroupSettings.GetBuildSettings().bBuildAdjacencyBuffer;
+						}
+					}
+					if (bUseFullPrecisionUVs || !GVertexElementTypeSupport.IsSupported(VET_Half2))
+					{
+						VertexBufferBuildFlags |= ESkeletalMeshVertexFlags::UseFullPrecisionUVs;
+					}
+					if (bUseHighPrecisionTangentBasis)
+					{
+						VertexBufferBuildFlags |= ESkeletalMeshVertexFlags::UseHighPrecisionTangentBasis;
+					}
+					if (bBuildAdjacencyBuffer)
+					{
+						VertexBufferBuildFlags |= ESkeletalMeshVertexFlags::BuildAdjacencyIndexBuffer;
+					}
+				}
+				LODData->BuildFromLODModel(LODModel, VertexBufferBuildFlags);
 			}
 
 			FMemoryWriter Ar(DerivedData, /*bIsPersistent=*/ true);
+			
+			//If we load an old asset we want to be sure the serialize ddc will be the same has before the skeletalmesh build refactor
+			//So we do not serialize the LODModel sections.
+			if (!Owner->UseLegacyMeshDerivedDataKey)
+			{
+				//Serialize the LODModel sections since they are dependent on the reduction
+				for (int32 LODIndex = 0; LODIndex < SkelMeshModel->LODModels.Num(); LODIndex++)
+				{
+					FSkeletalMeshLODModel* LODModel = &(SkelMeshModel->LODModels[LODIndex]);
+					int32 SectionNum = LODModel->Sections.Num();
+					Ar << SectionNum;
+					for (int32 SectionIndex = 0; SectionIndex < LODModel->Sections.Num(); ++SectionIndex)
+					{
+						Ar << LODModel->Sections[SectionIndex];
+					}
+				}
+			}
+
+			//Serialize the render data
 			Serialize(Ar, Owner);
 			for (int32 LODIndex = 0; LODIndex < LODRenderData.Num(); ++LODIndex)
 			{
diff --git a/Engine/Source/Runtime/Engine/Public/Rendering/SkeletalMeshLODImporterData.cpp b/Engine/Source/Runtime/Engine/Public/Rendering/SkeletalMeshLODImporterData.cpp
--- a/Engine/Source/Runtime/Engine/Public/Rendering/SkeletalMeshLODImporterData.cpp
+++ b/Engine/Source/Runtime/Engine/Public/Rendering/SkeletalMeshLODImporterData.cpp
@@ -3,7 +3,8 @@
 #if WITH_EDITOR
 
 #include "SkeletalMeshLODImporterData.h"
-#include "Serialization/MemoryWriter.h"
+#include "Serialization/BulkDataWriter.h"
+#include "Serialization/BulkDataReader.h"
 #include "Rendering/SkeletalMeshModel.h"
 #include "Engine/SkeletalMesh.h"
 #include "Factories/FbxSkeletalMeshImportData.h"
@@ -349,14 +350,16 @@ enum
 
 struct FReductionSkeletalMeshData
 {
-	FReductionSkeletalMeshData(FSkeletalMeshLODModel& InBaseLODModel, TMap<FString, TArray<FMorphTargetDelta>>& InBaseLODMorphTargetData)
+	FReductionSkeletalMeshData(FSkeletalMeshLODModel& InBaseLODModel, TMap<FString, TArray<FMorphTargetDelta>>& InBaseLODMorphTargetData, UObject* InOwner)
 		: BaseLODModel(InBaseLODModel)
 		, BaseLODMorphTargetData(InBaseLODMorphTargetData)
+		, Owner(InOwner)
 	{
 	}
 
 	FSkeletalMeshLODModel& BaseLODModel;
 	TMap<FString, TArray<FMorphTargetDelta>>& BaseLODMorphTargetData;
+	UObject* Owner;
 };
 
 FArchive& operator<<(FArchive& Ar, FReductionSkeletalMeshData& ReductionSkeletalMeshData)
@@ -365,7 +368,7 @@ FArchive& operator<<(FArchive& Ar, FReductionSkeletalMeshData& ReductionSkeletal
 	int32 LicenseeVersion = REDUCTION_BASE_SK_DATA_BULKDATA_LIC_VER;
 	Ar << Version;
 	Ar << LicenseeVersion;
-	ReductionSkeletalMeshData.BaseLODModel.Serialize(Ar, nullptr, 0);
+	ReductionSkeletalMeshData.BaseLODModel.Serialize(Ar, ReductionSkeletalMeshData.Owner, INDEX_NONE);
 	Ar << ReductionSkeletalMeshData.BaseLODMorphTargetData;
 	return Ar;
 }
@@ -386,10 +389,8 @@ void FReductionBaseSkeletalMeshBulkData::Serialize(FArchive& Ar, TArray<FReducti
 		for (int32 Index = 0; Index < NewNum; Index++)
 		{
 			FReductionBaseSkeletalMeshBulkData* EmptyData = new FReductionBaseSkeletalMeshBulkData();
-			ReductionBaseSkeletalMeshDatas.Add(EmptyData);
-		}
-		for (int32 Index = 0; Index < NewNum; Index++)
-		{
+			int32 NewEntryIndex = ReductionBaseSkeletalMeshDatas.Add(EmptyData);
+			check(NewEntryIndex == Index);
 			ReductionBaseSkeletalMeshDatas[Index]->Serialize(Ar, Owner);
 		}
 	}
@@ -407,68 +408,75 @@ void FReductionBaseSkeletalMeshBulkData::Serialize(FArchive& Ar, TArray<FReducti
 
 void FReductionBaseSkeletalMeshBulkData::Serialize(FArchive& Ar, UObject* Owner)
 {
-	if (Ar.IsLoading())
+	if (Ar.IsTransacting())
 	{
-		//Save the custom version so we can load FReductionSkeletalMeshData later
-		SerializeLoadingCustomVersionContainer = Ar.GetCustomVersions();
-		bUseSerializeLoadingCustomVersion = false;
-		const FCustomVersionContainer& RegisteredCustomVersionContainer = FCustomVersionContainer::GetRegistered();
-		const FCustomVersionArray& ArchiveCustomVersionArray = SerializeLoadingCustomVersionContainer.GetAllVersions();
-		for (const FCustomVersion& ArchiveVersion : ArchiveCustomVersionArray)
+		// If transacting, keep these members alive the other side of an undo, otherwise their values will get lost
+		SerializeLoadingCustomVersionContainer.Serialize(Ar);
+		Ar << bUseSerializeLoadingCustomVersion;
+	}
+	else
+	{
+		if (Ar.IsLoading())
 		{
-			const FCustomVersion* RegisteredVersion = RegisteredCustomVersionContainer.GetVersion(ArchiveVersion.Key);
-			if (!RegisteredVersion || RegisteredVersion->Version != ArchiveVersion.Version)
-			{
-				bUseSerializeLoadingCustomVersion = true;
-				break;
-			}
+			//Save the custom version so we can load FReductionSkeletalMeshData later
+			SerializeLoadingCustomVersionContainer = Ar.GetCustomVersions();
+			bUseSerializeLoadingCustomVersion = true;
 		}
-	}
 
-	if (Ar.IsSaving() && bUseSerializeLoadingCustomVersion == true)
-	{
-		//We need to update the FReductionSkeletalMeshData serialize version to the latest in case we save the Parent bulkdata
-		FSkeletalMeshLODModel BaseLODModel;
-		TMap<FString, TArray<FMorphTargetDelta>> BaseLODMorphTargetData;
-		LoadReductionData(BaseLODModel, BaseLODMorphTargetData);
-		SaveReductionData(BaseLODModel, BaseLODMorphTargetData);
+		if (Ar.IsSaving() && bUseSerializeLoadingCustomVersion == true)
+		{
+			//We need to update the FReductionSkeletalMeshData serialize version to the latest in case we save the Parent bulkdata
+			FSkeletalMeshLODModel BaseLODModel;
+			TMap<FString, TArray<FMorphTargetDelta>> BaseLODMorphTargetData;
+			LoadReductionData(BaseLODModel, BaseLODMorphTargetData, Owner);
+			SaveReductionData(BaseLODModel, BaseLODMorphTargetData, Owner);
+		}
 	}
 
 	BulkData.Serialize(Ar, Owner);
 }
 
-void FReductionBaseSkeletalMeshBulkData::SaveReductionData(FSkeletalMeshLODModel& BaseLODModel, TMap<FString, TArray<FMorphTargetDelta>>& BaseLODMorphTargetData)
+void FReductionBaseSkeletalMeshBulkData::SaveReductionData(FSkeletalMeshLODModel& BaseLODModel, TMap<FString, TArray<FMorphTargetDelta>>& BaseLODMorphTargetData, UObject* Owner)
 {
 	//Saving the bulk data mean we do not need anymore the SerializeLoadingCustomVersionContainer of the parent bulk data
 	SerializeLoadingCustomVersionContainer.Empty();
 	bUseSerializeLoadingCustomVersion = false;
-	FReductionSkeletalMeshData ReductionSkeletalMeshData(BaseLODModel, BaseLODMorphTargetData);
-	TArray<uint8> TempBytes;
-	FMemoryWriter Ar(TempBytes, /*bIsPersistent=*/ true);
-	Ar << ReductionSkeletalMeshData;
-	BulkData.Lock(LOCK_READ_WRITE);
-	uint8* Dest = (uint8*)BulkData.Realloc(TempBytes.Num());
-	FMemory::Memcpy(Dest, TempBytes.GetData(), TempBytes.Num());
-	BulkData.Unlock();
+	FReductionSkeletalMeshData ReductionSkeletalMeshData(BaseLODModel, BaseLODMorphTargetData, Owner);
+
+	//Clear the bulk data before writing it
+	BulkData.RemoveBulkData();
+
+	// Get a lock on the bulk data
+	{
+		const bool bIsPersistent = true;
+		FBulkDataWriter Ar(BulkData, bIsPersistent);
+		Ar << ReductionSkeletalMeshData;
+
+		// Preserve CustomVersions at save time so we can reuse the same ones when reloading direct from memory
+		SerializeLoadingCustomVersionContainer = Ar.GetCustomVersions();
+	}
+	// Unlock the bulk data
 }
 
-void FReductionBaseSkeletalMeshBulkData::LoadReductionData(FSkeletalMeshLODModel& BaseLODModel, TMap<FString, TArray<FMorphTargetDelta>>& BaseLODMorphTargetData)
+void FReductionBaseSkeletalMeshBulkData::LoadReductionData(FSkeletalMeshLODModel& BaseLODModel, TMap<FString, TArray<FMorphTargetDelta>>& BaseLODMorphTargetData, UObject* Owner)
 {
 	BaseLODMorphTargetData.Empty();
 	if (BulkData.GetElementCount() > 0)
 	{
-		FReductionSkeletalMeshData ReductionSkeletalMeshData(BaseLODModel, BaseLODMorphTargetData);
-		FBufferReader Ar(
-			BulkData.Lock(LOCK_READ_ONLY), BulkData.GetElementCount(),
-			/*bInFreeOnClose=*/ false, /*bIsPersistent=*/ true
-		);
-		//If we load the bulkdata content and the bulk was using customVersion, we must set the same until we save the reduction data
-		if (bUseSerializeLoadingCustomVersion)
+		FReductionSkeletalMeshData ReductionSkeletalMeshData(BaseLODModel, BaseLODMorphTargetData, Owner);
+
+		// Get a lock on the bulk data
 		{
+			const bool bIsPersistent = true;
+			FBulkDataReader Ar(BulkData, bIsPersistent);
+			
+			// Propagate the custom version information from the package to the bulk data, so that the MeshDescription
+			// is serialized with the same versioning.
 			Ar.SetCustomVersions(SerializeLoadingCustomVersionContainer);
+
+			Ar << ReductionSkeletalMeshData;
 		}
-		Ar << ReductionSkeletalMeshData;
-		BulkData.Unlock();
+		// Unlock the bulk data
 	}
 }
 
@@ -479,6 +487,8 @@ FRawSkeletalMeshBulkData
 FRawSkeletalMeshBulkData::FRawSkeletalMeshBulkData()
 	: bGuidIsHash(false)
 {
+	GeoImportVersion = ESkeletalMeshGeoImportVersions::Before_Versionning;
+	SkinningImportVersion = ESkeletalMeshSkinningImportVersions::Before_Versionning;
 }
 
 
@@ -563,30 +573,41 @@ FArchive& operator<<(FArchive& Ar, FSkeletalMeshImportData& RawMesh)
 
 void FRawSkeletalMeshBulkData::Serialize(FArchive& Ar, UObject* Owner)
 {
-	if (Ar.IsLoading())
+	if (Ar.IsTransacting())
 	{
-		//Save the custom version so we can load FReductionSkeletalMeshData later
-		SerializeLoadingCustomVersionContainer = Ar.GetCustomVersions();
-		bUseSerializeLoadingCustomVersion = false;
-		const FCustomVersionContainer& RegisteredCustomVersionContainer = FCustomVersionContainer::GetRegistered();
-		const FCustomVersionArray& ArchiveCustomVersionArray = SerializeLoadingCustomVersionContainer.GetAllVersions();
-		for (const FCustomVersion& ArchiveVersion : ArchiveCustomVersionArray)
+		// If transacting, keep these members alive the other side of an undo, otherwise their values will get lost
+		SerializeLoadingCustomVersionContainer.Serialize(Ar);
+		Ar << bUseSerializeLoadingCustomVersion;
+	}
+	else
+	{
+		if (Ar.IsLoading())
 		{
-			const FCustomVersion* RegisteredVersion = RegisteredCustomVersionContainer.GetVersion(ArchiveVersion.Key);
-			if (!RegisteredVersion || RegisteredVersion->Version != ArchiveVersion.Version)
-			{
-				bUseSerializeLoadingCustomVersion = true;
-				break;
-			}
+			//Save the custom version so we can load FReductionSkeletalMeshData later
+			SerializeLoadingCustomVersionContainer = Ar.GetCustomVersions();
+			bUseSerializeLoadingCustomVersion = true;
+		}
+
+		if (Ar.IsSaving() && bUseSerializeLoadingCustomVersion == true)
+		{
+			//We need to update the FReductionSkeletalMeshData serialize version to the latest in case we save the Parent bulkdata
+			FSkeletalMeshImportData MeshImportData;
+			LoadRawMesh(MeshImportData);
+			SaveRawMesh(MeshImportData);
 		}
 	}
 
-	if (Ar.IsSaving() && bUseSerializeLoadingCustomVersion == true)
+	Ar.UsingCustomVersion(FEditorObjectVersion::GUID);
+
+	if (Ar.CustomVer(FEditorObjectVersion::GUID) >= FEditorObjectVersion::SkeletalMeshBuildRefactor)
 	{
-		//We need to update the FReductionSkeletalMeshData serialize version to the latest in case we save the Parent bulkdata
-		FSkeletalMeshImportData MeshImportData;
-		LoadRawMesh(MeshImportData);
-		SaveRawMesh(MeshImportData);
+		Ar << GeoImportVersion;
+		Ar << SkinningImportVersion;
+	}
+	else
+	{
+		GeoImportVersion = ESkeletalMeshGeoImportVersions::Before_Versionning;
+		SkinningImportVersion = ESkeletalMeshSkinningImportVersions::Before_Versionning;
 	}
 	BulkData.Serialize(Ar, Owner);
 	Ar << Guid;
@@ -597,16 +618,33 @@ void FRawSkeletalMeshBulkData::SaveRawMesh(FSkeletalMeshImportData& InMesh)
 {
 	//Saving the bulk data mean we do not need anymore the SerializeLoadingCustomVersionContainer of the parent bulk data
 	SerializeLoadingCustomVersionContainer.Empty();
-	bUseSerializeLoadingCustomVersion = false;
 
-	TArray<uint8> TempBytes;
-	FMemoryWriter Ar(TempBytes, /*bIsPersistent=*/ true);
-	Ar << InMesh;
-	BulkData.Lock(LOCK_READ_WRITE);
-	uint8* Dest = (uint8*)BulkData.Realloc(TempBytes.Num());
-	FMemory::Memcpy(Dest, TempBytes.GetData(), TempBytes.Num());
-	BulkData.Unlock();
-	FPlatformMisc::CreateGuid(Guid);
+	//Clear the bulk data before writing it
+	BulkData.RemoveBulkData();
+
+	// Get a lock on the bulk data
+	{
+		const bool bIsPersistent = true;
+		FBulkDataWriter Ar(BulkData, bIsPersistent);
+		Ar << InMesh;
+
+		// Preserve CustomVersions at save time so we can reuse the same ones when reloading direct from memory
+		SerializeLoadingCustomVersionContainer = Ar.GetCustomVersions();
+	}
+	// Unlock bulk data when we leave scope
+
+	//Create the guid from the content, this allow to use the data into the ddc key
+	FSHA1 Sha;
+	if (BulkData.GetBulkDataSize() > 0)
+	{
+		uint8* Buffer = (uint8*)BulkData.Lock(LOCK_READ_ONLY);
+		Sha.Update(Buffer, BulkData.GetBulkDataSize());
+		BulkData.Unlock();
+	}
+	Sha.Final();
+	uint32 Hash[5];
+	Sha.GetHash((uint8*)Hash);
+	Guid = FGuid(Hash[0] ^ Hash[4], Hash[1], Hash[2], Hash[3]);
 }
 
 void FRawSkeletalMeshBulkData::LoadRawMesh(FSkeletalMeshImportData& OutMesh)
@@ -614,17 +652,17 @@ void FRawSkeletalMeshBulkData::LoadRawMesh(FSkeletalMeshImportData& OutMesh)
 	OutMesh.Empty();
 	if (BulkData.GetElementCount() > 0)
 	{
-		FBufferReader Ar(
-			BulkData.Lock(LOCK_READ_ONLY), BulkData.GetElementCount(),
-			/*bInFreeOnClose=*/ false, /*bIsPersistent=*/ true
-		);
-		//If we load the bulkdata content and the bulk was using customVersion, we must set the same until we save the reduction data
-		if (bUseSerializeLoadingCustomVersion)
+		// Get a lock on the bulk data
 		{
+			const bool bIsPersistent = true;
+			FBulkDataReader Ar(BulkData, bIsPersistent);
+
+			// Propagate the custom version information from the package to the bulk data, so that the MeshDescription
+			// is serialized with the same versioning.
 			Ar.SetCustomVersions(SerializeLoadingCustomVersionContainer);
+			Ar << OutMesh;
 		}
-		Ar << OutMesh;
-		BulkData.Unlock();
+		// Unlock bulk data when we leave scope
 	}
 }
 
diff --git a/Engine/Source/Runtime/Engine/Public/Rendering/SkeletalMeshLODImporterData.h b/Engine/Source/Runtime/Engine/Public/Rendering/SkeletalMeshLODImporterData.h
--- a/Engine/Source/Runtime/Engine/Public/Rendering/SkeletalMeshLODImporterData.h
+++ b/Engine/Source/Runtime/Engine/Public/Rendering/SkeletalMeshLODImporterData.h
@@ -15,6 +15,38 @@
 
 class FSkeletalMeshLODModel;
 
+#endif
+
+//////////////////////////////////////////////////////////////////////////
+//uenum class cannot be inside a preprocessor like #if WITH_EDITOR
+
+UENUM()
+enum class ESkeletalMeshGeoImportVersions : uint8
+{
+	Before_Versionning = 0,
+	SkeletalMeshBuildRefactor,
+
+	// -----<new versions can be added above this line>-------------------------------------------------
+	VersionPlusOne,
+	LatestVersion = VersionPlusOne - 1
+};
+
+UENUM()
+enum class ESkeletalMeshSkinningImportVersions : uint8
+{
+	Before_Versionning = 0,
+	SkeletalMeshBuildRefactor,
+
+	// -----<new versions can be added above this line>-------------------------------------------------
+	VersionPlusOne,
+	LatestVersion = VersionPlusOne - 1
+};
+
+// End of enum declaration
+//////////////////////////////////////////////////////////////////////////
+
+#if WITH_EDITOR
+
 namespace SkeletalMeshImportData
 {
 	struct FMeshWedge
@@ -391,10 +423,10 @@ public:
 	ENGINE_API void Serialize(class FArchive& Ar, class UObject* Owner);
 
 	/** Store a new raw mesh in the bulk data. */
-	ENGINE_API void SaveReductionData(FSkeletalMeshLODModel& BaseLODModel, TMap<FString, TArray<FMorphTargetDelta>>& BaseLODMorphTargetData);
+	ENGINE_API void SaveReductionData(FSkeletalMeshLODModel& BaseLODModel, TMap<FString, TArray<FMorphTargetDelta>>& BaseLODMorphTargetData, UObject* Owner);
 
 	/** Load the raw mesh from bulk data. */
-	ENGINE_API void LoadReductionData(FSkeletalMeshLODModel& BaseLODModel, TMap<FString, TArray<FMorphTargetDelta>>& BaseLODMorphTargetData);
+	ENGINE_API void LoadReductionData(FSkeletalMeshLODModel& BaseLODModel, TMap<FString, TArray<FMorphTargetDelta>>& BaseLODMorphTargetData, UObject* Owner);
 
 	ENGINE_API FByteBulkData& GetBulkData() { return BulkData; }
 
@@ -421,6 +453,18 @@ class FRawSkeletalMeshBulkData
 	bool bUseSerializeLoadingCustomVersion = false;
 
 public:
+	/*
+	 * The last geo imported version, we use this flag to know if we have some data or not.
+	 * This flag must be updated every time we import a new geometry
+	 */
+	ESkeletalMeshGeoImportVersions GeoImportVersion;
+
+	/*
+	 * The last skinning imported version, we use this flag to know if we have some data or not.
+	 * This flag must be updated every time we import the skinning
+	 */
+	ESkeletalMeshSkinningImportVersions SkinningImportVersion;
+
 	/** Default constructor. */
 	ENGINE_API FRawSkeletalMeshBulkData();
 
@@ -446,6 +490,13 @@ public:
 
 	/** Returns true if no bulk data is available for this mesh. */
 	FORCEINLINE bool IsEmpty() const { return BulkData.GetBulkDataSize() == 0; }
+
+	/** Returns true if the last import version is enough to use the new build system. WE cannot rebuild asset if we did not previously store the data*/
+	ENGINE_API bool IsBuildDataAvailable() const
+	{
+		return GeoImportVersion >= ESkeletalMeshGeoImportVersions::SkeletalMeshBuildRefactor &&
+			SkinningImportVersion >= ESkeletalMeshSkinningImportVersions::SkeletalMeshBuildRefactor;
+	}
 };
 
 namespace FWedgePositionHelper
diff --git a/Engine/Source/Runtime/Engine/Public/Rendering/SkeletalMeshLODModel.h b/Engine/Source/Runtime/Engine/Public/Rendering/SkeletalMeshLODModel.h
--- a/Engine/Source/Runtime/Engine/Public/Rendering/SkeletalMeshLODModel.h
+++ b/Engine/Source/Runtime/Engine/Public/Rendering/SkeletalMeshLODModel.h
@@ -127,6 +127,24 @@ struct FSkelMeshSection
 	 */
 	int32 GenerateUpToLodIndex;
 
+	/*
+	 * This represent the original section index in the imported data. The original data is chunk per material,
+	 * we use this index to store user section modification. The user cannot change a BONE chunked section data,
+	 * since the BONE chunk can be per-platform. Do not use this value to index the Sections array, only the user
+	 * section data should be index by this value.
+	 */
+	int32 OriginalDataSectionIndex;
+
+	/*
+	 * If this section was produce because of BONE chunking, the parent section index will be valid.
+	 * If the section is not the result of skin vertex chunking, this value will be INDEX_NONE.
+	 * Use this value to know if the section was BONE chunked:
+	 * if(ChunkedParentSectionIndex != INDEX_NONE) will be true if the section is BONE chunked
+	 */
+	int32 ChunkedParentSectionIndex;
+
+
+
 	FSkelMeshSection()
 		: MaterialIndex(0)
 		, BaseIndex(0)
@@ -141,7 +159,9 @@ struct FSkelMeshSection
 		, MaxBoneInfluences(4)
 		, CorrespondClothAssetIndex(INDEX_NONE)
 		, bDisabled(false)
-		, GenerateUpToLodIndex(-1)
+		, GenerateUpToLodIndex(INDEX_NONE)
+		, OriginalDataSectionIndex(INDEX_NONE)
+		, ChunkedParentSectionIndex(INDEX_NONE)
 	{}
 
 
@@ -185,6 +205,64 @@ struct FSkelMeshSection
 	friend FArchive& operator<<(FArchive& Ar, FSkelMeshSection& S);
 };
 
+/**
+* Structure containing all the section data a user can change.
+*
+* Some section data also impact dependent generated LOD, those member should be add to the DDC Key
+* and trig a rebuild if they are change.
+*/
+struct FSkelMeshSourceSectionUserData
+{
+	/** This section will recompute tangent in runtime */
+	bool bRecomputeTangent;
+
+	/** This section will cast shadow */
+	bool bCastShadow;
+
+	// INDEX_NONE if not set
+	int16 CorrespondClothAssetIndex;
+
+	/** Clothing data for this section, clothing is only present if ClothingData.IsValid() returns true */
+	FClothingSectionData ClothingData;
+
+
+	//////////////////////////////////////////////////////////////////////////
+	//Skeletalmesh DDC key members, Add sections member that impact generated skel mesh here
+
+	/** If disabled, we won't render this section */
+	bool bDisabled;
+
+	/*
+	 * The LOD index at which any generated lower quality LODs will include this section.
+	 * A value of -1 mean the section will always be include when generating a LOD
+	 */
+	int32 GenerateUpToLodIndex;
+
+	// End DDC members
+	//////////////////////////////////////////////////////////////////////////
+
+
+
+	FSkelMeshSourceSectionUserData()
+		: bRecomputeTangent(false)
+		, bCastShadow(true)
+		, CorrespondClothAssetIndex(INDEX_NONE)
+		, bDisabled(false)
+		, GenerateUpToLodIndex(INDEX_NONE)
+	{}
+
+	/**
+	* @return TRUE if we have cloth data for this section
+	*/
+	FORCEINLINE bool HasClothingData() const
+	{
+		return (ClothingData.AssetGuid.IsValid());
+	}
+
+	// Serialization.
+	friend FArchive& operator<<(FArchive& Ar, FSkelMeshSourceSectionUserData& S);
+};
+
 /**
 * All data to define a certain LOD model for a skeletal mesh.
 */
@@ -194,6 +272,14 @@ public:
 	/** Sections. */
 	TArray<FSkelMeshSection> Sections;
 
+	/*
+	 * When user change section data in the UI, we store it here to be able to regenerate the changes
+	 * Note: the key (int32) is the original imported section data, because of BONE chunk the size of
+	 * this array is not the same as the Sections array. Use the section's OriginalDataSectionIndex to
+	 * index it.
+	 */
+	TMap<int32, FSkelMeshSourceSectionUserData> UserSectionsData;
+
 	uint32						NumVertices;
 	/** The number of unique texture coordinate sets in this lod */
 	uint32						NumTexCoords;
@@ -234,6 +320,7 @@ public:
 		: NumVertices(0)
 		, NumTexCoords(0)
 		, MaxImportVertex(-1)
+		, BuildStringID(TEXT(""))
 	{
 	}
 
@@ -269,7 +356,7 @@ public:
 	/**
 	* @return true if any chunks have cloth data.
 	*/
-	bool HasClothData() const;
+	ENGINE_API bool HasClothData() const;
 
 	ENGINE_API int32 NumNonClothingSections() const;
 
@@ -286,6 +373,21 @@ public:
 	ENGINE_API bool DoSectionsNeedExtraBoneInfluences() const;
 
 	void GetResourceSizeEx(FResourceSizeEx& CumulativeResourceSize) const;
+
+	/**
+	* Make sure user section data is present for every original sections
+	*/
+	ENGINE_API void SyncronizeUserSectionsDataArray(bool bResetNonUsedSection = false);
+
+	//Temporary build String ID
+	//We use this string to store the LOD model data so we can know if the LOD need to be rebuild
+	//This GUID is set when we Cache the render data (build function)
+	mutable FString BuildStringID;
+
+	/**
+	* Build a derive data key with the user section data (UserSectionsData) and the original bulk data
+	*/
+	ENGINE_API FString GetLODModelDeriveDataKey() const;
 };
 
 #endif // WITH_EDITOR
\ No newline at end of file
diff --git a/Engine/Source/Runtime/Engine/Public/Rendering/SkeletalMeshModel.h b/Engine/Source/Runtime/Engine/Public/Rendering/SkeletalMeshModel.h
--- a/Engine/Source/Runtime/Engine/Public/Rendering/SkeletalMeshModel.h
+++ b/Engine/Source/Runtime/Engine/Public/Rendering/SkeletalMeshModel.h
@@ -42,6 +42,10 @@ public:
 	/** Get current GUID Id as a string, for DDC key */
 	FString GetIdString() const;
 
+	void SyncronizeLODUserSectionsData();
+
+	FString GetLODModelIdString() const;
+
 	ENGINE_API void EmptyOriginalReductionSourceMeshData();
 	/* When user reduce an imported LOD with himself (BaseLOD == TargetLOD), we need to store some imported model data so we can reduce again from the same data.*/
 	TArray<FReductionBaseSkeletalMeshBulkData*> OriginalReductionSourceMeshData;
diff --git a/Engine/Source/Runtime/Engine/Public/SkeletalMeshTypes.h b/Engine/Source/Runtime/Engine/Public/SkeletalMeshTypes.h
--- a/Engine/Source/Runtime/Engine/Public/SkeletalMeshTypes.h
+++ b/Engine/Source/Runtime/Engine/Public/SkeletalMeshTypes.h
@@ -5,6 +5,7 @@
 #include "CoreMinimal.h"
 #include "PrimitiveSceneProxy.h"
 #include "Materials/MaterialInterface.h"
+#include "ComponentReregisterContext.h"
 
 class FMaterialRenderProxy;
 class FMeshElementCollector;
@@ -122,7 +123,8 @@ struct ESkeletalMeshVertexFlags
 		None = 0x0,
 		UseFullPrecisionUVs = 0x1,
 		HasVertexColors = 0x2,
-		UseHighPrecisionTangentBasis = 0x4
+		UseHighPrecisionTangentBasis = 0x4,
+		BuildAdjacencyIndexBuffer = 0x8
 	};
 };
 
@@ -414,3 +416,35 @@ private:
 	bool bRefreshBounds;
 };
 
+#if WITH_EDITOR
+
+//Helper to scope skeletal mesh post edit change.
+class ENGINE_API FScopedSkeletalMeshPostEditChange
+{
+public:
+	/*
+	 * This constructor increment the skeletal mesh PostEditChangeStackCounter. If the stack counter is zero before the increment
+	 * the skeletal mesh component will be unregister from the world. The component will also release there rendering resources.
+	 * Parameters:
+	 * @param InbCallPostEditChange - if we are the first scope PostEditChange will be called.
+	 * @param InbReregisterComponents - if we are the first scope we will re register component from world and also component render data.
+	 */
+	FScopedSkeletalMeshPostEditChange(USkeletalMesh* InSkeletalMesh, bool InbCallPostEditChange = true, bool InbReregisterComponents = true);
+
+	/*
+	 * This destructor decrement the skeletal mesh PostEditChangeStackCounter. If the stack counter is zero after the decrement,
+	 * the skeletal mesh PostEditChange will be call. The component will also be register to the world and there render data resources will be rebuild.
+	 */
+	~FScopedSkeletalMeshPostEditChange();
+
+	void SetSkeletalMesh(USkeletalMesh* InSkeletalMesh);
+
+private:
+	USkeletalMesh* SkeletalMesh;
+	bool bReregisterComponents;
+	bool bCallPostEditChange;
+	FSkinnedMeshComponentRecreateRenderStateContext* RecreateExistingRenderStateContext;
+	TIndirectArray<FComponentReregisterContext> ComponentReregisterContexts;
+};
+
+#endif
diff --git a/Engine/Source/Runtime/MeshDescription/MeshDescription.Build.cs b/Engine/Source/Runtime/MeshDescription/MeshDescription.Build.cs
--- a/Engine/Source/Runtime/MeshDescription/MeshDescription.Build.cs
+++ b/Engine/Source/Runtime/MeshDescription/MeshDescription.Build.cs
@@ -9,11 +9,11 @@ namespace UnrealBuildTool.Rules
             PrivateIncludePaths.Add("Runtime/MeshDescription/Private");
             PublicIncludePaths.Add("Runtime/MeshDescription/Public");
 
-			PrivateDependencyModuleNames.AddRange(
+			PublicDependencyModuleNames.AddRange(
 				new string[]
 				{
 					"Core",
-					"CoreUObject",
+					"CoreUObject"
 				}
 			);
 		}
diff --git a/Engine/Source/Runtime/MeshDescription/Private/MeshAttributes.cpp b/Engine/Source/Runtime/MeshDescription/Private/MeshAttributes.cpp
--- a/Engine/Source/Runtime/MeshDescription/Private/MeshAttributes.cpp
+++ b/Engine/Source/Runtime/MeshDescription/Private/MeshAttributes.cpp
@@ -4,26 +4,12 @@
 
 namespace MeshAttribute
 {
-	const FName Vertex::Position( "Position ");
-	const FName Vertex::CornerSharpness( "CornerSharpness" );
-
-	const FName VertexInstance::TextureCoordinate( "TextureCoordinate" );
-	const FName VertexInstance::Normal( "Normal" );
-	const FName VertexInstance::Tangent( "Tangent" );
-	const FName VertexInstance::BinormalSign( "BinormalSign" );
-	const FName VertexInstance::Color( "Color" );
-
-	const FName Edge::IsHard( "IsHard" );
-	const FName Edge::IsUVSeam( "IsUVSeam" );
-	const FName Edge::CreaseSharpness( "CreaseSharpness" );
+	const FName Vertex::Position("Position ");
+}
 
-	const FName Polygon::Normal( "Normal" );
-	const FName Polygon::Tangent( "Tangent" );
-	const FName Polygon::Binormal( "Binormal" );
-	const FName Polygon::Center( "Center" );
 
-	const FName PolygonGroup::ImportedMaterialSlotName( "ImportedMaterialSlotName" );
-	const FName PolygonGroup::EnableCollision( "EnableCollision" );
-	const FName PolygonGroup::CastShadow( "CastShadow" );
+void FMeshAttributes::Register()
+{
+	// Nothing to do here: Vertex positions are already registered by the FMeshDescription constructor
 }
 
