Made bugit command also copy the bugitgo location to the clipboard

diff --git a/Engine/Source/Runtime/Engine/Private/CheatManager.cpp b/Engine/Source/Runtime/Engine/Private/CheatManager.cpp
--- a/Engine/Source/Runtime/Engine/Private/CheatManager.cpp
+++ b/Engine/Source/Runtime/Engine/Private/CheatManager.cpp
@@ -1189,6 +1189,10 @@ void UCheatManager::LogOutBugItGoToLogFile( const FString& InScreenShotDesc, con
 	//delete OutputFile;
 	OutputFile.TearDown();
 
+#if PLATFORM_DESKTOP
+	FPlatformMisc::ClipboardCopy(*InGoString);
+#endif
+
 	// so here we want to send this bad boy back to the PC
 	SendDataToPCViaUnrealConsole( TEXT("UE_PROFILER!BUGIT:"), *(FullFileName) );
 #endif // ALLOW_DEBUG_FILES
