Make sure the streamer track correct texel factors for meshes

diff --git a/Engine/Source/Runtime/Engine/Private/Components/StaticMeshComponent.cpp b/Engine/Source/Runtime/Engine/Private/Components/StaticMeshComponent.cpp
--- a/Engine/Source/Runtime/Engine/Private/Components/StaticMeshComponent.cpp
+++ b/Engine/Source/Runtime/Engine/Private/Components/StaticMeshComponent.cpp
@@ -865,7 +865,9 @@ void UStaticMeshComponent::GetStreamingRenderAssetInfo(FStreamingTextureLevelCon
 
 	if (IsStreamingRenderAsset(GetStaticMesh()))
 	{
-		const float TexelFactor = ForcedLodModel > 0 ? -(GetStaticMesh()->RenderData->LODResources.Num() - ForcedLodModel + 1) : Bounds.SphereRadius * 2.f;
+		const float TexelFactor = ForcedLodModel > 0 ?
+			-(GetStaticMesh()->RenderData->LODResources.Num() - ForcedLodModel + 1) :
+			(IsRegistered() ? Bounds.SphereRadius * 2.f : 0.f);
 		new (OutStreamingRenderAssets) FStreamingRenderAssetPrimitiveInfo(GetStaticMesh(), Bounds, TexelFactor, PackedRelativeBox_Identity, true);
 	}
 }
diff --git a/Engine/Source/Runtime/Engine/Private/Streaming/TextureInstanceView.cpp b/Engine/Source/Runtime/Engine/Private/Streaming/TextureInstanceView.cpp
--- a/Engine/Source/Runtime/Engine/Private/Streaming/TextureInstanceView.cpp
+++ b/Engine/Source/Runtime/Engine/Private/Streaming/TextureInstanceView.cpp
@@ -556,6 +556,15 @@ void FRenderAssetInstanceAsyncView::GetRenderAssetScreenSize(
 					const FRenderAssetInstanceView::FCompiledElement& CompiledElement = CompiledElementData[CompiledElementIndex];
 					if (ensure(BoundsViewInfo.IsValidIndex(CompiledElement.BoundsIndex)))
 					{
+						// Texel factor wasn't available because the component wasn't registered. Lazy initialize it now.
+						if (AssetType != FStreamingRenderAsset::AT_Texture
+							&& CompiledElement.TexelFactor == 0.f
+							&& ensure(CompiledElement.BoundsIndex < View->NumBounds4() * 4))
+						{
+							FRenderAssetInstanceView::FCompiledElement* MutableCompiledElement = const_cast<FRenderAssetInstanceView::FCompiledElement*>(CompiledElementData);
+							MutableCompiledElement->TexelFactor = View->GetBounds4(CompiledElement.BoundsIndex / 4).Radius.Component(CompiledElement.BoundsIndex % 4) * 2.f;
+						}
+
 						ProcessElement(
 							AssetType,
 							BoundsViewInfo[CompiledElement.BoundsIndex],
