Fix deprecation warning

diff --git a/Engine/Source/Runtime/Engine/Private/CheatManager.cpp b/Engine/Source/Runtime/Engine/Private/CheatManager.cpp
--- a/Engine/Source/Runtime/Engine/Private/CheatManager.cpp
+++ b/Engine/Source/Runtime/Engine/Private/CheatManager.cpp
@@ -38,6 +38,7 @@
 #include "GameFramework/InputSettings.h"
 #include "Misc/CoreDelegates.h"
 #include "Engine/NetConnection.h"
+#include "HAL/PlatformApplicationMisc.h"
 
 DEFINE_LOG_CATEGORY_STATIC(LogCheatManager, Log, All);
 
@@ -1190,7 +1191,7 @@ void UCheatManager::LogOutBugItGoToLogFile( const FString& InScreenShotDesc, con
 	OutputFile.TearDown();
 
 #if PLATFORM_DESKTOP
-	FPlatformMisc::ClipboardCopy(*InGoString);
+	FPlatformApplicationMisc::ClipboardCopy(*InGoString);
 #endif
 
 	// so here we want to send this bad boy back to the PC
