move the cpp file outside of public folder into the private rendering folder
#rnx

diff --git a/Engine/Source/Runtime/Engine/Private/Rendering/SkeletalMeshLODImporterData.cpp b/Engine/Source/Runtime/Engine/Private/Rendering/SkeletalMeshLODImporterData.cpp
new file mode 100644
--- /dev/null
+++ b/Engine/Source/Runtime/Engine/Private/Rendering/SkeletalMeshLODImporterData.cpp
@@ -0,0 +1,874 @@
+// Copyright 1998-2019 Epic Games, Inc. All Rights Reserved.
+
+#if WITH_EDITOR
+
+#include "Rendering/SkeletalMeshLODImporterData.h"
+#include "Serialization/BulkDataWriter.h"
+#include "Serialization/BulkDataReader.h"
+#include "Rendering/SkeletalMeshModel.h"
+#include "Engine/SkeletalMesh.h"
+#include "Factories/FbxSkeletalMeshImportData.h"
+
+void FSkeletalMeshImportData::CopyDataNeedByMorphTargetImport(FSkeletalMeshImportData& Other) const
+{
+	Other.Points = Points;
+	Other.PointToRawMap = PointToRawMap;
+	Other.bDiffPose = bDiffPose;
+	Other.bUseT0AsRefPose = bUseT0AsRefPose;
+}
+
+void FSkeletalMeshImportData::KeepAlternateSkinningBuildDataOnly()
+{
+	//No need of any alternate restore data, since we are this data if this function is called
+	AlternateInfluenceProfileNames.Empty();
+	AlternateInfluences.Empty();
+	
+	//No need of the morph target restore data
+	MorphTargetModifiedPoints.Empty();
+	MorphTargetNames.Empty();
+	MorphTargets.Empty();
+
+	//Remove material array and PointToRawMap
+	Materials.Empty();
+	PointToRawMap.Empty();
+}
+
+/**
+* Takes an imported bone name, removes any leading or trailing spaces, and converts the remaining spaces to dashes.
+*/
+FString FSkeletalMeshImportData::FixupBoneName(FString BoneName)
+{
+	BoneName.TrimStartAndEndInline();
+	BoneName.ReplaceInline(TEXT(" "), TEXT("-"), ESearchCase::IgnoreCase);
+	return BoneName;
+}
+
+/**
+* Copy mesh data for importing a single LOD
+*
+* @param LODPoints - vertex data.
+* @param LODWedges - wedge information to static LOD level.
+* @param LODFaces - triangle/ face data to static LOD level.
+* @param LODInfluences - weights/ influences to static LOD level.
+*/
+void FSkeletalMeshImportData::CopyLODImportData(
+	TArray<FVector>& LODPoints,
+	TArray<SkeletalMeshImportData::FMeshWedge>& LODWedges,
+	TArray<SkeletalMeshImportData::FMeshFace>& LODFaces,
+	TArray<SkeletalMeshImportData::FVertInfluence>& LODInfluences,
+	TArray<int32>& LODPointToRawMap) const
+{
+	// Copy vertex data.
+	LODPoints.Empty(Points.Num());
+	LODPoints.AddUninitialized(Points.Num());
+	for (int32 p = 0; p < Points.Num(); p++)
+	{
+		LODPoints[p] = Points[p];
+	}
+
+	// Copy wedge information to static LOD level.
+	LODWedges.Empty(Wedges.Num());
+	LODWedges.AddUninitialized(Wedges.Num());
+	for (int32 w = 0; w < Wedges.Num(); w++)
+	{
+		LODWedges[w].iVertex = Wedges[w].VertexIndex;
+		// Copy all texture coordinates
+		FMemory::Memcpy(LODWedges[w].UVs, Wedges[w].UVs, sizeof(FVector2D) * MAX_TEXCOORDS);
+		LODWedges[w].Color = Wedges[w].Color;
+
+	}
+
+	// Copy triangle/ face data to static LOD level.
+	LODFaces.Empty(Faces.Num());
+	LODFaces.AddUninitialized(Faces.Num());
+	for (int32 f = 0; f < Faces.Num(); f++)
+	{
+		SkeletalMeshImportData::FMeshFace Face;
+		Face.iWedge[0] = Faces[f].WedgeIndex[0];
+		Face.iWedge[1] = Faces[f].WedgeIndex[1];
+		Face.iWedge[2] = Faces[f].WedgeIndex[2];
+		Face.MeshMaterialIndex = Faces[f].MatIndex;
+
+		Face.TangentX[0] = Faces[f].TangentX[0];
+		Face.TangentX[1] = Faces[f].TangentX[1];
+		Face.TangentX[2] = Faces[f].TangentX[2];
+
+		Face.TangentY[0] = Faces[f].TangentY[0];
+		Face.TangentY[1] = Faces[f].TangentY[1];
+		Face.TangentY[2] = Faces[f].TangentY[2];
+
+		Face.TangentZ[0] = Faces[f].TangentZ[0];
+		Face.TangentZ[1] = Faces[f].TangentZ[1];
+		Face.TangentZ[2] = Faces[f].TangentZ[2];
+
+		Face.SmoothingGroups = Faces[f].SmoothingGroups;
+
+		LODFaces[f] = Face;
+	}
+
+	// Copy weights/ influences to static LOD level.
+	LODInfluences.Empty(Influences.Num());
+	LODInfluences.AddUninitialized(Influences.Num());
+	for (int32 i = 0; i < Influences.Num(); i++)
+	{
+		LODInfluences[i].Weight = Influences[i].Weight;
+		LODInfluences[i].VertIndex = Influences[i].VertexIndex;
+		LODInfluences[i].BoneIndex = Influences[i].BoneIndex;
+	}
+
+	// Copy mapping
+	LODPointToRawMap = PointToRawMap;
+}
+
+bool FSkeletalMeshImportData::ReplaceSkeletalMeshGeometryImportData(const USkeletalMesh* SkeletalMesh, FSkeletalMeshImportData* ImportData, int32 LodIndex)
+{
+	FSkeletalMeshModel *ImportedResource = SkeletalMesh->GetImportedModel();
+	check(ImportedResource && ImportedResource->LODModels.IsValidIndex(LodIndex));
+	FSkeletalMeshLODModel& SkeletalMeshLODModel = ImportedResource->LODModels[LodIndex];
+
+	const FSkeletalMeshLODInfo* LodInfo = SkeletalMesh->GetLODInfo(LodIndex);
+	check(LodInfo);
+
+	//Load the original skeletal mesh import data
+	FSkeletalMeshImportData OriginalSkeletalMeshImportData;
+	SkeletalMeshLODModel.RawSkeletalMeshBulkData.LoadRawMesh(OriginalSkeletalMeshImportData);
+
+	//Backup the new geometry and rig to be able to apply the rig to the old geometry
+	FSkeletalMeshImportData NewGeometryAndRigData = *ImportData;
+
+	ImportData->bHasNormals = OriginalSkeletalMeshImportData.bHasNormals;
+	ImportData->bHasTangents = OriginalSkeletalMeshImportData.bHasTangents;
+	ImportData->bHasVertexColors = OriginalSkeletalMeshImportData.bHasVertexColors;
+	ImportData->NumTexCoords = OriginalSkeletalMeshImportData.NumTexCoords;
+
+	ImportData->Materials.Reset();
+	ImportData->Points.Reset();
+	ImportData->Faces.Reset();
+	ImportData->Wedges.Reset();
+	ImportData->PointToRawMap.Reset();
+
+	//Material is a special case since we cannot serialize the UMaterialInstance when saving the RawSkeletalMeshBulkData
+	//So it has to be reconstructed.
+	ImportData->MaxMaterialIndex = 0;
+	for (int32 MaterialIndex = 0; MaterialIndex < SkeletalMesh->Materials.Num(); ++MaterialIndex)
+	{
+		SkeletalMeshImportData::FMaterial NewMaterial;
+
+		NewMaterial.MaterialImportName = SkeletalMesh->Materials[MaterialIndex].ImportedMaterialSlotName.ToString();
+		NewMaterial.Material = SkeletalMesh->Materials[MaterialIndex].MaterialInterface;
+		// Add an entry for each unique material
+		ImportData->MaxMaterialIndex = FMath::Max(ImportData->MaxMaterialIndex, (uint32)(ImportData->Materials.Add(NewMaterial)));
+	}
+
+	ImportData->NumTexCoords = OriginalSkeletalMeshImportData.NumTexCoords;
+	ImportData->Points += OriginalSkeletalMeshImportData.Points;
+	ImportData->Faces += OriginalSkeletalMeshImportData.Faces;
+	ImportData->Wedges += OriginalSkeletalMeshImportData.Wedges;
+	ImportData->PointToRawMap += OriginalSkeletalMeshImportData.PointToRawMap;
+
+	return ImportData->ApplyRigToGeo(NewGeometryAndRigData);
+}
+
+bool FSkeletalMeshImportData::ReplaceSkeletalMeshRigImportData(const USkeletalMesh* SkeletalMesh, FSkeletalMeshImportData* ImportData, int32 LodIndex)
+{
+	FSkeletalMeshModel *ImportedResource = SkeletalMesh->GetImportedModel();
+	check(ImportedResource && ImportedResource->LODModels.IsValidIndex(LodIndex));
+	FSkeletalMeshLODModel& SkeletalMeshLODModel = ImportedResource->LODModels[LodIndex];
+
+	const FSkeletalMeshLODInfo* LodInfo = SkeletalMesh->GetLODInfo(LodIndex);
+	check(LodInfo);
+
+	//Load the original skeletal mesh import data
+	FSkeletalMeshImportData OriginalSkeletalMeshImportData;
+	SkeletalMeshLODModel.RawSkeletalMeshBulkData.LoadRawMesh(OriginalSkeletalMeshImportData);
+
+	ImportData->bDiffPose = OriginalSkeletalMeshImportData.bDiffPose;
+	ImportData->bUseT0AsRefPose = OriginalSkeletalMeshImportData.bUseT0AsRefPose;
+
+	ImportData->RefBonesBinary.Reset();
+
+	ImportData->RefBonesBinary += OriginalSkeletalMeshImportData.RefBonesBinary;
+
+	//Fix the old rig to match the new geometry
+	return ImportData->ApplyRigToGeo(OriginalSkeletalMeshImportData);
+}
+
+bool FSkeletalMeshImportData::ApplyRigToGeo(FSkeletalMeshImportData& Other)
+{
+	//Reset the influence, we will regenerate it from the other data (the incoming rig)
+	Influences.Reset();
+
+	FWedgePosition OldGeoOverlappingPosition;
+	FWedgePosition::FillWedgePosition(OldGeoOverlappingPosition, Other.Points, Other.Wedges, THRESH_POINTS_ARE_SAME);
+	FOctreeQueryHelper OctreeQueryHelper(OldGeoOverlappingPosition.GetOctree());
+
+	int32 NewWedgesNum = Wedges.Num();
+	int32 OldWedgesNum = Other.Wedges.Num();
+
+	//////////////////////////////////////////////////////////////////////////
+	// Found the Remapping between old vertex index and new vertex index
+	// The old vertex index are the key, the index of the first array
+	// The new vertex indexes are the second array, because we can map many
+	// new vertex to one old vertex
+	//
+	// All new wedges get remap to a old wedge index, so we can be sure that all
+	// new vertex will have correct bone weight apply to them.
+	TArray<TArray<int32>> OldToNewRemap;
+	OldToNewRemap.AddDefaulted(Other.Points.Num());
+	for (int32 WedgeIndex = 0; WedgeIndex < NewWedgesNum; ++WedgeIndex)
+	{
+		const FVector2D& CurWedgeUV = Wedges[WedgeIndex].UVs[0];
+		int32 NewVertexIndex = (int32)(Wedges[WedgeIndex].VertexIndex);
+		int32 NewFaceIndex = (WedgeIndex / 3);
+		int32 NewFaceCorner = (WedgeIndex % 3);
+		FVector NewNormal = Faces[NewFaceIndex].TangentZ[NewFaceCorner];
+		TArray<int32> OldWedgeIndexes;
+		OldGeoOverlappingPosition.FindMatchingPositionWegdeIndexes(Points[NewVertexIndex], THRESH_POINTS_ARE_SAME, OldWedgeIndexes);
+		bool bFoundMatch = false;
+		if (OldWedgeIndexes.Num() > 0)
+		{
+			for (int32 OldWedgeIndex : OldWedgeIndexes)
+			{
+				int32 OldVertexIndex = Other.Wedges[OldWedgeIndex].VertexIndex;
+				int32 OldFaceIndex = (OldWedgeIndex / 3);
+				int32 OldFaceCorner = (OldWedgeIndex % 3);
+				FVector OldNormal = Other.Faces[OldFaceIndex].TangentZ[OldFaceCorner];
+				
+				if (Other.Wedges[OldWedgeIndex].UVs[0].Equals(CurWedgeUV, THRESH_UVS_ARE_SAME)
+					&& OldNormal.Equals(NewNormal, THRESH_NORMALS_ARE_SAME))
+				{
+					OldToNewRemap[OldVertexIndex].AddUnique(NewVertexIndex);
+					bFoundMatch = true;
+				}
+			}
+		}
+
+		//If some geometry was added, it will not found any exact match with the old geometry
+		//In this case we have to find the nearest list of wedge indexe
+		if(!bFoundMatch)
+		{
+			TArray<FWedgeInfo> NearestWedges;
+			FVector SearchPosition = Points[NewVertexIndex];
+			OctreeQueryHelper.FindNearestWedgeIndexes(SearchPosition, NearestWedges);
+			//The best old wedge match is base on those weight ratio
+			const int32 UVWeightRatioIndex = 0;
+			const int32 NormalWeightRatioIndex = 1;
+			const float MatchWeightRatio[3] = { 0.99f, 0.01f };
+			if (NearestWedges.Num() > 0)
+			{
+				int32 BestOldVertexIndex = INDEX_NONE;
+				float MaxUVDistance = 0.0f;
+				float MaxNormalDelta = 0.0f;
+				TArray<float> UvDistances;
+				UvDistances.Reserve(NearestWedges.Num());
+				TArray<float> NormalDeltas;
+				NormalDeltas.Reserve(NearestWedges.Num());
+				for (const FWedgeInfo& WedgeInfo : NearestWedges)
+				{
+					int32 OldWedgeIndex = WedgeInfo.WedgeIndex;
+					int32 OldVertexIndex = Other.Wedges[OldWedgeIndex].VertexIndex;
+					int32 OldFaceIndex = (OldWedgeIndex / 3);
+					int32 OldFaceCorner = (OldWedgeIndex % 3);
+					const FVector2D& OldUV = Other.Wedges[OldWedgeIndex].UVs[0];
+					const FVector& OldNormal = Other.Faces[OldFaceIndex].TangentZ[OldFaceCorner];
+					float UVDelta = FVector2D::DistSquared(CurWedgeUV, OldUV);
+					float NormalDelta = FMath::Abs(FMath::Acos(FVector::DotProduct(NewNormal, OldNormal)));
+					if (UVDelta > MaxUVDistance)
+					{
+						MaxUVDistance = UVDelta;
+					}
+					UvDistances.Add(UVDelta);
+					if (NormalDelta > MaxNormalDelta)
+					{
+						MaxNormalDelta = NormalDelta;
+					}
+					NormalDeltas.Add(NormalDelta);
+				}
+				float BestContribution = 0.0f;
+				for (int32 NearestWedgeIndex = 0; NearestWedgeIndex < UvDistances.Num(); ++NearestWedgeIndex)
+				{
+					float Contribution = ((MaxUVDistance - UvDistances[NearestWedgeIndex])/MaxUVDistance)*MatchWeightRatio[UVWeightRatioIndex];
+					Contribution += ((MaxNormalDelta - NormalDeltas[NearestWedgeIndex]) / MaxNormalDelta)*MatchWeightRatio[NormalWeightRatioIndex];
+					if (Contribution > BestContribution)
+					{
+						BestContribution = Contribution;
+						BestOldVertexIndex = Other.Wedges[NearestWedges[NearestWedgeIndex].WedgeIndex].VertexIndex;
+					}
+				}
+				if (BestOldVertexIndex == INDEX_NONE)
+				{
+					//Use the first NearestWedges entry, we end up here because all NearestWedges entries all equals, so the ratio will be zero in such a case
+					BestOldVertexIndex = Other.Wedges[NearestWedges[0].WedgeIndex].VertexIndex;
+				}
+				OldToNewRemap[BestOldVertexIndex].AddUnique(NewVertexIndex);
+			}
+		}
+	}
+
+	for (int32 InfluenceIndex = 0; InfluenceIndex < Other.Influences.Num(); ++InfluenceIndex)
+	{
+		int32 OldPointIndex = Other.Influences[InfluenceIndex].VertexIndex;
+
+		const TArray<int32>& NewInfluenceVertexIndexes = OldToNewRemap[OldPointIndex];
+
+		for (int32 NewPointIdx : NewInfluenceVertexIndexes)
+		{
+			SkeletalMeshImportData::FRawBoneInfluence& RawBoneInfluence = Influences.AddDefaulted_GetRef();
+			RawBoneInfluence.BoneIndex = Other.Influences[InfluenceIndex].BoneIndex;
+			RawBoneInfluence.Weight = Other.Influences[InfluenceIndex].Weight;
+			RawBoneInfluence.VertexIndex = NewPointIdx;
+		}
+	}
+
+	return true;
+}
+
+/**
+* Serialization of raw meshes uses its own versioning scheme because it is
+* stored in bulk data.
+*/
+enum
+{
+	// Engine raw mesh version:
+	REDUCTION_BASE_SK_DATA_BULKDATA_VER_INITIAL = 0,
+	
+	//////////////////////////////////////////////////////////////////////////
+	// Add new raw mesh versions here.
+
+	REDUCTION_BASE_SK_DATA_BULKDATA_VER_PLUS_ONE,
+	REDUCTION_BASE_SK_DATA_BULKDATA_VER = REDUCTION_BASE_SK_DATA_BULKDATA_VER_PLUS_ONE - 1,
+
+	// Licensee raw mesh version:
+	REDUCTION_BASE_SK_DATA_BULKDATA_LIC_VER_INITIAL = 0,
+	
+	//////////////////////////////////////////////////////////////////////////
+	// Licensees add new raw mesh versions here.
+
+	REDUCTION_BASE_SK_DATA_BULKDATA_LIC_VER_PLUS_ONE,
+	REDUCTION_BASE_SK_DATA_BULKDATA_LIC_VER = REDUCTION_BASE_SK_DATA_BULKDATA_LIC_VER_PLUS_ONE - 1
+};
+
+struct FReductionSkeletalMeshData
+{
+	FReductionSkeletalMeshData(FSkeletalMeshLODModel& InBaseLODModel, TMap<FString, TArray<FMorphTargetDelta>>& InBaseLODMorphTargetData, UObject* InOwner)
+		: BaseLODModel(InBaseLODModel)
+		, BaseLODMorphTargetData(InBaseLODMorphTargetData)
+		, Owner(InOwner)
+	{
+	}
+
+	FSkeletalMeshLODModel& BaseLODModel;
+	TMap<FString, TArray<FMorphTargetDelta>>& BaseLODMorphTargetData;
+	UObject* Owner;
+};
+
+FArchive& operator<<(FArchive& Ar, FReductionSkeletalMeshData& ReductionSkeletalMeshData)
+{
+	int32 Version = REDUCTION_BASE_SK_DATA_BULKDATA_VER;
+	int32 LicenseeVersion = REDUCTION_BASE_SK_DATA_BULKDATA_LIC_VER;
+	Ar << Version;
+	Ar << LicenseeVersion;
+	ReductionSkeletalMeshData.BaseLODModel.Serialize(Ar, ReductionSkeletalMeshData.Owner, INDEX_NONE);
+	Ar << ReductionSkeletalMeshData.BaseLODMorphTargetData;
+	return Ar;
+}
+
+FReductionBaseSkeletalMeshBulkData::FReductionBaseSkeletalMeshBulkData()
+{
+}
+
+void FReductionBaseSkeletalMeshBulkData::Serialize(FArchive& Ar, TArray<FReductionBaseSkeletalMeshBulkData*>& ReductionBaseSkeletalMeshDatas, UObject* Owner)
+{
+	Ar.CountBytes(ReductionBaseSkeletalMeshDatas.Num() * sizeof(FReductionBaseSkeletalMeshBulkData), ReductionBaseSkeletalMeshDatas.Num() * sizeof(FReductionBaseSkeletalMeshBulkData));
+	if (Ar.IsLoading())
+	{
+		// Load array.
+		int32 NewNum;
+		Ar << NewNum;
+		ReductionBaseSkeletalMeshDatas.Empty(NewNum);
+		for (int32 Index = 0; Index < NewNum; Index++)
+		{
+			FReductionBaseSkeletalMeshBulkData* EmptyData = new FReductionBaseSkeletalMeshBulkData();
+			int32 NewEntryIndex = ReductionBaseSkeletalMeshDatas.Add(EmptyData);
+			check(NewEntryIndex == Index);
+			ReductionBaseSkeletalMeshDatas[Index]->Serialize(Ar, Owner);
+		}
+	}
+	else
+	{
+		// Save array.
+		int32 Num = ReductionBaseSkeletalMeshDatas.Num();
+		Ar << Num;
+		for (int32 Index = 0; Index < Num; Index++)
+		{
+			(ReductionBaseSkeletalMeshDatas)[Index]->Serialize(Ar, Owner);
+		}
+	}
+}
+
+void FReductionBaseSkeletalMeshBulkData::Serialize(FArchive& Ar, UObject* Owner)
+{
+	if (Ar.IsTransacting())
+	{
+		// If transacting, keep these members alive the other side of an undo, otherwise their values will get lost
+		SerializeLoadingCustomVersionContainer.Serialize(Ar);
+		Ar << bUseSerializeLoadingCustomVersion;
+	}
+	else
+	{
+		if (Ar.IsLoading())
+		{
+			//Save the custom version so we can load FReductionSkeletalMeshData later
+			SerializeLoadingCustomVersionContainer = Ar.GetCustomVersions();
+			bUseSerializeLoadingCustomVersion = true;
+		}
+
+		if (Ar.IsSaving() && bUseSerializeLoadingCustomVersion == true)
+		{
+			//We need to update the FReductionSkeletalMeshData serialize version to the latest in case we save the Parent bulkdata
+			FSkeletalMeshLODModel BaseLODModel;
+			TMap<FString, TArray<FMorphTargetDelta>> BaseLODMorphTargetData;
+			LoadReductionData(BaseLODModel, BaseLODMorphTargetData, Owner);
+			SaveReductionData(BaseLODModel, BaseLODMorphTargetData, Owner);
+		}
+	}
+
+	BulkData.Serialize(Ar, Owner);
+}
+
+void FReductionBaseSkeletalMeshBulkData::SaveReductionData(FSkeletalMeshLODModel& BaseLODModel, TMap<FString, TArray<FMorphTargetDelta>>& BaseLODMorphTargetData, UObject* Owner)
+{
+	//Saving the bulk data mean we do not need anymore the SerializeLoadingCustomVersionContainer of the parent bulk data
+	SerializeLoadingCustomVersionContainer.Empty();
+	bUseSerializeLoadingCustomVersion = false;
+	FReductionSkeletalMeshData ReductionSkeletalMeshData(BaseLODModel, BaseLODMorphTargetData, Owner);
+
+	//Clear the bulk data before writing it
+	BulkData.RemoveBulkData();
+
+	// Get a lock on the bulk data
+	{
+		const bool bIsPersistent = true;
+		FBulkDataWriter Ar(BulkData, bIsPersistent);
+		Ar << ReductionSkeletalMeshData;
+
+		// Preserve CustomVersions at save time so we can reuse the same ones when reloading direct from memory
+		SerializeLoadingCustomVersionContainer = Ar.GetCustomVersions();
+	}
+	// Unlock the bulk data
+}
+
+void FReductionBaseSkeletalMeshBulkData::LoadReductionData(FSkeletalMeshLODModel& BaseLODModel, TMap<FString, TArray<FMorphTargetDelta>>& BaseLODMorphTargetData, UObject* Owner)
+{
+	BaseLODMorphTargetData.Empty();
+	if (BulkData.GetElementCount() > 0)
+	{
+		FReductionSkeletalMeshData ReductionSkeletalMeshData(BaseLODModel, BaseLODMorphTargetData, Owner);
+
+		// Get a lock on the bulk data
+		{
+			const bool bIsPersistent = true;
+			FBulkDataReader Ar(BulkData, bIsPersistent);
+			
+			// Propagate the custom version information from the package to the bulk data, so that the MeshDescription
+			// is serialized with the same versioning.
+			Ar.SetCustomVersions(SerializeLoadingCustomVersionContainer);
+
+			Ar << ReductionSkeletalMeshData;
+		}
+		// Unlock the bulk data
+	}
+}
+
+/*------------------------------------------------------------------------------
+FRawSkeletalMeshBulkData
+------------------------------------------------------------------------------*/
+
+FRawSkeletalMeshBulkData::FRawSkeletalMeshBulkData()
+	: bGuidIsHash(false)
+{
+	GeoImportVersion = ESkeletalMeshGeoImportVersions::Before_Versionning;
+	SkinningImportVersion = ESkeletalMeshSkinningImportVersions::Before_Versionning;
+}
+
+
+/**
+* Serialization of raw meshes uses its own versioning scheme because it is
+* stored in bulk data.
+*/
+enum
+{
+	// Engine raw mesh version:
+	RAW_SKELETAL_MESH_BULKDATA_VER_INITIAL = 0,
+	RAW_SKELETAL_MESH_BULKDATA_VER_AlternateInfluence = 1,
+	RAW_SKELETAL_MESH_BULKDATA_VER_RebuildSystem = 2,
+	// Add new raw mesh versions here.
+
+	RAW_SKELETAL_MESH_BULKDATA_VER_PLUS_ONE,
+	RAW_SKELETAL_MESH_BULKDATA_VER = RAW_SKELETAL_MESH_BULKDATA_VER_PLUS_ONE - 1,
+
+	// Licensee raw mesh version:
+	RAW_SKELETAL_MESH_BULKDATA_LIC_VER_INITIAL = 0,
+	// Licensees add new raw mesh versions here.
+
+	RAW_SKELETAL_MESH_BULKDATA_LIC_VER_PLUS_ONE,
+	RAW_SKELETAL_MESH_BULKDATA_LIC_VER = RAW_SKELETAL_MESH_BULKDATA_LIC_VER_PLUS_ONE - 1
+};
+
+FArchive& operator<<(FArchive& Ar, FSkeletalMeshImportData& RawMesh)
+{
+	int32 Version = RAW_SKELETAL_MESH_BULKDATA_VER;
+	int32 LicenseeVersion = RAW_SKELETAL_MESH_BULKDATA_LIC_VER;
+	Ar << Version;
+	Ar << LicenseeVersion;
+
+	/**
+	* Serialization should use the raw mesh version not the archive version.
+	* Additionally, stick to serializing basic types and arrays of basic types.
+	*/
+
+	Ar << RawMesh.bDiffPose;
+	Ar << RawMesh.bHasNormals;
+	Ar << RawMesh.bHasTangents;
+	Ar << RawMesh.bHasVertexColors;
+	Ar << RawMesh.bUseT0AsRefPose;
+	Ar << RawMesh.MaxMaterialIndex;
+	Ar << RawMesh.NumTexCoords;
+	
+	Ar << RawMesh.Faces;
+	Ar << RawMesh.Influences;
+	Ar << RawMesh.Materials;
+	Ar << RawMesh.Points;
+	Ar << RawMesh.PointToRawMap;
+	Ar << RawMesh.RefBonesBinary;
+	Ar << RawMesh.Wedges;
+	
+	//In the old version this processing was done after we save the asset
+	//We now save it after the processing is done so for old version we do it here when loading
+	if (Ar.IsLoading() && Version < RAW_SKELETAL_MESH_BULKDATA_VER_AlternateInfluence)
+	{
+		ProcessImportMeshInfluences(RawMesh);
+	}
+
+	
+	if (Version >= RAW_SKELETAL_MESH_BULKDATA_VER_RebuildSystem)
+	{
+		Ar << RawMesh.MorphTargets;
+		Ar << RawMesh.MorphTargetModifiedPoints;
+		Ar << RawMesh.MorphTargetNames;
+		Ar << RawMesh.AlternateInfluences;
+		Ar << RawMesh.AlternateInfluenceProfileNames;
+	}
+	else if (Ar.IsLoading())
+	{
+		RawMesh.MorphTargets.Empty();
+		RawMesh.MorphTargetModifiedPoints.Empty();
+		RawMesh.MorphTargetNames.Empty();
+		RawMesh.AlternateInfluences.Empty();
+		RawMesh.AlternateInfluenceProfileNames.Empty();
+	}
+
+	return Ar;
+}
+
+void FRawSkeletalMeshBulkData::Serialize(FArchive& Ar, UObject* Owner)
+{
+	if (Ar.IsTransacting())
+	{
+		// If transacting, keep these members alive the other side of an undo, otherwise their values will get lost
+		SerializeLoadingCustomVersionContainer.Serialize(Ar);
+		Ar << bUseSerializeLoadingCustomVersion;
+	}
+	else
+	{
+		if (Ar.IsLoading())
+		{
+			//Save the custom version so we can load FReductionSkeletalMeshData later
+			SerializeLoadingCustomVersionContainer = Ar.GetCustomVersions();
+			bUseSerializeLoadingCustomVersion = true;
+		}
+
+		if (Ar.IsSaving() && bUseSerializeLoadingCustomVersion == true)
+		{
+			//We need to update the FReductionSkeletalMeshData serialize version to the latest in case we save the Parent bulkdata
+			FSkeletalMeshImportData MeshImportData;
+			LoadRawMesh(MeshImportData);
+			SaveRawMesh(MeshImportData);
+		}
+	}
+
+	Ar.UsingCustomVersion(FEditorObjectVersion::GUID);
+
+	if (Ar.CustomVer(FEditorObjectVersion::GUID) >= FEditorObjectVersion::SkeletalMeshBuildRefactor)
+	{
+		Ar << GeoImportVersion;
+		Ar << SkinningImportVersion;
+	}
+	else
+	{
+		GeoImportVersion = ESkeletalMeshGeoImportVersions::Before_Versionning;
+		SkinningImportVersion = ESkeletalMeshSkinningImportVersions::Before_Versionning;
+	}
+	BulkData.Serialize(Ar, Owner);
+	Ar << Guid;
+	Ar << bGuidIsHash;
+}
+
+void FRawSkeletalMeshBulkData::SaveRawMesh(FSkeletalMeshImportData& InMesh)
+{
+	//Saving the bulk data mean we do not need anymore the SerializeLoadingCustomVersionContainer of the parent bulk data
+	SerializeLoadingCustomVersionContainer.Empty();
+
+	//Clear the bulk data before writing it
+	BulkData.RemoveBulkData();
+
+	// Get a lock on the bulk data
+	{
+		const bool bIsPersistent = true;
+		FBulkDataWriter Ar(BulkData, bIsPersistent);
+		Ar << InMesh;
+
+		// Preserve CustomVersions at save time so we can reuse the same ones when reloading direct from memory
+		SerializeLoadingCustomVersionContainer = Ar.GetCustomVersions();
+	}
+	// Unlock bulk data when we leave scope
+
+	//Create the guid from the content, this allow to use the data into the ddc key
+	FSHA1 Sha;
+	if (BulkData.GetBulkDataSize() > 0)
+	{
+		uint8* Buffer = (uint8*)BulkData.Lock(LOCK_READ_ONLY);
+		Sha.Update(Buffer, BulkData.GetBulkDataSize());
+		BulkData.Unlock();
+	}
+	Sha.Final();
+	uint32 Hash[5];
+	Sha.GetHash((uint8*)Hash);
+	Guid = FGuid(Hash[0] ^ Hash[4], Hash[1], Hash[2], Hash[3]);
+}
+
+void FRawSkeletalMeshBulkData::LoadRawMesh(FSkeletalMeshImportData& OutMesh)
+{
+	OutMesh.Empty();
+	if (BulkData.GetElementCount() > 0)
+	{
+		// Get a lock on the bulk data
+		{
+			const bool bIsPersistent = true;
+			FBulkDataReader Ar(BulkData, bIsPersistent);
+
+			// Propagate the custom version information from the package to the bulk data, so that the MeshDescription
+			// is serialized with the same versioning.
+			Ar.SetCustomVersions(SerializeLoadingCustomVersionContainer);
+			Ar << OutMesh;
+		}
+		// Unlock bulk data when we leave scope
+	}
+}
+
+FString FRawSkeletalMeshBulkData::GetIdString() const
+{
+	FString GuidString = Guid.ToString();
+	if (bGuidIsHash)
+	{
+		GuidString += TEXT("X");
+	}
+	return GuidString;
+}
+
+void FRawSkeletalMeshBulkData::UseHashAsGuid(UObject* Owner)
+{
+	// Build the hash from the path name + the contents of the bulk data.
+	FSHA1 Sha;
+	TArray<TCHAR> OwnerName = Owner->GetPathName().GetCharArray();
+	Sha.Update((uint8*)OwnerName.GetData(), OwnerName.Num() * OwnerName.GetTypeSize());
+	if (BulkData.GetBulkDataSize() > 0)
+	{
+		uint8* Buffer = (uint8*)BulkData.Lock(LOCK_READ_ONLY);
+		Sha.Update(Buffer, BulkData.GetBulkDataSize());
+		BulkData.Unlock();
+	}
+	Sha.Final();
+
+	// Retrieve the hash and use it to construct a pseudo-GUID. Use bGuidIsHash to distinguish from real guids.
+	uint32 Hash[5];
+	Sha.GetHash((uint8*)Hash);
+	Guid = FGuid(Hash[0] ^ Hash[4], Hash[1], Hash[2], Hash[3]);
+	bGuidIsHash = true;
+}
+
+FByteBulkData& FRawSkeletalMeshBulkData::GetBulkData()
+{
+	return BulkData;
+}
+
+
+/************************************************************************
+* FWedgePosition
+*/
+void FWedgePosition::FindMatchingPositionWegdeIndexes(const FVector &Position, float ComparisonThreshold, TArray<int32>& OutResults)
+{
+	int32 SortedPositionNumber = SortedPositions.Num();
+	OutResults.Empty();
+	if (SortedPositionNumber == 0)
+	{
+		//No possible match
+		return;
+	}
+	FWedgePositionHelper::FIndexAndZ PositionIndexAndZ(INDEX_NONE, Position);
+	int32 SortedIndex = SortedPositions.Num()/2;
+	int32 StartIndex = 0;
+	int32 LastTopIndex = SortedPositions.Num();
+	int32 LastBottomIndex = 0;
+	int32 SearchIterationCount = 0;
+
+	{
+		double Increments = ((double)SortedPositions[SortedPositionNumber - 1].Z - (double)SortedPositions[0].Z) / (double)SortedPositionNumber;
+
+		//Optimize the iteration count when a value is not in the middle
+		SortedIndex = FMath::RoundToInt(((double)PositionIndexAndZ.Z - (double)SortedPositions[0].Z) / Increments);
+	}
+
+	for (SearchIterationCount = 0; SortedPositions.IsValidIndex(SortedIndex); ++SearchIterationCount)
+	{
+		if (LastTopIndex - LastBottomIndex < 5)
+		{
+			break;
+		}
+		if (FMath::Abs(PositionIndexAndZ.Z - SortedPositions[SortedIndex].Z) < ComparisonThreshold)
+		{
+			//Continue since we want the lowest start
+			LastTopIndex = SortedIndex;
+			SortedIndex = LastBottomIndex + ((LastTopIndex - LastBottomIndex) / 2);
+			if (SortedIndex <= LastBottomIndex)
+			{
+				break;
+			}
+		}
+		else if (PositionIndexAndZ.Z > SortedPositions[SortedIndex].Z + ComparisonThreshold)
+		{
+			LastBottomIndex = SortedIndex;
+			SortedIndex = SortedIndex + FMath::Max(((LastTopIndex - SortedIndex) / 2), 1);
+		}
+		else
+		{
+			LastTopIndex = SortedIndex;
+			SortedIndex = SortedIndex - FMath::Max(((SortedIndex - LastBottomIndex) / 2), 1);
+		}
+	}
+	
+	//////////////////////////////////////////////////////////////////////////
+	//Closest point data (!bExactMatch)
+	float MinDistance = MAX_FLT;
+	int32 ClosestIndex = LastBottomIndex;
+
+	for (int32 i = LastBottomIndex; i < SortedPositionNumber; i++)
+	{
+		//Get fast to the close position
+		if (PositionIndexAndZ.Z > SortedPositions[i].Z + ComparisonThreshold)
+		{
+			continue;
+		}
+		//break when we pass point close to the position
+		if (SortedPositions[i].Z > PositionIndexAndZ.Z + ComparisonThreshold)
+			break; // can't be any more dups
+
+		//Point is close to the position, verify it
+		const FVector& PositionA = Points[Wedges[SortedPositions[i].Index].VertexIndex];
+		if (FWedgePositionHelper::PointsEqual(PositionA, Position, ComparisonThreshold))
+		{
+			OutResults.Add(SortedPositions[i].Index);
+		}
+	}
+}
+
+void FOctreeQueryHelper::FindNearestWedgeIndexes(const FVector& SearchPosition, TArray<FWedgeInfo>& OutNearestWedges)
+{
+	if (WedgePosOctree == nullptr)
+	{
+		return;
+	}
+	float MinSquaredDistance = MAX_FLT;
+	OutNearestWedges.Empty();
+	
+	FVector Extend(2.0f);
+	for (int i = 0; i < 2; ++i)
+	{
+		TWedgeInfoPosOctree::TConstIterator<> OctreeIter((*WedgePosOctree));
+		// Iterate through the octree attempting to find the vertices closest to the current new point
+		while (OctreeIter.HasPendingNodes())
+		{
+			const TWedgeInfoPosOctree::FNode& CurNode = OctreeIter.GetCurrentNode();
+			const FOctreeNodeContext& CurContext = OctreeIter.GetCurrentContext();
+
+			// Find the child of the current node, if any, that contains the current new point
+
+			//The first shot is an intersection with a 1 CM cube box around the search position, this ensure we dont fall in the wrong neighbourg
+			FOctreeChildNodeSubset ChilNodesSubset = CurContext.GetIntersectingChildren(FBoxCenterAndExtent(SearchPosition, Extend));
+			FOREACH_OCTREE_CHILD_NODE(OctreeChildRef)
+			{
+				if (ChilNodesSubset.Contains(OctreeChildRef) && CurNode.HasChild(OctreeChildRef))
+				{
+					OctreeIter.PushChild(OctreeChildRef);
+				}
+			}
+			// Add all of the elements in the current node to the list of points to consider for closest point calculations
+			for (const FWedgeInfo& WedgeInfo : CurNode.GetElements())
+			{
+				float VectorDelta = FVector::DistSquared(SearchPosition, WedgeInfo.Position);
+				MinSquaredDistance = FMath::Min(VectorDelta, MinSquaredDistance);
+				OutNearestWedges.Add(WedgeInfo);
+			}
+			OctreeIter.Advance();
+		}
+
+		if (i == 0)
+		{
+			float MinDistance = FMath::Sqrt(MinSquaredDistance);
+			if (MinDistance < Extend.X)
+			{
+				//We found the closest points
+				break;
+			}
+			OutNearestWedges.Empty();
+			//Change the extend to the distance we found so we are sure to find any closer point in the neighbourg
+			Extend = FVector(MinDistance + KINDA_SMALL_NUMBER);
+		}
+	}
+
+}
+
+void FWedgePosition::FillWedgePosition(
+	FWedgePosition& OutOverlappingPosition,
+	const TArray<FVector>& Points,
+	const TArray<SkeletalMeshImportData::FVertex> Wedges,
+	float ComparisonThreshold)
+{
+	OutOverlappingPosition.Points= Points;
+	OutOverlappingPosition.Wedges = Wedges;
+	const int32 NumWedges = OutOverlappingPosition.Wedges.Num();
+	// Create a list of vertex Z/index pairs
+	OutOverlappingPosition.SortedPositions.Reserve(NumWedges);
+	for (int32 WedgeIndex = 0; WedgeIndex < NumWedges; WedgeIndex++)
+	{
+		new(OutOverlappingPosition.SortedPositions)FWedgePositionHelper::FIndexAndZ(WedgeIndex, OutOverlappingPosition.Points[OutOverlappingPosition.Wedges[WedgeIndex].VertexIndex]);
+	}
+
+	// Sort the vertices by z value
+	OutOverlappingPosition.SortedPositions.Sort(FWedgePositionHelper::FCompareIndexAndZ());
+
+
+	FBox OldBounds(OutOverlappingPosition.Points);
+	OutOverlappingPosition.WedgePosOctree = new TWedgeInfoPosOctree(OldBounds.GetCenter(), OldBounds.GetExtent().GetMax());
+
+	// Add each old vertex to the octree
+	for (int32 WedgeIndex = 0; WedgeIndex < NumWedges; ++WedgeIndex)
+	{
+		FWedgeInfo WedgeInfo;
+		WedgeInfo.WedgeIndex = WedgeIndex;
+		WedgeInfo.Position = OutOverlappingPosition.Points[OutOverlappingPosition.Wedges[WedgeIndex].VertexIndex];
+		OutOverlappingPosition.WedgePosOctree->AddElement(WedgeInfo);
+	}
+}
+
+#endif // WITH_EDITOR
\ No newline at end of file
diff --git a/Engine/Source/Runtime/Engine/Public/Rendering/SkeletalMeshLODImporterData.cpp b/Engine/Source/Runtime/Engine/Public/Rendering/SkeletalMeshLODImporterData.cpp
deleted file mode 100644
