#UE4 Fix implicit int32 to bool conversion. It is now explicit.

diff --git a/Engine/Source/Runtime/Engine/Private/PhysicsEngine/ImmediatePhysics/ImmediatePhysicsChaos/ImmediatePhysicsSimulation_Chaos.cpp b/Engine/Source/Runtime/Engine/Private/PhysicsEngine/ImmediatePhysics/ImmediatePhysicsChaos/ImmediatePhysicsSimulation_Chaos.cpp
--- a/Engine/Source/Runtime/Engine/Private/PhysicsEngine/ImmediatePhysics/ImmediatePhysicsChaos/ImmediatePhysicsSimulation_Chaos.cpp
+++ b/Engine/Source/Runtime/Engine/Private/PhysicsEngine/ImmediatePhysics/ImmediatePhysicsChaos/ImmediatePhysicsSimulation_Chaos.cpp
@@ -406,7 +406,7 @@ namespace ImmediatePhysics_Chaos
 			Evolution->SetNumIterations(ChaosImmediate_Evolution_Iterations);
 
 			Evolution->GetCollisionConstraints().SetThickness(ChaosImmediate_Collision_Thickness);
-			Evolution->GetCollisionConstraints().SetVelocitySolveEnabled(ChaosImmediate_Collision_ApplyEnabled);
+			Evolution->GetCollisionConstraints().SetVelocitySolveEnabled(ChaosImmediate_Collision_ApplyEnabled != 0);
 			Evolution->GetCollisionConstraints().SetPushOutPairIterations(ChaosImmediate_Collision_PushOutPairIterations);
 			Evolution->GetCollisionConstraintsRule().SetPushOutIterations(ChaosImmediate_Collision_PushOutIterations);
 		}
