Copied ImmediatePhysics.h from //UE4/Dev-Physics to fix AnimGraph compilation problems

diff --git a/Engine/Source/Runtime/Engine/Public/Physics/ImmediatePhysics/ImmediatePhysicsSimulation.h b/Engine/Source/Runtime/Engine/Public/Physics/ImmediatePhysics/ImmediatePhysicsSimulation.h
--- a/Engine/Source/Runtime/Engine/Public/Physics/ImmediatePhysics/ImmediatePhysicsSimulation.h
+++ b/Engine/Source/Runtime/Engine/Public/Physics/ImmediatePhysics/ImmediatePhysicsSimulation.h
@@ -2,277 +2,7 @@
 
 #pragma once
 
-#if WITH_PHYSX && PHYSICS_INTERFACE_PHYSX
-#include "PhysXPublic.h"
+#include "Physics/ImmediatePhysics/ImmediatePhysicsPhysX/ImmediatePhysicsSimulation_PhysX.h"
+#include "Physics/ImmediatePhysics/ImmediatePhysicsChaos/ImmediatePhysicsSimulation_Chaos.h"
 
-#include "Physics/PhysicsInterfaceTypes.h"
-
-#include "Engine/EngineTypes.h"
-#include "Physics/ImmediatePhysics/ImmediatePhysicsPhysX/ImmediatePhysicsCore_PhysX.h"
-#include "Physics/ImmediatePhysics/ImmediatePhysicsPhysX/ImmediatePhysicsActor_PhysX.h"
-#include "Physics/ImmediatePhysics/ImmediatePhysicsPhysX/ImmediatePhysicsJoint_PhysX.h"
-#include "Physics/ImmediatePhysics/ImmediatePhysicsPhysX/ImmediatePhysicsLinearBlockAllocator_PhysX.h"
-#include "Physics/ImmediatePhysics/ImmediatePhysicsPhysX/ImmediatePhysicsContactPair_PhysX.h"
-#include "Physics/ImmediatePhysics/ImmediatePhysicsPhysX/ImmediatePhysicsPersistentContactPairData_PhysX.h"
-#include "Physics/ImmediatePhysics/ImmediatePhysicsPhysX/ImmediatePhysicsCacheAllocator_PhysX.h"
-#include "Physics/ImmediatePhysics/ImmediatePhysicsPhysX/ImmediatePhysicsConstraintAllocator_PhysX.h"
-#include "Physics/ImmediatePhysics/ImmediatePhysicsPhysX/ImmediatePhysicsKinematicTarget_PhysX.h"
-#include "Physics/ImmediatePhysics/ImmediatePhysicsPhysX/ImmediatePhysicsD6JointData_PhysX.h"
-
-class UBodySetup;
-class UPhysicsConstraintTemplate;
-struct FBodyInstance;
-struct FConstraintInstance;
-
-namespace ImmediatePhysics_PhysX
-{
-	struct FActorHandle;
-	struct FJointHandle;
-    struct FContactPointRecorder;
-}
-
-namespace ImmediatePhysics_PhysX
-{
-	struct FActorData
-	{
-		immediate::PxRigidBodyData RigidBodyData;
-		FTransform InitialTransform;
-
-		bool bStatic;
-		bool bKinematic;
-	};
-
-	FActorData CreateActorData(const FActorCreationParams& InParams);
-
-	/** Owns all the data associated with the simulation. Can be considered a single scene or world*/
-	struct ENGINE_API FSimulation
-	{
-	public:
-
-		FActorHandle* InsertActorData(const FActor& InActor, const FActorData& InData);
-
-		void RemoveActor(FActorHandle* InHandle);
-
-		/** Create a kinematic body with the same setup as the body instance and add it to the simulation */
-		FActorHandle* CreateKinematicActor(FBodyInstance* BodyInstance, const FTransform& TM);
-
-		/** Create a dynamic body with the same setup as the body instance and add it to the simulation */
-		FActorHandle* CreateDynamicActor(FBodyInstance* BodyInstance, const FTransform& TM);
-
-		/** Create a static body with the same setup as the body instance and add it to the simulation */
-		FActorHandle* CreateStaticActor(FBodyInstance* BodyInstance);
-
-		/** Create a body of the specified type with the same setup as the body instance and add it to the simulation */
-		FActorHandle* CreateActor(EActorType ActorType, FBodyInstance* BodyInstance, const FTransform& Transform);
-
-		/** Create a physical joint with the same setup as the constraint instance and add it to the simulation */
-		FJointHandle* CreateJoint(FConstraintInstance* ConstraintInstance, FActorHandle* Body1, FActorHandle* Body2);
-
-		const immediate::PxRigidBodyData& GetLowLevelBody(int32 ActorDataIndex) const
-		{
-			return RigidBodiesData[ActorDataIndex];
-		}
-
-		immediate::PxRigidBodyData& GetLowLevelBody(int32 ActorDataIndex)
-		{
-			return RigidBodiesData[ActorDataIndex];
-		}
-
-		/** Sets the number of active bodies. This number is reset any time a new simulated body is created */
-		void SetNumActiveBodies(int32 NumActiveBodies);
-
-		/** An array of actors to ignore. */
-		struct FIgnorePair
-		{
-			FActorHandle* A;
-			FActorHandle* B;
-		};
-
-		/** Set pair of bodies to ignore collision for */
-		void SetIgnoreCollisionPairTable(const TArray<FIgnorePair>& InIgnoreCollisionPairTable);
-
-		/** Set bodies that require no collision */
-		void SetIgnoreCollisionActors(const TArray<FActorHandle*>& InIgnoreCollision);
-
-		/** Get/Set whether the body is kinematic or not, kinematics do not simulate and move where they're told. They also act as if they have infinite mass */
-		bool GetIsKinematic(int32 ActorDataIndex);
-		void SetIsKinematic(int32 ActorDataIndex, bool bKinematic);
-
-		/** Advance the simulation by DeltaTime */
-		void Simulate(float DeltaTime, const FVector& Gravity);
-		void Simulate_AssumesLocked(float DeltaTime, const FVector& Gravity);
-
-		/** Whether or not an entity is simulated */
-		bool IsSimulated(int32 ActorDataIndex) const
-		{
-			return ActorDataIndex < NumSimulatedBodies;
-		}
-
-		/** Add a radial impulse to the given actor */
-		void AddRadialForce(int32 ActorDataIndex, const FVector& Origin, float Strength, float Radius, ERadialImpulseFalloff Falloff, EForceType ForceType);
-
-		/** Add a force to the given actor */
-		void AddForce(int32 ActorDataIndex, const FVector& Force);
-
-		FSimulation();
-
-		~FSimulation();
-
-		int32 NumActors() const { return ActorHandles.Num(); }
-
-		FActorHandle* GetActorHandle(int ActorHandleIndex) { return ActorHandles[ActorHandleIndex]; }
-		const FActorHandle* GetActorHandle(int ActorHandleIndex) const { return ActorHandles[ActorHandleIndex]; }
-
-		/* Set solver iteration counts per step */
-		void SetPositionIterationCount(int32 InIterationCount);
-		void SetVelocityIterationCount(int32 InIterationCount);
-
-		void SetSimulationSpaceTransform(const FTransform& Transform) { }
-
-	private:
-		friend FActorHandle;
-
-		const FImmediateKinematicTarget& GetKinematicTarget(int32 ActorDataIndex) const
-		{
-			return KinematicTargets[ActorDataIndex];
-		}
-
-		FImmediateKinematicTarget& GetKinematicTarget(int32 ActorDataIndex)
-		{
-			return KinematicTargets[ActorDataIndex];
-		}
-
-		template <EActorType ActorType>
-		int32 CreateActor(PxRigidActor* RigidActor, const FTransform& TM);
-
-		/** Swap actor data - that is move all data associated with the two actors in the various arrays*/
-		void SwapActorData(int32 Entity1Idx, int32 Entity2Idx);
-
-		/** Swap joint data - that is move all data associated with the two joints in the various arrays*/
-		void SwapJointData(int32 Joint1Idx, int32 Joint2Idx);
-
-		/** Ensure arrays are valid */
-		void ValidateArrays() const;
-
-		/** Constructs solver bodies */
-		void ConstructSolverBodies(float DeltaTime, const FVector& Gravity);
-
-		/** Generate contacts*/
-		void GenerateContacts();
-
-		/** Batch constraints and re-order them for optimal processing */
-		void BatchConstraints();
-
-		/** Prepares the various constraints (contact,joints) for the solver */
-		void PrepareConstraints(float DeltaTime);
-
-		/** Solve constraints and integrate velocities */
-		void SolveAndIntegrate(float DeltaTime);
-
-		/** Prepares iteration cache for generating contacts */
-		void PrepareIterationCache();
-
-		//void EvictCache();
-
-	private:
-
-		/** Mapping from entity index to handle */
-		TArray<FActorHandle*>	ActorHandles;
-
-		/** Mapping from constraint index to handle */
-		TArray<FJointHandle*> JointHandles;
-
-		/** Entities holding loose data. NOTE: for performance reasons we don't automatically cleanup on destructor (needed for tarray swaps etc...) it's very important that Terminate is called */
-		TArray<FActor> Actors;
-		TArray<FJoint> Joints;
-
-		/** Workspace memory that we use for per frame allocations */
-		FLinearBlockAllocator Workspace;
-
-		/** Low level rigid body data */
-		TArray<immediate::PxRigidBodyData> RigidBodiesData;
-
-		/** Low level solver bodies data */
-		TArray<PxSolverBodyData> SolverBodiesData;
-
-		/** Kinematic targets used to implicitly compute the velocity of moving kinematic actors */
-		TArray<FImmediateKinematicTarget> KinematicTargets;
-
-		TArray<PxVec3> PendingAcceleration;
-
-		/** Low level contact points generated for this frame. Points are grouped together by pairs */
-		TArray<Gu::ContactPoint> ContactPoints;
-
-		/** Shapes used in the entire simulation. Shapes are sorted in the same order as actors. Note that an actor can have multiple shapes which will be adjacent*/
-		struct FShapeSOA
-		{
-			TArray<PxTransform> LocalTMs;
-			TArray<FMaterial> Materials;
-			TArray<const PxGeometry*> Geometries;
-			TArray<float> Bounds;
-			TArray<PxVec3> BoundsOffsets;
-			TArray<int32> OwningActors;
-#if PERSISTENT_CONTACT_PAIRS
-			TArray<FPersistentContactPairData> ContactPairData;
-#endif
-		} ShapeSOA;
-
-		/** Low level solver bodies */
-		PxSolverBody* SolverBodies;
-
-		/** Low level constraint descriptors.*/
-		TArray<PxSolverConstraintDesc> OrderedDescriptors;
-		TArray<PxConstraintBatchHeader> BatchHeaders;
-
-		/** JointData as passed in from physics constraint template */
-		TArray<D6JointData> JointData;
-
-		/** When new joints are created we have to update the processing order */
-		bool bDirtyJointData;
-
-		PxU32 NumContactHeaders;
-		PxU32 NumJointHeaders;
-		int32 NumActiveJoints;
-
-		/** Contact pairs generated for this frame */
-		TArray<FContactPair> ContactPairs;
-
-		/** Number of dynamic bodies associated with the simulation */
-		int32 NumSimulatedBodies;
-
-		/** Number of dynamic bodies that are actually active */
-		int32 NumActiveSimulatedBodies;
-
-		/** Number of kinematic bodies (dynamic but not simulated) associated with the simulation */
-		int32 NumKinematicBodies;
-
-		/** Total number of simulated shapes in the scene */
-		int32 NumSimulatedShapesWithCollision;
-
-		/** Number of position iterations used by solver */
-		int32 NumPositionIterations;
-
-		/** Number of velocity iterations used by solver */
-		int32 NumVelocityIterations;
-
-		/** Count of how many times we've ticked. Useful for cache invalidation */
-		int32 SimCount;
-
-		/** Both of these are slow to access. Make sure to use iteration cache when possible */
-		TMap<FActorHandle*, TSet<FActorHandle*>> IgnoreCollisionPairTable;
-		TSet<FActorHandle*> IgnoreCollisionActors;
-
-		/** This cache is used to record which generate contact iteration we can skip. This assumes the iteration order has not changed (add/remove/swap actors must invalidate this) */
-		bool bRecreateIterationCache;
-
-		TArray<int32> SkipCollisionCache;	//Holds the iteration count that we should skip due to ignore filtering
-
-		friend FContactPointRecorder;
-
-		FCacheAllocator CacheAllocator;
-		FConstraintAllocator ConstraintAllocator;
-	};
-
-}
-
-#endif // WITH_PHYSX
\ No newline at end of file
+#include "Physics/ImmediatePhysics/ImmediatePhysicsDeclares.h"
