#UE4 Silence some cook warnings

diff --git a/Engine/Source/Runtime/Engine/Private/AudioDerivedData.cpp b/Engine/Source/Runtime/Engine/Private/AudioDerivedData.cpp
--- a/Engine/Source/Runtime/Engine/Private/AudioDerivedData.cpp
+++ b/Engine/Source/Runtime/Engine/Private/AudioDerivedData.cpp
@@ -468,7 +468,7 @@ class FStreamedAudioCacheDerivedDataWorker : public FNonAbandonableTask
 		}
 		else
 		{
-			UE_LOG(LogAudio, Warning, TEXT("Failed to build %s derived data for %s"),
+			UE_LOG(LogAudio, Display, TEXT("Failed to build %s derived data for %s"),
 				*AudioFormatName.GetPlainNameString(),
 				*SoundWave.GetPathName()
 				);
