Moved min roughness override into IntegrateBxDF. This merges existing roughness clamp with the overriden one.

diff --git a/Engine/Shaders/Private/CapsuleLightIntegrate.ush b/Engine/Shaders/Private/CapsuleLightIntegrate.ush
--- a/Engine/Shaders/Private/CapsuleLightIntegrate.ush
+++ b/Engine/Shaders/Private/CapsuleLightIntegrate.ush
@@ -90,7 +90,7 @@ FDirectLighting IntegrateBxDF( FGBufferData GBuffer, half3 N, half3 V, FCapsuleL
 	float InvDist = rsqrt( DistSqr );
 	float3 L = ToLight * InvDist;
 	
-	GBuffer.Roughness = max( GBuffer.Roughness, 0.02 );
+	GBuffer.Roughness = max( GBuffer.Roughness, View.MinRoughness );
 	float a = Pow2( GBuffer.Roughness );
     
 	FAreaLight AreaLight;
diff --git a/Engine/Shaders/Private/DeferredLightingCommon.ush b/Engine/Shaders/Private/DeferredLightingCommon.ush
--- a/Engine/Shaders/Private/DeferredLightingCommon.ush
+++ b/Engine/Shaders/Private/DeferredLightingCommon.ush
@@ -287,8 +287,6 @@ float4 GetDynamicLighting(
 	FDeferredLightData LightData, float4 LightAttenuation, float Dither, uint2 SVPos, FRectTexture SourceTexture,
 	inout float SurfaceShadow)
 {
-	GBuffer.Roughness = max(View.MinRoughness, GBuffer.Roughness);
-
 	FLightAccumulator LightAccumulator = (FLightAccumulator)0;
 
 	float3 V = -CameraVector;
diff --git a/Engine/Source/Runtime/Engine/Private/SceneView.cpp b/Engine/Source/Runtime/Engine/Private/SceneView.cpp
--- a/Engine/Source/Runtime/Engine/Private/SceneView.cpp
+++ b/Engine/Source/Runtime/Engine/Private/SceneView.cpp
@@ -2329,7 +2329,7 @@ void FSceneView::SetupCommonViewUniformBufferParameters(
 	ViewUniformShaderParameters.CameraCut = bCameraCut ? 1 : 0;
 
 	ViewUniformShaderParameters.VirtualTextureParams = FVector4(ForceInitToZero);
-	ViewUniformShaderParameters.MinRoughness = CVarGlobalMinRoughnessOverride.GetValueOnRenderThread();
+	ViewUniformShaderParameters.MinRoughness = FMath::Clamp(CVarGlobalMinRoughnessOverride.GetValueOnRenderThread(), 0.02f, 1.0f);
 
 	//to tail call keep the order and number of parameters of the caller function
 	SetupViewRectUniformBufferParameters(ViewUniformShaderParameters, BufferSize, EffectiveViewRect, InViewMatrices, InPrevViewMatrices);
