Sanitized sampler source being null. Changed check to an ensure.

diff --git a/Engine/Source/Runtime/Engine/Private/Materials/MaterialUniformExpressions.cpp b/Engine/Source/Runtime/Engine/Private/Materials/MaterialUniformExpressions.cpp
--- a/Engine/Source/Runtime/Engine/Private/Materials/MaterialUniformExpressions.cpp
+++ b/Engine/Source/Runtime/Engine/Private/Materials/MaterialUniformExpressions.cpp
@@ -624,13 +624,14 @@ void FUniformExpressionSet::FillUniformBuffer(const FMaterialRenderContext& Mate
 			// VirtualTexture is allowed here, as these may be demoted to regular textures on platforms that don't have VT support
 			const uint32 ValidTextureTypes = MCT_Texture2D | MCT_TextureVirtual | MCT_TextureExternal;
 
+			bool bValueValid = false;
+
 			// TextureReference.TextureReferenceRHI is cleared from a render command issued by UTexture::BeginDestroy
 			// It's possible for this command to trigger before a given material is cleaned up and removed from deferred update list
 			// Technically I don't think it's necessary to check 'Resource' for nullptr here, as if TextureReferenceRHI has been initialized, that should be enough
 			// Going to leave the check for now though, to hopefully avoid any unexpected problems
 			if (Value && Value->Resource && Value->TextureReference.TextureReferenceRHI && (Value->GetMaterialType() & ValidTextureTypes) != 0u)
 			{
-				*ResourceTableTexturePtr = Value->TextureReference.TextureReferenceRHI;
 				FSamplerStateRHIRef* SamplerSource = &Value->Resource->SamplerStateRHI;
 
 				ESamplerSourceMode SourceMode = Uniform2DTextureExpressions[ExpressionIndex]->GetSamplerSource();
@@ -643,12 +644,23 @@ void FUniformExpressionSet::FillUniformBuffer(const FMaterialRenderContext& Mate
 					SamplerSource = &Clamp_WorldGroupSettings->SamplerStateRHI;
 				}
 
-				checkf(*SamplerSource, TEXT("Texture %s of class %s had invalid sampler source. Material %s with texture expression in slot %i"),
-					*Value->GetName(), *Value->GetClass()->GetName(),
-					*MaterialRenderContext.Material.GetFriendlyName(), ExpressionIndex);
-				*ResourceTableSamplerPtr = *SamplerSource;
+				if (*SamplerSource)
+				{
+					*ResourceTableTexturePtr = Value->TextureReference.TextureReferenceRHI;
+					*ResourceTableSamplerPtr = *SamplerSource;
+					bValueValid = true;
+				}
+				else
+				{
+					ensureMsgf(false,
+						TEXT("Texture %s of class %s had invalid sampler source. Material %s with texture expression in slot %i. Sampler source mode %d. Resource initialized: %d"),
+						*Value->GetName(), *Value->GetClass()->GetName(),
+						*MaterialRenderContext.Material.GetFriendlyName(), ExpressionIndex, SourceMode,
+						Value->Resource->IsInitialized());
+				}
 			}
-			else
+
+			if (!bValueValid)
 			{
 				check(GWhiteTexture->TextureRHI);
 				*ResourceTableTexturePtr = GWhiteTexture->TextureRHI;
