Fix issue where replicated object references within arrays within structs may fail to resolve.

diff --git a/Engine/Source/Runtime/Engine/Private/RepLayout.cpp b/Engine/Source/Runtime/Engine/Private/RepLayout.cpp
--- a/Engine/Source/Runtime/Engine/Private/RepLayout.cpp
+++ b/Engine/Source/Runtime/Engine/Private/RepLayout.cpp
@@ -3437,12 +3437,12 @@ bool FRepLayout::MoveMappedObjectToUnmapped(
 }
 
 void FRepLayout::UpdateUnmappedObjects_r(
-	FReceivingRepState* RESTRICT RepState, 
+	FReceivingRepState* RESTRICT RepState,
 	FGuidReferencesMap* GuidReferencesMap,
 	UObject* OriginalObject,
-	UPackageMap* PackageMap, 
-	FRepShadowDataBuffer ShadowData, 
-	FRepObjectDataBuffer Data, 
+	UPackageMap* PackageMap,
+	FRepShadowDataBuffer ShadowData,
+	FRepObjectDataBuffer Data,
 	const int32 MaxAbsOffset,
 	bool& bCalledPreNetReceive,
 	bool& bOutSomeObjectsWereMapped,
@@ -3464,13 +3464,18 @@ void FRepLayout::UpdateUnmappedObjects_r(
 		const FRepLayoutCmd& Cmd = Cmds[GuidReferences.CmdIndex];
 		const FRepParentCmd& Parent = Parents[GuidReferences.ParentIndex];
 
+		// Make sure if we're touching an array element, we use the correct offset for shadow values.
+		// This should always be safe, because MaxAbsOffset will account for ShadowArray size for arrays.
+		// For non array properties, AbsOffset should always equal Cmd.Offset.
+		const int32 ShadowOffset = (AbsOffset - Cmd.Offset) + Cmd.ShadowOffset;
+
 		if (GuidReferences.Array != nullptr)
 		{
 			check(Cmd.Type == ERepLayoutCmdType::DynamicArray);
-			
+
 			if (ShadowData)
 			{
-				FScriptArray* ShadowArray = (FScriptArray*)(ShadowData + Cmd).Data;
+				FScriptArray* ShadowArray = (FScriptArray*)(ShadowData + ShadowOffset).Data;
 				FScriptArray* Array = (FScriptArray*)(Data + AbsOffset).Data;
 
 				FRepShadowDataBuffer ShadowArrayData(ShadowArray->GetData());
@@ -3494,7 +3499,7 @@ void FRepLayout::UpdateUnmappedObjects_r(
 		bool bMappedSomeGUIDs = false;
 
 		for (auto UnmappedIt = GuidReferences.UnmappedGUIDs.CreateIterator(); UnmappedIt; ++UnmappedIt)
-		{			
+		{
 			const FNetworkGUID& GUID = *UnmappedIt;
 
 			if (PackageMap->IsGUIDBroken(GUID, false))
@@ -3538,7 +3543,7 @@ void FRepLayout::UpdateUnmappedObjects_r(
 			// Copy current value over so we can check to see if it changed
 			if (bUpdateShadowState)
 			{
-				StoreProperty(Cmd, ShadowData + Cmd, Data + AbsOffset);
+				StoreProperty(Cmd, ShadowData + ShadowOffset, Data + AbsOffset);
 			}
 
 			// Initialize the reader with the stored buffer that we need to read from
@@ -3556,7 +3561,7 @@ void FRepLayout::UpdateUnmappedObjects_r(
 				// That could cause us to trigger RepNotifies more often for Dynamic Array properties.
 				// That goes for the above too.
 
-				if (Parent.RepNotifyCondition == REPNOTIFY_Always || !PropertiesAreIdentical(Cmd, ShadowData + Cmd, Data + AbsOffset))
+				if (Parent.RepNotifyCondition == REPNOTIFY_Always || !PropertiesAreIdentical(Cmd, ShadowData + ShadowOffset, Data + AbsOffset))
 				{
 					// If this properties needs an OnRep, queue that up to be handled later
 					RepState->RepNotifies.AddUnique(Parent.Property);
