Fix assertion that occurred when running the editor with -server with Stream Caching enabled.

diff --git a/Engine/Source/Runtime/Engine/Private/SoundWave.cpp b/Engine/Source/Runtime/Engine/Private/SoundWave.cpp
--- a/Engine/Source/Runtime/Engine/Private/SoundWave.cpp
+++ b/Engine/Source/Runtime/Engine/Private/SoundWave.cpp
@@ -822,7 +822,10 @@ void USoundWave::EnsureZerothChunkIsLoaded()
 	uint8* TempChunkBuffer = nullptr;
 	int32 ChunkSizeInBytes = RunningPlatformData->GetChunkFromDDC(0, &TempChunkBuffer, true);
 	// Since we block for the DDC in the previous call we should always have the chunk loaded.
-	check(ChunkSizeInBytes > 0);
+	if (ChunkSizeInBytes == 0)
+	{
+		return;
+	}
 
 	// TODO: Support passing a TArray by ref into FStreamedAudioPlatformData::GetChunkFromDDC.
 	// Currently not feasible unless FUntypedBulkData::GetCopy API was changed.
