#ue4 - UCharacterMovementComponent::ServerCheckClientError() uses ServerExceedsAllowablePositionError() for the actual position and movement mode checks, so you can override the behavior but still retain the ignore setting, logging, and random chance from CVars.

diff --git a/Engine/Source/Runtime/Engine/Classes/GameFramework/CharacterMovementComponent.h b/Engine/Source/Runtime/Engine/Classes/GameFramework/CharacterMovementComponent.h
--- a/Engine/Source/Runtime/Engine/Classes/GameFramework/CharacterMovementComponent.h
+++ b/Engine/Source/Runtime/Engine/Classes/GameFramework/CharacterMovementComponent.h
@@ -2197,6 +2197,11 @@ protected:
 	 */
 	virtual bool ServerCheckClientError(float ClientTimeStamp, float DeltaTime, const FVector& Accel, const FVector& ClientWorldLocation, const FVector& RelativeClientLocation, UPrimitiveComponent* ClientMovementBase, FName ClientBaseBoneName, uint8 ClientMovementMode);
 
+	/**
+	 * Check position error within ServerCheckClientError(). Set bNetworkLargeClientCorrection to true if the correction should be prioritized (delayed less in SendClientAdjustment).
+	 */
+	virtual bool ServerExceedsAllowablePositionError(float ClientTimeStamp, float DeltaTime, const FVector& Accel, const FVector& ClientWorldLocation, const FVector& RelativeClientLocation, UPrimitiveComponent* ClientMovementBase, FName ClientBaseBoneName, uint8 ClientMovementMode);
+
 	/**
 	 * If ServerCheckClientError() does not find an error, this determines if the server should also copy the client's movement params rather than keep the server sim result.
 	 */
diff --git a/Engine/Source/Runtime/Engine/Private/Components/CharacterMovementComponent.cpp b/Engine/Source/Runtime/Engine/Private/Components/CharacterMovementComponent.cpp
--- a/Engine/Source/Runtime/Engine/Private/Components/CharacterMovementComponent.cpp
+++ b/Engine/Source/Runtime/Engine/Private/Components/CharacterMovementComponent.cpp
@@ -8865,27 +8865,17 @@ bool UCharacterMovementComponent::ServerCheckClientError(float ClientTimeStamp,
 	// Check location difference against global setting
 	if (!bIgnoreClientMovementErrorChecksAndCorrection)
 	{
-		const FVector LocDiff = UpdatedComponent->GetComponentLocation() - ClientWorldLocation;
-
 #if ROOT_MOTION_DEBUG
 		if (RootMotionSourceDebug::CVarDebugRootMotionSources.GetValueOnGameThread() == 1)
 		{
+			const FVector LocDiff = UpdatedComponent->GetComponentLocation() - ClientWorldLocation;
 			FString AdjustedDebugString = FString::Printf(TEXT("ServerCheckClientError LocDiff(%.1f) ExceedsAllowablePositionError(%d) TimeStamp(%f)"),
 				LocDiff.Size(), GetDefault<AGameNetworkManager>()->ExceedsAllowablePositionError(LocDiff), ClientTimeStamp);
 			RootMotionSourceDebug::PrintOnScreen(*CharacterOwner, AdjustedDebugString);
 		}
 #endif
 
-		const AGameNetworkManager* GameNetworkManager = (const AGameNetworkManager*)(AGameNetworkManager::StaticClass()->GetDefaultObject());
-		if (GameNetworkManager->ExceedsAllowablePositionError(LocDiff))
-		{
-			bNetworkLargeClientCorrection = (LocDiff.SizeSquared() > FMath::Square(NetworkLargeClientCorrectionDistance));
-			return true;
-		}
-
-		// Check for disagreement in movement mode
-		const uint8 CurrentPackedMovementMode = PackNetworkMovementMode();
-		if (CurrentPackedMovementMode != ClientMovementMode)
+		if (ServerExceedsAllowablePositionError(ClientTimeStamp, DeltaTime, Accel, ClientWorldLocation, RelativeClientLocation, ClientMovementBase, ClientBaseBoneName, ClientMovementMode))
 		{
 			return true;
 		}
@@ -8915,6 +8905,28 @@ bool UCharacterMovementComponent::ServerCheckClientError(float ClientTimeStamp,
 }
 
 
+bool UCharacterMovementComponent::ServerExceedsAllowablePositionError(float ClientTimeStamp, float DeltaTime, const FVector& Accel, const FVector& ClientWorldLocation, const FVector& RelativeClientLocation, UPrimitiveComponent* ClientMovementBase, FName ClientBaseBoneName, uint8 ClientMovementMode)
+{
+	// Check for disagreement in movement mode
+	const uint8 CurrentPackedMovementMode = PackNetworkMovementMode();
+	if (CurrentPackedMovementMode != ClientMovementMode)
+	{
+		// Consider this a major correction, see SendClientAdjustment()
+		bNetworkLargeClientCorrection = true;
+		return true;
+	}
+
+	const FVector LocDiff = UpdatedComponent->GetComponentLocation() - ClientWorldLocation;	
+	const AGameNetworkManager* GameNetworkManager = (const AGameNetworkManager*)(AGameNetworkManager::StaticClass()->GetDefaultObject());
+	if (GameNetworkManager->ExceedsAllowablePositionError(LocDiff))
+	{
+		bNetworkLargeClientCorrection |= (LocDiff.SizeSquared() > FMath::Square(NetworkLargeClientCorrectionDistance));
+		return true;
+	}
+
+	return false;
+}
+
 bool UCharacterMovementComponent::ServerShouldUseAuthoritativePosition(float ClientTimeStamp, float DeltaTime, const FVector& Accel, const FVector& ClientWorldLocation, const FVector& RelativeClientLocation, UPrimitiveComponent* ClientMovementBase, FName ClientBaseBoneName, uint8 ClientMovementMode)
 {
 	if (bServerAcceptClientAuthoritativePosition)
