Fix a fallout from CL 9325556

diff --git a/Engine/Source/Runtime/Engine/Private/Streaming/TextureInstanceView.cpp b/Engine/Source/Runtime/Engine/Private/Streaming/TextureInstanceView.cpp
--- a/Engine/Source/Runtime/Engine/Private/Streaming/TextureInstanceView.cpp
+++ b/Engine/Source/Runtime/Engine/Private/Streaming/TextureInstanceView.cpp
@@ -561,7 +561,7 @@ void FRenderAssetInstanceAsyncView::GetRenderAssetScreenSize(
 							&& CompiledElement.TexelFactor == 0.f
 							&& ensure(CompiledElement.BoundsIndex < View->NumBounds4() * 4))
 						{
-							FRenderAssetInstanceView::FCompiledElement* MutableCompiledElement = const_cast<FRenderAssetInstanceView::FCompiledElement*>(CompiledElementData);
+							FRenderAssetInstanceView::FCompiledElement* MutableCompiledElement = const_cast<FRenderAssetInstanceView::FCompiledElement*>(&CompiledElement);
 							MutableCompiledElement->TexelFactor = View->GetBounds4(CompiledElement.BoundsIndex / 4).Radius.Component(CompiledElement.BoundsIndex % 4) * 2.f;
 						}
 
