Audio - Adding a static for preloading streaming audio.

diff --git a/Engine/Source/Runtime/Engine/Classes/Kismet/GameplayStatics.h b/Engine/Source/Runtime/Engine/Classes/Kismet/GameplayStatics.h
--- a/Engine/Source/Runtime/Engine/Classes/Kismet/GameplayStatics.h
+++ b/Engine/Source/Runtime/Engine/Classes/Kismet/GameplayStatics.h
@@ -710,6 +710,10 @@ public:
 	UFUNCTION(BlueprintCallable, Category="Audio", meta=(WorldContext = "WorldContextObject"))
 	static void SetBaseSoundMix(const UObject* WorldContextObject, class USoundMix* InSoundMix);
 
+	/** Primes the sound, caching the first chunk of streamed audio. **/
+	UFUNCTION(BlueprintCallable, Category = "Audio")
+	static void PrimeSound(USoundBase* InSound);
+
 	/** Overrides the sound class adjuster in the given sound mix. If the sound class does not exist in the input sound mix, the sound class adjustment will be added to the sound mix.
 	 * @param InSoundMixModifier The sound mix to modify.
 	 * @param InSoundClass The sound class to override (or add) in the sound mix.
