Audio - Adding a static for preloading streaming audio.

diff --git a/Engine/Source/Runtime/Engine/Private/GameplayStatics.cpp b/Engine/Source/Runtime/Engine/Private/GameplayStatics.cpp
--- a/Engine/Source/Runtime/Engine/Private/GameplayStatics.cpp
+++ b/Engine/Source/Runtime/Engine/Private/GameplayStatics.cpp
@@ -44,6 +44,8 @@
 #include "Async/Async.h"
 #include "Engine/SceneCapture2D.h"
 #include "Components/SceneCaptureComponent2D.h"
+#include "Sound/SoundCue.h"
+#include "Sound/SoundWave.h"
 
 #define LOCTEXT_NAMESPACE "GameplayStatics"
 
@@ -1607,6 +1609,26 @@ void UGameplayStatics::PushSoundMixModifier(const UObject* WorldContextObject, U
 	}
 }
 
+void UGameplayStatics::PrimeSound(USoundBase* InSound)
+{
+	if (!InSound || !GEngine || !GEngine->UseSound())
+	{
+		return;
+	}
+
+	if (USoundCue* InSoundCue = Cast<USoundCue>(InSound))
+	{
+		InSoundCue->PrimeSoundCue();
+	}
+	else if (USoundWave* InSoundWave = Cast<USoundWave>(InSound))
+	{
+		if (InSoundWave->GetNumChunks() > 1)
+		{
+			IStreamingManager::Get().GetAudioStreamingManager().RequestChunk(InSoundWave, 1, TFunction<void(EAudioChunkLoadResult)>());
+		}
+	}
+}
+
 void UGameplayStatics::SetSoundMixClassOverride(const UObject* WorldContextObject, class USoundMix* InSoundMixModifier, class USoundClass* InSoundClass, float Volume, float Pitch, float FadeInTime, bool bApplyToChildren)
 {
 	if (!InSoundMixModifier || !GEngine || !GEngine->UseSound())
