Fix G16 textures not being built correctly on console platforms
This caused world height virtual texture to be broken on console

diff --git a/Engine/Source/Runtime/Engine/Private/VT/VirtualTextureDataBuilder.cpp b/Engine/Source/Runtime/Engine/Private/VT/VirtualTextureDataBuilder.cpp
--- a/Engine/Source/Runtime/Engine/Private/VT/VirtualTextureDataBuilder.cpp
+++ b/Engine/Source/Runtime/Engine/Private/VT/VirtualTextureDataBuilder.cpp
@@ -781,6 +781,25 @@ int32 FVirtualTextureDataBuilder::FindSourceBlockIndex(int32 MipIndex, int32 Blo
 	return INDEX_NONE;
 }
 
+static const FName RemovePlatformPrefixFromName(FName const& InName)
+{
+	FString NameString = InName.ToString();
+	if (NameString.StartsWith(TEXT("PS4_")))
+	{
+		return *NameString.Replace(TEXT("PS4_"), TEXT(""));
+	}
+	else if (NameString.StartsWith("XBOXONE_"))
+	{
+		return *NameString.Replace(TEXT("XBOXONE_"), TEXT(""));
+	}
+	else if (NameString.StartsWith("SWITCH_"))
+	{
+		return *NameString.Replace(TEXT("SWITCH_"), TEXT(""));
+	}
+
+	return InName;
+}
+
 // This builds an uncompressed version of the texture containing all other build settings baked in
 // color corrections, mip sharpening, ....
 void FVirtualTextureDataBuilder::BuildSourcePixels(const FTextureSourceData& SourceData, const FTextureSourceData& CompositeSourceData)
@@ -800,7 +819,8 @@ void FVirtualTextureDataBuilder::BuildSourcePixels(const FTextureSourceData& Sou
 		LayerData.GammaSpace = BuildSettingsForLayer.GetGammaSpace();
 		LayerData.bHasAlpha = false;
 
-		const FName TextureFormatName = BuildSettingsForLayer.TextureFormatName;
+		const FName TextureFormatName = RemovePlatformPrefixFromName(BuildSettingsForLayer.TextureFormatName);
+
 		const bool bIsHdr = BuildSettingsForLayer.bHDRSource || TextureFormatName == "BC6H" || TextureFormatName == "RGBA16F";
 		const bool bIsG16 = TextureFormatName == "G16";
 
@@ -1065,20 +1085,7 @@ void FVirtualTextureDataBuilder::BuildSourcePixels(const FTextureSourceData& Sou
 
 		// Don't want platform specific swizzling for VT tile data, this tends to add extra padding for textures with odd dimensions
 		// (VT physical tiles generally not power-of-2 after adding border)
-		FName TextureFormatName = BuildSettingsForLayer.TextureFormatName;
-		FString BaseTextureFormatName = TextureFormatName.ToString();
-		if (BaseTextureFormatName.StartsWith(TEXT("PS4_")))
-		{
-			TextureFormatName = *BaseTextureFormatName.Replace(TEXT("PS4_"), TEXT(""));
-		}
-		else if (BaseTextureFormatName.StartsWith("XBOXONE_"))
-		{
-			TextureFormatName = *BaseTextureFormatName.Replace(TEXT("XBOXONE_"), TEXT(""));
-		}
-		else if (BaseTextureFormatName.StartsWith("SWITCH_"))
-		{
-			TextureFormatName = *BaseTextureFormatName.Replace(TEXT("SWITCH_"), TEXT(""));
-		}
+		FName TextureFormatName = RemovePlatformPrefixFromName(BuildSettingsForLayer.TextureFormatName);
 
 		// We handle AutoDXT specially here since otherwise the texture format compressor would choose a DXT format for every tile
 		// individually. Causing tiles in the same VT to use different formats which we don't allow.
