Ensure we log current network metrics when a user sends in-game feedback

#rnx

diff --git a/Engine/Source/Runtime/Engine/Public/Net/NetAnalyticsTypes.h b/Engine/Source/Runtime/Engine/Public/Net/NetAnalyticsTypes.h
--- a/Engine/Source/Runtime/Engine/Public/Net/NetAnalyticsTypes.h
+++ b/Engine/Source/Runtime/Engine/Public/Net/NetAnalyticsTypes.h
@@ -25,6 +25,13 @@ public:
 	float TimeSeconds;
 };
 
+// With this formatting, an array of these (or even just a single entry) will have
+// the same form as a map "key:value,key:value,..."
+inline FString LexToString(const FDelinquencyNameTimePair& Value)
+{
+	return FString::Printf(TEXT("%s:%f"), *Value.Name.ToString(), Value.TimeSeconds);
+}
+
 struct ENGINE_API FDelinquencyKeyFuncs : public BaseKeyFuncs<FDelinquencyNameTimePair, FDelinquencyNameTimePair, false>
 {
 	static KeyInitType GetSetKey(ElementInitType Element)
@@ -336,6 +343,11 @@ public:
 		return NumberOfTrackedFrames;
 	}
 
+	const float GetBurstyPacketLossPerSecond(double DurationSec) const
+	{
+		return DurationSec > 0.0 ? GetNumberOfFramesWithBurstsOfPacketLoss() / DurationSec : 0.0;
+	}
+
 	/** Resets the state of tracking. */
 	void Reset();
 
