Removed unused velocity shaders on translucency. It was caused by FMaterial::WritesEveryPixel returning true for translucent blend modes. Changed WritesEveryPixel back, to return true for the translucents, as they write all the pixels.

diff --git a/Engine/Source/Runtime/Engine/Private/Materials/MaterialShared.cpp b/Engine/Source/Runtime/Engine/Private/Materials/MaterialShared.cpp
--- a/Engine/Source/Runtime/Engine/Private/Materials/MaterialShared.cpp
+++ b/Engine/Source/Runtime/Engine/Private/Materials/MaterialShared.cpp
@@ -3008,7 +3008,6 @@ bool FMaterial::WritesEveryPixel(bool bShadowPass) const
 	}
 
 	return !IsMasked()
-		&& !IsTranslucentBlendMode(GetBlendMode())
 		// Render dithered material as masked if a stencil prepass is not used (UE-50064, UE-49537)
 		&& !((bShadowPass || !bStencilDitheredLOD) && IsDitheredLODTransition())
 		&& !IsWireframe()
