TimerManager merge fix.
Missing header fix.

diff --git a/Engine/Source/Editor/Blutility/Private/EditorUtilitySubsystem.cpp b/Engine/Source/Editor/Blutility/Private/EditorUtilitySubsystem.cpp
--- a/Engine/Source/Editor/Blutility/Private/EditorUtilitySubsystem.cpp
+++ b/Engine/Source/Editor/Blutility/Private/EditorUtilitySubsystem.cpp
@@ -3,6 +3,7 @@
 #include "EditorUtilitySubsystem.h"
 #include "EditorUtilityCommon.h"
 #include "Interfaces/IMainFrameModule.h"
+#include "Engine/Blueprint.h"
 
 UEditorUtilitySubsystem::UEditorUtilitySubsystem() :
 	UEditorSubsystem()
diff --git a/Engine/Source/Runtime/Engine/Private/TimerManager.cpp b/Engine/Source/Runtime/Engine/Private/TimerManager.cpp
--- a/Engine/Source/Runtime/Engine/Private/TimerManager.cpp
+++ b/Engine/Source/Runtime/Engine/Private/TimerManager.cpp
@@ -954,6 +954,9 @@ TStatId FTimerManager::GetStatId() const
 
 void FTimerManager::SetGameInstance(UGameInstance* InGameInstance)
 {
+	// not currently threadsafe
+	check(IsInGameThread());
+
 	OwningGameInstance = InGameInstance;
 
 #if UE_ENABLE_TRACKING_TIMER_SOURCES
@@ -1003,14 +1006,6 @@ void FTimerManager::ListTimers() const
 	UE_LOG(LogEngine, Log, TEXT("------- %d Total Timers -------"), PendingTimerSet.Num() + PausedTimerSet.Num() + ValidActiveTimers.Num());
 }
 
-void FTimerManager::SetGameInstance(UGameInstance* InGameInstance)
-{
-	// not currently threadsafe
-	check(IsInGameThread());
-
-	OwningGameInstance = InGameInstance;
-}
-
 FTimerHandle FTimerManager::AddTimer(FTimerData&& TimerData)
 {
 	const void* TimerIndicesByObjectKey = TimerData.TimerDelegate.GetBoundObject();
