PR #5828: Adding missing overload to Visual Logger for logging an arrow with a log category, instead of a log category name. (Contributed by [REDACTED_PERSON])

diff --git a/Engine/Source/Runtime/Engine/Public/VisualLogger/VisualLogger.h b/Engine/Source/Runtime/Engine/Public/VisualLogger/VisualLogger.h
--- a/Engine/Source/Runtime/Engine/Public/VisualLogger/VisualLogger.h
+++ b/Engine/Source/Runtime/Engine/Public/VisualLogger/VisualLogger.h
@@ -170,6 +170,7 @@ class ENGINE_API FVisualLogger : public FOutputDevice
 	static void GeometryConvexLogfImpl(const UObject* LogOwner, const FName& CategoryName, ELogVerbosity::Type Verbosity, const TArray<FVector>& Points, const FColor& Color, const TCHAR* Fmt, ...);
 	static void HistogramDataLogfImpl (const UObject* LogOwner, const FLogCategoryBase& Category, ELogVerbosity::Type Verbosity, FName GraphName, FName DataName, const FVector2D& Data, const FColor& Color, const TCHAR* Fmt, ...);
 	static void HistogramDataLogfImpl (const UObject* LogOwner, const FName& CategoryName, ELogVerbosity::Type Verbosity, FName GraphName, FName DataName, const FVector2D& Data, const FColor& Color, const TCHAR* Fmt, ...);
+	static void ArrowLogfImpl(const UObject* LogOwner, const FLogCategoryBase& Category, ELogVerbosity::Type Verbosity, const FVector& Start, const FVector& End, const FColor& Color, const TCHAR* Fmt, ...);
 	static void ArrowLogfImpl(const UObject* LogOwner, const FName& CategoryName, ELogVerbosity::Type Verbosity, const FVector& Start, const FVector& End, const FColor& Color, const TCHAR* Fmt, ...);
 
 public:
@@ -211,6 +212,13 @@ public:
 
 	// Arrow
 	template <typename FmtType, typename... Types>
+	static void ArrowLogf(const UObject* LogOwner, const FLogCategoryBase& Category, ELogVerbosity::Type Verbosity, const FVector& Start, const FVector& End, const FColor& Color, const FmtType& Fmt, Types... Args)
+	{
+		static_assert(TIsArrayOrRefOfType<FmtType, TCHAR>::Value, "Formatting string must be a TCHAR array.");
+		static_assert(TAnd<TIsValidVariadicFunctionArg<Types>...>::Value, "Invalid argument(s) passed to FVisualLogger::ArrowLogf");
+		ArrowLogfImpl(LogOwner, Category, Verbosity, Start, End, Color, Fmt, Args...);
+	}
+	template <typename FmtType, typename... Types>
 	static void ArrowLogf(const UObject* LogOwner, const FName& CategoryName, ELogVerbosity::Type Verbosity, const FVector& Start, const FVector& End, const FColor& Color, const FmtType& Fmt, Types... Args)
 	{
 		static_assert(TIsArrayOrRefOfType<FmtType, TCHAR>::Value, "Formatting string must be a TCHAR array.");
@@ -625,6 +633,13 @@ inline void FVisualLogger::GeometryShapeLogfImpl(const UObject* Object, const FN
 	);
 }
 
+inline void FVisualLogger::ArrowLogfImpl(const UObject* Object, const FLogCategoryBase& Category, ELogVerbosity::Type Verbosity, const FVector& Start, const FVector& End, const FColor& Color, const TCHAR* Fmt, ...)
+{
+	const FName CategoryName = Category.GetCategoryName();
+	COLLAPSED_LOGF(
+		CurrentEntry->AddArrow(Start, End, CategoryName, Verbosity, Color, Buffer);
+	);
+}
 inline void FVisualLogger::ArrowLogfImpl(const UObject* Object, const FName& CategoryName, ELogVerbosity::Type Verbosity, const FVector& Start, const FVector& End, const FColor& Color, const TCHAR* Fmt, ...)
 {
 	COLLAPSED_LOGF(
