Strip PIE prefix from path name when looking for a net export group (particularly for level scrript actor classes)

diff --git a/Engine/Source/Runtime/Engine/Private/PackageMapClient.cpp b/Engine/Source/Runtime/Engine/Private/PackageMapClient.cpp
--- a/Engine/Source/Runtime/Engine/Private/PackageMapClient.cpp
+++ b/Engine/Source/Runtime/Engine/Private/PackageMapClient.cpp
@@ -1269,9 +1269,9 @@ void UPackageMapClient::ReceiveNetGUIDBunch( FInBunch &InBunch )
 	UE_LOG(LogNetPackageMap, Log, TEXT("UPackageMapClient::ReceiveNetGUIDBunch end. BitPos: %d"), InBunch.GetPosBits() );
 }
 
-TSharedPtr< FNetFieldExportGroup > UPackageMapClient::GetNetFieldExportGroup( const FString& PathName )
+TSharedPtr<FNetFieldExportGroup> UPackageMapClient::GetNetFieldExportGroup(const FString& PathName)
 {
-	return GuidCache->NetFieldExportGroupMap.FindRef( PathName );
+	return GuidCache->NetFieldExportGroupMap.FindRef(UWorld::RemovePIEPrefix(PathName));
 }
 
 void UPackageMapClient::AddNetFieldExportGroup( const FString& PathName, TSharedPtr< FNetFieldExportGroup > NewNetFieldExportGroup )
