Presence of export bunches was resulting in bPartial bit of the export bunch and the next bunch set to 1 even though the size of these bunches were less than the max bunch size. Changed that behaviour so that the bPartial flag of these bunches are no longer set to 1.

diff --git a/Engine/Source/Runtime/Engine/Private/DataChannel.cpp b/Engine/Source/Runtime/Engine/Private/DataChannel.cpp
--- a/Engine/Source/Runtime/Engine/Private/DataChannel.cpp
+++ b/Engine/Source/Runtime/Engine/Private/DataChannel.cpp
@@ -536,6 +536,12 @@ bool UChannel::ReceivedNextBunch( FInBunch & Bunch, bool & bOutSkipAck )
 		Connection->InReliable[Bunch.ChIndex] = Bunch.ChSequence;
 	}
 
+	if (Bunch.bHasPackageMapExports && !Connection->InternalAck)
+	{
+		LogPartialBunch(TEXT("Received New partial bunch. It only contained NetGUIDs."), Bunch, *InPartialBunch);
+		return false;
+	}
+
 	FInBunch* HandleBunch = &Bunch;
 	if (Bunch.bPartial)
 	{
@@ -570,7 +576,7 @@ bool UChannel::ReceivedNextBunch( FInBunch & Bunch, bool & bOutSkipAck )
 			}
 
 			InPartialBunch = new FInBunch(Bunch, false);
-			if ( !Bunch.bHasPackageMapExports && Bunch.GetBitsLeft() > 0 )
+			if (ensureMsgf((!Bunch.bHasPackageMapExports && Bunch.GetBitsLeft() > 0), TEXT("Received new partial bunch with bHasPackageMapExports flag set to 1. This should not have happened. %s"), *Describe()))
 			{
 				if ( Bunch.GetBitsLeft() % 8 != 0 )
 				{
@@ -583,10 +589,6 @@ bool UChannel::ReceivedNextBunch( FInBunch & Bunch, bool & bOutSkipAck )
 
 				LogPartialBunch(TEXT("Received new partial bunch."), Bunch, *InPartialBunch);
 			}
-			else
-			{
-				LogPartialBunch(TEXT("Received New partial bunch. It only contained NetGUIDs."), Bunch, *InPartialBunch);
-			}
 		}
 		else
 		{
@@ -949,11 +951,13 @@ FPacketIdRange UChannel::SendBunch( FOutBunch* Bunch, bool Merge )
 	TArray<FOutBunch*>& OutgoingBunches = Connection->GetOutgoingBunches();
 	OutgoingBunches.Reset();
 
+	int32 NumOfExportBunches = 0;
 	// Add any export bunches
 	// Replay connections will manage export bunches separately.
 	if (!Connection->InternalAck)
 	{
 		AppendExportBunches( OutgoingBunches );
+		NumOfExportBunches = OutgoingBunches.Num();
 	}
 
 	if ( OutgoingBunches.Num() )
@@ -1073,7 +1077,7 @@ FPacketIdRange UChannel::SendBunch( FOutBunch* Bunch, bool Merge )
 		return PacketIdRange;
 	}
 
-	UE_CLOG((OutgoingBunches.Num() > 1), LogNetPartialBunch, Log, TEXT("Sending %d Bunches. Channel: %d %s"), OutgoingBunches.Num(), Bunch->ChIndex, *Describe());
+	UE_CLOG((OutgoingBunches.Num() > 1 + NumOfExportBunches), LogNetPartialBunch, Log, TEXT("Sending %d Bunches. Channel: %d %s"), OutgoingBunches.Num() - NumOfExportBunches, Bunch->ChIndex, *Describe());
 	for( int32 PartialNum = 0; PartialNum < OutgoingBunches.Num(); ++PartialNum)
 	{
 		FOutBunch * NextBunch = OutgoingBunches[PartialNum];
@@ -1093,12 +1097,12 @@ FPacketIdRange UChannel::SendBunch( FOutBunch* Bunch, bool Merge )
 			NextBunch->bHasMustBeMappedGUIDs |= Bunch->bHasMustBeMappedGUIDs;
 		}
 
-		if (OutgoingBunches.Num() > 1)
+		if (OutgoingBunches.Num() > 1 + NumOfExportBunches)
 		{
-			NextBunch->bPartial = 1;
-			NextBunch->bPartialInitial = (PartialNum == 0 ? 1: 0);
+			NextBunch->bPartial = (PartialNum >= NumOfExportBunches ? 1 : 0);               // do not set bPartial bit of export bunches to 1
+			NextBunch->bPartialInitial = (PartialNum == NumOfExportBunches ? 1: 0);
 			NextBunch->bPartialFinal = (PartialNum == OutgoingBunches.Num() - 1 ? 1: 0);
-			NextBunch->bOpen &= (PartialNum == 0);											// Only the first bunch should have the bOpen bit set
+			NextBunch->bOpen &= (PartialNum == NumOfExportBunches);											// Only the first bunch should have the bOpen bit set
 			NextBunch->bClose = (Bunch->bClose && (OutgoingBunches.Num()-1 == PartialNum)); // Only last bunch should have bClose bit set
 		}
 
