Updated log to fix the static analysis warning for the previous submit.

diff --git a/Engine/Source/Runtime/Engine/Private/DataChannel.cpp b/Engine/Source/Runtime/Engine/Private/DataChannel.cpp
--- a/Engine/Source/Runtime/Engine/Private/DataChannel.cpp
+++ b/Engine/Source/Runtime/Engine/Private/DataChannel.cpp
@@ -538,7 +538,7 @@ bool UChannel::ReceivedNextBunch( FInBunch & Bunch, bool & bOutSkipAck )
 
 	if (Bunch.bHasPackageMapExports && !Connection->InternalAck)
 	{
-		LogPartialBunch(TEXT("Received New partial bunch. It only contained NetGUIDs."), Bunch, *InPartialBunch);
+		UE_LOG(LogNet, VeryVerbose, TEXT("Received new bunch. It only contained NetGUIDs and was processed in UChannel::ReceivedRawBunch(). PacketId %d. ChSequence %d. ChIndex %d"), Bunch.PacketId, Bunch.ChSequence, Bunch.ChIndex);
 		return false;
 	}
 
