Log errors from demoplay console command as LogDemo/Error

diff --git a/Engine/Source/Runtime/Engine/Private/World.cpp b/Engine/Source/Runtime/Engine/Private/World.cpp
--- a/Engine/Source/Runtime/Engine/Private/World.cpp
+++ b/Engine/Source/Runtime/Engine/Private/World.cpp
@@ -3795,7 +3795,7 @@ bool UWorld::HandleDemoPlayCommand( const TCHAR* Cmd, FOutputDevice& Ar, UWorld*
 
 	if (ErrorString != nullptr)
 	{
-		Ar.Log(ErrorString);
+		UE_SUPPRESS(LogDemo, Error, Ar.Log(ErrorString));
 
 		if (GetGameInstance() != nullptr)
 		{
