Cleaning up control channel message macros

diff --git a/Engine/Plugins/NetcodeUnitTest/NetcodeUnitTest/Source/NetcodeUnitTest/Classes/NUTActor.h b/Engine/Plugins/NetcodeUnitTest/NetcodeUnitTest/Source/NetcodeUnitTest/Classes/NUTActor.h
--- a/Engine/Plugins/NetcodeUnitTest/NetcodeUnitTest/Source/NetcodeUnitTest/Classes/NUTActor.h
+++ b/Engine/Plugins/NetcodeUnitTest/NetcodeUnitTest/Source/NetcodeUnitTest/Classes/NUTActor.h
@@ -33,7 +33,7 @@ enum class ENUTControlCommand : uint8
 };
 
 // Custom control channel message, used to communicate with the server NUTActor over the control channel
-DEFINE_CONTROL_CHANNEL_MESSAGE_TWOPARAM(NUTControl, 250, ENUTControlCommand, FString);
+DEFINE_CONTROL_CHANNEL_MESSAGE(NUTControl, 250, ENUTControlCommand, FString);
 
 
 /**
diff --git a/Engine/Source/Runtime/Engine/Public/Net/DataChannel.h b/Engine/Source/Runtime/Engine/Public/Net/DataChannel.h
--- a/Engine/Source/Runtime/Engine/Public/Net/DataChannel.h
+++ b/Engine/Source/Runtime/Engine/Public/Net/DataChannel.h
@@ -20,7 +20,7 @@
 /** network control channel message types
  * 
  * to add a new message type, you need to:
- * - add a DEFINE_CONTROL_CHANNEL_MESSAGE_* for the message type with the appropriate parameters to this file
+ * - add a DEFINE_CONTROL_CHANNEL_MESSAGE for the message type with the appropriate parameters to this file
  * - add IMPLEMENT_CONTROL_CHANNEL_MESSAGE for the message type to DataChannel.cpp
  * - implement the fallback behavior (eat an unparsed message) to UControlChannel::ReceivedBunch()
  *
@@ -63,355 +63,109 @@ private:
 		Names[MessageType] = InName;
 	}
 
-	static const TCHAR* Names[256];
+	template<typename... ParamTypes>
+	static void SendParams(FControlChannelOutBunch& Bunch, ParamTypes&... Params) {}
+
+	template<typename FirstParamType, typename... ParamTypes>
+	static void SendParams(FControlChannelOutBunch& Bunch, FirstParamType& FirstParam, ParamTypes&... Params)
+	{
+		Bunch << FirstParam;
+		SendParams(Bunch, Params...);
+	}
+
+	template<typename... ParamTypes>
+	static void ReceiveParams(FInBunch& Bunch, ParamTypes&... Params) {}
+	
+	template<typename FirstParamType, typename... ParamTypes>
+	static void ReceiveParams(FInBunch& Bunch, FirstParamType& FirstParam, ParamTypes&... Params)
+	{
+		Bunch << FirstParam;
+		ReceiveParams(Bunch, Params...);
+	}
+	
+	static constexpr int MaxNames = TNumericLimits<uint8>::Max() + 1;
+	static const TCHAR* Names[MaxNames];
 };
 
-#define DEFINE_CONTROL_CHANNEL_MESSAGE_ZEROPARAM(Name, Index) \
-	enum { NMT_##Name = Index }; \
-	template<> class FNetControlMessage<Index> \
+#define DEFINE_CONTROL_CHANNEL_MESSAGE(Name, Index, ...) \
+enum { NMT_##Name = Index }; \
+template<> class FNetControlMessage<Index> \
+{ \
+public: \
+	static uint8 Initialize() \
 	{ \
-	public: \
-		static uint8 Initialize() \
-		{ \
-			FNetControlMessageInfo::SetName(Index, TEXT(#Name)); \
-			return 0; \
-		} \
-		/** sends a message of this type on the specified connection's control channel */ \
-		static void Send(UNetConnection* Conn) \
-		{ \
-			static_assert(Index < 256, "Control channel message must be a byte."); \
-			checkSlow(!Conn->IsA(UChildConnection::StaticClass())); /** control channel messages can only be sent on the parent connection */ \
-			if (Conn->Channels[0] != NULL && !Conn->Channels[0]->Closing) \
-			{ \
-				FControlChannelOutBunch Bunch(Conn->Channels[0], false); \
-				uint8 MessageType = Index; \
-				Bunch << MessageType; \
-				Conn->Channels[0]->SendBunch(&Bunch, true); \
-			} \
-		} \
-	};
-#define DEFINE_CONTROL_CHANNEL_MESSAGE_ONEPARAM(Name, Index, TypeA) \
-	enum { NMT_##Name = Index }; \
-	template<> class FNetControlMessage<Index> \
+		FNetControlMessageInfo::SetName(Index, TEXT(#Name)); \
+		return 0; \
+	} \
+	/** sends a message of this type on the specified connection's control channel \
+		* @note: const not used only because of the FArchive interface; the parameters are not modified \
+		*/ \
+	template<typename... ParamTypes> \
+	static void Send(UNetConnection* Conn, ParamTypes&... Params) \
 	{ \
-	public: \
-		static uint8 Initialize() \
-		{ \
-			FNetControlMessageInfo::SetName(Index, TEXT(#Name)); \
-			return 0; \
-		} \
-		/** sends a message of this type on the specified connection's control channel \
-		 * @note: const not used only because of the FArchive interface; the parameters are not modified \
-		 */ \
-		static void Send(UNetConnection* Conn, TypeA& ParamA) \
-		{ \
-			static_assert(Index < 256, "Control channel message must be a byte."); \
-			checkSlow(!Conn->IsA(UChildConnection::StaticClass())); /** control channel messages can only be sent on the parent connection */ \
-			if (Conn->Channels[0] != NULL && !Conn->Channels[0]->Closing) \
-			{ \
-				FControlChannelOutBunch Bunch(Conn->Channels[0], false); \
-				uint8 MessageType = Index; \
-				Bunch << MessageType; \
-				Bunch << ParamA; \
-				Conn->Channels[0]->SendBunch(&Bunch, true); \
-			} \
-		} \
-		/** receives a message of this type from the passed in bunch */ \
-		FUNCTION_CHECK_RETURN_START \
-		static bool Receive(FInBunch& Bunch, TypeA& ParamA) \
-		FUNCTION_CHECK_RETURN_END \
-		{ \
-			Bunch << ParamA; \
-			return !Bunch.IsError(); \
-		} \
-		/** throws away a message of this type from the passed in bunch */ \
-		static void Discard(FInBunch& Bunch) \
-		{ \
-			TypeA ParamA; \
-			if (Receive(Bunch, ParamA)) {} \
-		} \
-	};
-#define DEFINE_CONTROL_CHANNEL_MESSAGE_TWOPARAM(Name, Index, TypeA, TypeB) \
-	enum { NMT_##Name = Index }; \
-	template<> class FNetControlMessage<Index> \
-	{ \
-	public: \
-		static uint8 Initialize() \
-		{ \
-			FNetControlMessageInfo::SetName(Index, TEXT(#Name)); \
-			return 0; \
-		} \
-		/** sends a message of this type on the specified connection's control channel \
-		 * @note: const not used only because of the FArchive interface; the parameters are not modified \
-		 */ \
-		static void Send(UNetConnection* Conn, TypeA& ParamA, TypeB& ParamB) \
-		{ \
-			static_assert(Index < 256, "Control channel message must be a byte."); \
-			checkSlow(!Conn->IsA(UChildConnection::StaticClass())); /** control channel messages can only be sent on the parent connection */ \
-			if (Conn->Channels[0] != NULL && !Conn->Channels[0]->Closing) \
-			{ \
-				FControlChannelOutBunch Bunch(Conn->Channels[0], false); \
-				uint8 MessageType = Index; \
-				Bunch << MessageType; \
-				Bunch << ParamA; \
-				Bunch << ParamB; \
-				Conn->Channels[0]->SendBunch(&Bunch, true); \
-			} \
-		} \
-		FUNCTION_CHECK_RETURN_START \
-		static bool Receive(FInBunch& Bunch, TypeA& ParamA, TypeB& ParamB) \
-		FUNCTION_CHECK_RETURN_END \
-		{ \
-			Bunch << ParamA; \
-			Bunch << ParamB; \
-			return !Bunch.IsError(); \
-		} \
-		/** throws away a message of this type from the passed in bunch */ \
-		static void Discard(FInBunch& Bunch) \
-		{ \
-			TypeA ParamA; \
-			TypeB ParamB; \
-			if (Receive(Bunch, ParamA, ParamB)) {} \
-		} \
-	};
-#define DEFINE_CONTROL_CHANNEL_MESSAGE_THREEPARAM(Name, Index, TypeA, TypeB, TypeC) \
-	enum { NMT_##Name = Index }; \
-	template<> class FNetControlMessage<Index> \
-	{ \
-	public: \
-		static uint8 Initialize() \
-		{ \
-			FNetControlMessageInfo::SetName(Index, TEXT(#Name)); \
-			return 0; \
-		} \
-		/** sends a message of this type on the specified connection's control channel \
-		 * @note: const not used only because of the FArchive interface; the parameters are not modified \
-		 */ \
-		static void Send(UNetConnection* Conn, TypeA& ParamA, TypeB& ParamB, TypeC& ParamC) \
-		{ \
-			static_assert(Index < 256, "Control channel message must be a byte."); \
-			checkSlow(!Conn->IsA(UChildConnection::StaticClass())); /** control channel messages can only be sent on the parent connection */ \
-			if (Conn->Channels[0] != NULL && !Conn->Channels[0]->Closing) \
-			{ \
-				FControlChannelOutBunch Bunch(Conn->Channels[0], false); \
-				uint8 MessageType = Index; \
-				Bunch << MessageType; \
-				Bunch << ParamA; \
-				Bunch << ParamB; \
-				Bunch << ParamC; \
-				Conn->Channels[0]->SendBunch(&Bunch, true); \
-			} \
-		} \
-		FUNCTION_CHECK_RETURN_START \
-		static bool Receive(FInBunch& Bunch, TypeA& ParamA, TypeB& ParamB, TypeC& ParamC) \
-		FUNCTION_CHECK_RETURN_END \
-		{ \
-			Bunch << ParamA; \
-			Bunch << ParamB; \
-			Bunch << ParamC; \
-			return !Bunch.IsError(); \
-		} \
-		/** throws away a message of this type from the passed in bunch */ \
-		static void Discard(FInBunch& Bunch) \
-		{ \
-			TypeA ParamA; \
-			TypeB ParamB; \
-			TypeC ParamC; \
-			if (Receive(Bunch, ParamA, ParamB, ParamC)) {} \
-		} \
-	};
-#define DEFINE_CONTROL_CHANNEL_MESSAGE_FOURPARAM(Name, Index, TypeA, TypeB, TypeC, TypeD) \
-	enum { NMT_##Name = Index }; \
-	template<> class FNetControlMessage<Index> \
+		static_assert(Index < FNetControlMessageInfo::MaxNames, "Control channel message must be a byte."); \
+		checkSlow(!Conn->IsA(UChildConnection::StaticClass())); /** control channel messages can only be sent on the parent connection */ \
+		if (Conn->Channels[0] != NULL && !Conn->Channels[0]->Closing) \
+		{ \
+			FControlChannelOutBunch Bunch(Conn->Channels[0], false); \
+			uint8 MessageType = Index; \
+			Bunch << MessageType; \
+			FNetControlMessageInfo::SendParams(Bunch, Params...); \
+			Conn->Channels[0]->SendBunch(&Bunch, true); \
+		} \
+	} \
+	/** receives a message of this type from the passed in bunch */ \
+	template<typename... ParamTypes> \
+	FUNCTION_CHECK_RETURN_START \
+	static bool Receive(FInBunch& Bunch, ParamTypes&... Params) \
+	FUNCTION_CHECK_RETURN_END \
 	{ \
-	public: \
-		static uint8 Initialize() \
-		{ \
-			FNetControlMessageInfo::SetName(Index, TEXT(#Name)); \
-			return 0; \
-		} \
-		/** sends a message of this type on the specified connection's control channel \
-		 * @note: const not used only because of the FArchive interface; the parameters are not modified \
-		 */ \
-		static void Send(UNetConnection* Conn, TypeA& ParamA, TypeB& ParamB, TypeC& ParamC, TypeD& ParamD) \
-		{ \
-			static_assert(Index < 256, "Control channel message must be a byte."); \
-			checkSlow(!Conn->IsA(UChildConnection::StaticClass())); /** control channel messages can only be sent on the parent connection */ \
-			if (Conn->Channels[0] != NULL && !Conn->Channels[0]->Closing) \
-			{ \
-				FControlChannelOutBunch Bunch(Conn->Channels[0], false); \
-				uint8 MessageType = Index; \
-				Bunch << MessageType; \
-				Bunch << ParamA; \
-				Bunch << ParamB; \
-				Bunch << ParamC; \
-				Bunch << ParamD; \
-				Conn->Channels[0]->SendBunch(&Bunch, true); \
-			} \
-		} \
-		FUNCTION_CHECK_RETURN_START \
-		static bool Receive(FInBunch& Bunch, TypeA& ParamA, TypeB& ParamB, TypeC& ParamC, TypeD& ParamD) \
-		FUNCTION_CHECK_RETURN_END \
-		{ \
-			Bunch << ParamA; \
-			Bunch << ParamB; \
-			Bunch << ParamC; \
-			Bunch << ParamD; \
-			return !Bunch.IsError(); \
-		} \
-		/** throws away a message of this type from the passed in bunch */ \
-		static void Discard(FInBunch& Bunch) \
-		{ \
-			TypeA ParamA; \
-			TypeB ParamB; \
-			TypeC ParamC; \
-			TypeD ParamD; \
-			if (Receive(Bunch, ParamA, ParamB, ParamC, ParamD)) {} \
-		} \
-	};
-
-#define DEFINE_CONTROL_CHANNEL_MESSAGE_SEVENPARAM(Name, Index, TypeA, TypeB, TypeC, TypeD, TypeE, TypeF, TypeG) \
-	enum { NMT_##Name = Index }; \
-	template<> class FNetControlMessage<Index> \
+		FNetControlMessageInfo::ReceiveParams(Bunch, Params...); \
+		return !Bunch.IsError(); \
+	} \
+	/** throws away a message of this type from the passed in bunch */ \
+	static void Discard(FInBunch& Bunch) \
 	{ \
-	public: \
-		static uint8 Initialize() \
-		{ \
-			FNetControlMessageInfo::SetName(Index, TEXT(#Name)); \
-			return 0; \
-		} \
-		/** sends a message of this type on the specified connection's control channel \
-		 * @note: const not used only because of the FArchive interface; the parameters are not modified \
-		 */ \
-		static void Send(UNetConnection* Conn, TypeA& ParamA, TypeB& ParamB, TypeC& ParamC, TypeD& ParamD, TypeE& ParamE, TypeF& ParamF, TypeG& ParamG) \
+		TTuple<__VA_ARGS__> Params; \
+		VisitTupleElements([&Bunch](auto& Param) \
 		{ \
-			static_assert(Index < 256, "Control channel message must be a byte."); \
-			checkSlow(!Conn->IsA(UChildConnection::StaticClass())); /** control channel messages can only be sent on the parent connection */ \
-			if (Conn->Channels[0] != NULL && !Conn->Channels[0]->Closing) \
-			{ \
-				FControlChannelOutBunch Bunch(Conn->Channels[0], false); \
-				uint8 MessageType = Index; \
-				Bunch << MessageType; \
-				Bunch << ParamA; \
-				Bunch << ParamB; \
-				Bunch << ParamC; \
-				Bunch << ParamD; \
-				Bunch << ParamE; \
-				Bunch << ParamF; \
-				Bunch << ParamG; \
-				Conn->Channels[0]->SendBunch(&Bunch, true); \
-			} \
-		} \
-		FUNCTION_CHECK_RETURN_START \
-		static bool Receive(FInBunch& Bunch, TypeA& ParamA, TypeB& ParamB, TypeC& ParamC, TypeD& ParamD, TypeE& ParamE, TypeF& ParamF, TypeG& ParamG) \
-		FUNCTION_CHECK_RETURN_END \
-		{ \
-			Bunch << ParamA; \
-			Bunch << ParamB; \
-			Bunch << ParamC; \
-			Bunch << ParamD; \
-			Bunch << ParamE; \
-			Bunch << ParamF; \
-			Bunch << ParamG; \
-			return !Bunch.IsError(); \
-		} \
-		/** throws away a message of this type from the passed in bunch */ \
-		static void Discard(FInBunch& Bunch) \
-		{ \
-			TypeA ParamA; \
-			TypeB ParamB; \
-			TypeC ParamC; \
-			TypeD ParamD; \
-			TypeE ParamE; \
-			TypeF ParamF; \
-			TypeG ParamG; \
-			if (Receive(Bunch, ParamA, ParamB, ParamC, ParamD, ParamE, ParamF, ParamG)) {} \
-		} \
-	};
+			Bunch << Param; \
+		}, \
+		Params); \
+	} \
+};
 
-#define DEFINE_CONTROL_CHANNEL_MESSAGE_EIGHTPARAM(Name, Index, TypeA, TypeB, TypeC, TypeD, TypeE, TypeF, TypeG, TypeH) \
-	enum { NMT_##Name = Index }; \
-	template<> class FNetControlMessage<Index> \
-	{ \
-	public: \
-		static uint8 Initialize() \
-		{ \
-			FNetControlMessageInfo::SetName(Index, TEXT(#Name)); \
-			return 0; \
-		} \
-		/** sends a message of this type on the specified connection's control channel \
-		 * @note: const not used only because of the FArchive interface; the parameters are not modified \
-		 */ \
-		static void Send(UNetConnection* Conn, TypeA& ParamA, TypeB& ParamB, TypeC& ParamC, TypeD& ParamD, TypeE& ParamE, TypeF& ParamF, TypeG& ParamG, TypeH& ParamH) \
-		{ \
-			static_assert(Index < 256, "Control channel message must be a byte."); \
-			checkSlow(!Conn->IsA(UChildConnection::StaticClass())); /** control channel messages can only be sent on the parent connection */ \
-			if (Conn->Channels[0] != NULL && !Conn->Channels[0]->Closing) \
-			{ \
-				FControlChannelOutBunch Bunch(Conn->Channels[0], false); \
-				uint8 MessageType = Index; \
-				Bunch << MessageType; \
-				Bunch << ParamA; \
-				Bunch << ParamB; \
-				Bunch << ParamC; \
-				Bunch << ParamD; \
-				Bunch << ParamE; \
-				Bunch << ParamF; \
-				Bunch << ParamG; \
-				Bunch << ParamH; \
-				Conn->Channels[0]->SendBunch(&Bunch, true); \
-			} \
-		} \
-		FUNCTION_CHECK_RETURN_START \
-		static bool Receive(FInBunch& Bunch, TypeA& ParamA, TypeB& ParamB, TypeC& ParamC, TypeD& ParamD, TypeE& ParamE, TypeF& ParamF, TypeG& ParamG, TypeH& ParamH) \
-		FUNCTION_CHECK_RETURN_END \
-		{ \
-			Bunch << ParamA; \
-			Bunch << ParamB; \
-			Bunch << ParamC; \
-			Bunch << ParamD; \
-			Bunch << ParamE; \
-			Bunch << ParamF; \
-			Bunch << ParamG; \
-			Bunch << ParamH; \
-			return !Bunch.IsError(); \
-		} \
-		/** throws away a message of this type from the passed in bunch */ \
-		static void Discard(FInBunch& Bunch) \
-		{ \
-			TypeA ParamA; \
-			TypeB ParamB; \
-			TypeC ParamC; \
-			TypeD ParamD; \
-			TypeE ParamE; \
-			TypeF ParamF; \
-			TypeG ParamG; \
-			TypeH ParamH; \
-			if (Receive(Bunch, ParamA, ParamB, ParamC, ParamD, ParamE, ParamF, ParamG, ParamH)) {} \
-		} \
-	};
+#define DEFINE_CONTROL_CHANNEL_MESSAGE_ZEROPARAM(Name, Index) DEFINE_CONTROL_CHANNEL_MESSAGE(Name, Index)
+#define DEFINE_CONTROL_CHANNEL_MESSAGE_ONEPARAM(Name, Index, TypeA) DEFINE_CONTROL_CHANNEL_MESSAGE(Name, Index, TypeA)
+#define DEFINE_CONTROL_CHANNEL_MESSAGE_TWOPARAM(Name, Index, TypeA, TypeB) DEFINE_CONTROL_CHANNEL_MESSAGE(Name, Index, TypeA, TypeB)
+#define DEFINE_CONTROL_CHANNEL_MESSAGE_THREEPARAM(Name, Index, TypeA, TypeB, TypeC) DEFINE_CONTROL_CHANNEL_MESSAGE(Name, Index, TypeA, TypeB, TypeC)
+#define DEFINE_CONTROL_CHANNEL_MESSAGE_FOURPARAM(Name, Index, TypeA, TypeB, TypeC, TypeD) DEFINE_CONTROL_CHANNEL_MESSAGE(Name, Index, TypeA, TypeB, TypeC, TypeD)
+#define DEFINE_CONTROL_CHANNEL_MESSAGE_FIVEPARAM(Name, Index, TypeA, TypeB, TypeC, TypeD, TypeE) DEFINE_CONTROL_CHANNEL_MESSAGE(Name, Index, TypeA, TypeB, TypeC, TypeD, TypeE)
+#define DEFINE_CONTROL_CHANNEL_MESSAGE_SIXPARAM(Name, Index, TypeA, TypeB, TypeC, TypeD, TypeE, TypeF) DEFINE_CONTROL_CHANNEL_MESSAGE(Name, Index, TypeA, TypeB, TypeC, TypeD, TypeE, TypeF)
+#define DEFINE_CONTROL_CHANNEL_MESSAGE_SEVENPARAM(Name, Index, TypeA, TypeB, TypeC, TypeD, TypeE, TypeF, TypeG) DEFINE_CONTROL_CHANNEL_MESSAGE(Name, Index, TypeA, TypeB, TypeC, TypeD, TypeE, TypeF, TypeG)
+#define DEFINE_CONTROL_CHANNEL_MESSAGE_EIGHTPARAM(Name, Index, TypeA, TypeB, TypeC, TypeD, TypeE, TypeF, TypeG, TypeH) DEFINE_CONTROL_CHANNEL_MESSAGE(Name, Index, TypeA, TypeB, TypeC, TypeD, TypeE, TypeF, TypeG, TypeH)
 
 #define IMPLEMENT_CONTROL_CHANNEL_MESSAGE(Name) static uint8 Dummy##_FNetControlMessage_##Name = FNetControlMessage<NMT_##Name>::Initialize();
 
 // message type definitions
-DEFINE_CONTROL_CHANNEL_MESSAGE_THREEPARAM(Hello, 0, uint8, uint32, FString); // initial client connection message
-DEFINE_CONTROL_CHANNEL_MESSAGE_THREEPARAM(Welcome, 1, FString, FString, FString); // server tells client they're ok'ed to load the server's level
-DEFINE_CONTROL_CHANNEL_MESSAGE_ONEPARAM(Upgrade, 2, uint32); // server tells client their version is incompatible
-DEFINE_CONTROL_CHANNEL_MESSAGE_ONEPARAM(Challenge, 3, FString); // server sends client challenge string to verify integrity
-DEFINE_CONTROL_CHANNEL_MESSAGE_ONEPARAM(Netspeed, 4, int32); // client sends requested transfer rate
-DEFINE_CONTROL_CHANNEL_MESSAGE_FOURPARAM(Login, 5, FString, FString, FUniqueNetIdRepl, FString); // client requests to be admitted to the game
-DEFINE_CONTROL_CHANNEL_MESSAGE_ONEPARAM(Failure, 6, FString); // indicates connection failure
-DEFINE_CONTROL_CHANNEL_MESSAGE_ZEROPARAM(Join, 9); // final join request (spawns PlayerController)
-DEFINE_CONTROL_CHANNEL_MESSAGE_TWOPARAM(JoinSplit, 10, FString, FUniqueNetIdRepl); // child player (splitscreen) join request
-DEFINE_CONTROL_CHANNEL_MESSAGE_ONEPARAM(Skip, 12, FGuid); // client request to skip an optional package
-DEFINE_CONTROL_CHANNEL_MESSAGE_ONEPARAM(Abort, 13, FGuid); // client informs server that it aborted a not-yet-verified package due to an UNLOAD request
-DEFINE_CONTROL_CHANNEL_MESSAGE_ONEPARAM(PCSwap, 15, int32); // client tells server it has completed a swap of its Connection->Actor
-DEFINE_CONTROL_CHANNEL_MESSAGE_ONEPARAM(ActorChannelFailure, 16, int32); // client tells server that it failed to open an Actor channel sent by the server (e.g. couldn't serialize Actor archetype)
-DEFINE_CONTROL_CHANNEL_MESSAGE_ONEPARAM(DebugText, 17, FString); // debug text sent to all clients or to server
-DEFINE_CONTROL_CHANNEL_MESSAGE_TWOPARAM(NetGUIDAssign, 18, FNetworkGUID, FString); // Explicit NetworkGUID assignment. This is rare and only happens if a netguid is only serialized client->server (this msg goes server->client to tell client what ID to use in that case)
-DEFINE_CONTROL_CHANNEL_MESSAGE_ONEPARAM(SecurityViolation, 19, FString); // server tells client that it has violated security and has been disconnected
-DEFINE_CONTROL_CHANNEL_MESSAGE_TWOPARAM(GameSpecific, 20, uint8, FString); // custom game-specific message routed to UGameInstance for processing
-DEFINE_CONTROL_CHANNEL_MESSAGE_ZEROPARAM(EncryptionAck, 21);
+DEFINE_CONTROL_CHANNEL_MESSAGE(Hello, 0, uint8, uint32, FString); // initial client connection message
+DEFINE_CONTROL_CHANNEL_MESSAGE(Welcome, 1, FString, FString, FString); // server tells client they're ok'ed to load the server's level
+DEFINE_CONTROL_CHANNEL_MESSAGE(Upgrade, 2, uint32); // server tells client their version is incompatible
+DEFINE_CONTROL_CHANNEL_MESSAGE(Challenge, 3, FString); // server sends client challenge string to verify integrity
+DEFINE_CONTROL_CHANNEL_MESSAGE(Netspeed, 4, int32); // client sends requested transfer rate
+DEFINE_CONTROL_CHANNEL_MESSAGE(Login, 5, FString, FString, FUniqueNetIdRepl, FString); // client requests to be admitted to the game
+DEFINE_CONTROL_CHANNEL_MESSAGE(Failure, 6, FString); // indicates connection failure
+DEFINE_CONTROL_CHANNEL_MESSAGE(Join, 9); // final join request (spawns PlayerController)
+DEFINE_CONTROL_CHANNEL_MESSAGE(JoinSplit, 10, FString, FUniqueNetIdRepl); // child player (splitscreen) join request
+DEFINE_CONTROL_CHANNEL_MESSAGE(Skip, 12, FGuid); // client request to skip an optional package
+DEFINE_CONTROL_CHANNEL_MESSAGE(Abort, 13, FGuid); // client informs server that it aborted a not-yet-verified package due to an UNLOAD request
+DEFINE_CONTROL_CHANNEL_MESSAGE(PCSwap, 15, int32); // client tells server it has completed a swap of its Connection->Actor
+DEFINE_CONTROL_CHANNEL_MESSAGE(ActorChannelFailure, 16, int32); // client tells server that it failed to open an Actor channel sent by the server (e.g. couldn't serialize Actor archetype)
+DEFINE_CONTROL_CHANNEL_MESSAGE(DebugText, 17, FString); // debug text sent to all clients or to server
+DEFINE_CONTROL_CHANNEL_MESSAGE(NetGUIDAssign, 18, FNetworkGUID, FString); // Explicit NetworkGUID assignment. This is rare and only happens if a netguid is only serialized client->server (this msg goes server->client to tell client what ID to use in that case)
+DEFINE_CONTROL_CHANNEL_MESSAGE(SecurityViolation, 19, FString); // server tells client that it has violated security and has been disconnected
+DEFINE_CONTROL_CHANNEL_MESSAGE(GameSpecific, 20, uint8, FString); // custom game-specific message routed to UGameInstance for processing
+DEFINE_CONTROL_CHANNEL_MESSAGE(EncryptionAck, 21);
 
 // 			Beacon control channel flow
 // Client												Server
@@ -431,7 +185,7 @@ DEFINE_CONTROL_CHANNEL_MESSAGE_ZEROPARAM(EncryptionAck, 21);
 //		Send<BeaconNetGUIDAck> acknowledging receipt of NetGUID
 //														Receive<BeaconNetGUIDAck> - connection complete
 
-DEFINE_CONTROL_CHANNEL_MESSAGE_ZEROPARAM(BeaconWelcome, 25); // server tells client they're ok to attempt to join (client sends netspeed/beacontype)
-DEFINE_CONTROL_CHANNEL_MESSAGE_TWOPARAM(BeaconJoin, 26, FString, FUniqueNetIdRepl);  // server tries to create beacon type requested by client, sends NetGUID for actor sync
-DEFINE_CONTROL_CHANNEL_MESSAGE_ONEPARAM(BeaconAssignGUID, 27, FNetworkGUID); // client assigns NetGUID from server to beacon actor, sends NetGUIDAck
-DEFINE_CONTROL_CHANNEL_MESSAGE_ONEPARAM(BeaconNetGUIDAck, 28, FString); // server received NetGUIDAck from client, connection established successfully
+DEFINE_CONTROL_CHANNEL_MESSAGE(BeaconWelcome, 25); // server tells client they're ok to attempt to join (client sends netspeed/beacontype)
+DEFINE_CONTROL_CHANNEL_MESSAGE(BeaconJoin, 26, FString, FUniqueNetIdRepl);  // server tries to create beacon type requested by client, sends NetGUID for actor sync
+DEFINE_CONTROL_CHANNEL_MESSAGE(BeaconAssignGUID, 27, FNetworkGUID); // client assigns NetGUID from server to beacon actor, sends NetGUIDAck
+DEFINE_CONTROL_CHANNEL_MESSAGE(BeaconNetGUIDAck, 28, FString); // server received NetGUIDAck from client, connection established successfully
