SoundCue: added missing nullptr checks for FirstNode breaking EdTest

#rnx

diff --git a/Engine/Source/Runtime/Engine/Private/SoundCue.cpp b/Engine/Source/Runtime/Engine/Private/SoundCue.cpp
--- a/Engine/Source/Runtime/Engine/Private/SoundCue.cpp
+++ b/Engine/Source/Runtime/Engine/Private/SoundCue.cpp
@@ -89,7 +89,10 @@ void USoundCue::CacheAggregateValues()
 
 void USoundCue::PrimeSoundCue()
 {
-	FirstNode->PrimeChildWavePlayers(true);
+	if (FirstNode != nullptr)
+	{
+		FirstNode->PrimeChildWavePlayers(true);
+	}
 }
 
 void USoundCue::Serialize(FStructuredArchive::FRecord Record)
@@ -159,7 +162,7 @@ void USoundCue::PostLoad()
 
 	CacheAggregateValues();
 	
-	if (bPrimeOnLoad)
+	if (bPrimeOnLoad && FirstNode != nullptr)
 	{
 		FirstNode->PrimeChildWavePlayers(true);
 	}
