Add transform support to shapes (Spheres, Capsules, Boxes, Convex)
Fix CoM-relative collision and box point samping
More GJKDistance unit tests

diff --git a/Engine/Source/Runtime/Engine/Private/PhysicsEngine/Experimental/ChaosInterfaceUtils.cpp b/Engine/Source/Runtime/Engine/Private/PhysicsEngine/Experimental/ChaosInterfaceUtils.cpp
--- a/Engine/Source/Runtime/Engine/Private/PhysicsEngine/Experimental/ChaosInterfaceUtils.cpp
+++ b/Engine/Source/Runtime/Engine/Private/PhysicsEngine/Experimental/ChaosInterfaceUtils.cpp
@@ -155,7 +155,7 @@ namespace ChaosInterface
 			const FKSphereElem ScaledSphereElem = InParams.Geometry->SphereElems[i].GetFinalScaled(Scale, InParams.LocalTransform);
 			const float UseRadius = FMath::Max(ScaledSphereElem.Radius, KINDA_SMALL_NUMBER);
 
-			auto ImplicitSphere = MakeUnique<Chaos::TSphere<float, 3>>(Chaos::TVector<float, 3>(0.f, 0.f, 0.f), UseRadius);
+			auto ImplicitSphere = MakeUnique<Chaos::TSphere<float, 3>>(ScaledSphereElem.Center, UseRadius);
 			auto NewShape = NewShapeHelper(MakeSerializable(ImplicitSphere));
 			Shapes.Emplace(MoveTemp(NewShape));
 			Geoms.Add(MoveTemp(ImplicitSphere));
@@ -171,13 +171,21 @@ namespace ChaosInterface
 			HalfExtents.Y = FMath::Max(HalfExtents.Y, KINDA_SMALL_NUMBER);
 			HalfExtents.Z = FMath::Max(HalfExtents.Z, KINDA_SMALL_NUMBER);
 
-			//todo: properly support transform on all geometries, for now just fixing a specific broken case with a translated aabb
-	//		ensure(BoxTransform.GetRotation() == FQuat::Identity);
-			TUniquePtr<Chaos::TBox<float, 3>> Box;
-			auto ImplicitBox = MakeUnique<Chaos::TBox<float, 3>>(BoxTransform.GetTranslation() - HalfExtents, BoxTransform.GetTranslation() + HalfExtents);
-			auto NewShape = NewShapeHelper(MakeSerializable(ImplicitBox));
+			// TBox can handle translations internally but if we have a rotation we need to wrap it in a transform
+			TUniquePtr<Chaos::TImplicitObject<float, 3>> Implicit;
+			if (!BoxTransform.GetRotation().IsIdentity())
+			{
+				auto ImplicitBox = MakeUnique<Chaos::TBox<float, 3>>(-HalfExtents, HalfExtents);
+				Implicit = TUniquePtr<Chaos::TImplicitObject<float, 3>>(new Chaos::TImplicitObjectTransformed<float, 3>(MoveTemp(ImplicitBox), BoxTransform));
+			}
+			else
+			{
+				Implicit = MakeUnique<Chaos::TBox<float, 3>>(BoxTransform.GetTranslation() - HalfExtents, BoxTransform.GetTranslation() + HalfExtents);
+			}
+
+			auto NewShape = NewShapeHelper(MakeSerializable(Implicit));
 			Shapes.Emplace(MoveTemp(NewShape));
-			Geoms.Add(MoveTemp(ImplicitBox));
+			Geoms.Add(MoveTemp(Implicit));
 		}
 		for (uint32 i = 0; i < static_cast<uint32>(InParams.Geometry->SphylElems.Num()); ++i)
 		{
@@ -190,7 +198,7 @@ namespace ChaosInterface
 			if (HalfHeight < KINDA_SMALL_NUMBER)
 			{
 				//not a capsule just use a sphere
-				auto ImplicitSphere = MakeUnique<Chaos::TSphere<float, 3>>(Chaos::TVector<float, 3>(0), Radius);
+				auto ImplicitSphere = MakeUnique<Chaos::TSphere<float, 3>>(ScaledSphylElem.Center, Radius);
 				auto NewShape = NewShapeHelper(MakeSerializable(ImplicitSphere));
 				Shapes.Emplace(MoveTemp(NewShape));
 				Geoms.Add(MoveTemp(ImplicitSphere));
@@ -198,9 +206,9 @@ namespace ChaosInterface
 			}
 			else
 			{
-				Chaos::TVector<float, 3> HalfExtents(0, 0, HalfHeight);
+				Chaos::TVector<float, 3> HalfExtents = ScaledSphylElem.Rotation.RotateVector(Chaos::TVector<float, 3>(0, 0, HalfHeight));
 
-				auto ImplicitCapsule = MakeUnique<Chaos::TCapsule<float>>(-HalfExtents, HalfExtents, Radius);
+				auto ImplicitCapsule = MakeUnique<Chaos::TCapsule<float>>(ScaledSphylElem.Center - HalfExtents, ScaledSphylElem.Center + HalfExtents, Radius);
 				auto NewShape = NewShapeHelper(MakeSerializable(ImplicitCapsule));
 				Shapes.Emplace(MoveTemp(NewShape));
 				Geoms.Add(MoveTemp(ImplicitCapsule));
@@ -238,12 +246,24 @@ namespace ChaosInterface
 		for (uint32 i = 0; i < static_cast<uint32>(InParams.Geometry->ConvexElems.Num()); ++i)
 		{
 			const FKConvexElem& CollisionBody = InParams.Geometry->ConvexElems[i];
+			const FTransform& ConvexTransform = CollisionBody.GetTransform();
 			if (const auto& ConvexImplicit = CollisionBody.GetChaosConvexMesh())
 			{
-				auto Implicit = MakeUnique<Chaos::TImplicitObjectScaled<float, 3>>(MakeSerializable(ConvexImplicit), Scale);
-				auto NewShape = NewShapeHelper(MakeSerializable(Implicit));
-				Shapes.Emplace(MoveTemp(NewShape));
-				Geoms.Add(MoveTemp(Implicit));
+				if (!ConvexTransform.GetTranslation().IsNearlyZero() || !ConvexTransform.GetRotation().IsIdentity())
+				{
+					TUniquePtr<Chaos::TImplicitObject<float, 3>> TransformImplicit = TUniquePtr<Chaos::TImplicitObject<float, 3>>(new Chaos::TImplicitObjectTransformed<float, 3>(MakeSerializable(ConvexImplicit), ConvexTransform));
+					TUniquePtr<Chaos::TImplicitObjectScaled<float, 3, false>> Implicit = MakeUnique<Chaos::TImplicitObjectScaled<float, 3, false>>(MoveTemp(TransformImplicit), Scale);
+					auto NewShape = NewShapeHelper(MakeSerializable(Implicit));
+					Shapes.Emplace(MoveTemp(NewShape));
+					Geoms.Add(MoveTemp(Implicit));
+				}
+				else
+				{
+					TUniquePtr<Chaos::TImplicitObjectScaled<float, 3, true>> Implicit = MakeUnique<Chaos::TImplicitObjectScaled<float, 3>>(MakeSerializable(ConvexImplicit), Scale);
+					auto NewShape = NewShapeHelper(MakeSerializable(Implicit));
+					Shapes.Emplace(MoveTemp(NewShape));
+					Geoms.Add(MoveTemp(Implicit));
+				}
 			}
 		}
 
diff --git a/Engine/Source/Runtime/Engine/Private/PhysicsEngine/ImmediatePhysics/ImmediatePhysicsChaos/ImmediatePhysicsActorHandle_Chaos.cpp b/Engine/Source/Runtime/Engine/Private/PhysicsEngine/ImmediatePhysics/ImmediatePhysicsChaos/ImmediatePhysicsActorHandle_Chaos.cpp
--- a/Engine/Source/Runtime/Engine/Private/PhysicsEngine/ImmediatePhysics/ImmediatePhysicsChaos/ImmediatePhysicsActorHandle_Chaos.cpp
+++ b/Engine/Source/Runtime/Engine/Private/PhysicsEngine/ImmediatePhysics/ImmediatePhysicsChaos/ImmediatePhysicsActorHandle_Chaos.cpp
@@ -183,7 +183,7 @@ namespace ImmediatePhysics_Chaos
 		AddParams.Scale = Scale;
 		//AddParams.SimpleMaterial = SimpleMaterial;
 		//AddParams.ComplexMaterials = TArrayView<UPhysicalMaterial*>(ComplexMaterials);
-		AddParams.LocalTransform = Chaos::TRigidTransform<float, 3>(-OutCoMTransform.GetTranslation(), OutCoMTransform.GetRotation().Inverse());
+		AddParams.LocalTransform = Chaos::TRigidTransform<float, 3>(OutCoMTransform.GetRotation().Inverse() * -OutCoMTransform.GetTranslation(), OutCoMTransform.GetRotation().Inverse());
 		AddParams.Geometry = &BodySetup->AggGeom;
 #if WITH_PHYSX
 		AddParams.TriMeshes = TArrayView<PxTriangleMesh*>(BodySetup->TriMeshes);
diff --git a/Engine/Source/Runtime/Experimental/Chaos/Private/Chaos/ChaosDebugDraw.cpp b/Engine/Source/Runtime/Experimental/Chaos/Private/Chaos/ChaosDebugDraw.cpp
--- a/Engine/Source/Runtime/Experimental/Chaos/Private/Chaos/ChaosDebugDraw.cpp
+++ b/Engine/Source/Runtime/Experimental/Chaos/Private/Chaos/ChaosDebugDraw.cpp
@@ -60,7 +60,8 @@ namespace Chaos
 			case ImplicitObjectType::Box:
 			{
 				const TBox<T, d>* Box = Shape->template GetObject<TBox<T, d>>();
-				FDebugDrawQueue::GetInstance().DrawDebugBox(ShapeTransform.GetLocation(), (T)0.5 * Box->Extents(), ShapeTransform.GetRotation(), Color, false, KINDA_SMALL_NUMBER, DrawPriority, LineThickness);
+				const TVector<T, d> P = ShapeTransform.TransformPosition(Box->Center());
+				FDebugDrawQueue::GetInstance().DrawDebugBox(P, (T)0.5 * Box->Extents(), ShapeTransform.GetRotation(), Color, false, KINDA_SMALL_NUMBER, DrawPriority, LineThickness);
 				break;
 			}
 			case ImplicitObjectType::Plane:
diff --git a/Engine/Source/Runtime/Experimental/Chaos/Private/Chaos/ImplicitObject.cpp b/Engine/Source/Runtime/Experimental/Chaos/Private/Chaos/ImplicitObject.cpp
--- a/Engine/Source/Runtime/Experimental/Chaos/Private/Chaos/ImplicitObject.cpp
+++ b/Engine/Source/Runtime/Experimental/Chaos/Private/Chaos/ImplicitObject.cpp
@@ -32,6 +32,50 @@ TImplicitObject<T, d>::~TImplicitObject()
 {
 }
 
+template<class T, int d>
+ImplicitObjectType TImplicitObject<T, d>::GetType(bool bGetTrueType) const
+{
+	if (bIgnoreAnalyticCollisions && !bGetTrueType)
+	{
+		return ImplicitObjectType::Unknown;
+	}
+	return Type;
+}
+
+template<class T, int d>
+bool TImplicitObject<T, d>::IsValidGeometry() const
+{
+	return true;
+}
+
+template<class T, int d>
+TUniquePtr<TImplicitObject<T, d>> TImplicitObject<T, d>::Copy() const
+{
+	check(false);
+	return nullptr;
+}
+
+template<class T, int d>
+bool TImplicitObject<T, d>::IsUnderlyingUnion() const
+{
+	return Type == ImplicitObjectType::Union;
+}
+
+template<class T, int d>
+T TImplicitObject<T, d>::SignedDistance(const TVector<T, d>& x) const
+{
+	TVector<T, d> Normal;
+	return PhiWithNormal(x, Normal);
+}
+
+template<class T, int d>
+TVector<T, d> TImplicitObject<T, d>::Normal(const TVector<T, d>& x) const
+{
+	TVector<T, d> Normal;
+	PhiWithNormal(x, Normal);
+	return Normal;
+}
+
 template<class T, int d>
 TVector<T, d> TImplicitObject<T, d>::Support(const TVector<T, d>& Direction, const T Thickness) const
 {
diff --git a/Engine/Source/Runtime/Experimental/Chaos/Private/Chaos/PBDCollisionConstraint.cpp b/Engine/Source/Runtime/Experimental/Chaos/Private/Chaos/PBDCollisionConstraint.cpp
--- a/Engine/Source/Runtime/Experimental/Chaos/Private/Chaos/PBDCollisionConstraint.cpp
+++ b/Engine/Source/Runtime/Experimental/Chaos/Private/Chaos/PBDCollisionConstraint.cpp
@@ -1477,37 +1477,31 @@ void UpdateCapsuleBoxConstraint(const TCapsule<T>& A, const TRigidTransform<T, d
 
 	TRigidTransform<T, d> BToATransform = BTransform.GetRelativeTransform(ATransform);
 
-	// USe GJK to track closest points (not strictly necessary yet)
-	//TVector<T, d> NearPointALocal, NearPointBLocal;
-	//T NearPointDistance;
-	//if (GJKDistance<T>(A, B, BToATransform, NearPointDistance, NearPointALocal, NearPointBLocal))
-	//{
-	//	TVector<T, d> NearPointAWorld = ATransform.TransformPosition(NearPointALocal);
-	//	TVector<T, d> NearPointBWorld = BTransform.TransformPosition(NearPointBLocal);
-	//	TVector<T, d> NearPointBtoAWorld = NearPointAWorld - NearPointBWorld;
-	//	Constraint.Phi = NearPointDistance;
-	//	Constraint.Normal = NearPointBtoAWorld.GetSafeNormal();
-	//	Constraint.Location = NearPointAWorld;
-	//}
-	//else
+	// Use GJK to track closest points (not strictly necessary yet)
+	TVector<T, d> NearPointALocal, NearPointBLocal;
+	T NearPointDistance;
+	if (GJKDistance<T>(A, B, BToATransform, NearPointDistance, NearPointALocal, NearPointBLocal))
+	{
+		TVector<T, d> NearPointAWorld = ATransform.TransformPosition(NearPointALocal);
+		TVector<T, d> NearPointBWorld = BTransform.TransformPosition(NearPointBLocal);
+		TVector<T, d> NearPointBtoAWorld = NearPointAWorld - NearPointBWorld;
+		Constraint.Phi = NearPointDistance;
+		Constraint.Normal = NearPointBtoAWorld.GetSafeNormal();
+		Constraint.Location = NearPointAWorld;
+	}
+	else
 	{
 		// Use box particle samples against the implicit capsule
-		const TArray<TVector<T, d>> BParticles = B.ComputeLocalSamplePoints();
+		const TArray<TVector<T, d>> BParticles = B.ComputeSamplePoints();
 		const int32 NumParticles = BParticles.Num();
-		bool bFoundContact = false;
 		for (int32 ParticleIndex = 0; ParticleIndex < NumParticles; ++ParticleIndex)
 		{
 			if (SampleObjectHelper2(A, ATransform, BToATransform, BParticles[ParticleIndex], Thickness, Constraint))
 			{
-				bFoundContact = true;
+				// SampleObjectHelper2 expects A to be the box, so reverse the results
+				Constraint.Normal = -Constraint.Normal;
 			}
 		}
-
-		// SampleObjectHelper2 expects A to be the box, so reverse the results
-		if (bFoundContact)
-		{
-			Constraint.Normal = -Constraint.Normal;
-		}
 	}
 
 }
@@ -1804,7 +1798,7 @@ void UpdateBoxConstraint(const TBox<T, d>& Box1, const TRigidTransform<T, d>& Bo
 			//check(Constraint.Phi < MThickness);
 			// For now revert to doing all points vs lsv check until we can figure out a good way to get the deepest point without needing this
 			{
-				const TArray<TVector<T, d>> SampleParticles = Box1.ComputeLocalSamplePoints();
+				const TArray<TVector<T, d>> SampleParticles = Box1.ComputeSamplePoints();
 				const TRigidTransform<T, d> Box1ToBox2Transform = Box1Transform.GetRelativeTransform(Box2Transform);
 				int32 NumParticles = SampleParticles.Num();
 				for (int32 i = 0; i < NumParticles; ++i)
diff --git a/Engine/Source/Runtime/Experimental/Chaos/Public/Chaos/ImplicitObject.h b/Engine/Source/Runtime/Experimental/Chaos/Public/Chaos/ImplicitObject.h
--- a/Engine/Source/Runtime/Experimental/Chaos/Public/Chaos/ImplicitObject.h
+++ b/Engine/Source/Runtime/Experimental/Chaos/Public/Chaos/ImplicitObject.h
@@ -23,6 +23,9 @@ template<class T, int d>
 class TParticles;
 template<class T, int d>
 class TBVHParticles;
+template<class T, int d>
+class TImplicitObject;
+
 
 enum class ImplicitObjectType : int8
 {
@@ -55,6 +58,37 @@ namespace EImplicitObject
 	const int32 FiniteConvex = IsConvex | HasBoundingBox;
 }
 
+
+
+template<class T, int d, bool bSerializable>
+struct TImplicitObjectPtrStorage
+{
+};
+
+template<class T, int d>
+struct TImplicitObjectPtrStorage<T, d, false>
+{
+	using PtrType = TImplicitObject<T, d>*;
+
+	static PtrType Convert(const TUniquePtr<TImplicitObject<T, d>>& Object)
+	{
+		return Object.Get();
+	}
+};
+
+template<class T, int d>
+struct TImplicitObjectPtrStorage<T, d, true>
+{
+	using PtrType = TSerializablePtr<TImplicitObject<T, d>>;
+
+	static PtrType Convert(const TUniquePtr<TImplicitObject<T, d>>& Object)
+	{
+		return MakeSerializable(Object);
+	}
+};
+
+
+
 template<class T, int d>
 class CHAOS_API TImplicitObject
 {
@@ -100,44 +134,20 @@ public:
 		return static_cast<const T_DERIVED&>(*this);
 	}
 
-	ImplicitObjectType GetType(bool bGetTrueType = false) const
-	{
-		if (bIgnoreAnalyticCollisions && !bGetTrueType)
-		{
-			return ImplicitObjectType::Unknown;
-		}
-		return Type;
-	}
+	ImplicitObjectType GetType(bool bGetTrueType = false) const;
 
-	virtual bool IsValidGeometry() const
-	{
-		return true;
-	}
+	virtual bool IsValidGeometry() const;
 
-	virtual TUniquePtr<TImplicitObject<T, d>> Copy() const
-	{
-		check(false);
-		return nullptr;
-	}
+	virtual TUniquePtr<TImplicitObject<T, d>> Copy() const;
 
 	//This is strictly used for optimization purposes
-	bool IsUnderlyingUnion() const
-	{
-		return Type == ImplicitObjectType::Union;
-	}
+	bool IsUnderlyingUnion() const;
+
 	// Explicitly non-virtual.  Must cast to derived types to target their implementation.
-	T SignedDistance(const TVector<T, d>& x) const
-	{
-		TVector<T, d> Normal;
-		return PhiWithNormal(x, Normal);
-	}
+	T SignedDistance(const TVector<T, d>& x) const;
+
 	// Explicitly non-virtual.  Must cast to derived types to target their implementation.
-	TVector<T, d> Normal(const TVector<T, d>& x) const
-	{
-		TVector<T, d> Normal;
-		PhiWithNormal(x, Normal);
-		return Normal;
-	}
+	TVector<T, d> Normal(const TVector<T, d>& x) const;
 	virtual T PhiWithNormal(const TVector<T, d>& x, TVector<T, d>& Normal) const = 0;
 	virtual const class TBox<T, d>& BoundingBox() const;
 	virtual TVector<T, d> Support(const TVector<T, d>& Direction, const T Thickness) const;
diff --git a/Engine/Source/Runtime/Experimental/ChaosSolvers/Public/PhysicsProxy/AnalyticImplicitGroup.h b/Engine/Source/Runtime/Experimental/ChaosSolvers/Public/PhysicsProxy/AnalyticImplicitGroup.h
--- a/Engine/Source/Runtime/Experimental/ChaosSolvers/Public/PhysicsProxy/AnalyticImplicitGroup.h
+++ b/Engine/Source/Runtime/Experimental/ChaosSolvers/Public/PhysicsProxy/AnalyticImplicitGroup.h
@@ -427,7 +427,6 @@ public:
 				// Make a copy and transfer ownership to the transformed implicit.
 				TUniquePtr<Chaos::TImplicitObject<float, 3>> ObjPtr(TransferImplicitObj(0));
 				return new Chaos::TImplicitObjectTransformed<float, 3, true>(
-					MakeSerializable(ObjPtr),
 					MoveTemp(ObjPtr),
 					Chaos::TRigidTransform<float, 3>(Transforms[0]));
 			}
@@ -454,7 +453,6 @@ public:
 					ImplicitObjects.Add(
 						TUniquePtr<Chaos::TImplicitObject<float, 3>>(
 							new Chaos::TImplicitObjectTransformed<float, 3, true>(
-								MakeSerializable(ObjPtr),
 								MoveTemp(ObjPtr),
 								Chaos::TRigidTransform<float, 3>(Xf))));
 				}
diff --git a/Engine/Source/Runtime/PhysicsCore/Public/PhysXToChaosUtil.h b/Engine/Source/Runtime/PhysicsCore/Public/PhysXToChaosUtil.h
--- a/Engine/Source/Runtime/PhysicsCore/Public/PhysXToChaosUtil.h
+++ b/Engine/Source/Runtime/PhysicsCore/Public/PhysXToChaosUtil.h
@@ -174,7 +174,7 @@ inline TUniquePtr<Chaos::TImplicitObjectTransformed<float, 3>> PxShapeToChaosGeo
 	default: ensure(false); return nullptr;	//missing support for this geometry type
 	}
 
-	return MakeUnique<TImplicitObjectTransformed<float, 3>>(MakeSerializable(InnerObj), MoveTemp(InnerObj), P2UTransform(ShapeTM));
+	return MakeUnique<TImplicitObjectTransformed<float, 3>>(MoveTemp(InnerObj), P2UTransform(ShapeTM));
 }
 
 #endif
\ No newline at end of file
