Fix a bug where some static meshes don't cast DF shadow when they should

diff --git a/Engine/Source/Runtime/Engine/Private/StaticMesh.cpp b/Engine/Source/Runtime/Engine/Private/StaticMesh.cpp
--- a/Engine/Source/Runtime/Engine/Private/StaticMesh.cpp
+++ b/Engine/Source/Runtime/Engine/Private/StaticMesh.cpp
@@ -1476,6 +1476,12 @@ void FStaticMeshRenderData::InitResources(ERHIFeatureLevel::Type InFeatureLevel,
 			LODResources[LODIndex].InitResources(Owner);
 			LODVertexFactories[LODIndex].InitResources(LODResources[LODIndex], LODIndex, Owner);
 		}
+		else if (!LODIndex && LODResources[LODIndex].DistanceFieldData)
+		{
+			FDistanceFieldVolumeData* DistanceFieldData = LODResources[LODIndex].DistanceFieldData;
+			DistanceFieldData->VolumeTexture.Initialize(Owner);
+			INC_DWORD_STAT_BY(STAT_StaticMeshDistanceFieldMemory, DistanceFieldData->GetResourceSizeBytes());
+		}
 	}
 
 	ENQUEUE_RENDER_COMMAND(CmdSetStaticMeshReadyForStreaming)(
@@ -1496,6 +1502,12 @@ void FStaticMeshRenderData::ReleaseResources()
 			LODResources[LODIndex].ReleaseResources();
 			LODVertexFactories[LODIndex].ReleaseResources();
 		}
+		else if (!LODIndex && LODResources[LODIndex].DistanceFieldData)
+		{
+			FDistanceFieldVolumeData* DistanceFieldData = LODResources[LODIndex].DistanceFieldData;
+			DEC_DWORD_STAT_BY(STAT_StaticMeshDistanceFieldMemory, DistanceFieldData->GetResourceSizeBytes());
+			DistanceFieldData->VolumeTexture.Release();
+		}
 	}
 }
 
