Removed unused shadow shaders from the water. This halves water shader number.

diff --git a/Engine/Source/Runtime/Engine/Public/MaterialShared.h b/Engine/Source/Runtime/Engine/Public/MaterialShared.h
--- a/Engine/Source/Runtime/Engine/Public/MaterialShared.h
+++ b/Engine/Source/Runtime/Engine/Public/MaterialShared.h
@@ -1560,9 +1560,10 @@ public:
 	/** Returns whether this material should be considered for casting dynamic shadows. */
 	inline bool ShouldCastDynamicShadows() const
 	{
-		return (GetBlendMode() == BLEND_Opaque ||
-		        GetBlendMode() == BLEND_Masked ||
-		       (GetBlendMode() == BLEND_Translucent && GetCastDynamicShadowAsMasked()));
+		return !GetShadingModels().HasOnlyShadingModel(MSM_SingleLayerWater) &&
+				(GetBlendMode() == BLEND_Opaque
+ 				 || GetBlendMode() == BLEND_Masked
+  				 || (GetBlendMode() == BLEND_Translucent && GetCastDynamicShadowAsMasked()));
 	}
 
 
