Workaround for instanced static mesh vertex factory containing mismatch with shader. Added ensure to catch susepected issue without crashing the user.

diff --git a/Engine/Source/Runtime/Engine/Private/InstancedStaticMesh.cpp b/Engine/Source/Runtime/Engine/Private/InstancedStaticMesh.cpp
--- a/Engine/Source/Runtime/Engine/Private/InstancedStaticMesh.cpp
+++ b/Engine/Source/Runtime/Engine/Private/InstancedStaticMesh.cpp
@@ -384,6 +384,8 @@ void FStaticMeshInstanceBuffer::BindInstanceVertexBuffer(const class FVertexFact
 		InstancedStaticMeshData.InstanceOriginSRV = InstanceOriginSRV;
 		InstancedStaticMeshData.InstanceTransformSRV = InstanceTransformSRV;
 		InstancedStaticMeshData.InstanceLightmapSRV = InstanceLightmapSRV;
+		InstancedStaticMeshData.NumInstances = InstanceData->GetNumInstances();
+		InstancedStaticMeshData.bInitialized = true;
 	}
 
 	{
@@ -2945,10 +2947,17 @@ void FInstancedStaticMeshVertexFactoryShaderParameters::GetElementShaderBindings
 	{
 		if (InstancedVertexFactory->SupportsManualVertexFetch(FeatureLevel))
 		{
-			ShaderBindings.Add(VertexFetch_InstanceOriginBufferParameter, InstancedVertexFactory->GetInstanceOriginSRV());
-			ShaderBindings.Add(VertexFetch_InstanceTransformBufferParameter, InstancedVertexFactory->GetInstanceTransformSRV());
-			ShaderBindings.Add(VertexFetch_InstanceLightmapBufferParameter, InstancedVertexFactory->GetInstanceLightmapSRV());
-			ShaderBindings.Add(InstanceOffset, InstanceOffsetValue);
+			if (InstancedVertexFactory->GetNumInstances() > 0)
+			{
+				ShaderBindings.Add(VertexFetch_InstanceOriginBufferParameter, InstancedVertexFactory->GetInstanceOriginSRV());
+				ShaderBindings.Add(VertexFetch_InstanceTransformBufferParameter, InstancedVertexFactory->GetInstanceTransformSRV());
+				ShaderBindings.Add(VertexFetch_InstanceLightmapBufferParameter, InstancedVertexFactory->GetInstanceLightmapSRV());
+				ShaderBindings.Add(InstanceOffset, InstanceOffsetValue);
+			}
+			else
+			{
+				ensureMsgf(false, TEXT("Instanced static mesh rendered with no instances. Data initialized: %d"), InstancedVertexFactory->IsDataInitialized());
+			}
 		}
 
 		if (InstanceOffsetValue > 0 && VertexStreams.Num() > 0)
diff --git a/Engine/Source/Runtime/Engine/Private/InstancedStaticMesh.h b/Engine/Source/Runtime/Engine/Private/InstancedStaticMesh.h
--- a/Engine/Source/Runtime/Engine/Private/InstancedStaticMesh.h
+++ b/Engine/Source/Runtime/Engine/Private/InstancedStaticMesh.h
@@ -169,6 +169,10 @@ struct FInstancedStaticMeshDataType
 	FRHIShaderResourceView* InstanceOriginSRV = nullptr;
 	FRHIShaderResourceView* InstanceTransformSRV = nullptr;
 	FRHIShaderResourceView* InstanceLightmapSRV = nullptr;
+
+	/** Used to track state for debugging. */
+	uint32 NumInstances = 0;
+	bool bInitialized = false;
 };
 
 /**
@@ -258,6 +262,16 @@ public:
 	virtual bool SupportsNullPixelShader() const override { return false; }
 #endif
 
+	inline bool IsDataInitialized() const
+	{
+		return Data.bInitialized;
+	}
+
+	inline uint32 GetNumInstances() const
+	{
+		return Data.NumInstances;
+	}
+
 	inline FRHIShaderResourceView* GetInstanceOriginSRV() const
 	{
 		return Data.InstanceOriginSRV;
