Chaos - Temp undo use of transform implicit

diff --git a/Engine/Source/Runtime/Engine/Private/PhysicsEngine/Experimental/ChaosInterfaceUtils.cpp b/Engine/Source/Runtime/Engine/Private/PhysicsEngine/Experimental/ChaosInterfaceUtils.cpp
--- a/Engine/Source/Runtime/Engine/Private/PhysicsEngine/Experimental/ChaosInterfaceUtils.cpp
+++ b/Engine/Source/Runtime/Engine/Private/PhysicsEngine/Experimental/ChaosInterfaceUtils.cpp
@@ -173,12 +173,12 @@ namespace ChaosInterface
 
 			// TBox can handle translations internally but if we have a rotation we need to wrap it in a transform
 			TUniquePtr<Chaos::TImplicitObject<float, 3>> Implicit;
-			if (!BoxTransform.GetRotation().IsIdentity())
-			{
-				auto ImplicitBox = MakeUnique<Chaos::TBox<float, 3>>(-HalfExtents, HalfExtents);
-				Implicit = TUniquePtr<Chaos::TImplicitObject<float, 3>>(new Chaos::TImplicitObjectTransformed<float, 3>(MoveTemp(ImplicitBox), BoxTransform));
-			}
-			else
+			//if (!BoxTransform.GetRotation().IsIdentity())
+			//{
+			//	auto ImplicitBox = MakeUnique<Chaos::TBox<float, 3>>(-HalfExtents, HalfExtents);
+			//	Implicit = TUniquePtr<Chaos::TImplicitObject<float, 3>>(new Chaos::TImplicitObjectTransformed<float, 3>(MoveTemp(ImplicitBox), BoxTransform));
+			//}
+			//else
 			{
 				Implicit = MakeUnique<Chaos::TBox<float, 3>>(BoxTransform.GetTranslation() - HalfExtents, BoxTransform.GetTranslation() + HalfExtents);
 			}
@@ -249,15 +249,15 @@ namespace ChaosInterface
 			const FTransform& ConvexTransform = CollisionBody.GetTransform();
 			if (const auto& ConvexImplicit = CollisionBody.GetChaosConvexMesh())
 			{
-				if (!ConvexTransform.GetTranslation().IsNearlyZero() || !ConvexTransform.GetRotation().IsIdentity())
-				{
-					TUniquePtr<Chaos::TImplicitObject<float, 3>> TransformImplicit = TUniquePtr<Chaos::TImplicitObject<float, 3>>(new Chaos::TImplicitObjectTransformed<float, 3>(MakeSerializable(ConvexImplicit), ConvexTransform));
-					TUniquePtr<Chaos::TImplicitObjectScaled<float, 3, false>> Implicit = MakeUnique<Chaos::TImplicitObjectScaled<float, 3, false>>(MoveTemp(TransformImplicit), Scale);
-					auto NewShape = NewShapeHelper(MakeSerializable(Implicit));
-					Shapes.Emplace(MoveTemp(NewShape));
-					Geoms.Add(MoveTemp(Implicit));
-				}
-				else
+				//if (!ConvexTransform.GetTranslation().IsNearlyZero() || !ConvexTransform.GetRotation().IsIdentity())
+				//{
+				//	TUniquePtr<Chaos::TImplicitObject<float, 3>> TransformImplicit = TUniquePtr<Chaos::TImplicitObject<float, 3>>(new Chaos::TImplicitObjectTransformed<float, 3>(MakeSerializable(ConvexImplicit), ConvexTransform));
+				//	TUniquePtr<Chaos::TImplicitObjectScaled<float, 3, false>> Implicit = MakeUnique<Chaos::TImplicitObjectScaled<float, 3, false>>(MoveTemp(TransformImplicit), Scale);
+				//	auto NewShape = NewShapeHelper(MakeSerializable(Implicit));
+				//	Shapes.Emplace(MoveTemp(NewShape));
+				//	Geoms.Add(MoveTemp(Implicit));
+				//}
+				//else
 				{
 					TUniquePtr<Chaos::TImplicitObjectScaled<float, 3, true>> Implicit = MakeUnique<Chaos::TImplicitObjectScaled<float, 3>>(MakeSerializable(ConvexImplicit), Scale);
 					auto NewShape = NewShapeHelper(MakeSerializable(Implicit));
