Fix static anlysis warning
#rnx

diff --git a/Engine/Source/Runtime/Engine/Private/PhysicsEngine/Experimental/PhysScene_Chaos.cpp b/Engine/Source/Runtime/Engine/Private/PhysicsEngine/Experimental/PhysScene_Chaos.cpp
--- a/Engine/Source/Runtime/Engine/Private/PhysicsEngine/Experimental/PhysScene_Chaos.cpp
+++ b/Engine/Source/Runtime/Engine/Private/PhysicsEngine/Experimental/PhysScene_Chaos.cpp
@@ -596,24 +596,21 @@ void RemovePhysicsProxy(ObjectType* InObject, Chaos::FPhysicsSolver* InSolver, F
 
 	const bool bDedicatedThread = PhysDispatcher->GetMode() == Chaos::EThreadingMode::DedicatedThread;
 
-	if(PhysDispatcher)
+	// Remove the object from the solver
+	PhysDispatcher->EnqueueCommandImmediate([InObject, InSolver, bDedicatedThread](Chaos::FPersistentPhysicsTask* PhysThread)
 	{
-		// Remove the object from the solver
-		PhysDispatcher->EnqueueCommandImmediate([InObject, InSolver, bDedicatedThread](Chaos::FPersistentPhysicsTask* PhysThread)
-		{
 #if CHAOS_PARTICLEHANDLE_TODO
-			InSolver->UnregisterObject(InObject);
+		InSolver->UnregisterObject(InObject);
 #endif
-			InObject->OnRemoveFromScene();
+		InObject->OnRemoveFromScene();
 
-			if (!bDedicatedThread)
-			{
-				InObject->SyncBeforeDestroy();
-				delete InObject;
-			}
+		if (!bDedicatedThread)
+		{
+			InObject->SyncBeforeDestroy();
+			delete InObject;
+		}
 
-		});
-	}
+	});
 }
 
 void FPhysScene_Chaos::RemoveObject(FSkeletalMeshPhysicsProxy* InObject)
