Move encryption key response delegates from UWorld to UNetConnection to handle a seamless travel occurring during the key lookup

diff --git a/Engine/Source/Runtime/Engine/Classes/Engine/NetConnection.h b/Engine/Source/Runtime/Engine/Classes/Engine/NetConnection.h
--- a/Engine/Source/Runtime/Engine/Classes/Engine/NetConnection.h
+++ b/Engine/Source/Runtime/Engine/Classes/Engine/NetConnection.h
@@ -1158,7 +1158,13 @@ public:
 	 * Get the current number of sent packets for which we have received a delivery notification
 	 */
 	ENGINE_API uint32 GetOutTotalNotifiedPackets() const { return OutTotalNotifiedPackets; }
-	
+
+	/** Sends the NMT_Challenge message */
+	void SendChallengeControlMessage();
+
+	/** Sends the NMT_Challenge message based on encryption response */
+	void SendChallengeControlMessage(const FEncryptionKeyResponse& Response);
+
 protected:
 
 	ENGINE_API void SetPendingCloseDueToSocketSendFailure();
diff --git a/Engine/Source/Runtime/Engine/Classes/Engine/World.h b/Engine/Source/Runtime/Engine/Classes/Engine/World.h
--- a/Engine/Source/Runtime/Engine/Classes/Engine/World.h
+++ b/Engine/Source/Runtime/Engine/Classes/Engine/World.h
@@ -3225,10 +3225,6 @@ private:
 
 	APhysicsVolume* InternalGetDefaultPhysicsVolume() const;
 
-	// Sends the NMT_Challenge message to Connection.
-	void SendChallengeControlMessage(UNetConnection* Connection);
-	void SendChallengeControlMessage(const FEncryptionKeyResponse& Response, TWeakObjectPtr<UNetConnection> WeakConnection);
-
 public:
 
 	/**
diff --git a/Engine/Source/Runtime/Engine/Private/NetConnection.cpp b/Engine/Source/Runtime/Engine/Private/NetConnection.cpp
--- a/Engine/Source/Runtime/Engine/Private/NetConnection.cpp
+++ b/Engine/Source/Runtime/Engine/Private/NetConnection.cpp
@@ -3707,6 +3707,46 @@ void UNetConnection::TrackReplicationForAnalytics(const bool bWasSaturated)
 	SaturationAnalytics.TrackReplication(bWasSaturated);
 }
 
+void UNetConnection::SendChallengeControlMessage()
+{
+	if (State != USOCK_Invalid && State != USOCK_Closed && Driver)
+	{
+		Challenge = FString::Printf(TEXT("%08X"), FPlatformTime::Cycles());
+		SetExpectedClientLoginMsgType(NMT_Login);
+		FNetControlMessage<NMT_Challenge>::Send(this, Challenge);
+		FlushNet();
+	}
+	else
+	{
+		UE_LOG(LogNet, Log, TEXT("UWorld::SendChallengeControlMessage: connection in invalid state. %s"), *Describe());
+	}
+}
+
+void UNetConnection::SendChallengeControlMessage(const FEncryptionKeyResponse& Response)
+{
+	if (State != USOCK_Invalid && State != USOCK_Closed && Driver)
+	{
+		if (Response.Response == EEncryptionResponse::Success)
+		{
+			EnableEncryptionServer(Response.EncryptionData);
+			SendChallengeControlMessage();
+		}
+		else
+		{
+			FString ResponseStr(LexToString(Response.Response));
+			UE_LOG(LogNet, Warning, TEXT("UWorld::SendChallengeControlMessage: encryption failure [%s] %s"), *ResponseStr, *Response.ErrorMsg);
+			FNetControlMessage<NMT_Failure>::Send(this, ResponseStr);
+			FlushNet();
+			// Can't close the connection here since it will leave the failure message in the send buffer and just close the socket. 
+			// Close();
+		}
+	}
+	else
+	{
+		UE_LOG(LogNet, Warning, TEXT("UNetConnection::SendChallengeControlMessage: connection in invalid state. %s"), *Describe());
+	}
+}
+
 /*-----------------------------------------------------------------------------
 	USimulatedClientNetConnection.
 -----------------------------------------------------------------------------*/
diff --git a/Engine/Source/Runtime/Engine/Private/World.cpp b/Engine/Source/Runtime/Engine/Private/World.cpp
--- a/Engine/Source/Runtime/Engine/Private/World.cpp
+++ b/Engine/Source/Runtime/Engine/Private/World.cpp
@@ -4928,14 +4928,13 @@ void UWorld::NotifyControlMessage(UNetConnection* Connection, uint8 MessageType,
 					{
 						if (EncryptionToken.IsEmpty())
 						{
-							SendChallengeControlMessage(Connection);
+							Connection->SendChallengeControlMessage();
 						}
 						else
 						{
 							if (FNetDelegates::OnReceivedNetworkEncryptionToken.IsBound())
 							{
-								TWeakObjectPtr<UNetConnection> WeakConnection = Connection;
-								FNetDelegates::OnReceivedNetworkEncryptionToken.Execute(EncryptionToken, FOnEncryptionKeyResponse::CreateUObject(this, &UWorld::SendChallengeControlMessage, WeakConnection));
+								FNetDelegates::OnReceivedNetworkEncryptionToken.Execute(EncryptionToken, FOnEncryptionKeyResponse::CreateUObject(Connection, &UNetConnection::SendChallengeControlMessage));
 							}
 							else
 							{
@@ -5261,61 +5260,6 @@ void UWorld::NotifyControlMessage(UNetConnection* Connection, uint8 MessageType,
 	}
 }
 
-void UWorld::SendChallengeControlMessage(UNetConnection* Connection)
-{
-	if (Connection)
-	{
-		if (Connection->State != USOCK_Invalid && Connection->State != USOCK_Closed && Connection->Driver)
-		{
-			Connection->Challenge = FString::Printf(TEXT("%08X"), FPlatformTime::Cycles());
-			Connection->SetExpectedClientLoginMsgType(NMT_Login);
-			FNetControlMessage<NMT_Challenge>::Send(Connection, Connection->Challenge);
-			Connection->FlushNet();
-		}
-		else
-		{
-			UE_LOG(LogNet, Log, TEXT("UWorld::SendChallengeControlMessage: connection in invalid state. %s"), *Connection->Describe());
-		}
-	}
-	else
-	{
-		UE_LOG(LogNet, Log, TEXT("UWorld::SendChallengeControlMessage: Connection is null."));
-	}
-}
-
-void UWorld::SendChallengeControlMessage(const FEncryptionKeyResponse& Response, TWeakObjectPtr<UNetConnection> WeakConnection)
-{
-	UNetConnection* Connection = WeakConnection.Get();
-	if (Connection)
-	{
-		if (Connection->State != USOCK_Invalid && Connection->State != USOCK_Closed && Connection->Driver)
-		{
-			if (Response.Response == EEncryptionResponse::Success)
-			{
-				Connection->EnableEncryptionServer(Response.EncryptionData);
-				SendChallengeControlMessage(Connection);
-			}
-			else
-			{
-				FString ResponseStr(LexToString(Response.Response));
-				UE_LOG(LogNet, Warning, TEXT("UWorld::SendChallengeControlMessage: encryption failure [%s] %s"), *ResponseStr, *Response.ErrorMsg);
-				FNetControlMessage<NMT_Failure>::Send(Connection, ResponseStr);
-				Connection->FlushNet();
-				// Can't close the connection here since it will leave the failure message in the send buffer and just close the socket. 
-				// Connection->Close();
-			}
-		}
-		else
-		{
-			UE_LOG(LogNet, Warning, TEXT("UWorld::SendChallengeControlMessage: connection in invalid state. %s"), *Connection->Describe());
-		}
-	}
-	else
-	{
-		UE_LOG(LogNet, Warning, TEXT("UWorld::SendChallengeControlMessage: Connection is null."));
-	}
-}
-
 bool UWorld::Listen( FURL& InURL )
 {
 #if WITH_SERVER_CODE
