[NO NON-PII COMMIT MESSAGE]

diff --git a/Engine/Source/Runtime/Engine/Private/AudioDerivedData.cpp b/Engine/Source/Runtime/Engine/Private/AudioDerivedData.cpp
--- a/Engine/Source/Runtime/Engine/Private/AudioDerivedData.cpp
+++ b/Engine/Source/Runtime/Engine/Private/AudioDerivedData.cpp
@@ -1234,7 +1234,7 @@ static void CookSurroundWave( USoundWave* SoundWave, FName FormatName, const IAu
 
 			if (SampleRateOverride > 0 && SampleRateOverride != (float)WaveSampleRate)
 			{
-				for (int ChannelIndex = 1; ChannelIndex < ChannelCount; ChannelIndex++)
+				for (int32 ChannelIndex = 0; ChannelIndex < ChannelCount; ChannelIndex++)
 				{
 					size_t DataSize = SourceBuffers[ChannelIndex].Num();
 					ResampleWaveData(SourceBuffers[ChannelIndex], DataSize, 1, WaveSampleRate, SampleRateOverride);
