PR #6264: Push streaming level priority to package loading request priorty (Contributed by [REDACTED_PERSON])
#rnx

diff --git a/Engine/Source/Runtime/Engine/Private/LevelStreaming.cpp b/Engine/Source/Runtime/Engine/Private/LevelStreaming.cpp
--- a/Engine/Source/Runtime/Engine/Private/LevelStreaming.cpp
+++ b/Engine/Source/Runtime/Engine/Private/LevelStreaming.cpp
@@ -1023,7 +1023,7 @@ bool ULevelStreaming::RequestLevel(UWorld* PersistentWorld, bool bAllowLevelLoad
 			// Kick off async load request.
 			STAT_ADD_CUSTOMMESSAGE_NAME( STAT_NamedMarker, *(FString( TEXT( "RequestLevel - " ) + DesiredPackageName.ToString() )) );
 			TRACE_BOOKMARK(TEXT("RequestLevel - %s"), *DesiredPackageName.ToString());
-			LoadPackageAsync(DesiredPackageName.ToString(), nullptr, *PackageNameToLoadFrom, FLoadPackageAsyncDelegate::CreateUObject(this, &ULevelStreaming::AsyncLevelLoadComplete), PackageFlags, PIEInstanceID);
+			LoadPackageAsync(DesiredPackageName.ToString(), nullptr, *PackageNameToLoadFrom, FLoadPackageAsyncDelegate::CreateUObject(this, &ULevelStreaming::AsyncLevelLoadComplete), PackageFlags, PIEInstanceID, GetPriority());
 
 			// streamingServer: server loads everything?
 			// Editor immediately blocks on load and we also block if background level streaming is disabled.
