Add console variable scale factor for virtual texture physical pools
This allows per platform device overrides

diff --git a/Engine/Source/Runtime/Engine/Classes/VT/VirtualTexturePoolConfig.h b/Engine/Source/Runtime/Engine/Classes/VT/VirtualTexturePoolConfig.h
--- a/Engine/Source/Runtime/Engine/Classes/VT/VirtualTexturePoolConfig.h
+++ b/Engine/Source/Runtime/Engine/Classes/VT/VirtualTexturePoolConfig.h
@@ -18,7 +18,7 @@ struct FVirtualTextureSpacePoolConfig
 {
 	GENERATED_USTRUCT_BODY()
 
-	FVirtualTextureSpacePoolConfig() : MinTileSize(0), MaxTileSize(0), SizeInMegabyte(0) {}
+	FVirtualTextureSpacePoolConfig() : MinTileSize(0), MaxTileSize(0), SizeInMegabyte(0), bAllowSizeScale(false) {}
 
 	/** Minimum tile size to match (including tile border). */
 	UPROPERTY()
@@ -36,6 +36,10 @@ struct FVirtualTextureSpacePoolConfig
 	UPROPERTY()
 	int32 SizeInMegabyte;
 
+	/** Allow the size to allocate for the pool to be scaled by some factor. */
+	UPROPERTY()
+	bool bAllowSizeScale;
+
 	/** Is this the default config? Use this setting when we can't find any other match. */
 	bool IsDefault() const { return Formats.Num() == 0; }
 };
@@ -51,12 +55,5 @@ public:
 	UPROPERTY(config)
 	TArray<FVirtualTextureSpacePoolConfig> Pools; // All the VT pools specified in the config
 
-	const FVirtualTextureSpacePoolConfig *FindPoolConfig(TEnumAsByte<EPixelFormat> const* Formats, int32 NumLayers, int32 TileSize);
-
-	virtual ~UVirtualTexturePoolConfig();
-
-	virtual void PostLoad() override;
-
-private:
-	FVirtualTextureSpacePoolConfig DefaultConfig;
+	void FindPoolConfig(TEnumAsByte<EPixelFormat> const* InFormats, int32 InNumLayers, int32 InTileSize, FVirtualTextureSpacePoolConfig& OutConfig) const;
 };
diff --git a/Engine/Source/Runtime/Engine/Private/VT/RuntimeVirtualTexture.cpp b/Engine/Source/Runtime/Engine/Private/VT/RuntimeVirtualTexture.cpp
--- a/Engine/Source/Runtime/Engine/Private/VT/RuntimeVirtualTexture.cpp
+++ b/Engine/Source/Runtime/Engine/Private/VT/RuntimeVirtualTexture.cpp
@@ -7,18 +7,9 @@
 #include "VT/VirtualTextureBuildSettings.h"
 #include "VT/RuntimeVirtualTextureNotify.h"
 #include "VT/RuntimeVirtualTextureStreamingProxy.h"
-#include "VT/UploadingVirtualTexture.h"
 #include "VT/VirtualTextureLevelRedirector.h"
-
-
-/** Device scalability option for runtime virtual texture size. */
-static TAutoConsoleVariable<int32> CVarVTTileCountBias(
-	TEXT("r.VT.RVT.TileCountBias"),
-	0,
-	TEXT("Bias to apply to Runtime Virtual Texture size."),
-	ECVF_RenderThreadSafe
-);
-
+#include "VT/VirtualTextureScalability.h"
+#include "VT/UploadingVirtualTexture.h"
 
 namespace
 {
@@ -229,7 +220,7 @@ void URuntimeVirtualTexture::GetProducerDescription(FVTProducerDescription& OutD
 	OutDesc.HeightInBlocks = 1;
 
 	// Apply TileCount modifier here to allow size scalability option
-	const int32 TileCountBias = CVarVTTileCountBias.GetValueOnAnyThread();
+	const int32 TileCountBias = VirtualTextureScalability::GetRuntimeVirtualTextureSizeBias();
 	const int32 MaxSizeInTiles = GetTileCount(TileCount + TileCountBias);
 
 	// Set width and height to best match the runtime virtual texture volume's aspect ratio
diff --git a/Engine/Source/Runtime/Engine/Private/VT/VirtualTexturePoolConfig.cpp b/Engine/Source/Runtime/Engine/Private/VT/VirtualTexturePoolConfig.cpp
--- a/Engine/Source/Runtime/Engine/Private/VT/VirtualTexturePoolConfig.cpp
+++ b/Engine/Source/Runtime/Engine/Private/VT/VirtualTexturePoolConfig.cpp
@@ -2,32 +2,25 @@
 
 #include "VT/VirtualTexturePoolConfig.h"
 
+#include "VT/VirtualTextureScalability.h"
+
 UVirtualTexturePoolConfig::UVirtualTexturePoolConfig(const FObjectInitializer& ObjectInitializer)
 	: Super(ObjectInitializer)
 {
 }
 
-UVirtualTexturePoolConfig::~UVirtualTexturePoolConfig()
-{
-}
-
-void UVirtualTexturePoolConfig::PostLoad()
-{
-	DefaultConfig.SizeInMegabyte = DefaultSizeInMegabyte;
-}
-
-const FVirtualTextureSpacePoolConfig *UVirtualTexturePoolConfig::FindPoolConfig(TEnumAsByte<EPixelFormat> const* Formats, int32 NumLayers, int32 TileSize)
+void UVirtualTexturePoolConfig::FindPoolConfig(TEnumAsByte<EPixelFormat> const* InFormats, int32 InNumLayers, int32 InTileSize, FVirtualTextureSpacePoolConfig& OutConfig) const
 {
 	// Reverse iterate so that project config can override base config
 	for (int32 Id = Pools.Num() - 1; Id >= 0 ; Id--)
 	{
 		const FVirtualTextureSpacePoolConfig& Config = Pools[Id];
-		if (Config.MinTileSize <= TileSize && Config.MaxTileSize >= TileSize && NumLayers == Config.Formats.Num())
+		if (Config.MinTileSize <= InTileSize && Config.MaxTileSize >= InTileSize && InNumLayers == Config.Formats.Num())
 		{
 			bool bAllFormatsMatch = true;
-			for (int Layer = 0; Layer < NumLayers && bAllFormatsMatch; ++Layer)
+			for (int Layer = 0; Layer < InNumLayers && bAllFormatsMatch; ++Layer)
 			{
-				if (Formats[Layer] != Config.Formats[Layer])
+				if (InFormats[Layer] != Config.Formats[Layer])
 				{
 					bAllFormatsMatch = false;
 				}
@@ -35,11 +28,14 @@ const FVirtualTextureSpacePoolConfig *UVirtualTexturePoolConfig::FindPoolConfig(
 
 			if (bAllFormatsMatch)
 			{
-				return &Config;
+				OutConfig = Config;
+				const float Scale = Config.bAllowSizeScale ? VirtualTextureScalability::GetPoolSizeScale() : 1.f;
+				OutConfig.SizeInMegabyte = (int32)(Scale * (float)OutConfig.SizeInMegabyte);
+				return;
 			}
 		}
 	}
 
-	DefaultConfig.SizeInMegabyte = DefaultSizeInMegabyte;
-	return &DefaultConfig;
+	OutConfig = FVirtualTextureSpacePoolConfig();
+	OutConfig.SizeInMegabyte = DefaultSizeInMegabyte;
 }
diff --git a/Engine/Source/Runtime/Engine/Private/VT/VirtualTextureScalability.cpp b/Engine/Source/Runtime/Engine/Private/VT/VirtualTextureScalability.cpp
new file mode 100644
--- /dev/null
+++ b/Engine/Source/Runtime/Engine/Private/VT/VirtualTextureScalability.cpp
@@ -0,0 +1,121 @@
+// Copyright 1998-2019 Epic Games, Inc. All Rights Reserved.
+
+#include "VT/VirtualTextureScalability.h"
+
+#include "Components/RuntimeVirtualTextureComponent.h"
+#include "CoreGlobals.h"
+#include "HAL/IConsoleManager.h"
+#include "UObject/UObjectIterator.h"
+#include "VT/RuntimeVirtualTexture.h"
+
+namespace VirtualTextureScalability
+{
+#if WITH_EDITOR
+	static TAutoConsoleVariable<int32> CVarVTMaxUploadsPerFrameInEditor(
+		TEXT("r.VT.MaxUploadsPerFrameInEditor"),
+		64,
+		TEXT("Max number of page uploads per frame when in editor"),
+		ECVF_RenderThreadSafe
+	);
+#endif 
+
+	static TAutoConsoleVariable<int32> CVarVTMaxUploadsPerFrame(
+		TEXT("r.VT.MaxUploadsPerFrame"),
+		8,
+		TEXT("Max number of page uploads per frame in game"),
+		ECVF_RenderThreadSafe | ECVF_Scalability
+	);
+
+	static TAutoConsoleVariable<float> CVarVTPoolSizeScale(
+		TEXT("r.VT.PoolSizeScale"),
+		1.f,
+		TEXT("Scale factor for virtual texture physical pool sizes."),
+		ECVF_RenderThreadSafe | ECVF_Scalability
+	);
+
+	static TAutoConsoleVariable<int32> CVarVTTileCountBias(
+		TEXT("r.VT.RVT.TileCountBias"),
+		0,
+		TEXT("Bias to apply to Runtime Virtual Texture size."),
+		ECVF_RenderThreadSafe | ECVF_Scalability
+	);
+
+
+	/** Track changes and apply to relevant systems. This allows us to dynamically change the scalability settings. */
+	static void OnUpdate()
+	{
+		const float PoolSizeScale = CVarVTPoolSizeScale.GetValueOnGameThread();
+		const float TileCountBias = CVarVTTileCountBias.GetValueOnGameThread();
+
+		static float LastPoolSizeScale = PoolSizeScale;
+		static float LastTileCountBias = TileCountBias;
+
+		bool bUpdate = false;
+		if (LastPoolSizeScale != PoolSizeScale)
+		{
+			LastPoolSizeScale = PoolSizeScale;
+			bUpdate = true;
+		}
+		if (LastTileCountBias != TileCountBias)
+		{
+			LastTileCountBias = TileCountBias;
+			bUpdate = true;
+		}
+
+		if (bUpdate)
+		{
+			// Temporarily release runtime virtual textures
+			for (TObjectIterator<URuntimeVirtualTexture> It; It; ++It)
+			{
+				It->Release();
+			}
+
+			// Release streaming virtual textures
+			TArray<UTexture2D*> ReleasedVirtualTextures;
+			for (TObjectIterator<UTexture2D> It; It; ++It)
+			{
+				if (It->IsCurrentlyVirtualTextured())
+				{
+					ReleasedVirtualTextures.Add(*It);
+					BeginReleaseResource(It->Resource);
+				}
+			}
+
+			// Now all pools should be flushed...
+			// Reinit streaming virtual textures
+			for (UTexture2D* Texture : ReleasedVirtualTextures)
+			{
+				BeginReleaseResource(Texture->Resource);
+			}
+
+			// Reinit runtime virtual textures
+			for (TObjectIterator<URuntimeVirtualTextureComponent> It; It; ++It)
+			{
+				It->MarkRenderStateDirty();
+			}
+		}
+	}
+
+	FAutoConsoleVariableSink GConsoleVariableSink(FConsoleCommandDelegate::CreateStatic(&OnUpdate));
+
+
+	int32 GetMaxUploadsPerFrame()
+	{
+#if WITH_EDITOR
+		// Don't want this scalability setting to affect editor because we rely on reactive updates while editing.
+		return GIsEditor ? CVarVTMaxUploadsPerFrameInEditor.GetValueOnAnyThread() : CVarVTMaxUploadsPerFrame.GetValueOnAnyThread();
+#else
+		return CVarVTMaxUploadsPerFrame.GetValueOnAnyThread();
+#endif
+	}
+
+	float GetPoolSizeScale()
+	{
+		return CVarVTPoolSizeScale.GetValueOnAnyThread();
+	}
+
+	int32 GetRuntimeVirtualTextureSizeBias()
+	{
+		return CVarVTTileCountBias.GetValueOnAnyThread();
+	}
+}
diff --git a/Engine/Source/Runtime/Engine/Public/VT/VirtualTextureScalability.h b/Engine/Source/Runtime/Engine/Public/VT/VirtualTextureScalability.h
new file mode 100644
--- /dev/null
+++ b/Engine/Source/Runtime/Engine/Public/VT/VirtualTextureScalability.h
@@ -0,0 +1,13 @@
+// Copyright 1998-2019 Epic Games, Inc. All Rights Reserved.
+
+#pragma once
+
+namespace VirtualTextureScalability
+{
+	/** Get max upload rate to virtual textures. */
+	ENGINE_API int32 GetMaxUploadsPerFrame();
+	/** Get scale factor for virtual texture physical pool sizes. */
+	ENGINE_API float GetPoolSizeScale();
+	/** Get resolution bias for runtime virtual textures. */
+	ENGINE_API int32 GetRuntimeVirtualTextureSizeBias();
+}
diff --git a/Engine/Source/Runtime/Renderer/Private/VT/VirtualTexturePhysicalSpace.cpp b/Engine/Source/Runtime/Renderer/Private/VT/VirtualTexturePhysicalSpace.cpp
--- a/Engine/Source/Runtime/Renderer/Private/VT/VirtualTexturePhysicalSpace.cpp
+++ b/Engine/Source/Runtime/Renderer/Private/VT/VirtualTexturePhysicalSpace.cpp
@@ -18,10 +18,10 @@ FVirtualTexturePhysicalSpace::FVirtualTexturePhysicalSpace(const FVTPhysicalSpac
 	, bGpuTextureLimit(false)
 {
 	// Find matching physical pool
-	UVirtualTexturePoolConfig* PoolConfig = GetMutableDefault<UVirtualTexturePoolConfig>();
-	const FVirtualTextureSpacePoolConfig* Config = PoolConfig->FindPoolConfig(InDesc.Format, InDesc.NumLayers, InDesc.TileSize);
-	const uint32 PoolSizeInBytes = Config->SizeInMegabyte * 1024u * 1024u;
-	const bool bForce16BitPageTable = false;
+	FVirtualTextureSpacePoolConfig Config;
+	UVirtualTexturePoolConfig const* PoolConfig = GetDefault<UVirtualTexturePoolConfig>();
+	PoolConfig->FindPoolConfig(InDesc.Format, InDesc.NumLayers, InDesc.TileSize, Config);
+	const uint32 PoolSizeInBytes = Config.SizeInMegabyte * 1024u * 1024u;
 
 	const FPixelFormatInfo& FormatInfo = GPixelFormats[InDesc.Format[0]];
 	check(InDesc.TileSize % FormatInfo.BlockSizeX == 0);
@@ -31,14 +31,15 @@ FVirtualTexturePhysicalSpace::FVirtualTexturePhysicalSpace(const FVTPhysicalSpac
 	{
 		TileSizeBytes += CalculateImageBytes(InDesc.TileSize, InDesc.TileSize, 0, InDesc.Format[Layer]);
 	}
-	const uint32 MaxTiles = (uint32)(PoolSizeInBytes / TileSizeBytes);
-	
+	const uint32 MaxTiles = FMath::Max((uint32)(PoolSizeInBytes / TileSizeBytes), 1u);
 	TextureSizeInTiles = FMath::FloorToInt(FMath::Sqrt((float)MaxTiles));
-	if (bForce16BitPageTable)
-	{
-		// 16 bit page tables support max size of 64x64 (4096 tiles)
-		TextureSizeInTiles = FMath::Min(64u, TextureSizeInTiles);
-	}
+	
+// 	const bool bForce16BitPageTable = false;
+// 	if (bForce16BitPageTable)
+// 	{
+// 		// 16 bit page tables support max size of 64x64 (4096 tiles)
+// 		TextureSizeInTiles = FMath::Min(64u, TextureSizeInTiles);
+// 	}
 
 	if (TextureSizeInTiles * InDesc.TileSize > GetMax2DTextureDimension())
 	{
diff --git a/Engine/Source/Runtime/Renderer/Private/VT/VirtualTextureSystem.cpp b/Engine/Source/Runtime/Renderer/Private/VT/VirtualTextureSystem.cpp
--- a/Engine/Source/Runtime/Renderer/Private/VT/VirtualTextureSystem.cpp
+++ b/Engine/Source/Runtime/Renderer/Private/VT/VirtualTextureSystem.cpp
@@ -9,13 +9,14 @@
 #include "ScenePrivate.h"
 #include "SceneUtils.h"
 #include "Stats/Stats.h"
-#include "TexturePagePool.h"
-#include "UniquePageList.h"
-#include "UniqueRequestList.h"
-#include "VirtualTextureFeedback.h"
-#include "VirtualTexturePhysicalSpace.h"
-#include "VirtualTextureSpace.h"
 #include "VirtualTexturing.h"
+#include "VT/TexturePagePool.h"
+#include "VT/UniquePageList.h"
+#include "VT/UniqueRequestList.h"
+#include "VT/VirtualTextureFeedback.h"
+#include "VT/VirtualTexturePhysicalSpace.h"
+#include "VT/VirtualTextureScalability.h"
+#include "VT/VirtualTextureSpace.h"
 
 DECLARE_CYCLE_STAT(TEXT("Feedback Analysis"), STAT_FeedbackAnalysis, STATGROUP_VirtualTexturing);
 DECLARE_CYCLE_STAT(TEXT("VirtualTextureSystem Update"), STAT_VirtualTextureSystem_Update, STATGROUP_VirtualTexturing);
@@ -51,13 +52,6 @@ DECLARE_MEMORY_STAT_POOL(TEXT("Total Pagetable Memory"), STAT_TotalPagetableMemo
 
 DECLARE_GPU_STAT( VirtualTexture );
 
-static TAutoConsoleVariable<int32> CVarVTMaxUploadsPerFrame(
-	TEXT("r.VT.MaxUploadsPerFrame"),
-	64,
-	TEXT("Max number of page uploads per frame"),
-	ECVF_RenderThreadSafe
-	);
-
 static TAutoConsoleVariable<int32> CVarVTEnableFeedBack(
 	TEXT("r.VT.EnableFeedBack"),
 	1,
@@ -1100,7 +1094,7 @@ void FVirtualTextureSystem::Update(FRHICommandListImmediate& RHICmdList, ERHIFea
 
 		// Limit the number of uploads (account for MappedTilesToProduce this frame)
 		// Are all pages equal? Should there be different limits on different types of pages?
-		const int32 MaxNumUploads = CVarVTMaxUploadsPerFrame.GetValueOnRenderThread();
+		const int32 MaxNumUploads = VirtualTextureScalability::GetMaxUploadsPerFrame();
 		const int32 MaxRequestUploads = FMath::Max(MaxNumUploads - MappedTilesToProduce.Num(), 1);
 
 		MergedRequestList->SortRequests(Producers, MemStack, MaxRequestUploads);
