Fix issue where Lifetime Custom Delta State tracked by [REDACTED_PERSON] used improperly sized lookup tables.

diff --git a/Engine/Source/Runtime/Engine/Private/RepLayout.cpp b/Engine/Source/Runtime/Engine/Private/RepLayout.cpp
--- a/Engine/Source/Runtime/Engine/Private/RepLayout.cpp
+++ b/Engine/Source/Runtime/Engine/Private/RepLayout.cpp
@@ -316,9 +316,10 @@ struct FLifetimeCustomDeltaState
 {
 public:
 
-	FLifetimeCustomDeltaState(uint16 TotalNumberOfLifetimeProperties)
+	FLifetimeCustomDeltaState(int32 HighestCustomDeltaRepIndex)
 	{
-		LifetimeCustomDeltaIndexLookup.Init(static_cast<uint16>(INDEX_NONE), TotalNumberOfLifetimeProperties);
+		check(HighestCustomDeltaRepIndex >= 0);
+		LifetimeCustomDeltaIndexLookup.Init(static_cast<uint16>(INDEX_NONE), HighestCustomDeltaRepIndex + 1);
 	}
 
 	void CountBytes(FArchive& Ar) const
@@ -4993,6 +4994,7 @@ void FRepLayout::InitFromClass(
 
 	int32 RelativeHandle = 0;
 	int32 LastOffset = INDEX_NONE;
+	int32 HighestCustomDeltaRepIndex = INDEX_NONE;
 
 	InObjectClass->SetUpRuntimeReplicationData();
 	Parents.Empty(InObjectClass->ClassReps.Num());
@@ -5048,6 +5050,11 @@ void FRepLayout::InitFromClass(
 			    RemoteRoleIndex = ParentHandle;
 		    }
 	    }
+
+		if (EnumHasAnyFlags(Parents[ParentHandle].Flags, ERepParentFlags::IsCustomDelta))
+		{
+			HighestCustomDeltaRepIndex = Parents[ParentHandle].RepIndex;
+		}
 	}
 
 	// Make sure it either found both, or didn't find either
@@ -5118,7 +5125,12 @@ void FRepLayout::InitFromClass(
 
 			if (!LifetimeCustomPropertyState)
 			{
-				LifetimeCustomPropertyState.Reset(new FLifetimeCustomDeltaState(LifetimeProps.Num()));
+				// We can't use the number of Lifetime Properties, because that could be smaller than
+				// the highest RepIndex of a Custom Delta Property, because properties may be disabled, removed,
+				// or just never added.
+				// For similar reasons, we don't want to use the total number of replicated properties, especially
+				// if we know we'll never use anything beyond the last Custom Delta Property anyway.
+				LifetimeCustomPropertyState.Reset(new FLifetimeCustomDeltaState(HighestCustomDeltaRepIndex));
 			}
 
 			// If we're a FastArraySerializer, we'll look for our replicated item type.
