Fix issue where replicated properties sent reliably due to net.PartialBunchReliableThreshold still generate history items.
This would cause the History buffer in FSendingRepState to fill up with reliable changelists that weren't initialized.

diff --git a/Engine/Source/Runtime/Engine/Private/DataChannel.cpp b/Engine/Source/Runtime/Engine/Private/DataChannel.cpp
--- a/Engine/Source/Runtime/Engine/Private/DataChannel.cpp
+++ b/Engine/Source/Runtime/Engine/Private/DataChannel.cpp
@@ -2816,7 +2816,7 @@ int64 UActorChannel::ReplicateActor()
 			return 0;
 		}
 		bPausedUntilReliableACK = 0;
-		UE_LOG(LogNet, Log, TEXT("ReplicateActor: bPausedUntilReliableACK is ending now that reliables have been ACK'd. %s"), *Describe());
+		UE_LOG(LogNet, Verbose, TEXT("ReplicateActor: bPausedUntilReliableACK is ending now that reliables have been ACK'd. %s"), *Describe());
 	}
 
 	const TArray<FNetViewer>& NetViewers = ActorWorld->GetWorldSettings()->ReplicationViewers;
diff --git a/Engine/Source/Runtime/Engine/Private/DataReplication.cpp b/Engine/Source/Runtime/Engine/Private/DataReplication.cpp
--- a/Engine/Source/Runtime/Engine/Private/DataReplication.cpp
+++ b/Engine/Source/Runtime/Engine/Private/DataReplication.cpp
@@ -1647,6 +1647,27 @@ void FObjectReplicator::PostSendBunch( FPacketIdRange & PacketRange, uint8 bReli
 			}
 		}
 
+		// bPausedUntilReliableAck won't have been set until *after* we generated a changelist and tried to send
+		// data. Therefore, there may be a changelist already created, and we need to remove it from the history.
+		else
+		{
+			// We won't try to replicate again until this bunch has been acked, so there should only ever possibly
+			// be a single unset history item, and that should correspond to the one that is in this bunch.
+			// However, PostSendBunch will be called on all subobject replicators when we replicate an Actor, so it's
+			// possible that we either didn't send anything, or the last thing we sent wasn't reliable.
+			if (SendingRepState->HistoryEnd > SendingRepState->HistoryStart)
+			{
+				const int32 HistoryIndex = (SendingRepState->HistoryEnd - 1) % FSendingRepState::MAX_CHANGE_HISTORY;
+				FRepChangedHistory& HistoryItem = SendingRepState->ChangeHistory[HistoryIndex];
+	
+				if (!HistoryItem.WasSent())
+				{
+					HistoryItem.Reset();
+					--SendingRepState->HistoryEnd;
+				}
+			}
+		}
+
 		for (FPropertyRetirement& Retirement : SendingRepState->Retirement)
 		{
 			FPropertyRetirement* Next = Retirement.Next;
diff --git a/Engine/Source/Runtime/Engine/Private/RepLayout.cpp b/Engine/Source/Runtime/Engine/Private/RepLayout.cpp
--- a/Engine/Source/Runtime/Engine/Private/RepLayout.cpp
+++ b/Engine/Source/Runtime/Engine/Private/RepLayout.cpp
@@ -1523,8 +1523,9 @@ void FRepLayout::UpdateChangelistHistory(
 
 		if (HistoryItem.OutPacketIdRange.First == INDEX_NONE)
 		{
-			//  Hasn't been initialized in PostReplicate yet
-			continue;
+			// Hasn't been initialized in PostReplicate yet
+			// No need to go further, otherwise we'll overwrite entries incorrectly.
+			break;
 		}
 
 		// All active history items should contain a change list
@@ -1539,21 +1540,17 @@ void FRepLayout::UpdateChangelistHistory(
 				TArray<uint16> Temp = MoveTemp(*OutMerged);
 				MergeChangeList(Data, HistoryItem.Changed, Temp, *OutMerged);
 
-				HistoryItem.Changed.Empty();
-
 #ifdef SANITY_CHECK_MERGES
 				SanityCheckChangeList(Data, *OutMerged);
 #endif
 
 				if (HistoryItem.Resend)
 				{
-					HistoryItem.Resend = false;
 					RepState->NumNaks--;
 				}
 			}
 
-			HistoryItem.Changed.Empty();
-			HistoryItem.OutPacketIdRange = FPacketIdRange();
+			HistoryItem.Reset();
 			RepState->HistoryStart++;
 		}
 	}
@@ -1561,7 +1558,7 @@ void FRepLayout::UpdateChangelistHistory(
 	// Remove any tiling in the history markers to keep them from wrapping over time
 	const int32 NewHistoryCount	= RepState->HistoryEnd - RepState->HistoryStart;
 
-	check(NewHistoryCount <= FSendingRepState::MAX_CHANGE_HISTORY);
+	check(NewHistoryCount < FSendingRepState::MAX_CHANGE_HISTORY);
 
 	RepState->HistoryStart = RepState->HistoryStart % FSendingRepState::MAX_CHANGE_HISTORY;
 	RepState->HistoryEnd = RepState->HistoryStart + NewHistoryCount;
diff --git a/Engine/Source/Runtime/Engine/Public/Net/RepLayout.h b/Engine/Source/Runtime/Engine/Public/Net/RepLayout.h
--- a/Engine/Source/Runtime/Engine/Public/Net/RepLayout.h
+++ b/Engine/Source/Runtime/Engine/Public/Net/RepLayout.h
@@ -315,6 +315,18 @@ public:
 		Changed.CountBytes(Ar);
 	}
 
+	bool WasSent() const
+	{
+		return OutPacketIdRange.First != INDEX_NONE;
+	}
+
+	void Reset()
+	{
+		OutPacketIdRange = FPacketIdRange();
+		Changed.Empty();
+		Resend = false;
+	}
+
 	/** Range of the Packets that this changelist was last sent with. Used to track acknowledgments. */
 	FPacketIdRange OutPacketIdRange;
 
@@ -1063,6 +1075,25 @@ enum class UE_DEPRECATED(4.24, "Use FRepLayout::IsEmpty() instead.") ERepLayoutS
  * Receiving is very similar, except the Handles are baked into the serialized data so no
  * explicit changelist is required. As each Handle is read, a Layout Command is applied
  * that serializes the data from the network bunch and applies it to an object.
+ *
+ * RETRIES AND RELIABLES
+ *
+ * @FSendingRepState maintains a circular buffer that tracks recently sent Changelists (@FRepChangedHistory).
+ * These history items track the Changelist alongside the Packet ID that the bunches were sent in.
+ * 
+ * Once we receive ACKs for all associated packets, the history will be removed from the buffer.
+ * If NAKs are received for any of the packets, we will merge the changelist into the next set of properties we replicate.
+ *
+ * If we receive no NAKs or ACKs for an extended period, to prevent overflows in the history buffer,
+ * we will merge the entire buffer into a single monolithic changelist which will be sent alongside the next set of properties.
+ *
+ * In both cases of NAKs or no response, the merged changelists will be tracked in the latest history item
+ * alongside with other sent properties.
+ *
+ * When "net.PartialBunchReliableThreshold" is non-zero and property data bunches are split into partial bunches above
+ * the threshold, we will not generate a history item. Instead, we will rely on the reliable bunch framework for resends
+ * and replication of the Object will be completely paused until the property bunches are acknowledged.
+ * However, this will not affect other history items since they are still unreliable.
  */
 class ENGINE_VTABLE FRepLayout : public FGCObject, public TSharedFromThis<FRepLayout>
 {
