Back out changelist 9447300. CIS errors

diff --git a/Engine/Source/Runtime/Engine/Private/RepLayout.cpp b/Engine/Source/Runtime/Engine/Private/RepLayout.cpp
--- a/Engine/Source/Runtime/Engine/Private/RepLayout.cpp
+++ b/Engine/Source/Runtime/Engine/Private/RepLayout.cpp
@@ -316,10 +316,9 @@ struct FLifetimeCustomDeltaState
 {
 public:
 
-	FLifetimeCustomDeltaState(int32 HighestCustomDeltaRepIndex)
+	FLifetimeCustomDeltaState(uint16 TotalNumberOfLifetimeProperties)
 	{
-		check(HighestCustomDeltaRepIndex >= 0);
-		LifetimeCustomDeltaIndexLookup.Init(static_cast<uint16>(INDEX_NONE), HighestCustomDeltaRepIndex + 1);
+		LifetimeCustomDeltaIndexLookup.Init(static_cast<uint16>(INDEX_NONE), TotalNumberOfLifetimeProperties);
 	}
 
 	void CountBytes(FArchive& Ar) const
@@ -4991,7 +4990,6 @@ void FRepLayout::InitFromClass(
 
 	int32 RelativeHandle = 0;
 	int32 LastOffset = INDEX_NONE;
-	int32 HighestCustomDeltaRepIndex = INDEX_NONE;
 
 	InObjectClass->SetUpRuntimeReplicationData();
 	Parents.Empty(InObjectClass->ClassReps.Num());
@@ -5047,11 +5045,6 @@ void FRepLayout::InitFromClass(
 			    RemoteRoleIndex = ParentHandle;
 		    }
 	    }
-
-		if (EnumHasAnyFlags(Parents[ParentHandle].Flags, ERepParentFlags::IsCustomDelta))
-		{
-			HighestCustomDeltaRepIndex = Parents[ParentHandle].RepIndex;
-		}
 	}
 
 	// Make sure it either found both, or didn't find either
@@ -5122,12 +5115,7 @@ void FRepLayout::InitFromClass(
 
 			if (!LifetimeCustomPropertyState)
 			{
-				// We can't use the number of Lifetime Properties, because that could be smaller than
-				// the highest RepIndex of a Custom Delta Property, because properties may be disabled, removed,
-				// or just never added.
-				// For similar reasons, we don't want to use the total number of replicated properties, especially
-				// if we know we'll never use anything beyond the last Custom Delta Property anyway.
-				LifetimeCustomPropertyState.Reset(new FLifetimeCustomDeltaState(HighestCustomDeltaRepIndex));
+				LifetimeCustomPropertyState.Reset(new FLifetimeCustomDeltaState(LifetimeProps.Num()));
 			}
 
 			// If we're a FastArraySerializer, we'll look for our replicated item type.
