Fix bad merge
#rnx

diff --git a/Engine/Source/Runtime/Engine/Private/Audio.cpp b/Engine/Source/Runtime/Engine/Private/Audio.cpp
--- a/Engine/Source/Runtime/Engine/Private/Audio.cpp
+++ b/Engine/Source/Runtime/Engine/Private/Audio.cpp
@@ -85,6 +85,7 @@ FAutoConsoleVariableRef CVarOcclusionFilterScaleEnabled(
 	OcclusionFilterScaleEnabledCVar,
 	TEXT("Whether or not we scale occlusion by 0.25f to compensate for change in filter cutoff frequencies in audio mixer. \n")
 	TEXT("0: Not Enabled, 1: Enabled"),
+	ECVF_Default);
 
 static int32 AllowReverbForMultichannelSources = 1;
 FAutoConsoleVariableRef CvarAllowReverbForMultichannelSources(
