Chaos - reenable collisions and collision filtering in RBAN

diff --git a/Engine/Source/Runtime/Engine/Private/PhysicsEngine/ImmediatePhysics/ImmediatePhysicsChaos/ImmediatePhysicsSimulation_Chaos.cpp b/Engine/Source/Runtime/Engine/Private/PhysicsEngine/ImmediatePhysics/ImmediatePhysicsChaos/ImmediatePhysicsSimulation_Chaos.cpp
--- a/Engine/Source/Runtime/Engine/Private/PhysicsEngine/ImmediatePhysics/ImmediatePhysicsChaos/ImmediatePhysicsSimulation_Chaos.cpp
+++ b/Engine/Source/Runtime/Engine/Private/PhysicsEngine/ImmediatePhysics/ImmediatePhysicsChaos/ImmediatePhysicsSimulation_Chaos.cpp
@@ -20,7 +20,7 @@
 // @todo(ccaulfield): remove when finished
 float ChaosImmediate_Evolution_DeltaTime = 0.03f;
 int32 ChaosImmediate_Evolution_Iterations = 10;
-int32 ChaosImmediate_Collision_Enabled = 0;
+int32 ChaosImmediate_Collision_Enabled = 1;
 int32 ChaosImmediate_Collision_ApplyEnabled = 0;
 int32 ChaosImmediate_Collision_PushOutIterations = 5;
 int32 ChaosImmediate_Collision_PushOutPairIterations = 2;
@@ -96,25 +96,6 @@ namespace ImmediatePhysics_Chaos
 		return false;
 	}
 
-	template<typename T, int d>
-	void FilterCollisionConstraints(
-		TArray<Chaos::TRigidBodyContactConstraint<T, d>>& Constraints, 
-		const TMap<const Chaos::TGeometryParticleHandle<T, d>*, TSet<const Chaos::TGeometryParticleHandle<T, d>*>>& IgnoreSetMap)
-	{
-		using namespace Chaos;
-
-		for (int ConstraintIndex = 0; ConstraintIndex < Constraints.Num(); ++ConstraintIndex)
-		{
-			if (ShouldIgnoreCollisionConstraint(Constraints[ConstraintIndex].Particle, Constraints[ConstraintIndex].Levelset, IgnoreSetMap))
-			{
-				Constraints.RemoveAtSwap(ConstraintIndex);
-				--ConstraintIndex;
-				continue;
-			}
-		}
-	}
-
-
 	//
 	//
 	//
@@ -133,10 +114,20 @@ namespace ImmediatePhysics_Chaos
 
 		Evolution->AddConstraintRule(JointsRule.Get());
 
+		// Filter collisions after detection
+		// @todo(ccaulfield): Eventually we will build lists of potentially colliding pairs and won't need this
 		Collisions.SetPostComputeCallback(
-			[this](TArray<TRigidBodyContactConstraint<FReal, Dimensions>>& Constraints)
+			[this]()
 			{
-				FilterCollisionConstraints(Constraints, IgnoreCollisionParticlePairTable);
+				Evolution->GetCollisionConstraints().ApplyCollisionModifier(
+					[this](TRigidBodyContactConstraint<float, 3>& Constraint)
+					{
+						if (ShouldIgnoreCollisionConstraint(Constraint.Particle, Constraint.Levelset, IgnoreCollisionParticlePairTable))
+						{
+							return ECollisionModifierResult::Disabled;
+						}
+						return ECollisionModifierResult::Unchanged;
+					});
 			});
 		
 #if CHAOS_DEBUG_DRAW
diff --git a/Engine/Source/Runtime/Experimental/Chaos/Private/Chaos/PBDCollisionConstraint.cpp b/Engine/Source/Runtime/Experimental/Chaos/Private/Chaos/PBDCollisionConstraint.cpp
--- a/Engine/Source/Runtime/Experimental/Chaos/Private/Chaos/PBDCollisionConstraint.cpp
+++ b/Engine/Source/Runtime/Experimental/Chaos/Private/Chaos/PBDCollisionConstraint.cpp
@@ -128,6 +128,20 @@ void TPBDCollisionConstraint<T, d>::RemoveConstraint(int32 Idx)
 	}
 }
 
+template<typename T, int d>
+void TPBDCollisionConstraint<T, d>::ApplyCollisionModifier(const TFunction<ECollisionModifierResult(FRigidBodyContactConstraint& Constraint)>& CollisionModifier)
+{
+	for (int ConstraintIndex = 0; ConstraintIndex < NumConstraints(); ++ConstraintIndex)
+	{
+		ECollisionModifierResult Result = CollisionModifier(Constraints[ConstraintIndex]);
+		if (Result == ECollisionModifierResult::Disabled)
+		{
+			RemoveConstraint(ConstraintIndex);
+			--ConstraintIndex;
+		}
+	}
+}
+
 template<typename T, int d>
 void TPBDCollisionConstraint<T, d>::SetPostComputeCallback(const TRigidBodyContactConstraintsPostComputeCallback<T, d>& Callback)
 {
@@ -226,7 +240,7 @@ void TPBDCollisionConstraint<T, d>::ComputeConstraints(const FAccelerationStruct
 	
 	if (PostComputeCallback != nullptr)
 	{
-		PostComputeCallback(Constraints);
+		PostComputeCallback();
 	}
 }
 
diff --git a/Engine/Source/Runtime/Experimental/Chaos/Public/Chaos/PBDCollisionConstraint.h b/Engine/Source/Runtime/Experimental/Chaos/Public/Chaos/PBDCollisionConstraint.h
--- a/Engine/Source/Runtime/Experimental/Chaos/Public/Chaos/PBDCollisionConstraint.h
+++ b/Engine/Source/Runtime/Experimental/Chaos/Public/Chaos/PBDCollisionConstraint.h
@@ -45,6 +45,13 @@ enum class ECollisionUpdateType
 	Deepest	//find the deepest penetration. Compute location and normal
 };
 
+/** Return value of the collision modification callback */
+enum class ECollisionModifierResult
+{
+	Unchanged,	/** No change to the collision */
+	Modified,	/** Modified the collision, but want it to remain enabled */
+	Disabled,	/** Collision should be disabled */
+};
 
 template<class T, int d>
 class CHAOS_API TPBDCollisionConstraintHandle : public TContainerConstraintHandle<TPBDCollisionConstraint<T, d>>
@@ -65,7 +72,7 @@ protected:
 };
 
 template<typename T, int d>
-using TRigidBodyContactConstraintsPostComputeCallback = TFunction<void(TArray<TRigidBodyContactConstraint<T, d>>& Constraints)>;
+using TRigidBodyContactConstraintsPostComputeCallback = TFunction<void()>;
 
 template<typename T, int d>
 using TRigidBodyContactConstraintsPostApplyCallback = TFunction<void(const T Dt, const TArray<TPBDCollisionConstraintHandle<T, d>*>& InConstraintHandles)>;
@@ -152,9 +159,16 @@ public:
 
 	void RemoveConstraint(int32 Idx);
 
+	/**
+	 * Apply a modifier to the constraints and specify which constraints should be disabled.
+	 * You would probably call this in the PostComputeCallback. Prefer this to calling RemoveConstraints in a loop, 
+	 * so you don't have to worry about constraint iterator/indices changing.
+	 */
+	void ApplyCollisionModifier(const TFunction<ECollisionModifierResult(FRigidBodyContactConstraint& Constraint)>& CollisionModifier);
+
 	/**
 	 * Set the callback used just after contacts are generated at the start of a frame tick.
-	 * This can be used to modify or remove constraints from the system.
+	 * This can be used to modify or disable constraints (via ApplyCollisionModifier).
 	 */
 	void SetPostComputeCallback(const TRigidBodyContactConstraintsPostComputeCallback<T, d>& Callback);
 
