Change DOREPLIFETIME implementations to use a struct for parameters, and expose a new macro DOREPLIFETIME_WITH_PARAMS.
This should make any updates easier, while maintaining backwards compatibility.

diff --git a/Engine/Source/Runtime/Engine/Private/UnrealNetwork.cpp b/Engine/Source/Runtime/Engine/Private/UnrealNetwork.cpp
--- a/Engine/Source/Runtime/Engine/Private/UnrealNetwork.cpp
+++ b/Engine/Source/Runtime/Engine/Private/UnrealNetwork.cpp
@@ -8,7 +8,24 @@ FOnProcessGameSpecificDemoHeader FNetworkReplayDelegates::OnProcessGameSpecificD
 
 // ----------------------------------------------------------------
 
-void RegisterReplicatedLifetimeProperty(const UProperty* ReplicatedProperty, TArray< FLifetimeProperty >& OutLifetimeProps, ELifetimeCondition InCondition, ELifetimeRepNotifyCondition InRepNotifyCondition)
+PRAGMA_DISABLE_DEPRECATION_WARNINGS
+void RegisterReplicatedLifetimeProperty(
+	const UProperty* ReplicatedProperty,
+	TArray<FLifetimeProperty>& OutLifetimeProps,
+	ELifetimeCondition InCondition,
+	ELifetimeRepNotifyCondition InRepNotifyCondition)
+{
+	FDoRepLifetimeParams Params;
+	Params.Condition = InCondition;
+	Params.RepNotifyCondition = InRepNotifyCondition;
+	RegisterReplicatedLifetimeProperty(ReplicatedProperty, OutLifetimeProps, Params);
+}
+PRAGMA_ENABLE_DEPRECATION_WARNINGS
+
+void RegisterReplicatedLifetimeProperty(
+	const UProperty* ReplicatedProperty,
+	TArray<FLifetimeProperty>& OutLifetimeProps,
+	const FDoRepLifetimeParams& Params)
 {
 	if (!ReplicatedProperty) 
 	{
@@ -21,7 +38,7 @@ void RegisterReplicatedLifetimeProperty(const UProperty* ReplicatedProperty, TAr
 		const uint16 RepIndex = ReplicatedProperty->RepIndex + i;
 		FLifetimeProperty* RegisteredPropertyPtr = OutLifetimeProps.FindByPredicate([&RepIndex](const FLifetimeProperty& Var) { return Var.RepIndex == RepIndex; });
 
-		FLifetimeProperty LifetimeProp(RepIndex, InCondition, InRepNotifyCondition);
+		FLifetimeProperty LifetimeProp(RepIndex, Params.Condition, Params.RepNotifyCondition);
 
 		if (RegisteredPropertyPtr)
 		{
@@ -34,7 +51,7 @@ void RegisterReplicatedLifetimeProperty(const UProperty* ReplicatedProperty, TAr
 			else
 			{
 				// Conditions should be identical when calling DOREPLIFETIME twice on the same variable.
-				checkf((*RegisteredPropertyPtr) == LifetimeProp, TEXT("Property %s was registered twice with different conditions (old:%d) (new:%d)"), *ReplicatedProperty->GetName(), RegisteredPropertyPtr->Condition, InCondition);
+				checkf((*RegisteredPropertyPtr) == LifetimeProp, TEXT("Property %s was registered twice with different conditions (old:%d) (new:%d)"), *ReplicatedProperty->GetName(), RegisteredPropertyPtr->Condition, Params.Condition);
 			}
 		}
 		else
diff --git a/Engine/Source/Runtime/Engine/Public/Net/UnrealNetwork.h b/Engine/Source/Runtime/Engine/Public/Net/UnrealNetwork.h
--- a/Engine/Source/Runtime/Engine/Public/Net/UnrealNetwork.h
+++ b/Engine/Source/Runtime/Engine/Public/Net/UnrealNetwork.h
@@ -54,6 +54,22 @@ struct ENGINE_API FNetworkReplayDelegates
 	static FOnProcessGameSpecificDemoHeader OnProcessGameSpecificDemoHeader;
 };
 
+/**
+ * Struct containing various parameters that can be passed to DOREPLIFETIME_WITH_PARAMS to control
+ * how variables are replicated.
+ */
+struct ENGINE_API FDoRepLifetimeParams
+{
+	/** Replication Condition. The property will only be replicated to connections where this condition is met. */
+	ELifetimeCondition Condition = COND_None;
+	
+	/**
+	 * RepNotify Condition. The property will only trigger a RepNotify if this condition is met, and has been
+	 * properly set up to handle RepNotifies.
+	 */
+	ELifetimeRepNotifyCondition RepNotifyCondition = REPNOTIFY_OnChanged;
+};
+
 /*-----------------------------------------------------------------------------
 	Replication.
 -----------------------------------------------------------------------------*/
@@ -77,12 +93,14 @@ static UProperty* GetReplicatedProperty(const UClass* CallingClass, const UClass
 	return TheProperty;
 }
 
-#define DOREPLIFETIME(c,v) \
+#define DOREPLIFETIME_WITH_PARAMS(c,v,params) \
 { \
 	UProperty* ReplicatedProperty = GetReplicatedProperty(StaticClass(), c::StaticClass(),GET_MEMBER_NAME_CHECKED(c,v)); \
-	RegisterReplicatedLifetimeProperty(ReplicatedProperty, OutLifetimeProps, COND_None); \
+	RegisterReplicatedLifetimeProperty(ReplicatedProperty, OutLifetimeProps, params); \
 }
 
+#define DOREPLIFETIME(c,v) DOREPLIFETIME_WITH_PARAMS(c,v,FDoRepLifetimeParams())
+
 /** This macro is used by nativized code (DynamicClasses), so the Property may be recreated. */
 #define DOREPLIFETIME_DIFFNAMES(c,v, n) \
 { \
@@ -101,17 +119,21 @@ static UProperty* GetReplicatedProperty(const UClass* CallingClass, const UClass
 
 #define DOREPLIFETIME_CONDITION(c,v,cond) \
 { \
-	UProperty* ReplicatedProperty = GetReplicatedProperty(StaticClass(), c::StaticClass(),GET_MEMBER_NAME_CHECKED(c,v)); \
-	RegisterReplicatedLifetimeProperty(ReplicatedProperty, OutLifetimeProps, cond); \
+	FDoRepLifetimeParams LocalDoRepParams; \
+	LocalDoRepParams.Condition = cond; \
+	DOREPLIFETIME_WITH_PARAMS(c,v,LocalDoRepParams); \
 }
 
 /** Allows gamecode to specify RepNotify condition: REPNOTIFY_OnChanged (default) or REPNOTIFY_Always for when repnotify function is called  */
-#define DOREPLIFETIME_CONDITION_NOTIFY(c,v,cond, rncond) \
+#define DOREPLIFETIME_CONDITION_NOTIFY(c,v,cond,rncond) \
 { \
-	UProperty* ReplicatedProperty = GetReplicatedProperty(StaticClass(), c::StaticClass(),GET_MEMBER_NAME_CHECKED(c,v)); \
-	RegisterReplicatedLifetimeProperty(ReplicatedProperty, OutLifetimeProps, cond, rncond); \
+	FDoRepLifetimeParams LocalDoRepParams; \
+	LocalDoRepParams.Condition = cond; \
+	LocalDoRepParams.RepNotifyCondition = rncond; \
+	DOREPLIFETIME_WITH_PARAMS(c,v,LocalDoRepParams); \
 }
 
+
 #define DOREPLIFETIME_ACTIVE_OVERRIDE(c,v,active)	\
 {													\
 	static UProperty* sp##v = GetReplicatedProperty(StaticClass(), c::StaticClass(),GET_MEMBER_NAME_CHECKED(c,v)); \
@@ -121,7 +143,7 @@ static UProperty* GetReplicatedProperty(const UClass* CallingClass, const UClass
 	}																						\
 }
 
-UE_DEPRECATED(4.30, "Please use the RESET_REPLIFETIME_CONDITION macro")
+UE_DEPRECATED(4.24, "Please use the RESET_REPLIFETIME_CONDITION macro")
 ENGINE_API void DeprecatedChangeCondition(UProperty* ReplicatedProperty, TArray< FLifetimeProperty >& OutLifetimeProps, ELifetimeCondition InCondition);
 
 #define DOREPLIFETIME_CHANGE_CONDITION(c,v,cond) \
@@ -130,7 +152,17 @@ ENGINE_API void DeprecatedChangeCondition(UProperty* ReplicatedProperty, TArray<
 	DeprecatedChangeCondition(sp##v, OutLifetimeProps, cond);														\
 }
 
-ENGINE_API void RegisterReplicatedLifetimeProperty(const UProperty* ReplicatedProperty, TArray< FLifetimeProperty >& OutLifetimeProps, ELifetimeCondition InCondition, ELifetimeRepNotifyCondition InRepNotifyCondition = REPNOTIFY_OnChanged);
+UE_DEPRECATED(4.24, "Use RegisterReplicatedLifetimeProperty that takes FDoRepLifetimeParams.")
+ENGINE_API void RegisterReplicatedLifetimeProperty(
+	const UProperty* ReplicatedProperty,
+	TArray<FLifetimeProperty>& OutLifetimeProps,
+	ELifetimeCondition InCondition,
+	ELifetimeRepNotifyCondition InRepNotifyCondition = REPNOTIFY_OnChanged);
+
+ENGINE_API void RegisterReplicatedLifetimeProperty(
+	const UProperty* ReplicatedProperty,
+	TArray<FLifetimeProperty>& OutLifetimeProps,
+	const FDoRepLifetimeParams& Params);
 
 /*-----------------------------------------------------------------------------
 	Disable macros.
