#UE4 IWYU CIS Fix

diff --git a/Engine/Source/Runtime/Engine/Private/VT/VirtualTextureScalability.cpp b/Engine/Source/Runtime/Engine/Private/VT/VirtualTextureScalability.cpp
--- a/Engine/Source/Runtime/Engine/Private/VT/VirtualTextureScalability.cpp
+++ b/Engine/Source/Runtime/Engine/Private/VT/VirtualTextureScalability.cpp
@@ -4,6 +4,7 @@
 
 #include "Components/RuntimeVirtualTextureComponent.h"
 #include "CoreGlobals.h"
+#include "Engine/Texture2D.h"
 #include "HAL/IConsoleManager.h"
 #include "UObject/UObjectIterator.h"
 #include "VT/RuntimeVirtualTexture.h"
