Fix various issues with stream caching in editor. Including:
1) If you play a sound while stream caching is disabled, then enable stream caching and play it, the sound won't be loaded into the cache because it's already precached with DTYPE_Native or DTYPE_Realtime. This was fixed by invalidating and re-cooking the asset on playback if the stream caching settings were changed since last time we cooked it.
2) Fix for issue where the Force Streaming flag would be set to true for any asset you opened while stream caching is enabled in the editor.
3) Fix for issue where newly imported sounds would not play if stream caching is enabled until you restart the editor.
4) Fix for issue where, if you toggle streaming on an asset, the both streaming and non-streaming versions of the asset would stick around in memory.

diff --git a/Engine/Source/Editor/AudioEditor/Private/Factories/SoundFactory.cpp b/Engine/Source/Editor/AudioEditor/Private/Factories/SoundFactory.cpp
--- a/Engine/Source/Editor/AudioEditor/Private/Factories/SoundFactory.cpp
+++ b/Engine/Source/Editor/AudioEditor/Private/Factories/SoundFactory.cpp
@@ -20,6 +20,7 @@
 #include "Misc/MessageDialog.h"
 #include "Misc/FeedbackContext.h"
 #include "EditorFramework/AssetImportData.h"
+#include "AudioCompressionSettingsUtils.h"
 #include "SoundFileIO/SoundFileIO.h"
 
 
@@ -497,6 +498,12 @@ UObject* USoundFactory::CreateObject
 		Sound->NumChannels = ChannelCount;
 		Sound->TotalSamples = *WaveInfo.pSamplesPerSec * Sound->Duration;
 
+		// If stream caching is enabled, we need to make sure this asset is ready for playback.
+		if (FPlatformCompressionUtilities::IsCurrentPlatformUsingStreamCaching() && Sound->IsStreaming())
+		{
+			Sound->EnsureZerothChunkIsLoaded();
+		}
+
 		GEditor->GetEditorSubsystem<UImportSubsystem>()->BroadcastAssetPostImport(this, Sound);
 
 		if (ExistingSound && bUseExistingSettings)
diff --git a/Engine/Source/Runtime/AudioMixer/Private/AudioMixerBuffer.cpp b/Engine/Source/Runtime/AudioMixer/Private/AudioMixerBuffer.cpp
--- a/Engine/Source/Runtime/AudioMixer/Private/AudioMixerBuffer.cpp
+++ b/Engine/Source/Runtime/AudioMixer/Private/AudioMixerBuffer.cpp
@@ -177,6 +177,10 @@ namespace Audio
 			return nullptr;
 		}
 
+#if WITH_EDITOR
+		InWave->InvalidateSoundWaveIfNeccessary();
+#endif // WITH_EDITOR
+
 		FAudioDeviceManager* AudioDeviceManager = FAudioDevice::GetAudioDeviceManager();
 
 		FMixerBuffer* Buffer = nullptr;
diff --git a/Engine/Source/Runtime/Engine/Classes/Sound/SoundWave.h b/Engine/Source/Runtime/Engine/Classes/Sound/SoundWave.h
--- a/Engine/Source/Runtime/Engine/Classes/Sound/SoundWave.h
+++ b/Engine/Source/Runtime/Engine/Classes/Sound/SoundWave.h
@@ -100,7 +100,7 @@ struct FStreamedAudioPlatformData
 	FString DerivedDataKey;
 	/** Async cache task if one is outstanding. */
 	struct FStreamedAudioAsyncCacheDerivedDataTask* AsyncTask;
-#endif
+#endif // WITH_EDITORONLY_DATA
 
 	/** Default constructor. */
 	FStreamedAudioPlatformData();
@@ -128,7 +128,7 @@ struct FStreamedAudioPlatformData
 	bool IsFinishedCache() const;
 	ENGINE_API bool TryInlineChunkData();
 	bool AreDerivedChunksAvailable() const;
-#endif
+#endif // WITH_EDITORONLY_DATA
 
 private:
 
@@ -368,7 +368,10 @@ private:
 
 	// This is set when SetSampleRate is called to invalidate our cached sample rate while not re-parsing project settings.
 	uint8 bSampleRateManuallyReset:1;
-#endif
+#else
+	// Whether this was previously cooked with stream caching enabled.
+	uint8 bWasStreamCachingEnabledOnLastCook:1;
+#endif // !WITH_EDITOR
 
 	enum class ESoundWaveResourceState : uint8
 	{
@@ -428,7 +431,7 @@ public:
 	/** The release time in milliseconds. Describes how quickly the envelope analyzer responds to decreasing amplitudes. */
 	UPROPERTY(EditAnywhere, Category = "Analysis|Envelope", meta = (EditCondition = "bEnableAmplitudeEnvelopeAnalysis", ClampMin = "0", UIMin = "0"))
 	int32 EnvelopeFollowerReleaseTime;
-#endif
+#endif // WITH_EDITORONLY_DATA
 
 	/** The frequencies (in hz) to analyze when doing baked FFT analysis. */
 	UPROPERTY(EditAnywhere, Category = "Analysis|FFT", meta = (EditCondition = "bEnableBakedFFTAnalysis"))
@@ -597,7 +600,7 @@ public:
 	/** FByteBulkData doesn't currently support readonly access from multiple threads, so we limit access to RawData with a critical section on cook. */
 	FCriticalSection RawDataCriticalSection;
 
-#endif
+#endif // WITH_EDITORONLY_DATA
 
 	/** The streaming derived data for this sound on this platform. */
 	FStreamedAudioPlatformData* RunningPlatformData;
@@ -675,7 +678,7 @@ public:
 		// Ensure that we invalidate our cached sample rate if the UProperty sample rate is changed.
 		bCachedSampleRateFromPlatformSettings = false;
 		bSampleRateManuallyReset = true;
-#endif
+#endif //WITH_EDITOR
 	}
 
 	/**
@@ -686,9 +689,11 @@ public:
 	virtual int32 GetResourceSizeForFormat(FName Format);
 
 	/**
-	 * Frees up all the resources allocated in this class
+	 * Frees up all the resources allocated in this class.
+	 * @param bStopSoundsUsingThisResource if false, will leave any playing audio alive.
+	 *        This occurs when we force a re-cook of audio while starting to play a sound.
 	 */
-	void FreeResources();
+	void FreeResources(bool bStopSoundsUsingThisResource = true);
 
 	/** Will clean up the decompressor task if the task has finished or force it finish. Returns true if the decompressor is cleaned up. */
 	bool CleanupDecompressor(bool bForceCleanup = false);
@@ -755,7 +760,12 @@ public:
 #if WITH_EDITOR
 	/** Utility which returns imported PCM data and the parsed header for the file. Returns true if there was data, false if there wasn't. */
 	bool GetImportedSoundWaveData(TArray<uint8>& OutRawPCMData, uint32& OutSampleRate, uint16& OutNumChannels);
-#endif
+
+	/**
+	 * This function can be called before playing or using a SoundWave to check if any cook settings have been modified since this SoundWave was last cooked.
+	 */
+	void InvalidateSoundWaveIfNeccessary();
+#endif //WITH_EDITOR
 
 private:
 
@@ -764,13 +774,13 @@ private:
 #if WITH_EDITOR
 	void BakeFFTAnalysis();
 	void BakeEnvelopeAnalysis();
-#endif
+#endif //WITH_EDITOR
 
 public:
 
 #if WITH_EDITOR
 	void LogBakedData();
-#endif
+#endif //WITH_EDITOR
 
 	virtual void BeginGetCompressedData(FName Format, const FPlatformAudioCookOverrides* CompressionOverrides);
 
@@ -810,7 +820,7 @@ public:
 
 	/** Gets the member name for the Curves property of the USoundWave object. */
 	static FName GetCurvePropertyName() { return GET_MEMBER_NAME_CHECKED(USoundWave, Curves); }
-#endif
+#endif // WITH_EDITOR
 
 	/** Checks whether sound has been categorised as streaming. */
 	bool IsStreaming(const FPlatformAudioCookOverrides* Overrides = nullptr) const;
@@ -880,7 +890,7 @@ public:
 	virtual void WillNeverCacheCookedPlatformDataAgain() override;
 
 	uint32 bNeedsThumbnailGeneration:1;
-#endif
+#endif // WITH_EDITOR
 
 	/**
 	 * Caches platform data for the sound.
@@ -901,7 +911,7 @@ public:
 	 * Forces platform data to be rebuilt.
 	 */
 	void ForceRebuildPlatformData();
-#endif
+#endif // WITH_EDITORONLY_DATA
 
 	/**
 	 * Get Chunk data for a specified chunk index.
diff --git a/Engine/Source/Runtime/Engine/Private/SoundWave.cpp b/Engine/Source/Runtime/Engine/Private/SoundWave.cpp
--- a/Engine/Source/Runtime/Engine/Private/SoundWave.cpp
+++ b/Engine/Source/Runtime/Engine/Private/SoundWave.cpp
@@ -206,6 +206,9 @@ USoundWave::USoundWave(const FObjectInitializer& ObjectInitializer)
 	bCachedSampleRateFromPlatformSettings = false;
 	bSampleRateManuallyReset = false;
 	CachedSampleRateOverride = 0.0f;
+#else
+	bWasStreamCachingEnabledOnLastCook = FPlatformCompressionUtilities::IsCurrentPlatformUsingStreamCaching();
+	RunningPlatformData = nullptr;
 #endif //!WITH_EDITOR
 }
 
@@ -632,6 +635,9 @@ void USoundWave::BeginGetCompressedData(FName Format, const FPlatformAudioCookOv
 		return;
 	}
 
+	// If stream caching has been enabled or disabled since the previous DDC operation, we need to invalidate the current
+	InvalidateSoundWaveIfNeccessary();
+
 	FName PlatformSpecificFormat = GetPlatformSpecificFormat(Format, CompressionOverrides);
 
 	if (!CompressedFormatData.Contains(PlatformSpecificFormat) && !AsyncLoadingDataFormats.Contains(PlatformSpecificFormat))
@@ -707,7 +713,9 @@ FByteBulkData* USoundWave::GetCompressedData(FName Format, const FPlatformAudioC
 void USoundWave::InvalidateCompressedData()
 {
 	CompressedDataGuid = FGuid::NewGuid();
+	ZerothChunkData.Reset();
 	CompressedFormatData.FlushData();
+	FreeResources(false);
 }
 
 void USoundWave::PostLoad()
@@ -735,8 +743,10 @@ void USoundWave::PostLoad()
 		return;
 	}
 
+#if !WITH_EDITOR
 	// In case any code accesses bStreaming directly, we fix up bStreaming based on the current platform's cook overrides.
 	bStreaming = IsStreaming();
+#endif
 
 	// Compress to whatever formats the active target platforms want
 	// static here as an optimization
@@ -970,8 +980,32 @@ void USoundWave::RemoveAudioResource()
 #endif
 }
 
+
+
 #if WITH_EDITOR
 
+void USoundWave::InvalidateSoundWaveIfNeccessary()
+{
+	if (bProcedural)
+	{
+		return;
+	}
+
+	// if stream caching was enabled since the last time we invalidated the compressed audio, force a re-cook.
+	const bool bIsStreamCachingEnabled = FPlatformCompressionUtilities::IsCurrentPlatformUsingStreamCaching();
+	if (bWasStreamCachingEnabledOnLastCook != bIsStreamCachingEnabled)
+	{
+		InvalidateCompressedData();
+		bWasStreamCachingEnabledOnLastCook = bIsStreamCachingEnabled;
+		
+		// If stream caching is now turned on, recook the streaming audio if neccessary.
+		if(bIsStreamCachingEnabled && IsStreaming())
+		{
+			EnsureZerothChunkIsLoaded();
+		}
+	}
+}
+
 float USoundWave::GetSampleRateForTargetPlatform(const ITargetPlatform* TargetPlatform)
 {
 	const FPlatformAudioCookOverrides* Overrides = TargetPlatform->GetAudioCompressionSettings();
@@ -1404,7 +1438,7 @@ void USoundWave::PostEditChangeProperty(FPropertyChangedEvent& PropertyChangedEv
 }
 #endif // WITH_EDITOR
 
-void USoundWave::FreeResources()
+void USoundWave::FreeResources(bool bStopSoundsUsingThisResource)
 {
 	check(IsInAudioThread());
 
@@ -1418,7 +1452,7 @@ void USoundWave::FreeResources()
 	{
 		// Notify the audio device to free the bulk data associated with this wave.
 		FAudioDeviceManager* AudioDeviceManager = GEngine->GetAudioDeviceManager();
-		if (AudioDeviceManager)
+		if (bStopSoundsUsingThisResource && AudioDeviceManager)
 		{
 			AudioDeviceManager->StopSoundsUsingResource(this);
 			AudioDeviceManager->FreeResource(this);
diff --git a/Engine/Source/Runtime/Windows/XAudio2/Private/XAudio2Buffer.cpp b/Engine/Source/Runtime/Windows/XAudio2/Private/XAudio2Buffer.cpp
--- a/Engine/Source/Runtime/Windows/XAudio2/Private/XAudio2Buffer.cpp
+++ b/Engine/Source/Runtime/Windows/XAudio2/Private/XAudio2Buffer.cpp
@@ -526,6 +526,10 @@ FXAudio2SoundBuffer* FXAudio2SoundBuffer::Init( FAudioDevice* AudioDevice, USoun
 		return( NULL );
 	}
 
+#if WITH_EDITOR
+	Wave->InvalidateSoundWaveIfNeccessary();
+#endif // WITH_EDITOR
+
 	FAudioDeviceManager* AudioDeviceManager = GEngine->GetAudioDeviceManager();
 
 	FXAudio2Device* XAudio2Device = ( FXAudio2Device* )AudioDevice;
