Reuse an in memory version of the PIE level if one exists rather than duplicating from the editor. This is consistent with non-PIE behavior.

diff --git a/Engine/Source/Runtime/Engine/Private/LevelStreaming.cpp b/Engine/Source/Runtime/Engine/Private/LevelStreaming.cpp
--- a/Engine/Source/Runtime/Engine/Private/LevelStreaming.cpp
+++ b/Engine/Source/Runtime/Engine/Private/LevelStreaming.cpp
@@ -898,9 +898,12 @@ bool ULevelStreaming::RequestLevel(UWorld* PersistentWorld, bool bAllowLevelLoad
 	EPackageFlags PackageFlags = PKG_ContainsMap;
 	int32 PIEInstanceID = INDEX_NONE;
 
+	// Try to find the [to be] loaded package.
+	UPackage* LevelPackage = (UPackage*)StaticFindObjectFast(UPackage::StaticClass(), nullptr, DesiredPackageName, 0, 0, RF_NoFlags, EInternalObjectFlags::PendingKill);
+
 	// copy streaming level on demand if we are in PIE
 	// (the world is already loaded for the editor, just find it and copy it)
-	if ( PersistentWorld->IsPlayInEditor() )
+	if ( LevelPackage == nullptr && PersistentWorld->IsPlayInEditor() )
 	{
 		if (PersistentWorld->GetOutermost()->HasAnyPackageFlags(PKG_PlayInEditor))
 		{
@@ -943,9 +946,6 @@ bool ULevelStreaming::RequestLevel(UWorld* PersistentWorld, bool bAllowLevelLoad
 		}
 	}
 
-	// Try to find the [to be] loaded package.
-	UPackage* LevelPackage = (UPackage*)StaticFindObjectFast(UPackage::StaticClass(), nullptr, DesiredPackageName, 0, 0, RF_NoFlags, EInternalObjectFlags::PendingKill);
-
 	// Package is already or still loaded.
 	if (LevelPackage)
 	{
