[NO NON-PII COMMIT MESSAGE]

diff --git a/Engine/Source/Runtime/Engine/Classes/Sound/SoundWave.h b/Engine/Source/Runtime/Engine/Classes/Sound/SoundWave.h
--- a/Engine/Source/Runtime/Engine/Classes/Sound/SoundWave.h
+++ b/Engine/Source/Runtime/Engine/Classes/Sound/SoundWave.h
@@ -797,8 +797,9 @@ public:
 
 	/**
 	 * Change the guid and flush all compressed data
+	 * @param bFreeResources if true, will delete any precached compressed data as well.
 	 */
-	void InvalidateCompressedData();
+	void InvalidateCompressedData(bool bFreeResources = false);
 
 	/** Returns curves associated with this sound wave */
 	virtual class UCurveTable* GetCurveData() const override { return Curves; }
diff --git a/Engine/Source/Runtime/Engine/Private/SoundWave.cpp b/Engine/Source/Runtime/Engine/Private/SoundWave.cpp
--- a/Engine/Source/Runtime/Engine/Private/SoundWave.cpp
+++ b/Engine/Source/Runtime/Engine/Private/SoundWave.cpp
@@ -710,12 +710,16 @@ FByteBulkData* USoundWave::GetCompressedData(FName Format, const FPlatformAudioC
 	return Result->GetBulkDataSize() > 0 ? Result : NULL; // we don't return empty bulk data...but we save it to avoid thrashing the DDC
 }
 
-void USoundWave::InvalidateCompressedData()
+void USoundWave::InvalidateCompressedData(bool bFreeResources)
 {
 	CompressedDataGuid = FGuid::NewGuid();
 	ZerothChunkData.Reset();
 	CompressedFormatData.FlushData();
-	FreeResources(false);
+
+	if (bFreeResources)
+	{
+		FreeResources(false);
+	}
 }
 
 void USoundWave::PostLoad()
@@ -995,7 +999,7 @@ void USoundWave::InvalidateSoundWaveIfNeccessary()
 	const bool bIsStreamCachingEnabled = FPlatformCompressionUtilities::IsCurrentPlatformUsingStreamCaching();
 	if (bWasStreamCachingEnabledOnLastCook != bIsStreamCachingEnabled)
 	{
-		InvalidateCompressedData();
+		InvalidateCompressedData(true);
 		bWasStreamCachingEnabledOnLastCook = bIsStreamCachingEnabled;
 		
 		// If stream caching is now turned on, recook the streaming audio if neccessary.
