Minor fix up- fix cook commandlet building an invalid DDC hash when the platform the DDC hash is being generated on has different streaming settings than the target platform you're cooking to.

diff --git a/Engine/Source/Runtime/Engine/Private/AudioDerivedData.cpp b/Engine/Source/Runtime/Engine/Private/AudioDerivedData.cpp
--- a/Engine/Source/Runtime/Engine/Private/AudioDerivedData.cpp
+++ b/Engine/Source/Runtime/Engine/Private/AudioDerivedData.cpp
@@ -343,7 +343,8 @@ class FStreamedAudioCacheDerivedDataWorker : public FNonAbandonableTask
 				}
 				else if (bUseStreamCaching)
 				{
-					checkf(MinimumChunkSize != 0, TEXT("To use Load On Demand, please override GetMinimumSizeForInitialChunk"));
+					// Ensure that the minimum chunk size is nonzero if our compressed buffer is not empty.
+					checkf(CompressedBuffer.Num() == 0 || MinimumChunkSize != 0, TEXT("To use Load On Demand, please override GetMinimumSizeForInitialChunk"));
 
 					// Otherwise if we're using Load On Demand, the first chunk should be as small as possible:
 					FirstChunkSize = MinimumChunkSize;
@@ -1635,7 +1636,7 @@ void USoundWave::ForceRebuildPlatformData()
 {
 	if (RunningPlatformData)
 	{
-		const FPlatformAudioCookOverrides* CompressionOverrides = GetPlatformCompressionOverridesForCurrentPlatform();
+		const FPlatformAudioCookOverrides* CompressionOverrides = GetCookOverridesForRunningPlatform();
 		RunningPlatformData->Cache(
 			*this,
 			CompressionOverrides,
