Fix for server when stream caching is enabled: Do not call EnsureZerothChunkIsLoaded nor attempt to prime the first chunk if this is a server build and hence has no audio data.

diff --git a/Engine/Source/Runtime/Engine/Classes/Sound/SoundWave.h b/Engine/Source/Runtime/Engine/Classes/Sound/SoundWave.h
--- a/Engine/Source/Runtime/Engine/Classes/Sound/SoundWave.h
+++ b/Engine/Source/Runtime/Engine/Classes/Sound/SoundWave.h
@@ -804,6 +804,9 @@ public:
 	/** Returns curves associated with this sound wave */
 	virtual class UCurveTable* GetCurveData() const override { return Curves; }
 
+	// This function returns true if there are streamable chunks in this asset.
+	bool HasStreamingChunks();
+
 #if WITH_EDITOR
 	/** These functions are required for support for some custom details/editor functionality.*/
 
diff --git a/Engine/Source/Runtime/Engine/Private/SoundWave.cpp b/Engine/Source/Runtime/Engine/Private/SoundWave.cpp
--- a/Engine/Source/Runtime/Engine/Private/SoundWave.cpp
+++ b/Engine/Source/Runtime/Engine/Private/SoundWave.cpp
@@ -722,6 +722,11 @@ void USoundWave::InvalidateCompressedData(bool bFreeResources)
 	}
 }
 
+bool USoundWave::HasStreamingChunks()
+{
+	return RunningPlatformData != NULL && RunningPlatformData->Chunks.Num() > 0;
+}
+
 void USoundWave::PostLoad()
 {
 	LLM_SCOPE(ELLMTag::AudioSoundWaves);
@@ -793,14 +798,16 @@ void USoundWave::PostLoad()
 		IStreamingManager::Get().GetAudioStreamingManager().AddStreamingSoundWave(this);
 	}
 
-	if (ShouldUseStreamCaching() && IsStreaming())
+	const bool bDoesSoundWaveHaveStreamingAudioData = HasStreamingChunks();
+
+	if (ShouldUseStreamCaching() && IsStreaming() && bDoesSoundWaveHaveStreamingAudioData)
 	{
 		EnsureZerothChunkIsLoaded();
 	}
 
 	const bool bShouldPrimeSound = LoadIntoCacheOnPostLoadCVar || bLoadCompressedAudioWhenSoundWaveIsLoaded;
 
-	if (bShouldPrimeSound && FPlatformCompressionUtilities::IsCurrentPlatformUsingStreamCaching())
+	if (bShouldPrimeSound && FPlatformCompressionUtilities::IsCurrentPlatformUsingStreamCaching() && bDoesSoundWaveHaveStreamingAudioData)
 	{
 		// Load rest of the audio into cache.
 		IStreamingManager::Get().GetAudioStreamingManager().RequestChunk(this, 1, [](EAudioChunkLoadResult) {});
