Chaos
- fix multi-body inertia calculation
- better shape debug draw
- debug option to disable constraints with linear limits

diff --git a/Engine/Source/Runtime/Engine/Private/PhysicsEngine/ImmediatePhysics/ImmediatePhysicsChaos/ImmediatePhysicsActorHandle_Chaos.cpp b/Engine/Source/Runtime/Engine/Private/PhysicsEngine/ImmediatePhysics/ImmediatePhysicsChaos/ImmediatePhysicsActorHandle_Chaos.cpp
--- a/Engine/Source/Runtime/Engine/Private/PhysicsEngine/ImmediatePhysics/ImmediatePhysicsChaos/ImmediatePhysicsActorHandle_Chaos.cpp
+++ b/Engine/Source/Runtime/Engine/Private/PhysicsEngine/ImmediatePhysics/ImmediatePhysicsChaos/ImmediatePhysicsActorHandle_Chaos.cpp
@@ -24,7 +24,7 @@ namespace ImmediatePhysics_Chaos
 	//
 
 	template<typename T, int d>
-	Chaos::PMatrix<T, d, d> CalculateInertia_Solid(T Mass, const FKSphereElem& SphereElem)
+	Chaos::PMatrix<T, d, d> CalculateInertia_Solid(const T Mass, const FKSphereElem& SphereElem)
 	{
 		return Chaos::PMatrix<T, d, d>(
 			((T)2 / (T)5) * Mass * SphereElem.Radius * SphereElem.Radius,
@@ -34,7 +34,7 @@ namespace ImmediatePhysics_Chaos
 	}
 
 	template<typename T, int d>
-	Chaos::PMatrix<T, d, d> CalculateInertia_Solid(T Mass, const FKSphylElem& SphylElem)
+	Chaos::PMatrix<T, d, d> CalculateInertia_Solid(const T Mass, const FKSphylElem& SphylElem)
 	{
 		return Chaos::PMatrix<T, d, d>(
 			((T)1 / (T)12) * Mass * ((T)3 * SphylElem.Radius + SphylElem.Length),
@@ -44,7 +44,7 @@ namespace ImmediatePhysics_Chaos
 	}
 
 	template<typename T, int d>
-	Chaos::PMatrix<T, d, d> CalculateInertia_Solid(T Mass, const FKBoxElem& BoxElem)
+	Chaos::PMatrix<T, d, d> CalculateInertia_Solid(const T Mass, const FKBoxElem& BoxElem)
 	{
 		return Chaos::PMatrix<T, d, d>(
 			((T)1 / (T)12) * Mass * (BoxElem.Y * BoxElem.Y + BoxElem.Z * BoxElem.Z),
@@ -54,43 +54,34 @@ namespace ImmediatePhysics_Chaos
 	}
 
 	template<typename T, int d>
-	void CalculateMassProperties(T Density, const FVector& Scale, const FTransform& LocalTransform, const FKAggregateGeom& AggGeom, T& OutMass, Chaos::TMassProperties<T, d>& OutMassProperties)
+	void CalculateMassProperties(const FVector& Scale, const FTransform& LocalTransform, const FKAggregateGeom& AggGeom, Chaos::TMassProperties<T, d>& OutMassProperties)
 	{
 		using namespace Chaos;
 		TArray<TMassProperties<T, d>> AllMassProperties;
-		T AllMass = 0;
 
 		for (uint32 i = 0; i < static_cast<uint32>(AggGeom.SphereElems.Num()); ++i)
 		{
 			const FKSphereElem ScaledSphereElem = AggGeom.SphereElems[i].GetFinalScaled(Scale, LocalTransform);
 
-			T Volume = ScaledSphereElem.GetVolume(FVector::OneVector);
-			T Mass = Density * Volume;
-
 			TMassProperties<T, d> MassProperties;
 			MassProperties.CenterOfMass = LocalTransform.GetTranslation() + ScaledSphereElem.Center;
 			MassProperties.RotationOfMass = TRotation<T, d>::FromIdentity();
-			MassProperties.Volume = Volume;
-			MassProperties.InertiaTensor = CalculateInertia_Solid<T, d>(Mass, ScaledSphereElem);
+			MassProperties.Volume = ScaledSphereElem.GetVolume(FVector::OneVector);
+			MassProperties.InertiaTensor = CalculateInertia_Solid<T, d>(MassProperties.Volume, ScaledSphereElem);
 
 			AllMassProperties.Add(MassProperties);
-			AllMass += Mass;
 		}
 		for (uint32 i = 0; i < static_cast<uint32>(AggGeom.BoxElems.Num()); ++i)
 		{
 			const auto& BoxElem = AggGeom.BoxElems[i];
 
-			T Volume = BoxElem.GetVolume(Scale);
-			T Mass = Density * Volume;
-
 			TMassProperties<T, d> MassProperties;
 			MassProperties.CenterOfMass = LocalTransform.GetTranslation() + BoxElem.Center;
 			MassProperties.RotationOfMass = LocalTransform.GetRotation() * TRotation<T, d>(FQuat(BoxElem.Rotation));
-			MassProperties.Volume = Volume;
-			MassProperties.InertiaTensor = CalculateInertia_Solid<T, d>(Mass, BoxElem);
+			MassProperties.Volume = BoxElem.GetVolume(Scale);
+			MassProperties.InertiaTensor = CalculateInertia_Solid<T, d>(MassProperties.Volume, BoxElem);
 
 			AllMassProperties.Add(MassProperties);
-			AllMass += Mass;
 		}
 		for (uint32 i = 0; i < static_cast<uint32>(AggGeom.SphylElems.Num()); ++i)
 		{
@@ -103,31 +94,23 @@ namespace ImmediatePhysics_Chaos
 				//not a capsule just use a sphere
 				const FKSphereElem ScaledSphereElem = FKSphereElem(Radius);
 
-				T Volume = ScaledSphereElem.GetVolume(FVector::OneVector);
-				T Mass = Density * Volume;
-
 				TMassProperties<T, d> MassProperties;
 				MassProperties.CenterOfMass = LocalTransform.GetTranslation() + ScaledSphereElem.Center;
 				MassProperties.RotationOfMass = TRotation<T, d>::FromIdentity();
-				MassProperties.Volume = Volume;
-				MassProperties.InertiaTensor = CalculateInertia_Solid<T, d>(Mass, ScaledSphereElem);
+				MassProperties.Volume = ScaledSphereElem.GetVolume(FVector::OneVector);
+				MassProperties.InertiaTensor = CalculateInertia_Solid<T, d>(MassProperties.Volume, ScaledSphereElem);
 
 				AllMassProperties.Add(MassProperties);
-				AllMass += Mass;
 			}
 			else
 			{
-				T Volume = ScaledSphylElem.GetVolume(FVector::OneVector);
-				T Mass = Density * Volume;
-
 				TMassProperties<T, d> MassProperties;
 				MassProperties.CenterOfMass = LocalTransform.GetTranslation() + ScaledSphylElem.Center;
 				MassProperties.RotationOfMass = LocalTransform.GetRotation() * TRotation<T, d>(FQuat(ScaledSphylElem.Rotation));
-				MassProperties.Volume = Volume;
-				MassProperties.InertiaTensor = CalculateInertia_Solid<T, d>(Mass, ScaledSphylElem);
+				MassProperties.Volume = ScaledSphylElem.GetVolume(FVector::OneVector);
+				MassProperties.InertiaTensor = CalculateInertia_Solid<T, d>(MassProperties.Volume, ScaledSphylElem);
 
 				AllMassProperties.Add(MassProperties);
-				AllMass += Mass;
 			}
 		}
 #if WITH_CHAOS && !PHYSICS_INTERFACE_PHYSX
@@ -152,7 +135,6 @@ namespace ImmediatePhysics_Chaos
 		}
 
 		OutMassProperties = Combine(AllMassProperties);
-		OutMass = AllMass;
 	}
 
 
@@ -163,8 +145,9 @@ namespace ImmediatePhysics_Chaos
 #if WITH_CHAOS && !PHYSICS_INTERFACE_PHYSX
 		float Density = 1.e-3f;	// 1g/cm3
 		Chaos::TMassProperties<float, 3> MassProperties;
-		CalculateMassProperties<float, 3>(Density, Scale, FTransform::Identity, BodySetup->AggGeom, OutMass, MassProperties);
-		OutInertia = Chaos::TVector<float, 3>(MassProperties.InertiaTensor.M[0][0], MassProperties.InertiaTensor.M[1][1], MassProperties.InertiaTensor.M[2][2]);
+		CalculateMassProperties<float, 3>(Scale, FTransform::Identity, BodySetup->AggGeom, MassProperties);
+		OutMass = Density * MassProperties.Volume;
+		OutInertia = Density * Chaos::TVector<float, 3>(MassProperties.InertiaTensor.M[0][0], MassProperties.InertiaTensor.M[1][1], MassProperties.InertiaTensor.M[2][2]);
 		OutCoMTransform = FTransform(MassProperties.RotationOfMass, MassProperties.CenterOfMass);
 #else
 		OutMass = BodyInstance->GetBodyMass();
diff --git a/Engine/Source/Runtime/Engine/Private/PhysicsEngine/ImmediatePhysics/ImmediatePhysicsChaos/ImmediatePhysicsSimulation_Chaos.cpp b/Engine/Source/Runtime/Engine/Private/PhysicsEngine/ImmediatePhysics/ImmediatePhysicsChaos/ImmediatePhysicsSimulation_Chaos.cpp
--- a/Engine/Source/Runtime/Engine/Private/PhysicsEngine/ImmediatePhysics/ImmediatePhysicsChaos/ImmediatePhysicsSimulation_Chaos.cpp
+++ b/Engine/Source/Runtime/Engine/Private/PhysicsEngine/ImmediatePhysics/ImmediatePhysicsChaos/ImmediatePhysicsSimulation_Chaos.cpp
@@ -35,6 +35,7 @@ FAutoConsoleVariableRef CVarCollisionPushOutPairIterations(TEXT("p.Chaos.ImmPhys
 FAutoConsoleVariableRef CVarChaosImmPhysThickness(TEXT("p.Chaos.ImmPhys.CollisionThickness"), ChaosImmediate_Collision_Thickness, TEXT("ChaosImmediateThickness"));
 
 float ChaosImmediate_Joint_SwingTwistAngleTolerance = 1.0e-6f;
+int32 ChaosImmediate_Joint_EnableLinearLimits = 1;
 int32 ChaosImmediate_Joint_EnableTwistLimits = 1;
 int32 ChaosImmediate_Joint_EnableSwingLimits = 1;
 int32 ChaosImmediate_Joint_EnableDrives = 1;
@@ -44,6 +45,7 @@ float ChaosImmediate_Joint_PBDDriveStiffness = 0.0f;
 float ChaosImmediate_Joint_PBDMinParentMassRatio = 0.5f;
 float ChaosImmediate_Joint_PBDMaxInertiaRatio = 5.0f;
 FAutoConsoleVariableRef CVarSwingTwistAngleTolerance(TEXT("p.Chaos.ImmPhys.Joint.SwingTwistAngleTolerance"), ChaosImmediate_Joint_SwingTwistAngleTolerance, TEXT("SwingTwistAngleTolerance."));
+FAutoConsoleVariableRef CVarEnableLinearLimits(TEXT("p.Chaos.ImmPhys.Joint.EnableLinearLimits"), ChaosImmediate_Joint_EnableLinearLimits, TEXT("EnableLinearLimits."));
 FAutoConsoleVariableRef CVarEnableTwistLimits(TEXT("p.Chaos.ImmPhys.Joint.EnableTwistLimits"), ChaosImmediate_Joint_EnableTwistLimits, TEXT("EnableTwistLimits."));
 FAutoConsoleVariableRef CVarEnableSwingLimits(TEXT("p.Chaos.ImmPhys.Joint.EnableSwingLimits"), ChaosImmediate_Joint_EnableSwingLimits, TEXT("EnableSwingLimits."));
 FAutoConsoleVariableRef CVarEnableDrives(TEXT("p.Chaos.ImmPhys.Joint.EnableDrives"), ChaosImmediate_Joint_EnableDrives, TEXT("EnableDrives."));
@@ -388,9 +390,10 @@ namespace ImmediatePhysics_Chaos
 			JointsSettings.PBDMaxInertiaRatio = ChaosImmediate_Joint_PBDMaxInertiaRatio;
 			JointsSettings.FreezeIterations = ChaosImmediate_Joint_FreezeIterations;
 			JointsSettings.FrozenIterations = ChaosImmediate_Joint_FrozenIterations;
-			JointsSettings.bEnableTwistLimits = ChaosImmediate_Joint_EnableTwistLimits > 0;
-			JointsSettings.bEnableSwingLimits = ChaosImmediate_Joint_EnableSwingLimits > 0;
-			JointsSettings.bEnableDrives = ChaosImmediate_Joint_EnableDrives > 0;
+			JointsSettings.bEnableLinearLimits = ChaosImmediate_Joint_EnableLinearLimits != 0;
+			JointsSettings.bEnableTwistLimits = ChaosImmediate_Joint_EnableTwistLimits != 0;
+			JointsSettings.bEnableSwingLimits = ChaosImmediate_Joint_EnableSwingLimits != 0;
+			JointsSettings.bEnableDrives = ChaosImmediate_Joint_EnableDrives != 0;
 			JointsSettings.PBDDriveStiffness = ChaosImmediate_Joint_PBDDriveStiffness;
 			Joints->SetSettings(JointsSettings);
 
diff --git a/Engine/Source/Runtime/Experimental/Chaos/Private/Chaos/ChaosDebugDraw.cpp b/Engine/Source/Runtime/Experimental/Chaos/Private/Chaos/ChaosDebugDraw.cpp
--- a/Engine/Source/Runtime/Experimental/Chaos/Private/Chaos/ChaosDebugDraw.cpp
+++ b/Engine/Source/Runtime/Experimental/Chaos/Private/Chaos/ChaosDebugDraw.cpp
@@ -114,7 +114,8 @@ namespace Chaos
 		template <typename T, int d>
 		void DrawParticleShapesImpl(const TRigidTransform<float, 3>& SpaceTransform, const TGeometryParticleHandle<T, d>* Particle, T ColorScale)
 		{
-			FColor Color = (ColorScale * FColor::White).ToFColor(false);
+			FColor ShapeColor = Particle->AsDynamic() ? FColor::Yellow : FColor::Orange;
+			FColor Color = (ColorScale * ShapeColor).ToFColor(false);
 			TVector<T, d> P = SpaceTransform.TransformPosition(Particle->AsDynamic() ? Particle->AsDynamic()->P() : Particle->X());
 			TRotation<T, d> Q = SpaceTransform.GetRotation() * (Particle->AsDynamic() ? Particle->AsDynamic()->Q() : Particle->R());
 
diff --git a/Engine/Source/Runtime/Experimental/Chaos/Private/Chaos/PBDJointConstraints.cpp b/Engine/Source/Runtime/Experimental/Chaos/Private/Chaos/PBDJointConstraints.cpp
--- a/Engine/Source/Runtime/Experimental/Chaos/Private/Chaos/PBDJointConstraints.cpp
+++ b/Engine/Source/Runtime/Experimental/Chaos/Private/Chaos/PBDJointConstraints.cpp
@@ -112,6 +112,7 @@ namespace Chaos
 		, PBDMaxInertiaRatio(5.0f)
 		, FreezeIterations(0)
 		, FrozenIterations(0)
+		, bEnableLinearLimits(1)
 		, bEnableTwistLimits(1)
 		, bEnableSwingLimits(1)
 		, bEnableDrives(1)
@@ -393,10 +394,17 @@ namespace Chaos
 			}
 		}
 
+		const TVector<EJointMotionType, d>& LinearMotion = JointSettings.Motion.LinearMotionTypes;
 		EJointMotionType TwistMotion = JointSettings.Motion.AngularMotionTypes[(int32)EJointAngularConstraintIndex::Twist];
 		EJointMotionType Swing1Motion = JointSettings.Motion.AngularMotionTypes[(int32)EJointAngularConstraintIndex::Swing1];
 		EJointMotionType Swing2Motion = JointSettings.Motion.AngularMotionTypes[(int32)EJointAngularConstraintIndex::Swing2];
 
+		// Disable a constraint if it has any linear limits?
+		if (!Settings.bEnableLinearLimits && ((LinearMotion[0] == EJointMotionType::Limited) || (LinearMotion[1] == EJointMotionType::Limited) || (LinearMotion[1] == EJointMotionType::Limited)))
+		{
+			return;
+		}
+
 		// Apply angular drives
 		if (Settings.bEnableDrives)
 		{
@@ -462,7 +470,6 @@ namespace Chaos
 		}
 
 		// Apply linear constraints
-		const TVector<EJointMotionType, d>& LinearMotion = JointSettings.Motion.LinearMotionTypes;
 		if ((LinearMotion[0] != EJointMotionType::Free) || (LinearMotion[1] != EJointMotionType::Free) || (LinearMotion[2] != EJointMotionType::Free))
 		{
 			TPBDJointUtilities<T, d>::ApplyJointPositionConstraint(Dt, Settings, JointSettings, Index0, Index1, P0, Q0, P1, Q1, InvM0, InvIL0, InvM1, InvIL1);
diff --git a/Engine/Source/Runtime/Experimental/Chaos/Public/Chaos/PBDJointConstraintTypes.h b/Engine/Source/Runtime/Experimental/Chaos/Public/Chaos/PBDJointConstraintTypes.h
--- a/Engine/Source/Runtime/Experimental/Chaos/Public/Chaos/PBDJointConstraintTypes.h
+++ b/Engine/Source/Runtime/Experimental/Chaos/Public/Chaos/PBDJointConstraintTypes.h
@@ -119,6 +119,7 @@ namespace Chaos
 		int32 FrozenIterations;
 
 		// @todo(ccaulfield): remove these TEMP overrides for testing
+		bool bEnableLinearLimits;
 		bool bEnableTwistLimits;
 		bool bEnableSwingLimits;
 		bool bEnableDrives;
