Fix LODGroup change. Avoid touching LOD 0 SectionInfoMap and remove any entry for the other LODs

diff --git a/Engine/Source/Runtime/Engine/Private/StaticMesh.cpp b/Engine/Source/Runtime/Engine/Private/StaticMesh.cpp
--- a/Engine/Source/Runtime/Engine/Private/StaticMesh.cpp
+++ b/Engine/Source/Runtime/Engine/Private/StaticMesh.cpp
@@ -3053,22 +3053,23 @@ void UStaticMesh::SetLODGroup(FName NewGroup, bool bRebuildImmediately)
 			// Set reduction settings to the defaults.
 			SourceModel.ReductionSettings = GroupSettings.GetDefaultSettings(LODIndex);
 			
-			//Reset the section info map
-			if (bResetSectionInfoMap)
+			if (LODIndex != 0)
 			{
-				for (int32 SectionIndex = 0; SectionIndex < GetSectionInfoMap().GetSectionNumber(LODIndex); ++SectionIndex)
+				//Reset the section info map
+				if (bResetSectionInfoMap)
 				{
-					FMeshSectionInfo Info;
-					Info.MaterialIndex = SectionIndex;
-					GetSectionInfoMap().Set(LODIndex, SectionIndex, Info);
+					for (int32 SectionIndex = 0; SectionIndex < GetSectionInfoMap().GetSectionNumber(LODIndex); ++SectionIndex)
+					{
+						GetSectionInfoMap().Remove(LODIndex, SectionIndex);
+					}
+				}
+				//Clear the raw data if we change the LOD Group and we do not reduce ourself, this will force the user to do a import LOD which will manage the section info map properly
+				if (!SourceModel.IsRawMeshEmpty() && SourceModel.ReductionSettings.BaseLODModel != LODIndex)
+				{
+					FRawMesh EmptyRawMesh;
+					SourceModel.SaveRawMesh(EmptyRawMesh);
+					SourceModel.SourceImportFilename = FString();
 				}
-			}
-			//Clear the raw data if we change the LOD Group and we do not reduce ourself, this will force the user to do a import LOD which will manage the section info map properly
-			if (!SourceModel.IsRawMeshEmpty() && SourceModel.ReductionSettings.BaseLODModel != LODIndex)
-			{
-				FRawMesh EmptyRawMesh;
-				SourceModel.SaveRawMesh(EmptyRawMesh);
-				SourceModel.SourceImportFilename = FString();
 			}
 		}
 		LightMapResolution = GroupSettings.GetDefaultLightMapResolution();
