[NO NON-PII COMMIT MESSAGE]

diff --git a/Engine/Source/Runtime/Engine/Private/SoundNode.cpp b/Engine/Source/Runtime/Engine/Private/SoundNode.cpp
--- a/Engine/Source/Runtime/Engine/Private/SoundNode.cpp
+++ b/Engine/Source/Runtime/Engine/Private/SoundNode.cpp
@@ -109,7 +109,7 @@ void USoundNode::PrimeChildWavePlayers(bool bRecurse)
 				if (WavePlayer != nullptr)
 				{
 					USoundWave* SoundWave = WavePlayer->GetSoundWave();
-					if (SoundWave && SoundWave->IsStreaming())
+					if (SoundWave && SoundWave->IsStreaming() && SoundWave->GetNumChunks() > 1)
 					{
 						IStreamingManager::Get().GetAudioStreamingManager().RequestChunk(SoundWave, 1, [](EAudioChunkLoadResult) {});
 					}
