[NO NON-PII COMMIT MESSAGE]

diff --git a/Engine/Source/Runtime/Engine/Private/AudioCompressionSettingsUtils.cpp b/Engine/Source/Runtime/Engine/Private/AudioCompressionSettingsUtils.cpp
--- a/Engine/Source/Runtime/Engine/Private/AudioCompressionSettingsUtils.cpp
+++ b/Engine/Source/Runtime/Engine/Private/AudioCompressionSettingsUtils.cpp
@@ -87,7 +87,7 @@ void CacheCurrentPlatformAudioCookOverrides(FPlatformAudioCookOverrides& OutOver
 #elif PLATFORM_PS4
 	const TCHAR* CategoryName = TEXT("/Script/PS4PlatformEditor.PS4TargetSettings");
 #elif PLATFORM_XBOXONE
-	const TCHAR* CategoryName = TEXT("/Script/XboxOneTargetPlatform.XboxOneTargetSettings");
+	const TCHAR* CategoryName = TEXT("/Script/XboxOnePlatformEditor.XboxOneTargetSettings");
 #else
 	const TCHAR* CategoryName = TEXT("");
 #endif
