Fix audio attenuation distance not respecting listener attenuation override for shapes other than spheres

diff --git a/Engine/Source/Runtime/Engine/Private/ActiveSound.cpp b/Engine/Source/Runtime/Engine/Private/ActiveSound.cpp
--- a/Engine/Source/Runtime/Engine/Private/ActiveSound.cpp
+++ b/Engine/Source/Runtime/Engine/Private/ActiveSound.cpp
@@ -1546,8 +1546,6 @@ void FActiveSound::ParseAttenuation(FSoundParseParameters& OutParseParams, const
 
 void FActiveSound::UpdateAttenuation(float DeltaTime, FSoundParseParameters& ParseParams, const FListener& Listener, const FSoundAttenuationSettings* SettingsAttenuationNode)
 {
-	const FVector& ListenerLocation = Listener.Transform.GetTranslation();
-
 	// Get the attenuation settings to use for this application to the active sound
 	const FSoundAttenuationSettings* Settings = SettingsAttenuationNode ? SettingsAttenuationNode : &AttenuationSettings;
 
@@ -1598,7 +1596,8 @@ void FActiveSound::UpdateAttenuation(float DeltaTime, FSoundParseParameters& Par
 		}
 		else
 		{
-			ParseParams.DistanceAttenuation *= Settings->Evaluate(ParseParams.Transform, ListenerLocation, FocusDataToApply.DistanceScale);
+			const FVector ListenerTranslation = ListenerData.ListenerTransform.GetTranslation();
+			ParseParams.DistanceAttenuation *= Settings->Evaluate(ParseParams.Transform, ListenerTranslation, FocusDataToApply.DistanceScale);
 		}
 	}
 
diff --git a/Engine/Source/Runtime/Engine/Private/AudioDevice.cpp b/Engine/Source/Runtime/Engine/Private/AudioDevice.cpp
--- a/Engine/Source/Runtime/Engine/Private/AudioDevice.cpp
+++ b/Engine/Source/Runtime/Engine/Private/AudioDevice.cpp
@@ -178,10 +178,14 @@ FAttenuationListenerData FAttenuationListenerData::Create(const FAudioDevice& Au
 	// Store the actual distance for surround-panning sources with spread (AudioMixer)
 	ListenerData.ListenerToSoundDistanceForPanning = ListenerData.ListenerToSoundDistance;
 
+	// Calculating override listener-to-sound distance and transform must
+	// be applied after distance used for panning value is calculated.
 	if (AudioDevice.IsUsingListenerAttenuationOverride())
 	{
 		const FVector& AttenuationOverride = AudioDevice.GetListenerAttenuationOverride();
+
 		ListenerData.ListenerToSoundDistance = (SoundTranslation - AttenuationOverride).Size();
+		ListenerData.ListenerTransform.SetTranslation(AttenuationOverride);
 	}
 
 	const FSoundAttenuationSettings& AttenuationSettings = *ListenerData.AttenuationSettings;
