fix on a crash with post graph disabled but trying to update curve from it.

diff --git a/Engine/Source/Runtime/Engine/Private/Components/SkeletalMeshComponent.cpp b/Engine/Source/Runtime/Engine/Private/Components/SkeletalMeshComponent.cpp
--- a/Engine/Source/Runtime/Engine/Private/Components/SkeletalMeshComponent.cpp
+++ b/Engine/Source/Runtime/Engine/Private/Components/SkeletalMeshComponent.cpp
@@ -2004,7 +2004,8 @@ void USkeletalMeshComponent::RefreshBoneTransforms(FActorComponentTickFunction*
 
 	PRAGMA_ENABLE_DEPRECATION_WARNINGS
 
-	TArray<uint16> const* CurrentAnimCurveUIDFinder = (AnimScriptInstance) ? &AnimScriptInstance->GetRequiredBones().GetUIDToArrayLookupTable() : nullptr;
+	TArray<uint16> const* CurrentAnimCurveUIDFinder = (AnimScriptInstance) ? &AnimScriptInstance->GetRequiredBones().GetUIDToArrayLookupTable() : 
+		((ShouldEvaluatePostProcessInstance())? &PostProcessAnimInstance->GetRequiredBones().GetUIDToArrayLookupTable() : nullptr);
 	const bool bAnimInstanceHasCurveUIDList = CurrentAnimCurveUIDFinder != nullptr;
 
 	const int32 CurrentCurveCount = (CurrentAnimCurveUIDFinder) ? FBlendedCurve::GetValidElementCount(CurrentAnimCurveUIDFinder) : 0;
@@ -2038,7 +2039,7 @@ void USkeletalMeshComponent::RefreshBoneTransforms(FActorComponentTickFunction*
 
 	AnimEvaluationContext.SkeletalMesh = SkeletalMesh;
 	AnimEvaluationContext.AnimInstance = AnimScriptInstance;
-	AnimEvaluationContext.PostProcessAnimInstance = PostProcessAnimInstance;
+	AnimEvaluationContext.PostProcessAnimInstance = (ShouldEvaluatePostProcessInstance())? PostProcessAnimInstance: nullptr;
 
 	if (CurrentAnimCurveUIDFinder)
 	{
