Fix static analysis warning
#rnx

diff --git a/Engine/Source/Runtime/Engine/Private/Components/SkeletalMeshComponent.cpp b/Engine/Source/Runtime/Engine/Private/Components/SkeletalMeshComponent.cpp
--- a/Engine/Source/Runtime/Engine/Private/Components/SkeletalMeshComponent.cpp
+++ b/Engine/Source/Runtime/Engine/Private/Components/SkeletalMeshComponent.cpp
@@ -2005,7 +2005,7 @@ void USkeletalMeshComponent::RefreshBoneTransforms(FActorComponentTickFunction*
 	PRAGMA_ENABLE_DEPRECATION_WARNINGS
 
 	TArray<uint16> const* CurrentAnimCurveUIDFinder = (AnimScriptInstance) ? &AnimScriptInstance->GetRequiredBones().GetUIDToArrayLookupTable() : 
-		((ShouldEvaluatePostProcessInstance())? &PostProcessAnimInstance->GetRequiredBones().GetUIDToArrayLookupTable() : nullptr);
+		((ShouldEvaluatePostProcessInstance() && PostProcessAnimInstance) ? &PostProcessAnimInstance->GetRequiredBones().GetUIDToArrayLookupTable() : nullptr);
 	const bool bAnimInstanceHasCurveUIDList = CurrentAnimCurveUIDFinder != nullptr;
 
 	const int32 CurrentCurveCount = (CurrentAnimCurveUIDFinder) ? FBlendedCurve::GetValidElementCount(CurrentAnimCurveUIDFinder) : 0;
