Fix PhysX includes.

diff --git a/Engine/Source/Runtime/Engine/Private/PhysicsEngine/Experimental/ChaosInterfaceUtils.cpp b/Engine/Source/Runtime/Engine/Private/PhysicsEngine/Experimental/ChaosInterfaceUtils.cpp
--- a/Engine/Source/Runtime/Engine/Private/PhysicsEngine/Experimental/ChaosInterfaceUtils.cpp
+++ b/Engine/Source/Runtime/Engine/Private/PhysicsEngine/Experimental/ChaosInterfaceUtils.cpp
@@ -22,9 +22,7 @@
 #include "PhysicsEngine/AggregateGeom.h"
 
 #if WITH_PHYSX
-#include "foundation/PxSimpleTypes.h"
-#include "geometry/PxConvexMesh.h"
-#include "geometry/PxTriangleMesh.h"
+#include "PhysXIncludes.h"
 #endif
 
 #define FORCE_ANALYTICS 0
