Reinstate materialinstance postload caching to fix HLOD customversion build errors

diff --git a/Engine/Source/Runtime/Engine/Private/Materials/MaterialInstance.cpp b/Engine/Source/Runtime/Engine/Private/Materials/MaterialInstance.cpp
--- a/Engine/Source/Runtime/Engine/Private/Materials/MaterialInstance.cpp
+++ b/Engine/Source/Runtime/Engine/Private/Materials/MaterialInstance.cpp
@@ -3581,7 +3581,7 @@ void UMaterialInstance::PostLoad()
 
 		// Make sure static parameters are up to date and shaders are cached for the current platform
 		InitStaticPermutation();
-#if WITH_EDITOR && 0
+#if WITH_EDITOR && 1
 		// enable caching in postload for derived data cache commandlet and cook by the book
 		ITargetPlatformManagerModule* TPM = GetTargetPlatformManager();
 		if (TPM && (TPM->RestrictFormatsToRuntimeOnly() == false)) 
