AWorldSettings - Hide HierarchicalLODSetup & OverrideBaseMaterial if HLODSetupAsset is valid

diff --git a/Engine/Source/Runtime/Engine/Classes/GameFramework/WorldSettings.h b/Engine/Source/Runtime/Engine/Classes/GameFramework/WorldSettings.h
--- a/Engine/Source/Runtime/Engine/Classes/GameFramework/WorldSettings.h
+++ b/Engine/Source/Runtime/Engine/Classes/GameFramework/WorldSettings.h
@@ -634,7 +634,7 @@ public:
 
 protected:
 	/** Hierarchical LOD Setup */
-	UPROPERTY(EditAnywhere, Category=LODSystem, config, meta=(editcondition = "bEnableHierarchicalLODSystem"))
+	UPROPERTY(EditAnywhere, Category=LODSystem, config)
 	TArray<struct FHierarchicalSimplification>	HierarchicalLODSetup;
 
 public:
diff --git a/Engine/Source/Runtime/Engine/Private/WorldSettings.cpp b/Engine/Source/Runtime/Engine/Private/WorldSettings.cpp
--- a/Engine/Source/Runtime/Engine/Private/WorldSettings.cpp
+++ b/Engine/Source/Runtime/Engine/Private/WorldSettings.cpp
@@ -556,6 +556,12 @@ bool AWorldSettings::CanEditChange(const UProperty* InProperty) const
 				return LightmassSettings.EnvironmentIntensity > 0;
 			}
 		}
+
+		if (PropertyName == GET_MEMBER_NAME_STRING_CHECKED(AWorldSettings, OverrideBaseMaterial) ||
+			PropertyName == GET_MEMBER_NAME_STRING_CHECKED(AWorldSettings, HierarchicalLODSetup))
+		{
+			return bEnableHierarchicalLODSystem && HLODSetupAsset.IsNull();
+		}
 	}
 
 	return Super::CanEditChange(InProperty);
@@ -700,6 +706,8 @@ void UHierarchicalLODSetup::PostEditChangeProperty(struct FPropertyChangedEvent&
 			}
 		}
 	}
+
+	Super::PostEditChangeProperty(PropertyChangedEvent);
 }
 #endif // WITH_EDITOR
 
