Fix for crash when using non default curve compression assets.

diff --git a/Engine/Source/Runtime/Engine/Classes/Animation/AnimSequence.h b/Engine/Source/Runtime/Engine/Classes/Animation/AnimSequence.h
--- a/Engine/Source/Runtime/Engine/Classes/Animation/AnimSequence.h
+++ b/Engine/Source/Runtime/Engine/Classes/Animation/AnimSequence.h
@@ -373,6 +373,7 @@ public:
 	virtual void PostLoad() override;
 	virtual bool IsPostLoadThreadSafe() const override;
 	virtual void PreSave(const class ITargetPlatform* TargetPlatform) override;
+	virtual void GetPreloadDependencies(TArray<UObject*>& OutDeps) override;
 #if WITH_EDITOR
 	virtual void PostEditChangeProperty(FPropertyChangedEvent& PropertyChangedEvent) override;
 	virtual void PostDuplicate(bool bDuplicateForPIE) override;
diff --git a/Engine/Source/Runtime/Engine/Private/Animation/AnimCompressionTypes.cpp b/Engine/Source/Runtime/Engine/Private/Animation/AnimCompressionTypes.cpp
--- a/Engine/Source/Runtime/Engine/Private/Animation/AnimCompressionTypes.cpp
+++ b/Engine/Source/Runtime/Engine/Private/Animation/AnimCompressionTypes.cpp
@@ -418,6 +418,17 @@ void FUECompressedAnimData::ByteSwapData(TArrayView<uint8> CompressedData, TArch
 template void FUECompressedAnimData::ByteSwapData(TArrayView<uint8> CompressedData, FMemoryReader& MemoryStream);
 template void FUECompressedAnimData::ByteSwapData(TArrayView<uint8> CompressedData, FMemoryWriter& MemoryStream);
 
+void ValidateUObjectLoaded(UObject* Obj, UObject* Source)
+{
+#if WITH_EDITOR
+	if (FLinkerLoad* ObjLinker = Obj->GetLinker())
+	{
+		ObjLinker->Preload(Obj);
+	}
+#endif
+	checkf(!Obj->HasAnyFlags(RF_NeedLoad), TEXT("Failed to load %s in %s"), *Obj->GetFullName(), *Source->GetFullName()); // in non editor should have been preloaded by GetPreloadDependencies
+}
+
 void FCompressedAnimSequence::SerializeCompressedData(FArchive& Ar, bool bDDCData, UObject* DataOwner, USkeleton* Skeleton, UAnimCurveCompressionSettings* CurveCompressionSettings, bool bCanUseBulkData)
 {
 	Ar << CompressedRawDataSize;
@@ -501,6 +512,7 @@ void FCompressedAnimSequence::SerializeCompressedData(FArchive& Ar, bool bDDCDat
 		FString CurveCodecPath;
 		Ar << CurveCodecPath;
 
+		ValidateUObjectLoaded(CurveCompressionSettings, DataOwner);
 		CurveCompressionCodec = CurveCompressionSettings->GetCodec(CurveCodecPath);
 
 		int32 NumCurveBytes;
diff --git a/Engine/Source/Runtime/Engine/Private/Animation/AnimCurveCompressionCodec_UniformIndexable.cpp b/Engine/Source/Runtime/Engine/Private/Animation/AnimCurveCompressionCodec_UniformIndexable.cpp
--- a/Engine/Source/Runtime/Engine/Private/Animation/AnimCurveCompressionCodec_UniformIndexable.cpp
+++ b/Engine/Source/Runtime/Engine/Private/Animation/AnimCurveCompressionCodec_UniformIndexable.cpp
@@ -23,7 +23,7 @@ FAnimCurveBufferAccess::FAnimCurveBufferAccess(const UAnimSequenceBase* InSequen
 				}
 				else
 				{
-					UE_LOG(LogTemp, Warning, TEXT("Trying to use FAnimCurveBufferAccess on an animation that is not compressed with UniformIndexable, this is not supported and will not work in cooked builds. Animation: %s"), *InSequence->GetFullName());
+					UE_LOG(LogAnimation, Warning, TEXT("Trying to use FAnimCurveBufferAccess on an animation that is not compressed with UniformIndexable, this is not supported and will not work in cooked builds. Animation: %s"), *InSequence->GetFullName());
 				}
 			}
 		}
diff --git a/Engine/Source/Runtime/Engine/Private/Animation/AnimSequence.cpp b/Engine/Source/Runtime/Engine/Private/Animation/AnimSequence.cpp
--- a/Engine/Source/Runtime/Engine/Private/Animation/AnimSequence.cpp
+++ b/Engine/Source/Runtime/Engine/Private/Animation/AnimSequence.cpp
@@ -819,6 +819,12 @@ void UAnimSequence::SortSyncMarkers()
 	RefreshSyncMarkerDataFromAuthored();
 }
 
+void UAnimSequence::GetPreloadDependencies(TArray<UObject*>& OutDeps)
+{
+	Super::GetPreloadDependencies(OutDeps);
+	OutDeps.Add(CurveCompressionSettings);
+}
+
 void UAnimSequence::PreSave(const class ITargetPlatform* TargetPlatform)
 {
 #if WITH_EDITOR
