Added min directional light elevation angle used for transmittance to maintain shadows at night.
This works when on +z on the planet. Added a follow todo to make it work everywehre on the planet.

diff --git a/Engine/Source/Runtime/Engine/Classes/Components/SkyAtmosphereComponent.h b/Engine/Source/Runtime/Engine/Classes/Components/SkyAtmosphereComponent.h
--- a/Engine/Source/Runtime/Engine/Classes/Components/SkyAtmosphereComponent.h
+++ b/Engine/Source/Runtime/Engine/Classes/Components/SkyAtmosphereComponent.h
@@ -129,6 +129,10 @@ class USkyAtmosphereComponent : public USceneComponent
 	UPROPERTY(EditAnywhere, BlueprintReadOnly, interp, Category = "Art direction", meta = (UIMin = 0.0, UIMax = 1.0, ClampMin = 0.0, SliderExponent = 2.0))
 	float HeightFogContribution;
 
+	/** The minimum elevation angle in degree that should be used to evaluate the sun transmittance to the ground. Useful to maintain a visible sun light and shadow on meshes even when the sun has started going below the horizon. This does not affect the aerial perspective.*/
+	UPROPERTY(EditAnywhere, BlueprintReadOnly, interp, Category = "Art direction", meta = (UIMin = -90.0, UIMax = 90.0, ClampMin = -90.0f, ClampMax = 90.0f))
+	float TransmittanceMinLightElevationAngle = -90.0f;
+
 
 
 	UFUNCTION(BlueprintCallable, Category = "Rendering")
diff --git a/Engine/Source/Runtime/Engine/Private/Components/SkyAtmosphereComponent.cpp b/Engine/Source/Runtime/Engine/Private/Components/SkyAtmosphereComponent.cpp
--- a/Engine/Source/Runtime/Engine/Private/Components/SkyAtmosphereComponent.cpp
+++ b/Engine/Source/Runtime/Engine/Private/Components/SkyAtmosphereComponent.cpp
@@ -70,6 +70,8 @@ USkyAtmosphereComponent::USkyAtmosphereComponent(const FObjectInitializer& Objec
 	AerialPespectiveViewDistanceScale = 1.0f;
 	HeightFogContribution = 1.0f;
 
+	TransmittanceMinLightElevationAngle = -90.0f;
+
 	memset(OverrideAtmosphericLight, 0, sizeof(OverrideAtmosphericLight));
 
 	ValidateStaticLightingGUIDs();
diff --git a/Engine/Source/Runtime/Engine/Public/Rendering/SkyAtmosphereCommonData.cpp b/Engine/Source/Runtime/Engine/Public/Rendering/SkyAtmosphereCommonData.cpp
--- a/Engine/Source/Runtime/Engine/Public/Rendering/SkyAtmosphereCommonData.cpp
+++ b/Engine/Source/Runtime/Engine/Public/Rendering/SkyAtmosphereCommonData.cpp
@@ -53,6 +53,8 @@ FAtmosphereSetup::FAtmosphereSetup(const USkyAtmosphereComponent& SkyAtmosphereC
 
 	AbsorptionExtinction = (SkyAtmosphereComponent.OtherAbsorption * SkyAtmosphereComponent.OtherAbsorptionScale).GetClamped(0.0f, 1e38f);
 	TentToCoefficients(SkyAtmosphereComponent.OtherTentDistribution, AbsorptionDensity0LayerWidth, AbsorptionDensity0LinearTerm, AbsorptionDensity1LinearTerm, AbsorptionDensity0ConstantTerm, AbsorptionDensity1ConstantTerm);
+
+	TransmittanceMinLightElevationAngle = SkyAtmosphereComponent.TransmittanceMinLightElevationAngle;
 }
 
 FLinearColor FAtmosphereSetup::GetTransmittanceAtGroundLevel(const FVector& SunDirection) const
@@ -137,7 +139,9 @@ FLinearColor FAtmosphereSetup::GetTransmittanceAtGroundLevel(const FVector& SunD
 
 	// Assuming camera is along Z on (0,0,earthRadius + 500m)
 	const FVector WorldPos = FVector(0.0f, 0.0f, BottomRadius + 0.5);
-	const FVector WorldDir = SunDirection;
+	FVector2D AzimuthElevation = FMath::GetAzimuthAndElevation(SunDirection, FVector::ForwardVector, FVector::LeftVector, FVector::UpVector); // TODO: make it work over the entire virtual planet with a local basis
+	AzimuthElevation.Y = FMath::Max(FMath::DegreesToRadians(TransmittanceMinLightElevationAngle), AzimuthElevation.Y);
+	const FVector WorldDir = FVector(FMath::Cos(AzimuthElevation.Y), 0.0f, FMath::Sin(AzimuthElevation.Y)); // no need to take azimuth into account as transmittance is symmetrical around zenith axis.
 	FLinearColor OpticalDepthRGB = OpticalDepth(WorldPos, WorldDir);
 	return FLinearColor(FMath::Exp(-OpticalDepthRGB.R), FMath::Exp(-OpticalDepthRGB.G), FMath::Exp(-OpticalDepthRGB.B));
 }
diff --git a/Engine/Source/Runtime/Engine/Public/Rendering/SkyAtmosphereCommonData.h b/Engine/Source/Runtime/Engine/Public/Rendering/SkyAtmosphereCommonData.h
--- a/Engine/Source/Runtime/Engine/Public/Rendering/SkyAtmosphereCommonData.h
+++ b/Engine/Source/Runtime/Engine/Public/Rendering/SkyAtmosphereCommonData.h
@@ -39,6 +39,8 @@ struct FAtmosphereSetup
 
 	FLinearColor GroundAlbedo;
 
+	float TransmittanceMinLightElevationAngle;
+
 	ENGINE_API FAtmosphereSetup(const USkyAtmosphereComponent& SkyAtmosphereComponent);
 
 	ENGINE_API FLinearColor GetTransmittanceAtGroundLevel(const FVector& SunDirection) const;
