Make K2_GetComponentsByClass public as AllowPrivateAccess does not work for UFUNCTIONS
#rnx

diff --git a/Engine/Source/Runtime/Engine/Classes/GameFramework/Actor.h b/Engine/Source/Runtime/Engine/Classes/GameFramework/Actor.h
--- a/Engine/Source/Runtime/Engine/Classes/GameFramework/Actor.h
+++ b/Engine/Source/Runtime/Engine/Classes/GameFramework/Actor.h
@@ -2802,15 +2802,14 @@ public:
 	UFUNCTION(BlueprintCallable, Category = "Actor", meta = (ComponentClass = "ActorComponent"), meta = (DeterminesOutputType = "ComponentClass"))
 	UActorComponent* GetComponentByClass(TSubclassOf<UActorComponent> ComponentClass) const;
 
-private:
 	/**
 	 * Gets all the components that inherit from the given class.
 	 * Currently returns an array of UActorComponent which must be cast to the correct type.
+	 * This intended to only be used by blueprints. Use GetComponents() in C++.
 	 */
-	UFUNCTION(BlueprintCallable, Category = "Actor", meta = (ComponentClass = "ActorComponent", DisplayName = "GetComponentsByClass", DeterminesOutputType = "ComponentClass", AllowPrivateAccess=true))
+	UFUNCTION(BlueprintCallable, Category = "Actor", meta = (ComponentClass = "ActorComponent", DisplayName = "GetComponentsByClass", DeterminesOutputType = "ComponentClass"))
 	TArray<UActorComponent*> K2_GetComponentsByClass(TSubclassOf<UActorComponent> ComponentClass) const;
 
-public:
 	/** Gets all the components that inherit from the given class with a given tag. */
 	UFUNCTION(BlueprintCallable, Category = "Actor", meta = (ComponentClass = "ActorComponent"), meta = (DeterminesOutputType = "ComponentClass"))
 	TArray<UActorComponent*> GetComponentsByTag(TSubclassOf<UActorComponent> ComponentClass, FName Tag) const;
