Fix potential out-of-bounds array access.

diff --git a/Engine/Source/Runtime/Engine/Private/PackageMapClient.cpp b/Engine/Source/Runtime/Engine/Private/PackageMapClient.cpp
--- a/Engine/Source/Runtime/Engine/Private/PackageMapClient.cpp
+++ b/Engine/Source/Runtime/Engine/Private/PackageMapClient.cpp
@@ -1518,7 +1518,12 @@ void UPackageMapClient::AppendNetFieldExportsInternal(FArchive& Archive, const T
 
 void UPackageMapClient::ReceiveNetFieldExportsCompat(FInBunch &InBunch)
 {
-	check(Connection->InternalAck);
+	if(!Connection->InternalAck)
+	{
+		UE_LOG(LogNetPackageMap, Error, TEXT("ReceiveNetFieldExportsCompat: connection is not a replay connection."));
+		InBunch.SetError();
+		return;
+	}
 
 	// Read number of net field exports
 	uint32 NumLayoutCmdExports = 0;
@@ -1589,7 +1594,7 @@ void UPackageMapClient::ReceiveNetFieldExportsCompat(FInBunch &InBunch)
 
 		TArray<FNetFieldExport>& NetFieldExportsRef = NetFieldExportGroup->NetFieldExports;
 
-		if ((int32)NetFieldExport.Handle < NetFieldExportsRef.Num())
+		if (NetFieldExportsRef.IsValidIndex((int32)NetFieldExport.Handle))
 		{
 			// Assign it to the correct slot (NetFieldExport.Handle is just the index into the array)
 			NetFieldExportGroup->NetFieldExports[NetFieldExport.Handle] = NetFieldExport;
