Retain previous Cmdline mute/solo behaviour of exclusive solo cmds.

diff --git a/Engine/Source/Runtime/Engine/Private/AudioDebug.cpp b/Engine/Source/Runtime/Engine/Private/AudioDebug.cpp
--- a/Engine/Source/Runtime/Engine/Private/AudioDebug.cpp
+++ b/Engine/Source/Runtime/Engine/Private/AudioDebug.cpp
@@ -232,18 +232,31 @@ void FAudioDebugger::SetNameArray(FName InName, TArray<FName>& InNameArray, bool
 	});
 }
 
-void FAudioDebugger::ToggleNameArray(FName InName, TArray<FName>& InNameArray )
+void FAudioDebugger::ToggleNameArray(FName InName, TArray<FName>& InNameArray, bool bExclusive )
 {
-	ExecuteCmdOnAudioThread([InName, &InNameArray]
+	ExecuteCmdOnAudioThread([InName, &InNameArray, bExclusive]
 	{ 
 		// On already?
 		int32 IndexOf = InNameArray.IndexOfByKey(InName);
 		if (IndexOf != INDEX_NONE)
-		{
-			InNameArray.RemoveAtSwap(IndexOf);
-		}			
+		{			
+			if (bExclusive)
+			{
+				// Turn off everything if we are exclusive.
+				InNameArray.Empty();
+			}
+			else
+			{
+				InNameArray.RemoveAtSwap(IndexOf);
+			}
+		}
 		else // Add it.
 		{
+			// If we are exclusive, turn off everything else first.
+			if (bExclusive)
+			{
+				InNameArray.Empty();
+			}
 			InNameArray.Add(InName);
 		}
 	});
diff --git a/Engine/Source/Runtime/Engine/Private/AudioDevice.cpp b/Engine/Source/Runtime/Engine/Private/AudioDevice.cpp
--- a/Engine/Source/Runtime/Engine/Private/AudioDevice.cpp
+++ b/Engine/Source/Runtime/Engine/Private/AudioDevice.cpp
@@ -1335,33 +1335,40 @@ bool FAudioDevice::HandleAudio3dVisualizeCommand(const TCHAR* Cmd, FOutputDevice
 	return true;
 }
 
-bool FAudioDevice::HandleAudioSoloSoundClass(const TCHAR* Cmd, FOutputDevice& Ar)
+void FAudioDevice::HandleAudioSoloCommon(const TCHAR* Cmd, FOutputDevice& Ar, FToggleSoloPtr FPtr)
 {
 	FAudioDeviceManager* DeviceManager = GEngine->GetAudioDeviceManager();
 	if (DeviceManager)
 	{
-		DeviceManager->GetDebugger().ToggleSoloSoundClass(Cmd);
+		bool bExclusive = true;
+		if (FParse::Param(Cmd, TEXT("nonexclusive")))
+		{
+			bExclusive = false;
+		}
+		TArray<FString> Args;
+		FString(Cmd).ParseIntoArray(Args, TEXT(" "));
+		if (Args.Num())
+		{
+			(DeviceManager->GetDebugger().*FPtr)(*Args[0], bExclusive);
+		}
 	}
+}
+
+bool FAudioDevice::HandleAudioSoloSoundClass(const TCHAR* Cmd, FOutputDevice& Ar)
+{
+	HandleAudioSoloCommon(Cmd, Ar, &FAudioDebugger::ToggleSoloSoundClass);
 	return true;
 }
 
 bool FAudioDevice::HandleAudioSoloSoundWave(const TCHAR* Cmd, FOutputDevice& Ar)
-{
-	FAudioDeviceManager* DeviceManager = GEngine->GetAudioDeviceManager();
-	if (DeviceManager)
-	{
-		DeviceManager->GetDebugger().ToggleSoloSoundWave(Cmd);
-	}
+{	
+	HandleAudioSoloCommon(Cmd, Ar, &FAudioDebugger::ToggleSoloSoundWave);
 	return true;
 }
 
 bool FAudioDevice::HandleAudioSoloSoundCue(const TCHAR* Cmd, FOutputDevice& Ar)
 {
-	FAudioDeviceManager* DeviceManager = GEngine->GetAudioDeviceManager();
-	if (DeviceManager)
-	{
-		DeviceManager->GetDebugger().ToggleSoloSoundCue(Cmd);
-	}
+	HandleAudioSoloCommon(Cmd, Ar, &FAudioDebugger::ToggleSoloSoundCue);
 	return true;
 }
 
diff --git a/Engine/Source/Runtime/Engine/Public/Audio/AudioDebug.h b/Engine/Source/Runtime/Engine/Public/Audio/AudioDebug.h
--- a/Engine/Source/Runtime/Engine/Public/Audio/AudioDebug.h
+++ b/Engine/Source/Runtime/Engine/Public/Audio/AudioDebug.h
@@ -79,13 +79,13 @@ public:
 	bool IsMuteSoundWave(FName InName) const  { return DebugNames.MuteSoundWave.Contains(InName); }
 	bool IsMuteSoundCue(FName InName) const   { return DebugNames.MuteSoundCue.Contains(InName); }
 		
-	// Mute/Solos toggles. (any thread).
-	void ToggleSoloSoundClass(FName InName) { ToggleNameArray(InName, DebugNames.SoloSoundClass); }
-	void ToggleSoloSoundWave(FName InName)  { ToggleNameArray(InName, DebugNames.SoloSoundWave); }
-	void ToggleSoloSoundCue(FName InName)   { ToggleNameArray(InName, DebugNames.SoloSoundCue); }
-	void ToggleMuteSoundClass(FName InName) { ToggleNameArray(InName, DebugNames.MuteSoundClass); }
-	void ToggleMuteSoundWave(FName InName)  { ToggleNameArray(InName, DebugNames.MuteSoundWave); }
-	void ToggleMuteSoundCue(FName InName)   { ToggleNameArray(InName, DebugNames.MuteSoundCue); }
+	// Mute/Solos toggles. (any thread). (If exclusive, toggle-on will clear everything first, and toggle-off will clear all).
+	void ToggleSoloSoundClass(FName InName, bool bExclusive = false) { ToggleNameArray(InName, DebugNames.SoloSoundClass, bExclusive); }
+	void ToggleSoloSoundWave(FName InName, bool bExclusive = false)  { ToggleNameArray(InName, DebugNames.SoloSoundWave, bExclusive); }
+	void ToggleSoloSoundCue(FName InName, bool bExclusive = false)   { ToggleNameArray(InName, DebugNames.SoloSoundCue, bExclusive); }
+	void ToggleMuteSoundClass(FName InName, bool bExclusive = false) { ToggleNameArray(InName, DebugNames.MuteSoundClass, bExclusive); }
+	void ToggleMuteSoundWave(FName InName, bool bExclusive = false)  { ToggleNameArray(InName, DebugNames.MuteSoundWave, bExclusive); }
+	void ToggleMuteSoundCue(FName InName, bool bExclusive = false)   { ToggleNameArray(InName, DebugNames.MuteSoundCue,	bExclusive); }
 
 	// Set Mute/Solo. (any thread).
 	void SetMuteSoundCue(FName InName, bool bInOnOff)  { SetNameArray(InName, DebugNames.MuteSoundCue, bInOnOff); }
@@ -100,7 +100,7 @@ public:
 	bool GetAudioDebugSound(FString& OutDebugSound);	  	
 private:
 	void SetNameArray(FName InName, TArray<FName>& InNameArray, bool bOnOff);
-	void ToggleNameArray(FName InName, TArray<FName>& NameArray);
+	void ToggleNameArray(FName InName, TArray<FName>& NameArray, bool bExclusive);
 	void ExecuteCmdOnAudioThread(TFunction<void()> Cmd);
 	
 	void GetDebugSoloMuteStateX(
diff --git a/Engine/Source/Runtime/Engine/Public/AudioDevice.h b/Engine/Source/Runtime/Engine/Public/AudioDevice.h
--- a/Engine/Source/Runtime/Engine/Public/AudioDevice.h
+++ b/Engine/Source/Runtime/Engine/Public/AudioDevice.h
@@ -25,6 +25,7 @@
  */
 
 class FArchive;
+class FAudioDebugger;
 class FAudioDevice;
 class FAudioEffectsManager;
 class FCanvas;
@@ -414,6 +415,10 @@ private:
 	bool HandleGetDynamicSoundVolumeCommand(const TCHAR* Cmd, FOutputDevice& Ar);
 	bool HandleSetDynamicSoundCommand(const TCHAR* Cmd, FOutputDevice& Ar);
 
+	/** Handles all argument parsing for the solo commands in one place */
+	using FToggleSoloPtr = void (FAudioDebugger::*)(FName InName, bool bExclusive);
+	void HandleAudioSoloCommon(const TCHAR* Cmd, FOutputDevice& Ar, FToggleSoloPtr Funct);
+
 	/**
 	* Lists a summary of loaded sound collated by class
 	*/
@@ -1552,7 +1557,6 @@ private:
 	int32 GetNumPrecacheFrames() const;
 
 	bool RemoveVirtualLoop(FActiveSound& ActiveSound);
-
 public:
 
 	/** Query if the editor is in VR Preview for the current play world. Returns false for non-editor builds */
