Blueprint components added at runtime (say, via "Add Component" nodes) now replicate if "Component replicates" is set to true, but the class default had replication disabled.

diff --git a/Engine/Source/Runtime/Engine/Private/ActorConstruction.cpp b/Engine/Source/Runtime/Engine/Private/ActorConstruction.cpp
--- a/Engine/Source/Runtime/Engine/Private/ActorConstruction.cpp
+++ b/Engine/Source/Runtime/Engine/Private/ActorConstruction.cpp
@@ -1210,6 +1210,14 @@ void AActor::PostCreateBlueprintComponent(UActorComponent* NewActorComp)
 
 		// Need to do this so component gets saved - Components array is not serialized
 		BlueprintCreatedComponents.Add(NewActorComp);
+
+		// The component may not have been added to ReplicatedComponents if it was duplicated from
+		// a template, since ReplicatedComponents is normally only updated before the duplicated properties
+		// are copied over - in this case bReplicates would not have been set yet, but it will be now.
+		if (NewActorComp->GetIsReplicated())
+		{
+			ReplicatedComponents.AddUnique(NewActorComp);
+		}
 	}
 }
 
