Make it easier to save out slow queries

diff --git a/Engine/Source/Runtime/Engine/Private/Collision/SceneQueryLowLevel.cpp b/Engine/Source/Runtime/Engine/Private/Collision/SceneQueryLowLevel.cpp
--- a/Engine/Source/Runtime/Engine/Private/Collision/SceneQueryLowLevel.cpp
+++ b/Engine/Source/Runtime/Engine/Private/Collision/SceneQueryLowLevel.cpp
@@ -67,22 +67,22 @@ void LowLevelRaycast(FPhysScene& Scene, const FVector& Start, const FVector& Dir
 	if (const auto& SolverAccelerationStructure = Scene.GetScene().GetSpacialAcceleration())
 	{
 		FChaosSQAccelerator SQAccelerator(*SolverAccelerationStructure);
-		if (SerializeSQs && IsInGameThread())
+		double Time = 0.0;
+		{
+			FScopedDurationTimer Timer(Time);
+			SQAccelerator.Raycast(Start, Dir, DeltaMag, HitBuffer, OutputFlags, QueryFilterData, *QueryCallback, DebugParams);
+		}
+
+		/*if (((SerializeSQs && Time * 1000.0 * 1000.0 > SerializeSQs) || (false)) && IsInGameThread())
 		{
 			FPhysTestSerializer Serializer;
 			Serializer.SetPhysicsData(*Scene.GetSolver()->GetEvolution());
 			FSQCapture& RaycastCapture = Serializer.CaptureSQ();
 			RaycastCapture.StartCaptureChaosRaycast(*Scene.GetSolver()->GetEvolution(), Start, Dir, DeltaMag, OutputFlags, QueryFilterData, Filter, *QueryCallback);
-			SQAccelerator.Raycast(Start, Dir, DeltaMag, HitBuffer, OutputFlags, QueryFilterData, *QueryCallback, DebugParams);
 			RaycastCapture.EndCaptureChaosRaycast(HitBuffer);
 
 			FinalizeCapture(Serializer);
-		}
-		else
-		{
-			//ISQAccelerator* SQAccelerator = Scene.GetSQAccelerator();
-			SQAccelerator.Raycast(Start, Dir, DeltaMag, HitBuffer, OutputFlags, QueryFilterData, *QueryCallback, DebugParams);
-		}
+		}*/
 	}
 #else
 	if (SerializeSQs | ReplaySQs)
@@ -109,7 +109,27 @@ void LowLevelSweep(FPhysScene& Scene, const FPhysicsGeometry& QueryGeom, const F
 	if (const auto& SolverAccelerationStructure = Scene.GetScene().GetSpacialAcceleration())
 	{
 		FChaosSQAccelerator SQAccelerator(*SolverAccelerationStructure);
-		if (SerializeSQs && IsInGameThread())
+		{
+			//ISQAccelerator* SQAccelerator = Scene.GetSQAccelerator();
+			double Time = 0.0;
+			{
+				FScopedDurationTimer Timer(Time);
+				SQAccelerator.Sweep(QueryGeom, StartTM, Dir, DeltaMag, HitBuffer, OutputFlags, QueryFilterData, *QueryCallback, DebugParams);
+			}
+			
+			if (((SerializeSQs && Time * 1000.0 * 1000.0 > SerializeSQs) || (false)) && IsInGameThread())
+			{
+				FPhysTestSerializer Serializer;
+				Serializer.SetPhysicsData(*Scene.GetSolver()->GetEvolution());
+				FSQCapture& SweepCapture = Serializer.CaptureSQ();
+				SweepCapture.StartCaptureChaosSweep(*Scene.GetSolver()->GetEvolution(), QueryGeom, StartTM, Dir, DeltaMag, OutputFlags, QueryFilterData, Filter, *QueryCallback);
+				SweepCapture.EndCaptureChaosSweep(HitBuffer);
+
+				FinalizeCapture(Serializer);
+			}
+		}
+
+		/*if (SerializeSQs && IsInGameThread())
 		{
 			FPhysTestSerializer Serializer;
 			Serializer.SetPhysicsData(*Scene.GetSolver()->GetEvolution());
@@ -123,8 +143,13 @@ void LowLevelSweep(FPhysScene& Scene, const FPhysicsGeometry& QueryGeom, const F
 		else
 		{
 			//ISQAccelerator* SQAccelerator = Scene.GetSQAccelerator();
-			SQAccelerator.Sweep(QueryGeom, StartTM, Dir, DeltaMag, HitBuffer, OutputFlags, QueryFilterData, *QueryCallback);
-		}
+			double Time = 0.0;
+			{
+				FScopedDurationTimer Timer(Time);
+				SQAccelerator.Sweep(QueryGeom, StartTM, Dir, DeltaMag, HitBuffer, OutputFlags, QueryFilterData, *QueryCallback);
+			}
+			if(Time > )
+		}*/
 	}
 #else
 	if (SerializeSQs | ReplaySQs)
@@ -151,7 +176,7 @@ void LowLevelOverlap(FPhysScene& Scene, const FPhysicsGeometry& QueryGeom, const
 	if (const auto& SolverAccelerationStructure = Scene.GetScene().GetSpacialAcceleration())
 	{
 		FChaosSQAccelerator SQAccelerator(*SolverAccelerationStructure);
-		if (SerializeSQs && IsInGameThread())
+		if (false && SerializeSQs && IsInGameThread())
 		{
 			FPhysTestSerializer Serializer;
 			Serializer.SetPhysicsData(*Scene.GetSolver()->GetEvolution());
