Introduced a new debugging aid GPlayInEditorContextString, which can be added to the watch window to indicate the kind of world that is being processed when stopped at a breakpoint in the editor (e.g., "Client 1" vs "Dedicated Server")
- To use, add UE4Editor-Engine!GPlayInEditorContextString to your Watch Window
- This change also better manages GPlayInEditorID in a couple of circumstances, properly resetting it at the end of a scope

diff --git a/Engine/Source/Editor/UnrealEd/Private/Editor.cpp b/Engine/Source/Editor/UnrealEd/Private/Editor.cpp
--- a/Engine/Source/Editor/UnrealEd/Private/Editor.cpp
+++ b/Engine/Source/Editor/UnrealEd/Private/Editor.cpp
@@ -23,6 +23,7 @@
 #include "UObject/UObjectIterator.h"
 #include "EngineUtils.h"
 #include "Dialogs/Dialogs.h"
+#include "UnrealEngine.h"
 
 // needed for the RemotePropagator
 
@@ -806,6 +807,7 @@ UWorld* SetPlayInEditorWorld( UWorld* PlayInEditorWorld )
 	if (FWorldContext* WorldContext = GEngine->GetWorldContextFromWorld(PlayInEditorWorld))
 	{
 		GPlayInEditorID = WorldContext->PIEInstance;
+		UpdatePlayInEditorWorldDebugString(WorldContext);
 	}
 
 	return SavedWorld;
@@ -822,7 +824,8 @@ void RestoreEditorWorld( UWorld* EditorWorld )
 	check(GIsPlayInEditorWorld);
 	GIsPlayInEditorWorld = false;
 	GWorld = EditorWorld;
-	GPlayInEditorID = -1;
+	GPlayInEditorID = INDEX_NONE;
+	UpdatePlayInEditorWorldDebugString(nullptr);
 }
 
 /**
diff --git a/Engine/Source/Editor/UnrealEd/Private/EditorEngine.cpp b/Engine/Source/Editor/UnrealEd/Private/EditorEngine.cpp
--- a/Engine/Source/Editor/UnrealEd/Private/EditorEngine.cpp
+++ b/Engine/Source/Editor/UnrealEd/Private/EditorEngine.cpp
@@ -1636,16 +1636,13 @@ void UEditorEngine::Tick( float DeltaSeconds, bool bIdleMode )
 		// So we iterate on a local list here instead.
 		for (FWorldContext* PieContextPtr : LocalPieContextPtrs)
 		{
-			FWorldContext &PieContext = *PieContextPtr;
-
-			GPlayInEditorID = PieContext.PIEInstance;
+			FWorldContext& PieContext = *PieContextPtr;
 
 			PlayWorld = PieContext.World();
 			GameViewport = PieContext.GameViewport;
 
-			UWorld* OldGWorld = NULL;
-			// Use the PlayWorld as the GWorld, because who knows what will happen in the Tick.
-			OldGWorld = SetPlayInEditorWorld( PlayWorld );
+			// Switch worlds and set the play world ID
+			UWorld* OldGWorld = SetPlayInEditorWorld(PlayWorld);
 
 			// Transfer debug references to ensure debugging ref's are valid for this tick in case of multiple game instances.
 			if (OldGWorld && OldGWorld != PlayWorld)
@@ -1730,7 +1727,8 @@ void UEditorEngine::Tick( float DeltaSeconds, bool bIdleMode )
 		bFirstTick = false;
 	}
 
-	GPlayInEditorID = -1;
+	ensure(GPlayInEditorID == INDEX_NONE);
+	GPlayInEditorID = INDEX_NONE;
 
 	// Clean up any game viewports that may have been closed during the level tick (eg by Kismet).
 	CleanupGameViewport();
@@ -1876,11 +1874,10 @@ void UEditorEngine::Tick( float DeltaSeconds, bool bIdleMode )
 			GameViewport = PieContext.GameViewport;
 
 			// Render playworld. This needs to happen after the other viewports for screenshots to work correctly in PIE.
-			if(PlayWorld && GameViewport && !bIsSimulatingInEditor)
+			if (PlayWorld && GameViewport && !bIsSimulatingInEditor)
 			{
 				// Use the PlayWorld as the GWorld, because who knows what will happen in the Tick.
 				UWorld* OldGWorld = SetPlayInEditorWorld( PlayWorld );
-				GPlayInEditorID = PieContext.PIEInstance;
 
 				// Render everything.
 				GameViewport->LayoutPlayers();
@@ -1889,7 +1886,6 @@ void UEditorEngine::Tick( float DeltaSeconds, bool bIdleMode )
 
 				// Pop the world
 				RestoreEditorWorld( OldGWorld );
-				GPlayInEditorID = -1;
 			}
 		}
 	}
diff --git a/Engine/Source/Editor/UnrealEd/Private/PlayLevel.cpp b/Engine/Source/Editor/UnrealEd/Private/PlayLevel.cpp
--- a/Engine/Source/Editor/UnrealEd/Private/PlayLevel.cpp
+++ b/Engine/Source/Editor/UnrealEd/Private/PlayLevel.cpp
@@ -3552,9 +3552,11 @@ UGameInstance* UEditorEngine::CreatePIEGameInstance(int32 InPIEInstance, bool bI
 	GameInstance->RemoveFromRoot();
 
 	// Start the game instance, make sure to set the PIE instance global as this is basically a tick
-	GPlayInEditorID = InPIEInstance;
-	const FGameInstancePIEResult StartResult = GameInstance->StartPlayInEditorGameInstance(NewLocalPlayer, GameInstanceParams);
-	GPlayInEditorID = -1;
+	FGameInstancePIEResult StartResult = FGameInstancePIEResult::Success();
+	{
+		FTemporaryPlayInEditorIDOverride OverrideIDHelper(InPIEInstance);
+		StartResult = GameInstance->StartPlayInEditorGameInstance(NewLocalPlayer, GameInstanceParams);
+	}
 
 	if (!StartResult.IsSuccess())
 	{
@@ -3877,7 +3879,7 @@ UWorld* UEditorEngine::CreatePIEWorldByDuplication(FWorldContext &WorldContext,
 	const FScopedBusyCursor BusyCursor;
 
 	// Before loading the map, we need to set these flags to true so that postload will work properly
-	GIsPlayInEditorWorld = true;
+	TGuardValue<bool> OverrideIsPlayWorld(GIsPlayInEditorWorld, true);
 
 	const FName PlayWorldMapFName = FName(*PlayWorldMapName);
 	UWorld::WorldTypePreLoadMap.FindOrAdd(PlayWorldMapFName) = EWorldType::PIE;
@@ -3894,7 +3896,7 @@ UWorld* UEditorEngine::CreatePIEWorldByDuplication(FWorldContext &WorldContext,
 
 	// check(GPlayInEditorID == -1 || GPlayInEditorID == WorldContext.PIEInstance);
 	// Currently GPlayInEditorID is not correctly reset after map loading, so it's not safe to assert here
-	GPlayInEditorID = WorldContext.PIEInstance;
+	FTemporaryPlayInEditorIDOverride IDHelper(WorldContext.PIEInstance);
 
 	{
 		double SDOStart = FPlatformTime::Seconds();
@@ -3950,14 +3952,10 @@ UWorld* UEditorEngine::CreatePIEWorldByDuplication(FWorldContext &WorldContext,
 	// Clean up the world type list now that PostLoad has occurred
 	UWorld::WorldTypePreLoadMap.Remove(PlayWorldMapFName);
 
-	GPlayInEditorID = -1;
 	check( NewPIEWorld );
 	NewPIEWorld->FeatureLevel = InWorld->FeatureLevel;
 	PostCreatePIEWorld(NewPIEWorld);
 
-	// After loading the map, reset these so that things continue as normal
-	GIsPlayInEditorWorld = false;
-	
 	UE_LOG(LogPlayLevel, Log, TEXT("PIE: Created PIE world by copying editor world from %s to %s (%fs)"), *InWorld->GetPathName(), *NewPIEWorld->GetPathName(), float(FPlatformTime::Seconds() - StartTime));
 	return NewPIEWorld;
 }
diff --git a/Engine/Source/Runtime/Engine/Private/Level.cpp b/Engine/Source/Runtime/Engine/Private/Level.cpp
--- a/Engine/Source/Runtime/Engine/Private/Level.cpp
+++ b/Engine/Source/Runtime/Engine/Private/Level.cpp
@@ -2047,7 +2047,7 @@ void ULevel::BeginCacheForCookedPlatformData(const ITargetPlatform *TargetPlatfo
 
 void ULevel::FixupForPIE(int32 PIEInstanceID)
 {
-	TGuardValue<int32> SetPlayInEditorID(GPlayInEditorID, PIEInstanceID);
+	FTemporaryPlayInEditorIDOverride SetPlayInEditorID(PIEInstanceID);
 
 	struct FSoftPathPIEFixupSerializer : public FArchiveUObject
 	{
diff --git a/Engine/Source/Runtime/Engine/Private/UnrealEngine.cpp b/Engine/Source/Runtime/Engine/Private/UnrealEngine.cpp
--- a/Engine/Source/Runtime/Engine/Private/UnrealEngine.cpp
+++ b/Engine/Source/Runtime/Engine/Private/UnrealEngine.cpp
@@ -360,6 +360,81 @@ bool GIsTextureMemoryCorrupted = false;
 bool GIsPrepareMapChangeBroken = false;
 #endif
 
+#if WITH_EDITOR
+// A debugging aid set when we switch out different play worlds during Play In Editor / PIE
+ENGINE_API FString GPlayInEditorContextString(TEXT("invalid"));
+#endif
+
+ENGINE_API void UpdatePlayInEditorWorldDebugString(const FWorldContext* WorldContext)
+{
+#if WITH_EDITOR
+	if (WorldContext == nullptr)
+	{
+		ensure(GPlayInEditorID == INDEX_NONE);
+		GPlayInEditorContextString = NSLOCTEXT("Engine", "PlayWorldIsNotActive", "Not in a play world").ToString();
+	}
+	else
+	{
+		FString WorldName;
+		if (UWorld* World = WorldContext->World())
+		{
+			switch (World->GetNetMode())
+			{
+			case NM_Standalone:
+				WorldName = NSLOCTEXT("Engine", "PlayWorldIsStandalone", "Standalone").ToString();
+				break;
+
+			case NM_ListenServer:
+				WorldName = NSLOCTEXT("Engine", "PlayWorldIsListenServer", "Listen Server").ToString();
+				break;
+
+			case NM_DedicatedServer:
+				WorldName = NSLOCTEXT("Engine", "PlayWorldIsDedicatedServer", "Dedicated Server").ToString();
+				break;
+
+			case NM_Client:
+				WorldName = FText::Format(NSLOCTEXT("Engine", "PlayWorldIsClient", "Client {0}"), FText::AsNumber(WorldContext->PIEInstance - 1)).ToString();
+				break;
+
+			default:
+				unimplemented();
+			};
+		}
+		else
+		{
+			WorldName = NSLOCTEXT("Engine", "PlayWorldBeingCreated", "(World Being Created)").ToString();
+		}
+
+		if (!WorldContext->CustomDescription.IsEmpty())
+		{
+			WorldName += TEXT(" ") + WorldContext->CustomDescription;
+		}
+
+		GPlayInEditorContextString = WorldName;
+	}
+#endif
+}
+
+FTemporaryPlayInEditorIDOverride::FTemporaryPlayInEditorIDOverride(int32 NewOverrideID)
+	: PreviousID(GPlayInEditorID)
+{
+	SetID(NewOverrideID);
+}
+
+FTemporaryPlayInEditorIDOverride::~FTemporaryPlayInEditorIDOverride()
+{
+	SetID(PreviousID);
+}
+
+void FTemporaryPlayInEditorIDOverride::SetID(int32 NewID)
+{
+	if (GPlayInEditorID != NewID)
+	{
+		GPlayInEditorID = NewID;
+		UpdatePlayInEditorWorldDebugString(GEngine->GetWorldContextFromPIEInstance(GPlayInEditorID));
+	}
+}
+
 PRAGMA_DISABLE_DEPRECATION_WARNINGS
 FSimpleMulticastDelegate UEngine::OnPostEngineInit;
 PRAGMA_ENABLE_DEPRECATION_WARNINGS
@@ -903,6 +978,7 @@ namespace
 		return NewWorld;
 	}
 }
+
 /*-----------------------------------------------------------------------------
 Object class implementation.
 -----------------------------------------------------------------------------*/
@@ -12400,7 +12476,7 @@ bool UEngine::LoadMap( FWorldContext& WorldContext, FURL URL, class UPendingNetG
 
 	MALLOC_PROFILER( FMallocProfiler::SnapshotMemoryLoadMapMid( URL.Map ); )
 
-		WorldContext.OwningGameInstance->PreloadContentForURL(URL);
+	WorldContext.OwningGameInstance->PreloadContentForURL(URL);
 
 	UPackage* WorldPackage = NULL;
 	UWorld*	NewWorld = NULL;
@@ -12423,6 +12499,7 @@ bool UEngine::LoadMap( FWorldContext& WorldContext, FURL URL, class UPendingNetG
 					// We are loading a new world for this context, so clear out PIE fixups that might be lingering.
 					// (note we dont want to do this in DuplicateWorldForPIE, since that is also called on streaming worlds.
 					GPlayInEditorID = WorldContext.PIEInstance;
+					UpdatePlayInEditorWorldDebugString(&WorldContext);
 					FLazyObjectPtr::ResetPIEFixups();
 
 					NewWorld = UWorld::DuplicateWorldForPIE(SourceWorldPackage, nullptr);
@@ -12726,7 +12803,7 @@ bool UEngine::LoadMap( FWorldContext& WorldContext, FURL URL, class UPendingNetG
 	TRACE_BOOKMARK(TEXT("LoadMapComplete - %s"), *URL.Map);
 	MALLOC_PROFILER( FMallocProfiler::SnapshotMemoryLoadMapEnd( URL.Map ); )
 
-		double StopTime = FPlatformTime::Seconds();
+	double StopTime = FPlatformTime::Seconds();
 
 	UE_LOG(LogLoad, Log, TEXT("Took %f seconds to LoadMap(%s)"), StopTime - StartTime, *URL.Map);
 	FLoadTimeTracker::Get().DumpRawLoadTimes();
diff --git a/Engine/Source/Runtime/Engine/Private/World.cpp b/Engine/Source/Runtime/Engine/Private/World.cpp
--- a/Engine/Source/Runtime/Engine/Private/World.cpp
+++ b/Engine/Source/Runtime/Engine/Private/World.cpp
@@ -3006,7 +3006,7 @@ UWorld* UWorld::DuplicateWorldForPIE(const FString& PackageName, UWorld* OwningW
 		checkf(false, TEXT("Unable to determine PIEInstanceID to duplicate for PIE."));
 	}
 
-	GPlayInEditorID = PIEInstanceID;
+	FTemporaryPlayInEditorIDOverride IDHelper(PIEInstanceID);
 
 	FString PrefixedLevelName = ConvertToPIEPackageName(PackageName, PIEInstanceID);
 	const FName PrefixedLevelFName = FName(*PrefixedLevelName);
@@ -7223,7 +7223,7 @@ static ULevel* DuplicateLevelWithPrefix(ULevel* InLevel, int32 InstanceID )
 
 	FSoftObjectPath::AddPIEPackageName(NewPackage->GetFName());
 
-	GPlayInEditorID = InstanceID;
+	FTemporaryPlayInEditorIDOverride IDHelper(InstanceID);
 
 	// Create "vestigial" world for the persistent level - it's OwningWorld will still be the main world,
 	// but we're treating it like a streaming level (even though it's a duplicate of the persistent level).
@@ -7268,8 +7268,6 @@ static ULevel* DuplicateLevelWithPrefix(ULevel* InLevel, int32 InstanceID )
 
 	UE_LOG( LogNet, Log, TEXT( "DuplicateLevelWithPrefix. TotalSeconds: %2.2f" ), TotalSeconds );
 
-	GPlayInEditorID = -1;
-
 	return NewLevel;
 }
 
diff --git a/Engine/Source/Runtime/Engine/Public/UnrealEngine.h b/Engine/Source/Runtime/Engine/Public/UnrealEngine.h
--- a/Engine/Source/Runtime/Engine/Public/UnrealEngine.h
+++ b/Engine/Source/Runtime/Engine/Public/UnrealEngine.h
@@ -491,3 +491,19 @@ struct ENGINE_API FSystemResolution
 
 ENGINE_API extern FSystemResolution GSystemResolution;
 ENGINE_API extern int32 GUnbuiltHLODCount;
+
+// Update the debugging aid GPlayInEditorContextString based on the current world context (does nothing in WITH_EDITOR=0 builds)
+ENGINE_API void UpdatePlayInEditorWorldDebugString(const FWorldContext* WorldContext);
+
+// Used to temporarily override GPlayInEditorID, correctly updating the debug string and other state as necessary
+struct ENGINE_API FTemporaryPlayInEditorIDOverride
+{
+public:
+	FTemporaryPlayInEditorIDOverride(int32 NewOverrideID);
+	~FTemporaryPlayInEditorIDOverride();
+
+private:
+	int32 PreviousID;
+
+	void SetID(int32 NewID);
+};
\ No newline at end of file
