Added explicit include of UnrealEngine.h to Level.cpp to fix non-unity error
#rnx

diff --git a/Engine/Source/Runtime/Engine/Private/Level.cpp b/Engine/Source/Runtime/Engine/Private/Level.cpp
--- a/Engine/Source/Runtime/Engine/Private/Level.cpp
+++ b/Engine/Source/Runtime/Engine/Private/Level.cpp
@@ -45,6 +45,7 @@ Level.cpp: Level-related functions
 #include "PhysicsEngine/BodySetup.h"
 #include "EngineGlobals.h"
 #include "Engine/LevelBounds.h"
+#include "UnrealEngine.h"
 #if WITH_EDITOR
 #include "Kismet2/KismetEditorUtilities.h"
 #include "Kismet2/BlueprintEditorUtils.h"
