Fixed broken trails when spawning interpolated position.

diff --git a/Engine/Source/Runtime/Engine/Private/Particles/ParticleTrail2EmitterInstance.cpp b/Engine/Source/Runtime/Engine/Private/Particles/ParticleTrail2EmitterInstance.cpp
--- a/Engine/Source/Runtime/Engine/Private/Particles/ParticleTrail2EmitterInstance.cpp
+++ b/Engine/Source/Runtime/Engine/Private/Particles/ParticleTrail2EmitterInstance.cpp
@@ -45,6 +45,13 @@ DEFINE_STAT(STAT_AnimTrailNotifyTime);
 
 DECLARE_CYCLE_STAT(TEXT("TrailEmitterInstance Init GT"), STAT_TrailEmitterInstance_Init, STATGROUP_Particles);
 
+static int32 GSkipRibbonSpawnInterp = 1;
+static FAutoConsoleVariableRef CVarSkipRibbonSpawnInterp(
+	TEXT("r.Emitter.SkipRibbonSpawnInterp"),
+	GSkipRibbonSpawnInterp,
+	TEXT("Ignore velocity based offsets when interpolating. This prevents ribbon quads from overlapping eachother (default=1)"),
+	ECVF_Default
+);
 
 #define MAX_TRAIL_INDICES	65535
 
@@ -872,7 +879,7 @@ void TrailsBase_CalculateTangent(
 	NewTangent *= (1.0f / InOutCurrTrailData->SpawnedTessellationPoints);
 
 		InOutCurrTrailData->Tangent = NewTangent;
-	}
+}
 
 /**
  *	Tick sub-function that handles recalculation of tangents
@@ -901,7 +908,7 @@ void FParticleRibbonEmitterInstance::Tick_RecalculateTangents(float DeltaTime, U
 				//     END, prev, prev, ..., START
 				FBaseParticle* PrevParticle = StartParticle;
 				FRibbonTypeDataPayload* PrevTrailData = StartTrailData;
-				FBaseParticle* CurrParticle = NULL;
+				FBaseParticle* CurrParticle = NULL; 
 				FRibbonTypeDataPayload* CurrTrailData = NULL;
 				FBaseParticle* NextParticle = NULL;
 				FTrailsBaseTypeDataPayload* TempPayload = NULL;
@@ -1862,7 +1869,8 @@ bool FParticleRibbonEmitterInstance::Spawn_Source(float DeltaTime)
 					//@todo. Need to track TypeData offset into payload!
 					LODLevel->TypeDataModule->Spawn(this, TypeDataOffset, SpawnTime, Particle);
 				}
-				PostSpawn(Particle, 1.f - float(SpawnIdx + 1) / float(MovementSpawnCount), SpawnTime);
+				const float InterpolationPercentage = GSkipRibbonSpawnInterp ? 0.f : (1.f - float(SpawnIdx + 1) / float(MovementSpawnCount));
+				PostSpawn(Particle, InterpolationPercentage, SpawnTime);
 
 				GetParticleLifetimeAndSize(TrailIdx, Particle, bNoLivingParticles, Particle->OneOverMaxLifetime, Particle->Size.X);
 				Particle->RelativeTime = SpawnTime * Particle->OneOverMaxLifetime;
