SoundNodeQuality should be evaluated during cooking

diff --git a/Engine/Plugins/Runtime/AR/Microsoft/HoloLensAR/Source/HoloLensPlatformEditor/Private/HoloLensTargetSettings.h b/Engine/Plugins/Runtime/AR/Microsoft/HoloLensAR/Source/HoloLensPlatformEditor/Private/HoloLensTargetSettings.h
--- a/Engine/Plugins/Runtime/AR/Microsoft/HoloLensAR/Source/HoloLensPlatformEditor/Private/HoloLensTargetSettings.h
+++ b/Engine/Plugins/Runtime/AR/Microsoft/HoloLensAR/Source/HoloLensPlatformEditor/Private/HoloLensTargetSettings.h
@@ -138,4 +138,8 @@ public:
 	/** Which of the currently enabled occlusion plugins to use on Windows. */
 	UPROPERTY(config, EditAnywhere, Category = "Audio")
 	FString OcclusionPlugin;
+
+	/** Quality Level to COOK SoundCues at (if set, all other levels will be stripped by the cooker). */
+	UPROPERTY(config, EditAnywhere, Category = "Audio|CookOverrides", meta = (DisplayName = "Sound Cue Cook Quality"))
+	int32 SoundCueCookQualityIndex = INDEX_NONE;
 };
diff --git a/Engine/Plugins/Runtime/SoundMod/Source/SoundMod/Private/SoundNodeModPlayer.cpp b/Engine/Plugins/Runtime/SoundMod/Source/SoundMod/Private/SoundNodeModPlayer.cpp
--- a/Engine/Plugins/Runtime/SoundMod/Source/SoundMod/Private/SoundNodeModPlayer.cpp
+++ b/Engine/Plugins/Runtime/SoundMod/Source/SoundMod/Private/SoundNodeModPlayer.cpp
@@ -22,7 +22,7 @@ void USoundNodeModPlayer::Serialize(FArchive& Ar)
 		}
 		else if (Ar.IsSaving())
 		{
-			USoundMod* HardReference = (ShouldHardReferenceAsset() ? SoundMod : nullptr);
+			USoundMod* HardReference = (ShouldHardReferenceAsset(Ar.CookingTarget()) ? SoundMod : nullptr);
 			Ar << HardReference;
 		}
 	}
diff --git a/Engine/Source/Developer/Android/AndroidTargetPlatform/Private/AndroidTargetPlatform.cpp b/Engine/Source/Developer/Android/AndroidTargetPlatform/Private/AndroidTargetPlatform.cpp
--- a/Engine/Source/Developer/Android/AndroidTargetPlatform/Private/AndroidTargetPlatform.cpp
+++ b/Engine/Source/Developer/Android/AndroidTargetPlatform/Private/AndroidTargetPlatform.cpp
@@ -662,6 +662,12 @@ namespace Android
 	{
 		const TCHAR* CategoryName = TEXT("/Script/AndroidRuntimeSettings.AndroidRuntimeSettings");
 
+		int32 SoundCueQualityIndex = INDEX_NONE;
+		if (GConfig->GetInt(CategoryName, TEXT("SoundCueCookQualityIndex"), SoundCueQualityIndex, GEngineIni))
+		{
+			OutOverrides.SoundCueCookQualityIndex = SoundCueQualityIndex;
+		}
+
 		GConfig->GetBool(CategoryName, TEXT("bUseAudioStreamCaching"), OutOverrides.bUseStreamCaching, GEngineIni);
 
 		/** Memory Load On Demand Settings */
diff --git a/Engine/Source/Developer/AudioSettingsEditor/AudioSettingsEditor.Build.cs b/Engine/Source/Developer/AudioSettingsEditor/AudioSettingsEditor.Build.cs
--- a/Engine/Source/Developer/AudioSettingsEditor/AudioSettingsEditor.Build.cs
+++ b/Engine/Source/Developer/AudioSettingsEditor/AudioSettingsEditor.Build.cs
@@ -17,7 +17,8 @@ public class AudioSettingsEditor : ModuleRules
 				"PropertyEditor",
 				"SharedSettingsWidgets",
 				"UnrealEd",
-			}
+                "CoreUObject"
+            }
 		);
 	}
 }
diff --git a/Engine/Source/Developer/AudioSettingsEditor/Private/TargetPlatformAudioCustomization.cpp b/Engine/Source/Developer/AudioSettingsEditor/Private/TargetPlatformAudioCustomization.cpp
--- a/Engine/Source/Developer/AudioSettingsEditor/Private/TargetPlatformAudioCustomization.cpp
+++ b/Engine/Source/Developer/AudioSettingsEditor/Private/TargetPlatformAudioCustomization.cpp
@@ -14,7 +14,8 @@
 #include "DetailCategoryBuilder.h"
 #include "Widgets/Input/SEditableTextBox.h"
 #include "Features/IModularFeatures.h"
-
+#include "PropertyCustomizationHelpers.h"
+#include "Widgets/SToolTip.h"
 
 #if WITH_ENGINE
 #include "AudioDevice.h"
@@ -269,6 +270,47 @@ void FAudioPluginWidgetManager::BuildAudioCategory(IDetailLayoutBuilder& DetailL
 		[
 			MakeAudioPluginSelectorWidget(AudioOcclusionPropertyHandle, EAudioPlugin::OCCLUSION, AudioPlatform)
 		];
+		
+	
+	// Not really a plugin, but this is common to all TargetPlatforms
+	TSharedPtr<IPropertyHandle> SoundQualityNamePropHandle = DetailLayout.GetProperty("SoundCueCookQualityIndex");
+	ensure(SoundQualityNamePropHandle.IsValid());
+	IDetailPropertyRow& SoundQualityNamePropRow = AudioCategory.AddProperty(SoundQualityNamePropHandle);
+
+	SoundQualityNamePropRow.CustomWidget()
+		.NameContent()
+		[
+			SoundQualityNamePropHandle->CreatePropertyNameWidget()
+		]
+		.ValueContent()
+		.MaxDesiredWidth(500.0f)
+		.MinDesiredWidth(100.0f)		
+		[
+			PropertyCustomizationHelpers::MakePropertyComboBox(
+				SoundQualityNamePropHandle,
+				FOnGetPropertyComboBoxStrings::CreateLambda([](TArray<TSharedPtr<FString>>& OutComboBoxStrings, TArray<TSharedPtr<SToolTip>>& OutToolTips, TArray<bool>& OutRestrictedItems) -> void  
+				{									
+					for (const FAudioQualitySettings& i : GetDefault<UAudioSettings>()->QualityLevels)
+					{
+						OutComboBoxStrings.Add(MakeShared<FString>(i.DisplayName.ToString()));
+						OutRestrictedItems.Add(false);
+						OutToolTips.Add(SNew(SToolTip).Text(i.DisplayName));
+					}
+				}), 
+				FOnGetPropertyComboBoxValue::CreateLambda([SoundQualityNamePropHandle]() -> FString
+				{
+					int32 IndexVal = INDEX_NONE;
+					SoundQualityNamePropHandle->GetValue(IndexVal); 
+					return GetDefault<UAudioSettings>()->FindQualityNameByIndex(IndexVal);
+				}),
+				FOnPropertyComboBoxValueSelected::CreateLambda([SoundQualityNamePropHandle](FString Value) -> void
+				{
+					const TArray<FAudioQualitySettings>& Values = GetDefault<UAudioSettings>()->QualityLevels;
+					int32 Index = Values.IndexOfByPredicate([TextValue = FText::FromString(Value)](const FAudioQualitySettings& i) -> bool { return i.DisplayName.CompareTo(TextValue) == 0; });
+					SoundQualityNamePropHandle->SetValue(Index);
+				})
+			)
+		];
 }
 
 void FAudioPluginWidgetManager::OnPluginSelected(FString PluginName, TSharedPtr<IPropertyHandle> PropertyHandle)
diff --git a/Engine/Source/Developer/IOS/IOSTargetPlatform/Private/IOSTargetPlatform.cpp b/Engine/Source/Developer/IOS/IOSTargetPlatform/Private/IOSTargetPlatform.cpp
--- a/Engine/Source/Developer/IOS/IOSTargetPlatform/Private/IOSTargetPlatform.cpp
+++ b/Engine/Source/Developer/IOS/IOSTargetPlatform/Private/IOSTargetPlatform.cpp
@@ -722,6 +722,12 @@ namespace IOS
 	{
 		const TCHAR* CategoryName = TEXT("/Script/IOSRuntimeSettings.IOSRuntimeSettings");
 
+		int32 SoundCueQualityIndex = INDEX_NONE;
+		if (GConfig->GetInt(CategoryName, TEXT("SoundCueCookQualityIndex"), SoundCueQualityIndex, GEngineIni))
+		{
+			OutOverrides.SoundCueCookQualityIndex = SoundCueQualityIndex;
+		}
+
 		GConfig->GetBool(CategoryName, TEXT("bUseAudioStreamCaching"), OutOverrides.bUseStreamCaching, GEngineIni);
 
 		/** Memory Load On Demand Settings */
diff --git a/Engine/Source/Developer/Linux/LinuxTargetPlatform/Classes/LinuxTargetSettings.h b/Engine/Source/Developer/Linux/LinuxTargetPlatform/Classes/LinuxTargetSettings.h
--- a/Engine/Source/Developer/Linux/LinuxTargetPlatform/Classes/LinuxTargetSettings.h
+++ b/Engine/Source/Developer/Linux/LinuxTargetPlatform/Classes/LinuxTargetSettings.h
@@ -35,6 +35,10 @@ public:
 	UPROPERTY(config, EditAnywhere, Category = "Audio")
 	FString OcclusionPlugin;
 
+	/** Quality Level to COOK SoundCues at (if set, all other levels will be stripped by the cooker). */
+	UPROPERTY(config, EditAnywhere, Category = "Audio|CookOverrides", meta = (DisplayName = "Sound Cue Cook Quality"))
+	int32 SoundCueCookQualityIndex = INDEX_NONE;
+
 	/**
 	 * The collection of RHI's we want to support on this platform.
 	 * This is not always the full list of RHI we can support.
diff --git a/Engine/Source/Developer/Mac/MacTargetPlatform/Classes/MacTargetSettings.h b/Engine/Source/Developer/Mac/MacTargetPlatform/Classes/MacTargetSettings.h
--- a/Engine/Source/Developer/Mac/MacTargetPlatform/Classes/MacTargetSettings.h
+++ b/Engine/Source/Developer/Mac/MacTargetPlatform/Classes/MacTargetSettings.h
@@ -149,4 +149,9 @@ public:
 	/** When set to anything beyond 0, this will ensure any SoundWaves longer than this value, in seconds, to stream directly off of the disk. */
 	UPROPERTY(GlobalConfig, EditAnywhere, Category = "Audio|CookOverrides", meta = (DisplayName = "Stream All Soundwaves Longer Than: "))
 	float AutoStreamingThreshold;
+	
+	/** Quality Level to COOK SoundCues at (if set, all other levels will be stripped by the cooker). */
+	UPROPERTY(GlobalConfig, EditAnywhere, Category = "Audio|CookOverrides", meta = (DisplayName = "Sound Cue Cook Quality"))
+	int32 SoundCueCookQualityIndex = INDEX_NONE;
+
 };
diff --git a/Engine/Source/Developer/Windows/WindowsTargetPlatform/Classes/WindowsTargetSettings.h b/Engine/Source/Developer/Windows/WindowsTargetPlatform/Classes/WindowsTargetSettings.h
--- a/Engine/Source/Developer/Windows/WindowsTargetPlatform/Classes/WindowsTargetSettings.h
+++ b/Engine/Source/Developer/Windows/WindowsTargetPlatform/Classes/WindowsTargetSettings.h
@@ -146,4 +146,8 @@ public:
 	/** When set to anything beyond 0, this will ensure any SoundWaves longer than this value, in seconds, to stream directly off of the disk. */
 	UPROPERTY(GlobalConfig, EditAnywhere, Category = "Audio|CookOverrides", meta = (DisplayName = "Stream All Soundwaves Longer Than: "))
 	float AutoStreamingThreshold;
+
+	/** Quality Level to COOK SoundCues at (if set, all other levels will be stripped by the cooker). */
+	UPROPERTY(GlobalConfig, EditAnywhere, Category = "Audio|CookOverrides", meta = (DisplayName = "Sound Cue Cook Quality"))
+	int32 SoundCueCookQualityIndex = INDEX_NONE;
 };
diff --git a/Engine/Source/Developer/Windows/WindowsTargetPlatform/Private/GenericWindowsTargetPlatform.h b/Engine/Source/Developer/Windows/WindowsTargetPlatform/Private/GenericWindowsTargetPlatform.h
--- a/Engine/Source/Developer/Windows/WindowsTargetPlatform/Private/GenericWindowsTargetPlatform.h
+++ b/Engine/Source/Developer/Windows/WindowsTargetPlatform/Private/GenericWindowsTargetPlatform.h
@@ -26,6 +26,12 @@ namespace Windows
 	{
 		const TCHAR* CategoryName = TEXT("/Script/WindowsTargetPlatform.WindowsTargetSettings");
 
+		int32 SoundCueQualityIndex = INDEX_NONE;
+		if (GConfig->GetInt(CategoryName, TEXT("SoundCueCookQualityIndex"), SoundCueQualityIndex, GEngineIni))
+		{
+			OutOverrides.SoundCueCookQualityIndex = SoundCueQualityIndex;
+		}
+
 		GConfig->GetBool(CategoryName, TEXT("bUseAudioStreamCaching"), OutOverrides.bUseStreamCaching, GEngineIni);
 
 		/** Memory Load On Demand Settings */
diff --git a/Engine/Source/Runtime/Android/AndroidRuntimeSettings/Classes/AndroidRuntimeSettings.h b/Engine/Source/Runtime/Android/AndroidRuntimeSettings/Classes/AndroidRuntimeSettings.h
--- a/Engine/Source/Runtime/Android/AndroidRuntimeSettings/Classes/AndroidRuntimeSettings.h
+++ b/Engine/Source/Runtime/Android/AndroidRuntimeSettings/Classes/AndroidRuntimeSettings.h
@@ -521,6 +521,10 @@ public:
 	UPROPERTY(config, EditAnywhere, Category = "Audio|CookOverrides")
 	bool bResampleForDevice;
 
+	/** Quality Level to COOK SoundCues at (if set, all other levels will be stripped by the cooker). */
+	UPROPERTY(GlobalConfig, EditAnywhere, Category = "Audio|CookOverrides", meta = (DisplayName = "Sound Cue Cook Quality"))
+	int32 SoundCueCookQualityIndex = INDEX_NONE;
+
 	// Mapping of which sample rates are used for each sample rate quality for a specific platform.
 
 	UPROPERTY(config, EditAnywhere, Category = "Audio|CookOverrides|ResamplingQuality", meta = (DisplayName = "Max"))
diff --git a/Engine/Source/Runtime/AudioPlatformConfiguration/Public/AudioCompressionSettings.h b/Engine/Source/Runtime/AudioPlatformConfiguration/Public/AudioCompressionSettings.h
--- a/Engine/Source/Runtime/AudioPlatformConfiguration/Public/AudioCompressionSettings.h
+++ b/Engine/Source/Runtime/AudioPlatformConfiguration/Public/AudioCompressionSettings.h
@@ -56,6 +56,9 @@ struct FPlatformAudioCookOverrides
 	// Scales all compression qualities when cooking to this platform. For example, 0.5 will halve all compression qualities, and 1.0 will leave them unchanged.
 	float CompressionQualityModifier;
 
+	// If set, the cooker will keep only this level of quality
+	int32 SoundCueCookQualityIndex = INDEX_NONE;
+
 	// When set to any platform > 0.0, this will automatically set any USoundWave beyond this value to be streamed from disk.
 	// If StreamCaching is set to true, this will be used 
 	float AutoStreamingThreshold;
diff --git a/Engine/Source/Runtime/Engine/Classes/Sound/AudioSettings.h b/Engine/Source/Runtime/Engine/Classes/Sound/AudioSettings.h
--- a/Engine/Source/Runtime/Engine/Classes/Sound/AudioSettings.h
+++ b/Engine/Source/Runtime/Engine/Classes/Sound/AudioSettings.h
@@ -183,6 +183,9 @@ class ENGINE_API UAudioSettings : public UDeveloperSettings
 public:
 	// Get the quality level settings at the provided level index
 	const FAudioQualitySettings& GetQualityLevelSettings(int32 QualityLevel) const;
+	
+	// Get the quality name level for a given index
+	FString FindQualityNameByIndex(int32 Index) const;
 
 	// Get the total number of quality level settings
 	int32 GetQualityLevelSettingsNum() const;
diff --git a/Engine/Source/Runtime/Engine/Classes/Sound/SoundCue.h b/Engine/Source/Runtime/Engine/Classes/Sound/SoundCue.h
--- a/Engine/Source/Runtime/Engine/Classes/Sound/SoundCue.h
+++ b/Engine/Source/Runtime/Engine/Classes/Sound/SoundCue.h
@@ -110,7 +110,7 @@ protected:
 	float SubtitlePriority;
 
 private:
-	float MaxAudibleDistance;
+	float MaxAudibleDistance;	
 
 public:
 	/* Indicates whether attenuation should use the Attenuation Overrides or the Attenuation Settings asset */
@@ -122,6 +122,9 @@ public:
 	uint8 bExcludeFromRandomNodeBranchCulling : 1;
 
 private:
+	UPROPERTY()
+	int32 CookedQualityIndex = INDEX_NONE;
+	
 	/** Whether a sound has play when silent enabled (i.e. for a sound cue, if any sound wave player has it enabled). */
 	UPROPERTY()
 	uint8 bHasPlayWhenSilent : 1;
@@ -238,6 +241,9 @@ public:
 	static void StaticAudioQualityChanged(int32 NewQualityLevel);
 
 	FORCEINLINE static int32 GetCachedQualityLevel() { return CachedQualityLevel; }
+	
+	/** Set the Quality level that the Cue was cooked at, called by the SoundQualityNodes */
+	int32 GetCookedQualityIndex() const { return CookedQualityIndex; }
 
 	/** Call to cache any values which need to be computed from the sound cue graph. e.g. MaxDistance, Duration, etc. */
 	void CacheAggregateValues();
@@ -299,5 +305,7 @@ private:
 
 	/** Ptr to interface to sound cue editor operations. */
 	static TSharedPtr<ISoundCueAudioEditor> SoundCueAudioEditor;
+
+	FCriticalSection EditorOnlyCs;
 #endif // WITH_EDITOR
 };
diff --git a/Engine/Source/Runtime/Engine/Classes/Sound/SoundNodeAssetReferencer.h b/Engine/Source/Runtime/Engine/Classes/Sound/SoundNodeAssetReferencer.h
--- a/Engine/Source/Runtime/Engine/Classes/Sound/SoundNodeAssetReferencer.h
+++ b/Engine/Source/Runtime/Engine/Classes/Sound/SoundNodeAssetReferencer.h
@@ -20,7 +20,7 @@ public:
 	virtual void LoadAsset(bool bAddToRoot = false) PURE_VIRTUAL(USoundNodeAssetReferencer::LoadAsset,);
 	virtual void ClearAssetReferences() PURE_VIRTUAL(USoundNodeAssetReferencer::ClearAssetReferences, );
 
-	bool ShouldHardReferenceAsset() const;
+	bool ShouldHardReferenceAsset(const class ITargetPlatform*) const;
 
 #if WITH_EDITOR
 	virtual void PostEditImport() override;
diff --git a/Engine/Source/Runtime/Engine/Classes/Sound/SoundNodeQualityLevel.h b/Engine/Source/Runtime/Engine/Classes/Sound/SoundNodeQualityLevel.h
--- a/Engine/Source/Runtime/Engine/Classes/Sound/SoundNodeQualityLevel.h
+++ b/Engine/Source/Runtime/Engine/Classes/Sound/SoundNodeQualityLevel.h
@@ -24,13 +24,12 @@ class USoundNodeQualityLevel : public USoundNode
 
 public:
 
-#if WITH_EDITOR
 	//~ Begin UObject Interface
 	virtual void PostLoad() override;
 	//~ End UObject Interface
-#endif
 
 	//~ Begin USoundNode Interface.
+	virtual void Serialize(FArchive& Ar) override;	
 	virtual void ParseNodes( FAudioDevice* AudioDevice, const UPTRINT NodeWaveInstanceHash, FActiveSound& ActiveSound, const FSoundParseParameters& ParseParams, TArray<FWaveInstance*>& WaveInstances ) override;
 	virtual int32 GetMaxChildNodes() const override;
 	virtual int32 GetMinChildNodes() const override;
@@ -41,8 +40,15 @@ public:
 #endif
 	//~ End USoundNode Interface.
 
+	// A Property to indicate which quality this node was cooked with. (INDEX_NONE if not cooked, or unset).
+	UPROPERTY()
+	int32 CookedQualityLevelIndex = INDEX_NONE;
+
 #if WITH_EDITOR
 	void ReconcileNode(bool bReconstructNode);
+
+	// Critical section to protect Child nodes array while serializing
+	FCriticalSection EditorOnlyCs;
 #endif
 
 };
diff --git a/Engine/Source/Runtime/Engine/Private/AudioCompressionSettingsUtils.cpp b/Engine/Source/Runtime/Engine/Private/AudioCompressionSettingsUtils.cpp
--- a/Engine/Source/Runtime/Engine/Private/AudioCompressionSettingsUtils.cpp
+++ b/Engine/Source/Runtime/Engine/Private/AudioCompressionSettingsUtils.cpp
@@ -92,6 +92,12 @@ void CacheCurrentPlatformAudioCookOverrides(FPlatformAudioCookOverrides& OutOver
 	const TCHAR* CategoryName = TEXT("");
 #endif
 
+	int32 SoundCueCookQualityIndex = INDEX_NONE;
+	if (GConfig->GetInt(CategoryName, TEXT("SoundCueCookQualityIndex"), SoundCueCookQualityIndex, GEngineIni))
+	{
+		OutOverrides.SoundCueCookQualityIndex = SoundCueCookQualityIndex;
+	}
+	
 	GConfig->GetBool(CategoryName, TEXT("bResampleForDevice"), OutOverrides.bResampleForDevice, GEngineIni);
 
 	GConfig->GetFloat(CategoryName, TEXT("CompressionQualityModifier"), OutOverrides.CompressionQualityModifier, GEngineIni);
diff --git a/Engine/Source/Runtime/Engine/Private/AudioSettings.cpp b/Engine/Source/Runtime/Engine/Private/AudioSettings.cpp
--- a/Engine/Source/Runtime/Engine/Private/AudioSettings.cpp
+++ b/Engine/Source/Runtime/Engine/Private/AudioSettings.cpp
@@ -150,4 +150,11 @@ int32 UAudioSettings::GetHighestMaxChannels() const
 	return HighestMaxChannels;
 }
 
+FString UAudioSettings::FindQualityNameByIndex(int32 Index) const
+{
+	return QualityLevels.IsValidIndex(Index) ?
+		   QualityLevels[Index].DisplayName.ToString() :
+		   TEXT("");
+}
+
 #undef LOCTEXT_NAMESPACE
diff --git a/Engine/Source/Runtime/Engine/Private/SoundCue.cpp b/Engine/Source/Runtime/Engine/Private/SoundCue.cpp
--- a/Engine/Source/Runtime/Engine/Private/SoundCue.cpp
+++ b/Engine/Source/Runtime/Engine/Private/SoundCue.cpp
@@ -25,13 +25,16 @@
 #include "DSP/Dsp.h"
 #if WITH_EDITOR
 #include "Kismet2/BlueprintEditorUtils.h"
-#include "Sound/AudioSettings.h"
 #include "SoundCueGraph/SoundCueGraphNode.h"
 #include "SoundCueGraph/SoundCueGraph.h"
 #include "SoundCueGraph/SoundCueGraphNode_Root.h"
 #include "SoundCueGraph/SoundCueGraphSchema.h"
 #endif // WITH_EDITOR
 
+#include "Interfaces/ITargetPlatform.h"
+#include "AudioCompressionSettings.h"
+#include "Sound/AudioSettings.h"
+
 /*-----------------------------------------------------------------------------
 	USoundCue implementation.
 -----------------------------------------------------------------------------*/
@@ -48,6 +51,7 @@ USoundCue::USoundCue(const FObjectInitializer& ObjectInitializer)
 	VolumeMultiplier = 0.75f;
 	PitchMultiplier = 1.0f;
 	SubtitlePriority = DEFAULT_SUBTITLE_PRIORITY;
+	CookedQualityIndex = INDEX_NONE;
 }
 
 #if WITH_EDITOR
@@ -106,7 +110,23 @@ void USoundCue::Serialize(FStructuredArchive::FRecord Record)
 		CacheAggregateValues();
 	}
 
-	Super::Serialize(Record);
+#if WITH_EDITOR
+	// If we are cooking, record our cooked quality before serialize and then undo it.
+	if (UnderlyingArchive.IsCooking() && UnderlyingArchive.IsSaving() && UnderlyingArchive.CookingTarget())
+	{
+		if (const FPlatformAudioCookOverrides* AudioCookOverrides = UnderlyingArchive.CookingTarget()->GetAudioCompressionSettings())
+		{
+			FScopeLock Lock(&EditorOnlyCs);
+			CookedQualityIndex = AudioCookOverrides->SoundCueCookQualityIndex;
+			Super::Serialize(Record);
+			CookedQualityIndex = INDEX_NONE;
+		}
+	}
+	else
+#endif //WITH_EDITOR
+	{
+		Super::Serialize(Record);
+	}
 
 	if (UnderlyingArchive.UE4Ver() >= VER_UE4_COOKED_ASSETS_IN_EDITOR_SUPPORT)
 	{
@@ -151,6 +171,17 @@ void USoundCue::PostLoad()
 	}
 	else
 #endif // WITH_EDITOR
+
+	// Warn if the Quality index is set to something that we can't support.
+	UE_CLOG(USoundCue::GetCachedQualityLevel() != CookedQualityIndex && CookedQualityIndex != INDEX_NONE, LogAudio, Warning,
+		TEXT("'%s' is ingoring Quality Setting '%s'(%d) as it was cooked with '%s'(%d)"),
+		*GetFullNameSafe(this),
+		*GetDefault<UAudioSettings>()->FindQualityNameByIndex(USoundCue::GetCachedQualityLevel()),
+		USoundCue::GetCachedQualityLevel(),
+		*GetDefault<UAudioSettings>()->FindQualityNameByIndex(CookedQualityIndex),
+		CookedQualityIndex
+	);
+
 	if (GEngine && *GEngine->GameUserSettingsClass)
 	{
 		EvaluateNodes(false);
@@ -224,7 +255,14 @@ void USoundCue::EvaluateNodes(bool bAddToRoot)
 		}
 	};
 
-	EvaluateNodes_Internal(FirstNode);
+	// Only Evaluate nodes if we haven't been cooked, as cooked builds will hard-ref all SoundAssetReferences.	
+	UE_CLOG(CookedQualityIndex == INDEX_NONE, LogAudio, Verbose, TEXT("'%s', DOING EvaluateNodes as we are *NOT* cooked"), *GetName());
+	UE_CLOG(CookedQualityIndex != INDEX_NONE, LogAudio, Verbose, TEXT("'%s', SKIPPING EvaluateNodes as we *ARE* cooked"), *GetName());
+
+	if (CookedQualityIndex == INDEX_NONE)
+	{		
+		EvaluateNodes_Internal(FirstNode);
+	}
 }
 
 float USoundCue::FindMaxDistanceInternal() const
diff --git a/Engine/Source/Runtime/Engine/Private/SoundNodeAssetReferencer.cpp b/Engine/Source/Runtime/Engine/Private/SoundNodeAssetReferencer.cpp
--- a/Engine/Source/Runtime/Engine/Private/SoundNodeAssetReferencer.cpp
+++ b/Engine/Source/Runtime/Engine/Private/SoundNodeAssetReferencer.cpp
@@ -5,11 +5,31 @@
 #include "Sound/SoundNodeRandom.h"
 #include "Sound/SoundCue.h"
 #include "UObject/FrameworkObjectVersion.h"
+#include "Interfaces/ITargetPlatform.h"
+#include "AudioCompressionSettings.h"
+#include "Sound/AudioSettings.h"
 
 #define ASYNC_LOAD_RANDOMIZED_SOUNDS 1
 
-bool USoundNodeAssetReferencer::ShouldHardReferenceAsset() const
+bool USoundNodeAssetReferencer::ShouldHardReferenceAsset(const ITargetPlatform* TargetPlatform) const
 {
+	if (TargetPlatform)
+	{
+		if (const FPlatformAudioCookOverrides* CookOverrides = TargetPlatform->GetAudioCompressionSettings())
+		{			
+			// If the Quality nodes are cooked, everything is hard refs.
+			if (CookOverrides->SoundCueCookQualityIndex != INDEX_NONE)
+			{
+				UE_LOG(LogAudio, Verbose, TEXT("HARD reffing '%s:%s', as we are cooking using quality '%s'"), 
+					*GetNameSafe(GetOuter()),
+					*GetName(),
+					*GetDefault<UAudioSettings>()->FindQualityNameByIndex(CookOverrides->SoundCueCookQualityIndex)
+				)
+				return true;
+			}
+		}
+	}
+
 	bool bShouldHardReference = true;
 
 	if (USoundCue* Cue = Cast<USoundCue>(GetOuter()))
@@ -52,6 +72,12 @@ bool USoundNodeAssetReferencer::ShouldHardReferenceAsset() const
 		
 	}
 
+	UE_LOG(LogAudio, Verbose, TEXT("%s reffing '%s:%s'."),
+		bShouldHardReference ? TEXT("HARD") : TEXT("SOFT"),
+		*GetNameSafe(GetOuter()),
+		*GetName()
+	);
+
 	return bShouldHardReference;
 }
 
diff --git a/Engine/Source/Runtime/Engine/Private/SoundNodeQualityLevel.cpp b/Engine/Source/Runtime/Engine/Private/SoundNodeQualityLevel.cpp
--- a/Engine/Source/Runtime/Engine/Private/SoundNodeQualityLevel.cpp
+++ b/Engine/Source/Runtime/Engine/Private/SoundNodeQualityLevel.cpp
@@ -13,14 +13,9 @@
 #include "Settings/LevelEditorPlaySettings.h"
 #include "Editor.h"
 #endif
+#include "Interfaces/ITargetPlatform.h"
 
 #if WITH_EDITOR
-void USoundNodeQualityLevel::PostLoad()
-{
-	Super::PostLoad();
-
-	ReconcileNode(false);
-}
 
 void USoundNodeQualityLevel::ReconcileNode(bool bReconstructNode)
 {
@@ -47,6 +42,24 @@ FText USoundNodeQualityLevel::GetInputPinName(int32 PinIndex) const
 }
 #endif
 
+void USoundNodeQualityLevel::PostLoad()
+{
+	Super::PostLoad();
+#if WITH_EDITOR
+	ReconcileNode(false);
+#endif //WITH_EDITOR
+
+	UE_CLOG(Cast<USoundCue>(GetOuter()) && Cast<USoundCue>(GetOuter())->GetCookedQualityIndex() != CookedQualityLevelIndex && CookedQualityLevelIndex != INDEX_NONE,
+		LogAudio,
+		Warning, TEXT("'%s' has been cooked with multiple quality levels. '%s'(%d) vs '%s'(%d)"),
+		*GetFullNameSafe(this),
+		*GetDefault<UAudioSettings>()->FindQualityNameByIndex(USoundCue::GetCachedQualityLevel()),
+		USoundCue::GetCachedQualityLevel(),
+		*GetDefault<UAudioSettings>()->FindQualityNameByIndex(CookedQualityLevelIndex),
+		CookedQualityLevelIndex
+	);
+}
+
 void USoundNodeQualityLevel::PrimeChildWavePlayers(bool bRecurse)
 {
 	// If we're able to retrieve a valid cached quality level for this sound cue,
@@ -102,11 +115,62 @@ void USoundNodeQualityLevel::ParseNodes( FAudioDevice* AudioDevice, const UPTRIN
 		QualityLevel = USoundCue::GetCachedQualityLevel();
 	}
 #else
-	const int32 QualityLevel = USoundCue::GetCachedQualityLevel();
-#endif
-
+	
+	int32 QualityLevel = USoundCue::GetCachedQualityLevel();
+	
+	// If CookedQualityLevelIndex has been set, we will have a *single* quality level.
+	if (CookedQualityLevelIndex >= 0)
+	{	
+		// Remap to index 0 (as all other levels have been removed by cooker).
+		QualityLevel = 0;
+	}
+#endif	
+	
 	if (ChildNodes.IsValidIndex(QualityLevel) && ChildNodes[QualityLevel])
 	{
 		ChildNodes[QualityLevel]->ParseNodes( AudioDevice, GetNodeWaveInstanceHash(NodeWaveInstanceHash, ChildNodes[QualityLevel], QualityLevel), ActiveSound, ParseParams, WaveInstances );
 	}
 }
+
+void USoundNodeQualityLevel::Serialize(FArchive& Ar)
+{
+#if WITH_EDITOR
+
+	if (Ar.IsCooking() && Ar.IsSaving() && Ar.CookingTarget() )
+	{			
+		if (const FPlatformAudioCookOverrides* AudioCookOverrides = Ar.CookingTarget()->GetAudioCompressionSettings())
+		{
+			// Prevent any other thread saving this class while we are modifying the ChildNode array.
+			FScopeLock Lock(&EditorOnlyCs);
+			
+			// Are we doing any cook quality level optimization for this Node?			
+			if (ChildNodes.IsValidIndex(AudioCookOverrides->SoundCueCookQualityIndex))
+			{
+				// Set our cook quality, as we serialize. 
+				CookedQualityLevelIndex = AudioCookOverrides->SoundCueCookQualityIndex;
+
+				// Move out all nodes.
+				TArray<USoundNode*> ChildNodesBackup;
+				ChildNodesBackup = MoveTemp(ChildNodes);
+				
+				// Put *just* the node we care about in our child array to be serialized by the Super
+				ChildNodes.Add(ChildNodesBackup[CookedQualityLevelIndex]);
+
+				// Call base serialize that will walk all properties and serialize them.
+				Super::Serialize(Ar);
+
+				// Return to our original state. (careful the cook only variables don't leak out).
+				ChildNodes = MoveTemp(ChildNodesBackup);
+				CookedQualityLevelIndex = INDEX_NONE;
+
+				// We are done.
+				return;
+			}	
+		}		
+	}
+		
+#endif //WITH_EDITOR	
+
+	// ... in all other cases, we just call the super.
+	Super::Serialize(Ar);
+}
\ No newline at end of file
diff --git a/Engine/Source/Runtime/Engine/Private/SoundNodeWavePlayer.cpp b/Engine/Source/Runtime/Engine/Private/SoundNodeWavePlayer.cpp
--- a/Engine/Source/Runtime/Engine/Private/SoundNodeWavePlayer.cpp
+++ b/Engine/Source/Runtime/Engine/Private/SoundNodeWavePlayer.cpp
@@ -22,7 +22,7 @@ void USoundNodeWavePlayer::Serialize(FArchive& Ar)
 		}
 		else if (Ar.IsSaving())
 		{
-			USoundWave* HardReference = (ShouldHardReferenceAsset() ? SoundWave : nullptr);
+			USoundWave* HardReference = (ShouldHardReferenceAsset(Ar.CookingTarget()) ? SoundWave : nullptr);
 			Ar << HardReference;
 		}
 	}
@@ -45,6 +45,7 @@ void USoundNodeWavePlayer::LoadAsset(bool bAddToRoot)
 			const FString LongPackageName = SoundWaveAssetPtr.GetLongPackageName();
 			if (!LongPackageName.IsEmpty())
 			{
+				UE_LOG(LogAudio, VeryVerbose, TEXT(" '%s:%s', Async loading..."), *GetNameSafe(GetOuter()), *GetName());
 				bAsyncLoading = true;
 				LoadPackageAsync(LongPackageName, FLoadPackageAsyncDelegate::CreateUObject(this, &USoundNodeWavePlayer::OnSoundWaveLoaded, bAddToRoot));
 			}
@@ -60,6 +61,12 @@ void USoundNodeWavePlayer::LoadAsset(bool bAddToRoot)
 	}
 	else
 	{
+		UE_LOG(LogAudio, VeryVerbose, TEXT("'%s:%s', DOING SYNCRONOUS loading... %s"),
+			*GetNameSafe(GetOuter()),
+			*GetName(),
+			SoundWaveAssetPtr.IsPending() ? TEXT("pending.") : TEXT("not-pending.")
+		);
+
 		SoundWave = SoundWaveAssetPtr.LoadSynchronous();
 		if (SoundWave)
 		{
diff --git a/Engine/Source/Runtime/IOS/IOSRuntimeSettings/Classes/IOSRuntimeSettings.h b/Engine/Source/Runtime/IOS/IOSRuntimeSettings/Classes/IOSRuntimeSettings.h
--- a/Engine/Source/Runtime/IOS/IOSRuntimeSettings/Classes/IOSRuntimeSettings.h
+++ b/Engine/Source/Runtime/IOS/IOSRuntimeSettings/Classes/IOSRuntimeSettings.h
@@ -526,6 +526,10 @@ public:
 	UPROPERTY(config, EditAnywhere, Category = "Audio|CookOverrides")
 	bool bResampleForDevice;
 
+	/** Quality Level to COOK SoundCues at (if set, all other levels will be stripped by the cooker). */
+	UPROPERTY(GlobalConfig, EditAnywhere, Category = "Audio|CookOverrides", meta = (DisplayName = "Sound Cue Cook Quality"))
+	int32 SoundCueCookQualityIndex = INDEX_NONE;
+
 	// Mapping of which sample rates are used for each sample rate quality for a specific platform.
 
 	UPROPERTY(config, EditAnywhere, Category = "Audio|CookOverrides|ResamplingQuality", meta = (DisplayName = "Max"))
diff --git a/Engine/Source/Runtime/Lumin/LuminRuntimeSettings/Classes/LuminRuntimeSettings.h b/Engine/Source/Runtime/Lumin/LuminRuntimeSettings/Classes/LuminRuntimeSettings.h
--- a/Engine/Source/Runtime/Lumin/LuminRuntimeSettings/Classes/LuminRuntimeSettings.h
+++ b/Engine/Source/Runtime/Lumin/LuminRuntimeSettings/Classes/LuminRuntimeSettings.h
@@ -150,6 +150,10 @@ public:
 	UPROPERTY(config, EditAnywhere, Category = "Audio")
 	FString OcclusionPlugin;
 
+	/** Quality Level to COOK SoundCues at (if set, all other levels will be stripped by the cooker). */
+	UPROPERTY(GlobalConfig, EditAnywhere, Category = "Audio|CookOverrides", meta = (DisplayName = "Sound Cue Cook Quality"))
+	int32 SoundCueCookQualityIndex = INDEX_NONE;
+
 	// Strip debug symbols from packaged builds even if they aren't shipping builds.
 	UPROPERTY(GlobalConfig, EditAnywhere, Category = AdvancedBuild, meta = (DisplayName = "Strip debug symbols from packaged builds even if they aren't shipping builds"))
 	bool bRemoveDebugInfo;
