Add AudioDefines.h to AudioDevice

diff --git a/Engine/Source/Runtime/Engine/Private/AudioDevice.cpp b/Engine/Source/Runtime/Engine/Private/AudioDevice.cpp
--- a/Engine/Source/Runtime/Engine/Private/AudioDevice.cpp
+++ b/Engine/Source/Runtime/Engine/Private/AudioDevice.cpp
@@ -4,6 +4,7 @@
 #include "ActiveSound.h"
 #include "AudioCompressionSettingsUtils.h"
 #include "AudioDecompress.h"
+#include "AudioDefines.h"
 #include "AudioEffect.h"
 #include "AudioPluginUtilities.h"
 #include "Audio/AudioDebug.h"
