Files that are either identical from Fury but with unresolved integrations, or no significant changes (whitespace, include order)

diff --git a/Engine/Source/Runtime/Engine/Private/Animation/AnimInstanceProxy.cpp b/Engine/Source/Runtime/Engine/Private/Animation/AnimInstanceProxy.cpp
--- a/Engine/Source/Runtime/Engine/Private/Animation/AnimInstanceProxy.cpp
+++ b/Engine/Source/Runtime/Engine/Private/Animation/AnimInstanceProxy.cpp
@@ -340,9 +340,10 @@ void FAnimInstanceProxy::PreUpdate(UAnimInstance* InAnimInstance, float DeltaSec
 
 	USkeletalMeshComponent* SkelMeshComp = InAnimInstance->GetSkelMeshComponent();
 	UWorld* World = SkelMeshComp ? SkelMeshComp->GetWorld() : nullptr;
+	AWorldSettings* WorldSettings = World ? World->GetWorldSettings() : nullptr;
 
 	CurrentDeltaSeconds = DeltaSeconds;
-	CurrentTimeDilation = World ? World->GetWorldSettings()->GetEffectiveTimeDilation() : 1.0f;
+	CurrentTimeDilation = WorldSettings ? WorldSettings->GetEffectiveTimeDilation() : 1.0f;
 	RootMotionMode = InAnimInstance->RootMotionMode;
 	bShouldExtractRootMotion = InAnimInstance->ShouldExtractRootMotion();
 
diff --git a/Engine/Source/Runtime/Engine/Private/Components/SkinnedMeshComponent.cpp b/Engine/Source/Runtime/Engine/Private/Components/SkinnedMeshComponent.cpp
--- a/Engine/Source/Runtime/Engine/Private/Components/SkinnedMeshComponent.cpp
+++ b/Engine/Source/Runtime/Engine/Private/Components/SkinnedMeshComponent.cpp
@@ -1356,7 +1356,9 @@ bool USkinnedMeshComponent::GetTwistAndSwingAngleOfDeltaRotationFromRefPose(FNam
 {
 	const FReferenceSkeleton& RefSkeleton = SkeletalMesh->RefSkeleton;
 	const int32 BoneIndex = GetBoneIndex(BoneName);
-	if (BoneIndex != INDEX_NONE)
+	const TArray<FTransform>& Transforms = GetComponentSpaceTransforms();
+
+	if (BoneIndex != INDEX_NONE && ensureMsgf(BoneIndex < Transforms.Num(), TEXT("Invalid transform access in %s. Index=%d, Num=%d"), *GetPathName(), BoneIndex, Transforms.Num()))
 	{
 		FTransform LocalTransform = GetComponentSpaceTransforms()[BoneIndex];
 		FTransform ReferenceTransform = RefSkeleton.GetRefBonePose()[BoneIndex];
diff --git a/Engine/Source/Runtime/Engine/Private/PhysicsEngine/ImmediatePhysics/ImmediatePhysicsPhysX/ImmediatePhysicsContactPointRecorder.cpp b/Engine/Source/Runtime/Engine/Private/PhysicsEngine/ImmediatePhysics/ImmediatePhysicsPhysX/ImmediatePhysicsContactPointRecorder.cpp
--- a/Engine/Source/Runtime/Engine/Private/PhysicsEngine/ImmediatePhysics/ImmediatePhysicsPhysX/ImmediatePhysicsContactPointRecorder.cpp
+++ b/Engine/Source/Runtime/Engine/Private/PhysicsEngine/ImmediatePhysics/ImmediatePhysicsPhysX/ImmediatePhysicsContactPointRecorder.cpp
@@ -5,6 +5,7 @@
 #if WITH_PHYSX && PHYSICS_INTERFACE_PHYSX
 
 #include "Physics/ImmediatePhysics/ImmediatePhysicsPhysX/ImmediatePhysicsSimulation_PhysX.h"
+#include "Physics/ImmediatePhysics/ImmediatePhysicsPhysX/ImmediatePhysicsContactPair_PhysX.h"
 
 namespace ImmediatePhysics_PhysX
 {
