Fix crash in FScopedSceneLock_Chaos when given a nullptr actorhandle.

diff --git a/Engine/Source/Runtime/Engine/Private/PhysicsEngine/Experimental/PhysInterface_Chaos.cpp b/Engine/Source/Runtime/Engine/Private/PhysicsEngine/Experimental/PhysInterface_Chaos.cpp
--- a/Engine/Source/Runtime/Engine/Private/PhysicsEngine/Experimental/PhysInterface_Chaos.cpp
+++ b/Engine/Source/Runtime/Engine/Private/PhysicsEngine/Experimental/PhysInterface_Chaos.cpp
@@ -877,7 +877,7 @@ private:
 
 	FPhysScene_ChaosInterface* GetSceneForActor(FPhysicsActorHandle const * InActorHandle)
 	{
-		FBodyInstance* ActorInstance = FPhysicsUserData_Chaos::Get<FBodyInstance>((*InActorHandle)->UserData());
+		FBodyInstance* ActorInstance = (InActorHandle) ? FPhysicsUserData_Chaos::Get<FBodyInstance>((*InActorHandle)->UserData()) : nullptr;
 
 		if(ActorInstance)
 		{
