Handle case where chaos TConvex has no geometry. Using an ensure highlights problem while still allowing things to continue.

diff --git a/Engine/Source/Runtime/Engine/Private/PhysicsEngine/BodySetup.cpp b/Engine/Source/Runtime/Engine/Private/PhysicsEngine/BodySetup.cpp
--- a/Engine/Source/Runtime/Engine/Private/PhysicsEngine/BodySetup.cpp
+++ b/Engine/Source/Runtime/Engine/Private/PhysicsEngine/BodySetup.cpp
@@ -695,12 +695,20 @@ void UBodySetup::FinishCreatingPhysicsMeshes_Chaos(FChaosDerivedDataReader<float
 		for (int32 ElementIndex = 0; ElementIndex < AggGeom.ConvexElems.Num(); ElementIndex++)
 		{
 			FKConvexElem& ConvexElem = AggGeom.ConvexElems[ElementIndex];
-			ConvexElem.SetChaosConvexMesh(MoveTemp(InReader.ConvexImplicitObjects[ElementIndex]));
 
-			if (ConvexElem.GetChaosConvexMesh()->IsPerformanceWarning())
+			if (ensure(InReader.ConvexImplicitObjects[ElementIndex]->IsValidGeometry()))
 			{
-				const FString& PerformanceString = ConvexElem.GetChaosConvexMesh()->PerformanceWarningAndSimplifaction();
-				UE_LOG(LogPhysics, Warning, TEXT("TConvex Name:%s, Element [%d], %s"), FullName.GetCharArray().GetData(), ElementIndex, PerformanceString.GetCharArray().GetData());
+				ConvexElem.SetChaosConvexMesh(MoveTemp(InReader.ConvexImplicitObjects[ElementIndex]));
+
+				if (ConvexElem.GetChaosConvexMesh()->IsPerformanceWarning())
+				{
+					const FString& PerformanceString = ConvexElem.GetChaosConvexMesh()->PerformanceWarningAndSimplifaction();
+					UE_LOG(LogPhysics, Warning, TEXT("TConvex Name:%s, Element [%d], %s"), FullName.GetCharArray().GetData(), ElementIndex, PerformanceString.GetCharArray().GetData());
+				}
+			}
+			else
+			{
+				UE_LOG(LogPhysics, Warning, TEXT("TConvex Name:%s, Element [%d] has no Geometry"), FullName.GetCharArray().GetData(), ElementIndex);
 			}
 
 		}
diff --git a/Engine/Source/Runtime/Experimental/Chaos/Public/Chaos/Convex.h b/Engine/Source/Runtime/Experimental/Chaos/Public/Chaos/Convex.h
--- a/Engine/Source/Runtime/Experimental/Chaos/Public/Chaos/Convex.h
+++ b/Engine/Source/Runtime/Experimental/Chaos/Public/Chaos/Convex.h
@@ -227,6 +227,11 @@ namespace Chaos
 			SerializeImp(Ar);
 		}
 
+		virtual bool IsValidGeometry() const override
+		{
+			return (SurfaceParticles.Size() > 0 && Planes.Num() > 0);
+		}
+
 		virtual bool IsPerformanceWarning() const override
 		{
 			return TConvexBuilder<T>::IsPerformanceWarning(Planes.Num(), SurfaceParticles.Size());
