Chaos RBAN - Kinematic Targets

diff --git a/Engine/Source/Runtime/AnimGraphRuntime/Private/BoneControllers/AnimNode_RigidBody.cpp b/Engine/Source/Runtime/AnimGraphRuntime/Private/BoneControllers/AnimNode_RigidBody.cpp
--- a/Engine/Source/Runtime/AnimGraphRuntime/Private/BoneControllers/AnimNode_RigidBody.cpp
+++ b/Engine/Source/Runtime/AnimGraphRuntime/Private/BoneControllers/AnimNode_RigidBody.cpp
@@ -26,6 +26,9 @@ DEFINE_STAT(STAT_RigidBodyNodeInitTime);
 TAutoConsoleVariable<int32> CVarEnableRigidBodyNode(TEXT("p.RigidBodyNode"), 1, TEXT("Enables/disables rigid body node updates and evaluations"), ECVF_ReadOnly);
 TAutoConsoleVariable<int32> CVarRigidBodyLODThreshold(TEXT("p.RigidBodyLODThreshold"), -1, TEXT("Max LOD that rigid body node is allowed to run on. Provides a global threshold that overrides per-node the LODThreshold property. -1 means no override."), ECVF_Scalability);
 
+int32 RBAN_MaxSubSteps = 4;
+FAutoConsoleVariableRef CVarRigidBodyNodeMaxSteps(TEXT("p.RigidBodyNode.MaxSubSteps"), RBAN_MaxSubSteps, TEXT("Set the maximum number of simulation steps in the update loop"), ECVF_ReadOnly);
+
 FAnimNode_RigidBody::FAnimNode_RigidBody():
 	QueryParams(NAME_None, FCollisionQueryParams::GetUnknownStatId())
 {
@@ -455,6 +458,7 @@ void FAnimNode_RigidBody::EvaluateSkeletalControl_AnyThread(FComponentSpacePoseC
 				}
 			}
 
+			// @todo(ccaulfield): We should be interpolating kinematic targets for each sub-step below
 			for (const FOutputBoneData& OutputData : OutputBoneData)
 			{
 				const int32 BodyIndex = OutputData.BodyIndex;
@@ -472,11 +476,12 @@ void FAnimNode_RigidBody::EvaluateSkeletalControl_AnyThread(FComponentSpacePoseC
 
 			// Run simulation at a minimum of 30 FPS to prevent system from exploding.
 			// DeltaTime can be higher due to URO, so take multiple iterations in that case.
+			const int32 MaxSteps = RBAN_MaxSubSteps;
 			const float MaxDeltaSeconds = 1.f / 30.f;
-			const int32 NumIterations = FMath::Clamp(FMath::CeilToInt(DeltaSeconds / MaxDeltaSeconds), 1, 4);
-			const float StepDeltaTime = DeltaSeconds / float(NumIterations);
+			const int32 NumSteps = FMath::Clamp(FMath::CeilToInt(DeltaSeconds / MaxDeltaSeconds), 1, MaxSteps);
+			const float StepDeltaTime = DeltaSeconds / float(NumSteps);
 
-			for (int32 Step = 1; Step <= NumIterations; Step++)
+			for (int32 Step = 1; Step <= NumSteps; Step++)
 			{
 				// We call the _AssumesLocked version here without a lock as the simulation is local to this node and we know
 				// we're not going to alter anything while this is running.
diff --git a/Engine/Source/Runtime/Engine/Private/PhysicsEngine/ImmediatePhysics/ImmediatePhysicsChaos/ImmediatePhysicsActorHandle_Chaos.cpp b/Engine/Source/Runtime/Engine/Private/PhysicsEngine/ImmediatePhysics/ImmediatePhysicsChaos/ImmediatePhysicsActorHandle_Chaos.cpp
--- a/Engine/Source/Runtime/Engine/Private/PhysicsEngine/ImmediatePhysics/ImmediatePhysicsChaos/ImmediatePhysicsActorHandle_Chaos.cpp
+++ b/Engine/Source/Runtime/Engine/Private/PhysicsEngine/ImmediatePhysics/ImmediatePhysicsChaos/ImmediatePhysicsActorHandle_Chaos.cpp
@@ -308,6 +308,7 @@ namespace ImmediatePhysics_Chaos
 	void FActorHandle::SetIsKinematic(bool bKinematic)
 	{
 #if IMMEDIATEPHYSICS_CHAOS_TODO
+		// This needs to destroy and recreate the particle
 #endif
 	}
 
@@ -316,44 +317,40 @@ namespace ImmediatePhysics_Chaos
 		return Handle()->IsKinematic();
 	}
 
+	const FKinematicTarget& FActorHandle::GetKinematicTarget() const
+	{
+		check(ParticleHandle->AsKinematic());
+		return ParticleHandle->AsKinematic()->KinematicTarget();
+	}
 
-#if IMMEDIATEPHYSICS_CHAOS_TODO
-	FImmediateKinematicTarget& FActorHandle::GetKinematicTarget()
+	FKinematicTarget& FActorHandle::GetKinematicTarget()
 	{
+		check(ParticleHandle->AsKinematic());
+		return ParticleHandle->AsKinematic()->KinematicTarget();
 	}
-#endif
 
 	void FActorHandle::SetKinematicTarget(const FTransform& WorldTM)
 	{
-		using namespace Chaos;
-
-#if IMMEDIATEPHYSICS_CHAOS_TODO
-		// Proper kinematic targets
-#endif
-		if (TKinematicGeometryParticleHandle<FReal, Dimensions>* KinematicParticleHandle = ParticleHandle->AsKinematic())
+		if (ensure(GetIsKinematic()))
 		{
 			FTransform ParticleTransform = ActorToCoMTransform * WorldTM;
-
-			FMatrix AxesW = FRotationMatrix::Make(WorldTM.GetRotation());
-			FMatrix AxesP = FRotationMatrix::Make(ParticleTransform.GetRotation());
-
-			ParticleHandle->SetX(ParticleTransform.GetTranslation());
-			ParticleHandle->SetR(ParticleTransform.GetRotation());
+			GetKinematicTarget().SetTargetMode(ParticleTransform);
 		}
+
 	}
 
 	bool FActorHandle::HasKinematicTarget() const
 	{
-#if IMMEDIATEPHYSICS_CHAOS_TODO
-#endif
+		if (GetIsKinematic())
+		{
+			return GetKinematicTarget().GetMode() == Chaos::EKinematicTargetMode::Position;
+		}
 		return false;
 	}
 
 	bool FActorHandle::IsSimulated() const
 	{
-#if IMMEDIATEPHYSICS_CHAOS_TODO
-#endif
-		return true;
+		return ParticleHandle->AsDynamic() != nullptr;
 	}
 
 	FTransform FActorHandle::GetWorldTransform() const
diff --git a/Engine/Source/Runtime/Engine/Private/PhysicsEngine/ImmediatePhysics/ImmediatePhysicsChaos/ImmediatePhysicsSimulation_Chaos.cpp b/Engine/Source/Runtime/Engine/Private/PhysicsEngine/ImmediatePhysics/ImmediatePhysicsChaos/ImmediatePhysicsSimulation_Chaos.cpp
--- a/Engine/Source/Runtime/Engine/Private/PhysicsEngine/ImmediatePhysics/ImmediatePhysicsChaos/ImmediatePhysicsSimulation_Chaos.cpp
+++ b/Engine/Source/Runtime/Engine/Private/PhysicsEngine/ImmediatePhysics/ImmediatePhysicsChaos/ImmediatePhysicsSimulation_Chaos.cpp
@@ -21,7 +21,7 @@
 float ChaosImmediate_Evolution_DeltaTime = 0.03f;
 int32 ChaosImmediate_Evolution_Iterations = 10;
 int32 ChaosImmediate_Collision_Enabled = 1;
-int32 ChaosImmediate_Collision_ApplyEnabled = 0;
+int32 ChaosImmediate_Collision_ApplyEnabled = 1;
 int32 ChaosImmediate_Collision_PushOutIterations = 5;
 int32 ChaosImmediate_Collision_PushOutPairIterations = 2;
 float ChaosImmediate_Collision_Thickness = 0;
@@ -68,10 +68,10 @@ int32 ChaosImmediate_DebugDrawCollisions = 0;
 int32 ChaosImmediate_DebugDrawJoints = 0;
 #endif
 
-FAutoConsoleVariableRef CVarDebugDrawParticles(TEXT("p.Chaos.ImmPhys.DebugDrawParticles"), ChaosImmediate_DebugDrawParticles, TEXT("Draw Particle Transforms (0 = never; 1 = end of frame; 2 = begin and end of frame; 3 = each evolution phase)."));
-FAutoConsoleVariableRef CVarDebugDrawShapes(TEXT("p.Chaos.ImmPhys.DebugDrawShapes"), ChaosImmediate_DebugDrawShapes, TEXT("Draw Shapes (0 = never; 1 = end of frame; 2 = begin and end of frame; 3 = each evolution phase."));
-FAutoConsoleVariableRef CVarDebugDrawCollisions(TEXT("p.Chaos.ImmPhys.DebugDrawCollisions"), ChaosImmediate_DebugDrawCollisions, TEXT("Draw Collisions (0 = never; 1 = end of frame; 2 = begin and end of frame; 3 = each evolution phase)"));
-FAutoConsoleVariableRef CVarDebugDrawJoints(TEXT("p.Chaos.ImmPhys.DebugDrawJoints"), ChaosImmediate_DebugDrawJoints, TEXT("Draw Joints. (0 = never; 1 = end of frame; 2 = begin and end of frame; 3 = each evolution phase; 4 = each Apply step)."));
+FAutoConsoleVariableRef CVarDebugDrawParticles(TEXT("p.Chaos.ImmPhys.DebugDrawParticles"), ChaosImmediate_DebugDrawParticles, TEXT("Draw Particle Transforms (0 = never; 1 = end of frame; 2 = begin and end of frame; 3 = post-integate, post-apply and post-applypushout;)."));
+FAutoConsoleVariableRef CVarDebugDrawShapes(TEXT("p.Chaos.ImmPhys.DebugDrawShapes"), ChaosImmediate_DebugDrawShapes, TEXT("Draw Shapes (0 = never; 1 = end of frame; 2 = begin and end of frame; 3 = post-integate, post-apply and post-applypushout;"));
+FAutoConsoleVariableRef CVarDebugDrawCollisions(TEXT("p.Chaos.ImmPhys.DebugDrawCollisions"), ChaosImmediate_DebugDrawCollisions, TEXT("Draw Collisions (0 = never; 1 = end of frame; 2 = begin and end of frame; 3 = post-integate, post-apply and post-applypushout;)"));
+FAutoConsoleVariableRef CVarDebugDrawJoints(TEXT("p.Chaos.ImmPhys.DebugDrawJoints"), ChaosImmediate_DebugDrawJoints, TEXT("Draw Joints. (0 = never; 1 = end of frame; 2 = begin and end of frame; 3 = post-integate, post-apply and post-applypushout; 4 = each Apply step)."));
 
 namespace ImmediatePhysics_Chaos
 {
@@ -136,17 +136,23 @@ namespace ImmediatePhysics_Chaos
 		Evolution->SetPostIntegrateCallback(
 			[this]()
 			{
-				DebugDrawElementsAll(3, 3, 0.3f);
+				// Dynamic only - Kinematics get drawn once at the end of frame
+				//DebugDrawIslandParticles(Island, 3, 3, 0.3f, false, true);
+				//DebugDrawIslandConstraints(Island, 3, 3, 0.3f);
 			});
 		Evolution->SetPostApplyCallback(
 			[this](int32 Island)
 			{
-				DebugDrawElementsIsland(Island, 3, 3, 0.3f);
+				// Dynamic only - Kinematics get drawn once at the end of frame
+				DebugDrawIslandParticles(Island, 3, 3, 0.3f, false, true);
+				DebugDrawIslandConstraints(Island, 3, 3, 0.3f);
 			});
 		Evolution->SetPostApplyPushOutCallback(
 			[this](int32 Island)
 			{
-				DebugDrawElementsIsland(Island, 3, 3, 0.3f);
+				// Dynamic only - Kinematics get drawn once at the end of frame
+				DebugDrawIslandParticles(Island, 3, 3, 1.0f, false, true);
+				DebugDrawIslandConstraints(Island, 3, 3, 1.0f);
 			});
 		Collisions.SetPostApplyCallback(
 			[this](const float Dt, const TArray<TPBDCollisionConstraintHandle<float, 3>*>& InConstraintHandles)
@@ -161,7 +167,7 @@ namespace ImmediatePhysics_Chaos
 			{
 				if (ChaosImmediate_DebugDrawCollisions == 4)
 				{
-					DebugDraw::DrawCollisions(SimulationSpaceTransform, InConstraintHandles, 0.3f);
+					DebugDraw::DrawCollisions(SimulationSpaceTransform, InConstraintHandles, 0.6f);
 				}
 			});
 		Joints->SetPreApplyCallback(
@@ -177,7 +183,7 @@ namespace ImmediatePhysics_Chaos
 			{
 				if (ChaosImmediate_DebugDrawJoints == 4)
 				{
-					DebugDraw::DrawJointConstraints(SimulationSpaceTransform, InConstraintHandles, 1.0f);
+					DebugDraw::DrawJointConstraints(SimulationSpaceTransform, InConstraintHandles, 0.6f);
 				}
 			});
 #endif
@@ -408,18 +414,23 @@ namespace ImmediatePhysics_Chaos
 			Evolution->GetCollisionConstraints().SetCollisionsEnabled(ChaosImmediate_Collision_Enabled != 0);
 		}
 
-		DebugDrawElementsAll(2, 2, 0.7f);
+		DebugDrawParticles(2, 2, 0.7f, true, true);
+		DebugDrawConstraints(2, 2, 0.7f);
 
 		ConditionConstraints();
 
 		Evolution->GetGravityForces().SetAcceleration(InGravity);
+		
 		Evolution->AdvanceOneTimeStep(DeltaTime);
+		
 		Evolution->EndFrame(DeltaTime);
 
-		DebugDrawElementsAll(1, 3, 1.0f);
+		DebugDrawParticles(1, 2, 1.0f, true, true);
+		DebugDrawConstraints(1, 2, 1.0f);
+		DebugDrawParticles(3, 3, 1.0f, true, false);	// Kinematics only
 	}
 
-	void FSimulation::DebugDrawElementsAll(const int32 MinDebugLevel, const int32 MaxDebugLevel, const float ColorScale)
+	void FSimulation::DebugDrawParticles(const int32 MinDebugLevel, const int32 MaxDebugLevel, const float ColorScale, bool bDrawKinematic, bool bDrawDynamic)
 	{
 #if CHAOS_DEBUG_DRAW
 		using namespace Chaos;
@@ -427,12 +438,22 @@ namespace ImmediatePhysics_Chaos
 		{
 			if ((ChaosImmediate_DebugDrawParticles >= MinDebugLevel) && (ChaosImmediate_DebugDrawParticles <= MaxDebugLevel))
 			{
-				DebugDraw::DrawParticleTransforms(SimulationSpaceTransform, Evolution->GetParticles().GetAllParticlesView(), ColorScale);
+				DebugDraw::DrawParticleTransforms(SimulationSpaceTransform, Evolution->GetParticles().GetAllParticlesView(), ColorScale, bDrawKinematic, bDrawDynamic);
 			}
 			if ((ChaosImmediate_DebugDrawShapes >= MinDebugLevel) && (ChaosImmediate_DebugDrawShapes <= MaxDebugLevel))
 			{
-				DebugDraw::DrawParticleShapes(SimulationSpaceTransform, Evolution->GetParticles().GetAllParticlesView(), ColorScale);
+				DebugDraw::DrawParticleShapes(SimulationSpaceTransform, Evolution->GetParticles().GetAllParticlesView(), ColorScale, bDrawKinematic, bDrawDynamic);
 			}
+		}
+#endif
+	}
+
+	void FSimulation::DebugDrawConstraints(const int32 MinDebugLevel, const int32 MaxDebugLevel, const float ColorScale)
+	{
+#if CHAOS_DEBUG_DRAW
+		using namespace Chaos;
+		if (FDebugDrawQueue::IsDebugDrawingEnabled())
+		{
 			if ((ChaosImmediate_DebugDrawCollisions >= MinDebugLevel) && (ChaosImmediate_DebugDrawCollisions <= MaxDebugLevel))
 			{
 				DebugDraw::DrawCollisions(SimulationSpaceTransform, Evolution->GetCollisionConstraints(), ColorScale);
@@ -445,7 +466,7 @@ namespace ImmediatePhysics_Chaos
 #endif
 	}
 
-	void FSimulation::DebugDrawElementsIsland(const int32 Island, const int32 MinDebugLevel, const int32 MaxDebugLevel, const float ColorScale)
+	void FSimulation::DebugDrawIslandParticles(const int32 Island, const int32 MinDebugLevel, const int32 MaxDebugLevel, const float ColorScale, bool bDrawKinematic, bool bDrawDynamic)
 	{
 #if CHAOS_DEBUG_DRAW
 		using namespace Chaos;
@@ -453,12 +474,22 @@ namespace ImmediatePhysics_Chaos
 		{
 			if ((ChaosImmediate_DebugDrawParticles >= MinDebugLevel) && (ChaosImmediate_DebugDrawParticles <= MaxDebugLevel))
 			{
-				DebugDraw::DrawParticleTransforms(SimulationSpaceTransform, Evolution->GetIslandParticles(Island), ColorScale);
+				DebugDraw::DrawParticleTransforms(SimulationSpaceTransform, Evolution->GetIslandParticles(Island), ColorScale, bDrawKinematic, bDrawDynamic);
 			}
 			if ((ChaosImmediate_DebugDrawShapes >= MinDebugLevel) && (ChaosImmediate_DebugDrawShapes <= MaxDebugLevel))
 			{
-				DebugDraw::DrawParticleShapes(SimulationSpaceTransform, Evolution->GetIslandParticles(Island), ColorScale);
+				DebugDraw::DrawParticleShapes(SimulationSpaceTransform, Evolution->GetIslandParticles(Island), ColorScale, bDrawKinematic, bDrawDynamic);
 			}
+		}
+#endif
+	}
+
+	void FSimulation::DebugDrawIslandConstraints(const int32 Island, const int32 MinDebugLevel, const int32 MaxDebugLevel, const float ColorScale)
+	{
+#if CHAOS_DEBUG_DRAW
+		using namespace Chaos;
+		if (FDebugDrawQueue::IsDebugDrawingEnabled())
+		{
 			if ((ChaosImmediate_DebugDrawCollisions >= MinDebugLevel) && (ChaosImmediate_DebugDrawCollisions <= MaxDebugLevel))
 			{
 				Evolution->GetCollisionConstraintsRule().VisitIslandConstraints(Island,
diff --git a/Engine/Source/Runtime/Engine/Public/Physics/ImmediatePhysics/ImmediatePhysicsChaos/ImmediatePhysicsActorHandle_Chaos.h b/Engine/Source/Runtime/Engine/Public/Physics/ImmediatePhysics/ImmediatePhysicsChaos/ImmediatePhysicsActorHandle_Chaos.h
--- a/Engine/Source/Runtime/Engine/Public/Physics/ImmediatePhysics/ImmediatePhysicsChaos/ImmediatePhysicsActorHandle_Chaos.h
+++ b/Engine/Source/Runtime/Engine/Public/Physics/ImmediatePhysics/ImmediatePhysicsChaos/ImmediatePhysicsActorHandle_Chaos.h
@@ -8,7 +8,6 @@
 
 namespace ImmediatePhysics_Chaos
 {
-
 	/** handle associated with a physics actor. This is the proper way to read/write to the physics simulation */
 	struct ENGINE_API FActorHandle
 	{
@@ -26,10 +25,8 @@ namespace ImmediatePhysics_Chaos
 		/** Is the actor kinematic */
 		bool GetIsKinematic() const;
 
-#if IMMEDIATEPHYSICS_CHAOS_TODO
 		/** Gets the kinematic target (next transform) for the actor if one is set (check HasKinematicTarget() to see if a target is available) */
-		FImmediateKinematicTarget& GetKinematicTarget();
-#endif
+		const FKinematicTarget& GetKinematicTarget() const;
 
 		/** Sets the kinematic target. This will affect velocities as expected*/
 		void SetKinematicTarget(const FTransform& WorldTM);
@@ -117,6 +114,8 @@ namespace ImmediatePhysics_Chaos
 		void SetLevel(int32 InLevel);
 
 	private:
+		FKinematicTarget& GetKinematicTarget();
+
 		friend struct FSimulation;
 		friend struct FJointHandle;
 
diff --git a/Engine/Source/Runtime/Engine/Public/Physics/ImmediatePhysics/ImmediatePhysicsChaos/ImmediatePhysicsCore_Chaos.h b/Engine/Source/Runtime/Engine/Public/Physics/ImmediatePhysics/ImmediatePhysicsChaos/ImmediatePhysicsCore_Chaos.h
--- a/Engine/Source/Runtime/Engine/Public/Physics/ImmediatePhysics/ImmediatePhysicsChaos/ImmediatePhysicsCore_Chaos.h
+++ b/Engine/Source/Runtime/Engine/Public/Physics/ImmediatePhysics/ImmediatePhysicsChaos/ImmediatePhysicsCore_Chaos.h
@@ -11,6 +11,7 @@ namespace Chaos
 {
 	template<typename T, int D> class TImplicitObject;
 	template<typename T, int D> struct TKinematicGeometryParticleParameters;
+	template<typename T, int D> class TKinematicTarget;
 	template<typename P, typename T, int D> class TPBDConstraintIslandRule;
 	template<typename T, int D> class TPBDJointConstraintHandle;
 	template<typename T, int D> class TPBDJointConstraints;
@@ -27,6 +28,8 @@ namespace ImmediatePhysics_Chaos
 
 	using EActorType = ImmediatePhysics_Shared::EActorType;
 	using EForceType = ImmediatePhysics_Shared::EForceType;
+
+	using FKinematicTarget = Chaos::TKinematicTarget<FReal, Dimensions>;
 }
 
 struct FBodyInstance;
diff --git a/Engine/Source/Runtime/Engine/Public/Physics/ImmediatePhysics/ImmediatePhysicsChaos/ImmediatePhysicsSimulation_Chaos.h b/Engine/Source/Runtime/Engine/Public/Physics/ImmediatePhysics/ImmediatePhysicsChaos/ImmediatePhysicsSimulation_Chaos.h
--- a/Engine/Source/Runtime/Engine/Public/Physics/ImmediatePhysics/ImmediatePhysicsChaos/ImmediatePhysicsSimulation_Chaos.h
+++ b/Engine/Source/Runtime/Engine/Public/Physics/ImmediatePhysics/ImmediatePhysicsChaos/ImmediatePhysicsSimulation_Chaos.h
@@ -81,8 +81,10 @@ namespace ImmediatePhysics_Chaos
 	private:
 		void ConditionConstraints();
 
-		void DebugDrawElementsAll(const int32 MinDebugLevel, const int32 MaxDebugLevel, const float ColorScale);
-		void DebugDrawElementsIsland(const int32 Island, const int32 MinDebugLevel, const int32 MaxDebugLevel, const float ColorScale);
+		void DebugDrawParticles(const int32 MinDebugLevel, const int32 MaxDebugLevel, const float ColorScale, bool bDrawKinematic, bool bDrawDynamic);
+		void DebugDrawConstraints(const int32 MinDebugLevel, const int32 MaxDebugLevel, const float ColorScale);
+		void DebugDrawIslandParticles(const int32 Island, const int32 MinDebugLevel, const int32 MaxDebugLevel, const float ColorScale, bool bDrawKinematic, bool bDrawDynamic);
+		void DebugDrawIslandConstraints(const int32 Island, const int32 MinDebugLevel, const int32 MaxDebugLevel, const float ColorScale);
 
 		using FParticleHandle = Chaos::TGeometryParticleHandle<FReal, Dimensions>;
 		TUniquePtr<Chaos::TPBDRigidsEvolutionGBF<FReal, Dimensions>> Evolution;
diff --git a/Engine/Source/Runtime/Experimental/Chaos/Private/Chaos/ChaosDebugDraw.cpp b/Engine/Source/Runtime/Experimental/Chaos/Private/Chaos/ChaosDebugDraw.cpp
--- a/Engine/Source/Runtime/Experimental/Chaos/Private/Chaos/ChaosDebugDraw.cpp
+++ b/Engine/Source/Runtime/Experimental/Chaos/Private/Chaos/ChaosDebugDraw.cpp
@@ -220,53 +220,65 @@ namespace Chaos
 		//
 		//
 
-		void DrawParticleShapes(const TRigidTransform<float, 3>& SpaceTransform, const TParticleView<TGeometryParticles<float, 3>>& ParticlesView, float ColorScale)
+		void DrawParticleShapes(const TRigidTransform<float, 3>& SpaceTransform, const TParticleView<TGeometryParticles<float, 3>>& ParticlesView, float ColorScale, bool bDrawKinemtatic, bool bDrawDynamic)
 		{
 #if CHAOS_DEBUG_DRAW
 			if (FDebugDrawQueue::IsDebugDrawingEnabled())
 			{
 				for (auto& Particle : ParticlesView)
 				{
-					DrawParticleShapesImpl<float, 3>(SpaceTransform, GetHandleHelper<float, 3>(&Particle), ColorScale);
+					if ((bDrawKinemtatic && !Particle.AsDynamic()) || (bDrawDynamic && Particle.AsDynamic()))
+					{
+						DrawParticleShapesImpl<float, 3>(SpaceTransform, GetHandleHelper<float, 3>(&Particle), ColorScale);
+					}
 				}
 			}
 #endif
 		}
 
-		void DrawParticleShapes(const TRigidTransform<float, 3>& SpaceTransform, const TArray<TGeometryParticleHandle<float, 3>*>& Particles, float ColorScale)
+		void DrawParticleShapes(const TRigidTransform<float, 3>& SpaceTransform, const TArray<TGeometryParticleHandle<float, 3>*>& Particles, float ColorScale, bool bDrawKinemtatic, bool bDrawDynamic)
 		{
 #if CHAOS_DEBUG_DRAW
 			if (FDebugDrawQueue::IsDebugDrawingEnabled())
 			{
 				for (auto& Particle : Particles)
 				{
-					DrawParticleShapesImpl<float, 3>(SpaceTransform, Particle, ColorScale);
+					if ((bDrawKinemtatic && !Particle->AsDynamic()) || (bDrawDynamic && Particle->AsDynamic()))
+					{
+						DrawParticleShapesImpl<float, 3>(SpaceTransform, Particle, ColorScale);
+					}
 				}
 			}
 #endif
 		}
 
-		void DrawParticleTransforms(const TRigidTransform<float, 3>& SpaceTransform, const TParticleView<TGeometryParticles<float, 3>>& ParticlesView, float ColorScale)
+		void DrawParticleTransforms(const TRigidTransform<float, 3>& SpaceTransform, const TParticleView<TGeometryParticles<float, 3>>& ParticlesView, float ColorScale, bool bDrawKinemtatic, bool bDrawDynamic)
 		{
 #if CHAOS_DEBUG_DRAW
 			if (FDebugDrawQueue::IsDebugDrawingEnabled())
 			{
 				for (auto& Particle : ParticlesView)
 				{
-					DrawParticleTransformImpl<float, 3>(SpaceTransform, GetHandleHelper<float, 3>(&Particle), ColorScale);
+					if ((bDrawKinemtatic && !Particle.AsDynamic()) || (bDrawDynamic && Particle.AsDynamic()))
+					{
+						DrawParticleTransformImpl<float, 3>(SpaceTransform, GetHandleHelper<float, 3>(&Particle), ColorScale);
+					}
 				}
 			}
 #endif
 		}
 
-		void DrawParticleTransforms(const TRigidTransform<float, 3>& SpaceTransform, const TArray<TGeometryParticleHandle<float, 3>*>& Particles, float ColorScale)
+		void DrawParticleTransforms(const TRigidTransform<float, 3>& SpaceTransform, const TArray<TGeometryParticleHandle<float, 3>*>& Particles, float ColorScale, bool bDrawKinemtatic, bool bDrawDynamic)
 		{
 #if CHAOS_DEBUG_DRAW
 			if (FDebugDrawQueue::IsDebugDrawingEnabled())
 			{
 				for (auto& Particle : Particles)
 				{
-					DrawParticleTransformImpl<float, 3>(SpaceTransform, Particle, ColorScale);
+					if ((bDrawKinemtatic && !Particle->AsDynamic()) || (bDrawDynamic && Particle->AsDynamic()))
+					{
+						DrawParticleTransformImpl<float, 3>(SpaceTransform, Particle, ColorScale);
+					}
 				}
 			}
 #endif
diff --git a/Engine/Source/Runtime/Experimental/Chaos/Public/Chaos/ChaosDebugDraw.h b/Engine/Source/Runtime/Experimental/Chaos/Public/Chaos/ChaosDebugDraw.h
--- a/Engine/Source/Runtime/Experimental/Chaos/Public/Chaos/ChaosDebugDraw.h
+++ b/Engine/Source/Runtime/Experimental/Chaos/Public/Chaos/ChaosDebugDraw.h
@@ -10,10 +10,10 @@ namespace Chaos
 	namespace DebugDraw
 	{
 #if CHAOS_DEBUG_DRAW
-		CHAOS_API void DrawParticleShapes(const TRigidTransform<float, 3>& SpaceTransform, const TParticleView<TGeometryParticles<float, 3>>& ParticlesView, float ColorScale);
-		CHAOS_API void DrawParticleShapes(const TRigidTransform<float, 3>& SpaceTransform, const TArray<TGeometryParticleHandle<float, 3>*>& Particles, float ColorScale);
-		CHAOS_API void DrawParticleTransforms(const TRigidTransform<float, 3>& SpaceTransform, const TParticleView<TGeometryParticles<float, 3>>& ParticlesView, float ColorScale);
-		CHAOS_API void DrawParticleTransforms(const TRigidTransform<float, 3>& SpaceTransform, const TArray<TGeometryParticleHandle<float, 3>*>& Particles, float ColorScale);
+		CHAOS_API void DrawParticleShapes(const TRigidTransform<float, 3>& SpaceTransform, const TParticleView<TGeometryParticles<float, 3>>& ParticlesView, float ColorScale, bool bDrawKinematic = true, bool bDrawDynamic = true);
+		CHAOS_API void DrawParticleShapes(const TRigidTransform<float, 3>& SpaceTransform, const TArray<TGeometryParticleHandle<float, 3>*>& Particles, float ColorScale, bool bDrawKinematic = true, bool bDrawDynamic = true);
+		CHAOS_API void DrawParticleTransforms(const TRigidTransform<float, 3>& SpaceTransform, const TParticleView<TGeometryParticles<float, 3>>& ParticlesView, float ColorScale, bool bDrawKinematic = true, bool bDrawDynamic = true);
+		CHAOS_API void DrawParticleTransforms(const TRigidTransform<float, 3>& SpaceTransform, const TArray<TGeometryParticleHandle<float, 3>*>& Particles, float ColorScale, bool bDrawKinematic = true, bool bDrawDynamic = true);
 		CHAOS_API void DrawParticleCollisions(const TRigidTransform<float, 3>& SpaceTransform, const TGeometryParticleHandle<float, 3>* Particle, const TPBDCollisionConstraint<float, 3>& Collisions);
 		CHAOS_API void DrawCollisions(const TRigidTransform<float, 3>& SpaceTransform, const TPBDCollisionConstraint<float, 3>& Collisions, float ColorScale);
 		CHAOS_API void DrawCollisions(const TRigidTransform<float, 3>& SpaceTransform, const TArray<TPBDCollisionConstraintHandle<float, 3>*>& ConstraintHandles, float ColorScale);
