UE-62399 Graphical corruption when opening the debug console
-Switched from the fast hardware clear to the slower shader clear because the slower one works.  Perhaps a better fix exists that would make the faster clear work?  Costs 60us when debug ui is up, but only when rendering debug to a stereo layer (vr/ar).

diff --git a/Engine/Source/Runtime/Engine/Private/Slate/DebugCanvas.cpp b/Engine/Source/Runtime/Engine/Private/Slate/DebugCanvas.cpp
--- a/Engine/Source/Runtime/Engine/Private/Slate/DebugCanvas.cpp
+++ b/Engine/Source/Runtime/Engine/Private/Slate/DebugCanvas.cpp
@@ -206,6 +206,8 @@ void FDebugCanvasDrawer::DrawRenderThread(FRHICommandListImmediate& RHICmdList,
 {
 	check( IsInRenderingThread() );
 
+	SCOPED_DRAW_EVENT(RHICmdList, DrawDebugCanvas);
+
 	QUICK_SCOPE_CYCLE_COUNTER(Stat_DrawDebugCanvas);
 	if( RenderThreadCanvas.IsValid() )
 	{
@@ -219,7 +221,9 @@ void FDebugCanvasDrawer::DrawRenderThread(FRHICommandListImmediate& RHICmdList,
 
 			if (!LayerTexture)
 			{
-				const FPooledRenderTargetDesc Desc(FPooledRenderTargetDesc::Create2DDesc(RenderThreadCanvas->GetParentCanvasSize(), PF_B8G8R8A8, FClearValueBinding(), TexCreate_SRGB, TexCreate_RenderTargetable, false));
+				// Set TexCreate_NoFastClear because the fast CMASK clear was not working on ps4.
+				FPooledRenderTargetDesc Desc(FPooledRenderTargetDesc::Create2DDesc(RenderThreadCanvas->GetParentCanvasSize(), PF_B8G8R8A8, FClearValueBinding(), TexCreate_SRGB, TexCreate_RenderTargetable | TexCreate_NoFastClear, false));
+				Desc.DebugName = TEXT("DebugCanvasLayerTexture");
 				GetRendererModule().RenderTargetPoolFindFreeElement(RHICmdList, Desc, LayerTexture, TEXT("DebugCanvasLayerTexture"));
 				UE_LOG(LogProfilingDebugging, Log, TEXT("Allocated a %d x %d texture for HMD canvas layer"), RenderThreadCanvas->GetParentCanvasSize().X, RenderThreadCanvas->GetParentCanvasSize().Y);
 			}
