Don't call GetStereoLayers if stereo is not enabled.

diff --git a/Engine/Source/Runtime/Engine/Private/Components/StereoLayerComponent.cpp b/Engine/Source/Runtime/Engine/Private/Components/StereoLayerComponent.cpp
--- a/Engine/Source/Runtime/Engine/Private/Components/StereoLayerComponent.cpp
+++ b/Engine/Source/Runtime/Engine/Private/Components/StereoLayerComponent.cpp
@@ -39,7 +39,7 @@ void UStereoLayerComponent::BeginDestroy()
 	Super::BeginDestroy();
 
 	IStereoLayers* StereoLayers;
-	if (LayerId && GEngine->StereoRenderingDevice.IsValid() && (StereoLayers = GEngine->StereoRenderingDevice->GetStereoLayers()) != nullptr)
+	if (LayerId && GEngine->StereoRenderingDevice.IsValid() && GEngine->StereoRenderingDevice->IsStereoEnabled() && (StereoLayers = GEngine->StereoRenderingDevice->GetStereoLayers()) != nullptr)
 	{
 		StereoLayers->DestroyLayer(LayerId);
 		LayerId = 0;
@@ -51,7 +51,7 @@ void UStereoLayerComponent::TickComponent(float DeltaTime, enum ELevelTick TickT
 	Super::TickComponent( DeltaTime, TickType, ThisTickFunction );
 
 	IStereoLayers* StereoLayers;
-	if (!GEngine->StereoRenderingDevice.IsValid() || (StereoLayers = GEngine->StereoRenderingDevice->GetStereoLayers()) == nullptr || !Texture)
+	if (!GEngine->StereoRenderingDevice.IsValid() || !GEngine->StereoRenderingDevice->IsStereoEnabled() || (StereoLayers = GEngine->StereoRenderingDevice->GetStereoLayers()) == nullptr || !Texture )
 	{
 		return;
 	}
