[NO NON-PII COMMIT MESSAGE]

diff --git a/Engine/Source/Runtime/Engine/Private/AudioCompressionSettingsUtils.cpp b/Engine/Source/Runtime/Engine/Private/AudioCompressionSettingsUtils.cpp
--- a/Engine/Source/Runtime/Engine/Private/AudioCompressionSettingsUtils.cpp
+++ b/Engine/Source/Runtime/Engine/Private/AudioCompressionSettingsUtils.cpp
@@ -5,7 +5,7 @@
 
 #define ENABLE_PLATFORM_COMPRESSION_OVERRIDES 1
 
-#if PLATFORM_ANDROID && ENABLE_PLATFORM_COMPRESSION_OVERRIDES
+#if PLATFORM_ANDROID && !PLATFORM_LUMIN && ENABLE_PLATFORM_COMPRESSION_OVERRIDES
 #include "AndroidRuntimeSettings.h"
 #endif
 
@@ -21,7 +21,7 @@
 
 const FPlatformRuntimeAudioCompressionOverrides* FPlatformCompressionUtilities::GetRuntimeCompressionOverridesForCurrentPlatform()
 {
-#if PLATFORM_ANDROID && ENABLE_PLATFORM_COMPRESSION_OVERRIDES
+#if PLATFORM_ANDROID && !PLATFORM_LUMIN && ENABLE_PLATFORM_COMPRESSION_OVERRIDES
 	static const UAndroidRuntimeSettings* Settings = GetDefault<UAndroidRuntimeSettings>();
 	if (Settings)
 	{
@@ -44,13 +44,13 @@ const FPlatformRuntimeAudioCompressionOverrides* FPlatformCompressionUtilities::
 		return &(Settings->CompressionOverrides);
 	}
 
-#endif // PLATFORM_ANDROID
+#endif // PLATFORM_ANDROID && !PLATFORM_LUMIN
 	return nullptr;
 }
 
 void CacheCurrentPlatformAudioCookOverrides(FPlatformAudioCookOverrides& OutOverrides)
 {
-#if PLATFORM_ANDROID
+#if PLATFORM_ANDROID && !PLATFORM_LUMIN
 	const TCHAR* CategoryName = TEXT("/Script/AndroidRuntimeSettings.AndroidRuntimeSettings");
 #elif PLATFORM_IOS
 	const TCHAR* CategoryName = TEXT("/Script/IOSRuntimeSettings.IOSRuntimeSettings");
@@ -93,14 +93,14 @@ void CacheCurrentPlatformAudioCookOverrides(FPlatformAudioCookOverrides& OutOver
 	OutOverrides.PlatformSampleRates.Add(ESoundwaveSampleRateSettings::Min, RetrievedSampleRate);
 }
 
-#if PLATFORM_ANDROID || PLATFORM_IOS || PLATFORM_SWITCH
+#if (PLATFORM_ANDROID && !PLATFORM_LUMIN) || PLATFORM_IOS || PLATFORM_SWITCH
 static FPlatformAudioCookOverrides OutOverrides = FPlatformAudioCookOverrides();
 #endif
 
 
 const FPlatformAudioCookOverrides* FPlatformCompressionUtilities::GetCookOverridesForCurrentPlatform()
 {
-#if PLATFORM_ANDROID || PLATFORM_IOS || PLATFORM_SWITCH
+#if (PLATFORM_ANDROID && !PLATFORM_LUMIN) || PLATFORM_IOS || PLATFORM_SWITCH
 	static bool bCachedCookOverrides = false;
 	if (!bCachedCookOverrides)
 	{
