Fix for presets not informing their instances when they are deleted

diff --git a/Engine/Source/Runtime/Engine/Classes/Sound/SoundEffectPreset.h b/Engine/Source/Runtime/Engine/Classes/Sound/SoundEffectPreset.h
--- a/Engine/Source/Runtime/Engine/Classes/Sound/SoundEffectPreset.h
+++ b/Engine/Source/Runtime/Engine/Classes/Sound/SoundEffectPreset.h
@@ -18,7 +18,8 @@ class ENGINE_API USoundEffectPreset : public UObject
 
 public:
 	USoundEffectPreset(const FObjectInitializer& ObjectInitializer);
-
+	virtual ~USoundEffectPreset();
+	
 	virtual FText GetAssetActionName() const PURE_VIRTUAL(USoundEffectPreset::GetAssetActionName, return FText(););
 	virtual UClass* GetSupportedClass() const PURE_VIRTUAL(USoundEffectPreset::GetSupportedClass, return nullptr;);
 	virtual USoundEffectPreset* CreateNewPreset(UObject* InParent, FName Name, EObjectFlags Flags) const PURE_VIRTUAL(USoundEffectPreset::CreateNewPreset, return nullptr;);
diff --git a/Engine/Source/Runtime/Engine/Private/SoundEffectPreset.cpp b/Engine/Source/Runtime/Engine/Private/SoundEffectPreset.cpp
--- a/Engine/Source/Runtime/Engine/Private/SoundEffectPreset.cpp
+++ b/Engine/Source/Runtime/Engine/Private/SoundEffectPreset.cpp
@@ -13,6 +13,15 @@ USoundEffectPreset::USoundEffectPreset(const FObjectInitializer& ObjectInitializ
 	
 }
 
+USoundEffectPreset::~USoundEffectPreset()
+{
+	for (int32 i = 0; i < Instances.Num(); ++i)
+	{
+		Instances[i]->SetPreset(nullptr);
+	}
+}
+
+
 void USoundEffectPreset::EffectCommand(TFunction<void()> Command)
 {
 	for (int32 i = 0; i < Instances.Num(); ++i)
