UE-65133 Console is not visible in HMD with Vive on Mac
-With FDefaultStereoLayers the debug layer could be drawn to the spectator screen twice.  This is because it will be in the main scene render and also will be copied to the spectator screen.  Now IStereoLayer has an accessor to tell the DebugCanvas if it should copy the Debug layer, and the FDefaultStereoLayers version returns false.

diff --git a/Engine/Plugins/Runtime/Oculus/OculusVR/Source/OculusHMD/Private/OculusHMD.h b/Engine/Plugins/Runtime/Oculus/OculusVR/Source/OculusHMD/Private/OculusHMD.h
--- a/Engine/Plugins/Runtime/Oculus/OculusVR/Source/OculusHMD/Private/OculusHMD.h
+++ b/Engine/Plugins/Runtime/Oculus/OculusVR/Source/OculusHMD/Private/OculusHMD.h
@@ -190,6 +190,7 @@ public:
 	virtual void UpdateSplashScreen() override;
 	virtual IStereoLayers::FLayerDesc GetDebugCanvasLayerDesc(FTextureRHIRef Texture) override;
 	virtual void GetAllocatedTexture(uint32 LayerId, FTextureRHIRef &Texture, FTextureRHIRef &LeftTexture) override;
+	virtual bool ShouldCopyDebugLayersToSpectatorScreen() const override { return true; }
 
 	// ISceneViewExtension
 	virtual void SetupViewFamily(FSceneViewFamily& InViewFamily) override;
diff --git a/Engine/Plugins/Runtime/Steam/SteamVR/Source/SteamVR/Private/SteamVRHMD.h b/Engine/Plugins/Runtime/Steam/SteamVR/Source/SteamVR/Private/SteamVRHMD.h
--- a/Engine/Plugins/Runtime/Steam/SteamVR/Source/SteamVR/Private/SteamVRHMD.h
+++ b/Engine/Plugins/Runtime/Steam/SteamVR/Source/SteamVR/Private/SteamVRHMD.h
@@ -228,6 +228,7 @@ public:
 	// Create/Set/Get/Destroy inherited from TStereoLayerManager
 	virtual void UpdateSplashScreen() override;
 	virtual void GetAllocatedTexture(uint32 LayerId, FTextureRHIRef &Texture, FTextureRHIRef &LeftTexture) override;
+	virtual bool ShouldCopyDebugLayersToSpectatorScreen() const override { return true; }
 
 	// ISceneViewExtension interface
 	virtual void SetupViewFamily(FSceneViewFamily& InViewFamily) override {};
diff --git a/Engine/Source/Runtime/Engine/Private/Slate/DebugCanvas.cpp b/Engine/Source/Runtime/Engine/Private/Slate/DebugCanvas.cpp
--- a/Engine/Source/Runtime/Engine/Private/Slate/DebugCanvas.cpp
+++ b/Engine/Source/Runtime/Engine/Private/Slate/DebugCanvas.cpp
@@ -237,7 +237,7 @@ void FDebugCanvasDrawer::DrawRenderThread(FRHICommandListImmediate& RHICmdList,
 				StereoLayers->GetAllocatedTexture(LayerID, HMDSwapchain, HMDNull);
 
 				// If drawing to a layer tell the spectator screen controller to copy that layer to the spectator screen.
-				if (LayerID != INVALID_LAYER_ID && GEngine && GEngine->XRSystem)
+				if (StereoLayers->ShouldCopyDebugLayersToSpectatorScreen() && LayerID != INVALID_LAYER_ID && GEngine && GEngine->XRSystem)
 				{
 					IHeadMountedDisplay* HMD = GEngine->XRSystem->GetHMDDevice();
 					if (HMD)
diff --git a/Engine/Source/Runtime/Engine/Public/IStereoLayers.h b/Engine/Source/Runtime/Engine/Public/IStereoLayers.h
--- a/Engine/Source/Runtime/Engine/Public/IStereoLayers.h
+++ b/Engine/Source/Runtime/Engine/Public/IStereoLayers.h
@@ -122,6 +122,11 @@ public:
 	 */
 	virtual void UpdateSplashScreen() = 0;
 
+	/**
+	* If true the debug layers are copied to the spectator screen, because they do not naturally end up on the spectator screen as part of the 3d view.
+	*/
+	virtual bool ShouldCopyDebugLayersToSpectatorScreen() const = 0;
+
 public:
 	/**
 	* Set the splash screen attributes
diff --git a/Engine/Source/Runtime/HeadMountedDisplay/Public/DefaultStereoLayers.h b/Engine/Source/Runtime/HeadMountedDisplay/Public/DefaultStereoLayers.h
--- a/Engine/Source/Runtime/HeadMountedDisplay/Public/DefaultStereoLayers.h
+++ b/Engine/Source/Runtime/HeadMountedDisplay/Public/DefaultStereoLayers.h
@@ -28,6 +28,12 @@ public:
 
 	void GetAllocatedTexture(uint32 LayerId, FTextureRHIRef &Texture, FTextureRHIRef &LeftTexture);
 
+	virtual bool ShouldCopyDebugLayersToSpectatorScreen() const override
+	{
+		// Emulated layer support means that the debug layer will be in the 3d scene render that the spectator screen displays.
+		return false;
+	}
+
 protected:
 	
 	/**
