Merging //UE4/Dev-Main to Dev-VR (//UE4/Dev-VR)
CL 8 of 8
Actual integrations are in this CL, the previous 7 were no ops

diff --git a/Engine/Plugins/Runtime/WindowsMixedReality/Source/WindowsMixedRealityHMD/Private/WindowsMixedRealityHMD.cpp b/Engine/Plugins/Runtime/WindowsMixedReality/Source/WindowsMixedRealityHMD/Private/WindowsMixedRealityHMD.cpp
--- a/Engine/Plugins/Runtime/WindowsMixedReality/Source/WindowsMixedRealityHMD/Private/WindowsMixedRealityHMD.cpp
+++ b/Engine/Plugins/Runtime/WindowsMixedReality/Source/WindowsMixedRealityHMD/Private/WindowsMixedRealityHMD.cpp
@@ -216,7 +216,7 @@ namespace WindowsMixedReality
 
 		FlushRenderingCommands();
 
-		return std::move(D3D11DeviceLocal);
+		return D3D11DeviceLocal;
 	}
 
 	/** Helper function for acquiring the appropriate FSceneViewport */
@@ -1080,11 +1080,11 @@ namespace WindowsMixedReality
 	FWindowsMixedRealityHMD::FWindowsMixedRealityHMD(const FAutoRegister& AutoRegister)
 		: FSceneViewExtensionBase(AutoRegister)
 		, FHeadMountedDisplayBase(nullptr)
+		, ScreenScalePercentage(1.0f)
 		, mCustomPresent(nullptr)
+		, HMDTrackingOrigin(EHMDTrackingOrigin::Floor)
 		, CurrOrientation(FQuat::Identity)
 		, CurrPosition(FVector::ZeroVector)
-		, HMDTrackingOrigin(EHMDTrackingOrigin::Floor)
-		, ScreenScalePercentage(1.0f)
 	{
 		static const FName RendererModuleName("Renderer");
 		RendererModule = FModuleManager::GetModulePtr<IRendererModule>(RendererModuleName);
diff --git a/Engine/Source/Runtime/Engine/Private/PhysicsEngine/BodyInstance.cpp b/Engine/Source/Runtime/Engine/Private/PhysicsEngine/BodyInstance.cpp
--- a/Engine/Source/Runtime/Engine/Private/PhysicsEngine/BodyInstance.cpp
+++ b/Engine/Source/Runtime/Engine/Private/PhysicsEngine/BodyInstance.cpp
@@ -1,4 +1,4 @@
-// Copyright 1998-2018 Epic Games, Inc. All Rights Reserved.
+// Copyright 1998-2019 Epic Games, Inc. All Rights Reserved.
 
 #include "PhysicsEngine/BodyInstance.h"
 #include "EngineGlobals.h"
@@ -30,10 +30,10 @@
 #include "Misc/UObjectToken.h"
 
 #if WITH_PHYSX
-	#include "PhysXPublic.h"
-	#include "Physics/PhysicsFiltering.h"
-	#include "PhysicsEngine/PhysXSupport.h"
-	#include "Collision/CollisionConversions.h"
+#include "PhysXPublic.h"
+#include "Physics/PhysicsFiltering.h"
+#include "PhysicsEngine/PhysXSupport.h"
+#include "Collision/CollisionConversions.h"
 #include "PxShape.h"
 #endif // WITH_PHYSX
 
@@ -82,9 +82,9 @@ bool IsRigidBodyKinematic_AssumesLocked(const FPhysicsActorHandle& InActorRef)
 	return false;
 }
 
-int32 FillInlineShapeArray_AssumesLocked(PhysicsInterfaceTypes::FInlineShapeArray& Array, const FPhysicsActorHandle& Actor, EPhysicsSceneType InSceneType)
+int32 FillInlineShapeArray_AssumesLocked(PhysicsInterfaceTypes::FInlineShapeArray& Array, const FPhysicsActorHandle& Actor)
 {
-	FPhysicsInterface::GetAllShapes_AssumedLocked(Actor, Array, InSceneType);
+	FPhysicsInterface::GetAllShapes_AssumedLocked(Actor, Array);
 
 	return Array.Num();
 }
@@ -108,7 +108,7 @@ void FCollisionResponse::SetResponse(ECollisionChannel Channel, ECollisionRespon
 {
 #if 1// @hack until PostLoad is disabled for CDO of BP WITH_EDITOR
 	ECollisionResponse DefaultResponse = FCollisionResponseContainer::GetDefaultResponseContainer().GetResponse(Channel);
-	if (DefaultResponse == NewResponse)
+	if(DefaultResponse == NewResponse)
 	{
 		RemoveReponseFromArray(Channel);
 	}
@@ -153,11 +153,11 @@ void FCollisionResponse::SetResponsesArray(const TArray<FResponseChannel>& InCha
 #if DO_GUARD_SLOW
 	// verify if the name is overlapping, if so, ensure, do not remove in debug because it will cause inconsistent bug between debug/release
 	int32 const ResponseNum = InChannelResponses.Num();
-	for (int32 I=0; I<ResponseNum; ++I)
+	for(int32 I = 0; I < ResponseNum; ++I)
 	{
-		for (int32 J=I+1; J<ResponseNum; ++J)
+		for(int32 J = I + 1; J < ResponseNum; ++J)
 		{
-			if (InChannelResponses[I].Channel == InChannelResponses[J].Channel)
+			if(InChannelResponses[I].Channel == InChannelResponses[J].Channel)
 			{
 				UE_LOG(LogCollision, Warning, TEXT("Collision Channel : Redundant name exists"));
 			}
@@ -177,9 +177,9 @@ bool FCollisionResponse::RemoveReponseFromArray(ECollisionChannel Channel)
 	// without editor, this does not have to match 
 	// We'd need to save name just in case that name is gone or not
 	FName ChannelName = UCollisionProfile::Get()->ReturnChannelNameFromContainerIndex(Channel);
-	for (auto Iter=ResponseArray.CreateIterator(); Iter; ++Iter)
+	for(auto Iter = ResponseArray.CreateIterator(); Iter; ++Iter)
 	{
-		if (ChannelName == (*Iter).Channel)
+		if(ChannelName == (*Iter).Channel)
 		{
 			ResponseArray.RemoveAt(Iter.GetIndex());
 			return true;
@@ -194,9 +194,9 @@ bool FCollisionResponse::AddReponseToArray(ECollisionChannel Channel, ECollision
 	// so only allow it in editor
 	// without editor, this does not have to match 
 	FName ChannelName = UCollisionProfile::Get()->ReturnChannelNameFromContainerIndex(Channel);
-	for (auto Iter=ResponseArray.CreateIterator(); Iter; ++Iter)
+	for(auto Iter = ResponseArray.CreateIterator(); Iter; ++Iter)
 	{
-		if (ChannelName == (*Iter).Channel)
+		if(ChannelName == (*Iter).Channel)
 		{
 			(*Iter).Response = Response;
 			return true;
@@ -215,13 +215,13 @@ void FCollisionResponse::UpdateArrayFromResponseContainer()
 	const FCollisionResponseContainer& DefaultResponse = FCollisionResponseContainer::GetDefaultResponseContainer();
 	const UCollisionProfile* CollisionProfile = UCollisionProfile::Get();
 
-	for (int32 i = 0; i < ARRAY_COUNT(ResponseToChannels.EnumArray); i++)
+	for(int32 i = 0; i < ARRAY_COUNT(ResponseToChannels.EnumArray); i++)
 	{
 		// if not same as default
-		if (ResponseToChannels.EnumArray[i] != DefaultResponse.EnumArray[i])
+		if(ResponseToChannels.EnumArray[i] != DefaultResponse.EnumArray[i])
 		{
 			FName ChannelName = CollisionProfile->ReturnChannelNameFromContainerIndex(i);
-			if (ChannelName != NAME_None)
+			if(ChannelName != NAME_None)
 			{
 				FResponseChannel NewResponse;
 				NewResponse.Channel = ChannelName;
@@ -238,12 +238,12 @@ void FCollisionResponse::UpdateResponseContainerFromArray()
 {
 	ResponseToChannels = FCollisionResponseContainer::GetDefaultResponseContainer();
 
-	for (auto Iter = ResponseArray.CreateIterator(); Iter; ++Iter)
+	for(auto Iter = ResponseArray.CreateIterator(); Iter; ++Iter)
 	{
 		FResponseChannel& Response = *Iter;
 
 		int32 EnumIndex = UCollisionProfile::Get()->ReturnContainerIndexFromChannelName(Response.Channel);
-		if ( EnumIndex != INDEX_NONE )
+		if(EnumIndex != INDEX_NONE)
 		{
 			ResponseToChannels.SetResponse((ECollisionChannel)EnumIndex, Response.Response);
 		}
@@ -271,7 +271,7 @@ bool FCollisionResponse::operator==(const FCollisionResponse& Other) const
 					break;
 				}
 			}
-			
+
 		}
 	}
 
@@ -309,12 +309,10 @@ FBodyInstance::FBodyInstance()
 	, bLockYRotation(false)
 	, bLockZRotation(false)
 	, bOverrideMaxAngularVelocity(false)
-	, bUseAsyncScene(false)
 	, bOverrideMaxDepenetrationVelocity(false)
 	, bOverrideWalkableSlopeOnInstance(false)
 	, bInterpolateWhenSubStepping(true)
 	, bPendingCollisionProfileSetup(false)
-	, bHasSharedShapes(false)
 	, Scale3D(1.0f)
 	, CollisionProfileName(UCollisionProfile::CustomCollisionProfileName)
 	, MaxDepenetrationVelocity(0.f)
@@ -343,15 +341,15 @@ const FPhysicsActorHandle& FBodyInstance::GetActorReferenceWithWelding() const
 	return WeldParent ? WeldParent->ActorHandle : ActorHandle;
 }
 
-FArchive& operator<<(FArchive& Ar,FBodyInstance& BodyInst)
+FArchive& operator<<(FArchive& Ar, FBodyInstance& BodyInst)
 {
-	if (!Ar.IsLoading() && !Ar.IsSaving())
+	if(!Ar.IsLoading() && !Ar.IsSaving())
 	{
 		Ar << BodyInst.OwnerComponent;
 		Ar << BodyInst.PhysMaterialOverride;
 	}
 
-	if (Ar.IsLoading() && Ar.UE4Ver() < VER_UE4_MAX_ANGULAR_VELOCITY_DEFAULT)
+	if(Ar.IsLoading() && Ar.UE4Ver() < VER_UE4_MAX_ANGULAR_VELOCITY_DEFAULT)
 	{
 		if(BodyInst.MaxAngularVelocity != 400.f)
 		{
@@ -404,22 +402,20 @@ bool FBodyInstance::IsShapeBoundToBody(const FPhysicsShapeHandle& Shape) const
 const TMap<FPhysicsShapeHandle, FBodyInstance::FWeldInfo>* FBodyInstance::GetCurrentWeldInfo() const
 {
 	return ShapeToBodiesMap.Get();
-	}
+}
 
 int32 FBodyInstance::GetAllShapes_AssumesLocked(TArray<FPhysicsShapeHandle>& OutShapes) const
-	{
-	// If we're sharing shapes we only want to grab from the sync scene, otherwise get everything
-	EPhysicsSceneType SceneType = HasSharedShapes() ? PST_Sync : PST_MAX;
-	return FPhysicsInterface::GetAllShapes_AssumedLocked(ActorHandle, OutShapes, SceneType);
+{
+	return FPhysicsInterface::GetAllShapes_AssumedLocked(ActorHandle, OutShapes);
 }
 
 void FBodyInstance::UpdateTriMeshVertices(const TArray<FVector> & NewPositions)
 {
-#if WITH_APEIRON || WITH_IMMEDIATE_PHYSX
+#if WITH_CHAOS || WITH_IMMEDIATE_PHYSX
 	check(false);
 #else
 #if WITH_PHYSX
-	if (BodySetup.IsValid())
+	if(BodySetup.IsValid())
 	{
 		FPhysicsCommand::ExecuteWrite(ActorHandle, [&](const FPhysicsActorHandle& Actor)
 		{
@@ -477,7 +473,7 @@ void FBodyInstance::SetResponseToAllChannels(ECollisionResponse NewResponse)
 	CollisionResponses.SetAllChannels(NewResponse);
 	UpdatePhysicsFilterData();
 }
-	
+
 void FBodyInstance::ReplaceResponseToChannels(ECollisionResponse OldResponse, ECollisionResponse NewResponse)
 {
 	InvalidateCollisionProfileName();
@@ -522,7 +518,7 @@ void FBodyInstance::SetCollisionProfileName(FName InCollisionProfileName)
 	//Note that GetCollisionProfileName will use the external profile if one is set.
 	//GetCollisionProfileName will be consistent with the values set by LoadProfileData.
 	//This is why we can't use CollisionProfileName directly during the equality check
-	if (bPendingCollisionProfileSetup || GetCollisionProfileName() != InCollisionProfileName)
+	if(bPendingCollisionProfileSetup || GetCollisionProfileName() != InCollisionProfileName)
 	{
 		//LoadProfileData uses GetCollisionProfileName internally so we must now set the external collision data to null.
 		ExternalCollisionProfileBodySetup = nullptr;
@@ -532,18 +528,18 @@ void FBodyInstance::SetCollisionProfileName(FName InCollisionProfileName)
 
 		bPendingCollisionProfileSetup = false;
 	}
-	
+
 	ExternalCollisionProfileBodySetup = nullptr;	//Even if incoming is the same as GetCollisionProfileName we turn it into "manual mode"
 }
 
 FName FBodyInstance::GetCollisionProfileName() const
 {
 	FName ReturnProfileName = CollisionProfileName;
-	if (UBodySetup* BodySetupPtr = ExternalCollisionProfileBodySetup.Get(true))
+	if(UBodySetup* BodySetupPtr = ExternalCollisionProfileBodySetup.Get(true))
 	{
 		ReturnProfileName = BodySetupPtr->DefaultInstance.CollisionProfileName;
 	}
-	
+
 	return ReturnProfileName;
 }
 
@@ -561,14 +557,14 @@ void FBodyInstance::SetMassScale(float InMassScale)
 
 void FBodyInstance::SetCollisionEnabled(ECollisionEnabled::Type NewType, bool bUpdatePhysicsFilterData)
 {
-	if (CollisionEnabled != NewType)
+	if(CollisionEnabled != NewType)
 	{
 		ECollisionEnabled::Type OldType = CollisionEnabled;
 		InvalidateCollisionProfileName();
 		CollisionEnabled = NewType;
-		
+
 		// Only update physics filter data if required
-		if (bUpdatePhysicsFilterData)
+		if(bUpdatePhysicsFilterData)
 		{
 			UpdatePhysicsFilterData();
 		}
@@ -578,7 +574,7 @@ void FBodyInstance::SetCollisionEnabled(ECollisionEnabled::Type NewType, bool bU
 
 		// Whenever we change physics state, call Recreate
 		// This should also handle destroying the state (in case it's newly disabled).
-		if (bWasPhysicsEnabled != bIsPhysicsEnabled)
+		if(bWasPhysicsEnabled != bIsPhysicsEnabled)
 		{
 			if(UPrimitiveComponent* PrimComponent = OwnerComponent.Get())
 			{
@@ -593,13 +589,13 @@ void FBodyInstance::SetCollisionEnabled(ECollisionEnabled::Type NewType, bool bU
 EDOFMode::Type FBodyInstance::ResolveDOFMode(EDOFMode::Type DOFMode)
 {
 	EDOFMode::Type ResultDOF = DOFMode;
-	if (DOFMode == EDOFMode::Default)
+	if(DOFMode == EDOFMode::Default)
 	{
 		ESettingsDOF::Type SettingDefaultPlane = UPhysicsSettings::Get()->DefaultDegreesOfFreedom;
-		if (SettingDefaultPlane == ESettingsDOF::XYPlane) ResultDOF = EDOFMode::XYPlane;
-		if (SettingDefaultPlane == ESettingsDOF::XZPlane) ResultDOF = EDOFMode::XZPlane;
-		if (SettingDefaultPlane == ESettingsDOF::YZPlane) ResultDOF = EDOFMode::YZPlane;
-		if (SettingDefaultPlane == ESettingsDOF::Full3D) ResultDOF  = EDOFMode::SixDOF;
+		if(SettingDefaultPlane == ESettingsDOF::XYPlane) ResultDOF = EDOFMode::XYPlane;
+		if(SettingDefaultPlane == ESettingsDOF::XZPlane) ResultDOF = EDOFMode::XZPlane;
+		if(SettingDefaultPlane == ESettingsDOF::YZPlane) ResultDOF = EDOFMode::YZPlane;
+		if(SettingDefaultPlane == ESettingsDOF::Full3D) ResultDOF = EDOFMode::SixDOF;
 	}
 
 	return ResultDOF;
@@ -609,15 +605,15 @@ FVector FBodyInstance::GetLockedAxis() const
 {
 	EDOFMode::Type MyDOFMode = ResolveDOFMode(DOFMode);
 
-	switch (MyDOFMode)
+	switch(MyDOFMode)
 	{
-	case EDOFMode::None: return FVector::ZeroVector;
-	case EDOFMode::YZPlane: return FVector(1, 0, 0);
-	case EDOFMode::XZPlane: return FVector(0, 1, 0);
-	case EDOFMode::XYPlane: return FVector(0, 0, 1);
-	case EDOFMode::CustomPlane: return CustomDOFPlaneNormal;
-	case EDOFMode::SixDOF: return FVector::ZeroVector;
-	default:	check(0);	//unsupported locked axis type
+		case EDOFMode::None: return FVector::ZeroVector;
+		case EDOFMode::YZPlane: return FVector(1, 0, 0);
+		case EDOFMode::XZPlane: return FVector(0, 1, 0);
+		case EDOFMode::XYPlane: return FVector(0, 0, 1);
+		case EDOFMode::CustomPlane: return CustomDOFPlaneNormal;
+		case EDOFMode::SixDOF: return FVector::ZeroVector;
+		default:	check(0);	//unsupported locked axis type
 	}
 
 	return FVector::ZeroVector;
@@ -646,7 +642,7 @@ void FBodyInstance::SetDOFLock(EDOFMode::Type NewAxisMode)
 
 void FBodyInstance::CreateDOFLock()
 {
-	if (DOFConstraint)
+	if(DOFConstraint)
 	{
 		DOFConstraint->TermConstraint();
 		FConstraintInstance::Free(DOFConstraint);
@@ -656,13 +652,13 @@ void FBodyInstance::CreateDOFLock()
 	const FVector LockedAxis = GetLockedAxis();
 	const EDOFMode::Type DOF = ResolveDOFMode(DOFMode);
 
-	if (IsDynamic() == false || (LockedAxis.IsNearlyZero() && DOF != EDOFMode::SixDOF))
+	if(IsDynamic() == false || (LockedAxis.IsNearlyZero() && DOF != EDOFMode::SixDOF))
 	{
 		return;
 	}
 
 	//if we're using SixDOF make sure we have at least one constraint
-	if (DOF == EDOFMode::SixDOF && !bLockXTranslation && !bLockYTranslation && !bLockZTranslation && !bLockXRotation && !bLockYRotation && !bLockZRotation)
+	if(DOF == EDOFMode::SixDOF && !bLockXTranslation && !bLockYTranslation && !bLockZTranslation && !bLockXRotation && !bLockYRotation && !bLockZRotation)
 	{
 		return;
 	}
@@ -670,8 +666,8 @@ void FBodyInstance::CreateDOFLock()
 	DOFConstraint = FConstraintInstance::Alloc();
 	{
 		DOFConstraint->ProfileInstance.ConeLimit.bSoftConstraint = false;
-		DOFConstraint->ProfileInstance.TwistLimit.bSoftConstraint  = false;
-		DOFConstraint->ProfileInstance.LinearLimit.bSoftConstraint  = false;
+		DOFConstraint->ProfileInstance.TwistLimit.bSoftConstraint = false;
+		DOFConstraint->ProfileInstance.LinearLimit.bSoftConstraint = false;
 
 		const FTransform TM = GetUnrealWorldTransform(false);
 		FVector Normal = FVector(1, 0, 0);
@@ -683,7 +679,7 @@ void FBodyInstance::CreateDOFLock()
 			DOFConstraint->SetAngularSwing1Motion((bLockRotation || DOFMode != EDOFMode::CustomPlane) ? EAngularConstraintMotion::ACM_Locked : EAngularConstraintMotion::ACM_Free);
 			DOFConstraint->SetAngularSwing2Motion((bLockRotation || DOFMode != EDOFMode::CustomPlane) ? EAngularConstraintMotion::ACM_Locked : EAngularConstraintMotion::ACM_Free);
 			DOFConstraint->SetAngularTwistMotion(EAngularConstraintMotion::ACM_Free);
-			
+
 			DOFConstraint->SetLinearXMotion((bLockTranslation || DOFMode != EDOFMode::CustomPlane) ? ELinearConstraintMotion::LCM_Locked : ELinearConstraintMotion::LCM_Free);
 			DOFConstraint->SetLinearYMotion(ELinearConstraintMotion::LCM_Free);
 			DOFConstraint->SetLinearZMotion(ELinearConstraintMotion::LCM_Free);
@@ -718,26 +714,26 @@ void FBodyInstance::CreateDOFLock()
 ECollisionEnabled::Type FBodyInstance::GetCollisionEnabled_CheckOwner() const
 {
 	// Check actor override
-	if (const UPrimitiveComponent* OwnerComponentInst = OwnerComponent.Get())
+	const UPrimitiveComponent* OwnerComponentInst = OwnerComponent.Get();
+	AActor* Owner = OwnerComponentInst ? OwnerComponentInst->GetOwner() : nullptr;
+	if(Owner && !Owner->GetActorEnableCollision())
 	{
-		AActor* Owner = OwnerComponentInst->GetOwner();
-		if (Owner && !Owner->GetActorEnableCollision())
-		{
-			return ECollisionEnabled::NoCollision;
-		}
-		else if(const USkeletalMeshComponent* SkelMeshComp = Cast<USkeletalMeshComponent>(OwnerComponentInst))
-{
-			// Check component override (skel mesh case)
-			return SkelMeshComp->BodyInstance.CollisionEnabled;
-}
+		return ECollisionEnabled::NoCollision;
 	}
-
-	return CollisionEnabled;
-			}
+	else if(const USkeletalMeshComponent* SkelMeshComp = Cast<USkeletalMeshComponent>(OwnerComponentInst))
+	{
+		// Check component override (skel mesh case)
+		return (ECollisionEnabled::Type)FMath::Max(SkelMeshComp->BodyInstance.CollisionEnabled, CollisionEnabled);
+	}
+	else
+	{
+		return CollisionEnabled;
+	}
+}
 
 void FBodyInstance::SetMaskFilter(FMaskFilter InMaskFilter)
 {
-	if (MaskFilter == InMaskFilter)
+	if(MaskFilter == InMaskFilter)
 	{
 		return;
 	}
@@ -752,13 +748,13 @@ void FBodyInstance::SetMaskFilter(FMaskFilter InMaskFilter)
 			const FBodyInstance* BI = GetOriginalBodyInstance(Shape);
 
 			if(BI == this)
-				{
+			{
 				FPhysicsCommand::ExecuteShapeWrite(this, Shape, [&](const FPhysicsShapeHandle& InnerShape)
-					{
+				{
 					FPhysicsInterface::SetMaskFilter(InnerShape, InMaskFilter);
-					});
-				}
+				});
 			}
+		}
 	});
 
 	MaskFilter = InMaskFilter;
@@ -776,13 +772,13 @@ void FBodyInstance::UpdatePhysicsFilterData()
 	}
 
 	// Do nothing if no physics actor
-	if (!IsValidBodyInstance())
+	if(!IsValidBodyInstance())
 	{
 		return;
 	}
 
 	// this can happen in landscape height field collision component
-	if (!BodySetup.IsValid())
+	if(!BodySetup.IsValid())
 	{
 		return;
 	}
@@ -808,7 +804,7 @@ void FBodyInstance::UpdatePhysicsFilterData()
 			// we need to generate new filter data using the owning original instance (and its BodySetup) 
 			FBodyCollisionData PerShapeCollisionData;
 			if(BI != this)
-	{
+			{
 				BI->BuildBodyFilterData(PerShapeCollisionData.CollisionFilterData);
 
 				const bool bInstanceComplexAsSimple = BI->BodySetup.IsValid() ? (BI->BodySetup->GetCollisionTraceFlag() == CTF_UseComplexAsSimple) : false;
@@ -817,7 +813,7 @@ void FBodyInstance::UpdatePhysicsFilterData()
 			else
 			{
 				PerShapeCollisionData = BodyCollisionData;
-		}
+			}
 
 			FPhysicsCommand::ExecuteShapeWrite(this, Shape, [&](const FPhysicsShapeHandle& InnerShape)
 			{
@@ -837,18 +833,18 @@ void FBodyInstance::UpdatePhysicsFilterData()
 
 				// If we changed 'simulation collision' on a shape, we need to recalc mass properties
 				if(bWasSimulationShape != bNewSimShape)
-		{
+				{
 					bUpdateMassProperties = true;
-		}
+				}
 
 				// Apply new collision settings to this shape
 				FPhysicsInterface::SetSimulationFilter(InnerShape, FilterData.SimFilter);
 				FPhysicsInterface::SetQueryFilter(InnerShape, bIsTrimesh ? FilterData.QueryComplexFilter : FilterData.QuerySimpleFilter);
 			});
-	}
+		}
 
 		if(bUpdateMassProperties)
-	{
+		{
 			UpdateMassProperties();
 		}
 	});
@@ -922,27 +918,21 @@ void GetSimulatingAndBlendWeight(const USkeletalMeshComponent* SkelMeshComp, con
 	}
 }
 
-struct FPhysicsActorCreatePayload_PhysX
-{
-	bool bUseAsync;
-
-};
-
 template <bool bCompileStatic>
 struct FInitBodiesHelper
 {
 	FInitBodiesHelper(TArray<FBodyInstance*>& InBodies, TArray<FTransform>& InTransforms, class UBodySetup* InBodySetup, class UPrimitiveComponent* InPrimitiveComp, FPhysScene* InRBScene, const FBodyInstance::FInitBodySpawnParams& InSpawnParams, FPhysicsAggregateHandle InAggregate)
-	: Bodies(InBodies)
-	, Transforms(InTransforms)
-	, BodySetup(InBodySetup)
-	, PrimitiveComp(InPrimitiveComp)
+		: Bodies(InBodies)
+		, Transforms(InTransforms)
+		, BodySetup(InBodySetup)
+		, PrimitiveComp(InPrimitiveComp)
 		, PhysScene(InRBScene)
 		, Aggregate(InAggregate)
-	, DebugName(TEXT(""))
-	, InstanceBlendWeight(-1.f)
-	, bInstanceSimulatePhysics(false)
-	, SkelMeshComp(nullptr)
-	, SpawnParams(InSpawnParams)
+		, DebugName(TEXT(""))
+		, InstanceBlendWeight(-1.f)
+		, bInstanceSimulatePhysics(false)
+		, SkelMeshComp(nullptr)
+		, SpawnParams(InSpawnParams)
 	{
 		//Compute all the needed constants
 
@@ -954,13 +944,13 @@ struct FInitBodiesHelper
 		{
 			GetSimulatingAndBlendWeight(SkelMeshComp, BodySetup, InstanceBlendWeight, bInstanceSimulatePhysics);
 		}
-		
+
 	}
 
 	FORCEINLINE bool IsStatic() const { return bCompileStatic || bStatic; }
 
 	//The arguments passed into InitBodies
-	TArray<FBodyInstance*>& Bodies;   
+	TArray<FBodyInstance*>& Bodies;
 	TArray<FTransform>& Transforms;
 	class UBodySetup* BodySetup;
 	class UPrimitiveComponent* PrimitiveComp;
@@ -993,7 +983,7 @@ struct FInitBodiesHelper
 		ActorParams.bQueryOnly = bDisableSim;
 		ActorParams.Scene = PhysScene;
 
-		if (IsStatic())
+		if(IsStatic())
 		{
 			ActorParams.bStatic = true;
 
@@ -1001,15 +991,6 @@ struct FInitBodiesHelper
 		}
 		else
 		{
-			if(SpawnParams.DynamicActorScene == EDynamicActorScene::Default)
-			{
-				ActorParams.bUseAsyncScene = Instance->bUseAsyncScene;
-			}
-			else
-			{
-				ActorParams.bUseAsyncScene = (SpawnParams.DynamicActorScene == EDynamicActorScene::UseAsyncScene);
-			}
-
 			Instance->ActorHandle = FPhysicsInterface::CreateActor(ActorParams);
 
 			FPhysicsInterface::SetCcdEnabled_AssumesLocked(Instance->ActorHandle, Instance->bUseCCD);
@@ -1031,41 +1012,12 @@ struct FInitBodiesHelper
 
 		bool bInitFail = false;
 
-		// possibly remove
-		const bool bShapeSharing = Instance->HasSharedShapes(); //If we have a static actor we can reuse the shapes between sync and async scene
-		TArray<FPhysicsShapeHandle> SharedShapes;
-
 		// #PHYS2 Call interface AddGeometry
-		if(FPhysicsInterface::HasSyncSceneData(Instance->ActorHandle))
-		{
-			BodySetup->AddShapesToRigidActor_AssumesLocked(Instance, PST_Sync, Instance->Scale3D, SimplePhysMat, ComplexPhysMats, BodyCollisionData, FTransform::Identity, bShapeSharing ? &SharedShapes : nullptr, bShapeSharing);
+		BodySetup->AddShapesToRigidActor_AssumesLocked(Instance, Instance->Scale3D, SimplePhysMat, ComplexPhysMats, BodyCollisionData, FTransform::Identity);
 
-			int32 NumSyncShapes = 0;
-			int32 NumAsyncShapes = 0;
-			FPhysicsInterface::GetNumShapes(Instance->ActorHandle, NumSyncShapes, NumAsyncShapes);
-			bInitFail |= NumSyncShapes == 0;
-		}
+		const int32 NumShapes = FPhysicsInterface::GetNumShapes(Instance->ActorHandle);
+		bInitFail |= NumShapes == 0;
 
-		if(FPhysicsInterface::HasAsyncSceneData(Instance->ActorHandle))
-		{
-			if (bShapeSharing)
-			{
-				for(FPhysicsShapeHandle& PShape : SharedShapes)
-				{
-					FPhysicsInterface::AttachShape(Instance->ActorHandle, PShape, PST_Async);
-				}
-			}
-			else
-			{
-				BodySetup->AddShapesToRigidActor_AssumesLocked(Instance, PST_Async, Instance->Scale3D, SimplePhysMat, ComplexPhysMats, BodyCollisionData);
-			}
-
-			int32 NumSyncShapes = 0;
-			int32 NumAsyncShapes = 0;
-			FPhysicsInterface::GetNumShapes(Instance->ActorHandle, NumSyncShapes, NumAsyncShapes);
-
-			bInitFail |= NumAsyncShapes == 0;
-		}
 
 		return bInitFail;
 	}
@@ -1081,7 +1033,7 @@ struct FInitBodiesHelper
 		// Ensure we have the AggGeom inside the body setup so we can calculate the number of shapes
 		BodySetup->CreatePhysicsMeshes();
 
-		for (int32 BodyIdx = NumBodies - 1; BodyIdx >= 0; BodyIdx--)   // iterate in reverse since list might shrink
+		for(int32 BodyIdx = NumBodies - 1; BodyIdx >= 0; BodyIdx--)   // iterate in reverse since list might shrink
 		{
 			FBodyInstance* Instance = Bodies[BodyIdx];
 			const FTransform& Transform = Transforms[BodyIdx];
@@ -1092,22 +1044,21 @@ struct FInitBodiesHelper
 			Instance->BodySetup = BodySetup;
 			Instance->Scale3D = Transform.GetScale3D();
 			Instance->CharDebugName = PhysXName;
-			Instance->bHasSharedShapes = IsStatic() && PhysScene && PhysScene->HasAsyncScene() && UPhysicsSettings::Get()->bEnableShapeSharing;
 			Instance->bEnableGravity = Instance->bEnableGravity && (SkelMeshComp ? SkelMeshComp->BodyInstance.bEnableGravity : true);	//In the case of skeletal mesh component we AND bodies with the parent body
 
 			// Handle autowelding here to avoid extra work
-			if (!IsStatic() && Instance->bAutoWeld)
+			if(!IsStatic() && Instance->bAutoWeld)
 			{
 				ECollisionEnabled::Type CollisionType = Instance->GetCollisionEnabled();
-				if (CollisionType != ECollisionEnabled::QueryOnly)
+				if(CollisionType != ECollisionEnabled::QueryOnly)
 				{
-					if (UPrimitiveComponent * ParentPrimComponent = PrimitiveComp ? Cast<UPrimitiveComponent>(PrimitiveComp->GetAttachParent()) : NULL)
+					if(UPrimitiveComponent * ParentPrimComponent = PrimitiveComp ? Cast<UPrimitiveComponent>(PrimitiveComp->GetAttachParent()) : NULL)
 					{
 						UWorld* World = PrimitiveComp->GetWorld();
-						if (World && World->IsGameWorld())
+						if(World && World->IsGameWorld())
 						{
 							//if we have a parent we will now do the weld and exit any further initialization
-							if (PrimitiveComp->WeldToImplementation(ParentPrimComponent, PrimitiveComp->GetAttachSocketName(), false))	//welded new simulated body so initialization is done
+							if(PrimitiveComp->WeldToImplementation(ParentPrimComponent, PrimitiveComp->GetAttachSocketName(), false))	//welded new simulated body so initialization is done
 							{
 								return false;
 							}
@@ -1121,17 +1072,17 @@ struct FInitBodiesHelper
 			if(Instance->GetPhysicsActorHandle().IsValid())
 			{
 				Instance->OwnerComponent = nullptr;
-				Instance->BodySetup      = nullptr;
+				Instance->BodySetup = nullptr;
 				Bodies.RemoveAt(BodyIdx);  // so we wont add it to the physx scene again later.
 				Transforms.RemoveAt(BodyIdx);
 				continue;
 			}
 
 			// Set sim parameters for bodies from skeletal mesh components
-			if (!IsStatic() && SpawnParams.bPhysicsTypeDeterminesSimulation)
+			if(!IsStatic() && SpawnParams.bPhysicsTypeDeterminesSimulation)
 			{
 				Instance->bSimulatePhysics = bInstanceSimulatePhysics;
-				if (InstanceBlendWeight != -1.0f)
+				if(InstanceBlendWeight != -1.0f)
 				{
 					Instance->PhysicsBlendWeight = InstanceBlendWeight;
 				}
@@ -1142,7 +1093,7 @@ struct FInitBodiesHelper
 
 			CreateActor_AssumesLocked(Instance, Transform);
 			const bool bInitFail = CreateShapes_AssumesLocked(Instance);
-			if (bInitFail)
+			if(bInitFail)
 			{
 #if WITH_EDITOR
 				//In the editor we may have ended up here because of world trace ignoring our EnableCollision. Since we can't get at the data in that function we check for it here
@@ -1162,70 +1113,69 @@ struct FInitBodiesHelper
 			}
 
 			FPhysicsInterface::SetActorUserData_AssumesLocked(Instance->ActorHandle, &Instance->PhysxUserData);
-			}
+		}
 
 		return true;
-			}
+	}
 
 
 	void InitBodies()
-				{
+	{
 		LLM_SCOPE(ELLMTag::PhysX);
 
 		check(IsInGameThread());
 
 		if(CreateShapesAndActors())
-					{
+		{
 			FPhysicsCommand::ExecuteWrite(PhysScene, [&]()
-	{
+			{
 				// If an aggregate present, add to that
 				if(Aggregate.IsValid())
-		{
+				{
 					for(FBodyInstance* BI : Bodies)
-			{
+					{
 						const FPhysicsActorHandle& ActorHandle = BI->GetPhysicsActorHandle();
 						if(ActorHandle.IsValid())
-				{
+						{
 							FPhysicsInterface::AddActorToAggregate_AssumesLocked(Aggregate, ActorHandle);
+						}
+					}
 				}
-			}
-		}
 				else if(PhysScene)
-		{
+				{
 					TArray<FPhysicsActorHandle> ActorHandles;
 					ActorHandles.Reserve(Bodies.Num());
 
 					for(FBodyInstance* BI : Bodies)
-			{
+					{
 						const FPhysicsActorHandle& ActorHandle = BI->GetPhysicsActorHandle();
 						if(ActorHandle.IsValid())
-			{
+						{
 							ActorHandles.Add(ActorHandle);
-			}
-		}
+						}
+					}
 
 					PhysScene->AddActorsToScene_AssumesLocked(ActorHandles);
 				}
-				
+
 				// Set up dynamic instance data
-				if (!IsStatic())
+				if(!IsStatic())
 				{
 					SCOPE_CYCLE_COUNTER(STAT_InitBodyPostAdd);
-					for (int32 BodyIdx = 0, NumBodies = Bodies.Num(); BodyIdx < NumBodies; ++BodyIdx)
+					for(int32 BodyIdx = 0, NumBodies = Bodies.Num(); BodyIdx < NumBodies; ++BodyIdx)
 					{
 						FBodyInstance* Instance = Bodies[BodyIdx];
 						Instance->InitDynamicProperties_AssumesLocked();
 					}
 				}
 			});
-			}
+		}
 	}
 };
 
 FBodyInstance::FInitBodySpawnParams::FInitBodySpawnParams(const UPrimitiveComponent* PrimComp)
 {
 	bStaticPhysics = PrimComp == nullptr || PrimComp->Mobility != EComponentMobility::Movable;
-	DynamicActorScene = EDynamicActorScene::Default;
 
 	if(const USkeletalMeshComponent* SKOwner = Cast<USkeletalMeshComponent>(PrimComp))
 	{
@@ -1237,11 +1187,25 @@ FBodyInstance::FInitBodySpawnParams::FInitBodySpawnParams(const UPrimitiveCompon
 	}
 }
 
+// Chaos addition
+static TAutoConsoleVariable<int32> CVarAllowCreatePhysxBodies(
+	TEXT("p.chaos.AllowCreatePhysxBodies"),
+	1,
+	TEXT("")
+	TEXT(" 0 is off, 1 is on (default)"),
+	ECVF_ReadOnly);
+
+
 void FBodyInstance::InitBody(class UBodySetup* Setup, const FTransform& Transform, UPrimitiveComponent* PrimComp, FPhysScene* InRBScene, const FInitBodySpawnParams& SpawnParams)
 {
+	if (CVarAllowCreatePhysxBodies.GetValueOnGameThread() == 0)
+	{
+		return;
+	}
+
 	SCOPE_CYCLE_COUNTER(STAT_InitBody);
 	check(Setup);
-	
+
 	static TArray<FBodyInstance*> Bodies;
 	static TArray<FTransform> Transforms;
 
@@ -1271,12 +1235,12 @@ void FBodyInstance::InitBody(class UBodySetup* Setup, const FTransform& Transfor
 
 FVector GetInitialLinearVelocity(const AActor* OwningActor, bool& bComponentAwake)
 {
-	FVector InitialLinVel(ForceInitToZero);
-	if (OwningActor)
+	FVector InitialLinVel(EForceInit::ForceInitToZero);
+	if(OwningActor)
 	{
 		InitialLinVel = OwningActor->GetVelocity();
 
-		if (InitialLinVel.SizeSquared() > FMath::Square(KINDA_SMALL_NUMBER))
+		if(InitialLinVel.SizeSquared() > FMath::Square(KINDA_SMALL_NUMBER))
 		{
 			bComponentAwake = true;
 		}
@@ -1317,28 +1281,28 @@ void FBodyInstance::TermBody(bool bNeverDeferRelease)
 	OwnerComponent = NULL;
 	ExternalCollisionProfileBodySetup = nullptr;
 
-	if (DOFConstraint)
+	if(DOFConstraint)
 	{
 		DOFConstraint->TermConstraint();
 		FConstraintInstance::Free(DOFConstraint);
-			DOFConstraint = NULL;
+		DOFConstraint = NULL;
 	}
-	
+
 }
 
 bool FBodyInstance::Weld(FBodyInstance* TheirBody, const FTransform& TheirTM)
 {
 	check(IsInGameThread());
 	check(TheirBody);
-	if (TheirBody->BodySetup.IsValid() == false)	//attach actor can be called before body has been initialized. In this case just return false
+	if(TheirBody->BodySetup.IsValid() == false)	//attach actor can be called before body has been initialized. In this case just return false
 	{
 		return false;
 	}
 
-    if (TheirBody->WeldParent == this) // The body is already welded to this component. Do nothing.
-    {
-        return false;
-    }
+	if(TheirBody->WeldParent == this) // The body is already welded to this component. Do nothing.
+	{
+		return false;
+	}
 
 	TArray<FPhysicsShapeHandle> PNewShapes;
 
@@ -1361,16 +1325,7 @@ bool FBodyInstance::Weld(FBodyInstance* TheirBody, const FTransform& TheirTM)
 		BuildBodyFilterData(BodyCollisionData.CollisionFilterData);
 		BuildBodyCollisionFlags(BodyCollisionData.CollisionFlags, GetCollisionEnabled(), BodySetup->GetCollisionTraceFlag() == CTF_UseComplexAsSimple);
 
-		//child body gets placed into the same scenes as parent body
-		if(FPhysicsInterface::HasSyncSceneData(Actor))
-			{
-			TheirBody->BodySetup->AddShapesToRigidActor_AssumesLocked(this, PST_Sync, Scale3D, SimplePhysMat, ComplexPhysMats, BodyCollisionData, RelativeTM, &PNewShapes);
-		}
-
-		if(FPhysicsInterface::HasAsyncSceneData(Actor))
-			{
-			TheirBody->BodySetup->AddShapesToRigidActor_AssumesLocked(this, PST_Async, Scale3D, SimplePhysMat, ComplexPhysMats, BodyCollisionData, RelativeTM, &PNewShapes);
-		}
+		TheirBody->BodySetup->AddShapesToRigidActor_AssumesLocked(this, Scale3D, SimplePhysMat, ComplexPhysMats, BodyCollisionData, RelativeTM, &PNewShapes);
 
 		FPhysicsInterface::SetSendsSleepNotifies_AssumesLocked(Actor, TheirBody->bGenerateWakeEvents);
 
@@ -1381,7 +1336,7 @@ bool FBodyInstance::Weld(FBodyInstance* TheirBody, const FTransform& TheirTM)
 				ShapeToBodiesMap = TSharedPtr<TMap<FPhysicsShapeHandle, FWeldInfo>>(new TMap<FPhysicsShapeHandle, FWeldInfo>());
 			}
 
-			for (int32 ShapeIdx = 0; ShapeIdx < PNewShapes.Num(); ++ShapeIdx)
+			for(int32 ShapeIdx = 0; ShapeIdx < PNewShapes.Num(); ++ShapeIdx)
 			{
 				ShapeToBodiesMap->Add(PNewShapes[ShapeIdx], FWeldInfo(TheirBody, RelativeTM));
 			}
@@ -1411,7 +1366,7 @@ bool FBodyInstance::Weld(FBodyInstance* TheirBody, const FTransform& TheirTM)
 		// remove their body from scenes (don't call TermBody because we don't want to clear things like BodySetup)
 		FPhysicsInterface::ReleaseActor(TheirBody->ActorHandle, TheirBody->GetPhysicsScene());
 	});
-	
+
 	UpdateInterpolateWhenSubStepping();
 
 	TheirBody->UpdateDebugRendering();
@@ -1427,7 +1382,7 @@ void FBodyInstance::UnWeld(FBodyInstance* TheirBI)
 	bool bShapesChanged = false;
 
 	FPhysicsCommand::ExecuteWrite(ActorHandle, [&](const FPhysicsActorHandle& Actor)
-				{
+	{
 		TArray<FPhysicsShapeHandle> Shapes;
 		const int32 NumSyncShapes = GetAllShapes_AssumesLocked(Shapes);
 		const int32 NumTotalShapes = Shapes.Num();
@@ -1435,7 +1390,7 @@ void FBodyInstance::UnWeld(FBodyInstance* TheirBI)
 		for(FPhysicsShapeHandle& Shape : Shapes)
 		{
 			const FBodyInstance* BI = GetOriginalBodyInstance(Shape);
-			if (TheirBI == BI)
+			if(TheirBI == BI)
 			{
 				ShapeToBodiesMap->Remove(Shape);
 				FPhysicsInterface::DetachShape(Actor, Shape);
@@ -1443,10 +1398,10 @@ void FBodyInstance::UnWeld(FBodyInstance* TheirBI)
 			}
 		}
 
-	if (bShapesChanged)
-	{
-		PostShapeChange();
-	}
+		if(bShapesChanged)
+		{
+			PostShapeChange();
+		}
 
 		TheirBI->WeldParent = nullptr;
 	});
@@ -1471,11 +1426,11 @@ float AdjustForSmallThreshold(float NewVal, float OldVal)
 {
 	float Threshold = 0.1f;
 	float Delta = NewVal - OldVal;
-	if (Delta < 0 && FMath::Abs(NewVal) < Threshold)	//getting smaller and passed threshold so flip sign
+	if(Delta < 0 && FMath::Abs(NewVal) < Threshold)	//getting smaller and passed threshold so flip sign
 	{
 		return -Threshold;
 	}
-	else if (Delta > 0 && FMath::Abs(NewVal) < Threshold)	//getting bigger and passed small threshold so flip sign
+	else if(Delta > 0 && FMath::Abs(NewVal) < Threshold)	//getting bigger and passed small threshold so flip sign
 	{
 		return Threshold;
 	}
@@ -1505,35 +1460,35 @@ void ComputeScalingVectors(EScaleMode::Type ScaleMode, const FVector& InScale3D,
 	NewScale3D.Z = FMath::IsNearlyZero(InScale3D.Z) ? KINDA_SMALL_NUMBER : InScale3D.Z;
 
 	const FVector NewScale3DAbs = NewScale3D.GetAbs();
-	switch (ScaleMode)
+	switch(ScaleMode)
 	{
-	case EScaleMode::Free:
-	{
-		OutScale3D = NewScale3D;
-		break;
-	}
-	case EScaleMode::LockedXY:
-	{
-		float XYScaleAbs = FMath::Max(NewScale3DAbs.X, NewScale3DAbs.Y);
-		float XYScale = FMath::Max(NewScale3D.X, NewScale3D.Y) < 0.f ? -XYScaleAbs : XYScaleAbs;	//if both xy are negative we should make the xy scale negative
+		case EScaleMode::Free:
+		{
+			OutScale3D = NewScale3D;
+			break;
+		}
+		case EScaleMode::LockedXY:
+		{
+			float XYScaleAbs = FMath::Max(NewScale3DAbs.X, NewScale3DAbs.Y);
+			float XYScale = FMath::Max(NewScale3D.X, NewScale3D.Y) < 0.f ? -XYScaleAbs : XYScaleAbs;	//if both xy are negative we should make the xy scale negative
 
-		OutScale3D = NewScale3D;
-		OutScale3D.X = OutScale3D.Y = XYScale;
+			OutScale3D = NewScale3D;
+			OutScale3D.X = OutScale3D.Y = XYScale;
 
-		break;
-	}
-	case EScaleMode::LockedXYZ:
-	{
-		float UniformScaleAbs = NewScale3DAbs.GetMin();	//uniform scale uses the smallest magnitude
-		float UniformScale = FMath::Max3(NewScale3D.X, NewScale3D.Y, NewScale3D.Z) < 0.f ? -UniformScaleAbs : UniformScaleAbs;	//if all three values are negative we should make uniform scale negative
+			break;
+		}
+		case EScaleMode::LockedXYZ:
+		{
+			float UniformScaleAbs = NewScale3DAbs.GetMin();	//uniform scale uses the smallest magnitude
+			float UniformScale = FMath::Max3(NewScale3D.X, NewScale3D.Y, NewScale3D.Z) < 0.f ? -UniformScaleAbs : UniformScaleAbs;	//if all three values are negative we should make uniform scale negative
 
-		OutScale3D = FVector(UniformScale);
-		break;
-	}
-	default:
-	{
-		check(false);	//invalid scale mode
-	}
+			OutScale3D = FVector(UniformScale);
+			break;
+		}
+		default:
+		{
+			check(false);	//invalid scale mode
+		}
 	}
 
 	OutScale3DAbs = OutScale3D.GetAbs();
@@ -1570,7 +1525,7 @@ void FBodyInstance::SetMassOverride(float MassInKG, bool bNewOverrideMass)
 
 bool FBodyInstance::GetRigidBodyState(FRigidBodyState& OutState)
 {
-	if (IsInstanceSimulatingPhysics())
+	if(IsInstanceSimulatingPhysics())
 	{
 		FTransform BodyTM = GetUnrealWorldTransform();
 		OutState.Position = BodyTM.GetTranslation();
@@ -1588,14 +1543,14 @@ bool FBodyInstance::UpdateBodyScale(const FVector& InScale3D, bool bForceUpdate)
 {
 	SCOPE_CYCLE_COUNTER(STAT_BodyInstanceUpdateBodyScale);
 
-	if (!IsValidBodyInstance())
+	if(!IsValidBodyInstance())
 	{
 		//UE_LOG(LogPhysics, Log, TEXT("Body hasn't been initialized. Call InitBody to initialize."));
 		return false;
 	}
 
 	// if scale is already correct, and not forcing an update, do nothing
-	if (Scale3D.Equals(InScale3D) && !bForceUpdate)
+	if(Scale3D.Equals(InScale3D) && !bForceUpdate)
 	{
 		return false;
 	}
@@ -1603,7 +1558,7 @@ bool FBodyInstance::UpdateBodyScale(const FVector& InScale3D, bool bForceUpdate)
 	bool bSuccess = false;
 
 #if !(UE_BUILD_SHIPPING || UE_BUILD_TEST)
-	ensureMsgf ( !Scale3D.ContainsNaN() && !InScale3D.ContainsNaN(), TEXT("Scale3D = (%f,%f,%f) InScale3D = (%f,%f,%f)"), Scale3D.X, Scale3D.Y, Scale3D.Z, InScale3D.X, InScale3D.Y, InScale3D.Z );
+	ensureMsgf(!Scale3D.ContainsNaN() && !InScale3D.ContainsNaN(), TEXT("Scale3D = (%f,%f,%f) InScale3D = (%f,%f,%f)"), Scale3D.X, Scale3D.Y, Scale3D.Z, InScale3D.X, InScale3D.Y, InScale3D.Z);
 #endif
 
 	FVector UpdatedScale3D;
@@ -1642,10 +1597,10 @@ bool FBodyInstance::UpdateBodyScale(const FVector& InScale3D, bool bForceUpdate)
 			FKShapeElem* ShapeElem = FPhysxUserData::Get<FKShapeElem>(FPhysicsInterface::GetUserData(Shape));
 			const FTransform& RelativeTM = GetRelativeBodyTransform(Shape);
 
-#if WITH_APEIRON
-			check(false);
+#if WITH_CHAOS
+			//check(false);
 #else
-			switch (GeomType)
+			switch(GeomType)
 			{
 				case ECollisionShapeType::Sphere:
 				{
@@ -1653,13 +1608,13 @@ bool FBodyInstance::UpdateBodyScale(const FVector& InScale3D, bool bForceUpdate)
 					ensure(ScaleMode == EScaleMode::LockedXYZ);
 
 					GeoCollection.GetSphereGeometry(PSphereGeom);
-					 
+
 					PSphereGeom.radius = FMath::Max(SphereElem->Radius * AdjustedScale3DAbs.X, FCollisionShape::MinSphereRadius());
 
 					FVector NewTranslation = RelativeTM.TransformPosition(SphereElem->Center) * AdjustedScale3D.X;
 					LocalTransform.SetTranslation(NewTranslation);
 
-					if (PSphereGeom.isValid())
+					if(PSphereGeom.isValid())
 					{
 						UpdatedGeometry = &PSphereGeom;
 						bSuccess = true;
@@ -1682,7 +1637,7 @@ bool FBodyInstance::UpdateBodyScale(const FVector& InScale3D, bool bForceUpdate)
 					LocalTransform = BoxElem->GetTransform() * RelativeTM;
 					LocalTransform.ScaleTranslation(AdjustedScale3D);
 
-					if (PBoxGeom.isValid())
+					if(PBoxGeom.isValid())
 					{
 						UpdatedGeometry = &PBoxGeom;
 						bSuccess = true;
@@ -1720,7 +1675,7 @@ bool FBodyInstance::UpdateBodyScale(const FVector& InScale3D, bool bForceUpdate)
 					LocalTransform.SetRotation(SphylElem->Rotation.Quaternion() * U2PSphylBasis_UE); // #PHYS2 we probably want to put this behind the interface?
 					LocalTransform.ScaleTranslation(AdjustedScale3D);
 
-					if (PCapsuleGeom.isValid())
+					if(PCapsuleGeom.isValid())
 					{
 						UpdatedGeometry = &PCapsuleGeom;
 						bSuccess = true;
@@ -1745,7 +1700,7 @@ bool FBodyInstance::UpdateBodyScale(const FVector& InScale3D, bool bForceUpdate)
 					LocalTransform.TransformRotation(RelativeTM.GetRotation());
 					LocalTransform.ScaleTranslation(AdjustedScale3D);
 
-					if (PConvexGeom.isValid())
+					if(PConvexGeom.isValid())
 					{
 						UpdatedGeometry = &PConvexGeom;
 						bSuccess = true;
@@ -1765,19 +1720,19 @@ bool FBodyInstance::UpdateBodyScale(const FVector& InScale3D, bool bForceUpdate)
 
 					// find which trimesh elems it is
 					// it would be nice to know if the order of PShapes array index is in the order of createShape
-					if (BodySetup.IsValid())
+					if(BodySetup.IsValid())
 					{
-						for (PxTriangleMesh* TriMesh : BodySetup->TriMeshes)
+						for(PxTriangleMesh* TriMesh : BodySetup->TriMeshes)
 						{
 							// found it
-							if (TriMesh == PTriMeshGeom.triangleMesh)
+							if(TriMesh == PTriMeshGeom.triangleMesh)
 							{
 								PTriMeshGeom.scale.scale = U2PVector(AdjustedScale3D);
 
 								LocalTransform = RelativeTM;
 								LocalTransform.ScaleTranslation(AdjustedScale3D);
 
-								if (PTriMeshGeom.isValid())
+								if(PTriMeshGeom.isValid())
 								{
 									UpdatedGeometry = &PTriMeshGeom;
 									bSuccess = true;
@@ -1799,14 +1754,14 @@ bool FBodyInstance::UpdateBodyScale(const FVector& InScale3D, bool bForceUpdate)
 				}
 				default:
 				{
-						   UE_LOG(LogPhysics, Error, TEXT("Unknown geom type."));
+					UE_LOG(LogPhysics, Error, TEXT("Unknown geom type."));
 				}
 			}// end switch
 #endif
 
-			if (UpdatedGeometry)
+			if(UpdatedGeometry)
 			{
-				FPhysicsCommand::ExecuteShapeWrite(this, Shape, [&](const FPhysicsShapeHandle& InShape)
+				FPhysicsCommand::ExecuteShapeWrite(this, Shape, [&](FPhysicsShapeHandle& InShape)
 				{
 					FPhysicsInterface::SetLocalTransform(InShape, LocalTransform);
 					FPhysicsInterface::SetGeometry(InShape, *UpdatedGeometry);
@@ -1814,7 +1769,7 @@ bool FBodyInstance::UpdateBodyScale(const FVector& InScale3D, bool bForceUpdate)
 
 				UpdatedScale3D = AdjustedScale3D;
 			}
-			else if (bInvalid)
+			else if(bInvalid)
 			{
 #if !(UE_BUILD_SHIPPING || UE_BUILD_TEST)
 				FMessageLog("PIE").Warning()
@@ -1826,16 +1781,16 @@ bool FBodyInstance::UpdateBodyScale(const FVector& InScale3D, bool bForceUpdate)
 			}
 		}
 	});
-	
+
 #endif
 
 	// if success, overwrite old Scale3D, otherwise, just don't do it. It will have invalid scale next time
-	if (bSuccess)
+	if(bSuccess)
 	{
 		Scale3D = UpdatedScale3D;
 
 		// update mass if required
-		if (bUpdateMassWhenScaleChanges)
+		if(bUpdateMassWhenScaleChanges)
 		{
 			UpdateMassProperties();
 		}
@@ -1858,16 +1813,16 @@ void FBodyInstance::UpdateInstanceSimulatePhysics()
 		FPhysicsInterface::SetCcdEnabled_AssumesLocked(Actor, bUseCCD);
 
 		if(bSimulatePhysics && bStartAwake)
-				{
+		{
 			FPhysicsInterface::WakeUp_AssumesLocked(Actor);
 		}
 	});
 
 	//In the original physx only implementation this was wrapped in a PRigidDynamic != NULL check.
 	//We use bInitialized to check rigid actor has been created in either engine because if we haven't even initialized yet, we don't want to undo our settings
-	if (bInitialized)
+	if(bInitialized)
 	{
-		if (bUseSimulate)
+		if(bUseSimulate)
 		{
 			PhysicsBlendWeight = 1.f;
 		}
@@ -1898,11 +1853,11 @@ void FBodyInstance::ApplyWeldOnChildren()
 		TArray<FName> ChildrenLabels;
 		OwnerComponentInst->GetWeldedBodies(ChildrenBodies, ChildrenLabels, /*bIncludingAutoWeld=*/true);
 
-		for (int32 ChildIdx = 0; ChildIdx < ChildrenBodies.Num(); ++ChildIdx)
+		for(int32 ChildIdx = 0; ChildIdx < ChildrenBodies.Num(); ++ChildIdx)
 		{
 			FBodyInstance* ChildBI = ChildrenBodies[ChildIdx];
 			checkSlow(ChildBI);
-			if (ChildBI != this)
+			if(ChildBI != this)
 			{
 				const ECollisionEnabled::Type ChildCollision = ChildBI->GetCollisionEnabled();
 				if(CollisionEnabledHasPhysics(ChildCollision))
@@ -1915,7 +1870,7 @@ void FBodyInstance::ApplyWeldOnChildren()
 			}
 		}
 	}
-	
+
 }
 
 bool FBodyInstance::ShouldInstanceSimulatingPhysics() const
@@ -1926,33 +1881,33 @@ bool FBodyInstance::ShouldInstanceSimulatingPhysics() const
 
 void FBodyInstance::SetInstanceSimulatePhysics(bool bSimulate, bool bMaintainPhysicsBlending)
 {
-	if (bSimulate)
+	if(bSimulate)
 	{
 		UPrimitiveComponent* OwnerComponentInst = OwnerComponent.Get();
 
 		// If we are enabling simulation, and we are the root body of our component (or we are welded), we detach the component 
-		if (OwnerComponentInst && OwnerComponentInst->IsRegistered() && (OwnerComponentInst->GetBodyInstance() == this || OwnerComponentInst->IsWelded()))
+		if(OwnerComponentInst && OwnerComponentInst->IsRegistered() && (OwnerComponentInst->GetBodyInstance() == this || OwnerComponentInst->IsWelded()))
 		{
-			if (OwnerComponentInst->GetAttachParent())
+			if(OwnerComponentInst->GetAttachParent())
 			{
 				OwnerComponentInst->DetachFromComponent(FDetachmentTransformRules::KeepWorldTransform);
 			}
 
-			if (bSimulatePhysics == false)	//if we're switching from kinematic to simulated
+			if(bSimulatePhysics == false)	//if we're switching from kinematic to simulated
 			{
 				ApplyWeldOnChildren();
 			}
 		}
 
 #if !(UE_BUILD_SHIPPING || UE_BUILD_TEST)
-		if (OwnerComponentInst)
+		if(OwnerComponentInst)
 		{
-			if (!IsValidBodyInstance())
+			if(!IsValidBodyInstance())
 			{
 				FMessageLog("PIE").Warning(FText::Format(LOCTEXT("SimPhysNoBody", "Trying to simulate physics on ''{0}'' but no physics body."),
 					FText::FromString(GetPathNameSafe(OwnerComponentInst))));
 			}
-			else if (!IsDynamic())
+			else if(!IsDynamic())
 			{
 				FMessageLog("PIE").Warning(FText::Format(LOCTEXT("SimPhysStatic", "Trying to simulate physics on ''{0}'' but it is static."),
 					FText::FromString(GetPathNameSafe(OwnerComponentInst))));
@@ -1967,9 +1922,9 @@ void FBodyInstance::SetInstanceSimulatePhysics(bool bSimulate, bool bMaintainPhy
 	}
 
 	bSimulatePhysics = bSimulate;
-	if ( !bMaintainPhysicsBlending )
+	if(!bMaintainPhysicsBlending)
 	{
-		if (bSimulatePhysics)
+		if(bSimulatePhysics)
 		{
 			PhysicsBlendWeight = 1.f;
 		}
@@ -2028,7 +1983,7 @@ void FBodyInstance::SetBodyTransform(const FTransform& NewTransform, ETeleportTy
 
 #if !(UE_BUILD_SHIPPING || UE_BUILD_TEST)
 	extern bool GShouldLogOutAFrameOfSetBodyTransform;
-	if (GShouldLogOutAFrameOfSetBodyTransform == true)
+	if(GShouldLogOutAFrameOfSetBodyTransform == true)
 	{
 		UE_LOG(LogPhysics, Log, TEXT("SetBodyTransform: %s"), *GetBodyDebugName());
 	}
@@ -2036,7 +1991,7 @@ void FBodyInstance::SetBodyTransform(const FTransform& NewTransform, ETeleportTy
 
 	// Catch NaNs and elegantly bail out.
 
-	if( !ensureMsgf(!NewTransform.ContainsNaN(), TEXT("SetBodyTransform contains NaN (%s)\n%s"), (OwnerComponent.Get() ? *OwnerComponent->GetPathName() : TEXT("NONE")), *NewTransform.ToString()) )
+	if(!ensureMsgf(!NewTransform.ContainsNaN(), TEXT("SetBodyTransform contains NaN (%s)\n%s"), (OwnerComponent.Get() ? *OwnerComponent->GetPathName() : TEXT("NONE")), *NewTransform.ToString()))
 	{
 		return;
 	}
@@ -2068,27 +2023,27 @@ void FBodyInstance::SetBodyTransform(const FTransform& NewTransform, ETeleportTy
 				const bool bIsSimKinematic = bKinematic && bSimulated;
 
 				if(bIsSimKinematic && Teleport == ETeleportType::None)
-					{
+				{
 					Scene->SetKinematicTarget_AssumesLocked(this, NewTransform, true);
-					}
-					else
-					{
+				}
+				else
+				{
 					if(bIsSimKinematic)
-						{
+					{
 						FPhysicsInterface::SetKinematicTarget_AssumesLocked(Actor, NewTransform);
-						}
-
-					FPhysicsInterface::SetGlobalPose_AssumesLocked(Actor, NewTransform, bAutoWake);
 					}
-			});
+
+					FPhysicsInterface::SetGlobalPose_AssumesLocked(Actor, NewTransform);
 				}
+			});
+		}
 		else if(Scene)
 		{
 			FPhysicsCommand::ExecuteWrite(ActorHandle, [&](const FPhysicsActorHandle& Actor)
 			{
-				FPhysicsInterface::SetGlobalPose_AssumesLocked(Actor, NewTransform, bAutoWake);
+				FPhysicsInterface::SetGlobalPose_AssumesLocked(Actor, NewTransform);
 			});
-			}
+		}
 	}
 	else if(WeldParent)
 	{
@@ -2118,9 +2073,9 @@ FVector FBodyInstance::GetUnrealWorldVelocity_AssumesLocked() const
 {
 	FVector LinVel(EForceInit::ForceInitToZero);
 	if(ActorHandle.IsValid())
-{
+	{
 		LinVel = FPhysicsInterface::GetLinearVelocity_AssumesLocked(ActorHandle);
-}
+	}
 
 	return LinVel;
 }
@@ -2166,7 +2121,7 @@ FTransform FBodyInstance::GetMassSpaceToWorldSpace() const
 	FTransform MassSpaceToWorldSpace = FTransform::Identity;
 
 	FPhysicsCommand::ExecuteRead(ActorHandle, [&](const FPhysicsActorHandle& Actor)
-		{
+	{
 		MassSpaceToWorldSpace = FPhysicsInterface::GetComTransform_AssumesLocked(Actor);
 	});
 
@@ -2233,7 +2188,7 @@ FBox FBodyInstance::GetBodyBounds() const
 
 void FBodyInstance::DrawCOMPosition(FPrimitiveDrawInterface* PDI, float COMRenderSize, const FColor& COMRenderColor)
 {
-	if (IsValidBodyInstance())
+	if(IsValidBodyInstance())
 	{
 		DrawWireStar(PDI, GetCOMPosition(), COMRenderSize, COMRenderColor, SDPG_World);
 	}
@@ -2275,10 +2230,15 @@ FPhysScene* FBodyInstance::GetPhysicsScene() const
 	{
 		return FPhysicsInterface::GetCurrentScene(ActorHandle);
 	}
-	
+
 	return nullptr;
 }
 
+FPhysicsActorHandle& FBodyInstance::GetPhysicsActorHandle()
+{
+	return ActorHandle;
+}
+
 const FPhysicsActorHandle& FBodyInstance::GetPhysicsActorHandle() const
 {
 	return ActorHandle;
@@ -2286,7 +2246,7 @@ const FPhysicsActorHandle& FBodyInstance::GetPhysicsActorHandle() const
 
 const FWalkableSlopeOverride& FBodyInstance::GetWalkableSlopeOverride() const
 {
-	if (bOverrideWalkableSlopeOnInstance || !BodySetup.IsValid())
+	if(bOverrideWalkableSlopeOnInstance || !BodySetup.IsValid())
 	{
 		return WalkableSlopeOverride;
 	}
@@ -2303,12 +2263,12 @@ void FBodyInstance::SetWalkableSlopeOverride(const FWalkableSlopeOverride& NewOv
 }
 
 
-/** 
-*	Changes the current PhysMaterialOverride for this body. 
-*	Note that if physics is already running on this component, this will _not_ alter its mass/inertia etc, it will only change its 
+/**
+*	Changes the current PhysMaterialOverride for this body.
+*	Note that if physics is already running on this component, this will _not_ alter its mass/inertia etc, it will only change its
 *	surface properties like friction and the damping.
 */
-void FBodyInstance::SetPhysMaterialOverride( UPhysicalMaterial* NewPhysMaterial )
+void FBodyInstance::SetPhysMaterialOverride(UPhysicalMaterial* NewPhysMaterial)
 {
 	// Save ref to PhysicalMaterial
 	PhysMaterialOverride = NewPhysMaterial;
@@ -2336,11 +2296,11 @@ UPhysicalMaterial* FBodyInstance::GetSimplePhysicalMaterial(const FBodyInstance*
 
 	// Find the PhysicalMaterial we need to apply to the physics bodies.
 	// (LOW priority) Engine Mat, Material PhysMat, BodySetup Mat, Component Override, Body Override (HIGH priority)
-	
+
 	UPhysicalMaterial* ReturnPhysMaterial = NULL;
 
 	// BodyInstance override
-	if (BodyInstance->PhysMaterialOverride != NULL)
+	if(BodyInstance->PhysMaterialOverride != NULL)
 	{
 		ReturnPhysMaterial = BodyInstance->PhysMaterialOverride;
 		check(!ReturnPhysMaterial || ReturnPhysMaterial->IsValidLowLevel());
@@ -2349,7 +2309,7 @@ UPhysicalMaterial* FBodyInstance::GetSimplePhysicalMaterial(const FBodyInstance*
 	{
 		// Component override
 		UPrimitiveComponent* OwnerPrimComponent = OwnerComp.Get();
-		if (OwnerPrimComponent && OwnerPrimComponent->BodyInstance.PhysMaterialOverride != NULL)
+		if(OwnerPrimComponent && OwnerPrimComponent->BodyInstance.PhysMaterialOverride != NULL)
 		{
 			ReturnPhysMaterial = OwnerComp->BodyInstance.PhysMaterialOverride;
 			check(!ReturnPhysMaterial || ReturnPhysMaterial->IsValidLowLevel());
@@ -2358,7 +2318,7 @@ UPhysicalMaterial* FBodyInstance::GetSimplePhysicalMaterial(const FBodyInstance*
 		{
 			// BodySetup
 			UBodySetup* BodySetupRawPtr = BodySetupPtr.Get();
-			if (BodySetupRawPtr && BodySetupRawPtr->PhysMaterial != NULL)
+			if(BodySetupRawPtr && BodySetupRawPtr->PhysMaterial != NULL)
 			{
 				ReturnPhysMaterial = BodySetupPtr->PhysMaterial;
 				check(!ReturnPhysMaterial || ReturnPhysMaterial->IsValidLowLevel());
@@ -2368,16 +2328,16 @@ UPhysicalMaterial* FBodyInstance::GetSimplePhysicalMaterial(const FBodyInstance*
 				// See if the Material has a PhysicalMaterial
 				UMeshComponent* MeshComp = Cast<UMeshComponent>(OwnerPrimComponent);
 				UPhysicalMaterial* PhysMatFromMaterial = NULL;
-				if (MeshComp != NULL)
+				if(MeshComp != NULL)
 				{
 					UMaterialInterface* Material = MeshComp->GetMaterial(0);
-					if (Material != NULL)
+					if(Material != NULL)
 					{
 						PhysMatFromMaterial = Material->GetPhysicalMaterial();
 					}
 				}
 
-				if (PhysMatFromMaterial != NULL)
+				if(PhysMatFromMaterial != NULL)
 				{
 					ReturnPhysMaterial = PhysMatFromMaterial;
 					check(!ReturnPhysMaterial || ReturnPhysMaterial->IsValidLowLevel());
@@ -2391,7 +2351,7 @@ UPhysicalMaterial* FBodyInstance::GetSimplePhysicalMaterial(const FBodyInstance*
 			}
 		}
 	}
-	
+
 	return ReturnPhysMaterial;
 }
 
@@ -2413,16 +2373,16 @@ void FBodyInstance::GetComplexPhysicalMaterials(const FBodyInstance*, TWeakObjec
 	check(GEngine->DefaultPhysMaterial != NULL);
 	// See if the Material has a PhysicalMaterial
 	UPrimitiveComponent* PrimComp = OwnerComp.Get();
-	if (PrimComp)
+	if(PrimComp)
 	{
 		const int32 NumMaterials = PrimComp->GetNumMaterials();
 		OutPhysicalMaterials.SetNum(NumMaterials);
 
-		for (int32 MatIdx = 0; MatIdx < NumMaterials; MatIdx++)
+		for(int32 MatIdx = 0; MatIdx < NumMaterials; MatIdx++)
 		{
 			UPhysicalMaterial* PhysMat = GEngine->DefaultPhysMaterial;
 			UMaterialInterface* Material = PrimComp->GetMaterial(MatIdx);
-			if (Material)
+			if(Material)
 			{
 				PhysMat = Material->GetPhysicalMaterial();
 			}
@@ -2475,7 +2435,7 @@ PxMassProperties ComputeMassProperties(const FBodyInstance* OwningBodyInstance,
 	// physical material - nothing can weigh less than hydrogen (0.09 kg/m^3)
 	float DensityKGPerCubicUU = 1.0f;
 	float RaiseMassToPower = 0.75f;
-	if (UPhysicalMaterial* PhysMat = OwningBodyInstance->GetSimplePhysicalMaterial())
+	if(UPhysicalMaterial* PhysMat = OwningBodyInstance->GetSimplePhysicalMaterial())
 	{
 		DensityKGPerCubicUU = FMath::Max(KgPerM3ToKgPerCm3(0.09f), gPerCm3ToKgPerCm3(PhysMat->Density));
 		RaiseMassToPower = PhysMat->RaiseMassToPower;
@@ -2487,7 +2447,7 @@ PxMassProperties ComputeMassProperties(const FBodyInstance* OwningBodyInstance,
 	float OldMass = MassProps.mass;
 	float NewMass = 0.f;
 
-	if (OwningBodyInstance->bOverrideMass == false)
+	if(OwningBodyInstance->bOverrideMass == false)
 	{
 		float UsePow = FMath::Clamp<float>(RaiseMassToPower, KINDA_SMALL_NUMBER, 1.f);
 		NewMass = FMath::Pow(OldMass, UsePow);
@@ -2503,7 +2463,7 @@ PxMassProperties ComputeMassProperties(const FBodyInstance* OwningBodyInstance,
 	check(NewMass > 0.f);
 
 	float MassRatio = NewMass / OldMass;
-	
+
 	PxMassProperties FinalMassProps = MassProps * MassRatio;
 
 	FinalMassProps.centerOfMass += U2PVector(MassModifierTransform.TransformVector(OwningBodyInstance->COMNudge));
@@ -2526,17 +2486,11 @@ void FBodyInstance::UpdateMassProperties()
 
 			if(GetNumSimShapes_AssumesLocked(Actor) > 0)
 			{
-				int32 NumSyncShapes = 0;
-				int32 NumAsyncShapes = 0;
-				FPhysicsInterface::GetNumShapes(Actor, NumSyncShapes, NumAsyncShapes);
-
-				bool bSyncData = FPhysicsInterface::HasSyncSceneData(Actor);
-				check(bSyncData || FPhysicsInterface::HasAsyncSceneData(Actor));
+				const int32 NumShapes = FPhysicsInterface::GetNumShapes(Actor);
 
 				TArray<FPhysicsShapeHandle> Shapes;
-				const uint32 NumShapes = bSyncData ? NumSyncShapes : NumAsyncShapes;
 				Shapes.AddUninitialized(NumShapes);
-				FPhysicsInterface::GetAllShapes_AssumedLocked(Actor, Shapes, bSyncData ? PST_Sync : PST_Async);
+				FPhysicsInterface::GetAllShapes_AssumedLocked(Actor, Shapes);
 
 				// Ignore trimeshes & shapes which don't contribute to the mass
 				for(int32 ShapeIdx = Shapes.Num() - 1; ShapeIdx >= 0; --ShapeIdx)
@@ -2639,7 +2593,7 @@ void FBodyInstance::UpdateDebugRendering()
 {
 #if !(UE_BUILD_SHIPPING || UE_BUILD_TEST)
 	//After we update the mass properties, we should update any debug rendering
-	if (UPrimitiveComponent* OwnerPrim = OwnerComponent.Get())
+	if(UPrimitiveComponent* OwnerPrim = OwnerComponent.Get())
 	{
 		OwnerPrim->SendRenderDebugPhysics();
 	}
@@ -2651,7 +2605,7 @@ void FBodyInstance::UpdateDampingProperties()
 	FPhysicsCommand::ExecuteWrite(ActorHandle, [&](const FPhysicsActorHandle& Actor)
 	{
 		if(FPhysicsInterface::IsDynamic(Actor))
-	{
+		{
 			FPhysicsInterface::SetLinearDamping_AssumesLocked(Actor, LinearDamping);
 			FPhysicsInterface::SetAngularDamping_AssumesLocked(Actor, AngularDamping);
 		}
@@ -2665,7 +2619,7 @@ bool FBodyInstance::IsInstanceAwake() const
 	FPhysicsCommand::ExecuteRead(ActorHandle, [&](const FPhysicsActorHandle& Actor)
 	{
 		if(FPhysicsInterface::IsDynamic(Actor))
-	{
+		{
 			bIsAwake = !FPhysicsInterface::IsSleeping(Actor);
 		}
 	});
@@ -2697,11 +2651,11 @@ void FBodyInstance::PutInstanceToSleep()
 
 float FBodyInstance::GetSleepThresholdMultiplier() const
 {
-	if (SleepFamily == ESleepFamily::Sensitive)
+	if(SleepFamily == ESleepFamily::Sensitive)
 	{
 		return 1 / 20.0f;
 	}
-	else if (SleepFamily == ESleepFamily::Custom)
+	else if(SleepFamily == ESleepFamily::Custom)
 	{
 		return CustomSleepThresholdMultiplier;
 	}
@@ -2714,15 +2668,15 @@ void FBodyInstance::SetLinearVelocity(const FVector& NewVel, bool bAddToCurrent,
 	FPhysicsCommand::ExecuteWrite(ActorHandle, [&](const FPhysicsActorHandle& Actor)
 	{
 		if(FPhysicsInterface::IsRigidBody(Actor))
-	{
+		{
 			FVector FinalVelocity = NewVel;
 
-		if (bAddToCurrent)
-		{
+			if(bAddToCurrent)
+			{
 				FinalVelocity += FPhysicsInterface::GetLinearVelocity_AssumesLocked(Actor);
-		}
+			}
 
-			FPhysicsInterface::SetLinearVelocity_AssumesLocked(Actor, FinalVelocity, bAutoWake);
+			FPhysicsInterface::SetLinearVelocity_AssumesLocked(Actor, FinalVelocity);
 		}
 	});
 }
@@ -2733,15 +2687,15 @@ void FBodyInstance::SetAngularVelocityInRadians(const FVector& NewAngVel, bool b
 	FPhysicsCommand::ExecuteWrite(ActorHandle, [&](const FPhysicsActorHandle& Actor)
 	{
 		if(FPhysicsInterface::IsRigidBody(Actor))
-	{
+		{
 			FVector FinalVelocity = NewAngVel;
 
-		if (bAddToCurrent)
-		{
+			if(bAddToCurrent)
+			{
 				FinalVelocity += FPhysicsInterface::GetAngularVelocity_AssumesLocked(Actor);
-		}
+			}
 
-			FPhysicsInterface::SetAngularVelocity_AssumesLocked(Actor, FinalVelocity, bAutoWake);
+			FPhysicsInterface::SetAngularVelocity_AssumesLocked(Actor, FinalVelocity);
 		}
 	});
 }
@@ -2757,7 +2711,7 @@ void FBodyInstance::SetMaxAngularVelocityInRadians(float NewMaxAngVel, bool bAdd
 
 	FPhysicsCommand::ExecuteWrite(ActorHandle, [&](const FPhysicsActorHandle& Actor)
 	{
-		if (bAddToCurrent)
+		if(bAddToCurrent)
 		{
 			float OldValue = FPhysicsInterface::GetMaxAngularVelocity_AssumesLocked(Actor);
 			NewMaxAngVel += OldValue;
@@ -2778,7 +2732,6 @@ void FBodyInstance::SetMaxAngularVelocityInRadians(float NewMaxAngVel, bool bAdd
 
 void FBodyInstance::SetMaxDepenetrationVelocity(float MaxVelocity)
 {
-	bOverrideMaxDepenetrationVelocity = MaxVelocity > 0.f;
 	MaxDepenetrationVelocity = MaxVelocity;
 
 	FPhysicsCommand::ExecuteWrite(ActorHandle, [&](const FPhysicsActorHandle& Actor)
@@ -2872,7 +2825,6 @@ void FBodyInstance::ClearTorques(bool bAllowSubstepping)
 	});
 }
 
-
 void FBodyInstance::AddAngularImpulseInRadians(const FVector& AngularImpulse, bool bVelChange)
 {
 	FPhysicsCommand::ExecuteWrite(ActorHandle, [&](const FPhysicsActorHandle& Actor)
@@ -2880,11 +2832,11 @@ void FBodyInstance::AddAngularImpulseInRadians(const FVector& AngularImpulse, bo
 		if(FPhysicsInterface::IsRigidBody(Actor) && FPhysicsInterface::IsInScene(Actor) && !IsRigidBodyKinematic_AssumesLocked(Actor))
 		{
 			if(bVelChange)
-	{
+			{
 				FPhysicsInterface::AddTorqueMassIndependent_AssumesLocked(Actor, AngularImpulse);
 			}
 			else
-		{
+			{
 				FPhysicsInterface::AddTorque_AssumesLocked(Actor, AngularImpulse);
 			}
 		}
@@ -2898,11 +2850,11 @@ void FBodyInstance::AddImpulse(const FVector& Impulse, bool bVelChange)
 		if(FPhysicsInterface::IsRigidBody(Actor) && FPhysicsInterface::IsInScene(Actor) && !IsRigidBodyKinematic_AssumesLocked(Actor))
 		{
 			if(bVelChange)
-	{
+			{
 				FPhysicsInterface::AddForceMassIndependent_AssumesLocked(Actor, Impulse);
 			}
 			else
-		{
+			{
 				FPhysicsInterface::AddForce_AssumesLocked(Actor, Impulse);
 			}
 		}
@@ -2928,7 +2880,7 @@ void FBodyInstance::SetInstanceNotifyRBCollision(bool bNewNotifyCollision)
 
 void FBodyInstance::SetEnableGravity(bool bInGravityEnabled)
 {
-	if (bEnableGravity != bInGravityEnabled)
+	if(bEnableGravity != bInGravityEnabled)
 	{
 		bEnableGravity = bInGravityEnabled;
 
@@ -2939,10 +2891,10 @@ void FBodyInstance::SetEnableGravity(bool bInGravityEnabled)
 				{
 					FPhysicsInterface::SetGravityEnabled_AssumesLocked(Actor, bEnableGravity);
 				}
-		});
+			});
 		}
 
-		if (bEnableGravity)
+		if(bEnableGravity)
 		{
 			WakeInstance();
 		}
@@ -2951,7 +2903,7 @@ void FBodyInstance::SetEnableGravity(bool bInGravityEnabled)
 
 void FBodyInstance::SetContactModification(bool bNewContactModification)
 {
-	if (bNewContactModification != bContactModification)
+	if(bNewContactModification != bContactModification)
 	{
 		bContactModification = bNewContactModification;
 		UpdatePhysicsFilterData();
@@ -2960,7 +2912,7 @@ void FBodyInstance::SetContactModification(bool bNewContactModification)
 
 void FBodyInstance::SetUseCCD(bool bInUseCCD)
 {
-	if (bUseCCD != bInUseCCD)
+	if(bUseCCD != bInUseCCD)
 	{
 		bUseCCD = bInUseCCD;
 		// Need to set body flag
@@ -3007,17 +2959,17 @@ FString FBodyInstance::GetBodyDebugName() const
 	FString DebugName;
 
 	UPrimitiveComponent* OwnerComponentInst = OwnerComponent.Get();
-	if (OwnerComponentInst != NULL)
+	if(OwnerComponentInst != NULL)
 	{
 		DebugName = OwnerComponentInst->GetPathName();
-		if (const UObject* StatObject = OwnerComponentInst->AdditionalStatObject())
+		if(const UObject* StatObject = OwnerComponentInst->AdditionalStatObject())
 		{
 			DebugName += TEXT(" ");
 			StatObject->AppendName(DebugName);
 		}
 	}
 
-	if ((BodySetup != NULL) && (BodySetup->BoneName != NAME_None))
+	if((BodySetup != NULL) && (BodySetup->BoneName != NAME_None))
 	{
 		DebugName += FString(TEXT(" Bone: ")) + BodySetup->BoneName.ToString();
 	}
@@ -3040,7 +2992,7 @@ bool FBodyInstance::LineTrace(struct FHitResult& OutHit, const FVector& Start, c
 bool FBodyInstance::Sweep(struct FHitResult& OutHit, const FVector& Start, const FVector& End, const FQuat& ShapeWorldRotation, const FCollisionShape& CollisionShape, bool bTraceComplex) const
 {
 	return FPhysicsInterface::Sweep_Geom(OutHit, this, Start, End, ShapeWorldRotation, CollisionShape, bTraceComplex);
-		}
+}
 
 bool FBodyInstance::GetSquaredDistanceToBody(const FVector& Point, float& OutDistanceSquared, FVector& OutPointOnBody) const
 {
@@ -3069,7 +3021,7 @@ bool FBodyInstance::OverlapTestForBodiesImpl(const FVector& Pos, const FQuat& Ro
 
 		for(const FPhysicsShapeHandle& Shape : TargetShapes)
 		{
-#if WITH_APEIRON || WITH_IMMEDIATE_PHYSX
+#if WITH_CHAOS || WITH_IMMEDIATE_PHYSX
 			check(false);
 #else
 
@@ -3081,11 +3033,11 @@ bool FBodyInstance::OverlapTestForBodiesImpl(const FVector& Pos, const FQuat& Ro
 
 			// Calc shape global pose
 			FTransform PShapeGlobalPose = FPhysicsInterface::GetLocalTransform(Shape) * PTestGlobalPose;
-			for (const FBodyInstance* BodyInstance : Bodies)
+			for(const FBodyInstance* BodyInstance : Bodies)
 			{
 				bHaveOverlap = FPhysicsInterface::Overlap_Geom(BodyInstance, FPhysicsInterface::GetGeometryCollection(Shape), PShapeGlobalPose);
 
-				if (bHaveOverlap)
+				if(bHaveOverlap)
 				{
 					return;
 				}
@@ -3120,10 +3072,10 @@ bool FBodyInstance::OverlapTest(const FVector& Position, const FQuat& Rotation,
 
 FTransform RootSpaceToWeldedSpace(const FBodyInstance* BI, const FTransform& RootTM)
 {
-	if (BI->WeldParent)
+	if(BI->WeldParent)
 	{
 		UPrimitiveComponent* BIOwnerComponentInst = BI->OwnerComponent.Get();
-		if (BIOwnerComponentInst)
+		if(BIOwnerComponentInst)
 		{
 			FTransform RootToWelded = BIOwnerComponentInst->GetRelativeTransform().Inverse();
 			RootToWelded.ScaleTranslation(BI->Scale3D);
@@ -3140,7 +3092,7 @@ bool FBodyInstance::OverlapMulti(TArray<struct FOverlapResult>& InOutOverlaps, c
 	SCOPE_CYCLE_COUNTER(STAT_Collision_SceneQueryTotal);
 	SCOPE_CYCLE_COUNTER(STAT_Collision_FBodyInstance_OverlapMulti);
 
-	if ( !IsValidBodyInstance()  && (!WeldParent || !WeldParent->IsValidBodyInstance()))
+	if(!IsValidBodyInstance() && (!WeldParent || !WeldParent->IsValidBodyInstance()))
 	{
 		UE_LOG(LogCollision, Log, TEXT("FBodyInstance::OverlapMulti : (%s) No physics data"), *GetBodyDebugName());
 		return false;
@@ -3152,7 +3104,7 @@ bool FBodyInstance::OverlapMulti(TArray<struct FOverlapResult>& InOutOverlaps, c
 	const FTransform ComponentSpaceToTestSpace(Quat, Pos);
 
 	FTransform BodyInstanceSpaceToTestSpace(NoInit);
-	if (pWorldToComponent)
+	if(pWorldToComponent)
 	{
 		const FTransform RootTM = WeldParent ? WeldParent->GetUnrealWorldTransform() : GetUnrealWorldTransform();
 		const FTransform LocalOffset = (*pWorldToComponent) * RootTM;
@@ -3165,7 +3117,7 @@ bool FBodyInstance::OverlapMulti(TArray<struct FOverlapResult>& InOutOverlaps, c
 
 	//We want to test using global position. However, the global position of the body will be in terms of the root body which we are welded to. So we must undo the relative transform so that our shapes are centered
 	//Global = Parent * Relative => Global * RelativeInverse = Parent
-	if (WeldParent)
+	if(WeldParent)
 	{
 		BodyInstanceSpaceToTestSpace = RootSpaceToWeldedSpace(this, BodyInstanceSpaceToTestSpace);
 	}
@@ -3176,15 +3128,15 @@ bool FBodyInstance::OverlapMulti(TArray<struct FOverlapResult>& InOutOverlaps, c
 	{
 		if(Actor.IsValid())
 		{
-		// Get all the shapes from the actor
+			// Get all the shapes from the actor
 			FInlineShapeArray PShapes;
 			const int32 NumShapes = FillInlineShapeArray_AssumesLocked(PShapes, Actor);
 
 			// Iterate over each shape
 			TArray<struct FOverlapResult> TempOverlaps;
-			for (int32 ShapeIdx = 0; ShapeIdx < NumShapes; ShapeIdx++)
+			for(int32 ShapeIdx = 0; ShapeIdx < NumShapes; ShapeIdx++)
 			{
-#if WITH_APEIRON || WITH_IMMEDIATE_PHYSX
+#if WITH_CHAOS || WITH_IMMEDIATE_PHYSX
 				check(false);
 #else
 				FPhysicsShapeHandle& ShapeRef = PShapes[ShapeIdx];
@@ -3205,7 +3157,7 @@ bool FBodyInstance::OverlapMulti(TArray<struct FOverlapResult>& InOutOverlaps, c
 				// Calc shape global pose
 				const FTransform LocalTransform = FPhysicsInterface::GetLocalTransform(ShapeRef);
 				const FTransform GlobalTransform = LocalTransform * BodyInstanceSpaceToTestSpace;
-			
+
 				TempOverlaps.Reset();
 				if(FPhysicsInterface::GeomOverlapMulti(World, GeomCollection, GlobalTransform.GetTranslation(), GlobalTransform.GetRotation(), TempOverlaps, TestChannel, Params, ResponseParams, ObjectQueryParams))
 				{
@@ -3214,8 +3166,8 @@ bool FBodyInstance::OverlapMulti(TArray<struct FOverlapResult>& InOutOverlaps, c
 				InOutOverlaps.Append(TempOverlaps);
 #endif
 			}
-			}
-		});
+		}
+	});
 
 	return bHaveBlockingHit;
 }
@@ -3228,38 +3180,38 @@ bool FBodyInstance::IsValidCollisionProfileName(FName InCollisionProfileName)
 void FBodyInstance::LoadProfileData(bool bVerifyProfile)
 {
 	const FName UseCollisionProfileName = GetCollisionProfileName();
-	if ( bVerifyProfile )
+	if(bVerifyProfile)
 	{
 		// if collision profile name exists, 
 		// check with current settings
 		// if same, then keep the profile name
 		// if not same, that means it has been modified from default
 		// leave it as it is, and clear profile name
-		if ( IsValidCollisionProfileName(UseCollisionProfileName) )
+		if(IsValidCollisionProfileName(UseCollisionProfileName))
 		{
 			FCollisionResponseTemplate Template;
-			if ( UCollisionProfile::Get()->GetProfileTemplate(UseCollisionProfileName, Template) ) 
+			if(UCollisionProfile::Get()->GetProfileTemplate(UseCollisionProfileName, Template))
 			{
 				// this function is only used for old code that did require verification of using profile or not
 				// so that means it will have valid ResponsetoChannels value, so this is okay to access. 
-				if (Template.IsEqual(CollisionEnabled, ObjectType, CollisionResponses.GetResponseContainer()) == false)
+				if(Template.IsEqual(CollisionEnabled, ObjectType, CollisionResponses.GetResponseContainer()) == false)
 				{
-					InvalidateCollisionProfileName(); 
+					InvalidateCollisionProfileName();
 				}
 			}
 			else
 			{
 				UE_LOG(LogPhysics, Warning, TEXT("COLLISION PROFILE [%s] is not found"), *UseCollisionProfileName.ToString());
 				// if not nothing to do
-				InvalidateCollisionProfileName(); 
+				InvalidateCollisionProfileName();
 			}
 		}
 	}
 	else
 	{
-		if ( IsValidCollisionProfileName(UseCollisionProfileName) )
+		if(IsValidCollisionProfileName(UseCollisionProfileName))
 		{
-			if ( UCollisionProfile::Get()->ReadConfig(UseCollisionProfileName, *this) == false)
+			if(UCollisionProfile::Get()->ReadConfig(UseCollisionProfileName, *this) == false)
 			{
 				// clear the name
 				InvalidateCollisionProfileName();
@@ -3267,10 +3219,10 @@ void FBodyInstance::LoadProfileData(bool bVerifyProfile)
 		}
 
 		// no profile, so it just needs to update container from array data
-		if ( DoesUseCollisionProfile() == false )
+		if(DoesUseCollisionProfile() == false)
 		{
 			// if external profile copy the data over
-			if (ExternalCollisionProfileBodySetup.IsValid(true))
+			if(ExternalCollisionProfileBodySetup.IsValid(true))
 			{
 				UBodySetup* BodySetupInstance = ExternalCollisionProfileBodySetup.Get(true);
 				const FBodyInstance& ExternalBodyInstance = BodySetupInstance->DefaultInstance;
@@ -3294,70 +3246,47 @@ void FBodyInstance::GetBodyInstanceResourceSizeEx(FResourceSizeEx& CumulativeRes
 
 void FBodyInstance::FixupData(class UObject* Loader)
 {
-	check (Loader);
+	check(Loader);
 
 	int32 const UE4Version = Loader->GetLinkerUE4Version();
 
 #if WITH_EDITOR
-	if (UE4Version < VER_UE4_ADD_CUSTOMPROFILENAME_CHANGE)
+	if(UE4Version < VER_UE4_ADD_CUSTOMPROFILENAME_CHANGE)
 	{
-		if (CollisionProfileName == NAME_None)
+		if(CollisionProfileName == NAME_None)
 		{
 			CollisionProfileName = UCollisionProfile::CustomCollisionProfileName;
 		}
 	}
 
-	if (UE4Version < VER_UE4_SAVE_COLLISIONRESPONSE_PER_CHANNEL)
+	if(UE4Version < VER_UE4_SAVE_COLLISIONRESPONSE_PER_CHANNEL)
 	{
 		CollisionResponses.SetCollisionResponseContainer(ResponseToChannels_DEPRECATED);
 	}
 #endif // WITH_EDITORONLY_DATA
 
 	// Load profile. If older version, please verify profile name first
-	bool bNeedToVerifyProfile = (UE4Version < VER_UE4_COLLISION_PROFILE_SETTING) || 
+	bool bNeedToVerifyProfile = (UE4Version < VER_UE4_COLLISION_PROFILE_SETTING) ||
 		// or shape component needs to convert since we added profile
 		(UE4Version < VER_UE4_SAVE_COLLISIONRESPONSE_PER_CHANNEL && Loader->IsA(UShapeComponent::StaticClass()));
 	LoadProfileData(bNeedToVerifyProfile);
 
 	// if profile isn't set, then fix up channel responses
-	if( CollisionProfileName == UCollisionProfile::CustomCollisionProfileName ) 
+	if(CollisionProfileName == UCollisionProfile::CustomCollisionProfileName)
 	{
-		if (UE4Version >= VER_UE4_SAVE_COLLISIONRESPONSE_PER_CHANNEL)
+		if(UE4Version >= VER_UE4_SAVE_COLLISIONRESPONSE_PER_CHANNEL)
 		{
 			CollisionResponses.UpdateResponseContainerFromArray();
 		}
 	}
 }
 
-bool FBodyInstance::UseAsyncScene(const FPhysScene* PhysScene) const
+void FBodyInstance::ApplyMaterialToShape_AssumesLocked(const FPhysicsShapeHandle& InShape, UPhysicalMaterial* SimplePhysMat, const TArrayView<UPhysicalMaterial*>& ComplexPhysMats)
 {
-	bool bHasAsyncScene = true;
-	if (PhysScene)
-	{
-		bHasAsyncScene = PhysScene->HasAsyncScene();
-	}
-
-	return bUseAsyncScene && bHasAsyncScene;
-}
-
-
-void FBodyInstance::SetUseAsyncScene(bool bNewUseAsyncScene)
-{
-	// Set flag
-	bUseAsyncScene = bNewUseAsyncScene;
-}
-
-void FBodyInstance::ApplyMaterialToShape_AssumesLocked(const FPhysicsShapeHandle& InShape, UPhysicalMaterial* SimplePhysMat, const TArrayView<UPhysicalMaterial*>& ComplexPhysMats, const bool bSharedShape)
-{
-	if(!bSharedShape && FPhysicsInterface::IsShared(InShape))	//user says the shape is exclusive, but physx says it's shared
-	{
-		UE_LOG(LogPhysics, Warning, TEXT("FBodyInstance::ApplyMaterialToShape_AssumesLocked : Trying to change the physical material of a shared shape. If this is your intention pass bSharedShape = true"));
-	}
-
 	// If a triangle mesh, need to get array of materials...
 	ECollisionShapeType GeomType = FPhysicsInterface::GetShapeType(InShape);
 	if(GeomType == ECollisionShapeType::Trimesh)
-		{
+	{
 		if(ComplexPhysMats.Num())
 		{
 			FPhysicsInterface::SetMaterials(InShape, ComplexPhysMats);
@@ -3400,12 +3329,12 @@ void FBodyInstance::ApplyMaterialToInstanceShapes_AssumesLocked(UPhysicalMateria
 		if(TheirBI->IsShapeBoundToBody(Shape))
 		{
 			FPhysicsCommand::ExecuteShapeWrite(BIWithActor, Shape, [&](const FPhysicsShapeHandle& InnerShape)
-		{
-				ApplyMaterialToShape_AssumesLocked(InnerShape, SimplePhysMat, ComplexPhysMats, TheirBI->HasSharedShapes());
-		});		
+			{
+				ApplyMaterialToShape_AssumesLocked(InnerShape, SimplePhysMat, ComplexPhysMats);
+			});
+		}
 	}
 }
-}
 
 bool FBodyInstance::ValidateTransform(const FTransform &Transform, const FString& DebugName, const UBodySetup* Setup)
 {
@@ -3442,7 +3371,7 @@ bool FBodyInstance::ValidateTransform(const FTransform &Transform, const FString
 
 void FBodyInstance::InitDynamicProperties_AssumesLocked()
 {
-	if (!BodySetup.IsValid())
+	if(!BodySetup.IsValid())
 	{
 		// This may be invalid following an undo if the BodySetup was a transient object (e.g. in Mesh Paint mode)
 		// Just exit gracefully if so.
@@ -3451,16 +3380,16 @@ void FBodyInstance::InitDynamicProperties_AssumesLocked()
 
 	//QueryOnly bodies cannot become simulated at runtime. To do this they must change their CollisionEnabled which recreates the physics state
 	//So early out to save a lot of useless work
-	if (GetCollisionEnabled() == ECollisionEnabled::QueryOnly)
+	if(GetCollisionEnabled() == ECollisionEnabled::QueryOnly)
 	{
 		return;
 	}
-	
+
 	if(FPhysicsInterface::IsDynamic(ActorHandle))
 	{
 		//A non simulated body may become simulated at runtime, so we need to compute its mass.
 		//However, this is not supported for complexAsSimple since a trimesh cannot itself be simulated, it can only be used for collision of other simple shapes.
-		if (BodySetup->GetCollisionTraceFlag() != ECollisionTraceFlag::CTF_UseComplexAsSimple)
+		if(BodySetup->GetCollisionTraceFlag() != ECollisionTraceFlag::CTF_UseComplexAsSimple)
 		{
 			UpdateMassProperties();
 			UpdateDampingProperties();
@@ -3470,7 +3399,7 @@ void FBodyInstance::InitDynamicProperties_AssumesLocked()
 		else
 		{
 #if !(UE_BUILD_SHIPPING || UE_BUILD_TEST)
-			if (bSimulatePhysics)
+			if(bSimulatePhysics)
 			{
 				if(UPrimitiveComponent* OwnerComponentInst = OwnerComponent.Get())
 				{
@@ -3487,7 +3416,7 @@ void FBodyInstance::InitDynamicProperties_AssumesLocked()
 		bool bComponentAwake = false;
 		FVector InitialLinVel = GetInitialLinearVelocity(OwningActor, bComponentAwake);
 
-		if (ShouldInstanceSimulatingPhysics())
+		if(ShouldInstanceSimulatingPhysics())
 		{
 			FPhysicsInterface::SetLinearVelocity_AssumesLocked(ActorHandle, InitialLinVel);
 		}
@@ -3520,7 +3449,7 @@ void FBodyInstance::InitDynamicProperties_AssumesLocked()
 }
 
 void FBodyInstance::BuildBodyFilterData(FBodyCollisionFilterData& OutFilterData) const
-	{
+{
 	// this can happen in landscape height field collision component
 	if(!BodySetup.IsValid())
 	{
@@ -3551,7 +3480,7 @@ void FBodyInstance::BuildBodyFilterData(FBodyCollisionFilterData& OutFilterData)
 		{
 			UseResponse.SetAllChannels(ECR_Ignore);
 			UseCollisionEnabled = ECollisionEnabled::PhysicsOnly;
-			}
+		}
 		else if(BodySetup->CollisionReponse == EBodyCollisionResponse::BodyCollision_Enabled)
 		{
 			UseResponse.SetAllChannels(ECR_Block);
@@ -3565,7 +3494,7 @@ void FBodyInstance::BuildBodyFilterData(FBodyCollisionFilterData& OutFilterData)
 		const bool bDisableSkelComponentOverride = CVarEnableDynamicPerBodyFilterHacks.GetValueOnGameThread() && bHACK_DisableSkelComponentFilterOverriding;
 		if(!bDisableSkelComponentOverride)
 		{
-		UseResponse = FCollisionResponseContainer::CreateMinContainer(UseResponse, SkelMeshComp->BodyInstance.CollisionResponses.GetResponseContainer());
+			UseResponse = FCollisionResponseContainer::CreateMinContainer(UseResponse, SkelMeshComp->BodyInstance.CollisionResponses.GetResponseContainer());
 		}
 
 		bUseNotifyRBCollision = bUseNotifyRBCollision && SkelMeshComp->BodyInstance.bNotifyRigidBodyCollision;
@@ -3578,7 +3507,7 @@ void FBodyInstance::BuildBodyFilterData(FBodyCollisionFilterData& OutFilterData)
 		UWorld* World = Owner->GetWorld();
 		UPrimitiveComponent* PrimComp = OwnerComponentInst;
 		if(World && World->bEnableTraceCollision &&
-		   (PrimComp->IsA(UStaticMeshComponent::StaticClass()) || PrimComp->IsA(USkeletalMeshComponent::StaticClass()) || PrimComp->IsA(UBrushComponent::StaticClass())))
+			(PrimComp->IsA(UStaticMeshComponent::StaticClass()) || PrimComp->IsA(USkeletalMeshComponent::StaticClass()) || PrimComp->IsA(UBrushComponent::StaticClass())))
 		{
 			//UE_LOG(LogPhysics, Warning, TEXT("Enabling collision %s : %s"), *GetNameSafe(Owner), *GetNameSafe(OwnerComponent.Get()));
 			// clear all other channel just in case other people using those channels to do something
@@ -3599,28 +3528,28 @@ void FBodyInstance::BuildBodyFilterData(FBodyCollisionFilterData& OutFilterData)
 		FCollisionFilterData SimFilterData;
 		FCollisionFilterData SimpleQueryData;
 
-			uint32 ActorID = Owner ? Owner->GetUniqueID() : 0;
-			uint32 CompID = (OwnerComponentInst != nullptr) ? OwnerComponentInst->GetUniqueID() : 0;
+		uint32 ActorID = Owner ? Owner->GetUniqueID() : 0;
+		uint32 CompID = (OwnerComponentInst != nullptr) ? OwnerComponentInst->GetUniqueID() : 0;
 		CreateShapeFilterData(UseChannel, MaskFilter, ActorID, UseResponse, CompID, InstanceBodyIndex, SimpleQueryData, SimFilterData, bRootCCD && !bPhysicsStatic, bUseNotifyRBCollision, bPhysicsStatic, bUseContactModification);
 
 		FCollisionFilterData ComplexQueryData = SimpleQueryData;
-			
-			// Set output sim data
+
+		// Set output sim data
 		OutFilterData.SimFilter = SimFilterData;
 
-			// Build filterdata variations for complex and simple
+		// Build filterdata variations for complex and simple
 		SimpleQueryData.Word3 |= EPDF_SimpleCollision;
-			if(bUseSimpleAsComplex)
-			{
+		if(bUseSimpleAsComplex)
+		{
 			SimpleQueryData.Word3 |= EPDF_ComplexCollision;
-			}
+		}
 
 		ComplexQueryData.Word3 |= EPDF_ComplexCollision;
-			if(bUseComplexAsSimple)
-			{
+		if(bUseComplexAsSimple)
+		{
 			ComplexQueryData.Word3 |= EPDF_SimpleCollision;
-			}
-			
+		}
+
 		OutFilterData.QuerySimpleFilter = SimpleQueryData;
 		OutFilterData.QueryComplexFilter = ComplexQueryData;
 	}
@@ -3653,7 +3582,7 @@ void FBodyInstance::InitStaticBodies(const TArray<FBodyInstance*>& Bodies, const
 
 
 void FBodyInstance::BuildBodyCollisionFlags(FBodyCollisionFlags& OutFlags, ECollisionEnabled::Type UseCollisionEnabled, bool bUseComplexAsSimple)
-	{
+{
 	if(UseCollisionEnabled != ECollisionEnabled::NoCollision)
 	{
 		// Query collision
@@ -3666,15 +3595,15 @@ void FBodyInstance::BuildBodyCollisionFlags(FBodyCollisionFlags& OutFlags, EColl
 		if(bSimCollision)
 		{
 			OutFlags.bEnableSimCollisionSimple = true;
-			
+
 			// on dynamic objects and objects which don't use complex as simple, tri mesh not used for sim
 			if(bUseComplexAsSimple)
 			{
 				OutFlags.bEnableSimCollisionComplex = true;
 			}
 		}
-			}
-		}
+	}
+}
 
 void FBodyInstance::UpdateInterpolateWhenSubStepping()
 {
@@ -3690,7 +3619,7 @@ void FBodyInstance::UpdateInterpolateWhenSubStepping()
 		{
 			UseCollisionEnabled = GetCollisionEnabled();
 		}
-	
+
 		bInterpolateWhenSubStepping = UseCollisionEnabled == ECollisionEnabled::PhysicsOnly || UseCollisionEnabled == ECollisionEnabled::QueryAndPhysics;
 
 		// If we have a weld parent we should take into account that too as that may be simulating while we are not
@@ -3710,7 +3639,7 @@ void FBodyInstance::UpdateInterpolateWhenSubStepping()
 
 void FBodyInstanceEditorHelpers::EnsureConsistentMobilitySimulationSettingsOnPostEditChange(UPrimitiveComponent* Component, FPropertyChangedEvent& PropertyChangedEvent)
 {
-	if (UProperty* PropertyThatChanged = PropertyChangedEvent.Property)
+	if(UProperty* PropertyThatChanged = PropertyChangedEvent.Property)
 	{
 		const FName PropertyName = PropertyThatChanged->GetFName();
 
@@ -3718,15 +3647,15 @@ void FBodyInstanceEditorHelpers::EnsureConsistentMobilitySimulationSettingsOnPos
 		const bool bMobilityChanged = PropertyName == GET_MEMBER_NAME_CHECKED(USceneComponent, Mobility);
 		const bool bSimulatePhysicsChanged = PropertyName == GET_MEMBER_NAME_CHECKED(FBodyInstance, bSimulatePhysics);
 
-		if (bMobilityChanged || bSimulatePhysicsChanged)
+		if(bMobilityChanged || bSimulatePhysicsChanged)
 		{
 			// If we enabled physics simulation, but we are not marked movable, do that for them
-			if (bSimulatePhysicsChanged && Component->BodyInstance.bSimulatePhysics && (Component->Mobility != EComponentMobility::Movable))
+			if(bSimulatePhysicsChanged && Component->BodyInstance.bSimulatePhysics && (Component->Mobility != EComponentMobility::Movable))
 			{
 				Component->SetMobility(EComponentMobility::Movable);
 			}
 			// If we made the component no longer movable, but simulation was enabled, disable that for them
-			else if (bMobilityChanged && (Component->Mobility != EComponentMobility::Movable) && Component->BodyInstance.bSimulatePhysics)
+			else if(bMobilityChanged && (Component->Mobility != EComponentMobility::Movable) && Component->BodyInstance.bSimulatePhysics)
 			{
 				Component->BodyInstance.bSimulatePhysics = false;
 			}
@@ -3734,13 +3663,13 @@ void FBodyInstanceEditorHelpers::EnsureConsistentMobilitySimulationSettingsOnPos
 			// If the collision profile is one of the 'default' ones for a StaticMeshActor, make sure it is the correct one
 			// If user has changed it to something else, don't touch it
 			const FName CurrentProfileName = Component->BodyInstance.GetCollisionProfileName();
-			if ((CurrentProfileName == UCollisionProfile::BlockAll_ProfileName) ||
+			if((CurrentProfileName == UCollisionProfile::BlockAll_ProfileName) ||
 				(CurrentProfileName == UCollisionProfile::BlockAllDynamic_ProfileName) ||
 				(CurrentProfileName == UCollisionProfile::PhysicsActor_ProfileName))
 			{
-				if (Component->Mobility == EComponentMobility::Movable)
+				if(Component->Mobility == EComponentMobility::Movable)
 				{
-					if (Component->BodyInstance.bSimulatePhysics)
+					if(Component->BodyInstance.bSimulatePhysics)
 					{
 						Component->SetCollisionProfileName(UCollisionProfile::PhysicsActor_ProfileName);
 					}
diff --git a/Engine/Source/Runtime/Experimental/Apeiron/Private/Apeiron/PBDCollisionConstraint.cpp b/Engine/Source/Runtime/Experimental/Apeiron/Private/Apeiron/PBDCollisionConstraint.cpp
deleted file mode 100644
diff --git a/Engine/Source/Runtime/Experimental/Chaos/Private/Chaos/PBDCollisionConstraint.cpp b/Engine/Source/Runtime/Experimental/Chaos/Private/Chaos/PBDCollisionConstraint.cpp
new file mode 100644
--- /dev/null
+++ b/Engine/Source/Runtime/Experimental/Chaos/Private/Chaos/PBDCollisionConstraint.cpp
@@ -0,0 +1,1555 @@
+// Copyright 1998-2019 Epic Games, Inc. All Rights Reserved.
+
+#include "Chaos/PBDCollisionConstraint.h"
+
+#include "Chaos/BoundingVolume.h"
+#include "Chaos/BoundingVolumeHierarchy.h"
+#include "Chaos/Defines.h"
+#include "Chaos/Pair.h"
+#include "Chaos/PBDContactGraph.h"
+#include "Chaos/Sphere.h"
+#include "Chaos/Transform.h"
+#include "ChaosStats.h"
+#include "ChaosLog.h"
+#include "ProfilingDebugging/ScopedTimers.h"
+#include "Chaos/ImplicitObjectUnion.h"
+#include "Chaos/PBDCollisionTypes.h"
+
+#define USE_SHOCK_PROPOGATION 1
+
+int32 CollisionParticlesBVHDepth = 4;
+FAutoConsoleVariableRef CVarCollisionParticlesBVHDepth(TEXT("p.CollisionParticlesBVHDepth"), CollisionParticlesBVHDepth, TEXT("The maximum depth for collision particles bvh"));
+
+int32 EnableCollisions = 1;
+FAutoConsoleVariableRef CVarEnableCollisions(TEXT("p.EnableCollisions"), EnableCollisions, TEXT("Enable/Disable collisions on the Chaos solver."));
+
+extern int32 UseLevelsetCollision;
+
+namespace Chaos
+{
+template<class T, int d>
+TPBDCollisionConstraint<T, d>::TPBDCollisionConstraint(TPBDRigidParticles<T, d>& InParticles, TArrayCollectionArray<bool>& Collided, const int32 PushOutIterations /*= 1*/, const int32 PushOutPairIterations /*= 1*/, const T Thickness /*= (T)0*/, const T Restitution /*= (T)0*/, const T Friction /*= (T)0*/)
+    : MCollided(Collided), MContactGraph(InParticles), MNumIterations(PushOutIterations), MPairIterations(PushOutPairIterations), MThickness(Thickness), MRestitution(Restitution), MFriction(Friction), MAngularFriction(0), bUseCCD(false)
+{
+}
+
+DECLARE_CYCLE_STAT(TEXT("ComputeConstraints"), STAT_ComputeConstraints, STATGROUP_Chaos);
+DECLARE_CYCLE_STAT(TEXT("ComputeConstraintsNP"), STAT_ComputeConstraintsNP, STATGROUP_Chaos);
+//DECLARE_CYCLE_STAT(TEXT("ComputeConstraintsNP2"), STAT_ComputeConstraintsNP2, STATGROUP_Chaos);
+//DECLARE_CYCLE_STAT(TEXT("ComputeConstraintsNP3"), STAT_ComputeConstraintsNP3, STATGROUP_Chaos);
+template<class T, int d>
+void TPBDCollisionConstraint<T, d>::ComputeConstraints(const TPBDRigidParticles<T, d>& InParticles, T Dt)
+{
+	SCOPE_CYCLE_COUNTER(STAT_ComputeConstraints);	
+	if (!EnableCollisions) return;
+	// Broad phase
+//	float TotalTested = 0;
+//	float TotalRejected = 0;
+	TBoundingVolume<TPBDRigidParticles<T, d>, T, d> Hierarchy(InParticles, true, Dt * BoundsThicknessMultiplier); 	//todo(ocohen): should we pass MThickness into this structure?
+	{
+		SCOPE_CYCLE_COUNTER(STAT_ComputeConstraintsNP);
+
+		
+		// Narrow phase
+		FCriticalSection CriticalSection;
+		PhysicsParallelFor(InParticles.Size(), [&](int32 Body1Index)
+		{
+			if (InParticles.Disabled(Body1Index))
+			{
+				return;
+			}
+
+			if (InParticles.InvM(Body1Index) == 0)
+			{
+				return;
+			}
+
+			//SCOPE_CYCLE_COUNTER(STAT_ComputeConstraintsNP2);
+
+			TArray<int32> PotentialIntersections;
+
+			const bool bBody1Bounded = InParticles.Geometry(Body1Index)->HasBoundingBox();
+
+			TBox<T, d> Box1;
+			const T Box1Thickness = ComputeThickness(InParticles, Dt, Body1Index).Size();
+
+			if (bBody1Bounded)
+			{
+				Box1 = Hierarchy.GetWorldSpaceBoundingBox(InParticles, Body1Index);
+				PotentialIntersections = Hierarchy.FindAllIntersections(Box1);
+			}
+			else
+			{
+				PotentialIntersections = Hierarchy.GlobalObjects();
+			}
+
+			for (int32 i = 0; i < PotentialIntersections.Num(); ++i)
+			{
+				const int32 Body2Index = PotentialIntersections[i];
+				const bool bBody2Bounded = InParticles.Geometry(Body2Index)->HasBoundingBox();
+
+				if (Body1Index == Body2Index || ((bBody1Bounded == bBody2Bounded) && InParticles.InvM(Body1Index) && InParticles.InvM(Body2Index) && Body2Index > Body1Index))
+				{
+					//if both are dynamic, assume index order matters
+					continue;
+				}
+
+				if (bBody1Bounded && bBody2Bounded)
+				{
+					const TBox<T, d>& Box2 = Hierarchy.GetWorldSpaceBoundingBox(InParticles, Body2Index);
+
+					if (!Box1.Intersects(Box2))
+					{
+						continue;
+					}
+				}
+//				++TotalTested;
+
+
+				const TVector<T, d> Box2Thickness = ComputeThickness(InParticles, Dt, Body2Index);
+				const T UseThickness = FMath::Max(Box1Thickness, Box2Thickness.Size());// + MThickness
+
+				auto Constraint = ComputeConstraint(InParticles, Body1Index, Body2Index, UseThickness);
+
+				//if (true || !InParticles.Geometry(Body1Index)->HasBoundingBox() || !InParticles.Geometry(Body2Index)->HasBoundingBox())
+				{
+					//SCOPE_CYCLE_COUNTER(STAT_ComputeConstraintsNP3);
+					//use narrow phase to determine if constraint is needed. Without this we can't do shock propagation
+					UpdateConstraint(InParticles, UseThickness, Constraint);
+					if (Constraint.Phi < UseThickness)
+					{
+						CriticalSection.Lock();
+						MConstraints.Add(Constraint);
+						CriticalSection.Unlock();
+					}
+//					else
+//					{
+//						++TotalRejected;
+//					}
+				}
+				/*else
+				{
+					CriticalSection.Lock();
+					MConstraints.Add(Constraint);
+					CriticalSection.Unlock();
+				}*/
+
+
+			}
+		});
+	}
+	MContactGraph.ComputeGraph(InParticles, MConstraints);
+//	if (TotalTested > 0)
+//	{
+		//UE_LOG(LogChaos, Warning, TEXT("ComputeConstraints: rejected:%f out of rejected:%f = %.3f"), TotalRejected, TotalTested, TotalRejected / TotalTested);
+//	}
+}
+
+template<class T, int d>
+void TPBDCollisionConstraint<T, d>::RemoveConstraints(const TSet<uint32>& RemovedParticles)
+{
+	for (int32 i = 0; i < MConstraints.Num(); ++i)
+	{
+		const auto& Constraint = MConstraints[i];
+		if (RemovedParticles.Contains(Constraint.ParticleIndex) || RemovedParticles.Contains(Constraint.LevelsetIndex))
+		{
+			MConstraints.RemoveAtSwap(i);
+			i--;
+		}
+	}
+}
+
+DECLARE_CYCLE_STAT(TEXT("UpdateConstraints"), STAT_UpdateConstraints, STATGROUP_Chaos);
+
+template<class T, int d>
+void TPBDCollisionConstraint<T, d>::UpdateConstraints(const TPBDRigidParticles<T, d>& InParticles, T Dt, const TSet<uint32>& AddedParticles, const TArray<uint32>& ActiveParticles)
+{
+	SCOPE_CYCLE_COUNTER(STAT_UpdateConstraints);
+	double Time = 0;
+	FDurationTimer Timer(Time);
+
+	//
+	// Broad phase
+	//
+
+	// @todo(mlentine): We only need to construct the hierarchy for the islands we care about
+	TBoundingVolume<TPBDRigidParticles<T, d>, T, d> Hierarchy(InParticles, ActiveParticles, true, Dt * BoundsThicknessMultiplier); 	//todo(ocohen): should we pass MThickness into this structure?
+	Timer.Stop();
+	UE_LOG(LogChaos, Verbose, TEXT("\tPBDCollisionConstraint Construct Hierarchy %f"), Time);
+
+	//
+	// Narrow phase
+	//
+
+	FCriticalSection CriticalSection;
+	Time = 0;
+	Timer.Start();
+	TArray<uint32> AddedParticlesArray = AddedParticles.Array();
+	PhysicsParallelFor(AddedParticlesArray.Num(), [&](int32 Index) {
+		int32 Body1Index = AddedParticlesArray[Index];
+		if (InParticles.Disabled(Body1Index))
+			return;
+		TArray<int32> PotentialIntersections;
+		TBox<T, d> Box1 = Hierarchy.GetWorldSpaceBoundingBox(InParticles, Body1Index);
+		const T Box1Thickness = ComputeThickness(InParticles, Dt, Body1Index).Size();
+
+		if (InParticles.Geometry(Body1Index)->HasBoundingBox())
+		{
+			PotentialIntersections = Hierarchy.FindAllIntersections(Box1);
+		}
+		else
+		{
+			PotentialIntersections = Hierarchy.GlobalObjects();
+		}
+		for (int32 i = 0; i < PotentialIntersections.Num(); ++i)
+		{
+			int32 Body2Index = PotentialIntersections[i];
+
+			if (Body1Index == Body2Index || ((InParticles.Geometry(Body1Index)->HasBoundingBox() == InParticles.Geometry(Body2Index)->HasBoundingBox()) && AddedParticles.Contains(Body2Index) && AddedParticles.Contains(Body1Index) && Body2Index > Body1Index))
+			{
+				continue;
+			}
+
+			if (InParticles.InvM(Body1Index) && InParticles.InvM(Body2Index) && (InParticles.Island(Body1Index) != InParticles.Island(Body2Index)))	//todo(ocohen): this is a hack - we should not even consider dynamics from other islands
+			{
+				continue;
+			}
+
+			const auto& Box2 = Hierarchy.GetWorldSpaceBoundingBox(InParticles, Body2Index);
+			if (InParticles.Geometry(Body1Index)->HasBoundingBox() && InParticles.Geometry(Body2Index)->HasBoundingBox() && !Box1.Intersects(Box2))
+			{
+				continue;
+			}
+
+			//todo(ocohen): this should not be needed in theory, but in practice we accidentally merge islands. We should be doing this test within an island for clusters
+			if (InParticles.Island(Body1Index) >= 0 && InParticles.Island(Body2Index) >= 0 && InParticles.Island(Body1Index) != InParticles.Island(Body2Index))
+			{
+				continue;
+			}
+
+			const TVector<T, d> Box2Thickness = ComputeThickness(InParticles, Dt, Body2Index);
+			const T UseThickness = FMath::Max(Box1Thickness, Box2Thickness.Size());// + MThickness
+
+			auto Constraint = ComputeConstraint(InParticles, Body1Index, Body2Index, UseThickness);
+
+			//if (true || !InParticles.Geometry(Body1Index)->HasBoundingBox() || !InParticles.Geometry(Body2Index)->HasBoundingBox())
+			{
+				//use narrow phase to determine if constraint is needed. Without this we can't do shock propagation
+				UpdateConstraint(InParticles, UseThickness, Constraint);
+				if (Constraint.Phi < UseThickness)
+				{
+
+					CriticalSection.Lock();
+					MConstraints.Add(Constraint);
+					CriticalSection.Unlock();
+				}
+			}
+			/*else
+			{
+			CriticalSection.Lock();
+			MConstraints.Add(Constraint);
+			CriticalSection.Unlock();
+			}*/
+		}
+	});
+	MContactGraph.Reset(InParticles, MConstraints);
+	Timer.Stop();
+	UE_LOG(LogChaos, Verbose, TEXT("\tPBDCollisionConstraint Update %d Constraints with Potential Collisions %f"), MConstraints.Num(), Time);
+}
+
+template<class T>
+PMatrix<T, 3, 3> ComputeFactorMatrix(const TVector<T, 3>& V, const PMatrix<T, 3, 3>& M, const T& Im)
+{
+	// Rigid objects rotational contribution to the impulse.
+	// Vx*M*VxT+Im
+	check(Im > FLT_MIN)
+	return PMatrix<T, 3, 3>(
+		-V[2] * (-V[2] * M.M[1][1] + V[1] * M.M[2][1]) + V[1] * (-V[2] * M.M[2][1] + V[1] * M.M[2][2]) + Im,
+		V[2] * (-V[2] * M.M[1][0] + V[1] * M.M[2][0]) - V[0] * (-V[2] * M.M[2][1] + V[1] * M.M[2][2]),
+		-V[1] * (-V[2] * M.M[1][0] + V[1] * M.M[2][0]) + V[0] * (-V[2] * M.M[1][1] + V[1] * M.M[2][1]),
+		V[2] * (V[2] * M.M[0][0] - V[0] * M.M[2][0]) - V[0] * (V[2] * M.M[2][0] - V[0] * M.M[2][2]) + Im,
+		-V[1] * (V[2] * M.M[0][0] - V[0] * M.M[2][0]) + V[0] * (V[2] * M.M[1][0] - V[0] * M.M[2][1]),
+		-V[1] * (-V[1] * M.M[0][0] + V[0] * M.M[1][0]) + V[0] * (-V[1] * M.M[1][0] + V[0] * M.M[1][1]) + Im);
+}
+
+template<class T, int d>
+TVector<T, d> GetEnergyClampedImpulse(const TPBDRigidParticles<T, d>& InParticles, const TRigidBodyContactConstraint<T, d>& Constraint, const TVector<T, d>& Impulse, const TVector<T, d>& VectorToPoint1, const TVector<T, d>& VectorToPoint2, const TVector<T, d>& Velocity1, const TVector<T, d>& Velocity2)
+{
+	TVector<T, d> Jr0, Jr1, IInvJr0, IInvJr1;
+	T ImpulseRatioNumerator0 = 0, ImpulseRatioNumerator1 = 0, ImpulseRatioDenom0 = 0, ImpulseRatioDenom1 = 0;
+	T ImpulseSize = Impulse.SizeSquared();
+	TVector<T, d> KinematicVelocity = !InParticles.InvM(Constraint.ParticleIndex) ? Velocity1 : !InParticles.InvM(Constraint.LevelsetIndex) ? Velocity2 : TVector<T, d>(0);
+	if (InParticles.InvM(Constraint.ParticleIndex))
+	{
+		Jr0 = TVector<T, d>::CrossProduct(VectorToPoint1, Impulse);
+		IInvJr0 = InParticles.Q(Constraint.ParticleIndex).RotateVector(InParticles.InvI(Constraint.ParticleIndex) * InParticles.Q(Constraint.ParticleIndex).UnrotateVector(Jr0));
+		ImpulseRatioNumerator0 = TVector<T, d>::DotProduct(Impulse, InParticles.V(Constraint.ParticleIndex) - KinematicVelocity) + TVector<T, d>::DotProduct(IInvJr0, InParticles.I(Constraint.ParticleIndex) * InParticles.W(Constraint.ParticleIndex));
+		ImpulseRatioDenom0 = ImpulseSize / InParticles.M(Constraint.ParticleIndex) + TVector<T, d>::DotProduct(Jr0, IInvJr0);
+	}
+	if (InParticles.InvM(Constraint.LevelsetIndex))
+	{
+		Jr1 = TVector<T, d>::CrossProduct(VectorToPoint2, Impulse);
+		IInvJr1 = InParticles.Q(Constraint.LevelsetIndex).RotateVector(InParticles.InvI(Constraint.LevelsetIndex) * InParticles.Q(Constraint.LevelsetIndex).UnrotateVector(Jr1));
+		ImpulseRatioNumerator1 = TVector<T, d>::DotProduct(Impulse, InParticles.V(Constraint.LevelsetIndex) - KinematicVelocity) + TVector<T, d>::DotProduct(IInvJr1, InParticles.I(Constraint.LevelsetIndex) * InParticles.W(Constraint.LevelsetIndex));
+		ImpulseRatioDenom1 = ImpulseSize / InParticles.M(Constraint.LevelsetIndex) + TVector<T, d>::DotProduct(Jr1, IInvJr1);
+	}
+	T Numerator = -2 * (ImpulseRatioNumerator0 - ImpulseRatioNumerator1);
+	if (Numerator < 0)
+	{
+		return TVector<T, d>(0);
+	}
+	check(Numerator >= 0);
+	T Denominator = ImpulseRatioDenom0 + ImpulseRatioDenom1;
+	return Numerator < Denominator ? (Impulse * Numerator / Denominator) : Impulse;
+}
+
+DECLARE_CYCLE_STAT(TEXT("Apply"), STAT_Apply, STATGROUP_ChaosWide);
+template<class T, int d>
+void TPBDCollisionConstraint<T, d>::Apply(TPBDRigidParticles<T, d>& InParticles, const T Dt, const int32 Island)
+{
+	SCOPE_CYCLE_COUNTER(STAT_Apply);
+	// @todo : The current implementation supports just the no-level approach.
+
+	TArray<typename TPBDContactGraph<FRigidBodyContactConstraint, T, d>::ContactMap>& LevelToColorToContactMap = MContactGraph.GetContactMapAt(Island);
+	int32 MaxColor = MContactGraph.GetMaxColorAt(Island);
+	int32 MaxLevel = MContactGraph.GetMaxLevelAt(Island);
+	for (int32 Level = 0; Level <= MaxLevel; ++Level)
+	{
+		for (int32 i = 0; i <= MaxColor; ++i)
+		{
+			if (LevelToColorToContactMap[Level].Contains(i))
+			{
+				PhysicsParallelFor(LevelToColorToContactMap[Level][i].Num(), [&](int32 ConstraintIndex) {
+					FRigidBodyContactConstraint& Constraint = LevelToColorToContactMap[Level][i][ConstraintIndex];
+					if (InParticles.Sleeping(Constraint.ParticleIndex))
+					{
+						check(InParticles.Sleeping(Constraint.LevelsetIndex) || InParticles.InvM(Constraint.LevelsetIndex) == 0);
+						return;
+					}
+					const_cast<TPBDCollisionConstraint<T, d>*>(this)->UpdateConstraint(InParticles, MThickness, Constraint);
+					if (Constraint.Phi >= MThickness)
+					{
+						return;
+					}
+					MCollided[Constraint.LevelsetIndex] = true;
+					MCollided[Constraint.ParticleIndex] = true;
+					TVector<T, d> VectorToPoint1 = Constraint.Location - InParticles.P(Constraint.ParticleIndex);
+					TVector<T, d> VectorToPoint2 = Constraint.Location - InParticles.P(Constraint.LevelsetIndex);
+					TVector<T, d> Body1Velocity = InParticles.V(Constraint.ParticleIndex) + TVector<T, d>::CrossProduct(InParticles.W(Constraint.ParticleIndex), VectorToPoint1);
+					TVector<T, d> Body2Velocity = InParticles.V(Constraint.LevelsetIndex) + TVector<T, d>::CrossProduct(InParticles.W(Constraint.LevelsetIndex), VectorToPoint2);
+					TVector<T, d> RelativeVelocity = Body1Velocity - Body2Velocity;
+					if (TVector<T, d>::DotProduct(RelativeVelocity, Constraint.Normal) < 0) // ignore separating constraints
+					{
+						PMatrix<T, d, d> WorldSpaceInvI1 = (InParticles.Q(Constraint.ParticleIndex) * FMatrix::Identity) * InParticles.InvI(Constraint.ParticleIndex) * (InParticles.Q(Constraint.ParticleIndex) * FMatrix::Identity).GetTransposed();
+						PMatrix<T, d, d> WorldSpaceInvI2 = (InParticles.Q(Constraint.LevelsetIndex) * FMatrix::Identity) * InParticles.InvI(Constraint.LevelsetIndex) * (InParticles.Q(Constraint.LevelsetIndex) * FMatrix::Identity).GetTransposed();
+						PMatrix<T, d, d> Factor =
+							(InParticles.InvM(Constraint.ParticleIndex) > FLT_MIN ? ComputeFactorMatrix(VectorToPoint1, WorldSpaceInvI1, InParticles.InvM(Constraint.ParticleIndex)) : PMatrix<T, d, d>(0)) +
+							(InParticles.InvM(Constraint.LevelsetIndex) > FLT_MIN ? ComputeFactorMatrix(VectorToPoint2, WorldSpaceInvI2, InParticles.InvM(Constraint.LevelsetIndex)) : PMatrix<T, d, d>(0));
+						TVector<T, d> Impulse;
+						TVector<T, d> AngularImpulse(0);
+						// Resting contact if very close to the surface
+						T Restitution = RelativeVelocity.Size() < (2 * 980 * Dt) ? 0 : MRestitution;
+						if (MFriction)
+						{
+							T RelativeNormalVelocity = TVector<T, d>::DotProduct(RelativeVelocity, Constraint.Normal);
+							if (RelativeNormalVelocity > 0)
+							{
+								RelativeNormalVelocity = 0;
+							}
+							TVector<T, d> VelocityChange = -(Restitution * RelativeNormalVelocity * Constraint.Normal + RelativeVelocity);
+							T NormalVelocityChange = TVector<T, d>::DotProduct(VelocityChange, Constraint.Normal);
+							PMatrix<T, d, d> FactorInverse = Factor.Inverse();
+							TVector<T, d> MinimalImpulse = FactorInverse * VelocityChange;
+							// Friction should stop the object
+							if ((VelocityChange - NormalVelocityChange * Constraint.Normal).Size() <= MFriction * NormalVelocityChange)
+							{
+								Impulse = MinimalImpulse;
+								if (MAngularFriction)
+								{
+									TVector<T, d> RelativeAngularVelocity = InParticles.W(Constraint.ParticleIndex) - InParticles.W(Constraint.LevelsetIndex);
+									T AngularNormal = TVector<T, d>::DotProduct(RelativeAngularVelocity, Constraint.Normal);
+									TVector<T, d> AngularTangent = RelativeAngularVelocity - AngularNormal * Constraint.Normal;
+									TVector<T, d> FinalAngularVelocity = FMath::Sign(AngularNormal) * FMath::Max((T)0, FMath::Abs(AngularNormal) - MAngularFriction * NormalVelocityChange) * Constraint.Normal + FMath::Max((T)0, AngularTangent.Size() - MAngularFriction * NormalVelocityChange) * AngularTangent.GetSafeNormal();
+									TVector<T, d> Delta = FinalAngularVelocity - RelativeAngularVelocity;
+									if (!InParticles.InvM(Constraint.ParticleIndex))
+									{
+										PMatrix<T, d, d> WorldSpaceI2 = (InParticles.Q(Constraint.LevelsetIndex) * FMatrix::Identity) * InParticles.I(Constraint.LevelsetIndex) * (InParticles.Q(Constraint.LevelsetIndex) * FMatrix::Identity).GetTransposed();
+										TVector<T, d> ImpulseDelta = InParticles.M(Constraint.LevelsetIndex) * TVector<T, d>::CrossProduct(VectorToPoint2, Delta);
+										Impulse += ImpulseDelta;
+										AngularImpulse += WorldSpaceI2 * Delta - TVector<T, d>::CrossProduct(VectorToPoint2, ImpulseDelta);
+									}
+									else if (!InParticles.InvM(Constraint.LevelsetIndex))
+									{
+										PMatrix<T, d, d> WorldSpaceI1 = (InParticles.Q(Constraint.ParticleIndex) * FMatrix::Identity) * InParticles.I(Constraint.ParticleIndex) * (InParticles.Q(Constraint.ParticleIndex) * FMatrix::Identity).GetTransposed();
+										TVector<T, d> ImpulseDelta = InParticles.M(Constraint.ParticleIndex) * TVector<T, d>::CrossProduct(VectorToPoint1, Delta);
+										Impulse += ImpulseDelta;
+										AngularImpulse += WorldSpaceI1 * Delta - TVector<T, d>::CrossProduct(VectorToPoint1, ImpulseDelta);
+									}
+									else
+									{
+										PMatrix<T, d, d> Cross1(0, VectorToPoint1.Z, -VectorToPoint1.Y, -VectorToPoint1.Z, 0, VectorToPoint1.X, VectorToPoint1.Y, -VectorToPoint1.X, 0);
+										PMatrix<T, d, d> Cross2(0, VectorToPoint2.Z, -VectorToPoint2.Y, -VectorToPoint2.Z, 0, VectorToPoint2.X, VectorToPoint2.Y, -VectorToPoint2.X, 0);
+										PMatrix<T, d, d> CrossI1 = Cross1 * WorldSpaceInvI1;
+										PMatrix<T, d, d> CrossI2 = Cross2 * WorldSpaceInvI2;
+										PMatrix<T, d, d> Diag1 = CrossI1 * Cross1.GetTransposed() + CrossI2 * Cross2.GetTransposed();
+										Diag1.M[0][0] += InParticles.InvM(Constraint.ParticleIndex) + InParticles.InvM(Constraint.LevelsetIndex);
+										Diag1.M[1][1] += InParticles.InvM(Constraint.ParticleIndex) + InParticles.InvM(Constraint.LevelsetIndex);
+										Diag1.M[2][2] += InParticles.InvM(Constraint.ParticleIndex) + InParticles.InvM(Constraint.LevelsetIndex);
+										PMatrix<T, d, d> OffDiag1 = (CrossI1 + CrossI2) * -1;
+										PMatrix<T, d, d> Diag2 = (WorldSpaceInvI1 + WorldSpaceInvI2).Inverse();
+										PMatrix<T, d, d> OffDiag1Diag2 = OffDiag1 * Diag2;
+										TVector<T, d> ImpulseDelta = PMatrix<T, d, d>((Diag1 - OffDiag1Diag2 * OffDiag1.GetTransposed()).Inverse()) * ((OffDiag1Diag2 * -1) * Delta);
+										Impulse += ImpulseDelta;
+										AngularImpulse += Diag2 * (Delta - PMatrix<T, d, d>(OffDiag1.GetTransposed()) * ImpulseDelta);
+									}
+								}
+							}
+							else
+							{
+								TVector<T, d> Tangent = (RelativeVelocity - TVector<T, d>::DotProduct(RelativeVelocity, Constraint.Normal) * Constraint.Normal).GetSafeNormal();
+								TVector<T, d> DirectionalFactor = Factor * (Constraint.Normal - MFriction * Tangent);
+								T ImpulseDenominator = TVector<T, d>::DotProduct(Constraint.Normal, DirectionalFactor);
+								if (!ensureMsgf(FMath::Abs(ImpulseDenominator) > SMALL_NUMBER, TEXT("Constraint:%s\n\nParticle:%s\n\nLevelset:%s\n\nDirectionalFactor:%s, ImpulseDenominator:%f"),
+									*Constraint.ToString(),
+									*InParticles.ToString(Constraint.ParticleIndex),
+									*InParticles.ToString(Constraint.LevelsetIndex),
+									*DirectionalFactor.ToString(), ImpulseDenominator))
+								{
+									ImpulseDenominator = (T)1;
+								}
+
+								Impulse = FactorInverse * -(1 + Restitution) * RelativeNormalVelocity * DirectionalFactor / ImpulseDenominator;
+							}
+						}
+						else
+						{
+							T ImpulseDenominator = TVector<T, d>::DotProduct(Constraint.Normal, Factor * Constraint.Normal);
+							TVector<T, d> ImpulseNumerator = -(1 + Restitution) * TVector<T, d>::DotProduct(RelativeVelocity, Constraint.Normal) * Constraint.Normal;
+							if (!ensureMsgf(FMath::Abs(ImpulseDenominator) > SMALL_NUMBER, TEXT("Constraint:%s\n\nParticle:%s\n\nLevelset:%s\n\nFactor*Constraint.Normal:%s, ImpulseDenominator:%f"),
+								*Constraint.ToString(),
+								*InParticles.ToString(Constraint.ParticleIndex),
+								*InParticles.ToString(Constraint.LevelsetIndex),
+								*(Factor*Constraint.Normal).ToString(), ImpulseDenominator))
+							{
+								ImpulseDenominator = (T)1;
+							}
+							Impulse = ImpulseNumerator / ImpulseDenominator;
+						}
+						Impulse = GetEnergyClampedImpulse(InParticles, Constraint, Impulse, VectorToPoint1, VectorToPoint2, Body1Velocity, Body2Velocity);
+						Constraint.AccumulatedImpulse += Impulse;
+						TVector<T, d> AngularImpulse1 = TVector<T, d>::CrossProduct(VectorToPoint1, Impulse) + AngularImpulse;
+						TVector<T, d> AngularImpulse2 = TVector<T, d>::CrossProduct(VectorToPoint2, -Impulse) - AngularImpulse;
+						// Velocity update for next step
+						InParticles.V(Constraint.ParticleIndex) += InParticles.InvM(Constraint.ParticleIndex) * Impulse;
+						InParticles.W(Constraint.ParticleIndex) += WorldSpaceInvI1 * AngularImpulse1;
+						InParticles.V(Constraint.LevelsetIndex) -= InParticles.InvM(Constraint.LevelsetIndex) * Impulse;
+						InParticles.W(Constraint.LevelsetIndex) += WorldSpaceInvI2 * AngularImpulse2;
+						// Position update as part of pbd
+						InParticles.P(Constraint.ParticleIndex) += (InParticles.InvM(Constraint.ParticleIndex) * Impulse) * Dt;
+						InParticles.Q(Constraint.ParticleIndex) += TRotation<T, d>(WorldSpaceInvI1 * AngularImpulse1, 0.f) * InParticles.Q(Constraint.ParticleIndex) * Dt * T(0.5);
+						InParticles.Q(Constraint.ParticleIndex).Normalize();
+						InParticles.P(Constraint.LevelsetIndex) -= (InParticles.InvM(Constraint.LevelsetIndex) * Impulse) * Dt;
+						InParticles.Q(Constraint.LevelsetIndex) += TRotation<T, d>(WorldSpaceInvI2 * AngularImpulse2, 0.f) * InParticles.Q(Constraint.LevelsetIndex) * Dt * T(0.5);
+						InParticles.Q(Constraint.LevelsetIndex).Normalize();
+					}
+				});
+			}
+		}
+	}
+}
+
+DECLARE_CYCLE_STAT(TEXT("ApplyPushOut"), STAT_ApplyPushOut, STATGROUP_ChaosWide);
+template<class T, int d>
+void TPBDCollisionConstraint<T, d>::ApplyPushOut(TPBDRigidParticles<T, d>& InParticles, const T Dt, const TArray<int32>& ActiveIndices, const int32 Island)
+{
+	SCOPE_CYCLE_COUNTER(STAT_ApplyPushOut);
+	// @todo : The current implementation supports just the no-level approach.
+
+	TArray<typename TPBDContactGraph<FRigidBodyContactConstraint, T, d>::ContactMap>& LevelToColorToContactMap = MContactGraph.GetContactMapAt(Island);
+	int32 MaxColor = MContactGraph.GetMaxColorAt(Island);
+	int32 MaxLevel = MContactGraph.GetMaxLevelAt(Island);
+
+	bool NeedsAnotherIteration = false;
+	TArray<bool> IsTemporarilyStatic;
+	IsTemporarilyStatic.Init(false, InParticles.Size());
+	for (int32 Iteration = 0; (Iteration == 0 || NeedsAnotherIteration) && Iteration < MNumIterations; ++Iteration)
+	{
+		NeedsAnotherIteration = false;
+		for (int32 Level = 0; Level <= MaxLevel; ++Level)
+		{
+			for (int32 i = 0; i <= MaxColor; ++i)
+			{
+				if (LevelToColorToContactMap[Level].Contains(i))
+				{
+					PhysicsParallelFor(LevelToColorToContactMap[Level][i].Num(), [&](int32 ConstraintIndex) {
+						TRigidBodyContactConstraint<T, d>& Constraint = LevelToColorToContactMap[Level][i][ConstraintIndex];
+						if (InParticles.Sleeping(Constraint.ParticleIndex))
+						{
+							check(InParticles.Sleeping(Constraint.LevelsetIndex) || InParticles.InvM(Constraint.LevelsetIndex) == 0);
+							return;
+						}
+						for (int32 PairIteration = 0; PairIteration < MPairIterations; ++PairIteration)
+						{
+							UpdateConstraint(InParticles, MThickness, Constraint);
+							if (Constraint.Phi >= MThickness)
+							{
+								break;
+							}
+							NeedsAnotherIteration = true;
+							PMatrix<T, d, d> WorldSpaceInvI1 = (InParticles.Q(Constraint.ParticleIndex) * FMatrix::Identity) * InParticles.InvI(Constraint.ParticleIndex) * (InParticles.Q(Constraint.ParticleIndex) * FMatrix::Identity).GetTransposed();
+							PMatrix<T, d, d> WorldSpaceInvI2 = (InParticles.Q(Constraint.LevelsetIndex) * FMatrix::Identity) * InParticles.InvI(Constraint.LevelsetIndex) * (InParticles.Q(Constraint.LevelsetIndex) * FMatrix::Identity).GetTransposed();
+							TVector<T, d> VectorToPoint1 = Constraint.Location - InParticles.P(Constraint.ParticleIndex);
+							TVector<T, d> VectorToPoint2 = Constraint.Location - InParticles.P(Constraint.LevelsetIndex);
+							PMatrix<T, d, d> Factor =
+								((InParticles.InvM(Constraint.ParticleIndex) && !IsTemporarilyStatic[Constraint.ParticleIndex]) ? ComputeFactorMatrix(VectorToPoint1, WorldSpaceInvI1, InParticles.InvM(Constraint.ParticleIndex)) : PMatrix<T, d, d>(0)) +
+								((InParticles.InvM(Constraint.LevelsetIndex) && !IsTemporarilyStatic[Constraint.LevelsetIndex]) ? ComputeFactorMatrix(VectorToPoint2, WorldSpaceInvI2, InParticles.InvM(Constraint.LevelsetIndex)) : PMatrix<T, d, d>(0));
+							T Numerator = FMath::Min((T)(Iteration + 2), (T)MNumIterations);
+							T ScalingFactor = Numerator / (T)MNumIterations;
+							TVector<T, d> Impulse = PMatrix<T, d, d>(Factor.Inverse()) * ((-Constraint.Phi + MThickness) * ScalingFactor * Constraint.Normal);
+							TVector<T, d> AngularImpulse1 = TVector<T, d>::CrossProduct(VectorToPoint1, Impulse);
+							TVector<T, d> AngularImpulse2 = TVector<T, d>::CrossProduct(VectorToPoint2, -Impulse);
+							if (!IsTemporarilyStatic[Constraint.ParticleIndex])
+							{
+								InParticles.P(Constraint.ParticleIndex) += InParticles.InvM(Constraint.ParticleIndex) * Impulse;
+								InParticles.Q(Constraint.ParticleIndex) = TRotation<T, d>::FromVector(WorldSpaceInvI1 * AngularImpulse1) * InParticles.Q(Constraint.ParticleIndex);
+								InParticles.Q(Constraint.ParticleIndex).Normalize();
+							}
+							if (!IsTemporarilyStatic[Constraint.LevelsetIndex])
+							{
+								InParticles.P(Constraint.LevelsetIndex) -= InParticles.InvM(Constraint.LevelsetIndex) * Impulse;
+								InParticles.Q(Constraint.LevelsetIndex) = TRotation<T, d>::FromVector(WorldSpaceInvI2 * AngularImpulse2) * InParticles.Q(Constraint.LevelsetIndex);
+								InParticles.Q(Constraint.LevelsetIndex).Normalize();
+							}
+						}
+					});
+				}
+			}
+#if USE_SHOCK_PROPOGATION
+			for (int32 i = 0; i <= MaxColor; ++i)
+			{
+				if (LevelToColorToContactMap[Level].Contains(i))
+				{
+					PhysicsParallelFor(LevelToColorToContactMap[Level][i].Num(), [&](int32 ConstraintIndex) {
+						auto Constraint = LevelToColorToContactMap[Level][i][ConstraintIndex];
+						if (Iteration == MNumIterations - 1)
+						{
+							if (!InParticles.InvM(Constraint.ParticleIndex) || IsTemporarilyStatic[Constraint.ParticleIndex])
+							{
+								IsTemporarilyStatic[Constraint.LevelsetIndex] = true;
+							}else if (!InParticles.InvM(Constraint.LevelsetIndex) || IsTemporarilyStatic[Constraint.LevelsetIndex])
+							{
+								IsTemporarilyStatic[Constraint.ParticleIndex] = true;
+							}
+						}
+					});
+				}
+			}
+#endif
+		}
+	}
+}
+
+DECLARE_CYCLE_STAT(TEXT("CopyOutConstraints"), STAT_CopyOutConstraints, STATGROUP_Chaos);
+
+template<class T, int d>
+void TPBDCollisionConstraint<T, d>::CopyOutConstraints(int32 NumIslands)
+{
+	SCOPE_CYCLE_COUNTER(STAT_CopyOutConstraints);
+	MConstraints.Reset();
+	for (int32 Island = 0; Island < NumIslands; ++Island)
+	{
+		//todo(ocohen): could be part of the parallel for, but need output buffer to be thread safe
+		const TArray<typename TPBDContactGraph<FRigidBodyContactConstraint, T, d>::ContactMap>& LevelToColorToContactMap = MContactGraph.GetContactMapAt(Island);
+		int32 MaxColor = MContactGraph.GetMaxColorAt(Island);
+		int32 MaxLevel = MContactGraph.GetMaxLevelAt(Island);
+
+		for (int32 Level = 0; Level <= MaxLevel; ++Level)
+		{
+			for (int32 Color = 0; Color <= MaxColor; ++Color)
+			{
+				if (LevelToColorToContactMap[Level].Contains(Color))
+				{
+					MConstraints.Append(LevelToColorToContactMap[Level][Color]);
+				}
+			}
+		}
+	}
+}
+
+template<class T, int d>
+bool TPBDCollisionConstraint<T, d>::NearestPoint(TArray<Pair<TVector<T, d>, TVector<T, d>>>& Points, TVector<T, d>& Direction, TVector<T, d>& ClosestPoint)
+{
+	check(Points.Num() > 1 && Points.Num() <= 4);
+	if (Points.Num() == 2)
+	{
+		TPlane<T, d> LocalPlane(Points[1].First, Points[0].First - Points[1].First);
+		TVector<T, d> Normal;
+		const auto& Phi = LocalPlane.PhiWithNormal(TVector<T, d>(0), Normal);
+		if ((TVector<T, d>::DotProduct(-Points[1].First, Normal.GetSafeNormal()) - Points[1].First.Size()) < SMALL_NUMBER)
+		{
+			T Alpha = Points[0].First.Size() / (Points[1].First - Points[0].First).Size();
+			ClosestPoint = (1 - Alpha) * Points[0].Second + Alpha * Points[1].Second;
+			return true;
+		}
+		if (Phi > 0)
+		{
+			check(Points.Num() == 2);
+			Direction = TVector<T, d>::CrossProduct(TVector<T, d>::CrossProduct(Normal, -Points[1].First), Normal);
+		}
+		else
+		{
+			Direction = -Points[1].First;
+			Points.RemoveAtSwap(0);
+			check(Points.Num() == 1);
+		}
+		check(Points.Num() > 1 && Points.Num() < 4);
+		return false;
+	}
+	if (Points.Num() == 3)
+	{
+		TVector<T, d> TriangleNormal = TVector<T, d>::CrossProduct(Points[0].First - Points[2].First, Points[0].First - Points[1].First);
+		TPlane<T, d> LocalPlane1(Points[2].First, TVector<T, d>::CrossProduct(Points[0].First - Points[2].First, TriangleNormal));
+		TPlane<T, d> LocalPlane2(Points[2].First, TVector<T, d>::CrossProduct(Points[1].First - Points[2].First, TriangleNormal));
+		TVector<T, d> Normal;
+		T Phi = LocalPlane1.PhiWithNormal(TVector<T, d>(0), Normal);
+		if (Phi > 0)
+		{
+			TVector<T, d> Delta = Points[0].First - Points[2].First;
+			if (TVector<T, d>::DotProduct(-Points[2].First, Delta) > 0)
+			{
+				Direction = TVector<T, d>::CrossProduct(TVector<T, d>::CrossProduct(Delta, -Points[2].First), Delta);
+				Points.RemoveAtSwap(1);
+				check(Points.Num() == 2);
+			}
+			else
+			{
+				Delta = Points[1].First - Points[2].First;
+				if (TVector<T, d>::DotProduct(-Points[2].First, Delta) > 0)
+				{
+					Direction = TVector<T, d>::CrossProduct(TVector<T, d>::CrossProduct(Delta, -Points[2].First), Delta);
+					Points.RemoveAtSwap(0);
+					check(Points.Num() == 2);
+				}
+				else
+				{
+					Direction = -Points[2].First;
+					Points = {Points[2]};
+					check(Points.Num() == 1);
+				}
+			}
+		}
+		else
+		{
+			Phi = LocalPlane2.PhiWithNormal(TVector<T, d>(0), Normal);
+			if (Phi > 0)
+			{
+				TVector<T, d> Delta = Points[1].First - Points[2].First;
+				if (TVector<T, d>::DotProduct(-Points[2].First, Delta) > 0)
+				{
+					Direction = TVector<T, d>::CrossProduct(TVector<T, d>::CrossProduct(Delta, -Points[2].First), Delta);
+					Points.RemoveAtSwap(0);
+					check(Points.Num() == 2);
+				}
+				else
+				{
+					Direction = -Points[2].First;
+					Points = {Points[2]};
+					check(Points.Num() == 1);
+				}
+			}
+			else
+			{
+				const auto DotResult = TVector<T, d>::DotProduct(TriangleNormal, -Points[2].First);
+				// We are inside the triangle
+				if (DotResult < SMALL_NUMBER)
+				{
+					TVector<T, 3> Bary;
+					TVector<T, d> P10 = Points[1].First - Points[0].First;
+					TVector<T, d> P20 = Points[2].First - Points[0].First;
+					TVector<T, d> PP0 = -Points[0].First;
+					T Size10 = P10.SizeSquared();
+					T Size20 = P20.SizeSquared();
+					T ProjSides = TVector<T, d>::DotProduct(P10, P20);
+					T ProjP1 = TVector<T, d>::DotProduct(PP0, P10);
+					T ProjP2 = TVector<T, d>::DotProduct(PP0, P20);
+					T Denom = Size10 * Size20 - ProjSides * ProjSides;
+					Bary.Y = (Size20 * ProjP1 - ProjSides * ProjP2) / Denom;
+					Bary.Z = (Size10 * ProjP2 - ProjSides * ProjP1) / Denom;
+					Bary.X = 1.0f - Bary.Z - Bary.Y;
+					ClosestPoint = Points[0].Second * Bary.X + Points[1].Second * Bary.Y + Points[2].Second * Bary.Z;
+					return true;
+				}
+				if (DotResult > 0)
+				{
+					Direction = TriangleNormal;
+				}
+				else
+				{
+					Direction = -TriangleNormal;
+					Points.Swap(0, 1);
+					check(Points.Num() == 3);
+				}
+			}
+		}
+		check(Points.Num() > 0 && Points.Num() < 4);
+		return false;
+	}
+	if (Points.Num() == 4)
+	{
+		TVector<T, d> TriangleNormal = TVector<T, d>::CrossProduct(Points[1].First - Points[3].First, Points[1].First - Points[2].First);
+		if (TVector<T, d>::DotProduct(TriangleNormal, Points[0].First - Points[3].First) > 0)
+		{
+			TriangleNormal *= -1;
+		}
+		T DotResult = TVector<T, d>::DotProduct(TriangleNormal, -Points[3].First);
+		if (DotResult > 0)
+		{
+			Points = {Points[1], Points[2], Points[3]};
+			check(Points.Num() == 3);
+			return NearestPoint(Points, Direction, ClosestPoint);
+		}
+		TriangleNormal = TVector<T, d>::CrossProduct(Points[2].First - Points[0].First, Points[2].First - Points[3].First);
+		if (TVector<T, d>::DotProduct(TriangleNormal, Points[1].First - Points[3].First) > 0)
+		{
+			TriangleNormal *= -1;
+		}
+		DotResult = TVector<T, d>::DotProduct(TriangleNormal, -Points[3].First);
+		if (DotResult > 0)
+		{
+			Points = {Points[0], Points[2], Points[3]};
+			check(Points.Num() == 3);
+			return NearestPoint(Points, Direction, ClosestPoint);
+		}
+		TriangleNormal = TVector<T, d>::CrossProduct(Points[3].First - Points[1].First, Points[3].First - Points[0].First);
+		if (TVector<T, d>::DotProduct(TriangleNormal, Points[2].First - Points[3].First) > 0)
+		{
+			TriangleNormal *= -1;
+		}
+		DotResult = TVector<T, d>::DotProduct(TriangleNormal, -Points[3].First);
+		if (DotResult > 0)
+		{
+			Points = {Points[0], Points[1], Points[3]};
+			check(Points.Num() == 3);
+			return NearestPoint(Points, Direction, ClosestPoint);
+		}
+		TVector<T, 4> Bary;
+		TVector<T, d> PP0 = -Points[0].First;
+		TVector<T, d> PP1 = -Points[1].First;
+		TVector<T, d> P10 = Points[1].First - Points[0].First;
+		TVector<T, d> P20 = Points[2].First - Points[0].First;
+		TVector<T, d> P30 = Points[3].First - Points[0].First;
+		TVector<T, d> P21 = Points[2].First - Points[1].First;
+		TVector<T, d> P31 = Points[3].First - Points[1].First;
+		Bary[0] = TVector<T, d>::DotProduct(PP1, TVector<T, d>::CrossProduct(P31, P21));
+		Bary[1] = TVector<T, d>::DotProduct(PP0, TVector<T, d>::CrossProduct(P20, P30));
+		Bary[2] = TVector<T, d>::DotProduct(PP0, TVector<T, d>::CrossProduct(P30, P10));
+		Bary[3] = TVector<T, d>::DotProduct(PP0, TVector<T, d>::CrossProduct(P10, P20));
+		T Denom = TVector<T, d>::DotProduct(P10, TVector<T, d>::CrossProduct(P20, P30));
+		ClosestPoint = (Bary[0] * Points[0].Second + Bary[1] * Points[1].Second + Bary[2] * Points[2].Second + Bary[3] * Points[3].Second) / Denom;
+		return true;
+	}
+	check(Points.Num() > 1 && Points.Num() < 4);
+	return false;
+}
+
+template<class T_PARTICLES, class T, int d>
+TVector<T, d> GetPosition(const T_PARTICLES& InParticles)
+{
+	check(false);
+	return TVector<T, d>();
+}
+
+template<class T, int d>
+TVector<T, d> GetPosition(const TParticles<T, d>& InParticles, const int32 Index)
+{
+	return InParticles.X(Index);
+}
+
+template<class T, int d>
+TVector<T, d> GetPosition(const TPBDRigidParticles<T, d>& InParticles, const int32 Index)
+{
+	return InParticles.P(Index);
+}
+
+template<class T_PARTICLES, class T, int d>
+TRotation<T, d> GetRotation(const T_PARTICLES& InParticles)
+{
+	check(false);
+	return TRotation<T, d>();
+}
+
+template<class T, int d>
+TRotation<T, d> GetRotation(const TRigidParticles<T, d>& InParticles, const int32 Index)
+{
+	return InParticles.R(Index);
+}
+
+template<class T, int d>
+TRotation<T, d> GetRotation(const TPBDRigidParticles<T, d>& InParticles, const int32 Index)
+{
+	return InParticles.Q(Index);
+}
+
+template<class T_PARTICLES, class T, int d>
+TRigidTransform<T, d> GetTransform(const T_PARTICLES& InParticles)
+{
+	check(false);
+	return TRigidTransform<T, d>();
+}
+
+template<class T, int d>
+TRigidTransform<T, d> GetTransform(const TRigidParticles<T, d>& InParticles, const int32 Index)
+{
+	return TRigidTransform<T, d>(InParticles.X(Index), InParticles.R(Index));
+}
+
+template<class T, int d>
+TRigidTransform<T, d> GetTransform(const TPBDRigidParticles<T, d>& InParticles, const int32 Index)
+{
+	return TRigidTransform<T, d>(InParticles.P(Index), InParticles.Q(Index));
+}
+
+#if 0
+template<class T, int d>
+void UpdateLevelsetConstraintHelperCCD(const TRigidParticles<T, d>& InParticles, const int32 j, const TRigidTransform<T, d> LocalToWorld1, const TRigidTransform<T, d> LocalToWorld2, const T Thickness, TRigidBodyContactConstraint<T, d>& Constraint)
+{
+	if (InParticles.CollisionParticles(Constraint.ParticleIndex))
+	{
+		const TRigidTransform<T, d> PreviousLocalToWorld1 = GetTransform(InParticles, Constraint.ParticleIndex);
+		TVector<T, d> WorldSpacePointStart = PreviousLocalToWorld1.TransformPosition(InParticles.CollisionParticles(Constraint.ParticleIndex)->X(j));
+		TVector<T, d> WorldSpacePointEnd = LocalToWorld1.TransformPosition(InParticles.CollisionParticles(Constraint.ParticleIndex)->X(j));
+		TVector<T, d> Body2SpacePointStart = LocalToWorld2.InverseTransformPosition(WorldSpacePointStart);
+		TVector<T, d> Body2SpacePointEnd = LocalToWorld2.InverseTransformPosition(WorldSpacePointEnd);
+		Pair<TVector<T, d>, bool> PointPair = InParticles.Geometry(Constraint.LevelsetIndex)->FindClosestIntersection(Body2SpacePointStart, Body2SpacePointEnd, Thickness);
+		if (PointPair.Second)
+		{
+			const TVector<T, d> WorldSpaceDelta = WorldSpacePointEnd - TVector<T, d>(LocalToWorld2.TransformPosition(PointPair.First));
+			Constraint.Phi = -WorldSpaceDelta.Size();
+			Constraint.Normal = LocalToWorld2.TransformVector(InParticles.Geometry(Constraint.LevelsetIndex)->Normal(PointPair.First));
+			// @todo(mlentine): Should we be using the actual collision point or that point evolved to the current time step?
+			Constraint.Location = WorldSpacePointEnd;
+		}
+	}
+}
+#endif
+
+template <typename T, int d>
+bool SampleObjectHelper(const TImplicitObject<T, d>& Object, const TRigidTransform<T,d>& ObjectTransform, const TRigidTransform<T, d>& SampleToObjectTransform, const TVector<T, d>& SampleParticle, T Thickness, TRigidBodyContactConstraint<T, d>& Constraint)
+{
+	TVector<T, d> LocalPoint = SampleToObjectTransform.TransformPositionNoScale(SampleParticle);
+	TVector<T, d> LocalNormal;
+	T LocalPhi = Object.PhiWithNormal(LocalPoint, LocalNormal);
+	if (LocalPhi < Constraint.Phi)
+	{
+		Constraint.Phi = LocalPhi;
+		return true;
+	}
+	return false;
+}
+
+template <typename T, int d>
+void SampleObjectNormalAverageHelper(const TImplicitObject<T, d>& Object, const TRigidTransform<T, d>& ObjectTransform, const TRigidTransform<T, d>& SampleToObjectTransform, const TVector<T, d>& SampleParticle, T Thickness, T& TotalThickness, TRigidBodyContactConstraint<T, d>& Constraint)
+{
+	TVector<T, d> LocalPoint = SampleToObjectTransform.TransformPositionNoScale(SampleParticle);
+	TVector<T, d> LocalNormal;
+	T LocalPhi = Object.PhiWithNormal(LocalPoint, LocalNormal);
+	T LocalThickness = LocalPhi - Thickness;
+	if (LocalThickness < -KINDA_SMALL_NUMBER)
+	{
+		Constraint.Location += LocalPoint * LocalThickness;
+		TotalThickness += LocalThickness;
+	}
+}
+
+DECLARE_CYCLE_STAT(TEXT("UpdateLevelsetOther"), STAT_UpdateLevelsetOther, STATGROUP_ChaosWide);
+DECLARE_CYCLE_STAT(TEXT("UpdateLevelsetFindParticles"), STAT_UpdateLevelsetFindParticles, STATGROUP_ChaosWide);
+DECLARE_CYCLE_STAT(TEXT("UpdateLevelsetBVHTraversal"), STAT_UpdateLevelsetBVHTraversal, STATGROUP_ChaosWide);
+DECLARE_CYCLE_STAT(TEXT("UpdateLevelsetSignedDistance"), STAT_UpdateLevelsetSignedDistance, STATGROUP_ChaosWide);
+DECLARE_CYCLE_STAT(TEXT("UpdateLevelsetPlane"), STAT_UpdateLevelsetPlane, STATGROUP_ChaosWide);
+DECLARE_CYCLE_STAT(TEXT("SampleObject"), STAT_SampleObject, STATGROUP_ChaosWide);
+
+int32 NormalAveraging = 1;
+FAutoConsoleVariableRef CVarNormalAveraging(TEXT("p.NormalAveraging"), NormalAveraging, TEXT(""));
+
+template <typename T, int d>
+void SampleObject(const TImplicitObject<T, d>& Object, const TRigidTransform<T, d>& ObjectTransform, const TBVHParticles<T, d>& SampleParticles, const TRigidTransform<T, d>& SampleParticlesTransform, T Thickness, TRigidBodyContactConstraint<T, d>& Constraint)
+{
+	SCOPE_CYCLE_COUNTER(STAT_SampleObject);
+	TRigidBodyContactConstraint<T, d> AvgConstraint;
+	AvgConstraint.ParticleIndex = Constraint.ParticleIndex;
+	AvgConstraint.LevelsetIndex = Constraint.LevelsetIndex;
+	AvgConstraint.Location = TVector<T, d>::ZeroVector;
+	AvgConstraint.Normal = TVector<T, d>::ZeroVector;
+	AvgConstraint.Phi = Thickness;
+	T TotalThickness = T(0);
+
+	int32 DeepestParticle = -1;
+
+	const TRigidTransform<T, d> SampleToObjectTM = SampleParticlesTransform.GetRelativeTransform(ObjectTransform);
+	if (Object.HasBoundingBox())
+	{
+		SCOPE_CYCLE_COUNTER(STAT_UpdateLevelsetOther);
+		TBox<T, d> ImplicitBox = Object.BoundingBox().TransformedBox(ObjectTransform * SampleParticlesTransform.Inverse());
+		ImplicitBox.Thicken(Thickness);
+		TArray<int32> PotentialParticles;
+		{
+			SCOPE_CYCLE_COUNTER(STAT_UpdateLevelsetFindParticles);
+			PotentialParticles = SampleParticles.FindAllIntersections(ImplicitBox);
+		}
+		{
+			SCOPE_CYCLE_COUNTER(STAT_UpdateLevelsetSignedDistance);
+			for (int32 i : PotentialParticles)
+			{
+				if (NormalAveraging)
+				{
+					SampleObjectNormalAverageHelper(Object, ObjectTransform, SampleToObjectTM, SampleParticles.X(i), Thickness, TotalThickness, AvgConstraint);
+				}
+				else
+				{
+					if (SampleObjectHelper(Object, ObjectTransform, SampleToObjectTM, SampleParticles.X(i), Thickness, AvgConstraint))
+					{
+						DeepestParticle = i;
+					}
+				}
+			}
+		}
+	}
+	else
+	{
+		SCOPE_CYCLE_COUNTER(STAT_UpdateLevelsetPlane);
+		int32 NumParticles = SampleParticles.Size();
+		for (int32 i = 0; i < NumParticles; ++i)
+		{
+			if (NormalAveraging)
+			{
+				SampleObjectNormalAverageHelper(Object, ObjectTransform, SampleToObjectTM, SampleParticles.X(i), Thickness, TotalThickness, AvgConstraint);
+			}
+			else
+			{
+				if (SampleObjectHelper(Object, ObjectTransform, SampleToObjectTM, SampleParticles.X(i), Thickness, AvgConstraint))
+				{
+					DeepestParticle = i;
+				}
+			}
+		}
+	}
+
+	if (NormalAveraging)
+	{
+		if (TotalThickness < -KINDA_SMALL_NUMBER)
+		{
+			TVector<T, d> LocalPoint = AvgConstraint.Location / TotalThickness;
+			TVector<T, d> LocalNormal;
+			const T NewPhi = Object.PhiWithNormal(LocalPoint, LocalNormal);
+			if (NewPhi < Constraint.Phi)
+			{
+				Constraint.Phi = NewPhi;
+				Constraint.Location = ObjectTransform.TransformPositionNoScale(LocalPoint);
+				Constraint.Normal = ObjectTransform.TransformVectorNoScale(LocalNormal);
+			}
+		}
+		else
+		{
+			check(AvgConstraint.Phi >= Thickness);
+		}
+	}
+	else if(AvgConstraint.Phi < Constraint.Phi)
+	{
+		check(DeepestParticle >= 0);
+		TVector<T,d> LocalPoint = SampleToObjectTM.TransformPositionNoScale(SampleParticles.X(DeepestParticle));
+		TVector<T, d> LocalNormal;
+		Constraint.Phi = Object.PhiWithNormal(LocalPoint, LocalNormal);
+		Constraint.Location = ObjectTransform.TransformPositionNoScale(LocalPoint);
+		Constraint.Normal = ObjectTransform.TransformVectorNoScale(LocalNormal);
+			
+	}
+}
+
+template <typename T, int d>
+bool UpdateBoxPlaneConstraint(const TBox<T, d>& Box, const TRigidTransform<T, d>& BoxTransform, const TPlane<T, d>& Plane, const TRigidTransform<T, d>& PlaneTransform, T Thickness, TRigidBodyContactConstraint<T, d>& Constraint)
+{
+	
+#if USING_CODE_ANALYSIS
+	MSVC_PRAGMA( warning( push ) )
+	MSVC_PRAGMA( warning( disable : ALL_CODE_ANALYSIS_WARNINGS ) )
+#endif	// USING_CODE_ANALYSIS
+
+	bool bApplied = false;
+	const TRigidTransform<T, d> BoxToPlaneTransform(BoxTransform * PlaneTransform.Inverse());
+	const TVector<T, d> Extents = Box.Extents();
+	constexpr int32 NumCorners = 2 + 2 * d;
+	constexpr T Epsilon = KINDA_SMALL_NUMBER;
+
+	TVector<T, d> Corners[NumCorners];
+	int32 CornerIdx = 0;
+	Corners[CornerIdx++] = BoxToPlaneTransform.TransformPosition(Box.Max());
+	Corners[CornerIdx++] = BoxToPlaneTransform.TransformPosition(Box.Min());
+	for (int32 j = 0; j < d; ++j)
+	{
+		Corners[CornerIdx++] = BoxToPlaneTransform.TransformPosition(Box.Min() + TVector<T, d>::AxisVector(j) * Extents);
+		Corners[CornerIdx++] = BoxToPlaneTransform.TransformPosition(Box.Max() - TVector<T, d>::AxisVector(j) * Extents);
+	}
+
+#if USING_CODE_ANALYSIS
+	MSVC_PRAGMA( warning( pop ) )
+#endif	// USING_CODE_ANALYSIS
+
+	TVector<T, d> PotentialConstraints[NumCorners];
+	int32 NumConstraints = 0;
+	for (int32 i = 0; i < NumCorners; ++i)
+	{
+		TVector<T, d> Normal;
+		const T NewPhi = Plane.PhiWithNormal(Corners[i], Normal);
+		if (NewPhi < Constraint.Phi + Epsilon)
+		{
+			if (NewPhi <= Constraint.Phi - Epsilon)
+			{
+				NumConstraints = 0;
+			}
+			Constraint.Phi = NewPhi;
+			Constraint.Normal = PlaneTransform.TransformVector(Normal);
+			Constraint.Location = PlaneTransform.TransformPosition(Corners[i]);
+			PotentialConstraints[NumConstraints++] = Constraint.Location;
+			bApplied = true;
+		}
+	}
+	if (NumConstraints > 1)
+	{
+		TVector<T, d> AverageLocation(0);
+		for (int32 ConstraintIdx = 0; ConstraintIdx < NumConstraints; ++ConstraintIdx)
+		{
+			AverageLocation += PotentialConstraints[ConstraintIdx];
+		}
+		Constraint.Location = AverageLocation / NumConstraints;
+	}
+
+	return bApplied;
+}
+
+template <typename T, int d>
+void UpdateSphereConstraint(const TSphere<T, d>& Sphere1, const TRigidTransform<T, d>& Sphere1Transform, const TSphere<T, d>& Sphere2, const TRigidTransform<T, d>& Sphere2Transform, T Thickness, TRigidBodyContactConstraint<T, d>& Constraint)
+{
+	const TVector<T, d> Center1 = Sphere1Transform.TransformPosition(Sphere1.Center());
+	const TVector<T, d> Center2 = Sphere2Transform.TransformPosition(Sphere2.Center());
+	const TVector<T, d> Direction = Center1 - Center2;
+	const T Size = Direction.Size();
+	const T NewPhi = Size - (Sphere1.Radius() + Sphere2.Radius());
+	if (NewPhi < Constraint.Phi)
+	{
+		Constraint.Normal = Size > SMALL_NUMBER ? Direction / Size : TVector<T, d>(0, 0, 1);
+		Constraint.Phi = NewPhi;
+		Constraint.Location = Center1 - Sphere1.Radius() * Constraint.Normal;
+	}
+}
+
+template <typename T, int d>
+void UpdateSpherePlaneConstraint(const TSphere<T, d>& Sphere, const TRigidTransform<T, d>& SphereTransform, const TPlane<T, d>& Plane, const TRigidTransform<T, d>& PlaneTransform, T Thickness, TRigidBodyContactConstraint<T, d>& Constraint)
+{
+	const TRigidTransform<T, d> SphereToPlaneTransform(PlaneTransform.Inverse() * SphereTransform);
+	const TVector<T, d> SphereCenter = SphereToPlaneTransform.TransformPosition(Sphere.Center());
+
+	TVector<T, d> NewNormal;
+	T NewPhi = Plane.PhiWithNormal(SphereCenter, NewNormal);
+	NewPhi -= Sphere.Radius();
+
+	if (NewPhi < Constraint.Phi)
+	{
+		Constraint.Phi = NewPhi;
+		Constraint.Normal = PlaneTransform.TransformVectorNoScale(NewNormal);
+		Constraint.Location = SphereCenter - Constraint.Normal * Sphere.Radius();
+	}
+}
+
+template <typename T, int d>
+bool UpdateSphereBoxConstraint(const TSphere<T, d>& Sphere, const TRigidTransform<T, d>& SphereTransform, const TBox<T, d>& Box, const TRigidTransform<T, d>& BoxTransform, T Thickness, TRigidBodyContactConstraint<T, d>& Constraint)
+{
+	const TRigidTransform<T, d> SphereToBoxTransform(SphereTransform * BoxTransform.Inverse());
+	const TVector<T, d> SphereCenterInBox = SphereToBoxTransform.TransformPosition(Sphere.Center());
+
+	TVector<T, d> NewNormal;
+	T NewPhi = Box.PhiWithNormal(SphereCenterInBox, NewNormal);
+	NewPhi -= Sphere.Radius();
+
+	if (NewPhi < Constraint.Phi)
+	{
+		Constraint.Phi = NewPhi;
+		Constraint.Normal = BoxTransform.TransformVectorNoScale(NewNormal);
+		Constraint.Location = SphereTransform.TransformPosition(Sphere.Center()) - Constraint.Normal * Sphere.Radius();
+		return true;
+	}
+
+	return false;
+}
+
+DECLARE_CYCLE_STAT(TEXT("FindRelevantShapes"), STAT_FindRelevantShapes, STATGROUP_ChaosWide);
+template <typename T, int d>
+TArray<Pair<const TImplicitObject<T, d>*, TRigidTransform<T, d>>> FindRelevantShapes(const TImplicitObject<T,d>& ParticleObj, const TRigidTransform<T,d>& ParticlesTM, const TImplicitObject<T,d>& LevelsetObj, const TRigidTransform<T,d>& LevelsetTM, const T Thickness)
+{
+	SCOPE_CYCLE_COUNTER(STAT_FindRelevantShapes);
+	TArray<Pair<const TImplicitObject<T, d>*, TRigidTransform<T, d>>> RelevantShapes;
+	//find all levelset inner objects
+	if (ParticleObj.HasBoundingBox())
+	{
+		const TRigidTransform<T, d> ParticlesToLevelsetTM = ParticlesTM * LevelsetTM.Inverse();
+		TBox<T, d> ParticleBoundsInLevelset = ParticleObj.BoundingBox().TransformedBox(ParticlesToLevelsetTM);
+		ParticleBoundsInLevelset.Thicken(Thickness);
+		{
+			LevelsetObj.FindAllIntersectingObjects(RelevantShapes, ParticleBoundsInLevelset);
+		}
+	}
+	else
+	{
+		LevelsetObj.AccumulateAllImplicitObjects(RelevantShapes, TRigidTransform<T, d>::Identity);
+	}
+
+	return RelevantShapes;
+}
+
+DECLARE_CYCLE_STAT(TEXT("UpdateUnionUnionConstraint"), STAT_UpdateUnionUnionConstraint, STATGROUP_ChaosWide);
+template<typename T_PARTICLES, typename T, int d>
+void UpdateUnionUnionConstraint(const T_PARTICLES& InParticles, const T Thickness, TRigidBodyContactConstraint<T, d>& Constraint)
+{
+	SCOPE_CYCLE_COUNTER(STAT_UpdateUnionUnionConstraint);
+	Constraint.Phi = Thickness;
+
+	TRigidTransform<T, d> ParticlesTM = GetTransform(InParticles, Constraint.ParticleIndex);
+	TRigidTransform<T, d> LevelsetTM = GetTransform(InParticles, Constraint.LevelsetIndex);
+
+	const TImplicitObject<T, d>* ParticleObj = InParticles.Geometry(Constraint.ParticleIndex);
+	const TImplicitObject<T,d>* LevelsetObj = InParticles.Geometry(Constraint.LevelsetIndex);
+	const TArray<Pair<const TImplicitObject<T, d>*, TRigidTransform<T, d>>> LevelsetShapes = FindRelevantShapes(*ParticleObj, ParticlesTM, *LevelsetObj, LevelsetTM, Thickness);
+
+	for (const Pair<const TImplicitObject<T, d>*, TRigidTransform<T, d>>& LevelsetObjPair : LevelsetShapes)
+	{
+		const TImplicitObject<T, d>& LevelsetInnerObj = *LevelsetObjPair.First;
+		const TRigidTransform<T, d>& LevelsetInnerObjTM = LevelsetObjPair.Second * LevelsetTM;
+
+		//now find all particle inner objects
+		const TArray<Pair<const TImplicitObject<T, d>*, TRigidTransform<T, d>>> ParticleShapes = FindRelevantShapes(LevelsetInnerObj, LevelsetInnerObjTM, *ParticleObj, ParticlesTM, Thickness);
+
+		//for each inner obj pair, update constraint
+		for (const Pair<const TImplicitObject<T, d>*, TRigidTransform<T, d>>& ParticlePair : ParticleShapes)
+		{
+			const TImplicitObject<T, d>& ParticleInnerObj = *ParticlePair.First;
+			const TRigidTransform<T, d> ParticleInnerObjTM = ParticlePair.Second * ParticlesTM;
+			UpdateConstraintImp(InParticles, ParticleInnerObj, ParticleInnerObjTM, LevelsetInnerObj, LevelsetInnerObjTM, Thickness, Constraint);
+		}
+	}
+}
+
+DECLARE_CYCLE_STAT(TEXT("UpdateSingleUnionConstraint"), STAT_UpdateSingleUnionConstraint, STATGROUP_ChaosWide);
+template<typename T_PARTICLES, typename T, int d>
+void UpdateSingleUnionConstraint(const T_PARTICLES& InParticles, const T Thickness, TRigidBodyContactConstraint<T, d>& Constraint)
+{
+	SCOPE_CYCLE_COUNTER(STAT_UpdateSingleUnionConstraint);
+	Constraint.Phi = Thickness;
+
+	TRigidTransform<T, d> ParticlesTM = GetTransform(InParticles, Constraint.ParticleIndex);
+	TRigidTransform<T, d> LevelsetTM = GetTransform(InParticles, Constraint.LevelsetIndex);
+
+	const TImplicitObject<T, d>* ParticleObj = InParticles.Geometry(Constraint.ParticleIndex);
+	const TImplicitObject<T, d>* LevelsetObj = InParticles.Geometry(Constraint.LevelsetIndex);
+	const TArray<Pair<const TImplicitObject<T, d>*, TRigidTransform<T, d>>> LevelsetShapes = FindRelevantShapes(*ParticleObj, ParticlesTM, *LevelsetObj, LevelsetTM, Thickness);
+	
+	for (const Pair<const TImplicitObject<T, d>*, TRigidTransform<T, d>>& LevelsetObjPair : LevelsetShapes)
+	{
+		const TImplicitObject<T, d>& LevelsetInnerObj = *LevelsetObjPair.First;
+		const TRigidTransform<T, d> LevelsetInnerObjTM = LevelsetObjPair.Second * LevelsetTM;
+		UpdateConstraintImp(InParticles, *ParticleObj, ParticlesTM, LevelsetInnerObj, LevelsetInnerObjTM, Thickness, Constraint);
+	}
+}
+
+DECLARE_CYCLE_STAT(TEXT("UpdateLevelsetConstraint"), STAT_UpdateLevelsetConstraint, STATGROUP_ChaosWide);
+template<typename T, int d>
+template<typename T_PARTICLES>
+void TPBDCollisionConstraint<T, d>::UpdateLevelsetConstraint(const T_PARTICLES& InParticles, const T Thickness, TRigidBodyContactConstraint<T, d>& Constraint)
+{
+	SCOPE_CYCLE_COUNTER(STAT_UpdateLevelsetConstraint);
+	Constraint.Phi = Thickness;
+
+	TRigidTransform<T, d> ParticlesTM = GetTransform(InParticles, Constraint.ParticleIndex);
+	TRigidTransform<T, d> LevelsetTM = GetTransform(InParticles, Constraint.LevelsetIndex);
+
+	const TBVHParticles<T, d>* SampleParticles = nullptr;
+	SampleParticles = InParticles.CollisionParticles(Constraint.ParticleIndex).Get();
+
+	if(SampleParticles)
+	{
+		SampleObject(*InParticles.Geometry(Constraint.LevelsetIndex), LevelsetTM, *SampleParticles, ParticlesTM, Thickness, Constraint);
+	}
+}
+
+DECLARE_CYCLE_STAT(TEXT("UpdateUnionLevelsetConstraint"), STAT_UpdateUnionLevelsetConstraint, STATGROUP_ChaosWide);
+template<typename T_PARTICLES, typename T, int d>
+void UpdateUnionLevelsetConstraint(const T_PARTICLES& InParticles, const T Thickness, TRigidBodyContactConstraint<T, d>& Constraint)
+{
+	SCOPE_CYCLE_COUNTER(STAT_UpdateUnionLevelsetConstraint);
+	Constraint.Phi = Thickness;
+
+	TRigidTransform<T, d> ParticlesTM = GetTransform(InParticles, Constraint.ParticleIndex);
+	TRigidTransform<T, d> LevelsetTM = GetTransform(InParticles, Constraint.LevelsetIndex);
+
+	const TImplicitObject<T, d>* ParticleObj = InParticles.Geometry(Constraint.ParticleIndex);
+	const TImplicitObject<T, d>* LevelsetObj = InParticles.Geometry(Constraint.LevelsetIndex);
+	TArray<Pair<const TImplicitObject<T, d>*, TRigidTransform<T, d>>> LevelsetShapes = FindRelevantShapes(*ParticleObj, ParticlesTM, *LevelsetObj, LevelsetTM, Thickness);
+
+	const TBVHParticles<T, d>& SampleParticles = *InParticles.CollisionParticles(Constraint.ParticleIndex).Get();
+	for (const Pair<const TImplicitObject<T, d>*, TRigidTransform<T, d>>& LevelsetObjPair : LevelsetShapes)
+	{
+		const TImplicitObject<T, d>* Object = LevelsetObjPair.First;
+		const TRigidTransform<T, d> ObjectTM = LevelsetObjPair.Second * LevelsetTM;
+		SampleObject(*Object, ObjectTM, SampleParticles, ParticlesTM, Thickness, Constraint);
+	}
+}
+
+template<typename T, int d>
+template<typename T_PARTICLES>
+void TPBDCollisionConstraint<T, d>::UpdateLevelsetConstraintGJK(const T_PARTICLES& InParticles, const T Thickness, TRigidBodyContactConstraint<T, d>& Constraint)
+{
+	static int32 MaxIterations = 100;
+	Constraint.Phi = Thickness;
+	const TRigidTransform<T, d> LocalToWorld1 = GetTransform(InParticles, Constraint.ParticleIndex);
+	const TRigidTransform<T, d> LocalToWorld2 = GetTransform(InParticles, Constraint.LevelsetIndex);
+	TVector<T, d> Direction = LocalToWorld1.GetTranslation() - LocalToWorld2.GetTranslation();
+	TVector<T, d> SupportA = LocalToWorld1.TransformPosition(InParticles.Geometry(Constraint.ParticleIndex)->Support(LocalToWorld1.InverseTransformVector(-Direction), Thickness));
+	TVector<T, d> SupportB = LocalToWorld2.TransformPosition(InParticles.Geometry(Constraint.LevelsetIndex)->Support(LocalToWorld2.InverseTransformVector(Direction), Thickness));
+	TVector<T, d> Point = SupportB - SupportA;
+	TArray<Pair<TVector<T, d>, TVector<T, d>>> Points = {MakePair(Point, SupportA)};
+	Direction = -Point;
+	for (int32 i = 0; i < MaxIterations; ++i)
+	{
+		SupportA = LocalToWorld1.TransformPosition(InParticles.Geometry(Constraint.ParticleIndex)->Support(LocalToWorld1.InverseTransformVector(-Direction), Thickness));
+		SupportB = LocalToWorld2.TransformPosition(InParticles.Geometry(Constraint.LevelsetIndex)->Support(LocalToWorld2.InverseTransformVector(Direction), Thickness));
+		Point = SupportB - SupportA;
+		if (TVector<T, d>::DotProduct(Point, Direction) < 0)
+		{
+			break;
+		}
+		Points.Add(MakePair(Point, SupportA));
+		TVector<T, d> ClosestPoint;
+		if (NearestPoint(Points, Direction, ClosestPoint))
+		{
+			TVector<T, d> Body1Location = LocalToWorld1.InverseTransformPosition(ClosestPoint);
+			TVector<T, d> Normal;
+			T Phi = InParticles.Geometry(Constraint.ParticleIndex)->PhiWithNormal(Body1Location, Normal);
+			Normal = LocalToWorld1.TransformVector(Normal);
+			Constraint.Location = ClosestPoint - Phi * Normal;
+			TVector<T, d> Body2Location = LocalToWorld2.InverseTransformPosition(Constraint.Location);
+			Constraint.Phi = InParticles.Geometry(Constraint.LevelsetIndex)->PhiWithNormal(Body2Location, Constraint.Normal);
+			Constraint.Normal = LocalToWorld2.TransformVector(Constraint.Normal);
+			break;
+		}
+	}
+}
+
+template <typename T, int d>
+void UpdateBoxConstraint(const TBox<T, d>& Box1, const TRigidTransform<T, d>& Box1Transform, const TBox<T, d>& Box2, const TRigidTransform<T, d>& Box2Transform, T Thickness, TRigidBodyContactConstraint<T, d>& Constraint)
+{
+	TBox<T,d> Box2SpaceBox1 = Box1.TransformedBox(Box1Transform * Box2Transform.Inverse());
+	TBox<T,d> Box1SpaceBox2 = Box2.TransformedBox(Box2Transform * Box1Transform.Inverse());
+	Box2SpaceBox1.Thicken(Thickness);
+	Box1SpaceBox2.Thicken(Thickness);
+	if (Box1SpaceBox2.Intersects(Box1) && Box2SpaceBox1.Intersects(Box2))
+	{
+		const TVector<T, d> Box1Center = (Box1Transform * Box2Transform.Inverse()).TransformPosition(Box1.Center());
+		bool bDeepOverlap = false;
+		if (Box2.SignedDistance(Box1Center) < 0)
+		{
+			//If Box1 is overlapping Box2 by this much the signed distance approach will fail (box1 gets sucked into box2). In this case just use two spheres
+			TSphere<T, d> Sphere1(Box1Transform.TransformPosition(Box1.Center()), Box1.Extents().Min() / 2);
+			TSphere<T, d> Sphere2(Box2Transform.TransformPosition(Box2.Center()), Box2.Extents().Min() / 2);
+			const TVector<T, d> Direction = Sphere1.Center() - Sphere2.Center();
+			T Size = Direction.Size();
+			if (Size < (Sphere1.Radius() + Sphere2.Radius()))
+			{
+				const T NewPhi = Size - (Sphere1.Radius() + Sphere2.Radius());;
+				if (NewPhi < Constraint.Phi)
+				{
+					bDeepOverlap = true;
+					Constraint.Normal = Size > SMALL_NUMBER ? Direction / Size : TVector<T, d>(0, 0, 1);
+					Constraint.Phi = NewPhi;
+					Constraint.Location = Sphere1.Center() - Sphere1.Radius() * Constraint.Normal;
+				}
+			}
+		}
+		if (!bDeepOverlap || Constraint.Phi >= 0)
+		{
+			//if we didn't have deep penetration use signed distance per particle. If we did have deep penetration but the spheres did not overlap use signed distance per particle
+
+			//UpdateLevelsetConstraintGJK(InParticles, Thickness, Constraint);
+			//check(Constraint.Phi < MThickness);
+			// For now revert to doing all points vs lsv check until we can figure out a good way to get the deepest point without needing this
+			{
+				const TArray<TVector<T, d>> SampleParticles = Box1.ComputeLocalSamplePoints();
+				const TRigidTransform<T, d> Box1ToBox2Transform = Box1Transform.GetRelativeTransform(Box2Transform);
+				int32 NumParticles = SampleParticles.Num();
+				for (int32 i = 0; i < NumParticles; ++i)
+				{
+					SampleObjectHelper(Box2, Box2Transform, Box1ToBox2Transform, SampleParticles[i], Thickness, Constraint);
+				}
+			}
+		}
+	}
+}
+
+
+template<class T, int d>
+TRigidBodyContactConstraint<T, d> ComputeLevelsetConstraint(const TPBDRigidParticles<T, d>& InParticles, int32 ParticleIndex, int32 LevelsetIndex, const T Thickness)
+{
+	if (!InParticles.CollisionParticlesSize(ParticleIndex))
+	{
+		int32 TmpIndex = ParticleIndex;
+		ParticleIndex = LevelsetIndex;
+		LevelsetIndex = TmpIndex;
+	}
+	//todo(ocohen):if both have collision particles, use the one with fewer?
+	// Find Deepest Point
+	TRigidBodyContactConstraint<T, d> Constraint;
+	Constraint.ParticleIndex = ParticleIndex;
+	Constraint.LevelsetIndex = LevelsetIndex;
+	return Constraint;
+}
+
+template<class T, int d>
+TRigidBodyContactConstraint<T, d> ComputeLevelsetConstraintGJK(const TPBDRigidParticles<T, d>& InParticles, int32 ParticleIndex, int32 LevelsetIndex, const T Thickness)
+{
+	TRigidBodyContactConstraint<T, d> Constraint;
+	Constraint.ParticleIndex = ParticleIndex;
+	Constraint.LevelsetIndex = LevelsetIndex;
+	return Constraint;
+}
+
+template<class T, int d>
+TRigidBodyContactConstraint<T, d> ComputeBoxConstraint(const TPBDRigidParticles<T, d>& InParticles, int32 Box1Index, int32 Box2Index, const T Thickness)
+{
+	TRigidBodyContactConstraint<T, d> Constraint;
+	Constraint.ParticleIndex = Box1Index;
+	Constraint.LevelsetIndex = Box2Index;
+	return Constraint;
+}
+
+template<class T, int d>
+TRigidBodyContactConstraint<T, d> ComputeBoxPlaneConstraint(const TPBDRigidParticles<T, d>& InParticles, int32 BoxIndex, int32 PlaneIndex, const T Thickness)
+{
+	TRigidBodyContactConstraint<T, d> Constraint;
+	Constraint.ParticleIndex = BoxIndex;
+	Constraint.LevelsetIndex = PlaneIndex;
+	return Constraint;
+}
+
+template<class T, int d>
+TRigidBodyContactConstraint<T, d> ComputeSphereConstraint(const TPBDRigidParticles<T, d>& InParticles, int32 Sphere1Index, int32 Sphere2Index, const T Thickness)
+{
+	TRigidBodyContactConstraint<T, d> Constraint;
+	Constraint.ParticleIndex = Sphere1Index;
+	Constraint.LevelsetIndex = Sphere2Index;
+	return Constraint;
+}
+
+template<class T, int d>
+TRigidBodyContactConstraint<T, d> ComputeSpherePlaneConstraint(const TPBDRigidParticles<T, d>& InParticles, int32 SphereIndex, int32 PlaneIndex, const T Thickness)
+{
+	TRigidBodyContactConstraint<T, d> Constraint;
+	Constraint.ParticleIndex = SphereIndex;
+	Constraint.LevelsetIndex = PlaneIndex;
+	return Constraint;
+}
+
+template<class T, int d>
+TRigidBodyContactConstraint<T, d> ComputeSphereBoxConstraint(const TPBDRigidParticles<T, d>& InParticles, int32 SphereIndex, int32 BoxIndex, const T Thickness)
+{
+	TRigidBodyContactConstraint<T, d> Constraint;
+	Constraint.ParticleIndex = SphereIndex;
+	Constraint.LevelsetIndex = BoxIndex;
+	return Constraint;
+}
+
+template <typename T, int d>
+TRigidBodyContactConstraint<T, d> ComputeSingleUnionConstraint(const TPBDRigidParticles<T, d>& InParticles, int32 NonUnionIndex, int32 UnionIndex, const T Thickness)
+{
+	TRigidBodyContactConstraint<T, d> Constraint;
+	Constraint.ParticleIndex = NonUnionIndex;
+	Constraint.LevelsetIndex = UnionIndex;
+	return Constraint;
+}
+
+template <typename T, int d>
+TRigidBodyContactConstraint<T, d> ComputeUnionUnionConstraint(const TPBDRigidParticles<T, d>& InParticles, int32 Union1Index, int32 Union2Index, const T Thickness)
+{
+	TRigidBodyContactConstraint<T, d> Constraint;
+	Constraint.ParticleIndex = Union1Index;
+	Constraint.LevelsetIndex = Union2Index;
+	//todo(ocohen): some heuristic for determining the order?
+	return Constraint;
+}
+
+template<class T, int d>
+typename TPBDCollisionConstraint<T, d>::FRigidBodyContactConstraint TPBDCollisionConstraint<T, d>::ComputeConstraint(const TPBDRigidParticles<T, d>& InParticles, int32 Body1Index, int32 Body2Index, const T Thickness)
+{
+	if (InParticles.Geometry(Body1Index)->GetType() == TBox<T, d>::GetType() && InParticles.Geometry(Body2Index)->GetType() == TBox<T, d>::GetType())
+	{
+		return ComputeBoxConstraint(InParticles, Body1Index, Body2Index, Thickness);
+	}
+	else if (InParticles.Geometry(Body1Index)->GetType() == TSphere<T, d>::GetType() && InParticles.Geometry(Body2Index)->GetType() == TSphere<T, d>::GetType())
+	{
+		return ComputeSphereConstraint(InParticles, Body1Index, Body2Index, Thickness);
+	}
+	else if (InParticles.Geometry(Body1Index)->GetType() == TBox<T, d>::GetType() && InParticles.Geometry(Body2Index)->GetType() == TPlane<T, d>::GetType())
+	{
+		return ComputeBoxPlaneConstraint(InParticles, Body1Index, Body2Index, Thickness);
+	}
+	else if (InParticles.Geometry(Body2Index)->GetType() == TPlane<T, d>::GetType() && InParticles.Geometry(Body1Index)->GetType() == TBox<T, d>::GetType())
+	{
+		return ComputeBoxPlaneConstraint(InParticles, Body2Index, Body1Index, Thickness);
+	}
+	else if (InParticles.Geometry(Body1Index)->GetType() == TSphere<T, d>::GetType() && InParticles.Geometry(Body2Index)->GetType() == TPlane<T, d>::GetType())
+	{
+		return ComputeSpherePlaneConstraint(InParticles, Body1Index, Body2Index, Thickness);
+	}
+	else if (InParticles.Geometry(Body2Index)->GetType() == TPlane<T, d>::GetType() && InParticles.Geometry(Body1Index)->GetType() == TSphere<T, d>::GetType())
+	{
+		return ComputeSpherePlaneConstraint(InParticles, Body2Index, Body1Index, Thickness);
+	}
+	else if (InParticles.Geometry(Body1Index)->GetType() == TSphere<T, d>::GetType() && InParticles.Geometry(Body2Index)->GetType() == TBox<T, d>::GetType())
+	{
+		return ComputeSphereBoxConstraint(InParticles, Body1Index, Body2Index, Thickness);
+	}
+	else if (InParticles.Geometry(Body1Index)->GetType() < TImplicitObjectUnion<T, d>::GetType() && InParticles.Geometry(Body2Index)->GetType() == TImplicitObjectUnion<T, d>::GetType())
+	{
+		return ComputeSingleUnionConstraint(InParticles, Body1Index, Body2Index, Thickness);
+	}
+	else if (InParticles.Geometry(Body1Index)->GetType() == TImplicitObjectUnion<T, d>::GetType() && InParticles.Geometry(Body2Index)->GetType() < TImplicitObjectUnion<T, d>::GetType())
+	{
+		return ComputeSingleUnionConstraint(InParticles, Body2Index, Body1Index, Thickness);
+	}
+	else if(InParticles.Geometry(Body1Index)->GetType() == TImplicitObjectUnion<T, d>::GetType() && InParticles.Geometry(Body2Index)->GetType() == TImplicitObjectUnion<T, d>::GetType())
+	{
+		return ComputeUnionUnionConstraint(InParticles, Body1Index, Body2Index, Thickness);
+	}
+#if 0
+	else if (InParticles.Geometry(Body1Index)->IsConvex() && InParticles.Geometry(Body2Index)->IsConvex())
+	{
+		return ComputeLevelsetConstraintGJK(InParticles, Body1Index, Body2Index, Thickness);
+	}
+#endif
+	return ComputeLevelsetConstraint(InParticles, Body1Index, Body2Index, Thickness);
+}
+
+DECLARE_CYCLE_STAT(TEXT("UpdateConstraint"), STAT_UpdateConstraint, STATGROUP_ChaosWide);
+
+template <typename T_PARTICLES, typename T, int d>
+void UpdateConstraintImp(const T_PARTICLES& InParticles, const TImplicitObject<T, d>& ParticleObject, const TRigidTransform<T, d>& ParticleTM, const TImplicitObject<T, d>& LevelsetObject, const TRigidTransform<T, d>& LevelsetTM, T Thickness, TRigidBodyContactConstraint<T, d>& Constraint)
+{
+	if (ParticleObject.GetType() == TBox<T, d>::GetType() && LevelsetObject.GetType() == TBox<T, d>::GetType())
+	{
+		UpdateBoxConstraint(*ParticleObject.template GetObject<TBox<T,d>>(), ParticleTM, *LevelsetObject.template GetObject<TBox<T, d>>(), LevelsetTM, Thickness, Constraint);
+	}
+	else if (ParticleObject.GetType() == TSphere<T, d>::GetType() && LevelsetObject.GetType() == TSphere<T, d>::GetType())
+	{
+		UpdateSphereConstraint(*ParticleObject.template GetObject<TSphere<T, d>>(), ParticleTM, *LevelsetObject.template GetObject<TSphere<T, d>>(), LevelsetTM, Thickness, Constraint);
+	}
+	else if (ParticleObject.GetType() == TBox<T, d>::GetType() && LevelsetObject.GetType() == TPlane<T, d>::GetType())
+	{
+		UpdateBoxPlaneConstraint(*ParticleObject.template GetObject<TBox<T, d>>(), ParticleTM, *LevelsetObject.template GetObject<TPlane<T, d>>(), LevelsetTM, Thickness, Constraint);
+	}
+	else if (ParticleObject.GetType() == TSphere<T, d>::GetType() && LevelsetObject.GetType() == TPlane<T, d>::GetType())
+	{
+		UpdateSpherePlaneConstraint(*ParticleObject.template GetObject<TSphere<T, d>>(), ParticleTM, *LevelsetObject.template GetObject<TPlane<T, d>>(), LevelsetTM, Thickness, Constraint);
+	}
+	else if (ParticleObject.GetType() == TSphere<T, d>::GetType() && LevelsetObject.GetType() == TBox<T, d>::GetType())
+	{
+		UpdateSphereBoxConstraint(*ParticleObject.template GetObject<TSphere<T, d>>(), ParticleTM, *LevelsetObject.template GetObject<TBox<T, d>>(), LevelsetTM, Thickness, Constraint);
+	}
+	else if (ParticleObject.GetType() == TPlane<T, d>::GetType() && LevelsetObject.GetType() == TBox<T, d>::GetType())
+	{
+		TRigidBodyContactConstraint<T, d> TmpConstraint = Constraint;
+		UpdateBoxPlaneConstraint(*LevelsetObject.template GetObject<TBox<T, d>>(), LevelsetTM, *ParticleObject.template GetObject<TPlane<T, d>>(), ParticleTM, Thickness, TmpConstraint);
+		if (TmpConstraint.Phi < Constraint.Phi)
+		{
+			Constraint = TmpConstraint;
+			Constraint.Normal = -Constraint.Normal;
+		}
+	}
+	else if (ParticleObject.GetType() == TPlane<T, d>::GetType() && LevelsetObject.GetType() == TSphere<T, d>::GetType())
+	{
+		TRigidBodyContactConstraint<T, d> TmpConstraint = Constraint;
+		UpdateSpherePlaneConstraint(*LevelsetObject.template GetObject<TSphere<T, d>>(), LevelsetTM, *ParticleObject.template GetObject<TPlane<T, d>>(), ParticleTM, Thickness, TmpConstraint);
+		if (TmpConstraint.Phi < Constraint.Phi)
+		{
+			Constraint = TmpConstraint;
+			Constraint.Normal = -Constraint.Normal;
+		}
+	}
+	else if (ParticleObject.GetType() == TBox<T, d>::GetType() && LevelsetObject.GetType() == TSphere<T, d>::GetType())
+	{
+		TRigidBodyContactConstraint<T, d> TmpConstraint = Constraint;
+		UpdateSphereBoxConstraint(*LevelsetObject.template GetObject<TSphere<T, d>>(), LevelsetTM, *ParticleObject.template GetObject<TBox<T, d>>(), ParticleTM, Thickness, TmpConstraint);
+		if (TmpConstraint.Phi < Constraint.Phi)
+		{
+			Constraint = TmpConstraint;
+			Constraint.Normal = -Constraint.Normal;
+		}
+	}
+	else if (ParticleObject.GetType() < TImplicitObjectUnion<T, d>::GetType() && LevelsetObject.GetType() == TImplicitObjectUnion<T, d>::GetType())
+	{
+		return UpdateSingleUnionConstraint(InParticles, Thickness, Constraint);
+	}
+	else if (ParticleObject.GetType() == TImplicitObjectUnion<T, d>::GetType() && LevelsetObject.GetType() < TImplicitObjectUnion<T, d>::GetType())
+	{
+		check(false);	//should not be possible to get this ordering (see ComputeConstraint)
+	}
+	else if (ParticleObject.GetType() == TImplicitObjectUnion<T, d>::GetType() && LevelsetObject.GetType() == TImplicitObjectUnion<T, d>::GetType())
+	{
+		return UpdateUnionUnionConstraint(InParticles, Thickness, Constraint);
+	}
+#if 0
+	else if (ParticleObject.IsConvex() && LevelsetObject.IsConvex())
+	{
+		UpdateLevelsetConstraintGJK(InParticles, Thickness, Constraint);
+	}
+#endif
+	else if (LevelsetObject.IsUnderlyingUnion())
+	{
+		UpdateUnionLevelsetConstraint(InParticles, Thickness, Constraint);
+	}
+	else
+	{
+		TPBDCollisionConstraint<T, d>::UpdateLevelsetConstraint(InParticles, Thickness, Constraint);
+	}
+}
+
+template<typename T, int d>
+template<typename T_PARTICLES>
+void TPBDCollisionConstraint<T, d>::UpdateConstraint(const T_PARTICLES& InParticles, const T Thickness, TRigidBodyContactConstraint<T, d>& Constraint)
+{
+	SCOPE_CYCLE_COUNTER(STAT_UpdateConstraint);
+	Constraint.Phi = Thickness;
+	const TRigidTransform<T, d> ParticleTM = GetTransform(InParticles, Constraint.ParticleIndex);
+	const TRigidTransform<T, d> LevelsetTM = GetTransform(InParticles, Constraint.LevelsetIndex);
+
+	UpdateConstraintImp(InParticles, *InParticles.Geometry(Constraint.ParticleIndex), ParticleTM, *InParticles.Geometry(Constraint.LevelsetIndex), LevelsetTM, Thickness, Constraint);
+}
+
+template<class T, int d>
+bool TPBDCollisionConstraint<T, d>::SleepInactive(TPBDRigidParticles<T, d>& InParticles, const TArray<int32>& ActiveIndices, int32& IslandSleepCount, const int32 Island, const T LinearSleepThreshold, const T AngularSleepThreshold) const
+{
+	return MContactGraph.SleepInactive(InParticles, ActiveIndices, IslandSleepCount, Island, LinearSleepThreshold, AngularSleepThreshold);
+}
+
+template<class T, int d>
+void TPBDCollisionConstraint<T, d>::UpdateIslandsFromConstraints(TPBDRigidParticles<T, d>& InParticles, TArray<TSet<int32>>& IslandParticles, TArray<int32>& IslandSleepCounts, TSet<int32>& ActiveIndices)
+{
+	MContactGraph.UpdateIslandsFromConstraints(InParticles, IslandParticles, IslandSleepCounts, ActiveIndices, MConstraints);
+}
+
+template<class T, int d>
+void TPBDCollisionConstraint<T, d>::UpdateAccelerationStructures(const TPBDRigidParticles<T, d>& InParticles, const TArray<int32>& ActiveIndices, const int32 Island)
+{
+	MContactGraph.UpdateAccelerationStructures(InParticles, ActiveIndices, Island);
+}
+
+template class TPBDCollisionConstraint<float, 3>;
+template void TPBDCollisionConstraint<float, 3>::UpdateConstraint<TPBDRigidParticles<float, 3>>(const TPBDRigidParticles<float, 3>&, const float, FRigidBodyContactConstraint& Constraint);
+template void TPBDCollisionConstraint<float, 3>::UpdateLevelsetConstraint<TPBDRigidParticles<float, 3>>(const TPBDRigidParticles<float, 3>&, const float, FRigidBodyContactConstraint& Constraint);
+template void TPBDCollisionConstraint<float, 3>::UpdateLevelsetConstraintGJK<TPBDRigidParticles<float, 3>>(const TPBDRigidParticles<float, 3>&, const float, FRigidBodyContactConstraint& Constraint);
+}
\ No newline at end of file
diff --git a/Engine/Source/Runtime/Sockets/Private/BSDSockets/SocketsBSD.cpp b/Engine/Source/Runtime/Sockets/Private/BSDSockets/SocketsBSD.cpp
--- a/Engine/Source/Runtime/Sockets/Private/BSDSockets/SocketsBSD.cpp
+++ b/Engine/Source/Runtime/Sockets/Private/BSDSockets/SocketsBSD.cpp
@@ -1,4 +1,4 @@
-// Copyright 1998-2018 Epic Games, Inc. All Rights Reserved.
+// Copyright 1998-2019 Epic Games, Inc. All Rights Reserved.
 
 #include "BSDSockets/SocketsBSD.h"
 
@@ -386,9 +386,9 @@ bool FSocketBSD::JoinMulticastGroup(const FInternetAddr& GroupAddress)
 	if (BSDAddr.GetProtocolFamily() == ESocketProtocolFamily::IPv6)
 	{
 		ipv6_mreq imr;
-		imr.ipv6mr_interface = htonl(BSDAddr.GetScopeId());
+		imr.ipv6mr_interface = 0;
 		imr.ipv6mr_multiaddr = ((sockaddr_in6*)&(BSDAddr.Addr))->sin6_addr;
-		return (setsockopt(Socket, IPPROTO_IPV6, IPV6_JOIN_GROUP, (char*)&imr, sizeof(imr)) == 0);
+		return (setsockopt(Socket, IPPROTO_IPV6, IP_ADD_MEMBERSHIP, (char*)&imr, sizeof(imr)) == 0);
 	}
 #endif
 
@@ -399,6 +399,28 @@ bool FSocketBSD::JoinMulticastGroup(const FInternetAddr& GroupAddress)
 }
 
 
+bool FSocketBSD::JoinMulticastGroup(const FInternetAddr& GroupAddress, const FInternetAddr& InterfaceAddress)
+{
+	const FInternetAddrBSD& BSDAddr = static_cast<const FInternetAddrBSD&>(GroupAddress);
+	const FInternetAddrBSD& BSDIFAddr = static_cast<const FInternetAddrBSD&>(InterfaceAddress);
+
+#if PLATFORM_HAS_BSD_IPV6_SOCKETS
+	if (BSDAddr.GetProtocolFamily() == ESocketProtocolFamily::IPv6)
+	{
+		ipv6_mreq imr;
+		imr.ipv6mr_interface = htonl(BSDIFAddr.GetScopeId());
+		imr.ipv6mr_multiaddr = ((sockaddr_in6*)&(BSDAddr.Addr))->sin6_addr;
+		return (setsockopt(Socket, IPPROTO_IPV6, IP_ADD_MEMBERSHIP, (char*)&imr, sizeof(imr)) == 0);
+	}
+#endif
+
+	ip_mreq imr;
+	imr.imr_interface.s_addr = ((sockaddr_in*)&(BSDIFAddr.Addr))->sin_addr.s_addr;
+	imr.imr_multiaddr = ((sockaddr_in*)&(BSDAddr.Addr))->sin_addr;
+	return (setsockopt(Socket, IPPROTO_IP, IP_ADD_MEMBERSHIP, (char*)&imr, sizeof(imr)) == 0);
+}
+
+
 bool FSocketBSD::LeaveMulticastGroup(const FInternetAddr& GroupAddress)
 {
 	const FInternetAddrBSD& BSDAddr = static_cast<const FInternetAddrBSD&>(GroupAddress);
@@ -407,9 +429,9 @@ bool FSocketBSD::LeaveMulticastGroup(const FInternetAddr& GroupAddress)
 	if (BSDAddr.GetProtocolFamily() == ESocketProtocolFamily::IPv6)
 	{
 		ipv6_mreq imr;
-		imr.ipv6mr_interface = htonl(BSDAddr.GetScopeId());
+		imr.ipv6mr_interface = 0;
 		imr.ipv6mr_multiaddr = ((sockaddr_in6*)&(BSDAddr.Addr))->sin6_addr;
-		return (setsockopt(Socket, IPPROTO_IPV6, IPV6_LEAVE_GROUP, (char*)&imr, sizeof(imr)) == 0);
+		return (setsockopt(Socket, IPPROTO_IPV6, IP_DROP_MEMBERSHIP, (char*)&imr, sizeof(imr)) == 0);
 	}
 #endif
 
@@ -420,11 +442,32 @@ bool FSocketBSD::LeaveMulticastGroup(const FInternetAddr& GroupAddress)
 }
 
 
+bool FSocketBSD::LeaveMulticastGroup(const FInternetAddr& GroupAddress, const FInternetAddr& InterfaceAddress)
+{
+	const FInternetAddrBSD& BSDAddr = static_cast<const FInternetAddrBSD&>(GroupAddress);
+	const FInternetAddrBSD& BSDIFAddr = static_cast<const FInternetAddrBSD&>(InterfaceAddress);
+
+#if PLATFORM_HAS_BSD_IPV6_SOCKETS
+	if (BSDAddr.GetProtocolFamily() == ESocketProtocolFamily::IPv6)
+	{
+		ipv6_mreq imr;
+		imr.ipv6mr_interface = htonl(BSDIFAddr.GetScopeId());
+		imr.ipv6mr_multiaddr = ((sockaddr_in6*)&(BSDAddr.Addr))->sin6_addr;
+		return (setsockopt(Socket, IPPROTO_IPV6, IP_DROP_MEMBERSHIP, (char*)&imr, sizeof(imr)) == 0);
+	}
+#endif
+
+	ip_mreq imr;
+	imr.imr_interface.s_addr = ((sockaddr_in*)&(BSDIFAddr.Addr))->sin_addr.s_addr;
+	imr.imr_multiaddr = ((sockaddr_in*)&(BSDAddr.Addr))->sin_addr;
+	return (setsockopt(Socket, IPPROTO_IP, IP_DROP_MEMBERSHIP, (char*)&imr, sizeof(imr)) == 0);
+}
+
+
 bool FSocketBSD::SetMulticastLoopback(bool bLoopback)
 {
 #if PLATFORM_HAS_BSD_IPV6_SOCKETS
-	uint32 ShouldLoopback = bLoopback ? 1 : 0;
-	return (setsockopt(Socket, IPPROTO_IPV6, IP_MULTICAST_LOOP, (char*)&ShouldLoopback, sizeof(ShouldLoopback)) == 0);
+	return (setsockopt(Socket, IPPROTO_IPV6, IP_MULTICAST_LOOP, (char*)&bLoopback, sizeof(bLoopback)) == 0);
 #endif
 
 	return (setsockopt(Socket, IPPROTO_IP, IP_MULTICAST_LOOP, (char*)&bLoopback, sizeof(bLoopback)) == 0);
@@ -434,14 +477,27 @@ bool FSocketBSD::SetMulticastLoopback(bool bLoopback)
 bool FSocketBSD::SetMulticastTtl(uint8 TimeToLive)
 {
 #if PLATFORM_HAS_BSD_IPV6_SOCKETS
-	uint32 RealTimeToLive = TimeToLive;
-	return (setsockopt(Socket, IPPROTO_IPV6, IP_MULTICAST_TTL, (char*)&RealTimeToLive, sizeof(RealTimeToLive)) == 0);
+	return (setsockopt(Socket, IPPROTO_IPV6, IPV6_MULTICAST_HOPS, (char*)&TimeToLive, sizeof(TimeToLive)) == 0);
 #endif
 
 	return (setsockopt(Socket, IPPROTO_IP, IP_MULTICAST_TTL, (char*)&TimeToLive, sizeof(TimeToLive)) == 0);
 }
 
 
+bool FSocketBSD::SetMulticastInterface(const FInternetAddr& InterfaceAddress)
+{
+	const FInternetAddrBSD& BSDIFAddr = static_cast<const FInternetAddrBSD&>(InterfaceAddress);
+
+#if PLATFORM_HAS_BSD_IPV6_SOCKETS
+	uint32 InterfaceIndex = htonl(BSDIFAddr.GetScopeId());
+	return (setsockopt(Socket, IPPROTO_IPV6, IPV6_MULTICAST_IF, (char*)&InterfaceIndex, sizeof(InterfaceIndex)) == 0);
+#endif
+
+	in_addr InterfaceAddr = ((sockaddr_in*)&(BSDIFAddr.Addr))->sin_addr;
+	return (setsockopt(Socket, IPPROTO_IP, IP_MULTICAST_IF, (char*)&InterfaceAddr, sizeof(InterfaceAddr)) == 0);
+}
+
+
 bool FSocketBSD::SetReuseAddr(bool bAllowReuse)
 {
 	int Param = bAllowReuse ? 1 : 0;
@@ -529,7 +585,7 @@ bool FSocketBSD::SetIPv6Only(bool bIPv6Only)
 {
 #if PLATFORM_HAS_BSD_IPV6_SOCKETS
 	int v6only = bIPv6Only ? 1 : 0;
-	bool bOk = (setsockopt(Socket, IPPROTO_IPV6, IPV6_V6ONLY, (char*)&v6only, sizeof(v6only)) == 0);
+	bool bOk = setsockopt(Socket, IPPROTO_IPV6, IPV6_V6ONLY, (char*)&v6only, sizeof(v6only)) == 0;
 
 	if (bOk == false)
 	{
diff --git a/Engine/Source/Runtime/VulkanRHI/Private/Lumin/VulkanLuminPlatform.h b/Engine/Source/Runtime/VulkanRHI/Private/Lumin/VulkanLuminPlatform.h
--- a/Engine/Source/Runtime/VulkanRHI/Private/Lumin/VulkanLuminPlatform.h
+++ b/Engine/Source/Runtime/VulkanRHI/Private/Lumin/VulkanLuminPlatform.h
@@ -1,4 +1,4 @@
-// Copyright 1998-2018 Epic Games, Inc. All Rights Reserved.
+// Copyright 1998-2019 Epic Games, Inc. All Rights Reserved.
 
 #pragma once
 
@@ -50,13 +50,7 @@ public:
 	static bool SupportsASTCTextureFormats() { return true; }
 	static bool SupportsQuerySurfaceProperties() { return false; }
 
-	static void SetupFeatureLevels()
-	{
-		GShaderPlatformForFeatureLevel[ERHIFeatureLevel::ES2] = SP_VULKAN_ES3_1_LUMIN;
-		GShaderPlatformForFeatureLevel[ERHIFeatureLevel::ES3_1] = SP_VULKAN_ES3_1_LUMIN;
-		GShaderPlatformForFeatureLevel[ERHIFeatureLevel::SM4] = SP_NumPlatforms;
-		GShaderPlatformForFeatureLevel[ERHIFeatureLevel::SM5] = SP_VULKAN_SM5_LUMIN;
-	}
+	static void SetupFeatureLevels();
 
 	static bool SupportsStandardSwapchain() { return false; }
 	static EPixelFormat GetPixelFormatForNonDefaultSwapchain() { return PF_R8G8B8A8; }
