Merging //UE4/Dev-Main to Dev-VR (//UE4/Dev-VR)

diff --git a/Engine/Plugins/Developer/AnimationSharing/Source/AnimationSharing/Public/AnimationSharingModule.h b/Engine/Plugins/Developer/AnimationSharing/Source/AnimationSharing/Public/AnimationSharingModule.h
--- a/Engine/Plugins/Developer/AnimationSharing/Source/AnimationSharing/Public/AnimationSharingModule.h
+++ b/Engine/Plugins/Developer/AnimationSharing/Source/AnimationSharing/Public/AnimationSharingModule.h
@@ -12,7 +12,7 @@
 class UAnimationSharingManager;
 class UAnimationSharingSetup;
 
-class ANIMATIONSHARING_API FAnimSharingModule : public FDefaultGameModuleImpl, public FGCObject
+class ANIMATIONSHARING_API FAnimSharingModule : public FDefaultModuleImpl, public FGCObject
 {
 public:
 	// Begin IModuleInterface overrides
diff --git a/Engine/Plugins/Developer/AnimationSharing/Source/AnimationSharing/Public/AnimationSharingTypes.h b/Engine/Plugins/Developer/AnimationSharing/Source/AnimationSharing/Public/AnimationSharingTypes.h
--- a/Engine/Plugins/Developer/AnimationSharing/Source/AnimationSharing/Public/AnimationSharingTypes.h
+++ b/Engine/Plugins/Developer/AnimationSharing/Source/AnimationSharing/Public/AnimationSharingTypes.h
@@ -96,14 +96,14 @@ class ANIMATIONSHARING_API UAnimationSharingStateProcessor : public UObject
 public:
 	UFUNCTION(BlueprintNativeEvent, Category = AnimationSharing)
 	void ProcessActorState(int32& OutState, AActor* InActor, uint8 CurrentState, uint8 OnDemandState, bool& bShouldProcess);
-	void ProcessActorState_Implementation(int32& OutState, AActor* InActor, uint8 CurrentState, uint8 OnDemandState, bool& bShouldProcess)
+	virtual void ProcessActorState_Implementation(int32& OutState, AActor* InActor, uint8 CurrentState, uint8 OnDemandState, bool& bShouldProcess)
 	{
 		ProcessActorState_Internal(OutState, InActor, CurrentState, OnDemandState, bShouldProcess);
 	}
 
 	UFUNCTION(BlueprintNativeEvent, Category = AnimationSharing)
 	UEnum* GetAnimationStateEnum();
-	UEnum* GetAnimationStateEnum_Implementation()
+	virtual UEnum* GetAnimationStateEnum_Implementation()
 	{
 		return GetAnimationStateEnum_Internal();
 	}
diff --git a/Engine/Plugins/Editor/PluginBrowser/Source/PluginBrowser/Private/SNewPluginWizard.cpp b/Engine/Plugins/Editor/PluginBrowser/Source/PluginBrowser/Private/SNewPluginWizard.cpp
--- a/Engine/Plugins/Editor/PluginBrowser/Source/PluginBrowser/Private/SNewPluginWizard.cpp
+++ b/Engine/Plugins/Editor/PluginBrowser/Source/PluginBrowser/Private/SNewPluginWizard.cpp
@@ -742,7 +742,7 @@ FReply SNewPluginWizard::OnCreatePluginClicked()
 	// If it contains code, we need the user to restart to enable it. Otherwise, we can just mount it now.
 	if (bSucceeded && bHasModules)
 	{
-		FString ProjectFileName = FPaths::GetProjectFilePath();
+		FString ProjectFileName = IFileManager::Get().ConvertToAbsolutePathForExternalAppForWrite(*FPaths::GetProjectFilePath());
 		FString Arguments = FString::Printf(TEXT("%s %s -TargetType=Editor -Plugin=\"%s\" -Project=\"%s\" -Progress -NoHotReloadFromIDE"), FModuleManager::Get().GetUBTConfiguration(), FPlatformMisc::GetUBTPlatform(), *UPluginFilePath, *ProjectFileName);
 		if (!FDesktopPlatformModule::Get()->RunUnrealBuildTool(LOCTEXT("Compiling", "Compiling..."), FPaths::RootDir(), Arguments, GWarn))
 		{
diff --git a/Engine/Plugins/Experimental/RemoteSession/Source/RemoteSession/Private/MessageHandler/RecordingMessageHandler.cpp b/Engine/Plugins/Experimental/RemoteSession/Source/RemoteSession/Private/MessageHandler/RecordingMessageHandler.cpp
--- a/Engine/Plugins/Experimental/RemoteSession/Source/RemoteSession/Private/MessageHandler/RecordingMessageHandler.cpp
+++ b/Engine/Plugins/Experimental/RemoteSession/Source/RemoteSession/Private/MessageHandler/RecordingMessageHandler.cpp
@@ -48,12 +48,14 @@ FRecordingMessageHandler::FRecordingMessageHandler(const TSharedPtr<FGenericAppl
 	BIND_PLAYBACK_HANDLER(TEXT("OnKeyDown"), PlayOnKeyDown);
 
 	BIND_PLAYBACK_HANDLER(TEXT("OnTouchStarted"), PlayOnTouchStarted);
+	BIND_PLAYBACK_HANDLER(TEXT("OnTouchFirstMove"), PlayOnTouchFirstMove);
 	BIND_PLAYBACK_HANDLER(TEXT("OnTouchMoved"), PlayOnTouchMoved);
 	BIND_PLAYBACK_HANDLER(TEXT("OnTouchEnded"), PlayOnTouchEnded);
 	BIND_PLAYBACK_HANDLER(TEXT("OnMotionDetected"), PlayOnMotionDetected);
 	BIND_PLAYBACK_HANDLER(TEXT("OnBeginGesture"), PlayOnBeginGesture);
 	BIND_PLAYBACK_HANDLER(TEXT("OnTouchGesture"), PlayOnTouchGesture);
 	BIND_PLAYBACK_HANDLER(TEXT("OnEndGesture"), PlayOnEndGesture);
+	BIND_PLAYBACK_HANDLER(TEXT("OnTouchForceChanged"), PlayOnTouchForceChanged);
 }
 
 #undef BIND_PLAYBACK_HANDLER
diff --git a/Engine/Source/Developer/BlueprintCompilerCppBackend/Private/BlueprintCompilerCppBackendUtils.cpp b/Engine/Source/Developer/BlueprintCompilerCppBackend/Private/BlueprintCompilerCppBackendUtils.cpp
--- a/Engine/Source/Developer/BlueprintCompilerCppBackend/Private/BlueprintCompilerCppBackendUtils.cpp
+++ b/Engine/Source/Developer/BlueprintCompilerCppBackend/Private/BlueprintCompilerCppBackendUtils.cpp
@@ -1604,7 +1604,8 @@ bool FEmitHelper::GenerateAutomaticCast(FEmitterLocalContext& EmitterContext, co
 
 				if (bForceReference)
 				{
-					OutCastBegin = TEXT("*static_cast<uint8*>(&(");
+					// An enum and its underlying type are not related by inheritance, so 'static_cast' cannot be used here (5.2.9/11).
+					OutCastBegin = TEXT("*reinterpret_cast<uint8*>(&(");
 					OutCastEnd = TEXT("))");
 				}
 				else
diff --git a/Engine/Source/Developer/HotReload/Private/HotReload.cpp b/Engine/Source/Developer/HotReload/Private/HotReload.cpp
--- a/Engine/Source/Developer/HotReload/Private/HotReload.cpp
+++ b/Engine/Source/Developer/HotReload/Private/HotReload.cpp
@@ -1632,12 +1632,6 @@ bool FHotReloadModule::StartCompilingModuleDLLs(const TArray< FModuleToRecompile
 		ExtraArg += TEXT( "-FailIfGeneratedCodeChanges " );
 	}
 
-	// Shared PCH does no work with hot-reloading engine/editor modules as we don't scan all modules for them.
-	if (!ContainsOnlyGameModules(ModuleNames))
-	{
-		ExtraArg += TEXT("-nosharedpch ");
-	}
-	
 	FString CmdLineParams = FString::Printf( TEXT( "%s %s %s %s%s -IgnoreJunk" ), 
 		*ModuleArg, 
 		BuildPlatformName, BuildConfigurationName, 
diff --git a/Engine/Source/Programs/UnrealBuildTool/Configuration/TargetDescriptor.cs b/Engine/Source/Programs/UnrealBuildTool/Configuration/TargetDescriptor.cs
--- a/Engine/Source/Programs/UnrealBuildTool/Configuration/TargetDescriptor.cs
+++ b/Engine/Source/Programs/UnrealBuildTool/Configuration/TargetDescriptor.cs
@@ -252,6 +252,12 @@ namespace UnrealBuildTool
 				throw new BuildException("No configurations specified for target");
 			}
 
+			// Make sure the project file exists
+			if(ProjectFile != null && !FileReference.Exists(ProjectFile))
+			{
+				throw new BuildException("Unable to find project '{0}'.", ProjectFile);
+			}
+
 			// Expand all the platforms, architectures and configurations
 			foreach(UnrealTargetPlatform Platform in Platforms)
 			{
diff --git a/Engine/Source/Runtime/Engine/Classes/Components/SkinnedMeshComponent.h b/Engine/Source/Runtime/Engine/Classes/Components/SkinnedMeshComponent.h
--- a/Engine/Source/Runtime/Engine/Classes/Components/SkinnedMeshComponent.h
+++ b/Engine/Source/Runtime/Engine/Classes/Components/SkinnedMeshComponent.h
@@ -518,6 +518,9 @@ protected:
 	UPROPERTY(Transient)
 	mutable uint8 bCachedLocalBoundsUpToDate:1;
 
+	/** Whether we have updated bone visibility this tick */
+	uint8 bBoneVisibilityDirty:1;
+
 private:
 	/** If true, UpdateTransform will always result in a call to MeshObject->Update. */
 	UPROPERTY(transient)
diff --git a/Engine/Source/Runtime/Engine/Private/Components/CharacterMovementComponent.cpp b/Engine/Source/Runtime/Engine/Private/Components/CharacterMovementComponent.cpp
--- a/Engine/Source/Runtime/Engine/Private/Components/CharacterMovementComponent.cpp
+++ b/Engine/Source/Runtime/Engine/Private/Components/CharacterMovementComponent.cpp
@@ -7672,7 +7672,11 @@ bool UCharacterMovementComponent::ForcePositionUpdate(float DeltaTime)
 	ServerData->ServerAccumulatedClientTimeStamp += DeltaTime;
 
 #if !(UE_BUILD_SHIPPING)
-	UE_LOG(LogNetPlayerMovement, Warning, TEXT("ForcePositionUpdate %s (DeltaTime %.2f)"), *CharacterOwner->GetName(), DeltaTime);
+	const bool bServerMoveHasOccurred = (ServerData->ServerTimeStampLastServerMove != 0.f);
+	if (bServerMoveHasOccurred)
+	{
+		UE_LOG(LogNetPlayerMovement, Warning, TEXT("ForcePositionUpdate %s (DeltaTime %.2f)"), *CharacterOwner->GetName(), DeltaTime);
+	}
 #endif
 
 	// Force movement update.
diff --git a/Engine/Source/Runtime/Engine/Private/Components/SkinnedMeshComponent.cpp b/Engine/Source/Runtime/Engine/Private/Components/SkinnedMeshComponent.cpp
--- a/Engine/Source/Runtime/Engine/Private/Components/SkinnedMeshComponent.cpp
+++ b/Engine/Source/Runtime/Engine/Private/Components/SkinnedMeshComponent.cpp
@@ -348,6 +348,7 @@ USkinnedMeshComponent::USkinnedMeshComponent(const FObjectInitializer& ObjectIni
 	CurrentEditableComponentTransforms = 0;
 	CurrentReadComponentTransforms = 1;
 	bNeedToFlipSpaceBaseBuffers = false;
+	bBoneVisibilityDirty = false;
 
 	bCanEverAffectNavigation = false;
 	MasterBoneMapCacheCount = 0;
@@ -944,6 +945,8 @@ void USkinnedMeshComponent::RebuildVisibilityArray()
 				}
 			}
 		}
+
+		bBoneVisibilityDirty = true;
 	}
 }
 
@@ -2798,6 +2801,13 @@ void USkinnedMeshComponent::FlipEditableSpaceBases()
 			{
 				GetEditableComponentSpaceTransforms() = GetComponentSpaceTransforms();
 				GetEditableBoneVisibilityStates() = GetBoneVisibilityStates();
+				bBoneVisibilityDirty = false;
+			}
+			// If we have changed bone visibility, then we need to reflect that next frame
+			else if(bBoneVisibilityDirty)
+			{
+				GetEditableBoneVisibilityStates() = GetBoneVisibilityStates();
+				bBoneVisibilityDirty = false;
 			}
 		}
 		else
diff --git a/Engine/Source/Runtime/Engine/Private/Net/NetPacketNotify.cpp b/Engine/Source/Runtime/Engine/Private/Net/NetPacketNotify.cpp
--- a/Engine/Source/Runtime/Engine/Private/Net/NetPacketNotify.cpp
+++ b/Engine/Source/Runtime/Engine/Private/Net/NetPacketNotify.cpp
@@ -4,6 +4,12 @@
 #include "Serialization/BitWriter.h"
 #include "Serialization/BitReader.h"
 
+FNetPacketNotify::FNetPacketNotify()
+	: AckRecord(64)
+	, WrittenHistoryWordCount(0)
+{
+}
+
 SIZE_T FNetPacketNotify::GetCurrentSequenceHistoryLength() const
 {
 	if (InAckSeq >= InAckSeqAck)
@@ -17,6 +23,32 @@ SIZE_T FNetPacketNotify::GetCurrentSequenceHistoryLength() const
 	}
 }
 
+FNetPacketNotify::SequenceNumberT FNetPacketNotify::UpdateInAckSeqAck(SequenceNumberT::DifferenceT AckCount, SequenceNumberT AckedSeq)
+{
+	check((SIZE_T)AckCount <= AckRecord.Count());
+
+	if (AckCount > 1)
+	{
+		AckRecord.PopNoCheck(AckCount - 1);
+	}
+
+	FSentAckData AckData = AckRecord.PeekNoCheck();
+	AckRecord.PopNoCheck();
+
+	// verify that we have a matching sequence number
+	if (AckData.OutSeq == AckedSeq)
+	{
+		return AckData.InAckSeq;
+	}
+	else
+	{
+		UE_LOG_PACKET_NOTIFY_WARNING(TEXT("FNetPacketNotify::UpdateInAckSeqAck - Failed to find matching AckRecord for %u, (Found %u)"), AckedSeq.Get(), AckData.OutSeq.Get());
+
+		// Pessimistic view, should never occur
+		return SequenceNumberT(AckedSeq.Get() - MaxSequenceHistoryLength);
+	}
+}
+
 void FNetPacketNotify::Init(SequenceNumberT InitialInSeq, SequenceNumberT InitialOutSeq)
 {
 	InSeqHistory.Reset();
@@ -75,7 +107,7 @@ namespace
 }
 
 // These methods must always write and read the exact same number of bits, that is the reason for not using WriteInt/WrittedWrappedInt
-void FNetPacketNotify::WriteHeader(FBitWriter& Writer, bool bRefresh)
+bool FNetPacketNotify::WriteHeader(FBitWriter& Writer, bool bRefresh)
 {
 	// we always write at least 1 word
 	SIZE_T CurrentHistoryWordCount = FMath::Clamp<SIZE_T>((GetCurrentSequenceHistoryLength() + SequenceHistoryT::BitsPerWord - 1u) / SequenceHistoryT::BitsPerWord, 1u, SequenceHistoryT::WordCount);
@@ -83,7 +115,7 @@ void FNetPacketNotify::WriteHeader(FBitWriter& Writer, bool bRefresh)
 	// We can only do a refresh if we do not need more space for the history
 	if (bRefresh && (CurrentHistoryWordCount > WrittenHistoryWordCount))
 	{
-		return;
+		return false;
 	}
 
 	// How many words of ack data should we write? If this is a refresh we must write the same size as the original header
@@ -104,6 +136,8 @@ void FNetPacketNotify::WriteHeader(FBitWriter& Writer, bool bRefresh)
 	InSeqHistory.Write(Writer, WrittenHistoryWordCount);
 
 	UE_LOG_PACKET_NOTIFY(TEXT("FNetPacketNotify::WriteHeader - Seq %u, AckedSeq %u bReFresh %u HistorySizeInWords %u"), Seq, AckedSeq, bRefresh ? 1u : 0u, WrittenHistoryWordCount);
+
+	return true;
 }
 
 bool FNetPacketNotify::ReadHeader(FNotificationHeader& Data, FBitReader& Reader) const
diff --git a/Engine/Source/Runtime/Engine/Private/NetConnection.cpp b/Engine/Source/Runtime/Engine/Private/NetConnection.cpp
--- a/Engine/Source/Runtime/Engine/Private/NetConnection.cpp
+++ b/Engine/Source/Runtime/Engine/Private/NetConnection.cpp
@@ -1228,8 +1228,7 @@ void UNetConnection::FlushNet(bool bIgnoreSimulation)
 		// Increase outgoing sequence number
 		if (!InternalAck)
 		{
-			PacketNotify.CommitAndIncrementOutSeq();
-			HasDirtyAcks = 0u;
+			PacketNotify.CommitAndIncrementOutSeq();			
 		}
 		++OutPacketId; 
 
@@ -1386,12 +1385,18 @@ void UNetConnection::WritePacketHeader(FBitWriter& Writer)
 	Reset.PopWithoutClear(Writer);
 	
 	// Write notification header or refresh the header if used space is the same.
-	PacketNotify.WriteHeader(Writer, bIsHeaderUpdate);
+	bool bWroteHeader = PacketNotify.WriteHeader(Writer, bIsHeaderUpdate);
 
 	// Jump back to where we came from.
 	if (bIsHeaderUpdate)
 	{
 		Restore.PopWithoutClear(Writer);
+
+		// if we wrote the header and successfully refreshed the header status we no longer has any dirty acks
+		if (bWroteHeader)
+		{
+			HasDirtyAcks = 0u;
+		}
 	}
 }
 
@@ -2018,11 +2023,17 @@ PRAGMA_ENABLE_DEPRECATION_WARNINGS
 		++HasDirtyAcks;
 
 		// This is to allow us to recover from hitches were we process more than FNetPacketNotify::SequenceHistoryLength packets in a row withouht sending out any packets.
-		// In most cases this allows us to recover smoothly without reporting any pessimistic naks which is the case is we overshoot our ack history.
+		// In most cases this allows us to recover smoothly without reporting any pessimistic naks which will be the case if we overshoot the ack history.
+		// Note: This should only occur if we are running with no timeouts!
 		if (HasDirtyAcks >= FNetPacketNotify::MaxSequenceHistoryLength)
 		{
-			UE_LOG(LogNet, Warning, TEXT("UNetConnection::ReceivedPacket - Too many received acks(%u) since last sent packet. %s "), HasDirtyAcks, *Describe());
+			UE_LOG(LogNet, Warning, TEXT("UNetConnection::ReceivedPacket - Too many received packets to ack (%u) since last sent packet. InSeq: %u %s NextOutGoingSeq: %u"), HasDirtyAcks, PacketNotify.GetInSeq().Get(), *Describe(), PacketNotify.GetOutSeq().Get());
+
 			FlushNet();
+			if (HasDirtyAcks) // if acks still are dirty, flush again
+			{
+				FlushNet();
+			}
 		}
 	}
 }
@@ -3332,4 +3343,4 @@ static void	PrintActorReportFunc(const TArray<FString>& Args, UWorld* InWorld)
 	ObjKeyMap.Dump(*GLog);
 }
 
-FAutoConsoleCommandWithWorldAndArgs PrintActorReportCmd(TEXT("net.ActorReport"), TEXT(""),	FConsoleCommandWithWorldAndArgsDelegate::CreateStatic(PrintActorReportFunc) );
\ No newline at end of file
+FAutoConsoleCommandWithWorldAndArgs PrintActorReportCmd(TEXT("net.ActorReport"), TEXT(""),	FConsoleCommandWithWorldAndArgsDelegate::CreateStatic(PrintActorReportFunc) );
diff --git a/Engine/Source/Runtime/Engine/Private/VorbisAudioInfo.cpp b/Engine/Source/Runtime/Engine/Private/VorbisAudioInfo.cpp
--- a/Engine/Source/Runtime/Engine/Private/VorbisAudioInfo.cpp
+++ b/Engine/Source/Runtime/Engine/Private/VorbisAudioInfo.cpp
@@ -483,6 +483,11 @@ void FVorbisAudioInfo::EnableHalfRate( bool HalfRate )
 
 bool FVorbisAudioInfo::StreamCompressedInfo(USoundWave* Wave, struct FSoundQualityInfo* QualityInfo)
 {
+	if (!bDllLoaded)
+	{
+		return false;
+	}
+
 	SCOPE_CYCLE_COUNTER( STAT_VorbisPrepareDecompressionTime );
 
 	FScopeLock ScopeLock(&VorbisCriticalSection);
@@ -511,6 +516,11 @@ bool FVorbisAudioInfo::StreamCompressedInfo(USoundWave* Wave, struct FSoundQuali
 
 bool FVorbisAudioInfo::StreamCompressedData(uint8* InDestination, bool bLooping, uint32 BufferSize)
 {
+	if (!bDllLoaded)
+	{
+		return false;
+	}
+
 	check( VFWrapper != NULL );
 
 #if PLATFORM_ANDROID
diff --git a/Engine/Source/Runtime/Engine/Public/Net/NetPacketNotify.h b/Engine/Source/Runtime/Engine/Public/Net/NetPacketNotify.h
--- a/Engine/Source/Runtime/Engine/Public/Net/NetPacketNotify.h
+++ b/Engine/Source/Runtime/Engine/Public/Net/NetPacketNotify.h
@@ -48,10 +48,7 @@ public:
 	};
 
 	/** Constructor */
-	FNetPacketNotify()
-	: AckRecord(64)
-	, WrittenHistoryWordCount(0)
-	{}
+	FNetPacketNotify();
 
 	/** Init notification with expected initial sequence numbers */
 	void Init(SequenceNumberT InitialInSeq, SequenceNumberT InitialOutSeq);
@@ -66,9 +63,10 @@ public:
 	SequenceNumberT CommitAndIncrementOutSeq();
 
 	/** Write NotificationHeader, and update outgoing ack record 
-	 if bRefresh is true we will attempt to refresh a previously written header 
-	if the size remains the same */
-	void WriteHeader(FBitWriter& Writer, bool bRefresh = false);
+		if bRefresh is true we will attempt to refresh a previously written header if the resulting size will be the same as the already written header.
+		returns true if data was written, and false if no data was written which might be the case if we try to rewrite an existing header but the required size differs.
+	*/
+	bool WriteHeader(FBitWriter& Writer, bool bRefresh = false);
 	
 	/** Read header from stream */
 	bool ReadHeader(FNotificationHeader& Data, FBitReader& Reader) const;
@@ -80,26 +78,7 @@ public:
 		Returns 0 if the received sequence number is out of the current window or if the acknowledged seq received by the remote is invalid.
 	*/
 	template<class Functor>
-	SequenceNumberT::DifferenceT Update(const FNotificationHeader& NotificationData, Functor&& InFunc)
-	{		
-		if (NotificationData.Seq > InSeq && NotificationData.AckedSeq >= OutAckSeq)
-		{
-			UE_LOG_PACKET_NOTIFY(TEXT("FNetPacketNotify::Update - Seq %u, InSeq %u"), NotificationData.Seq.Get(), InSeq.Get());
-
-			const SequenceNumberT::DifferenceT InSeqDelta = SequenceNumberT::Diff(NotificationData.Seq, InSeq);
-			
-			ProcessReceivedAcks(NotificationData, InFunc);
-
-			// accept sequence
-			InSeq = NotificationData.Seq;
-
-			return InSeqDelta;
-		}
-		else
-		{
-			return 0;
-		}
-	}
+	SequenceNumberT::DifferenceT Update(const FNotificationHeader& NotificationData, Functor&& InFunc);
 
 	/** Get the current SequenceHistory */
 	const SequenceHistoryT& GetInSeqHistory() const { return InSeqHistory; }
@@ -145,72 +124,78 @@ private:
 	SequenceNumberT OutAckSeq;			// Last sequence number that we know that the remote side have received.
 
 private:
-	SequenceNumberT UpdateInAckSeqAck(SequenceNumberT::DifferenceT AckCount, SequenceNumberT AckedSeq)
+	SequenceNumberT UpdateInAckSeqAck(SequenceNumberT::DifferenceT AckCount, SequenceNumberT AckedSeq);
+
+	template<class Functor>
+	inline void ProcessReceivedAcks(const FNotificationHeader& NotificationData, Functor&& InFunc);
+	void AckSeq(SequenceNumberT AckedSeq, bool IsAck);
+
+#if WITH_DEV_AUTOMATION_TESTS
+	friend struct FNetPacketNotifyTestUtil;
+#endif
+};
+
+template<class Functor>
+FNetPacketNotify::SequenceNumberT::DifferenceT FNetPacketNotify::Update(const FNotificationHeader& NotificationData, Functor&& InFunc)
+{
+	if (NotificationData.Seq > InSeq && NotificationData.AckedSeq >= OutAckSeq)
 	{
-		check((SIZE_T)AckCount <= AckRecord.Count());
+		UE_LOG_PACKET_NOTIFY(TEXT("FNetPacketNotify::Update - Seq %u, InSeq %u"), NotificationData.Seq.Get(), InSeq.Get());
 
-		if (AckCount > 1)
-		{
-			AckRecord.PopNoCheck(AckCount - 1);
-		}
-		
-		FSentAckData AckData = AckRecord.PeekNoCheck();
-		AckRecord.PopNoCheck();
+		const SequenceNumberT::DifferenceT InSeqDelta = SequenceNumberT::Diff(NotificationData.Seq, InSeq);
 
-		// verify that we have a matching sequence number
-		if (AckData.OutSeq == AckedSeq)
-		{
-			return AckData.InAckSeq;
-		}
-		else
-		{
-			UE_LOG_PACKET_NOTIFY_WARNING(TEXT("FNetPacketNotify::UpdateInAckSeqAck - Failed to find matching AckRecord for %u, (Found %u)"), AckedSeq.Get(), AckData.OutSeq.Get());
+		ProcessReceivedAcks(NotificationData, InFunc);
 
-			// Pessimistic view, should never occur
-			return SequenceNumberT(AckedSeq.Get() - MaxSequenceHistoryLength);
-		}
+		// accept sequence
+		InSeq = NotificationData.Seq;
+
+		return InSeqDelta;
+	}
+	else
+	{
+		return 0;
 	}
+}
 
-	template<class Functor>
-	void ProcessReceivedAcks(const FNotificationHeader& NotificationData, Functor&& InFunc)
+template<class Functor>
+void FNetPacketNotify::ProcessReceivedAcks(const FNotificationHeader& NotificationData, Functor&& InFunc)
+{
+	if (NotificationData.AckedSeq > OutAckSeq)
 	{
-		if (NotificationData.AckedSeq > OutAckSeq)
-		{
-			UE_LOG_PACKET_NOTIFY(TEXT("Notification::ProcessReceivedAcks - AckedSeq: %u, OutAckSeq: %u"), NotificationData.AckedSeq.Get(), OutAckSeq.Get());
+		UE_LOG_PACKET_NOTIFY(TEXT("Notification::ProcessReceivedAcks - AckedSeq: %u, OutAckSeq: %u"), NotificationData.AckedSeq.Get(), OutAckSeq.Get());
+
+		SequenceNumberT::DifferenceT AckCount = SequenceNumberT::Diff(NotificationData.AckedSeq, OutAckSeq);
+
+		// Update InAckSeqAck used to track the needed number of bits to transmit our ack history
+		InAckSeqAck = UpdateInAckSeqAck(AckCount, NotificationData.AckedSeq);
 
-			SequenceNumberT::DifferenceT AckCount = SequenceNumberT::Diff(NotificationData.AckedSeq, OutAckSeq);
+		// ExpectedAck = OutAckSeq + 1
+		SequenceNumberT CurrentAck(OutAckSeq);
+		++CurrentAck;
 
-			// Update InAckSeqAck used to track the needed number of bits to transmit our ack history
-			InAckSeqAck = UpdateInAckSeqAck(AckCount, NotificationData.AckedSeq);
+		if (AckCount > (SequenceNumberT::DifferenceT)(SequenceHistoryT::Size))
+		{
+			UE_LOG_PACKET_NOTIFY_WARNING(TEXT("Notification::ProcessReceivedAcks - Missed Acks: AckedSeq: %u, OutAckSeq: %u, FirstMissingSeq: %u Count: %u"), NotificationData.AckedSeq.Get(), OutAckSeq.Get(), CurrentAck.Get(), AckCount - (SequenceNumberT::DifferenceT)(SequenceHistoryT::Size));
+		}
 
-			// ExpectedAck = OutAckSeq + 1
-			SequenceNumberT CurrentAck(OutAckSeq);
+		// Everything not found in the history buffer is treated as lost
+		while (AckCount > (SequenceNumberT::DifferenceT)(SequenceHistoryT::Size))
+		{
+			--AckCount;
+			InFunc(CurrentAck, false);
 			++CurrentAck;
+		}
 
-			// Everything not found in the history buffer is treated as lost
-			while (AckCount > (SequenceNumberT::DifferenceT)(SequenceHistoryT::Size))
-			{
-				--AckCount;
-				UE_LOG_PACKET_NOTIFY_WARNING(TEXT("Notification::ProcessReceivedAcks Seq: %u - IsAck: 0 HistoryIndex: N/A"), CurrentAck.Get());
-				InFunc(CurrentAck, false);
-				++CurrentAck;
-			}
-
-			// For sequence numbers contained in the history we lookup the delivery status from the history
-			while (AckCount > 0)
-			{
-				--AckCount;
-				UE_LOG_PACKET_NOTIFY(TEXT("Notification::ProcessReceivedAcks Seq: %u - IsAck: %u HistoryIndex: %u"), CurrentAck.Get(), NotificationData.History.IsDelivered(AckCount) ? 1u : 0u, AckCount);
-				InFunc(CurrentAck, NotificationData.History.IsDelivered(AckCount));
-				++CurrentAck;
-			}
-			OutAckSeq = NotificationData.AckedSeq;
+		// For sequence numbers contained in the history we lookup the delivery status from the history
+		while (AckCount > 0)
+		{
+			--AckCount;
+			UE_LOG_PACKET_NOTIFY(TEXT("Notification::ProcessReceivedAcks Seq: %u - IsAck: %u HistoryIndex: %u"), CurrentAck.Get(), NotificationData.History.IsDelivered(AckCount) ? 1u : 0u, AckCount);
+			InFunc(CurrentAck, NotificationData.History.IsDelivered(AckCount));
+			++CurrentAck;
 		}
+		OutAckSeq = NotificationData.AckedSeq;
 	}
+}
 
-	void AckSeq(SequenceNumberT AckedSeq, bool IsAck);
 
-#if WITH_DEV_AUTOMATION_TESTS
-	friend struct FNetPacketNotifyTestUtil;
-#endif
-};
diff --git a/Engine/Source/Runtime/Launch/Private/IOS/IOSAppDelegateConsoleHandling.cpp b/Engine/Source/Runtime/Launch/Private/IOS/IOSAppDelegateConsoleHandling.cpp
--- a/Engine/Source/Runtime/Launch/Private/IOS/IOSAppDelegateConsoleHandling.cpp
+++ b/Engine/Source/Runtime/Launch/Private/IOS/IOSAppDelegateConsoleHandling.cpp
@@ -119,7 +119,7 @@ extern bool GShowSplashScreen;
 {
 	if ([ConsoleCommand length] > 0)
 	{
-		if (self.bEngineInit)
+		if (self.bEngineInit && GEngine != nullptr)
 		{
 			TArray<TCHAR> Ch;
 			Ch.AddZeroed([ConsoleCommand length]);
diff --git a/Engine/Source/Runtime/PreLoadScreen/Private/BuildPatchServicesPreLoadManager.cpp b/Engine/Source/Runtime/PreLoadScreen/Private/BuildPatchServicesPreLoadManager.cpp
--- a/Engine/Source/Runtime/PreLoadScreen/Private/BuildPatchServicesPreLoadManager.cpp
+++ b/Engine/Source/Runtime/PreLoadScreen/Private/BuildPatchServicesPreLoadManager.cpp
@@ -64,6 +64,14 @@ void FBuildPatchServicesPreLoadManagerBase::ResumeBuildPatchInstall()
     }
 }
 
+void FBuildPatchServicesPreLoadManagerBase::CancelBuildPatchInstall()
+{
+	if (ContentBuildInstaller.IsValid())
+	{
+		ContentBuildInstaller->CancelInstall();
+	}
+}
+
 const FText& FBuildPatchServicesPreLoadManagerBase::GetStatusText() const
 {
 	// Static const fixed FText values so that they are not constantly constructed
diff --git a/Engine/Source/Runtime/PreLoadScreen/Public/BuildPatchServicesPreLoadManager.h b/Engine/Source/Runtime/PreLoadScreen/Public/BuildPatchServicesPreLoadManager.h
--- a/Engine/Source/Runtime/PreLoadScreen/Public/BuildPatchServicesPreLoadManager.h
+++ b/Engine/Source/Runtime/PreLoadScreen/Public/BuildPatchServicesPreLoadManager.h
@@ -31,6 +31,7 @@ public:
 
     virtual void PauseBuildPatchInstall();
     virtual void ResumeBuildPatchInstall();
+	virtual void CancelBuildPatchInstall();
 
     float GetProgressPercent() const { return ContentBuildInstaller.IsValid() ? ContentBuildInstaller->GetUpdateProgress() : 0.0f; }
     EBuildPatchDownloadHealth GetDownloadHealth() const { return ContentBuildInstaller.IsValid() ? ContentBuildInstaller->GetDownloadHealth() : EBuildPatchDownloadHealth::NUM_Values; }
