Fixed motion controllers not respecting listen servers

diff --git a/Engine/Source/Runtime/Engine/Classes/GameFramework/Actor.h b/Engine/Source/Runtime/Engine/Classes/GameFramework/Actor.h
--- a/Engine/Source/Runtime/Engine/Classes/GameFramework/Actor.h
+++ b/Engine/Source/Runtime/Engine/Classes/GameFramework/Actor.h
@@ -262,6 +262,13 @@ public:
 	 */
 	virtual bool HasNetOwner() const;
 
+	/**
+	 * Does this actor have a locally controlled owner responsible for replication? (APlayerController typically)
+	 *
+	 * @return true if this actor can call RPCs or false if no such owner chain exists
+	 */
+	virtual bool HasLocalNetOwner() const;
+
 protected:
 	/**
 	 * If true, this actor will replicate to remote machines
diff --git a/Engine/Source/Runtime/Engine/Private/Actor.cpp b/Engine/Source/Runtime/Engine/Private/Actor.cpp
--- a/Engine/Source/Runtime/Engine/Private/Actor.cpp
+++ b/Engine/Source/Runtime/Engine/Private/Actor.cpp
@@ -1600,6 +1600,24 @@ void AActor::SetOwner(AActor* NewOwner)
 	}
 }
 
+bool AActor::HasLocalNetOwner() const
+{
+	if (Owner == nullptr)
+	{
+		// all basic AActors are unable to call RPCs without special AActors as their owners (ie APlayerController)
+		return false;
+	}
+
+	// Find the topmost actor in this owner chain
+	AActor* TopOwner = nullptr;
+	for (TopOwner = Owner; TopOwner->Owner; TopOwner = TopOwner->Owner)
+	{
+	}
+
+	AController* Controller = Cast<AController>(TopOwner);
+	return Controller && Controller->IsLocalController();
+}
+
 bool AActor::HasNetOwner() const
 {
 	if (Owner == nullptr)
diff --git a/Engine/Source/Runtime/HeadMountedDisplay/Private/MotionControllerComponent.cpp b/Engine/Source/Runtime/HeadMountedDisplay/Private/MotionControllerComponent.cpp
--- a/Engine/Source/Runtime/HeadMountedDisplay/Private/MotionControllerComponent.cpp
+++ b/Engine/Source/Runtime/HeadMountedDisplay/Private/MotionControllerComponent.cpp
@@ -471,7 +471,7 @@ bool UMotionControllerComponent::PollControllerState(FVector& Position, FRotator
 	{
 		// Cache state from the game thread for use on the render thread
 		const AActor* MyOwner = GetOwner();
-		bHasAuthority = MyOwner->HasNetOwner();
+		bHasAuthority = MyOwner->HasLocalNetOwner();
 	}
 
 	if(bHasAuthority)
