Slight tweak to the HasLocalNetOwner() logic

diff --git a/Engine/Source/Runtime/Engine/Private/Actor.cpp b/Engine/Source/Runtime/Engine/Private/Actor.cpp
--- a/Engine/Source/Runtime/Engine/Private/Actor.cpp
+++ b/Engine/Source/Runtime/Engine/Private/Actor.cpp
@@ -1615,6 +1615,12 @@ bool AActor::HasLocalNetOwner() const
 	{
 	}
 
+	// Top owner is will normally be a Pawn or a Controller
+	if (APawn* Pawn = Cast<APawn>(TopOwner))
+	{
+		return Pawn->IsLocallyControlled();
+	}
+
 	AController* Controller = Cast<AController>(TopOwner);
 	return Controller && Controller->IsLocalController();
 }
