Fixed the case where HasLocalNetOwner() was called on an actor that was the top most in the ownership chain (thanks JB!)

diff --git a/Engine/Source/Runtime/Engine/Private/Actor.cpp b/Engine/Source/Runtime/Engine/Private/Actor.cpp
--- a/Engine/Source/Runtime/Engine/Private/Actor.cpp
+++ b/Engine/Source/Runtime/Engine/Private/Actor.cpp
@@ -1603,19 +1603,18 @@ void AActor::SetOwner(AActor* NewOwner)
 
 bool AActor::HasLocalNetOwner() const
 {
-	if (Owner == nullptr)
-	{
-		// all basic AActors are unable to call RPCs without special AActors as their owners (ie APlayerController)
-		return false;
-	}
+	// I might be the top owner if I am a Pawn or a Controller (owner will be null)
+	AActor* TopOwner = this;
 
-	// Find the topmost actor in this owner chain
-	AActor* TopOwner = nullptr;
-	for (TopOwner = Owner; TopOwner->Owner; TopOwner = TopOwner->Owner)
+	if (Owner != nullptr)
 	{
+		// I have an owner so search that for the top owner
+		for (TopOwner = Owner; TopOwner->Owner; TopOwner = TopOwner->Owner)
+		{
+		}
 	}
 
-	// Top owner is will normally be a Pawn or a Controller
+	// Top owner will normally be a Pawn or a Controller
 	if (APawn* Pawn = Cast<APawn>(TopOwner))
 	{
 		return Pawn->IsLocallyControlled();
