Fixed my const correctness rush to fix bug

diff --git a/Engine/Source/Runtime/Engine/Private/Actor.cpp b/Engine/Source/Runtime/Engine/Private/Actor.cpp
--- a/Engine/Source/Runtime/Engine/Private/Actor.cpp
+++ b/Engine/Source/Runtime/Engine/Private/Actor.cpp
@@ -1604,7 +1604,7 @@ void AActor::SetOwner(AActor* NewOwner)
 bool AActor::HasLocalNetOwner() const
 {
 	// I might be the top owner if I am a Pawn or a Controller (owner will be null)
-	AActor* TopOwner = this;
+	const AActor* TopOwner = this;
 
 	if (Owner != nullptr)
 	{
