Added weighted moving average smoothing functions

diff --git a/Engine/Source/Runtime/AnimGraphRuntime/Private/AnimNodes/AnimNode_ModifyCurve.cpp b/Engine/Source/Runtime/AnimGraphRuntime/Private/AnimNodes/AnimNode_ModifyCurve.cpp
--- a/Engine/Source/Runtime/AnimGraphRuntime/Private/AnimNodes/AnimNode_ModifyCurve.cpp
+++ b/Engine/Source/Runtime/AnimGraphRuntime/Private/AnimNodes/AnimNode_ModifyCurve.cpp
@@ -64,7 +64,7 @@ void FAnimNode_ModifyCurve::Evaluate_AnyThread(FPoseContext& Output)
 				check(LastCurveValues.Num() == CurveValues.Num());
 
 				const float LastCurveValue = LastCurveValues[ModIdx];
-				const float WAvg = (CurrentValue * Alpha) + (LastCurveValue * (1.f - Alpha));
+				const float WAvg = FMath::WeightedMovingAverage(CurrentValue, LastCurveValue, Alpha);
 				// Update the last curve value for next run
 				LastCurveValues[ModIdx] = WAvg;
 
diff --git a/Engine/Source/Runtime/Core/Public/Math/UnrealMathUtility.h b/Engine/Source/Runtime/Core/Public/Math/UnrealMathUtility.h
--- a/Engine/Source/Runtime/Core/Public/Math/UnrealMathUtility.h
+++ b/Engine/Source/Runtime/Core/Public/Math/UnrealMathUtility.h
@@ -1578,4 +1578,47 @@ struct FMath : public FPlatformMath
 	 * @return	Perlin noise in the range of -1.0 to 1.0
 	 */
 	static CORE_API float PerlinNoise1D(const float Value);
+
+	/**
+	 * Calculates the new value in a weighted moving average series using the previous value and the weight
+	 *
+	 * @param CurrentSample - The value to blend with the previous sample to get a new weighted value
+	 * @param PreviousSample - The last value from the series
+	 * @param Weight - The weight to blend with
+	 *
+	 * @return the next value in the series
+	 */
+	static CORE_API inline float WeightedMovingAverage(float CurrentSample, float PreviousSample, float Weight)
+	{
+		Weight = Clamp<float>(Weight, 0.f, 1.f);
+		float WAvg = (CurrentSample * Weight) + (PreviousSample * (1.f - Weight));
+		return WAvg;
+	}
+
+	/**
+	 * Calculates the new value in a weighted moving average series using the previous value and a weight range.
+	 * The weight range is used to dynamically adjust based upon distance between the samples
+	 * This allows you to smooth a value more aggressively for small noise and let large movements be smoothed less (or vice versa)
+	 *
+	 * @param CurrentSample - The value to blend with the previous sample to get a new weighted value
+	 * @param PreviousSample - The last value from the series
+	 * @param MaxDistance - Distance to use as the blend between min weight or max weight
+	 * @param MinWeight - The weight use when the distance is small
+	 * @param MaxWeight - The weight use when the distance is large
+	 *
+	 * @return the next value in the series
+	 */
+	static CORE_API inline float DynamicWeightedMovingAverage(float CurrentSample, float PreviousSample, float MaxDistance, float MinWeight, float MaxWeight)
+	{
+		// We need the distance between samples to determine how much of each weight to use
+		const float Distance = Abs<float>(CurrentSample - PreviousSample);
+		float Weight = MinWeight;
+		if (MaxDistance > 0)
+		{
+			// Figure out the lerp value to use between the min/max weights
+			const float LerpAlpha = Clamp<float>(Distance / MaxDistance, 0.f, 1.f);
+			Weight = Lerp<float>(MinWeight, MaxWeight, LerpAlpha);
+		}
+		return WeightedMovingAverage(CurrentSample, PreviousSample, Weight);
+	}
 };
diff --git a/Engine/Source/Runtime/Engine/Classes/Kismet/KismetMathLibrary.h b/Engine/Source/Runtime/Engine/Classes/Kismet/KismetMathLibrary.h
--- a/Engine/Source/Runtime/Engine/Classes/Kismet/KismetMathLibrary.h
+++ b/Engine/Source/Runtime/Engine/Classes/Kismet/KismetMathLibrary.h
@@ -3503,6 +3503,89 @@ class ENGINE_API UKismetMathLibrary : public UBlueprintFunctionLibrary
 	UFUNCTION(BlueprintPure, Category = "Math|Intersection", meta = (DisplayName = "Line Plane Intersection (Origin & Normal)"))
 	static bool LinePlaneIntersection_OriginNormal(const FVector& LineStart, const FVector& LineEnd, FVector PlaneOrigin, FVector PlaneNormal, float& T, FVector& Intersection);
 
+	/**
+	 * Calculates the new value in a weighted moving average series using the previous value and the weight
+	 *
+	 * @param CurrentSample - The value to blend with the previous sample to get a new weighted value
+	 * @param PreviousSample - The last value from the series
+	 * @param Weight - The weight to blend with
+	 *
+	 * @return the next value in the series
+	 */
+	UFUNCTION(BlueprintPure, Category="Math|Smoothing", meta=(DisplayName="Weighted Moving Average Float"))
+	static float WeightedMovingAverage_Float(float CurrentSample, float PreviousSample, float Weight);
+
+	/**
+	 * Calculates the new value in a weighted moving average series using the previous value and the weight
+	 *
+	 * @param CurrentSample - The value to blend with the previous sample to get a new weighted value
+	 * @param PreviousSample - The last value from the series
+	 * @param Weight - The weight to blend with
+	 *
+	 * @return the next value in the series
+	 */
+	UFUNCTION(BlueprintPure, Category="Math|Smoothing", meta=(DisplayName="Weighted Moving Average Vector"))
+	static FVector WeightedMovingAverage_FVector(FVector CurrentSample, FVector PreviousSample, float Weight);
+
+	/**
+	 * Calculates the new value in a weighted moving average series using the previous value and the weight
+	 *
+	 * @param CurrentSample - The value to blend with the previous sample to get a new weighted value
+	 * @param PreviousSample - The last value from the series
+	 * @param Weight - The weight to blend with
+	 *
+	 * @return the next value in the series
+	 */
+	UFUNCTION(BlueprintPure, Category="Math|Smoothing", meta=(DisplayName="Weighted Moving Average Rotator"))
+	static FRotator WeightedMovingAverage_FRotator(FRotator CurrentSample, FRotator PreviousSample, float Weight);
+
+	/**
+	 * Calculates the new value in a weighted moving average series using the previous value and a weight range.
+	 * The weight range is used to dynamically adjust based upon distance between the samples
+	 * This allows you to smooth a value more aggressively for small noise and let large movements be smoothed less (or vice versa)
+	 *
+	 * @param CurrentSample - The value to blend with the previous sample to get a new weighted value
+	 * @param PreviousSample - The last value from the series
+	 * @param MaxDistance - Distance to use as the blend between min weight or max weight
+	 * @param MinWeight - The weight use when the distance is small
+	 * @param MaxWeight - The weight use when the distance is large
+	 *
+	 * @return the next value in the series
+	 */
+	UFUNCTION(BlueprintPure, Category="Math|Smoothing", meta=(DisplayName="Dynamic Weighted Moving Average Float"))
+	static float DynamicWeightedMovingAverage_Float(float CurrentSample, float PreviousSample, float MaxDistance, float MinWeight, float MaxWeight);
+
+	/**
+	 * Calculates the new value in a weighted moving average series using the previous value and a weight range.
+	 * The weight range is used to dynamically adjust based upon distance between the samples
+	 * This allows you to smooth a value more aggressively for small noise and let large movements be smoothed less (or vice versa)
+	 *
+	 * @param CurrentSample - The value to blend with the previous sample to get a new weighted value
+	 * @param PreviousSample - The last value from the series
+	 * @param MaxDistance - Distance to use as the blend between min weight or max weight
+	 * @param MinWeight - The weight use when the distance is small
+	 * @param MaxWeight - The weight use when the distance is large
+	 *
+	 * @return the next value in the series
+	 */
+	UFUNCTION(BlueprintPure, Category="Math|Smoothing", meta=(DisplayName="Dynamic Weighted Moving Average Vector"))
+	static FVector DynamicWeightedMovingAverage_FVector(FVector CurrentSample, FVector PreviousSample, float MaxDistance, float MinWeight, float MaxWeight);
+
+	/**
+	 * Calculates the new value in a weighted moving average series using the previous value and a weight range.
+	 * The weight range is used to dynamically adjust based upon distance between the samples
+	 * This allows you to smooth a value more aggressively for small noise and let large movements be smoothed less (or vice versa)
+	 *
+	 * @param CurrentSample - The value to blend with the previous sample to get a new weighted value
+	 * @param PreviousSample - The last value from the series
+	 * @param MaxDistance - Distance to use as the blend between min weight or max weight
+	 * @param MinWeight - The weight use when the distance is small
+	 * @param MaxWeight - The weight use when the distance is large
+	 *
+	 * @return the next value in the series
+	 */
+	UFUNCTION(BlueprintPure, Category="Math|Smoothing", meta=(DisplayName="Dynamic Weighted Moving Average Rotator"))
+	static FRotator DynamicWeightedMovingAverage_FRotator(FRotator CurrentSample, FRotator PreviousSample, float MaxDistance, float MinWeight, float MaxWeight);
 
 private:
 
diff --git a/Engine/Source/Runtime/Engine/Private/KismetMathLibrary.cpp b/Engine/Source/Runtime/Engine/Private/KismetMathLibrary.cpp
--- a/Engine/Source/Runtime/Engine/Private/KismetMathLibrary.cpp
+++ b/Engine/Source/Runtime/Engine/Private/KismetMathLibrary.cpp
@@ -1131,4 +1131,51 @@ float UKismetMathLibrary::PerlinNoise1D(const float Value)
 	return FMath::PerlinNoise1D(Value);
 }
 
+float UKismetMathLibrary::WeightedMovingAverage_Float(float CurrentSample, float PreviousSample, float Weight)
+{
+	return FMath::WeightedMovingAverage(CurrentSample, PreviousSample, Weight);
+}
+
+FVector UKismetMathLibrary::WeightedMovingAverage_FVector(FVector CurrentSample, FVector PreviousSample, float Weight)
+{
+	FVector OutVector;
+	OutVector.X = FMath::WeightedMovingAverage(CurrentSample.X, PreviousSample.X, Weight);
+	OutVector.Y = FMath::WeightedMovingAverage(CurrentSample.Y, PreviousSample.Y, Weight);
+	OutVector.Z = FMath::WeightedMovingAverage(CurrentSample.Z, PreviousSample.Z, Weight);
+	return OutVector;
+}
+
+FRotator UKismetMathLibrary::WeightedMovingAverage_FRotator(FRotator CurrentSample, FRotator PreviousSample, float Weight)
+{
+	FRotator OutRotator;
+	OutRotator.Yaw = FMath::Clamp(FMath::WeightedMovingAverage(CurrentSample.Yaw, PreviousSample.Yaw, Weight), -180.f, 180.f);
+	OutRotator.Pitch = FMath::Clamp(FMath::WeightedMovingAverage(CurrentSample.Pitch, PreviousSample.Pitch, Weight), -180.f, 180.f);
+	OutRotator.Roll = FMath::Clamp(FMath::WeightedMovingAverage(CurrentSample.Roll, PreviousSample.Roll, Weight), -180.f, 180.f);
+	return OutRotator;
+}
+
+float UKismetMathLibrary::DynamicWeightedMovingAverage_Float(float CurrentSample, float PreviousSample, float MaxDistance, float MinWeight, float MaxWeight)
+{
+	return FMath::DynamicWeightedMovingAverage(CurrentSample, PreviousSample, MaxDistance, MinWeight, MaxWeight);
+}
+
+FVector UKismetMathLibrary::DynamicWeightedMovingAverage_FVector(FVector CurrentSample, FVector PreviousSample, float MaxDistance, float MinWeight, float MaxWeight)
+{
+	FVector OutVector;
+	OutVector.X = FMath::DynamicWeightedMovingAverage(CurrentSample.X, PreviousSample.X, MaxDistance, MinWeight, MaxWeight);
+	OutVector.Y = FMath::DynamicWeightedMovingAverage(CurrentSample.Y, PreviousSample.Y, MaxDistance, MinWeight, MaxWeight);
+	OutVector.Z = FMath::DynamicWeightedMovingAverage(CurrentSample.Z, PreviousSample.Z, MaxDistance, MinWeight, MaxWeight);
+	return OutVector;
+}
+
+FRotator UKismetMathLibrary::DynamicWeightedMovingAverage_FRotator(FRotator CurrentSample, FRotator PreviousSample, float MaxDistance, float MinWeight, float MaxWeight)
+{
+	FRotator OutRotator;
+	OutRotator.Yaw = FMath::Clamp(FMath::DynamicWeightedMovingAverage(CurrentSample.Yaw, PreviousSample.Yaw, MaxDistance, MinWeight, MaxWeight), -180.f, 180.f);
+	OutRotator.Pitch = FMath::Clamp(FMath::DynamicWeightedMovingAverage(CurrentSample.Pitch, PreviousSample.Pitch, MaxDistance, MinWeight, MaxWeight), -180.f, 180.f);
+	OutRotator.Roll = FMath::Clamp(FMath::DynamicWeightedMovingAverage(CurrentSample.Roll, PreviousSample.Roll, MaxDistance, MinWeight, MaxWeight), -180.f, 180.f);
+	return OutRotator;
+}
+
+
 #undef LOCTEXT_NAMESPACE
