Implement material node to correct UVs for passthrough camera

diff --git a/Engine/Plugins/Runtime/AR/Google/GoogleARCore/Source/GoogleARCoreBase/Private/GoogleARCoreXRCamera.cpp b/Engine/Plugins/Runtime/AR/Google/GoogleARCore/Source/GoogleARCoreBase/Private/GoogleARCoreXRCamera.cpp
--- a/Engine/Plugins/Runtime/AR/Google/GoogleARCore/Source/GoogleARCoreBase/Private/GoogleARCoreXRCamera.cpp
+++ b/Engine/Plugins/Runtime/AR/Google/GoogleARCore/Source/GoogleARCoreBase/Private/GoogleARCoreXRCamera.cpp
@@ -66,10 +66,11 @@ bool FGoogleARCoreXRCamera::GetPassthroughCameraUVs_RenderThread(TArray<FVector2
 		GoogleARCoreTrackingSystem.ARCoreDeviceInstance->GetPassthroughCameraImageUVs(PassthroughRenderer->OverlayQuadUVs, TransformedUVs);
 		PassthroughRenderer->UpdateOverlayUVCoordinate_RenderThread(TransformedUVs, FGoogleARCoreAndroidHelper::GetDisplayRotation());
 		OutUVs.SetNumUninitialized(4);
-		OutUVs[0] = FVector2D(TransformedUVs[0], TransformedUVs[1]);
-		OutUVs[1] = FVector2D(TransformedUVs[2], TransformedUVs[3]);
-		OutUVs[2] = FVector2D(TransformedUVs[4], TransformedUVs[5]);
-		OutUVs[3] = FVector2D(TransformedUVs[6], TransformedUVs[7]);
+
+		OutUVs[0] = FVector2D(TransformedUVs[4], TransformedUVs[5]);
+		OutUVs[1] = FVector2D(TransformedUVs[6], TransformedUVs[7]);
+		OutUVs[2] = FVector2D(TransformedUVs[0], TransformedUVs[1]);
+		OutUVs[3] = FVector2D(TransformedUVs[2], TransformedUVs[3]);
 		return true;
 	}
 	else
diff --git a/Engine/Source/Runtime/Engine/Classes/Materials/MaterialExpressionMapARPassthroughCameraUV.h b/Engine/Source/Runtime/Engine/Classes/Materials/MaterialExpressionMapARPassthroughCameraUV.h
new file mode 100644
--- /dev/null
+++ b/Engine/Source/Runtime/Engine/Classes/Materials/MaterialExpressionMapARPassthroughCameraUV.h
@@ -0,0 +1,32 @@
+// Copyright 1998-2019 Epic Games, Inc. All Rights Reserved.
+
+/**
+ *
+ *	A material expression that maps viewport UVs to AR passthrough camera UVs taking into account aspect ratio and device rotation.
+ */
+
+#pragma once
+
+#include "CoreMinimal.h"
+#include "UObject/ObjectMacros.h"
+#include "Materials/MaterialExpression.h"
+#include "MaterialExpressionMapARPassthroughCameraUV.generated.h"
+
+UCLASS(collapsecategories, hidecategories = Object)
+class UMaterialExpressionMapARPassthroughCameraUV : public UMaterialExpression
+{
+	GENERATED_UCLASS_BODY()
+
+	UPROPERTY(meta = (RequiredInput = "true", ToolTip = "Input UV to map to AR Camera UV"))
+	FExpressionInput Coordinates;
+
+	//~ Begin UMaterialExpression Interface
+#if WITH_EDITOR
+	virtual int32 Compile(class FMaterialCompiler* Compiler, int32 OutputIndex) override;
+	virtual void GetCaption(TArray<FString>& OutCaptions) const override;
+#endif
+	//~ End UMaterialExpression Interface
+};
+
+
+
diff --git a/Engine/Source/Runtime/Engine/Private/Materials/HLSLMaterialTranslator.h b/Engine/Source/Runtime/Engine/Private/Materials/HLSLMaterialTranslator.h
--- a/Engine/Source/Runtime/Engine/Private/Materials/HLSLMaterialTranslator.h
+++ b/Engine/Source/Runtime/Engine/Private/Materials/HLSLMaterialTranslator.h
@@ -5464,6 +5464,20 @@ protected:
 
 	}
 
+	virtual int32 MapARPassthroughCameraUV(int32 UV) override
+	{
+		if (UV == INDEX_NONE)
+		{
+			return INDEX_NONE;
+		}
+
+		int32 UVPair0 = AddInlinedCodeChunk(MCT_Float4, TEXT("ResolvedView.XRPassthroughCameraUVs[0]"));
+		int32 UVPair1 = AddInlinedCodeChunk(MCT_Float4, TEXT("ResolvedView.XRPassthroughCameraUVs[1]"));
+
+		int32 ULerp = Lerp(UVPair0, UVPair1, ComponentMask(UV, 1, 0, 0, 0));
+		return Lerp(ComponentMask(ULerp, 1, 1, 0, 0), ComponentMask(ULerp, 0, 0, 1, 1), ComponentMask(UV, 0, 1, 0, 0));
+	}
+
 	virtual int32 CustomExpression( class UMaterialExpressionCustom* Custom, TArray<int32>& CompiledInputs ) override
 	{
 		int32 ResultIdx = INDEX_NONE;
diff --git a/Engine/Source/Runtime/Engine/Private/Materials/MaterialExpressions.cpp b/Engine/Source/Runtime/Engine/Private/Materials/MaterialExpressions.cpp
--- a/Engine/Source/Runtime/Engine/Private/Materials/MaterialExpressions.cpp
+++ b/Engine/Source/Runtime/Engine/Private/Materials/MaterialExpressions.cpp
@@ -209,6 +209,7 @@
 #include "Materials/MaterialExpressionAtmosphericLightColor.h"
 #include "Materials/MaterialExpressionMaterialLayerOutput.h"
 #include "Materials/MaterialExpressionCurveAtlasRowParameter.h"
+#include "Materials/MaterialExpressionMapARPassthroughCameraUV.h"
 #include "Materials/MaterialUniformExpressions.h"
 #include "EditorSupportDelegates.h"
 #include "MaterialCompiler.h"
@@ -15203,4 +15204,48 @@ void UMaterialExpressionCurveAtlasRowParameter::PostEditChangeProperty(FProperty
 	Super::PostEditChangeProperty(PropertyChangedEvent);
 }
 #endif
+
+//
+//  UMaterialExpressionARPassthroughCameraUVs
+//
+
+UMaterialExpressionMapARPassthroughCameraUV::UMaterialExpressionMapARPassthroughCameraUV(const FObjectInitializer& ObjectInitializer)
+	: Super(ObjectInitializer)
+{
+#if WITH_EDITORONLY_DATA
+	// Structure to hold one-time initialization
+	struct FConstructorStatics
+	{
+		FText NAME_Coordinates;
+		FConstructorStatics()
+			: NAME_Coordinates(LOCTEXT("Coordinates", "Coordinates"))
+		{
+		}
+	};
+	static FConstructorStatics ConstructorStatics;
+
+	MenuCategories.Add(ConstructorStatics.NAME_Coordinates);
+#endif // WITH_EDITORONLY_DATA
+}
+
+#if WITH_EDITOR
+int32 UMaterialExpressionMapARPassthroughCameraUV::Compile(class FMaterialCompiler* Compiler, int32 OutputIndex)
+{
+	if (!Coordinates.GetTracedInput().Expression)
+	{
+		return CompilerError(Compiler, TEXT("UV input missing"));
+	}
+	else
+	{
+		int32 Index = Coordinates.Compile(Compiler);
+		return Compiler->MapARPassthroughCameraUV(Index);
+	}
+}
+
+void UMaterialExpressionMapARPassthroughCameraUV::GetCaption(TArray<FString>& OutCaptions) const
+{
+	OutCaptions.Add(TEXT("Map AR Passthrough Camera UVs"));
+}
+#endif // WITH_EDITOR
+
 #undef LOCTEXT_NAMESPACE
diff --git a/Engine/Source/Runtime/Engine/Public/MaterialCompiler.h b/Engine/Source/Runtime/Engine/Public/MaterialCompiler.h
--- a/Engine/Source/Runtime/Engine/Public/MaterialCompiler.h
+++ b/Engine/Source/Runtime/Engine/Public/MaterialCompiler.h
@@ -293,6 +293,8 @@ public:
 	virtual int32 EyeAdaptation() = 0;
 	virtual int32 AtmosphericLightVector() = 0;
 	virtual int32 AtmosphericLightColor() = 0;
+
+	virtual int32 MapARPassthroughCameraUV(int32 UV) = 0;
 	// The compiler can run in a different state and this affects caching of sub expression, Expressions are different (e.g. View.PrevWorldViewOrigin) when using previous frame's values
 	// If possible we should re-factor this to avoid having to deal with compiler state
 	virtual bool IsCurrentlyCompilingForPreviousFrame() const { return false; }
@@ -549,6 +551,11 @@ public:
 		return Compiler->AtmosphericLightColor();
 	}
 
+	virtual int32 MapARPassthroughCameraUV(int32 UV) override
+	{
+		return Compiler->MapARPassthroughCameraUV(UV);
+	}
+
 	virtual int32 TextureCoordinateOffset() override
 	{
 		return Compiler->TextureCoordinateOffset();
diff --git a/Engine/Source/Runtime/Engine/Public/SceneView.h b/Engine/Source/Runtime/Engine/Public/SceneView.h
--- a/Engine/Source/Runtime/Engine/Public/SceneView.h
+++ b/Engine/Source/Runtime/Engine/Public/SceneView.h
@@ -702,7 +702,9 @@ enum ETranslucencyVolumeCascade
 	VIEW_UNIFORM_BUFFER_MEMBER(FVector, VolumetricLightmapBrickTexelSize) \
 	VIEW_UNIFORM_BUFFER_MEMBER(float, StereoIPD) \
 	VIEW_UNIFORM_BUFFER_MEMBER(float, IndirectLightingCacheShowFlag) \
-	VIEW_UNIFORM_BUFFER_MEMBER(float, EyeToPixelSpreadAngle)
+	VIEW_UNIFORM_BUFFER_MEMBER(float, EyeToPixelSpreadAngle) \
+	VIEW_UNIFORM_BUFFER_MEMBER_ARRAY(FVector4, XRPassthroughCameraUVs, [2])
+
 
 #define VIEW_UNIFORM_BUFFER_MEMBER(type, identifier) \
 	SHADER_PARAMETER(type, identifier)
diff --git a/Engine/Source/Runtime/Renderer/Private/SceneRendering.cpp b/Engine/Source/Runtime/Renderer/Private/SceneRendering.cpp
--- a/Engine/Source/Runtime/Renderer/Private/SceneRendering.cpp
+++ b/Engine/Source/Runtime/Renderer/Private/SceneRendering.cpp
@@ -50,6 +50,9 @@
 #include "VisualizeTexturePresent.h"
 #include "MeshDrawCommands.h"
 #include "HAL/LowLevelMemTracker.h"
+#include "IXRTrackingSystem.h"
+#include "IXRCamera.h"
+#include "IXRCamera.h"
 
 /*-----------------------------------------------------------------------------
 	Globals
@@ -1429,7 +1432,22 @@ void FViewInfo::SetupUniformBufferParameters(
 
 	ViewUniformShaderParameters.StereoPassIndex = GEngine->StereoRenderingDevice ? GEngine->StereoRenderingDevice->GetViewIndexForPass(StereoPass) : 0;
 	ViewUniformShaderParameters.StereoIPD = StereoIPD;
-	
+
+	{
+		auto XRCamera = GEngine->XRSystem ? GEngine->XRSystem->GetXRCamera() : nullptr;
+		TArray<FVector2D> CameraUVs;
+		if (XRCamera.IsValid() && XRCamera->GetPassthroughCameraUVs_RenderThread(CameraUVs) && CameraUVs.Num() == 4)
+		{
+			ViewUniformShaderParameters.XRPassthroughCameraUVs[0] = FVector4(CameraUVs[0], CameraUVs[1]);
+			ViewUniformShaderParameters.XRPassthroughCameraUVs[1] = FVector4(CameraUVs[2], CameraUVs[3]);
+		}
+		else
+		{
+			ViewUniformShaderParameters.XRPassthroughCameraUVs[0] = FVector4(0, 0, 0, 1);
+			ViewUniformShaderParameters.XRPassthroughCameraUVs[1] = FVector4(1, 0, 1, 1);
+		}
+	}
+
 	ViewUniformShaderParameters.PreIntegratedBRDF = GEngine->PreIntegratedSkinBRDFTexture->Resource->TextureRHI;
 
 	if (UseGPUScene(GMaxRHIShaderPlatform, RHIFeatureLevel))
