PR #4723: Add SetLeftTexture for UStereoLayerComponent (Contributed by [REDACTED_PERSON])
Also added GetLeftTexture to preserve symmetry with Get/SetTexture().

diff --git a/Engine/Source/Runtime/Engine/Classes/Components/StereoLayerComponent.h b/Engine/Source/Runtime/Engine/Classes/Components/StereoLayerComponent.h
--- a/Engine/Source/Runtime/Engine/Classes/Components/StereoLayerComponent.h
+++ b/Engine/Source/Runtime/Engine/Classes/Components/StereoLayerComponent.h
@@ -59,16 +59,29 @@ public:
 	virtual void TickComponent(float DeltaTime, enum ELevelTick TickType, FActorComponentTickFunction *ThisTickFunction) override;
 	
 	/** 
-	 * Change the texture displayed on the stereo layer quad
+	 * Change the texture displayed on the stereo layer. 
+	 *
+	 * If stereoscopic layer textures are supported on the platform and LeftTexture is set, this property controls the texture for the right eye.
 	 * @param	InTexture: new Texture2D
 	 */
 	UFUNCTION(BlueprintCallable, Category="Components|Stereo Layer")
 	void SetTexture(UTexture* InTexture);
 
-	// @return the texture mapped to the stereo layer quad
+	/** 
+	 * Change the texture displayed on the stereo layer for left eye, if stereoscopic layer textures are supported on the platform.
+	 * @param	InTexture: new Texture2D
+	 */
+	UFUNCTION(BlueprintCallable, Category="Components|Stereo Layer")
+	void SetLeftTexture(UTexture* InTexture);
+
+	// @return the texture mapped to the stereo layer.
 	UFUNCTION(BlueprintCallable, Category="Components|Stereo Layer")
 	UTexture* GetTexture() const { return Texture; }
-	
+
+	// @return the texture mapped to the stereo layer for left eye, if stereoscopic layer textures are supported on the platform.
+	UFUNCTION(BlueprintCallable, Category = "Components|Stereo Layer")
+	UTexture* GetLeftTexture() const { return LeftTexture; }
+
 	/** 
 	 * Change the quad size. This is the unscaled height and width, before component scale is applied.
 	 * @param	InQuadSize: new quad size.
diff --git a/Engine/Source/Runtime/Engine/Private/Components/StereoLayerComponent.cpp b/Engine/Source/Runtime/Engine/Private/Components/StereoLayerComponent.cpp
--- a/Engine/Source/Runtime/Engine/Private/Components/StereoLayerComponent.cpp
+++ b/Engine/Source/Runtime/Engine/Private/Components/StereoLayerComponent.cpp
@@ -175,6 +175,17 @@ void UStereoLayerComponent::SetTexture(UTexture* InTexture)
 	bIsDirty = true;
 }
 
+void UStereoLayerComponent::SetLeftTexture(UTexture* InTexture)
+{
+	if (LeftTexture == InTexture)
+	{
+		return;
+	}
+
+	LeftTexture = InTexture;
+	bIsDirty = true;
+}
+
 void UStereoLayerComponent::SetQuadSize(FVector2D InQuadSize)
 {
 	if (QuadSize == InQuadSize)
