Incorrect cast.

diff --git a/Engine/Source/Runtime/Engine/Classes/GameFramework/PlayerInput.h b/Engine/Source/Runtime/Engine/Classes/GameFramework/PlayerInput.h
--- a/Engine/Source/Runtime/Engine/Classes/GameFramework/PlayerInput.h
+++ b/Engine/Source/Runtime/Engine/Classes/GameFramework/PlayerInput.h
@@ -306,13 +306,13 @@ public:
 	bool operator<(const FInputActionSpeechMapping& Other) const
 	{
 		bool bResult = false;
-		if (ActionName < Other.ActionName)
+		if (ActionName.LexicalLess(Other.ActionName))
 		{
 			bResult = true;
 		}
 		else if (ActionName == Other.ActionName)
 		{
-			bResult = (SpeechKeyword < Other.SpeechKeyword);
+			bResult = (SpeechKeyword.LexicalLess(Other.SpeechKeyword));
 		}
 		return bResult;
 	}
