Merging //UE4/Dev-Main to Dev-VR (//UE4/Dev-VR)

diff --git a/Engine/Build/BatchFiles/FindPlatformExtensionSources.bat b/Engine/Build/BatchFiles/FindPlatformExtensionSources.bat
--- a/Engine/Build/BatchFiles/FindPlatformExtensionSources.bat
+++ b/Engine/Build/BatchFiles/FindPlatformExtensionSources.bat
@@ -12,17 +12,25 @@ pushd "%~dp0..\..\Source"
 mkdir ..\Intermediate\ProjectFiles >NUL 2>NUL
 
 rem Look for any platform extension .cs files, and add them to a file that will be refernced by UBT project, so it can bring them in automatically
-set REFERENCE_FILE="..\Intermediate\ProjectFiles\UnrealBuildTool.csproj.References"
+set REFERENCE_FILE="%~dp0..\..\Intermediate\ProjectFiles\UnrealBuildTool.csproj.References"
 echo ^<?xml version="1.0" encoding="utf-8"?^> > %REFERENCE_FILE%
 echo ^<Project ToolsVersion="15.0" DefaultTargets="Build" xmlns="http://schemas.microsoft.com/developer/msbuild/2003"^> >> %REFERENCE_FILE%
 echo ^<ItemGroup^> >> %REFERENCE_FILE%
 
 if exist "..\..\Platforms" (
-	for /d %%i in ("..\..\Platforms\*") do (
-		for %%j in (%%i\Engine\Source\Programs\UnrealBuildTool\*) do (
-			echo ^<PlatformExtensionCompile Include="..\..\%%j" /^> >> %REFERENCE_FILE%
+	pushd "..\..\Platforms"
+	for /d %%i in ("*") do (
+		if exist "%%i\Engine\Source\Programs\UnrealBuildTool" (
+			pushd "%%i\Engine\Source\Programs\UnrealBuildTool"
+			for %%j in ("*.cs") do (
+				echo   ^<PlatformExtensionCompile Include="..\..\..\..\Platforms\%%i\Engine\Source\Programs\UnrealBuildTool\%%j"^> >> %REFERENCE_FILE%
+				echo     ^<Link^>Platform\%%i\%%j^</Link^> >> %REFERENCE_FILE%
+				echo   ^</PlatformExtensionCompile^> >> %REFERENCE_FILE%
+			)
+			popd
 		)
 	)
+	popd
 )
 echo ^</ItemGroup^> >> %REFERENCE_FILE%
 echo ^</Project^> >> %REFERENCE_FILE%
\ No newline at end of file
diff --git a/Engine/Source/Editor/DetailCustomizations/Private/DetailCustomizations.cpp b/Engine/Source/Editor/DetailCustomizations/Private/DetailCustomizations.cpp
--- a/Engine/Source/Editor/DetailCustomizations/Private/DetailCustomizations.cpp
+++ b/Engine/Source/Editor/DetailCustomizations/Private/DetailCustomizations.cpp
@@ -383,7 +383,6 @@ void FDetailCustomizationsModule::RegisterObjectCustomizations()
 	RegisterCustomClassLayout("MotionControllerComponent", FOnGetDetailCustomizationInstance::CreateStatic(&FMotionControllerDetails::MakeInstance));
 
 	RegisterCustomClassLayout("Landscape", FOnGetDetailCustomizationInstance::CreateStatic(&FLandscapeUIDetails::MakeInstance));
-	RegisterCustomClassLayout("LandscapeProxy", FOnGetDetailCustomizationInstance::CreateStatic(&FLandscapeUIDetails::MakeInstance));
 	RegisterCustomClassLayout("LandscapeSettings", FOnGetDetailCustomizationInstance::CreateStatic(&FLandscapeSettingsDetails::MakeInstance));
 }
 
diff --git a/Engine/Source/Editor/DetailCustomizations/Private/LandscapeUIDetails.cpp b/Engine/Source/Editor/DetailCustomizations/Private/LandscapeUIDetails.cpp
--- a/Engine/Source/Editor/DetailCustomizations/Private/LandscapeUIDetails.cpp
+++ b/Engine/Source/Editor/DetailCustomizations/Private/LandscapeUIDetails.cpp
@@ -59,29 +59,29 @@ void FLandscapeUIDetails::CustomizeDetails( IDetailLayoutBuilder& DetailBuilder
 			DetailBuilder.HideProperty(PropertyHandle);
 			if(GetMutableDefault<UEditorExperimentalSettings>()->bLandscapeLayerSystem)
 			{
-				IDetailCategoryBuilder& Category = DetailBuilder.EditCategory(PropertyHandle->GetDefaultCategoryName());
-				Category.AddCustomRow(FText(), true)
+				const FText DisplayAndFilterText(NSLOCTEXT("UnrealEd", "LandscapeToggleLayerName", "Enable Layer System"));
+				DetailBuilder.AddCustomRowToCategory(PropertyHandle, DisplayAndFilterText)
 				.NameContent()
 				[
-					PropertyHandle->CreatePropertyNameWidget(NSLOCTEXT("UnrealEd", "LandscapeToggleLayerName", "Enable Layer System"))
+					PropertyHandle->CreatePropertyNameWidget(DisplayAndFilterText)
 				]
 				.ValueContent()
-					[
-						SNew(SCheckBox)
-						.ToolTipText(NSLOCTEXT("UnrealEd", "LandscapeToggleLayerToolTip", "Toggle whether or not to support layers on this Landscape and its streaming proxies. Toggling this will clear the undo stack."))
-						.Type(ESlateCheckBoxType::CheckBox)
-						.IsChecked_Lambda([=]() 
-						{ 
-							return Landscape->CanHaveLayersContent() ? ECheckBoxState::Checked : ECheckBoxState::Unchecked;
-						})
-						.OnCheckStateChanged_Lambda([=](ECheckBoxState NewState) 
-						{ 
-							bool bChecked = (NewState == ECheckBoxState::Checked);
-							if (Landscape->CanHaveLayersContent() != bChecked)
-							{
-								ToggleCanHaveLayersContent(Landscape);
-							}
-						})
+				[
+					SNew(SCheckBox)
+					.ToolTipText(NSLOCTEXT("UnrealEd", "LandscapeToggleLayerToolTip", "Toggle whether or not to support layers on this Landscape and its streaming proxies. Toggling this will clear the undo stack."))
+					.Type(ESlateCheckBoxType::CheckBox)
+					.IsChecked_Lambda([=]() 
+					{ 
+						return Landscape->CanHaveLayersContent() ? ECheckBoxState::Checked : ECheckBoxState::Unchecked;
+					})
+					.OnCheckStateChanged_Lambda([=](ECheckBoxState NewState) 
+					{ 
+						bool bChecked = (NewState == ECheckBoxState::Checked);
+						if (Landscape->CanHaveLayersContent() != bChecked)
+						{
+							ToggleCanHaveLayersContent(Landscape);
+						}
+					})
 				];
 			}
 			
diff --git a/Engine/Source/Runtime/Engine/Classes/Animation/AnimBlueprintGeneratedClass.h b/Engine/Source/Runtime/Engine/Classes/Animation/AnimBlueprintGeneratedClass.h
--- a/Engine/Source/Runtime/Engine/Classes/Animation/AnimBlueprintGeneratedClass.h
+++ b/Engine/Source/Runtime/Engine/Classes/Animation/AnimBlueprintGeneratedClass.h
@@ -194,7 +194,7 @@ class ENGINE_API UAnimBlueprintGeneratedClass : public UBlueprintGeneratedClass,
 	UPROPERTY()
 	TMap<FName, FCachedPoseIndices> OrderedSavedPoseIndicesMap;
 
-	// The various anim functions that this class holds (created during Link)
+	// The various anim functions that this class holds (created during GenerateAnimationBlueprintFunctions)
 	TArray<FAnimBlueprintFunction> AnimBlueprintFunctions;
 
 	// The arrays of anim nodes; this is transient generated data (created during Link)
@@ -223,7 +223,7 @@ public:
 	virtual int32 GetSyncGroupIndex(FName SyncGroupName) const override { return SyncGroupNames.IndexOfByKey(SyncGroupName); }
 	virtual const TArray<FExposedValueHandler>& GetExposedValueHandlers() const { return EvaluateGraphExposedInputs; }
 	virtual const TArray<FAnimBlueprintFunction>& GetAnimBlueprintFunctions() const override { return AnimBlueprintFunctions; }
-
+	
 public:
 #if WITH_EDITORONLY_DATA
 	FAnimBlueprintDebugData AnimBlueprintDebugData;
@@ -332,9 +332,12 @@ public:
 	const int32* GetNodePropertyIndexFromGuid(FGuid Guid, EPropertySearchMode::Type SearchMode = EPropertySearchMode::OnlyThis);
 #endif
 
-	// Called after Link to patch up refereces to the nodes in the CDO
+	// Called after Link to patch up references to the nodes in the CDO
 	void LinkFunctionsToDefaultObjectNodes(UObject* DefaultObject);
 
+	// Populates AnimBlueprintFunctions according to the UFunction(s) on this class
+	void GenerateAnimationBlueprintFunctions();
+
 	// UObject interface
 	virtual void Serialize(FArchive& Ar) override;
 	// End of UObject interface
@@ -347,6 +350,7 @@ public:
 	virtual void PurgeClass(bool bRecompilingOnLoad) override;
 	virtual uint8* GetPersistentUberGraphFrame(UObject* Obj, UFunction* FuncToCheck) const override;
 	virtual void PostLoadDefaultObject(UObject* Object) override;
+	virtual void PostLoad() override;
 	// End of UClass interface
 };
 
diff --git a/Engine/Source/Runtime/Engine/Private/Animation/AnimBlueprintGeneratedClass.cpp b/Engine/Source/Runtime/Engine/Private/Animation/AnimBlueprintGeneratedClass.cpp
--- a/Engine/Source/Runtime/Engine/Private/Animation/AnimBlueprintGeneratedClass.cpp
+++ b/Engine/Source/Runtime/Engine/Private/Animation/AnimBlueprintGeneratedClass.cpp
@@ -251,92 +251,15 @@
 {
 	Super::Link(Ar, bRelinkExistingProperties);
 
-	static const FName DefaultAnimGraphName("AnimGraph");
-
 	// @TODO: Shouldn't be necessary to clear these, but currently the class gets linked twice during compilation
 	AnimNodeProperties.Empty();
 	SubInstanceNodeProperties.Empty();
 	LayerNodeProperties.Empty();
-	AnimBlueprintFunctions.Empty();
-
-	// Patch up blueprint function info
-	for (TFieldIterator<UFunction> It(this); It; ++It)
-	{
-		bool bFoundOutput = false;
-#if WITH_EDITOR
-		// In editor we can grab the group from metadata, otherwise we need to wait until CDO post load (LinkFunctionsToDefaultObjectNodes)
-		FText CategoryText = It->GetMetaDataText(TEXT("Category"), TEXT("UObjectCategory"), It->GetFullGroupName(false));
-		FName Group = CategoryText.IsEmpty() ? NAME_None : FName(*CategoryText.ToString());
-#endif
-		UStructProperty* OutputPoseNodeProperty = nullptr;
-		TArray<FName> InputPoseNames;
-		TArray<int32> InputPoseNodeIndices;
-		TArray<UStructProperty*> InputPoseNodeProperties;
-		TArray<UProperty*> InputProperties;
-
-		// grab the input/output poses, their indices will be patched up later once the CDO is loaded in PostLoadDefaultObject
-		for (TFieldIterator<UProperty> ItParam(*It); ItParam; ++ItParam)
-		{
-			if(UStructProperty* StructProperty = Cast<UStructProperty>(*ItParam))
-			{
-				if(StructProperty->Struct->IsChildOf(FPoseLink::StaticStruct()))
-				{
-					if(StructProperty->GetPropertyFlags() & CPF_OutParm)
-					{
-						if(!bFoundOutput)
-						{
-							OutputPoseNodeProperty = StructProperty;
-							bFoundOutput = true;
-						}
-						else
-						{
-							// our required signature needs us to have a single post link output, so null it out if we find more than one
-							OutputPoseNodeProperty = nullptr;
-						}
-					}
-					else
-					{
-						InputPoseNames.Add(StructProperty->GetFName());
-						InputPoseNodeIndices.Add(INDEX_NONE);
-						InputPoseNodeProperties.Add(nullptr);
-					}
-				}
-				else
-				{
-					InputProperties.Add(*ItParam);
-				}
-			}
-			else
-			{
-				InputProperties.Add(*ItParam);
-			}
-		}
-
-		if(OutputPoseNodeProperty)
-		{
-			// We use the undecorated name here, so trim the postfix
-			FAnimBlueprintFunction* AnimBlueprintFunction = nullptr;
-
-			// Make sure that the default graph is at index 0
-			FName FunctionName = It->GetFName();
-			if(FunctionName == DefaultAnimGraphName)
-			{
-				AnimBlueprintFunction = &AnimBlueprintFunctions.Insert_GetRef(FAnimBlueprintFunction(FunctionName), 0);
-			}
-			else
-			{
-				AnimBlueprintFunction = &AnimBlueprintFunctions.Emplace_GetRef(FunctionName);
-			}
 
 #if WITH_EDITOR
-			AnimBlueprintFunction->Group = Group;
-#endif
-			AnimBlueprintFunction->InputPoseNames.Append(MoveTemp(InputPoseNames));
-			AnimBlueprintFunction->InputPoseNodeIndices.Append(MoveTemp(InputPoseNodeIndices));
-			AnimBlueprintFunction->InputPoseNodeProperties.Append(MoveTemp(InputPoseNodeProperties));
-			AnimBlueprintFunction->InputProperties.Append(MoveTemp(InputProperties));
-		}
-	}
+	// This relies on the entire class being fully loaded, this is not the case with EDL async-loading, in which case the functions are generated in PostLoad
+	GenerateAnimationBlueprintFunctions();
+#endif // WITH_EDITOR
 
 	// Initialize the various tracked node arrays & fix up function internals
 	for (TFieldIterator<UProperty> It(this); It; ++It)
@@ -423,6 +346,98 @@
 	LinkFunctionsToDefaultObjectNodes(Object);
 }
 
+void UAnimBlueprintGeneratedClass::PostLoad()
+{
+	Super::PostLoad();
+	GenerateAnimationBlueprintFunctions();
+}
+
+void UAnimBlueprintGeneratedClass::GenerateAnimationBlueprintFunctions()
+{
+	AnimBlueprintFunctions.Empty();
+
+	static const FName DefaultAnimGraphName("AnimGraph");
+	// Patch up blueprint function info
+	for (TFieldIterator<UFunction> It(this); It; ++It)
+	{
+		bool bFoundOutput = false;
+#if WITH_EDITOR
+		// In editor we can grab the group from metadata, otherwise we need to wait until CDO post load (LinkFunctionsToDefaultObjectNodes)
+		FText CategoryText = It->GetMetaDataText(TEXT("Category"), TEXT("UObjectCategory"), It->GetFullGroupName(false));
+		FName Group = CategoryText.IsEmpty() ? NAME_None : FName(*CategoryText.ToString());
+#endif
+		UStructProperty* OutputPoseNodeProperty = nullptr;
+		TArray<FName> InputPoseNames;
+		TArray<int32> InputPoseNodeIndices;
+		TArray<UStructProperty*> InputPoseNodeProperties;
+		TArray<UProperty*> InputProperties;
+
+		// grab the input/output poses, their indices will be patched up later once the CDO is loaded in PostLoadDefaultObject
+		for (TFieldIterator<UProperty> ItParam(*It); ItParam; ++ItParam)
+		{
+			if (UStructProperty* StructProperty = Cast<UStructProperty>(*ItParam))
+			{
+				if (StructProperty->Struct->IsChildOf(FPoseLink::StaticStruct()))
+				{
+					if (StructProperty->GetPropertyFlags() & CPF_OutParm)
+					{
+						if (!bFoundOutput)
+						{
+							OutputPoseNodeProperty = StructProperty;
+							bFoundOutput = true;
+						}
+						else
+						{
+							// our required signature needs us to have a single post link output, so null it out if we find more than one
+							OutputPoseNodeProperty = nullptr;
+						}
+					}
+					else
+					{
+						InputPoseNames.Add(StructProperty->GetFName());
+						InputPoseNodeIndices.Add(INDEX_NONE);
+						InputPoseNodeProperties.Add(nullptr);
+					}
+				}
+				else
+				{
+					InputProperties.Add(*ItParam);
+				}
+			}
+			else
+			{
+				InputProperties.Add(*ItParam);
+			}
+		}
+
+		if (OutputPoseNodeProperty)
+		{
+			// We use the undecorated name here, so trim the postfix
+			FAnimBlueprintFunction* AnimBlueprintFunction = nullptr;
+
+			FName FunctionName = It->GetFName();
+
+			// Make sure that the default graph is at index 0
+			if (FunctionName == DefaultAnimGraphName)
+			{
+				AnimBlueprintFunction = &AnimBlueprintFunctions.Insert_GetRef(FAnimBlueprintFunction(FunctionName), 0);
+			}
+			else
+			{
+				AnimBlueprintFunction = &AnimBlueprintFunctions.Emplace_GetRef(FunctionName);
+			}
+
+#if WITH_EDITOR
+			AnimBlueprintFunction->Group = Group;
+#endif
+			AnimBlueprintFunction->InputPoseNames.Append(MoveTemp(InputPoseNames));
+			AnimBlueprintFunction->InputPoseNodeIndices.Append(MoveTemp(InputPoseNodeIndices));
+			AnimBlueprintFunction->InputPoseNodeProperties.Append(MoveTemp(InputPoseNodeProperties));
+			AnimBlueprintFunction->InputProperties.Append(MoveTemp(InputProperties));
+		}
+	}
+}
+
 void UAnimBlueprintGeneratedClass::LinkFunctionsToDefaultObjectNodes(UObject* DefaultObject)
 {
 	// Link functions to their nodes
diff --git a/Engine/Source/Runtime/Landscape/Classes/Landscape.h b/Engine/Source/Runtime/Landscape/Classes/Landscape.h
--- a/Engine/Source/Runtime/Landscape/Classes/Landscape.h
+++ b/Engine/Source/Runtime/Landscape/Classes/Landscape.h
@@ -355,7 +355,7 @@ private:
 public:
 
 #if WITH_EDITORONLY_DATA
-	UPROPERTY(EditAnywhere, Category=Experimental, AdvancedDisplay)
+	UPROPERTY(EditAnywhere, Category=Experimental)
 	bool bCanHaveLayersContent = false;
 
 	DECLARE_EVENT(ALandscape, FLandscapeBPCustomBrushChangedDelegate);
diff --git a/Engine/Source/Runtime/TraceLog/Private/Trace/Trace.cpp b/Engine/Source/Runtime/TraceLog/Private/Trace/Trace.cpp
--- a/Engine/Source/Runtime/TraceLog/Private/Trace/Trace.cpp
+++ b/Engine/Source/Runtime/TraceLog/Private/Trace/Trace.cpp
@@ -63,8 +63,7 @@ void Flush()
 
 #else
 
-#if PLATFORM_MAC
-void Dummy() {}
-#endif // PLATFORM_MAC
+// Workaround for module not having any exported symbols
+TRACELOG_API int TraceLogExportedSymbol = 0;
 
 #endif // UE_TRACE_ENABLED
