[NO NON-PII COMMIT MESSAGE]

diff --git a/Engine/Source/Runtime/Engine/Public/ADPCMAudioInfo.h b/Engine/Source/Runtime/Engine/Public/ADPCMAudioInfo.h
--- a/Engine/Source/Runtime/Engine/Public/ADPCMAudioInfo.h
+++ b/Engine/Source/Runtime/Engine/Public/ADPCMAudioInfo.h
@@ -118,7 +118,7 @@ public:
 	const uint8*	SrcBufferData;
 	uint32			SrcBufferDataSize;
 
-	uint32	UncompressedBlockSize;
+	uint32	NumBytesPerUncompressedChannel;
 	uint32	CompressedBlockSize;
 	uint32	BlockSize;
 	int32	StreamBufferSize;
@@ -126,13 +126,13 @@ public:
 	int32	NumChannels;
 	int32	Format;
 
-	uint8*			UncompressedBlockData;			// This holds the current block of compressed data for all channels
-	uint32			CurrentUncompressedBlockSampleIndex;	// This is the sample index within the current uncompressed block data
+	int16*			DeinterleavedUncompressedAudio;			// This holds the current block of compressed data for all channels
+	int32			CurrentUncompressedFrameIndex;	// This is the sample index within the current uncompressed block data
 	uint32			CurrentChunkIndex;				// This is the index that is currently being used, needed by streaming engine to make sure it stays loaded and the next chunk gets preloaded
 	uint32			CurrentChunkBufferOffset;		// This is this byte offset within the current chunk, used by streaming engine to prioritize a load if more then half way through current chunk
-	uint32			CurrentChunkDataSize;			// The size of the current chunk, the first chunk is bigger to accommodate the header info
-	uint32			TotalSamplesStreamed;			// The number of samples streamed so far (per channel)
-	uint32			TotalSamplesPerChannel;			// Number of samples per channel, used to detect when an audio waveform has ended
+	uint32			CurrentChunkDataSize;			// The size of the current chunk, the first chunk is bigger to accomodate the header info
+	uint32			TotalFramesStreamed;			// The number of samples streamed so far
+	uint32			TotalFrames;			// Number of samples per channel, used to detect when an audio waveform has ended
 	uint32			SamplesPerBlock;				// The number of samples per block
 	uint32			FirstChunkSampleDataOffset;		// The size of the header in the first chunk, used to skip over it when looping or starting the sample over
 	const uint8*	CurCompressedChunkData;			// A pointer to the current chunk of data
