Merging manually to clean up fallout from 7705805

//UE4/Main/Engine/Source/Runtime/Engine/Public/Physics/ImmediatePhysics/...

to //UE4/Dev-VR/Engine/Source/Runtime/Engine/Public/Physics/ImmediatePhysics/...

diff --git a/Engine/Source/Runtime/Engine/Public/Physics/ImmediatePhysics/ImmediatePhysicsActorHandle.h b/Engine/Source/Runtime/Engine/Public/Physics/ImmediatePhysics/ImmediatePhysicsActorHandle.h
new file mode 100644
--- /dev/null
+++ b/Engine/Source/Runtime/Engine/Public/Physics/ImmediatePhysics/ImmediatePhysicsActorHandle.h
@@ -0,0 +1,8 @@
+// Copyright 1998-2019 Epic Games, Inc. All Rights Reserved.
+
+#pragma once
+
+#include "Physics/ImmediatePhysics/ImmediatePhysicsPhysX/ImmediatePhysicsActorHandle_PhysX.h"
+#include "Physics/ImmediatePhysics/ImmediatePhysicsChaos/ImmediatePhysicsActorHandle_Chaos.h"
+
+#include "Physics/ImmediatePhysics/ImmediatePhysicsDeclares.h"
diff --git a/Engine/Source/Runtime/Engine/Public/Physics/ImmediatePhysics/ImmediatePhysicsSimulation.h b/Engine/Source/Runtime/Engine/Public/Physics/ImmediatePhysics/ImmediatePhysicsSimulation.h
new file mode 100644
--- /dev/null
+++ b/Engine/Source/Runtime/Engine/Public/Physics/ImmediatePhysics/ImmediatePhysicsSimulation.h
@@ -0,0 +1,8 @@
+// Copyright 1998-2019 Epic Games, Inc. All Rights Reserved.
+
+#pragma once
+
+#include "Physics/ImmediatePhysics/ImmediatePhysicsPhysX/ImmediatePhysicsSimulation_PhysX.h"
+#include "Physics/ImmediatePhysics/ImmediatePhysicsChaos/ImmediatePhysicsSimulation_Chaos.h"
+
+#include "Physics/ImmediatePhysics/ImmediatePhysicsDeclares.h"
