Add a flag to hide StereoLayers and use it from StereoLayerComponent instead of destroying and recreating the underlying layer object.

diff --git a/Engine/Plugins/Runtime/Oculus/OculusVR/Source/OculusHMD/Private/OculusHMD_Layer.cpp b/Engine/Plugins/Runtime/Oculus/OculusVR/Source/OculusHMD/Private/OculusHMD_Layer.cpp
--- a/Engine/Plugins/Runtime/Oculus/OculusVR/Source/OculusHMD/Private/OculusHMD_Layer.cpp
+++ b/Engine/Plugins/Runtime/Oculus/OculusVR/Source/OculusHMD/Private/OculusHMD_Layer.cpp
@@ -343,6 +343,10 @@ void FLayer::Initialize_RenderThread(const FSettings* Settings, FCustomPresent*
 		bInvertY = (CustomPresent->GetLayerFlags() & ovrpLayerFlag_TextureOriginAtBottomLeft) != 0;
 
 		uint32 SizeX = 0, SizeY = 0;
+		if (Desc.Flags & IStereoLayers::LAYER_FLAG_HIDDEN)
+		{
+			return;
+		}
 
 		if (Desc.Texture.IsValid())
 		{
diff --git a/Engine/Source/Runtime/Engine/Private/Components/StereoLayerComponent.cpp b/Engine/Source/Runtime/Engine/Private/Components/StereoLayerComponent.cpp
--- a/Engine/Source/Runtime/Engine/Private/Components/StereoLayerComponent.cpp
+++ b/Engine/Source/Runtime/Engine/Private/Components/StereoLayerComponent.cpp
@@ -52,7 +52,7 @@ void UStereoLayerComponent::TickComponent(float DeltaTime, enum ELevelTick TickT
 	Super::TickComponent( DeltaTime, TickType, ThisTickFunction );
 
 	IStereoLayers* StereoLayers;
-	if (!GEngine->StereoRenderingDevice.IsValid() || (StereoLayers = GEngine->StereoRenderingDevice->GetStereoLayers()) == nullptr || !Texture )
+	if (!GEngine->StereoRenderingDevice.IsValid() || (StereoLayers = GEngine->StereoRenderingDevice->GetStereoLayers()) == nullptr)
 	{
 		return;
 	}
@@ -81,85 +81,79 @@ void UStereoLayerComponent::TickComponent(float DeltaTime, enum ELevelTick TickT
 
 	if (bIsDirty)
 	{
-		if (!bCurrVisible)
+		IStereoLayers::FLayerDesc LayerDsec;
+		LayerDsec.Priority = Priority;
+		LayerDsec.QuadSize = QuadSize;
+		LayerDsec.UVRect = UVRect;
+		LayerDsec.EquirectProps = EquirectProps;
+		LayerDsec.Transform = Transform;
+
+		if (Texture)
 		{
-			if (LayerId)
-			{
-				StereoLayers->DestroyLayer(LayerId);
-				LayerId = 0;
-			}
+			LayerDsec.Texture = Texture->Resource->TextureRHI;
+			LayerDsec.Flags |= (Texture->GetMaterialType() == MCT_TextureExternal) ? IStereoLayers::LAYER_FLAG_TEX_EXTERNAL : 0;
 		}
-		else
+		if (LeftTexture)
 		{
-			IStereoLayers::FLayerDesc LayerDsec;
-			LayerDsec.Priority = Priority;
-			LayerDsec.QuadSize = QuadSize;
-			LayerDsec.UVRect = UVRect;
-			LayerDsec.EquirectProps = EquirectProps;
+			LayerDsec.LeftTexture = LeftTexture->Resource->TextureRHI;
+		}
+		LayerDsec.CylinderRadius = CylinderRadius; 
+		LayerDsec.CylinderOverlayArc = CylinderOverlayArc;
+		LayerDsec.CylinderHeight = CylinderHeight;
 			
-			LayerDsec.Transform = Transform;
-			if (Texture)
-			{
-				LayerDsec.Texture = Texture->Resource->TextureRHI;
-			}
-			if (LeftTexture)
-			{
-				LayerDsec.LeftTexture = LeftTexture->Resource->TextureRHI;
-			}
-			LayerDsec.CylinderRadius = CylinderRadius; 
-			LayerDsec.CylinderOverlayArc = CylinderOverlayArc;
-			LayerDsec.CylinderHeight = CylinderHeight;
 			
-			LayerDsec.Flags |= (bLiveTexture) ? IStereoLayers::LAYER_FLAG_TEX_CONTINUOUS_UPDATE : 0;
-			LayerDsec.Flags |= (bNoAlphaChannel) ? IStereoLayers::LAYER_FLAG_TEX_NO_ALPHA_CHANNEL : 0;
-			LayerDsec.Flags |= (bQuadPreserveTextureRatio) ? IStereoLayers::LAYER_FLAG_QUAD_PRESERVE_TEX_RATIO : 0;
-			LayerDsec.Flags |= (bSupportsDepth) ? IStereoLayers::LAYER_FLAG_SUPPORT_DEPTH : 0;
-			LayerDsec.Flags |= (Texture->GetMaterialType() == MCT_TextureExternal) ? IStereoLayers::LAYER_FLAG_TEX_EXTERNAL : 0;
+		LayerDsec.Flags |= (bLiveTexture) ? IStereoLayers::LAYER_FLAG_TEX_CONTINUOUS_UPDATE : 0;
+		LayerDsec.Flags |= (bNoAlphaChannel) ? IStereoLayers::LAYER_FLAG_TEX_NO_ALPHA_CHANNEL : 0;
+		LayerDsec.Flags |= (bQuadPreserveTextureRatio) ? IStereoLayers::LAYER_FLAG_QUAD_PRESERVE_TEX_RATIO : 0;
+		LayerDsec.Flags |= (bSupportsDepth) ? IStereoLayers::LAYER_FLAG_SUPPORT_DEPTH : 0;
+		LayerDsec.Flags |= (!bCurrVisible) ? IStereoLayers::LAYER_FLAG_HIDDEN : 0;
+
+
+		switch (StereoLayerType)
+		{
+		case SLT_WorldLocked:
+			LayerDsec.PositionType = IStereoLayers::WorldLocked;
+			break;
+		case SLT_TrackerLocked:
+			LayerDsec.PositionType = IStereoLayers::TrackerLocked;
+			break;
+		case SLT_FaceLocked:
+			LayerDsec.PositionType = IStereoLayers::FaceLocked;
+			break;
+		}
+
+		switch (StereoLayerShape)
+		{
+		case SLSH_QuadLayer:
+			LayerDsec.ShapeType = IStereoLayers::QuadLayer;
+			break;
 
-			switch (StereoLayerType)
-			{
-			case SLT_WorldLocked:
-				LayerDsec.PositionType = IStereoLayers::WorldLocked;
-				break;
-			case SLT_TrackerLocked:
-				LayerDsec.PositionType = IStereoLayers::TrackerLocked;
-				break;
-			case SLT_FaceLocked:
-				LayerDsec.PositionType = IStereoLayers::FaceLocked;
-				break;
-			}
-
-			switch (StereoLayerShape)
-			{
-			case SLSH_QuadLayer:
-				LayerDsec.ShapeType = IStereoLayers::QuadLayer;
-				break;
-
-			case SLSH_CylinderLayer:
-				LayerDsec.ShapeType = IStereoLayers::CylinderLayer;
-				break;
-
-			case SLSH_CubemapLayer:
-				LayerDsec.ShapeType = IStereoLayers::CubemapLayer;
-				break;
-
-			case SLSH_EquirectLayer:
-				LayerDsec.ShapeType = IStereoLayers::EquirectLayer;
-				break;
-
-			default:
-				break;
-			}
-
-			if (LayerId)
-			{
-				StereoLayers->SetLayerDesc(LayerId, LayerDsec);
-			}
-			else
-			{
-				LayerId = StereoLayers->CreateLayer(LayerDsec);
-			}
+		case SLSH_CylinderLayer:
+			LayerDsec.ShapeType = IStereoLayers::CylinderLayer;
+			break;
+
+		case SLSH_CubemapLayer:
+			LayerDsec.ShapeType = IStereoLayers::CubemapLayer;
+			break;
+
+		case SLSH_EquirectLayer:
+			LayerDsec.ShapeType = IStereoLayers::EquirectLayer;
+			break;
+
+		default:
+			break;
+		}
+
+
+		if (LayerId)
+		{
+			StereoLayers->SetLayerDesc(LayerId, LayerDsec);
+		}
+		else
+		{
+			LayerId = StereoLayers->CreateLayer(LayerDsec);
 		}
+		
 		LastTransform = Transform;
 		bLastVisible = bCurrVisible;
 		bIsDirty = false;
diff --git a/Engine/Source/Runtime/Engine/Public/IStereoLayers.h b/Engine/Source/Runtime/Engine/Public/IStereoLayers.h
--- a/Engine/Source/Runtime/Engine/Public/IStereoLayers.h
+++ b/Engine/Source/Runtime/Engine/Public/IStereoLayers.h
@@ -104,6 +104,8 @@ public:
 		LAYER_FLAG_SUPPORT_DEPTH = 0x00000008,
 		// Required on some platforms to enable rendering of external textures.
 		LAYER_FLAG_TEX_EXTERNAL = 0x00000010,
+		// When set, this layer will not be rendered.
+		LAYER_FLAG_HIDDEN = 0x00000020,
 	};
 
 	/**
@@ -113,6 +115,7 @@ public:
 	{
 		void SetLayerId(uint32 InId) { Id = InId; }
 		uint32 GetLayerId() const { return Id; }
+		bool IsVisible() const { return Texture != nullptr && !(Flags & LAYER_FLAG_HIDDEN); }
 
 		// Layer IDs must be larger than 0
 		const static uint32	INVALID_LAYER_ID = 0; 
diff --git a/Engine/Source/Runtime/HeadMountedDisplay/Private/DefaultStereoLayers.cpp b/Engine/Source/Runtime/HeadMountedDisplay/Private/DefaultStereoLayers.cpp
--- a/Engine/Source/Runtime/HeadMountedDisplay/Private/DefaultStereoLayers.cpp
+++ b/Engine/Source/Runtime/HeadMountedDisplay/Private/DefaultStereoLayers.cpp
@@ -95,7 +95,7 @@ void FDefaultStereoLayers::StereoLayerRender(FRHICommandListImmediate& RHICmdLis
 	for (uint32 LayerIndex : LayersToRender)
 	{
 		const FLayerDesc& Layer = RenderThreadLayers[LayerIndex];
-		check(Layer.Texture.IsValid());
+		check(Layer.IsVisible());
 		const bool bIsOpaque = (Layer.Flags & LAYER_FLAG_TEX_NO_ALPHA_CHANNEL) != 0;
 		const bool bIsExternal = (Layer.Flags & LAYER_FLAG_TEX_EXTERNAL) != 0;
 		bool bPipelineStateNeedsUpdate = false;
@@ -184,7 +184,7 @@ void FDefaultStereoLayers::PreRenderViewFamily_RenderThread(FRHICommandListImmed
 	for (uint32 LayerIndex = 0; LayerIndex < LayerCount; ++LayerIndex)
 	{
 		const auto& Layer = RenderThreadLayers[LayerIndex];
-		if (!Layer.Texture.IsValid())
+		if (!Layer.IsVisible())
 		{
 			continue;
 		}
diff --git a/Engine/Source/Runtime/HeadMountedDisplay/Public/StereoLayerManager.h b/Engine/Source/Runtime/HeadMountedDisplay/Public/StereoLayerManager.h
--- a/Engine/Source/Runtime/HeadMountedDisplay/Public/StereoLayerManager.h
+++ b/Engine/Source/Runtime/HeadMountedDisplay/Public/StereoLayerManager.h
@@ -147,7 +147,7 @@ public:
 		check(LayerId != FLayerDesc::INVALID_LAYER_ID);
 		LayerType& NewLayer = StereoLayers().Emplace(LayerId, InLayerDesc);
 		NewLayer.SetLayerId(LayerId);
-		UpdateLayer(NewLayer, LayerId, InLayerDesc.Texture != nullptr);
+		UpdateLayer(NewLayer, LayerId, InLayerDesc.IsVisible());
 		bStereoLayersDirty = true;
 		return LayerId;
 	}
@@ -195,7 +195,7 @@ public:
 			// If the layer is currently active, update layer state
 			if (FoundLevel == 0)
 			{
-				UpdateLayer(*Found, LayerId, InLayerDesc.Texture != nullptr);
+				UpdateLayer(*Found, LayerId, InLayerDesc.IsVisible());
 				bStereoLayersDirty = true;
 			}
 		}
@@ -290,7 +290,7 @@ public:
 		{
 			FLayerDesc LayerDesc;
 			GetLayerDescMember(Pair.Value, LayerDesc);
-			UpdateLayer(Pair.Value, Pair.Key, LayerDesc.Texture != nullptr);
+			UpdateLayer(Pair.Value, Pair.Key, LayerDesc.IsVisible());
 		}
 
 		bStereoLayersDirty = true;
