Adding generic SCOPED_GPU_EVENT macros for cases where we don't know if the RHICommandList is compute or gfx.

diff --git a/Engine/Source/Runtime/Engine/Private/SceneUtils.cpp b/Engine/Source/Runtime/Engine/Private/SceneUtils.cpp
--- a/Engine/Source/Runtime/Engine/Private/SceneUtils.cpp
+++ b/Engine/Source/Runtime/Engine/Private/SceneUtils.cpp
@@ -78,6 +78,7 @@ void TDrawEvent<TRHICmdList>::Stop()
 }
 template struct TDrawEvent<FRHICommandList>;
 template struct TDrawEvent<FRHIAsyncComputeCommandList>;
+template struct TDrawEvent<FRHIAsyncComputeCommandListImmediate>;
 
 void FDrawEventRHIExecute::Start(IRHIComputeContext& InRHICommandContext, FColor Color, const TCHAR* Fmt, ...)
 {
diff --git a/Engine/Source/Runtime/Engine/Public/SceneUtils.h b/Engine/Source/Runtime/Engine/Public/SceneUtils.h
--- a/Engine/Source/Runtime/Engine/Public/SceneUtils.h
+++ b/Engine/Source/Runtime/Engine/Public/SceneUtils.h
@@ -94,6 +94,18 @@ class FScopedGPUStatEvent;
 		void Stop();
 	};
 
+	#define SCOPED_GPU_EVENT(RHICmdList, Name) TDrawEvent<typename TDecay<decltype(RHICmdList)>::Type> PREPROCESSOR_JOIN(Event_##Name,__LINE__); if(GetEmitDrawEvents()) PREPROCESSOR_JOIN(Event_##Name,__LINE__).Start(RHICmdList, FColor(0), TEXT(#Name));
+	#define SCOPED_GPU_EVENT_COLOR(RHICmdList, Color, Name) TDrawEvent<typename TDecay<decltype(RHICmdList)>::Type> PREPROCESSOR_JOIN(Event_##Name,__LINE__); if(GetEmitDrawEvents()) PREPROCESSOR_JOIN(Event_##Name,__LINE__).Start(RHICmdList, Color, TEXT(#Name));
+	#define SCOPED_GPU_EVENTF(RHICmdList, Name, Format, ...) TDrawEvent<typename TDecay<decltype(RHICmdList)>::Type> PREPROCESSOR_JOIN(Event_##Name,__LINE__); if(GetEmitDrawEvents()) PREPROCESSOR_JOIN(Event_##Name,__LINE__).Start(RHICmdList, FColor(0), Format, ##__VA_ARGS__);
+	#define SCOPED_GPU_EVENTF_COLOR(RHICmdList, Color, Name, Format, ...) TDrawEvent<typename TDecay<decltype(RHICmdList)>::Type> PREPROCESSOR_JOIN(Event_##Name,__LINE__); if(GetEmitDrawEvents()) PREPROCESSOR_JOIN(Event_##Name,__LINE__).Start(RHICmdList, Color, Format, ##__VA_ARGS__);
+	#define SCOPED_CONDITIONAL_GPU_EVENT(RHICmdList, Name, Condition) TDrawEvent<typename TDecay<decltype(RHICmdList)>::Type> PREPROCESSOR_JOIN(Event_##Name,__LINE__); if(GetEmitDrawEvents() && (Condition)) PREPROCESSOR_JOIN(Event_##Name,__LINE__).Start(RHICmdList, FColor(0), TEXT(#Name));
+	#define SCOPED_CONDITIONAL_GPU_EVENT_COLOR(RHICmdList, Name, Color, Condition) TDrawEvent<typename TDecay<decltype(RHICmdList)>::Type> PREPROCESSOR_JOIN(Event_##Name,__LINE__); if(GetEmitDrawEvents() && (Condition)) PREPROCESSOR_JOIN(Event_##Name,__LINE__).Start(RHICmdList, Color, TEXT(#Name));
+	#define SCOPED_CONDITIONAL_GPU_EVENTF(RHICmdList, Name, Condition, Format, ...) TDrawEvent<typename TDecay<decltype(RHICmdList)>::Type> PREPROCESSOR_JOIN(Event_##Name,__LINE__); if(GetEmitDrawEvents() && (Condition)) PREPROCESSOR_JOIN(Event_##Name,__LINE__).Start(RHICmdList, FColor(0), Format, ##__VA_ARGS__);
+	#define SCOPED_CONDITIONAL_GPU_EVENTF_COLOR(RHICmdList, Color, Name, Condition, Format, ...) TDrawEvent<typename TDecay<decltype(RHICmdList)>::Type> PREPROCESSOR_JOIN(Event_##Name,__LINE__); if(GetEmitDrawEvents() && (Condition)) PREPROCESSOR_JOIN(Event_##Name,__LINE__).Start(RHICmdList, Color, Format, ##__VA_ARGS__);
+	#define BEGIN_GPU_EVENTF(RHICmdList, Name, Event, Format, ...) if(GetEmitDrawEvents()) Event.Start(RHICmdList, FColor(0), Format, ##__VA_ARGS__);
+	#define BEGIN_GPU_EVENTF_COLOR(RHICmdList, Color, Name, Event, Format, ...) if(GetEmitDrawEvents()) Event.Start(RHICmdList, Color, Format, ##__VA_ARGS__);
+	#define STOP_GPU_EVENT(Event) (Event).Stop();
+
 	// Macros to allow for scoping of draw events outside of RHI function implementations
 	#define SCOPED_DRAW_EVENT(RHICmdList, Name) TDrawEvent<FRHICommandList> PREPROCESSOR_JOIN(Event_##Name,__LINE__); if(GetEmitDrawEvents()) PREPROCESSOR_JOIN(Event_##Name,__LINE__).Start(RHICmdList, FColor(0), TEXT(#Name));
 	#define SCOPED_DRAW_EVENT_COLOR(RHICmdList, Color, Name) TDrawEvent<FRHICommandList> PREPROCESSOR_JOIN(Event_##Name,__LINE__); if(GetEmitDrawEvents()) PREPROCESSOR_JOIN(Event_##Name,__LINE__).Start(RHICmdList, Color, TEXT(#Name));
@@ -133,6 +145,18 @@ class FScopedGPUStatEvent;
 	{
 	};
 
+	#define SCOPED_GPU_EVENT(...)
+	#define SCOPED_GPU_EVENT_COLOR(...)
+	#define SCOPED_GPU_EVENTF(...)
+	#define SCOPED_GPU_EVENTF_COLOR(...)
+	#define SCOPED_CONDITIONAL_GPU_EVENT(...)
+	#define SCOPED_CONDITIONAL_GPU_EVENT_COLOR(...)
+	#define SCOPED_CONDITIONAL_GPU_EVENTF(...)
+	#define SCOPED_CONDITIONAL_GPU_EVENTF_COLOR(...)
+	#define BEGIN_GPU_EVENTF(...)
+	#define BEGIN_GPU_EVENTF_COLOR(...)
+	#define STOP_GPU_EVENT(...)
+
 	#define SCOPED_DRAW_EVENT(...)
 	#define SCOPED_DRAW_EVENT_COLOR(...)
 	#define SCOPED_DRAW_EVENTF(...)
