DR - Fix bad merge causing corruption
#rnx

diff --git a/Engine/Source/Runtime/Engine/Private/StaticMeshRender.cpp b/Engine/Source/Runtime/Engine/Private/StaticMeshRender.cpp
--- a/Engine/Source/Runtime/Engine/Private/StaticMeshRender.cpp
+++ b/Engine/Source/Runtime/Engine/Private/StaticMeshRender.cpp
@@ -783,6 +783,7 @@ uint32 FStaticMeshSceneProxy::SetMeshElementGeometrySource(
 		}
 	}
 
+	if (bRequiresAdjacencyInformation)
 	{
 		check( LODModel.bHasAdjacencyInfo );
 		OutMeshBatchElement.IndexBuffer = &LODModel.AdditionalIndexBuffers->AdjacencyIndexBuffer;
