Guard PostEditPropertyVChanged() against null deference. The property pointer may be null.

diff --git a/Engine/Source/Runtime/Engine/Private/Components/SkyLightComponent.cpp b/Engine/Source/Runtime/Engine/Private/Components/SkyLightComponent.cpp
--- a/Engine/Source/Runtime/Engine/Private/Components/SkyLightComponent.cpp
+++ b/Engine/Source/Runtime/Engine/Private/Components/SkyLightComponent.cpp
@@ -474,7 +474,8 @@ void USkyLightComponent::PostEditChangeProperty(FPropertyChangedEvent& PropertyC
 {
 	Super::PostEditChangeProperty(PropertyChangedEvent);
 
-	if (PropertyChangedEvent.Property->GetFName() == GET_MEMBER_NAME_CHECKED(USkyLightComponent, CubemapResolution) &&
+	if (PropertyChangedEvent.Property && 
+		PropertyChangedEvent.Property->GetFName() == GET_MEMBER_NAME_CHECKED(USkyLightComponent, CubemapResolution) &&
 		PropertyChangedEvent.ChangeType != EPropertyChangeType::Interactive)
 	{
 		SanitizeCubemapSize();
