UE-25637 - Introduce 8bit and 16 bit data texture groups.

diff --git a/Engine/Config/BaseDeviceProfiles.ini b/Engine/Config/BaseDeviceProfiles.ini
--- a/Engine/Config/BaseDeviceProfiles.ini
+++ b/Engine/Config/BaseDeviceProfiles.ini
@@ -116,6 +116,8 @@ TextureLODGroups=(Group=TEXTUREGROUP_World,MinLODSize=1,MaxLODSize=8192,LODBias=
 +TextureLODGroups=(Group=TEXTUREGROUP_Terrain_Weightmap,MinLODSize=1,MaxLODSize=8192,LODBias=0,MinMagFilter=aniso,MipFilter=point,MipGenSettings=TMGS_SimpleAverage)
 +TextureLODGroups=(Group=TEXTUREGROUP_Bokeh,MinLODSize=1,MaxLODSize=256,LODBias=0,MinMagFilter=linear,MipFilter=linear,MipGenSettings=TMGS_SimpleAverage)
 +TextureLODGroups=(Group=TEXTUREGROUP_Pixels2D,MinLODSize=1,MaxLODSize=8192,LODBias=0,MinMagFilter=point,MipFilter=point,MipGenSettings=TMGS_SimpleAverage)
++TextureLODGroups=(Group=TEXTUREGROUP_8BitData,LODBias=0,MinMagFilter=point,MipFilter=point,MipGenSettings=TMGS_NoMipmaps)
++TextureLODGroups=(Group=TEXTUREGROUP_16BitData,LODBias=0,MinMagFilter=point,MipFilter=point,MipGenSettings=TMGS_NoMipmaps)
 
 [Windows DeviceProfile]
 DeviceType=Windows
diff --git a/Engine/Source/Runtime/Engine/Classes/Engine/EngineTypes.h b/Engine/Source/Runtime/Engine/Classes/Engine/EngineTypes.h
--- a/Engine/Source/Runtime/Engine/Classes/Engine/EngineTypes.h
+++ b/Engine/Source/Runtime/Engine/Classes/Engine/EngineTypes.h
@@ -498,6 +498,7 @@ enum EMaterialSamplerType
 	SAMPLERTYPE_DistanceFieldFont UMETA(DisplayName="Distance Field Font"),
 	SAMPLERTYPE_LinearColor UMETA(DisplayName = "Linear Color"),
 	SAMPLERTYPE_LinearGrayscale UMETA(DisplayName = "Linear Grayscale"),
+	SAMPLERTYPE_Data UMETA(DisplayName = "Data"),
 	SAMPLERTYPE_External UMETA(DisplayName = "External"),
 	SAMPLERTYPE_MAX,
 };
diff --git a/Engine/Source/Runtime/Engine/Classes/Engine/TextureDefines.h b/Engine/Source/Runtime/Engine/Classes/Engine/TextureDefines.h
--- a/Engine/Source/Runtime/Engine/Classes/Engine/TextureDefines.h
+++ b/Engine/Source/Runtime/Engine/Classes/Engine/TextureDefines.h
@@ -58,7 +58,11 @@ enum TextureGroup
 	/** Impostor Color Textures*/
 	TEXTUREGROUP_Impostor UMETA(DisplayName="ini:Impostor Color"),
 	/** Impostor Normal and Depth, use default compression*/
-	TEXTUREGROUP_ImpostorNormalDepth UMETA(DisplayName="ini:Impostor Normal and Depth"),
+	TEXTUREGROUP_ImpostorNormalDepth UMETA(DisplayName = "ini:Impostor Normal and Depth"),
+	/** 8 bit data stored in textures */
+	TEXTUREGROUP_8BitData UMETA(DisplayName = "ini:8 Bit Data"),
+	/** 16 bit data stored in textures */
+	TEXTUREGROUP_16BitData UMETA(DisplayName = "ini:16 Bit Data"),
 	/** Project specific group, rename in Engine.ini, [EnumRemap] TEXTUREGROUP_Project**.DisplayName=My Fun Group */
 	TEXTUREGROUP_Project01 UMETA(DisplayName="ini:Project Group 01"),
 	TEXTUREGROUP_Project02 UMETA(DisplayName="ini:Project Group 02"),
diff --git a/Engine/Source/Runtime/Engine/Private/Materials/HLSLMaterialTranslator.h b/Engine/Source/Runtime/Engine/Private/Materials/HLSLMaterialTranslator.h
--- a/Engine/Source/Runtime/Engine/Private/Materials/HLSLMaterialTranslator.h
+++ b/Engine/Source/Runtime/Engine/Private/Materials/HLSLMaterialTranslator.h
@@ -3579,8 +3579,12 @@ protected:
 				// Normal maps need to be unpacked in the pixel shader.
 				SampleCode = FString::Printf( TEXT("UnpackNormalMap(%s)"), *SampleCode );
 				break;
+
 			case SAMPLERTYPE_Masks:
 				break;
+
+			case SAMPLERTYPE_Data:
+				break;
 		}
 
 		FString TextureName;
diff --git a/Engine/Source/Runtime/Engine/Private/Materials/MaterialExpressions.cpp b/Engine/Source/Runtime/Engine/Private/Materials/MaterialExpressions.cpp
--- a/Engine/Source/Runtime/Engine/Private/Materials/MaterialExpressions.cpp
+++ b/Engine/Source/Runtime/Engine/Private/Materials/MaterialExpressions.cpp
@@ -1371,6 +1371,10 @@ EMaterialSamplerType UMaterialExpressionTextureBase::GetSamplerTypeForTexture(co
 		{
 			return SAMPLERTYPE_External;
 		}
+		else if (Texture->LODGroup == TEXTUREGROUP_8BitData || Texture->LODGroup == TEXTUREGROUP_16BitData)
+		{
+			return SAMPLERTYPE_Data;
+		}
 
 		switch (Texture->CompressionSettings)
 		{
diff --git a/Engine/Source/Runtime/Engine/Private/Texture.cpp b/Engine/Source/Runtime/Engine/Private/Texture.cpp
--- a/Engine/Source/Runtime/Engine/Private/Texture.cpp
+++ b/Engine/Source/Runtime/Engine/Private/Texture.cpp
@@ -168,6 +168,20 @@ void UTexture::PostEditChangeProperty(FPropertyChangedEvent& PropertyChangedEven
 			(PropertyName == SrgbName))
 		{
 			RequiresNotifyMaterials = true;
+
+			if (PropertyName == LODGroupName)
+			{
+				if (LODGroup == TEXTUREGROUP_8BitData)
+				{
+					CompressionSettings = TC_VectorDisplacementmap;
+					SRGB = false;
+				}
+				else if (LODGroup == TEXTUREGROUP_16BitData)
+				{
+					CompressionSettings = TC_HDR;
+					SRGB = false;
+				}
+			}
 		}
 		else if (PropertyName == DeferCompressionName)
 		{
