Update serialization of FLightmap2D to support new members used in Dev-VT

diff --git a/Engine/Source/Runtime/Core/Public/UObject/RenderingObjectVersion.h b/Engine/Source/Runtime/Core/Public/UObject/RenderingObjectVersion.h
--- a/Engine/Source/Runtime/Core/Public/UObject/RenderingObjectVersion.h
+++ b/Engine/Source/Runtime/Core/Public/UObject/RenderingObjectVersion.h
@@ -79,6 +79,8 @@ struct CORE_API FRenderingObjectVersion
 
 		GeometryCacheFastDecoder,
 
+		LightmapHasShadowmapData,
+
 		// -----<new versions can be added above this line>-------------------------------------------------
 		VersionPlusOne,
 		LatestVersion = VersionPlusOne - 1
diff --git a/Engine/Source/Runtime/Engine/Private/LightMap.cpp b/Engine/Source/Runtime/Engine/Private/LightMap.cpp
--- a/Engine/Source/Runtime/Engine/Private/LightMap.cpp
+++ b/Engine/Source/Runtime/Engine/Private/LightMap.cpp
@@ -2364,6 +2364,21 @@ void FLightMap2D::Serialize(FArchive& Ar)
 
 	Ar << CoordinateScale << CoordinateBias;
 
+	if (Ar.CustomVer(FRenderingObjectVersion::GUID) >= FRenderingObjectVersion::LightmapHasShadowmapData)
+	{
+		// In Dev-VT, when VT lightmaps are enabled, baked shadow data is combined with lightmap data, and serialized along with lightmap
+		// This is here as a placeholder
+		FVector4 InvUniformPenumbraSize(ForceInitToZero);
+		bool bShadowChannelValid = false;
+
+		for (int Channel = 0; Channel < 4; Channel++)
+		{
+			Ar << bShadowChannelValid;
+		}
+
+		Ar << InvUniformPenumbraSize;
+	}
+
 	if (Ar.CustomVer(FRenderingObjectVersion::GUID) >= FRenderingObjectVersion::VirtualTexturedLightmaps)
 	{
 		// Don't save VT's if they are disabled for rendering
