RHI.h isn't included in RenderSettings.cpp on Mac builds only. Must be some platform specific way RHI.h gets included. Added RHI.h explicitly to fix the nightly build break.

diff --git a/Engine/Source/Runtime/Engine/Private/RendererSettings.cpp b/Engine/Source/Runtime/Engine/Private/RendererSettings.cpp
--- a/Engine/Source/Runtime/Engine/Private/RendererSettings.cpp
+++ b/Engine/Source/Runtime/Engine/Private/RendererSettings.cpp
@@ -2,6 +2,7 @@
 
 #include "Engine/RendererSettings.h"
 #include "PixelFormat.h"
+#include "RHI.h"
 
 #if WITH_EDITOR
 #include "Editor/EditorEngine.h"
