Disabled SanatizeReflectionCaptureResolution() popup message when cooking from the command line.

diff --git a/Engine/Source/Runtime/Engine/Private/RendererSettings.cpp b/Engine/Source/Runtime/Engine/Private/RendererSettings.cpp
--- a/Engine/Source/Runtime/Engine/Private/RendererSettings.cpp
+++ b/Engine/Source/Runtime/Engine/Private/RendererSettings.cpp
@@ -150,36 +150,38 @@ void URendererSettings::SanatizeReflectionCaptureResolution()
 	ReflectionCaptureResolution = FMath::Clamp(int32(FMath::RoundUpToPowerOfTwo(ReflectionCaptureResolution)), MinCubemapResolution, MaxCubemapResolution);
 
 #if WITH_EDITOR
-	SIZE_T TexMemRequired = CalcTextureSize(ReflectionCaptureResolution, ReflectionCaptureResolution, PF_FloatRGBA, FMath::CeilLogTwo(ReflectionCaptureResolution) + 1) * CubeFace_MAX;
+	if (FApp::CanEverRender())
+	{
+		SIZE_T TexMemRequired = CalcTextureSize(ReflectionCaptureResolution, ReflectionCaptureResolution, PF_FloatRGBA, FMath::CeilLogTwo(ReflectionCaptureResolution) + 1) * CubeFace_MAX;
 
-	FTextureMemoryStats TextureMemStats;
-	RHIGetTextureMemoryStats(TextureMemStats);
+		FTextureMemoryStats TextureMemStats;
+		RHIGetTextureMemoryStats(TextureMemStats);
 
-	if (TexMemRequired > SIZE_T(TextureMemStats.DedicatedVideoMemory / 4))
-	{
-		FNumberFormattingOptions FmtOpts = FNumberFormattingOptions()
-			.SetUseGrouping(false)
-			.SetMaximumFractionalDigits(2)
-			.SetMinimumFractionalDigits(0)
-			.SetRoundingMode(HalfFromZero);
-
-		EAppReturnType::Type Response = FPlatformMisc::MessageBoxExt(
-			EAppMsgType::YesNo,
-			*FText::Format(
-				LOCTEXT("MemAllocWarning_Message_ReflectionCubemap", "A resolution of {0} will require {1} of video memory PER reflection capture component. Are you sure?"),
-				FText::AsNumber(ReflectionCaptureResolution, &FmtOpts),
-				FText::AsMemory(TexMemRequired, &FmtOpts)
+		if (TexMemRequired > SIZE_T(TextureMemStats.DedicatedVideoMemory / 4))
+		{
+			FNumberFormattingOptions FmtOpts = FNumberFormattingOptions()
+				.SetUseGrouping(false)
+				.SetMaximumFractionalDigits(2)
+				.SetMinimumFractionalDigits(0)
+				.SetRoundingMode(HalfFromZero);
+
+			EAppReturnType::Type Response = FPlatformMisc::MessageBoxExt(
+				EAppMsgType::YesNo,
+				*FText::Format(
+					LOCTEXT("MemAllocWarning_Message_ReflectionCubemap", "A resolution of {0} will require {1} of video memory PER reflection capture component. Are you sure?"),
+					FText::AsNumber(ReflectionCaptureResolution, &FmtOpts),
+					FText::AsMemory(TexMemRequired, &FmtOpts)
 				).ToString(),
-			*LOCTEXT("MemAllocWarning_Title_ReflectionCubemap", "Memory Allocation Warning").ToString()
+				*LOCTEXT("MemAllocWarning_Title_ReflectionCubemap", "Memory Allocation Warning").ToString()
 			);
 
-		if (Response == EAppReturnType::No)
-		{
-			ReflectionCaptureResolution = PreEditReflectionCaptureResolution;
+			if (Response == EAppReturnType::No)
+			{
+				ReflectionCaptureResolution = PreEditReflectionCaptureResolution;
+			}
 		}
+		PreEditReflectionCaptureResolution = ReflectionCaptureResolution;
 	}
-
-	PreEditReflectionCaptureResolution = ReflectionCaptureResolution;
 #endif // WITH_EDITOR
 }
 
