Fixed possible non deterministic cooking with UMaterialInterface because of how we resolve texture references.

diff --git a/Engine/Source/Runtime/Engine/Private/Materials/MaterialInterface.cpp b/Engine/Source/Runtime/Engine/Private/Materials/MaterialInterface.cpp
--- a/Engine/Source/Runtime/Engine/Private/Materials/MaterialInterface.cpp
+++ b/Engine/Source/Runtime/Engine/Private/Materials/MaterialInterface.cpp
@@ -521,13 +521,31 @@ void UMaterialInterface::SortTextureStreamingData(bool bForceSort, bool bFinalSo
 	// In cook that was already done in the save.
 	if (!bTextureStreamingDataSorted || bForceSort)
 	{
+		TArray<UTexture*> UsedTextures;
+		if (bFinalSort)
+		{
+			AppendReferencedTextures(UsedTextures);
+			for (int32 TextureIndex = 0; TextureIndex < UsedTextures.Num(); ++TextureIndex)
+			{
+				UTexture* UsedTexture = UsedTextures[TextureIndex];
+				// Sort some of the conditions that could make the texture unstreamable, to make the data leaner.
+				// Note that because we are cooking, UStreamableRenderAsset::bIsStreamable is not reliable here.
+				if (!UsedTexture || UsedTexture->NeverStream || UsedTexture->LODGroup == TEXTUREGROUP_UI || UsedTexture->MipGenSettings == TMGS_NoMipmaps)
+				{
+					UsedTextures.RemoveAtSwap(TextureIndex);
+					--TextureIndex;
+				}
+			}
+		}
+
 		for (int32 Index = 0; Index < TextureStreamingData.Num(); ++Index)
 		{
 			FMaterialTextureInfo& TextureData = TextureStreamingData[Index];
-			UObject* Texture = TextureData.TextureReference.ResolveObject();
+			UTexture* Texture = Cast<UTexture>(TextureData.TextureReference.ResolveObject());
 
-			// In the final data it must also be a streaming texture, to make the data leaner.
-			if (Texture)
+			// Also, when cooking, only keep textures that are directly referenced by this material to prevent non-deterministic cooking.
+			// This would happen if a texture reference resolves to a texture not used anymore by this material. The resolved object could then be valid or not.
+			if (Texture && (!bFinalSort || UsedTextures.Contains(Texture)))
 			{
 				TextureData.TextureName = Texture->GetFName();
 			}
