Cleanup of FViewInfo::PrimitivesWithCustomData and FPrimitiveSceneProxy::PostInitViewCustomData() which isn't used anymore

diff --git a/Engine/Source/Runtime/Engine/Public/PrimitiveSceneProxy.h b/Engine/Source/Runtime/Engine/Public/PrimitiveSceneProxy.h
--- a/Engine/Source/Runtime/Engine/Public/PrimitiveSceneProxy.h
+++ b/Engine/Source/Runtime/Engine/Public/PrimitiveSceneProxy.h
@@ -688,15 +688,6 @@ public:
    	 * @param InMeshScreenSizeSquared - Computed mesh batch screen size, passed to prevent recalculation
 	 */
 	ENGINE_API virtual void* InitViewCustomData(const FSceneView& InView, float InViewLODScale, FMemStackBase& InCustomDataMemStack, bool InIsStaticRelevant, bool InIsShadowOnly, const struct FLODMask* InVisiblePrimitiveLODMask = nullptr, float InMeshScreenSizeSquared = -1.0f) { return nullptr; }
-	
-	/**
-	 * Called during post visibility and shadow setup, just before the frame is rendered. It can be used to update custom data that had a dependency between them.
-	 * Keep in mind this can be called in multihread.
-	 * This will only be called on primitive that added view custom data during the InitViewCustomData.
-	 * @param InView - Current View
- 	 * @param InViewCustomData - Custom data to update
-	 */	
-	ENGINE_API virtual void PostInitViewCustomData(const FSceneView& InView, void* InViewCustomData) const { }
 
 	/** Tell us if this proxy is drawn in game.*/
 	ENGINE_API virtual bool IsDrawnInGame() const { return DrawInGame; }
diff --git a/Engine/Source/Runtime/Landscape/Public/LandscapeRender.h b/Engine/Source/Runtime/Landscape/Public/LandscapeRender.h
--- a/Engine/Source/Runtime/Landscape/Public/LandscapeRender.h
+++ b/Engine/Source/Runtime/Landscape/Public/LandscapeRender.h
@@ -684,7 +684,7 @@ public:
 	virtual void CreateRenderThreadResources() override;
 	virtual void OnLevelAddedToWorld() override;
 	virtual void* InitViewCustomData(const FSceneView& InView, float InViewLODScale, FMemStackBase& InCustomDataMemStack, bool InIsStaticRelevant, bool InIsShadowOnly, const FLODMask* InVisiblePrimitiveLODMask = nullptr, float InMeshScreenSizeSquared = -1.0f) override;
-	virtual void PostInitViewCustomData(const FSceneView& InView, void* InViewCustomData) const override;
+	void PostInitViewCustomData(const FSceneView& InView, void* InViewCustomData) const;
 	virtual bool IsUsingCustomLODRules() const override;
 	virtual FLODMask GetCustomLOD(const FSceneView& InView, float InViewLODScale, int32 InForcedLODLevel, float& OutScreenSizeSquared) const override;
 	virtual bool IsUsingCustomWholeSceneShadowLODRules() const override;
diff --git a/Engine/Source/Runtime/Renderer/Private/SceneRendering.cpp b/Engine/Source/Runtime/Renderer/Private/SceneRendering.cpp
--- a/Engine/Source/Runtime/Renderer/Private/SceneRendering.cpp
+++ b/Engine/Source/Runtime/Renderer/Private/SceneRendering.cpp
@@ -1731,9 +1731,6 @@ void FViewInfo::SetCustomData(const FPrimitiveSceneInfo* InPrimitiveSceneInfo, v
 	{
 		check(PrimitivesCustomData.IsValidIndex(InPrimitiveSceneInfo->GetIndex()));
 		PrimitivesCustomData[InPrimitiveSceneInfo->GetIndex()] = InCustomData;
-
-		check(!PrimitivesWithCustomData.Contains(InPrimitiveSceneInfo));
-		PrimitivesWithCustomData.Add(InPrimitiveSceneInfo);
 	}
 }
 
diff --git a/Engine/Source/Runtime/Renderer/Private/SceneRendering.h b/Engine/Source/Runtime/Renderer/Private/SceneRendering.h
--- a/Engine/Source/Runtime/Renderer/Private/SceneRendering.h
+++ b/Engine/Source/Runtime/Renderer/Private/SceneRendering.h
@@ -849,7 +849,6 @@ public:
 	FMobileCSMVisibilityInfo MobileCSMVisibilityInfo;
 
 	// Primitive CustomData
-	TArray<const FPrimitiveSceneInfo*, SceneRenderingAllocator> PrimitivesWithCustomData;	// Size == Amount of Primitive With Custom Data
 	TArray<FMemStackBase, SceneRenderingAllocator> PrimitiveCustomDataMemStack; // Size == 1 global stack + 1 per visibility thread (if multithread)
 
 	/** Parameters for exponential height fog. */
diff --git a/Engine/Source/Runtime/Renderer/Private/SceneVisibility.cpp b/Engine/Source/Runtime/Renderer/Private/SceneVisibility.cpp
--- a/Engine/Source/Runtime/Renderer/Private/SceneVisibility.cpp
+++ b/Engine/Source/Runtime/Renderer/Private/SceneVisibility.cpp
@@ -3451,7 +3451,6 @@ void FSceneRenderer::ComputeViewVisibility(FRHICommandListImmediate& RHICmdList,
 		View.PrimitivesLODMask.Init(FLODMask(), Scene->Primitives.Num());
 
 		View.PrimitivesCustomData.Init(nullptr, Scene->Primitives.Num());
-		View.PrimitivesWithCustomData.Reserve(Scene->Primitives.Num());
 
 		// We must reserve to prevent realloc otherwise it will cause memory leak if we Execute In Parallel
 		const bool WillExecuteInParallel = FApp::ShouldUseThreadingForPerformance() && CVarParallelInitViews.GetValueOnRenderThread() > 0;
