Added debug option to constantly update reflection capture volumes.

diff --git a/Engine/Source/Runtime/Engine/Private/Components/ReflectionCaptureComponent.cpp b/Engine/Source/Runtime/Engine/Private/Components/ReflectionCaptureComponent.cpp
--- a/Engine/Source/Runtime/Engine/Private/Components/ReflectionCaptureComponent.cpp
+++ b/Engine/Source/Runtime/Engine/Private/Components/ReflectionCaptureComponent.cpp
@@ -32,6 +32,7 @@
 #include "Components/BoxReflectionCaptureComponent.h"
 #include "Engine/PlaneReflectionCapture.h"
 #include "Engine/BoxReflectionCapture.h"
+#include "EngineUtils.h"
 #include "Components/PlaneReflectionCaptureComponent.h"
 #include "Components/BoxComponent.h"
 #include "Components/SkyLightComponent.h"
@@ -57,6 +58,12 @@ ENGINE_API TAutoConsoleVariable<int32> CVarReflectionCaptureSize(
 	TEXT("Set the resolution for all reflection capture cubemaps. Should be set via project's Render Settings. Must be power of 2. Defaults to 128.\n")
 	);
 
+TAutoConsoleVariable<int32> CVarUpdateReflectionCaptureEveryFrame(
+	TEXT("r.UpdateReflectionCaptureEveryFrame"),
+	0,
+	TEXT("When set, reflection captures will constantly be scheduled for update.\n")
+);
+
 static int32 SanitizeReflectionCaptureSize(int32 ReflectionCaptureSize)
 {
 	const int32 MaxReflectionCaptureSize = GetMaxCubeTextureDimension();
@@ -1042,6 +1049,19 @@ void UReflectionCaptureComponent::UpdateReflectionCaptureContents(UWorld* WorldT
 		//guarantee that all render proxies are up to date before kicking off this render
 		WorldToUpdate->SendAllEndOfFrameUpdates();
 
+		if (CVarUpdateReflectionCaptureEveryFrame.GetValueOnGameThread())
+		{
+			for (FActorIterator It(WorldToUpdate); It; ++It)
+			{
+				TInlineComponentArray<UReflectionCaptureComponent*> Components;
+				(*It)->GetComponents(Components);
+				for (int32 ComponentIndex = 0; ComponentIndex < Components.Num(); ComponentIndex++)
+				{
+					Components[ComponentIndex]->MarkDirtyForRecapture(); // Continuously refresh reflection captures
+				}
+			}
+		}
+
 		TArray<UReflectionCaptureComponent*> WorldCombinedCaptures;
 
 		for (int32 CaptureIndex = ReflectionCapturesToUpdate.Num() - 1; CaptureIndex >= 0; CaptureIndex--)
