DR - Viewmode list more readable
#rnx

diff --git a/Engine/Source/Runtime/Engine/Private/GameViewportClient.cpp b/Engine/Source/Runtime/Engine/Private/GameViewportClient.cpp
--- a/Engine/Source/Runtime/Engine/Private/GameViewportClient.cpp
+++ b/Engine/Source/Runtime/Engine/Private/GameViewportClient.cpp
@@ -3042,7 +3042,14 @@ bool UGameViewportClient::HandleViewModeCommand( const TCHAR* Cmd, FOutputDevice
 		{
 			if(i != 0)
 			{
-				ViewModes += TEXT(", ");
+				if ((i % 5) == 0)
+				{
+					ViewModes += TEXT("\n     ");
+				}
+				else
+				{
+					ViewModes += TEXT(", ");
+				}
 			}
 			ViewModes += GetViewModeName((EViewModeIndex)i);
 		}
