Speculative fix since still unable to repro locally.
- Guard message from displaying in non-rendering apps.
- Be more conservative when we display the message box.

diff --git a/Engine/Source/Runtime/Engine/Private/Components/SkyLightComponent.cpp b/Engine/Source/Runtime/Engine/Private/Components/SkyLightComponent.cpp
--- a/Engine/Source/Runtime/Engine/Private/Components/SkyLightComponent.cpp
+++ b/Engine/Source/Runtime/Engine/Private/Components/SkyLightComponent.cpp
@@ -273,36 +273,39 @@ void USkyLightComponent::SanitizeCubemapSize()
 	CubemapResolution = FMath::Clamp(int32(FMath::RoundUpToPowerOfTwo(CubemapResolution)), MinCubemapResolution, MaxCubemapResolution);
 
 #if WITH_EDITOR
-	SIZE_T TexMemRequired = CalcTextureSize(CubemapResolution, CubemapResolution, SKYLIGHT_CUBEMAP_FORMAT, FMath::CeilLogTwo(CubemapResolution) + 1) * CubeFace_MAX;
-
-	FTextureMemoryStats TextureMemStats;
-	RHIGetTextureMemoryStats(TextureMemStats);
-
-	if (TexMemRequired > SIZE_T(TextureMemStats.DedicatedVideoMemory / 2))
+	if (FApp::CanEverRender() && !FApp::IsUnattended())
 	{
-		FNumberFormattingOptions FmtOpts = FNumberFormattingOptions()
-			.SetUseGrouping(false)
-			.SetMaximumFractionalDigits(2)
-			.SetMinimumFractionalDigits(0)
-			.SetRoundingMode(HalfFromZero);
+		SIZE_T TexMemRequired = CalcTextureSize(CubemapResolution, CubemapResolution, SKYLIGHT_CUBEMAP_FORMAT, FMath::CeilLogTwo(CubemapResolution) + 1) * CubeFace_MAX;
+
+		FTextureMemoryStats TextureMemStats;
+		RHIGetTextureMemoryStats(TextureMemStats);
 
-		EAppReturnType::Type Response = FPlatformMisc::MessageBoxExt(
-			EAppMsgType::YesNo, 
-			*FText::Format(
-				LOCTEXT("MemAllocWarning_Message_SkylightCubemap", "A resolution of {0} will require {1} of video memory. Are you sure?"), 
-				FText::AsNumber(CubemapResolution, &FmtOpts), 
-				FText::AsMemory(TexMemRequired, &FmtOpts)
+		if (TextureMemStats.DedicatedVideoMemory > 0 && TexMemRequired > SIZE_T(TextureMemStats.DedicatedVideoMemory / 4))
+		{
+			FNumberFormattingOptions FmtOpts = FNumberFormattingOptions()
+				.SetUseGrouping(false)
+				.SetMaximumFractionalDigits(2)
+				.SetMinimumFractionalDigits(0)
+				.SetRoundingMode(HalfFromZero);
+
+			EAppReturnType::Type Response = FPlatformMisc::MessageBoxExt(
+				EAppMsgType::YesNo,
+				*FText::Format(
+					LOCTEXT("MemAllocWarning_Message_SkylightCubemap", "A resolution of {0} will require {1} of video memory. Are you sure?"),
+					FText::AsNumber(CubemapResolution, &FmtOpts),
+					FText::AsMemory(TexMemRequired, &FmtOpts)
 				).ToString(),
-			*LOCTEXT("MemAllocWarning_Title_SkylightCubemap", "Memory Allocation Warning").ToString()
+				*LOCTEXT("MemAllocWarning_Title_SkylightCubemap", "Memory Allocation Warning").ToString()
 			);
 
-		if (Response == EAppReturnType::No)
-		{
-			CubemapResolution = PreEditCubemapResolution;
+			if (Response == EAppReturnType::No)
+			{
+				CubemapResolution = PreEditCubemapResolution;
+			}
 		}
-	}
 
-	PreEditCubemapResolution = CubemapResolution;
+		PreEditCubemapResolution = CubemapResolution;
+	}
 #endif // WITH_EDITOR
 }
 
diff --git a/Engine/Source/Runtime/Engine/Private/RendererSettings.cpp b/Engine/Source/Runtime/Engine/Private/RendererSettings.cpp
--- a/Engine/Source/Runtime/Engine/Private/RendererSettings.cpp
+++ b/Engine/Source/Runtime/Engine/Private/RendererSettings.cpp
@@ -177,14 +177,14 @@ void URendererSettings::SanatizeReflectionCaptureResolution()
 	ReflectionCaptureResolution = FMath::Clamp(int32(FMath::RoundUpToPowerOfTwo(ReflectionCaptureResolution)), MinCubemapResolution, MaxCubemapResolution);
 
 #if WITH_EDITOR
-	if (FApp::CanEverRender())
+	if (FApp::CanEverRender() && !FApp::IsUnattended())
 	{
 		SIZE_T TexMemRequired = CalcTextureSize(ReflectionCaptureResolution, ReflectionCaptureResolution, PF_FloatRGBA, FMath::CeilLogTwo(ReflectionCaptureResolution) + 1) * CubeFace_MAX;
 
 		FTextureMemoryStats TextureMemStats;
 		RHIGetTextureMemoryStats(TextureMemStats);
 
-		if (TexMemRequired > SIZE_T(TextureMemStats.DedicatedVideoMemory / 4))
+		if (TextureMemStats.DedicatedVideoMemory > 0 && TexMemRequired > SIZE_T(TextureMemStats.DedicatedVideoMemory / 8))
 		{
 			FNumberFormattingOptions FmtOpts = FNumberFormattingOptions()
 				.SetUseGrouping(false)
