DR - Fix uninitialized variable

diff --git a/Engine/Source/Runtime/Engine/Public/Rendering/SkeletalMeshRenderData.h b/Engine/Source/Runtime/Engine/Public/Rendering/SkeletalMeshRenderData.h
--- a/Engine/Source/Runtime/Engine/Public/Rendering/SkeletalMeshRenderData.h
+++ b/Engine/Source/Runtime/Engine/Public/Rendering/SkeletalMeshRenderData.h
@@ -54,5 +54,5 @@ public:
 
 private:
 	/** True if the resource has been initialized. */
-	bool bInitialized;
+	bool bInitialized = false;
 };
\ No newline at end of file
