Fixed compile error in shipping configs due to CVarFreezeMouseCursor not being available. ViewUniformShaderParameters.CursorPosition is no longer updated in shipping and test.

diff --git a/Engine/Source/Runtime/Engine/Private/SceneView.cpp b/Engine/Source/Runtime/Engine/Private/SceneView.cpp
--- a/Engine/Source/Runtime/Engine/Private/SceneView.cpp
+++ b/Engine/Source/Runtime/Engine/Private/SceneView.cpp
@@ -2345,12 +2345,16 @@ void FSceneView::SetupCommonViewUniformBufferParameters(
 	ViewUniformShaderParameters.GameTime = Family->CurrentWorldTime;
 	ViewUniformShaderParameters.RealTime = Family->CurrentRealTime;
 	ViewUniformShaderParameters.DeltaTime = Family->DeltaWorldTime;
+
+#if !(UE_BUILD_SHIPPING || UE_BUILD_TEST)
 	static FIntPoint LockedCursorPos = CursorPos;
 	if (CVarFreezeMouseCursor.GetValueOnRenderThread() == 0 && CursorPos.X >= 0 && CursorPos.Y >= 0)
 	{
 		LockedCursorPos = CursorPos;
 	}
 	ViewUniformShaderParameters.CursorPosition = LockedCursorPos;
+#endif
+
 	ViewUniformShaderParameters.Random = FMath::Rand();
 	ViewUniformShaderParameters.FrameNumber = Family->FrameNumber;
 
