Cherry-pick 5407391: Kills the HLR artifact of the editor depth test boolean.

This got killed in 3718066 refactoring editor primitive compositing to remove this manual depth test to have temporal upsample working in editor viewports.

diff --git a/Engine/Plugins/2D/Paper2D/Source/Paper2D/Private/PaperRenderSceneProxy.cpp b/Engine/Plugins/2D/Paper2D/Source/Paper2D/Private/PaperRenderSceneProxy.cpp
--- a/Engine/Plugins/2D/Paper2D/Source/Paper2D/Private/PaperRenderSceneProxy.cpp
+++ b/Engine/Plugins/2D/Paper2D/Source/Paper2D/Private/PaperRenderSceneProxy.cpp
@@ -175,12 +175,12 @@ void FPaperRenderSceneProxy::DebugDrawBodySetup(const FSceneView* View, int32 Vi
 
 			Collector.RegisterOneFrameMaterialProxy(SolidMaterialInstance);
 
-			BodySetup->AggGeom.GetAggGeom(GeomTransform, GetWireframeColor().ToFColor(true), SolidMaterialInstance, false, true, UseEditorDepthTest(), ViewIndex, Collector);
+			BodySetup->AggGeom.GetAggGeom(GeomTransform, GetWireframeColor().ToFColor(true), SolidMaterialInstance, false, true, ViewIndex, Collector);
 		}
 		else
 		{
 			// wireframe
-			BodySetup->AggGeom.GetAggGeom(GeomTransform, GetSelectionColor(CollisionColor, IsSelected(), IsHovered()).ToFColor(true), nullptr, ( Owner == nullptr ), false, UseEditorDepthTest(), ViewIndex, Collector);
+			BodySetup->AggGeom.GetAggGeom(GeomTransform, GetSelectionColor(CollisionColor, IsSelected(), IsHovered()).ToFColor(true), nullptr, ( Owner == nullptr ), false, ViewIndex, Collector);
 		}
 	}
 }
diff --git a/Engine/Plugins/2D/Paper2D/Source/Paper2D/Private/PaperSpriteSceneProxy.cpp b/Engine/Plugins/2D/Paper2D/Source/Paper2D/Private/PaperSpriteSceneProxy.cpp
--- a/Engine/Plugins/2D/Paper2D/Source/Paper2D/Private/PaperSpriteSceneProxy.cpp
+++ b/Engine/Plugins/2D/Paper2D/Source/Paper2D/Private/PaperSpriteSceneProxy.cpp
@@ -86,7 +86,7 @@ void FPaperSpriteSceneProxy::GetDynamicMeshElements(const TArray<const FSceneVie
 
 						const bool bUseSeparateColorPerHull = (Owner == nullptr);
 						const bool bDrawSolid = false;
-						BodySetup->AggGeom.GetAggGeom(GeomTransform, GetSelectionColor(CollisionColor, bDrawWireSelected, IsHovered()).ToFColor(true), CollisionMaterialInstance, bUseSeparateColorPerHull, bDrawSolid, false, ViewIndex, Collector);
+						BodySetup->AggGeom.GetAggGeom(GeomTransform, GetSelectionColor(CollisionColor, bDrawWireSelected, IsHovered()).ToFColor(true), CollisionMaterialInstance, bUseSeparateColorPerHull, bDrawSolid, ViewIndex, Collector);
 					}
 				}
 			}
diff --git a/Engine/Plugins/2D/Paper2D/Source/Paper2D/Private/PaperTileMapRenderSceneProxy.cpp b/Engine/Plugins/2D/Paper2D/Source/Paper2D/Private/PaperTileMapRenderSceneProxy.cpp
--- a/Engine/Plugins/2D/Paper2D/Source/Paper2D/Private/PaperTileMapRenderSceneProxy.cpp
+++ b/Engine/Plugins/2D/Paper2D/Source/Paper2D/Private/PaperTileMapRenderSceneProxy.cpp
@@ -321,7 +321,7 @@ void FPaperTileMapRenderSceneProxy::GetDynamicMeshElements(const TArray<const FS
 
 							const bool bPerHullColor = false;
 							const bool bDrawSimpleSolid = false;
-							BodySetup->AggGeom.GetAggGeom(GeomTransform, GetSelectionColor(CollisionColor, bDrawWireSelected, IsHovered()).ToFColor(true), CollisionMaterialInstance, bPerHullColor, bDrawSimpleSolid, UseEditorDepthTest(), ViewIndex, Collector);
+							BodySetup->AggGeom.GetAggGeom(GeomTransform, GetSelectionColor(CollisionColor, bDrawWireSelected, IsHovered()).ToFColor(true), CollisionMaterialInstance, bPerHullColor, bDrawSimpleSolid, ViewIndex, Collector);
 						}
 					}
 				}
diff --git a/Engine/Plugins/Editor/MeshEditor/Source/MeshEditor/OverlayComponent.cpp b/Engine/Plugins/Editor/MeshEditor/Source/MeshEditor/OverlayComponent.cpp
--- a/Engine/Plugins/Editor/MeshEditor/Source/MeshEditor/OverlayComponent.cpp
+++ b/Engine/Plugins/Editor/MeshEditor/Source/MeshEditor/OverlayComponent.cpp
@@ -253,7 +253,7 @@ public:
 					Mesh.MaterialRenderProxy = MeshBatchData.MaterialProxy;
 
 					FDynamicPrimitiveUniformBuffer& DynamicPrimitiveUniformBuffer = Collector.AllocateOneFrameResource<FDynamicPrimitiveUniformBuffer>();
-					DynamicPrimitiveUniformBuffer.Set(GetLocalToWorld(), GetLocalToWorld(), GetBounds(), GetLocalBounds(), true, false, UseEditorDepthTest());
+					DynamicPrimitiveUniformBuffer.Set(GetLocalToWorld(), GetLocalToWorld(), GetBounds(), GetLocalBounds(), true, false);
 					BatchElement.PrimitiveUniformBufferResource = &DynamicPrimitiveUniformBuffer.UniformBuffer;
 
 					BatchElement.FirstIndex = MeshBatchData.StartIndex;
diff --git a/Engine/Plugins/Editor/MeshEditor/Source/MeshEditor/WireframeMeshComponent.cpp b/Engine/Plugins/Editor/MeshEditor/Source/MeshEditor/WireframeMeshComponent.cpp
--- a/Engine/Plugins/Editor/MeshEditor/Source/MeshEditor/WireframeMeshComponent.cpp
+++ b/Engine/Plugins/Editor/MeshEditor/Source/MeshEditor/WireframeMeshComponent.cpp
@@ -344,7 +344,7 @@ public:
 				Mesh.MaterialRenderProxy = MaterialProxy;
 
 				FDynamicPrimitiveUniformBuffer& DynamicPrimitiveUniformBuffer = Collector.AllocateOneFrameResource<FDynamicPrimitiveUniformBuffer>();
-				DynamicPrimitiveUniformBuffer.Set(GetLocalToWorld(), GetLocalToWorld(), GetBounds(), GetLocalBounds(), true, false, UseEditorDepthTest());
+				DynamicPrimitiveUniformBuffer.Set(GetLocalToWorld(), GetLocalToWorld(), GetBounds(), GetLocalBounds(), true, false);
 				BatchElement.PrimitiveUniformBufferResource = &DynamicPrimitiveUniformBuffer.UniformBuffer;
 
 				BatchElement.FirstIndex = 0;
diff --git a/Engine/Plugins/Experimental/GeometryCache/Source/GeometryCache/Private/GeometryCacheSceneProxy.cpp b/Engine/Plugins/Experimental/GeometryCache/Source/GeometryCache/Private/GeometryCacheSceneProxy.cpp
--- a/Engine/Plugins/Experimental/GeometryCache/Source/GeometryCache/Private/GeometryCacheSceneProxy.cpp
+++ b/Engine/Plugins/Experimental/GeometryCache/Source/GeometryCache/Private/GeometryCacheSceneProxy.cpp
@@ -357,7 +357,7 @@ void FGeometryCacheSceneProxy::CreateMeshBatch(
 
 	const FMatrix& LocalToWorldTransform = TrackProxy->WorldMatrix * GetLocalToWorld();
 
-	DynamicPrimitiveUniformBuffer.Set(LocalToWorldTransform, LocalToWorldTransform, GetBounds(), GetLocalBounds(), true, false, UseEditorDepthTest());
+	DynamicPrimitiveUniformBuffer.Set(LocalToWorldTransform, LocalToWorldTransform, GetBounds(), GetLocalBounds(), true, false);
 	BatchElement.PrimitiveUniformBuffer = DynamicPrimitiveUniformBuffer.UniformBuffer.GetUniformBufferRHI();
 
 	BatchElement.FirstIndex = BatchInfo.StartIndex;
diff --git a/Engine/Plugins/FX/Niagara/Source/Niagara/Private/NiagaraRendererMeshes.cpp b/Engine/Plugins/FX/Niagara/Source/Niagara/Private/NiagaraRendererMeshes.cpp
--- a/Engine/Plugins/FX/Niagara/Source/Niagara/Private/NiagaraRendererMeshes.cpp
+++ b/Engine/Plugins/FX/Niagara/Source/Niagara/Private/NiagaraRendererMeshes.cpp
@@ -165,7 +165,6 @@ void NiagaraRendererMeshes::GetDynamicMeshElements(const TArray<const FSceneView
 				false,
 				false,
 				false,
-				SceneProxy->UseEditorDepthTest(),
 				SceneProxy->GetLightingChannelMask(),
 				0,
 				INDEX_NONE,
diff --git a/Engine/Plugins/FX/Niagara/Source/Niagara/Private/NiagaraRendererRibbons.cpp b/Engine/Plugins/FX/Niagara/Source/Niagara/Private/NiagaraRendererRibbons.cpp
--- a/Engine/Plugins/FX/Niagara/Source/Niagara/Private/NiagaraRendererRibbons.cpp
+++ b/Engine/Plugins/FX/Niagara/Source/Niagara/Private/NiagaraRendererRibbons.cpp
@@ -132,7 +132,6 @@ void NiagaraRendererRibbons::GetDynamicMeshElements(const TArray<const FSceneVie
 				false,
 				SceneProxy->UseSingleSampleShadowFromStationaryLights(),
 				SceneProxy->GetScene().HasPrecomputedVolumetricLightmap_RenderThread(),
-				SceneProxy->UseEditorDepthTest(),
 				SceneProxy->GetLightingChannelMask(),
 				0,
 				INDEX_NONE,
diff --git a/Engine/Plugins/FX/Niagara/Source/Niagara/Private/NiagaraRendererSprites.cpp b/Engine/Plugins/FX/Niagara/Source/Niagara/Private/NiagaraRendererSprites.cpp
--- a/Engine/Plugins/FX/Niagara/Source/Niagara/Private/NiagaraRendererSprites.cpp
+++ b/Engine/Plugins/FX/Niagara/Source/Niagara/Private/NiagaraRendererSprites.cpp
@@ -130,7 +130,6 @@ void NiagaraRendererSprites::ConditionalInitPrimitiveUniformBuffer(const FNiagar
 			false,
 			SceneProxy->UseSingleSampleShadowFromStationaryLights(),
 			SceneProxy->GetScene().HasPrecomputedVolumetricLightmap_RenderThread(),
-			SceneProxy->UseEditorDepthTest(),
 			SceneProxy->GetLightingChannelMask(),
 			0,
 			INDEX_NONE,
diff --git a/Engine/Plugins/Runtime/CableComponent/Source/CableComponent/Private/CableComponent.cpp b/Engine/Plugins/Runtime/CableComponent/Source/CableComponent/Private/CableComponent.cpp
--- a/Engine/Plugins/Runtime/CableComponent/Source/CableComponent/Private/CableComponent.cpp
+++ b/Engine/Plugins/Runtime/CableComponent/Source/CableComponent/Private/CableComponent.cpp
@@ -288,7 +288,7 @@ public:
 				GetScene().GetPrimitiveUniformShaderParameters_RenderThread(GetPrimitiveSceneInfo(), bHasPrecomputedVolumetricLightmap, PreviousLocalToWorld, SingleCaptureIndex);
 
 				FDynamicPrimitiveUniformBuffer& DynamicPrimitiveUniformBuffer = Collector.AllocateOneFrameResource<FDynamicPrimitiveUniformBuffer>();
-				DynamicPrimitiveUniformBuffer.Set(GetLocalToWorld(), PreviousLocalToWorld, GetBounds(), GetLocalBounds(), true, bHasPrecomputedVolumetricLightmap, UseEditorDepthTest());
+				DynamicPrimitiveUniformBuffer.Set(GetLocalToWorld(), PreviousLocalToWorld, GetBounds(), GetLocalBounds(), true, bHasPrecomputedVolumetricLightmap);
 				BatchElement.PrimitiveUniformBufferResource = &DynamicPrimitiveUniformBuffer.UniformBuffer;
 
 				BatchElement.FirstIndex = 0;
diff --git a/Engine/Plugins/Runtime/CustomMeshComponent/Source/CustomMeshComponent/Private/CustomMeshComponent.cpp b/Engine/Plugins/Runtime/CustomMeshComponent/Source/CustomMeshComponent/Private/CustomMeshComponent.cpp
--- a/Engine/Plugins/Runtime/CustomMeshComponent/Source/CustomMeshComponent/Private/CustomMeshComponent.cpp
+++ b/Engine/Plugins/Runtime/CustomMeshComponent/Source/CustomMeshComponent/Private/CustomMeshComponent.cpp
@@ -135,7 +135,7 @@ public:
 				GetScene().GetPrimitiveUniformShaderParameters_RenderThread(GetPrimitiveSceneInfo(), bHasPrecomputedVolumetricLightmap, PreviousLocalToWorld, SingleCaptureIndex);
 
 				FDynamicPrimitiveUniformBuffer& DynamicPrimitiveUniformBuffer = Collector.AllocateOneFrameResource<FDynamicPrimitiveUniformBuffer>();
-				DynamicPrimitiveUniformBuffer.Set(GetLocalToWorld(), PreviousLocalToWorld, GetBounds(), GetLocalBounds(), true, bHasPrecomputedVolumetricLightmap, UseEditorDepthTest());
+				DynamicPrimitiveUniformBuffer.Set(GetLocalToWorld(), PreviousLocalToWorld, GetBounds(), GetLocalBounds(), true, bHasPrecomputedVolumetricLightmap);
 				BatchElement.PrimitiveUniformBufferResource = &DynamicPrimitiveUniformBuffer.UniformBuffer;
 
 				BatchElement.FirstIndex = 0;
diff --git a/Engine/Plugins/Runtime/Oculus/OculusVR/Source/OculusMR/Private/OculusMR_BoundaryMeshComponent.cpp b/Engine/Plugins/Runtime/Oculus/OculusVR/Source/OculusMR/Private/OculusMR_BoundaryMeshComponent.cpp
--- a/Engine/Plugins/Runtime/Oculus/OculusVR/Source/OculusMR/Private/OculusMR_BoundaryMeshComponent.cpp
+++ b/Engine/Plugins/Runtime/Oculus/OculusVR/Source/OculusMR/Private/OculusMR_BoundaryMeshComponent.cpp
@@ -150,7 +150,7 @@ public:
 				//Mesh.bWireframe = bWireframe;
 				//Mesh.VertexFactory = &VertexFactory;
 				//Mesh.MaterialRenderProxy = MaterialProxy;
-				//BatchElement.PrimitiveUniformBuffer = CreatePrimitiveUniformBufferImmediate(Transform.ToMatrixWithScale(), GetBounds(), GetLocalBounds(), true, UseEditorDepthTest());
+				//BatchElement.PrimitiveUniformBuffer = CreatePrimitiveUniformBufferImmediate(Transform.ToMatrixWithScale(), GetBounds(), GetLocalBounds(), true);
 				//BatchElement.FirstIndex = 0;
 				//BatchElement.NumPrimitives = IndexBuffer.Indices.Num() / 3;
 				//BatchElement.MinVertexIndex = 0;
diff --git a/Engine/Plugins/Runtime/Oculus/OculusVR/Source/OculusMR/Private/OculusMR_PlaneMeshComponent.cpp b/Engine/Plugins/Runtime/Oculus/OculusVR/Source/OculusMR/Private/OculusMR_PlaneMeshComponent.cpp
--- a/Engine/Plugins/Runtime/Oculus/OculusVR/Source/OculusMR/Private/OculusMR_PlaneMeshComponent.cpp
+++ b/Engine/Plugins/Runtime/Oculus/OculusVR/Source/OculusMR/Private/OculusMR_PlaneMeshComponent.cpp
@@ -127,7 +127,7 @@ public:
 					//Mesh.bWireframe = bWireframe;
 					//Mesh.VertexFactory = &VertexFactory;
 					//Mesh.MaterialRenderProxy = MaterialProxy;
-					//BatchElement.PrimitiveUniformBuffer = CreatePrimitiveUniformBufferImmediate(GetLocalToWorld(), GetBounds(), GetLocalBounds(), true, UseEditorDepthTest());
+					//BatchElement.PrimitiveUniformBuffer = CreatePrimitiveUniformBufferImmediate(GetLocalToWorld(), GetBounds(), GetLocalBounds(), true);
 					//BatchElement.FirstIndex = 0;
 					//BatchElement.NumPrimitives = IndexBuffer.Indices.Num() / 3;
 					//BatchElement.MinVertexIndex = 0;
diff --git a/Engine/Plugins/Runtime/ProceduralMeshComponent/Source/ProceduralMeshComponent/Private/ProceduralMeshComponent.cpp b/Engine/Plugins/Runtime/ProceduralMeshComponent/Source/ProceduralMeshComponent/Private/ProceduralMeshComponent.cpp
--- a/Engine/Plugins/Runtime/ProceduralMeshComponent/Source/ProceduralMeshComponent/Private/ProceduralMeshComponent.cpp
+++ b/Engine/Plugins/Runtime/ProceduralMeshComponent/Source/ProceduralMeshComponent/Private/ProceduralMeshComponent.cpp
@@ -307,7 +307,7 @@ public:
 						GetScene().GetPrimitiveUniformShaderParameters_RenderThread(GetPrimitiveSceneInfo(), bHasPrecomputedVolumetricLightmap, PreviousLocalToWorld, SingleCaptureIndex);
 
 						FDynamicPrimitiveUniformBuffer& DynamicPrimitiveUniformBuffer = Collector.AllocateOneFrameResource<FDynamicPrimitiveUniformBuffer>();
-						DynamicPrimitiveUniformBuffer.Set(GetLocalToWorld(), PreviousLocalToWorld, GetBounds(), GetLocalBounds(), true, bHasPrecomputedVolumetricLightmap, UseEditorDepthTest());
+						DynamicPrimitiveUniformBuffer.Set(GetLocalToWorld(), PreviousLocalToWorld, GetBounds(), GetLocalBounds(), true, bHasPrecomputedVolumetricLightmap);
 						BatchElement.PrimitiveUniformBufferResource = &DynamicPrimitiveUniformBuffer.UniformBuffer;
 
 						BatchElement.FirstIndex = 0;
@@ -334,7 +334,7 @@ public:
 				if (ViewFamily.EngineShowFlags.Collision && IsCollisionEnabled() && BodySetup->GetCollisionTraceFlag() != ECollisionTraceFlag::CTF_UseComplexAsSimple)
 				{
 					FTransform GeomTransform(GetLocalToWorld());
-					BodySetup->AggGeom.GetAggGeom(GeomTransform, GetSelectionColor(FColor(157, 149, 223, 255), IsSelected(), IsHovered()).ToFColor(true), NULL, false, false, UseEditorDepthTest(), ViewIndex, Collector);
+					BodySetup->AggGeom.GetAggGeom(GeomTransform, GetSelectionColor(FColor(157, 149, 223, 255), IsSelected(), IsHovered()).ToFColor(true), NULL, false, false, ViewIndex, Collector);
 				}
 
 				// Render bounds
diff --git a/Engine/Shaders/Private/SceneData.ush b/Engine/Shaders/Private/SceneData.ush
--- a/Engine/Shaders/Private/SceneData.ush
+++ b/Engine/Shaders/Private/SceneData.ush
@@ -25,7 +25,7 @@ struct FPrimitiveSceneData
 	float DecalReceiverMask;
 	float PerObjectGBufferData;
 	float UseVolumetricLightmapShadowFromStationaryLights;
-	float UseEditorDepthTest;
+	float Unused;
 	float4 ObjectOrientation;
 	float4 NonUniformScale;
 	float3 LocalObjectBoundsMin;
@@ -78,7 +78,7 @@ FPrimitiveSceneData GetPrimitiveData(uint PrimitiveId)
 	PrimitiveData.DecalReceiverMask = View.PrimitiveSceneData[PrimitiveBaseOffset + 20].x;
 	PrimitiveData.PerObjectGBufferData = View.PrimitiveSceneData[PrimitiveBaseOffset + 20].y;
 	PrimitiveData.UseVolumetricLightmapShadowFromStationaryLights = View.PrimitiveSceneData[PrimitiveBaseOffset + 20].z;
-	PrimitiveData.UseEditorDepthTest = View.PrimitiveSceneData[PrimitiveBaseOffset + 20].w;
+	PrimitiveData.Unused = false;
 
 	PrimitiveData.ObjectOrientation = View.PrimitiveSceneData[PrimitiveBaseOffset + 21];
 	PrimitiveData.NonUniformScale = View.PrimitiveSceneData[PrimitiveBaseOffset + 22];
diff --git a/Engine/Source/Runtime/Engine/Classes/PhysicsEngine/AggregateGeom.h b/Engine/Source/Runtime/Engine/Classes/PhysicsEngine/AggregateGeom.h
--- a/Engine/Source/Runtime/Engine/Classes/PhysicsEngine/AggregateGeom.h
+++ b/Engine/Source/Runtime/Engine/Classes/PhysicsEngine/AggregateGeom.h
@@ -76,7 +76,7 @@ struct ENGINE_API FKAggregateGeom
 	void FixupDeprecated(FArchive& Ar);
 #endif
 
-	void GetAggGeom(const FTransform& Transform, const FColor Color, const FMaterialRenderProxy* MatInst, bool bPerHullColor, bool bDrawSolid, bool bUseEditorDepthTest, int32 ViewIndex, class FMeshElementCollector& Collector) const;
+	void GetAggGeom(const FTransform& Transform, const FColor Color, const FMaterialRenderProxy* MatInst, bool bPerHullColor, bool bDrawSolid, int32 ViewIndex, class FMeshElementCollector& Collector) const;
 
 	/** Release the RenderInfo (if its there) and safely clean up any resources. Call on the game thread. */
 	void FreeRenderInfo();
diff --git a/Engine/Source/Runtime/Engine/Private/Components/ArrowComponent.cpp b/Engine/Source/Runtime/Engine/Private/Components/ArrowComponent.cpp
--- a/Engine/Source/Runtime/Engine/Private/Components/ArrowComponent.cpp
+++ b/Engine/Source/Runtime/Engine/Private/Components/ArrowComponent.cpp
@@ -148,7 +148,7 @@ public:
 				Mesh.MaterialRenderProxy = ArrowMaterialRenderProxy;
 
 				FDynamicPrimitiveUniformBuffer& DynamicPrimitiveUniformBuffer = Collector.AllocateOneFrameResource<FDynamicPrimitiveUniformBuffer>();
-				DynamicPrimitiveUniformBuffer.Set(FScaleMatrix(ViewScale) * EffectiveLocalToWorld, FScaleMatrix(ViewScale) * EffectiveLocalToWorld, GetBounds(), GetLocalBounds(), true, false, UseEditorDepthTest());
+				DynamicPrimitiveUniformBuffer.Set(FScaleMatrix(ViewScale) * EffectiveLocalToWorld, FScaleMatrix(ViewScale) * EffectiveLocalToWorld, GetBounds(), GetLocalBounds(), true, false);
 				BatchElement.PrimitiveUniformBufferResource = &DynamicPrimitiveUniformBuffer.UniformBuffer;
 
 				BatchElement.FirstIndex = 0;
diff --git a/Engine/Source/Runtime/Engine/Private/Components/BrushComponent.cpp b/Engine/Source/Runtime/Engine/Private/Components/BrushComponent.cpp
--- a/Engine/Source/Runtime/Engine/Private/Components/BrushComponent.cpp
+++ b/Engine/Source/Runtime/Engine/Private/Components/BrushComponent.cpp
@@ -141,8 +141,6 @@ public:
 			}
 		}
 
-		bUseEditorDepthTest = !bInManipulation;
-
 		// Get a color for property coloration.
 		FColor NewPropertyColor;
 		GEngine->GetPropertyColorationColor( (UObject*)Component, NewPropertyColor );
@@ -253,7 +251,7 @@ public:
 							Collector.RegisterOneFrameMaterialProxy(SolidMaterialInstance);
 
 							FTransform GeomTransform(GetLocalToWorld());
-							BodySetup->AggGeom.GetAggGeom(GeomTransform, DrawColor.ToFColor(true), /*Material=*/SolidMaterialInstance, false, /*bSolid=*/ true, UseEditorDepthTest(), ViewIndex, Collector);
+							BodySetup->AggGeom.GetAggGeom(GeomTransform, DrawColor.ToFColor(true), /*Material=*/SolidMaterialInstance, false, /*bSolid=*/ true, ViewIndex, Collector);
 						}
 					}
 					// WIREFRAME
@@ -290,7 +288,7 @@ public:
 							// If not, use the body setup for wireframe
 						{
 							FTransform GeomTransform(GetLocalToWorld());
-							BodySetup->AggGeom.GetAggGeom(GeomTransform, GetSelectionColor(DrawColor, IsSelected(), IsHovered()).ToFColor(true), /* Material=*/ NULL, false, /* bSolid=*/ false, UseEditorDepthTest(), ViewIndex, Collector);
+							BodySetup->AggGeom.GetAggGeom(GeomTransform, GetSelectionColor(DrawColor, IsSelected(), IsHovered()).ToFColor(true), /* Material=*/ NULL, false, /* bSolid=*/ false, ViewIndex, Collector);
 						}
 
 					}
diff --git a/Engine/Source/Runtime/Engine/Private/DynamicMeshBuilder.cpp b/Engine/Source/Runtime/Engine/Private/DynamicMeshBuilder.cpp
--- a/Engine/Source/Runtime/Engine/Private/DynamicMeshBuilder.cpp
+++ b/Engine/Source/Runtime/Engine/Private/DynamicMeshBuilder.cpp
@@ -732,7 +732,6 @@ void FDynamicMeshBuilder::GetMesh(const FMatrix& LocalToWorld,const FMaterialRen
 				false,
 				false,
 				false,
-				false,
 				GetDefaultLightingChannelMask(),
 				0,
 				INDEX_NONE,
@@ -822,7 +821,6 @@ void FDynamicMeshBuilder::GetMeshElement(const FMatrix& LocalToWorld, const FMat
 				false,
 				false,
 				false,
-				false,
 				GetDefaultLightingChannelMask(),
 				0,
 				INDEX_NONE,
@@ -899,7 +897,6 @@ void FDynamicMeshBuilder::Draw(FPrimitiveDrawInterface* PDI,const FMatrix& Local
 			false,
 			false,
 			false,
-			false,
 			GetDefaultLightingChannelMask(),
 			0,
 			INDEX_NONE,
diff --git a/Engine/Source/Runtime/Engine/Private/Particles/ParticleSystemRender.cpp b/Engine/Source/Runtime/Engine/Private/Particles/ParticleSystemRender.cpp
--- a/Engine/Source/Runtime/Engine/Private/Particles/ParticleSystemRender.cpp
+++ b/Engine/Source/Runtime/Engine/Private/Particles/ParticleSystemRender.cpp
@@ -7271,7 +7271,6 @@ void FParticleSystemSceneProxy::UpdateWorldSpacePrimitiveUniformBuffer() const
 			false,
 			UseSingleSampleShadowFromStationaryLights(),
 			GetScene().HasPrecomputedVolumetricLightmap_RenderThread(),
-			UseEditorDepthTest(),
 			GetLightingChannelMask(),
 			0,
 			INDEX_NONE,
diff --git a/Engine/Source/Runtime/Engine/Private/PhysicsEngine/PhysDrawing.cpp b/Engine/Source/Runtime/Engine/Private/PhysicsEngine/PhysDrawing.cpp
--- a/Engine/Source/Runtime/Engine/Private/PhysicsEngine/PhysDrawing.cpp
+++ b/Engine/Source/Runtime/Engine/Private/PhysicsEngine/PhysDrawing.cpp
@@ -800,7 +800,7 @@ bool FKConvexGeomRenderInfo::HasValidGeometry()
 // FKAggregateGeom
 /////////////////////////////////////////////////////////////////////////////////////
 
-void FKAggregateGeom::GetAggGeom(const FTransform& Transform, const FColor Color, const FMaterialRenderProxy* MatInst, bool bPerHullColor, bool bDrawSolid, bool bUseEditorDepthTest, int32 ViewIndex, FMeshElementCollector& Collector) const
+void FKAggregateGeom::GetAggGeom(const FTransform& Transform, const FColor Color, const FMaterialRenderProxy* MatInst, bool bPerHullColor, bool bDrawSolid, int32 ViewIndex, FMeshElementCollector& Collector) const
 {
 	const FVector Scale3D = Transform.GetScale3D();
 	FTransform ParentTM = Transform;
@@ -894,7 +894,7 @@ void FKAggregateGeom::GetAggGeom(const FTransform& Transform, const FColor Color
 				CalcBoxSphereBounds(LocalBounds, FTransform::Identity);
 
 				FDynamicPrimitiveUniformBuffer& DynamicPrimitiveUniformBuffer = Collector.AllocateOneFrameResource<FDynamicPrimitiveUniformBuffer>();
-				DynamicPrimitiveUniformBuffer.Set(LocalToWorld.ToMatrixWithScale(), LocalToWorld.ToMatrixWithScale(), WorldBounds, LocalBounds, true, false, bUseEditorDepthTest);
+				DynamicPrimitiveUniformBuffer.Set(LocalToWorld.ToMatrixWithScale(), LocalToWorld.ToMatrixWithScale(), WorldBounds, LocalBounds, true, false);
 				BatchElement.PrimitiveUniformBufferResource = &DynamicPrimitiveUniformBuffer.UniformBuffer;
 
 			 	// previous l2w not used so treat as static
@@ -1004,7 +1004,7 @@ void UPhysicsAsset::GetCollisionMesh(int32 ViewIndex, FMeshElementCollector& Col
 		// BoneTransform.SetScale3D(Scale3D);
 		if (SkeletalBodySetups[i]->bCreatedPhysicsMeshes)
 		{
-			SkeletalBodySetups[i]->AggGeom.GetAggGeom(BoneTransform, *BoneColor, NULL, false, false, false, ViewIndex, Collector);
+			SkeletalBodySetups[i]->AggGeom.GetAggGeom(BoneTransform, *BoneColor, NULL, false, false, ViewIndex, Collector);
 		}
 	}
 }
diff --git a/Engine/Source/Runtime/Engine/Private/PrimitiveSceneProxy.cpp b/Engine/Source/Runtime/Engine/Private/PrimitiveSceneProxy.cpp
--- a/Engine/Source/Runtime/Engine/Private/PrimitiveSceneProxy.cpp
+++ b/Engine/Source/Runtime/Engine/Private/PrimitiveSceneProxy.cpp
@@ -120,7 +120,6 @@ FPrimitiveSceneProxy::FPrimitiveSceneProxy(const UPrimitiveComponent* InComponen
 ,	bStaticElementsAlwaysUseProxyPrimitiveUniformBuffer(false)
 ,	bVFRequiresPrimitiveUniformBuffer(true)
 ,	bAlwaysHasVelocity(false)
-,	bUseEditorDepthTest(true)
 ,	bSupportsDistanceFieldRepresentation(false)
 ,	bSupportsHeightfieldRepresentation(false)
 ,	bNeedsLevelAddedToWorldNotification(false)
@@ -283,7 +282,6 @@ void FPrimitiveSceneProxy::UpdateUniformBuffer()
 				HasDynamicIndirectShadowCasterRepresentation(), 
 				UseSingleSampleShadowFromStationaryLights(),
 				bHasPrecomputedVolumetricLightmap,
-				UseEditorDepthTest(), 
 				GetLightingChannelMask(),
 				LpvBiasMultiplier,
 				PrimitiveSceneInfo ? PrimitiveSceneInfo->GetLightmapDataOffset() : 0,
diff --git a/Engine/Source/Runtime/Engine/Private/PrimitiveUniformShaderParameters.cpp b/Engine/Source/Runtime/Engine/Private/PrimitiveUniformShaderParameters.cpp
--- a/Engine/Source/Runtime/Engine/Private/PrimitiveUniformShaderParameters.cpp
+++ b/Engine/Source/Runtime/Engine/Private/PrimitiveUniformShaderParameters.cpp
@@ -55,7 +55,6 @@ FPrimitiveSceneShaderData::FPrimitiveSceneShaderData(const FPrimitiveSceneProxy*
 		Proxy->HasDynamicIndirectShadowCasterRepresentation(), 
 		Proxy->UseSingleSampleShadowFromStationaryLights(),
 		bHasPrecomputedVolumetricLightmap,
-		Proxy->UseEditorDepthTest(), 
 		Proxy->GetLightingChannelMask(),
 		Proxy->GetLpvBiasMultiplier(),
 		Proxy->GetPrimitiveSceneInfo()->GetLightmapDataOffset(),
@@ -95,7 +94,7 @@ void FPrimitiveSceneShaderData::Setup(const FPrimitiveUniformShaderParameters& P
 		PrimitiveUniformShaderParameters.DecalReceiverMask, 
 		PrimitiveUniformShaderParameters.PerObjectGBufferData, 
 		PrimitiveUniformShaderParameters.UseVolumetricLightmapShadowFromStationaryLights, 
-		PrimitiveUniformShaderParameters.UseEditorDepthTest);
+		1.0f);
 	Data[21] = PrimitiveUniformShaderParameters.ObjectOrientation;
 	Data[22] = PrimitiveUniformShaderParameters.NonUniformScale;
 
diff --git a/Engine/Source/Runtime/Engine/Private/StaticMeshRender.cpp b/Engine/Source/Runtime/Engine/Private/StaticMeshRender.cpp
--- a/Engine/Source/Runtime/Engine/Private/StaticMeshRender.cpp
+++ b/Engine/Source/Runtime/Engine/Private/StaticMeshRender.cpp
@@ -1390,13 +1390,13 @@ void FStaticMeshSceneProxy::GetDynamicMeshElements(const TArray<const FSceneView
 						Collector.RegisterOneFrameMaterialProxy(SolidMaterialInstance);
 
 						FTransform GeomTransform(GetLocalToWorld());
-						BodySetup->AggGeom.GetAggGeom(GeomTransform, GetWireframeColor().ToFColor(true), SolidMaterialInstance, false, true, UseEditorDepthTest(), ViewIndex, Collector);
+						BodySetup->AggGeom.GetAggGeom(GeomTransform, GetWireframeColor().ToFColor(true), SolidMaterialInstance, false, true, ViewIndex, Collector);
 					}
 					// wireframe
 					else
 					{
 						FTransform GeomTransform(GetLocalToWorld());
-						BodySetup->AggGeom.GetAggGeom(GeomTransform, GetSelectionColor(SimpleCollisionColor, bProxyIsSelected, IsHovered()).ToFColor(true), NULL, ( Owner == NULL ), false, UseEditorDepthTest(), ViewIndex, Collector);
+						BodySetup->AggGeom.GetAggGeom(GeomTransform, GetSelectionColor(SimpleCollisionColor, bProxyIsSelected, IsHovered()).ToFColor(true), NULL, ( Owner == NULL ), false, ViewIndex, Collector);
 					}
 
 
diff --git a/Engine/Source/Runtime/Engine/Public/PrimitiveSceneProxy.h b/Engine/Source/Runtime/Engine/Public/PrimitiveSceneProxy.h
--- a/Engine/Source/Runtime/Engine/Public/PrimitiveSceneProxy.h
+++ b/Engine/Source/Runtime/Engine/Public/PrimitiveSceneProxy.h
@@ -514,7 +514,6 @@ public:
 	inline bool WillEverBeLit() const { return bWillEverBeLit; }
 	inline bool HasValidSettingsForStaticLighting() const { return bHasValidSettingsForStaticLighting; }
 	inline bool AlwaysHasVelocity() const { return bAlwaysHasVelocity; }
-	inline bool UseEditorDepthTest() const { return bUseEditorDepthTest; }
 	inline bool SupportsDistanceFieldRepresentation() const { return bSupportsDistanceFieldRepresentation; }
 	inline bool SupportsHeightfieldRepresentation() const { return bSupportsHeightfieldRepresentation; }
 	inline bool TreatAsBackgroundForOcclusion() const { return bTreatAsBackgroundForOcclusion; }
@@ -897,9 +896,6 @@ protected:
 	/** Whether the primitive should always be considered to have velocities, even if it hasn't moved. */
 	uint8 bAlwaysHasVelocity : 1;
 
-	/** Whether editor compositing depth testing should be used for this primitive.  Only matters for primitives with bUseEditorCompositing. */
-	uint8 bUseEditorDepthTest : 1;
-
 	/** Whether the primitive type supports a distance field representation.  Does not mean the primitive has a valid representation. */
 	uint8 bSupportsDistanceFieldRepresentation : 1;
 
diff --git a/Engine/Source/Runtime/Engine/Public/SceneManagement.h b/Engine/Source/Runtime/Engine/Public/SceneManagement.h
--- a/Engine/Source/Runtime/Engine/Public/SceneManagement.h
+++ b/Engine/Source/Runtime/Engine/Public/SceneManagement.h
@@ -2203,8 +2203,7 @@ public:
 		const FBoxSphereBounds& LocalBounds,
 		const FBoxSphereBounds& PreSkinnedLocalBounds,
 		bool bReceivesDecals,
-		bool bHasPrecomputedVolumetricLightmap,
-		bool bUseEditorDepthTest);
+		bool bHasPrecomputedVolumetricLightmap);
 
 	/** Pass-through implementation which calls the overloaded Set function with LocalBounds for PreSkinnedLocalBounds. */
 	ENGINE_API void Set(
@@ -2213,8 +2212,7 @@ public:
 		const FBoxSphereBounds& WorldBounds,
 		const FBoxSphereBounds& LocalBounds,
 		bool bReceivesDecals,
-		bool bHasPrecomputedVolumetricLightmap,
-		bool bUseEditorDepthTest);
+		bool bHasPrecomputedVolumetricLightmap);
 };
 
 //
diff --git a/Engine/Source/Runtime/Experimental/GeometryCollectionEngine/Private/GeometryCollection/GeometryCollectionSceneProxy.cpp b/Engine/Source/Runtime/Experimental/GeometryCollectionEngine/Private/GeometryCollection/GeometryCollectionSceneProxy.cpp
--- a/Engine/Source/Runtime/Experimental/GeometryCollectionEngine/Private/GeometryCollection/GeometryCollectionSceneProxy.cpp
+++ b/Engine/Source/Runtime/Experimental/GeometryCollectionEngine/Private/GeometryCollection/GeometryCollectionSceneProxy.cpp
@@ -286,7 +286,7 @@ void FGeometryCollectionSceneProxy::GetDynamicMeshElements(const TArray<const FS
                     
                     // Collector owns the uniform buffer
                     FDynamicPrimitiveUniformBuffer& DynamicPrimitiveUniformBuffer = Collector.AllocateOneFrameResource<FDynamicPrimitiveUniformBuffer>();
-                    DynamicPrimitiveUniformBuffer.Set(GetLocalToWorld(), GetLocalToWorld(), GetBounds(), GetLocalBounds(), true, false, UseEditorDepthTest());
+                    DynamicPrimitiveUniformBuffer.Set(GetLocalToWorld(), GetLocalToWorld(), GetBounds(), GetLocalBounds(), true, false);
                     BatchElement.PrimitiveUniformBuffer = DynamicPrimitiveUniformBuffer.UniformBuffer.GetUniformBufferRHI();
                     
 					BatchElement.FirstIndex = Section.FirstIndex;
@@ -318,7 +318,7 @@ void FGeometryCollectionSceneProxy::GetDynamicMeshElements(const TArray<const FS
 					Mesh.MaterialRenderProxy = MaterialRenderProxy;
                     
                     FDynamicPrimitiveUniformBuffer& DynamicPrimitiveUniformBuffer = Collector.AllocateOneFrameResource<FDynamicPrimitiveUniformBuffer>();
-                    DynamicPrimitiveUniformBuffer.Set(GetLocalToWorld(), GetLocalToWorld(), GetBounds(), GetLocalBounds(), true, false, UseEditorDepthTest());
+                    DynamicPrimitiveUniformBuffer.Set(GetLocalToWorld(), GetLocalToWorld(), GetBounds(), GetLocalBounds(), true, false);
                     BatchElement.PrimitiveUniformBuffer = DynamicPrimitiveUniformBuffer.UniformBuffer.GetUniformBufferRHI();
                     
 					BatchElement.FirstIndex = 0;
diff --git a/Engine/Source/Runtime/MRMesh/Private/MRMeshComponent.cpp b/Engine/Source/Runtime/MRMesh/Private/MRMeshComponent.cpp
--- a/Engine/Source/Runtime/MRMesh/Private/MRMeshComponent.cpp
+++ b/Engine/Source/Runtime/MRMesh/Private/MRMeshComponent.cpp
@@ -325,7 +325,7 @@ private:
 						Mesh.MaterialRenderProxy = MaterialProxy;
 
 						FDynamicPrimitiveUniformBuffer& DynamicPrimitiveUniformBuffer = Collector.AllocateOneFrameResource<FDynamicPrimitiveUniformBuffer>();
-						DynamicPrimitiveUniformBuffer.Set(GetLocalToWorld(), GetLocalToWorld(), InfiniteBounds, InfiniteBounds, true, false, UseEditorDepthTest());
+						DynamicPrimitiveUniformBuffer.Set(GetLocalToWorld(), GetLocalToWorld(), InfiniteBounds, InfiniteBounds, true, false);
 						BatchElement.PrimitiveUniformBufferResource = &DynamicPrimitiveUniformBuffer.UniformBuffer;
 
 						BatchElement.FirstIndex = 0;
diff --git a/Engine/Source/Runtime/NavigationSystem/Private/NavMesh/NavMeshRenderingComponent.cpp b/Engine/Source/Runtime/NavigationSystem/Private/NavMesh/NavMeshRenderingComponent.cpp
--- a/Engine/Source/Runtime/NavigationSystem/Private/NavMesh/NavMeshRenderingComponent.cpp
+++ b/Engine/Source/Runtime/NavigationSystem/Private/NavMesh/NavMeshRenderingComponent.cpp
@@ -952,7 +952,7 @@ void FNavMeshSceneProxy::GetDynamicMeshElements(const TArray<const FSceneView*>&
 					BatchElement = MeshBatchElements[Index];
 
 					FDynamicPrimitiveUniformBuffer& DynamicPrimitiveUniformBuffer = Collector.AllocateOneFrameResource<FDynamicPrimitiveUniformBuffer>();
-					DynamicPrimitiveUniformBuffer.Set(FMatrix::Identity, FMatrix::Identity, GetBounds(), GetLocalBounds(), false, false, UseEditorDepthTest());
+					DynamicPrimitiveUniformBuffer.Set(FMatrix::Identity, FMatrix::Identity, GetBounds(), GetLocalBounds(), false, false);
 					BatchElement.PrimitiveUniformBufferResource = &DynamicPrimitiveUniformBuffer.UniformBuffer;
 
 					Mesh.bWireframe = false;
diff --git a/Engine/Source/Runtime/UMG/Private/Components/WidgetComponent.cpp b/Engine/Source/Runtime/UMG/Private/Components/WidgetComponent.cpp
--- a/Engine/Source/Runtime/UMG/Private/Components/WidgetComponent.cpp
+++ b/Engine/Source/Runtime/UMG/Private/Components/WidgetComponent.cpp
@@ -487,14 +487,14 @@ public:
 						Collector.RegisterOneFrameMaterialProxy(SolidMaterialInstance);
 
 						FTransform GeomTransform(GetLocalToWorld());
-						InBodySetup->AggGeom.GetAggGeom(GeomTransform, GetWireframeColor().ToFColor(true), SolidMaterialInstance, false, true, UseEditorDepthTest(), ViewIndex, Collector);
+						InBodySetup->AggGeom.GetAggGeom(GeomTransform, GetWireframeColor().ToFColor(true), SolidMaterialInstance, false, true, ViewIndex, Collector);
 					}
 					// wireframe
 					else
 					{
 						FColor CollisionColor = FColor(157, 149, 223, 255);
 						FTransform GeomTransform(GetLocalToWorld());
-						InBodySetup->AggGeom.GetAggGeom(GeomTransform, GetSelectionColor(CollisionColor, bProxyIsSelected, IsHovered()).ToFColor(true), nullptr, false, false, UseEditorDepthTest(), ViewIndex, Collector);
+						InBodySetup->AggGeom.GetAggGeom(GeomTransform, GetSelectionColor(CollisionColor, bProxyIsSelected, IsHovered()).ToFColor(true), nullptr, false, false, ViewIndex, Collector);
 					}
 				}
 			}
