Oups... forget this guys...

diff --git a/Engine/Source/Runtime/Engine/Private/SceneManagement.cpp b/Engine/Source/Runtime/Engine/Private/SceneManagement.cpp
--- a/Engine/Source/Runtime/Engine/Private/SceneManagement.cpp
+++ b/Engine/Source/Runtime/Engine/Private/SceneManagement.cpp
@@ -307,12 +307,11 @@ void FDynamicPrimitiveUniformBuffer::Set(
 	const FBoxSphereBounds& LocalBounds,
 	const FBoxSphereBounds& PreSkinnedLocalBounds,
 	bool bReceivesDecals,
-	bool bHasPrecomputedVolumetricLightmap,
-	bool bUseEditorDepthTest)
+	bool bHasPrecomputedVolumetricLightmap)
 {
 	check(IsInRenderingThread());
 	UniformBuffer.SetContents(
-		GetPrimitiveUniformShaderParameters(LocalToWorld, PreviousLocalToWorld, WorldBounds.Origin, WorldBounds, LocalBounds, PreSkinnedLocalBounds, bReceivesDecals, false, false, false, bHasPrecomputedVolumetricLightmap, bUseEditorDepthTest, GetDefaultLightingChannelMask(), 1.0f, INDEX_NONE, INDEX_NONE));
+		GetPrimitiveUniformShaderParameters(LocalToWorld, PreviousLocalToWorld, WorldBounds.Origin, WorldBounds, LocalBounds, PreSkinnedLocalBounds, bReceivesDecals, false, false, false, bHasPrecomputedVolumetricLightmap, GetDefaultLightingChannelMask(), 1.0f, INDEX_NONE, INDEX_NONE));
 	UniformBuffer.InitResource();
 }
 
@@ -322,10 +321,9 @@ void FDynamicPrimitiveUniformBuffer::Set(
 	const FBoxSphereBounds& WorldBounds,
 	const FBoxSphereBounds& LocalBounds,
 	bool bReceivesDecals,
-	bool bHasPrecomputedVolumetricLightmap,
-	bool bUseEditorDepthTest)
+	bool bHasPrecomputedVolumetricLightmap)
 {
-	Set(LocalToWorld, PreviousLocalToWorld, WorldBounds, LocalBounds, LocalBounds, bReceivesDecals, bHasPrecomputedVolumetricLightmap, bUseEditorDepthTest);
+	Set(LocalToWorld, PreviousLocalToWorld, WorldBounds, LocalBounds, LocalBounds, bReceivesDecals, bHasPrecomputedVolumetricLightmap);
 }
 
 FLightMapInteraction FLightMapInteraction::Texture(
diff --git a/Engine/Source/Runtime/Engine/Public/PrimitiveUniformShaderParameters.h b/Engine/Source/Runtime/Engine/Public/PrimitiveUniformShaderParameters.h
--- a/Engine/Source/Runtime/Engine/Public/PrimitiveUniformShaderParameters.h
+++ b/Engine/Source/Runtime/Engine/Public/PrimitiveUniformShaderParameters.h
@@ -27,8 +27,8 @@ BEGIN_GLOBAL_SHADER_PARAMETER_STRUCT(FPrimitiveUniformShaderParameters,ENGINE_AP
 	SHADER_PARAMETER(float,LpvBiasMultiplier)
 	SHADER_PARAMETER_EX(float,DecalReceiverMask,EShaderPrecisionModifier::Half)
 	SHADER_PARAMETER_EX(float,PerObjectGBufferData,EShaderPrecisionModifier::Half)		// 0..1, 2 bits, bDistanceFieldRepresentation, bHeightfieldRepresentation
-	SHADER_PARAMETER_EX(float,UseVolumetricLightmapShadowFromStationaryLights,EShaderPrecisionModifier::Half)		
-	SHADER_PARAMETER_EX(float,UseEditorDepthTest,EShaderPrecisionModifier::Half)
+	SHADER_PARAMETER_EX(float,UseVolumetricLightmapShadowFromStationaryLights,EShaderPrecisionModifier::Half)	
+	SHADER_PARAMETER_EX(float,Unused,EShaderPrecisionModifier::Half)
 	SHADER_PARAMETER_EX(FVector4,ObjectOrientation,EShaderPrecisionModifier::Half)
 	SHADER_PARAMETER_EX(FVector4,NonUniformScale,EShaderPrecisionModifier::Half)
 	SHADER_PARAMETER(FVector, LocalObjectBoundsMin)		// This is used in a custom material function (ObjectLocalBounds.uasset)
@@ -52,7 +52,6 @@ inline FPrimitiveUniformShaderParameters GetPrimitiveUniformShaderParameters(
 	bool bHasCapsuleRepresentation,
 	bool bUseSingleSampleShadowFromStationaryLights,
 	bool bUseVolumetricLightmap,
-	bool bUseEditorDepthTest,
 	uint32 LightingChannelMask,
 	float LpvBiasMultiplier,
 	uint32 LightmapDataIndex,
@@ -94,7 +93,7 @@ inline FPrimitiveUniformShaderParameters GetPrimitiveUniformShaderParameters(
 	Result.PerObjectGBufferData = (2 * (int32)bHasCapsuleRepresentation + (int32)bHasDistanceFieldRepresentation) / 3.0f;
 	Result.UseSingleSampleShadowFromStationaryLights = bUseSingleSampleShadowFromStationaryLights ? 1.0f : 0.0f;
 	Result.UseVolumetricLightmapShadowFromStationaryLights = bUseVolumetricLightmap && bUseSingleSampleShadowFromStationaryLights ? 1.0f : 0.0f;
-	Result.UseEditorDepthTest = bUseEditorDepthTest ? 1 : 0;
+	Result.Unused = 0;
 	Result.LightmapDataIndex = LightmapDataIndex;
 	Result.SingleCaptureIndex = SingleCaptureIndex;
 	return Result;
@@ -112,7 +111,6 @@ inline FPrimitiveUniformShaderParameters GetPrimitiveUniformShaderParameters(
 	bool bHasCapsuleRepresentation,
 	bool bUseSingleSampleShadowFromStationaryLights,
 	bool bUseVolumetricLightmap,
-	bool bUseEditorDepthTest,
 	uint32 LightingChannelMask,
 	float LpvBiasMultiplier,
 	uint32 LightmapDataIndex,
@@ -121,7 +119,7 @@ inline FPrimitiveUniformShaderParameters GetPrimitiveUniformShaderParameters(
 {
 	// Pass through call
 	return GetPrimitiveUniformShaderParameters(LocalToWorld, PreviousLocalToWorld, ActorPosition, WorldBounds, LocalBounds, LocalBounds, bReceivesDecals, bHasDistanceFieldRepresentation, bHasCapsuleRepresentation, 
-		bUseSingleSampleShadowFromStationaryLights, bUseVolumetricLightmap, bUseEditorDepthTest, LightingChannelMask, LpvBiasMultiplier, LightmapDataIndex, SingleCaptureIndex);
+		bUseSingleSampleShadowFromStationaryLights, bUseVolumetricLightmap, LightingChannelMask, LpvBiasMultiplier, LightmapDataIndex, SingleCaptureIndex);
 }
 
 inline TUniformBufferRef<FPrimitiveUniformShaderParameters> CreatePrimitiveUniformBufferImmediate(
@@ -130,13 +128,12 @@ inline TUniformBufferRef<FPrimitiveUniformShaderParameters> CreatePrimitiveUnifo
 	const FBoxSphereBounds& LocalBounds,
 	const FBoxSphereBounds& PreSkinnedLocalBounds,
 	bool bReceivesDecals,
-	bool bUseEditorDepthTest,
 	float LpvBiasMultiplier = 1.0f
 	)
 {
 	check(IsInRenderingThread());
 	return TUniformBufferRef<FPrimitiveUniformShaderParameters>::CreateUniformBufferImmediate(
-		GetPrimitiveUniformShaderParameters(LocalToWorld, LocalToWorld, WorldBounds.Origin, WorldBounds, LocalBounds, PreSkinnedLocalBounds, bReceivesDecals, false, false, false, false, bUseEditorDepthTest, GetDefaultLightingChannelMask(), LpvBiasMultiplier, INDEX_NONE, INDEX_NONE),
+		GetPrimitiveUniformShaderParameters(LocalToWorld, LocalToWorld, WorldBounds.Origin, WorldBounds, LocalBounds, PreSkinnedLocalBounds, bReceivesDecals, false, false, false, false, GetDefaultLightingChannelMask(), LpvBiasMultiplier, INDEX_NONE, INDEX_NONE),
 		UniformBuffer_MultiFrame
 		);
 }
@@ -156,7 +153,6 @@ inline FPrimitiveUniformShaderParameters GetIdentityPrimitiveParameters()
 		false,
 		false,
 		false,
-		true,
 		GetDefaultLightingChannelMask(),
 		1.0f,		// LPV bias
 		INDEX_NONE,
