Fixes TAA's dynamic antighosting heuristic when GBuffer outputs velocity.

The source cause of the bug was all geometry was drawing velocity, instead of just movable objects, meaning TAA could no longer make the difference between a static and dynamic object in the scene, increasing ghosting of dynamic object over static object.

diff --git a/Engine/Plugins/2D/Paper2D/Source/Paper2D/Private/PaperRenderSceneProxy.cpp b/Engine/Plugins/2D/Paper2D/Source/Paper2D/Private/PaperRenderSceneProxy.cpp
--- a/Engine/Plugins/2D/Paper2D/Source/Paper2D/Private/PaperRenderSceneProxy.cpp
+++ b/Engine/Plugins/2D/Paper2D/Source/Paper2D/Private/PaperRenderSceneProxy.cpp
@@ -175,12 +175,12 @@ void FPaperRenderSceneProxy::DebugDrawBodySetup(const FSceneView* View, int32 Vi
 
 			Collector.RegisterOneFrameMaterialProxy(SolidMaterialInstance);
 
-			BodySetup->AggGeom.GetAggGeom(GeomTransform, GetWireframeColor().ToFColor(true), SolidMaterialInstance, false, true, ViewIndex, Collector);
+			BodySetup->AggGeom.GetAggGeom(GeomTransform, GetWireframeColor().ToFColor(true), SolidMaterialInstance, false, true, DrawsVelocity(), ViewIndex, Collector);
 		}
 		else
 		{
 			// wireframe
-			BodySetup->AggGeom.GetAggGeom(GeomTransform, GetSelectionColor(CollisionColor, IsSelected(), IsHovered()).ToFColor(true), nullptr, ( Owner == nullptr ), false, ViewIndex, Collector);
+			BodySetup->AggGeom.GetAggGeom(GeomTransform, GetSelectionColor(CollisionColor, IsSelected(), IsHovered()).ToFColor(true), nullptr, ( Owner == nullptr ), false, DrawsVelocity(), ViewIndex, Collector);
 		}
 	}
 }
diff --git a/Engine/Plugins/2D/Paper2D/Source/Paper2D/Private/PaperSpriteSceneProxy.cpp b/Engine/Plugins/2D/Paper2D/Source/Paper2D/Private/PaperSpriteSceneProxy.cpp
--- a/Engine/Plugins/2D/Paper2D/Source/Paper2D/Private/PaperSpriteSceneProxy.cpp
+++ b/Engine/Plugins/2D/Paper2D/Source/Paper2D/Private/PaperSpriteSceneProxy.cpp
@@ -86,7 +86,7 @@ void FPaperSpriteSceneProxy::GetDynamicMeshElements(const TArray<const FSceneVie
 
 						const bool bUseSeparateColorPerHull = (Owner == nullptr);
 						const bool bDrawSolid = false;
-						BodySetup->AggGeom.GetAggGeom(GeomTransform, GetSelectionColor(CollisionColor, bDrawWireSelected, IsHovered()).ToFColor(true), CollisionMaterialInstance, bUseSeparateColorPerHull, bDrawSolid, ViewIndex, Collector);
+						BodySetup->AggGeom.GetAggGeom(GeomTransform, GetSelectionColor(CollisionColor, bDrawWireSelected, IsHovered()).ToFColor(true), CollisionMaterialInstance, bUseSeparateColorPerHull, bDrawSolid, DrawsVelocity(), ViewIndex, Collector);
 					}
 				}
 			}
diff --git a/Engine/Plugins/2D/Paper2D/Source/Paper2D/Private/PaperTileMapRenderSceneProxy.cpp b/Engine/Plugins/2D/Paper2D/Source/Paper2D/Private/PaperTileMapRenderSceneProxy.cpp
--- a/Engine/Plugins/2D/Paper2D/Source/Paper2D/Private/PaperTileMapRenderSceneProxy.cpp
+++ b/Engine/Plugins/2D/Paper2D/Source/Paper2D/Private/PaperTileMapRenderSceneProxy.cpp
@@ -321,7 +321,7 @@ void FPaperTileMapRenderSceneProxy::GetDynamicMeshElements(const TArray<const FS
 
 							const bool bPerHullColor = false;
 							const bool bDrawSimpleSolid = false;
-							BodySetup->AggGeom.GetAggGeom(GeomTransform, GetSelectionColor(CollisionColor, bDrawWireSelected, IsHovered()).ToFColor(true), CollisionMaterialInstance, bPerHullColor, bDrawSimpleSolid, ViewIndex, Collector);
+							BodySetup->AggGeom.GetAggGeom(GeomTransform, GetSelectionColor(CollisionColor, bDrawWireSelected, IsHovered()).ToFColor(true), CollisionMaterialInstance, bPerHullColor, bDrawSimpleSolid, DrawsVelocity(), ViewIndex, Collector);
 						}
 					}
 				}
diff --git a/Engine/Plugins/Editor/MeshEditor/Source/MeshEditor/OverlayComponent.cpp b/Engine/Plugins/Editor/MeshEditor/Source/MeshEditor/OverlayComponent.cpp
--- a/Engine/Plugins/Editor/MeshEditor/Source/MeshEditor/OverlayComponent.cpp
+++ b/Engine/Plugins/Editor/MeshEditor/Source/MeshEditor/OverlayComponent.cpp
@@ -253,7 +253,7 @@ public:
 					Mesh.MaterialRenderProxy = MeshBatchData.MaterialProxy;
 
 					FDynamicPrimitiveUniformBuffer& DynamicPrimitiveUniformBuffer = Collector.AllocateOneFrameResource<FDynamicPrimitiveUniformBuffer>();
-					DynamicPrimitiveUniformBuffer.Set(GetLocalToWorld(), GetLocalToWorld(), GetBounds(), GetLocalBounds(), true, false);
+					DynamicPrimitiveUniformBuffer.Set(GetLocalToWorld(), GetLocalToWorld(), GetBounds(), GetLocalBounds(), true, false, DrawsVelocity());
 					BatchElement.PrimitiveUniformBufferResource = &DynamicPrimitiveUniformBuffer.UniformBuffer;
 
 					BatchElement.FirstIndex = MeshBatchData.StartIndex;
diff --git a/Engine/Plugins/Editor/MeshEditor/Source/MeshEditor/WireframeMeshComponent.cpp b/Engine/Plugins/Editor/MeshEditor/Source/MeshEditor/WireframeMeshComponent.cpp
--- a/Engine/Plugins/Editor/MeshEditor/Source/MeshEditor/WireframeMeshComponent.cpp
+++ b/Engine/Plugins/Editor/MeshEditor/Source/MeshEditor/WireframeMeshComponent.cpp
@@ -344,7 +344,7 @@ public:
 				Mesh.MaterialRenderProxy = MaterialProxy;
 
 				FDynamicPrimitiveUniformBuffer& DynamicPrimitiveUniformBuffer = Collector.AllocateOneFrameResource<FDynamicPrimitiveUniformBuffer>();
-				DynamicPrimitiveUniformBuffer.Set(GetLocalToWorld(), GetLocalToWorld(), GetBounds(), GetLocalBounds(), true, false);
+				DynamicPrimitiveUniformBuffer.Set(GetLocalToWorld(), GetLocalToWorld(), GetBounds(), GetLocalBounds(), true, false, DrawsVelocity());
 				BatchElement.PrimitiveUniformBufferResource = &DynamicPrimitiveUniformBuffer.UniformBuffer;
 
 				BatchElement.FirstIndex = 0;
diff --git a/Engine/Plugins/Experimental/GeometryCache/Source/GeometryCache/Private/GeometryCacheSceneProxy.cpp b/Engine/Plugins/Experimental/GeometryCache/Source/GeometryCache/Private/GeometryCacheSceneProxy.cpp
--- a/Engine/Plugins/Experimental/GeometryCache/Source/GeometryCache/Private/GeometryCacheSceneProxy.cpp
+++ b/Engine/Plugins/Experimental/GeometryCache/Source/GeometryCache/Private/GeometryCacheSceneProxy.cpp
@@ -357,7 +357,7 @@ void FGeometryCacheSceneProxy::CreateMeshBatch(
 
 	const FMatrix& LocalToWorldTransform = TrackProxy->WorldMatrix * GetLocalToWorld();
 
-	DynamicPrimitiveUniformBuffer.Set(LocalToWorldTransform, LocalToWorldTransform, GetBounds(), GetLocalBounds(), true, false);
+	DynamicPrimitiveUniformBuffer.Set(LocalToWorldTransform, LocalToWorldTransform, GetBounds(), GetLocalBounds(), true, false, DrawsVelocity());
 	BatchElement.PrimitiveUniformBuffer = DynamicPrimitiveUniformBuffer.UniformBuffer.GetUniformBufferRHI();
 
 	BatchElement.FirstIndex = BatchInfo.StartIndex;
diff --git a/Engine/Plugins/FX/Niagara/Source/Niagara/Private/NiagaraRendererMeshes.cpp b/Engine/Plugins/FX/Niagara/Source/Niagara/Private/NiagaraRendererMeshes.cpp
--- a/Engine/Plugins/FX/Niagara/Source/Niagara/Private/NiagaraRendererMeshes.cpp
+++ b/Engine/Plugins/FX/Niagara/Source/Niagara/Private/NiagaraRendererMeshes.cpp
@@ -165,6 +165,7 @@ void NiagaraRendererMeshes::GetDynamicMeshElements(const TArray<const FSceneView
 				false,
 				false,
 				false,
+				SceneProxy->DrawsVelocity(),
 				SceneProxy->GetLightingChannelMask(),
 				0,
 				INDEX_NONE,
diff --git a/Engine/Plugins/FX/Niagara/Source/Niagara/Private/NiagaraRendererRibbons.cpp b/Engine/Plugins/FX/Niagara/Source/Niagara/Private/NiagaraRendererRibbons.cpp
--- a/Engine/Plugins/FX/Niagara/Source/Niagara/Private/NiagaraRendererRibbons.cpp
+++ b/Engine/Plugins/FX/Niagara/Source/Niagara/Private/NiagaraRendererRibbons.cpp
@@ -132,6 +132,7 @@ void NiagaraRendererRibbons::GetDynamicMeshElements(const TArray<const FSceneVie
 				false,
 				SceneProxy->UseSingleSampleShadowFromStationaryLights(),
 				SceneProxy->GetScene().HasPrecomputedVolumetricLightmap_RenderThread(),
+				SceneProxy->DrawsVelocity(),
 				SceneProxy->GetLightingChannelMask(),
 				0,
 				INDEX_NONE,
diff --git a/Engine/Plugins/FX/Niagara/Source/Niagara/Private/NiagaraRendererSprites.cpp b/Engine/Plugins/FX/Niagara/Source/Niagara/Private/NiagaraRendererSprites.cpp
--- a/Engine/Plugins/FX/Niagara/Source/Niagara/Private/NiagaraRendererSprites.cpp
+++ b/Engine/Plugins/FX/Niagara/Source/Niagara/Private/NiagaraRendererSprites.cpp
@@ -130,6 +130,7 @@ void NiagaraRendererSprites::ConditionalInitPrimitiveUniformBuffer(const FNiagar
 			false,
 			SceneProxy->UseSingleSampleShadowFromStationaryLights(),
 			SceneProxy->GetScene().HasPrecomputedVolumetricLightmap_RenderThread(),
+			SceneProxy->DrawsVelocity(),
 			SceneProxy->GetLightingChannelMask(),
 			0,
 			INDEX_NONE,
diff --git a/Engine/Plugins/Runtime/CableComponent/Source/CableComponent/Private/CableComponent.cpp b/Engine/Plugins/Runtime/CableComponent/Source/CableComponent/Private/CableComponent.cpp
--- a/Engine/Plugins/Runtime/CableComponent/Source/CableComponent/Private/CableComponent.cpp
+++ b/Engine/Plugins/Runtime/CableComponent/Source/CableComponent/Private/CableComponent.cpp
@@ -288,7 +288,7 @@ public:
 				GetScene().GetPrimitiveUniformShaderParameters_RenderThread(GetPrimitiveSceneInfo(), bHasPrecomputedVolumetricLightmap, PreviousLocalToWorld, SingleCaptureIndex);
 
 				FDynamicPrimitiveUniformBuffer& DynamicPrimitiveUniformBuffer = Collector.AllocateOneFrameResource<FDynamicPrimitiveUniformBuffer>();
-				DynamicPrimitiveUniformBuffer.Set(GetLocalToWorld(), PreviousLocalToWorld, GetBounds(), GetLocalBounds(), true, bHasPrecomputedVolumetricLightmap);
+				DynamicPrimitiveUniformBuffer.Set(GetLocalToWorld(), PreviousLocalToWorld, GetBounds(), GetLocalBounds(), true, bHasPrecomputedVolumetricLightmap, DrawsVelocity());
 				BatchElement.PrimitiveUniformBufferResource = &DynamicPrimitiveUniformBuffer.UniformBuffer;
 
 				BatchElement.FirstIndex = 0;
diff --git a/Engine/Plugins/Runtime/CustomMeshComponent/Source/CustomMeshComponent/Private/CustomMeshComponent.cpp b/Engine/Plugins/Runtime/CustomMeshComponent/Source/CustomMeshComponent/Private/CustomMeshComponent.cpp
--- a/Engine/Plugins/Runtime/CustomMeshComponent/Source/CustomMeshComponent/Private/CustomMeshComponent.cpp
+++ b/Engine/Plugins/Runtime/CustomMeshComponent/Source/CustomMeshComponent/Private/CustomMeshComponent.cpp
@@ -135,7 +135,7 @@ public:
 				GetScene().GetPrimitiveUniformShaderParameters_RenderThread(GetPrimitiveSceneInfo(), bHasPrecomputedVolumetricLightmap, PreviousLocalToWorld, SingleCaptureIndex);
 
 				FDynamicPrimitiveUniformBuffer& DynamicPrimitiveUniformBuffer = Collector.AllocateOneFrameResource<FDynamicPrimitiveUniformBuffer>();
-				DynamicPrimitiveUniformBuffer.Set(GetLocalToWorld(), PreviousLocalToWorld, GetBounds(), GetLocalBounds(), true, bHasPrecomputedVolumetricLightmap);
+				DynamicPrimitiveUniformBuffer.Set(GetLocalToWorld(), PreviousLocalToWorld, GetBounds(), GetLocalBounds(), true, bHasPrecomputedVolumetricLightmap, DrawsVelocity());
 				BatchElement.PrimitiveUniformBufferResource = &DynamicPrimitiveUniformBuffer.UniformBuffer;
 
 				BatchElement.FirstIndex = 0;
diff --git a/Engine/Plugins/Runtime/Oculus/OculusVR/Source/OculusMR/Private/OculusMR_BoundaryMeshComponent.cpp b/Engine/Plugins/Runtime/Oculus/OculusVR/Source/OculusMR/Private/OculusMR_BoundaryMeshComponent.cpp
--- a/Engine/Plugins/Runtime/Oculus/OculusVR/Source/OculusMR/Private/OculusMR_BoundaryMeshComponent.cpp
+++ b/Engine/Plugins/Runtime/Oculus/OculusVR/Source/OculusMR/Private/OculusMR_BoundaryMeshComponent.cpp
@@ -150,7 +150,7 @@ public:
 				//Mesh.bWireframe = bWireframe;
 				//Mesh.VertexFactory = &VertexFactory;
 				//Mesh.MaterialRenderProxy = MaterialProxy;
-				//BatchElement.PrimitiveUniformBuffer = CreatePrimitiveUniformBufferImmediate(Transform.ToMatrixWithScale(), GetBounds(), GetLocalBounds(), true);
+				//BatchElement.PrimitiveUniformBuffer = CreatePrimitiveUniformBufferImmediate(Transform.ToMatrixWithScale(), GetBounds(), GetLocalBounds(), true, DrawsVelocity());
 				//BatchElement.FirstIndex = 0;
 				//BatchElement.NumPrimitives = IndexBuffer.Indices.Num() / 3;
 				//BatchElement.MinVertexIndex = 0;
diff --git a/Engine/Plugins/Runtime/Oculus/OculusVR/Source/OculusMR/Private/OculusMR_PlaneMeshComponent.cpp b/Engine/Plugins/Runtime/Oculus/OculusVR/Source/OculusMR/Private/OculusMR_PlaneMeshComponent.cpp
--- a/Engine/Plugins/Runtime/Oculus/OculusVR/Source/OculusMR/Private/OculusMR_PlaneMeshComponent.cpp
+++ b/Engine/Plugins/Runtime/Oculus/OculusVR/Source/OculusMR/Private/OculusMR_PlaneMeshComponent.cpp
@@ -127,7 +127,7 @@ public:
 					//Mesh.bWireframe = bWireframe;
 					//Mesh.VertexFactory = &VertexFactory;
 					//Mesh.MaterialRenderProxy = MaterialProxy;
-					//BatchElement.PrimitiveUniformBuffer = CreatePrimitiveUniformBufferImmediate(GetLocalToWorld(), GetBounds(), GetLocalBounds(), true);
+					//BatchElement.PrimitiveUniformBuffer = CreatePrimitiveUniformBufferImmediate(GetLocalToWorld(), GetBounds(), GetLocalBounds(), true, DrawsVelocity());
 					//BatchElement.FirstIndex = 0;
 					//BatchElement.NumPrimitives = IndexBuffer.Indices.Num() / 3;
 					//BatchElement.MinVertexIndex = 0;
diff --git a/Engine/Plugins/Runtime/ProceduralMeshComponent/Source/ProceduralMeshComponent/Private/ProceduralMeshComponent.cpp b/Engine/Plugins/Runtime/ProceduralMeshComponent/Source/ProceduralMeshComponent/Private/ProceduralMeshComponent.cpp
--- a/Engine/Plugins/Runtime/ProceduralMeshComponent/Source/ProceduralMeshComponent/Private/ProceduralMeshComponent.cpp
+++ b/Engine/Plugins/Runtime/ProceduralMeshComponent/Source/ProceduralMeshComponent/Private/ProceduralMeshComponent.cpp
@@ -307,7 +307,7 @@ public:
 						GetScene().GetPrimitiveUniformShaderParameters_RenderThread(GetPrimitiveSceneInfo(), bHasPrecomputedVolumetricLightmap, PreviousLocalToWorld, SingleCaptureIndex);
 
 						FDynamicPrimitiveUniformBuffer& DynamicPrimitiveUniformBuffer = Collector.AllocateOneFrameResource<FDynamicPrimitiveUniformBuffer>();
-						DynamicPrimitiveUniformBuffer.Set(GetLocalToWorld(), PreviousLocalToWorld, GetBounds(), GetLocalBounds(), true, bHasPrecomputedVolumetricLightmap);
+						DynamicPrimitiveUniformBuffer.Set(GetLocalToWorld(), PreviousLocalToWorld, GetBounds(), GetLocalBounds(), true, bHasPrecomputedVolumetricLightmap, DrawsVelocity());
 						BatchElement.PrimitiveUniformBufferResource = &DynamicPrimitiveUniformBuffer.UniformBuffer;
 
 						BatchElement.FirstIndex = 0;
@@ -334,7 +334,7 @@ public:
 				if (ViewFamily.EngineShowFlags.Collision && IsCollisionEnabled() && BodySetup->GetCollisionTraceFlag() != ECollisionTraceFlag::CTF_UseComplexAsSimple)
 				{
 					FTransform GeomTransform(GetLocalToWorld());
-					BodySetup->AggGeom.GetAggGeom(GeomTransform, GetSelectionColor(FColor(157, 149, 223, 255), IsSelected(), IsHovered()).ToFColor(true), NULL, false, false, ViewIndex, Collector);
+					BodySetup->AggGeom.GetAggGeom(GeomTransform, GetSelectionColor(FColor(157, 149, 223, 255), IsSelected(), IsHovered()).ToFColor(true), NULL, false, false, DrawsVelocity(), ViewIndex, Collector);
 				}
 
 				// Render bounds
diff --git a/Engine/Shaders/Private/BasePassPixelShader.usf b/Engine/Shaders/Private/BasePassPixelShader.usf
--- a/Engine/Shaders/Private/BasePassPixelShader.usf
+++ b/Engine/Shaders/Private/BasePassPixelShader.usf
@@ -883,7 +883,7 @@ void FPixelShaderInOut_MainPS(
 	#endif
 
 		// Make sure not to touch 0,0 which is clear color
-		GBuffer.Velocity = float4(EncodeVelocityToTexture(Velocity), 0, 0);
+		GBuffer.Velocity = float4(EncodeVelocityToTexture(Velocity) * GetPrimitiveData(MaterialParameters.PrimitiveId).DrawsVelocity, 0, 0);
 	}
 	#else
 		GBuffer.Velocity = 0;
diff --git a/Engine/Shaders/Private/SceneData.ush b/Engine/Shaders/Private/SceneData.ush
--- a/Engine/Shaders/Private/SceneData.ush
+++ b/Engine/Shaders/Private/SceneData.ush
@@ -25,7 +25,7 @@ struct FPrimitiveSceneData
 	float DecalReceiverMask;
 	float PerObjectGBufferData;
 	float UseVolumetricLightmapShadowFromStationaryLights;
-	float Unused;
+	float DrawsVelocity;
 	float4 ObjectOrientation;
 	float4 NonUniformScale;
 	float3 LocalObjectBoundsMin;
@@ -78,7 +78,7 @@ FPrimitiveSceneData GetPrimitiveData(uint PrimitiveId)
 	PrimitiveData.DecalReceiverMask = View.PrimitiveSceneData[PrimitiveBaseOffset + 20].x;
 	PrimitiveData.PerObjectGBufferData = View.PrimitiveSceneData[PrimitiveBaseOffset + 20].y;
 	PrimitiveData.UseVolumetricLightmapShadowFromStationaryLights = View.PrimitiveSceneData[PrimitiveBaseOffset + 20].z;
-	PrimitiveData.Unused = false;
+	PrimitiveData.DrawsVelocity = View.PrimitiveSceneData[PrimitiveBaseOffset + 20].w;
 
 	PrimitiveData.ObjectOrientation = View.PrimitiveSceneData[PrimitiveBaseOffset + 21];
 	PrimitiveData.NonUniformScale = View.PrimitiveSceneData[PrimitiveBaseOffset + 22];
diff --git a/Engine/Source/Runtime/Engine/Classes/PhysicsEngine/AggregateGeom.h b/Engine/Source/Runtime/Engine/Classes/PhysicsEngine/AggregateGeom.h
--- a/Engine/Source/Runtime/Engine/Classes/PhysicsEngine/AggregateGeom.h
+++ b/Engine/Source/Runtime/Engine/Classes/PhysicsEngine/AggregateGeom.h
@@ -76,7 +76,7 @@ struct ENGINE_API FKAggregateGeom
 	void FixupDeprecated(FArchive& Ar);
 #endif
 
-	void GetAggGeom(const FTransform& Transform, const FColor Color, const FMaterialRenderProxy* MatInst, bool bPerHullColor, bool bDrawSolid, int32 ViewIndex, class FMeshElementCollector& Collector) const;
+	void GetAggGeom(const FTransform& Transform, const FColor Color, const FMaterialRenderProxy* MatInst, bool bPerHullColor, bool bDrawSolid, bool bDrawsVelocity, int32 ViewIndex, class FMeshElementCollector& Collector) const;
 
 	/** Release the RenderInfo (if its there) and safely clean up any resources. Call on the game thread. */
 	void FreeRenderInfo();
diff --git a/Engine/Source/Runtime/Engine/Private/Components/ArrowComponent.cpp b/Engine/Source/Runtime/Engine/Private/Components/ArrowComponent.cpp
--- a/Engine/Source/Runtime/Engine/Private/Components/ArrowComponent.cpp
+++ b/Engine/Source/Runtime/Engine/Private/Components/ArrowComponent.cpp
@@ -148,7 +148,7 @@ public:
 				Mesh.MaterialRenderProxy = ArrowMaterialRenderProxy;
 
 				FDynamicPrimitiveUniformBuffer& DynamicPrimitiveUniformBuffer = Collector.AllocateOneFrameResource<FDynamicPrimitiveUniformBuffer>();
-				DynamicPrimitiveUniformBuffer.Set(FScaleMatrix(ViewScale) * EffectiveLocalToWorld, FScaleMatrix(ViewScale) * EffectiveLocalToWorld, GetBounds(), GetLocalBounds(), true, false);
+				DynamicPrimitiveUniformBuffer.Set(FScaleMatrix(ViewScale) * EffectiveLocalToWorld, FScaleMatrix(ViewScale) * EffectiveLocalToWorld, GetBounds(), GetLocalBounds(), true, false, DrawsVelocity());
 				BatchElement.PrimitiveUniformBufferResource = &DynamicPrimitiveUniformBuffer.UniformBuffer;
 
 				BatchElement.FirstIndex = 0;
diff --git a/Engine/Source/Runtime/Engine/Private/Components/BrushComponent.cpp b/Engine/Source/Runtime/Engine/Private/Components/BrushComponent.cpp
--- a/Engine/Source/Runtime/Engine/Private/Components/BrushComponent.cpp
+++ b/Engine/Source/Runtime/Engine/Private/Components/BrushComponent.cpp
@@ -251,7 +251,7 @@ public:
 							Collector.RegisterOneFrameMaterialProxy(SolidMaterialInstance);
 
 							FTransform GeomTransform(GetLocalToWorld());
-							BodySetup->AggGeom.GetAggGeom(GeomTransform, DrawColor.ToFColor(true), /*Material=*/SolidMaterialInstance, false, /*bSolid=*/ true, ViewIndex, Collector);
+							BodySetup->AggGeom.GetAggGeom(GeomTransform, DrawColor.ToFColor(true), /*Material=*/SolidMaterialInstance, false, /*bSolid=*/ true, DrawsVelocity(), ViewIndex, Collector);
 						}
 					}
 					// WIREFRAME
@@ -288,7 +288,7 @@ public:
 							// If not, use the body setup for wireframe
 						{
 							FTransform GeomTransform(GetLocalToWorld());
-							BodySetup->AggGeom.GetAggGeom(GeomTransform, GetSelectionColor(DrawColor, IsSelected(), IsHovered()).ToFColor(true), /* Material=*/ NULL, false, /* bSolid=*/ false, ViewIndex, Collector);
+							BodySetup->AggGeom.GetAggGeom(GeomTransform, GetSelectionColor(DrawColor, IsSelected(), IsHovered()).ToFColor(true), /* Material=*/ NULL, false, /* bSolid=*/ false, DrawsVelocity(), ViewIndex, Collector);
 						}
 
 					}
diff --git a/Engine/Source/Runtime/Engine/Private/DynamicMeshBuilder.cpp b/Engine/Source/Runtime/Engine/Private/DynamicMeshBuilder.cpp
--- a/Engine/Source/Runtime/Engine/Private/DynamicMeshBuilder.cpp
+++ b/Engine/Source/Runtime/Engine/Private/DynamicMeshBuilder.cpp
@@ -732,6 +732,7 @@ void FDynamicMeshBuilder::GetMesh(const FMatrix& LocalToWorld,const FMaterialRen
 				false,
 				false,
 				false,
+				/* bDrawsVelocity = */ true,
 				GetDefaultLightingChannelMask(),
 				0,
 				INDEX_NONE,
@@ -821,6 +822,7 @@ void FDynamicMeshBuilder::GetMeshElement(const FMatrix& LocalToWorld, const FMat
 				false,
 				false,
 				false,
+				/* bDrawsVelocity = */ true,
 				GetDefaultLightingChannelMask(),
 				0,
 				INDEX_NONE,
@@ -897,6 +899,7 @@ void FDynamicMeshBuilder::Draw(FPrimitiveDrawInterface* PDI,const FMatrix& Local
 			false,
 			false,
 			false,
+			/* bDrawsVelocity = */ true,
 			GetDefaultLightingChannelMask(),
 			0,
 			INDEX_NONE,
diff --git a/Engine/Source/Runtime/Engine/Private/Particles/ParticleSystemRender.cpp b/Engine/Source/Runtime/Engine/Private/Particles/ParticleSystemRender.cpp
--- a/Engine/Source/Runtime/Engine/Private/Particles/ParticleSystemRender.cpp
+++ b/Engine/Source/Runtime/Engine/Private/Particles/ParticleSystemRender.cpp
@@ -7271,6 +7271,7 @@ void FParticleSystemSceneProxy::UpdateWorldSpacePrimitiveUniformBuffer() const
 			false,
 			UseSingleSampleShadowFromStationaryLights(),
 			GetScene().HasPrecomputedVolumetricLightmap_RenderThread(),
+			DrawsVelocity(),
 			GetLightingChannelMask(),
 			0,
 			INDEX_NONE,
diff --git a/Engine/Source/Runtime/Engine/Private/PhysicsEngine/PhysDrawing.cpp b/Engine/Source/Runtime/Engine/Private/PhysicsEngine/PhysDrawing.cpp
--- a/Engine/Source/Runtime/Engine/Private/PhysicsEngine/PhysDrawing.cpp
+++ b/Engine/Source/Runtime/Engine/Private/PhysicsEngine/PhysDrawing.cpp
@@ -800,7 +800,7 @@ bool FKConvexGeomRenderInfo::HasValidGeometry()
 // FKAggregateGeom
 /////////////////////////////////////////////////////////////////////////////////////
 
-void FKAggregateGeom::GetAggGeom(const FTransform& Transform, const FColor Color, const FMaterialRenderProxy* MatInst, bool bPerHullColor, bool bDrawSolid, int32 ViewIndex, FMeshElementCollector& Collector) const
+void FKAggregateGeom::GetAggGeom(const FTransform& Transform, const FColor Color, const FMaterialRenderProxy* MatInst, bool bPerHullColor, bool bDrawSolid, bool bDrawsVelocity, int32 ViewIndex, FMeshElementCollector& Collector) const
 {
 	const FVector Scale3D = Transform.GetScale3D();
 	FTransform ParentTM = Transform;
@@ -894,7 +894,7 @@ void FKAggregateGeom::GetAggGeom(const FTransform& Transform, const FColor Color
 				CalcBoxSphereBounds(LocalBounds, FTransform::Identity);
 
 				FDynamicPrimitiveUniformBuffer& DynamicPrimitiveUniformBuffer = Collector.AllocateOneFrameResource<FDynamicPrimitiveUniformBuffer>();
-				DynamicPrimitiveUniformBuffer.Set(LocalToWorld.ToMatrixWithScale(), LocalToWorld.ToMatrixWithScale(), WorldBounds, LocalBounds, true, false);
+				DynamicPrimitiveUniformBuffer.Set(LocalToWorld.ToMatrixWithScale(), LocalToWorld.ToMatrixWithScale(), WorldBounds, LocalBounds, true, false, bDrawsVelocity);
 				BatchElement.PrimitiveUniformBufferResource = &DynamicPrimitiveUniformBuffer.UniformBuffer;
 
 			 	// previous l2w not used so treat as static
@@ -1004,7 +1004,7 @@ void UPhysicsAsset::GetCollisionMesh(int32 ViewIndex, FMeshElementCollector& Col
 		// BoneTransform.SetScale3D(Scale3D);
 		if (SkeletalBodySetups[i]->bCreatedPhysicsMeshes)
 		{
-			SkeletalBodySetups[i]->AggGeom.GetAggGeom(BoneTransform, *BoneColor, NULL, false, false, ViewIndex, Collector);
+			SkeletalBodySetups[i]->AggGeom.GetAggGeom(BoneTransform, *BoneColor, NULL, false, false, true, ViewIndex, Collector);
 		}
 	}
 }
diff --git a/Engine/Source/Runtime/Engine/Private/PrimitiveSceneProxy.cpp b/Engine/Source/Runtime/Engine/Private/PrimitiveSceneProxy.cpp
--- a/Engine/Source/Runtime/Engine/Private/PrimitiveSceneProxy.cpp
+++ b/Engine/Source/Runtime/Engine/Private/PrimitiveSceneProxy.cpp
@@ -282,6 +282,7 @@ void FPrimitiveSceneProxy::UpdateUniformBuffer()
 				HasDynamicIndirectShadowCasterRepresentation(), 
 				UseSingleSampleShadowFromStationaryLights(),
 				bHasPrecomputedVolumetricLightmap,
+				DrawsVelocity(), 
 				GetLightingChannelMask(),
 				LpvBiasMultiplier,
 				PrimitiveSceneInfo ? PrimitiveSceneInfo->GetLightmapDataOffset() : 0,
diff --git a/Engine/Source/Runtime/Engine/Private/PrimitiveUniformShaderParameters.cpp b/Engine/Source/Runtime/Engine/Private/PrimitiveUniformShaderParameters.cpp
--- a/Engine/Source/Runtime/Engine/Private/PrimitiveUniformShaderParameters.cpp
+++ b/Engine/Source/Runtime/Engine/Private/PrimitiveUniformShaderParameters.cpp
@@ -55,6 +55,7 @@ FPrimitiveSceneShaderData::FPrimitiveSceneShaderData(const FPrimitiveSceneProxy*
 		Proxy->HasDynamicIndirectShadowCasterRepresentation(), 
 		Proxy->UseSingleSampleShadowFromStationaryLights(),
 		bHasPrecomputedVolumetricLightmap,
+		Proxy->DrawsVelocity(), 
 		Proxy->GetLightingChannelMask(),
 		Proxy->GetLpvBiasMultiplier(),
 		Proxy->GetPrimitiveSceneInfo()->GetLightmapDataOffset(),
@@ -94,7 +95,7 @@ void FPrimitiveSceneShaderData::Setup(const FPrimitiveUniformShaderParameters& P
 		PrimitiveUniformShaderParameters.DecalReceiverMask, 
 		PrimitiveUniformShaderParameters.PerObjectGBufferData, 
 		PrimitiveUniformShaderParameters.UseVolumetricLightmapShadowFromStationaryLights, 
-		1.0f);
+		PrimitiveUniformShaderParameters.DrawsVelocity);
 	Data[21] = PrimitiveUniformShaderParameters.ObjectOrientation;
 	Data[22] = PrimitiveUniformShaderParameters.NonUniformScale;
 
diff --git a/Engine/Source/Runtime/Engine/Private/SceneManagement.cpp b/Engine/Source/Runtime/Engine/Private/SceneManagement.cpp
--- a/Engine/Source/Runtime/Engine/Private/SceneManagement.cpp
+++ b/Engine/Source/Runtime/Engine/Private/SceneManagement.cpp
@@ -307,11 +307,12 @@ void FDynamicPrimitiveUniformBuffer::Set(
 	const FBoxSphereBounds& LocalBounds,
 	const FBoxSphereBounds& PreSkinnedLocalBounds,
 	bool bReceivesDecals,
-	bool bHasPrecomputedVolumetricLightmap)
+	bool bHasPrecomputedVolumetricLightmap,
+	bool bDrawsVelocity)
 {
 	check(IsInRenderingThread());
 	UniformBuffer.SetContents(
-		GetPrimitiveUniformShaderParameters(LocalToWorld, PreviousLocalToWorld, WorldBounds.Origin, WorldBounds, LocalBounds, PreSkinnedLocalBounds, bReceivesDecals, false, false, false, bHasPrecomputedVolumetricLightmap, GetDefaultLightingChannelMask(), 1.0f, INDEX_NONE, INDEX_NONE));
+		GetPrimitiveUniformShaderParameters(LocalToWorld, PreviousLocalToWorld, WorldBounds.Origin, WorldBounds, LocalBounds, PreSkinnedLocalBounds, bReceivesDecals, false, false, false, bHasPrecomputedVolumetricLightmap, bDrawsVelocity, GetDefaultLightingChannelMask(), 1.0f, INDEX_NONE, INDEX_NONE));
 	UniformBuffer.InitResource();
 }
 
@@ -321,9 +322,10 @@ void FDynamicPrimitiveUniformBuffer::Set(
 	const FBoxSphereBounds& WorldBounds,
 	const FBoxSphereBounds& LocalBounds,
 	bool bReceivesDecals,
-	bool bHasPrecomputedVolumetricLightmap)
+	bool bHasPrecomputedVolumetricLightmap,
+	bool bDrawsVelocity)
 {
-	Set(LocalToWorld, PreviousLocalToWorld, WorldBounds, LocalBounds, LocalBounds, bReceivesDecals, bHasPrecomputedVolumetricLightmap);
+	Set(LocalToWorld, PreviousLocalToWorld, WorldBounds, LocalBounds, LocalBounds, bReceivesDecals, bHasPrecomputedVolumetricLightmap, bDrawsVelocity);
 }
 
 FLightMapInteraction FLightMapInteraction::Texture(
diff --git a/Engine/Source/Runtime/Engine/Private/StaticMeshRender.cpp b/Engine/Source/Runtime/Engine/Private/StaticMeshRender.cpp
--- a/Engine/Source/Runtime/Engine/Private/StaticMeshRender.cpp
+++ b/Engine/Source/Runtime/Engine/Private/StaticMeshRender.cpp
@@ -1390,13 +1390,13 @@ void FStaticMeshSceneProxy::GetDynamicMeshElements(const TArray<const FSceneView
 						Collector.RegisterOneFrameMaterialProxy(SolidMaterialInstance);
 
 						FTransform GeomTransform(GetLocalToWorld());
-						BodySetup->AggGeom.GetAggGeom(GeomTransform, GetWireframeColor().ToFColor(true), SolidMaterialInstance, false, true, ViewIndex, Collector);
+						BodySetup->AggGeom.GetAggGeom(GeomTransform, GetWireframeColor().ToFColor(true), SolidMaterialInstance, false, true, DrawsVelocity(), ViewIndex, Collector);
 					}
 					// wireframe
 					else
 					{
 						FTransform GeomTransform(GetLocalToWorld());
-						BodySetup->AggGeom.GetAggGeom(GeomTransform, GetSelectionColor(SimpleCollisionColor, bProxyIsSelected, IsHovered()).ToFColor(true), NULL, ( Owner == NULL ), false, ViewIndex, Collector);
+						BodySetup->AggGeom.GetAggGeom(GeomTransform, GetSelectionColor(SimpleCollisionColor, bProxyIsSelected, IsHovered()).ToFColor(true), NULL, ( Owner == NULL ), false, DrawsVelocity(), ViewIndex, Collector);
 					}
 
 
diff --git a/Engine/Source/Runtime/Engine/Public/PrimitiveSceneProxy.h b/Engine/Source/Runtime/Engine/Public/PrimitiveSceneProxy.h
--- a/Engine/Source/Runtime/Engine/Public/PrimitiveSceneProxy.h
+++ b/Engine/Source/Runtime/Engine/Public/PrimitiveSceneProxy.h
@@ -522,6 +522,11 @@ public:
 	inline bool IsStaticPathAvailable() const { return !bHasMobileMovablePointLightInteraction; }
 	inline bool ShouldReceiveMobileCSMShadows() const { return bReceiveMobileCSMShadows; }
 
+	/** Returns whether draws velocity in base pass. */
+	inline bool DrawsVelocity() const {
+		return IsMovable();
+	}
+
 #if WITH_EDITOR
 	inline int32 GetNumUncachedStaticLightingInteractions() { return NumUncachedStaticLightingInteractions; }
 
diff --git a/Engine/Source/Runtime/Engine/Public/PrimitiveUniformShaderParameters.h b/Engine/Source/Runtime/Engine/Public/PrimitiveUniformShaderParameters.h
--- a/Engine/Source/Runtime/Engine/Public/PrimitiveUniformShaderParameters.h
+++ b/Engine/Source/Runtime/Engine/Public/PrimitiveUniformShaderParameters.h
@@ -27,8 +27,8 @@ BEGIN_GLOBAL_SHADER_PARAMETER_STRUCT(FPrimitiveUniformShaderParameters,ENGINE_AP
 	SHADER_PARAMETER(float,LpvBiasMultiplier)
 	SHADER_PARAMETER_EX(float,DecalReceiverMask,EShaderPrecisionModifier::Half)
 	SHADER_PARAMETER_EX(float,PerObjectGBufferData,EShaderPrecisionModifier::Half)		// 0..1, 2 bits, bDistanceFieldRepresentation, bHeightfieldRepresentation
-	SHADER_PARAMETER_EX(float,UseVolumetricLightmapShadowFromStationaryLights,EShaderPrecisionModifier::Half)	
-	SHADER_PARAMETER_EX(float,Unused,EShaderPrecisionModifier::Half)
+	SHADER_PARAMETER_EX(float,UseVolumetricLightmapShadowFromStationaryLights,EShaderPrecisionModifier::Half)		
+	SHADER_PARAMETER_EX(float,DrawsVelocity,EShaderPrecisionModifier::Half)
 	SHADER_PARAMETER_EX(FVector4,ObjectOrientation,EShaderPrecisionModifier::Half)
 	SHADER_PARAMETER_EX(FVector4,NonUniformScale,EShaderPrecisionModifier::Half)
 	SHADER_PARAMETER(FVector, LocalObjectBoundsMin)		// This is used in a custom material function (ObjectLocalBounds.uasset)
@@ -52,6 +52,7 @@ inline FPrimitiveUniformShaderParameters GetPrimitiveUniformShaderParameters(
 	bool bHasCapsuleRepresentation,
 	bool bUseSingleSampleShadowFromStationaryLights,
 	bool bUseVolumetricLightmap,
+	bool bDrawsVelocity,
 	uint32 LightingChannelMask,
 	float LpvBiasMultiplier,
 	uint32 LightmapDataIndex,
@@ -93,7 +94,7 @@ inline FPrimitiveUniformShaderParameters GetPrimitiveUniformShaderParameters(
 	Result.PerObjectGBufferData = (2 * (int32)bHasCapsuleRepresentation + (int32)bHasDistanceFieldRepresentation) / 3.0f;
 	Result.UseSingleSampleShadowFromStationaryLights = bUseSingleSampleShadowFromStationaryLights ? 1.0f : 0.0f;
 	Result.UseVolumetricLightmapShadowFromStationaryLights = bUseVolumetricLightmap && bUseSingleSampleShadowFromStationaryLights ? 1.0f : 0.0f;
-	Result.Unused = 0;
+	Result.DrawsVelocity = bDrawsVelocity ? 1 : 0;
 	Result.LightmapDataIndex = LightmapDataIndex;
 	Result.SingleCaptureIndex = SingleCaptureIndex;
 	return Result;
@@ -111,6 +112,7 @@ inline FPrimitiveUniformShaderParameters GetPrimitiveUniformShaderParameters(
 	bool bHasCapsuleRepresentation,
 	bool bUseSingleSampleShadowFromStationaryLights,
 	bool bUseVolumetricLightmap,
+	bool bDrawsVelocity,
 	uint32 LightingChannelMask,
 	float LpvBiasMultiplier,
 	uint32 LightmapDataIndex,
@@ -119,7 +121,7 @@ inline FPrimitiveUniformShaderParameters GetPrimitiveUniformShaderParameters(
 {
 	// Pass through call
 	return GetPrimitiveUniformShaderParameters(LocalToWorld, PreviousLocalToWorld, ActorPosition, WorldBounds, LocalBounds, LocalBounds, bReceivesDecals, bHasDistanceFieldRepresentation, bHasCapsuleRepresentation, 
-		bUseSingleSampleShadowFromStationaryLights, bUseVolumetricLightmap, LightingChannelMask, LpvBiasMultiplier, LightmapDataIndex, SingleCaptureIndex);
+		bUseSingleSampleShadowFromStationaryLights, bUseVolumetricLightmap, bDrawsVelocity, LightingChannelMask, LpvBiasMultiplier, LightmapDataIndex, SingleCaptureIndex);
 }
 
 inline TUniformBufferRef<FPrimitiveUniformShaderParameters> CreatePrimitiveUniformBufferImmediate(
@@ -128,12 +130,13 @@ inline TUniformBufferRef<FPrimitiveUniformShaderParameters> CreatePrimitiveUnifo
 	const FBoxSphereBounds& LocalBounds,
 	const FBoxSphereBounds& PreSkinnedLocalBounds,
 	bool bReceivesDecals,
+	bool bDrawsVelocity,
 	float LpvBiasMultiplier = 1.0f
 	)
 {
 	check(IsInRenderingThread());
 	return TUniformBufferRef<FPrimitiveUniformShaderParameters>::CreateUniformBufferImmediate(
-		GetPrimitiveUniformShaderParameters(LocalToWorld, LocalToWorld, WorldBounds.Origin, WorldBounds, LocalBounds, PreSkinnedLocalBounds, bReceivesDecals, false, false, false, false, GetDefaultLightingChannelMask(), LpvBiasMultiplier, INDEX_NONE, INDEX_NONE),
+		GetPrimitiveUniformShaderParameters(LocalToWorld, LocalToWorld, WorldBounds.Origin, WorldBounds, LocalBounds, PreSkinnedLocalBounds, bReceivesDecals, false, false, false, false, bDrawsVelocity, GetDefaultLightingChannelMask(), LpvBiasMultiplier, INDEX_NONE, INDEX_NONE),
 		UniformBuffer_MultiFrame
 		);
 }
@@ -153,6 +156,7 @@ inline FPrimitiveUniformShaderParameters GetIdentityPrimitiveParameters()
 		false,
 		false,
 		false,
+		/* bDrawsVelocity = */ true,
 		GetDefaultLightingChannelMask(),
 		1.0f,		// LPV bias
 		INDEX_NONE,
diff --git a/Engine/Source/Runtime/Engine/Public/SceneManagement.h b/Engine/Source/Runtime/Engine/Public/SceneManagement.h
--- a/Engine/Source/Runtime/Engine/Public/SceneManagement.h
+++ b/Engine/Source/Runtime/Engine/Public/SceneManagement.h
@@ -2203,7 +2203,8 @@ public:
 		const FBoxSphereBounds& LocalBounds,
 		const FBoxSphereBounds& PreSkinnedLocalBounds,
 		bool bReceivesDecals,
-		bool bHasPrecomputedVolumetricLightmap);
+		bool bHasPrecomputedVolumetricLightmap,
+		bool bDrawsVelocity);
 
 	/** Pass-through implementation which calls the overloaded Set function with LocalBounds for PreSkinnedLocalBounds. */
 	ENGINE_API void Set(
@@ -2212,7 +2213,8 @@ public:
 		const FBoxSphereBounds& WorldBounds,
 		const FBoxSphereBounds& LocalBounds,
 		bool bReceivesDecals,
-		bool bHasPrecomputedVolumetricLightmap);
+		bool bHasPrecomputedVolumetricLightmap,
+		bool bDrawsVelocity);
 };
 
 //
diff --git a/Engine/Source/Runtime/Experimental/GeometryCollectionEngine/Private/GeometryCollection/GeometryCollectionSceneProxy.cpp b/Engine/Source/Runtime/Experimental/GeometryCollectionEngine/Private/GeometryCollection/GeometryCollectionSceneProxy.cpp
--- a/Engine/Source/Runtime/Experimental/GeometryCollectionEngine/Private/GeometryCollection/GeometryCollectionSceneProxy.cpp
+++ b/Engine/Source/Runtime/Experimental/GeometryCollectionEngine/Private/GeometryCollection/GeometryCollectionSceneProxy.cpp
@@ -286,7 +286,7 @@ void FGeometryCollectionSceneProxy::GetDynamicMeshElements(const TArray<const FS
                     
                     // Collector owns the uniform buffer
                     FDynamicPrimitiveUniformBuffer& DynamicPrimitiveUniformBuffer = Collector.AllocateOneFrameResource<FDynamicPrimitiveUniformBuffer>();
-                    DynamicPrimitiveUniformBuffer.Set(GetLocalToWorld(), GetLocalToWorld(), GetBounds(), GetLocalBounds(), true, false);
+                    DynamicPrimitiveUniformBuffer.Set(GetLocalToWorld(), GetLocalToWorld(), GetBounds(), GetLocalBounds(), true, false, DrawsVelocity());
                     BatchElement.PrimitiveUniformBuffer = DynamicPrimitiveUniformBuffer.UniformBuffer.GetUniformBufferRHI();
                     
 					BatchElement.FirstIndex = Section.FirstIndex;
@@ -318,7 +318,7 @@ void FGeometryCollectionSceneProxy::GetDynamicMeshElements(const TArray<const FS
 					Mesh.MaterialRenderProxy = MaterialRenderProxy;
                     
                     FDynamicPrimitiveUniformBuffer& DynamicPrimitiveUniformBuffer = Collector.AllocateOneFrameResource<FDynamicPrimitiveUniformBuffer>();
-                    DynamicPrimitiveUniformBuffer.Set(GetLocalToWorld(), GetLocalToWorld(), GetBounds(), GetLocalBounds(), true, false);
+                    DynamicPrimitiveUniformBuffer.Set(GetLocalToWorld(), GetLocalToWorld(), GetBounds(), GetLocalBounds(), true, false, DrawsVelocity());
                     BatchElement.PrimitiveUniformBuffer = DynamicPrimitiveUniformBuffer.UniformBuffer.GetUniformBufferRHI();
                     
 					BatchElement.FirstIndex = 0;
diff --git a/Engine/Source/Runtime/MRMesh/Private/MRMeshComponent.cpp b/Engine/Source/Runtime/MRMesh/Private/MRMeshComponent.cpp
--- a/Engine/Source/Runtime/MRMesh/Private/MRMeshComponent.cpp
+++ b/Engine/Source/Runtime/MRMesh/Private/MRMeshComponent.cpp
@@ -325,7 +325,7 @@ private:
 						Mesh.MaterialRenderProxy = MaterialProxy;
 
 						FDynamicPrimitiveUniformBuffer& DynamicPrimitiveUniformBuffer = Collector.AllocateOneFrameResource<FDynamicPrimitiveUniformBuffer>();
-						DynamicPrimitiveUniformBuffer.Set(GetLocalToWorld(), GetLocalToWorld(), InfiniteBounds, InfiniteBounds, true, false);
+						DynamicPrimitiveUniformBuffer.Set(GetLocalToWorld(), GetLocalToWorld(), InfiniteBounds, InfiniteBounds, true, false, DrawsVelocity());
 						BatchElement.PrimitiveUniformBufferResource = &DynamicPrimitiveUniformBuffer.UniformBuffer;
 
 						BatchElement.FirstIndex = 0;
diff --git a/Engine/Source/Runtime/NavigationSystem/Private/NavMesh/NavMeshRenderingComponent.cpp b/Engine/Source/Runtime/NavigationSystem/Private/NavMesh/NavMeshRenderingComponent.cpp
--- a/Engine/Source/Runtime/NavigationSystem/Private/NavMesh/NavMeshRenderingComponent.cpp
+++ b/Engine/Source/Runtime/NavigationSystem/Private/NavMesh/NavMeshRenderingComponent.cpp
@@ -952,7 +952,7 @@ void FNavMeshSceneProxy::GetDynamicMeshElements(const TArray<const FSceneView*>&
 					BatchElement = MeshBatchElements[Index];
 
 					FDynamicPrimitiveUniformBuffer& DynamicPrimitiveUniformBuffer = Collector.AllocateOneFrameResource<FDynamicPrimitiveUniformBuffer>();
-					DynamicPrimitiveUniformBuffer.Set(FMatrix::Identity, FMatrix::Identity, GetBounds(), GetLocalBounds(), false, false);
+					DynamicPrimitiveUniformBuffer.Set(FMatrix::Identity, FMatrix::Identity, GetBounds(), GetLocalBounds(), false, false, DrawsVelocity());
 					BatchElement.PrimitiveUniformBufferResource = &DynamicPrimitiveUniformBuffer.UniformBuffer;
 
 					Mesh.bWireframe = false;
diff --git a/Engine/Source/Runtime/UMG/Private/Components/WidgetComponent.cpp b/Engine/Source/Runtime/UMG/Private/Components/WidgetComponent.cpp
--- a/Engine/Source/Runtime/UMG/Private/Components/WidgetComponent.cpp
+++ b/Engine/Source/Runtime/UMG/Private/Components/WidgetComponent.cpp
@@ -487,14 +487,14 @@ public:
 						Collector.RegisterOneFrameMaterialProxy(SolidMaterialInstance);
 
 						FTransform GeomTransform(GetLocalToWorld());
-						InBodySetup->AggGeom.GetAggGeom(GeomTransform, GetWireframeColor().ToFColor(true), SolidMaterialInstance, false, true, ViewIndex, Collector);
+						InBodySetup->AggGeom.GetAggGeom(GeomTransform, GetWireframeColor().ToFColor(true), SolidMaterialInstance, false, true, DrawsVelocity(), ViewIndex, Collector);
 					}
 					// wireframe
 					else
 					{
 						FColor CollisionColor = FColor(157, 149, 223, 255);
 						FTransform GeomTransform(GetLocalToWorld());
-						InBodySetup->AggGeom.GetAggGeom(GeomTransform, GetSelectionColor(CollisionColor, bProxyIsSelected, IsHovered()).ToFColor(true), nullptr, false, false, ViewIndex, Collector);
+						InBodySetup->AggGeom.GetAggGeom(GeomTransform, GetSelectionColor(CollisionColor, bProxyIsSelected, IsHovered()).ToFColor(true), nullptr, false, false, DrawsVelocity(), ViewIndex, Collector);
 					}
 				}
 			}
