Fixed distance field DDC key setup for non default r.DistanceFields.MaxPerMeshResolution and r.DistanceFields.DefaultVoxelDensity values.

diff --git a/Engine/Source/Runtime/Engine/Private/DistanceFieldAtlas.cpp b/Engine/Source/Runtime/Engine/Private/DistanceFieldAtlas.cpp
--- a/Engine/Source/Runtime/Engine/Private/DistanceFieldAtlas.cpp
+++ b/Engine/Source/Runtime/Engine/Private/DistanceFieldAtlas.cpp
@@ -585,11 +585,11 @@ FString BuildDistanceFieldDerivedDataKey(const FString& InMeshKey)
 {
 	static const auto CVar = IConsoleManager::Get().FindTConsoleVariableDataInt(TEXT("r.DistanceFields.MaxPerMeshResolution"));
 	const int32 PerMeshMax = CVar->GetValueOnAnyThread();
-	const FString PerMeshMaxString = PerMeshMax == 128 ? TEXT("") : FString(TEXT("_%u"), PerMeshMax);
+	const FString PerMeshMaxString = PerMeshMax == 128 ? TEXT("") : FString::Printf(TEXT("_%u"), PerMeshMax);
 
 	static const auto CVarDensity = IConsoleManager::Get().FindTConsoleVariableDataFloat(TEXT("r.DistanceFields.DefaultVoxelDensity"));
 	const float VoxelDensity = CVarDensity->GetValueOnAnyThread();
-	const FString VoxelDensityString = VoxelDensity == .1f ? TEXT("") : FString(TEXT("_%.3f"), VoxelDensity);
+	const FString VoxelDensityString = VoxelDensity == .1f ? TEXT("") : FString::Printf(TEXT("_%.3f"), VoxelDensity);
 
 	static const auto CVarCompress = IConsoleManager::Get().FindTConsoleVariableDataInt(TEXT("r.DistanceFieldBuild.Compress"));
 	const bool bCompress = CVarCompress->GetValueOnAnyThread() != 0;
