DR - Static analysis fix
#rnx

diff --git a/Engine/Source/Runtime/Engine/Private/EdGraph/EdGraphPin.cpp b/Engine/Source/Runtime/Engine/Private/EdGraph/EdGraphPin.cpp
--- a/Engine/Source/Runtime/Engine/Private/EdGraph/EdGraphPin.cpp
+++ b/Engine/Source/Runtime/Engine/Private/EdGraph/EdGraphPin.cpp
@@ -685,7 +685,8 @@ FText UEdGraphPin::GetDefaultAsText() const
 
 bool UEdGraphPin::DoesDefaultValueMatchAutogenerated() const
 {
-	if (const UEdGraphSchema* Schema = GetSchema())
+	const UEdGraphSchema* Schema = GetSchema();
+	if (Schema)
 	{
 		return Schema->DoesDefaultValueMatchAutogenerated(*this);
 	}
