Fixed naming inconsistencies in ray tracing parameters

diff --git a/Engine/Source/Runtime/Engine/Classes/Components/LightComponent.h b/Engine/Source/Runtime/Engine/Classes/Components/LightComponent.h
--- a/Engine/Source/Runtime/Engine/Classes/Components/LightComponent.h
+++ b/Engine/Source/Runtime/Engine/Classes/Components/LightComponent.h
@@ -203,7 +203,7 @@ class ENGINE_API ULightComponent : public ULightComponentBase
 	 * They have less aliasing artifacts than standard shadowmaps, but inherit all the limitations of distance field representations (only uniform scale, no deformation).
 	 * These shadows have a low per-object cost (and don't depend on triangle count) so they are effective for distant shadows from a dynamic sun.
 	 */
-	UPROPERTY(EditAnywhere, BlueprintReadOnly, Category=DistanceFieldShadows, meta=(DisplayName = "RayTraced DistanceField Shadows"))
+	UPROPERTY(EditAnywhere, BlueprintReadOnly, Category=DistanceFieldShadows, meta=(DisplayName = "Ray Tracing DistanceField Shadows"))
 	bool bUseRayTracedDistanceFieldShadows;
 
 	/** 
diff --git a/Engine/Source/Runtime/Engine/Classes/Components/LightComponentBase.h b/Engine/Source/Runtime/Engine/Classes/Components/LightComponentBase.h
--- a/Engine/Source/Runtime/Engine/Classes/Components/LightComponentBase.h
+++ b/Engine/Source/Runtime/Engine/Classes/Components/LightComponentBase.h
@@ -78,7 +78,7 @@ class ENGINE_API ULightComponentBase : public USceneComponent
 	uint32 bCastVolumetricShadow : 1;
 
 	/** Whether the light shadows are computed with shadow-mapping or ray-tracing (when available). */
-	UPROPERTY(EditAnywhere, BlueprintReadOnly, Category = Light, AdvancedDisplay)
+	UPROPERTY(EditAnywhere, BlueprintReadOnly, Category = Light, meta = (DisplayName = "Cast Ray Tracing Shadow"), AdvancedDisplay)
 	uint32 bCastRaytracedShadow : 1;
 
 	/** Whether the light affects objects in reflections, when ray-traced reflection is enabled. */
diff --git a/Engine/Source/Runtime/Engine/Private/Materials/MaterialExpressions.cpp b/Engine/Source/Runtime/Engine/Private/Materials/MaterialExpressions.cpp
--- a/Engine/Source/Runtime/Engine/Private/Materials/MaterialExpressions.cpp
+++ b/Engine/Source/Runtime/Engine/Private/Materials/MaterialExpressions.cpp
@@ -12755,14 +12755,13 @@ UMaterialExpressionRayTracingQualitySwitch::UMaterialExpressionRayTracingQuality
 #if WITH_EDITOR
 int32 UMaterialExpressionRayTracingQualitySwitch::Compile(class FMaterialCompiler* Compiler, int32 OutputIndex)
 {
-
 	if (!Normal.GetTracedInput().Expression)
 	{
-		return Compiler->Errorf(TEXT("Missing RayTraceQualitySwitch input 'Normal'"));
+		return Compiler->Errorf(TEXT("Missing RayTracingQualitySwitch input 'Normal'"));
 	}
 	else if (!RayTraced.GetTracedInput().Expression)
 	{
-		return Compiler->Errorf(TEXT("Missing RayTraceQualitySwitch input 'RayTraced'"));
+		return Compiler->Errorf(TEXT("Missing RayTracingQualitySwitch input 'RayTracing'"));
 	}
 	else
 	{
