Adds new AlphaHoldout blend mode

diff --git a/Engine/Plugins/2D/Paper2D/Source/Paper2DEditor/Private/SpriteEditor/SpriteEditorViewportClient.cpp b/Engine/Plugins/2D/Paper2D/Source/Paper2DEditor/Private/SpriteEditor/SpriteEditorViewportClient.cpp
--- a/Engine/Plugins/2D/Paper2D/Source/Paper2DEditor/Private/SpriteEditor/SpriteEditorViewportClient.cpp
+++ b/Engine/Plugins/2D/Paper2D/Source/Paper2DEditor/Private/SpriteEditor/SpriteEditorViewportClient.cpp
@@ -365,6 +365,7 @@ void FSpriteEditorViewportClient::AnalyzeSpriteMaterialType(UPaperSprite* Sprite
 				case EBlendMode::BLEND_Additive:
 				case EBlendMode::BLEND_Modulate:
 				case EBlendMode::BLEND_AlphaComposite:
+				case EBlendMode::BLEND_AlphaHoldout:
 					NumTranslucentTriangles += NumTriangles;
 					break;
 				case EBlendMode::BLEND_Masked:
diff --git a/Engine/Plugins/2D/Paper2D/Source/Paper2DEditor/Private/TileMapEditing/TileMapEditorViewportClient.cpp b/Engine/Plugins/2D/Paper2D/Source/Paper2DEditor/Private/TileMapEditing/TileMapEditorViewportClient.cpp
--- a/Engine/Plugins/2D/Paper2D/Source/Paper2DEditor/Private/TileMapEditing/TileMapEditorViewportClient.cpp
+++ b/Engine/Plugins/2D/Paper2D/Source/Paper2DEditor/Private/TileMapEditing/TileMapEditorViewportClient.cpp
@@ -190,6 +190,7 @@ void FTileMapEditorViewportClient::DrawCanvas(FViewport& InViewport, FSceneView&
 					case EBlendMode::BLEND_Additive:
 					case EBlendMode::BLEND_Modulate:
 					case EBlendMode::BLEND_AlphaComposite:
+					case EBlendMode::BLEND_AlphaHoldout:
 						MaterialType = Translucent;
 						break;
 					case EBlendMode::BLEND_Masked:
diff --git a/Engine/Plugins/FX/Niagara/Source/Niagara/Private/NiagaraRendererMeshes.cpp b/Engine/Plugins/FX/Niagara/Source/Niagara/Private/NiagaraRendererMeshes.cpp
--- a/Engine/Plugins/FX/Niagara/Source/Niagara/Private/NiagaraRendererMeshes.cpp
+++ b/Engine/Plugins/FX/Niagara/Source/Niagara/Private/NiagaraRendererMeshes.cpp
@@ -239,7 +239,7 @@ void NiagaraRendererMeshes::GetDynamicMeshElements(const TArray<const FSceneView
 					if (MaterialProxy)
 					{
 						EBlendMode BlendMode = MaterialProxy->GetMaterial(VertexFactory->GetFeatureLevel())->GetBlendMode();
-						bHasTranslucentMaterials |= BlendMode == BLEND_AlphaComposite || BlendMode == BLEND_Translucent;
+						bHasTranslucentMaterials |= BlendMode == BLEND_AlphaComposite || BlendMode == BLEND_AlphaHoldout || BlendMode == BLEND_Translucent;
 					}
 				}
 
diff --git a/Engine/Plugins/FX/Niagara/Source/Niagara/Private/NiagaraRendererSprites.cpp b/Engine/Plugins/FX/Niagara/Source/Niagara/Private/NiagaraRendererSprites.cpp
--- a/Engine/Plugins/FX/Niagara/Source/Niagara/Private/NiagaraRendererSprites.cpp
+++ b/Engine/Plugins/FX/Niagara/Source/Niagara/Private/NiagaraRendererSprites.cpp
@@ -262,7 +262,7 @@ void NiagaraRendererSprites::SetVertexFactoryParticleData(
 
 	int32 NumInstances = DynamicDataSprites->RTParticleData.GetNumInstances();
 	FNiagaraGPUSortInfo SortInfo;
-	if (View && Properties->SortMode != ENiagaraSortMode::None && (BlendMode == BLEND_AlphaComposite || BlendMode == BLEND_Translucent || !Properties->bSortOnlyWhenTranslucent))
+	if (View && Properties->SortMode != ENiagaraSortMode::None && (BlendMode == BLEND_AlphaComposite || BlendMode == BLEND_AlphaHoldout|| BlendMode == BLEND_Translucent || !Properties->bSortOnlyWhenTranslucent))
 	{
 		SortInfo.ParticleCount = NumInstances;
 		SortInfo.SortMode = Properties->SortMode;
diff --git a/Engine/Shaders/Private/BasePassPixelShader.usf b/Engine/Shaders/Private/BasePassPixelShader.usf
--- a/Engine/Shaders/Private/BasePassPixelShader.usf
+++ b/Engine/Shaders/Private/BasePassPixelShader.usf
@@ -1113,6 +1113,10 @@ void FPixelShaderInOut_MainPS(
 			Out.MRT[0] = half4(Color, 0);
 		#endif
 		Out.MRT[0] = RETURN_COLOR(Out.MRT[0]);
+	#elif MATERIALBLENDING_ALPHAHOLDOUT
+		// not implemented for holdout
+		Out.MRT[0] = half4(Color * Fogging.a + Fogging.rgb * Opacity, Opacity);
+		Out.MRT[0] = RETURN_COLOR(Out.MRT[0]);
 	#elif MATERIALBLENDING_ALPHACOMPOSITE
 		Out.MRT[0] = half4(Color * Fogging.a + Fogging.rgb * Opacity, Opacity);
 		Out.MRT[0] = RETURN_COLOR(Out.MRT[0]);
diff --git a/Engine/Shaders/Private/Definitions.usf b/Engine/Shaders/Private/Definitions.usf
--- a/Engine/Shaders/Private/Definitions.usf
+++ b/Engine/Shaders/Private/Definitions.usf
@@ -42,6 +42,10 @@
 #define MATERIALBLENDING_ALPHACOMPOSITE					0
 #endif
 
+#ifndef MATERIALBLENDING_ALPHAHOLDOUT
+#define MATERIALBLENDING_ALPHAHOLDOUT					0
+#endif
+
 #ifndef MATERIAL_SHADINGMODEL_DEFAULT_LIT
 #define MATERIAL_SHADINGMODEL_DEFAULT_LIT				0
 #endif
diff --git a/Engine/Shaders/Private/MobileBasePassPixelShader.usf b/Engine/Shaders/Private/MobileBasePassPixelShader.usf
--- a/Engine/Shaders/Private/MobileBasePassPixelShader.usf
+++ b/Engine/Shaders/Private/MobileBasePassPixelShader.usf
@@ -309,6 +309,10 @@ half3 FrameBufferBlendOp(half4 Source)
 // so ensure  MATERIALBLENDING_ALPHACOMPOSITE gets first in line
 #elif MATERIALBLENDING_ALPHACOMPOSITE
 	return Source.rgb + (Dest.rgb*(1.0 - Source.a));
+// AlphaHoldout will set both MATERIALBLENDING_TRANSLUCENT and MATERIALBLENDING_ALPHAHOLDOUT defines
+// so ensure  MATERIALBLENDING_ALPHAHOLDOUT gets first in line
+#elif MATERIALBLENDING_ALPHAHOLDOUT
+	return (Dest.rgb*(1.0 - Source.a));
 #elif MATERIALBLENDING_TRANSLUCENT
 	return (Source.rgb*Source.a) + (Dest.rgb*(1.0 - Source.a));
 #elif MATERIALBLENDING_ADDITIVE
@@ -650,6 +654,9 @@ void Main(
 
 	#if MATERIALBLENDING_ALPHACOMPOSITE
 		OutColor = half4(Color * VertexFog.a + VertexFog.rgb * Opacity, Opacity);
+	#if MATERIALBLENDING_ALPHAHOLDOUT
+		// not implemented for holdout
+		OutColor = half4(Color * VertexFog.a + VertexFog.rgb * Opacity, Opacity);
 	#elif MATERIALBLENDING_TRANSLUCENT
 		OutColor = half4(Color * VertexFog.a + VertexFog.rgb, Opacity);
 	#elif MATERIALBLENDING_ADDITIVE
diff --git a/Engine/Shaders/Private/RayTracing/RayTracingCommon.ush b/Engine/Shaders/Private/RayTracing/RayTracingCommon.ush
--- a/Engine/Shaders/Private/RayTracing/RayTracingCommon.ush
+++ b/Engine/Shaders/Private/RayTracing/RayTracingCommon.ush
@@ -89,6 +89,8 @@ float3 ReconstructWorldPositionFromDepth(float2 UV, float Depth)
 #define RAY_TRACING_BLEND_MODE_TRANSLUCENT			2
 #define RAY_TRACING_BLEND_MODE_ADDITIVE				3
 #define RAY_TRACING_BLEND_MODE_MODULATE				4
+#define RAY_TRACING_BLEND_MODE_ALPHA_HOLDOUT		5
+
 
 struct FMaterialClosestHitPayload
 {
diff --git a/Engine/Shaders/Private/RayTracing/RayTracingMaterialHitShaders.usf b/Engine/Shaders/Private/RayTracing/RayTracingMaterialHitShaders.usf
--- a/Engine/Shaders/Private/RayTracing/RayTracingMaterialHitShaders.usf
+++ b/Engine/Shaders/Private/RayTracing/RayTracingMaterialHitShaders.usf
@@ -314,6 +314,8 @@ void MaterialCHS(inout FPackedMaterialClosestHitPayload PackedPayload, in FDefau
 
  #if MATERIALBLENDING_ALPHACOMPOSITE
 	uint MaterialBlendingMode = RAY_TRACING_BLEND_MODE_ALPHA_COMPOSITE;
+ #elif MATERIALBLENDING_ALPHAHOLDOUT
+	uint MaterialBlendingMode = RAY_TRACING_BLEND_MODE_ALPHA_HOLDOUT;
  #elif MATERIALBLENDING_TRANSLUCENT
 	uint MaterialBlendingMode = RAY_TRACING_BLEND_MODE_TRANSLUCENT;
  #elif MATERIALBLENDING_ADDITIVE
diff --git a/Engine/Shaders/Private/ShadowDepthPixelShader.usf b/Engine/Shaders/Private/ShadowDepthPixelShader.usf
--- a/Engine/Shaders/Private/ShadowDepthPixelShader.usf
+++ b/Engine/Shaders/Private/ShadowDepthPixelShader.usf
@@ -150,7 +150,7 @@ void Main(
 #if MATERIAL_BLOCK_GI
 	OutColor = 1.0f;
 	OutReflectiveShadowMapDiffuse = opacity;
-#elif MATERIALBLENDING_TRANSLUCENT || MATERIALBLENDING_ALPHACOMPOSITE
+#elif MATERIALBLENDING_TRANSLUCENT || MATERIALBLENDING_ALPHACOMPOSITE || MATERIALBLENDING_ALPHAHOLDOUT
 	OutColor = 0;
 	OutReflectiveShadowMapDiffuse = 0;
 	clip(-1.0);
diff --git a/Engine/Shaders/Private/SimpleElementPixelShader.usf b/Engine/Shaders/Private/SimpleElementPixelShader.usf
--- a/Engine/Shaders/Private/SimpleElementPixelShader.usf
+++ b/Engine/Shaders/Private/SimpleElementPixelShader.usf
@@ -48,6 +48,10 @@ float3 SimpleElementFrameBufferBlendOp(float4 Source)
 // so ensure  MATERIALBLENDING_ALPHACOMPOSITE gets first in line
 #elif MATERIALBLENDING_ALPHACOMPOSITE
 	return Source.rgb + (Dest.rgb*(1.0 - Source.a));
+// AlphaHoldout will set both MATERIALBLENDING_TRANSLUCENT and MATERIALBLENDING_ALPHAHOLDOUT defines
+// so ensure  MATERIALBLENDING_ALPHAHOLDOUT gets first in line
+#elif MATERIALBLENDING_ALPHAHOLDOUT
+	return (Dest.rgb*(1.0 - Source.a));
 #elif SE_BLEND_MODE == SE_BLEND_TRANSLUCENT || SE_BLEND_MODE == SE_BLEND_TRANSLUCENTALPHAONLY || SE_BLEND_MODE == SE_BLEND_ALPHABLEND || SE_BLEND_MODE == SE_BLEND_TRANSLUCENTDISTANCEFIELD || SE_BLEND_MODE == SE_BLEND_TRANSLUCENTDISTANCEFIELDSHADOWED || SE_BLEND_MODE == SE_BLEND_TRANSLUCENTALPHAONLYWRITEALPHA
 	return (Source.rgb*Source.a) + (Dest.rgb*(1.0 - Source.a));
 #elif SE_BLEND_MODE == SE_BLEND_ADDITIVE
@@ -56,6 +60,8 @@ float3 SimpleElementFrameBufferBlendOp(float4 Source)
 	return Source.rgb * Dest.rgb;
 #elif SE_BLEND_MODE == SE_BLEND_ALPHACOMPOSITE
 	return Source.rgb + (Dest.rgb*(1.0 - Source.a));
+#elif SE_BLEND_MODE == SE_BLEND_ALPHAHOLDOUT
+	return (Dest.rgb*(1.0 - Source.a));
 #endif
 
 	return Source.rgb;
diff --git a/Engine/Source/Editor/MaterialEditor/Private/SMaterialEditorViewport.cpp b/Engine/Source/Editor/MaterialEditor/Private/SMaterialEditorViewport.cpp
--- a/Engine/Source/Editor/MaterialEditor/Private/SMaterialEditorViewport.cpp
+++ b/Engine/Source/Editor/MaterialEditor/Private/SMaterialEditorViewport.cpp
@@ -172,7 +172,7 @@ FLinearColor FMaterialEditorViewportClient::GetBackgroundColor() const
 				{
 					BackgroundColor = FLinearColor::White;
 				}
-				else if (PreviewBlendMode == BLEND_Translucent || PreviewBlendMode == BLEND_AlphaComposite)
+				else if (PreviewBlendMode == BLEND_Translucent || PreviewBlendMode == BLEND_AlphaComposite || PreviewBlendMode == BLEND_AlphaHoldout)
 				{
 					BackgroundColor = FColor(64, 64, 64);
 				}
diff --git a/Engine/Source/Editor/UnrealEd/Private/Lightmass/LightmassRender.cpp b/Engine/Source/Editor/UnrealEd/Private/Lightmass/LightmassRender.cpp
--- a/Engine/Source/Editor/UnrealEd/Private/Lightmass/LightmassRender.cpp
+++ b/Engine/Source/Editor/UnrealEd/Private/Lightmass/LightmassRender.cpp
@@ -422,7 +422,7 @@ public:
 						return Compiler->Saturate(MaterialInterface->CompileProperty(Compiler, DiffuseInput, ForceCast_Exact_Replicate));
 					}
 				}
-				else if ((BlendMode == BLEND_Translucent) || (BlendMode == BLEND_Additive || (BlendMode == BLEND_AlphaComposite)))
+				else if ((BlendMode == BLEND_Translucent) || (BlendMode == BLEND_Additive) || (BlendMode == BLEND_AlphaComposite) || (BlendMode == BLEND_AlphaHoldout))
 				{
 					int32 ColoredOpacity = INDEX_NONE;
 					if (ShadingModel == MSM_Unlit)
@@ -651,7 +651,8 @@ public:
 			if ((BlendMode == BLEND_Modulate) ||
 				(BlendMode == BLEND_Translucent) || 
 				(BlendMode == BLEND_Additive) ||
-				(BlendMode == BLEND_AlphaComposite))
+				(BlendMode == BLEND_AlphaComposite) ||
+				(BlendMode == BLEND_AlphaHoldout))
 			{
 				bool bColorInputIsNULL = false;
 				if (ShadingModel == MSM_Unlit)
@@ -664,7 +665,8 @@ public:
 				}
 				if (BlendMode == BLEND_Translucent
 					|| BlendMode == BLEND_Additive
-					|| BlendMode == BLEND_AlphaComposite)
+					|| BlendMode == BLEND_AlphaComposite
+					|| BlendMode == BLEND_AlphaHoldout)
 				{
 					bExpressionIsNULL = bColorInputIsNULL && !IsMaterialInputConnected(Material, PropertyToCompile);
 				}
@@ -775,6 +777,7 @@ public:
 		case BLEND_Translucent:
 		case BLEND_Additive:
 		case BLEND_AlphaComposite:
+		case BLEND_AlphaHoldout:
 			{
 				switch (InMaterialProperty)
 				{
diff --git a/Engine/Source/Programs/UnrealLightmass/Private/ImportExport/Material.cpp b/Engine/Source/Programs/UnrealLightmass/Private/ImportExport/Material.cpp
--- a/Engine/Source/Programs/UnrealLightmass/Private/ImportExport/Material.cpp
+++ b/Engine/Source/Programs/UnrealLightmass/Private/ImportExport/Material.cpp
@@ -56,7 +56,7 @@ namespace Lightmass
 		else
 		{
 			// Materials with a translucent blend mode should always import transmission
-			check(BlendMode != BLEND_Translucent && BlendMode != BLEND_Additive && BlendMode != BLEND_Modulate && BlendMode != BLEND_AlphaComposite);
+			check(BlendMode != BLEND_Translucent && BlendMode != BLEND_Additive && BlendMode != BLEND_Modulate && BlendMode != BLEND_AlphaComposite && BlendMode != BLEND_AlphaHoldout);
 		}
 
 		// Normal
diff --git a/Engine/Source/Programs/UnrealLightmass/Public/MaterialExport.h b/Engine/Source/Programs/UnrealLightmass/Public/MaterialExport.h
--- a/Engine/Source/Programs/UnrealLightmass/Public/MaterialExport.h
+++ b/Engine/Source/Programs/UnrealLightmass/Public/MaterialExport.h
@@ -50,7 +50,8 @@ enum EBlendMode
     BLEND_Additive          =3,
     BLEND_Modulate          =4,
 	BLEND_AlphaComposite    =5,
-    BLEND_MAX               =6,
+	BLEND_AlphaHoldout      =6,
+	BLEND_MAX               =7,
 };
 
 struct FMaterialData
diff --git a/Engine/Source/Runtime/Engine/Classes/Engine/EngineTypes.h b/Engine/Source/Runtime/Engine/Classes/Engine/EngineTypes.h
--- a/Engine/Source/Runtime/Engine/Classes/Engine/EngineTypes.h
+++ b/Engine/Source/Runtime/Engine/Classes/Engine/EngineTypes.h
@@ -243,7 +243,8 @@ enum EBlendMode
 	BLEND_Translucent UMETA(DisplayName="Translucent"),
 	BLEND_Additive UMETA(DisplayName="Additive"),
 	BLEND_Modulate UMETA(DisplayName="Modulate"),
-	BLEND_AlphaComposite UMETA(DisplayName ="AlphaComposite (Premultiplied Alpha)"),
+	BLEND_AlphaComposite UMETA(DisplayName = "AlphaComposite (Premultiplied Alpha)"),
+	BLEND_AlphaHoldout UMETA(DisplayName = "AlphaHoldout"),
 	BLEND_MAX,
 };
 
diff --git a/Engine/Source/Runtime/Engine/Private/BatchedElements.cpp b/Engine/Source/Runtime/Engine/Private/BatchedElements.cpp
--- a/Engine/Source/Runtime/Engine/Private/BatchedElements.cpp
+++ b/Engine/Source/Runtime/Engine/Private/BatchedElements.cpp
@@ -136,7 +136,10 @@ void FBatchedElements::AddTriangle(int32 V0,int32 V1,int32 V2,const FTexture* Te
 	case BLEND_AlphaComposite:
 		SimpleElementBlendMode = SE_BLEND_AlphaComposite;
 		break;
-	};	
+	case BLEND_AlphaHoldout:
+		SimpleElementBlendMode = SE_BLEND_AlphaHoldout;
+		break;
+	};
 	AddTriangle(V0,V1,V2,Texture,SimpleElementBlendMode);
 }
 
@@ -385,6 +388,9 @@ static void SetBlendState(FRHICommandList& RHICmdList, FGraphicsPipelineStateIni
 	case SE_BLEND_AlphaComposite:
 		GraphicsPSOInit.BlendState = TStaticBlendState<CW_RGBA, BO_Add, BF_One, BF_InverseSourceAlpha, BO_Add, BF_One, BF_InverseSourceAlpha>::GetRHI();
 		break;
+	case SE_BLEND_AlphaHoldout:
+		GraphicsPSOInit.BlendState = TStaticBlendState<CW_RGBA, BO_Add, BF_Zero, BF_InverseSourceAlpha, BO_Add, BF_Zero, BF_InverseSourceAlpha>::GetRHI();
+		break;
 	case SE_BLEND_TranslucentAlphaOnlyWriteAlpha:
 	case SE_BLEND_AlphaBlend:
 		GraphicsPSOInit.BlendState = TStaticBlendState<CW_RGBA, BO_Add, BF_SourceAlpha, BF_InverseSourceAlpha, BO_Add, BF_InverseDestAlpha, BF_One>::GetRHI();
@@ -409,6 +415,7 @@ static void SetHitTestingBlendState(FRHICommandList& RHICmdList, FGraphicsPipeli
 		GraphicsPSOInit.BlendState = TStaticBlendState<CW_RGBA, BO_Add, BF_One, BF_Zero, BO_Add, BF_One, BF_Zero>::GetRHI();
 		break;
 	case SE_BLEND_AlphaComposite:
+	case SE_BLEND_AlphaHoldout:
 	case SE_BLEND_AlphaBlend:
 	case SE_BLEND_Translucent:
 	case SE_BLEND_TranslucentDistanceField:
@@ -460,6 +467,8 @@ static TSimpleElementPixelShader* GetPixelShader(bool bEncoded, ESimpleElementBl
 				return *TShaderMapRef<FEncodedSimpleElement<TSimpleElementPixelShader, SE_BLEND_MaskedDistanceFieldShadowed> >(GetGlobalShaderMap(FeatureLevel));
 			case SE_BLEND_AlphaComposite:
 				return *TShaderMapRef<FEncodedSimpleElement<TSimpleElementPixelShader, SE_BLEND_AlphaComposite> >(GetGlobalShaderMap(FeatureLevel));
+			case SE_BLEND_AlphaHoldout:
+				return *TShaderMapRef<FEncodedSimpleElement<TSimpleElementPixelShader, SE_BLEND_AlphaHoldout> >(GetGlobalShaderMap(FeatureLevel));
 			case SE_BLEND_AlphaBlend:
 				return *TShaderMapRef<FEncodedSimpleElement<TSimpleElementPixelShader, SE_BLEND_AlphaBlend> >(GetGlobalShaderMap(FeatureLevel));
 			case SE_BLEND_TranslucentAlphaOnly:
diff --git a/Engine/Source/Runtime/Engine/Private/Materials/HLSLMaterialTranslator.h b/Engine/Source/Runtime/Engine/Private/Materials/HLSLMaterialTranslator.h
--- a/Engine/Source/Runtime/Engine/Private/Materials/HLSLMaterialTranslator.h
+++ b/Engine/Source/Runtime/Engine/Private/Materials/HLSLMaterialTranslator.h
@@ -811,6 +811,11 @@ public:
 				}
 			}
 
+			if (BlendMode == BLEND_AlphaHoldout && MaterialShadingModel != MSM_Unlit)
+			{
+				Errorf(TEXT("Alpha Holdout blend mode must use unlit shading model."));
+			}
+
 			if (Domain == MD_Volume && BlendMode != BLEND_Additive)
 			{
 				Errorf(TEXT("Volume materials must use an Additive blend mode."));
diff --git a/Engine/Source/Runtime/Engine/Private/Materials/Material.cpp b/Engine/Source/Runtime/Engine/Private/Materials/Material.cpp
--- a/Engine/Source/Runtime/Engine/Private/Materials/Material.cpp
+++ b/Engine/Source/Runtime/Engine/Private/Materials/Material.cpp
@@ -6034,7 +6034,7 @@ static bool IsPropertyActive_Internal(EMaterialProperty InProperty,
 		Active = false;
 		break;
 	case MP_Refraction:
-		Active = bIsTranslucentBlendMode && BlendMode != BLEND_Modulate;
+		Active = bIsTranslucentBlendMode && BlendMode != BLEND_AlphaHoldout && BlendMode != BLEND_Modulate;
 		break;
 	case MP_Opacity:
 		Active = bIsTranslucentBlendMode && BlendMode != BLEND_Modulate;
@@ -6075,8 +6075,9 @@ static bool IsPropertyActive_Internal(EMaterialProperty InProperty,
 		Active = bHasTessellation;
 		break;
 	case MP_EmissiveColor:
-		// Emissive is always active, even for light functions and post process materials
-		Active = true;
+		// Emissive is always active, even for light functions and post process materials, 
+		// but not for AlphaHoldout
+		Active = BlendMode != BLEND_AlphaHoldout;
 		break;
 	case MP_WorldPositionOffset:
 		Active = true;
diff --git a/Engine/Source/Runtime/Engine/Private/Materials/MaterialShader.cpp b/Engine/Source/Runtime/Engine/Private/Materials/MaterialShader.cpp
--- a/Engine/Source/Runtime/Engine/Private/Materials/MaterialShader.cpp
+++ b/Engine/Source/Runtime/Engine/Private/Materials/MaterialShader.cpp
@@ -88,6 +88,7 @@ FString GetBlendModeString(EBlendMode BlendMode)
 		case BLEND_Additive: BlendModeName = TEXT("BLEND_Additive"); break;
 		case BLEND_Modulate: BlendModeName = TEXT("BLEND_Modulate"); break;
 		case BLEND_AlphaComposite: BlendModeName = TEXT("BLEND_AlphaComposite"); break;
+		case BLEND_AlphaHoldout: BlendModeName = TEXT("BLEND_AlphaHoldout"); break;
 		default: BlendModeName = TEXT("Unknown"); break;
 	}
 	return BlendModeName;
diff --git a/Engine/Source/Runtime/Engine/Private/Materials/MaterialShared.cpp b/Engine/Source/Runtime/Engine/Private/Materials/MaterialShared.cpp
--- a/Engine/Source/Runtime/Engine/Private/Materials/MaterialShared.cpp
+++ b/Engine/Source/Runtime/Engine/Private/Materials/MaterialShared.cpp
@@ -1467,11 +1467,20 @@ void FMaterial::SetupMaterialEnvironment(
 			}
 			break;
 		}
-	case BLEND_AlphaComposite: 
-		{
-			// Fall through the case, blend mode will reuse MATERIALBLENDING_TRANSLUCENT
-			OutEnvironment.SetDefine(TEXT("MATERIALBLENDING_ALPHACOMPOSITE"), TEXT("1"));
-		}
+	case BLEND_AlphaComposite:
+	{
+		// Blend mode will reuse MATERIALBLENDING_TRANSLUCENT
+		OutEnvironment.SetDefine(TEXT("MATERIALBLENDING_ALPHACOMPOSITE"), TEXT("1"));
+		OutEnvironment.SetDefine(TEXT("MATERIALBLENDING_TRANSLUCENT"), TEXT("1"));
+		break;
+	}
+	case BLEND_AlphaHoldout:
+	{
+		// Blend mode will reuse MATERIALBLENDING_TRANSLUCENT
+		OutEnvironment.SetDefine(TEXT("MATERIALBLENDING_ALPHAHOLDOUT"), TEXT("1"));
+		OutEnvironment.SetDefine(TEXT("MATERIALBLENDING_TRANSLUCENT"), TEXT("1"));
+		break;
+	}
 	case BLEND_Translucent: OutEnvironment.SetDefine(TEXT("MATERIALBLENDING_TRANSLUCENT"),TEXT("1")); break;
 	case BLEND_Additive: OutEnvironment.SetDefine(TEXT("MATERIALBLENDING_ADDITIVE"),TEXT("1")); break;
 	case BLEND_Modulate: OutEnvironment.SetDefine(TEXT("MATERIALBLENDING_MODULATE"),TEXT("1")); break;
diff --git a/Engine/Source/Runtime/Engine/Private/Particles/ParticleSystemRender.cpp b/Engine/Source/Runtime/Engine/Private/Particles/ParticleSystemRender.cpp
--- a/Engine/Source/Runtime/Engine/Private/Particles/ParticleSystemRender.cpp
+++ b/Engine/Source/Runtime/Engine/Private/Particles/ParticleSystemRender.cpp
@@ -1111,7 +1111,7 @@ void FDynamicSpriteEmitterData::GetDynamicMeshElementsEmitter(const FParticleSys
 						const FMaterial* Material = MaterialResource ? MaterialResource->GetMaterial(FeatureLevel) : nullptr;
 
 						if (Material && 
-							(Material->GetBlendMode() == BLEND_Translucent || Material->GetBlendMode() == BLEND_AlphaComposite ||
+							(Material->GetBlendMode() == BLEND_Translucent || Material->GetBlendMode() == BLEND_AlphaComposite || Material->GetBlendMode() == BLEND_AlphaHoldout ||
 							((SourceData->SortMode == PSORTMODE_Age_OldestFirst) || (SourceData->SortMode == PSORTMODE_Age_NewestFirst)))
 							)
 						{
diff --git a/Engine/Source/Runtime/Engine/Private/SimpleElementShaders.cpp b/Engine/Source/Runtime/Engine/Private/SimpleElementShaders.cpp
--- a/Engine/Source/Runtime/Engine/Private/SimpleElementShaders.cpp
+++ b/Engine/Source/Runtime/Engine/Private/SimpleElementShaders.cpp
@@ -300,6 +300,7 @@ IMPLEMENT_SHADER_TYPE(,FSimpleElementColorChannelMaskPS,TEXT("/Engine/Private/Si
 	BLEND_VARIATION(SHADERCLASS, SE_BLEND_MaskedDistanceFieldShadowed, SHADERFILENAME, SHADERENTRYFUNC) \
 	BLEND_VARIATION(SHADERCLASS, SE_BLEND_TranslucentDistanceField, SHADERFILENAME, SHADERENTRYFUNC) \
 	BLEND_VARIATION(SHADERCLASS, SE_BLEND_AlphaComposite, SHADERFILENAME, SHADERENTRYFUNC) \
+	BLEND_VARIATION(SHADERCLASS, SE_BLEND_AlphaHoldout, SHADERFILENAME, SHADERENTRYFUNC) \
 	BLEND_VARIATION(SHADERCLASS, SE_BLEND_TranslucentDistanceFieldShadowed, SHADERFILENAME, SHADERENTRYFUNC) \
 	BLEND_VARIATION(SHADERCLASS, SE_BLEND_AlphaBlend, SHADERFILENAME, SHADERENTRYFUNC) \
 	BLEND_VARIATION(SHADERCLASS, SE_BLEND_TranslucentAlphaOnly, SHADERFILENAME, SHADERENTRYFUNC) \
diff --git a/Engine/Source/Runtime/Engine/Private/UserInterface/Canvas.cpp b/Engine/Source/Runtime/Engine/Private/UserInterface/Canvas.cpp
--- a/Engine/Source/Runtime/Engine/Private/UserInterface/Canvas.cpp
+++ b/Engine/Source/Runtime/Engine/Private/UserInterface/Canvas.cpp
@@ -1476,6 +1476,8 @@ ESimpleElementBlendMode FCanvas::BlendToSimpleElementBlend(EBlendMode BlendMode)
 			return SE_BLEND_Modulate;
 		case BLEND_AlphaComposite:
 			return SE_BLEND_AlphaComposite;
+		case BLEND_AlphaHoldout:
+			return SE_BLEND_AlphaHoldout;
 		case BLEND_Translucent:
 		default:
 			return SE_BLEND_Translucent;
diff --git a/Engine/Source/Runtime/Engine/Public/SceneTypes.h b/Engine/Source/Runtime/Engine/Public/SceneTypes.h
--- a/Engine/Source/Runtime/Engine/Public/SceneTypes.h
+++ b/Engine/Source/Runtime/Engine/Public/SceneTypes.h
@@ -182,6 +182,7 @@ enum ESimpleElementBlendMode
 	SE_BLEND_TranslucentDistanceField,
 	SE_BLEND_TranslucentDistanceFieldShadowed,
 	SE_BLEND_AlphaComposite,
+	SE_BLEND_AlphaHoldout,
 	// Like SE_BLEND_Translucent, but modifies destination alpha
 	SE_BLEND_AlphaBlend,
 	// Like SE_BLEND_Translucent, but reads from an alpha-only texture
diff --git a/Engine/Source/Runtime/Engine/Public/SimpleElementShaders.h b/Engine/Source/Runtime/Engine/Public/SimpleElementShaders.h
--- a/Engine/Source/Runtime/Engine/Public/SimpleElementShaders.h
+++ b/Engine/Source/Runtime/Engine/Public/SimpleElementShaders.h
@@ -341,6 +341,7 @@ public:
 		OutEnvironment.SetDefine(TEXT("SE_BLEND_MASKEDDISTANCEFIELD"), (uint32)ESimpleElementBlendMode::SE_BLEND_MaskedDistanceField);
 		OutEnvironment.SetDefine(TEXT("SE_BLEND_MASKEDDISTANCEFIELDSHADOWED"), (uint32)ESimpleElementBlendMode::SE_BLEND_MaskedDistanceFieldShadowed);
 		OutEnvironment.SetDefine(TEXT("SE_BLEND_ALPHACOMPOSITE"), (uint32)ESimpleElementBlendMode::SE_BLEND_AlphaComposite);
+		OutEnvironment.SetDefine(TEXT("SE_BLEND_ALPHAHOLDOUT"), (uint32)ESimpleElementBlendMode::SE_BLEND_AlphaHoldout);
 		OutEnvironment.SetDefine(TEXT("SE_BLEND_ALPHABLEND"), (uint32)ESimpleElementBlendMode::SE_BLEND_AlphaBlend);
 		OutEnvironment.SetDefine(TEXT("SE_BLEND_TRANSLUCENTALPHAONLY"), (uint32)ESimpleElementBlendMode::SE_BLEND_TranslucentAlphaOnly);
 		OutEnvironment.SetDefine(TEXT("SE_BLEND_TRANSLUCENTALPHAONLYWRITEALPHA"), (uint32)ESimpleElementBlendMode::SE_BLEND_TranslucentAlphaOnlyWriteAlpha);
diff --git a/Engine/Source/Runtime/Renderer/Private/BasePassRendering.cpp b/Engine/Source/Runtime/Renderer/Private/BasePassRendering.cpp
--- a/Engine/Source/Runtime/Renderer/Private/BasePassRendering.cpp
+++ b/Engine/Source/Runtime/Renderer/Private/BasePassRendering.cpp
@@ -193,6 +193,10 @@ void SetTranslucentRenderState(FMeshPassProcessorRenderState& DrawRenderState, c
 		// Blend with existing scene color. New color is already pre-multiplied by alpha.
 		DrawRenderState.SetBlendState(TStaticBlendState<CW_RGBA, BO_Add, BF_One, BF_InverseSourceAlpha, BO_Add, BF_Zero, BF_InverseSourceAlpha>::GetRHI());
 		break;
+	case BLEND_AlphaHoldout:
+		// Blend by holding out the matte shape of the source alpha
+		DrawRenderState.SetBlendState(TStaticBlendState<CW_RGBA, BO_Add, BF_Zero, BF_InverseSourceAlpha, BO_Add, BF_One, BF_InverseSourceAlpha>::GetRHI());
+		break;
 	};
 
 	const bool bDisableDepthTest = Material.ShouldDisableDepthTest();
diff --git a/Engine/Source/Runtime/Renderer/Private/MobileBasePass.cpp b/Engine/Source/Runtime/Renderer/Private/MobileBasePass.cpp
--- a/Engine/Source/Runtime/Renderer/Private/MobileBasePass.cpp
+++ b/Engine/Source/Runtime/Renderer/Private/MobileBasePass.cpp
@@ -412,6 +412,10 @@ void MobileBasePass::SetTranslucentRenderState(FMeshPassProcessorRenderState& Dr
 			// Blend with existing scene color. New color is already pre-multiplied by alpha.
 			DrawRenderState.SetBlendState(TStaticBlendState<CW_RGB, BO_Add, BF_One, BF_InverseSourceAlpha, BO_Add, BF_Zero, BF_InverseSourceAlpha>::GetRHI());
 			break;
+		case BLEND_AlphaHoldout:
+			// Blend by holding out the matte shape of the source alpha
+			DrawRenderState.SetBlendState(TStaticBlendState<CW_RGBA, BO_Add, BF_Zero, BF_InverseSourceAlpha, BO_Add, BF_One, BF_InverseSourceAlpha>::GetRHI());
+			break;
 		default:
 			check(0);
 		};
diff --git a/Engine/Source/Runtime/Renderer/Private/PostProcess/PostProcessMaterial.cpp b/Engine/Source/Runtime/Renderer/Private/PostProcess/PostProcessMaterial.cpp
--- a/Engine/Source/Runtime/Renderer/Private/PostProcess/PostProcessMaterial.cpp
+++ b/Engine/Source/Runtime/Renderer/Private/PostProcess/PostProcessMaterial.cpp
@@ -358,6 +358,7 @@ void FRCPassPostProcessMaterial::Process(FRenderingCompositePassContext& Context
 			TStaticBlendState<CW_RGB, BO_Add, BF_One, BF_One>::GetRHI(),
 			TStaticBlendState<CW_RGB, BO_Add, BF_DestColor, BF_Zero>::GetRHI(),
 			TStaticBlendState<CW_RGBA, BO_Add, BF_One, BF_InverseSourceAlpha, BO_Add, BF_One, BF_InverseSourceAlpha>::GetRHI(),
+			TStaticBlendState<CW_RGBA, BO_Add, BF_Zero, BF_InverseSourceAlpha, BO_Add, BF_Zero, BF_InverseSourceAlpha>::GetRHI(),
 		};
 		static_assert(EBlendMode::BLEND_MAX == ARRAY_COUNT(BlendStates), "Ensure that all EBlendMode values are accounted for.");
 
diff --git a/Engine/Source/Runtime/SlateRHIRenderer/Private/SlateMaterialShader.cpp b/Engine/Source/Runtime/SlateRHIRenderer/Private/SlateMaterialShader.cpp
--- a/Engine/Source/Runtime/SlateRHIRenderer/Private/SlateMaterialShader.cpp
+++ b/Engine/Source/Runtime/SlateRHIRenderer/Private/SlateMaterialShader.cpp
@@ -109,6 +109,11 @@ void FSlateMaterialShaderPS::SetBlendState(FGraphicsPipelineStateInitializer& Gr
 		// Blend with existing scene color. New color is already pre-multiplied by alpha.
 		GraphicsPSOInit.BlendState = TStaticBlendState<CW_RGBA, BO_Add, BF_One, BF_InverseSourceAlpha, BO_Add, BF_One, BF_InverseSourceAlpha>::GetRHI();
 		break;
+	case BLEND_AlphaHoldout:
+		// Blend by holding out the matte shape of the source alpha
+		GraphicsPSOInit.BlendState = TStaticBlendState<CW_RGBA, BO_Add, BF_Zero, BF_InverseSourceAlpha, BO_Add, BF_Zero, BF_InverseSourceAlpha>::GetRHI();
+		break;
+
 	};
 }
 
