Add r.Metal.ForceDXC to control DXC for Metal shaders.

diff --git a/Engine/Source/Runtime/Engine/Private/ShaderCompiler/ShaderCompiler.cpp b/Engine/Source/Runtime/Engine/Private/ShaderCompiler/ShaderCompiler.cpp
--- a/Engine/Source/Runtime/Engine/Private/ShaderCompiler/ShaderCompiler.cpp
+++ b/Engine/Source/Runtime/Engine/Private/ShaderCompiler/ShaderCompiler.cpp
@@ -243,6 +243,15 @@ static TAutoConsoleVariable<int32> CVarD3DForceDXC(
 	TEXT(" 0: Disable (default)\n")
 	TEXT(" 1: Force new compiler for all shaders"),
 	ECVF_ReadOnly);
+	);
+
+static TAutoConsoleVariable<int32> CVarMetalForceDXC(
+	TEXT("r.Metal.ForceDXC"),
+	0,
+	TEXT("Forces DirectX Shader Compiler (DXC) to be used for all Metal shaders instead of hlslcc.\n")
+	TEXT(" 0: Disable (default)\n")
+	TEXT(" 1: Force new compiler for all shaders"),
+	ECVF_ReadOnly);
 
 int32 GCreateShadersOnLoad = 0;
 static FAutoConsoleVariableRef CVarCreateShadersOnLoad(
@@ -3223,7 +3232,7 @@ void GlobalBeginCompileShader(
 			}
 		}
 		{
-			static const auto CVar = IConsoleManager::Get().FindConsoleVariable(TEXT("r.D3D.ForceDXC"));
+			static const auto CVar = IConsoleManager::Get().FindConsoleVariable(TEXT("r.Metal.ForceDXC"));
 			if (CVar && CVar->GetInt() != 0)
 			{
 				Input.Environment.CompilerFlags.Add(CFLAG_ForceDXC);
diff --git a/Engine/Source/Runtime/RenderCore/Private/Shader.cpp b/Engine/Source/Runtime/RenderCore/Private/Shader.cpp
--- a/Engine/Source/Runtime/RenderCore/Private/Shader.cpp
+++ b/Engine/Source/Runtime/RenderCore/Private/Shader.cpp
@@ -2316,7 +2316,7 @@ void ShaderMapAppendKeyString(EShaderPlatform Platform, FString& KeyString)
 			KeyString += TEXT("_ARCHIVE");
 		}
 		{
-			static const auto CVar = IConsoleManager::Get().FindConsoleVariable(TEXT("r.D3D.ForceDXC"));
+			static const auto CVar = IConsoleManager::Get().FindConsoleVariable(TEXT("r.Metal.ForceDXC"));
 			if (CVar && CVar->GetInt() != 0)
 			{
 				KeyString += TEXT("_DXC");
