Fix copy-paste mistake.

diff --git a/Engine/Source/Runtime/Engine/Private/ShaderCompiler/ShaderCompiler.cpp b/Engine/Source/Runtime/Engine/Private/ShaderCompiler/ShaderCompiler.cpp
--- a/Engine/Source/Runtime/Engine/Private/ShaderCompiler/ShaderCompiler.cpp
+++ b/Engine/Source/Runtime/Engine/Private/ShaderCompiler/ShaderCompiler.cpp
@@ -243,7 +243,6 @@ static TAutoConsoleVariable<int32> CVarD3DForceDXC(
 	TEXT(" 0: Disable (default)\n")
 	TEXT(" 1: Force new compiler for all shaders"),
 	ECVF_ReadOnly);
-	);
 
 static TAutoConsoleVariable<int32> CVarMetalForceDXC(
 	TEXT("r.Metal.ForceDXC"),
