Fixed ray tracing debug visualization views should not use tonemapping or any other post-processing

diff --git a/Engine/Source/Runtime/Engine/Private/ShowFlags.cpp b/Engine/Source/Runtime/Engine/Private/ShowFlags.cpp
--- a/Engine/Source/Runtime/Engine/Private/ShowFlags.cpp
+++ b/Engine/Source/Runtime/Engine/Private/ShowFlags.cpp
@@ -285,6 +285,12 @@ void ApplyViewMode(EViewModeIndex ViewModeIndex, bool bPerspective, FEngineShowF
 		case VMI_HLODColoration:
 			bPostProcessing = true;
 			break;
+		case VMI_RayTracingDebug:
+			bPostProcessing = true;
+			break;
+		case VMI_PathTracing:
+			bPostProcessing = true;
+			break;
 	}
 
 	if(!bPerspective)
@@ -318,6 +324,8 @@ void ApplyViewMode(EViewModeIndex ViewModeIndex, bool bPerspective, FEngineShowF
 	EngineShowFlags.SetCollisionVisibility(ViewModeIndex == VMI_CollisionVisibility);
 	EngineShowFlags.SetLODColoration(ViewModeIndex == VMI_LODColoration);
 	EngineShowFlags.SetHLODColoration(ViewModeIndex == VMI_HLODColoration);
+	EngineShowFlags.SetRayTracingDebug(ViewModeIndex == VMI_RayTracingDebug);
+	EngineShowFlags.SetPathTracing(ViewModeIndex == VMI_PathTracing);
 }
 
 void EngineShowFlagOverride(EShowFlagInitMode ShowFlagInitMode, EViewModeIndex ViewModeIndex, FEngineShowFlags& EngineShowFlags, bool bCanDisableTonemapper)
@@ -474,6 +482,11 @@ void EngineShowFlagOverride(EShowFlagInitMode ShowFlagInitMode, EViewModeIndex V
 		if (ViewModeIndex == VMI_RayTracingDebug)
 		{
 			EngineShowFlags.SetRayTracingDebug(true);
+			EngineShowFlags.VisualizeHDR = 0;
+			EngineShowFlags.VisualizeBloom = 0;
+			EngineShowFlags.VisualizeMotionBlur = 0;
+			EngineShowFlags.DepthOfField = 0;
+			EngineShowFlags.PostProcessMaterial = 0;
 		}
 	}
 
@@ -686,6 +699,8 @@ const TCHAR* GetViewModeName(EViewModeIndex ViewModeIndex)
 		case VMI_LitLightmapDensity:		return TEXT("LitLightmapDensity");
 		case VMI_ReflectionOverride:		return TEXT("ReflectionOverride");
 		case VMI_VisualizeBuffer:			return TEXT("VisualizeBuffer");
+		case VMI_RayTracingDebug:			return TEXT("RayTracingDebug");
+		case VMI_PathTracing:				return TEXT("PathTracing");
 		case VMI_CollisionPawn:				return TEXT("CollisionPawn");
 		case VMI_CollisionVisibility:		return TEXT("CollisionVis");
 		case VMI_LODColoration:				return TEXT("LODColoration");
