Fixed unlit materials do not show in Ray Tracing reflections

(Small side effect change/opt: LightCullMask calculations now only happen when direct lighting flag is enabled.)

diff --git a/Engine/Shaders/Private/RayTracing/RayTracingLightingCommon.ush b/Engine/Shaders/Private/RayTracing/RayTracingLightingCommon.ush
--- a/Engine/Shaders/Private/RayTracing/RayTracingLightingCommon.ush
+++ b/Engine/Shaders/Private/RayTracing/RayTracingLightingCommon.ush
@@ -521,16 +521,22 @@ void AccumulateResults(
 {
 	if (IsHit(Ray, Payload))
 	{
-		uint2 LightCullMask = (uint2)0;
-		const float3 WorldPosition = Ray.Origin + Payload.HitT * Ray.Direction;
-		LightCullMask = CullDirectLighting(Ray, RayCone, WorldPosition, Payload.WorldNormal, RandSequence, ReflectedShadowsType > 0, ShadowMaxNormalBias);
 
 		Payload.DiffuseColor = Payload.BaseColor - Payload.BaseColor * Payload.Metallic;
 		Payload.SpecularColor = ComputeF0(Payload.Specular, Payload.BaseColor, Payload.Metallic);
 
 		if (ShouldDoDirectLighting)
 		{
-			float3 DirectLighting = ComputeDirectLightingCulled(LightCullMask, Ray, RayCone, TLAS, Payload, RandSequence, ReflectedShadowsType, ShadowMaxNormalBias);
+			uint2 LightCullMask = (uint2)0;
+			const float3 WorldPosition = Ray.Origin + Payload.HitT * Ray.Direction;
+			LightCullMask = CullDirectLighting(Ray, RayCone, WorldPosition, Payload.WorldNormal, RandSequence, ReflectedShadowsType > 0, ShadowMaxNormalBias);
+			
+			float3 DirectLighting = 0;
+				
+			if (Payload.ShadingModelID != SHADINGMODELID_UNLIT)
+			{
+				DirectLighting = ComputeDirectLightingCulled(LightCullMask, Ray, RayCone, TLAS, Payload, RandSequence, ReflectedShadowsType, ShadowMaxNormalBias);
+			}
 
 			// Transform NaNs to black, transform negative colors to black.
 			DirectLighting = -min(-DirectLighting, 0.0);
diff --git a/Engine/Source/Runtime/Engine/Private/StaticMeshRender.cpp b/Engine/Source/Runtime/Engine/Private/StaticMeshRender.cpp
--- a/Engine/Source/Runtime/Engine/Private/StaticMeshRender.cpp
+++ b/Engine/Source/Runtime/Engine/Private/StaticMeshRender.cpp
@@ -196,7 +196,6 @@ FStaticMeshSceneProxy::FStaticMeshSceneProxy(UStaticMeshComponent* InComponent,
 		// Under certain error conditions an LOD's material will be set to 
 		// DefaultMaterial. Ensure our material view relevance is set properly.
 		const int32 NumSections = NewLODInfo->Sections.Num();
-		bool bHasNonUnlitSections = false;
 		for (int32 SectionIndex = 0; SectionIndex < NumSections; ++SectionIndex)
 		{
 			const FLODInfo::FSectionInfo& SectionInfo = NewLODInfo->Sections[SectionIndex];
@@ -205,22 +204,10 @@ FStaticMeshSceneProxy::FStaticMeshSceneProxy(UStaticMeshComponent* InComponent,
 			{
 				MaterialRelevance |= UMaterial::GetDefaultMaterial(MD_Surface)->GetRelevance(FeatureLevel);
 			}
-
-			// #dxr_todo UE-72552: (skybox should be excluded from RT by default) warning this is a hack:
-			// #dxr_todo UE-72554 (handling of unlit geometries):
-			//  - The only reason we exclude unlit from any RT acceleration structure right now is to not hit the sky dome. 
-			//  - It works but this would need consideration (e.g. infinite RTAO would not work if the sky dome is hit, double reflection due sky light cubemap, reflection testing no intersection, etc.).
-			if (SectionInfo.Material->GetShadingModel() != MSM_Unlit)
-			{
-				bHasNonUnlitSections = true;
-			}
 		}
 
 #if RHI_RAYTRACING
-		if (bHasNonUnlitSections)
-		{
-			RayTracingGeometries[LODIndex] = RenderData->LODResources[LODIndex].RayTracingGeometry.RayTracingGeometryRHI;
-		}
+		RayTracingGeometries[LODIndex] = RenderData->LODResources[LODIndex].RayTracingGeometry.RayTracingGeometryRHI;
 #endif
 	}
 
