Rename/move file to match current component file convention.

Only change in the moved file is added #include "Atmosphere/AtmosphericFogComponent.h" because it was required to compile.

Discussed with Daniel

diff --git a/Engine/Source/Runtime/Engine/Private/Atmosphere/Atmosphere.cpp b/Engine/Source/Runtime/Engine/Private/Atmosphere/Atmosphere.cpp
deleted file mode 100644
diff --git a/Engine/Source/Runtime/Engine/Private/Atmosphere/AtmosphericFogComponent.cpp b/Engine/Source/Runtime/Engine/Private/Atmosphere/AtmosphericFogComponent.cpp
new file mode 100644
--- /dev/null
+++ b/Engine/Source/Runtime/Engine/Private/Atmosphere/AtmosphericFogComponent.cpp
@@ -0,0 +1,850 @@
+// Copyright 1998-2019 Epic Games, Inc. All Rights Reserved.
+
+#include "Atmosphere/AtmosphericFogComponent.h"
+#include "Atmosphere/Atmosphere.h"
+#include "UObject/ConstructorHelpers.h"
+#include "EngineDefines.h"
+#include "RenderingThread.h"
+#include "Components/ArrowComponent.h"
+#include "UObject/UObjectHash.h"
+#include "UObject/UObjectIterator.h"
+#include "Engine/Texture2D.h"
+#include "ComponentReregisterContext.h"
+#include "Atmosphere/AtmosphericFog.h"
+#include "Components/BillboardComponent.h"
+#include "Runtime/Renderer/Private/ScenePrivate.h"
+#include "Runtime/Renderer/Private/AtmosphereRendering.h"
+
+#if WITH_EDITOR
+#include "ObjectEditorUtils.h"
+#endif
+
+AAtmosphericFog::AAtmosphericFog(const FObjectInitializer& ObjectInitializer)
+	: Super(ObjectInitializer)
+{
+	AtmosphericFogComponent = CreateDefaultSubobject<UAtmosphericFogComponent>(TEXT("AtmosphericFogComponent0"));
+	RootComponent = AtmosphericFogComponent;
+
+#if WITH_EDITORONLY_DATA
+	ArrowComponent = CreateEditorOnlyDefaultSubobject<UArrowComponent>(TEXT("ArrowComponent0"));
+
+	if (!IsRunningCommandlet())
+	{
+		// Structure to hold one-time initialization
+		struct FConstructorStatics
+		{
+			ConstructorHelpers::FObjectFinderOptional<UTexture2D> FogTextureObject;
+			FName ID_Fog;
+			FText NAME_Fog;
+			FConstructorStatics()
+				: FogTextureObject(TEXT("/Engine/EditorResources/S_ExpoHeightFog"))
+				, ID_Fog(TEXT("Fog"))
+				, NAME_Fog(NSLOCTEXT("SpriteCategory", "Fog", "Fog"))
+			{
+			}
+		};
+		static FConstructorStatics ConstructorStatics;
+
+		if (GetSpriteComponent())
+		{
+			GetSpriteComponent()->Sprite = ConstructorStatics.FogTextureObject.Get();
+			GetSpriteComponent()->RelativeScale3D = FVector(0.5f, 0.5f, 0.5f);
+			GetSpriteComponent()->SpriteInfo.Category = ConstructorStatics.ID_Fog;
+			GetSpriteComponent()->SpriteInfo.DisplayName = ConstructorStatics.NAME_Fog;
+			GetSpriteComponent()->SetupAttachment(AtmosphericFogComponent);
+		}
+
+		if (ArrowComponent)
+		{
+			ArrowComponent->ArrowColor = FColor(150, 200, 255);
+
+			ArrowComponent->bTreatAsASprite = true;
+			ArrowComponent->SpriteInfo.Category = ConstructorStatics.ID_Fog;
+			ArrowComponent->SpriteInfo.DisplayName = ConstructorStatics.NAME_Fog;
+			ArrowComponent->SetupAttachment(AtmosphericFogComponent);
+			ArrowComponent->bLightAttachment = true;
+			ArrowComponent->bIsScreenSizeScaled = true;
+		}
+	}
+#endif // WITH_EDITORONLY_DATA
+
+
+	PrimaryActorTick.bCanEverTick = true;
+	bHidden = false;
+}
+
+// Precomputation parameters
+FAtmospherePrecomputeParameters::FAtmospherePrecomputeParameters():
+	DensityHeight(0.5f),
+	DecayHeight_DEPRECATED(0.5f),
+	MaxScatteringOrder(4),
+	TransmittanceTexWidth(256),
+	TransmittanceTexHeight(64),
+	IrradianceTexWidth(64),
+	IrradianceTexHeight(16),
+	InscatterAltitudeSampleNum(2),
+	InscatterMuNum(128),
+	InscatterMuSNum(32),
+	InscatterNuNum(8)
+{
+}
+
+UAtmosphericFogComponent::UAtmosphericFogComponent(const FObjectInitializer& ObjectInitializer)
+	: Super(ObjectInitializer), TransmittanceResource(NULL), IrradianceResource(NULL), InscatterResource(NULL)
+#if WITH_EDITORONLY_DATA
+	, PrecomputeDataHandler(NULL)
+#endif
+{
+	SunMultiplier = 1.f;
+	FogMultiplier = 1.f;
+	DensityMultiplier = 1.f;
+	DensityOffset = 0.f;
+	DistanceScale = 1.f;
+	AltitudeScale = 1.f;
+	SunDiscScale = 1.f;
+	StartDistance = 15000;
+	DistanceOffset = 0.f;
+	GroundOffset = -100000.f; // -1 km in default distance scale, 100K UU
+
+	// Default lighting
+	DefaultBrightness = 50.f;
+	DefaultLightColor = FColor::White;
+
+	bDisableSunDisk = false;
+	bDisableGroundScattering = false;
+}
+
+bool UAtmosphericFogComponent::IsPostLoadThreadSafe() const
+{
+	return true;
+}
+
+void UAtmosphericFogComponent::PostLoad()
+{
+	Super::PostLoad();
+
+	if ( !IsTemplate() )
+	{
+#if WITH_EDITOR
+		if (TransmittanceTexture_DEPRECATED)
+		{
+			// Copy data from Previous data
+			if (TransmittanceTexture_DEPRECATED->Source.IsValid())
+			{
+				TArray<uint8> RawData;
+				TArray<FColor> OutData;
+				TransmittanceTexture_DEPRECATED->Source.GetMipData(RawData, 0);
+				// Convert from FFloat16Color to FColor
+				for (int32 i = 0; i < RawData.Num(); i+=sizeof(FFloat16Color))
+				{
+					FFloat16Color* OriginalColor = (FFloat16Color*)&RawData[i];
+					FColor TempColor;
+					TempColor.R = FMath::Clamp<uint8>(OriginalColor->R.GetFloat() * 255, 0, 255);
+					TempColor.G = FMath::Clamp<uint8>(OriginalColor->G.GetFloat() * 255, 0, 255);
+					TempColor.B = FMath::Clamp<uint8>(OriginalColor->B.GetFloat() * 255, 0, 255);
+					OutData.Add(TempColor);
+				}
+
+				int32 TotalByte = OutData.Num() * sizeof(FColor);
+				TransmittanceData.Lock(LOCK_READ_WRITE);
+				FColor* TextureData = (FColor*)TransmittanceData.Realloc(TotalByte);
+				FMemory::Memcpy(TextureData, OutData.GetData(), TotalByte);
+				TransmittanceData.Unlock();
+
+				TransmittanceTexture_DEPRECATED = NULL;
+			}
+		}
+
+		if (IrradianceTexture_DEPRECATED)
+		{
+			if (IrradianceTexture_DEPRECATED->Source.IsValid())
+			{
+				TArray<uint8> RawData;
+				TArray<FColor> OutData;
+				IrradianceTexture_DEPRECATED->Source.GetMipData(RawData, 0);
+				// Convert from FFloat16Color to FColor
+				for (int32 i = 0; i < RawData.Num(); i+=sizeof(FFloat16Color))
+				{
+					FFloat16Color* OriginalColor = (FFloat16Color*)&RawData[i];
+					FColor TempColor;
+					TempColor.R = FMath::Clamp<uint8>(OriginalColor->R.GetFloat() * 255, 0, 255);
+					TempColor.G = FMath::Clamp<uint8>(OriginalColor->G.GetFloat() * 255, 0, 255);
+					TempColor.B = FMath::Clamp<uint8>(OriginalColor->B.GetFloat() * 255, 0, 255);
+					OutData.Add(TempColor);
+				}
+
+				int32 TotalByte = OutData.Num() * sizeof(FColor);
+				IrradianceData.Lock(LOCK_READ_WRITE);
+				FColor* TextureData = (FColor*)IrradianceData.Realloc(TotalByte);
+				FMemory::Memcpy(TextureData, OutData.GetData(), TotalByte);
+				IrradianceData.Unlock();
+
+				IrradianceTexture_DEPRECATED = NULL;
+			}
+		}
+#endif
+		InitResource();
+	}
+}
+
+#if WITH_EDITOR
+#include "TickableEditorObject.h"
+
+class FAtmospherePrecomputeDataHandler : public FTickableEditorObject
+{
+public:
+	UAtmosphericFogComponent* Component;
+
+	FAtmospherePrecomputeDataHandler(UAtmosphericFogComponent* InComponent) : Component(InComponent)
+	{
+	}
+
+	/** FTickableEditorObject interface */
+	virtual void Tick( float DeltaTime ) override
+	{
+		if( Component && Component->GameThreadServiceRequest.GetValue()) 
+		{
+			Component->UpdatePrecomputedData();
+			Component->GameThreadServiceRequest.Decrement();
+		}              
+	}
+
+	virtual ETickableTickType GetTickableTickType() const override
+	{
+		return ETickableTickType::Always;
+	}
+
+	TStatId GetStatId() const override
+	{
+		RETURN_QUICK_DECLARE_CYCLE_STAT(FAtmospherePrecomputeDataHandler, STATGROUP_Tickables);
+	}
+};
+#endif
+
+
+static TAutoConsoleVariable<int32> CVarAtmosphereRender(
+	TEXT("r.Atmosphere"),
+	1,
+	TEXT("Defines atmosphere will render or not. Only changed by r.Atmosphere console command.\n")
+	TEXT("Enable/Disable Atmosphere, Load/Unload related data.\n")
+	TEXT(" 0: off (To save GPU memory)\n")
+	TEXT(" 1: on (default)"));
+
+// On CPU
+void UAtmosphericFogComponent::InitResource()
+{
+	bool bNewAtmosphere = CVarAtmosphereRender.GetValueOnGameThread() != 0;
+
+	if(!bNewAtmosphere)
+	{
+		return; // Don't do initialize resource when Atmosphere Render is off
+	}
+
+	if (PrecomputeCounter >= EValid)
+	{
+		// A little inefficient for thread-safe
+		if (TransmittanceData.GetElementCount() && TransmittanceResource == NULL)
+		{
+			TransmittanceResource = new FAtmosphereTextureResource( PrecomputeParams, TransmittanceData, FAtmosphereTextureResource::E_Transmittance ); 
+			BeginInitResource( TransmittanceResource );
+		}
+
+		if (IrradianceData.GetElementCount() && IrradianceResource == NULL)
+		{
+			IrradianceResource = new FAtmosphereTextureResource( PrecomputeParams, IrradianceData, FAtmosphereTextureResource::E_Irradiance ); 
+			BeginInitResource( IrradianceResource );
+		}
+
+		if (InscatterData.GetElementCount() && InscatterResource == NULL)
+		{
+			InscatterResource = new FAtmosphereTextureResource( PrecomputeParams, InscatterData, FAtmosphereTextureResource::E_Inscatter ); 
+			BeginInitResource( InscatterResource );
+		}
+	}
+#if WITH_EDITOR
+	else if (!PrecomputeDataHandler && !IsTemplate())
+	{
+		PrecomputeDataHandler = new FAtmospherePrecomputeDataHandler(this);
+	}
+#endif
+}
+
+void UAtmosphericFogComponent::BeginDestroy()
+{
+	ReleaseResource();
+	Super::BeginDestroy();
+}
+
+void UAtmosphericFogComponent::ReleaseResource()
+{
+	FSceneInterface* Scene = GetScene();
+	if ( TransmittanceResource != nullptr )
+	{
+		FRenderResource* InTransmittanceResource = TransmittanceResource;
+		ENQUEUE_RENDER_COMMAND(ReleaseAtmosphericFogTransmittanceTextureCommand)(
+			[InTransmittanceResource, Scene](FRHICommandListImmediate& RHICmdList)
+			{
+				if (Scene)
+				{
+					Scene->RemoveAtmosphericFogResource_RenderThread(InTransmittanceResource);
+				}
+				InTransmittanceResource->ReleaseResource();
+				delete InTransmittanceResource;
+			});
+
+		TransmittanceResource = nullptr;
+	}
+
+	if ( IrradianceResource != nullptr)
+	{
+		FRenderResource* InIrradianceResource = IrradianceResource;
+		ENQUEUE_RENDER_COMMAND(ReleaseAtmosphericFogIrradianceTextureCommand)(
+			[InIrradianceResource, Scene](FRHICommandListImmediate& RHICmdList)
+			{
+				if (Scene)
+				{
+					Scene->RemoveAtmosphericFogResource_RenderThread(InIrradianceResource);
+				}
+				InIrradianceResource->ReleaseResource();
+				delete InIrradianceResource;
+			});
+
+		IrradianceResource = nullptr;
+	}
+
+	if ( InscatterResource != nullptr )
+	{
+		FRenderResource* InInscatterResource = InscatterResource;
+		ENQUEUE_RENDER_COMMAND(ReleaseAtmosphericFogInscatterTextureCommand)(
+			[InInscatterResource, Scene](FRHICommandListImmediate& RHICmdList)
+			{
+				if (Scene)
+				{
+					Scene->RemoveAtmosphericFogResource_RenderThread(InInscatterResource);
+				}
+				InInscatterResource->ReleaseResource();
+				delete InInscatterResource;
+			});
+
+		InscatterResource = nullptr;
+	}
+}
+
+void UAtmosphericFogComponent::CreateRenderState_Concurrent()
+{
+	Super::CreateRenderState_Concurrent();
+	AddFogIfNeeded();
+}
+
+void UAtmosphericFogComponent::SendRenderTransform_Concurrent()
+{
+	GetWorld()->Scene->RemoveAtmosphericFog(this);
+	AddFogIfNeeded();
+	Super::SendRenderTransform_Concurrent();
+}
+
+void UAtmosphericFogComponent::DestroyRenderState_Concurrent()
+{
+	Super::DestroyRenderState_Concurrent();
+	GetWorld()->Scene->RemoveAtmosphericFog(this);
+	// Search for new fog component
+	for(TObjectIterator<UAtmosphericFogComponent> It; It; ++It)
+	{
+		UAtmosphericFogComponent* Component = *It;
+		checkSlow(Component);
+		if (Component != this && Component->IsRegistered())
+		{
+			Component->AddFogIfNeeded();
+			break;
+		}
+	}
+}
+
+/** Set brightness of the light */
+void UAtmosphericFogComponent::SetDefaultBrightness(float NewBrightness)
+{
+	if( DefaultBrightness != NewBrightness )
+	{
+		DefaultBrightness = NewBrightness;
+		MarkRenderStateDirty();
+	}
+}
+
+/** Set color of the light */
+void UAtmosphericFogComponent::SetDefaultLightColor(FLinearColor NewLightColor)
+{
+	FColor NewColor(NewLightColor.ToFColor(true));
+	if( DefaultLightColor != NewColor )
+	{
+		DefaultLightColor = NewColor;
+		MarkRenderStateDirty();
+	}
+}
+
+void UAtmosphericFogComponent::SetSunMultiplier(float NewSunMultiplier)
+{
+	if( SunMultiplier != NewSunMultiplier )
+	{
+		SunMultiplier = NewSunMultiplier;
+		MarkRenderStateDirty();
+	}
+}
+
+void UAtmosphericFogComponent::SetFogMultiplier(float NewFogMultiplier)
+{
+	if( FogMultiplier != NewFogMultiplier )
+	{
+		FogMultiplier = NewFogMultiplier;
+		MarkRenderStateDirty();
+	}
+}
+
+void UAtmosphericFogComponent::SetDensityMultiplier(float NewDensityMultiplier)
+{
+	if( DensityMultiplier != NewDensityMultiplier )
+	{
+		DensityMultiplier = NewDensityMultiplier;
+		MarkRenderStateDirty();
+	}
+}
+
+void UAtmosphericFogComponent::SetDensityOffset(float NewDensityOffset)
+{
+	if( DensityOffset != NewDensityOffset )
+	{
+		DensityOffset = NewDensityOffset;
+		MarkRenderStateDirty();
+	}
+}
+
+void UAtmosphericFogComponent::SetDistanceScale(float NewDistanceScale)
+{
+	if( DistanceScale != NewDistanceScale )
+	{
+		DistanceScale = NewDistanceScale;
+		MarkRenderStateDirty();
+	}
+}
+
+void UAtmosphericFogComponent::SetAltitudeScale(float NewAltitudeScale)
+{
+	if( AltitudeScale != NewAltitudeScale )
+	{
+		AltitudeScale = NewAltitudeScale;
+		MarkRenderStateDirty();
+	}
+}
+
+void UAtmosphericFogComponent::SetStartDistance(float NewStartDistance)
+{
+	if( StartDistance != NewStartDistance )
+	{
+		StartDistance = NewStartDistance;
+		MarkRenderStateDirty();
+	}
+}
+
+void UAtmosphericFogComponent::SetDistanceOffset(float NewDistanceOffset)
+{
+	if( DistanceOffset != NewDistanceOffset )
+	{
+		DistanceOffset = NewDistanceOffset;
+		MarkRenderStateDirty();
+	}
+}
+
+void UAtmosphericFogComponent::DisableSunDisk(bool NewSunDisk)
+{
+	if( bDisableSunDisk != NewSunDisk )
+	{
+		bDisableSunDisk = NewSunDisk;
+		MarkRenderStateDirty();
+	}
+}
+
+void UAtmosphericFogComponent::DisableGroundScattering(bool NewGroundScattering)
+{
+	if( bDisableGroundScattering != NewGroundScattering )
+	{
+		bDisableGroundScattering = NewGroundScattering;
+		MarkRenderStateDirty();
+	}
+}
+
+void UAtmosphericFogComponent::SetPrecomputeParams(float DensityHeight, int32 MaxScatteringOrder, int32 InscatterAltitudeSampleNum)
+{
+#if WITH_EDITOR
+	// Range for valid value
+	if (DensityHeight <= 0.f) // Assume default value, using default
+	{
+		DensityHeight = 0.5f;
+	}
+
+	if (MaxScatteringOrder <= 0)
+	{
+		MaxScatteringOrder = 4;
+	}
+
+	if (InscatterAltitudeSampleNum <= 0)
+	{
+		InscatterAltitudeSampleNum = 2;
+	}
+
+	FAtmospherePrecomputeParameters NewParams = PrecomputeParams;
+
+	DensityHeight = FMath::Clamp(DensityHeight, 0.1f, 1.f);
+	MaxScatteringOrder = FMath::Clamp(MaxScatteringOrder, 1, 4);
+	InscatterAltitudeSampleNum = FMath::Clamp(InscatterAltitudeSampleNum, 2, 32);
+
+	NewParams.DensityHeight = DensityHeight;
+	NewParams.MaxScatteringOrder = MaxScatteringOrder;
+	NewParams.InscatterAltitudeSampleNum = InscatterAltitudeSampleNum;
+
+	if (PrecomputeParams != NewParams)
+	{
+		PrecomputeParams = NewParams;
+		StartPrecompute();
+	}
+#endif
+}
+
+void UAtmosphericFogComponent::AddFogIfNeeded()
+{
+	if (ShouldComponentAddToScene() && ShouldRender() && IsRegistered() && SunMultiplier > DELTA && FogMultiplier > DELTA
+		&& (GetOuter() == NULL || !GetOuter()->HasAnyFlags(RF_ClassDefaultObject)))
+	{
+		GetWorld()->Scene->AddAtmosphericFog(this);
+	}
+}
+
+#if WITH_EDITOR
+void UAtmosphericFogComponent::PostEditChangeProperty(FPropertyChangedEvent& PropertyChangedEvent)
+{
+	const FName CategoryName = FObjectEditorUtils::GetCategoryFName(PropertyChangedEvent.Property);
+
+	bool bNeedsPrecompute = false;
+	if( CategoryName == FName(TEXT("AtmosphereParam")) )
+	{
+		// Recompute when precompute parameters were changed
+		PrecomputeParams.DensityHeight = FMath::Clamp(PrecomputeParams.DensityHeight, 0.1f, 1.f);
+		PrecomputeParams.MaxScatteringOrder = FMath::Clamp(PrecomputeParams.MaxScatteringOrder, 1, 4);
+		PrecomputeParams.InscatterAltitudeSampleNum = FMath::Clamp(PrecomputeParams.InscatterAltitudeSampleNum, 2, 32);
+		bNeedsPrecompute = true;
+	}
+	else
+	{
+		SunMultiplier = FMath::Clamp(SunMultiplier, 0.0f, 1000.f);
+		FogMultiplier = FMath::Clamp(FogMultiplier, 0.0f, 1000.f);
+		DensityMultiplier = FMath::Clamp(DensityMultiplier, 0.001f, 1000.f);
+		DensityOffset = FMath::Clamp(DensityOffset, -1.0f, 1.f);
+		DistanceScale = FMath::Clamp(DistanceScale, 0.1f, 1000.f);
+		AltitudeScale = FMath::Clamp(AltitudeScale, 0.1f, 1000.f);
+		SunDiscScale = FMath::Clamp(SunDiscScale, 0.1f, 1000.f);
+		GroundOffset = FMath::Clamp(GroundOffset, (float)-WORLD_MAX, (float)WORLD_MAX);
+		StartDistance = FMath::Clamp(StartDistance, 100.f, (float)WORLD_MAX);
+	}
+
+	Super::PostEditChangeProperty(PropertyChangedEvent);
+
+	if (bNeedsPrecompute)
+	{
+		StartPrecompute();
+	}
+}
+#endif // WITH_EDITOR
+
+void UAtmosphericFogComponent::PostInterpChange(UProperty* PropertyThatChanged)
+{
+	Super::PostInterpChange(PropertyThatChanged);
+
+	MarkRenderStateDirty();
+}
+
+void UAtmosphericFogComponent::StartPrecompute()
+{
+#if WITH_EDITOR
+	if (GIsEditor && !IsTemplate() && !GUsingNullRHI)
+	{
+		FSceneInterface* AtmosphericFogScene = GetScene();
+		if (AtmosphericFogScene)
+		{
+			PrecomputeCounter = EInvalid;
+
+			if (!PrecomputeDataHandler)
+			{
+				PrecomputeDataHandler = new FAtmospherePrecomputeDataHandler(this);
+			}
+
+			// This is largely redundant, the component will be reattached anyway, thus it will be recomputed.
+			UAtmosphericFogComponent* Component = this;
+			ENQUEUE_RENDER_COMMAND(FStartPrecompute)(
+				[AtmosphericFogScene, Component](FRHICommandListImmediate& RHICmdList)
+				{
+					FAtmosphericFogSceneInfo* AtmosphericFogSceneInfo = AtmosphericFogScene->GetAtmosphericFogSceneInfo();
+
+					if (AtmosphericFogSceneInfo && AtmosphericFogSceneInfo->Component == Component)
+					{
+						AtmosphericFogSceneInfo->bNeedRecompute = true;
+					}
+				});
+		}
+	}
+#endif
+}
+UAtmosphericFogComponent::~UAtmosphericFogComponent()
+{
+#if WITH_EDITOR
+	if (PrecomputeDataHandler)
+	{
+		delete PrecomputeDataHandler;
+		PrecomputeDataHandler = NULL;
+	}
+#endif
+}
+
+#if WITH_EDITOR
+// Prepare render targets when new actor spawned
+void AAtmosphericFog::PostActorCreated()
+{
+	Super::PostActorCreated();
+	if (GIsEditor)
+	{
+		if ( !IsTemplate() && AtmosphericFogComponent )
+		{
+			AtmosphericFogComponent->InitResource();
+		}
+	}
+}
+
+void UAtmosphericFogComponent::UpdatePrecomputedData()
+{
+	if (GIsEditor) 
+	{
+		// Prevent atmosphere pre-computation texture read/write from rendering thread during this process
+		FlushRenderingCommands();
+		FScene* Scene = GetScene() ? GetScene()->GetRenderScene() : NULL;
+		if (Scene && Scene->AtmosphericFog && this == Scene->AtmosphericFog->Component && Scene->AtmosphericFog->bPrecomputationFinished && !Scene->AtmosphericFog->bPrecomputationAcceptedByGameThread)
+		{
+			// When the precomputation is done, should save result to final texture for rendering
+			// Need to resolve to actor textures and remove render targets
+			// Resolve to data...
+			{
+				int32 SizeX = PrecomputeParams.TransmittanceTexWidth;
+				int32 SizeY = PrecomputeParams.TransmittanceTexHeight;
+				int32 TotalByte = sizeof(FColor) * SizeX * SizeY;
+				check(TotalByte == Scene->AtmosphericFog->PrecomputeTransmittance.GetBulkDataSize());
+				const FColor* PrecomputeData = (const FColor*)Scene->AtmosphericFog->PrecomputeTransmittance.Lock(LOCK_READ_ONLY);
+				TransmittanceData.Lock(LOCK_READ_WRITE);
+				FColor* TextureData = (FColor*)TransmittanceData.Realloc(TotalByte);
+				FMemory::Memcpy(TextureData, PrecomputeData, TotalByte);
+				TransmittanceData.Unlock();
+				Scene->AtmosphericFog->PrecomputeTransmittance.Unlock();
+			}
+
+			{
+				int32 SizeX = PrecomputeParams.IrradianceTexWidth;
+				int32 SizeY = PrecomputeParams.IrradianceTexHeight;
+				int32 TotalByte = sizeof(FColor) * SizeX * SizeY;
+				check(TotalByte == Scene->AtmosphericFog->PrecomputeIrradiance.GetBulkDataSize());
+				const FColor* PrecomputeData = (const FColor*)Scene->AtmosphericFog->PrecomputeIrradiance.Lock(LOCK_READ_ONLY);
+				IrradianceData.Lock(LOCK_READ_WRITE);
+				FColor* TextureData = (FColor*)IrradianceData.Realloc(TotalByte);
+				FMemory::Memcpy(TextureData, PrecomputeData, TotalByte);
+				IrradianceData.Unlock();
+				Scene->AtmosphericFog->PrecomputeIrradiance.Unlock();
+			}
+
+			{
+				int32 SizeX = PrecomputeParams.InscatterMuSNum * PrecomputeParams.InscatterNuNum;
+				int32 SizeY = PrecomputeParams.InscatterMuNum;
+				int32 SizeZ = PrecomputeParams.InscatterAltitudeSampleNum;
+				int32 TotalByte = sizeof(FFloat16Color) * SizeX * SizeY * SizeZ;
+				check(TotalByte == Scene->AtmosphericFog->PrecomputeInscatter.GetBulkDataSize());
+				const FFloat16Color* PrecomputeData = (const FFloat16Color*)Scene->AtmosphericFog->PrecomputeInscatter.Lock(LOCK_READ_ONLY);
+				InscatterData.Lock(LOCK_READ_WRITE);
+				FFloat16Color* TextureData = (FFloat16Color*)InscatterData.Realloc(TotalByte);
+				FMemory::Memcpy(TextureData, PrecomputeData, TotalByte);
+				InscatterData.Unlock();
+				Scene->AtmosphericFog->PrecomputeInscatter.Unlock();
+			}
+			PrecomputeCounter = EValid;
+			FPlatformMisc::MemoryBarrier();
+			Scene->AtmosphericFog->bPrecomputationAcceptedByGameThread = true;
+
+			// Resolve to data...
+			ReleaseResource();
+			// Wait for release...
+			FlushRenderingCommands();
+
+			InitResource();
+			FComponentReregisterContext ReregisterContext(this);
+		}
+	}
+}
+#endif
+
+void UAtmosphericFogComponent::Serialize(FArchive& Ar)
+{
+	Super::Serialize(Ar);
+
+	if (Ar.IsLoading() && Ar.UE4Ver() < VER_UE4_ATMOSPHERIC_FOG_DECAY_NAME_CHANGE)
+	{
+		PrecomputeParams.DensityHeight = PrecomputeParams.DecayHeight_DEPRECATED;
+	}
+
+	if (Ar.UE4Ver() >= VER_UE4_ATMOSPHERIC_FOG_CACHE_DATA)
+	{
+		TransmittanceData.Serialize(Ar, this);
+		IrradianceData.Serialize(Ar, this);
+	}
+
+	InscatterData.Serialize(Ar,this);
+
+	if (Ar.IsLoading())
+	{
+		int32 CounterVal;
+		Ar << CounterVal;
+		// Precomputation was not successful, just ignore it
+		if (CounterVal < EValid || !TransmittanceData.GetElementCount())
+		{
+			CounterVal = EInvalid;
+		}
+		PrecomputeCounter = CounterVal;
+	}
+	else
+	{
+		int32 CounterVal = PrecomputeCounter;
+		Ar << CounterVal;
+	}
+
+	if (Ar.IsLoading() && Ar.UE4Ver() < VER_UE4_ATMOSPHERIC_FOG_CACHE_DATA && PrecomputeCounter == EValid) 
+	{
+		// InscatterAltitudeSampleNum default value has been changed (32 -> 2)
+		// Recalculate InscatterAltitudeSampleNum based on Inscatter Size
+		PrecomputeParams.InscatterAltitudeSampleNum = FMath::Max<int32>(1, 
+			InscatterData.GetBulkDataSize() / sizeof(FFloat16Color) 
+			/ (PrecomputeParams.InscatterMuSNum * PrecomputeParams.InscatterNuNum * PrecomputeParams.InscatterMuNum) 
+			);
+	}
+}
+
+
+// Backup the precomputed data before re-running Blueprint construction script
+TStructOnScope<FActorComponentInstanceData> UAtmosphericFogComponent::GetComponentInstanceData() const
+{
+	TStructOnScope<FActorComponentInstanceData> InstanceData;
+
+	if (TransmittanceData.GetElementCount() && IrradianceData.GetElementCount() && InscatterData.GetElementCount() && PrecomputeCounter == EValid)
+	{
+		// Allocate new struct for holding light map data
+		InstanceData.InitializeAs<FAtmospherePrecomputeInstanceData>(this);
+		FAtmospherePrecomputeInstanceData* PrecomputedData = InstanceData.Cast<FAtmospherePrecomputeInstanceData>();
+
+		// Fill in info
+		PrecomputedData->PrecomputeParameter = PrecomputeParams;
+		{
+			int32 TotalByte = TransmittanceData.GetBulkDataSize();
+			PrecomputedData->TransmittanceData.Lock(LOCK_READ_WRITE);
+			void* OutData = PrecomputedData->TransmittanceData.Realloc(TotalByte);
+			TransmittanceData.GetCopy(&OutData, false);
+			PrecomputedData->TransmittanceData.Unlock();
+		}
+
+		{
+			int32 TotalByte = IrradianceData.GetBulkDataSize();
+			PrecomputedData->IrradianceData.Lock(LOCK_READ_WRITE);
+			void* OutData = PrecomputedData->IrradianceData.Realloc(TotalByte);
+			IrradianceData.GetCopy(&OutData, false);
+			PrecomputedData->IrradianceData.Unlock();
+		}
+
+		{
+			int32 TotalByte = InscatterData.GetBulkDataSize();
+			PrecomputedData->InscatterData.Lock(LOCK_READ_WRITE);
+			void* OutData = PrecomputedData->InscatterData.Realloc(TotalByte);
+			InscatterData.GetCopy(&OutData, false);
+			PrecomputedData->InscatterData.Unlock();
+		}
+	}
+	else
+	{
+		InstanceData = Super::GetComponentInstanceData();
+	}
+	return InstanceData;
+}
+
+// Restore the precomputed data after re-running Blueprint construction script
+void UAtmosphericFogComponent::ApplyComponentInstanceData(FAtmospherePrecomputeInstanceData* PrecomputedData)
+{
+	check(PrecomputedData);
+
+	if (PrecomputedData->PrecomputeParameter != GetPrecomputeParameters())
+	{
+		return;
+	}
+
+	FComponentReregisterContext ReregisterContext(this);
+	ReleaseResource();
+
+	{
+		int32 TotalByte = PrecomputedData->TransmittanceData.GetBulkDataSize();
+		TransmittanceData.Lock(LOCK_READ_WRITE);
+		void* OutData = TransmittanceData.Realloc(TotalByte);
+		PrecomputedData->TransmittanceData.GetCopy(&OutData, false);
+		TransmittanceData.Unlock();
+	}
+
+	{
+		int32 TotalByte = PrecomputedData->IrradianceData.GetBulkDataSize();
+		IrradianceData.Lock(LOCK_READ_WRITE);
+		void* OutData = IrradianceData.Realloc(TotalByte);
+		PrecomputedData->IrradianceData.GetCopy(&OutData, false);
+		IrradianceData.Unlock();
+	}
+
+	{
+		int32 TotalByte = PrecomputedData->InscatterData.GetBulkDataSize();
+		InscatterData.Lock(LOCK_READ_WRITE);
+		void* OutData = InscatterData.Realloc(TotalByte);
+		PrecomputedData->InscatterData.GetCopy(&OutData, false);
+		InscatterData.Unlock();
+	}
+
+	PrecomputeCounter = EValid;
+	InitResource();
+}
+
+/**
+ * Gets called any time cvars change (on the main thread), we check if r.Atmosphere has changed and update the components.
+ */
+static void AtmosphereRenderSinkFunction()
+{
+	bool bNewAtmosphere = CVarAtmosphereRender.GetValueOnGameThread() != 0;
+
+	// by default we assume the state is true
+	static bool GAtmosphere = true;
+
+	if (GAtmosphere != bNewAtmosphere)
+	{
+		GAtmosphere = bNewAtmosphere;
+
+		for(TObjectIterator<UAtmosphericFogComponent> It; It; ++It)
+		{
+			UAtmosphericFogComponent* Comp = *It;
+			if (!Comp->GetOuter()->HasAnyFlags(RF_ClassDefaultObject))
+			{
+				if (bNewAtmosphere && Comp->IsRegistered())
+				{
+					Comp->InitResource();
+				}
+				else
+				{
+					Comp->ReleaseResource();
+				}
+			}
+		}
+	}
+}
+
+FAutoConsoleVariableSink CVarAtmosphereRenderSink(FConsoleCommandDelegate::CreateStatic(&AtmosphereRenderSinkFunction));
+
