Fixed compile errors in shipping configs due to attempting to directly write to certain engine show flags

diff --git a/Engine/Source/Runtime/Engine/Private/ShowFlags.cpp b/Engine/Source/Runtime/Engine/Private/ShowFlags.cpp
--- a/Engine/Source/Runtime/Engine/Private/ShowFlags.cpp
+++ b/Engine/Source/Runtime/Engine/Private/ShowFlags.cpp
@@ -482,11 +482,11 @@ void EngineShowFlagOverride(EShowFlagInitMode ShowFlagInitMode, EViewModeIndex V
 		if (ViewModeIndex == VMI_RayTracingDebug)
 		{
 			EngineShowFlags.SetRayTracingDebug(true);
-			EngineShowFlags.VisualizeHDR = 0;
-			EngineShowFlags.VisualizeBloom = 0;
-			EngineShowFlags.VisualizeMotionBlur = 0;
-			EngineShowFlags.DepthOfField = 0;
-			EngineShowFlags.PostProcessMaterial = 0;
+			EngineShowFlags.SetVisualizeHDR(false);
+			EngineShowFlags.SetVisualizeBloom(false);
+			EngineShowFlags.SetVisualizeMotionBlur(false);
+			EngineShowFlags.SetDepthOfField(false);
+			EngineShowFlags.SetPostProcessMaterial(false);
 		}
 	}
 
