Remove ancient unnecessary rendertarget clamp.  Just causes unexpected problems for people

diff --git a/Engine/Source/Runtime/Engine/Private/TextureRenderTarget2D.cpp b/Engine/Source/Runtime/Engine/Private/TextureRenderTarget2D.cpp
--- a/Engine/Source/Runtime/Engine/Private/TextureRenderTarget2D.cpp
+++ b/Engine/Source/Runtime/Engine/Private/TextureRenderTarget2D.cpp
@@ -213,14 +213,6 @@ void UTextureRenderTarget2D::PostLoad()
 {
 	float OriginalSizeX = SizeX;
 	float OriginalSizeY = SizeY;
-	
-	if (!FPlatformProperties::SupportsWindowedMode())
-	{
-		// Clamp the render target size in order to avoid reallocating the scene render targets,
-		// before the FTextureRenderTarget2DResource() is created in Super::PostLoad().
-		SizeX = FMath::Min<int32>(SizeX, GSystemResolution.ResX);
-		SizeY = FMath::Min<int32>(SizeY, GSystemResolution.ResY);
-	}
 
 	SizeX = FMath::Min<int32>(SizeX, GTextureRenderTarget2DMaxSizeX);
 	SizeY = FMath::Min<int32>(SizeY, GTextureRenderTarget2DMaxSizeY);
