Fixes an issue where r.DOF.Algorithm was still in renderer settings.

diff --git a/Engine/Source/Runtime/Engine/Classes/Engine/RendererSettings.h b/Engine/Source/Runtime/Engine/Classes/Engine/RendererSettings.h
--- a/Engine/Source/Runtime/Engine/Classes/Engine/RendererSettings.h
+++ b/Engine/Source/Runtime/Engine/Classes/Engine/RendererSettings.h
@@ -357,11 +357,6 @@ class ENGINE_API URendererSettings : public UDeveloperSettings
 		ConfigRestartRequired = true))
 	TEnumAsByte<EAlphaChannelMode::Type> bEnableAlphaChannelInPostProcessing;
 
-	UPROPERTY(config, EditAnywhere, Category = Postprocessing, meta = (
-		ConsoleVariable = "r.DOF.Algorithm", DisplayName = "Use new DOF algorithm",
-		ToolTip = "Whether to use the new DOF implementation for Circle DOF method."))
-	uint32 bUseNewAlgorithm : 1;
-
 	UPROPERTY(config, EditAnywhere, Category = DefaultSettings, meta = (
 		ConsoleVariable = "r.DefaultFeature.Bloom", DisplayName = "Bloom",
 		ToolTip = "Whether the default for Bloom is enabled or not (postprocess volume/camera/game setting can still override and enable or disable it independently)"))
