Fixes compilation failure in Scene.cpp

diff --git a/Engine/Source/Runtime/Engine/Private/Scene.cpp b/Engine/Source/Runtime/Engine/Private/Scene.cpp
--- a/Engine/Source/Runtime/Engine/Private/Scene.cpp
+++ b/Engine/Source/Runtime/Engine/Private/Scene.cpp
@@ -2,6 +2,7 @@
 
 #include "Engine/Scene.h"
 #include "HAL/IConsoleManager.h"
+#include "UObject/RenderingObjectVersion.h"
 
 void FColorGradingSettings::ExportToPostProcessSettings(FPostProcessSettings* OutPostProcessSettings) const
 {
