Minor code clean up and comment

diff --git a/Engine/Shaders/Private/AtmosphereCommon.ush b/Engine/Shaders/Private/AtmosphereCommon.ush
--- a/Engine/Shaders/Private/AtmosphereCommon.ush
+++ b/Engine/Shaders/Private/AtmosphereCommon.ush
@@ -129,8 +129,8 @@ void GetTransmittanceRMuS(float2 UV, out float Radius, out float MuS)
     Radius = RadiusGround + (Radius * Radius) * (RadiusAtmosphere - RadiusGround);
     MuS = -0.15 + tan(1.5 * MuS) / tan(1.5) * (1.0 + 0.15);
 #else
-    r = RadiusGround + r * (RadiusAtmosphere - RadiusGround);
-    muS = -0.15 + muS * (1.0 + 0.15);
+	Radius = RadiusGround + Radius * (RadiusAtmosphere - RadiusGround);
+	MuS = -0.15 + MuS * (1.0 + 0.15);
 #endif
 }
 
diff --git a/Engine/Source/Runtime/Engine/Classes/Atmosphere/AtmosphericFogComponent.h b/Engine/Source/Runtime/Engine/Classes/Atmosphere/AtmosphericFogComponent.h
--- a/Engine/Source/Runtime/Engine/Classes/Atmosphere/AtmosphericFogComponent.h
+++ b/Engine/Source/Runtime/Engine/Classes/Atmosphere/AtmosphericFogComponent.h
@@ -99,39 +99,39 @@ class UAtmosphericFogComponent : public USceneComponent
 
 	~UAtmosphericFogComponent();
 
-	/** Global scattering factor. */
+	/** Scale the scattered luminance from the atmosphere sun light. Only affect the sky and atmospheric fog. */
 	UPROPERTY(EditAnywhere, BlueprintReadWrite, interp, Category=Atmosphere)
 	float SunMultiplier;
 
-	/** Scattering factor on object. */
+	/** Scale the scattered luminance from the atmosphere sun light only on surfaces, excludes the sky. */
 	UPROPERTY(EditAnywhere, BlueprintReadWrite, interp, Category=Atmosphere)
 	float FogMultiplier;
 
-	/** Fog density control factor. */
+	/** Scales the atmosphere transmittance over background. */
 	UPROPERTY(EditAnywhere, BlueprintReadWrite, interp, Category=Atmosphere)
 	float DensityMultiplier;
 
-	/** Fog density offset to control opacity [-1.f ~ 1.f]. */
+	/** Offset the atmosphere transmittance over background [-1.f ~ 1.f]. */
 	UPROPERTY(EditAnywhere, BlueprintReadWrite, interp, Category=Atmosphere)
 	float DensityOffset;
 
-	/** Distance scale. */
+	/** Scale the view position. */
 	UPROPERTY(EditAnywhere, BlueprintReadWrite, interp, Category=Atmosphere)
 	float DistanceScale;
 
-	/** Altitude scale (only Z scale). */
+	/** Scale the view altitude (only Z scale). */
 	UPROPERTY(EditAnywhere, BlueprintReadWrite, interp, Category=Atmosphere)
 	float AltitudeScale;
 
-	/** Distance offset, in km (to handle large distance) */
+	/** Apply a distance offset before evaluating the atmospheric fog, in km (to handle large distance). Only on surfaces, excludes the sky. */
 	UPROPERTY(EditAnywhere, BlueprintReadWrite, interp, Category=Atmosphere)
 	float DistanceOffset;
 
-	/** Ground offset. */
+	/** Offset the view altitude (along Z). */
 	UPROPERTY(EditAnywhere, BlueprintReadWrite, interp, Category=Atmosphere)
 	float GroundOffset;
 
-	/** Start Distance. */
+	/** The atmospheric fog start distance in centimeters. */
 	UPROPERTY(EditAnywhere, BlueprintReadWrite, interp, Category=Atmosphere)
 	float StartDistance;
 
@@ -139,19 +139,19 @@ class UAtmosphericFogComponent : public USceneComponent
 	UPROPERTY(EditAnywhere, BlueprintReadWrite, interp, Category = Atmosphere)
 	float SunDiscScale;
 
-	/** Default light brightness. Used when there is no sunlight placed in the level. Unit is lumens */
+	/** Default atmospheric sun light disc luminance. Used when there is no atmospheric sun light selected in the level. */
 	UPROPERTY(EditAnywhere, BlueprintReadOnly, interp, Category=Lighting)
 	float DefaultBrightness;
 
-	/** Default light color. Used when there is no sunlight placed in the level. */
+	/** Default atmospheric sun light disc color. Used when there is no sunlight placed in the level. */
 	UPROPERTY(EditAnywhere, BlueprintReadOnly, interp, Category=Lighting)
 	FColor DefaultLightColor;
 
-	/** Disable Sun Disk rendering. */
+	/** Disable sun disk rendering. */
 	UPROPERTY(EditAnywhere, BlueprintReadOnly, interp, Category=Lighting, meta=(ScriptName="DisableSunDiskValue"))
 	uint32 bDisableSunDisk : 1;
 
-	/** Disable Color scattering from ground. */
+	/** Disable color scattering from ground. */
 	UPROPERTY(EditAnywhere, BlueprintReadOnly, interp, Category=Lighting, meta=(ScriptName="DisableGroundScatteringValue"))
 	uint32 bDisableGroundScattering : 1;
 
@@ -242,9 +242,11 @@ protected:
 	 void AddFogIfNeeded();
 
 public:
-	/** The resource for Inscatter. */
+	// Stores colored transmittance from outer space to point in atmosphere.
 	class FAtmosphereTextureResource* TransmittanceResource;
+	// Stores ground illuminance as a function of sun direction and atmosphere radius.
 	class FAtmosphereTextureResource* IrradianceResource;
+	// Stores in-scattered luminance toward a point according to height and sun direction.
 	class FAtmosphereTextureResource* InscatterResource;
 	
 	/** Source vector data. */
@@ -283,7 +285,7 @@ private:
 	friend class FAtmosphericFogSceneInfo;
 };
 
-/** Used to store lightmap data during RerunConstructionScripts */
+/** Used to store data during RerunConstructionScripts */
 USTRUCT()
 struct FAtmospherePrecomputeInstanceData : public FSceneComponentInstanceData
 {
diff --git a/Engine/Source/Runtime/Engine/Private/Atmosphere/AtmosphericFogComponent.cpp b/Engine/Source/Runtime/Engine/Private/Atmosphere/AtmosphericFogComponent.cpp
--- a/Engine/Source/Runtime/Engine/Private/Atmosphere/AtmosphericFogComponent.cpp
+++ b/Engine/Source/Runtime/Engine/Private/Atmosphere/AtmosphericFogComponent.cpp
@@ -40,6 +40,45 @@ FAtmospherePrecomputeParameters::FAtmospherePrecomputeParameters():
 }
 
 
+//////////////////////////////////////////////////////////////////////////
+// FAtmospherePrecomputeDataHandler
+
+
+#if WITH_EDITOR
+#include "TickableEditorObject.h"
+
+class FAtmospherePrecomputeDataHandler : public FTickableEditorObject
+{
+public:
+	UAtmosphericFogComponent* Component;
+
+	FAtmospherePrecomputeDataHandler(UAtmosphericFogComponent* InComponent) : Component(InComponent)
+	{
+	}
+
+	/** FTickableEditorObject interface */
+	virtual void Tick(float DeltaTime) override
+	{
+		if (Component && Component->GameThreadServiceRequest.GetValue())
+		{
+			Component->UpdatePrecomputedData();
+			Component->GameThreadServiceRequest.Decrement();
+		}
+	}
+
+	virtual ETickableTickType GetTickableTickType() const override
+	{
+		return ETickableTickType::Always;
+	}
+
+	TStatId GetStatId() const override
+	{
+		RETURN_QUICK_DECLARE_CYCLE_STAT(FAtmospherePrecomputeDataHandler, STATGROUP_Tickables);
+	}
+};
+#endif
+
+
 //////////////////////////////////////////////////////////////////////////
 // UAtmosphericFogComponent
 
@@ -142,40 +181,6 @@ void UAtmosphericFogComponent::PostLoad()
 	}
 }
 
-#if WITH_EDITOR
-#include "TickableEditorObject.h"
-
-class FAtmospherePrecomputeDataHandler : public FTickableEditorObject
-{
-public:
-	UAtmosphericFogComponent* Component;
-
-	FAtmospherePrecomputeDataHandler(UAtmosphericFogComponent* InComponent) : Component(InComponent)
-	{
-	}
-
-	/** FTickableEditorObject interface */
-	virtual void Tick( float DeltaTime ) override
-	{
-		if( Component && Component->GameThreadServiceRequest.GetValue()) 
-		{
-			Component->UpdatePrecomputedData();
-			Component->GameThreadServiceRequest.Decrement();
-		}              
-	}
-
-	virtual ETickableTickType GetTickableTickType() const override
-	{
-		return ETickableTickType::Always;
-	}
-
-	TStatId GetStatId() const override
-	{
-		RETURN_QUICK_DECLARE_CYCLE_STAT(FAtmospherePrecomputeDataHandler, STATGROUP_Tickables);
-	}
-};
-#endif
-
 
 static TAutoConsoleVariable<int32> CVarAtmosphereRender(
 	TEXT("r.Atmosphere"),
diff --git a/Engine/Source/Runtime/Renderer/Private/AtmosphereRendering.h b/Engine/Source/Runtime/Renderer/Private/AtmosphereRendering.h
--- a/Engine/Source/Runtime/Renderer/Private/AtmosphereRendering.h
+++ b/Engine/Source/Runtime/Renderer/Private/AtmosphereRendering.h
@@ -35,7 +35,7 @@ namespace EAtmosphereRenderFlag
 	};
 }
 
-/** The properties of a atmospheric fog layer which are used for rendering. */
+/** The properties of a atmospheric fog layer which are used for rendering. (Render side of the application) */
 class FAtmosphericFogSceneInfo : public FRenderResource
 {
 public:
diff --git a/Engine/Source/Runtime/Renderer/Private/SceneRendering.cpp b/Engine/Source/Runtime/Renderer/Private/SceneRendering.cpp
--- a/Engine/Source/Runtime/Renderer/Private/SceneRendering.cpp
+++ b/Engine/Source/Runtime/Renderer/Private/SceneRendering.cpp
@@ -1140,7 +1140,7 @@ void FViewInfo::SetupUniformBufferParameters(
 			ViewUniformShaderParameters.AtmosphericFogSunDiscScale = Scene->AtmosphericFog->SunDiscScale;
 			ViewUniformShaderParameters.AtmosphericFogSunDiscHalfApexAngleRadian = Scene->AtmosphericFog->SunDiscScale * SunLightDiskHalfApexAngleRadian;
 			ViewUniformShaderParameters.AtmosphericFogSunDiscLuminance = Scene->SunLight ? Scene->SunLight->Proxy->GetOuterSpaceLuminance() : FLinearColor::White;
-			ViewUniformShaderParameters.AtmosphericFogSunColor = Scene->SunLight ? Scene->SunLight->Proxy->GetColor() : Scene->AtmosphericFog->DefaultSunColor;
+			ViewUniformShaderParameters.AtmosphericFogSunColor = Scene->SunLight ? Scene->SunLight->Proxy->GetColor() : Scene->AtmosphericFog->DefaultSunColor; // Sun light color unaffected by atmosphere transmittance
 			ViewUniformShaderParameters.AtmosphericFogSunDirection = Scene->SunLight ? -Scene->SunLight->Proxy->GetDirection() : -Scene->AtmosphericFog->DefaultSunDirection;
 			ViewUniformShaderParameters.AtmosphericFogRenderMask = Scene->AtmosphericFog->RenderFlag & (EAtmosphereRenderFlag::E_DisableGroundScattering | EAtmosphereRenderFlag::E_DisableSunDisk);
 			ViewUniformShaderParameters.AtmosphericFogInscatterAltitudeSampleNum = Scene->AtmosphericFog->InscatterAltitudeSampleNum;
