Exposes temporal sample as view portperty to avoid having engine material function accessing view uniform buffer with a custom node.

diff --git a/Engine/Source/Runtime/Engine/Classes/Materials/MaterialExpressionViewProperty.h b/Engine/Source/Runtime/Engine/Classes/Materials/MaterialExpressionViewProperty.h
--- a/Engine/Source/Runtime/Engine/Classes/Materials/MaterialExpressionViewProperty.h
+++ b/Engine/Source/Runtime/Engine/Classes/Materials/MaterialExpressionViewProperty.h
@@ -25,6 +25,12 @@ enum EMaterialExposedViewProperty
 	MEVP_WorldSpaceCameraPosition UMETA(DisplayName = "Camera Position (Absolute World Space)"),
 	/** Horizontal and vertical position of the viewport in pixels within the buffer. */
 	MEVP_ViewportOffset UMETA(DisplayName = "Viewport Offset"),
+	/** Number of temporal AA sample used across multiple to converge to anti aliased output. */
+	MEVP_TemporalSampleCount UMETA(DisplayName = "Temporal Sample Count"),
+	/** Index of the Temporal AA jitter for this frame. */
+	MEVP_TemporalSampleIndex UMETA(DisplayName = "Temporal Sample Index"),
+	/** Offset of the temporal sample for this frame in pixel size. */
+	MEVP_TemporalSampleOffset UMETA(DisplayName = "Temporal Sample Offset"),
 
 	MEVP_MAX,
 };
diff --git a/Engine/Source/Runtime/Engine/Private/Materials/HLSLMaterialTranslator.h b/Engine/Source/Runtime/Engine/Private/Materials/HLSLMaterialTranslator.h
--- a/Engine/Source/Runtime/Engine/Private/Materials/HLSLMaterialTranslator.h
+++ b/Engine/Source/Runtime/Engine/Private/Materials/HLSLMaterialTranslator.h
@@ -2424,6 +2424,9 @@ protected:
 			{MEVP_WorldSpaceViewPosition, MCT_Float3, TEXT("ResolvedView.<PREV>WorldViewOrigin"), nullptr},
 			{MEVP_WorldSpaceCameraPosition, MCT_Float3, TEXT("ResolvedView.<PREV>WorldCameraOrigin"), nullptr},
 			{MEVP_ViewportOffset, MCT_Float2, TEXT("View.ViewRectMin.xy"), nullptr},
+			{MEVP_TemporalSampleCount, MCT_Float1, TEXT("View.TemporalAAParams.y"), nullptr},
+			{MEVP_TemporalSampleIndex, MCT_Float1, TEXT("View.TemporalAAParams.x"), nullptr},
+			{MEVP_TemporalSampleOffset, MCT_Float2, TEXT("View.TemporalAAParams.zw"), nullptr},
 		};
 		static_assert((sizeof(ViewPropertyMetaArray) / sizeof(ViewPropertyMetaArray[0])) == MEVP_MAX, "incoherency between EMaterialExposedViewProperty and ViewPropertyMetaArray");
 
diff --git a/Engine/Source/Runtime/Renderer/Private/SceneRendering.cpp b/Engine/Source/Runtime/Renderer/Private/SceneRendering.cpp
--- a/Engine/Source/Runtime/Renderer/Private/SceneRendering.cpp
+++ b/Engine/Source/Runtime/Renderer/Private/SceneRendering.cpp
@@ -862,7 +862,8 @@ void FViewInfo::Init()
 
 	// Disable HDR encoding for editor elements.
 	EditorSimpleElementCollector.BatchedElements.EnableMobileHDREncoding(false);
-
+	
+	TemporalJitterSequenceLength = 1;
 	TemporalJitterIndex = 0;
 	TemporalJitterPixels = FVector2D::ZeroVector;
 
@@ -1263,7 +1264,7 @@ void FViewInfo::SetupUniformBufferParameters(
 	{
 		ViewUniformShaderParameters.TemporalAAParams = FVector4(
 			TemporalJitterIndex, 
-			0,
+			TemporalJitterSequenceLength,
 			TemporalJitterPixels.X,
 			TemporalJitterPixels.Y);
 
diff --git a/Engine/Source/Runtime/Renderer/Private/SceneRendering.h b/Engine/Source/Runtime/Renderer/Private/SceneRendering.h
--- a/Engine/Source/Runtime/Renderer/Private/SceneRendering.h
+++ b/Engine/Source/Runtime/Renderer/Private/SceneRendering.h
@@ -884,10 +884,13 @@ public:
 	float TranslucencyVolumeVoxelSize[TVC_MAX];
 	FVector TranslucencyLightingVolumeSize[TVC_MAX];
 
-	/** Index of the temporal jitter in the sequence. */
+	/** Number of samples in the temporal AA sequqnce */
+	int32 TemporalJitterSequenceLength;
+
+	/** Index of the temporal AA jitter in the sequence. */
 	int32 TemporalJitterIndex;
 
-	/** Temporal jitter at the pixel scale. */
+	/** Temporal AA jitter at the pixel scale. */
 	FVector2D TemporalJitterPixels;
 
 	/** Whether view state may be updated with this view. */
diff --git a/Engine/Source/Runtime/Renderer/Private/SceneVisibility.cpp b/Engine/Source/Runtime/Renderer/Private/SceneVisibility.cpp
--- a/Engine/Source/Runtime/Renderer/Private/SceneVisibility.cpp
+++ b/Engine/Source/Runtime/Renderer/Private/SceneVisibility.cpp
@@ -3122,6 +3122,7 @@ void FSceneRenderer::PreVisibilityFrameSetup(FRHICommandListImmediate& RHICmdLis
 				SampleY = r * FMath::Sin( Theta );
 			}
 
+			View.TemporalJitterSequenceLength = TemporalAASamples;
 			View.TemporalJitterIndex = TemporalSampleIndex;
 			View.TemporalJitterPixels.X = SampleX;
 			View.TemporalJitterPixels.Y = SampleY;
