Added option on AtmosphereComponent to conditionally affect the sunlight illuminance. Off by default.

diff --git a/Engine/Source/Editor/UnrealEd/Private/Lightmass/Lightmass.cpp b/Engine/Source/Editor/UnrealEd/Private/Lightmass/Lightmass.cpp
--- a/Engine/Source/Editor/UnrealEd/Private/Lightmass/Lightmass.cpp
+++ b/Engine/Source/Editor/UnrealEd/Private/Lightmass/Lightmass.cpp
@@ -1027,7 +1027,7 @@ void FLightmassExporter::WriteLights( int32 Channel )
 	// If found, we keep the pointer to apply atmosphere transmittance to it in the light loop.
 	const UDirectionalLightComponent* SunLight = nullptr;
 	FLinearColor SunLightAtmosphereTransmittance(FLinearColor::White);
-	if (AtmosphericFogComponent)
+	if (AtmosphericFogComponent && AtmosphericFogComponent->bAtmosphereAffectsSunIlluminance)
 	{
 		float SunLightEnergy = 0.0f;
 		for (int32 LightIndex = 0; LightIndex < DirectionalLights.Num(); ++LightIndex)
diff --git a/Engine/Source/Runtime/Engine/Classes/Atmosphere/AtmosphericFogComponent.h b/Engine/Source/Runtime/Engine/Classes/Atmosphere/AtmosphericFogComponent.h
--- a/Engine/Source/Runtime/Engine/Classes/Atmosphere/AtmosphericFogComponent.h
+++ b/Engine/Source/Runtime/Engine/Classes/Atmosphere/AtmosphericFogComponent.h
@@ -141,23 +141,27 @@ class UAtmosphericFogComponent : public USceneComponent
 	float StartDistance;
 
 	/** Sun half apex angle in degree, see https://en.wikipedia.org/wiki/Solid_angle */
-	UPROPERTY(EditAnywhere, BlueprintReadWrite, interp, Category = Atmosphere)
+	UPROPERTY(EditAnywhere, BlueprintReadWrite, interp, Category=Sun)
 	float SunDiscScale;
 
 	/** Default atmospheric sun light disc luminance. Used when there is no atmospheric sun light selected in the level. */
-	UPROPERTY(EditAnywhere, BlueprintReadOnly, interp, Category=Lighting)
+	UPROPERTY(EditAnywhere, BlueprintReadOnly, interp, Category=Sun)
 	float DefaultBrightness;
 
 	/** Default atmospheric sun light disc color. Used when there is no sunlight placed in the level. */
-	UPROPERTY(EditAnywhere, BlueprintReadOnly, interp, Category=Lighting)
+	UPROPERTY(EditAnywhere, BlueprintReadOnly, interp, Category=Sun)
 	FColor DefaultLightColor;
 
 	/** Disable sun disk rendering. */
-	UPROPERTY(EditAnywhere, BlueprintReadOnly, interp, Category=Lighting, meta=(ScriptName="DisableSunDiskValue"))
+	UPROPERTY(EditAnywhere, BlueprintReadOnly, interp, Category=Sun, meta=(ScriptName="DisableSunDiskValue"))
 	uint32 bDisableSunDisk : 1;
 
+	/** Set to true if the atmosphere should affect the selected sun light illuminance. The light will be tinted based on its zenith angle and atmosphere properties as if all surfaces were at the ground level 0 meter. */
+	UPROPERTY(EditAnywhere, BlueprintReadOnly, interp, Category=Sun)
+	uint32 bAtmosphereAffectsSunIlluminance : 1;
+
 	/** Disable color scattering from ground. */
-	UPROPERTY(EditAnywhere, BlueprintReadOnly, interp, Category=Lighting, meta=(ScriptName="DisableGroundScatteringValue"))
+	UPROPERTY(EditAnywhere, BlueprintReadOnly, interp, Category=Ground, meta=(ScriptName="DisableGroundScatteringValue"))
 	uint32 bDisableGroundScattering : 1;
 
 	/** Set brightness of the light */
@@ -288,7 +292,8 @@ public:
 	 */
 	static FLinearColor GetTransmittance(const FVector& SunDirection, float AtmosphericFogHeightScaleRayleigh)
 	{
-		// Code from HLSL here as lambda function. It will simulate atmosphere transmittance according to the current sky hardcoded parameterization.
+		// The following code is from atmosphere HLSL shaders and has been converted to lambda functions. 
+		// It simulates atmosphere transmittance according to the current sky hardcoded parameterization.
 		// This will change in the future when the sky parameterization and workflow/ui will be updated.
 
 		const float RadiusGround = 6360;
diff --git a/Engine/Source/Runtime/Engine/Private/Atmosphere/AtmosphericFogComponent.cpp b/Engine/Source/Runtime/Engine/Private/Atmosphere/AtmosphericFogComponent.cpp
--- a/Engine/Source/Runtime/Engine/Private/Atmosphere/AtmosphericFogComponent.cpp
+++ b/Engine/Source/Runtime/Engine/Private/Atmosphere/AtmosphericFogComponent.cpp
@@ -106,6 +106,8 @@ UAtmosphericFogComponent::UAtmosphericFogComponent(const FObjectInitializer& Obj
 
 	bDisableSunDisk = false;
 	bDisableGroundScattering = false;
+
+	bAtmosphereAffectsSunIlluminance = false;
 }
 
 bool UAtmosphericFogComponent::IsPostLoadThreadSafe() const
diff --git a/Engine/Source/Runtime/Renderer/Private/AtmosphereRendering.cpp b/Engine/Source/Runtime/Renderer/Private/AtmosphereRendering.cpp
--- a/Engine/Source/Runtime/Renderer/Private/AtmosphereRendering.cpp
+++ b/Engine/Source/Runtime/Renderer/Private/AtmosphereRendering.cpp
@@ -1826,16 +1826,17 @@ void FAtmosphericFogSceneInfo::PrecomputeTextures(FRHICommandListImmediate& RHIC
 void FAtmosphericFogSceneInfo::PrepareSunLightProxy(FLightSceneInfo& SunLight) const
 {
 	// See explanation in https://media.contentapi.ea.com/content/dam/eacom/frostbite/files/s2016-pbs-frostbite-sky-clouds-new.pdf page 26
-	FLinearColor TransmittanceTowardSun = UAtmosphericFogComponent::GetTransmittance(-SunLight.Proxy->GetDirection(), RHeight);
+	FLinearColor TransmittanceTowardSun = bAtmosphereAffectsSunIlluminance ? UAtmosphericFogComponent::GetTransmittance(-SunLight.Proxy->GetDirection(), RHeight) : FLinearColor(FLinearColor::White);
+	FLinearColor TransmittanceAtZenithFinal = bAtmosphereAffectsSunIlluminance ? TransmittanceAtZenith : FLinearColor(FLinearColor::White);
 
 	FLinearColor SunZenithIlluminance = SunLight.Proxy->GetColor();
-	FLinearColor SunOuterSpaceIlluminance = SunZenithIlluminance / TransmittanceAtZenith;
+	FLinearColor SunOuterSpaceIlluminance = SunZenithIlluminance / TransmittanceAtZenithFinal;
 
 	// SunDiscScale is only considered as a visual tweak so we do not make it influence the sun disk outerspace luminance.
 	const float SunSolidAngle = 2.0f * PI * (1.0f - FMath::Cos(SunLight.Proxy->GetSunLightHalfApexAngleRadian())); // Solid angle from aperture https://en.wikipedia.org/wiki/Solid_angle 
 	FLinearColor SunDiskOuterSpaceLuminance = SunOuterSpaceIlluminance / SunSolidAngle; // approximation  
 
-	SunLight.Proxy->SetAtmosphereRelatedProperties(TransmittanceTowardSun / TransmittanceAtZenith, SunDiskOuterSpaceLuminance);
+	SunLight.Proxy->SetAtmosphereRelatedProperties(TransmittanceTowardSun / TransmittanceAtZenithFinal, SunDiskOuterSpaceLuminance);
 }
 
 /** Initialization constructor. */
@@ -1854,6 +1855,7 @@ FAtmosphericFogSceneInfo::FAtmosphericFogSceneInfo(const UAtmosphericFogComponen
 	, SunDiscScale(InComponent->SunDiscScale)
 	, RenderFlag(EAtmosphereRenderFlag::E_EnableAll)
 	, InscatterAltitudeSampleNum(InComponent->PrecomputeParams.InscatterAltitudeSampleNum)
+	, bAtmosphereAffectsSunIlluminance(InComponent->bAtmosphereAffectsSunIlluminance)
 
 #if WITH_EDITORONLY_DATA
 	, bNeedRecompute(false)
diff --git a/Engine/Source/Runtime/Renderer/Private/AtmosphereRendering.h b/Engine/Source/Runtime/Renderer/Private/AtmosphereRendering.h
--- a/Engine/Source/Runtime/Renderer/Private/AtmosphereRendering.h
+++ b/Engine/Source/Runtime/Renderer/Private/AtmosphereRendering.h
@@ -56,6 +56,7 @@ public:
 	FVector DefaultSunDirection;
 	uint32 RenderFlag;
 	uint32 InscatterAltitudeSampleNum;
+	bool bAtmosphereAffectsSunIlluminance;
 	class FAtmosphereTextureResource* TransmittanceResource;
 	class FAtmosphereTextureResource* IrradianceResource;
 	class FAtmosphereTextureResource* InscatterResource;
