Moved GPU profiling tools into RenderCore to allow RenderGraph to utilize them.

diff --git a/Engine/Source/Runtime/Engine/Private/ProfilingDebugging/TracingProfiler.cpp b/Engine/Source/Runtime/Engine/Private/ProfilingDebugging/TracingProfiler.cpp
deleted file mode 100644
diff --git a/Engine/Source/Runtime/Engine/Private/SceneUtils.cpp b/Engine/Source/Runtime/Engine/Private/SceneUtils.cpp
--- a/Engine/Source/Runtime/Engine/Private/SceneUtils.cpp
+++ b/Engine/Source/Runtime/Engine/Private/SceneUtils.cpp
@@ -6,103 +6,6 @@
 
 DEFINE_LOG_CATEGORY_STATIC(LogSceneUtils,All,All);
 
-// Only exposed for debugging. Disabling this carries a severe performance penalty
-#define RENDER_QUERY_POOLING_ENABLED 1
-
-#if HAS_GPU_STATS 
-
-CSV_DEFINE_CATEGORY_MODULE(ENGINE_API, GPU, true);
-
-static TAutoConsoleVariable<int> CVarGPUStatsEnabled(
-	TEXT("r.GPUStatsEnabled"),
-	1,
-	TEXT("Enables or disables GPU stat recording"));
-
-
-static TAutoConsoleVariable<int> CVarGPUStatsMaxQueriesPerFrame(
-	TEXT("r.GPUStatsMaxQueriesPerFrame"),
-	-1,
-	TEXT("Limits the number of timestamps allocated per frame. -1 = no limit"), 
-	ECVF_RenderThreadSafe);
-
-
-static TAutoConsoleVariable<int> CVarGPUCsvStatsEnabled(
-	TEXT("r.GPUCsvStatsEnabled"),
-	0,
-	TEXT("Enables or disables GPU stat recording to CSVs"));
-
-DECLARE_GPU_STAT_NAMED( Total, TEXT("[TOTAL]") );
-
-static TAutoConsoleVariable<int> CVarGPUTracingStatsEnabled(
-	TEXT("r.GPUTracingStatsEnabled"),
-	1,
-	TEXT("Enables or disables GPU stat recording to tracing profiler"));
-
-static TAutoConsoleVariable<int> CVarGPUStatsChildTimesIncluded(
-	TEXT("r.GPUStatsChildTimesIncluded"),
-	0,
-	TEXT("If this is enabled, the child stat timings will be included in their parents' times.\n")
-	TEXT("This presents problems for non-hierarchical stats if we're expecting them to add up\n")
-	TEXT("to the total GPU time, so we probably want this disabled.\n")
-);
-
-#endif //HAS_GPU_STATS
-
-
-#if WANTS_DRAW_MESH_EVENTS
-
-template<typename TRHICmdList>
-void TDrawEvent<TRHICmdList>::Start(TRHICmdList& InRHICmdList, FColor Color, const TCHAR* Fmt, ...)
-{
-	check(IsInParallelRenderingThread() || IsInRHIThread());
-	{
-		va_list ptr;
-		va_start(ptr, Fmt);
-		TCHAR TempStr[256];
-		// Build the string in the temp buffer
-		FCString::GetVarArgs(TempStr, ARRAY_COUNT(TempStr), Fmt, ptr);
-		InRHICmdList.PushEvent(TempStr, Color);
-		RHICmdList = &InRHICmdList;
-		va_end(ptr);
-	}
-}
-
-template<typename TRHICmdList>
-void TDrawEvent<TRHICmdList>::Stop()
-{
-	if (RHICmdList)
-	{
-		RHICmdList->PopEvent();
-		RHICmdList = NULL;
-	}
-}
-template struct TDrawEvent<FRHICommandList>;
-template struct TDrawEvent<FRHIAsyncComputeCommandList>;
-template struct TDrawEvent<FRHIAsyncComputeCommandListImmediate>;
-
-void FDrawEventRHIExecute::Start(IRHIComputeContext& InRHICommandContext, FColor Color, const TCHAR* Fmt, ...)
-{
-	check(IsInParallelRenderingThread() || IsInRHIThread() || (!IsRunningRHIInSeparateThread() && IsInRenderingThread()));
-	{
-		va_list ptr;
-		va_start(ptr, Fmt);
-		TCHAR TempStr[256];
-		// Build the string in the temp buffer
-		FCString::GetVarArgs(TempStr, ARRAY_COUNT(TempStr), Fmt, ptr);
-		RHICommandContext = &InRHICommandContext;
-		RHICommandContext->RHIPushEvent(TempStr, Color);
-		va_end(ptr);
-	}
-}
-
-void FDrawEventRHIExecute::Stop()
-{
-	RHICommandContext->RHIPopEvent();
-}
-
-#endif // WANTS_DRAW_MESH_EVENTS
-
-
 bool IsMobileHDR()
 {
 	static auto* MobileHDRCvar = IConsoleManager::Get().FindTConsoleVariableDataInt(TEXT("r.MobileHDR"));
@@ -171,617 +74,4 @@ ENGINE_API bool IsMobileColorsRGB()
 	const bool bMobileUseHWsRGBEncoding = (MobileUseHWsRGBEncodingCVAR && MobileUseHWsRGBEncodingCVAR->GetValueOnAnyThread() == 1);
 
 	return !IsMobileHDR() && bMobileUseHWsRGBEncoding;
-}
-
-#if HAS_GPU_STATS
-static const int32 NumGPUProfilerBufferedFrames = 4;
-
-/*-----------------------------------------------------------------------------
-FRealTimeGPUProfilerEvent class
------------------------------------------------------------------------------*/
-class FRealtimeGPUProfilerEvent
-{
-public:
-	static const uint64 InvalidQueryResult = 0xFFFFFFFFFFFFFFFFull;
-
-public:
-	FRealtimeGPUProfilerEvent(const FName& InName, const FName& InStatName, FRenderQueryPool* RenderQueryPool)
-		: StartResultMicroseconds(InvalidQueryResult)
-		, EndResultMicroseconds(InvalidQueryResult)
-		, FrameNumber(-1)
-		, bInsideQuery(false)
-		, bBeginQueryInFlight(false)
-		, bEndQueryInFlight(false)
-	{
-#if STATS
-		StatName = InStatName;
-#endif
-		Name = InName;
-
-		const int MaxGPUQueries = CVarGPUStatsMaxQueriesPerFrame.GetValueOnRenderThread();
-		if ( MaxGPUQueries == -1 || RenderQueryPool->GetAllocatedQueryCount() < MaxGPUQueries )
-		{
-			StartQuery = RenderQueryPool->AllocateQuery();
-			EndQuery = RenderQueryPool->AllocateQuery();
-		}
-	}
-
-	bool HasQueriesAllocated() const 
-	{ 
-		return IsValidRef(StartQuery); 
-	}
-
-	void ReleaseQueries(FRenderQueryPool* RenderQueryPool, FRHICommandListImmediate* RHICmdListPtr)
-	{
-		if ( HasQueriesAllocated() )
-		{
-			if (RHICmdListPtr)
-			{
-				// If we have queries in flight then get results before releasing back to the pool to avoid an ensure fail in the gnm RHI
-				uint64 Temp;
-				if (bBeginQueryInFlight)
-				{
-					RHICmdListPtr->GetRenderQueryResult(StartQuery, Temp, false);
-				}
-
-				if (bEndQueryInFlight)
-				{
-					RHICmdListPtr->GetRenderQueryResult(EndQuery, Temp, false);
-				}
-			}
-			RenderQueryPool->ReleaseQuery(StartQuery);
-			RenderQueryPool->ReleaseQuery(EndQuery);
-		}
-	}
-
-	void Begin(FRHICommandListImmediate& RHICmdList)
-	{
-		check(IsInRenderingThread());
-		check(!bInsideQuery);
-		bInsideQuery = true;
-
-		if (IsValidRef(StartQuery))
-		{
-			RHICmdList.EndRenderQuery(StartQuery);
-			bBeginQueryInFlight = true;
-		}
-		StartResultMicroseconds = InvalidQueryResult;
-		EndResultMicroseconds = InvalidQueryResult;
-		FrameNumber = GFrameNumberRenderThread;
-	}
-
-	void End(FRHICommandListImmediate& RHICmdList)
-	{
-		check(IsInRenderingThread());
-		check(bInsideQuery);
-		bInsideQuery = false;
-
-		if ( HasQueriesAllocated() )
-		{
-			RHICmdList.EndRenderQuery(EndQuery);
-			bEndQueryInFlight = true;
-		}
-	}
-
-	bool GatherQueryResults(FRHICommandListImmediate& RHICmdList)
-	{
-		QUICK_SCOPE_CYCLE_COUNTER(STAT_SceneUtils_GatherQueryResults);
-		// Get the query results which are still outstanding
-		check(GFrameNumberRenderThread != FrameNumber);
-		if ( HasQueriesAllocated() )
-		{
-			if (StartResultMicroseconds == InvalidQueryResult)
-			{
-				if (!RHICmdList.GetRenderQueryResult(StartQuery, StartResultMicroseconds, false))
-				{
-					StartResultMicroseconds = InvalidQueryResult;
-				}
-				bBeginQueryInFlight = false;
-			}
-			if (EndResultMicroseconds == InvalidQueryResult)
-			{
-				if (!RHICmdList.GetRenderQueryResult(EndQuery, EndResultMicroseconds, false))
-				{
-					EndResultMicroseconds = InvalidQueryResult;
-				}
-				bEndQueryInFlight = false;
-			}
-		}
-		else
-		{
-			// If we don't have a query allocated, just set the results to zero
-			EndResultMicroseconds = StartResultMicroseconds = 0;
-		}
-		return HasValidResult();
-	}
-
-	float GetResultMS() const
-	{
-		check(HasValidResult());
-		if (EndResultMicroseconds < StartResultMicroseconds)
-		{
-			// This should never happen...
-			return 0.0f;
-		}
-		return float(EndResultMicroseconds - StartResultMicroseconds) / 1000.0f;
-	}
-
-	bool HasValidResult() const
-	{
-		return StartResultMicroseconds != FRealtimeGPUProfilerEvent::InvalidQueryResult && EndResultMicroseconds != FRealtimeGPUProfilerEvent::InvalidQueryResult;
-	}
-
-#if STATS
-	const FName& GetStatName() const
-	{
-		return StatName;
-	}
-#endif
-	const FName& GetName() const
-	{
-		return Name;
-	}
-
-	uint64 GetStartResultMicroseconds(uint32 GPUIndex = 0) const
-	{
-		return StartResultMicroseconds;
-	}
-
-	uint64 GetEndResultMicroseconds(uint32 GPUIndex = 0) const
-	{
-		return EndResultMicroseconds;
-	}
-
-	uint32 GetFrameNumber() const
-	{
-		return FrameNumber;
-	}
-
-private:
-	FRenderQueryRHIRef StartQuery;
-	FRenderQueryRHIRef EndQuery;
-#if STATS
-	FName StatName;
-#endif
-	FName Name;
-	uint64 StartResultMicroseconds;
-	uint64 EndResultMicroseconds;
-	uint32 FrameNumber;
-
-	bool bInsideQuery;
-	bool bBeginQueryInFlight;
-	bool bEndQueryInFlight;
-};
-
-/*-----------------------------------------------------------------------------
-FRealtimeGPUProfilerFrame class
-Container for a single frame's GPU stats
------------------------------------------------------------------------------*/
-class FRealtimeGPUProfilerFrame
-{
-public:
-	FRealtimeGPUProfilerFrame(FRenderQueryPool* InRenderQueryPool)
-		: FrameNumber(-1)
-		, RenderQueryPool(InRenderQueryPool)
-	{}
-
-	~FRealtimeGPUProfilerFrame()
-	{
-		Clear(nullptr);
-	}
-
-	void PushEvent(FRHICommandListImmediate& RHICmdList, const FName& Name, const FName& StatName)
-	{
-		// TODO: this should really use a pool / free list
-		FRealtimeGPUProfilerEvent* Event = new FRealtimeGPUProfilerEvent(Name, StatName, RenderQueryPool);
-		const int32 EventIndex = GpuProfilerEvents.Num();
-
-		GpuProfilerEvents.Add(Event);
-
-		FRealtimeGPUProfilerTimelineEvent TimelineEvent = {};
-		TimelineEvent.Type = FRealtimeGPUProfilerTimelineEvent::EType::PushEvent;
-		TimelineEvent.EventIndex = EventIndex;
-
-		GpuProfilerTimelineEvents.Push(TimelineEvent);
-
-		EventStack.Push(EventIndex);
-		Event->Begin(RHICmdList);
-	}
-
-	void PopEvent(FRHICommandListImmediate& RHICmdList)
-	{
-		int32 Index = EventStack.Pop();
-
-		FRealtimeGPUProfilerTimelineEvent TimelineEvent = {};
-		TimelineEvent.Type = FRealtimeGPUProfilerTimelineEvent::EType::PopEvent;
-		TimelineEvent.EventIndex = Index;
-		GpuProfilerTimelineEvents.Push(TimelineEvent);
-
-		GpuProfilerEvents[Index]->End(RHICmdList);
-	}
-
-	void Clear(FRHICommandListImmediate* RHICommandListPtr)
-	{
-		EventStack.Empty();
-
-		for (int Index = 0; Index < GpuProfilerEvents.Num(); Index++)
-		{
-			if (GpuProfilerEvents[Index])
-			{
-				GpuProfilerEvents[Index]->ReleaseQueries(RenderQueryPool, RHICommandListPtr);
-				delete GpuProfilerEvents[Index];
-			}
-		}
-		GpuProfilerEvents.Empty();
-		GpuProfilerTimelineEvents.Empty();
-		EventAggregates.Empty();
-	}
-
-	bool UpdateStats(FRHICommandListImmediate& RHICmdList)
-	{
-		const bool bCsvStatsEnabled = !!CVarGPUCsvStatsEnabled.GetValueOnRenderThread();
-		const bool bTracingStatsEnabled = !!CVarGPUTracingStatsEnabled.GetValueOnRenderThread();
-
-		// Gather any remaining results and check all the results are ready
-		bool bAllQueriesAllocated = true;
-		for (int Index = 0; Index < GpuProfilerEvents.Num(); Index++)
-		{
-			FRealtimeGPUProfilerEvent* Event = GpuProfilerEvents[Index];
-			check(Event != nullptr);
-			if (!Event->HasValidResult())
-			{
-				Event->GatherQueryResults(RHICmdList);
-			}
-			if (!Event->HasValidResult())
-			{
-				// The frame isn't ready yet. Don't update stats - we'll try again next frame. 
-				return false;
-			}
-			if (!Event->HasQueriesAllocated())
-			{
-				bAllQueriesAllocated = false;
-			}
-		}
-
-		if (!bAllQueriesAllocated)
-		{
-			static bool bWarned = false;
-
-			if (!bWarned)
-			{
-				bWarned = true;
-				UE_LOG(LogSceneUtils, Warning, TEXT("Ran out of GPU queries! Results for this frame will be incomplete"));
-			}
-		}
-
-		// Calculate inclusive and exclusive time for all events
-
-		EventAggregates.Reserve(GpuProfilerEvents.Num());
-		TArray<int32, TInlineAllocator<32>> TimelineEventStack;
-
-		for (const FRealtimeGPUProfilerEvent* Event : GpuProfilerEvents)
-		{
-			FGPUEventTimeAggregate Aggregate;
-			Aggregate.InclusiveTime = Event->GetResultMS();
-			Aggregate.ExclusiveTime = Aggregate.InclusiveTime;
-			EventAggregates.Push(Aggregate);
-		}
-
-		for (const FRealtimeGPUProfilerTimelineEvent& TimelineEvent : GpuProfilerTimelineEvents)
-		{
-			if (TimelineEvent.Type == FRealtimeGPUProfilerTimelineEvent::EType::PushEvent)
-			{
-				if (TimelineEventStack.Num() != 0)
-				{
-					EventAggregates[TimelineEventStack.Last()].ExclusiveTime -= EventAggregates[TimelineEvent.EventIndex].InclusiveTime;
-				}
-				TimelineEventStack.Push(TimelineEvent.EventIndex);
-			}
-			else
-			{
-				TimelineEventStack.Pop();
-			}
-		}
-
-		// Update the stats
-
-		const bool GPUStatsChildTimesIncluded = !!CVarGPUStatsChildTimesIncluded.GetValueOnRenderThread();
-		float TotalMS = 0.0f;
-
-		TMap<FName, bool> StatSeenMap;
-		for (int Index = 0; Index < GpuProfilerEvents.Num(); Index++)
-		{
-			FRealtimeGPUProfilerEvent* Event = GpuProfilerEvents[Index];
-			check(Event != nullptr);
-			check(Event->HasValidResult());
-			const FName& StatName = Event->GetName();
-
-			// Check if we've seen this stat yet 
-			bool bIsNew = false;
-			if (StatSeenMap.Find(StatName) == nullptr)
-			{
-				StatSeenMap.Add(StatName, true);
-				bIsNew = true;
-			}
-
-			const float EventTime = GPUStatsChildTimesIncluded
-				? EventAggregates[Index].InclusiveTime
-				: EventAggregates[Index].ExclusiveTime;
-
-#if STATS
-			EStatOperation::Type StatOp = bIsNew ? EStatOperation::Set : EStatOperation::Add;
-			FThreadStats::AddMessage(Event->GetStatName(), StatOp, double(EventTime));
-#endif
-
-			if (bCsvStatsEnabled)
-			{
-				ECsvCustomStatOp CsvStatOp = bIsNew ? ECsvCustomStatOp::Set : ECsvCustomStatOp::Accumulate;
-				FCsvProfiler::Get()->RecordCustomStat(Event->GetName(), CSV_CATEGORY_INDEX(GPU), EventTime, CsvStatOp);
-			}
-
-#if TRACING_PROFILER
-			if (bTracingStatsEnabled)
-			{
-				ANSICHAR EventName[NAME_SIZE];
-				Event->GetName().GetPlainANSIString(EventName);
-
-				const uint32 GPUIndex = 0;
-				FTracingProfiler::Get()->AddGPUEvent(
-					EventName,
-					Event->GetStartResultMicroseconds(),
-					Event->GetEndResultMicroseconds(),
-					GPUIndex,
-					Event->GetFrameNumber());
-			}
-#endif //TRACING_PROFILER
-
-			TotalMS += EventAggregates[Index].ExclusiveTime;
-		}
-
-#if STATS
-		FThreadStats::AddMessage( GET_STATFNAME(Stat_GPU_Total), EStatOperation::Set, double(TotalMS) );
-#endif 
-
-#if CSV_PROFILER
-		if (bCsvStatsEnabled)
-		{
-			FCsvProfiler::Get()->RecordCustomStat(CSV_STAT_FNAME(Total), CSV_CATEGORY_INDEX(GPU), TotalMS, ECsvCustomStatOp::Set);
-		}
-#endif
-		return true;
-	}
-
-private:
-
-	TArray<FRealtimeGPUProfilerEvent*> GpuProfilerEvents;
-	TArray<int32> EventStack;
-
-	struct FRealtimeGPUProfilerTimelineEvent
-	{
-		enum class EType { PushEvent, PopEvent };
-		EType Type;
-		int32 EventIndex;
-	};
-
-	// All profiler push and pop events are recorded to calculate inclusive and exclusive timing
-	// while maintaining hierarchy and not splitting events unnecessarily.
-	TArray<FRealtimeGPUProfilerTimelineEvent> GpuProfilerTimelineEvents;
-
-	struct FGPUEventTimeAggregate
-	{
-		float ExclusiveTime;
-		float InclusiveTime;
-	};
-	TArray<FGPUEventTimeAggregate> EventAggregates;
-
-	uint32 FrameNumber;
-	FRenderQueryPool* RenderQueryPool;
-};
-
-/*-----------------------------------------------------------------------------
-FRealtimeGPUProfiler
------------------------------------------------------------------------------*/
-FRealtimeGPUProfiler* FRealtimeGPUProfiler::Instance = nullptr;
-
-FRealtimeGPUProfiler* FRealtimeGPUProfiler::Get()
-{
-	if (Instance == nullptr)
-	{
-		Instance = new FRealtimeGPUProfiler;
-	}
-	return Instance;
-}
-
-FRealtimeGPUProfiler::FRealtimeGPUProfiler()
-	: WriteBufferIndex(0)
-	, ReadBufferIndex(1) 
-	, WriteFrameNumber(-1)
-	, bStatGatheringPaused(false)
-	, bInBeginEndBlock(false)
-{
-	RenderQueryPool = new FRenderQueryPool(RQT_AbsoluteTime);
-	for (int Index = 0; Index < NumGPUProfilerBufferedFrames; Index++)
-	{
-		Frames.Add(new FRealtimeGPUProfilerFrame(RenderQueryPool));
-	}
-}
-
-void FRealtimeGPUProfiler::Release()
-{
-	for (int Index = 0; Index < Frames.Num(); Index++)
-	{
-		delete Frames[Index];
-	}
-	Frames.Empty();
-	delete RenderQueryPool;
-	RenderQueryPool = nullptr;
-}
-
-void FRealtimeGPUProfiler::BeginFrame(FRHICommandListImmediate& RHICmdList)
-{
-	check(bInBeginEndBlock == false);
-	bInBeginEndBlock = true;
-}
-
-inline bool AreGPUStatsEnabled()
-{
-	if (GSupportsTimestampRenderQueries == false || !CVarGPUStatsEnabled.GetValueOnRenderThread())
-	{
-		return false;
-	}
-
-#if STATS 
-	return true;
-#elif !CSV_PROFILER
-	return false;
-#else
-
-	// If we only have CSV stats, only capture if CSV GPU stats are enabled, and we're capturing
-	if (!CVarGPUCsvStatsEnabled.GetValueOnRenderThread())
-	{
-		return false;
-	}
-	if (!FCsvProfiler::Get()->IsCapturing_Renderthread())
-	{
-		return false;
-	}
-
-	return true;
-#endif
-}
-
-void FRealtimeGPUProfiler::EndFrame(FRHICommandListImmediate& RHICmdList)
-{
-	// This is called at the end of the renderthread frame. Note that the RHI thread may still be processing commands for the frame at this point, however
-	// The read buffer index is always 3 frames beind the write buffer index in order to prevent us reading from the frame the RHI thread is still processing. 
-	// This should also ensure the GPU is done with the queries before we try to read them
-	check(Frames.Num() > 0);
-	check(IsInRenderingThread());
-	check(bInBeginEndBlock == true);
-	bInBeginEndBlock = false;
-	if (!AreGPUStatsEnabled())
-	{
-		return;
-	}
-
-	if (Frames[ReadBufferIndex]->UpdateStats(RHICmdList))
-	{
-		// On a successful read, advance the ReadBufferIndex and WriteBufferIndex and clear the frame we just read
-		Frames[ReadBufferIndex]->Clear(&RHICmdList);
-		WriteFrameNumber = GFrameNumberRenderThread;
-		WriteBufferIndex = (WriteBufferIndex + 1) % Frames.Num();
-		ReadBufferIndex = (ReadBufferIndex + 1) % Frames.Num();
-		bStatGatheringPaused = false;
-	}
-	else
-	{
-		// The stats weren't ready; skip the next frame and don't advance the indices. We'll try to read the stats again next frame
-		bStatGatheringPaused = true;
-	}
-}
-
-void FRealtimeGPUProfiler::PushEvent(FRHICommandListImmediate& RHICmdList, const FName& Name, const FName& StatName)
-{
-	check(IsInRenderingThread());
-	if (bStatGatheringPaused || !bInBeginEndBlock)
-	{
-		return;
-	}
-	check(Frames.Num() > 0);
-	if (WriteBufferIndex >= 0)
-	{
-		Frames[WriteBufferIndex]->PushEvent(RHICmdList, Name, StatName);
-	}
-}
-
-void FRealtimeGPUProfiler::PopEvent(FRHICommandListImmediate& RHICmdList)
-{
-	check(IsInRenderingThread());
-	if (bStatGatheringPaused || !bInBeginEndBlock)
-	{
-		return;
-	}
-	check(Frames.Num() > 0);
-	if (WriteBufferIndex >= 0)
-	{
-		Frames[WriteBufferIndex]->PopEvent(RHICmdList);
-	}
-}
-
-/*-----------------------------------------------------------------------------
-FScopedGPUStatEvent
------------------------------------------------------------------------------*/
-void FScopedGPUStatEvent::Begin(FRHICommandList& InRHICmdList, const FName& Name, const FName& StatName)
-{
-	check(IsInRenderingThread());
-	if (!AreGPUStatsEnabled())
-	{
-		return;
-	}
-
-	// Non-immediate command lists are not supported (silently fail)
-	if (InRHICmdList.IsImmediate())
-	{
-		RHICmdList = (FRHICommandListImmediate*)&InRHICmdList;
-		FRealtimeGPUProfiler::Get()->PushEvent(*RHICmdList, Name, StatName);
-	}
-}
-
-void FScopedGPUStatEvent::End()
-{
-	check(IsInRenderingThread());
-	if (!AreGPUStatsEnabled())
-	{
-		return;
-	}
-	if (RHICmdList != nullptr)
-	{
-		FRealtimeGPUProfiler::Get()->PopEvent(*RHICmdList);
-	}
-}
-#endif // HAS_GPU_STATS
-
-/*-----------------------------------------------------------------------------
-FRenderQueryPool
------------------------------------------------------------------------------*/
-FRenderQueryPool::~FRenderQueryPool()
-{
-	Release();
-}
-
-void FRenderQueryPool::Release()
-{
-	Queries.Empty();
-	NumQueriesAllocated = 0;
-}
-
-FRenderQueryRHIRef FRenderQueryPool::AllocateQuery()
-{
-	// Are we out of available render queries?
-	NumQueriesAllocated++;
-	if (Queries.Num() == 0)
-	{
-		// Create a new render query.
-		return RHICreateRenderQuery(QueryType);
-	}
-
-	return Queries.Pop(/*bAllowShrinking=*/ false);
-}
-
-void FRenderQueryPool::ReleaseQuery(FRenderQueryRHIRef &Query)
-{
-	if (IsValidRef(Query))
-	{
-		NumQueriesAllocated--;
-#if RENDER_QUERY_POOLING_ENABLED
-		// Is no one else keeping a refcount to the query?
-		if (Query.GetRefCount() == 1)
-		{
-			// Return it to the pool.
-			Queries.Add(Query);
-		}
-#endif
-		// De-ref without deleting.
-		Query = NULL;
-	}
-}
+}
\ No newline at end of file
diff --git a/Engine/Source/Runtime/Engine/Public/ProfilingDebugging/TracingProfiler.h b/Engine/Source/Runtime/Engine/Public/ProfilingDebugging/TracingProfiler.h
deleted file mode 100644
diff --git a/Engine/Source/Runtime/Engine/Public/SceneUtils.h b/Engine/Source/Runtime/Engine/Public/SceneUtils.h
--- a/Engine/Source/Runtime/Engine/Public/SceneUtils.h
+++ b/Engine/Source/Runtime/Engine/Public/SceneUtils.h
@@ -11,9 +11,8 @@
 #pragma once
 
 #include "CoreMinimal.h"
-#include "Stats/Stats.h"
 #include "RHI.h"
-#include "ProfilingDebugging/CsvProfiler.h"
+#include "ProfilingDebugging/RealtimeGPUProfiler.h"
 
 enum class EShadingPath
 {
@@ -22,284 +21,6 @@ enum class EShadingPath
 	Num,
 };
 
-// Note:  WITH_PROFILEGPU should be 0 for final builds
-#define WANTS_DRAW_MESH_EVENTS (RHI_COMMAND_LIST_DEBUG_TRACES || (WITH_PROFILEGPU && PLATFORM_SUPPORTS_DRAW_MESH_EVENTS))
-
-class FRealtimeGPUProfiler;
-class FRealtimeGPUProfilerEvent;
-class FRealtimeGPUProfilerFrame;
-class FRenderQueryPool;
-class FScopedGPUStatEvent;
-
-#if WANTS_DRAW_MESH_EVENTS
-
-	/**
-	 * Class that logs draw events based upon class scope. Draw events can be seen
-	 * in PIX
-	 */
-	template<typename TRHICmdList>
-	struct ENGINE_API TDrawEvent
-	{
-		/** Cmdlist to push onto. */
-		TRHICmdList* RHICmdList;
-
-		/** Default constructor, initializing all member variables. */
-		FORCEINLINE TDrawEvent()
-			: RHICmdList(nullptr)
-		{}
-
-		/**
-		 * Terminate the event based upon scope
-		 */
-		FORCEINLINE ~TDrawEvent()
-		{
-			if (RHICmdList)
-			{
-				Stop();
-			}
-		}
-
-		/**
-		 * Function for logging a PIX event with var args
-		 */
-		void CDECL Start(TRHICmdList& RHICmdList, FColor Color, const TCHAR* Fmt, ...);
-		void Stop();
-	};
-
-	struct ENGINE_API FDrawEventRHIExecute
-	{
-		/** Context to execute on*/
-		class IRHIComputeContext* RHICommandContext;
-
-		/** Default constructor, initializing all member variables. */
-		FORCEINLINE FDrawEventRHIExecute()
-			: RHICommandContext(nullptr)
-		{}
-
-		/**
-		* Terminate the event based upon scope
-		*/
-		FORCEINLINE ~FDrawEventRHIExecute()
-		{
-			if (RHICommandContext)
-			{
-				Stop();
-			}
-		}
-
-		/**
-		* Function for logging a PIX event with var args
-		*/
-		void CDECL Start(IRHIComputeContext& InRHICommandContext, FColor Color, const TCHAR* Fmt, ...);
-		void Stop();
-	};
-
-	#define SCOPED_GPU_EVENT(RHICmdList, Name) TDrawEvent<typename TDecay<decltype(RHICmdList)>::Type> PREPROCESSOR_JOIN(Event_##Name,__LINE__); if(GetEmitDrawEvents()) PREPROCESSOR_JOIN(Event_##Name,__LINE__).Start(RHICmdList, FColor(0), TEXT(#Name));
-	#define SCOPED_GPU_EVENT_COLOR(RHICmdList, Color, Name) TDrawEvent<typename TDecay<decltype(RHICmdList)>::Type> PREPROCESSOR_JOIN(Event_##Name,__LINE__); if(GetEmitDrawEvents()) PREPROCESSOR_JOIN(Event_##Name,__LINE__).Start(RHICmdList, Color, TEXT(#Name));
-	#define SCOPED_GPU_EVENTF(RHICmdList, Name, Format, ...) TDrawEvent<typename TDecay<decltype(RHICmdList)>::Type> PREPROCESSOR_JOIN(Event_##Name,__LINE__); if(GetEmitDrawEvents()) PREPROCESSOR_JOIN(Event_##Name,__LINE__).Start(RHICmdList, FColor(0), Format, ##__VA_ARGS__);
-	#define SCOPED_GPU_EVENTF_COLOR(RHICmdList, Color, Name, Format, ...) TDrawEvent<typename TDecay<decltype(RHICmdList)>::Type> PREPROCESSOR_JOIN(Event_##Name,__LINE__); if(GetEmitDrawEvents()) PREPROCESSOR_JOIN(Event_##Name,__LINE__).Start(RHICmdList, Color, Format, ##__VA_ARGS__);
-	#define SCOPED_CONDITIONAL_GPU_EVENT(RHICmdList, Name, Condition) TDrawEvent<typename TDecay<decltype(RHICmdList)>::Type> PREPROCESSOR_JOIN(Event_##Name,__LINE__); if(GetEmitDrawEvents() && (Condition)) PREPROCESSOR_JOIN(Event_##Name,__LINE__).Start(RHICmdList, FColor(0), TEXT(#Name));
-	#define SCOPED_CONDITIONAL_GPU_EVENT_COLOR(RHICmdList, Name, Color, Condition) TDrawEvent<typename TDecay<decltype(RHICmdList)>::Type> PREPROCESSOR_JOIN(Event_##Name,__LINE__); if(GetEmitDrawEvents() && (Condition)) PREPROCESSOR_JOIN(Event_##Name,__LINE__).Start(RHICmdList, Color, TEXT(#Name));
-	#define SCOPED_CONDITIONAL_GPU_EVENTF(RHICmdList, Name, Condition, Format, ...) TDrawEvent<typename TDecay<decltype(RHICmdList)>::Type> PREPROCESSOR_JOIN(Event_##Name,__LINE__); if(GetEmitDrawEvents() && (Condition)) PREPROCESSOR_JOIN(Event_##Name,__LINE__).Start(RHICmdList, FColor(0), Format, ##__VA_ARGS__);
-	#define SCOPED_CONDITIONAL_GPU_EVENTF_COLOR(RHICmdList, Color, Name, Condition, Format, ...) TDrawEvent<typename TDecay<decltype(RHICmdList)>::Type> PREPROCESSOR_JOIN(Event_##Name,__LINE__); if(GetEmitDrawEvents() && (Condition)) PREPROCESSOR_JOIN(Event_##Name,__LINE__).Start(RHICmdList, Color, Format, ##__VA_ARGS__);
-	#define BEGIN_GPU_EVENTF(RHICmdList, Name, Event, Format, ...) if(GetEmitDrawEvents()) Event.Start(RHICmdList, FColor(0), Format, ##__VA_ARGS__);
-	#define BEGIN_GPU_EVENTF_COLOR(RHICmdList, Color, Name, Event, Format, ...) if(GetEmitDrawEvents()) Event.Start(RHICmdList, Color, Format, ##__VA_ARGS__);
-	#define STOP_GPU_EVENT(Event) (Event).Stop();
-
-	// Macros to allow for scoping of draw events outside of RHI function implementations
-	#define SCOPED_DRAW_EVENT(RHICmdList, Name) TDrawEvent<FRHICommandList> PREPROCESSOR_JOIN(Event_##Name,__LINE__); if(GetEmitDrawEvents()) PREPROCESSOR_JOIN(Event_##Name,__LINE__).Start(RHICmdList, FColor(0), TEXT(#Name));
-	#define SCOPED_DRAW_EVENT_COLOR(RHICmdList, Color, Name) TDrawEvent<FRHICommandList> PREPROCESSOR_JOIN(Event_##Name,__LINE__); if(GetEmitDrawEvents()) PREPROCESSOR_JOIN(Event_##Name,__LINE__).Start(RHICmdList, Color, TEXT(#Name));
-	#define SCOPED_DRAW_EVENTF(RHICmdList, Name, Format, ...) TDrawEvent<FRHICommandList> PREPROCESSOR_JOIN(Event_##Name,__LINE__); if(GetEmitDrawEvents()) PREPROCESSOR_JOIN(Event_##Name,__LINE__).Start(RHICmdList, FColor(0), Format, ##__VA_ARGS__);
-	#define SCOPED_DRAW_EVENTF_COLOR(RHICmdList, Color, Name, Format, ...) TDrawEvent<FRHICommandList> PREPROCESSOR_JOIN(Event_##Name,__LINE__); if(GetEmitDrawEvents()) PREPROCESSOR_JOIN(Event_##Name,__LINE__).Start(RHICmdList, Color, Format, ##__VA_ARGS__);
-	#define SCOPED_CONDITIONAL_DRAW_EVENT(RHICmdList, Name, Condition) TDrawEvent<FRHICommandList> PREPROCESSOR_JOIN(Event_##Name,__LINE__); if(GetEmitDrawEvents() && (Condition)) PREPROCESSOR_JOIN(Event_##Name,__LINE__).Start(RHICmdList, FColor(0), TEXT(#Name));
-	#define SCOPED_CONDITIONAL_DRAW_EVENT_COLOR(RHICmdList, Name, Color, Condition) TDrawEvent<FRHICommandList> PREPROCESSOR_JOIN(Event_##Name,__LINE__); if(GetEmitDrawEvents() && (Condition)) PREPROCESSOR_JOIN(Event_##Name,__LINE__).Start(RHICmdList, Color, TEXT(#Name));
-	#define SCOPED_CONDITIONAL_DRAW_EVENTF(RHICmdList, Name, Condition, Format, ...) TDrawEvent<FRHICommandList> PREPROCESSOR_JOIN(Event_##Name,__LINE__); if(GetEmitDrawEvents() && (Condition)) PREPROCESSOR_JOIN(Event_##Name,__LINE__).Start(RHICmdList, FColor(0), Format, ##__VA_ARGS__);
-	#define SCOPED_CONDITIONAL_DRAW_EVENTF_COLOR(RHICmdList, Color, Name, Condition, Format, ...) TDrawEvent<FRHICommandList> PREPROCESSOR_JOIN(Event_##Name,__LINE__); if(GetEmitDrawEvents() && (Condition)) PREPROCESSOR_JOIN(Event_##Name,__LINE__).Start(RHICmdList, Color, Format, ##__VA_ARGS__);
-	#define BEGIN_DRAW_EVENTF(RHICmdList, Name, Event, Format, ...) if(GetEmitDrawEvents()) Event.Start(RHICmdList, FColor(0), Format, ##__VA_ARGS__);
-	#define BEGIN_DRAW_EVENTF_COLOR(RHICmdList, Color, Name, Event, Format, ...) if(GetEmitDrawEvents()) Event.Start(RHICmdList, Color, Format, ##__VA_ARGS__);
-	#define STOP_DRAW_EVENT(Event) (Event).Stop();
-
-	#define SCOPED_COMPUTE_EVENT(RHICmdList, Name) TDrawEvent<FRHIAsyncComputeCommandList> PREPROCESSOR_JOIN(Event_##Name,__LINE__); if(GetEmitDrawEvents()) PREPROCESSOR_JOIN(Event_##Name,__LINE__).Start(RHICmdList, FColor(0), TEXT(#Name));
-	#define SCOPED_COMPUTE_EVENT_COLOR(RHICmdList, Color, Name) TDrawEvent<FRHIAsyncComputeCommandList> PREPROCESSOR_JOIN(Event_##Name,__LINE__); if(GetEmitDrawEvents()) PREPROCESSOR_JOIN(Event_##Name,__LINE__).Start(RHICmdList, Color, TEXT(#Name));
-	#define SCOPED_COMPUTE_EVENTF(RHICmdList, Name, Format, ...) TDrawEvent<FRHIAsyncComputeCommandList> PREPROCESSOR_JOIN(Event_##Name,__LINE__); if(GetEmitDrawEvents()) PREPROCESSOR_JOIN(Event_##Name,__LINE__).Start(RHICmdList, FColor(0), Format, ##__VA_ARGS__);
-	#define SCOPED_COMPUTE_EVENTF_COLOR(RHICmdList, Color, Name, Format, ...) TDrawEvent<FRHIAsyncComputeCommandList> PREPROCESSOR_JOIN(Event_##Name,__LINE__); if(GetEmitDrawEvents()) PREPROCESSOR_JOIN(Event_##Name,__LINE__).Start(RHICmdList, Color, Format, ##__VA_ARGS__);
-	#define SCOPED_CONDITIONAL_COMPUTE_EVENT(RHICmdList, Name, Condition) TDrawEvent<FRHIAsyncComputeCommandList> PREPROCESSOR_JOIN(Event_##Name,__LINE__); if(GetEmitDrawEvents() && (Condition)) PREPROCESSOR_JOIN(Event_##Name,__LINE__).Start(RHICmdList, FColor(0), TEXT(#Name));
-	#define SCOPED_CONDITIONAL_COMPUTE_EVENT_COLOR(RHICmdList, Color, Name, Condition) TDrawEvent<FRHIAsyncComputeCommandList> PREPROCESSOR_JOIN(Event_##Name,__LINE__); if(GetEmitDrawEvents() && (Condition)) PREPROCESSOR_JOIN(Event_##Name,__LINE__).Start(RHICmdList, Color, TEXT(#Name));
-	#define SCOPED_CONDITIONAL_COMPUTE_EVENTF(RHICmdList, Name, Condition, Format, ...) TDrawEvent<FRHIAsyncComputeCommandList> PREPROCESSOR_JOIN(Event_##Name,__LINE__); if(GetEmitDrawEvents() && (Condition)) PREPROCESSOR_JOIN(Event_##Name,__LINE__).Start(RHICmdList, FColor(0), Format, ##__VA_ARGS__);
-	#define SCOPED_CONDITIONAL_COMPUTE_EVENTF_COLOR(RHICmdList, Color, Name, Condition, Format, ...) TDrawEvent<FRHIAsyncComputeCommandList> PREPROCESSOR_JOIN(Event_##Name,__LINE__); if(GetEmitDrawEvents() && (Condition)) PREPROCESSOR_JOIN(Event_##Name,__LINE__).Start(RHICmdList, Color, Format, ##__VA_ARGS__);
-
-	// Macros to allow for scoping of draw events within RHI function implementations
-	#define SCOPED_RHI_DRAW_EVENT(RHICmdContext, Name) FDrawEventRHIExecute PREPROCESSOR_JOIN(Event_##Name,__LINE__); if(GetEmitDrawEvents()) PREPROCESSOR_JOIN(Event_##Name,__LINE__).Start(RHICmdContext, FColor(0), TEXT(#Name));
-	#define SCOPED_RHI_DRAW_EVENT_COLOR(RHICmdContext, Color, Name) FDrawEventRHIExecute PREPROCESSOR_JOIN(Event_##Name,__LINE__); if(GetEmitDrawEvents()) PREPROCESSOR_JOIN(Event_##Name,__LINE__).Start(RHICmdContext, Color, TEXT(#Name));
-	#define SCOPED_RHI_DRAW_EVENTF(RHICmdContext, Name, Format, ...) FDrawEventRHIExecute PREPROCESSOR_JOIN(Event_##Name,__LINE__); if(GetEmitDrawEvents()) PREPROCESSOR_JOIN(Event_##Name,__LINE__).Start(RHICmdContext, FColor(0), Format, ##__VA_ARGS__);
-	#define SCOPED_RHI_DRAW_EVENTF_COLOR(RHICmdContext, Color, Name, Format, ...) FDrawEventRHIExecute PREPROCESSOR_JOIN(Event_##Name,__LINE__); if(GetEmitDrawEvents()) PREPROCESSOR_JOIN(Event_##Name,__LINE__).Start(RHICmdContext, Color, Format, ##__VA_ARGS__);
-	#define SCOPED_RHI_CONDITIONAL_DRAW_EVENT(RHICmdContext, Name, Condition) FDrawEventRHIExecute PREPROCESSOR_JOIN(Event_##Name,__LINE__); if(GetEmitDrawEvents() && (Condition)) PREPROCESSOR_JOIN(Event_##Name,__LINE__).Start(RHICmdContext, FColor(0), TEXT(#Name));
-	#define SCOPED_RHI_CONDITIONAL_DRAW_EVENT_COLOR(RHICmdContext, Color, Name, Condition) FDrawEventRHIExecute PREPROCESSOR_JOIN(Event_##Name,__LINE__); if(GetEmitDrawEvents() && (Condition)) PREPROCESSOR_JOIN(Event_##Name,__LINE__).Start(RHICmdContext, Color, TEXT(#Name));
-	#define SCOPED_RHI_CONDITIONAL_DRAW_EVENTF(RHICmdContext, Name, Condition, Format, ...) FDrawEventRHIExecute PREPROCESSOR_JOIN(Event_##Name,__LINE__); if(GetEmitDrawEvents() && (Condition)) PREPROCESSOR_JOIN(Event_##Name,__LINE__).Start(RHICmdContext, FColor(0), Format, ##__VA_ARGS__);
-	#define SCOPED_RHI_CONDITIONAL_DRAW_EVENTF_COLOR(RHICmdContext, Color, Name, Condition, Format, ...) FDrawEventRHIExecute PREPROCESSOR_JOIN(Event_##Name,__LINE__); if(GetEmitDrawEvents() && (Condition)) PREPROCESSOR_JOIN(Event_##Name,__LINE__).Start(RHICmdContext, Color, Format, ##__VA_ARGS__);
-
-#else
-
-	template<typename TRHICmdList>
-	struct ENGINE_API TDrawEvent
-	{
-	};
-
-	#define SCOPED_GPU_EVENT(...)
-	#define SCOPED_GPU_EVENT_COLOR(...)
-	#define SCOPED_GPU_EVENTF(...)
-	#define SCOPED_GPU_EVENTF_COLOR(...)
-	#define SCOPED_CONDITIONAL_GPU_EVENT(...)
-	#define SCOPED_CONDITIONAL_GPU_EVENT_COLOR(...)
-	#define SCOPED_CONDITIONAL_GPU_EVENTF(...)
-	#define SCOPED_CONDITIONAL_GPU_EVENTF_COLOR(...)
-	#define BEGIN_GPU_EVENTF(...)
-	#define BEGIN_GPU_EVENTF_COLOR(...)
-	#define STOP_GPU_EVENT(...)
-
-	#define SCOPED_DRAW_EVENT(...)
-	#define SCOPED_DRAW_EVENT_COLOR(...)
-	#define SCOPED_DRAW_EVENTF(...)
-	#define SCOPED_DRAW_EVENTF_COLOR(...)
-	#define SCOPED_CONDITIONAL_DRAW_EVENT(...)
-	#define SCOPED_CONDITIONAL_DRAW_EVENT_COLOR(...)
-	#define SCOPED_CONDITIONAL_DRAW_EVENTF(...)
-	#define SCOPED_CONDITIONAL_DRAW_EVENTF_COLOR(...)
-	#define BEGIN_DRAW_EVENTF(...)
-	#define BEGIN_DRAW_EVENTF_COLOR(...)
-	#define STOP_DRAW_EVENT(...)
-
-	#define SCOPED_COMPUTE_EVENT(RHICmdList, Name)
-	#define SCOPED_COMPUTE_EVENT_COLOR(RHICmdList, Name)
-	#define SCOPED_COMPUTE_EVENTF(RHICmdList, Name, Format, ...)
-	#define SCOPED_COMPUTE_EVENTF_COLOR(RHICmdList, Name)
-	#define SCOPED_CONDITIONAL_COMPUTE_EVENT(RHICmdList, Name, Condition)
-	#define SCOPED_CONDITIONAL_COMPUTE_EVENT_COLOR(RHICmdList, Name, Condition)
-	#define SCOPED_CONDITIONAL_COMPUTE_EVENTF(RHICmdList, Name, Condition, Format, ...)
-	#define SCOPED_CONDITIONAL_COMPUTE_EVENTF_COLOR(RHICmdList, Name, Condition)
-
-	#define SCOPED_RHI_DRAW_EVENT(...)
-	#define SCOPED_RHI_DRAW_EVENT_COLOR(...)
-	#define SCOPED_RHI_DRAW_EVENTF(...)
-	#define SCOPED_RHI_DRAW_EVENTF_COLOR(...)
-	#define SCOPED_RHI_CONDITIONAL_DRAW_EVENT(...)
-	#define SCOPED_RHI_CONDITIONAL_DRAW_EVENT_COLOR(...)
-	#define SCOPED_RHI_CONDITIONAL_DRAW_EVENTF(...)
-	#define SCOPED_RHI_CONDITIONAL_DRAW_EVENTF_COLOR(...)
-
-#endif
-
-// GPU stats
-#if ( STATS || CSV_PROFILER ) && !PLATFORM_HTML5 && (!UE_BUILD_SHIPPING)
-#define HAS_GPU_STATS 1
-#else
-#define HAS_GPU_STATS 0
-#endif
-
-#if HAS_GPU_STATS
- CSV_DECLARE_CATEGORY_MODULE_EXTERN(ENGINE_API,GPU);
- // The DECLARE_GPU_STAT macros both declare and define a stat (for use in a single CPP)
- #define DECLARE_GPU_STAT(StatName) DECLARE_FLOAT_COUNTER_STAT(TEXT(#StatName), Stat_GPU_##StatName, STATGROUP_GPU); CSV_DEFINE_STAT(GPU,StatName);
- #define DECLARE_GPU_STAT_NAMED(StatName, NameString) DECLARE_FLOAT_COUNTER_STAT(NameString, Stat_GPU_##StatName, STATGROUP_GPU); CSV_DEFINE_STAT(GPU,StatName);
-
- // Extern GPU stats are needed where a stat is used in multiple CPPs. Use the DECLARE_GPU_STAT_NAMED_EXTERN in the header and DEFINE_GPU_STAT in the CPPs
- #define DECLARE_GPU_STAT_NAMED_EXTERN(StatName, NameString) DECLARE_FLOAT_COUNTER_STAT_EXTERN(NameString, Stat_GPU_##StatName, STATGROUP_GPU, ); CSV_DECLARE_STAT_EXTERN(GPU,StatName);
- #define DEFINE_GPU_STAT(StatName) DEFINE_STAT(Stat_GPU_##StatName); CSV_DEFINE_STAT(GPU,StatName);
-#if STATS
-  #define SCOPED_GPU_STAT(RHICmdList, StatName) FScopedGPUStatEvent PREPROCESSOR_JOIN(GPUStatEvent_##StatName,__LINE__); PREPROCESSOR_JOIN(GPUStatEvent_##StatName,__LINE__).Begin(RHICmdList, CSV_STAT_FNAME(StatName), GET_STATID( Stat_GPU_##StatName ).GetName() );
- #else
-  #define SCOPED_GPU_STAT(RHICmdList, StatName) FScopedGPUStatEvent PREPROCESSOR_JOIN(GPUStatEvent_##StatName,__LINE__); PREPROCESSOR_JOIN(GPUStatEvent_##StatName,__LINE__).Begin(RHICmdList, CSV_STAT_FNAME(StatName), FName() );
- #endif
- #define GPU_STATS_BEGINFRAME(RHICmdList) FRealtimeGPUProfiler::Get()->BeginFrame(RHICmdList);
- #define GPU_STATS_ENDFRAME(RHICmdList) FRealtimeGPUProfiler::Get()->EndFrame(RHICmdList);
-#else
- #define DECLARE_GPU_STAT(StatName)
- #define DECLARE_GPU_STAT_NAMED(StatName, NameString)
- #define DECLARE_GPU_STAT_NAMED_EXTERN(StatName, NameString)
- #define DEFINE_GPU_STAT(StatName)
- #define SCOPED_GPU_STAT(RHICmdList, StatName) 
- #define GPU_STATS_BEGINFRAME(RHICmdList) 
- #define GPU_STATS_ENDFRAME(RHICmdList) 
-#endif
-
-#if HAS_GPU_STATS
-
-class FRealtimeGPUProfilerEvent;
-class FRealtimeGPUProfilerFrame;
-class FRenderQueryPool;
-
-/**
-* FRealtimeGPUProfiler class. This manages recording and reporting all for GPU stats
-*/
-class FRealtimeGPUProfiler
-{
-	static FRealtimeGPUProfiler* Instance;
-public:
-	// Singleton interface
-	static ENGINE_API FRealtimeGPUProfiler* Get();
-
-	/** Per-frame update */
-	ENGINE_API void BeginFrame(FRHICommandListImmediate& RHICmdList);
-	ENGINE_API void EndFrame(FRHICommandListImmediate& RHICmdList);
-
-	/** Final cleanup */
-	ENGINE_API void Release();
-
-	/** Push/pop events */
-	void PushEvent(FRHICommandListImmediate& RHICmdList, const FName& Name, const FName& StatName);
-	void PopEvent(FRHICommandListImmediate& RHICmdList);
-
-private:
-	FRealtimeGPUProfiler();
-
-	/** Ringbuffer of profiler frames */
-	TArray<FRealtimeGPUProfilerFrame*> Frames;
-
-	int32 WriteBufferIndex;
-	int32 ReadBufferIndex;
-	uint32 WriteFrameNumber;
-	FRenderQueryPool* RenderQueryPool;
-	bool bStatGatheringPaused;
-	bool bInBeginEndBlock;
-};
-
-/**
-* Class that logs GPU Stat events for the realtime GPU profiler
-*/
-class FScopedGPUStatEvent
-{
-	/** Cmdlist to push onto. */
-	FRHICommandListImmediate* RHICmdList;
-
-	/** The stat event used to record timings */
-	FRealtimeGPUProfilerEvent* RealtimeGPUProfilerEvent;
-
-public:
-	/** Default constructor, initializing all member variables. */
-	FORCEINLINE FScopedGPUStatEvent()
-		: RHICmdList(nullptr)
-		, RealtimeGPUProfilerEvent(nullptr)
-	{}
-
-	/**
-	* Terminate the event based upon scope
-	*/
-	FORCEINLINE ~FScopedGPUStatEvent()
-	{
-		if (RHICmdList)
-		{
-			End();
-		}
-	}
-
-	/**
-	* Start/Stop functions for timer stats
-	*/
-	ENGINE_API void Begin(FRHICommandList& InRHICmdList, const FName& Name, const FName& StatName );
-	ENGINE_API void End();
-};
-#endif // HAS_GPU_STATS
-
 enum class EMobileHDRMode
 {
 	Unset,
@@ -323,40 +44,6 @@ ENGINE_API EMobileHDRMode GetMobileHDRMode();
 
 ENGINE_API bool IsMobileColorsRGB();
 
-/**
-* A pool of render (e.g. occlusion/timer) queries which are allocated individually, and returned to the pool as a group.
-*/
-class ENGINE_API FRenderQueryPool
-{
-public:
-	FRenderQueryPool(ERenderQueryType InQueryType)
-		: QueryType(InQueryType)
-		, NumQueriesAllocated(0)
-	{ }
-
-	virtual ~FRenderQueryPool();
-
-	/** Releases all the render queries in the pool. */
-	void Release();
-
-	/** Allocates an render query from the pool. */
-	FRenderQueryRHIRef AllocateQuery();
-
-	/** De-reference an render query, returning it to the pool instead of deleting it when the refcount reaches 0. */
-	void ReleaseQuery(FRenderQueryRHIRef &Query);
-
-	/** Returns the number of currently allocated queries. This is not necessarily the same as the pool size */
-	int32 GetAllocatedQueryCount() const { return NumQueriesAllocated;  }
-
-private:
-	/** Container for available render queries. */
-	TArray<FRenderQueryRHIRef> Queries;
-
-	ERenderQueryType QueryType;
-
-	int32 NumQueriesAllocated;
-};
-
 // Callback for calling one action (typical use case: delay a clear until it's actually needed)
 class FDelayedRendererAction
 {
diff --git a/Engine/Source/Runtime/RenderCore/Private/ProfilingDebugging/RealtimeGPUProfiler.cpp b/Engine/Source/Runtime/RenderCore/Private/ProfilingDebugging/RealtimeGPUProfiler.cpp
new file mode 100644
--- /dev/null
+++ b/Engine/Source/Runtime/RenderCore/Private/ProfilingDebugging/RealtimeGPUProfiler.cpp
@@ -0,0 +1,714 @@
+// Copyright 1998-2019 Epic Games, Inc. All Rights Reserved.
+
+#include "ProfilingDebugging/RealtimeGPUProfiler.h"
+#include "ProfilingDebugging/CsvProfiler.h"
+#include "ProfilingDebugging/TracingProfiler.h"
+#include "RenderCore.h"
+
+// Only exposed for debugging. Disabling this carries a severe performance penalty
+#define RENDER_QUERY_POOLING_ENABLED 1
+
+#if HAS_GPU_STATS 
+
+CSV_DEFINE_CATEGORY_MODULE(RENDERCORE_API, GPU, true);
+
+static TAutoConsoleVariable<int> CVarGPUStatsEnabled(
+	TEXT("r.GPUStatsEnabled"),
+	1,
+	TEXT("Enables or disables GPU stat recording"));
+
+
+static TAutoConsoleVariable<int> CVarGPUStatsMaxQueriesPerFrame(
+	TEXT("r.GPUStatsMaxQueriesPerFrame"),
+	-1,
+	TEXT("Limits the number of timestamps allocated per frame. -1 = no limit"), 
+	ECVF_RenderThreadSafe);
+
+
+static TAutoConsoleVariable<int> CVarGPUCsvStatsEnabled(
+	TEXT("r.GPUCsvStatsEnabled"),
+	0,
+	TEXT("Enables or disables GPU stat recording to CSVs"));
+
+DECLARE_GPU_STAT_NAMED( Total, TEXT("[TOTAL]") );
+
+static TAutoConsoleVariable<int> CVarGPUTracingStatsEnabled(
+	TEXT("r.GPUTracingStatsEnabled"),
+	1,
+	TEXT("Enables or disables GPU stat recording to tracing profiler"));
+
+static TAutoConsoleVariable<int> CVarGPUStatsChildTimesIncluded(
+	TEXT("r.GPUStatsChildTimesIncluded"),
+	0,
+	TEXT("If this is enabled, the child stat timings will be included in their parents' times.\n")
+	TEXT("This presents problems for non-hierarchical stats if we're expecting them to add up\n")
+	TEXT("to the total GPU time, so we probably want this disabled.\n")
+);
+
+#endif //HAS_GPU_STATS
+
+#if WANTS_DRAW_MESH_EVENTS
+
+template<typename TRHICmdList>
+void TDrawEvent<TRHICmdList>::Start(TRHICmdList& InRHICmdList, FColor Color, const TCHAR* Fmt, ...)
+{
+	check(IsInParallelRenderingThread() || IsInRHIThread());
+	{
+		va_list ptr;
+		va_start(ptr, Fmt);
+		TCHAR TempStr[256];
+		// Build the string in the temp buffer
+		FCString::GetVarArgs(TempStr, ARRAY_COUNT(TempStr), Fmt, ptr);
+		InRHICmdList.PushEvent(TempStr, Color);
+		RHICmdList = &InRHICmdList;
+		va_end(ptr);
+	}
+}
+
+template<typename TRHICmdList>
+void TDrawEvent<TRHICmdList>::Stop()
+{
+	if (RHICmdList)
+	{
+		RHICmdList->PopEvent();
+		RHICmdList = NULL;
+	}
+}
+template struct TDrawEvent<FRHICommandList>;
+template struct TDrawEvent<FRHIAsyncComputeCommandList>;
+template struct TDrawEvent<FRHIAsyncComputeCommandListImmediate>;
+
+void FDrawEventRHIExecute::Start(IRHIComputeContext& InRHICommandContext, FColor Color, const TCHAR* Fmt, ...)
+{
+	check(IsInParallelRenderingThread() || IsInRHIThread() || (!IsRunningRHIInSeparateThread() && IsInRenderingThread()));
+	{
+		va_list ptr;
+		va_start(ptr, Fmt);
+		TCHAR TempStr[256];
+		// Build the string in the temp buffer
+		FCString::GetVarArgs(TempStr, ARRAY_COUNT(TempStr), Fmt, ptr);
+		RHICommandContext = &InRHICommandContext;
+		RHICommandContext->RHIPushEvent(TempStr, Color);
+		va_end(ptr);
+	}
+}
+
+void FDrawEventRHIExecute::Stop()
+{
+	RHICommandContext->RHIPopEvent();
+}
+
+#endif // WANTS_DRAW_MESH_EVENTS
+
+#if HAS_GPU_STATS
+static const int32 NumGPUProfilerBufferedFrames = 4;
+
+/*-----------------------------------------------------------------------------
+FRealTimeGPUProfilerEvent class
+-----------------------------------------------------------------------------*/
+class FRealtimeGPUProfilerEvent
+{
+public:
+	static const uint64 InvalidQueryResult = 0xFFFFFFFFFFFFFFFFull;
+
+public:
+	FRealtimeGPUProfilerEvent(const FName& InName, const FName& InStatName, FRenderQueryPool* RenderQueryPool)
+		: StartResultMicroseconds(InvalidQueryResult)
+		, EndResultMicroseconds(InvalidQueryResult)
+		, FrameNumber(-1)
+		, bInsideQuery(false)
+		, bBeginQueryInFlight(false)
+		, bEndQueryInFlight(false)
+	{
+#if STATS
+		StatName = InStatName;
+#endif
+		Name = InName;
+
+		const int MaxGPUQueries = CVarGPUStatsMaxQueriesPerFrame.GetValueOnRenderThread();
+		if ( MaxGPUQueries == -1 || RenderQueryPool->GetAllocatedQueryCount() < MaxGPUQueries )
+		{
+			StartQuery = RenderQueryPool->AllocateQuery();
+			EndQuery = RenderQueryPool->AllocateQuery();
+		}
+	}
+
+	bool HasQueriesAllocated() const 
+	{ 
+		return IsValidRef(StartQuery); 
+	}
+
+	void ReleaseQueries(FRenderQueryPool* RenderQueryPool, FRHICommandListImmediate* RHICmdListPtr)
+	{
+		if ( HasQueriesAllocated() )
+		{
+			if (RHICmdListPtr)
+			{
+				// If we have queries in flight then get results before releasing back to the pool to avoid an ensure fail in the gnm RHI
+				uint64 Temp;
+				if (bBeginQueryInFlight)
+				{
+					RHICmdListPtr->GetRenderQueryResult(StartQuery, Temp, false);
+				}
+
+				if (bEndQueryInFlight)
+				{
+					RHICmdListPtr->GetRenderQueryResult(EndQuery, Temp, false);
+				}
+			}
+			RenderQueryPool->ReleaseQuery(StartQuery);
+			RenderQueryPool->ReleaseQuery(EndQuery);
+		}
+	}
+
+	void Begin(FRHICommandListImmediate& RHICmdList)
+	{
+		check(IsInRenderingThread());
+		check(!bInsideQuery);
+		bInsideQuery = true;
+
+		if (IsValidRef(StartQuery))
+		{
+			RHICmdList.EndRenderQuery(StartQuery);
+			bBeginQueryInFlight = true;
+		}
+		StartResultMicroseconds = InvalidQueryResult;
+		EndResultMicroseconds = InvalidQueryResult;
+		FrameNumber = GFrameNumberRenderThread;
+	}
+
+	void End(FRHICommandListImmediate& RHICmdList)
+	{
+		check(IsInRenderingThread());
+		check(bInsideQuery);
+		bInsideQuery = false;
+
+		if ( HasQueriesAllocated() )
+		{
+			RHICmdList.EndRenderQuery(EndQuery);
+			bEndQueryInFlight = true;
+		}
+	}
+
+	bool GatherQueryResults(FRHICommandListImmediate& RHICmdList)
+	{
+		QUICK_SCOPE_CYCLE_COUNTER(STAT_SceneUtils_GatherQueryResults);
+		// Get the query results which are still outstanding
+		check(GFrameNumberRenderThread != FrameNumber);
+		if ( HasQueriesAllocated() )
+		{
+			if (StartResultMicroseconds == InvalidQueryResult)
+			{
+				if (!RHICmdList.GetRenderQueryResult(StartQuery, StartResultMicroseconds, false))
+				{
+					StartResultMicroseconds = InvalidQueryResult;
+				}
+				bBeginQueryInFlight = false;
+			}
+			if (EndResultMicroseconds == InvalidQueryResult)
+			{
+				if (!RHICmdList.GetRenderQueryResult(EndQuery, EndResultMicroseconds, false))
+				{
+					EndResultMicroseconds = InvalidQueryResult;
+				}
+				bEndQueryInFlight = false;
+			}
+		}
+		else
+		{
+			// If we don't have a query allocated, just set the results to zero
+			EndResultMicroseconds = StartResultMicroseconds = 0;
+		}
+		return HasValidResult();
+	}
+
+	float GetResultMS() const
+	{
+		check(HasValidResult());
+		if (EndResultMicroseconds < StartResultMicroseconds)
+		{
+			// This should never happen...
+			return 0.0f;
+		}
+		return float(EndResultMicroseconds - StartResultMicroseconds) / 1000.0f;
+	}
+
+	bool HasValidResult() const
+	{
+		return StartResultMicroseconds != FRealtimeGPUProfilerEvent::InvalidQueryResult && EndResultMicroseconds != FRealtimeGPUProfilerEvent::InvalidQueryResult;
+	}
+
+#if STATS
+	const FName& GetStatName() const
+	{
+		return StatName;
+	}
+#endif
+	const FName& GetName() const
+	{
+		return Name;
+	}
+
+	uint64 GetStartResultMicroseconds(uint32 GPUIndex = 0) const
+	{
+		return StartResultMicroseconds;
+	}
+
+	uint64 GetEndResultMicroseconds(uint32 GPUIndex = 0) const
+	{
+		return EndResultMicroseconds;
+	}
+
+	uint32 GetFrameNumber() const
+	{
+		return FrameNumber;
+	}
+
+private:
+	FRenderQueryRHIRef StartQuery;
+	FRenderQueryRHIRef EndQuery;
+#if STATS
+	FName StatName;
+#endif
+	FName Name;
+	uint64 StartResultMicroseconds;
+	uint64 EndResultMicroseconds;
+	uint32 FrameNumber;
+
+	bool bInsideQuery;
+	bool bBeginQueryInFlight;
+	bool bEndQueryInFlight;
+};
+
+/*-----------------------------------------------------------------------------
+FRealtimeGPUProfilerFrame class
+Container for a single frame's GPU stats
+-----------------------------------------------------------------------------*/
+class FRealtimeGPUProfilerFrame
+{
+public:
+	FRealtimeGPUProfilerFrame(FRenderQueryPool* InRenderQueryPool)
+		: FrameNumber(-1)
+		, RenderQueryPool(InRenderQueryPool)
+	{}
+
+	~FRealtimeGPUProfilerFrame()
+	{
+		Clear(nullptr);
+	}
+
+	void PushEvent(FRHICommandListImmediate& RHICmdList, const FName& Name, const FName& StatName)
+	{
+		// TODO: this should really use a pool / free list
+		FRealtimeGPUProfilerEvent* Event = new FRealtimeGPUProfilerEvent(Name, StatName, RenderQueryPool);
+		const int32 EventIndex = GpuProfilerEvents.Num();
+
+		GpuProfilerEvents.Add(Event);
+
+		FRealtimeGPUProfilerTimelineEvent TimelineEvent = {};
+		TimelineEvent.Type = FRealtimeGPUProfilerTimelineEvent::EType::PushEvent;
+		TimelineEvent.EventIndex = EventIndex;
+
+		GpuProfilerTimelineEvents.Push(TimelineEvent);
+
+		EventStack.Push(EventIndex);
+		Event->Begin(RHICmdList);
+	}
+
+	void PopEvent(FRHICommandListImmediate& RHICmdList)
+	{
+		int32 Index = EventStack.Pop();
+
+		FRealtimeGPUProfilerTimelineEvent TimelineEvent = {};
+		TimelineEvent.Type = FRealtimeGPUProfilerTimelineEvent::EType::PopEvent;
+		TimelineEvent.EventIndex = Index;
+		GpuProfilerTimelineEvents.Push(TimelineEvent);
+
+		GpuProfilerEvents[Index]->End(RHICmdList);
+	}
+
+	void Clear(FRHICommandListImmediate* RHICommandListPtr)
+	{
+		EventStack.Empty();
+
+		for (int Index = 0; Index < GpuProfilerEvents.Num(); Index++)
+		{
+			if (GpuProfilerEvents[Index])
+			{
+				GpuProfilerEvents[Index]->ReleaseQueries(RenderQueryPool, RHICommandListPtr);
+				delete GpuProfilerEvents[Index];
+			}
+		}
+		GpuProfilerEvents.Empty();
+		GpuProfilerTimelineEvents.Empty();
+		EventAggregates.Empty();
+	}
+
+	bool UpdateStats(FRHICommandListImmediate& RHICmdList)
+	{
+		const bool bCsvStatsEnabled = !!CVarGPUCsvStatsEnabled.GetValueOnRenderThread();
+		const bool bTracingStatsEnabled = !!CVarGPUTracingStatsEnabled.GetValueOnRenderThread();
+
+		// Gather any remaining results and check all the results are ready
+		bool bAllQueriesAllocated = true;
+		for (int Index = 0; Index < GpuProfilerEvents.Num(); Index++)
+		{
+			FRealtimeGPUProfilerEvent* Event = GpuProfilerEvents[Index];
+			check(Event != nullptr);
+			if (!Event->HasValidResult())
+			{
+				Event->GatherQueryResults(RHICmdList);
+			}
+			if (!Event->HasValidResult())
+			{
+				// The frame isn't ready yet. Don't update stats - we'll try again next frame. 
+				return false;
+			}
+			if (!Event->HasQueriesAllocated())
+			{
+				bAllQueriesAllocated = false;
+			}
+		}
+
+		if (!bAllQueriesAllocated)
+		{
+			static bool bWarned = false;
+
+			if (!bWarned)
+			{
+				bWarned = true;
+				UE_LOG(LogRendererCore, Warning, TEXT("Ran out of GPU queries! Results for this frame will be incomplete"));
+			}
+		}
+
+		// Calculate inclusive and exclusive time for all events
+
+		EventAggregates.Reserve(GpuProfilerEvents.Num());
+		TArray<int32, TInlineAllocator<32>> TimelineEventStack;
+
+		for (const FRealtimeGPUProfilerEvent* Event : GpuProfilerEvents)
+		{
+			FGPUEventTimeAggregate Aggregate;
+			Aggregate.InclusiveTime = Event->GetResultMS();
+			Aggregate.ExclusiveTime = Aggregate.InclusiveTime;
+			EventAggregates.Push(Aggregate);
+		}
+
+		for (const FRealtimeGPUProfilerTimelineEvent& TimelineEvent : GpuProfilerTimelineEvents)
+		{
+			if (TimelineEvent.Type == FRealtimeGPUProfilerTimelineEvent::EType::PushEvent)
+			{
+				if (TimelineEventStack.Num() != 0)
+				{
+					EventAggregates[TimelineEventStack.Last()].ExclusiveTime -= EventAggregates[TimelineEvent.EventIndex].InclusiveTime;
+				}
+				TimelineEventStack.Push(TimelineEvent.EventIndex);
+			}
+			else
+			{
+				TimelineEventStack.Pop();
+			}
+		}
+
+		// Update the stats
+
+		const bool GPUStatsChildTimesIncluded = !!CVarGPUStatsChildTimesIncluded.GetValueOnRenderThread();
+		float TotalMS = 0.0f;
+
+		TMap<FName, bool> StatSeenMap;
+		for (int Index = 0; Index < GpuProfilerEvents.Num(); Index++)
+		{
+			FRealtimeGPUProfilerEvent* Event = GpuProfilerEvents[Index];
+			check(Event != nullptr);
+			check(Event->HasValidResult());
+			const FName& StatName = Event->GetName();
+
+			// Check if we've seen this stat yet 
+			bool bIsNew = false;
+			if (StatSeenMap.Find(StatName) == nullptr)
+			{
+				StatSeenMap.Add(StatName, true);
+				bIsNew = true;
+			}
+
+			const float EventTime = GPUStatsChildTimesIncluded
+				? EventAggregates[Index].InclusiveTime
+				: EventAggregates[Index].ExclusiveTime;
+
+#if STATS
+			EStatOperation::Type StatOp = bIsNew ? EStatOperation::Set : EStatOperation::Add;
+			FThreadStats::AddMessage(Event->GetStatName(), StatOp, double(EventTime));
+#endif
+
+			if (bCsvStatsEnabled)
+			{
+				ECsvCustomStatOp CsvStatOp = bIsNew ? ECsvCustomStatOp::Set : ECsvCustomStatOp::Accumulate;
+				FCsvProfiler::Get()->RecordCustomStat(Event->GetName(), CSV_CATEGORY_INDEX(GPU), EventTime, CsvStatOp);
+			}
+
+#if TRACING_PROFILER
+			if (bTracingStatsEnabled)
+			{
+				ANSICHAR EventName[NAME_SIZE];
+				Event->GetName().GetPlainANSIString(EventName);
+
+				const uint32 GPUIndex = 0;
+				FTracingProfiler::Get()->AddGPUEvent(
+					EventName,
+					Event->GetStartResultMicroseconds(),
+					Event->GetEndResultMicroseconds(),
+					GPUIndex,
+					Event->GetFrameNumber());
+			}
+#endif //TRACING_PROFILER
+
+			TotalMS += EventAggregates[Index].ExclusiveTime;
+		}
+
+#if STATS
+		FThreadStats::AddMessage( GET_STATFNAME(Stat_GPU_Total), EStatOperation::Set, double(TotalMS) );
+#endif 
+
+#if CSV_PROFILER
+		if (bCsvStatsEnabled)
+		{
+			FCsvProfiler::Get()->RecordCustomStat(CSV_STAT_FNAME(Total), CSV_CATEGORY_INDEX(GPU), TotalMS, ECsvCustomStatOp::Set);
+		}
+#endif
+		return true;
+	}
+
+private:
+
+	TArray<FRealtimeGPUProfilerEvent*> GpuProfilerEvents;
+	TArray<int32> EventStack;
+
+	struct FRealtimeGPUProfilerTimelineEvent
+	{
+		enum class EType { PushEvent, PopEvent };
+		EType Type;
+		int32 EventIndex;
+	};
+
+	// All profiler push and pop events are recorded to calculate inclusive and exclusive timing
+	// while maintaining hierarchy and not splitting events unnecessarily.
+	TArray<FRealtimeGPUProfilerTimelineEvent> GpuProfilerTimelineEvents;
+
+	struct FGPUEventTimeAggregate
+	{
+		float ExclusiveTime;
+		float InclusiveTime;
+	};
+	TArray<FGPUEventTimeAggregate> EventAggregates;
+
+	uint32 FrameNumber;
+	FRenderQueryPool* RenderQueryPool;
+};
+
+/*-----------------------------------------------------------------------------
+FRealtimeGPUProfiler
+-----------------------------------------------------------------------------*/
+FRealtimeGPUProfiler* FRealtimeGPUProfiler::Instance = nullptr;
+
+FRealtimeGPUProfiler* FRealtimeGPUProfiler::Get()
+{
+	if (Instance == nullptr)
+	{
+		Instance = new FRealtimeGPUProfiler;
+	}
+	return Instance;
+}
+
+FRealtimeGPUProfiler::FRealtimeGPUProfiler()
+	: WriteBufferIndex(0)
+	, ReadBufferIndex(1) 
+	, WriteFrameNumber(-1)
+	, bStatGatheringPaused(false)
+	, bInBeginEndBlock(false)
+{
+	RenderQueryPool = new FRenderQueryPool(RQT_AbsoluteTime);
+	for (int Index = 0; Index < NumGPUProfilerBufferedFrames; Index++)
+	{
+		Frames.Add(new FRealtimeGPUProfilerFrame(RenderQueryPool));
+	}
+}
+
+void FRealtimeGPUProfiler::Release()
+{
+	for (int Index = 0; Index < Frames.Num(); Index++)
+	{
+		delete Frames[Index];
+	}
+	Frames.Empty();
+	delete RenderQueryPool;
+	RenderQueryPool = nullptr;
+}
+
+void FRealtimeGPUProfiler::BeginFrame(FRHICommandListImmediate& RHICmdList)
+{
+	check(bInBeginEndBlock == false);
+	bInBeginEndBlock = true;
+}
+
+inline bool AreGPUStatsEnabled()
+{
+	if (GSupportsTimestampRenderQueries == false || !CVarGPUStatsEnabled.GetValueOnRenderThread())
+	{
+		return false;
+	}
+
+#if STATS 
+	return true;
+#elif !CSV_PROFILER
+	return false;
+#else
+
+	// If we only have CSV stats, only capture if CSV GPU stats are enabled, and we're capturing
+	if (!CVarGPUCsvStatsEnabled.GetValueOnRenderThread())
+	{
+		return false;
+	}
+	if (!FCsvProfiler::Get()->IsCapturing_Renderthread())
+	{
+		return false;
+	}
+
+	return true;
+#endif
+}
+
+void FRealtimeGPUProfiler::EndFrame(FRHICommandListImmediate& RHICmdList)
+{
+	// This is called at the end of the renderthread frame. Note that the RHI thread may still be processing commands for the frame at this point, however
+	// The read buffer index is always 3 frames beind the write buffer index in order to prevent us reading from the frame the RHI thread is still processing. 
+	// This should also ensure the GPU is done with the queries before we try to read them
+	check(Frames.Num() > 0);
+	check(IsInRenderingThread());
+	check(bInBeginEndBlock == true);
+	bInBeginEndBlock = false;
+	if (!AreGPUStatsEnabled())
+	{
+		return;
+	}
+
+	if (Frames[ReadBufferIndex]->UpdateStats(RHICmdList))
+	{
+		// On a successful read, advance the ReadBufferIndex and WriteBufferIndex and clear the frame we just read
+		Frames[ReadBufferIndex]->Clear(&RHICmdList);
+		WriteFrameNumber = GFrameNumberRenderThread;
+		WriteBufferIndex = (WriteBufferIndex + 1) % Frames.Num();
+		ReadBufferIndex = (ReadBufferIndex + 1) % Frames.Num();
+		bStatGatheringPaused = false;
+	}
+	else
+	{
+		// The stats weren't ready; skip the next frame and don't advance the indices. We'll try to read the stats again next frame
+		bStatGatheringPaused = true;
+	}
+}
+
+void FRealtimeGPUProfiler::PushEvent(FRHICommandListImmediate& RHICmdList, const FName& Name, const FName& StatName)
+{
+	check(IsInRenderingThread());
+	if (bStatGatheringPaused || !bInBeginEndBlock)
+	{
+		return;
+	}
+	check(Frames.Num() > 0);
+	if (WriteBufferIndex >= 0)
+	{
+		Frames[WriteBufferIndex]->PushEvent(RHICmdList, Name, StatName);
+	}
+}
+
+void FRealtimeGPUProfiler::PopEvent(FRHICommandListImmediate& RHICmdList)
+{
+	check(IsInRenderingThread());
+	if (bStatGatheringPaused || !bInBeginEndBlock)
+	{
+		return;
+	}
+	check(Frames.Num() > 0);
+	if (WriteBufferIndex >= 0)
+	{
+		Frames[WriteBufferIndex]->PopEvent(RHICmdList);
+	}
+}
+
+/*-----------------------------------------------------------------------------
+FScopedGPUStatEvent
+-----------------------------------------------------------------------------*/
+void FScopedGPUStatEvent::Begin(FRHICommandList& InRHICmdList, const FName& Name, const FName& StatName)
+{
+	check(IsInRenderingThread());
+	if (!AreGPUStatsEnabled())
+	{
+		return;
+	}
+
+	// Non-immediate command lists are not supported (silently fail)
+	if (InRHICmdList.IsImmediate())
+	{
+		RHICmdList = (FRHICommandListImmediate*)&InRHICmdList;
+		FRealtimeGPUProfiler::Get()->PushEvent(*RHICmdList, Name, StatName);
+	}
+}
+
+void FScopedGPUStatEvent::End()
+{
+	check(IsInRenderingThread());
+	if (!AreGPUStatsEnabled())
+	{
+		return;
+	}
+	if (RHICmdList != nullptr)
+	{
+		FRealtimeGPUProfiler::Get()->PopEvent(*RHICmdList);
+	}
+}
+#endif // HAS_GPU_STATS
+
+/*-----------------------------------------------------------------------------
+FRenderQueryPool
+-----------------------------------------------------------------------------*/
+FRenderQueryPool::~FRenderQueryPool()
+{
+	Release();
+}
+
+void FRenderQueryPool::Release()
+{
+	Queries.Empty();
+	NumQueriesAllocated = 0;
+}
+
+FRenderQueryRHIRef FRenderQueryPool::AllocateQuery()
+{
+	// Are we out of available render queries?
+	NumQueriesAllocated++;
+	if (Queries.Num() == 0)
+	{
+		// Create a new render query.
+		return RHICreateRenderQuery(QueryType);
+	}
+
+	return Queries.Pop(/*bAllowShrinking=*/ false);
+}
+
+void FRenderQueryPool::ReleaseQuery(FRenderQueryRHIRef &Query)
+{
+	if (IsValidRef(Query))
+	{
+		NumQueriesAllocated--;
+#if RENDER_QUERY_POOLING_ENABLED
+		// Is no one else keeping a refcount to the query?
+		if (Query.GetRefCount() == 1)
+		{
+			// Return it to the pool.
+			Queries.Add(Query);
+		}
+#endif
+		// De-ref without deleting.
+		Query = NULL;
+	}
+}
diff --git a/Engine/Source/Runtime/RenderCore/Private/ProfilingDebugging/TracingProfiler.cpp b/Engine/Source/Runtime/RenderCore/Private/ProfilingDebugging/TracingProfiler.cpp
new file mode 100644
--- /dev/null
+++ b/Engine/Source/Runtime/RenderCore/Private/ProfilingDebugging/TracingProfiler.cpp
@@ -0,0 +1,366 @@
+// Copyright 1998-2019 Epic Games, Inc. All Rights Reserved.
+
+#include "ProfilingDebugging/TracingProfiler.h"
+
+#if TRACING_PROFILER
+
+#include "GPUProfiler.h"
+#include "HAL/ConsoleManager.h"
+#include "HAL/FileManager.h"
+#include "HAL/LowLevelMemTracker.h"
+#include "Misc/CoreDelegates.h"
+#include "Misc/Paths.h"
+#include "RHI.h"
+#include "RenderingThread.h"
+
+static TAutoConsoleVariable<int32> GTracingProfileBufferSize(
+	TEXT("TracingProfiler.BufferSize"),
+	65536,
+	TEXT("Defines the maximum umber of events stored in the internal ring buffer of the tracing profiler. ")
+	TEXT("Only read at process startup and can't be changed at runtime."),
+	ECVF_Default);
+
+FTracingProfiler* FTracingProfiler::Instance = nullptr;
+
+FTracingProfiler* FTracingProfiler::Get()
+{
+	if (Instance == nullptr)
+	{
+		Instance = new FTracingProfiler;
+	}
+	return Instance;
+}
+
+static void HandleTracingProfileCommand(const TArray<FString>& Args)
+{
+	if (Args.Num() < 1)
+	{
+		return;
+	}
+	FString Param = Args[0];
+	if (Param == TEXT("START"))
+	{
+		FTracingProfiler::Get()->BeginCapture();
+	}
+	else if (Param == TEXT("STOP"))
+	{
+		FTracingProfiler::Get()->EndCapture();
+	}
+	else
+	{
+		int32 CaptureFrames = 0;
+		if (FParse::Value(*Param, TEXT("FRAMES="), CaptureFrames))
+		{
+			FTracingProfiler::Get()->BeginCapture(CaptureFrames);
+		}
+	}
+}
+
+static FAutoConsoleCommand HandleTracingProfileCmd(
+	TEXT("TracingProfile"),
+	TEXT("Starts or stops tracing profiler"),
+	FConsoleCommandWithArgsDelegate::CreateStatic(&HandleTracingProfileCommand)
+);
+
+FTracingProfiler::FTracingProfiler()
+{
+	LLM_SCOPE(ELLMTag::Stats);
+
+	MaxNumCapturedEvents = GTracingProfileBufferSize.GetValueOnAnyThread();
+
+	FEvent DefaultEvent = {};
+	CapturedEvents.Init(DefaultEvent, MaxNumCapturedEvents);
+
+	FCoreDelegates::OnBeginFrame.AddLambda([]() { FTracingProfiler::Get()->BeginFrame(); });
+	FCoreDelegates::OnEndFrame.AddLambda([]() { FTracingProfiler::Get()->EndFrame(); });
+
+	FCoreDelegates::OnBeginFrameRT.AddLambda([]() { FTracingProfiler::Get()->BeginFrameRT(); });
+	FCoreDelegates::OnEndFrameRT.AddLambda([]() { FTracingProfiler::Get()->EndFrameRT(); });
+}
+
+void FTracingProfiler::Init()
+{
+	
+}
+
+void FTracingProfiler::BeginFrame()
+{
+	check(IsInGameThread());
+
+	if (bRequestStartCapture)
+	{
+		bCapturing = true;
+		bRequestStartCapture = false;
+	}
+
+	GameThreadFrameBeginCycle = FPlatformTime::Cycles64();
+}
+
+static inline uint64 CyclesToMicroseconds64(uint64 Cycles)
+{
+	return uint64(FPlatformTime::ToSeconds64(Cycles) * 1e6);
+}
+
+void FTracingProfiler::EndFrame()
+{
+	check(IsInGameThread());
+
+	if (!bCapturing)
+	{
+		return;
+	}
+
+	GameThreadFrameEndCycle = FPlatformTime::Cycles64();
+
+	AddCPUEvent("GameThread",
+		GameThreadFrameBeginCycle,
+		GameThreadFrameEndCycle,
+		GGameThreadId,
+		GFrameNumber);
+
+	if (NumFramesToCapture >= 0)
+	{
+		NumFramesToCapture--;
+		if (NumFramesToCapture == 0)
+		{
+			bRequestStopCapture = true;
+		}
+	}
+
+	if (bRequestStopCapture)
+	{
+		FlushRenderingCommands();
+
+		bCapturing = false;
+		bRequestStopCapture = false;
+
+		WriteCaptureToFile();
+
+	}
+
+	CaptureFrameNumber++;
+}
+
+void FTracingProfiler::BeginFrameRT()
+{
+	RenderThreadFrameBeginCycle = FPlatformTime::Cycles64();
+
+	if (!bCapturingRT && bCapturing)
+	{
+		GDynamicRHI->RHICalibrateTimers();
+	}
+
+	bCapturingRT = bCapturing;
+}
+
+void FTracingProfiler::EndFrameRT()
+{
+	RenderThreadFrameEndCycle = FPlatformTime::Cycles64();
+
+	if (!bCapturingRT)
+	{
+		return;
+	}
+
+	AddCPUEvent("RenderThread",
+		RenderThreadFrameBeginCycle,
+		RenderThreadFrameEndCycle,
+		GRenderThreadId,
+		GFrameNumberRenderThread);
+}
+
+
+int32 FTracingProfiler::GetCaptureFrameNumber()
+{
+	return CaptureFrameNumber;
+}
+
+void FTracingProfiler::BeginCapture(int InNumFramesToCapture)
+{
+	check(IsInGameThread());
+
+	NumFramesToCapture = InNumFramesToCapture;
+	CaptureFrameNumber = 0;
+
+	bRequestStartCapture = true;
+
+	SessionId++;
+}
+
+void FTracingProfiler::EndCapture()
+{
+	bRequestStopCapture = true;
+}
+
+bool FTracingProfiler::IsCapturing() const
+{
+	check(IsInGameThread());
+	return bRequestStartCapture || bCapturing;
+}
+
+void FTracingProfiler::WriteCaptureToFile()
+{
+	check(!bCapturing);
+
+	FString Filename = FString::Printf(TEXT("Profile(%s)"), *FDateTime::Now().ToString(TEXT("%Y%m%d_%H%M%S")));
+	FString TracingRootPath = FPaths::ProfilingDir() + TEXT("Traces/");
+	FString OutputFilename = TracingRootPath + Filename + TEXT(".json");
+
+	FArchive* OutputFile = IFileManager::Get().CreateDebugFileWriter(*OutputFilename);
+
+	auto WriteString = [OutputFile](const char* String)
+	{
+		OutputFile->Serialize((void*)String, sizeof(ANSICHAR)*FCStringAnsi::Strlen(String));
+	};
+
+	WriteString(R"({"traceEvents":[)" "\n");
+
+	enum { StringBufferSize = 1024 };
+	char StringBuffer[StringBufferSize];
+
+	const uint32 Pid = FPlatformProcess::GetCurrentProcessId();
+
+	const uint32 NumEvents = CapturedEvents.Num();
+
+	const uint32 LocalSessionId = SessionId;
+	auto IsEventValid = [LocalSessionId](const FEvent& Event)
+	{
+		if (Event.SessionId != LocalSessionId)
+		{
+			return false;
+		}
+
+		if (Event.Type == EEventType::GPU)
+		{
+			return Event.GPU.EndMicroseconds > Event.GPU.BeginMicroseconds
+				&& (Event.GPU.EndMicroseconds - Event.GPU.BeginMicroseconds) < 10000000; // Arbitrary threshold of 10 seconds for any one GPU event
+		}
+		else
+		{
+			return Event.CPU.EndCycles > Event.CPU.BeginCycles;
+		}
+	};
+
+	int32 MaxGPUIndex = -1;
+	uint64 FirstCPUTimestampCycles = UINT64_MAX;
+	uint64 FirstGPUTimestampMicroseconds = UINT64_MAX;
+	for (const FEvent& Event : CapturedEvents)
+	{
+		if (IsEventValid(Event))
+		{
+			if (Event.Type == EEventType::GPU)
+			{
+				MaxGPUIndex = FMath::Max<int32>(MaxGPUIndex, Event.GPU.GPUIndex);
+				FirstGPUTimestampMicroseconds = FMath::Min(FirstGPUTimestampMicroseconds, Event.GPU.BeginMicroseconds);
+			}
+			else
+			{
+				FirstCPUTimestampCycles = FMath::Min(FirstCPUTimestampCycles, Event.CPU.BeginCycles);
+			}
+		}
+	}
+
+	// Write metadata (thread names, sorting order, etc.)
+
+	int32 SortIndex = 0; // Lower numbers result in higher position in the visualizer.
+
+	for (int32 GPUIndex = 0; GPUIndex <= MaxGPUIndex; ++GPUIndex)
+	{
+		FCStringAnsi::Snprintf(StringBuffer, StringBufferSize,
+			R"({"pid":%d, "tid":%d, "ph": "M", "name": "thread_name", "args":{"name":"GPU %d"}},)"
+			R"({"pid":%d, "tid":%d, "ph": "M", "name": "thread_sort_index", "args":{"sort_index": %d}},)"
+			"\n",
+			Pid, GPUIndex, GPUIndex, Pid, GPUIndex, SortIndex);
+		WriteString(StringBuffer);
+		SortIndex++;
+	}
+
+	FCStringAnsi::Snprintf(StringBuffer, StringBufferSize,
+		R"({"pid":%d, "tid":%d, "ph": "M", "name": "thread_name", "args":{"name":"Render thread"}},)"
+		R"({"pid":%d, "tid":%d, "ph": "M", "name": "thread_sort_index", "args":{"sort_index": %d}},)"
+		"\n",
+		Pid, GRenderThreadId, Pid, GRenderThreadId, SortIndex);
+	WriteString(StringBuffer);
+	SortIndex++;
+
+	FCStringAnsi::Snprintf(StringBuffer, StringBufferSize,
+		R"({"pid":%d, "tid":%d, "ph": "M", "name": "thread_name", "args":{"name":"Game thread"}},)"
+		R"({"pid":%d, "tid":%d, "ph": "M", "name": "thread_sort_index", "args":{"sort_index": %d}},)"
+		"\n",
+		Pid, GGameThreadId, Pid, GGameThreadId, SortIndex);
+	WriteString(StringBuffer);
+	SortIndex++;
+
+	// Align GPU and CPU timestamps
+
+	FGPUTimingCalibrationTimestamp CalibrationTimestamp = FGPUTiming::GetCalibrationTimestamp();
+
+	// If platform does not support GPU/CPU timer alignment, then simply align GPU and CPU on first event
+	if (CalibrationTimestamp.CPUMicroseconds == 0 || CalibrationTimestamp.GPUMicroseconds == 0)
+	{
+		CalibrationTimestamp.CPUMicroseconds = CyclesToMicroseconds64(FirstCPUTimestampCycles);
+		CalibrationTimestamp.GPUMicroseconds = FirstGPUTimestampMicroseconds;
+	}
+
+	const uint64 GPUTimeOffset = CalibrationTimestamp.CPUMicroseconds - CalibrationTimestamp.GPUMicroseconds;
+
+	// Write out all events
+
+	uint32 NumProcessedEvents = 0;
+	for (const FEvent& Event : CapturedEvents)
+	{
+		if (!IsEventValid(Event))
+		{
+			continue;
+		}
+
+		uint64 TimeBeginMicroseconds = 0;
+		uint64 TimeEndMicroseconds = 0;
+
+		switch (Event.Type)
+		{
+		case EEventType::CPU:
+			TimeBeginMicroseconds = CyclesToMicroseconds64(Event.CPU.BeginCycles);
+			TimeEndMicroseconds = CyclesToMicroseconds64(Event.CPU.EndCycles);
+			break;
+		case EEventType::GPU:
+			// Note: We could also use `clock_sync` metadata to synchronize events in Chrome viewer.
+			// Advantage of doing it manually is that trace log can be consumed by a less sophisticated
+			// parser that does not implement all the features of Chrome viewer.
+			TimeBeginMicroseconds = Event.GPU.BeginMicroseconds + GPUTimeOffset;
+			TimeEndMicroseconds = Event.GPU.EndMicroseconds + GPUTimeOffset;
+			break;
+		default:
+			verifyf(0, TEXT("Unexpected profiling event type"));
+		}
+
+		uint32 ThreadOrGpuId;
+		if (Event.Type == EEventType::CPU)
+		{
+			ThreadOrGpuId = Event.CPU.ThreadId;
+		}
+		else
+		{
+			ThreadOrGpuId = Event.GPU.GPUIndex;
+		}
+
+		FCStringAnsi::Snprintf(StringBuffer, StringBufferSize,
+			R"({"pid":%d, "tid":%d, "ph": "X", "name": "%s", "ts": %llu, "dur": %llu, "args":{"frame":%d}},)"
+			"\n",
+			Pid, ThreadOrGpuId, Event.Name, TimeBeginMicroseconds,
+			TimeEndMicroseconds - TimeBeginMicroseconds,
+			Event.FrameNumber);
+
+		WriteString(StringBuffer);
+
+		++NumProcessedEvents;
+	}
+
+	// All done
+
+	WriteString("{}]}");
+
+	OutputFile->Close();
+}
+
+#endif //TRACING_PROFILER
diff --git a/Engine/Source/Runtime/RenderCore/Public/ProfilingDebugging/RealtimeGPUProfiler.h b/Engine/Source/Runtime/RenderCore/Public/ProfilingDebugging/RealtimeGPUProfiler.h
new file mode 100644
--- /dev/null
+++ b/Engine/Source/Runtime/RenderCore/Public/ProfilingDebugging/RealtimeGPUProfiler.h
@@ -0,0 +1,328 @@
+// Copyright 1998-2019 Epic Games, Inc. All Rights Reserved.
+
+/**
+ *
+ * This file contains the various draw mesh macros that display draw calls
+ * inside of PIX.
+ */
+
+// Colors that are defined for a particular mesh type
+// Each event type will be displayed using the defined color
+#pragma once
+
+#include "CoreMinimal.h"
+#include "Stats/Stats.h"
+#include "RHI.h"
+#include "ProfilingDebugging/CsvProfiler.h"
+
+// Note:  WITH_PROFILEGPU should be 0 for final builds
+#define WANTS_DRAW_MESH_EVENTS (RHI_COMMAND_LIST_DEBUG_TRACES || (WITH_PROFILEGPU && PLATFORM_SUPPORTS_DRAW_MESH_EVENTS))
+
+class FRealtimeGPUProfiler;
+class FRealtimeGPUProfilerEvent;
+class FRealtimeGPUProfilerFrame;
+class FRenderQueryPool;
+class FScopedGPUStatEvent;
+
+#if WANTS_DRAW_MESH_EVENTS
+
+	/**
+	 * Class that logs draw events based upon class scope. Draw events can be seen
+	 * in PIX
+	 */
+	template<typename TRHICmdList>
+	struct RENDERCORE_API TDrawEvent
+	{
+		/** Cmdlist to push onto. */
+		TRHICmdList* RHICmdList;
+
+		/** Default constructor, initializing all member variables. */
+		FORCEINLINE TDrawEvent()
+			: RHICmdList(nullptr)
+		{}
+
+		/**
+		 * Terminate the event based upon scope
+		 */
+		FORCEINLINE ~TDrawEvent()
+		{
+			if (RHICmdList)
+			{
+				Stop();
+			}
+		}
+
+		/**
+		 * Function for logging a PIX event with var args
+		 */
+		void CDECL Start(TRHICmdList& RHICmdList, FColor Color, const TCHAR* Fmt, ...);
+		void Stop();
+	};
+
+	struct RENDERCORE_API FDrawEventRHIExecute
+	{
+		/** Context to execute on*/
+		class IRHIComputeContext* RHICommandContext;
+
+		/** Default constructor, initializing all member variables. */
+		FORCEINLINE FDrawEventRHIExecute()
+			: RHICommandContext(nullptr)
+		{}
+
+		/**
+		* Terminate the event based upon scope
+		*/
+		FORCEINLINE ~FDrawEventRHIExecute()
+		{
+			if (RHICommandContext)
+			{
+				Stop();
+			}
+		}
+
+		/**
+		* Function for logging a PIX event with var args
+		*/
+		void CDECL Start(IRHIComputeContext& InRHICommandContext, FColor Color, const TCHAR* Fmt, ...);
+		void Stop();
+	};
+
+	#define SCOPED_GPU_EVENT(RHICmdList, Name) TDrawEvent<typename TDecay<decltype(RHICmdList)>::Type> PREPROCESSOR_JOIN(Event_##Name,__LINE__); if(GetEmitDrawEvents()) PREPROCESSOR_JOIN(Event_##Name,__LINE__).Start(RHICmdList, FColor(0), TEXT(#Name));
+	#define SCOPED_GPU_EVENT_COLOR(RHICmdList, Color, Name) TDrawEvent<typename TDecay<decltype(RHICmdList)>::Type> PREPROCESSOR_JOIN(Event_##Name,__LINE__); if(GetEmitDrawEvents()) PREPROCESSOR_JOIN(Event_##Name,__LINE__).Start(RHICmdList, Color, TEXT(#Name));
+	#define SCOPED_GPU_EVENTF(RHICmdList, Name, Format, ...) TDrawEvent<typename TDecay<decltype(RHICmdList)>::Type> PREPROCESSOR_JOIN(Event_##Name,__LINE__); if(GetEmitDrawEvents()) PREPROCESSOR_JOIN(Event_##Name,__LINE__).Start(RHICmdList, FColor(0), Format, ##__VA_ARGS__);
+	#define SCOPED_GPU_EVENTF_COLOR(RHICmdList, Color, Name, Format, ...) TDrawEvent<typename TDecay<decltype(RHICmdList)>::Type> PREPROCESSOR_JOIN(Event_##Name,__LINE__); if(GetEmitDrawEvents()) PREPROCESSOR_JOIN(Event_##Name,__LINE__).Start(RHICmdList, Color, Format, ##__VA_ARGS__);
+	#define SCOPED_CONDITIONAL_GPU_EVENT(RHICmdList, Name, Condition) TDrawEvent<typename TDecay<decltype(RHICmdList)>::Type> PREPROCESSOR_JOIN(Event_##Name,__LINE__); if(GetEmitDrawEvents() && (Condition)) PREPROCESSOR_JOIN(Event_##Name,__LINE__).Start(RHICmdList, FColor(0), TEXT(#Name));
+	#define SCOPED_CONDITIONAL_GPU_EVENT_COLOR(RHICmdList, Name, Color, Condition) TDrawEvent<typename TDecay<decltype(RHICmdList)>::Type> PREPROCESSOR_JOIN(Event_##Name,__LINE__); if(GetEmitDrawEvents() && (Condition)) PREPROCESSOR_JOIN(Event_##Name,__LINE__).Start(RHICmdList, Color, TEXT(#Name));
+	#define SCOPED_CONDITIONAL_GPU_EVENTF(RHICmdList, Name, Condition, Format, ...) TDrawEvent<typename TDecay<decltype(RHICmdList)>::Type> PREPROCESSOR_JOIN(Event_##Name,__LINE__); if(GetEmitDrawEvents() && (Condition)) PREPROCESSOR_JOIN(Event_##Name,__LINE__).Start(RHICmdList, FColor(0), Format, ##__VA_ARGS__);
+	#define SCOPED_CONDITIONAL_GPU_EVENTF_COLOR(RHICmdList, Color, Name, Condition, Format, ...) TDrawEvent<typename TDecay<decltype(RHICmdList)>::Type> PREPROCESSOR_JOIN(Event_##Name,__LINE__); if(GetEmitDrawEvents() && (Condition)) PREPROCESSOR_JOIN(Event_##Name,__LINE__).Start(RHICmdList, Color, Format, ##__VA_ARGS__);
+	#define BEGIN_GPU_EVENTF(RHICmdList, Name, Event, Format, ...) if(GetEmitDrawEvents()) Event.Start(RHICmdList, FColor(0), Format, ##__VA_ARGS__);
+	#define BEGIN_GPU_EVENTF_COLOR(RHICmdList, Color, Name, Event, Format, ...) if(GetEmitDrawEvents()) Event.Start(RHICmdList, Color, Format, ##__VA_ARGS__);
+	#define STOP_GPU_EVENT(Event) (Event).Stop();
+
+	// Macros to allow for scoping of draw events outside of RHI function implementations
+	#define SCOPED_DRAW_EVENT(RHICmdList, Name) TDrawEvent<FRHICommandList> PREPROCESSOR_JOIN(Event_##Name,__LINE__); if(GetEmitDrawEvents()) PREPROCESSOR_JOIN(Event_##Name,__LINE__).Start(RHICmdList, FColor(0), TEXT(#Name));
+	#define SCOPED_DRAW_EVENT_COLOR(RHICmdList, Color, Name) TDrawEvent<FRHICommandList> PREPROCESSOR_JOIN(Event_##Name,__LINE__); if(GetEmitDrawEvents()) PREPROCESSOR_JOIN(Event_##Name,__LINE__).Start(RHICmdList, Color, TEXT(#Name));
+	#define SCOPED_DRAW_EVENTF(RHICmdList, Name, Format, ...) TDrawEvent<FRHICommandList> PREPROCESSOR_JOIN(Event_##Name,__LINE__); if(GetEmitDrawEvents()) PREPROCESSOR_JOIN(Event_##Name,__LINE__).Start(RHICmdList, FColor(0), Format, ##__VA_ARGS__);
+	#define SCOPED_DRAW_EVENTF_COLOR(RHICmdList, Color, Name, Format, ...) TDrawEvent<FRHICommandList> PREPROCESSOR_JOIN(Event_##Name,__LINE__); if(GetEmitDrawEvents()) PREPROCESSOR_JOIN(Event_##Name,__LINE__).Start(RHICmdList, Color, Format, ##__VA_ARGS__);
+	#define SCOPED_CONDITIONAL_DRAW_EVENT(RHICmdList, Name, Condition) TDrawEvent<FRHICommandList> PREPROCESSOR_JOIN(Event_##Name,__LINE__); if(GetEmitDrawEvents() && (Condition)) PREPROCESSOR_JOIN(Event_##Name,__LINE__).Start(RHICmdList, FColor(0), TEXT(#Name));
+	#define SCOPED_CONDITIONAL_DRAW_EVENT_COLOR(RHICmdList, Name, Color, Condition) TDrawEvent<FRHICommandList> PREPROCESSOR_JOIN(Event_##Name,__LINE__); if(GetEmitDrawEvents() && (Condition)) PREPROCESSOR_JOIN(Event_##Name,__LINE__).Start(RHICmdList, Color, TEXT(#Name));
+	#define SCOPED_CONDITIONAL_DRAW_EVENTF(RHICmdList, Name, Condition, Format, ...) TDrawEvent<FRHICommandList> PREPROCESSOR_JOIN(Event_##Name,__LINE__); if(GetEmitDrawEvents() && (Condition)) PREPROCESSOR_JOIN(Event_##Name,__LINE__).Start(RHICmdList, FColor(0), Format, ##__VA_ARGS__);
+	#define SCOPED_CONDITIONAL_DRAW_EVENTF_COLOR(RHICmdList, Color, Name, Condition, Format, ...) TDrawEvent<FRHICommandList> PREPROCESSOR_JOIN(Event_##Name,__LINE__); if(GetEmitDrawEvents() && (Condition)) PREPROCESSOR_JOIN(Event_##Name,__LINE__).Start(RHICmdList, Color, Format, ##__VA_ARGS__);
+	#define BEGIN_DRAW_EVENTF(RHICmdList, Name, Event, Format, ...) if(GetEmitDrawEvents()) Event.Start(RHICmdList, FColor(0), Format, ##__VA_ARGS__);
+	#define BEGIN_DRAW_EVENTF_COLOR(RHICmdList, Color, Name, Event, Format, ...) if(GetEmitDrawEvents()) Event.Start(RHICmdList, Color, Format, ##__VA_ARGS__);
+	#define STOP_DRAW_EVENT(Event) (Event).Stop();
+
+	#define SCOPED_COMPUTE_EVENT(RHICmdList, Name) TDrawEvent<FRHIAsyncComputeCommandList> PREPROCESSOR_JOIN(Event_##Name,__LINE__); if(GetEmitDrawEvents()) PREPROCESSOR_JOIN(Event_##Name,__LINE__).Start(RHICmdList, FColor(0), TEXT(#Name));
+	#define SCOPED_COMPUTE_EVENT_COLOR(RHICmdList, Color, Name) TDrawEvent<FRHIAsyncComputeCommandList> PREPROCESSOR_JOIN(Event_##Name,__LINE__); if(GetEmitDrawEvents()) PREPROCESSOR_JOIN(Event_##Name,__LINE__).Start(RHICmdList, Color, TEXT(#Name));
+	#define SCOPED_COMPUTE_EVENTF(RHICmdList, Name, Format, ...) TDrawEvent<FRHIAsyncComputeCommandList> PREPROCESSOR_JOIN(Event_##Name,__LINE__); if(GetEmitDrawEvents()) PREPROCESSOR_JOIN(Event_##Name,__LINE__).Start(RHICmdList, FColor(0), Format, ##__VA_ARGS__);
+	#define SCOPED_COMPUTE_EVENTF_COLOR(RHICmdList, Color, Name, Format, ...) TDrawEvent<FRHIAsyncComputeCommandList> PREPROCESSOR_JOIN(Event_##Name,__LINE__); if(GetEmitDrawEvents()) PREPROCESSOR_JOIN(Event_##Name,__LINE__).Start(RHICmdList, Color, Format, ##__VA_ARGS__);
+	#define SCOPED_CONDITIONAL_COMPUTE_EVENT(RHICmdList, Name, Condition) TDrawEvent<FRHIAsyncComputeCommandList> PREPROCESSOR_JOIN(Event_##Name,__LINE__); if(GetEmitDrawEvents() && (Condition)) PREPROCESSOR_JOIN(Event_##Name,__LINE__).Start(RHICmdList, FColor(0), TEXT(#Name));
+	#define SCOPED_CONDITIONAL_COMPUTE_EVENT_COLOR(RHICmdList, Color, Name, Condition) TDrawEvent<FRHIAsyncComputeCommandList> PREPROCESSOR_JOIN(Event_##Name,__LINE__); if(GetEmitDrawEvents() && (Condition)) PREPROCESSOR_JOIN(Event_##Name,__LINE__).Start(RHICmdList, Color, TEXT(#Name));
+	#define SCOPED_CONDITIONAL_COMPUTE_EVENTF(RHICmdList, Name, Condition, Format, ...) TDrawEvent<FRHIAsyncComputeCommandList> PREPROCESSOR_JOIN(Event_##Name,__LINE__); if(GetEmitDrawEvents() && (Condition)) PREPROCESSOR_JOIN(Event_##Name,__LINE__).Start(RHICmdList, FColor(0), Format, ##__VA_ARGS__);
+	#define SCOPED_CONDITIONAL_COMPUTE_EVENTF_COLOR(RHICmdList, Color, Name, Condition, Format, ...) TDrawEvent<FRHIAsyncComputeCommandList> PREPROCESSOR_JOIN(Event_##Name,__LINE__); if(GetEmitDrawEvents() && (Condition)) PREPROCESSOR_JOIN(Event_##Name,__LINE__).Start(RHICmdList, Color, Format, ##__VA_ARGS__);
+
+	// Macros to allow for scoping of draw events within RHI function implementations
+	#define SCOPED_RHI_DRAW_EVENT(RHICmdContext, Name) FDrawEventRHIExecute PREPROCESSOR_JOIN(Event_##Name,__LINE__); if(GetEmitDrawEvents()) PREPROCESSOR_JOIN(Event_##Name,__LINE__).Start(RHICmdContext, FColor(0), TEXT(#Name));
+	#define SCOPED_RHI_DRAW_EVENT_COLOR(RHICmdContext, Color, Name) FDrawEventRHIExecute PREPROCESSOR_JOIN(Event_##Name,__LINE__); if(GetEmitDrawEvents()) PREPROCESSOR_JOIN(Event_##Name,__LINE__).Start(RHICmdContext, Color, TEXT(#Name));
+	#define SCOPED_RHI_DRAW_EVENTF(RHICmdContext, Name, Format, ...) FDrawEventRHIExecute PREPROCESSOR_JOIN(Event_##Name,__LINE__); if(GetEmitDrawEvents()) PREPROCESSOR_JOIN(Event_##Name,__LINE__).Start(RHICmdContext, FColor(0), Format, ##__VA_ARGS__);
+	#define SCOPED_RHI_DRAW_EVENTF_COLOR(RHICmdContext, Color, Name, Format, ...) FDrawEventRHIExecute PREPROCESSOR_JOIN(Event_##Name,__LINE__); if(GetEmitDrawEvents()) PREPROCESSOR_JOIN(Event_##Name,__LINE__).Start(RHICmdContext, Color, Format, ##__VA_ARGS__);
+	#define SCOPED_RHI_CONDITIONAL_DRAW_EVENT(RHICmdContext, Name, Condition) FDrawEventRHIExecute PREPROCESSOR_JOIN(Event_##Name,__LINE__); if(GetEmitDrawEvents() && (Condition)) PREPROCESSOR_JOIN(Event_##Name,__LINE__).Start(RHICmdContext, FColor(0), TEXT(#Name));
+	#define SCOPED_RHI_CONDITIONAL_DRAW_EVENT_COLOR(RHICmdContext, Color, Name, Condition) FDrawEventRHIExecute PREPROCESSOR_JOIN(Event_##Name,__LINE__); if(GetEmitDrawEvents() && (Condition)) PREPROCESSOR_JOIN(Event_##Name,__LINE__).Start(RHICmdContext, Color, TEXT(#Name));
+	#define SCOPED_RHI_CONDITIONAL_DRAW_EVENTF(RHICmdContext, Name, Condition, Format, ...) FDrawEventRHIExecute PREPROCESSOR_JOIN(Event_##Name,__LINE__); if(GetEmitDrawEvents() && (Condition)) PREPROCESSOR_JOIN(Event_##Name,__LINE__).Start(RHICmdContext, FColor(0), Format, ##__VA_ARGS__);
+	#define SCOPED_RHI_CONDITIONAL_DRAW_EVENTF_COLOR(RHICmdContext, Color, Name, Condition, Format, ...) FDrawEventRHIExecute PREPROCESSOR_JOIN(Event_##Name,__LINE__); if(GetEmitDrawEvents() && (Condition)) PREPROCESSOR_JOIN(Event_##Name,__LINE__).Start(RHICmdContext, Color, Format, ##__VA_ARGS__);
+
+#else
+
+	template<typename TRHICmdList>
+	struct RENDERCORE_API TDrawEvent
+	{
+	};
+
+	#define SCOPED_GPU_EVENT(...)
+	#define SCOPED_GPU_EVENT_COLOR(...)
+	#define SCOPED_GPU_EVENTF(...)
+	#define SCOPED_GPU_EVENTF_COLOR(...)
+	#define SCOPED_CONDITIONAL_GPU_EVENT(...)
+	#define SCOPED_CONDITIONAL_GPU_EVENT_COLOR(...)
+	#define SCOPED_CONDITIONAL_GPU_EVENTF(...)
+	#define SCOPED_CONDITIONAL_GPU_EVENTF_COLOR(...)
+	#define BEGIN_GPU_EVENTF(...)
+	#define BEGIN_GPU_EVENTF_COLOR(...)
+	#define STOP_GPU_EVENT(...)
+
+	#define SCOPED_DRAW_EVENT(...)
+	#define SCOPED_DRAW_EVENT_COLOR(...)
+	#define SCOPED_DRAW_EVENTF(...)
+	#define SCOPED_DRAW_EVENTF_COLOR(...)
+	#define SCOPED_CONDITIONAL_DRAW_EVENT(...)
+	#define SCOPED_CONDITIONAL_DRAW_EVENT_COLOR(...)
+	#define SCOPED_CONDITIONAL_DRAW_EVENTF(...)
+	#define SCOPED_CONDITIONAL_DRAW_EVENTF_COLOR(...)
+	#define BEGIN_DRAW_EVENTF(...)
+	#define BEGIN_DRAW_EVENTF_COLOR(...)
+	#define STOP_DRAW_EVENT(...)
+
+	#define SCOPED_COMPUTE_EVENT(RHICmdList, Name)
+	#define SCOPED_COMPUTE_EVENT_COLOR(RHICmdList, Name)
+	#define SCOPED_COMPUTE_EVENTF(RHICmdList, Name, Format, ...)
+	#define SCOPED_COMPUTE_EVENTF_COLOR(RHICmdList, Name)
+	#define SCOPED_CONDITIONAL_COMPUTE_EVENT(RHICmdList, Name, Condition)
+	#define SCOPED_CONDITIONAL_COMPUTE_EVENT_COLOR(RHICmdList, Name, Condition)
+	#define SCOPED_CONDITIONAL_COMPUTE_EVENTF(RHICmdList, Name, Condition, Format, ...)
+	#define SCOPED_CONDITIONAL_COMPUTE_EVENTF_COLOR(RHICmdList, Name, Condition)
+
+	#define SCOPED_RHI_DRAW_EVENT(...)
+	#define SCOPED_RHI_DRAW_EVENT_COLOR(...)
+	#define SCOPED_RHI_DRAW_EVENTF(...)
+	#define SCOPED_RHI_DRAW_EVENTF_COLOR(...)
+	#define SCOPED_RHI_CONDITIONAL_DRAW_EVENT(...)
+	#define SCOPED_RHI_CONDITIONAL_DRAW_EVENT_COLOR(...)
+	#define SCOPED_RHI_CONDITIONAL_DRAW_EVENTF(...)
+	#define SCOPED_RHI_CONDITIONAL_DRAW_EVENTF_COLOR(...)
+
+#endif
+
+// GPU stats
+#if ( STATS || CSV_PROFILER ) && !PLATFORM_HTML5 && (!UE_BUILD_SHIPPING)
+#define HAS_GPU_STATS 1
+#else
+#define HAS_GPU_STATS 0
+#endif
+
+#if HAS_GPU_STATS
+ CSV_DECLARE_CATEGORY_MODULE_EXTERN(RENDERCORE_API,GPU);
+ // The DECLARE_GPU_STAT macros both declare and define a stat (for use in a single CPP)
+ #define DECLARE_GPU_STAT(StatName) DECLARE_FLOAT_COUNTER_STAT(TEXT(#StatName), Stat_GPU_##StatName, STATGROUP_GPU); CSV_DEFINE_STAT(GPU,StatName);
+ #define DECLARE_GPU_STAT_NAMED(StatName, NameString) DECLARE_FLOAT_COUNTER_STAT(NameString, Stat_GPU_##StatName, STATGROUP_GPU); CSV_DEFINE_STAT(GPU,StatName);
+
+ // Extern GPU stats are needed where a stat is used in multiple CPPs. Use the DECLARE_GPU_STAT_NAMED_EXTERN in the header and DEFINE_GPU_STAT in the CPPs
+ #define DECLARE_GPU_STAT_NAMED_EXTERN(StatName, NameString) DECLARE_FLOAT_COUNTER_STAT_EXTERN(NameString, Stat_GPU_##StatName, STATGROUP_GPU, ); CSV_DECLARE_STAT_EXTERN(GPU,StatName);
+ #define DEFINE_GPU_STAT(StatName) DEFINE_STAT(Stat_GPU_##StatName); CSV_DEFINE_STAT(GPU,StatName);
+#if STATS
+  #define SCOPED_GPU_STAT(RHICmdList, StatName) FScopedGPUStatEvent PREPROCESSOR_JOIN(GPUStatEvent_##StatName,__LINE__); PREPROCESSOR_JOIN(GPUStatEvent_##StatName,__LINE__).Begin(RHICmdList, CSV_STAT_FNAME(StatName), GET_STATID( Stat_GPU_##StatName ).GetName() );
+ #else
+  #define SCOPED_GPU_STAT(RHICmdList, StatName) FScopedGPUStatEvent PREPROCESSOR_JOIN(GPUStatEvent_##StatName,__LINE__); PREPROCESSOR_JOIN(GPUStatEvent_##StatName,__LINE__).Begin(RHICmdList, CSV_STAT_FNAME(StatName), FName() );
+ #endif
+ #define GPU_STATS_BEGINFRAME(RHICmdList) FRealtimeGPUProfiler::Get()->BeginFrame(RHICmdList);
+ #define GPU_STATS_ENDFRAME(RHICmdList) FRealtimeGPUProfiler::Get()->EndFrame(RHICmdList);
+#else
+ #define DECLARE_GPU_STAT(StatName)
+ #define DECLARE_GPU_STAT_NAMED(StatName, NameString)
+ #define DECLARE_GPU_STAT_NAMED_EXTERN(StatName, NameString)
+ #define DEFINE_GPU_STAT(StatName)
+ #define SCOPED_GPU_STAT(RHICmdList, StatName) 
+ #define GPU_STATS_BEGINFRAME(RHICmdList) 
+ #define GPU_STATS_ENDFRAME(RHICmdList) 
+#endif
+
+#if HAS_GPU_STATS
+
+class FRealtimeGPUProfilerEvent;
+class FRealtimeGPUProfilerFrame;
+class FRenderQueryPool;
+
+/**
+* FRealtimeGPUProfiler class. This manages recording and reporting all for GPU stats
+*/
+class FRealtimeGPUProfiler
+{
+	static FRealtimeGPUProfiler* Instance;
+public:
+	// Singleton interface
+	static RENDERCORE_API FRealtimeGPUProfiler* Get();
+
+	/** Per-frame update */
+	RENDERCORE_API void BeginFrame(FRHICommandListImmediate& RHICmdList);
+	RENDERCORE_API void EndFrame(FRHICommandListImmediate& RHICmdList);
+
+	/** Final cleanup */
+	RENDERCORE_API void Release();
+
+	/** Push/pop events */
+	void PushEvent(FRHICommandListImmediate& RHICmdList, const FName& Name, const FName& StatName);
+	void PopEvent(FRHICommandListImmediate& RHICmdList);
+
+private:
+	FRealtimeGPUProfiler();
+
+	/** Ringbuffer of profiler frames */
+	TArray<FRealtimeGPUProfilerFrame*> Frames;
+
+	int32 WriteBufferIndex;
+	int32 ReadBufferIndex;
+	uint32 WriteFrameNumber;
+	FRenderQueryPool* RenderQueryPool;
+	bool bStatGatheringPaused;
+	bool bInBeginEndBlock;
+};
+
+/**
+* Class that logs GPU Stat events for the realtime GPU profiler
+*/
+class FScopedGPUStatEvent
+{
+	/** Cmdlist to push onto. */
+	FRHICommandListImmediate* RHICmdList;
+
+	/** The stat event used to record timings */
+	FRealtimeGPUProfilerEvent* RealtimeGPUProfilerEvent;
+
+public:
+	/** Default constructor, initializing all member variables. */
+	FORCEINLINE FScopedGPUStatEvent()
+		: RHICmdList(nullptr)
+		, RealtimeGPUProfilerEvent(nullptr)
+	{}
+
+	/**
+	* Terminate the event based upon scope
+	*/
+	FORCEINLINE ~FScopedGPUStatEvent()
+	{
+		if (RHICmdList)
+		{
+			End();
+		}
+	}
+
+	/**
+	* Start/Stop functions for timer stats
+	*/
+	RENDERCORE_API void Begin(FRHICommandList& InRHICmdList, const FName& Name, const FName& StatName );
+	RENDERCORE_API void End();
+};
+#endif // HAS_GPU_STATS
+
+/**
+* A pool of render (e.g. occlusion/timer) queries which are allocated individually, and returned to the pool as a group.
+*/
+class RENDERCORE_API FRenderQueryPool
+{
+public:
+	FRenderQueryPool(ERenderQueryType InQueryType)
+		: QueryType(InQueryType)
+		, NumQueriesAllocated(0)
+	{ }
+
+	virtual ~FRenderQueryPool();
+
+	/** Releases all the render queries in the pool. */
+	void Release();
+
+	/** Allocates an render query from the pool. */
+	FRenderQueryRHIRef AllocateQuery();
+
+	/** De-reference an render query, returning it to the pool instead of deleting it when the refcount reaches 0. */
+	void ReleaseQuery(FRenderQueryRHIRef &Query);
+
+	/** Returns the number of currently allocated queries. This is not necessarily the same as the pool size */
+	int32 GetAllocatedQueryCount() const { return NumQueriesAllocated;  }
+
+private:
+	/** Container for available render queries. */
+	TArray<FRenderQueryRHIRef> Queries;
+
+	ERenderQueryType QueryType;
+
+	int32 NumQueriesAllocated;
+};
\ No newline at end of file
diff --git a/Engine/Source/Runtime/RenderCore/Public/ProfilingDebugging/TracingProfiler.h b/Engine/Source/Runtime/RenderCore/Public/ProfilingDebugging/TracingProfiler.h
new file mode 100644
--- /dev/null
+++ b/Engine/Source/Runtime/RenderCore/Public/ProfilingDebugging/TracingProfiler.h
@@ -0,0 +1,152 @@
+// Copyright 1998-2019 Epic Games, Inc. All Rights Reserved.
+
+/**
+*
+* A lightweight profiler that can output logs compatible with Google Chrome tracing visualizer.
+* Captured events are written as a flat array (fixed size ring buffer), without any kind of aggregation.
+* Tracing events may be added from multiple threads simultaneously. 
+* Old trace events are overwritten when ring buffer wraps.
+*/
+
+#pragma once
+
+#include "HAL/PlatformAtomics.h"
+#include "Modules/ModuleInterface.h"
+#include "Modules/ModuleManager.h"
+#include "CoreTypes.h"
+#include "CoreMinimal.h"
+
+#define TRACING_PROFILER (WITH_ENGINE && !UE_BUILD_SHIPPING)
+
+#if TRACING_PROFILER
+
+class FTracingProfiler
+{
+private:
+
+	static FTracingProfiler* Instance;
+	FTracingProfiler();
+
+public:
+
+	enum class EEventType
+	{
+		CPU,
+		GPU,
+	};
+
+	struct FEvent
+	{
+		const char* Name;
+		uint32 FrameNumber;
+		uint32 SessionId;
+		EEventType Type;
+		union
+		{
+			struct
+			{
+				uint64 BeginMicroseconds;
+				uint64 EndMicroseconds;
+				uint64 GPUIndex;
+			} GPU;
+
+			struct
+			{
+				uint64 BeginCycles;
+				uint64 EndCycles;
+				uint64 ThreadId;
+			} CPU;
+		};
+	};
+
+	static RENDERCORE_API FTracingProfiler* Get();
+
+	RENDERCORE_API void Init();
+
+	uint32 AddCPUEvent(const char* Name,
+		uint64 TimestampBeginCycles,
+		uint64 TimestempEndCycles,
+		uint32 ThreadId,
+		uint32 FrameNumber)
+	{
+		FEvent Event;
+		Event.Name = Name;
+		Event.Type = EEventType::CPU;
+		Event.FrameNumber = FrameNumber;
+		Event.SessionId = SessionId;
+		Event.CPU.BeginCycles = TimestampBeginCycles;
+		Event.CPU.EndCycles = TimestempEndCycles;
+		Event.CPU.ThreadId = ThreadId;
+		return AddEvent(Event);
+	}
+
+	uint32 AddGPUEvent(const char* Name,
+		uint64 TimestempBeginMicroseconds,
+		uint64 TimestampEndMicroseconds,
+		uint64 GPUIndex,
+		uint32 FrameNumber)
+	{
+		FEvent Event;
+		Event.Name = Name;
+		Event.Type = EEventType::GPU;
+		Event.FrameNumber = FrameNumber;
+		Event.SessionId = SessionId;
+		Event.GPU.BeginMicroseconds = TimestempBeginMicroseconds;
+		Event.GPU.EndMicroseconds = TimestampEndMicroseconds;
+		Event.GPU.GPUIndex = GPUIndex;
+		return AddEvent(Event);
+	}
+
+	uint32 AddEvent(const FEvent& Event)
+	{
+		if (!bCapturing)
+		{
+			return ~0u;
+		}
+
+		uint32 EventId = (FPlatformAtomics::InterlockedIncrement(&EventAtomicConter)-1) % MaxNumCapturedEvents;
+		CapturedEvents[EventId] = Event;
+		return EventId;
+	}
+
+	RENDERCORE_API int32 GetCaptureFrameNumber();
+
+	RENDERCORE_API void BeginCapture(int InNumFramesToCapture = -1);
+	RENDERCORE_API void EndCapture();
+
+	RENDERCORE_API bool IsCapturing() const;
+
+private:
+
+	void WriteCaptureToFile();
+
+	void BeginFrame();
+	void EndFrame();
+
+	void BeginFrameRT();
+	void EndFrameRT();
+
+	TArray<FEvent> CapturedEvents;
+	uint32 MaxNumCapturedEvents = 0;
+
+	int32 NumFramesToCapture = -1;
+	int32 CaptureFrameNumber = 0;
+
+	bool bRequestStartCapture = false;
+	bool bRequestStopCapture = false;
+	bool bCapturing = false;
+	bool bCapturingRT = false;
+
+	uint64 GameThreadFrameBeginCycle = 0;
+	uint64 GameThreadFrameEndCycle = 0;
+
+	uint64 RenderThreadFrameBeginCycle = 0;
+	uint64 RenderThreadFrameEndCycle = 0;
+
+	uint32 SessionId = 0;
+
+	uint32 Pad[(PLATFORM_CACHE_LINE_SIZE / 4 - 2)];
+	volatile int32 EventAtomicConter = 0;
+};
+
+#endif //TRACING_PROFILER
