DR - Remove use of Dynamic Res variables in Server builds

diff --git a/Engine/Source/Runtime/Engine/Private/UnrealEngine.cpp b/Engine/Source/Runtime/Engine/Private/UnrealEngine.cpp
--- a/Engine/Source/Runtime/Engine/Private/UnrealEngine.cpp
+++ b/Engine/Source/Runtime/Engine/Private/UnrealEngine.cpp
@@ -1704,8 +1704,10 @@ void UEngine::PreExit()
 
 	ShutdownHMD();
 
+#if WITH_DYNAMIC_RESOLUTION
 	DynamicResolutionState.Reset();
 	NextDynamicResolutionState.Reset();
+#endif
 }
 
 void UEngine::ShutdownHMD()
