DR  - Fix some issues regarding to query poools and made them more robust.

diff --git a/Engine/Source/Runtime/Engine/Private/DynamicResolution.cpp b/Engine/Source/Runtime/Engine/Private/DynamicResolution.cpp
--- a/Engine/Source/Runtime/Engine/Private/DynamicResolution.cpp
+++ b/Engine/Source/Runtime/Engine/Private/DynamicResolution.cpp
@@ -649,16 +649,22 @@ private:
 
 		InFlightFrameQueries()
 			: QueryPool(RHICreateRenderQueryPool(RQT_AbsoluteTime, 4))
-			, BeginFrameQuery(QueryPool->AllocateQuery())
-			, BeginDynamicResolutionQuery(QueryPool->AllocateQuery())
-			, EndDynamicResolutionQuery(QueryPool->AllocateQuery())
-			, EndFrameQuery(QueryPool->AllocateQuery())
 			, HeuristicHistoryEntry(FDynamicResolutionHeuristicProxy::kInvalidEntryId) 
 			, BeginFrameQuerySubmitted(false)
 			, BeginDynamicResolutionQuerySubmitted(false)
 			, EndDynamicResolutionQuerySubmitted(false)
 			, EndFrameQuerySubmitted(false)
-		{ }
+		{ 
+			InFlightFrameQueries* self = this;
+			ENQUEUE_RENDER_COMMAND(InFlightFrameQueries)(
+			[self](FRHICommandListImmediate& RHICmdList)
+			{
+				self->BeginFrameQuery = self->QueryPool->AllocateQuery();
+				self->BeginDynamicResolutionQuery = self->QueryPool->AllocateQuery();
+				self->EndDynamicResolutionQuery = self->QueryPool->AllocateQuery();
+				self->EndFrameQuery = self->QueryPool->AllocateQuery();
+			});
+		}
 	};
 
 	// List of frame queries in flight.
diff --git a/Engine/Source/Runtime/Engine/Private/SceneUtils.cpp b/Engine/Source/Runtime/Engine/Private/SceneUtils.cpp
--- a/Engine/Source/Runtime/Engine/Private/SceneUtils.cpp
+++ b/Engine/Source/Runtime/Engine/Private/SceneUtils.cpp
@@ -202,7 +202,9 @@ public:
 		if ( MaxGPUQueries == -1 || QueryCount < uint32(MaxGPUQueries) )
 		{
 			StartQuery = RenderQueryPool->AllocateQuery();
+			check(StartQuery.GetQuery() != nullptr);
 			EndQuery = RenderQueryPool->AllocateQuery();
+			check(EndQuery.GetQuery() != nullptr);
 			QueryCount += 2;
 		}
 	}
@@ -222,12 +224,12 @@ public:
 				uint64 Temp;
 				if (bBeginQueryInFlight)
 				{
-					RHICmdListPtr->GetRenderQueryResult(StartQuery.GetQuery(), Temp, false);
+					RHICmdListPtr->GetRenderQueryResult(StartQuery.GetQuery(), Temp, true);
 				}
 
 				if (bEndQueryInFlight)
 				{
-					RHICmdListPtr->GetRenderQueryResult(EndQuery.GetQuery(), Temp, false);
+					RHICmdListPtr->GetRenderQueryResult(EndQuery.GetQuery(), Temp, true);
 				}
 			}
 			StartQuery.ReleaseQuery();
@@ -428,6 +430,7 @@ public:
 
 		// Gather any remaining results and check all the results are ready
 		bool bAllQueriesAllocated = true;
+		bool bAnyEventFailed = false;
 		for (int Index = 0; Index < GpuProfilerEvents.Num(); Index++)
 		{
 			FRealtimeGPUProfilerEvent* Event = GpuProfilerEvents[Index];
@@ -438,8 +441,10 @@ public:
 			}
 			if (!Event->HasValidResult())
 			{
+				UE_LOG(LogSceneUtils, Warning, TEXT("Query '%s' not ready."), *Event->GetName().ToString());
 				// The frame isn't ready yet. Don't update stats - we'll try again next frame. 
-				return false;
+				bAnyEventFailed = true;
+				continue;
 			}
 			if (!Event->HasQueriesAllocated())
 			{
@@ -447,6 +452,11 @@ public:
 			}
 		}
 
+		if (bAnyEventFailed)
+		{
+			return false;
+		}
+
 		if (!bAllQueriesAllocated)
 		{
 			static bool bWarned = false;
@@ -604,7 +614,7 @@ FRealtimeGPUProfiler::FRealtimeGPUProfiler()
 	, bInBeginEndBlock(false)
 {
 	const int MaxGPUQueries = CVarGPUStatsMaxQueriesPerFrame.GetValueOnRenderThread();
-	RenderQueryPool = RHICreateRenderQueryPool(RQT_AbsoluteTime, (MaxGPUQueries > 0) ? MaxGPUQueries * 2 : 10240);
+	RenderQueryPool = RHICreateRenderQueryPool(RQT_AbsoluteTime, (MaxGPUQueries > 0) ? MaxGPUQueries * 2 : UINT32_MAX);
 	for (int Index = 0; Index < NumGPUProfilerBufferedFrames; Index++)
 	{
 		Frames.Add(new FRealtimeGPUProfilerFrame(RenderQueryPool, QueryCount));
diff --git a/Engine/Source/Runtime/RHI/Private/DynamicRHI.cpp b/Engine/Source/Runtime/RHI/Private/DynamicRHI.cpp
--- a/Engine/Source/Runtime/RHI/Private/DynamicRHI.cpp
+++ b/Engine/Source/Runtime/RHI/Private/DynamicRHI.cpp
@@ -370,6 +370,7 @@ FDefaultRHIRenderQueryPool::FDefaultRHIRenderQueryPool(ERenderQueryType InQueryT
 
 FRHIPooledRenderQuery FDefaultRHIRenderQueryPool::AllocateQuery()
 {
+	check(IsInRenderingThread());
 	if (Queries.Num() > 0)
 	{
 		return FRHIPooledRenderQuery(this, Queries.Pop());
@@ -383,6 +384,7 @@ FRHIPooledRenderQuery FDefaultRHIRenderQueryPool::AllocateQuery()
 
 void FDefaultRHIRenderQueryPool::ReleaseQuery(TRefCountPtr<FRHIRenderQuery>&& Query)
 {
+	check(IsInRenderingThread());
 	//Hard to validate because of Resource resurrection, better to remove GetQueryRef entirely
 	//checkf(Query.IsValid() && Query.GetRefCount() <= 2, TEXT("Query has been released but reference still held: use FRHIPooledRenderQuery::GetQueryRef() with extreme caution"));
 	Queries.Push(MoveTemp(Query));
diff --git a/Engine/Source/Runtime/RHI/Public/DynamicRHI.h b/Engine/Source/Runtime/RHI/Public/DynamicRHI.h
--- a/Engine/Source/Runtime/RHI/Public/DynamicRHI.h
+++ b/Engine/Source/Runtime/RHI/Public/DynamicRHI.h
@@ -159,6 +159,7 @@ public:
 
 	~FDefaultRHIRenderQueryPool() override 
 	{
+		check(IsInRenderingThread());
 		checkf(AllocatedQueries == Queries.Num(), TEXT("Querypool deleted before all Queries have been released"));
 	}
 
diff --git a/Engine/Source/Runtime/RHI/Public/RHIResources.h b/Engine/Source/Runtime/RHI/Public/RHIResources.h
--- a/Engine/Source/Runtime/RHI/Public/RHIResources.h
+++ b/Engine/Source/Runtime/RHI/Public/RHIResources.h
@@ -987,6 +987,7 @@ inline FRHIPooledRenderQuery::FRHIPooledRenderQuery(FRHIRenderQueryPool* InQuery
 	: Query(MoveTemp(InQuery))
 	, QueryPool(InQueryPool)
 {
+	check(IsInRenderingThread());
 }
 
 inline void FRHIPooledRenderQuery::ReleaseQuery()
@@ -1001,6 +1002,7 @@ inline void FRHIPooledRenderQuery::ReleaseQuery()
 
 inline FRHIPooledRenderQuery::~FRHIPooledRenderQuery()
 {
+	check(IsInRenderingThread());
 	ReleaseQuery();
 }
 
