Add new FMaterialCompiler method GetTextureForExpression.
Previously this logic was implemented in terms of GetParameterType/GetParameterUniformExpression, but these methods are overriden to do nothing in lightmass material compiler.
By moving this logic to a new method, lightmass can continue to override the previous methods, but still work correctly with the new method.
This fixes materials that use texture objects showing up black during lightmass bakes (which means they didn't reflect any indirect light)

diff --git a/Engine/Source/Runtime/Engine/Private/Materials/HLSLMaterialTranslator.h b/Engine/Source/Runtime/Engine/Private/Materials/HLSLMaterialTranslator.h
--- a/Engine/Source/Runtime/Engine/Private/Materials/HLSLMaterialTranslator.h
+++ b/Engine/Source/Runtime/Engine/Private/Materials/HLSLMaterialTranslator.h
@@ -1866,6 +1866,43 @@ protected:
 		return Chunk.UniformExpression;
 	}
 
+	virtual bool GetTextureForExpression(int32 Index, int32& OutTextureIndex, EMaterialSamplerType& OutSamplerType, TOptional<FName>& OutParameterName) const override
+	{
+		check(Index >= 0 && Index < CurrentScopeChunks->Num());
+		const FShaderCodeChunk& Chunk = (*CurrentScopeChunks)[Index];
+		const EMaterialValueType TexInputType = Chunk.Type;
+		if (!(TexInputType & MCT_Texture))
+		{
+			return false;
+		}
+
+		// If 'InputExpression' is connected, we use need to find the texture object that was passed in
+		// In this case, the texture/sampler assigned on this expression node are not used
+		FMaterialUniformExpression* TextureUniformBase = Chunk.UniformExpression;
+		checkf(TextureUniformBase, TEXT("TexInputType is %d, but missing FMaterialUniformExpression"), TexInputType);
+
+		if (FMaterialUniformExpressionTexture* TextureUniform = TextureUniformBase->GetTextureUniformExpression())
+		{
+			OutSamplerType = TextureUniform->GetSamplerType();
+			OutTextureIndex = TextureUniform->GetTextureIndex();
+			if (FMaterialUniformExpressionTextureParameter* TextureParameterUniform = TextureUniform->GetTextureParameterUniformExpression())
+			{
+				OutParameterName = TextureParameterUniform->GetParameterName();
+			}
+		}
+		else if (FMaterialUniformExpressionExternalTexture* ExternalTextureUniform = TextureUniformBase->GetExternalTextureUniformExpression())
+		{
+			OutTextureIndex = ExternalTextureUniform->GetSourceTextureIndex();
+			OutSamplerType = SAMPLERTYPE_External;
+			if (FMaterialUniformExpressionExternalTextureParameter* ExternalTextureParameterUniform = ExternalTextureUniform->GetExternalTextureParameterUniformExpression())
+			{
+				OutParameterName = ExternalTextureParameterUniform->GetParameterName();
+			}
+		}
+
+		return true;
+	}
+
 	// GetArithmeticResultType
 	EMaterialValueType GetArithmeticResultType(EMaterialValueType TypeA,EMaterialValueType TypeB)
 	{
diff --git a/Engine/Source/Runtime/Engine/Private/Materials/MaterialExpressions.cpp b/Engine/Source/Runtime/Engine/Private/Materials/MaterialExpressions.cpp
--- a/Engine/Source/Runtime/Engine/Private/Materials/MaterialExpressions.cpp
+++ b/Engine/Source/Runtime/Engine/Private/Materials/MaterialExpressions.cpp
@@ -1678,36 +1678,13 @@ int32 UMaterialExpressionTextureSample::Compile(class FMaterialCompiler* Compile
 		TOptional<FName> EffectiveParameterName;
 		if (InputExpression)
 		{
-			const EMaterialValueType TexInputType = Compiler->GetParameterType(TextureCodeIndex);
-			if (!(TexInputType & MCT_Texture))
+			if (!Compiler->GetTextureForExpression(TextureCodeIndex, TextureReferenceIndex, EffectiveSamplerType, EffectiveParameterName))
 			{
 				return CompilerError(Compiler, TEXT("Tex input requires a texture value"));
 			}
-
-			// If 'InputExpression' is connected, we use need to find the texture object that was passed in
-			// In this case, the texture/sampler assigned on this expression node are not used
-			FMaterialUniformExpression* TextureUniformBase = Compiler->GetParameterUniformExpression(TextureCodeIndex);
-			checkf(TextureUniformBase, TEXT("TexInputType is %d, but missing FMaterialUniformExpression"), TexInputType);
-
-			if (FMaterialUniformExpressionTexture* TextureUniform = TextureUniformBase->GetTextureUniformExpression())
+			if (TextureReferenceIndex != INDEX_NONE)
 			{
-				EffectiveSamplerType = TextureUniform->GetSamplerType();
-				TextureReferenceIndex = TextureUniform->GetTextureIndex();
 				EffectiveTexture = Compiler->GetReferencedTexture(TextureReferenceIndex);
-				if (FMaterialUniformExpressionTextureParameter* TextureParameterUniform = TextureUniform->GetTextureParameterUniformExpression())
-				{
-					EffectiveParameterName = TextureParameterUniform->GetParameterName();
-				}
-			}
-			else if (FMaterialUniformExpressionExternalTexture* ExternalTextureUniform = TextureUniformBase->GetExternalTextureUniformExpression())
-			{
-				TextureReferenceIndex = ExternalTextureUniform->GetSourceTextureIndex();
-				EffectiveTexture = Compiler->GetReferencedTexture(TextureReferenceIndex);
-				EffectiveSamplerType = SAMPLERTYPE_External;
-				if (FMaterialUniformExpressionExternalTextureParameter* ExternalTextureParameterUniform = ExternalTextureUniform->GetExternalTextureParameterUniformExpression())
-				{
-					EffectiveParameterName = ExternalTextureParameterUniform->GetParameterName();
-				}
 			}
 		}
 
diff --git a/Engine/Source/Runtime/Engine/Public/MaterialCompiler.h b/Engine/Source/Runtime/Engine/Public/MaterialCompiler.h
--- a/Engine/Source/Runtime/Engine/Public/MaterialCompiler.h
+++ b/Engine/Source/Runtime/Engine/Public/MaterialCompiler.h
@@ -77,6 +77,8 @@ public:
 
 	virtual FMaterialUniformExpression* GetParameterUniformExpression(int32 Index) const = 0;
 
+	virtual bool GetTextureForExpression(int32 Index, int32& OutTextureIndex, EMaterialSamplerType& OutSamplerType, TOptional<FName>& OutParameterName) const = 0;
+
 	/** 
 	 * Casts the passed in code to DestType, or generates a compile error if the cast is not valid. 
 	 * This will truncate a type (float4 -> float3) but not add components (float2 -> float3), however a float1 can be cast to any float type by replication. 
@@ -318,6 +320,7 @@ public:
 	virtual EMaterialShadingModel GetMaterialShadingModel() const { return Compiler->GetMaterialShadingModel();  }
 	virtual EMaterialValueType GetParameterType(int32 Index) const { return Compiler->GetParameterType(Index); }
 	virtual FMaterialUniformExpression* GetParameterUniformExpression(int32 Index) const { return Compiler->GetParameterUniformExpression(Index); }
+	virtual bool GetTextureForExpression(int32 Index, int32& OutTextureIndex, EMaterialSamplerType& OutSamplerType, TOptional<FName>& OutParameterName) const override { return Compiler->GetTextureForExpression(Index, OutTextureIndex, OutSamplerType, OutParameterName); }
 	virtual void SetMaterialProperty(EMaterialProperty InProperty, EShaderFrequency OverrideShaderFrequency, bool bUsePreviousFrameTime) override { Compiler->SetMaterialProperty(InProperty, OverrideShaderFrequency, bUsePreviousFrameTime); }
 	virtual void PushMaterialAttribute(const FGuid& InAttributeID) override { Compiler->PushMaterialAttribute(InAttributeID); }
 	virtual FGuid PopMaterialAttribute() override { return Compiler->PopMaterialAttribute(); }
