Fixed compilation error in Linux due a bug introduced during merging

diff --git a/Engine/Source/Runtime/Engine/Private/Materials/HLSLMaterialTranslator.h b/Engine/Source/Runtime/Engine/Private/Materials/HLSLMaterialTranslator.h
--- a/Engine/Source/Runtime/Engine/Private/Materials/HLSLMaterialTranslator.h
+++ b/Engine/Source/Runtime/Engine/Private/Materials/HLSLMaterialTranslator.h
@@ -814,7 +814,7 @@ ResourcesString = TEXT("");
 				}
 			}
 
-			if (BlendMode == BLEND_AlphaHoldout && BlendMode != MSM_Unlit)
+			if (BlendMode == BLEND_AlphaHoldout && !MaterialShadingModels.IsUnlit())
 			{
 				Errorf(TEXT("Alpha Holdout blend mode must use unlit shading model."));
 			}
