Always pass in the platform for GMaxRHIShaderPlatform

diff --git a/Engine/Source/Editor/Persona/Private/PersonaMeshDetails.cpp b/Engine/Source/Editor/Persona/Private/PersonaMeshDetails.cpp
--- a/Engine/Source/Editor/Persona/Private/PersonaMeshDetails.cpp
+++ b/Engine/Source/Editor/Persona/Private/PersonaMeshDetails.cpp
@@ -2488,7 +2488,7 @@ TSharedRef<SWidget> FPersonaMeshDetails::OnGenerateCustomNameWidgetsForSection(i
 
 TSharedRef<SWidget> FPersonaMeshDetails::OnGenerateCustomSectionWidgetsForSection(int32 LODIndex, int32 SectionIndex)
 {
-	extern ENGINE_API bool IsGPUSkinCacheAvailable();
+	extern ENGINE_API bool IsGPUSkinCacheAvailable(EShaderPlatform Platform);
 
 	TSharedRef<SVerticalBox> SectionWidget = SNew(SVerticalBox);
 
@@ -2558,7 +2558,7 @@ TSharedRef<SWidget> FPersonaMeshDetails::OnGenerateCustomSectionWidgetsForSectio
 		.Padding(2, 0, 2, 0)
 		[
 			SNew(SCheckBox)
-			.IsEnabled(IsGPUSkinCacheAvailable())
+			.IsEnabled(IsGPUSkinCacheAvailable(GMaxRHIShaderPlatform))
 			.IsChecked(this, &FPersonaMeshDetails::IsSectionRecomputeTangentEnabled, LODIndex, SectionIndex)
 			.OnCheckStateChanged(this, &FPersonaMeshDetails::OnSectionRecomputeTangentChanged, LODIndex, SectionIndex)
 			[
diff --git a/Engine/Source/Editor/UnrealEd/Private/PreviewMaterial.cpp b/Engine/Source/Editor/UnrealEd/Private/PreviewMaterial.cpp
--- a/Engine/Source/Editor/UnrealEd/Private/PreviewMaterial.cpp
+++ b/Engine/Source/Editor/UnrealEd/Private/PreviewMaterial.cpp
@@ -102,8 +102,8 @@ public:
 					return false;
 				}
 
-				extern ENGINE_API bool IsGPUSkinCacheAvailable();
-				bool bSkinCache = IsGPUSkinCacheAvailable() && (VertexFactoryType == FindVertexFactoryType(FName(TEXT("FGPUSkinPassthroughVertexFactory"), FNAME_Find)));
+				extern ENGINE_API bool IsGPUSkinCacheAvailable(EShaderPlatform Platform);
+				bool bSkinCache = IsGPUSkinCacheAvailable(Platform) && (VertexFactoryType == FindVertexFactoryType(FName(TEXT("FGPUSkinPassthroughVertexFactory"), FNAME_Find)));
 					
 				if (
 					VertexFactoryType != FindVertexFactoryType(FName(TEXT("TGPUSkinVertexFactoryfalse"), FNAME_Find)) &&
diff --git a/Engine/Source/Runtime/Engine/Private/GPUSkinCache.cpp b/Engine/Source/Runtime/Engine/Private/GPUSkinCache.cpp
--- a/Engine/Source/Runtime/Engine/Private/GPUSkinCache.cpp
+++ b/Engine/Source/Runtime/Engine/Private/GPUSkinCache.cpp
@@ -114,9 +114,9 @@ static inline bool DoesPlatformSupportGPUSkinCache(EShaderPlatform Platform)
 	return Platform == SP_PCD3D_SM5 || IsMetalSM5Platform(Platform) || IsVulkanSM5Platform(Platform) || Platform == SP_OPENGL_SM5;
 }
 
-ENGINE_API bool IsGPUSkinCacheAvailable()
+ENGINE_API bool IsGPUSkinCacheAvailable(EShaderPlatform Platform)
 {
-	return (GEnableGPUSkinCacheShaders != 0 || GForceRecomputeTangents != 0) && DoesPlatformSupportGPUSkinCache(GMaxRHIShaderPlatform);
+	return (GEnableGPUSkinCacheShaders != 0 || GForceRecomputeTangents != 0) && DoesPlatformSupportGPUSkinCache(Platform);
 }
 
 // We don't have it always enabled as it's not clear if this has a performance cost
@@ -126,7 +126,7 @@ ENGINE_API bool DoSkeletalMeshIndexBuffersNeedSRV()
 {
 	// currently only implemented and tested on Window SM5 (needs Compute, Atomics, SRV for index buffers, UAV for VertexBuffers)
 	//#todo-gpuskin: Enable on PS4 when SRVs for IB exist
-	return IsGPUSkinCacheAvailable();
+	return IsGPUSkinCacheAvailable(GMaxRHIShaderPlatform);
 }
 
 ENGINE_API bool DoRecomputeSkinTangentsOnGPU_RT()
@@ -535,7 +535,7 @@ public:
 
 	static bool ShouldCompilePermutation(const FGlobalShaderPermutationParameters& Parameters)
 	{
-		return IsGPUSkinCacheAvailable() && IsFeatureLevelSupported(Parameters.Platform, ERHIFeatureLevel::SM5);
+		return IsGPUSkinCacheAvailable(Parameters.Platform) && IsFeatureLevelSupported(Parameters.Platform, ERHIFeatureLevel::SM5);
 	}
 
 	static void ModifyCompilationEnvironment(const FGlobalShaderPermutationParameters& Parameters, FShaderCompilerEnvironment& OutEnvironment)
@@ -647,7 +647,7 @@ public:
 	static bool ShouldCompilePermutation(const FGlobalShaderPermutationParameters& Parameters)
 	{
 		// currently only implemented and tested on Window SM5 (needs Compute, Atomics, SRV for index buffers, UAV for VertexBuffers)
-		return DoesPlatformSupportGPUSkinCache(Parameters.Platform) && IsGPUSkinCacheAvailable();
+		return DoesPlatformSupportGPUSkinCache(Parameters.Platform) && IsGPUSkinCacheAvailable(Parameters.Platform);
 	}
 
 	static const uint32 ThreadGroupSizeX = 64;
@@ -788,7 +788,7 @@ class FRecomputeTangentsPerVertexPassCS : public FGlobalShader
 	static bool ShouldCompilePermutation(const FGlobalShaderPermutationParameters& Parameters)
 	{
 		// currently only implemented and tested on Window SM5 (needs Compute, Atomics, SRV for index buffers, UAV for VertexBuffers)
-		return DoesPlatformSupportGPUSkinCache(Parameters.Platform) && IsGPUSkinCacheAvailable();
+		return DoesPlatformSupportGPUSkinCache(Parameters.Platform) && IsGPUSkinCacheAvailable(Parameters.Platform);
 	}
 
 	static void ModifyCompilationEnvironment(const FGlobalShaderPermutationParameters& Parameters, FShaderCompilerEnvironment& OutEnvironment)
diff --git a/Engine/Source/Runtime/Engine/Private/GPUSkinVertexFactory.cpp b/Engine/Source/Runtime/Engine/Private/GPUSkinVertexFactory.cpp
--- a/Engine/Source/Runtime/Engine/Private/GPUSkinVertexFactory.cpp
+++ b/Engine/Source/Runtime/Engine/Private/GPUSkinVertexFactory.cpp
@@ -655,7 +655,7 @@ void FGPUSkinPassthroughVertexFactory::ModifyCompilationEnvironment( const FVert
 bool FGPUSkinPassthroughVertexFactory::ShouldCompilePermutation(EShaderPlatform Platform, const class FMaterial* Material, const FShaderType* ShaderType)
 {
 	// Passthrough is only valid on platforms with Compute Shader support AND for (skeletal meshes or default materials)
-	return IsGPUSkinCacheAvailable() && IsFeatureLevelSupported(Platform, ERHIFeatureLevel::SM5) && Super::ShouldCompilePermutation(Platform, Material, ShaderType) && (Material->IsUsedWithSkeletalMesh() || Material->IsSpecialEngineMaterial());
+	return IsGPUSkinCacheAvailable(Platform) && IsFeatureLevelSupported(Platform, ERHIFeatureLevel::SM5) && Super::ShouldCompilePermutation(Platform, Material, ShaderType) && (Material->IsUsedWithSkeletalMesh() || Material->IsSpecialEngineMaterial());
 }
 
 void FGPUSkinPassthroughVertexFactory::InternalUpdateVertexDeclaration(FGPUBaseSkinVertexFactory* SourceVertexFactory, struct FRWBuffer* PositionRWBuffer, struct FRWBuffer* TangentRWBuffer)
diff --git a/Engine/Source/Runtime/Engine/Private/PositionVertexBuffer.cpp b/Engine/Source/Runtime/Engine/Private/PositionVertexBuffer.cpp
--- a/Engine/Source/Runtime/Engine/Private/PositionVertexBuffer.cpp
+++ b/Engine/Source/Runtime/Engine/Private/PositionVertexBuffer.cpp
@@ -223,7 +223,7 @@ void FPositionVertexBuffer::InitRHI()
 {
 	VertexBufferRHI = CreateRHIBuffer_RenderThread();
 	// we have decide to create the SRV based on GMaxRHIShaderPlatform because this is created once and shared between feature levels for editor preview.
-	if (VertexBufferRHI && (RHISupportsManualVertexFetch(GMaxRHIShaderPlatform) || IsGPUSkinCacheAvailable()))
+	if (VertexBufferRHI && (RHISupportsManualVertexFetch(GMaxRHIShaderPlatform) || IsGPUSkinCacheAvailable(GMaxRHIShaderPlatform)))
 	{
 		PositionComponentSRV = RHICreateShaderResourceView(VertexData ? VertexBufferRHI : nullptr, 4, PF_R32_FLOAT);
 	}
@@ -254,4 +254,4 @@ void FPositionVertexBuffer::BindPositionVertexBuffer(const FVertexFactory* Verte
 		VET_Float3
 	);
 	StaticMeshData.PositionComponentSRV = PositionComponentSRV;
-}
\ No newline at end of file
+}
diff --git a/Engine/Source/Runtime/Engine/Private/Rendering/StaticMeshVertexBuffer.cpp b/Engine/Source/Runtime/Engine/Private/Rendering/StaticMeshVertexBuffer.cpp
--- a/Engine/Source/Runtime/Engine/Private/Rendering/StaticMeshVertexBuffer.cpp
+++ b/Engine/Source/Runtime/Engine/Private/Rendering/StaticMeshVertexBuffer.cpp
@@ -331,7 +331,7 @@ void FStaticMeshVertexBuffer::InitRHI()
 {
 	TangentsVertexBuffer.VertexBufferRHI = CreateTangentsRHIBuffer_RenderThread();
 	TexCoordVertexBuffer.VertexBufferRHI = CreateTexCoordRHIBuffer_RenderThread();
-	if (TangentsVertexBuffer.VertexBufferRHI && (RHISupportsManualVertexFetch(GMaxRHIShaderPlatform) || IsGPUSkinCacheAvailable()))
+	if (TangentsVertexBuffer.VertexBufferRHI && (RHISupportsManualVertexFetch(GMaxRHIShaderPlatform) || IsGPUSkinCacheAvailable(GMaxRHIShaderPlatform)))
 	{
 		TangentsSRV = RHICreateShaderResourceView(
 			TangentsData ? TangentsVertexBuffer.VertexBufferRHI : nullptr,
diff --git a/Engine/Source/Runtime/Engine/Private/SkeletalRenderGPUSkin.cpp b/Engine/Source/Runtime/Engine/Private/SkeletalRenderGPUSkin.cpp
--- a/Engine/Source/Runtime/Engine/Private/SkeletalRenderGPUSkin.cpp
+++ b/Engine/Source/Runtime/Engine/Private/SkeletalRenderGPUSkin.cpp
@@ -82,7 +82,7 @@ void FMorphVertexBuffer::InitDynamicRHI()
 	Flags = (EBufferUsageFlags)(Flags | BUF_ShaderResource);
 
 	VertexBufferRHI = RHICreateVertexBuffer(Size, Flags, CreateInfo);
-	bool bUsesSkinCache = bSupportsComputeShaders && IsGPUSkinCacheAvailable() && GEnableGPUSkinCache;
+	bool bUsesSkinCache = bSupportsComputeShaders && IsGPUSkinCacheAvailable(GMaxRHIShaderPlatform) && GEnableGPUSkinCache;
 	if (bUsesSkinCache)
 	{
 		SRVValue = RHICreateShaderResourceView(VertexBufferRHI, 4, PF_R32_FLOAT);
diff --git a/Engine/Source/Runtime/Engine/Public/GPUSkinCache.h b/Engine/Source/Runtime/Engine/Public/GPUSkinCache.h
--- a/Engine/Source/Runtime/Engine/Public/GPUSkinCache.h
+++ b/Engine/Source/Runtime/Engine/Public/GPUSkinCache.h
@@ -51,7 +51,7 @@ struct FSkelMeshRenderSection;
 struct FVertexBufferAndSRV;
 
 // Can the skin cache be used (ie shaders added, etc)
-extern ENGINE_API bool IsGPUSkinCacheAvailable();
+extern ENGINE_API bool IsGPUSkinCacheAvailable(EShaderPlatform Platform);
 
 // Is it actually enabled?
 extern ENGINE_API int32 GEnableGPUSkinCache;
diff --git a/Engine/Source/Runtime/Renderer/Private/RendererScene.cpp b/Engine/Source/Runtime/Renderer/Private/RendererScene.cpp
--- a/Engine/Source/Runtime/Renderer/Private/RendererScene.cpp
+++ b/Engine/Source/Runtime/Renderer/Private/RendererScene.cpp
@@ -1079,7 +1079,7 @@ FScene::FScene(UWorld* InWorld, bool bInRequiresHitProxies, bool bInIsEditorScen
 		SetFXSystem(NULL);
 	}
 
-	if (IsGPUSkinCacheAvailable())
+	if (IsGPUSkinCacheAvailable(GetFeatureLevelShaderPlatform(InFeatureLevel)))
 	{
 		const bool bRequiresMemoryLimit = !bInIsEditorScene;
 		GPUSkinCache = new FGPUSkinCache(bRequiresMemoryLimit);
diff --git a/Engine/Source/Runtime/Renderer/Private/SceneRendering.cpp b/Engine/Source/Runtime/Renderer/Private/SceneRendering.cpp
--- a/Engine/Source/Runtime/Renderer/Private/SceneRendering.cpp
+++ b/Engine/Source/Runtime/Renderer/Private/SceneRendering.cpp
@@ -2440,8 +2440,8 @@ void FSceneRenderer::RenderFinish(FRHICommandListImmediate& RHICmdList)
 		static auto* CVarSkinCacheOOM = IConsoleManager::Get().FindTConsoleVariableDataFloat(TEXT("r.SkinCache.SceneMemoryLimitInMB"));
 
 		uint64 GPUSkinCacheExtraRequiredMemory = 0;
-		extern ENGINE_API bool IsGPUSkinCacheAvailable();
-		if (IsGPUSkinCacheAvailable())
+		extern ENGINE_API bool IsGPUSkinCacheAvailable(EShaderPlatform Platform);
+		if (IsGPUSkinCacheAvailable(ShaderPlatform))
 		{
 			if (FGPUSkinCache* SkinCache = Scene->GetGPUSkinCache())
 			{
@@ -3816,4 +3816,4 @@ FTextureRHIParamRef FSceneRenderer::GetMultiViewSceneColor(const FSceneRenderTar
 	{
 		return static_cast<FTextureRHIRef>(ViewFamily.RenderTarget->GetRenderTargetTexture());
 	}
-}
\ No newline at end of file
+}
