Bump VT DDC key

diff --git a/Engine/Source/Runtime/Engine/Private/TextureDerivedData.cpp b/Engine/Source/Runtime/Engine/Private/TextureDerivedData.cpp
--- a/Engine/Source/Runtime/Engine/Private/TextureDerivedData.cpp
+++ b/Engine/Source/Runtime/Engine/Private/TextureDerivedData.cpp
@@ -73,7 +73,7 @@ static TAutoConsoleVariable<int32> CVarVTTileBorderSize(
 
 // This GUID is mixed into DDC version for virtual textures only, this allows updating DDC version for VT without invalidating DDC for all textures
 // This is useful during development, but once large numbers of VT are present in shipped content, it will have the same problem as TEXTURE_DERIVEDDATA_VER
-#define TEXTURE_VT_DERIVEDDATA_VER	TEXT("05576F4E7FCC420BBB91F91A43BF3E30")
+#define TEXTURE_VT_DERIVEDDATA_VER	TEXT("186764627EE1418BBAD6ADBA2C2CAE63")
 
 #if ENABLE_COOK_STATS
 namespace TextureCookStats
