Update serialization of lightmaps to work with new VT lightmap types
(This change has been in Dev-VT for a while, but avoided changing serialize format until moving to Dev-Rendering)

diff --git a/Engine/Source/Runtime/Core/Public/UObject/RenderingObjectVersion.h b/Engine/Source/Runtime/Core/Public/UObject/RenderingObjectVersion.h
--- a/Engine/Source/Runtime/Core/Public/UObject/RenderingObjectVersion.h
+++ b/Engine/Source/Runtime/Core/Public/UObject/RenderingObjectVersion.h
@@ -84,6 +84,9 @@ struct CORE_API FRenderingObjectVersion
 		// Removed old gaussian and bokeh DOF methods from deferred shading renderer.
 		DiaphragmDOFOnlyForDeferredShadingRenderer,
 
+		// Lightmaps replace ULightMapVirtualTexture (non-UTexture derived class) with ULightMapVirtualTexture2D (derived from UTexture)
+		VirtualTexturedLightmapsV2,
+
 		// -----<new versions can be added above this line>-------------------------------------------------
 		VersionPlusOne,
 		LatestVersion = VersionPlusOne - 1
diff --git a/Engine/Source/Runtime/Engine/Private/LightMap.cpp b/Engine/Source/Runtime/Engine/Private/LightMap.cpp
--- a/Engine/Source/Runtime/Engine/Private/LightMap.cpp
+++ b/Engine/Source/Runtime/Engine/Private/LightMap.cpp
@@ -2788,6 +2788,7 @@ void FLightMap2D::AddReferencedObjects( FReferenceCollector& Collector )
 void FLightMap2D::Serialize(FArchive& Ar)
 {
 	Ar.UsingCustomVersion(FRenderingObjectVersion::GUID);
+	const int32 RenderCustomVersion = Ar.CustomVer(FRenderingObjectVersion::GUID);
 
 	FLightMap::Serialize(Ar);
 
@@ -2871,7 +2872,7 @@ void FLightMap2D::Serialize(FArchive& Ar)
 
 	Ar << CoordinateScale << CoordinateBias;
 
-	if (Ar.CustomVer(FRenderingObjectVersion::GUID) >= FRenderingObjectVersion::LightmapHasShadowmapData)
+	if (RenderCustomVersion >= FRenderingObjectVersion::LightmapHasShadowmapData)
 	{
 		for (int Channel = 0; Channel < ARRAY_COUNT(bShadowChannelValid); Channel++)
 		{
@@ -2881,28 +2882,28 @@ void FLightMap2D::Serialize(FArchive& Ar)
 		Ar << InvUniformPenumbraSize;
 	}
 
-	if (Ar.CustomVer(FRenderingObjectVersion::GUID) >= FRenderingObjectVersion::VirtualTexturedLightmaps)
+	if (RenderCustomVersion >= FRenderingObjectVersion::VirtualTexturedLightmaps)
 	{
-		// Don't save VT's if they are disabled for rendering
-		if ( bUsingVTLightmaps )
+		if (RenderCustomVersion >= FRenderingObjectVersion::VirtualTexturedLightmapsV2)
 		{
-			// TEMPORARY - Once this is merged into Dev-Rendering, this will be removed, and an updated FRenderingObjectVersion will cause new VT lightmaps to be serialized
-			// This is just in place so we can see VT lightmaps in world settings, without messing up normal UObject serialization
-			if (Ar.ArIsObjectReferenceCollector)
+			// Don't save VT's if they are disabled for rendering
+			if (bUsingVTLightmaps)
 			{
 				Ar << VirtualTexture;
 			}
-
-			// TODO - add serialize code for new ULightMapVirtualTexture2D
-			ULightMapVirtualTexture* Dummy = NULL;
-			Ar << Dummy;
-			if (Ar.IsLoading())
+			else
 			{
-				VirtualTexture = nullptr;
+				ULightMapVirtualTexture2D* Dummy = NULL;
+				Ar << Dummy;
+				if (Ar.IsLoading())
+				{
+					VirtualTexture = nullptr;
+				}
 			}
 		}
 		else
 		{
+			// Older version using older lightmap texture type
 			ULightMapVirtualTexture* Dummy = NULL;
 			Ar << Dummy;
 			if (Ar.IsLoading())
