DR - remove some back to back queued rendercommands

diff --git a/Engine/Source/Runtime/Engine/Private/SkeletalRenderGPUSkin.cpp b/Engine/Source/Runtime/Engine/Private/SkeletalRenderGPUSkin.cpp
--- a/Engine/Source/Runtime/Engine/Private/SkeletalRenderGPUSkin.cpp
+++ b/Engine/Source/Runtime/Engine/Private/SkeletalRenderGPUSkin.cpp
@@ -1185,12 +1185,10 @@ static VertexFactoryType* CreateVertexFactory(TArray<TUniquePtr<VertexFactoryTyp
 			typename VertexFactoryType::FDataType Data;
 			InitGPUSkinVertexFactoryComponents<VertexFactoryType>(&Data, VertexUpdateData.VertexBuffers, VertexUpdateData.VertexFactory);
 			VertexUpdateData.VertexFactory->SetData(Data);
+			VertexUpdateData.VertexFactory->InitResource();
 		}
 	);
 
-	// init rendering resource	
-	BeginInitResource(VertexFactory);
-
 	return VertexFactory;
 }
 
@@ -1216,11 +1214,9 @@ static void CreatePassthroughVertexFactory(ERHIFeatureLevel::Type InFeatureLevel
 		[NewPassthroughVertexFactory, SourceVertexFactory](FRHICommandList& RHICmdList)
 		{
 			SourceVertexFactory->CopyDataTypeForPassthroughFactory(NewPassthroughVertexFactory);
+			NewPassthroughVertexFactory->InitResource();
 		}
 	);
-
-	// init rendering resource	
-	BeginInitResource(NewPassthroughVertexFactory);
 }
 
 /**
@@ -1247,12 +1243,10 @@ static VertexFactoryType* CreateVertexFactoryMorph(TArray<TUniquePtr<VertexFacto
 			InitGPUSkinVertexFactoryComponents<VertexFactoryType>(&Data, VertexUpdateData.VertexBuffers, VertexUpdateData.VertexFactory);
 			InitMorphVertexFactoryComponents<VertexFactoryType>(&Data, VertexUpdateData.VertexBuffers);
 			VertexUpdateData.VertexFactory->SetData(Data);
+			VertexUpdateData.VertexFactory->InitResource();
 		}
 	);
 
-	// init rendering resource	
-	BeginInitResource(VertexFactory);
-
 	return VertexFactory;
 }
 
@@ -1282,11 +1276,9 @@ static void CreateVertexFactoryCloth(TArray<TUniquePtr<VertexFactoryTypeBase>>&
 			InitGPUSkinVertexFactoryComponents<VertexFactoryType>(&Data, VertexUpdateData.VertexBuffers, VertexUpdateData.VertexFactory);
 			InitAPEXClothVertexFactoryComponents<VertexFactoryType>(&Data, VertexUpdateData.VertexBuffers);
 			VertexUpdateData.VertexFactory->SetData(Data);
+			VertexUpdateData.VertexFactory->InitResource();
 		}
 	);
-
-	// init rendering resource	
-	BeginInitResource(VertexFactory);
 }
 
 /**
