Disable lossy (Crunch) compression on VT lightmaps by default, add CVar to enable it

diff --git a/Engine/Source/Runtime/Engine/Private/LightMap.cpp b/Engine/Source/Runtime/Engine/Private/LightMap.cpp
--- a/Engine/Source/Runtime/Engine/Private/LightMap.cpp
+++ b/Engine/Source/Runtime/Engine/Private/LightMap.cpp
@@ -83,6 +83,19 @@ static TAutoConsoleVariable<int32> CVarTexelDebugging(
 	TEXT("Whether T + Left mouse click in the editor selects lightmap texels for debugging Lightmass.  Lightmass must be recompiled with ALLOW_LIGHTMAP_SAMPLE_DEBUGGING enabled for this to work."),
 	ECVF_Default);
 
+static TAutoConsoleVariable<int32> CVarVirtualTexturedLightMaps(
+	TEXT("r.VirtualTexturedLightmaps"),
+	0,
+	TEXT("Controls wether to stream the lightmaps using virtual texturing.\n") \
+	TEXT(" 0: Disabled.\n") \
+	TEXT(" 1: Enabled."),
+	ECVF_ReadOnly);
+
+static TAutoConsoleVariable<int32> CVarVTEnableLossyCompressLightmaps(
+	TEXT("r.VT.EnableLossyCompressLightmaps"),
+	0,
+	TEXT("Enables lossy compression on virtual texture lightmaps. Lossy compression tends to have lower quality on lightmap textures, vs regular color textures."));
+
 bool IsTexelDebuggingEnabled()
 {
 	return CVarTexelDebugging.GetValueOnGameThread() != 0;
@@ -1171,8 +1184,7 @@ void FLightMapPendingTexture::CreateUObjects()
 	}
 
 	// Only build VT lightmaps if they are enabled
-	static const auto CVar = IConsoleManager::Get().FindTConsoleVariableDataInt(TEXT("r.VirtualTexturedLightmaps"));
-	const bool bUseVirtualTextures = (CVar->GetValueOnAnyThread() != 0) && UseVirtualTexturing(GMaxRHIFeatureLevel);
+	const bool bUseVirtualTextures = (CVarVirtualTexturedLightMaps.GetValueOnAnyThread() != 0) && UseVirtualTexturing(GMaxRHIFeatureLevel);
 	if (bUseVirtualTextures)
 	{
 		VirtualTexture = NewObject<ULightMapVirtualTexture2D>(Outer, GetVirtualTextureName(GLightmapCounter));
@@ -1699,6 +1711,7 @@ void FLightMapPendingTexture::StartEncoding(ULevel* LightingScenario, ITextureCo
 		VirtualTexture->LODGroup = TEXTUREGROUP_Lightmap;
 		VirtualTexture->CompressionNoAlpha = false;
 		VirtualTexture->CompressionNone = !GCompressLightmaps;
+		VirtualTexture->LossyCompressionAmount = CVarVTEnableLossyCompressLightmaps.GetValueOnAnyThread() ? TLCA_Default : TLCA_None;
 
 		FTextureFormatSettings DefaultFormatSettings;
 		VirtualTexture->GetDefaultFormatSettings(DefaultFormatSettings);
@@ -2792,8 +2805,7 @@ void FLightMap2D::Serialize(FArchive& Ar)
 
 	FLightMap::Serialize(Ar);
 
-	static const auto CVar = IConsoleManager::Get().FindTConsoleVariableDataInt(TEXT("r.VirtualTexturedLightmaps"));
-	const bool bUsingVTLightmaps = (CVar->GetValueOnAnyThread() != 0) && UseVirtualTexturing(GMaxRHIFeatureLevel, Ar.CookingTarget());
+	const bool bUsingVTLightmaps = (CVarVirtualTexturedLightMaps.GetValueOnAnyThread() != 0) && UseVirtualTexturing(GMaxRHIFeatureLevel, Ar.CookingTarget());
 
 	if( Ar.IsLoading() && Ar.UE4Ver() < VER_UE4_LOW_QUALITY_DIRECTIONAL_LIGHTMAPS )
 	{
@@ -2960,8 +2972,7 @@ FLightMapInteraction FLightMap2D::GetInteraction(ERHIFeatureLevel::Type InFeatur
 
 	int32 LightmapIndex = bHighQuality ? 0 : 1;
 
-	static const auto CVar = IConsoleManager::Get().FindTConsoleVariableDataInt(TEXT("r.VirtualTexturedLightmaps"));
-	const bool bUseVirtualTextures = (CVar->GetValueOnAnyThread() != 0) && UseVirtualTexturing(InFeatureLevel);
+	const bool bUseVirtualTextures = (CVarVirtualTexturedLightMaps.GetValueOnAnyThread() != 0) && UseVirtualTexturing(InFeatureLevel);
 	if (!bUseVirtualTextures)
 	{
 		bool bValidTextures = Textures[LightmapIndex] && Textures[LightmapIndex]->Resource;
@@ -2986,8 +2997,7 @@ FLightMapInteraction FLightMap2D::GetInteraction(ERHIFeatureLevel::Type InFeatur
 
 FShadowMapInteraction FLightMap2D::GetShadowInteraction(ERHIFeatureLevel::Type InFeatureLevel) const
 {
-	static const auto CVar = IConsoleManager::Get().FindTConsoleVariableDataInt(TEXT("r.VirtualTexturedLightmaps"));
-	const bool bUseVirtualTextures = (CVar->GetValueOnAnyThread() != 0) && UseVirtualTexturing(InFeatureLevel);
+	const bool bUseVirtualTextures = (CVarVirtualTexturedLightMaps.GetValueOnAnyThread() != 0) && UseVirtualTexturing(InFeatureLevel);
 	if (bUseVirtualTextures)
 	{
 		const bool bValidVirtualTexture = VirtualTexture && VirtualTexture->Resource;
@@ -3192,9 +3202,8 @@ void FLightmapResourceCluster::UpdateUniformBuffer(ERHIFeatureLevel::Type InFeat
 	ENQUEUE_RENDER_COMMAND(SetFeatureLevel)(
 		[Cluster, InFeatureLevel](FRHICommandList& RHICmdList)
 	{
-		static const auto CVar = IConsoleManager::Get().FindTConsoleVariableDataInt(TEXT("r.VirtualTexturedLightmaps"));
 		const bool bAllowHighQualityLightMaps = AllowHighQualityLightmaps(InFeatureLevel);
-		const bool bUseVirtualTextures = bAllowHighQualityLightMaps && (CVar->GetValueOnRenderThread() != 0) && UseVirtualTexturing(InFeatureLevel);
+		const bool bUseVirtualTextures = bAllowHighQualityLightMaps && (CVarVirtualTexturedLightMaps.GetValueOnRenderThread() != 0) && UseVirtualTexturing(InFeatureLevel);
 	
 		Cluster->FeatureLevel = InFeatureLevel;
 
diff --git a/Engine/Source/Runtime/Engine/Private/VT/VirtualTexture.cpp b/Engine/Source/Runtime/Engine/Private/VT/VirtualTexture.cpp
--- a/Engine/Source/Runtime/Engine/Private/VT/VirtualTexture.cpp
+++ b/Engine/Source/Runtime/Engine/Private/VT/VirtualTexture.cpp
@@ -11,15 +11,6 @@
 
 DEFINE_LOG_CATEGORY(LogVirtualTexturingModule);
 
-// Is there a cleaner place to put this??
-static TAutoConsoleVariable<int32> CVarVirtualTexturedLightMaps(
-	TEXT("r.VirtualTexturedLightmaps"),
-	0,
-	TEXT("Controls wether to stream the lightmaps using virtual texturing.\n") \
-	TEXT(" 0: Disabled.\n") \
-	TEXT(" 1: Enabled."),
-	ECVF_ReadOnly);
-
 static FAutoConsoleCommand GVTFlushAndEvictFileCacheCommand(
 	TEXT("r.VT.FlushAndEvictFileCache"),
 	TEXT("Flush both the virtual texture physcial page cache and disk file cachje"),
