Remove a bit of old code left over from when FVirtualTexture2DResource could be created by non-UTexture derived owner

diff --git a/Engine/Source/Runtime/Engine/Private/Texture2D.cpp b/Engine/Source/Runtime/Engine/Private/Texture2D.cpp
--- a/Engine/Source/Runtime/Engine/Private/Texture2D.cpp
+++ b/Engine/Source/Runtime/Engine/Private/Texture2D.cpp
@@ -1072,7 +1072,7 @@ FTextureResource* UTexture2D::CreateResource()
 	{
 		UnlinkStreaming();
 		bIsStreamable = false;
-		FVirtualTexture2DResource* ResourceVT = new FVirtualTexture2DResource(GetFName(), this, PlatformData->VTData, GetCachedLODBias());
+		FVirtualTexture2DResource* ResourceVT = new FVirtualTexture2DResource(this, PlatformData->VTData, GetCachedLODBias());
 		return ResourceVT;
  	}
 
@@ -1822,12 +1822,12 @@ void FTexture2DResource::GetData( uint32 MipIndex, void* Dest, uint32 DestPitch
 
 
 
-FVirtualTexture2DResource::FVirtualTexture2DResource(const FName& InName, const UTexture2D* InOwner, FVirtualTextureBuiltData* InVTData, int32 SetFirstMipToUse)
+FVirtualTexture2DResource::FVirtualTexture2DResource(const UTexture2D* InOwner, FVirtualTextureBuiltData* InVTData, int32 SetFirstMipToUse)
 	: AllocatedVT(nullptr)
 	, VTData(InVTData)
 	, TextureOwner(InOwner)
-	, OwnerName(InName)
 {
+	check(InOwner);
 	check(InVTData);
 	FirstMipToUse = FMath::Min((int32)VTData->GetNumMips() - 1, SetFirstMipToUse);
 	if (InOwner)
@@ -1871,7 +1871,7 @@ void FVirtualTexture2DResource::InitRHI()
 	check(MaxLevel >= 0);
 
 	FVTProducerDescription ProducerDesc;
-	ProducerDesc.Name = OwnerName;
+	ProducerDesc.Name = TextureOwner->GetFName();
 	ProducerDesc.Dimensions = 2;
 	ProducerDesc.TileSize = VTData->TileSize;
 	ProducerDesc.TileBorderSize = VTData->TileBorderSize;
@@ -1899,7 +1899,7 @@ void FVirtualTexture2DResource::InitializeEditorResources(IVirtualTexture* InVir
 {
 	// Create a texture resource from the lowest resolution VT page data
 	// this will then be used during asset tumbnails/hitproxies/...
-	if (GIsEditor && TextureOwner)
+	if (GIsEditor)
 	{
 		struct FPageToProduce
 		{
diff --git a/Engine/Source/Runtime/Engine/Public/TextureResource.h b/Engine/Source/Runtime/Engine/Public/TextureResource.h
--- a/Engine/Source/Runtime/Engine/Public/TextureResource.h
+++ b/Engine/Source/Runtime/Engine/Public/TextureResource.h
@@ -257,7 +257,7 @@ private:
 class FVirtualTexture2DResource : public FTextureResource
 {
 public:
-	FVirtualTexture2DResource(const FName& InName, const UTexture2D* InOwner, struct FVirtualTextureBuiltData* InVTData, int32 FirstMipToUse);
+	FVirtualTexture2DResource(const UTexture2D* InOwner, struct FVirtualTextureBuiltData* InVTData, int32 FirstMipToUse);
 	virtual ~FVirtualTexture2DResource();
 
 	virtual void InitRHI() override;
@@ -304,7 +304,6 @@ private:
 	class IAllocatedVirtualTexture* AllocatedVT;
 	struct FVirtualTextureBuiltData* VTData;
 	const UTexture2D* TextureOwner;
-	FName OwnerName;
 	FVirtualTextureProducerHandle ProducerHandle;
 	int32 FirstMipToUse;
 };
