Add type templated texture copy global shader.
Intended for use by PS4 CopySubTextureRegion.
#rnx

diff --git a/Engine/Shaders/Private/UpdateTextureShaders.usf b/Engine/Shaders/Private/UpdateTextureShaders.usf
--- a/Engine/Shaders/Private/UpdateTextureShaders.usf
+++ b/Engine/Shaders/Private/UpdateTextureShaders.usf
@@ -92,4 +92,16 @@ void TUpdateTexture2DSubresourceCS( uint2 ThreadId : SV_DispatchThreadID )
 	}
 }
 
+Texture2D<COMPONENT_TYPE> TSrcTexture;
+uint2 SrcPos;
+
+[numthreads(8,8,1)]
+void TCopyTexture2DCS( uint3 ThreadId : SV_DispatchThreadID )
+{
+	if (all(ThreadId.xy < DestPosSize.zw))
+	{
+		TDestTexture[DestPosSize.xy + ThreadId.xy] = TSrcTexture.Load(int3(SrcPos + ThreadId.xy, 0));
+	}
+}
+
 #endif // COMPONENT_TYPE
diff --git a/Engine/Source/Runtime/Engine/Private/UpdateTextureShaders.cpp b/Engine/Source/Runtime/Engine/Private/UpdateTextureShaders.cpp
--- a/Engine/Source/Runtime/Engine/Private/UpdateTextureShaders.cpp
+++ b/Engine/Source/Runtime/Engine/Private/UpdateTextureShaders.cpp
@@ -1,12 +1,16 @@
 // Copyright 1998-2019 Epic Games, Inc. All Rights Reserved.
 
-
 #include "UpdateTextureShaders.h"
 
 IMPLEMENT_SHADER_TYPE(,FUpdateTexture2DSubresouceCS,TEXT("/Engine/Private/UpdateTextureShaders.usf"),TEXT("UpdateTexture2DSubresourceCS"),SF_Compute);
 IMPLEMENT_SHADER_TYPE(,FUpdateTexture3DSubresouceCS, TEXT("/Engine/Private/UpdateTextureShaders.usf"), TEXT("UpdateTexture3DSubresourceCS"), SF_Compute);
+
 IMPLEMENT_SHADER_TYPE(,FCopyTexture2DCS,TEXT("/Engine/Private/UpdateTextureShaders.usf"),TEXT("CopyTexture2DCS"),SF_Compute);
 
+IMPLEMENT_SHADER_TYPE(template<>, TCopyTexture2DCS<1u>, TEXT("/Engine/Private/UpdateTextureShaders.usf"), TEXT("TCopyTexture2DCS"), SF_Compute);
+IMPLEMENT_SHADER_TYPE(template<>, TCopyTexture2DCS<2u>, TEXT("/Engine/Private/UpdateTextureShaders.usf"), TEXT("TCopyTexture2DCS"), SF_Compute);
+IMPLEMENT_SHADER_TYPE(template<>, TCopyTexture2DCS<4u>, TEXT("/Engine/Private/UpdateTextureShaders.usf"), TEXT("TCopyTexture2DCS"), SF_Compute);
+
 IMPLEMENT_SHADER_TYPE(template<>, TUpdateTexture2DSubresouceCS<1u>, TEXT("/Engine/Private/UpdateTextureShaders.usf"), TEXT("TUpdateTexture2DSubresourceCS"), SF_Compute);
 IMPLEMENT_SHADER_TYPE(template<>, TUpdateTexture2DSubresouceCS<2u>, TEXT("/Engine/Private/UpdateTextureShaders.usf"), TEXT("TUpdateTexture2DSubresourceCS"), SF_Compute);
 IMPLEMENT_SHADER_TYPE(template<>, TUpdateTexture2DSubresouceCS<3u>, TEXT("/Engine/Private/UpdateTextureShaders.usf"), TEXT("TUpdateTexture2DSubresourceCS"), SF_Compute);
diff --git a/Engine/Source/Runtime/Engine/Public/UpdateTextureShaders.h b/Engine/Source/Runtime/Engine/Public/UpdateTextureShaders.h
--- a/Engine/Source/Runtime/Engine/Public/UpdateTextureShaders.h
+++ b/Engine/Source/Runtime/Engine/Public/UpdateTextureShaders.h
@@ -167,6 +167,54 @@ public:
 	FShaderParameter DestPosSize;
 };
 
+template<uint32 NumComponents>
+class TCopyTexture2DCS : public FGlobalShader
+{
+	DECLARE_SHADER_TYPE(TCopyTexture2DCS, Global);
+public:
+	TCopyTexture2DCS() {}
+	TCopyTexture2DCS(const ShaderMetaType::CompiledShaderInitializerType& Initializer)
+		: FGlobalShader(Initializer)
+	{
+		SrcTexture.Bind(Initializer.ParameterMap, TEXT("TSrcTexture"), SPF_Mandatory);
+		DestTexture.Bind(Initializer.ParameterMap, TEXT("TDestTexture"), SPF_Mandatory);
+		SrcPosParameter.Bind(Initializer.ParameterMap, TEXT("SrcPos"), SPF_Mandatory);
+		DestPosSizeParameter.Bind(Initializer.ParameterMap, TEXT("DestPosSize"), SPF_Mandatory);
+	}
+
+	// FShader interface.
+	virtual bool Serialize(FArchive& Ar) override
+	{
+		bool bShaderHasOutdatedParameters = FGlobalShader::Serialize(Ar);
+		Ar << SrcTexture << DestTexture << SrcPosParameter << DestPosSizeParameter;
+		return bShaderHasOutdatedParameters;
+	}
+
+	static void ModifyCompilationEnvironment(const FGlobalShaderPermutationParameters& Parameters, FShaderCompilerEnvironment& OutEnvironment)
+	{
+		const TCHAR* Type = nullptr;
+		static_assert(NumComponents == 1u || NumComponents == 2u || NumComponents == 4u, "Invalid NumComponents");
+		switch (NumComponents)
+		{
+		case 1u: Type = TEXT("uint"); break;
+		case 2u: Type = TEXT("uint2"); break;
+		case 4u: Type = TEXT("uint4"); break;
+		}
+
+		OutEnvironment.SetDefine(TEXT("COMPONENT_TYPE"), Type);
+	}
+
+	static bool ShouldCompilePermutation(const FGlobalShaderPermutationParameters& Parameters)
+	{
+		return IsFeatureLevelSupported(Parameters.Platform, ERHIFeatureLevel::SM5);
+	}
+
+	FShaderResourceParameter SrcTexture;
+	FShaderResourceParameter DestTexture;
+	FShaderParameter SrcPosParameter;
+	FShaderParameter DestPosSizeParameter;
+};
+
 template<uint32 ElementsPerThread>
 class TCopyDataCS : public FGlobalShader
 {
