Fixing compilation error.

Declaring the constexpr in the class definition causes compile to fail on MSVS 2017, 15.6.5 but works on 15.9.7.

diff --git a/Engine/Source/Runtime/Engine/Private/VT/VirtualTextureBuiltData.h b/Engine/Source/Runtime/Engine/Private/VT/VirtualTextureBuiltData.h
--- a/Engine/Source/Runtime/Engine/Private/VT/VirtualTextureBuiltData.h
+++ b/Engine/Source/Runtime/Engine/Private/VT/VirtualTextureBuiltData.h
@@ -36,9 +36,6 @@ struct FVirtualTextureDataChunk
 	uint16 CodecPayloadOffset[VIRTUALTEXTURE_DATA_MAXLAYERS];
 	EVirtualTextureCodec CodecType[VIRTUALTEXTURE_DATA_MAXLAYERS];
 
-	// Don't include editor-only data
-	static const uint32 MemoryFootprint = sizeof(BulkData) + sizeof(SizeInBytes) + sizeof(CodecPayloadOffset) + sizeof(CodecPayloadSize) + sizeof(CodecType);
-
 	inline FVirtualTextureDataChunk()
 		: SizeInBytes(0u)
 		, CodecPayloadSize(0u)
@@ -47,7 +44,12 @@ struct FVirtualTextureDataChunk
 		FMemory::Memzero(CodecType);
 	}
 
-	inline uint32 GetMemoryFootprint() const { return MemoryFootprint; }
+	inline uint32 GetMemoryFootprint() const
+	{
+		// Don't include editor-only data
+		static constexpr uint32 MemoryFootprint = uint32(sizeof(BulkData) + sizeof(SizeInBytes) + sizeof(CodecPayloadOffset) + sizeof(CodecPayloadSize) + sizeof(CodecType));
+		return MemoryFootprint;
+	}
 
 #if WITH_EDITORONLY_DATA
 	/** Key if stored in the derived data cache. */
