SCW executed locally when XGE job failed, preventing crash.

#rnx

diff --git a/Engine/Source/Runtime/Engine/Private/ShaderCompiler/ShaderCompiler.cpp b/Engine/Source/Runtime/Engine/Private/ShaderCompiler/ShaderCompiler.cpp
--- a/Engine/Source/Runtime/Engine/Private/ShaderCompiler/ShaderCompiler.cpp
+++ b/Engine/Source/Runtime/Engine/Private/ShaderCompiler/ShaderCompiler.cpp
@@ -29,7 +29,6 @@
 #include "Materials/Material.h"
 #include "EditorSupportDelegates.h"
 #include "GlobalShader.h"
-#include "Interfaces/ITargetPlatformManagerModule.h"
 #include "Interfaces/IShaderFormat.h"
 #include "Interfaces/IShaderFormatModule.h"
 #include "RendererInterface.h"
@@ -1347,73 +1346,78 @@ void FShaderCompileThreadRunnable::CompileDirectlyThroughDll()
 				check(!CurrentJob.bFinalized);
 				CurrentJob.bFinalized = true;
 
-				static ITargetPlatformManagerModule& TPM = GetTargetPlatformManagerRef();
-				auto* SingleJob = CurrentJob.GetSingleShaderJob();
-				if (SingleJob)
-				{
-					const FName Format = LegacyShaderPlatformToShaderFormat(EShaderPlatform(SingleJob->Input.Target.Platform));
-					const IShaderFormat* Compiler = TPM.FindShaderFormat(Format);
+				FShaderCompileUtilities::ExecuteShaderCompileJob(CurrentJob);
+			}
 
-					if (!Compiler)
-					{
-						UE_LOG(LogShaderCompilers, Fatal, TEXT("Can't compile shaders for format %s, couldn't load compiler dll"), *Format.ToString());
-					}
-					CA_ASSUME(Compiler != NULL);
+			CurrentWorkerInfo.bComplete = true;
+		}
+	}
+}
 
-					if (IsValidRef(SingleJob->Input.SharedEnvironment))
-					{
-						// Merge the shared environment into the per-shader environment before calling into the compile function
-						// Normally this happens in the worker
-						SingleJob->Input.Environment.Merge(*SingleJob->Input.SharedEnvironment);
-					}
+void FShaderCompileUtilities::ExecuteShaderCompileJob(FShaderCommonCompileJob& Job)
+{
+	static ITargetPlatformManagerModule& TPM = GetTargetPlatformManagerRef();
+	auto* SingleJob = Job.GetSingleShaderJob();
+	if (SingleJob)
+	{
+		const FName Format = LegacyShaderPlatformToShaderFormat(EShaderPlatform(SingleJob->Input.Target.Platform));
+		const IShaderFormat* Compiler = TPM.FindShaderFormat(Format);
 
-					// Compile the shader directly through the platform dll (directly from the shader dir as the working directory)
-					Compiler->CompileShader(Format, SingleJob->Input, SingleJob->Output, FString(FPlatformProcess::ShaderDir()));
+		if (!Compiler)
+		{
+			UE_LOG(LogShaderCompilers, Fatal, TEXT("Can't compile shaders for format %s, couldn't load compiler dll"), *Format.ToString());
+		}
+		CA_ASSUME(Compiler != NULL);
 
-					SingleJob->bSucceeded = SingleJob->Output.bSucceeded;
+		if (IsValidRef(SingleJob->Input.SharedEnvironment))
+		{
+			// Merge the shared environment into the per-shader environment before calling into the compile function
+			// Normally this happens in the worker
+			SingleJob->Input.Environment.Merge(*SingleJob->Input.SharedEnvironment);
+		}
 
-					if (SingleJob->Output.bSucceeded)
-					{
-						// Generate a hash of the output and cache it
-						// The shader processing this output will use it to search for existing FShaderResources
-						SingleJob->Output.GenerateOutputHash();
-					}
-				}
-				else
-				{
-					auto* PipelineJob = CurrentJob.GetShaderPipelineJob();
-					check(PipelineJob);
+		// Compile the shader directly through the platform dll (directly from the shader dir as the working directory)
+		Compiler->CompileShader(Format, SingleJob->Input, SingleJob->Output, FString(FPlatformProcess::ShaderDir()));
 
-					EShaderPlatform Platform = (EShaderPlatform)PipelineJob->StageJobs[0]->GetSingleShaderJob()->Input.Target.Platform;
-					const FName Format = LegacyShaderPlatformToShaderFormat(Platform);
-					const IShaderFormat* Compiler = TPM.FindShaderFormat(Format);
+		SingleJob->bSucceeded = SingleJob->Output.bSucceeded;
 
-					if (!Compiler)
-					{
-						UE_LOG(LogShaderCompilers, Fatal, TEXT("Can't compile shaders for format %s, couldn't load compiler dll"), *Format.ToString());
-					}
-					CA_ASSUME(Compiler != NULL);
+		if (SingleJob->Output.bSucceeded)
+		{
+			// Generate a hash of the output and cache it
+			// The shader processing this output will use it to search for existing FShaderResources
+			SingleJob->Output.GenerateOutputHash();
+		}
+	}
+	else
+	{
+		FShaderPipelineCompileJob* PipelineJob = Job.GetShaderPipelineJob();
+		check(PipelineJob);
 
-					// Verify same platform on all stages
-					for (int32 Index = 1; Index < PipelineJob->StageJobs.Num(); ++Index)
-					{
-						auto* SingleStage = PipelineJob->StageJobs[Index]->GetSingleShaderJob();
-						if (!SingleStage)
-						{
-							UE_LOG(LogShaderCompilers, Fatal, TEXT("Can't nest Shader Pipelines inside Shader Pipeline '%s'!"), PipelineJob->ShaderPipeline->GetName());
-						}
-						if (Platform != SingleStage->Input.Target.Platform)
-						{
-							UE_LOG(LogShaderCompilers, Fatal, TEXT("Mismatched Target Platform %s while compiling Shader Pipeline '%s'."), *Format.GetPlainNameString(), PipelineJob->ShaderPipeline->GetName());
-						}
-					}
+		EShaderPlatform Platform = (EShaderPlatform)PipelineJob->StageJobs[0]->GetSingleShaderJob()->Input.Target.Platform;
+		const FName Format = LegacyShaderPlatformToShaderFormat(Platform);
+		const IShaderFormat* Compiler = TPM.FindShaderFormat(Format);
 
-					CompileShaderPipeline(Compiler, Format, PipelineJob, FString(FPlatformProcess::ShaderDir()));
-				}
-			}
+		if (!Compiler)
+		{
+			UE_LOG(LogShaderCompilers, Fatal, TEXT("Can't compile shaders for format %s, couldn't load compiler dll"), *Format.ToString());
+		}
+		CA_ASSUME(Compiler != NULL);
 
-			CurrentWorkerInfo.bComplete = true;
+		// Verify same platform on all stages
+		for (int32 Index = 1; Index < PipelineJob->StageJobs.Num(); ++Index)
+		{
+			auto* SingleStage = PipelineJob->StageJobs[Index]->GetSingleShaderJob();
+			if (!SingleStage)
+			{
+				UE_LOG(LogShaderCompilers, Fatal, TEXT("Can't nest Shader Pipelines inside Shader Pipeline '%s'!"), PipelineJob->ShaderPipeline->GetName());
+			}
+			if (Platform != SingleStage->Input.Target.Platform)
+			{
+				UE_LOG(LogShaderCompilers, Fatal, TEXT("Mismatched Target Platform %s while compiling Shader Pipeline '%s'."), *Format.GetPlainNameString(), PipelineJob->ShaderPipeline->GetName());
+			}
 		}
+
+		CompileShaderPipeline(Compiler, Format, PipelineJob, FString(FPlatformProcess::ShaderDir()));
 	}
 }
 
@@ -1475,7 +1479,6 @@ FShaderCompilingManager::FShaderCompilingManager() :
 	SuppressedShaderPlatforms(0)
 {
 	WorkersBusyTime = 0;
-	bFallBackToDirectCompiles = false;
 
 	// Threads must use absolute paths on Windows in case the current directory is changed on another thread!
 	ShaderCompileWorkerName = FPaths::ConvertRelativePathToFull(ShaderCompileWorkerName);
diff --git a/Engine/Source/Runtime/Engine/Private/ShaderCompiler/ShaderCompilerXGE.cpp b/Engine/Source/Runtime/Engine/Private/ShaderCompiler/ShaderCompilerXGE.cpp
--- a/Engine/Source/Runtime/Engine/Private/ShaderCompiler/ShaderCompilerXGE.cpp
+++ b/Engine/Source/Runtime/Engine/Private/ShaderCompiler/ShaderCompilerXGE.cpp
@@ -858,6 +858,8 @@ int32 FShaderCompileXGEThreadRunnable_InterceptionInterface::CompilingLoop()
 
 	for (auto Iter = DispatchedTasks.CreateIterator(); Iter; ++Iter)
 	{
+		bool bOutputFileReadFailed = false;
+
 		FXGEShaderCompilerTask* Task = *Iter;
 		if (!Task->Future.IsReady())
 		{
@@ -875,9 +877,24 @@ int32 FShaderCompileXGEThreadRunnable_InterceptionInterface::CompilingLoop()
 		if (Result.bCompleted)
 		{
 			// Open the output file, and serialize in the completed jobs.
-			FArchive* OutputFileAr = IFileManager::Get().CreateFileReader(*Task->OutputFilePath, FILEREAD_NoFail);
-			FShaderCompileUtilities::DoReadTaskResults(Task->ShaderJobs, *OutputFileAr);
-			delete OutputFileAr;
+			FArchive* OutputFileAr = IFileManager::Get().CreateFileReader(*Task->OutputFilePath, FILEREAD_Silent);
+			if (OutputFileAr)
+			{
+				FShaderCompileUtilities::DoReadTaskResults(Task->ShaderJobs, *OutputFileAr);
+				delete OutputFileAr;
+			}
+			else
+			{
+				// Reading result from XGE job failed, so recompile shaders in current job batch locally
+				bOutputFileReadFailed = true;
+
+				UE_LOG(LogShaderCompilers, Error, TEXT("Reschedule shader compilation after XGE job failed: %s"), *Task->OutputFilePath);
+
+				for (FShaderCommonCompileJob* Job : Task->ShaderJobs)
+				{
+					FShaderCompileUtilities::ExecuteShaderCompileJob(*Job);
+				}
+			}
 
 			// Enter the critical section so we can access the input and output queues
 			{
@@ -905,9 +922,13 @@ int32 FShaderCompileXGEThreadRunnable_InterceptionInterface::CompilingLoop()
 		{
 			FPlatformProcess::Sleep(0.01f);
 		}
-		while (!IFileManager::Get().Delete(*Task->OutputFilePath, false, true, true))
+
+		if (!bOutputFileReadFailed)
 		{
-			FPlatformProcess::Sleep(0.01f);
+			while (!IFileManager::Get().Delete(*Task->OutputFilePath, false, true, true))
+			{
+				FPlatformProcess::Sleep(0.01f);
+			}
 		}
 
 		Iter.RemoveCurrent();
diff --git a/Engine/Source/Runtime/Engine/Public/ShaderCompiler.h b/Engine/Source/Runtime/Engine/Public/ShaderCompiler.h
--- a/Engine/Source/Runtime/Engine/Public/ShaderCompiler.h
+++ b/Engine/Source/Runtime/Engine/Public/ShaderCompiler.h
@@ -225,6 +225,9 @@ namespace FShaderCompileUtilities
 {
 	bool DoWriteTasks(const TArray<FShaderCommonCompileJob*>& QueuedJobs, FArchive& TransferFile);
 	void DoReadTaskResults(const TArray<FShaderCommonCompileJob*>& QueuedJobs, FArchive& OutputFile);
+
+	/** Execute the specified (single or pipeline) shader compile job. */
+	void ExecuteShaderCompileJob(FShaderCommonCompileJob& Job);
 }
 
 #if PLATFORM_WINDOWS // XGE shader compilation is only supported on Windows.
@@ -445,8 +448,6 @@ private:
 	FString AbsoluteShaderDebugInfoDirectory;
 	/** Name of the shader worker application. */
 	FString ShaderCompileWorkerName;
-	/** Whether the SCW has crashed and we should fall back to calling the compiler dll's directly. */
-	bool bFallBackToDirectCompiles;
 
 	/** 
 	 * Tracks the total time that shader compile workers have been busy since startup.  
