Fix bad lightmap UVs on instances when using foliage.DensityScale
#rn minor Editor/Foliage

diff --git a/Engine/Source/Runtime/Engine/Private/HierarchicalInstancedStaticMesh.cpp b/Engine/Source/Runtime/Engine/Private/HierarchicalInstancedStaticMesh.cpp
--- a/Engine/Source/Runtime/Engine/Private/HierarchicalInstancedStaticMesh.cpp
+++ b/Engine/Source/Runtime/Engine/Private/HierarchicalInstancedStaticMesh.cpp
@@ -2854,14 +2854,9 @@ void UHierarchicalInstancedStaticMeshComponent::SetPerInstanceLightMapAndEditorD
 
 	if (MeshMapBuildData != nullptr || GIsEditor)
 	{
-		for (int32 Index = 0; Index < NumInstances; ++Index)
+		for (int32 RenderIndex = 0; RenderIndex < NumInstances; ++RenderIndex)
 		{
-			int32 RenderIndex = InstanceReorderTable.IsValidIndex(Index) ? InstanceReorderTable[Index] : Index;
-			if (RenderIndex == INDEX_NONE)
-			{
-				// could be skipped by density settings
-				continue;
-			}
+			int32 Index = SortedInstances[RenderIndex];
 						
 			FVector2D LightmapUVBias = FVector2D(-1.0f, -1.0f);
 			FVector2D ShadowmapUVBias = FVector2D(-1.0f, -1.0f);
