Load UMaterialExpressionRuntimeVirtualTextureOutput objects on client
Fixes missing landscape in virtual texture
Better (incoming) fix will be not to load this, and have a cheaper serialized flag for virtual texture materials

diff --git a/Engine/Source/Runtime/Engine/Classes/Materials/MaterialExpressionRuntimeVirtualTextureOutput.h b/Engine/Source/Runtime/Engine/Classes/Materials/MaterialExpressionRuntimeVirtualTextureOutput.h
--- a/Engine/Source/Runtime/Engine/Classes/Materials/MaterialExpressionRuntimeVirtualTextureOutput.h
+++ b/Engine/Source/Runtime/Engine/Classes/Materials/MaterialExpressionRuntimeVirtualTextureOutput.h
@@ -38,6 +38,11 @@ class UMaterialExpressionRuntimeVirtualTextureOutput : public UMaterialExpressio
 	FExpressionInput Opacity;
 
 public:
+	//~ Begin UObject Interface
+	// Only needed for (bad, slow) implementation of FMaterialResource::HasRuntimeVirtualTextureOutput()
+	virtual bool NeedsLoadForClient() const override { return true; }
+	//~ End UObject Interface
+
 #if WITH_EDITOR
 	//~ Begin UMaterialExpression Interface
 	virtual int32 Compile(class FMaterialCompiler* Compiler, int32 OutputIndex) override;
