Fixed static analysis warnings in VT code.

diff --git a/Engine/Source/Runtime/Engine/Private/LightMap.cpp b/Engine/Source/Runtime/Engine/Private/LightMap.cpp
--- a/Engine/Source/Runtime/Engine/Private/LightMap.cpp
+++ b/Engine/Source/Runtime/Engine/Private/LightMap.cpp
@@ -1691,15 +1691,15 @@ void FLightMapPendingTexture::StartEncoding(ULevel* LightingScenario, ITextureCo
 		LayerFormat.Init(TSF_Invalid, NumVirtualTextureLayers);
 		LayerFormat[0] = Textures[0]->Source.GetFormat();
 		LayerFormat[1] = Textures[0]->Source.GetFormat();
-		if (SkyOcclusionLayer != ~0u)
+		if (SkyOcclusionLayer != ~0u && SkyOcclusionTexture)
 		{
 			LayerFormat[SkyOcclusionLayer] = SkyOcclusionTexture->Source.GetFormat();
 		}
-		if (AOMaterialMaskLayer != ~0u)
+		if (AOMaterialMaskLayer != ~0u && AOMaterialMaskTexture)
 		{
 			LayerFormat[AOMaterialMaskLayer] = AOMaterialMaskTexture->Source.GetFormat();
 		}
-		if (ShadowMaskLayer != ~0u)
+		if (ShadowMaskLayer != ~0u && ShadowMapTexture)
 		{
 			LayerFormat[ShadowMaskLayer] = ShadowMapTexture->Source.GetFormat();
 		}
diff --git a/Engine/Source/Runtime/Engine/Private/Texture2D.cpp b/Engine/Source/Runtime/Engine/Private/Texture2D.cpp
--- a/Engine/Source/Runtime/Engine/Private/Texture2D.cpp
+++ b/Engine/Source/Runtime/Engine/Private/Texture2D.cpp
@@ -1847,25 +1847,24 @@ void FVirtualTexture2DResource::RefreshSamplerStates()
 
 void FVirtualTexture2DResource::InitRHI()
 {
+	check(TextureOwner);
+
 	// We always create a sampler state if we're attached to a texture. This is used to sample the cache texture during actual rendering and the miptails editor resource.
 	// If we're not attached to a texture it likely means we're light maps which have sampling handled differently.
-	if (TextureOwner != nullptr)
-	{
-		FSamplerStateInitializerRHI SamplerStateInitializer
-		(
-			// This will ensure nearest/linear/trilinear which does matter when sampling both the cache and the miptail
-			(ESamplerFilter)UDeviceProfileManager::Get().GetActiveProfile()->GetTextureLODSettings()->GetSamplerFilter(TextureOwner),
+	FSamplerStateInitializerRHI SamplerStateInitializer
+	(
+		// This will ensure nearest/linear/trilinear which does matter when sampling both the cache and the miptail
+		(ESamplerFilter)UDeviceProfileManager::Get().GetActiveProfile()->GetTextureLODSettings()->GetSamplerFilter(TextureOwner),
 
-			// This doesn't really matter when sampling the cache texture but it does when sampling the miptail texture
-			TextureOwner->AddressX == TA_Wrap ? AM_Wrap : (TextureOwner->AddressX == TA_Clamp ? AM_Clamp : AM_Mirror),
-			TextureOwner->AddressY == TA_Wrap ? AM_Wrap : (TextureOwner->AddressY == TA_Clamp ? AM_Clamp : AM_Mirror),
-			AM_Wrap,
+		// This doesn't really matter when sampling the cache texture but it does when sampling the miptail texture
+		TextureOwner->AddressX == TA_Wrap ? AM_Wrap : (TextureOwner->AddressX == TA_Clamp ? AM_Clamp : AM_Mirror),
+		TextureOwner->AddressY == TA_Wrap ? AM_Wrap : (TextureOwner->AddressY == TA_Clamp ? AM_Clamp : AM_Mirror),
+		AM_Wrap,
 
-			// This doesn't really matter when sampling the cache texture (as it only has a level 0, so whatever the bias that is sampled) but it does when we sample miptail texture
-			0 // VT currently ignores global mip bias ensure the miptail works the same -> UTexture2D::GetGlobalMipMapLODBias()
-		);
-		SamplerStateRHI = RHICreateSamplerState(SamplerStateInitializer);
-	}
+		// This doesn't really matter when sampling the cache texture (as it only has a level 0, so whatever the bias that is sampled) but it does when we sample miptail texture
+		0 // VT currently ignores global mip bias ensure the miptail works the same -> UTexture2D::GetGlobalMipMapLODBias()
+	);
+	SamplerStateRHI = RHICreateSamplerState(SamplerStateInitializer);
 
 	const int32 MaxLevel = VTData->GetNumMips() - FirstMipToUse - 1;
 	check(MaxLevel >= 0);
