Fix support for non-power-2 padding when building virtual textures

diff --git a/Engine/Source/Runtime/Engine/Private/VT/VirtualTextureDataBuilder.cpp b/Engine/Source/Runtime/Engine/Private/VT/VirtualTextureDataBuilder.cpp
--- a/Engine/Source/Runtime/Engine/Private/VT/VirtualTextureDataBuilder.cpp
+++ b/Engine/Source/Runtime/Engine/Private/VT/VirtualTextureDataBuilder.cpp
@@ -278,10 +278,31 @@ void FVirtualTextureDataBuilder::Build(const FTextureSourceData& InSourceData, c
 	SettingsPerLayer.AddUninitialized(NumLayers);
 	FMemory::Memcpy(&SettingsPerLayer[0], InSettingsPerLayer, sizeof(FTextureBuildSettings) * NumLayers);
 
-	SizeInBlocksX = InSourceData.SizeInBlocksX;
-	SizeInBlocksY = InSourceData.SizeInBlocksY;
 	BlockSizeX = InSourceData.BlockSizeX;
 	BlockSizeY = InSourceData.BlockSizeY;
+
+	// BlockSize is potentially adjusted by rounding to power of 2
+	switch (SettingsPerLayer[0].PowerOfTwoMode)
+	{
+	case ETexturePowerOfTwoSetting::None:
+		break;
+	case ETexturePowerOfTwoSetting::PadToPowerOfTwo:
+		BlockSizeX = FMath::RoundUpToPowerOfTwo(BlockSizeX);
+		BlockSizeY = FMath::RoundUpToPowerOfTwo(BlockSizeY);
+		break;
+	case ETexturePowerOfTwoSetting::PadToSquarePowerOfTwo:
+		BlockSizeX = FMath::RoundUpToPowerOfTwo(BlockSizeX);
+		BlockSizeY = FMath::RoundUpToPowerOfTwo(BlockSizeY);
+		BlockSizeX = FMath::Max(BlockSizeX, BlockSizeY);
+		BlockSizeY = BlockSizeX;
+		break;
+	default:
+		checkNoEntry();
+		break;
+	}
+
+	SizeInBlocksX = InSourceData.SizeInBlocksX;
+	SizeInBlocksY = InSourceData.SizeInBlocksY;
 	SizeX = BlockSizeX * SizeInBlocksX;
 	SizeY = BlockSizeY * SizeInBlocksY;
 
@@ -806,11 +827,11 @@ void FVirtualTextureDataBuilder::BuildSourcePixels(const FTextureSourceData& Sou
 		FTextureSourceBlockData& BlockData = SourceBlocks[BlockIndex];
 		BlockData.BlockX = SourceBlockData.BlockX;
 		BlockData.BlockY = SourceBlockData.BlockY;
-		BlockData.SizeX = SourceBlockData.SizeX;
-		BlockData.SizeY = SourceBlockData.SizeY;
 		BlockData.NumMips = SourceBlockData.NumMips;
 		BlockData.NumSlices = SourceBlockData.NumSlices;
 		BlockData.MipBias = SourceBlockData.MipBias;
+		BlockData.SizeX = 0u;
+		BlockData.SizeY = 0u;
 		BlockData.MipsPerLayer.AddDefaulted(NumLayers);
 		for (int32 LayerIndex = 0; LayerIndex < NumLayers; ++LayerIndex)
 		{
@@ -858,10 +879,11 @@ void FVirtualTextureDataBuilder::BuildSourcePixels(const FTextureSourceData& Sou
 				check(false);
 			}
 
-			check(BlockData.SizeX == CompressedMips[0].SizeX);
-			check(BlockData.SizeY == CompressedMips[0].SizeY);
+			// Get size of block from Compressor output, since it may have been padded/adjusted
+			BlockData.SizeX = CompressedMips[0].SizeX;
+			BlockData.SizeY = CompressedMips[0].SizeY;
 
-			const uint32 BlockSize = FMath::Max(SourceBlockData.SizeX, SourceBlockData.SizeY);
+			const uint32 BlockSize = FMath::Max(BlockData.SizeX, BlockData.SizeY);
 			const uint32 BlockSizeInTiles = FMath::DivideAndRoundUp<uint32>(BlockSize, TileSize);
 			const uint32 MaxMipInBlock = FMath::CeilLogTwo(BlockSizeInTiles);
 
@@ -890,13 +912,13 @@ void FVirtualTextureDataBuilder::BuildSourcePixels(const FTextureSourceData& Sou
 	// If we have more than 1 block, need to create miptail that contains mips made from multiple blocks
 	if (NumBlocks > 1)
 	{
-		const uint32 BlockSize = FMath::Max(SourceData.BlockSizeX, SourceData.BlockSizeY);
+		const uint32 BlockSize = FMath::Max(BlockSizeX, BlockSizeY);
 		const uint32 BlockSizeInTiles = FMath::DivideAndRoundUp<uint32>(BlockSize, TileSize);
 		const uint32 MaxMipInBlock = FMath::CeilLogTwo(BlockSizeInTiles);
-		const uint32 MipWidthInBlock = FMath::Max<uint32>(SourceData.BlockSizeX >> MaxMipInBlock, 1);
-		const uint32 MipHeightInBlock = FMath::Max<uint32>(SourceData.BlockSizeY >> MaxMipInBlock, 1);
-		const uint32 MipInputSizeX = FMath::RoundUpToPowerOfTwo(SourceData.SizeInBlocksX * MipWidthInBlock);
-		const uint32 MipInputSizeY = FMath::RoundUpToPowerOfTwo(SourceData.SizeInBlocksY * MipHeightInBlock);
+		const uint32 MipWidthInBlock = FMath::Max<uint32>(BlockSizeX >> MaxMipInBlock, 1);
+		const uint32 MipHeightInBlock = FMath::Max<uint32>(BlockSizeY >> MaxMipInBlock, 1);
+		const uint32 MipInputSizeX = FMath::RoundUpToPowerOfTwo(SizeInBlocksX * MipWidthInBlock);
+		const uint32 MipInputSizeY = FMath::RoundUpToPowerOfTwo(SizeInBlocksY * MipHeightInBlock);
 		const uint32 MipInputSize = FMath::Max(MipInputSizeX, MipInputSizeY);
 		const uint32 MipInputSizeInTiles = FMath::DivideAndRoundUp<uint32>(MipInputSize, TileSize);
 
