fix out-of-bounds indexing when packing certain textures

diff --git a/Engine/Source/Runtime/Engine/Private/TextureDerivedDataTask.cpp b/Engine/Source/Runtime/Engine/Private/TextureDerivedDataTask.cpp
--- a/Engine/Source/Runtime/Engine/Private/TextureDerivedDataTask.cpp
+++ b/Engine/Source/Runtime/Engine/Private/TextureDerivedDataTask.cpp
@@ -255,7 +255,7 @@ void FTextureCacheDerivedDataWorker::BuildTexture()
 		// Compress the texture.
 		TArray<FCompressedImage2D> CompressedMips;
 		if (Compressor->BuildTexture(TextureData.Blocks[0].MipsPerLayer[0],
-			(bool)Texture.CompositeTexture ? CompositeTextureData.Blocks[0].MipsPerLayer[0] : TArray<FImage>(),
+			((bool)Texture.CompositeTexture && CompositeTextureData.Blocks.Num() && CompositeTextureData.Blocks[0].MipsPerLayer.Num()) ? CompositeTextureData.Blocks[0].MipsPerLayer[0] : TArray<FImage>(),
 			BuildSettingsPerLayer[0],
 			CompressedMips))
 		{
