Fix StaticMesh RayTracingGeometries race condition between GT and RT

diff --git a/Engine/Source/Runtime/Engine/Private/StaticMeshRender.cpp b/Engine/Source/Runtime/Engine/Private/StaticMeshRender.cpp
--- a/Engine/Source/Runtime/Engine/Private/StaticMeshRender.cpp
+++ b/Engine/Source/Runtime/Engine/Private/StaticMeshRender.cpp
@@ -214,9 +214,9 @@ FStaticMeshSceneProxy::FStaticMeshSceneProxy(UStaticMeshComponent* InComponent,
 			}
 		}
 
-#if RHI_RAYTRACING
-		RayTracingGeometries[LODIndex] = RenderData->LODResources[LODIndex].RayTracingGeometry.RayTracingGeometryRHI;
-#endif
+	#if RHI_RAYTRACING
+		RayTracingGeometries[LODIndex] = &RenderData->LODResources[LODIndex].RayTracingGeometry;
+	#endif
 	}
 
 	// WPO is typically used for ambient animations, so don't include in cached shadowmaps
diff --git a/Engine/Source/Runtime/Engine/Public/PrimitiveSceneProxy.h b/Engine/Source/Runtime/Engine/Public/PrimitiveSceneProxy.h
--- a/Engine/Source/Runtime/Engine/Public/PrimitiveSceneProxy.h
+++ b/Engine/Source/Runtime/Engine/Public/PrimitiveSceneProxy.h
@@ -198,9 +198,9 @@ public:
 	/** Gathers dynamic ray tracing instances from this proxy. */
 	virtual void GetDynamicRayTracingInstances(struct FRayTracingMaterialGatheringContext& Context, TArray<struct FRayTracingInstance>& OutRayTracingInstances) {}
 
-	TArray<FRayTracingGeometryRHIRef>&& MoveRayTracingGeometries()
+	TArray<FRayTracingGeometry*>&& MoveRayTracingGeometries()
 	{
-		return static_cast<TArray<FRayTracingGeometryRHIRef>&&>(RayTracingGeometries);
+		return static_cast<TArray<FRayTracingGeometry*>&&>(RayTracingGeometries);
 	}
 #endif // RHI_RAYTRACING
 
@@ -745,7 +745,7 @@ protected:
 	}
 
 #if RHI_RAYTRACING
-	TArray<FRayTracingGeometryRHIRef> RayTracingGeometries;
+	TArray<FRayTracingGeometry*> RayTracingGeometries;
 #endif
 
 private:
diff --git a/Engine/Source/Runtime/Renderer/Private/PrimitiveSceneInfo.cpp b/Engine/Source/Runtime/Renderer/Private/PrimitiveSceneInfo.cpp
--- a/Engine/Source/Runtime/Renderer/Private/PrimitiveSceneInfo.cpp
+++ b/Engine/Source/Runtime/Renderer/Private/PrimitiveSceneInfo.cpp
@@ -185,7 +185,7 @@ FRayTracingGeometryRHIRef FPrimitiveSceneInfo::GetStaticRayTracingGeometryInstan
 {
 	if (RayTracingGeometries.Num() > LodLevel)
 	{
-		return RayTracingGeometries[LodLevel];
+		return RayTracingGeometries[LodLevel]->RayTracingGeometryRHI;
 	}
 	else
 	{
diff --git a/Engine/Source/Runtime/Renderer/Public/PrimitiveSceneInfo.h b/Engine/Source/Runtime/Renderer/Public/PrimitiveSceneInfo.h
--- a/Engine/Source/Runtime/Renderer/Public/PrimitiveSceneInfo.h
+++ b/Engine/Source/Runtime/Renderer/Public/PrimitiveSceneInfo.h
@@ -490,7 +490,7 @@ private:
 	void RemoveCachedMeshDrawCommands();
 
 #if RHI_RAYTRACING
-	TArray<FRayTracingGeometryRHIRef> RayTracingGeometries;
+	TArray<FRayTracingGeometry*> RayTracingGeometries;
 #endif
 };
 
