Add depth bias parameter for controlling bias strength accross casades
Add clamp to on cascade transition for avoiding far cascade to leak

This adds a parameter on directional light to scale depth bias based on world texel size and allows user to tweak/smooth shadow acne across cascades transition.

diff --git a/Engine/Source/Runtime/Engine/Classes/Components/DirectionalLightComponent.h b/Engine/Source/Runtime/Engine/Classes/Components/DirectionalLightComponent.h
--- a/Engine/Source/Runtime/Engine/Classes/Components/DirectionalLightComponent.h
+++ b/Engine/Source/Runtime/Engine/Classes/Components/DirectionalLightComponent.h
@@ -19,6 +19,14 @@ class ENGINE_API UDirectionalLightComponent : public ULightComponent
 {
 	GENERATED_UCLASS_BODY()
 
+	/**
+	* Controls the depth bias scaling across cascades. This allows to mitigage the shadow acne difference on shadow cascades transition.
+	* A value of 1 scales shadow bias based on each cascade size (Default).
+	* A value of 0 scales shadow bias uniformly accross all cacascade.
+	*/
+	UPROPERTY(EditAnywhere, BlueprintReadOnly, Category = Light, AdvancedDisplay, meta = (UIMin = "0", UIMax = "1"))
+	float ShadowCascadeBiasDistribution;
+
 	/** Whether to occlude fog and atmosphere inscattering with screenspace blurred occlusion from this light. */
 	UPROPERTY(EditAnywhere, BlueprintReadOnly, Category=LightShafts, meta=(DisplayName = "Light Shaft Occlusion"))
 	uint32 bEnableLightShaftOcclusion:1;
diff --git a/Engine/Source/Runtime/Engine/Private/Components/DirectionalLightComponent.cpp b/Engine/Source/Runtime/Engine/Private/Components/DirectionalLightComponent.cpp
--- a/Engine/Source/Runtime/Engine/Private/Components/DirectionalLightComponent.cpp
+++ b/Engine/Source/Runtime/Engine/Private/Components/DirectionalLightComponent.cpp
@@ -46,6 +46,9 @@ class FDirectionalLightSceneProxy : public FLightSceneProxy
 {
 public:
 
+	/** Control how the cascade size influence the shadow depth bias. */
+	float ShadowCascadeBiasDistribution;
+
 	/** Whether to occlude fog and atmosphere inscattering with screenspace blurred occlusion from this light. */
 	bool bEnableLightShaftOcclusion;
 
@@ -115,6 +118,7 @@ public:
 	/** Initialization constructor. */
 	FDirectionalLightSceneProxy(const UDirectionalLightComponent* Component) :
 		FLightSceneProxy(Component),
+		ShadowCascadeBiasDistribution(Component->ShadowCascadeBiasDistribution),
 		bEnableLightShaftOcclusion(Component->bEnableLightShaftOcclusion),
 		bUseInsetShadowsForMovableObjects(Component->bUseInsetShadowsForMovableObjects),
 		OcclusionMaskDarkness(Component->OcclusionMaskDarkness),
@@ -734,6 +738,7 @@ private:
 			OutCascadeSettings->SplitNear = SplitNear;
 			OutCascadeSettings->FadePlaneOffset = FadePlane;
 			OutCascadeSettings->FadePlaneLength = SplitFar - FadePlane;
+			OutCascadeSettings->CascadeBiasDistribution = ShadowCascadeBiasDistribution;
 		}
 
 		const FSphere CascadeSphere = FDirectionalLightSceneProxy::GetShadowSplitBoundsDepthRange(View, View.ViewMatrices.GetViewOrigin(), SplitNear, SplitFar, OutCascadeSettings);
@@ -764,6 +769,7 @@ UDirectionalLightComponent::UDirectionalLightComponent(const FObjectInitializer&
 	OcclusionDepthRange = 100000.f;
 	OcclusionMaskDarkness = 0.05f;
 
+	ShadowCascadeBiasDistribution = 1;
 	WholeSceneDynamicShadowRadius_DEPRECATED = 20000.0f;
 	DynamicShadowDistanceMovableLight = 20000.0f;
 	DynamicShadowDistanceStationaryLight = 0.f;
@@ -800,11 +806,12 @@ void UDirectionalLightComponent::PostEditChangeProperty(FPropertyChangedEvent& P
 
 	DynamicShadowDistanceMovableLight = FMath::Max(DynamicShadowDistanceMovableLight, 0.0f);
 	DynamicShadowDistanceStationaryLight = FMath::Max(DynamicShadowDistanceStationaryLight, 0.0f);
+	ShadowCascadeBiasDistribution = FMath::Clamp(ShadowCascadeBiasDistribution, 0.f, 1.f);
 
 	DynamicShadowCascades = FMath::Clamp(DynamicShadowCascades, 0, 10);
 	FarShadowCascadeCount = FMath::Clamp(FarShadowCascadeCount, 0, 10);
 	CascadeDistributionExponent = FMath::Clamp(CascadeDistributionExponent, .1f, 10.0f);
-	CascadeTransitionFraction = FMath::Clamp(CascadeTransitionFraction, 0.0f, 0.3f);
+	CascadeTransitionFraction = FMath::Clamp(CascadeTransitionFraction, KINDA_SMALL_NUMBER, 0.3f);
 	ShadowDistanceFadeoutFraction = FMath::Clamp(ShadowDistanceFadeoutFraction, 0.0f, 1.0f);
 	// max range is larger than UI
 	ShadowBias = FMath::Clamp(ShadowBias, 0.0f, 10.0f);
diff --git a/Engine/Source/Runtime/Engine/Public/SceneManagement.h b/Engine/Source/Runtime/Engine/Public/SceneManagement.h
--- a/Engine/Source/Runtime/Engine/Public/SceneManagement.h
+++ b/Engine/Source/Runtime/Engine/Public/SceneManagement.h
@@ -897,6 +897,9 @@ public:
 	 * Or index of the direction if this is a whole scene shadow from a point light, otherwise INDEX_NONE. 
 	 */
 	int32 ShadowSplitIndex;
+
+	/** Strength of depth bias across cascades. */
+	float CascadeBiasDistribution;
 	
 	FShadowCascadeSettings()
 		: SplitNear(0.0f)
@@ -907,6 +910,7 @@ public:
 		, FadePlaneLength(SplitFar - FadePlaneOffset)
 		, bFarShadowCascade(false)
 		, ShadowSplitIndex(INDEX_NONE)
+		, CascadeBiasDistribution(1)
 	{
 	}
 };
diff --git a/Engine/Source/Runtime/Renderer/Private/ShadowRendering.cpp b/Engine/Source/Runtime/Renderer/Private/ShadowRendering.cpp
--- a/Engine/Source/Runtime/Renderer/Private/ShadowRendering.cpp
+++ b/Engine/Source/Runtime/Renderer/Private/ShadowRendering.cpp
@@ -1236,9 +1236,8 @@ void FProjectedShadowInfo::UpdateShaderDepthBias()
 
 		// the z range is adjusted to we need to adjust here as well
 		DepthBias = CVarCSMShadowDepthBias.GetValueOnRenderThread() / (MaxSubjectZ - MinSubjectZ);
-
-		float WorldSpaceTexelScale = ShadowBounds.W / ResolutionX;
-		DepthBias *= WorldSpaceTexelScale;
+		const float WorldSpaceTexelScale = ShadowBounds.W / ResolutionX;
+		DepthBias = FMath::Lerp(DepthBias, DepthBias * WorldSpaceTexelScale, CascadeSettings.CascadeBiasDistribution);
 		DepthBias *= LightSceneInfo->Proxy->GetUserShadowBias();
 
 		SlopeScaleDepthBias = CVarCSMShadowSlopeScaleDepthBias.GetValueOnRenderThread();
