Fix crash when compiling runtime VT material without feature level support

diff --git a/Engine/Source/Runtime/Engine/Private/Materials/HLSLMaterialTranslator.h b/Engine/Source/Runtime/Engine/Private/Materials/HLSLMaterialTranslator.h
--- a/Engine/Source/Runtime/Engine/Private/Materials/HLSLMaterialTranslator.h
+++ b/Engine/Source/Runtime/Engine/Private/Materials/HLSLMaterialTranslator.h
@@ -4817,20 +4817,23 @@ protected:
 
 	virtual int32 VirtualTextureUnpack(int32 CodeIndex, EVirtualTextureUnpackType UnpackType) override
 	{
-		if (UnpackType == EVirtualTextureUnpackType::NormalBC3)
+		if (CodeIndex != INDEX_NONE)
 		{
-			FString	SampleCode(TEXT("VirtualTextureUnpackNormalBC3(%s)"));
-			return AddCodeChunk(MCT_Float3, *SampleCode, *GetParameterCode(CodeIndex));
-		}
-		if (UnpackType == EVirtualTextureUnpackType::NormalBC5)
-		{
-			FString	SampleCode(TEXT("VirtualTextureUnpackNormalBC5(%s)"));
-			return AddCodeChunk(MCT_Float3, *SampleCode, *GetParameterCode(CodeIndex));
-		}
-		else if (UnpackType == EVirtualTextureUnpackType::HeightR8G8)
-		{
-			FString	SampleCode(TEXT("VirtualTextureUnpackHeightR8G8(%s)"));
-			return AddCodeChunk(MCT_Float, *SampleCode, *GetParameterCode(CodeIndex));
+			if (UnpackType == EVirtualTextureUnpackType::NormalBC3)
+			{
+				FString	SampleCode(TEXT("VirtualTextureUnpackNormalBC3(%s)"));
+				return AddCodeChunk(MCT_Float3, *SampleCode, *GetParameterCode(CodeIndex));
+			}
+			if (UnpackType == EVirtualTextureUnpackType::NormalBC5)
+			{
+				FString	SampleCode(TEXT("VirtualTextureUnpackNormalBC5(%s)"));
+				return AddCodeChunk(MCT_Float3, *SampleCode, *GetParameterCode(CodeIndex));
+			}
+			else if (UnpackType == EVirtualTextureUnpackType::HeightR8G8)
+			{
+				FString	SampleCode(TEXT("VirtualTextureUnpackHeightR8G8(%s)"));
+				return AddCodeChunk(MCT_Float, *SampleCode, *GetParameterCode(CodeIndex));
+			}
 		}
 
 		return CodeIndex;
