Prevent some log spam generated by compiling VT enabled materials

diff --git a/Engine/Source/Runtime/Engine/Private/Materials/HLSLMaterialTranslator.h b/Engine/Source/Runtime/Engine/Private/Materials/HLSLMaterialTranslator.h
--- a/Engine/Source/Runtime/Engine/Private/Materials/HLSLMaterialTranslator.h
+++ b/Engine/Source/Runtime/Engine/Private/Materials/HLSLMaterialTranslator.h
@@ -724,6 +724,12 @@ public:
 		for (uint32 PropertyIndex = 0u; PropertyIndex < MP_MAX; ++PropertyIndex)
 		{
 			const EMaterialProperty PropertyToValidate = (EMaterialProperty)PropertyIndex;
+			if (PropertyToValidate == MP_MaterialAttributes || PropertyToValidate == MP_CustomOutput)
+			{
+				// These properties are "special", attempting to pass them to FMaterialAttributeDefinitionMap::GetShaderFrequency() will generate log spam
+				continue;
+			}
+
 			const EShaderFrequency ShaderFrequencyToValidate = FMaterialAttributeDefinitionMap::GetShaderFrequency(PropertyToValidate);
 
 			// check to see if this is a property that doesn't support virtual texture connections
