Marked RHICopySubTextureRegion deprecated for UE-4.23.

#rnx

diff --git a/Engine/Source/Runtime/Engine/Private/Texture2D.cpp b/Engine/Source/Runtime/Engine/Private/Texture2D.cpp
--- a/Engine/Source/Runtime/Engine/Private/Texture2D.cpp
+++ b/Engine/Source/Runtime/Engine/Private/Texture2D.cpp
@@ -1974,8 +1974,9 @@ void FVirtualTexture2DResource::InitializeEditorResources(IVirtualTexture* InVir
 			// Logical dimensions of mip image may be smaller than tile size (in this case tile will contain mirrored/wrapped padding)
 			// In this case, copy the proper sub-image from the tiled texture we produced into a new texture of the correct size
 			FTexture2DRHIRef ResizedTexture2DRHI = RHICreateTexture2D(MipWidth, MipHeight, PixelFormat, 1, 1, TexCreateFlags, CreateInfo);
-			const FBox2D CopyBox(FVector2D(0.0f, 0.0f), FVector2D(MipWidth, MipHeight));
-			RHICommandList.CopySubTextureRegion(Texture2DRHI, ResizedTexture2DRHI, CopyBox, CopyBox);
+			FRHICopyTextureInfo CopyInfo;
+			CopyInfo.Size = FIntVector(MipWidth, MipHeight, 1);
+			RHICommandList.CopyTexture(Texture2DRHI, ResizedTexture2DRHI, CopyInfo);
 			Texture2DRHI = MoveTemp(ResizedTexture2DRHI);
 		}
 
diff --git a/Engine/Source/Runtime/Engine/Private/VT/VirtualTextureUploadCache.cpp b/Engine/Source/Runtime/Engine/Private/VT/VirtualTextureUploadCache.cpp
--- a/Engine/Source/Runtime/Engine/Private/VT/VirtualTextureUploadCache.cpp
+++ b/Engine/Source/Runtime/Engine/Private/VT/VirtualTextureUploadCache.cpp
@@ -155,12 +155,15 @@ void FVirtualTextureUploadCache::Finalize(FRHICommandListImmediate& RHICmdList)
 
 			const uint32 SkipBorderSize = Entry.SubmitSkipBorderSize;
 			const uint32 SubmitTileSize = TileSize - SkipBorderSize * 2;
-			const FVector2D SourceBoxStart(SrcTileX * TileSize + SkipBorderSize, SrcTileY * TileSize + SkipBorderSize);
-			const FVector2D DestinationBoxStart(Entry.SubmitDestX * SubmitTileSize, Entry.SubmitDestY * SubmitTileSize);
-			const FBox2D SourceBox(SourceBoxStart, SourceBoxStart + FVector2D(SubmitTileSize, SubmitTileSize));
-			const FBox2D DestinationBox(DestinationBoxStart, DestinationBoxStart + FVector2D(SubmitTileSize, SubmitTileSize));
+			const FIntVector SourceBoxStart(SrcTileX * TileSize + SkipBorderSize, SrcTileY * TileSize + SkipBorderSize, 0);
+			const FIntVector DestinationBoxStart(Entry.SubmitDestX * SubmitTileSize, Entry.SubmitDestY * SubmitTileSize, 0);
+
+			FRHICopyTextureInfo CopyInfo;
+			CopyInfo.Size = FIntVector(SubmitTileSize, SubmitTileSize, 1);
+			CopyInfo.SourcePosition = SourceBoxStart;
+			CopyInfo.DestPosition = DestinationBoxStart;
+			RHICmdList.CopyTexture(StagingTexture.RHITexture, Entry.RHISubmitTexture, CopyInfo);
 
-			RHICmdList.CopySubTextureRegion(StagingTexture.RHITexture, Entry.RHISubmitTexture, SourceBox, DestinationBox);
 			Entry.RHISubmitTexture = nullptr;
 			Entry.SubmitBatchIndex = 0u;
 			Entry.SubmitDestX = 0;
diff --git a/Engine/Source/Runtime/RHI/Public/RHICommandList.h b/Engine/Source/Runtime/RHI/Public/RHICommandList.h
--- a/Engine/Source/Runtime/RHI/Public/RHICommandList.h
+++ b/Engine/Source/Runtime/RHI/Public/RHICommandList.h
@@ -4273,6 +4273,7 @@ public:
 		GDynamicRHI->VirtualTextureSetFirstMipVisible_RenderThread(*this, Texture, FirstMip);
 	}
 
+	UE_DEPRECATED(4.23, "CopySubTextureRegion API is deprecated; please use CopyTexture instead.")
 	FORCEINLINE void CopySubTextureRegion(FTexture2DRHIParamRef SourceTexture, FTexture2DRHIParamRef DestinationTexture, FBox2D SourceBox, FBox2D DestinationBox)
 	{
 		GDynamicRHI->RHICopySubTextureRegion_RenderThread(*this, SourceTexture, DestinationTexture, SourceBox, DestinationBox);
