Move RuntimeVirtualTexturePlane out of Landscape library and into Engine

diff --git a/Engine/Source/Editor/VirtualTexturingEditor/Private/RuntimeVirtualTextureDetailsCustomization.cpp b/Engine/Source/Editor/VirtualTexturingEditor/Private/RuntimeVirtualTextureDetailsCustomization.cpp
--- a/Engine/Source/Editor/VirtualTexturingEditor/Private/RuntimeVirtualTextureDetailsCustomization.cpp
+++ b/Engine/Source/Editor/VirtualTexturingEditor/Private/RuntimeVirtualTextureDetailsCustomization.cpp
@@ -5,12 +5,12 @@
 #include "DetailCategoryBuilder.h"
 #include "DetailLayoutBuilder.h"
 #include "DetailWidgetRow.h"
-#include "RuntimeVirtualTexturePlane.h"
 #include "SResetToDefaultMenu.h"
 #include "VT/RuntimeVirtualTexture.h"
-#include "Widgets/Layout/SWrapBox.h"
-#include "Widgets/Layout/SBox.h"
+#include "VT/RuntimeVirtualTexturePlane.h"
 #include "Widgets/Input/SButton.h"
+#include "Widgets/Layout/SBox.h"
+#include "Widgets/Layout/SWrapBox.h"
 
 #define LOCTEXT_NAMESPACE "VirtualTexturingEditorModule"
 
diff --git a/Engine/Source/Editor/VirtualTexturingEditor/Private/RuntimeVirtualTexturePlaneActorFactory.cpp b/Engine/Source/Editor/VirtualTexturingEditor/Private/RuntimeVirtualTexturePlaneActorFactory.cpp
--- a/Engine/Source/Editor/VirtualTexturingEditor/Private/RuntimeVirtualTexturePlaneActorFactory.cpp
+++ b/Engine/Source/Editor/VirtualTexturingEditor/Private/RuntimeVirtualTexturePlaneActorFactory.cpp
@@ -2,7 +2,7 @@
 
 #include "RuntimeVirtualTexturePlaneActorFactory.h"
 
-#include "RuntimeVirtualTexturePlane.h"
+#include "VT/RuntimeVirtualTexturePlane.h"
 
 #define LOCTEXT_NAMESPACE "VirtualTexturingEditorModule"
 
diff --git a/Engine/Source/Editor/VirtualTexturingEditor/VirtualTexturingEditor.Build.cs b/Engine/Source/Editor/VirtualTexturingEditor/VirtualTexturingEditor.Build.cs
--- a/Engine/Source/Editor/VirtualTexturingEditor/VirtualTexturingEditor.Build.cs
+++ b/Engine/Source/Editor/VirtualTexturingEditor/VirtualTexturingEditor.Build.cs
@@ -21,21 +21,20 @@ public class VirtualTexturingEditor : ModuleRules
 
         PrivateDependencyModuleNames.AddRange(
             new string[] {
-                "AppFramework",
-                "AssetRegistry",
-                "ContentBrowser",
-                "Core",
-                "CoreUObject",
-                "DesktopPlatform",
-                "DesktopWidgets",
-                "EditorStyle",
-                "Engine",
-                "InputCore",
-                "Landscape",
-		"PropertyEditor",
-                "Slate",
-                "SlateCore",
-                "UnrealEd",
+				"AppFramework",
+				"AssetRegistry",
+				"ContentBrowser",
+				"Core",
+				"CoreUObject",
+				"DesktopPlatform",
+				"DesktopWidgets",
+				"EditorStyle",
+				"Engine",
+				"InputCore",
+				"PropertyEditor",
+				"Slate",
+				"SlateCore",
+				"UnrealEd",
             }
         );
     }
diff --git a/Engine/Source/Runtime/Engine/Classes/VT/RuntimeVirtualTexture.h b/Engine/Source/Runtime/Engine/Classes/VT/RuntimeVirtualTexture.h
--- a/Engine/Source/Runtime/Engine/Classes/VT/RuntimeVirtualTexture.h
+++ b/Engine/Source/Runtime/Engine/Classes/VT/RuntimeVirtualTexture.h
@@ -115,12 +115,6 @@ public:
 	/** Getter for the shader uniform parameters. */
 	FVector4 GetUniformParameter(int32 Index);
 
-#if WITH_EDITOR
-	/** Delegate to broadcast property changes. Use this so that the IVirtualTexture can be updated accordingly. */
-	DECLARE_DELEGATE_OneParam(FOnEditProperty, URuntimeVirtualTexture const*);
-	FOnEditProperty OnEditProperty;
-#endif
-
 protected:
 	/** Initialize the render resource. This kicks off render thread work. */
 	void InitResource(IVirtualTexture* InProducer);
@@ -131,10 +125,10 @@ protected:
 
 	//~ Begin UObject Interface.
 	virtual void BeginDestroy() override;
+	virtual void GetAssetRegistryTags(TArray<FAssetRegistryTag>& OutTags) const override;
 #if WITH_EDITOR
 	virtual void PostEditChangeProperty(FPropertyChangedEvent& PropertyChangedEvent) override;
 #endif
-	virtual void GetAssetRegistryTags(TArray<FAssetRegistryTag>& OutTags) const override;
 	//~ End UObject Interface.
 
 private:
diff --git a/Engine/Source/Runtime/Engine/Classes/VT/RuntimeVirtualTexturePlane.h b/Engine/Source/Runtime/Engine/Classes/VT/RuntimeVirtualTexturePlane.h
new file mode 100644
--- /dev/null
+++ b/Engine/Source/Runtime/Engine/Classes/VT/RuntimeVirtualTexturePlane.h
@@ -0,0 +1,76 @@
+// Copyright 1998-2019 Epic Games, Inc. All Rights Reserved.
+
+#pragma once
+
+#include "CoreMinimal.h"
+#include "GameFramework/Actor.h"
+#include "RuntimeVirtualTexturePlane.generated.h"
+
+class UMaterialInterface;
+class URuntimeVirtualTexture;
+class URuntimeVirtualTextureComponent;
+
+/** Actor used to place a URuntimeVirtualTexture in the world. */
+UCLASS(hidecategories=(Actor, Collision, Cooking, Input, LOD, Replication), MinimalAPI)
+class ARuntimeVirtualTexturePlane : public AActor
+{
+	GENERATED_UCLASS_BODY()
+
+private:
+	/** Component that owns the runtime virtual texture. */
+	UPROPERTY(VisibleAnywhere, BlueprintReadOnly, Category = VirtualTexture, meta = (AllowPrivateAccess = "true"))
+	class URuntimeVirtualTextureComponent* VirtualTextureComponent;
+
+#if WITH_EDITORONLY_DATA
+	/** Box for visualizing virtual texture extents. */
+	UPROPERTY(Transient)
+	class UBoxComponent* Box = nullptr;
+#endif // WITH_EDITORONLY_DATA
+
+protected:
+	//~ Begin UObject Interface.
+	virtual bool NeedsLoadForServer() const override { return false; }
+	//~ End UObject Interface.
+	//~ Begin AActor Interface.
+	virtual bool IsLevelBoundsRelevant() const override { return false; }
+	//~ End AActor Interface
+};
+
+/** Component used to place a URuntimeVirtualTexture in the world. */
+UCLASS(ClassGroup = Rendering, collapsecategories, hidecategories = (Activation, Collision, Cooking, Mobility, LOD, Object, Physics, Rendering), editinlinenew)
+class ENGINE_API URuntimeVirtualTextureComponent : public USceneComponent
+{
+	GENERATED_UCLASS_BODY()
+
+private:
+	/** The virtual texture object to use. */
+	UPROPERTY(EditAnywhere, DuplicateTransient, Category = VirtualTexture)
+	URuntimeVirtualTexture* VirtualTexture = nullptr;
+
+	/** Actor to copy the bounds from to set up the transform. */
+	UPROPERTY(EditAnywhere, DuplicateTransient, Category = TransformFromBounds, meta = (DisplayName = "Source Actor"))
+	AActor* BoundsSourceActor = nullptr;
+
+public:
+	/** Get the runtime virtual texture object set on this component */
+	URuntimeVirtualTexture* GetVirtualTexture() const { return VirtualTexture; }
+
+#if WITH_EDITOR
+	/** Copy the rotation from BoundsSourceActor to this component. Called by our UI details customization. */
+	void SetRotation();
+
+	/** Set this component transform to include the BoundsSourceActor bounds. Called by our UI details customization. */
+	void SetTransformToBounds();
+#endif
+
+protected:
+	//~ Begin UActorComponent Interface
+	virtual void CreateRenderState_Concurrent() override;
+	virtual void SendRenderTransform_Concurrent() override;
+	virtual void DestroyRenderState_Concurrent() override;
+	//~ End UActorComponent Interface
+
+public:
+	/** Scene proxy object. Managed by the scene but stored here. */
+	class FRuntimeVirtualTextureSceneProxy* SceneProxy;
+};
diff --git a/Engine/Source/Runtime/Engine/Private/VT/RuntimeVirtualTexture.cpp b/Engine/Source/Runtime/Engine/Private/VT/RuntimeVirtualTexture.cpp
--- a/Engine/Source/Runtime/Engine/Private/VT/RuntimeVirtualTexture.cpp
+++ b/Engine/Source/Runtime/Engine/Private/VT/RuntimeVirtualTexture.cpp
@@ -7,6 +7,7 @@
 #include "Materials/Material.h"
 #include "RendererInterface.h"
 #include "UObject/UObjectIterator.h"
+#include "VT/RuntimeVirtualTexturePlane.h"
 #include "VT/VirtualTextureUtility.h"
 
 
@@ -171,7 +172,7 @@ void URuntimeVirtualTexture::ReleaseResource()
 {
 	if (Resource != nullptr)
 	{
-		ReleaseResourceAndFlush(Resource);
+		BeginReleaseResource(Resource);
 		Resource = nullptr;
 	}
 }
@@ -182,17 +183,6 @@ void URuntimeVirtualTexture::BeginDestroy()
 	Super::BeginDestroy();
 }
 
-#if WITH_EDITOR
-
-void URuntimeVirtualTexture::PostEditChangeProperty(FPropertyChangedEvent& PropertyChangedEvent)
-{
-	Super::PostEditChangeProperty(PropertyChangedEvent);
-
-	OnEditProperty.ExecuteIfBound(this);
-}
-
-#endif
-
 void URuntimeVirtualTexture::GetAssetRegistryTags(TArray<FAssetRegistryTag>& OutTags) const
 {
 	Super::GetAssetRegistryTags(OutTags);
@@ -203,12 +193,34 @@ void URuntimeVirtualTexture::GetAssetRegistryTags(TArray<FAssetRegistryTag>& Out
 	OutTags.Add(FAssetRegistryTag("TileBorderSize", FString::FromInt(GetTileBorderSize()), FAssetRegistryTag::TT_Numerical));
 }
 
+#if WITH_EDITOR
+
+void URuntimeVirtualTexture::PostEditChangeProperty(FPropertyChangedEvent& PropertyChangedEvent)
+{
+	Super::PostEditChangeProperty(PropertyChangedEvent);
+
+	// Find any components that reference this object and mark them dirty.
+	for (TObjectIterator<URuntimeVirtualTextureComponent> It; It; ++It)
+	{
+		if (It->GetVirtualTexture() == this)
+		{
+			It->MarkRenderStateDirty();
+		}
+	}
+}
+
+#endif
+
 void URuntimeVirtualTexture::NotifyMaterials()
 {
-	//todo[vt]: This can be expensive. We need to make sure we don't call this too often (or at all?) in non-editor builds?
-	
-	// In non-editor builds we can come here > 1 times during PostLoad callbacks.
-	// Is there a way to express the dependency so that materials are init'd after the VT is setup and then we don't need this at all?
+	// Find materials referencing this virtual texture and re-cache the uniforms
+	// We need to do this after any operation that reallocates the virtual texture since the material caches info about the VT allocation in it's uniform buffer
+
+	// todo [vt]:
+	//  This operation is very slow! We should only do it during edition, but currently we do it on runtime load/unload.
+	//  Can we pre-calculate the list of materials to touch during cook? But even then touching here them will be extra work.
+	//  Is there a way to set up dependencies so that we load the materials _after_ the virtual texture is allocated?
+	//  Or maybe we should consider serializing the WorldToUVTransform here and not depending on a URuntimeVirtualTextureComponent at runtime?
 
 	TSet<UMaterial*> BaseMaterialsThatUseThisTexture;
 	for (TObjectIterator<UMaterialInterface> It; It; ++It)
@@ -230,7 +242,7 @@ void URuntimeVirtualTexture::NotifyMaterials()
 
 	if (BaseMaterialsThatUseThisTexture.Num() > 0)
 	{
-		FMaterialUpdateContext UpdateContext;
+		FMaterialUpdateContext UpdateContext(FMaterialUpdateContext::EOptions::SyncWithRenderingThread);
 		for (TSet<UMaterial*>::TConstIterator It(BaseMaterialsThatUseThisTexture); It; ++It)
 		{
 			(*It)->RecacheUniformExpressions(false);
diff --git a/Engine/Source/Runtime/Engine/Private/VT/RuntimeVirtualTexturePlane.cpp b/Engine/Source/Runtime/Engine/Private/VT/RuntimeVirtualTexturePlane.cpp
new file mode 100644
--- /dev/null
+++ b/Engine/Source/Runtime/Engine/Private/VT/RuntimeVirtualTexturePlane.cpp
@@ -0,0 +1,116 @@
+// Copyright 1998-2019 Epic Games, Inc. All Rights Reserved.
+
+#include "VT/RuntimeVirtualTexturePlane.h"
+
+#include "Components/BoxComponent.h"
+
+
+ARuntimeVirtualTexturePlane::ARuntimeVirtualTexturePlane(const FObjectInitializer& ObjectInitializer)
+	: Super(ObjectInitializer)
+{
+	RootComponent = VirtualTextureComponent = CreateDefaultSubobject<URuntimeVirtualTextureComponent>(TEXT("VirtualTextureComponent"));
+
+#if WITH_EDITORONLY_DATA
+	// Add box for visualization of bounds
+	Box = CreateDefaultSubobject<UBoxComponent>(TEXT("Box"));
+	Box->SetBoxExtent(FVector(0.5f, 0.5f, 1.f), false);
+	Box->SetIsVisualizationComponent(true);
+	Box->SetCollisionEnabled(ECollisionEnabled::NoCollision);
+	Box->SetCanEverAffectNavigation(false);
+	Box->CanCharacterStepUpOn = ECanBeCharacterBase::ECB_No;
+	Box->SetGenerateOverlapEvents(false);
+	Box->SetupAttachment(VirtualTextureComponent);
+#endif
+}
+
+URuntimeVirtualTextureComponent::URuntimeVirtualTextureComponent(const FObjectInitializer& ObjectInitializer)
+	: Super(ObjectInitializer)
+	, SceneProxy(nullptr)
+{
+}
+
+void URuntimeVirtualTextureComponent::CreateRenderState_Concurrent()
+{
+	if (ShouldRender() && VirtualTexture != nullptr)
+	{
+		// This will internally allocate modify the URuntimeVirtualTexture and allocate its VT
+		GetScene()->AddRuntimeVirtualTexture(this);
+	}
+
+	Super::CreateRenderState_Concurrent();
+}
+
+void URuntimeVirtualTextureComponent::SendRenderTransform_Concurrent()
+{
+	if (ShouldRender() && VirtualTexture != nullptr)
+	{
+		// This will internally allocate modify the URuntimeVirtualTexture and allocate its VT
+		GetScene()->AddRuntimeVirtualTexture(this);
+	}
+
+	Super::SendRenderTransform_Concurrent();
+}
+
+void URuntimeVirtualTextureComponent::DestroyRenderState_Concurrent()
+{
+	// This will internally allocate modify the URuntimeVirtualTexture and free its VT
+	GetScene()->RemoveRuntimeVirtualTexture(this);
+
+	Super::DestroyRenderState_Concurrent();
+}
+
+#if WITH_EDITOR
+
+void URuntimeVirtualTextureComponent::SetRotation()
+{
+	if (BoundsSourceActor != nullptr)
+	{
+		// Copy the source actor rotation and notify the parent actor
+		SetWorldRotation(BoundsSourceActor->GetTransform().GetRotation());
+		GetOwner()->PostEditMove(true);
+	}
+}
+
+void URuntimeVirtualTextureComponent::SetTransformToBounds()
+{
+	if (BoundsSourceActor != nullptr)
+	{
+		// Calculate the bounds in our local rotation space translated to the BoundsSourceActor center
+		const FQuat TargetRotation = GetComponentToWorld().GetRotation();
+		const FVector InitialPosition = BoundsSourceActor->GetComponentsBoundingBox().GetCenter();
+		const FVector InitialScale = FVector(0.5f, 0.5, 1.f);
+
+		FTransform LocalTransform;
+		LocalTransform.SetComponents(TargetRotation, InitialPosition, InitialScale);
+		FTransform WorldToLocal = LocalTransform.Inverse();
+
+		FBox BoundBox(ForceInit);
+		for (const UActorComponent* Component : BoundsSourceActor->GetComponents())
+		{
+			// Only gather visual components in the bounds calculation
+			const UPrimitiveComponent* PrimitiveComponent = Cast<const UPrimitiveComponent>(Component);
+			if (PrimitiveComponent != nullptr && PrimitiveComponent->IsRegistered())
+			{
+				const FTransform ComponentToActor = PrimitiveComponent->GetComponentTransform() * WorldToLocal;
+				FBoxSphereBounds LocalSpaceComponentBounds = PrimitiveComponent->CalcBounds(ComponentToActor);
+				BoundBox += LocalSpaceComponentBounds.GetBox();
+			}
+		}
+
+		// Create transform from bounds
+		FVector Origin;
+		FVector Extent;
+		BoundBox.GetCenterAndExtents(Origin, Extent);
+
+		Origin = LocalTransform.TransformPosition(Origin);
+
+		FTransform Transform;
+		Transform.SetComponents(TargetRotation, Origin, Extent);
+
+		// Apply final result and notify the parent actor
+		SetWorldTransform(Transform);
+		GetOwner()->PostEditMove(true);
+	}
+}
+
+#endif
diff --git a/Engine/Source/Runtime/Engine/Public/SceneInterface.h b/Engine/Source/Runtime/Engine/Public/SceneInterface.h
--- a/Engine/Source/Runtime/Engine/Public/SceneInterface.h
+++ b/Engine/Source/Runtime/Engine/Public/SceneInterface.h
@@ -162,6 +162,12 @@ public:
 	virtual void AddPrecomputedVolumetricLightmap(const class FPrecomputedVolumetricLightmap* Volume) {}
 	virtual void RemovePrecomputedVolumetricLightmap(const class FPrecomputedVolumetricLightmap* Volume) {}
 
+	/** Add a runtime virtual texture object to the scene. */
+	virtual void AddRuntimeVirtualTexture(class URuntimeVirtualTextureComponent* Component) {}
+
+	/** Removes a runtime virtual texture object from the scene. */
+	virtual void RemoveRuntimeVirtualTexture(class URuntimeVirtualTextureComponent* Component) {}
+
 	/** 
 	 * Retrieves primitive uniform shader parameters that are internal to the renderer.
 	 */
diff --git a/Engine/Source/Runtime/Landscape/Classes/RuntimeVirtualTexturePlane.h b/Engine/Source/Runtime/Landscape/Classes/RuntimeVirtualTexturePlane.h
deleted file mode 100644
diff --git a/Engine/Source/Runtime/Landscape/Private/RuntimeVirtualTexturePlane.cpp b/Engine/Source/Runtime/Landscape/Private/RuntimeVirtualTexturePlane.cpp
deleted file mode 100644
diff --git a/Engine/Source/Runtime/Renderer/Private/RendererScene.cpp b/Engine/Source/Runtime/Renderer/Private/RendererScene.cpp
--- a/Engine/Source/Runtime/Renderer/Private/RendererScene.cpp
+++ b/Engine/Source/Runtime/Renderer/Private/RendererScene.cpp
@@ -54,6 +54,8 @@
 #include "DynamicShadowMapChannelBindingHelper.h"
 #include "GPUScene.h"
 #include "HAL/LowLevelMemTracker.h"
+#include "VT/RuntimeVirtualTexturePlane.h"
+#include "VT/RuntimeVirtualTextureSceneProxy.h"
 #if RHI_RAYTRACING
 #include "RayTracingDynamicGeometryCollection.h"
 #endif
@@ -2359,6 +2361,18 @@ void FScene::RemovePrecomputedVolumetricLightmap(const FPrecomputedVolumetricLig
 		});
 }
 
+void FScene::AddRuntimeVirtualTexture(class URuntimeVirtualTextureComponent* Component)
+{
+	delete Component->SceneProxy;
+	Component->SceneProxy = new FRuntimeVirtualTextureSceneProxy(Component);
+}
+
+void FScene::RemoveRuntimeVirtualTexture(class URuntimeVirtualTextureComponent* Component)
+{
+	delete Component->SceneProxy;
+	Component->SceneProxy = nullptr;
+}
+
 bool FScene::GetPreviousLocalToWorld(const FPrimitiveSceneInfo* PrimitiveSceneInfo, FMatrix& OutPreviousLocalToWorld) const
 {
 	return VelocityData.GetComponentPreviousLocalToWorld(PrimitiveSceneInfo->PrimitiveComponentId, OutPreviousLocalToWorld);
diff --git a/Engine/Source/Runtime/Renderer/Private/ScenePrivate.h b/Engine/Source/Runtime/Renderer/Private/ScenePrivate.h
--- a/Engine/Source/Runtime/Renderer/Private/ScenePrivate.h
+++ b/Engine/Source/Runtime/Renderer/Private/ScenePrivate.h
@@ -2624,6 +2624,8 @@ public:
 	virtual bool HasPrecomputedVolumetricLightmap_RenderThread() const override;
 	virtual void AddPrecomputedVolumetricLightmap(const class FPrecomputedVolumetricLightmap* Volume) override;
 	virtual void RemovePrecomputedVolumetricLightmap(const class FPrecomputedVolumetricLightmap* Volume) override;
+	virtual void AddRuntimeVirtualTexture(class URuntimeVirtualTextureComponent* Component) override;
+	virtual void RemoveRuntimeVirtualTexture(class URuntimeVirtualTextureComponent* Component) override;
 	virtual void GetPrimitiveUniformShaderParameters_RenderThread(const FPrimitiveSceneInfo* PrimitiveSceneInfo, bool& bHasPrecomputedVolumetricLightmap, FMatrix& PreviousLocalToWorld, int32& SingleCaptureIndex, bool& bOutputVelocity) const override;
 	virtual void UpdateLightTransform(ULightComponent* Light) override;
 	virtual void UpdateLightColorAndBrightness(ULightComponent* Light) override;
diff --git a/Engine/Source/Runtime/Renderer/Private/VT/RuntimeVirtualTextureProducer.h b/Engine/Source/Runtime/Renderer/Private/VT/RuntimeVirtualTextureProducer.h
new file mode 100644
--- /dev/null
+++ b/Engine/Source/Runtime/Renderer/Private/VT/RuntimeVirtualTextureProducer.h
@@ -0,0 +1,87 @@
+// Copyright 1998-2019 Epic Games, Inc. All Rights Reserved.
+
+#pragma once
+
+#include "CoreMinimal.h"
+#include "UObject/ObjectMacros.h"
+#include "VirtualTexturing.h"
+
+enum class ERuntimeVirtualTextureMaterialType;
+class FRHITexture2D;
+class FSceneInterface;
+
+/** IVirtualTextureFinalizer implementation that renders the virtual texture pages on demand. */
+class FRuntimeVirtualTextureFinalizer : public IVirtualTextureFinalizer
+{
+public:
+	FRuntimeVirtualTextureFinalizer(FVTProducerDescription const& InDesc, ERuntimeVirtualTextureMaterialType InMaterialType, FSceneInterface* InScene, FTransform const& InUVToWorld);
+	virtual ~FRuntimeVirtualTextureFinalizer() {}
+
+	/** A description for a single tile to render. */
+	struct FTileEntry
+	{
+		FRHITexture2D* Texture0 = nullptr;
+		FRHITexture2D* Texture1 = nullptr;
+		int32 DestX0 = 0;
+		int32 DestY0 = 0;
+		int32 DestX1 = 0;
+		int32 DestY1 = 0;
+		uint32 vAddress = 0;
+		uint8 vLevel = 0;
+	};
+
+	/** Returns false if we don't yet have everything we need to render a VT page. */
+	bool IsReady();
+
+	/** Add a tile to the finalize queue. */
+	void AddTile(FTileEntry& Tile);
+
+	//~ Begin IVirtualTextureFinalizer Interface.
+	virtual void Finalize(FRHICommandListImmediate& RHICmdList) override;
+	//~ End IVirtualTextureFinalizer Interface.
+
+private:
+	/** Description of our virtual texture. */
+	const FVTProducerDescription Desc;
+	/** Contents of virtual texture layer stack. */
+	ERuntimeVirtualTextureMaterialType MaterialType;
+	/** Scene that the virtual texture is placed within. */
+	FSceneInterface* Scene;
+	/** Transform from UV space to world space. */
+	FTransform UVToWorld;
+	/** Array of tiles in the queue to finalize. */
+	TArray<FTileEntry> Tiles;
+};
+
+/** IVirtualTexture implementation that is handling runtime rendered page data requests. */
+class FRuntimeVirtualTextureProducer : public IVirtualTexture
+{
+public:
+	RENDERER_API FRuntimeVirtualTextureProducer(FVTProducerDescription const& InDesc, ERuntimeVirtualTextureMaterialType InMaterialType, FSceneInterface* InScene, FTransform const& InUVToWorld);
+	RENDERER_API virtual ~FRuntimeVirtualTextureProducer() {}
+
+	//~ Begin IVirtualTexture Interface.
+	virtual FVTRequestPageResult RequestPageData(
+		const FVirtualTextureProducerHandle& ProducerHandle,
+		uint8 LayerMask,
+		uint8 vLevel,
+		uint32 vAddress,
+		EVTRequestPagePriority Priority
+	) override;
+
+	virtual IVirtualTextureFinalizer* ProducePageData(
+		FRHICommandListImmediate& RHICmdList,
+		ERHIFeatureLevel::Type FeatureLevel,
+		EVTProducePageFlags Flags,
+		const FVirtualTextureProducerHandle& ProducerHandle,
+		uint8 LayerMask,
+		uint8 vLevel,
+		uint32 vAddress,
+		uint64 RequestHandle,
+		const FVTProduceTargetLayer* TargetLayers
+	) override;
+	//~ End IVirtualTexture Interface.
+
+private:
+	FRuntimeVirtualTextureFinalizer Finalizer;
+};
diff --git a/Engine/Source/Runtime/Renderer/Private/VT/RuntimeVirtualTextureSceneProxy.cpp b/Engine/Source/Runtime/Renderer/Private/VT/RuntimeVirtualTextureSceneProxy.cpp
new file mode 100644
--- /dev/null
+++ b/Engine/Source/Runtime/Renderer/Private/VT/RuntimeVirtualTextureSceneProxy.cpp
@@ -0,0 +1,35 @@
+// Copyright 1998-2019 Epic Games, Inc. All Rights Reserved.
+
+#include "VT/RuntimeVirtualTextureSceneProxy.h"
+
+#include "VT/RuntimeVirtualTexture.h"
+#include "VT/RuntimeVirtualTexturePlane.h"
+#include "VT/RuntimeVirtualTextureProducer.h"
+
+
+FRuntimeVirtualTextureSceneProxy::FRuntimeVirtualTextureSceneProxy(URuntimeVirtualTextureComponent* InComponent)
+	: VirtualTexture(InComponent->GetVirtualTexture())
+{
+	if (VirtualTexture != nullptr)
+	{
+		// The Producer object created here will be passed into the Virtual Texture system which will take ownership
+		FVTProducerDescription Desc;
+		VirtualTexture->GetProducerDescription(Desc);
+
+		const ERuntimeVirtualTextureMaterialType MaterialType = VirtualTexture->GetMaterialType();
+
+		// Transform is based on bottom left of the URuntimeVirtualTextureComponent unit box (which is centered on the origin)
+		FTransform Transform = FTransform(FVector(-0.5f, -0.5f, 0.f)) * InComponent->GetComponentTransform();
+
+		FRuntimeVirtualTextureProducer* Producer = new FRuntimeVirtualTextureProducer(Desc, MaterialType, InComponent->GetScene(), Transform);
+		VirtualTexture->Initialize(Producer, Transform);
+	}
+}
+
+FRuntimeVirtualTextureSceneProxy::~FRuntimeVirtualTextureSceneProxy()
+{
+	if (VirtualTexture != nullptr)
+	{
+		VirtualTexture->Release();
+	}
+}
diff --git a/Engine/Source/Runtime/Renderer/Private/VT/RuntimeVirtualTextureSceneProxy.h b/Engine/Source/Runtime/Renderer/Private/VT/RuntimeVirtualTextureSceneProxy.h
new file mode 100644
--- /dev/null
+++ b/Engine/Source/Runtime/Renderer/Private/VT/RuntimeVirtualTextureSceneProxy.h
@@ -0,0 +1,16 @@
+// Copyright 1998-2019 Epic Games, Inc. All Rights Reserved.
+
+#pragma once
+
+#include "CoreMinimal.h"
+
+/** Scene proxy for the URuntimeVirtualTextureComponent. Manages a runtime virtual texture in the renderer scene. */
+class FRuntimeVirtualTextureSceneProxy
+{
+public:
+	FRuntimeVirtualTextureSceneProxy(URuntimeVirtualTextureComponent* InComponent);
+	~FRuntimeVirtualTextureSceneProxy();
+
+	/** Pointer to linked URuntimeVirtualTexture. Not for dereferencing, just for debug! */
+	class URuntimeVirtualTexture* VirtualTexture;
+};
diff --git a/Engine/Source/Runtime/Renderer/Public/RuntimeVirtualTextureProducer.h b/Engine/Source/Runtime/Renderer/Public/RuntimeVirtualTextureProducer.h
deleted file mode 100644
