Non-unity build fix

diff --git a/Engine/Source/Runtime/Engine/Private/Materials/MaterialShared.cpp b/Engine/Source/Runtime/Engine/Private/Materials/MaterialShared.cpp
--- a/Engine/Source/Runtime/Engine/Private/Materials/MaterialShared.cpp
+++ b/Engine/Source/Runtime/Engine/Private/Materials/MaterialShared.cpp
@@ -16,6 +16,7 @@
 #include "Materials/Material.h"
 #include "Materials/MaterialInstanceBasePropertyOverrides.h"
 #include "Materials/MaterialInstance.h"
+#include "Materials/MaterialUniformExpressions.h"
 #include "UObject/UObjectIterator.h"
 #include "ComponentReregisterContext.h"
 #include "Materials/MaterialExpressionBreakMaterialAttributes.h"
