Move runtime virtual texture material notification used during edition out of concurrent worker thread

diff --git a/Engine/Source/Runtime/Engine/Classes/VT/RuntimeVirtualTexture.h b/Engine/Source/Runtime/Engine/Classes/VT/RuntimeVirtualTexture.h
--- a/Engine/Source/Runtime/Engine/Classes/VT/RuntimeVirtualTexture.h
+++ b/Engine/Source/Runtime/Engine/Classes/VT/RuntimeVirtualTexture.h
@@ -120,8 +120,6 @@ protected:
 	void InitResource(IVirtualTexture* InProducer);
 	/** Release the render resource. This kicks off render thread work. */
 	void ReleaseResource();
-	/** Trigger an update for materials that reference this object. */
-	void NotifyMaterials();
 
 	//~ Begin UObject Interface.
 	virtual void BeginDestroy() override;
diff --git a/Engine/Source/Runtime/Engine/Classes/VT/RuntimeVirtualTexturePlane.h b/Engine/Source/Runtime/Engine/Classes/VT/RuntimeVirtualTexturePlane.h
--- a/Engine/Source/Runtime/Engine/Classes/VT/RuntimeVirtualTexturePlane.h
+++ b/Engine/Source/Runtime/Engine/Classes/VT/RuntimeVirtualTexturePlane.h
@@ -51,6 +51,9 @@ private:
 	UPROPERTY(EditAnywhere, DuplicateTransient, Category = TransformFromBounds, meta = (DisplayName = "Source Actor"))
 	AActor* BoundsSourceActor = nullptr;
 
+	/** Flag used for deferred material notification after render state changes. */
+	bool bNotifyInNextTick;
+
 public:
 	/** Get the runtime virtual texture object set on this component */
 	URuntimeVirtualTexture* GetVirtualTexture() const { return VirtualTexture; }
@@ -68,6 +71,7 @@ protected:
 	virtual void CreateRenderState_Concurrent() override;
 	virtual void SendRenderTransform_Concurrent() override;
 	virtual void DestroyRenderState_Concurrent() override;
+	virtual void TickComponent(float DeltaTime, ELevelTick TickType, FActorComponentTickFunction* ThisTickFunction) override;
 	//~ End UActorComponent Interface
 
 public:
diff --git a/Engine/Source/Runtime/Engine/Private/VT/RuntimeVirtualTexture.cpp b/Engine/Source/Runtime/Engine/Private/VT/RuntimeVirtualTexture.cpp
--- a/Engine/Source/Runtime/Engine/Private/VT/RuntimeVirtualTexture.cpp
+++ b/Engine/Source/Runtime/Engine/Private/VT/RuntimeVirtualTexture.cpp
@@ -3,12 +3,8 @@
 #include "VT/RuntimeVirtualTexture.h"
 
 #include "EngineModule.h"
-#include "MaterialShared.h"
-#include "Materials/Material.h"
 #include "RendererInterface.h"
-#include "UObject/UObjectIterator.h"
-#include "VT/RuntimeVirtualTexturePlane.h"
-#include "VT/VirtualTextureUtility.h"
+#include "VT/RuntimeVirtualTextureNotify.h"
 
 
 FRuntimeVirtualTextureRenderResource::FRuntimeVirtualTextureRenderResource(FVTProducerDescription const& InProducerDesc, IVirtualTexture* InVirtualTextureProducer)
@@ -147,13 +143,11 @@ void URuntimeVirtualTexture::Initialize(IVirtualTexture* InProducer, FTransform
 
 	ReleaseResource();
 	InitResource(InProducer);
-	NotifyMaterials();
 }
 
 void URuntimeVirtualTexture::Release()
 {
 	ReleaseResource();
-	NotifyMaterials();
 }
 
 void URuntimeVirtualTexture::InitResource(IVirtualTexture* InProducer)
@@ -199,54 +193,8 @@ void URuntimeVirtualTexture::PostEditChangeProperty(FPropertyChangedEvent& Prope
 {
 	Super::PostEditChangeProperty(PropertyChangedEvent);
 
-	// Find any components that reference this object and mark them dirty.
-	for (TObjectIterator<URuntimeVirtualTextureComponent> It; It; ++It)
-	{
-		if (It->GetVirtualTexture() == this)
-		{
-			It->MarkRenderStateDirty();
-		}
-	}
+	RuntimeVirtualTexture::NotifyComponents(this);
+	RuntimeVirtualTexture::NotifyPrimitives(this);
 }
 
 #endif
-
-void URuntimeVirtualTexture::NotifyMaterials()
-{
-	// Find materials referencing this virtual texture and re-cache the uniforms
-	// We need to do this after any operation that reallocates the virtual texture since the material caches info about the VT allocation in it's uniform buffer
-
-	// todo [vt]:
-	//  This operation is very slow! We should only do it during edition, but currently we do it on runtime load/unload.
-	//  Can we pre-calculate the list of materials to touch during cook? But even then touching here them will be extra work.
-	//  Is there a way to set up dependencies so that we load the materials _after_ the virtual texture is allocated?
-	//  Or maybe we should consider serializing the WorldToUVTransform here and not depending on a URuntimeVirtualTextureComponent at runtime?
-
-	TSet<UMaterial*> BaseMaterialsThatUseThisTexture;
-	for (TObjectIterator<UMaterialInterface> It; It; ++It)
-	{
-		UMaterialInterface* MaterialInterface = *It;
-
-		TArray<UObject*> Textures;
-		MaterialInterface->AppendReferencedTextures(Textures);
-
-		for (auto Texture : Textures)
-		{
-			if (Texture == this)
-			{
-				BaseMaterialsThatUseThisTexture.Add(MaterialInterface->GetMaterial());
-				break;
-			}
-		}
-	}
-
-	if (BaseMaterialsThatUseThisTexture.Num() > 0)
-	{
-		FMaterialUpdateContext UpdateContext(FMaterialUpdateContext::EOptions::SyncWithRenderingThread);
-		for (TSet<UMaterial*>::TConstIterator It(BaseMaterialsThatUseThisTexture); It; ++It)
-		{
-			(*It)->RecacheUniformExpressions(false);
-			UpdateContext.AddMaterial(*It);
-		}
-	}
-}
diff --git a/Engine/Source/Runtime/Engine/Private/VT/RuntimeVirtualTextureNotify.cpp b/Engine/Source/Runtime/Engine/Private/VT/RuntimeVirtualTextureNotify.cpp
new file mode 100644
--- /dev/null
+++ b/Engine/Source/Runtime/Engine/Private/VT/RuntimeVirtualTextureNotify.cpp
@@ -0,0 +1,98 @@
+// Copyright 1998-2019 Epic Games, Inc. All Rights Reserved.
+
+#include "VT/RuntimeVirtualTextureNotify.h"
+
+#include "Components/PrimitiveComponent.h"
+#include "MaterialShared.h"
+#include "Materials/Material.h"
+#include "UObject/UObjectIterator.h"
+#include "VT/RuntimeVirtualTexture.h"
+#include "VT/RuntimeVirtualTexturePlane.h"
+
+namespace RuntimeVirtualTexture
+{
+#if WITH_EDITOR
+
+	void NotifyComponents(URuntimeVirtualTexture const* VirtualTexture)
+	{
+		for (TObjectIterator<URuntimeVirtualTextureComponent> It; It; ++It)
+		{
+			if (It->GetVirtualTexture() == VirtualTexture)
+			{
+				It->MarkRenderStateDirty();
+			}
+		}
+	}
+
+	void NotifyPrimitives(URuntimeVirtualTexture const* VirtualTexture)
+	{
+		for (TObjectIterator<UPrimitiveComponent> It; It; ++It)
+		{
+			for (URuntimeVirtualTexture* ItVirtualTexture : It->RuntimeVirtualTextures)
+			{
+				if (ItVirtualTexture == VirtualTexture)
+				{
+					It->MarkRenderStateDirty();
+					break;
+				}
+			}
+		}
+
+		/*
+		for (TObjectIterator<ULandscapeComponent> It; It; ++It)
+		{
+			ALandscape* Landscape = It->GetLandscapeActor();
+			if (Landscape != nullptr)
+			{
+				for (URuntimeVirtualTexture* ItVirtualTexture : Landscape->RuntimeVirtualTextures)
+				{
+					if (ItVirtualTexture == VirtualTexture)
+					{
+						It->MarkRenderStateDirty();
+						break;
+					}
+				}
+			}
+		}
+		*/
+	}
+
+#endif
+
+	void NotifyMaterials(URuntimeVirtualTexture const* VirtualTexture)
+	{
+		//todo [vt]:
+		// This operation is very slow! We should only do it during edition, but currently we do it once per virtual texture at runtime load/unload.
+		// Can we pre-calculate the list of materials to touch during cook? But even then touching here them will be extra work...
+		// Is there a way to set up dependencies so that we load the materials _after_ the virtual texture is allocated?
+		// Or maybe we should consider serializing the WorldToUVTransform in the URuntimeVirtualTexture and not depending on a URuntimeVirtualTextureComponent at runtime?
+
+		TSet<UMaterial*> BaseMaterialsThatUseThisTexture;
+		for (TObjectIterator<UMaterialInterface> It; It; ++It)
+		{
+			UMaterialInterface* MaterialInterface = *It;
+
+			TArray<UObject*> Textures;
+			MaterialInterface->AppendReferencedTextures(Textures);
+
+			for (auto Texture : Textures)
+			{
+				if (Texture == VirtualTexture)
+				{
+					BaseMaterialsThatUseThisTexture.Add(MaterialInterface->GetMaterial());
+					break;
+				}
+			}
+		}
+
+		if (BaseMaterialsThatUseThisTexture.Num() > 0)
+		{
+			FMaterialUpdateContext UpdateContext(FMaterialUpdateContext::EOptions::SyncWithRenderingThread);
+			for (TSet<UMaterial*>::TConstIterator It(BaseMaterialsThatUseThisTexture); It; ++It)
+			{
+				(*It)->RecacheUniformExpressions(false);
+				UpdateContext.AddMaterial(*It);
+			}
+		}
+	}
+}
diff --git a/Engine/Source/Runtime/Engine/Private/VT/RuntimeVirtualTextureNotify.h b/Engine/Source/Runtime/Engine/Private/VT/RuntimeVirtualTextureNotify.h
new file mode 100644
--- /dev/null
+++ b/Engine/Source/Runtime/Engine/Private/VT/RuntimeVirtualTextureNotify.h
@@ -0,0 +1,28 @@
+// Copyright 1998-2019 Epic Games, Inc. All Rights Reserved.
+
+#pragma once
+
+#include "CoreMinimal.h"
+
+namespace RuntimeVirtualTexture
+{
+#if WITH_EDITOR
+	/**
+	 * Find any URuntimeVirtualTextureComponent that reference this virtual texture and mark them dirty.
+	 * We need to do this after editing the URuntimeVirtualTexture settings.
+	 */
+	void NotifyComponents(URuntimeVirtualTexture const* VirtualTexture);
+
+	/**
+	 * Find any primitive components that render to this virtual texture and mark them dirty.
+	 * We need to do this after editing the URuntimeVirtualTexture settings.
+	 */
+	void NotifyPrimitives(URuntimeVirtualTexture const* VirtualTexture);
+#endif
+
+	/**
+	 * Find materials referencing this virtual texture and re-cache the uniforms.
+	 * We need to do this after any operation that reallocates the virtual texture since the material caches info about the VT allocation in it's uniform buffer.
+	 */
+	void NotifyMaterials(URuntimeVirtualTexture const* VirtualTexture);
+}
\ No newline at end of file
diff --git a/Engine/Source/Runtime/Engine/Private/VT/RuntimeVirtualTexturePlane.cpp b/Engine/Source/Runtime/Engine/Private/VT/RuntimeVirtualTexturePlane.cpp
--- a/Engine/Source/Runtime/Engine/Private/VT/RuntimeVirtualTexturePlane.cpp
+++ b/Engine/Source/Runtime/Engine/Private/VT/RuntimeVirtualTexturePlane.cpp
@@ -3,6 +3,7 @@
 #include "VT/RuntimeVirtualTexturePlane.h"
 
 #include "Components/BoxComponent.h"
+#include "VT/RuntimeVirtualTextureNotify.h"
 
 
 ARuntimeVirtualTexturePlane::ARuntimeVirtualTexturePlane(const FObjectInitializer& ObjectInitializer)
@@ -25,8 +26,11 @@ ARuntimeVirtualTexturePlane::ARuntimeVirtualTexturePlane(const FObjectInitialize
 
 URuntimeVirtualTextureComponent::URuntimeVirtualTextureComponent(const FObjectInitializer& ObjectInitializer)
 	: Super(ObjectInitializer)
+	, bNotifyInNextTick(false)
 	, SceneProxy(nullptr)
 {
+	PrimaryComponentTick.bCanEverTick = true;
+	bTickInEditor = true;
 }
 
 void URuntimeVirtualTextureComponent::CreateRenderState_Concurrent()
@@ -35,6 +39,7 @@ void URuntimeVirtualTextureComponent::CreateRenderState_Concurrent()
 	{
 		// This will internally allocate modify the URuntimeVirtualTexture and allocate its VT
 		GetScene()->AddRuntimeVirtualTexture(this);
+		bNotifyInNextTick = true;
 	}
 
 	Super::CreateRenderState_Concurrent();
@@ -46,6 +51,7 @@ void URuntimeVirtualTextureComponent::SendRenderTransform_Concurrent()
 	{
 		// This will internally allocate modify the URuntimeVirtualTexture and allocate its VT
 		GetScene()->AddRuntimeVirtualTexture(this);
+		bNotifyInNextTick = true;
 	}
 
 	Super::SendRenderTransform_Concurrent();
@@ -55,10 +61,21 @@ void URuntimeVirtualTextureComponent::DestroyRenderState_Concurrent()
 {
 	// This will internally allocate modify the URuntimeVirtualTexture and free its VT
 	GetScene()->RemoveRuntimeVirtualTexture(this);
+	bNotifyInNextTick = true;
 
 	Super::DestroyRenderState_Concurrent();
 }
 
+void URuntimeVirtualTextureComponent::TickComponent(float DeltaTime, ELevelTick TickType, FActorComponentTickFunction* ThisTickFunction)
+{
+	if (bNotifyInNextTick)
+	{
+		// Notify materials of reallocation of the virtual texture caused by render state update
+		RuntimeVirtualTexture::NotifyMaterials(VirtualTexture);
+		bNotifyInNextTick = false;
+	}
+}
+
 #if WITH_EDITOR
 
 void URuntimeVirtualTextureComponent::SetRotation()
