Forward declare compile fix :(

diff --git a/Engine/Source/Runtime/Engine/Private/VT/RuntimeVirtualTextureNotify.h b/Engine/Source/Runtime/Engine/Private/VT/RuntimeVirtualTextureNotify.h
--- a/Engine/Source/Runtime/Engine/Private/VT/RuntimeVirtualTextureNotify.h
+++ b/Engine/Source/Runtime/Engine/Private/VT/RuntimeVirtualTextureNotify.h
@@ -4,6 +4,8 @@
 
 #include "CoreMinimal.h"
 
+class URuntimeVirtualTexture;
+
 namespace RuntimeVirtualTexture
 {
 #if WITH_EDITOR
