Allow creation of render targets while cooking.

This change reverts 3130965 which was originally submitted to solve a crash with UTexture::Resource. However, this was independently fixed in 3189134 (a later change in a different stream).

diff --git a/Engine/Source/Runtime/Engine/Private/KismetRenderingLibrary.cpp b/Engine/Source/Runtime/Engine/Private/KismetRenderingLibrary.cpp
--- a/Engine/Source/Runtime/Engine/Private/KismetRenderingLibrary.cpp
+++ b/Engine/Source/Runtime/Engine/Private/KismetRenderingLibrary.cpp
@@ -64,12 +64,12 @@ UTextureRenderTarget2D* UKismetRenderingLibrary::CreateRenderTarget2D(UObject* W
 {
 	UWorld* World = GEngine->GetWorldFromContextObject(WorldContextObject, EGetWorldErrorMode::LogAndReturnNull);
 
-	if (Width > 0 && Height > 0 && World && FApp::CanEverRender())
+	if (Width > 0 && Height > 0 && World)
 	{
 		UTextureRenderTarget2D* NewRenderTarget2D = NewObject<UTextureRenderTarget2D>(WorldContextObject);
 		check(NewRenderTarget2D);
 		NewRenderTarget2D->RenderTargetFormat = Format;
-		NewRenderTarget2D->InitAutoFormat(Width, Height); 
+		NewRenderTarget2D->InitAutoFormat(Width, Height);
 		NewRenderTarget2D->UpdateResourceImmediate(true);
 
 		return NewRenderTarget2D; 
