Fixed Paper2D render proxy to have correct local-to-world transform.

diff --git a/Engine/Plugins/2D/Paper2D/Source/Paper2D/Private/PaperRenderSceneProxy.cpp b/Engine/Plugins/2D/Paper2D/Source/Paper2D/Private/PaperRenderSceneProxy.cpp
--- a/Engine/Plugins/2D/Paper2D/Source/Paper2D/Private/PaperRenderSceneProxy.cpp
+++ b/Engine/Plugins/2D/Paper2D/Source/Paper2D/Private/PaperRenderSceneProxy.cpp
@@ -336,6 +336,13 @@ void FPaperRenderSceneProxy::GetBatchMesh(const FSceneView* View, class UMateria
 
 	const bool bIsWireframeView = View->Family->EngineShowFlags.Wireframe;
 
+	FMatrix PreviousLocalToWorld;
+
+	if (!GetScene().GetPreviousLocalToWorld(GetPrimitiveSceneInfo(), PreviousLocalToWorld))
+	{
+		PreviousLocalToWorld = GetLocalToWorld();
+	}
+
 	for (const FSpriteDrawCallRecord& Record : Batch)
 	{
 		if (Record.IsValid())
@@ -391,7 +398,7 @@ void FPaperRenderSceneProxy::GetBatchMesh(const FSceneView* View, class UMateria
 			Settings.bReceivesDecals = true;
 			Settings.bUseSelectionOutline = true;
 			Settings.CastShadow = bCastShadow;
-			DynamicMeshBuilder.GetMesh(GetLocalToWorld(), TextureOverrideMaterialProxy, DPG, Settings, nullptr, ViewIndex, Collector);
+			DynamicMeshBuilder.GetMesh(GetLocalToWorld(), PreviousLocalToWorld, TextureOverrideMaterialProxy, DPG, Settings, nullptr, ViewIndex, Collector);
 		}
 	}
 }
diff --git a/Engine/Source/Runtime/Engine/Private/DynamicMeshBuilder.cpp b/Engine/Source/Runtime/Engine/Private/DynamicMeshBuilder.cpp
--- a/Engine/Source/Runtime/Engine/Private/DynamicMeshBuilder.cpp
+++ b/Engine/Source/Runtime/Engine/Private/DynamicMeshBuilder.cpp
@@ -694,7 +694,29 @@ void FDynamicMeshBuilder::GetMesh(const FMatrix& LocalToWorld, const FMaterialRe
 	GetMesh(LocalToWorld, MaterialRenderProxy, DepthPriorityGroup, Settings, nullptr, ViewIndex, Collector, HitProxyId);
 }
 
-void FDynamicMeshBuilder::GetMesh(const FMatrix& LocalToWorld,const FMaterialRenderProxy* MaterialRenderProxy, uint8 DepthPriorityGroup, const FDynamicMeshBuilderSettings& Settings, FDynamicMeshDrawOffset const * const DrawOffset, int32 ViewIndex, FMeshElementCollector& Collector, const FHitProxyId HitProxyId )
+void FDynamicMeshBuilder::GetMesh(
+	const FMatrix& LocalToWorld,
+	const FMaterialRenderProxy* MaterialRenderProxy,
+	uint8 DepthPriorityGroup,
+	const FDynamicMeshBuilderSettings& Settings,
+	FDynamicMeshDrawOffset const * const DrawOffset,
+	int32 ViewIndex,
+	FMeshElementCollector& Collector,
+	const FHitProxyId HitProxyId)
+{
+	GetMesh(LocalToWorld, LocalToWorld, MaterialRenderProxy, DepthPriorityGroup, Settings, DrawOffset, ViewIndex, Collector, HitProxyId);
+}
+
+void FDynamicMeshBuilder::GetMesh(
+	const FMatrix& LocalToWorld,
+	const FMatrix& PreviousLocalToWorld,
+	const FMaterialRenderProxy* MaterialRenderProxy,
+	uint8 DepthPriorityGroup,
+	const FDynamicMeshBuilderSettings& Settings,
+	FDynamicMeshDrawOffset const * const DrawOffset,
+	int32 ViewIndex,
+	FMeshElementCollector& Collector,
+	const FHitProxyId HitProxyId )
 {
 	// Only draw non-empty meshes.
 	if((VertexBuffer && VertexBuffer->Vertices.Num() > 0) || (DrawOffset != nullptr))
@@ -723,7 +745,7 @@ void FDynamicMeshBuilder::GetMesh(const FMatrix& LocalToWorld,const FMaterialRen
 			OneFrameResources->PrimitiveUniformBuffer = new FDynamicMeshPrimitiveUniformBuffer();
 			FPrimitiveUniformShaderParameters PrimitiveParams = GetPrimitiveUniformShaderParameters(
 				LocalToWorld,
-				LocalToWorld,
+				PreviousLocalToWorld,
 				LocalToWorld.GetOrigin(),
 				FBoxSphereBounds(EForceInit::ForceInit),
 				FBoxSphereBounds(EForceInit::ForceInit),
diff --git a/Engine/Source/Runtime/Engine/Public/DynamicMeshBuilder.h b/Engine/Source/Runtime/Engine/Public/DynamicMeshBuilder.h
--- a/Engine/Source/Runtime/Engine/Public/DynamicMeshBuilder.h
+++ b/Engine/Source/Runtime/Engine/Public/DynamicMeshBuilder.h
@@ -180,6 +180,8 @@ public:
 							FMeshElementCollector& Collector, const FHitProxyId HitProxyId = FHitProxyId());
 	ENGINE_API void GetMesh(const FMatrix& LocalToWorld, const FMaterialRenderProxy* MaterialRenderProxy, uint8 DepthPriorityGroup, const FDynamicMeshBuilderSettings& Settings, FDynamicMeshDrawOffset const * const DrawOffset, int32 ViewIndex,
 		FMeshElementCollector& Collector, const FHitProxyId HitProxyId = FHitProxyId());
+	ENGINE_API void GetMesh(const FMatrix& LocalToWorld, const FMatrix& PrevLocalToWorld, const FMaterialRenderProxy* MaterialRenderProxy, uint8 DepthPriorityGroup, const FDynamicMeshBuilderSettings& Settings,
+		FDynamicMeshDrawOffset const * const DrawOffset, int32 ViewIndex, FMeshElementCollector& Collector, const FHitProxyId HitProxyId = FHitProxyId());
 
 	ENGINE_API void GetMeshElement(const FMatrix& LocalToWorld, const FMaterialRenderProxy* MaterialRenderProxy, uint8 DepthPriorityGroup, bool bDisableBackfaceCulling, bool bReceivesDecals, int32 ViewIndex, FMeshBuilderOneFrameResources& OneFrameResource, FMeshBatch& Mesh);
 
