Fix Android build warning

diff --git a/Engine/Source/Runtime/Engine/Private/VT/RuntimeVirtualTexturePlane.cpp b/Engine/Source/Runtime/Engine/Private/VT/RuntimeVirtualTexturePlane.cpp
--- a/Engine/Source/Runtime/Engine/Private/VT/RuntimeVirtualTexturePlane.cpp
+++ b/Engine/Source/Runtime/Engine/Private/VT/RuntimeVirtualTexturePlane.cpp
@@ -37,7 +37,7 @@ void URuntimeVirtualTextureComponent::CreateRenderState_Concurrent()
 {
 	if (ShouldRender() && VirtualTexture != nullptr)
 	{
-		// This will internally allocate modify the URuntimeVirtualTexture and allocate its VT
+		// This will modify the URuntimeVirtualTexture and allocate its VT
 		GetScene()->AddRuntimeVirtualTexture(this);
 		bNotifyInNextTick = true;
 	}
@@ -49,7 +49,7 @@ void URuntimeVirtualTextureComponent::SendRenderTransform_Concurrent()
 {
 	if (ShouldRender() && VirtualTexture != nullptr)
 	{
-		// This will internally allocate modify the URuntimeVirtualTexture and allocate its VT
+		// This will modify the URuntimeVirtualTexture and allocate its VT
 		GetScene()->AddRuntimeVirtualTexture(this);
 		bNotifyInNextTick = true;
 	}
@@ -59,7 +59,7 @@ void URuntimeVirtualTextureComponent::SendRenderTransform_Concurrent()
 
 void URuntimeVirtualTextureComponent::DestroyRenderState_Concurrent()
 {
-	// This will internally allocate modify the URuntimeVirtualTexture and free its VT
+	// This will modify the URuntimeVirtualTexture and free its VT
 	GetScene()->RemoveRuntimeVirtualTexture(this);
 	bNotifyInNextTick = true;
 
@@ -70,7 +70,8 @@ void URuntimeVirtualTextureComponent::TickComponent(float DeltaTime, ELevelTick
 {
 	if (bNotifyInNextTick)
 	{
-		// Notify materials of reallocation of the virtual texture caused by render state update
+		// Notify materials of reallocation of the virtual texture caused by render state update.
+		// This is slow and ideally we will find a different approach.
 		RuntimeVirtualTexture::NotifyMaterials(VirtualTexture);
 		bNotifyInNextTick = false;
 	}
diff --git a/Engine/Source/Runtime/Renderer/Private/VT/RuntimeVirtualTextureProducer.cpp b/Engine/Source/Runtime/Renderer/Private/VT/RuntimeVirtualTextureProducer.cpp
--- a/Engine/Source/Runtime/Renderer/Private/VT/RuntimeVirtualTextureProducer.cpp
+++ b/Engine/Source/Runtime/Renderer/Private/VT/RuntimeVirtualTextureProducer.cpp
@@ -77,11 +77,7 @@ FVTRequestPageResult FRuntimeVirtualTextureProducer::RequestPageData(
 	EVTRequestPagePriority Priority)
 {
 	//todo[vt]: 
-	// Investigate what causes a partial layer mask to be requested.
-	// If we can't avoid it then look at ways to handle it efficiently (right now we render all layers even for partial requests).
-
-	//todo[vt]: 
-	// Possibly throttle rendering according to performance by returning Saturated here.
+	// Possibly throttle rendering according to performance and return Saturated here.
 
 	FVTRequestPageResult result;
 	result.Handle = 0;
@@ -105,7 +101,12 @@ IVirtualTextureFinalizer* FRuntimeVirtualTextureProducer::ProducePageData(
 	FRuntimeVirtualTextureFinalizer::FTileEntry Tile;
 	Tile.vAddress = vAddress;
 	Tile.vLevel = vLevel;
-	
+
+	//todo[vt]: 
+	// Partial layer masks can happen when one layer has more physical space available so that old pages are evicted at different rates.
+	// This can be almost always be avoided by setting up the physical pools correctly for the application's needs.
+	// If we can't avoid partial layer masks then we could look at ways to handle it more efficiently (right now we render all layers even for these partial requests).
+
 	//todo[vt]: Add support for more than two layers
 	if (LayerMask & 1)
 	{
diff --git a/Engine/Source/Runtime/Renderer/Private/VT/RuntimeVirtualTextureRender.cpp b/Engine/Source/Runtime/Renderer/Private/VT/RuntimeVirtualTextureRender.cpp
--- a/Engine/Source/Runtime/Renderer/Private/VT/RuntimeVirtualTextureRender.cpp
+++ b/Engine/Source/Runtime/Renderer/Private/VT/RuntimeVirtualTextureRender.cpp
@@ -617,9 +617,9 @@ namespace RuntimeVirtualTexture
 				RDG_EVENT_NAME("VirtualTextureDraw"),
 				PassParameters,
 				ERenderGraphPassFlags::None,
-				[Scene, View, MaterialType](FRHICommandListImmediate& RHICmdList)
+				[Scene, View, MaterialType](FRHICommandListImmediate& RHICmdListImmediate)
 			{
-				DrawMeshes(RHICmdList, Scene, View, MaterialType);
+				DrawMeshes(RHICmdListImmediate, Scene, View, MaterialType);
 			});
 		}
 
