Performance improvement when deleting an instanced static mesh instance.

#5711

diff --git a/Engine/Source/Runtime/Engine/Private/InstancedStaticMesh.cpp b/Engine/Source/Runtime/Engine/Private/InstancedStaticMesh.cpp
--- a/Engine/Source/Runtime/Engine/Private/InstancedStaticMesh.cpp
+++ b/Engine/Source/Runtime/Engine/Private/InstancedStaticMesh.cpp
@@ -1850,11 +1850,22 @@ bool UInstancedStaticMeshComponent::RemoveInstanceInternal(int32 InstanceIndex,
 #endif
 
 	// update the physics state
-	if (bPhysicsStateCreated)
+	if (bPhysicsStateCreated && InstanceBodies.IsValidIndex(InstanceIndex))
 	{
-		// TODO: it may be possible to instead just update the BodyInstanceIndex for all bodies after the removed instance. 
-		ClearAllInstanceBodies();
-		CreateAllInstanceBodies();
+		if (FBodyInstance*& InstanceBody = InstanceBodies[InstanceIndex])
+		{
+			InstanceBody->TermBody();
+			delete InstanceBody;
+			InstanceBody = nullptr;
+
+			InstanceBodies.RemoveAt(InstanceIndex);
+
+			// Re-target instance indices for shifting of array.
+			for (int32 i = InstanceIndex; i < InstanceBodies.Num(); ++i)
+			{
+				InstanceBodies[i]->InstanceBodyIndex = i;
+			}
+		}
 	}
 
 	// Force recreation of the render data
