Report error on corrupted output file in ShaderCompileWorker.

#rnx

diff --git a/Engine/Source/Runtime/Engine/Private/ShaderCompiler/ShaderCompiler.cpp b/Engine/Source/Runtime/Engine/Private/ShaderCompiler/ShaderCompiler.cpp
--- a/Engine/Source/Runtime/Engine/Private/ShaderCompiler/ShaderCompiler.cpp
+++ b/Engine/Source/Runtime/Engine/Private/ShaderCompiler/ShaderCompiler.cpp
@@ -113,7 +113,7 @@ static void ModalErrorOrLog(const FString& Text, int64 CurrentFilePos = 0, int64
 	}
 }
 
-// Set to 1 to debug ShaderCompilerWorker.exe. Set a breakpoint in LaunchWorker() to get the cmd-line.
+// Set to 1 to debug ShaderCompileWorker.exe. Set a breakpoint in LaunchWorker() to get the cmd-line.
 #define DEBUG_SHADERCOMPILEWORKER 0
 
 // Default value comes from bPromptToRetryFailedShaderCompiles in BaseEngine.ini
@@ -345,6 +345,11 @@ namespace SCWErrorCode
 	{
 		ModalErrorOrLog(FString::Printf(TEXT("ShaderCompileWorker failed:\n%s\n"), Data));
 	}
+
+	void HandleOutputFileCorrupted(const TCHAR* Filename, int64 ExpectedSize, int64 ActualSize)
+	{
+		ModalErrorOrLog(FString::Printf(TEXT("Output file corrupted (expected %I64d bytes, but only got %I64d): %s"), Filename, ExpectedSize, ActualSize));
+	}
 }
 
 static inline void GetFormatVersionMap(TMap<FString, uint32>& OutFormatVersionMap)
@@ -608,14 +613,20 @@ void FShaderCompileUtilities::DoReadTaskResults(const TArray<FShaderCommonCompil
 
 	if (ShaderCompileWorkerOutputVersion != OutputVersion)
 	{
-		FString Text = FString::Printf(TEXT("Expecting ShaderCompilerWorker output version %d, got %d instead! Forgot to build ShaderCompilerWorker?"), ShaderCompileWorkerOutputVersion, OutputVersion);
+		FString Text = FString::Printf(TEXT("Expecting ShaderCompileWorker output version %d, got %d instead! Forgot to build ShaderCompileWorker?"), ShaderCompileWorkerOutputVersion, OutputVersion);
 		ModalErrorOrLog(Text);
 	}
 
 	int64 FileSize = 0;
 	OutputFile << FileSize;
 
-	int32 ErrorCode;
+	// Check for corrupted output file
+	if (FileSize > OutputFile.TotalSize())
+	{
+		SCWErrorCode::HandleOutputFileCorrupted(*OutputFile.GetArchiveName(), FileSize, OutputFile.TotalSize());
+	}
+
+	int32 ErrorCode = 0;
 	OutputFile << ErrorCode;
 
 	int32 NumProcessedJobs = 0;
@@ -726,7 +737,7 @@ void FShaderCompileUtilities::DoReadTaskResults(const TArray<FShaderCommonCompil
 		OutputFile << SingleJobHeader;
 		if (SingleJobHeader != ShaderCompileWorkerSingleJobHeader)
 		{
-			FString Text = FString::Printf(TEXT("Expecting ShaderCompilerWorker Single Jobs %d, got %d instead! Forgot to build ShaderCompilerWorker?"), ShaderCompileWorkerSingleJobHeader, SingleJobHeader);
+			FString Text = FString::Printf(TEXT("Expecting ShaderCompileWorker Single Jobs %d, got %d instead! Forgot to build ShaderCompileWorker?"), ShaderCompileWorkerSingleJobHeader, SingleJobHeader);
 			ModalErrorOrLog(Text, OutputFile.Tell(), FileSize);
 		}
 
@@ -751,7 +762,7 @@ void FShaderCompileUtilities::DoReadTaskResults(const TArray<FShaderCommonCompil
 		OutputFile << PipelineJobHeader;
 		if (PipelineJobHeader != ShaderCompileWorkerPipelineJobHeader)
 		{
-			FString Text = FString::Printf(TEXT("Expecting ShaderCompilerWorker Pipeline Jobs %d, got %d instead! Forgot to build ShaderCompilerWorker?"), ShaderCompileWorkerPipelineJobHeader, PipelineJobHeader);
+			FString Text = FString::Printf(TEXT("Expecting ShaderCompileWorker Pipeline Jobs %d, got %d instead! Forgot to build ShaderCompileWorker?"), ShaderCompileWorkerPipelineJobHeader, PipelineJobHeader);
 			ModalErrorOrLog(Text, OutputFile.Tell(), FileSize);
 		}
 
@@ -1059,7 +1070,7 @@ int32 FShaderCompileThreadRunnable::PullTasksFromQueue()
 					}
 
 					// Update the worker state as having new tasks that need to be issued					
-					// don't reset worker app ID, because the shadercompilerworkers don't shutdown immediately after finishing a single job queue.
+					// don't reset worker app ID, because the shadercompileworkers don't shutdown immediately after finishing a single job queue.
 					CurrentWorkerInfo.bIssuedTasksToWorker = false;					
 					CurrentWorkerInfo.bLaunchedWorker = false;
 					CurrentWorkerInfo.StartTime = FPlatformTime::Seconds();
@@ -3880,7 +3891,7 @@ void VerifyGlobalShaders(EShaderPlatform Platform, bool bLoadedFromCacheFile)
 		}
 
 		ensureMsgf(
-			PermutationCountToCompile < 200,	// Please try to not bump this
+			PermutationCountToCompile < 397,	// ToneMapper today (2019-04-17) can go up to 396 permutations
 			TEXT("Global shader %s has %i permutation: probably more that it needs."),
 			GlobalShaderType->GetName(), PermutationCountToCompile);
 
diff --git a/Engine/Source/Runtime/RenderCore/Private/ShaderCore.cpp b/Engine/Source/Runtime/RenderCore/Private/ShaderCore.cpp
--- a/Engine/Source/Runtime/RenderCore/Private/ShaderCore.cpp
+++ b/Engine/Source/Runtime/RenderCore/Private/ShaderCore.cpp
@@ -402,7 +402,7 @@ void FShaderCompilerOutput::GenerateOutputHash()
 	
 	const TArray<uint8>& Code = ShaderCode.GetReadAccess();
 
-	// we don't hash the optional attachments as they would prevent sharing (e.g. many material share the save VS)
+	// we don't hash the optional attachments as they would prevent sharing (e.g. many materials share the same VS)
 	uint32 ShaderCodeSize = ShaderCode.GetShaderCodeSize();
 
 	HashState.Update(Code.GetData(), ShaderCodeSize * Code.GetTypeSize());
diff --git a/Engine/Source/Runtime/RenderCore/Public/ShaderCore.h b/Engine/Source/Runtime/RenderCore/Public/ShaderCore.h
--- a/Engine/Source/Runtime/RenderCore/Public/ShaderCore.h
+++ b/Engine/Source/Runtime/RenderCore/Public/ShaderCore.h
@@ -864,7 +864,7 @@ public:
 		return 0;
 	}
 
-	// Returns nullptr and Size -1 if not key was not found
+	// Returns nullptr and Size -1 if key was not found
 	const uint8* FindOptionalDataAndSize(uint8 InKey, int32& OutSize) const
 	{
 		check(ShaderCode.Num() >= 4);
