Fixed static analysis warnings.

diff --git a/Engine/Source/Runtime/Engine/Private/LightMap.cpp b/Engine/Source/Runtime/Engine/Private/LightMap.cpp
--- a/Engine/Source/Runtime/Engine/Private/LightMap.cpp
+++ b/Engine/Source/Runtime/Engine/Private/LightMap.cpp
@@ -1682,6 +1682,8 @@ void FLightMapPendingTexture::StartEncoding(ULevel* LightingScenario, ITextureCo
 
 	if (VirtualTexture)
 	{
+		const uint32 InvalidLayerId = ~0u;
+
 		// Copy data from all the separate Lightmap textures into the proper layers of the VT source
 		const uint32 SkyOcclusionLayer = VirtualTexture->GetLayerForType(ELightMapVirtualTextureType::SkyOcclusion);
 		const uint32 AOMaterialMaskLayer = VirtualTexture->GetLayerForType(ELightMapVirtualTextureType::AOMaterialMask);
@@ -1691,16 +1693,20 @@ void FLightMapPendingTexture::StartEncoding(ULevel* LightingScenario, ITextureCo
 		LayerFormat.Init(TSF_Invalid, NumVirtualTextureLayers);
 		LayerFormat[0] = Textures[0]->Source.GetFormat();
 		LayerFormat[1] = Textures[0]->Source.GetFormat();
-		if (SkyOcclusionLayer != ~0u && SkyOcclusionTexture)
+
+		if (SkyOcclusionLayer != InvalidLayerId)
 		{
+			check(SkyOcclusionTexture);
 			LayerFormat[SkyOcclusionLayer] = SkyOcclusionTexture->Source.GetFormat();
 		}
-		if (AOMaterialMaskLayer != ~0u && AOMaterialMaskTexture)
+		if (AOMaterialMaskLayer != InvalidLayerId)
 		{
+			check(AOMaterialMaskTexture);
 			LayerFormat[AOMaterialMaskLayer] = AOMaterialMaskTexture->Source.GetFormat();
 		}
-		if (ShadowMaskLayer != ~0u && ShadowMapTexture)
+		if (ShadowMaskLayer != InvalidLayerId)
 		{
+			check(ShadowMapTexture);
 			LayerFormat[ShadowMaskLayer] = ShadowMapTexture->Source.GetFormat();
 		}
 
@@ -1732,7 +1738,7 @@ void FLightMapPendingTexture::StartEncoding(ULevel* LightingScenario, ITextureCo
 			FMemory::Memcpy(Layer1Data, &SourceData[SourceData.Num() / 2], SourceData.Num() / 2);
 			VirtualTexture->Source.UnlockMip(0, 1, MipIndex);
 
-			if (SkyOcclusionLayer != ~0u)
+			if (SkyOcclusionLayer != InvalidLayerId)
 			{
 				FTextureFormatSettings& FormatSettings = VirtualTexture->LayerFormatSettings[SkyOcclusionLayer];
 				SkyOcclusionTexture->GetDefaultFormatSettings(FormatSettings);
@@ -1742,7 +1748,7 @@ void FLightMapPendingTexture::StartEncoding(ULevel* LightingScenario, ITextureCo
 				FMemory::Memcpy(LayerData, &SourceData[0], SourceData.Num());
 				VirtualTexture->Source.UnlockMip(0, SkyOcclusionLayer, MipIndex);
 			}
-			if (AOMaterialMaskLayer != ~0u)
+			if (AOMaterialMaskLayer != InvalidLayerId)
 			{
 				FTextureFormatSettings& FormatSettings = VirtualTexture->LayerFormatSettings[AOMaterialMaskLayer];
 				AOMaterialMaskTexture->GetDefaultFormatSettings(FormatSettings);
@@ -1752,7 +1758,7 @@ void FLightMapPendingTexture::StartEncoding(ULevel* LightingScenario, ITextureCo
 				FMemory::Memcpy(LayerData, &SourceData[0], SourceData.Num());
 				VirtualTexture->Source.UnlockMip(0, AOMaterialMaskLayer, MipIndex);
 			}
-			if (ShadowMaskLayer != ~0u)
+			if (ShadowMaskLayer != InvalidLayerId)
 			{
 				FTextureFormatSettings& FormatSettings = VirtualTexture->LayerFormatSettings[ShadowMaskLayer];
 				ShadowMapTexture->GetDefaultFormatSettings(FormatSettings);
