Added more description to pre-exposure
Fixed bad exposure values when r.EyeAdaptation.MethodOveride==-2 or r.DefaultFeature.AutoExposure=1, with "r.DefaultFeature.AutoExposure.ExtendDefaultLuminanceRange=1"

diff --git a/Engine/Source/Runtime/Engine/Private/SceneView.cpp b/Engine/Source/Runtime/Engine/Private/SceneView.cpp
--- a/Engine/Source/Runtime/Engine/Private/SceneView.cpp
+++ b/Engine/Source/Runtime/Engine/Private/SceneView.cpp
@@ -301,6 +301,11 @@ static TAutoConsoleVariable<int32> CVarEnableTemporalUpsample(
 	TEXT(" 1: TemporalAA performs spatial and temporal upscale as screen percentage method."),
 	ECVF_Default);
 
+// Conversion factor used when "r.DefaultFeature.AutoExposure.ExtendDefaultLuminanceRange" 
+FORCEINLINE float LuminanceToEV100(float Luminance)
+{
+	return FMath::Log2(Luminance / 1.2f);
+}
 
 /** Global vertex color view mode setting when SHOW_VertexColors show flag is set */
 EVertexColorViewMode::Type GVertexColorViewMode = EVertexColorViewMode::Color;
@@ -1671,6 +1676,11 @@ void FSceneView::StartFinalPostprocessSettings(FVector InViewLocation)
 		{
 			FinalPostProcessSettings.AutoExposureMinBrightness = 1;
 			FinalPostProcessSettings.AutoExposureMaxBrightness = 1;
+			if (CVarDefaultAutoExposureExtendDefaultLuminanceRange.GetValueOnGameThread())
+			{
+				FinalPostProcessSettings.AutoExposureMinBrightness = LuminanceToEV100(FinalPostProcessSettings.AutoExposureMinBrightness);
+				FinalPostProcessSettings.AutoExposureMaxBrightness = LuminanceToEV100(FinalPostProcessSettings.AutoExposureMaxBrightness);
+			}
 		}
 		else
 		{
@@ -1729,9 +1739,14 @@ void FSceneView::EndFinalPostprocessSettings(const FSceneViewInitOptions& ViewIn
 		FinalPostProcessSettings.AutoExposureMethod = AEM_Basic;
 		FinalPostProcessSettings.AutoExposureBias = -0.6f;
 		FinalPostProcessSettings.AutoExposureMaxBrightness = 2.f;
-		FinalPostProcessSettings.AutoExposureMinBrightness = 0.05;
+		FinalPostProcessSettings.AutoExposureMinBrightness = 0.05f;
 		FinalPostProcessSettings.AutoExposureSpeedDown = 1.f;
 		FinalPostProcessSettings.AutoExposureSpeedUp = 3.f;
+		if (CVarDefaultAutoExposureExtendDefaultLuminanceRange.GetValueOnGameThread())
+		{
+			FinalPostProcessSettings.AutoExposureMinBrightness = LuminanceToEV100(FinalPostProcessSettings.AutoExposureMinBrightness);
+			FinalPostProcessSettings.AutoExposureMaxBrightness = LuminanceToEV100(FinalPostProcessSettings.AutoExposureMaxBrightness);
+		}
 	}
 #endif
 
diff --git a/Engine/Source/Runtime/Renderer/Private/SceneRendering.cpp b/Engine/Source/Runtime/Renderer/Private/SceneRendering.cpp
--- a/Engine/Source/Runtime/Renderer/Private/SceneRendering.cpp
+++ b/Engine/Source/Runtime/Renderer/Private/SceneRendering.cpp
@@ -167,7 +167,10 @@ static TAutoConsoleVariable<int32> CVarRoundRobinOcclusion(
 static TAutoConsoleVariable<int32> CVarUsePreExposure(
 	TEXT("r.UsePreExposure"),
 	1,
-	TEXT("0 to disable pre-exposure, 1 to enable (default)."),
+	TEXT("0 to disable pre-exposure, 1 to enable it (default).\n")
+	TEXT("Pre-exposure allows the engine to apply the last frame exposure to luminance values before writing them in rendertargets.\n")
+	TEXT("It avoids rendertarget overflow when using low precision formats like fp16.\n")
+	TEXT("The pre-exposure value can be overriden through r.EyeAdaptation.PreExposureOverride\n"),
 	ECVF_ReadOnly | ECVF_RenderThreadSafe);
 
 static TAutoConsoleVariable<int32> CVarODSCapture(
