- Low level material uniform buffer won't allow invalid texture types to be bound (will be replaced with appropriate default white texture), this should fix crashes in all cases
- Material instance editor will display message if invalid texture is assigned to overridden parameter
- Still no feedback if invalid parameter is set on MID (but will silently replace with default white texture)

diff --git a/Engine/Source/Editor/MaterialEditor/Private/MaterialInstanceEditor.cpp b/Engine/Source/Editor/MaterialEditor/Private/MaterialInstanceEditor.cpp
--- a/Engine/Source/Editor/MaterialEditor/Private/MaterialInstanceEditor.cpp
+++ b/Engine/Source/Editor/MaterialEditor/Private/MaterialInstanceEditor.cpp
@@ -1262,35 +1262,52 @@ void FMaterialInstanceEditor::DrawSamplerWarningStrings(FCanvas* Canvas, int32&
 							EMaterialSamplerType SamplerType = UMaterialExpressionTextureBase::GetSamplerTypeForTexture( Texture );
 							UMaterialExpressionTextureSampleParameter* Expression = BaseMaterial->FindExpressionByGUID<UMaterialExpressionTextureSampleParameter>( TextureParameterValue->ExpressionId );
 
-							if ( Expression && Expression->SamplerType != SamplerType )
+							FString ErrorMessage;
+							if (!Expression->TextureIsValid(Texture, ErrorMessage))
 							{
-								FString SamplerTypeDisplayName = SamplerTypeEnum->GetDisplayNameTextByValue(Expression->SamplerType).ToString();
-
 								Canvas->DrawShadowedString(
 									5,
 									DrawPositionY,
-									*FString::Printf( TEXT("Warning: %s samples %s as %s."),
-									*TextureParameterValue->ParameterInfo.Name.ToString(),
-									*Texture->GetPathName(),
-									*SamplerTypeDisplayName ),
+									*FString::Printf(TEXT("Error: %s has invalid texture %s: %s."),
+										*TextureParameterValue->ParameterInfo.Name.ToString(),
+										*Texture->GetPathName(),
+										*ErrorMessage),
 									FontToUse,
-									FLinearColor(1,1,0) );
+									FLinearColor(1, 0, 0));
 								DrawPositionY += SpacingBetweenLines;
 							}
-							if( Expression && ((Expression->SamplerType == (EMaterialSamplerType)TC_Normalmap || Expression->SamplerType ==  (EMaterialSamplerType)TC_Masks) && Texture->SRGB))
+							else
 							{
-								FString SamplerTypeDisplayName = SamplerTypeEnum->GetDisplayNameTextByValue(Expression->SamplerType).ToString();
-								
-								Canvas->DrawShadowedString(
-									5,
-									DrawPositionY,
-									*FString::Printf( TEXT("Warning: %s samples texture as '%s'. SRGB should be disabled for '%s'."),
-									*TextureParameterValue->ParameterInfo.Name.ToString(),
-									*SamplerTypeDisplayName,
-									*Texture->GetPathName()),
-									FontToUse,
-									FLinearColor(1,1,0) );
-								DrawPositionY += SpacingBetweenLines;
+								if (Expression && Expression->SamplerType != SamplerType)
+								{
+									FString SamplerTypeDisplayName = SamplerTypeEnum->GetDisplayNameTextByValue(Expression->SamplerType).ToString();
+
+									Canvas->DrawShadowedString(
+										5,
+										DrawPositionY,
+										*FString::Printf(TEXT("Warning: %s samples %s as %s."),
+											*TextureParameterValue->ParameterInfo.Name.ToString(),
+											*Texture->GetPathName(),
+											*SamplerTypeDisplayName),
+										FontToUse,
+										FLinearColor(1, 1, 0));
+									DrawPositionY += SpacingBetweenLines;
+								}
+								if (Expression && ((Expression->SamplerType == (EMaterialSamplerType)TC_Normalmap || Expression->SamplerType == (EMaterialSamplerType)TC_Masks) && Texture->SRGB))
+								{
+									FString SamplerTypeDisplayName = SamplerTypeEnum->GetDisplayNameTextByValue(Expression->SamplerType).ToString();
+
+									Canvas->DrawShadowedString(
+										5,
+										DrawPositionY,
+										*FString::Printf(TEXT("Warning: %s samples texture as '%s'. SRGB should be disabled for '%s'."),
+											*TextureParameterValue->ParameterInfo.Name.ToString(),
+											*SamplerTypeDisplayName,
+											*Texture->GetPathName()),
+										FontToUse,
+										FLinearColor(1, 1, 0));
+									DrawPositionY += SpacingBetweenLines;
+								}
 							}
 						}
 					}
diff --git a/Engine/Source/Runtime/Engine/Classes/Materials/MaterialExpressionAntialiasedTextureMask.h b/Engine/Source/Runtime/Engine/Classes/Materials/MaterialExpressionAntialiasedTextureMask.h
--- a/Engine/Source/Runtime/Engine/Classes/Materials/MaterialExpressionAntialiasedTextureMask.h
+++ b/Engine/Source/Runtime/Engine/Classes/Materials/MaterialExpressionAntialiasedTextureMask.h
@@ -33,8 +33,7 @@ class UMaterialExpressionAntialiasedTextureMask : public UMaterialExpressionText
 
 
 	//~ Begin UMaterialExpressionTextureSampleParameter Interface
-	virtual bool TextureIsValid( UTexture* InTexture ) override;
-	virtual const TCHAR* GetRequirements() override;
+	virtual bool TextureIsValid(UTexture* InTexture, FString& OutMessage) override;
 	virtual void SetDefaultTexture() override;
 	//~ End UMaterialExpressionTextureSampleParameter Interface
 
diff --git a/Engine/Source/Runtime/Engine/Classes/Materials/MaterialExpressionTextureObjectParameter.h b/Engine/Source/Runtime/Engine/Classes/Materials/MaterialExpressionTextureObjectParameter.h
--- a/Engine/Source/Runtime/Engine/Classes/Materials/MaterialExpressionTextureObjectParameter.h
+++ b/Engine/Source/Runtime/Engine/Classes/Materials/MaterialExpressionTextureObjectParameter.h
@@ -27,11 +27,6 @@ class UMaterialExpressionTextureObjectParameter : public UMaterialExpressionText
 	virtual const TArray<FExpressionInput*> GetInputs() override;
 #endif
 	//~ End UMaterialExpression Interface
-
-	//~ Begin UMaterialExpressionTextureSampleParameter Interface
-	virtual const TCHAR* GetRequirements() override;
-	//~ End UMaterialExpressionTextureSampleParameter Interface
-
 };
 
 
diff --git a/Engine/Source/Runtime/Engine/Classes/Materials/MaterialExpressionTextureSampleParameter.h b/Engine/Source/Runtime/Engine/Classes/Materials/MaterialExpressionTextureSampleParameter.h
--- a/Engine/Source/Runtime/Engine/Classes/Materials/MaterialExpressionTextureSampleParameter.h
+++ b/Engine/Source/Runtime/Engine/Classes/Materials/MaterialExpressionTextureSampleParameter.h
@@ -62,17 +62,11 @@ class ENGINE_API UMaterialExpressionTextureSampleParameter : public UMaterialExp
 	 * Return true if the texture is a movie texture
 	 *
 	 * @param	InTexture - texture to test
+	 * @param	OutMessage - if texture isn't valid, gives a description of the problem
 	 * @return	true/false
 	 */	
-	virtual bool TextureIsValid( UTexture* InTexture );
+	virtual bool TextureIsValid(UTexture* InTexture, FString& OutMessage);
 
-	/**
-	 * Called when TextureIsValid==false
-	 *
-	 * @return	Descriptive error text
-	 */	
-	virtual const TCHAR* GetRequirements();
-	
 	/**
 	 *	Sets the default texture if none is set
 	 */
diff --git a/Engine/Source/Runtime/Engine/Classes/Materials/MaterialExpressionTextureSampleParameter2D.h b/Engine/Source/Runtime/Engine/Classes/Materials/MaterialExpressionTextureSampleParameter2D.h
--- a/Engine/Source/Runtime/Engine/Classes/Materials/MaterialExpressionTextureSampleParameter2D.h
+++ b/Engine/Source/Runtime/Engine/Classes/Materials/MaterialExpressionTextureSampleParameter2D.h
@@ -23,8 +23,7 @@ class ENGINE_API UMaterialExpressionTextureSampleParameter2D : public UMaterialE
 	//~ End UMaterialExpression Interface
 	
 	//~ Begin UMaterialExpressionTextureSampleParameter Interface
-	virtual bool TextureIsValid( UTexture* InTexture ) override;
-	virtual const TCHAR* GetRequirements() override;
+	virtual bool TextureIsValid(UTexture* InTexture, FString& OutMessage) override;
 	virtual void SetDefaultTexture() override;
 	//~ End UMaterialExpressionTextureSampleParameter Interface
 };
diff --git a/Engine/Source/Runtime/Engine/Classes/Materials/MaterialExpressionTextureSampleParameterCube.h b/Engine/Source/Runtime/Engine/Classes/Materials/MaterialExpressionTextureSampleParameterCube.h
--- a/Engine/Source/Runtime/Engine/Classes/Materials/MaterialExpressionTextureSampleParameterCube.h
+++ b/Engine/Source/Runtime/Engine/Classes/Materials/MaterialExpressionTextureSampleParameterCube.h
@@ -26,8 +26,7 @@ class UMaterialExpressionTextureSampleParameterCube : public UMaterialExpression
 	//~ End UMaterialExpression Interface
 
 	//~ Begin UMaterialExpressionTextureSampleParameter Interface
-	virtual bool TextureIsValid( UTexture* InTexture ) override;
-	virtual const TCHAR* GetRequirements() override;
+	virtual bool TextureIsValid(UTexture* InTexture, FString& OutMessage) override;
 	virtual void SetDefaultTexture() override;
 	//~ End UMaterialExpressionTextureSampleParameter Interface
 };
diff --git a/Engine/Source/Runtime/Engine/Classes/Materials/MaterialExpressionTextureSampleParameterSubUV.h b/Engine/Source/Runtime/Engine/Classes/Materials/MaterialExpressionTextureSampleParameterSubUV.h
--- a/Engine/Source/Runtime/Engine/Classes/Materials/MaterialExpressionTextureSampleParameterSubUV.h
+++ b/Engine/Source/Runtime/Engine/Classes/Materials/MaterialExpressionTextureSampleParameterSubUV.h
@@ -20,8 +20,7 @@ class UMaterialExpressionTextureSampleParameterSubUV : public UMaterialExpressio
 
 
 	//~ Begin UMaterialExpressionTextureSampleParameter Interface
-	virtual bool TextureIsValid( UTexture* InTexture ) override;
-	virtual const TCHAR* GetRequirements() override;
+	virtual bool TextureIsValid(UTexture* InTexture, FString& OutMessage) override;
 	//~ End UMaterialExpressionTextureSampleParameter Interface
 
 	//~ Begin UMaterialExpression Interface
diff --git a/Engine/Source/Runtime/Engine/Classes/Materials/MaterialExpressionTextureSampleParameterVolume.h b/Engine/Source/Runtime/Engine/Classes/Materials/MaterialExpressionTextureSampleParameterVolume.h
--- a/Engine/Source/Runtime/Engine/Classes/Materials/MaterialExpressionTextureSampleParameterVolume.h
+++ b/Engine/Source/Runtime/Engine/Classes/Materials/MaterialExpressionTextureSampleParameterVolume.h
@@ -22,8 +22,7 @@ class UMaterialExpressionTextureSampleParameterVolume : public UMaterialExpressi
 	//~ End UMaterialExpression Interface
 
 	//~ Begin UMaterialExpressionTextureSampleParameter Interface
-	virtual bool TextureIsValid( UTexture* InTexture ) override;
-	virtual const TCHAR* GetRequirements() override;
+	virtual bool TextureIsValid(UTexture* InTexture, FString& OutMessage) override;
 	virtual void SetDefaultTexture() override;
 	//~ End UMaterialExpressionTextureSampleParameter Interface
 };
diff --git a/Engine/Source/Runtime/Engine/Private/Materials/MaterialExpressions.cpp b/Engine/Source/Runtime/Engine/Private/Materials/MaterialExpressions.cpp
--- a/Engine/Source/Runtime/Engine/Private/Materials/MaterialExpressions.cpp
+++ b/Engine/Source/Runtime/Engine/Private/Materials/MaterialExpressions.cpp
@@ -2182,14 +2182,10 @@ UMaterialExpressionTextureSampleParameter::UMaterialExpressionTextureSampleParam
 #if WITH_EDITOR
 int32 UMaterialExpressionTextureSampleParameter::Compile(class FMaterialCompiler* Compiler, int32 OutputIndex)
 {
-	if (Texture == NULL)
+	FString ErrorMessage;
+	if (!TextureIsValid(Texture, ErrorMessage))
 	{
-		return CompilerError(Compiler, GetRequirements());
-	}
-
-	if (!TextureIsValid(Texture))
-	{
-		return CompilerError(Compiler, GetRequirements());
+		return CompilerError(Compiler, *ErrorMessage);
 	}
 
 	if (!VerifySamplerType(Compiler, (Desc.Len() > 0 ? *Desc : TEXT("TextureSampleParameter")), Texture, SamplerType))
@@ -2260,17 +2256,12 @@ bool UMaterialExpressionTextureSampleParameter::SetParameterValue(FName InParame
 }
 #endif
 
-bool UMaterialExpressionTextureSampleParameter::TextureIsValid( UTexture* /*InTexture*/ )
+bool UMaterialExpressionTextureSampleParameter::TextureIsValid(UTexture* /*InTexture*/, FString& OutMessage)
 {
+	OutMessage = TEXT("Invalid texture type");
 	return false;
 }
 
-const TCHAR* UMaterialExpressionTextureSampleParameter::GetRequirements()
-{
-	return TEXT("Invalid texture type");
-}
-
-
 void UMaterialExpressionTextureSampleParameter::SetDefaultTexture()
 {
 	// Does nothing in the base case...
@@ -2336,14 +2327,6 @@ void UMaterialExpressionTextureObjectParameter::GetCaption(TArray<FString>& OutC
 	OutCaptions.Add(TEXT("Param Tex Object")); 
 	OutCaptions.Add(FString::Printf(TEXT("'%s'"), *ParameterName.ToString()));
 }
-#endif // WITH_EDITOR
-
-const TCHAR* UMaterialExpressionTextureObjectParameter::GetRequirements()
-{
-	return TEXT("Requires valid texture");
-}
-
-#if WITH_EDITOR
 
 const TArray<FExpressionInput*> UMaterialExpressionTextureObjectParameter::GetInputs()
 {
@@ -2355,7 +2338,7 @@ int32 UMaterialExpressionTextureObjectParameter::Compile(class FMaterialCompiler
 {
 	if (!Texture)
 	{
-		return CompilerError(Compiler, GetRequirements());
+		return CompilerError(Compiler, TEXT("Requires valid texture"));
 	}
 
 	// It seems like this error should be checked here, but this can break existing materials, see https://jira.it.epicgames.net/browse/UE-68862
@@ -2371,7 +2354,7 @@ int32 UMaterialExpressionTextureObjectParameter::CompilePreview(class FMaterialC
 {
 	if (!Texture)
 	{
-		return CompilerError(Compiler, GetRequirements());
+		return CompilerError(Compiler, TEXT("Requires valid texture"));
 	}
 
 	// Preview the texture object by actually sampling it
@@ -2617,37 +2600,33 @@ void UMaterialExpressionTextureSampleParameter2D::GetCaption(TArray<FString>& Ou
 }
 #endif // WITH_EDITOR
 
-bool UMaterialExpressionTextureSampleParameter2D::TextureIsValid( UTexture* InTexture )
+bool UMaterialExpressionTextureSampleParameter2D::TextureIsValid(UTexture* InTexture, FString& OutMessage)
 {
-	bool Result=false;
-	if (InTexture)		
+	const bool bRequiresVirtualTexture = IsVirtualSamplerType(SamplerType);
+	if (!InTexture)
 	{
-		if (InTexture->IsA(UTexture2D::StaticClass()))
-		{
-			Result = true;
-		}
-		if( InTexture->IsA(UTextureRenderTarget2D::StaticClass()) )	
-		{
-			Result = true;
-		}
-		if ( InTexture->IsA(UTexture2DDynamic::StaticClass()) )
-		{
-			Result = true;
-		}
-		if ( InTexture->GetMaterialType() == MCT_TextureExternal )
-		{
-			Result = true;
-		}
+		OutMessage = TEXT("Found NULL, requires Texture2D");
+		return false;
+	}
+	else if (!(InTexture->GetMaterialType() & (MCT_Texture2D | MCT_TextureExternal | MCT_TextureVirtual)))
+	{
+		OutMessage = FString::Printf(TEXT("Found %s, requires Texture2D"), *InTexture->GetClass()->GetName());
+		return false;
+	}
+	else if (bRequiresVirtualTexture && !InTexture->VirtualTextureStreaming)
+	{
+		OutMessage = TEXT("Sampler requires VirtualTexture");
+		return false;
+	}
+	else if (!bRequiresVirtualTexture && InTexture->VirtualTextureStreaming)
+	{
+		OutMessage = TEXT("Sampler requires non-VirtualTexture");
+		return false;
 	}
-	return Result;
-}
 
-const TCHAR* UMaterialExpressionTextureSampleParameter2D::GetRequirements()
-{
-	return TEXT("Requires Texture2D");
+	return true;
 }
 
-
 void UMaterialExpressionTextureSampleParameter2D::SetDefaultTexture()
 {
 	Texture = LoadObject<UTexture2D>(NULL, TEXT("/Engine/EngineResources/DefaultTexture.DefaultTexture"), NULL, LOAD_None, NULL);
@@ -2725,24 +2704,20 @@ void UMaterialExpressionTextureSampleParameterCube::GetCaption(TArray<FString>&
 }
 #endif // WITH_EDITOR
 
-bool UMaterialExpressionTextureSampleParameterCube::TextureIsValid( UTexture* InTexture )
+bool UMaterialExpressionTextureSampleParameterCube::TextureIsValid(UTexture* InTexture, FString& OutMessage)
 {
-	bool Result=false;
-	if (InTexture)
+	if (!InTexture)
 	{
-		if( InTexture->GetClass() == UTextureCube::StaticClass() ) {
-			Result = true;
-		}
-		if( InTexture->IsA(UTextureRenderTargetCube::StaticClass()) ) {
-			Result = true;
-		}
+		OutMessage = TEXT("Found NULL, requires TextureCube");
+		return false;
+	}
+	else if (!(InTexture->GetMaterialType() & MCT_TextureCube))
+	{
+		OutMessage = FString::Printf(TEXT("Found %s, requires TextureCube"), *InTexture->GetClass()->GetName());
+		return false;
 	}
-	return Result;
-}
 
-const TCHAR* UMaterialExpressionTextureSampleParameterCube::GetRequirements()
-{
-	return TEXT("Requires TextureCube");
+	return true;
 }
 
 void UMaterialExpressionTextureSampleParameterCube::SetDefaultTexture()
@@ -2798,22 +2773,20 @@ void UMaterialExpressionTextureSampleParameterVolume::GetCaption(TArray<FString>
 }
 #endif // WITH_EDITOR
 
-bool UMaterialExpressionTextureSampleParameterVolume::TextureIsValid( UTexture* InTexture )
+bool UMaterialExpressionTextureSampleParameterVolume::TextureIsValid(UTexture* InTexture, FString& OutMessage)
 {
-	bool Result=false;
-	if (InTexture)
+	if (!InTexture)
 	{
-		if( InTexture->GetClass() == UVolumeTexture::StaticClass() )
-		{
-			Result = true;
-		}
+		OutMessage = TEXT("Found NULL, requires VolumeTexture");
+		return false;
+	}
+	else if (!(InTexture->GetMaterialType() & MCT_VolumeTexture))
+	{
+		OutMessage = FString::Printf(TEXT("Found %s, requires VolumeTexture"), *InTexture->GetClass()->GetName());
+		return false;
 	}
-	return Result;
-}
 
-const TCHAR* UMaterialExpressionTextureSampleParameterVolume::GetRequirements()
-{
-	return TEXT("Requires VolumeTexture");
+	return true;
 }
 
 void UMaterialExpressionTextureSampleParameterVolume::SetDefaultTexture()
@@ -2857,14 +2830,10 @@ UMaterialExpressionTextureSampleParameterSubUV::UMaterialExpressionTextureSample
 #if WITH_EDITOR
 int32 UMaterialExpressionTextureSampleParameterSubUV::Compile(class FMaterialCompiler* Compiler, int32 OutputIndex)
 {
-	if (Texture == NULL)
+	FString ErrorMessage;
+	if (!TextureIsValid(Texture, ErrorMessage))
 	{
-		return CompilerError(Compiler, GetRequirements());
-	}
-
-	if (!TextureIsValid(Texture))
-	{
-		return CompilerError(Compiler, GetRequirements());
+		return CompilerError(Compiler, *ErrorMessage);
 	}
 
 	if (!VerifySamplerType(Compiler, (Desc.Len() > 0 ? *Desc : TEXT("TextureSampleParameterSubUV")), Texture, SamplerType))
@@ -2882,14 +2851,9 @@ void UMaterialExpressionTextureSampleParameterSubUV::GetCaption(TArray<FString>&
 }
 #endif // WITH_EDITOR
 
-bool UMaterialExpressionTextureSampleParameterSubUV::TextureIsValid( UTexture* InTexture )
-{
-	return UMaterialExpressionTextureSampleParameter2D::TextureIsValid(InTexture);
-}
-
-const TCHAR* UMaterialExpressionTextureSampleParameterSubUV::GetRequirements()
+bool UMaterialExpressionTextureSampleParameterSubUV::TextureIsValid(UTexture* InTexture, FString& OutMessage)
 {
-	return UMaterialExpressionTextureSampleParameter2D::GetRequirements();
+	return UMaterialExpressionTextureSampleParameter2D::TextureIsValid(InTexture, OutMessage);
 }
 
 #if WITH_EDITOR
@@ -14224,9 +14188,10 @@ int32 UMaterialExpressionAntialiasedTextureMask::Compile(class FMaterialCompiler
 
 	int32 ArgCoord = Coordinates.Expression ? Coordinates.Compile(Compiler) : Compiler->TextureCoordinate(ConstCoordinate, false, false);
 
-	if (!TextureIsValid(Texture))
+	FString ErrorMessage;
+	if (!TextureIsValid(Texture, ErrorMessage))
 	{
-		return CompilerError(Compiler, GetRequirements());
+		return CompilerError(Compiler, *ErrorMessage);
 	}
 
 	int32 TextureCodeIndex;
@@ -14255,26 +14220,21 @@ void UMaterialExpressionAntialiasedTextureMask::GetCaption(TArray<FString>& OutC
 }
 #endif // WITH_EDITOR
 
-bool UMaterialExpressionAntialiasedTextureMask::TextureIsValid( UTexture* InTexture )
+bool UMaterialExpressionAntialiasedTextureMask::TextureIsValid(UTexture* InTexture, FString& OutMessage)
 {
-	bool Result=false;
-	if (InTexture)		
+	if (!InTexture)
 	{
-		if( InTexture->GetClass() == UTexture2D::StaticClass() ) 
-		{
-			Result = true;
-		}
-		if( InTexture->IsA(UTextureRenderTarget2D::StaticClass()) )	
-		{
-			Result = true;
-		}
+		OutMessage = TEXT("Found NULL, requires Texture2D");
+		return false;
 	}
-	return Result;
-}
-
-const TCHAR* UMaterialExpressionAntialiasedTextureMask::GetRequirements()
-{
-	return TEXT("Requires Texture2D");
+	// Doesn't allow virtual/external textures here
+	else if (!(InTexture->GetMaterialType() & MCT_Texture2D))
+	{
+		OutMessage = FString::Printf(TEXT("Found %s, requires Texture2D"), *InTexture->GetClass()->GetName());
+		return false;
+	}
+	
+	return true;
 }
 
 void UMaterialExpressionAntialiasedTextureMask::SetDefaultTexture()
diff --git a/Engine/Source/Runtime/Engine/Private/Materials/MaterialUniformExpressions.cpp b/Engine/Source/Runtime/Engine/Private/Materials/MaterialUniformExpressions.cpp
--- a/Engine/Source/Runtime/Engine/Private/Materials/MaterialUniformExpressions.cpp
+++ b/Engine/Source/Runtime/Engine/Private/Materials/MaterialUniformExpressions.cpp
@@ -586,11 +586,15 @@ void FUniformExpressionSet::FillUniformBuffer(const FMaterialRenderContext& Mate
 			BufferCursor = ((uint8*)BufferCursor) + (SHADER_PARAMETER_POINTER_ALIGNMENT * 2);
 			check(BufferCursor <= TempBuffer + TempBufferSize);
 
+			// ExternalTexture is allowed here, with warning above
+			// VirtualTexture is allowed here, as these may be demoted to regular textures on platforms that don't have VT support
+			const uint32 ValidTextureTypes = MCT_Texture2D | MCT_TextureVirtual | MCT_TextureExternal;
+
 			// TextureReference.TextureReferenceRHI is cleared from a render command issued by UTexture::BeginDestroy
 			// It's possible for this command to trigger before a given material is cleaned up and removed from deferred update list
-			// Technically I don't think it's necesary to check 'Resource' for nullptr here, as if TextureReferenceRHI has been initialized, that should be enough
+			// Technically I don't think it's necessary to check 'Resource' for nullptr here, as if TextureReferenceRHI has been initialized, that should be enough
 			// Going to leave the check for now though, to hopefully avoid any unexpected problems
-			if (Value && Value->Resource && Value->TextureReference.TextureReferenceRHI)
+			if (Value && Value->Resource && Value->TextureReference.TextureReferenceRHI && (Value->GetMaterialType() & ValidTextureTypes) != 0u)
 			{
 				//@todo-rco: Help track down a invalid values
 				checkf(Value->IsA(UTexture::StaticClass()), TEXT("Expecting a UTexture! Value='%s' class='%s'"), *Value->GetName(), *Value->GetClass()->GetName());
@@ -633,7 +637,7 @@ void FUniformExpressionSet::FillUniformBuffer(const FMaterialRenderContext& Mate
 			BufferCursor = ((uint8*)BufferCursor) + (SHADER_PARAMETER_POINTER_ALIGNMENT * 2);
 			check(BufferCursor <= TempBuffer + TempBufferSize);
 
-			if(Value && Value->Resource)
+			if(Value && Value->Resource && (Value->GetMaterialType() & MCT_TextureCube) != 0u)
 			{
 				check(Value->TextureReference.TextureReferenceRHI);
 				*ResourceTableTexturePtr = Value->TextureReference.TextureReferenceRHI;
@@ -672,7 +676,7 @@ void FUniformExpressionSet::FillUniformBuffer(const FMaterialRenderContext& Mate
 			BufferCursor = ((uint8*)BufferCursor) + (SHADER_PARAMETER_POINTER_ALIGNMENT * 2);
 			check(BufferCursor <= TempBuffer + TempBufferSize);
 
-			if(Value && Value->Resource)
+			if(Value && Value->Resource && (Value->GetMaterialType() & MCT_VolumeTexture) != 0u)
 			{
 				check(Value->TextureReference.TextureReferenceRHI);
 				*ResourceTableTexturePtr = Value->TextureReference.TextureReferenceRHI;
