Fix crashes when material instance texture parameter doesn't have VT set when expected
- UI normally enforces this, but possible to get into bad state by modifying base material

diff --git a/Engine/Source/Editor/UnrealEd/Private/Lightmass/LightmassRender.cpp b/Engine/Source/Editor/UnrealEd/Private/Lightmass/LightmassRender.cpp
--- a/Engine/Source/Editor/UnrealEd/Private/Lightmass/LightmassRender.cpp
+++ b/Engine/Source/Editor/UnrealEd/Private/Lightmass/LightmassRender.cpp
@@ -1136,7 +1136,10 @@ bool FLightmassMaterialRenderer::GenerateMaterialPropertyData(
 						const FUniformExpressionCache& UniformExpressionCache = InMaterialProxy->UniformExpressionCache[InFeatureLevel];
 						for (IAllocatedVirtualTexture* AllocatedVT : UniformExpressionCache.AllocatedVTs)
 						{
-							GetRendererModule().RequestVirtualTextureTilesForRegion(AllocatedVT, InScreenSpaceSize, FIntRect(), -1);
+							if (AllocatedVT)
+							{
+								GetRendererModule().RequestVirtualTextureTilesForRegion(AllocatedVT, InScreenSpaceSize, FIntRect(), -1);
+							}
 						}
 
 						ENQUEUE_RENDER_COMMAND(LoadTiles)(
diff --git a/Engine/Source/Runtime/Engine/Private/Materials/MaterialShared.cpp b/Engine/Source/Runtime/Engine/Private/Materials/MaterialShared.cpp
--- a/Engine/Source/Runtime/Engine/Private/Materials/MaterialShared.cpp
+++ b/Engine/Source/Runtime/Engine/Private/Materials/MaterialShared.cpp
@@ -2296,27 +2296,34 @@ static IAllocatedVirtualTexture* AllocateVTStack(const FMaterialRenderContext& C
 
 	const UTexture2D* LayerTextures[VIRTUALTEXTURE_SPACE_MAXLAYERS] = { nullptr };
 	VTStack.GetTextureValues(Context, UniformExpressionSet, LayerTextures);
-	check(LayerTextures[0]);
-	check(LayerTextures[0]->IsCurrentlyVirtualTextured());
-	const FVirtualTexture2DResource* VirtualTextureResource = (FVirtualTexture2DResource*)LayerTextures[0]->Resource;
 
 	FAllocatedVTDescription VTDesc;
 	VTDesc.Dimensions = 2;
-	VTDesc.TileSize = VirtualTextureResource->GetTileSize();
-	VTDesc.TileBorderSize = VirtualTextureResource->GetBorderSize();
 	VTDesc.NumLayers = NumLayers;
+	bool bFoundValidLayer = false;
 	for (uint32 LayerIndex = 0u; LayerIndex < VTDesc.NumLayers; ++LayerIndex)
 	{
-		check(LayerTextures[LayerIndex]->IsCurrentlyVirtualTextured());
-		const FVirtualTexture2DResource* VirtualTextureResourceForLayer = (FVirtualTexture2DResource*)LayerTextures[LayerIndex]->Resource;
+		const UTexture2D* Texture = LayerTextures[LayerIndex];
+		const FVirtualTexture2DResource* VirtualTextureResourceForLayer = (Texture && Texture->IsCurrentlyVirtualTextured()) ? (FVirtualTexture2DResource*)Texture->Resource : nullptr;
 		if (VirtualTextureResourceForLayer != nullptr)
 		{
+			// All tile sizes need to match
+			check(!bFoundValidLayer || VTDesc.TileSize == VirtualTextureResourceForLayer->GetTileSize());
+			check(!bFoundValidLayer || VTDesc.TileBorderSize == VirtualTextureResourceForLayer->GetBorderSize());
+
+			VTDesc.TileSize = VirtualTextureResourceForLayer->GetTileSize();
+			VTDesc.TileBorderSize = VirtualTextureResourceForLayer->GetBorderSize();
 			VTDesc.ProducerHandle[LayerIndex] = VirtualTextureResourceForLayer->GetProducerHandle();
 			VTDesc.LocalLayerToProduce[LayerIndex] = 0u;
+			bFoundValidLayer = true;
 		}
 	}
 
-	return GetRendererModule().AllocateVirtualTexture(VTDesc);
+	if (bFoundValidLayer)
+	{
+		return GetRendererModule().AllocateVirtualTexture(VTDesc);
+	}
+	return nullptr;
 }
 
 FUniformExpressionCache::~FUniformExpressionCache()
diff --git a/Engine/Source/Runtime/Renderer/Private/VT/VirtualTextureProducer.cpp b/Engine/Source/Runtime/Renderer/Private/VT/VirtualTextureProducer.cpp
--- a/Engine/Source/Runtime/Renderer/Private/VT/VirtualTextureProducer.cpp
+++ b/Engine/Source/Runtime/Renderer/Private/VT/VirtualTextureProducer.cpp
@@ -27,6 +27,7 @@ void FVirtualTextureProducer::Release(FVirtualTextureSystem* System, const FVirt
 FVirtualTextureProducerCollection::FVirtualTextureProducerCollection()
 {
 	Producers.AddDefaulted(1);
+	Producers[0].Magic = 1u; // make sure FVirtualTextureProducerHandle(0) will not resolve to the dummy producer entry
 }
 
 FVirtualTextureProducerHandle FVirtualTextureProducerCollection::RegisterProducer(FVirtualTextureSystem* System, const FVTProducerDescription& InDesc, IVirtualTexture* InProducer)
diff --git a/Engine/Source/Runtime/Renderer/Private/VT/VirtualTextureSystem.cpp b/Engine/Source/Runtime/Renderer/Private/VT/VirtualTextureSystem.cpp
--- a/Engine/Source/Runtime/Renderer/Private/VT/VirtualTextureSystem.cpp
+++ b/Engine/Source/Runtime/Renderer/Private/VT/VirtualTextureSystem.cpp
@@ -1729,11 +1729,17 @@ void FVirtualTextureSystem::SubmitRequests(FRHICommandListImmediate& RHICmdList,
 			uint32 NumFullyMappedLayers = 0u;
 			for (uint32 LayerIndex = 0u; LayerIndex < AllocatedVT->GetNumLayers(); ++LayerIndex)
 			{
+				const FVirtualTextureProducerHandle ProducerHandle = AllocatedVT->GetProducerHandle(LayerIndex);
+				const FVirtualTextureProducer* Producer = Producers.FindProducer(ProducerHandle);
+				if (!Producer)
+				{
+					++NumFullyMappedLayers;
+					continue;
+				}
+
 				const uint32 ProducerIndex = AllocatedVT->GetUniqueProducerIndexForLayer(LayerIndex);
 				const uint32 ProducerMipBias = AllocatedVT->GetUniqueProducerMipBias(ProducerIndex);
 				const uint32 LocalLayerIndex = AllocatedVT->GetLocalLayerToProduce(LayerIndex);
-				const FVirtualTextureProducerHandle ProducerHandle = AllocatedVT->GetUniqueProducerHandle(ProducerIndex);
-				const FVirtualTextureProducer* Producer = Producers.FindProducer(ProducerHandle);
 				const uint32 WidthInTiles = Producer->GetWidthInTiles();
 				const uint32 HeightInTiles = Producer->GetHeightInTiles();
 				const uint32 MaxLevel = FMath::Min(Producer->GetMaxLevel(), AllocatedVT->GetMaxLevel());
