Fix GPUSkinCache GPU crash when ray tracing is enabled. Skeletal meshes can receive multiple updates every frame which causes mismatched states and BLAS refit failure

diff --git a/Engine/Source/Runtime/Engine/Private/GPUSkinCache.cpp b/Engine/Source/Runtime/Engine/Private/GPUSkinCache.cpp
--- a/Engine/Source/Runtime/Engine/Private/GPUSkinCache.cpp
+++ b/Engine/Source/Runtime/Engine/Private/GPUSkinCache.cpp
@@ -608,9 +608,20 @@ void FGPUSkinCache::TransitionAllToReadable(FRHICommandList& RHICmdList)
 #if RHI_RAYTRACING
 void FGPUSkinCache::CommitRayTracingGeometryUpdates(FRHICommandList& RHICmdList)
 {
+	SCOPED_DRAW_EVENT(RHICmdList, CommitSkeletalRayTracingGeometryUpdates);
+
 	if (RayTracingGeometriesToUpdate.Num())
 	{
-		RHICmdList.UpdateAccelerationStructures(RayTracingGeometriesToUpdate);
+		TArray<FAccelerationStructureUpdateParams> Updates;
+		for (const auto& RayTracingGeometry : RayTracingGeometriesToUpdate)
+		{
+			FAccelerationStructureUpdateParams Params;
+			Params.Geometry = RayTracingGeometry->RayTracingGeometryRHI;
+			Params.VertexBuffer = RayTracingGeometry->Initializer.PositionVertexBuffer;
+			Updates.Add(Params);
+		}
+
+		RHICmdList.UpdateAccelerationStructures(Updates);
 		RayTracingGeometriesToUpdate.Reset();
 	}
 }
@@ -1322,6 +1333,7 @@ void FGPUSkinCache::FRWBuffersAllocation::RemoveAllFromTransitionArray(TArray<FU
 void FGPUSkinCache::ReleaseSkinCacheEntry(FGPUSkinCacheEntry* SkinCacheEntry)
 {
 	FGPUSkinCache* SkinCache = SkinCacheEntry->SkinCache;
+	SkinCache->RemoveRayTracingGeometryUpdate(&SkinCacheEntry->GPUSkin->RayTracingGeometry);
 	FRWBuffersAllocation* PositionAllocation = SkinCacheEntry->PositionAllocation;
 	if (PositionAllocation)
 	{
diff --git a/Engine/Source/Runtime/Engine/Public/GPUSkinCache.h b/Engine/Source/Runtime/Engine/Public/GPUSkinCache.h
--- a/Engine/Source/Runtime/Engine/Public/GPUSkinCache.h
+++ b/Engine/Source/Runtime/Engine/Public/GPUSkinCache.h
@@ -283,19 +283,21 @@ public:
 #if RHI_RAYTRACING
 	void AddRayTracingGeometryToUpdate(FRayTracingGeometry* RayTracingGeometry)
 	{
-		FAccelerationStructureUpdateParams Params;
-		Params.Geometry     = RayTracingGeometry->RayTracingGeometryRHI;
-		Params.VertexBuffer = RayTracingGeometry->Initializer.PositionVertexBuffer;
-
-		RayTracingGeometriesToUpdate.Add(Params);
+		RayTracingGeometriesToUpdate.Add(RayTracingGeometry);
 	}
 
 	void CommitRayTracingGeometryUpdates(FRHICommandList& RHICmdList);
+
+	void RemoveRayTracingGeometryUpdate(FRayTracingGeometry* RayTracingGeometry)
+	{
+		if (RayTracingGeometriesToUpdate.Find(RayTracingGeometry) != nullptr)
+			RayTracingGeometriesToUpdate.Remove(RayTracingGeometry);
+	}
 #endif // RHI_RAYTRACING
 
 protected:
 	TArray<FUnorderedAccessViewRHIParamRef> BuffersToTransition;
-	TArray<FAccelerationStructureUpdateParams> RayTracingGeometriesToUpdate;
+	TSet<FRayTracingGeometry*> RayTracingGeometriesToUpdate;
 
 	TArray<FRWBuffersAllocation*> Allocations;
 	TArray<FGPUSkinCacheEntry*> Entries;
