Fix crashes when building light map in levels with foiage and the foliage.DensityScale < 1
With density scaling the removed instances need to be skipped when setting up light map UVs after encoding

diff --git a/Engine/Source/Runtime/Engine/Private/HierarchicalInstancedStaticMesh.cpp b/Engine/Source/Runtime/Engine/Private/HierarchicalInstancedStaticMesh.cpp
--- a/Engine/Source/Runtime/Engine/Private/HierarchicalInstancedStaticMesh.cpp
+++ b/Engine/Source/Runtime/Engine/Private/HierarchicalInstancedStaticMesh.cpp
@@ -2844,9 +2844,11 @@ void UHierarchicalInstancedStaticMeshComponent::PropagateLightingScenarioChange(
 				for (int32 InstanceIndex = 0; InstanceIndex < PerInstanceSMData.Num(); ++InstanceIndex)
 				{
 					int32 RenderIndex = InstanceReorderTable.IsValidIndex(InstanceIndex) ? InstanceReorderTable[InstanceIndex] : InstanceIndex;
-
-					InstanceUpdateCmdBuffer.SetLightMapData(RenderIndex, MeshMapBuildData->PerInstanceLightmapData[InstanceIndex].LightmapUVBias);
-					InstanceUpdateCmdBuffer.SetShadowMapData(RenderIndex, MeshMapBuildData->PerInstanceLightmapData[InstanceIndex].ShadowmapUVBias);
+					if (RenderIndex != INDEX_NONE)
+					{
+						InstanceUpdateCmdBuffer.SetLightMapData(RenderIndex, MeshMapBuildData->PerInstanceLightmapData[InstanceIndex].LightmapUVBias);
+						InstanceUpdateCmdBuffer.SetShadowMapData(RenderIndex, MeshMapBuildData->PerInstanceLightmapData[InstanceIndex].ShadowmapUVBias);
+					}
 				}
 			}
 		}
diff --git a/Engine/Source/Runtime/Engine/Private/LightMap.cpp b/Engine/Source/Runtime/Engine/Private/LightMap.cpp
--- a/Engine/Source/Runtime/Engine/Private/LightMap.cpp
+++ b/Engine/Source/Runtime/Engine/Private/LightMap.cpp
@@ -299,9 +299,12 @@ struct FLightMapAllocation
 				// Need to create per-LOD instance data to fix that
 				MeshBuildData->PerInstanceLightmapData[InstanceIndex].LightmapUVBias = LightMap->GetCoordinateBias();
 
-				int32 RenderIndex = (Component->InstanceReorderTable.IsValidIndex(InstanceIndex) && Component->InstanceReorderTable[InstanceIndex] != INDEX_NONE) ? Component->InstanceReorderTable[InstanceIndex] : InstanceIndex;
-				Component->InstanceUpdateCmdBuffer.SetLightMapData(RenderIndex, MeshBuildData->PerInstanceLightmapData[InstanceIndex].LightmapUVBias);
-				Component->MarkRenderStateDirty();
+				int32 RenderIndex = Component->InstanceReorderTable.IsValidIndex(InstanceIndex) ? Component->InstanceReorderTable[InstanceIndex] : InstanceIndex;
+				if (RenderIndex != INDEX_NONE)
+				{
+					Component->InstanceUpdateCmdBuffer.SetLightMapData(RenderIndex, MeshBuildData->PerInstanceLightmapData[InstanceIndex].LightmapUVBias);
+					Component->MarkRenderStateDirty();
+				}
 			}
 		}
 	}
diff --git a/Engine/Source/Runtime/Engine/Private/ShadowMap.cpp b/Engine/Source/Runtime/Engine/Private/ShadowMap.cpp
--- a/Engine/Source/Runtime/Engine/Private/ShadowMap.cpp
+++ b/Engine/Source/Runtime/Engine/Private/ShadowMap.cpp
@@ -109,10 +109,12 @@ struct FShadowMapAllocation
 				// TODO: We currently only support one LOD of static lighting in foliage
 				// Need to create per-LOD instance data to fix that
 				MeshBuildData->PerInstanceLightmapData[InstanceIndex].ShadowmapUVBias = ShadowMap->GetCoordinateBias();
-				int32 RenderIndex = (Component->InstanceReorderTable.IsValidIndex(InstanceIndex) && Component->InstanceReorderTable[InstanceIndex] != INDEX_NONE) ? Component->InstanceReorderTable[InstanceIndex] : InstanceIndex;
-
-				Component->InstanceUpdateCmdBuffer.SetShadowMapData(RenderIndex, MeshBuildData->PerInstanceLightmapData[InstanceIndex].ShadowmapUVBias);
-				Component->MarkRenderStateDirty();
+				int32 RenderIndex = Component->InstanceReorderTable.IsValidIndex(InstanceIndex) ? Component->InstanceReorderTable[InstanceIndex] : InstanceIndex;
+				if (RenderIndex != INDEX_NONE)
+				{
+					Component->InstanceUpdateCmdBuffer.SetShadowMapData(RenderIndex, MeshBuildData->PerInstanceLightmapData[InstanceIndex].ShadowmapUVBias);
+					Component->MarkRenderStateDirty();
+				}
 			}
 		}
 	}
