Fix compile on non-ray-tracing platforms

diff --git a/Engine/Source/Runtime/Engine/Private/GPUSkinCache.cpp b/Engine/Source/Runtime/Engine/Private/GPUSkinCache.cpp
--- a/Engine/Source/Runtime/Engine/Private/GPUSkinCache.cpp
+++ b/Engine/Source/Runtime/Engine/Private/GPUSkinCache.cpp
@@ -1333,7 +1333,9 @@ void FGPUSkinCache::FRWBuffersAllocation::RemoveAllFromTransitionArray(TArray<FU
 void FGPUSkinCache::ReleaseSkinCacheEntry(FGPUSkinCacheEntry* SkinCacheEntry)
 {
 	FGPUSkinCache* SkinCache = SkinCacheEntry->SkinCache;
+#if RHI_RAYTRACING
 	SkinCache->RemoveRayTracingGeometryUpdate(&SkinCacheEntry->GPUSkin->RayTracingGeometry);
+#endif // RHI_RAYTRACING
 	FRWBuffersAllocation* PositionAllocation = SkinCacheEntry->PositionAllocation;
 	if (PositionAllocation)
 	{
diff --git a/Engine/Source/Runtime/Engine/Public/GPUSkinCache.h b/Engine/Source/Runtime/Engine/Public/GPUSkinCache.h
--- a/Engine/Source/Runtime/Engine/Public/GPUSkinCache.h
+++ b/Engine/Source/Runtime/Engine/Public/GPUSkinCache.h
@@ -297,7 +297,9 @@ public:
 
 protected:
 	TArray<FUnorderedAccessViewRHIParamRef> BuffersToTransition;
+#if RHI_RAYTRACING
 	TSet<FRayTracingGeometry*> RayTracingGeometriesToUpdate;
+#endif // RHI_RAYTRACING
 
 	TArray<FRWBuffersAllocation*> Allocations;
 	TArray<FGPUSkinCacheEntry*> Entries;
