Tile and Triangle rendering now use Index Buffers to work around DrawPrimitive + BaseVertexOffset + Manual Vertex Fetch issue.

DrawPrimitive converts FirstIndex into BaseVertexOffset. Unfortunately, SV_VertexId does not take this into account, which results in incorrect vertex fetch. However, DrawIndexedPrimitive with FirstIndex does work correctly, as SV_VertexId becomes the indexed vertex instead. To work around this issue, Tile and Triangle canvas rendering now uses index buffers.

diff --git a/Engine/Source/Runtime/Engine/Private/TileRendering.cpp b/Engine/Source/Runtime/Engine/Private/TileRendering.cpp
--- a/Engine/Source/Runtime/Engine/Private/TileRendering.cpp
+++ b/Engine/Source/Runtime/Engine/Private/TileRendering.cpp
@@ -19,7 +19,7 @@
 #include "EngineModule.h"
 #include "MeshPassProcessor.h"
 
-#define NUM_MATERIAL_TILE_VERTS	6
+#define NUM_MATERIAL_TILE_VERTS	4
 
 /** 
 * Vertex buffer
@@ -43,7 +43,7 @@ public:
 	virtual void InitRHI() override
 	{
 		FVector* PositionBufferData = nullptr;
-		// used with a trilist, so 6 vertices are needed
+
 		uint32 PositionSize = NUM_MATERIAL_TILE_VERTS * sizeof(FVector);
 		// create vertex buffer
 		{
@@ -59,7 +59,7 @@ public:
 		}
 
 		FPackedNormal* TangentBufferData = nullptr;
-		// used with a trilist, so 6 vertices are needed
+
 		uint32 TangentSize = NUM_MATERIAL_TILE_VERTS * 2 * sizeof(FPackedNormal);
 		// create vertex buffer
 		{
@@ -74,7 +74,7 @@ public:
 		}
 
 		FVector2D* TexCoordBufferData = nullptr;
-		// used with a trilist, so 6 vertices are needed
+
 		uint32 TexCoordSize = NUM_MATERIAL_TILE_VERTS * sizeof(FVector2D);
 		// create vertex buffer
 		{
@@ -90,7 +90,7 @@ public:
 		}
 
 		uint32* ColorBufferData = nullptr;
-		// used with a trilist, so 6 vertices are needed
+
 		uint32 ColorSize = NUM_MATERIAL_TILE_VERTS * sizeof(uint32);
 		// create vertex buffer
 		{
@@ -108,16 +108,12 @@ public:
 		PositionBufferData[0] = FVector( 1, -1, 0);
 		PositionBufferData[1] = FVector( 1,  1, 0);
 		PositionBufferData[2] = FVector(-1, -1, 0);
-		PositionBufferData[3] = FVector(-1, -1, 0);
-		PositionBufferData[4] = FVector( 1,  1, 0);	
-		PositionBufferData[5] = FVector(-1,  1, 0);
+		PositionBufferData[3] = FVector(-1,  1, 0);
 
 		TexCoordBufferData[0] = FVector2D(1, 1);
 		TexCoordBufferData[1] = FVector2D(1, 0);
 		TexCoordBufferData[2] = FVector2D(0, 1);
-		TexCoordBufferData[3] = FVector2D(0, 1);
-		TexCoordBufferData[4] = FVector2D(1, 0);
-		TexCoordBufferData[5] = FVector2D(0, 0);
+		TexCoordBufferData[3] = FVector2D(0, 0);
 
 		for (int i = 0; i < NUM_MATERIAL_TILE_VERTS; i++)
 		{
@@ -179,8 +175,11 @@ FCanvasTileRendererItem::FTileVertexFactory::FTileVertexFactory(ERHIFeatureLevel
 /**
  * Mesh used to render tiles.
  */
-FCanvasTileRendererItem::FTileMesh::FTileMesh(FCanvasTileRendererItem::FTileVertexFactory* InVertexFactory)
-	: VertexFactory(InVertexFactory)
+FCanvasTileRendererItem::FTileMesh::FTileMesh(
+	const FRawIndexBuffer* InIndexBuffer,
+	const FCanvasTileRendererItem::FTileVertexFactory* InVertexFactory)
+	: IndexBuffer(InIndexBuffer)
+	, VertexFactory(InVertexFactory)
 {
 }
 
@@ -188,6 +187,7 @@ void FCanvasTileRendererItem::FTileMesh::InitRHI()
 {
 	FMeshBatchElement& BatchElement = MeshElement.Elements[0];
 	MeshElement.VertexFactory = VertexFactory;
+	BatchElement.IndexBuffer = IndexBuffer;
 	BatchElement.FirstIndex = 0;
 	BatchElement.NumPrimitives = 2;
 	BatchElement.MinVertexIndex = 0;
@@ -206,7 +206,7 @@ FCanvasTileRendererItem::FRenderData::FRenderData(ERHIFeatureLevel::Type InFeatu
 	const FMaterialRenderProxy* InMaterialRenderProxy,
 	const FCanvas::FTransformEntry& InTransform )
 	: VertexFactory(InFeatureLevel)
-	, TileMesh(&VertexFactory)
+	, TileMesh(&IndexBuffer, &VertexFactory)
 	, MaterialRenderProxy(InMaterialRenderProxy)
 	, Transform(InTransform)
 {
@@ -270,14 +270,24 @@ FCanvasTileRendererItem::~FCanvasTileRendererItem()
 
 void FCanvasTileRendererItem::InitTileBuffers(FLocalVertexFactory* VertexFactory, TArray<FTileInst>& Tiles, const FSceneView& View, bool bNeedsToSwitchVerticalAxis)
 {
-	static_assert(NUM_MATERIAL_TILE_VERTS == 6, "Invalid tile tri-list size.");
+	static_assert(NUM_MATERIAL_TILE_VERTS == 4, "Invalid tile tri-list size.");
 	Data->StaticMeshVertexBuffers.PositionVertexBuffer.Init(Tiles.Num() * NUM_MATERIAL_TILE_VERTS);
 	Data->StaticMeshVertexBuffers.StaticMeshVertexBuffer.Init(Tiles.Num() * NUM_MATERIAL_TILE_VERTS, 1);
 	Data->StaticMeshVertexBuffers.ColorVertexBuffer.Init(Tiles.Num() * NUM_MATERIAL_TILE_VERTS);
 
+	Data->IndexBuffer.Indices.SetNum(Tiles.Num() * NUM_MATERIAL_TILE_VERTS);
+
 	for (int32 i = 0; i < Tiles.Num(); i++)
 	{
 		const FTileInst& Tile = Tiles[i];
+		const uint32 FirstIndex = i * 6;
+
+		Data->IndexBuffer.Indices[FirstIndex + 0] = FirstIndex + 0;
+		Data->IndexBuffer.Indices[FirstIndex + 1] = FirstIndex + 1;
+		Data->IndexBuffer.Indices[FirstIndex + 2] = FirstIndex + 2;
+		Data->IndexBuffer.Indices[FirstIndex + 3] = FirstIndex + 2;
+		Data->IndexBuffer.Indices[FirstIndex + 4] = FirstIndex + 1;
+		Data->IndexBuffer.Indices[FirstIndex + 5] = FirstIndex + 3;
 
 		const float X = Tile.X;
 		const float Y = Tile.Y;
@@ -293,32 +303,24 @@ void FCanvasTileRendererItem::InitTileBuffers(FLocalVertexFactory* VertexFactory
 			Data->StaticMeshVertexBuffers.PositionVertexBuffer.VertexPosition(i * NUM_MATERIAL_TILE_VERTS + 0) = FVector(X + SizeX, View.UnscaledViewRect.Height() - (Y + SizeY), 0.0f);
 			Data->StaticMeshVertexBuffers.PositionVertexBuffer.VertexPosition(i * NUM_MATERIAL_TILE_VERTS + 1) = FVector(X, View.UnscaledViewRect.Height() - (Y + SizeY), 0.0f);
 			Data->StaticMeshVertexBuffers.PositionVertexBuffer.VertexPosition(i * NUM_MATERIAL_TILE_VERTS + 2) = FVector(X + SizeX, View.UnscaledViewRect.Height() - Y, 0.0f);
-			Data->StaticMeshVertexBuffers.PositionVertexBuffer.VertexPosition(i * NUM_MATERIAL_TILE_VERTS + 3) = FVector(X + SizeX, View.UnscaledViewRect.Height() - Y, 0.0f);
-			Data->StaticMeshVertexBuffers.PositionVertexBuffer.VertexPosition(i * NUM_MATERIAL_TILE_VERTS + 4) = FVector(X, View.UnscaledViewRect.Height() - (Y + SizeY), 0.0f);
-			Data->StaticMeshVertexBuffers.PositionVertexBuffer.VertexPosition(i * NUM_MATERIAL_TILE_VERTS + 5) = FVector(X, View.UnscaledViewRect.Height() - Y, 0.0f);
+			Data->StaticMeshVertexBuffers.PositionVertexBuffer.VertexPosition(i * NUM_MATERIAL_TILE_VERTS + 3) = FVector(X, View.UnscaledViewRect.Height() - Y, 0.0f);
 
 			Data->StaticMeshVertexBuffers.StaticMeshVertexBuffer.SetVertexUV(i * NUM_MATERIAL_TILE_VERTS + 0, 0, FVector2D(U + SizeU, V + SizeV));
 			Data->StaticMeshVertexBuffers.StaticMeshVertexBuffer.SetVertexUV(i * NUM_MATERIAL_TILE_VERTS + 1, 0, FVector2D(U, V + SizeV));
 			Data->StaticMeshVertexBuffers.StaticMeshVertexBuffer.SetVertexUV(i * NUM_MATERIAL_TILE_VERTS + 2, 0, FVector2D(U + SizeU, V));
-			Data->StaticMeshVertexBuffers.StaticMeshVertexBuffer.SetVertexUV(i * NUM_MATERIAL_TILE_VERTS + 3, 0, FVector2D(U + SizeU, V));
-			Data->StaticMeshVertexBuffers.StaticMeshVertexBuffer.SetVertexUV(i * NUM_MATERIAL_TILE_VERTS + 4, 0, FVector2D(U, V + SizeV));
-			Data->StaticMeshVertexBuffers.StaticMeshVertexBuffer.SetVertexUV(i * NUM_MATERIAL_TILE_VERTS + 5, 0, FVector2D(U, V));
+			Data->StaticMeshVertexBuffers.StaticMeshVertexBuffer.SetVertexUV(i * NUM_MATERIAL_TILE_VERTS + 3, 0, FVector2D(U, V));
 		}
 		else
 		{
 			Data->StaticMeshVertexBuffers.PositionVertexBuffer.VertexPosition(i * NUM_MATERIAL_TILE_VERTS + 0) = FVector(X + SizeX, Y, 0.0f);
 			Data->StaticMeshVertexBuffers.PositionVertexBuffer.VertexPosition(i * NUM_MATERIAL_TILE_VERTS + 1) = FVector(X, Y, 0.0f);
 			Data->StaticMeshVertexBuffers.PositionVertexBuffer.VertexPosition(i * NUM_MATERIAL_TILE_VERTS + 2) = FVector(X + SizeX, Y + SizeY, 0.0f);
-			Data->StaticMeshVertexBuffers.PositionVertexBuffer.VertexPosition(i * NUM_MATERIAL_TILE_VERTS + 3) = FVector(X + SizeX, Y + SizeY, 0.0f);
-			Data->StaticMeshVertexBuffers.PositionVertexBuffer.VertexPosition(i * NUM_MATERIAL_TILE_VERTS + 4) = FVector(X, Y, 0.0f);
-			Data->StaticMeshVertexBuffers.PositionVertexBuffer.VertexPosition(i * NUM_MATERIAL_TILE_VERTS + 5) = FVector(X, Y + SizeY, 0.0f);
+			Data->StaticMeshVertexBuffers.PositionVertexBuffer.VertexPosition(i * NUM_MATERIAL_TILE_VERTS + 3) = FVector(X, Y + SizeY, 0.0f);
 
 			Data->StaticMeshVertexBuffers.StaticMeshVertexBuffer.SetVertexUV(i * NUM_MATERIAL_TILE_VERTS + 0, 0, FVector2D(U + SizeU, V));
 			Data->StaticMeshVertexBuffers.StaticMeshVertexBuffer.SetVertexUV(i * NUM_MATERIAL_TILE_VERTS + 1, 0, FVector2D(U, V));
 			Data->StaticMeshVertexBuffers.StaticMeshVertexBuffer.SetVertexUV(i * NUM_MATERIAL_TILE_VERTS + 2, 0, FVector2D(U + SizeU, V + SizeV));
-			Data->StaticMeshVertexBuffers.StaticMeshVertexBuffer.SetVertexUV(i * NUM_MATERIAL_TILE_VERTS + 3, 0, FVector2D(U + SizeU, V + SizeV));
-			Data->StaticMeshVertexBuffers.StaticMeshVertexBuffer.SetVertexUV(i * NUM_MATERIAL_TILE_VERTS + 4, 0, FVector2D(U, V));
-			Data->StaticMeshVertexBuffers.StaticMeshVertexBuffer.SetVertexUV(i * NUM_MATERIAL_TILE_VERTS + 5, 0, FVector2D(U, V + SizeV));
+			Data->StaticMeshVertexBuffers.StaticMeshVertexBuffer.SetVertexUV(i * NUM_MATERIAL_TILE_VERTS + 3, 0, FVector2D(U, V + SizeV));
 		}
 
 		for (int j = 0; j < NUM_MATERIAL_TILE_VERTS; j++)
@@ -344,6 +346,7 @@ void FCanvasTileRendererItem::InitTileBuffers(FLocalVertexFactory* VertexFactory
 		RenderData->StaticMeshVertexBuffers.ColorVertexBuffer.BindColorVertexBuffer(VertexFactory, RData);
 		VertexFactory->SetData(RData);
 
+		RenderData->IndexBuffer.InitResource();
 		RenderData->VertexFactory.InitResource();
 		RenderData->TileMesh.InitResource();
 	});
@@ -398,7 +401,7 @@ bool FCanvasTileRendererItem::Render_RenderThread(FRHICommandListImmediate& RHIC
 		// update the FMeshBatch
 		FMeshBatch& Mesh = Data->TileMesh.MeshElement;
 		Mesh.MaterialRenderProxy = Data->MaterialRenderProxy;
-		Mesh.Elements[0].BaseVertexIndex = NUM_MATERIAL_TILE_VERTS * TileIdx;
+		Mesh.Elements[0].FirstIndex = 6 * TileIdx;
 
 		GetRendererModule().DrawTileMesh(RHICmdList, DrawRenderState, *View, Mesh, Canvas->IsHitTesting(), Tile.HitProxyId);
 	}
@@ -502,6 +505,7 @@ bool FCanvasTileRendererItem::Render_GameThread(const FCanvas* Canvas, FRenderTh
 		DrawTileParameters.RenderData->StaticMeshVertexBuffers.ColorVertexBuffer.ReleaseResource();
 		DrawTileParameters.RenderData->TileMesh.ReleaseResource();
 		DrawTileParameters.RenderData->VertexFactory.ReleaseResource();
+		DrawTileParameters.RenderData->IndexBuffer.ReleaseResource();
 
 		delete DrawTileParameters.View->Family;
 		delete DrawTileParameters.View;
diff --git a/Engine/Source/Runtime/Engine/Private/TriangleRendering.cpp b/Engine/Source/Runtime/Engine/Private/TriangleRendering.cpp
--- a/Engine/Source/Runtime/Engine/Private/TriangleRendering.cpp
+++ b/Engine/Source/Runtime/Engine/Private/TriangleRendering.cpp
@@ -190,8 +190,9 @@ FCanvasTriangleRendererItem::FTriangleVertexFactory::FTriangleVertexFactory(ERHI
 /**
  * Mesh used to render triangles.
  */
-FCanvasTriangleRendererItem::FTriangleMesh::FTriangleMesh(FTriangleVertexFactory* InVertexFactory)
-	: VertexFactory(InVertexFactory)
+FCanvasTriangleRendererItem::FTriangleMesh::FTriangleMesh(const FRawIndexBuffer* InIndexBuffer, const FTriangleVertexFactory* InVertexFactory)
+	: IndexBuffer(InIndexBuffer)
+	, VertexFactory(InVertexFactory)
 {
 
 }
@@ -203,6 +204,7 @@ void FCanvasTriangleRendererItem::FTriangleMesh::InitRHI()
 {
 	FMeshBatchElement& BatchElement = TriMeshElement.Elements[0];
 	TriMeshElement.VertexFactory = VertexFactory;
+	BatchElement.IndexBuffer = IndexBuffer;
 	BatchElement.FirstIndex = 0;
 	BatchElement.NumPrimitives = 1;
 	BatchElement.MinVertexIndex = 0;
@@ -224,35 +226,46 @@ void FCanvasTriangleRendererItem::InitTriangleBuffers(FLocalVertexFactory* Verte
 	Data->StaticMeshVertexBuffers.StaticMeshVertexBuffer.Init(Triangles.Num() * 3, 1);
 	Data->StaticMeshVertexBuffers.ColorVertexBuffer.Init(Triangles.Num() * 3);
 
+	Data->IndexBuffer.Indices.SetNum(Triangles.Num() * 3);
+
 	for (int32 i = 0; i < Triangles.Num(); i++)
 	{
+		const uint32 StartIndex = i * 3;
+
+		/** The use of an index buffer here is actually necessary to workaround an issue with BaseVertexIndex, DrawPrimitive, and manual vertex fetch.
+		 *  In short, DrawIndexedPrimitive with StartIndex map SV_VertexId to the correct location, but DrawPrimitive with BaseVertexIndex will not.
+		 */
+		Data->IndexBuffer.Indices[StartIndex + 0] = StartIndex + 0;
+		Data->IndexBuffer.Indices[StartIndex + 1] = StartIndex + 1;
+		Data->IndexBuffer.Indices[StartIndex + 2] = StartIndex + 2;
+
 		const FCanvasUVTri& Tri = Triangles[i].Tri;
 
 		// create verts. Notice the order is (1, 0, 2)
 		if (bNeedsToSwitchVerticalAxis)
 		{
-			Data->StaticMeshVertexBuffers.PositionVertexBuffer.VertexPosition(i * 3 + 0) = FVector(Tri.V1_Pos.X, View.UnscaledViewRect.Height() - Tri.V1_Pos.Y, 0.0f);
-			Data->StaticMeshVertexBuffers.PositionVertexBuffer.VertexPosition(i * 3 + 1) = FVector(Tri.V0_Pos.X, View.UnscaledViewRect.Height() - Tri.V0_Pos.Y, 0.0f);
-			Data->StaticMeshVertexBuffers.PositionVertexBuffer.VertexPosition(i * 3 + 2) = FVector(Tri.V2_Pos.X, View.UnscaledViewRect.Height() - Tri.V2_Pos.Y, 0.0f);
+			Data->StaticMeshVertexBuffers.PositionVertexBuffer.VertexPosition(StartIndex + 0) = FVector(Tri.V1_Pos.X, View.UnscaledViewRect.Height() - Tri.V1_Pos.Y, 0.0f);
+			Data->StaticMeshVertexBuffers.PositionVertexBuffer.VertexPosition(StartIndex + 1) = FVector(Tri.V0_Pos.X, View.UnscaledViewRect.Height() - Tri.V0_Pos.Y, 0.0f);
+			Data->StaticMeshVertexBuffers.PositionVertexBuffer.VertexPosition(StartIndex + 2) = FVector(Tri.V2_Pos.X, View.UnscaledViewRect.Height() - Tri.V2_Pos.Y, 0.0f);
 		}
 		else
 		{
-			Data->StaticMeshVertexBuffers.PositionVertexBuffer.VertexPosition(i * 3 + 0) = FVector(Tri.V1_Pos.X, Tri.V1_Pos.Y, 0.0f);
-			Data->StaticMeshVertexBuffers.PositionVertexBuffer.VertexPosition(i * 3 + 1) = FVector(Tri.V0_Pos.X, Tri.V0_Pos.Y, 0.0f);
-			Data->StaticMeshVertexBuffers.PositionVertexBuffer.VertexPosition(i * 3 + 2) = FVector(Tri.V2_Pos.X, Tri.V2_Pos.Y, 0.0f);
+			Data->StaticMeshVertexBuffers.PositionVertexBuffer.VertexPosition(StartIndex + 0) = FVector(Tri.V1_Pos.X, Tri.V1_Pos.Y, 0.0f);
+			Data->StaticMeshVertexBuffers.PositionVertexBuffer.VertexPosition(StartIndex + 1) = FVector(Tri.V0_Pos.X, Tri.V0_Pos.Y, 0.0f);
+			Data->StaticMeshVertexBuffers.PositionVertexBuffer.VertexPosition(StartIndex + 2) = FVector(Tri.V2_Pos.X, Tri.V2_Pos.Y, 0.0f);
 		}
 
-		Data->StaticMeshVertexBuffers.StaticMeshVertexBuffer.SetVertexTangents(i * 3 + 0, FVector(1.0f, 0.0f, 0.0f), FVector(0.0f, 1.0f, 0.0f), FVector(0.0f, 0.0f, 1.0f));
-		Data->StaticMeshVertexBuffers.StaticMeshVertexBuffer.SetVertexTangents(i * 3 + 1, FVector(1.0f, 0.0f, 0.0f), FVector(0.0f, 1.0f, 0.0f), FVector(0.0f, 0.0f, 1.0f));
-		Data->StaticMeshVertexBuffers.StaticMeshVertexBuffer.SetVertexTangents(i * 3 + 2, FVector(1.0f, 0.0f, 0.0f), FVector(0.0f, 1.0f, 0.0f), FVector(0.0f, 0.0f, 1.0f));
+		Data->StaticMeshVertexBuffers.StaticMeshVertexBuffer.SetVertexTangents(StartIndex + 0, FVector(1.0f, 0.0f, 0.0f), FVector(0.0f, 1.0f, 0.0f), FVector(0.0f, 0.0f, 1.0f));
+		Data->StaticMeshVertexBuffers.StaticMeshVertexBuffer.SetVertexTangents(StartIndex + 1, FVector(1.0f, 0.0f, 0.0f), FVector(0.0f, 1.0f, 0.0f), FVector(0.0f, 0.0f, 1.0f));
+		Data->StaticMeshVertexBuffers.StaticMeshVertexBuffer.SetVertexTangents(StartIndex + 2, FVector(1.0f, 0.0f, 0.0f), FVector(0.0f, 1.0f, 0.0f), FVector(0.0f, 0.0f, 1.0f));
 
-		Data->StaticMeshVertexBuffers.StaticMeshVertexBuffer.SetVertexUV(i * 3 + 0, 0, FVector2D(Tri.V1_UV.X, Tri.V1_UV.Y));
-		Data->StaticMeshVertexBuffers.StaticMeshVertexBuffer.SetVertexUV(i * 3 + 1, 0, FVector2D(Tri.V0_UV.X, Tri.V0_UV.Y));
-		Data->StaticMeshVertexBuffers.StaticMeshVertexBuffer.SetVertexUV(i * 3 + 2, 0, FVector2D(Tri.V2_UV.X, Tri.V2_UV.Y));
+		Data->StaticMeshVertexBuffers.StaticMeshVertexBuffer.SetVertexUV(StartIndex + 0, 0, FVector2D(Tri.V1_UV.X, Tri.V1_UV.Y));
+		Data->StaticMeshVertexBuffers.StaticMeshVertexBuffer.SetVertexUV(StartIndex + 1, 0, FVector2D(Tri.V0_UV.X, Tri.V0_UV.Y));
+		Data->StaticMeshVertexBuffers.StaticMeshVertexBuffer.SetVertexUV(StartIndex + 2, 0, FVector2D(Tri.V2_UV.X, Tri.V2_UV.Y));
 
-		Data->StaticMeshVertexBuffers.ColorVertexBuffer.VertexColor(i * 3 + 0) = Tri.V1_Color.ToFColor(true);
-		Data->StaticMeshVertexBuffers.ColorVertexBuffer.VertexColor(i * 3 + 1) = Tri.V0_Color.ToFColor(true);
-		Data->StaticMeshVertexBuffers.ColorVertexBuffer.VertexColor(i * 3 + 2) = Tri.V2_Color.ToFColor(true);
+		Data->StaticMeshVertexBuffers.ColorVertexBuffer.VertexColor(StartIndex + 0) = Tri.V1_Color.ToFColor(true);
+		Data->StaticMeshVertexBuffers.ColorVertexBuffer.VertexColor(StartIndex + 1) = Tri.V0_Color.ToFColor(true);
+		Data->StaticMeshVertexBuffers.ColorVertexBuffer.VertexColor(StartIndex + 2) = Tri.V2_Color.ToFColor(true);
 	}
 
 	FCanvasTriangleRendererItem::FRenderData* RenderData = Data;
@@ -271,6 +284,7 @@ void FCanvasTriangleRendererItem::InitTriangleBuffers(FLocalVertexFactory* Verte
 		RenderData->StaticMeshVertexBuffers.ColorVertexBuffer.BindColorVertexBuffer(VertexFactory, RData);
 		VertexFactory->SetData(RData);
 
+		RenderData->IndexBuffer.InitResource();
 		RenderData->VertexFactory.InitResource();
 		RenderData->TriMesh.InitResource();
 	});
@@ -326,7 +340,7 @@ bool FCanvasTriangleRendererItem::Render_RenderThread(FRHICommandListImmediate&
 		FMeshBatch& TriMesh = Data->TriMesh.TriMeshElement;
 		TriMesh.VertexFactory = &Data->VertexFactory;
 		TriMesh.MaterialRenderProxy = Data->MaterialRenderProxy;
-		TriMesh.Elements[0].BaseVertexIndex = 3 * TriIdx;
+		TriMesh.Elements[0].FirstIndex = 3 * TriIdx;
 
 		GetRendererModule().DrawTileMesh(RHICmdList, DrawRenderState, *View, TriMesh, Canvas->IsHitTesting(), Tri.HitProxyId);
 	}
@@ -334,6 +348,7 @@ bool FCanvasTriangleRendererItem::Render_RenderThread(FRHICommandListImmediate&
 	Data->StaticMeshVertexBuffers.PositionVertexBuffer.ReleaseResource();
 	Data->StaticMeshVertexBuffers.StaticMeshVertexBuffer.ReleaseResource();
 	Data->StaticMeshVertexBuffers.ColorVertexBuffer.ReleaseResource();
+	Data->IndexBuffer.ReleaseResource();
 	Data->TriMesh.ReleaseResource();
 	Data->VertexFactory.ReleaseResource();
 
@@ -420,7 +435,7 @@ bool FCanvasTriangleRendererItem::Render_GameThread(const FCanvas* Canvas, FRend
 			FMeshBatch& TriMesh = Parameters.RenderData->TriMesh.TriMeshElement;
 			TriMesh.VertexFactory = &Parameters.RenderData->VertexFactory;
 			TriMesh.MaterialRenderProxy = Parameters.RenderData->MaterialRenderProxy;
-			TriMesh.Elements[0].BaseVertexIndex = 3 * TriIdx;
+			TriMesh.Elements[0].FirstIndex = 3 * TriIdx;
 
 			GetRendererModule().DrawTileMesh(RHICmdList, DrawRenderState, *Parameters.View, TriMesh, Parameters.bIsHitTesting, Tri.HitProxyId);
 		}
@@ -428,6 +443,7 @@ bool FCanvasTriangleRendererItem::Render_GameThread(const FCanvas* Canvas, FRend
 		Parameters.RenderData->StaticMeshVertexBuffers.PositionVertexBuffer.ReleaseResource();
 		Parameters.RenderData->StaticMeshVertexBuffers.StaticMeshVertexBuffer.ReleaseResource();
 		Parameters.RenderData->StaticMeshVertexBuffers.ColorVertexBuffer.ReleaseResource();
+		Parameters.RenderData->IndexBuffer.ReleaseResource();
 		Parameters.RenderData->TriMesh.ReleaseResource();
 		Parameters.RenderData->VertexFactory.ReleaseResource();
 
diff --git a/Engine/Source/Runtime/Engine/Public/CanvasTypes.h b/Engine/Source/Runtime/Engine/Public/CanvasTypes.h
--- a/Engine/Source/Runtime/Engine/Public/CanvasTypes.h
+++ b/Engine/Source/Runtime/Engine/Public/CanvasTypes.h
@@ -1049,7 +1049,7 @@ private:
 	class FTileMesh : public FRenderResource
 	{
 	public:
-		FTileMesh(FTileVertexFactory* VertexFactory);
+		FTileMesh(const FRawIndexBuffer* IndexBuffer, const FTileVertexFactory* VertexFactory);
 
 		/** The mesh element. */
 		FMeshBatch MeshElement;
@@ -1057,7 +1057,8 @@ private:
 		virtual void InitRHI() override;
 		virtual void ReleaseRHI() override;
 	private:
-		FTileVertexFactory* VertexFactory;
+		const FRawIndexBuffer* IndexBuffer;
+		const FTileVertexFactory* VertexFactory;
 	};
 
 	class FRenderData
@@ -1066,6 +1067,7 @@ private:
 
 	private:
 		/** The buffer containing vertex data. */
+		FRawIndexBuffer IndexBuffer;
 		FStaticMeshVertexBuffers StaticMeshVertexBuffers;
 		FTileVertexFactory VertexFactory;
 		FTileMesh TileMesh;
@@ -1220,14 +1222,15 @@ private:
 	class FTriangleMesh : public FRenderResource
 	{
 	public:
-		FTriangleMesh(FTriangleVertexFactory* VertexFactory);
+		FTriangleMesh(const FRawIndexBuffer* IndexBuffer, const FTriangleVertexFactory* VertexFactory);
 
 		/** The mesh element. */
 		FMeshBatch TriMeshElement;
 		virtual void InitRHI() override;
 		virtual void ReleaseRHI() override;
 	private:
-		FTriangleVertexFactory* VertexFactory;
+		const FRawIndexBuffer* IndexBuffer;
+		const FTriangleVertexFactory* VertexFactory;
 	};
 
 	class FRenderData
@@ -1239,6 +1242,7 @@ private:
 		FCanvas::FTransformEntry Transform;
 
 		/** The buffer containing vertex data. */
+		FRawIndexBuffer IndexBuffer;
 		FStaticMeshVertexBuffers StaticMeshVertexBuffers;
 		FTriangleVertexFactory VertexFactory;
 		FTriangleMesh TriMesh;
@@ -1257,7 +1261,7 @@ private:
 			: MaterialRenderProxy(InMaterialRenderProxy)
 			, Transform(InTransform)
 			, VertexFactory(InFeatureLevel)
-			, TriMesh(&VertexFactory)
+			, TriMesh(&IndexBuffer, &VertexFactory)
 		{
 		}
 
