Fixed broken tile rendering.

diff --git a/Engine/Source/Runtime/Engine/Private/TileRendering.cpp b/Engine/Source/Runtime/Engine/Private/TileRendering.cpp
--- a/Engine/Source/Runtime/Engine/Private/TileRendering.cpp
+++ b/Engine/Source/Runtime/Engine/Private/TileRendering.cpp
@@ -275,7 +275,7 @@ void FCanvasTileRendererItem::InitTileBuffers(FLocalVertexFactory* VertexFactory
 	Data->StaticMeshVertexBuffers.StaticMeshVertexBuffer.Init(Tiles.Num() * NUM_MATERIAL_TILE_VERTS, 1);
 	Data->StaticMeshVertexBuffers.ColorVertexBuffer.Init(Tiles.Num() * NUM_MATERIAL_TILE_VERTS);
 
-	Data->IndexBuffer.Indices.SetNum(Tiles.Num() * NUM_MATERIAL_TILE_VERTS);
+	Data->IndexBuffer.Indices.SetNum(Tiles.Num() * 6);
 
 	for (int32 i = 0; i < Tiles.Num(); i++)
 	{
@@ -409,6 +409,7 @@ bool FCanvasTileRendererItem::Render_RenderThread(FRHICommandListImmediate& RHIC
 	Data->StaticMeshVertexBuffers.PositionVertexBuffer.ReleaseResource();
 	Data->StaticMeshVertexBuffers.StaticMeshVertexBuffer.ReleaseResource();
 	Data->StaticMeshVertexBuffers.ColorVertexBuffer.ReleaseResource();
+	Data->IndexBuffer.ReleaseResource();
 	Data->TileMesh.ReleaseResource();
 	Data->VertexFactory.ReleaseResource();
 
