Added basic ray tracing support for ISM and HISM.

diff --git a/Engine/Source/Runtime/Engine/Private/InstancedStaticMesh.cpp b/Engine/Source/Runtime/Engine/Private/InstancedStaticMesh.cpp
--- a/Engine/Source/Runtime/Engine/Private/InstancedStaticMesh.cpp
+++ b/Engine/Source/Runtime/Engine/Private/InstancedStaticMesh.cpp
@@ -31,6 +31,9 @@
 #endif // WITH_EDITOR
 #include "MeshMaterialShader.h"
 
+#if RHI_RAYTRACING
+#include "RayTracingInstance.h"
+#endif
 
 #include "Interfaces/ITargetPlatform.h"
 #if WITH_EDITOR
@@ -41,6 +44,11 @@
 #include "UObject/EditorObjectVersion.h"
 
 
+static TAutoConsoleVariable<int32> CVarRayTracingRenderInstances(
+	TEXT("r.RayTracing.Instances"),
+	1,
+	TEXT("Include static mesh instances in the ray tracing acceleration structure"));
+
 const int32 InstancedStaticMeshMaxTexCoord = 8;
 
 IMPLEMENT_HIT_PROXY(HInstancedStaticMeshInstance, HHitProxy);
@@ -981,6 +989,51 @@ HHitProxy* FInstancedStaticMeshSceneProxy::CreateHitProxies(UPrimitiveComponent*
 	return FStaticMeshSceneProxy::CreateHitProxies(Component, OutHitProxies);
 }
 
+#if RHI_RAYTRACING
+void FInstancedStaticMeshSceneProxy::GetDynamicRayTracingInstances(struct FRayTracingMaterialGatheringContext& Context, TArray<FRayTracingInstance>& OutRayTracingInstances)
+{
+	if (!CVarRayTracingRenderInstances.GetValueOnRenderThread())
+	{
+		return;
+	}
+
+	//#dxr_todo: select the appropriate LOD depending on Context.View
+	//#dxr_todo: add culling logic for instances that are far away
+	uint32 LOD = 0;
+
+	const int32 InstanceCount = InstancedRenderData.Component->PerInstanceSMData.Num();
+
+	for (int32 InstanceIdx = 0; InstanceIdx < InstanceCount; ++InstanceIdx)
+	{
+		if (InstancedRenderData.Component->PerInstanceSMData.IsValidIndex(InstanceIdx))
+		{
+			const FInstancedStaticMeshInstanceData& InstanceData = InstancedRenderData.Component->PerInstanceSMData[InstanceIdx];
+			FMatrix LocalToWorld = InstancedRenderData.Component->GetComponentTransform().ToMatrixWithScale();
+			FMatrix InstanceTransform = InstanceData.Transform * LocalToWorld;
+
+			FRayTracingInstance RayTracingInstance;
+			RayTracingInstance.Geometry = &RenderData->LODResources[LOD].RayTracingGeometry;
+			RayTracingInstance.InstanceTransforms.Add(InstanceTransform);
+
+			int32 SectionCount = InstancedRenderData.LODModels[LOD].Sections.Num();
+
+			for (int32 SectionIdx = 0; SectionIdx < SectionCount; ++SectionIdx)
+			{
+				//#dxr_todo: so far we use the parent static mesh path to get material data
+				FMeshBatch MeshBatch;
+				FStaticMeshSceneProxy::GetMeshElement(LOD, 0, SectionIdx, 0, false, false, MeshBatch);
+
+				RayTracingInstance.Materials.Add(MeshBatch);
+			}
+
+			RayTracingInstance.BuildInstanceMaskAndFlags();
+			OutRayTracingInstances.Add(RayTracingInstance);
+		}
+	}
+}
+#endif
+
+
 /*-----------------------------------------------------------------------------
 	UInstancedStaticMeshComponent
 -----------------------------------------------------------------------------*/
diff --git a/Engine/Source/Runtime/Engine/Private/InstancedStaticMesh.h b/Engine/Source/Runtime/Engine/Private/InstancedStaticMesh.h
--- a/Engine/Source/Runtime/Engine/Private/InstancedStaticMesh.h
+++ b/Engine/Source/Runtime/Engine/Private/InstancedStaticMesh.h
@@ -570,6 +570,15 @@ public:
 		return Result;
 	}
 
+#if RHI_RAYTRACING
+	virtual bool IsRayTracingStaticRelevant() const override
+	{
+		return false;
+	}
+
+	virtual void GetDynamicRayTracingInstances(struct FRayTracingMaterialGatheringContext& Context, TArray<FRayTracingInstance>& OutRayTracingInstances) final override;
+#endif
+
 	virtual void GetLightRelevance(const FLightSceneProxy* LightSceneProxy, bool& bDynamic, bool& bRelevant, bool& bLightMapped, bool& bShadowMapped) const override;
 	virtual void GetDynamicMeshElements(const TArray<const FSceneView*>& Views, const FSceneViewFamily& ViewFamily, uint32 VisibilityMap, FMeshElementCollector& Collector) const override;
 
