Filter runtime virtual texture casters in the main pass according to the component setting

diff --git a/Engine/Source/Runtime/Engine/Classes/Components/PrimitiveComponent.h b/Engine/Source/Runtime/Engine/Classes/Components/PrimitiveComponent.h
--- a/Engine/Source/Runtime/Engine/Classes/Components/PrimitiveComponent.h
+++ b/Engine/Source/Runtime/Engine/Classes/Components/PrimitiveComponent.h
@@ -563,6 +563,11 @@ public:
 	UPROPERTY(EditAnywhere, BlueprintReadOnly, Category = VirtualTexture, meta = (DisplayName = "Virtual Texture Pass Type"))
 	ERuntimeVirtualTextureMainPassType VirtualTextureRenderPassType = ERuntimeVirtualTextureMainPassType::Exclusive;
 
+	/** Get the array of runtime virtual textures into which we render the mesh for this actor. */
+	virtual TArray<URuntimeVirtualTexture*> const& GetRuntimeVirtualTextures() const { return RuntimeVirtualTextures; }
+	/** Get the runtime virtual texture pass settings. */
+	virtual ERuntimeVirtualTextureMainPassType GetVirtualTextureRenderPassType() const { return VirtualTextureRenderPassType; }
+
 	/** Used by the renderer, to identify a component across re-registers. */
 	FPrimitiveComponentId ComponentId;
 
diff --git a/Engine/Source/Runtime/Engine/Classes/VT/RuntimeVirtualTexturePlane.h b/Engine/Source/Runtime/Engine/Classes/VT/RuntimeVirtualTexturePlane.h
--- a/Engine/Source/Runtime/Engine/Classes/VT/RuntimeVirtualTexturePlane.h
+++ b/Engine/Source/Runtime/Engine/Classes/VT/RuntimeVirtualTexturePlane.h
@@ -67,11 +67,15 @@ public:
 #endif
 
 protected:
+	/** Apply any deferred material notifications. */
+	void NotifyMaterials();
+
 	//~ Begin UActorComponent Interface
 	virtual void CreateRenderState_Concurrent() override;
 	virtual void SendRenderTransform_Concurrent() override;
 	virtual void DestroyRenderState_Concurrent() override;
 	virtual void TickComponent(float DeltaTime, ELevelTick TickType, FActorComponentTickFunction* ThisTickFunction) override;
+	virtual void OnUnregister() override;
 	//~ End UActorComponent Interface
 
 public:
diff --git a/Engine/Source/Runtime/Engine/Private/PrimitiveSceneProxy.cpp b/Engine/Source/Runtime/Engine/Private/PrimitiveSceneProxy.cpp
--- a/Engine/Source/Runtime/Engine/Private/PrimitiveSceneProxy.cpp
+++ b/Engine/Source/Runtime/Engine/Private/PrimitiveSceneProxy.cpp
@@ -209,15 +209,10 @@ FPrimitiveSceneProxy::FPrimitiveSceneProxy(const UPrimitiveComponent* InComponen
 	bRequiresVisibleLevelToRender = (ComponentLevel && ComponentLevel->bRequireFullVisibilityToRender);
 	bIsComponentLevelVisible = (!ComponentLevel || ComponentLevel->bIsVisible);
 
-#if WITH_EDITOR
-	const bool bGetDebugMaterials = true;
-	InComponent->GetUsedMaterials(UsedMaterialsForVerification, bGetDebugMaterials);
-#endif
-
-	// Store the target runtime virtual texture information
+	// Setup the runtime virtual texture information and flush the virtual texture if necessary
 	if (UseVirtualTexturing(GetScene().GetFeatureLevel()))
 	{
-		for (URuntimeVirtualTexture* VirtualTexture : InComponent->RuntimeVirtualTextures)
+		for (URuntimeVirtualTexture* VirtualTexture : InComponent->GetRuntimeVirtualTextures())
 		{
 			if (VirtualTexture != nullptr && VirtualTexture->GetEnabled())
 			{
@@ -230,6 +225,21 @@ FPrimitiveSceneProxy::FPrimitiveSceneProxy(const UPrimitiveComponent* InComponen
 			}
 		}
 	}
+
+	// Conditionally remove from the main render pass based on the runtime virtual texture setup
+	ERuntimeVirtualTextureMainPassType MainPassType = InComponent->GetVirtualTextureRenderPassType();
+	const bool bRequestVirtualTexture = InComponent->GetRuntimeVirtualTextures().Num() > 0;
+	const bool bUseVirtualTexture = RuntimeVirtualTextures.Num() > 0;
+	if ((MainPassType == ERuntimeVirtualTextureMainPassType::Never && bRequestVirtualTexture) ||
+		(MainPassType == ERuntimeVirtualTextureMainPassType::Exclusive && bUseVirtualTexture))
+	{
+		bRenderInMainPass = false;
+	}
+
+#if WITH_EDITOR
+	const bool bGetDebugMaterials = true;
+	InComponent->GetUsedMaterials(UsedMaterialsForVerification, bGetDebugMaterials);
+#endif
 }
 
 #if WITH_EDITOR
diff --git a/Engine/Source/Runtime/Engine/Private/VT/RuntimeVirtualTextureNotify.cpp b/Engine/Source/Runtime/Engine/Private/VT/RuntimeVirtualTextureNotify.cpp
--- a/Engine/Source/Runtime/Engine/Private/VT/RuntimeVirtualTextureNotify.cpp
+++ b/Engine/Source/Runtime/Engine/Private/VT/RuntimeVirtualTextureNotify.cpp
@@ -28,7 +28,7 @@ namespace RuntimeVirtualTexture
 	{
 		for (TObjectIterator<UPrimitiveComponent> It; It; ++It)
 		{
-			for (URuntimeVirtualTexture* ItVirtualTexture : It->RuntimeVirtualTextures)
+			for (URuntimeVirtualTexture* ItVirtualTexture : It->GetRuntimeVirtualTextures())
 			{
 				if (ItVirtualTexture == VirtualTexture)
 				{
@@ -37,24 +37,6 @@ namespace RuntimeVirtualTexture
 				}
 			}
 		}
-
-		/*
-		for (TObjectIterator<ULandscapeComponent> It; It; ++It)
-		{
-			ALandscape* Landscape = It->GetLandscapeActor();
-			if (Landscape != nullptr)
-			{
-				for (URuntimeVirtualTexture* ItVirtualTexture : Landscape->RuntimeVirtualTextures)
-				{
-					if (ItVirtualTexture == VirtualTexture)
-					{
-						It->MarkRenderStateDirty();
-						break;
-					}
-				}
-			}
-		}
-		*/
 	}
 
 #endif
@@ -67,6 +49,11 @@ namespace RuntimeVirtualTexture
 		// Is there a way to set up dependencies so that we load the materials _after_ the virtual texture is allocated?
 		// Or maybe we should consider serializing the WorldToUVTransform in the URuntimeVirtualTexture and not depending on a URuntimeVirtualTextureComponent at runtime?
 
+		if (VirtualTexture == nullptr)
+		{
+			return;
+		}
+
 		TSet<UMaterial*> BaseMaterialsThatUseThisTexture;
 		for (TObjectIterator<UMaterialInterface> It; It; ++It)
 		{
diff --git a/Engine/Source/Runtime/Engine/Private/VT/RuntimeVirtualTexturePlane.cpp b/Engine/Source/Runtime/Engine/Private/VT/RuntimeVirtualTexturePlane.cpp
--- a/Engine/Source/Runtime/Engine/Private/VT/RuntimeVirtualTexturePlane.cpp
+++ b/Engine/Source/Runtime/Engine/Private/VT/RuntimeVirtualTexturePlane.cpp
@@ -66,7 +66,7 @@ void URuntimeVirtualTextureComponent::DestroyRenderState_Concurrent()
 	Super::DestroyRenderState_Concurrent();
 }
 
-void URuntimeVirtualTextureComponent::TickComponent(float DeltaTime, ELevelTick TickType, FActorComponentTickFunction* ThisTickFunction)
+void URuntimeVirtualTextureComponent::NotifyMaterials()
 {
 	if (bNotifyInNextTick)
 	{
@@ -77,6 +77,17 @@ void URuntimeVirtualTextureComponent::TickComponent(float DeltaTime, ELevelTick
 	}
 }
 
+void URuntimeVirtualTextureComponent::TickComponent(float DeltaTime, ELevelTick TickType, FActorComponentTickFunction* ThisTickFunction)
+{
+	NotifyMaterials();
+}
+
+void URuntimeVirtualTextureComponent::OnUnregister()
+{
+	NotifyMaterials();
+	Super::OnUnregister();
+}
+
 #if WITH_EDITOR
 
 void URuntimeVirtualTextureComponent::SetRotation()
diff --git a/Engine/Source/Runtime/Engine/Public/PrimitiveSceneProxy.h b/Engine/Source/Runtime/Engine/Public/PrimitiveSceneProxy.h
--- a/Engine/Source/Runtime/Engine/Public/PrimitiveSceneProxy.h
+++ b/Engine/Source/Runtime/Engine/Public/PrimitiveSceneProxy.h
@@ -970,7 +970,7 @@ protected:
 
 	/** Array of virtual textures that this proxy should render to. */
 	TArray<URuntimeVirtualTexture*> RuntimeVirtualTextures;
-	TArray<ERuntimeVirtualTextureMaterialType> RuntimeVirtualTextureMaterialTypes;
+	TSet<ERuntimeVirtualTextureMaterialType> RuntimeVirtualTextureMaterialTypes;
 
 private:
 	/** The hierarchy of owners of this primitive.  These must not be dereferenced on the rendering thread, but the pointer values can be used for identification.  */
diff --git a/Engine/Source/Runtime/Foliage/Private/InstancedFoliage.cpp b/Engine/Source/Runtime/Foliage/Private/InstancedFoliage.cpp
--- a/Engine/Source/Runtime/Foliage/Private/InstancedFoliage.cpp
+++ b/Engine/Source/Runtime/Foliage/Private/InstancedFoliage.cpp
@@ -1354,13 +1354,11 @@ void FFoliageStaticMesh::UpdateComponentSettings(const UFoliageType_InstancedSta
 		}
 		if (Component->RuntimeVirtualTextures != FoliageType->RuntimeVirtualTextures)
 		{
-			//todo[vt]: Trigger a virtual texture flush somewhere withing the update chain
 			Component->RuntimeVirtualTextures = FoliageType->RuntimeVirtualTextures;
 			bNeedsMarkRenderStateDirty = true;
 		}
 		if (Component->VirtualTextureRenderPassType != FoliageType->VirtualTextureRenderPassType)
 		{
-			//todo[vt]: Trigger a virtual texture flush somewhere withing the update chain
 			Component->VirtualTextureRenderPassType = FoliageType->VirtualTextureRenderPassType;
 			bNeedsMarkRenderStateDirty = true;
 		}
diff --git a/Engine/Source/Runtime/Landscape/Classes/LandscapeComponent.h b/Engine/Source/Runtime/Landscape/Classes/LandscapeComponent.h
--- a/Engine/Source/Runtime/Landscape/Classes/LandscapeComponent.h
+++ b/Engine/Source/Runtime/Landscape/Classes/LandscapeComponent.h
@@ -327,7 +327,7 @@ enum ELandscapeLayerUpdateMode : uint32
 
 #endif
 
-UCLASS(hidecategories=(Display, Attachment, Physics, Debug, Collision, Movement, Rendering, PrimitiveComponent, Object, Transform, Mobility), showcategories=("Rendering|Material"), MinimalAPI, Within=LandscapeProxy)
+UCLASS(hidecategories=(Display, Attachment, Physics, Debug, Collision, Movement, Rendering, PrimitiveComponent, Object, Transform, Mobility, VirtualTexture), showcategories=("Rendering|Material"), MinimalAPI, Within=LandscapeProxy)
 class ULandscapeComponent : public UPrimitiveComponent
 {
 	GENERATED_UCLASS_BODY()
@@ -584,6 +584,9 @@ public:
 	virtual void GetStreamingRenderAssetInfo(FStreamingTextureLevelContext& LevelContext, TArray<FStreamingRenderAssetPrimitiveInfo>& OutStreamingRenderAssets) const override;
 	virtual bool IsPrecomputedLightingValid() const override;
 
+	virtual TArray<URuntimeVirtualTexture*> const& GetRuntimeVirtualTextures() const override;
+	virtual ERuntimeVirtualTextureMainPassType GetVirtualTextureRenderPassType() const override;
+
 	LANDSCAPE_API UTexture2D* GetHeightmap(bool InReturnEditingHeightmap = false) const;
 	LANDSCAPE_API TArray<UTexture2D*>& GetWeightmapTextures(bool InReturnEditingWeightmap = false);
 	LANDSCAPE_API const TArray<UTexture2D*>& GetWeightmapTextures(bool InReturnEditingWeightmap = false) const;
diff --git a/Engine/Source/Runtime/Landscape/Private/Landscape.cpp b/Engine/Source/Runtime/Landscape/Private/Landscape.cpp
--- a/Engine/Source/Runtime/Landscape/Private/Landscape.cpp
+++ b/Engine/Source/Runtime/Landscape/Private/Landscape.cpp
@@ -1289,6 +1289,16 @@ void ULandscapeComponent::PropagateLightingScenarioChange()
 	FComponentRecreateRenderStateContext Context(this);
 }
 
+TArray<URuntimeVirtualTexture*> const& ULandscapeComponent::GetRuntimeVirtualTextures() const
+{
+	return GetLandscapeProxy()->RuntimeVirtualTextures;
+}
+
+ERuntimeVirtualTextureMainPassType ULandscapeComponent::GetVirtualTextureRenderPassType() const
+{
+	return GetLandscapeProxy()->VirtualTextureRenderPassType;
+}
+
 #if WITH_EDITOR
 ULandscapeInfo* ULandscapeComponent::GetLandscapeInfo() const
 {
diff --git a/Engine/Source/Runtime/Landscape/Private/LandscapeRender.cpp b/Engine/Source/Runtime/Landscape/Private/LandscapeRender.cpp
--- a/Engine/Source/Runtime/Landscape/Private/LandscapeRender.cpp
+++ b/Engine/Source/Runtime/Landscape/Private/LandscapeRender.cpp
@@ -862,23 +862,6 @@ FLandscapeComponentSceneProxy::FLandscapeComponentSceneProxy(ULandscapeComponent
 
 	bSupportsHeightfieldRepresentation = true;
 
-	// Store the target runtime virtual texture information
-	if (UseVirtualTexturing(FeatureLevel))
-	{
-		//todo[vt]: Should the virtual texture settings be in the component overrides?
-		for (URuntimeVirtualTexture* VirtualTexture : InComponent->GetLandscapeProxy()->RuntimeVirtualTextures)
-		{
-			if (VirtualTexture != nullptr && VirtualTexture->GetEnabled())
-			{
-				RuntimeVirtualTextures.Add(VirtualTexture);
-				RuntimeVirtualTextureMaterialTypes.Add(VirtualTexture->GetMaterialType());
-
-				//todo[vt]: Only flush specific virtual textures
-				//todo[vt]: Only flush primitive bounds 
-				GetRendererModule().FlushVirtualTextureCache();
-			}
-		}
-	}
 #if WITH_EDITOR
 	TArray<FWeightmapLayerAllocationInfo>& ComponentWeightmapLayerAllocations = InComponent->GetWeightmapLayerAllocations();
 	
@@ -5041,13 +5024,6 @@ FLandscapeMeshProxySceneProxy::~FLandscapeMeshProxySceneProxy()
 	{
 		Info.UnregisterNeighbors();
 	}
-
-	for (URuntimeVirtualTexture* VirtualTexture : RuntimeVirtualTextures)
-	{
-		//todo[vt]: Only flush specific virtual textures
-		//todo[vt]: Only flush primitive bounds 
-		GetRendererModule().FlushVirtualTextureCache();
-	}
 }
 
 FPrimitiveSceneProxy* ULandscapeMeshProxyComponent::CreateSceneProxy()
