Refactor of FCanvasTileRendererItem implementation.

Removed dead code in TileRendering.cpp and did some restructuring of the internal FRenderData class to make the data flow cleaner.

diff --git a/Engine/Source/Runtime/Engine/Private/TileRendering.cpp b/Engine/Source/Runtime/Engine/Private/TileRendering.cpp
--- a/Engine/Source/Runtime/Engine/Private/TileRendering.cpp
+++ b/Engine/Source/Runtime/Engine/Private/TileRendering.cpp
@@ -19,275 +19,86 @@
 #include "EngineModule.h"
 #include "MeshPassProcessor.h"
 
-#define NUM_MATERIAL_TILE_VERTS	4
+static const uint32 CanvasTileVertexCount = 4;
+static const uint32 CanvasTileIndexCount = 6;
 
-/** 
-* Vertex buffer
-*/
-class FMaterialTileVertexBuffer : public FRenderResource
-{
-public:
-	FVertexBuffer PositionBuffer;
-	FVertexBuffer TangentBuffer;
-	FVertexBuffer TexCoordBuffer;
-	FVertexBuffer ColorBuffer;
-
-	FShaderResourceViewRHIRef TangentBufferSRV;
-	FShaderResourceViewRHIRef TexCoordBufferSRV;
-	FShaderResourceViewRHIRef ColorBufferSRV;
-	FShaderResourceViewRHIRef PositionBufferSRV;
-
-	/** 
-	* Initialize the RHI for this rendering resource 
-	*/
-	virtual void InitRHI() override
-	{
-		FVector* PositionBufferData = nullptr;
-
-		uint32 PositionSize = NUM_MATERIAL_TILE_VERTS * sizeof(FVector);
-		// create vertex buffer
-		{
-			FRHIResourceCreateInfo CreateInfo;
-			void* Data = nullptr;
-			PositionBuffer.VertexBufferRHI = RHICreateAndLockVertexBuffer(PositionSize, BUF_Static | BUF_ShaderResource, CreateInfo, Data);
-			PositionBufferData = static_cast<FVector*>(Data);
-			if (RHISupportsManualVertexFetch(GMaxRHIShaderPlatform))
-			{
-				PositionBufferSRV = RHICreateShaderResourceView(PositionBuffer.VertexBufferRHI, sizeof(float), PF_R32_FLOAT);
-			}
-
-		}
-
-		FPackedNormal* TangentBufferData = nullptr;
-
-		uint32 TangentSize = NUM_MATERIAL_TILE_VERTS * 2 * sizeof(FPackedNormal);
-		// create vertex buffer
-		{
-			FRHIResourceCreateInfo CreateInfo;
-			void* Data = nullptr;
-			TangentBuffer.VertexBufferRHI = RHICreateAndLockVertexBuffer(TangentSize, BUF_Static | BUF_ShaderResource, CreateInfo, Data);
-			TangentBufferData = static_cast<FPackedNormal*>(Data);
-			if (RHISupportsManualVertexFetch(GMaxRHIShaderPlatform))
-			{
-				TangentBufferSRV = RHICreateShaderResourceView(TangentBuffer.VertexBufferRHI, sizeof(FPackedNormal), PF_R8G8B8A8_SNORM);
-			}
-		}
-
-		FVector2D* TexCoordBufferData = nullptr;
-
-		uint32 TexCoordSize = NUM_MATERIAL_TILE_VERTS * sizeof(FVector2D);
-		// create vertex buffer
-		{
-			FRHIResourceCreateInfo CreateInfo;
-			void* Data = nullptr;
-			TexCoordBuffer.VertexBufferRHI = RHICreateAndLockVertexBuffer(TexCoordSize, BUF_Static | BUF_ShaderResource, CreateInfo, Data);
-			TexCoordBufferData = static_cast<FVector2D*>(Data);
-
-			if (RHISupportsManualVertexFetch(GMaxRHIShaderPlatform))
-			{
-				TexCoordBufferSRV = RHICreateShaderResourceView(TexCoordBuffer.VertexBufferRHI, sizeof(FVector2D), PF_G32R32F);
-			}
-		}
-
-		uint32* ColorBufferData = nullptr;
-
-		uint32 ColorSize = NUM_MATERIAL_TILE_VERTS * sizeof(uint32);
-		// create vertex buffer
-		{
-			FRHIResourceCreateInfo CreateInfo;
-			void* Data = nullptr;
-			ColorBuffer.VertexBufferRHI = RHICreateAndLockVertexBuffer(ColorSize, BUF_Static | BUF_ShaderResource, CreateInfo, Data);
-			ColorBufferData = static_cast<uint32*>(Data);
-			if (RHISupportsManualVertexFetch(GMaxRHIShaderPlatform))
-			{
-				ColorBufferSRV = RHICreateShaderResourceView(ColorBuffer.VertexBufferRHI, sizeof(uint32), PF_R8G8B8A8);
-			}
-		}
-
-		// fill out the verts
-		PositionBufferData[0] = FVector( 1, -1, 0);
-		PositionBufferData[1] = FVector( 1,  1, 0);
-		PositionBufferData[2] = FVector(-1, -1, 0);
-		PositionBufferData[3] = FVector(-1,  1, 0);
-
-		TexCoordBufferData[0] = FVector2D(1, 1);
-		TexCoordBufferData[1] = FVector2D(1, 0);
-		TexCoordBufferData[2] = FVector2D(0, 1);
-		TexCoordBufferData[3] = FVector2D(0, 0);
-
-		for (int i = 0; i < NUM_MATERIAL_TILE_VERTS; i++)
-		{
-			TangentBufferData[2 * i + 0] = FVector(1, 0, 0);
-			TangentBufferData[2 * i + 1] = FVector(0, 0, 1);
-			ColorBufferData[i] = FColor(255, 255, 255, 255).DWColor();
-		}
-
-		// Unlock the buffer.
-		RHIUnlockVertexBuffer(PositionBuffer.VertexBufferRHI);
-		RHIUnlockVertexBuffer(TangentBuffer.VertexBufferRHI);
-		RHIUnlockVertexBuffer(TexCoordBuffer.VertexBufferRHI);
-		RHIUnlockVertexBuffer(ColorBuffer.VertexBufferRHI);
-	}
+FCanvasTileRendererItem::FTileVertexFactory::FTileVertexFactory(
+	const FStaticMeshVertexBuffers* InVertexBuffers,
+	ERHIFeatureLevel::Type InFeatureLevel)
+	: FLocalVertexFactory(InFeatureLevel, "FTileVertexFactory")
+	, VertexBuffers(InVertexBuffers)
+{}
 
-	void InitResource() override
-	{
-		FRenderResource::InitResource();
-		PositionBuffer.InitResource();
-		TangentBuffer.InitResource();
-		TexCoordBuffer.InitResource();
-		ColorBuffer.InitResource();
-	}
-
-	void ReleaseResource() override
-	{
-		FRenderResource::ReleaseResource();
-		PositionBuffer.ReleaseResource();
-		TangentBuffer.ReleaseResource();
-		TexCoordBuffer.ReleaseResource();
-		ColorBuffer.ReleaseResource();
-	}
-
-	virtual void ReleaseRHI() override
-	{
-		PositionBuffer.ReleaseRHI();
-		TangentBuffer.ReleaseRHI();
-		TexCoordBuffer.ReleaseRHI();
-		ColorBuffer.ReleaseRHI();
-
-		TangentBufferSRV.SafeRelease();
-		TexCoordBufferSRV.SafeRelease();
-		ColorBufferSRV.SafeRelease();
-		PositionBufferSRV.SafeRelease();
-	}
-};
-TGlobalResource<FMaterialTileVertexBuffer> GTileRendererVertexBuffer;
-
-
-/**
-* Vertex factory for rendering tiles.
-*/
-	/** Default constructor. */
-FCanvasTileRendererItem::FTileVertexFactory::FTileVertexFactory(ERHIFeatureLevel::Type InFeatureLevel) : FLocalVertexFactory(InFeatureLevel, "FTileVertexFactory")
+void FCanvasTileRendererItem::FTileVertexFactory::InitResource()
 {
-	
+	FLocalVertexFactory::FDataType VertexData;
+	VertexBuffers->PositionVertexBuffer.BindPositionVertexBuffer(this, VertexData);
+	VertexBuffers->StaticMeshVertexBuffer.BindTangentVertexBuffer(this, VertexData);
+	VertexBuffers->StaticMeshVertexBuffer.BindPackedTexCoordVertexBuffer(this, VertexData);
+	VertexBuffers->StaticMeshVertexBuffer.BindLightMapVertexBuffer(this, VertexData, 0);
+	VertexBuffers->ColorVertexBuffer.BindColorVertexBuffer(this, VertexData);
+	SetData(VertexData);
+
+	FLocalVertexFactory::InitResource();
 }
 
-/**
- * Mesh used to render tiles.
- */
 FCanvasTileRendererItem::FTileMesh::FTileMesh(
 	const FRawIndexBuffer* InIndexBuffer,
 	const FCanvasTileRendererItem::FTileVertexFactory* InVertexFactory)
 	: IndexBuffer(InIndexBuffer)
 	, VertexFactory(InVertexFactory)
-{
-}
+{}
 
 void FCanvasTileRendererItem::FTileMesh::InitRHI()
 {
-	FMeshBatchElement& BatchElement = MeshElement.Elements[0];
 	MeshElement.VertexFactory = VertexFactory;
+	MeshElement.ReverseCulling = false;
+	MeshElement.Type = PT_TriangleList;
+	MeshElement.DepthPriorityGroup = SDPG_Foreground;
+
+	FMeshBatchElement& BatchElement = MeshElement.Elements[0];
 	BatchElement.IndexBuffer = IndexBuffer;
 	BatchElement.FirstIndex = 0;
 	BatchElement.NumPrimitives = 2;
 	BatchElement.MinVertexIndex = 0;
-	BatchElement.MaxVertexIndex = NUM_MATERIAL_TILE_VERTS - 1;
-	MeshElement.ReverseCulling = false;
-	MeshElement.Type = PT_TriangleList;
-	MeshElement.DepthPriorityGroup = SDPG_Foreground;
+	BatchElement.MaxVertexIndex = CanvasTileVertexCount - 1;
 	BatchElement.PrimitiveUniformBufferResource = &GIdentityPrimitiveUniformBuffer;
 }
 
-void FCanvasTileRendererItem::FTileMesh::ReleaseRHI()
-{
-}
-
-FCanvasTileRendererItem::FRenderData::FRenderData(ERHIFeatureLevel::Type InFeatureLevel,
+FCanvasTileRendererItem::FRenderData::FRenderData(
+	ERHIFeatureLevel::Type InFeatureLevel,
 	const FMaterialRenderProxy* InMaterialRenderProxy,
-	const FCanvas::FTransformEntry& InTransform )
-	: VertexFactory(InFeatureLevel)
-	, TileMesh(&IndexBuffer, &VertexFactory)
-	, MaterialRenderProxy(InMaterialRenderProxy)
+	const FCanvas::FTransformEntry& InTransform)
+	: MaterialRenderProxy(InMaterialRenderProxy)
 	, Transform(InTransform)
-{
-	FTileVertexFactory* TileVertexFactory = &VertexFactory;
-	ENQUEUE_RENDER_COMMAND(FCanvasTileRendererItemFTileVertexFactoryInit)(
-		[TileVertexFactory](FRHICommandListImmediate& RHICmdList)
-		{
-			FLocalVertexFactory::FDataType VertexData;
-			VertexData.NumTexCoords = 1;
-
-			{
-				VertexData.LightMapCoordinateIndex = 0;
-				VertexData.TangentsSRV = GTileRendererVertexBuffer.TangentBufferSRV;
-				VertexData.TextureCoordinatesSRV = GTileRendererVertexBuffer.TexCoordBufferSRV;
-				VertexData.ColorComponentsSRV = GTileRendererVertexBuffer.ColorBufferSRV;
-				VertexData.PositionComponentSRV = GTileRendererVertexBuffer.PositionBufferSRV;
-			}
-
-			{
-				// position
-				VertexData.PositionComponent = FVertexStreamComponent(
-					&GTileRendererVertexBuffer.PositionBuffer,
-					0,
-					sizeof(FVector),
-					VET_Float3,
-					EVertexStreamUsage::Default
-				);
-
-				// tangents
-				VertexData.TangentBasisComponents[0] = FVertexStreamComponent(&GTileRendererVertexBuffer.TangentBuffer, 0, 2 * sizeof(FPackedNormal), VET_PackedNormal, EVertexStreamUsage::ManualFetch);
-				VertexData.TangentBasisComponents[1] = FVertexStreamComponent(&GTileRendererVertexBuffer.TangentBuffer, sizeof(FPackedNormal), 2 * sizeof(FPackedNormal), VET_PackedNormal, EVertexStreamUsage::ManualFetch);
-
-				// color
-				VertexData.ColorComponent = FVertexStreamComponent(&GTileRendererVertexBuffer.ColorBuffer, 0, sizeof(FColor), VET_Color, EVertexStreamUsage::ManualFetch);
-				// UVs
-				VertexData.TextureCoordinates.Add(FVertexStreamComponent(
-					&GTileRendererVertexBuffer.TexCoordBuffer,
-					0,
-					sizeof(FVector2D),
-					VET_Float2,
-					EVertexStreamUsage::ManualFetch
-				));
-
-				// update the data
-				TileVertexFactory->SetData(VertexData);
-			}
-	});
-}
+	, VertexFactory(&StaticMeshVertexBuffers, InFeatureLevel)
+	, TileMesh(&IndexBuffer, &VertexFactory)
+{}
 
-FCanvasTileRendererItem::~FCanvasTileRendererItem()
+void FCanvasTileRendererItem::FRenderData::InitTileMesh(const FSceneView& View, bool bNeedsToSwitchVerticalAxis)
 {
-	FCanvasTileRendererItem::FRenderData* RenderData = Data;
-	ENQUEUE_RENDER_COMMAND(FCanvasTileRendererItemDeleteRenderData)(
-		[RenderData](FRHICommandListImmediate& RHICmdList)
-		{
-			delete RenderData;
-		});
-	
-	Data = nullptr;
-}
+	static_assert(CanvasTileVertexCount == 4, "Invalid tile tri-list size.");
+	static_assert(CanvasTileIndexCount == 6, "Invalid tile tri-list size.");
 
-void FCanvasTileRendererItem::InitTileBuffers(FLocalVertexFactory* VertexFactory, TArray<FTileInst>& Tiles, const FSceneView& View, bool bNeedsToSwitchVerticalAxis)
-{
-	static_assert(NUM_MATERIAL_TILE_VERTS == 4, "Invalid tile tri-list size.");
-	Data->StaticMeshVertexBuffers.PositionVertexBuffer.Init(Tiles.Num() * NUM_MATERIAL_TILE_VERTS);
-	Data->StaticMeshVertexBuffers.StaticMeshVertexBuffer.Init(Tiles.Num() * NUM_MATERIAL_TILE_VERTS, 1);
-	Data->StaticMeshVertexBuffers.ColorVertexBuffer.Init(Tiles.Num() * NUM_MATERIAL_TILE_VERTS);
+	const uint32 TotalVertexCount = Tiles.Num() * CanvasTileVertexCount;
+	const uint32 TotalIndexCount = Tiles.Num() * CanvasTileIndexCount;
 
-	Data->IndexBuffer.Indices.SetNum(Tiles.Num() * 6);
+	StaticMeshVertexBuffers.PositionVertexBuffer.Init(TotalVertexCount);
+	StaticMeshVertexBuffers.StaticMeshVertexBuffer.Init(TotalVertexCount, 1);
+	StaticMeshVertexBuffers.ColorVertexBuffer.Init(TotalVertexCount);
+	IndexBuffer.Indices.SetNum(TotalIndexCount);
 
 	for (int32 i = 0; i < Tiles.Num(); i++)
 	{
 		const FTileInst& Tile = Tiles[i];
-		const uint32 FirstIndex = i * 6;
+		const uint32 FirstIndex = i * CanvasTileIndexCount;
 
-		Data->IndexBuffer.Indices[FirstIndex + 0] = FirstIndex + 0;
-		Data->IndexBuffer.Indices[FirstIndex + 1] = FirstIndex + 1;
-		Data->IndexBuffer.Indices[FirstIndex + 2] = FirstIndex + 2;
-		Data->IndexBuffer.Indices[FirstIndex + 3] = FirstIndex + 2;
-		Data->IndexBuffer.Indices[FirstIndex + 4] = FirstIndex + 1;
-		Data->IndexBuffer.Indices[FirstIndex + 5] = FirstIndex + 3;
+		IndexBuffer.Indices[FirstIndex + 0] = FirstIndex + 0;
+		IndexBuffer.Indices[FirstIndex + 1] = FirstIndex + 1;
+		IndexBuffer.Indices[FirstIndex + 2] = FirstIndex + 2;
+		IndexBuffer.Indices[FirstIndex + 3] = FirstIndex + 2;
+		IndexBuffer.Indices[FirstIndex + 4] = FirstIndex + 1;
+		IndexBuffer.Indices[FirstIndex + 5] = FirstIndex + 3;
 
 		const float X = Tile.X;
 		const float Y = Tile.Y;
@@ -300,56 +111,81 @@ void FCanvasTileRendererItem::InitTileBuffers(FLocalVertexFactory* VertexFactory
 
 		if (bNeedsToSwitchVerticalAxis)
 		{
-			Data->StaticMeshVertexBuffers.PositionVertexBuffer.VertexPosition(i * NUM_MATERIAL_TILE_VERTS + 0) = FVector(X + SizeX, View.UnscaledViewRect.Height() - (Y + SizeY), 0.0f);
-			Data->StaticMeshVertexBuffers.PositionVertexBuffer.VertexPosition(i * NUM_MATERIAL_TILE_VERTS + 1) = FVector(X, View.UnscaledViewRect.Height() - (Y + SizeY), 0.0f);
-			Data->StaticMeshVertexBuffers.PositionVertexBuffer.VertexPosition(i * NUM_MATERIAL_TILE_VERTS + 2) = FVector(X + SizeX, View.UnscaledViewRect.Height() - Y, 0.0f);
-			Data->StaticMeshVertexBuffers.PositionVertexBuffer.VertexPosition(i * NUM_MATERIAL_TILE_VERTS + 3) = FVector(X, View.UnscaledViewRect.Height() - Y, 0.0f);
-
-			Data->StaticMeshVertexBuffers.StaticMeshVertexBuffer.SetVertexUV(i * NUM_MATERIAL_TILE_VERTS + 0, 0, FVector2D(U + SizeU, V + SizeV));
-			Data->StaticMeshVertexBuffers.StaticMeshVertexBuffer.SetVertexUV(i * NUM_MATERIAL_TILE_VERTS + 1, 0, FVector2D(U, V + SizeV));
-			Data->StaticMeshVertexBuffers.StaticMeshVertexBuffer.SetVertexUV(i * NUM_MATERIAL_TILE_VERTS + 2, 0, FVector2D(U + SizeU, V));
-			Data->StaticMeshVertexBuffers.StaticMeshVertexBuffer.SetVertexUV(i * NUM_MATERIAL_TILE_VERTS + 3, 0, FVector2D(U, V));
+			StaticMeshVertexBuffers.PositionVertexBuffer.VertexPosition(i * CanvasTileVertexCount + 0) = FVector(X + SizeX, View.UnscaledViewRect.Height() - (Y + SizeY), 0.0f);
+			StaticMeshVertexBuffers.PositionVertexBuffer.VertexPosition(i * CanvasTileVertexCount + 1) = FVector(X, View.UnscaledViewRect.Height() - (Y + SizeY), 0.0f);
+			StaticMeshVertexBuffers.PositionVertexBuffer.VertexPosition(i * CanvasTileVertexCount + 2) = FVector(X + SizeX, View.UnscaledViewRect.Height() - Y, 0.0f);
+			StaticMeshVertexBuffers.PositionVertexBuffer.VertexPosition(i * CanvasTileVertexCount + 3) = FVector(X, View.UnscaledViewRect.Height() - Y, 0.0f);
+
+			StaticMeshVertexBuffers.StaticMeshVertexBuffer.SetVertexUV(i * CanvasTileVertexCount + 0, 0, FVector2D(U + SizeU, V + SizeV));
+			StaticMeshVertexBuffers.StaticMeshVertexBuffer.SetVertexUV(i * CanvasTileVertexCount + 1, 0, FVector2D(U, V + SizeV));
+			StaticMeshVertexBuffers.StaticMeshVertexBuffer.SetVertexUV(i * CanvasTileVertexCount + 2, 0, FVector2D(U + SizeU, V));
+			StaticMeshVertexBuffers.StaticMeshVertexBuffer.SetVertexUV(i * CanvasTileVertexCount + 3, 0, FVector2D(U, V));
 		}
 		else
 		{
-			Data->StaticMeshVertexBuffers.PositionVertexBuffer.VertexPosition(i * NUM_MATERIAL_TILE_VERTS + 0) = FVector(X + SizeX, Y, 0.0f);
-			Data->StaticMeshVertexBuffers.PositionVertexBuffer.VertexPosition(i * NUM_MATERIAL_TILE_VERTS + 1) = FVector(X, Y, 0.0f);
-			Data->StaticMeshVertexBuffers.PositionVertexBuffer.VertexPosition(i * NUM_MATERIAL_TILE_VERTS + 2) = FVector(X + SizeX, Y + SizeY, 0.0f);
-			Data->StaticMeshVertexBuffers.PositionVertexBuffer.VertexPosition(i * NUM_MATERIAL_TILE_VERTS + 3) = FVector(X, Y + SizeY, 0.0f);
-
-			Data->StaticMeshVertexBuffers.StaticMeshVertexBuffer.SetVertexUV(i * NUM_MATERIAL_TILE_VERTS + 0, 0, FVector2D(U + SizeU, V));
-			Data->StaticMeshVertexBuffers.StaticMeshVertexBuffer.SetVertexUV(i * NUM_MATERIAL_TILE_VERTS + 1, 0, FVector2D(U, V));
-			Data->StaticMeshVertexBuffers.StaticMeshVertexBuffer.SetVertexUV(i * NUM_MATERIAL_TILE_VERTS + 2, 0, FVector2D(U + SizeU, V + SizeV));
-			Data->StaticMeshVertexBuffers.StaticMeshVertexBuffer.SetVertexUV(i * NUM_MATERIAL_TILE_VERTS + 3, 0, FVector2D(U, V + SizeV));
+			StaticMeshVertexBuffers.PositionVertexBuffer.VertexPosition(i * CanvasTileVertexCount + 0) = FVector(X + SizeX, Y, 0.0f);
+			StaticMeshVertexBuffers.PositionVertexBuffer.VertexPosition(i * CanvasTileVertexCount + 1) = FVector(X, Y, 0.0f);
+			StaticMeshVertexBuffers.PositionVertexBuffer.VertexPosition(i * CanvasTileVertexCount + 2) = FVector(X + SizeX, Y + SizeY, 0.0f);
+			StaticMeshVertexBuffers.PositionVertexBuffer.VertexPosition(i * CanvasTileVertexCount + 3) = FVector(X, Y + SizeY, 0.0f);
+
+			StaticMeshVertexBuffers.StaticMeshVertexBuffer.SetVertexUV(i * CanvasTileVertexCount + 0, 0, FVector2D(U + SizeU, V));
+			StaticMeshVertexBuffers.StaticMeshVertexBuffer.SetVertexUV(i * CanvasTileVertexCount + 1, 0, FVector2D(U, V));
+			StaticMeshVertexBuffers.StaticMeshVertexBuffer.SetVertexUV(i * CanvasTileVertexCount + 2, 0, FVector2D(U + SizeU, V + SizeV));
+			StaticMeshVertexBuffers.StaticMeshVertexBuffer.SetVertexUV(i * CanvasTileVertexCount + 3, 0, FVector2D(U, V + SizeV));
 		}
 
-		for (int j = 0; j < NUM_MATERIAL_TILE_VERTS; j++)
+		for (int j = 0; j < CanvasTileVertexCount; j++)
 		{
-			Data->StaticMeshVertexBuffers.StaticMeshVertexBuffer.SetVertexTangents(i * NUM_MATERIAL_TILE_VERTS + j, FVector(1.0f, 0.0f, 0.0f), FVector(0.0f, 1.0f, 0.0f), FVector(0.0f, 0.0f, 1.0f));
-			Data->StaticMeshVertexBuffers.ColorVertexBuffer.VertexColor(i * NUM_MATERIAL_TILE_VERTS + j) = Tile.InColor;
+			StaticMeshVertexBuffers.StaticMeshVertexBuffer.SetVertexTangents(i * CanvasTileVertexCount + j, FVector(1.0f, 0.0f, 0.0f), FVector(0.0f, 1.0f, 0.0f), FVector(0.0f, 0.0f, 1.0f));
+			StaticMeshVertexBuffers.ColorVertexBuffer.VertexColor(i * CanvasTileVertexCount + j) = Tile.InColor;
 		}
 	}
 
-	FCanvasTileRendererItem::FRenderData* RenderData = Data;
-	ENQUEUE_RENDER_COMMAND(CreateTileVerticesBuffersLegacyInit)(
-		[VertexFactory, RenderData](FRHICommandListImmediate& RHICmdList)
+	StaticMeshVertexBuffers.PositionVertexBuffer.InitResource();
+	StaticMeshVertexBuffers.StaticMeshVertexBuffer.InitResource();
+	StaticMeshVertexBuffers.ColorVertexBuffer.InitResource();
+	IndexBuffer.InitResource();
+	VertexFactory.InitResource();
+	TileMesh.InitResource();
+}
+
+void FCanvasTileRendererItem::FRenderData::ReleaseTileMesh()
+{
+	TileMesh.ReleaseResource();
+	VertexFactory.ReleaseResource();
+	IndexBuffer.ReleaseResource();
+	StaticMeshVertexBuffers.PositionVertexBuffer.ReleaseResource();
+	StaticMeshVertexBuffers.StaticMeshVertexBuffer.ReleaseResource();
+	StaticMeshVertexBuffers.ColorVertexBuffer.ReleaseResource();
+}
+
+void FCanvasTileRendererItem::FRenderData::RenderTiles(
+	FRHICommandListImmediate& RHICmdList,
+	FMeshPassProcessorRenderState& DrawRenderState,
+	const FSceneView& View,
+	bool bIsHitTesting,
+	bool bNeedsToSwitchVerticalAxis)
+{
+	check(IsInRenderingThread());
+
+	IRendererModule& RendererModule = GetRendererModule();
+
+	InitTileMesh(View, bNeedsToSwitchVerticalAxis);
+
+	SCOPED_DRAW_EVENTF(RHICmdList, CanvasDrawTile, *MaterialRenderProxy->GetMaterial(GMaxRHIFeatureLevel)->GetFriendlyName());
+
+	for (int32 TileIdx = 0; TileIdx < Tiles.Num(); TileIdx++)
 	{
-		RenderData->StaticMeshVertexBuffers.PositionVertexBuffer.InitResource();
-		RenderData->StaticMeshVertexBuffers.StaticMeshVertexBuffer.InitResource();
-		RenderData->StaticMeshVertexBuffers.ColorVertexBuffer.InitResource();
-
-		FLocalVertexFactory::FDataType RData;
-		RenderData->StaticMeshVertexBuffers.PositionVertexBuffer.BindPositionVertexBuffer(VertexFactory, RData);
-		RenderData->StaticMeshVertexBuffers.StaticMeshVertexBuffer.BindTangentVertexBuffer(VertexFactory, RData);
-		RenderData->StaticMeshVertexBuffers.StaticMeshVertexBuffer.BindPackedTexCoordVertexBuffer(VertexFactory, RData);
-		RenderData->StaticMeshVertexBuffers.StaticMeshVertexBuffer.BindLightMapVertexBuffer(VertexFactory, RData, 0);
-		RenderData->StaticMeshVertexBuffers.ColorVertexBuffer.BindColorVertexBuffer(VertexFactory, RData);
-		VertexFactory->SetData(RData);
-
-		RenderData->IndexBuffer.InitResource();
-		RenderData->VertexFactory.InitResource();
-		RenderData->TileMesh.InitResource();
-	});
+		FRenderData::FTileInst& Tile = Tiles[TileIdx];
+
+		FMeshBatch& Mesh = TileMesh.MeshElement;
+		Mesh.MaterialRenderProxy = MaterialRenderProxy;
+		Mesh.Elements[0].FirstIndex = CanvasTileIndexCount * TileIdx;
+
+		RendererModule.DrawTileMesh(RHICmdList, DrawRenderState, View, Mesh, bIsHitTesting, Tile.HitProxyId);
+	}
+
+	ReleaseTileMesh();
 }
 
 bool FCanvasTileRendererItem::Render_RenderThread(FRHICommandListImmediate& RHICmdList, FMeshPassProcessorRenderState& DrawRenderState, const FCanvas* Canvas)
@@ -366,20 +202,21 @@ bool FCanvasTileRendererItem::Render_RenderThread(FRHICommandListImmediate& RHIC
 	}
 
 	checkSlow(Data);
-	// current render target set for the canvas
+
 	const FRenderTarget* CanvasRenderTarget = Canvas->GetRenderTarget();
-	FSceneViewFamily* ViewFamily = new FSceneViewFamily(FSceneViewFamily::ConstructionValues(
+
+	TUniquePtr<const FSceneViewFamily> ViewFamily = MakeUnique<const FSceneViewFamily>(FSceneViewFamily::ConstructionValues(
 		CanvasRenderTarget,
 		nullptr,
 		FEngineShowFlags(ESFIM_Game))
 		.SetWorldTimes(CurrentWorldTime, DeltaWorldTime, CurrentRealTime)
 		.SetGammaCorrection(CanvasRenderTarget->GetDisplayGamma()));
 
-	FIntRect ViewRect(FIntPoint(0, 0), CanvasRenderTarget->GetSizeXY());
+	const FIntRect ViewRect(FIntPoint(0, 0), CanvasRenderTarget->GetSizeXY());
 
 	// make a temporary view
 	FSceneViewInitOptions ViewInitOptions;
-	ViewInitOptions.ViewFamily = ViewFamily;
+	ViewInitOptions.ViewFamily = ViewFamily.Get();
 	ViewInitOptions.SetViewRectangle(ViewRect);
 	ViewInitOptions.ViewOrigin = FVector::ZeroVector;
 	ViewInitOptions.ViewRotationMatrix = FMatrix::Identity;
@@ -387,39 +224,17 @@ bool FCanvasTileRendererItem::Render_RenderThread(FRHICommandListImmediate& RHIC
 	ViewInitOptions.BackgroundColor = FLinearColor::Black;
 	ViewInitOptions.OverlayColor = FLinearColor::White;
 
-	FSceneView* View = new FSceneView(ViewInitOptions);
-	
-	bool bNeedsToSwitchVerticalAxis = RHINeedsToSwitchVerticalAxis(Canvas->GetShaderPlatform()) && IsMobileHDR();
-
-	Data->TileMesh.InitResource();
-	InitTileBuffers(&Data->VertexFactory, Data->Tiles, *View, bNeedsToSwitchVerticalAxis);
-
-	for (int32 TileIdx = 0; TileIdx < Data->Tiles.Num(); TileIdx++)
-	{
-		FRenderData::FTileInst& Tile = Data->Tiles[TileIdx];
-
-		// update the FMeshBatch
-		FMeshBatch& Mesh = Data->TileMesh.MeshElement;
-		Mesh.MaterialRenderProxy = Data->MaterialRenderProxy;
-		Mesh.Elements[0].FirstIndex = 6 * TileIdx;
+	TUniquePtr<const FSceneView> View = MakeUnique<const FSceneView>(ViewInitOptions);
 
-		GetRendererModule().DrawTileMesh(RHICmdList, DrawRenderState, *View, Mesh, Canvas->IsHitTesting(), Tile.HitProxyId);
-	}
+	const bool bNeedsToSwitchVerticalAxis = RHINeedsToSwitchVerticalAxis(Canvas->GetShaderPlatform()) && IsMobileHDR();
 
-	Data->StaticMeshVertexBuffers.PositionVertexBuffer.ReleaseResource();
-	Data->StaticMeshVertexBuffers.StaticMeshVertexBuffer.ReleaseResource();
-	Data->StaticMeshVertexBuffers.ColorVertexBuffer.ReleaseResource();
-	Data->IndexBuffer.ReleaseResource();
-	Data->TileMesh.ReleaseResource();
-	Data->VertexFactory.ReleaseResource();
+	Data->RenderTiles(RHICmdList, DrawRenderState, *View, Canvas->IsHitTesting(), bNeedsToSwitchVerticalAxis);
 
-	delete View->Family;
-	delete View;
 	if (Canvas->GetAllowedModes() & FCanvas::Allow_DeleteOnRender)
 	{
-		delete Data;
-		Data = NULL;
+		Data = nullptr;
 	}
+
 	return true;
 }
 
@@ -437,16 +252,17 @@ bool FCanvasTileRendererItem::Render_GameThread(const FCanvas* Canvas, FRenderTh
 	}
 
 	checkSlow(Data);
-	// current render target set for the canvas
+
 	const FRenderTarget* CanvasRenderTarget = Canvas->GetRenderTarget();
-	FSceneViewFamily* ViewFamily = new FSceneViewFamily(FSceneViewFamily::ConstructionValues(
+
+	const FSceneViewFamily* ViewFamily = new FSceneViewFamily(FSceneViewFamily::ConstructionValues(
 		CanvasRenderTarget,
 		Canvas->GetScene(),
 		FEngineShowFlags(ESFIM_Game))
 		.SetWorldTimes(CurrentWorldTime, DeltaWorldTime, CurrentRealTime)
 		.SetGammaCorrection(CanvasRenderTarget->GetDisplayGamma()));
 
-	FIntRect ViewRect(FIntPoint(0, 0), CanvasRenderTarget->GetSizeXY());
+	const FIntRect ViewRect(FIntPoint(0, 0), CanvasRenderTarget->GetSizeXY());
 
 	// make a temporary view
 	FSceneViewInitOptions ViewInitOptions;
@@ -458,67 +274,31 @@ bool FCanvasTileRendererItem::Render_GameThread(const FCanvas* Canvas, FRenderTh
 	ViewInitOptions.BackgroundColor = FLinearColor::Black;
 	ViewInitOptions.OverlayColor = FLinearColor::White;
 
-	FSceneView* View = new FSceneView(ViewInitOptions);
-
-	bool bNeedsToSwitchVerticalAxis = RHINeedsToSwitchVerticalAxis(Canvas->GetShaderPlatform()) && IsMobileHDR();
-
-	struct FDrawTileParameters
-	{
-		FSceneView* View;
-		FRenderData* RenderData;
-		uint32 bIsHitTesting : 1;
-		uint32 AllowedCanvasModes;
-	};
-	FDrawTileParameters DrawTileParameters =
-	{
-		View,
-		Data,
-		Canvas->IsHitTesting() ? uint32(1) : uint32(0),
-		Canvas->GetAllowedModes()
-	};
+	const FSceneView* View = new FSceneView(ViewInitOptions);
 
-	InitTileBuffers(&Data->VertexFactory, Data->Tiles, *View, bNeedsToSwitchVerticalAxis);
+	const bool bNeedsToSwitchVerticalAxis = RHINeedsToSwitchVerticalAxis(Canvas->GetShaderPlatform()) && IsMobileHDR();
+	const bool bIsHitTesting = Canvas->IsHitTesting();
+	const bool bDeleteOnRender = Canvas->GetAllowedModes() & FCanvas::Allow_DeleteOnRender;
 
 	RenderScope.EnqueueRenderCommand(
-		[DrawTileParameters](FRHICommandListImmediate& RHICmdList)
+		[LocalData = Data, View, bIsHitTesting, bNeedsToSwitchVerticalAxis]
+		(FRHICommandListImmediate& RHICmdList)
 	{
-		SCOPED_DRAW_EVENTF(RHICmdList, CanvasDrawTile, *DrawTileParameters.RenderData->MaterialRenderProxy->GetMaterial(GMaxRHIFeatureLevel)->GetFriendlyName());
-
-		FMeshPassProcessorRenderState DrawRenderState(*DrawTileParameters.View);
+		FMeshPassProcessorRenderState DrawRenderState(*View);
 
 		// disable depth test & writes
 		DrawRenderState.SetDepthStencilState(TStaticDepthStencilState<false, CF_Always>::GetRHI());
 
-		for (int32 TileIdx = 0; TileIdx < DrawTileParameters.RenderData->Tiles.Num(); TileIdx++)
-		{
-			FRenderData::FTileInst& Tile = DrawTileParameters.RenderData->Tiles[TileIdx];
-
-			// update the FMeshBatch
-			FMeshBatch& Mesh = DrawTileParameters.RenderData->TileMesh.MeshElement;
-			Mesh.MaterialRenderProxy = DrawTileParameters.RenderData->MaterialRenderProxy;
-			Mesh.Elements[0].BaseVertexIndex = NUM_MATERIAL_TILE_VERTS * TileIdx;
-
-			GetRendererModule().DrawTileMesh(RHICmdList, DrawRenderState, *DrawTileParameters.View, Mesh, DrawTileParameters.bIsHitTesting, Tile.HitProxyId);
-		}
-
-		DrawTileParameters.RenderData->StaticMeshVertexBuffers.PositionVertexBuffer.ReleaseResource();
-		DrawTileParameters.RenderData->StaticMeshVertexBuffers.StaticMeshVertexBuffer.ReleaseResource();
-		DrawTileParameters.RenderData->StaticMeshVertexBuffers.ColorVertexBuffer.ReleaseResource();
-		DrawTileParameters.RenderData->TileMesh.ReleaseResource();
-		DrawTileParameters.RenderData->VertexFactory.ReleaseResource();
-		DrawTileParameters.RenderData->IndexBuffer.ReleaseResource();
+		LocalData->RenderTiles(RHICmdList, DrawRenderState, *View, bIsHitTesting, bNeedsToSwitchVerticalAxis);
 
-		delete DrawTileParameters.View->Family;
-		delete DrawTileParameters.View;
-		if (DrawTileParameters.AllowedCanvasModes & FCanvas::Allow_DeleteOnRender)
-		{
-			delete DrawTileParameters.RenderData;
-		}
+		delete View->Family;
+		delete View;
 	});
 
-	if (Canvas->GetAllowedModes() & FCanvas::Allow_DeleteOnRender)
+	if (bDeleteOnRender)
 	{
-		Data = NULL;
+		Data = nullptr;
 	}
+
 	return true;
-}
+}
\ No newline at end of file
diff --git a/Engine/Source/Runtime/Engine/Public/CanvasTypes.h b/Engine/Source/Runtime/Engine/Public/CanvasTypes.h
--- a/Engine/Source/Runtime/Engine/Public/CanvasTypes.h
+++ b/Engine/Source/Runtime/Engine/Public/CanvasTypes.h
@@ -969,15 +969,10 @@ public:
 		const FCanvas::FTransformEntry& InTransform=FCanvas::FTransformEntry(FMatrix::Identity),
 		bool bInFreezeTime=false)
 		// this data is deleted after rendering has completed
-		:	Data(new FRenderData(InFeatureLevel,InMaterialRenderProxy,InTransform))
-		,	bFreezeTime(bInFreezeTime)
+		: Data(MakeShared<FRenderData>(InFeatureLevel,InMaterialRenderProxy,InTransform))
+		, bFreezeTime(bInFreezeTime)
 	{}
 
-	/**
-	* Destructor to delete data in case nothing rendered
-	*/
-	virtual ~FCanvasTileRendererItem();	
-
 	/**
 	* FCanvasTileRendererItem instance accessor
 	*
@@ -1042,8 +1037,11 @@ private:
 	class FTileVertexFactory : public FLocalVertexFactory
 	{
 	public:
-		/** Default constructor. */
-		FTileVertexFactory(ERHIFeatureLevel::Type InFeatureLevel);
+		FTileVertexFactory(const FStaticMeshVertexBuffers* VertexBuffers, ERHIFeatureLevel::Type InFeatureLevel);
+		void InitResource() override;
+
+	private:
+		const FStaticMeshVertexBuffers* VertexBuffers;
 	};
 
 	class FTileMesh : public FRenderResource
@@ -1051,11 +1049,9 @@ private:
 	public:
 		FTileMesh(const FRawIndexBuffer* IndexBuffer, const FTileVertexFactory* VertexFactory);
 
-		/** The mesh element. */
 		FMeshBatch MeshElement;
 
-		virtual void InitRHI() override;
-		virtual void ReleaseRHI() override;
+		void InitRHI() override;
 	private:
 		const FRawIndexBuffer* IndexBuffer;
 		const FTileVertexFactory* VertexFactory;
@@ -1063,51 +1059,56 @@ private:
 
 	class FRenderData
 	{
-		friend class FCanvasTileRendererItem;
+	public:
+		FRenderData(
+			ERHIFeatureLevel::Type InFeatureLevel,
+			const FMaterialRenderProxy* InMaterialRenderProxy,
+			const FCanvas::FTransformEntry& InTransform);
+
+		void RenderTiles(
+			FRHICommandListImmediate& RHICmdList,
+			FMeshPassProcessorRenderState& DrawRenderState,
+			const FSceneView& View,
+			bool bIsHitTesting,
+			bool bNeedsToSwitchVerticalAxis);
+
+		const FMaterialRenderProxy* const MaterialRenderProxy;
+		const FCanvas::FTransformEntry Transform;
+
+		inline int32 AddTile(float X, float Y, float SizeX, float SizeY, float U, float V, float SizeU, float SizeV, FHitProxyId HitProxyId, FColor InColor)
+		{
+			FTileInst NewTile = { X,Y,SizeX,SizeY,U,V,SizeU,SizeV,HitProxyId,InColor };
+			return Tiles.Add(NewTile);
+		};
 
 	private:
-		/** The buffer containing vertex data. */
+		void InitTileMesh(const FSceneView& View, bool bNeedsToSwitchVerticalAxis);
+		void ReleaseTileMesh();
+
 		FRawIndexBuffer IndexBuffer;
 		FStaticMeshVertexBuffers StaticMeshVertexBuffers;
 		FTileVertexFactory VertexFactory;
 		FTileMesh TileMesh;
 
-	public:
-
-		FRenderData(ERHIFeatureLevel::Type InFeatureLevel,
-			const FMaterialRenderProxy* InMaterialRenderProxy = nullptr,
-			const FCanvas::FTransformEntry& InTransform = FCanvas::FTransformEntry(FMatrix::Identity));
-		
-		const FMaterialRenderProxy* MaterialRenderProxy;
-		FCanvas::FTransformEntry Transform;
-
 		struct FTileInst
 		{
-			float X,Y;
-			float SizeX,SizeY;
-			float U,V;
-			float SizeU,SizeV;
+			float X, Y;
+			float SizeX, SizeY;
+			float U, V;
+			float SizeU, SizeV;
 			FHitProxyId HitProxyId;
 			FColor InColor;
 		};
 		TArray<FTileInst> Tiles;
-
-		FORCEINLINE int32 AddTile(float X, float Y, float SizeX, float SizeY, float U, float V, float SizeU, float SizeV, FHitProxyId HitProxyId, FColor InColor)
-		{
-			FTileInst NewTile = {X,Y,SizeX,SizeY,U,V,SizeU,SizeV,HitProxyId,InColor};
-			return Tiles.Add(NewTile);
-		};
 	};
+
 	/**
-	* Render data which is allocated when a new FCanvasTileRendererItem is added for rendering.
-	* This data is only freed on the rendering thread once the item has finished rendering
-	*/
-	FRenderData* Data;	
+	 * Render data which is allocated when a new FCanvasTileRendererItem is added for rendering.
+	 * This data is only freed on the rendering thread once the item has finished rendering
+	 */
+	TSharedPtr<FRenderData> Data;
 
 	const bool bFreezeTime;
-
-	typedef FRenderData::FTileInst FTileInst;
-	void InitTileBuffers(FLocalVertexFactory* VertexFactory, TArray<FTileInst>& Tiles, const FSceneView& View, bool bNeedsToSwitchVerticalAxis);
 };
 
 /**
