Callbacks when VT producers are destroyed, used to signal uniform buffers that need to be updated/recreated with new VT values
- Material VT stacks register callback to recreate allocated VT
- LightmapResourceCluster registers callback to recreate allocated VT
- PrimitiveSceneInfo (if lit by VT lightmap) registers callback to dirty itself in GPUScene (since VT uniforms for lightmaps are stroed here)

diff --git a/Engine/Source/Runtime/Engine/Private/LightMap.cpp b/Engine/Source/Runtime/Engine/Private/LightMap.cpp
--- a/Engine/Source/Runtime/Engine/Private/LightMap.cpp
+++ b/Engine/Source/Runtime/Engine/Private/LightMap.cpp
@@ -3198,10 +3198,7 @@ bool FQuantizedLightmapData::HasNonZeroData() const
 
 FLightmapResourceCluster::~FLightmapResourceCluster()
 {
-	if (AllocatedVT)
-	{
-		GetRendererModule().DestroyVirtualTexture(AllocatedVT);
-	}
+	check(AllocatedVT == nullptr);
 }
 
 void FLightmapResourceCluster::UpdateUniformBuffer(ERHIFeatureLevel::Type InFeatureLevel)
@@ -3211,16 +3208,28 @@ void FLightmapResourceCluster::UpdateUniformBuffer(ERHIFeatureLevel::Type InFeat
 	ENQUEUE_RENDER_COMMAND(SetFeatureLevel)(
 		[Cluster, InFeatureLevel](FRHICommandList& RHICmdList)
 	{
-		const bool bAllowHighQualityLightMaps = AllowHighQualityLightmaps(InFeatureLevel);
-		const bool bUseVirtualTextures = bAllowHighQualityLightMaps && (CVarVirtualTexturedLightMaps.GetValueOnRenderThread() != 0) && UseVirtualTexturing(InFeatureLevel);
-	
 		Cluster->FeatureLevel = InFeatureLevel;
+		Cluster->UpdateUniformBuffer_RenderThread();
+	});
+}
 
-		FLightmapResourceClusterShaderParameters Parameters;
-		GetLightmapClusterResourceParameters(InFeatureLevel, Cluster->Input, bUseVirtualTextures ? Cluster->AcquireAllocatedVT() : nullptr, Parameters);
+void FLightmapResourceCluster::UpdateUniformBuffer_RenderThread()
+{
+	check(IsInRenderingThread());
+	const bool bAllowHighQualityLightMaps = AllowHighQualityLightmaps(FeatureLevel);
+	const bool bUseVirtualTextures = bAllowHighQualityLightMaps && (CVarVirtualTexturedLightMaps.GetValueOnRenderThread() != 0) && UseVirtualTexturing(FeatureLevel);
 
-		RHIUpdateUniformBuffer(Cluster->UniformBuffer, &Parameters);
-	});
+	FLightmapResourceClusterShaderParameters Parameters;
+	GetLightmapClusterResourceParameters(FeatureLevel, Input, bUseVirtualTextures ? AcquireAllocatedVT() : nullptr, Parameters);
+
+	RHIUpdateUniformBuffer(UniformBuffer, &Parameters);
+}
+
+static void OnVirtualTextureDestroyed(const FVirtualTextureProducerHandle& InHandle, void* Baton)
+{
+	FLightmapResourceCluster* Cluster = static_cast<FLightmapResourceCluster*>(Baton);
+	Cluster->ReleaseAllocatedVT();
+	Cluster->UpdateUniformBuffer_RenderThread();
 }
 
 IAllocatedVirtualTexture* FLightmapResourceCluster::AcquireAllocatedVT() const
@@ -3232,6 +3241,9 @@ IAllocatedVirtualTexture* FLightmapResourceCluster::AcquireAllocatedVT() const
 	{
 		check(VirtualTexture->VirtualTextureStreaming);
 		const FVirtualTexture2DResource* Resource = (FVirtualTexture2DResource*)VirtualTexture->Resource;
+		const FVirtualTextureProducerHandle ProducerHandle = Resource->GetProducerHandle();
+
+		GetRendererModule().AddVirtualTextureProducerDestroyedCallback(ProducerHandle, &OnVirtualTextureDestroyed, const_cast<FLightmapResourceCluster*>(this));
 
 		FAllocatedVTDescription VTDesc;
 		VTDesc.Dimensions = 2;
@@ -3245,7 +3257,7 @@ IAllocatedVirtualTexture* FLightmapResourceCluster::AcquireAllocatedVT() const
 			if (LayerIndex != ~0u)
 			{
 				VTDesc.NumLayers = TypeIndex + 1u;
-				VTDesc.ProducerHandle[TypeIndex] = Resource->GetProducerHandle(); // use the same producer for each layer
+				VTDesc.ProducerHandle[TypeIndex] = ProducerHandle; // use the same producer for each layer
 				VTDesc.LocalLayerToProduce[TypeIndex] = LayerIndex;
 			}
 			else
@@ -3263,7 +3275,7 @@ IAllocatedVirtualTexture* FLightmapResourceCluster::AcquireAllocatedVT() const
 				// this isn't strictly necessary, but without this VT feedback analysis will see an unmapped page each frame for the empty layer,
 				// and attempt to do some extra work before determining there's nothing else to do...this wastes CPU time
 				// By mapping to layer0, we ensure that every layer has a valid mapping, and the overhead of mapping the empty layer to layer0 is very small, since layer0 will already be resident
-				VTDesc.ProducerHandle[LayerIndex] = Resource->GetProducerHandle();
+				VTDesc.ProducerHandle[LayerIndex] = ProducerHandle;
 				VTDesc.LocalLayerToProduce[LayerIndex] = 0u;
 			}
 		}
@@ -3274,10 +3286,26 @@ IAllocatedVirtualTexture* FLightmapResourceCluster::AcquireAllocatedVT() const
 	return AllocatedVT;
 }
 
+void FLightmapResourceCluster::ReleaseAllocatedVT()
+{
+	if (AllocatedVT)
+	{
+		GetRendererModule().RemoveAllVirtualTextureProducerDestroyedCallbacks(this);
+		GetRendererModule().DestroyVirtualTexture(AllocatedVT);
+		AllocatedVT = nullptr;
+	}
+}
+
 void FLightmapResourceCluster::InitRHI()
 {
 	FLightmapResourceClusterShaderParameters Parameters;
 	GetLightmapClusterResourceParameters(GMaxRHIFeatureLevel, FLightmapClusterResourceInput(), nullptr, Parameters);
 
 	UniformBuffer = FLightmapResourceClusterShaderParameters::CreateUniformBuffer(Parameters, UniformBuffer_MultiFrame);
-}
\ No newline at end of file
+}
+
+void FLightmapResourceCluster::ReleaseRHI()
+{
+	ReleaseAllocatedVT();
+	UniformBuffer = nullptr;
+}
diff --git a/Engine/Source/Runtime/Engine/Private/Materials/MaterialShared.cpp b/Engine/Source/Runtime/Engine/Private/Materials/MaterialShared.cpp
--- a/Engine/Source/Runtime/Engine/Private/Materials/MaterialShared.cpp
+++ b/Engine/Source/Runtime/Engine/Private/Materials/MaterialShared.cpp
@@ -2284,7 +2284,13 @@ static IAllocatedVirtualTexture* GetPreallocatedVTStack(const FMaterialRenderCon
 	return (Texture == nullptr) ? nullptr : Texture->GetAllocatedVirtualTexture();
 }
 
-static IAllocatedVirtualTexture* AllocateVTStack(const FMaterialRenderContext& Context, const FUniformExpressionSet& UniformExpressionSet, const FMaterialVirtualTextureStack& VTStack)
+static void OnVirtualTextureDestroyed(const FVirtualTextureProducerHandle& InHandle, void* Baton)
+{
+	FMaterialRenderProxy* MaterialProxy = static_cast<FMaterialRenderProxy*>(Baton);
+	MaterialProxy->CacheUniformExpressions(false);
+}
+
+static IAllocatedVirtualTexture* AllocateVTStack(const FMaterialRenderProxy* MaterialProxy, const FMaterialRenderContext& Context, const FUniformExpressionSet& UniformExpressionSet, const FMaterialVirtualTextureStack& VTStack)
 {
 	check(!VTStack.IsPreallocatedStack())
 
@@ -2313,8 +2319,10 @@ static IAllocatedVirtualTexture* AllocateVTStack(const FMaterialRenderContext& C
 
 			VTDesc.TileSize = VirtualTextureResourceForLayer->GetTileSize();
 			VTDesc.TileBorderSize = VirtualTextureResourceForLayer->GetBorderSize();
-			VTDesc.ProducerHandle[LayerIndex] = VirtualTextureResourceForLayer->GetProducerHandle();
+			const FVirtualTextureProducerHandle& ProducerHandle = VirtualTextureResourceForLayer->GetProducerHandle();
+			VTDesc.ProducerHandle[LayerIndex] = ProducerHandle;
 			VTDesc.LocalLayerToProduce[LayerIndex] = 0u;
+			GetRendererModule().AddVirtualTextureProducerDestroyedCallback(ProducerHandle, &OnVirtualTextureDestroyed, const_cast<FMaterialRenderProxy*>(MaterialProxy));
 			bFoundValidLayer = true;
 		}
 	}
@@ -2366,7 +2374,7 @@ void FMaterialRenderProxy::EvaluateUniformExpressions(FUniformExpressionCache& O
 		}
 		else
 		{
-			AllocatedVT = AllocateVTStack(Context, UniformExpressionSet, VTStack);
+			AllocatedVT = AllocateVTStack(this, Context, UniformExpressionSet, VTStack);
 			if (AllocatedVT != nullptr)
 			{
 				OutUniformExpressionCache.OwnedAllocatedVTs.Add(AllocatedVT);
@@ -2460,6 +2468,8 @@ void FMaterialRenderProxy::InvalidateUniformExpressionCache(bool bRecreateUnifor
 		UniformExpressionCache[i].CachedUniformExpressionShaderMap = nullptr;
 		UniformExpressionCache[i].ResetAllocatedVTs();
 
+		GetRendererModule().RemoveAllVirtualTextureProducerDestroyedCallbacks(this);
+
 		if (bRecreateUniformBuffer)
 		{
 			// This is required if the FMaterial is being recompiled (the uniform buffer layout will change).
diff --git a/Engine/Source/Runtime/Engine/Private/SceneManagement.cpp b/Engine/Source/Runtime/Engine/Private/SceneManagement.cpp
--- a/Engine/Source/Runtime/Engine/Private/SceneManagement.cpp
+++ b/Engine/Source/Runtime/Engine/Private/SceneManagement.cpp
@@ -787,8 +787,32 @@ void FLightCacheInterface::CreatePrecomputedLightingUniformBuffer_RenderingThrea
 	{
 		FPrecomputedLightingUniformParameters Parameters;
 		GetPrecomputedLightingParameters(FeatureLevel, Parameters, this);
-		PrecomputedLightingUniformBuffer = FPrecomputedLightingUniformParameters::CreateUniformBuffer(Parameters, UniformBuffer_MultiFrame);
+		if (PrecomputedLightingUniformBuffer)
+		{
+			// Don't recreate the buffer if it already exists
+			RHIUpdateUniformBuffer(PrecomputedLightingUniformBuffer, &Parameters);
+		}
+		else
+		{
+			PrecomputedLightingUniformBuffer = FPrecomputedLightingUniformParameters::CreateUniformBuffer(Parameters, UniformBuffer_MultiFrame);
+		}
+	}
+}
+
+bool FLightCacheInterface::GetVirtualTextureLightmapProducer(ERHIFeatureLevel::Type FeatureLevel, FVirtualTextureProducerHandle& OutProducerHandle)
+{
+	const FLightMapInteraction LightMapInteraction = GetLightMapInteraction(FeatureLevel);
+	if (LightMapInteraction.GetType() == LMIT_Texture)
+	{
+		const ULightMapVirtualTexture2D* VirtualTexture = LightMapInteraction.GetVirtualTexture();
+		if (VirtualTexture)
+		{
+			FVirtualTexture2DResource* Resource = (FVirtualTexture2DResource*)VirtualTexture->Resource;
+			OutProducerHandle = Resource->GetProducerHandle();
+			return true;
+		}
 	}
+	return false;
 }
 
 IMPLEMENT_GLOBAL_SHADER_PARAMETER_STRUCT(FLightmapResourceClusterShaderParameters, "LightmapResourceCluster");
diff --git a/Engine/Source/Runtime/Engine/Public/LightMap.h b/Engine/Source/Runtime/Engine/Public/LightMap.h
--- a/Engine/Source/Runtime/Engine/Public/LightMap.h
+++ b/Engine/Source/Runtime/Engine/Public/LightMap.h
@@ -640,20 +640,19 @@ public:
 	virtual ~FLightmapResourceCluster();
 
 	ENGINE_API virtual void InitRHI();
+	ENGINE_API virtual void ReleaseRHI();
 
-	virtual void ReleaseRHI()
-	{
-		UniformBuffer = nullptr;
-	}
-	
 	ENGINE_API void UpdateUniformBuffer(ERHIFeatureLevel::Type InFeatureLevel);
 
+	void UpdateUniformBuffer_RenderThread();
+
 	/**
 	 * Allocates virtual texture on demand and returns it, may return nullptr if not using virtual texture
 	 * 'const' method setting 'mutable' member is not amazing, but this is required to work around ordering of render commands submitted from main thread,
 	 * it's possible for commands that want to access AllocatedVT from here execute before this has a chance to run InitRHI()
 	 */
 	IAllocatedVirtualTexture* AcquireAllocatedVT() const;
+	void ReleaseAllocatedVT();
 
 	FLightmapClusterResourceInput Input;
 	mutable IAllocatedVirtualTexture* AllocatedVT;
diff --git a/Engine/Source/Runtime/Engine/Public/SceneManagement.h b/Engine/Source/Runtime/Engine/Public/SceneManagement.h
--- a/Engine/Source/Runtime/Engine/Public/SceneManagement.h
+++ b/Engine/Source/Runtime/Engine/Public/SceneManagement.h
@@ -753,6 +753,8 @@ public:
 	
 	ENGINE_API void CreatePrecomputedLightingUniformBuffer_RenderingThread(ERHIFeatureLevel::Type FeatureLevel);
 
+	ENGINE_API bool GetVirtualTextureLightmapProducer(ERHIFeatureLevel::Type FeatureLevel, FVirtualTextureProducerHandle& OutProducerHandle);
+
 	// @param InLightMap may be 0
 	void SetLightMap(const FLightMap* InLightMap)
 	{
diff --git a/Engine/Source/Runtime/RenderCore/Public/RendererInterface.h b/Engine/Source/Runtime/RenderCore/Public/RendererInterface.h
--- a/Engine/Source/Runtime/RenderCore/Public/RendererInterface.h
+++ b/Engine/Source/Runtime/RenderCore/Public/RendererInterface.h
@@ -698,6 +698,8 @@ public:
 
 	virtual FVirtualTextureProducerHandle RegisterVirtualTextureProducer(const FVTProducerDescription& Desc, IVirtualTexture* Producer) = 0;
 	virtual void ReleaseVirtualTextureProducer(const FVirtualTextureProducerHandle& Handle) = 0;
+	virtual void AddVirtualTextureProducerDestroyedCallback(const FVirtualTextureProducerHandle& Handle, FVTProducerDestroyedFunction* Function, void* Baton) = 0;
+	virtual void RemoveAllVirtualTextureProducerDestroyedCallbacks(const void* Baton) = 0;
 
 	/**	Provided a list of packed virtual texture tile ids, let the VT system request them. Note this should be called as long as the tiles are needed.*/
 	virtual void RequestVirtualTextureTilesForRegion(IAllocatedVirtualTexture* AllocatedVT, const FVector2D& InScreenSpaceSize, const FIntRect& InTextureRegion, int32 InMipLevel) = 0;
diff --git a/Engine/Source/Runtime/RenderCore/Public/VirtualTexturing.h b/Engine/Source/Runtime/RenderCore/Public/VirtualTexturing.h
--- a/Engine/Source/Runtime/RenderCore/Public/VirtualTexturing.h
+++ b/Engine/Source/Runtime/RenderCore/Public/VirtualTexturing.h
@@ -116,6 +116,8 @@ struct FVTProducerDescription
 	TEnumAsByte<EPixelFormat> LayerFormat[VIRTUALTEXTURE_SPACE_MAXLAYERS] = { PF_Unknown };
 };
 
+typedef void (FVTProducerDestroyedFunction)(const FVirtualTextureProducerHandle& InHandle, void* Baton);
+
 class IVirtualTextureFinalizer
 {
 public:
diff --git a/Engine/Source/Runtime/Renderer/Private/PrimitiveSceneInfo.cpp b/Engine/Source/Runtime/Renderer/Private/PrimitiveSceneInfo.cpp
--- a/Engine/Source/Runtime/Renderer/Private/PrimitiveSceneInfo.cpp
+++ b/Engine/Source/Runtime/Renderer/Private/PrimitiveSceneInfo.cpp
@@ -13,6 +13,8 @@
 #include "ScenePrivate.h"
 #include "HAL/LowLevelMemTracker.h"
 #include "RayTracing/RayTracingMaterialHitShaders.h"
+#include "VT/VirtualTextureSystem.h"
+#include "GPUScene.h"
 
 /** An implementation of FStaticPrimitiveDrawInterface that stores the drawn elements for the rendering thread to use. */
 class FBatchingSPDI : public FStaticPrimitiveDrawInterface
@@ -134,6 +136,7 @@ FPrimitiveSceneInfo::FPrimitiveSceneInfo(UPrimitiveComponent* InComponent,FScene
 	bNeedsStaticMeshUpdateWithoutVisibilityCheck(false),
 	bNeedsUniformBufferUpdate(false),
 	bIndirectLightingCacheBufferDirty(false),
+	bRegisteredVirtualTextureProducerCallback(false),
 	LightmapDataOffset(INDEX_NONE),
 	NumLightmapDataEntries(0)
 {
@@ -431,6 +434,52 @@ void FPrimitiveSceneInfo::AddStaticMeshes(FRHICommandListImmediate& RHICmdList,
 	}
 }
 
+static void OnVirtualTextureDestroyed(const FVirtualTextureProducerHandle& InHandle, void* Baton)
+{
+	FPrimitiveSceneInfo* PrimitiveSceneInfo = static_cast<FPrimitiveSceneInfo*>(Baton);
+
+	// Update the main uniform buffer
+	PrimitiveSceneInfo->UpdateStaticLightingBuffer();
+
+	// Also need to update lightmap data inside GPUScene, if that's enabled
+	AddPrimitiveToUpdateGPU(*PrimitiveSceneInfo->Scene, PrimitiveSceneInfo->GetIndex());
+}
+
+int32 FPrimitiveSceneInfo::UpdateStaticLightingBuffer()
+{
+	checkSlow(IsInRenderingThread());
+
+	if (bRegisteredVirtualTextureProducerCallback)
+	{
+		// Remove any previous VT callbacks
+		FVirtualTextureSystem::Get().RemoveAllProducerDestroyedCallbacks(this);
+		bRegisteredVirtualTextureProducerCallback = false;
+	}
+
+	FPrimitiveSceneProxy::FLCIArray LCIs;
+	Proxy->GetLCIs(LCIs);
+	for (int32 i = 0; i < LCIs.Num(); ++i)
+	{
+		FLightCacheInterface* LCI = LCIs[i];
+
+		if (LCI)
+		{
+			LCI->CreatePrecomputedLightingUniformBuffer_RenderingThread(Scene->GetFeatureLevel());
+
+			// If lightmap is using virtual texture, need to set a callback to update our uniform buffers if VT is destroyed,
+			// since we cache VT parameters inside these uniform buffers
+			FVirtualTextureProducerHandle VTProducerHandle;
+			if (LCI->GetVirtualTextureLightmapProducer(Scene->GetFeatureLevel(), VTProducerHandle))
+			{
+				FVirtualTextureSystem::Get().AddProducerDestroyedCallback(VTProducerHandle, &OnVirtualTextureDestroyed, this);
+				bRegisteredVirtualTextureProducerCallback = true;
+			}
+		}
+	}
+
+	return LCIs.Num();
+}
+
 void FPrimitiveSceneInfo::AddToScene(FRHICommandListImmediate& RHICmdList, bool bUpdateStaticDrawLists, bool bAddToStaticDrawLists)
 {
 	check(IsInRenderingThread());
@@ -482,20 +531,7 @@ void FPrimitiveSceneInfo::AddToScene(FRHICommandListImmediate& RHICmdList, bool
 	const bool bAllowStaticLighting = FReadOnlyCVARCache::Get().bAllowStaticLighting;
 	if (bAllowStaticLighting)
 	{
-		FPrimitiveSceneProxy::FLCIArray LCIs;
-		Proxy->GetLCIs(LCIs);
-		for (int32 i = 0; i < LCIs.Num(); ++i)
-		{
-			FLightCacheInterface* LCI = LCIs[i];
-
-			if (LCI)
-			{
-				LCI->CreatePrecomputedLightingUniformBuffer_RenderingThread(Scene->GetFeatureLevel());
-			}
-		}
-
-		NumLightmapDataEntries = LCIs.Num();
-
+		NumLightmapDataEntries = UpdateStaticLightingBuffer();
 		if (NumLightmapDataEntries > 0 && UseGPUScene(GMaxRHIShaderPlatform, Scene->GetFeatureLevel()))
 		{
 			LightmapDataOffset = Scene->GPUScene.LightmapDataAllocator.Allocate(NumLightmapDataEntries);
@@ -660,6 +696,12 @@ void FPrimitiveSceneInfo::RemoveFromScene(bool bUpdateStaticDrawLists)
 
 		RemoveStaticMeshes();
 	}
+
+	if (bRegisteredVirtualTextureProducerCallback)
+	{
+		FVirtualTextureSystem::Get().RemoveAllProducerDestroyedCallbacks(this);
+		bRegisteredVirtualTextureProducerCallback = false;
+	}
 }
 
 bool FPrimitiveSceneInfo::NeedsUpdateStaticMeshes()
diff --git a/Engine/Source/Runtime/Renderer/Private/RendererModule.h b/Engine/Source/Runtime/Renderer/Private/RendererModule.h
--- a/Engine/Source/Runtime/Renderer/Private/RendererModule.h
+++ b/Engine/Source/Runtime/Renderer/Private/RendererModule.h
@@ -101,6 +101,8 @@ public:
 	virtual void DestroyVirtualTexture(IAllocatedVirtualTexture* AllocatedVT) override;
 	virtual FVirtualTextureProducerHandle RegisterVirtualTextureProducer(const FVTProducerDescription& Desc, IVirtualTexture* Producer) override;
 	virtual void ReleaseVirtualTextureProducer(const FVirtualTextureProducerHandle& Handle) override;
+	virtual void AddVirtualTextureProducerDestroyedCallback(const FVirtualTextureProducerHandle& Handle, FVTProducerDestroyedFunction* Function, void* Baton) override;
+	virtual void RemoveAllVirtualTextureProducerDestroyedCallbacks(const void* Baton) override;
 	virtual void RequestVirtualTextureTilesForRegion(IAllocatedVirtualTexture* AllocatedVT, const FVector2D& InScreenSpaceSize, const FIntRect& InTextureRegion, int32 InMipLevel) override;
 	virtual void LoadPendingVirtualTextureTiles(FRHICommandListImmediate& RHICmdList, ERHIFeatureLevel::Type FeatureLevel) override;
 	virtual void FlushVirtualTextureCache() override;
diff --git a/Engine/Source/Runtime/Renderer/Private/SceneRendering.cpp b/Engine/Source/Runtime/Renderer/Private/SceneRendering.cpp
--- a/Engine/Source/Runtime/Renderer/Private/SceneRendering.cpp
+++ b/Engine/Source/Runtime/Renderer/Private/SceneRendering.cpp
@@ -3504,6 +3504,16 @@ void FRendererModule::ReleaseVirtualTextureProducer(const FVirtualTextureProduce
 	FVirtualTextureSystem::Get().ReleaseProducer(Handle);
 }
 
+void FRendererModule::AddVirtualTextureProducerDestroyedCallback(const FVirtualTextureProducerHandle& Handle, FVTProducerDestroyedFunction* Function, void* Baton)
+{
+	FVirtualTextureSystem::Get().AddProducerDestroyedCallback(Handle, Function, Baton);
+}
+
+void FRendererModule::RemoveAllVirtualTextureProducerDestroyedCallbacks(const void* Baton)
+{
+	FVirtualTextureSystem::Get().RemoveAllProducerDestroyedCallbacks(Baton);
+}
+
 void FRendererModule::RequestVirtualTextureTilesForRegion(IAllocatedVirtualTexture* AllocatedVT, const FVector2D& InScreenSpaceSize, const FIntRect& InTextureRegion, int32 InMipLevel)
 {
 	FVirtualTextureSystem::Get().RequestTilesForRegion(AllocatedVT, InScreenSpaceSize, InTextureRegion, InMipLevel);
diff --git a/Engine/Source/Runtime/Renderer/Private/VT/VirtualTextureProducer.cpp b/Engine/Source/Runtime/Renderer/Private/VT/VirtualTextureProducer.cpp
--- a/Engine/Source/Runtime/Renderer/Private/VT/VirtualTextureProducer.cpp
+++ b/Engine/Source/Runtime/Renderer/Private/VT/VirtualTextureProducer.cpp
@@ -28,6 +28,13 @@ FVirtualTextureProducerCollection::FVirtualTextureProducerCollection()
 {
 	Producers.AddDefaulted(1);
 	Producers[0].Magic = 1u; // make sure FVirtualTextureProducerHandle(0) will not resolve to the dummy producer entry
+
+	Callbacks.AddDefaulted(CallbackList_Count);
+	for (uint32 CallbackIndex = 0u; CallbackIndex < CallbackList_Count; ++CallbackIndex)
+	{
+		FCallbackEntry& Callback = Callbacks[CallbackIndex];
+		Callback.NextIndex = Callback.PrevIndex = CallbackIndex;
+	}
 }
 
 FVirtualTextureProducerHandle FVirtualTextureProducerCollection::RegisterProducer(FVirtualTextureSystem* System, const FVTProducerDescription& InDesc, IVirtualTexture* InProducer)
@@ -44,6 +51,7 @@ FVirtualTextureProducerHandle FVirtualTextureProducerCollection::RegisterProduce
 	FProducerEntry& Entry = Producers[Index];
 	Entry.Producer.Description = InDesc;
 	Entry.Producer.VirtualTexture = InProducer;
+	Entry.DestroyedCallbacksIndex = AcquireCallback();
 
 	check(InDesc.NumLayers <= VIRTUALTEXTURE_SPACE_MAXLAYERS);
 	for (uint32 LayerIndex = 0u; LayerIndex < InDesc.NumLayers; ++LayerIndex)
@@ -67,12 +75,91 @@ void FVirtualTextureProducerCollection::ReleaseProducer(FVirtualTextureSystem* S
 
 	if (FProducerEntry* Entry = GetEntry(Handle))
 	{
-		Entry->Producer.Release(System, Handle);
+		uint32 CallbackIndex = Callbacks[Entry->DestroyedCallbacksIndex].NextIndex;
+		while (CallbackIndex != Entry->DestroyedCallbacksIndex)
+		{
+			FCallbackEntry& Callback = Callbacks[CallbackIndex];
+			const uint32 NextIndex = Callback.NextIndex;
+			check(Callback.OwnerHandle == Handle);
+
+			// Move callback to pending list
+			RemoveCallbackFromList(CallbackIndex);
+			AddCallbackToList(CallbackList_Pending, CallbackIndex);
+			
+			CallbackIndex = NextIndex;
+		}
+
+		ReleaseCallback(Entry->DestroyedCallbacksIndex);
+		Entry->DestroyedCallbacksIndex = 0u;
+
 		Entry->Magic = (Entry->Magic + 1u) & 1023u;
+		Entry->Producer.Release(System, Handle);
 		ReleaseEntry(Handle.Index);
 	}
 }
 
+void FVirtualTextureProducerCollection::CallPendingCallbacks()
+{
+	uint32 CallbackIndex = Callbacks[CallbackList_Pending].NextIndex;
+	while (CallbackIndex != CallbackList_Pending)
+	{
+		FCallbackEntry& Callback = Callbacks[CallbackIndex];
+		check(Callback.DestroyedFunction);
+		check(Callback.OwnerHandle.PackedValue != 0u);
+
+		// Make a copy, then release the callback entry before calling the callback function
+		// (The destroyed callback may try to remove this or other callbacks, so need to make sure state is valid before calling)
+		const FCallbackEntry CallbackCopy(Callback);
+		Callback.DestroyedFunction = nullptr;
+		Callback.Baton = nullptr;
+		Callback.OwnerHandle = FVirtualTextureProducerHandle();
+		ReleaseCallback(CallbackIndex);
+
+		CallbackCopy.DestroyedFunction(CallbackCopy.OwnerHandle, CallbackCopy.Baton);
+		CallbackIndex = CallbackCopy.NextIndex;
+	}
+}
+
+void FVirtualTextureProducerCollection::AddDestroyedCallback(const FVirtualTextureProducerHandle& Handle, FVTProducerDestroyedFunction* Function, void* Baton)
+{
+	check(IsInRenderingThread());
+	check(Function);
+
+	FProducerEntry* Entry = GetEntry(Handle);
+	if (Entry)
+	{
+		const uint32 CallbackIndex = AcquireCallback();
+		AddCallbackToList(Entry->DestroyedCallbacksIndex, CallbackIndex);
+		FCallbackEntry& Callback = Callbacks[CallbackIndex];
+		Callback.DestroyedFunction = Function;
+		Callback.Baton = Baton;
+		Callback.OwnerHandle = Handle;
+	}
+}
+
+void FVirtualTextureProducerCollection::RemoveAllCallbacks(const void* Baton)
+{
+	check(IsInRenderingThread());
+	check(Baton);
+
+	for (int32 CallbackIndex = 0u; CallbackIndex < Callbacks.Num(); ++CallbackIndex)
+	{
+		FCallbackEntry& Callback = Callbacks[CallbackIndex];
+		if (Callback.Baton == Baton)
+		{
+			check(Callback.DestroyedFunction);
+			FProducerEntry* Entry = GetEntry(Callback.OwnerHandle);
+			if (Entry)
+			{
+				Callback.DestroyedFunction = nullptr;
+				Callback.Baton = nullptr;
+				Callback.OwnerHandle = FVirtualTextureProducerHandle();
+				ReleaseCallback(CallbackIndex);
+			}
+		}
+	}
+}
+
 FVirtualTextureProducer* FVirtualTextureProducerCollection::FindProducer(const FVirtualTextureProducerHandle& Handle)
 {
 	FProducerEntry* Entry = GetEntry(Handle);
diff --git a/Engine/Source/Runtime/Renderer/Private/VT/VirtualTextureProducer.h b/Engine/Source/Runtime/Renderer/Private/VT/VirtualTextureProducer.h
--- a/Engine/Source/Runtime/Renderer/Private/VT/VirtualTextureProducer.h
+++ b/Engine/Source/Runtime/Renderer/Private/VT/VirtualTextureProducer.h
@@ -42,6 +42,10 @@ public:
 	FVirtualTextureProducerCollection();
 	FVirtualTextureProducerHandle RegisterProducer(FVirtualTextureSystem* System, const FVTProducerDescription& InDesc, IVirtualTexture* InProducer);
 	void ReleaseProducer(FVirtualTextureSystem* System, const FVirtualTextureProducerHandle& Handle);
+	
+	void AddDestroyedCallback(const FVirtualTextureProducerHandle& Handle, FVTProducerDestroyedFunction* Function, void* Baton);
+	void RemoveAllCallbacks(const void* Baton);
+	void CallPendingCallbacks();
 
 	/**
 	 * Gets the producer associated with the given handle, or nullptr if handle is invalid
@@ -56,11 +60,28 @@ private:
 	struct FProducerEntry
 	{
 		FVirtualTextureProducer Producer;
+		uint32 DestroyedCallbacksIndex = 0u;
 		uint32 NextIndex = 0u;
 		uint32 PrevIndex = 0u;
 		uint16 Magic = 0u;
 	};
 
+	struct FCallbackEntry
+	{
+		FVTProducerDestroyedFunction* DestroyedFunction = nullptr;
+		void* Baton = nullptr;
+		FVirtualTextureProducerHandle OwnerHandle;
+		uint32 NextIndex = 0u;
+		uint32 PrevIndex = 0u;
+	};
+
+	enum ECallbackListType
+	{
+		CallbackList_Free,
+		CallbackList_Pending,
+		CallbackList_Count,
+	};
+
 	FProducerEntry* GetEntry(const FVirtualTextureProducerHandle& Handle);
 
 	void RemoveEntryFromList(uint32 Index)
@@ -109,5 +130,52 @@ private:
 		AddEntryToList(0u, Index);
 	}
 
+	void RemoveCallbackFromList(uint32 Index)
+	{
+		FCallbackEntry& Entry = Callbacks[Index];
+		Callbacks[Entry.PrevIndex].NextIndex = Entry.NextIndex;
+		Callbacks[Entry.NextIndex].PrevIndex = Entry.PrevIndex;
+		Entry.NextIndex = Entry.PrevIndex = Index;
+	}
+
+	void AddCallbackToList(uint32 HeadIndex, uint32 Index)
+	{
+		FCallbackEntry& Head = Callbacks[HeadIndex];
+		FCallbackEntry& Entry = Callbacks[Index];
+		check(Index >= CallbackList_Count); // make sure we're not trying to add a list head to another list
+
+		// make sure we're not currently in any list
+		check(Entry.NextIndex == Index);
+		check(Entry.PrevIndex == Index);
+
+		Entry.NextIndex = HeadIndex;
+		Entry.PrevIndex = Head.PrevIndex;
+		Callbacks[Head.PrevIndex].NextIndex = Index;
+		Head.PrevIndex = Index;
+	}
+
+	uint32 AcquireCallback()
+	{
+		FCallbackEntry& FreeHead = Callbacks[CallbackList_Free];
+		uint32 Index = FreeHead.NextIndex;
+		if (Index != CallbackList_Free)
+		{
+			RemoveCallbackFromList(Index);
+			return Index;
+		}
+
+		Index = Callbacks.AddDefaulted();
+		FCallbackEntry& Entry = Callbacks[Index];
+		Entry.NextIndex = Entry.PrevIndex = Index;
+		return Index;
+	}
+
+	void ReleaseCallback(uint32 Index)
+	{
+		RemoveCallbackFromList(Index);
+		AddCallbackToList(CallbackList_Free, Index);
+	}
+
 	TArray<FProducerEntry> Producers;
+	TArray<FCallbackEntry> Callbacks;
 };
diff --git a/Engine/Source/Runtime/Renderer/Private/VT/VirtualTextureSystem.cpp b/Engine/Source/Runtime/Renderer/Private/VT/VirtualTextureSystem.cpp
--- a/Engine/Source/Runtime/Renderer/Private/VT/VirtualTextureSystem.cpp
+++ b/Engine/Source/Runtime/Renderer/Private/VT/VirtualTextureSystem.cpp
@@ -476,6 +476,16 @@ void FVirtualTextureSystem::ReleaseProducer(const FVirtualTextureProducerHandle&
 	Producers.ReleaseProducer(this, Handle);
 }
 
+void FVirtualTextureSystem::AddProducerDestroyedCallback(const FVirtualTextureProducerHandle& Handle, FVTProducerDestroyedFunction* Function, void* Baton)
+{
+	Producers.AddDestroyedCallback(Handle, Function, Baton);
+}
+
+void FVirtualTextureSystem::RemoveAllProducerDestroyedCallbacks(const void* Baton)
+{
+	Producers.RemoveAllCallbacks(Baton);
+}
+
 FVirtualTextureProducer* FVirtualTextureSystem::FindProducer(const FVirtualTextureProducerHandle& Handle)
 {
 	return Producers.FindProducer(Handle);
@@ -769,6 +779,8 @@ void FVirtualTextureSystem::Update(FRHICommandListImmediate& RHICmdList, ERHIFea
 		bFlushCaches = false;
 	}
 
+	Producers.CallPendingCallbacks();
+
 	DestroyPendingVirtualTextures();
 
 	FMemStack& MemStack = FMemStack::Get();
diff --git a/Engine/Source/Runtime/Renderer/Private/VT/VirtualTextureSystem.h b/Engine/Source/Runtime/Renderer/Private/VT/VirtualTextureSystem.h
--- a/Engine/Source/Runtime/Renderer/Private/VT/VirtualTextureSystem.h
+++ b/Engine/Source/Runtime/Renderer/Private/VT/VirtualTextureSystem.h
@@ -43,6 +43,8 @@ public:
 
 	FVirtualTextureProducerHandle RegisterProducer(const FVTProducerDescription& InDesc, IVirtualTexture* InProducer);
 	void ReleaseProducer(const FVirtualTextureProducerHandle& Handle);
+	void AddProducerDestroyedCallback(const FVirtualTextureProducerHandle& Handle, FVTProducerDestroyedFunction* Function, void* Baton);
+	void RemoveAllProducerDestroyedCallbacks(const void* Baton);
 	FVirtualTextureProducer* FindProducer(const FVirtualTextureProducerHandle& Handle);
 
 	FVirtualTextureSpace* AcquireSpace(const FVTSpaceDescription& InDesc, uint32 InSizeNeeded);
diff --git a/Engine/Source/Runtime/Renderer/Public/PrimitiveSceneInfo.h b/Engine/Source/Runtime/Renderer/Public/PrimitiveSceneInfo.h
--- a/Engine/Source/Runtime/Renderer/Public/PrimitiveSceneInfo.h
+++ b/Engine/Source/Runtime/Renderer/Public/PrimitiveSceneInfo.h
@@ -440,6 +440,9 @@ public:
 	/** Helper function for writing out to the last render times to the game thread */
 	void UpdateComponentLastRenderTime(float CurrentWorldTime, bool bUpdateLastRenderTimeOnScreen) const;
 
+	/** Updates static lighting uniform buffer, returns the number of entries needed for GPUScene */
+	int32 UpdateStaticLightingBuffer();
+
 #if RHI_RAYTRACING
 	RENDERER_API FRayTracingGeometryRHIRef GetStaticRayTracingGeometryInstance(int LodLevel);
 #endif
@@ -471,6 +474,8 @@ private:
 	/** If this is TRUE, this primitive's indirect lighting cache buffer needs to be updated before it can be rendered. */
 	bool bIndirectLightingCacheBufferDirty : 1;
 
+	bool bRegisteredVirtualTextureProducerCallback : 1;
+
 	/** Offset into the scene's lightmap data buffer, when GPUScene is enabled. */
 	int32 LightmapDataOffset;
 	/** Number of entries in the scene's lightmap data buffer. */
