DR - Fix Indexmemory calculation for StaticMeshes

diff --git a/Engine/Source/Runtime/Engine/Private/StaticMesh.cpp b/Engine/Source/Runtime/Engine/Private/StaticMesh.cpp
--- a/Engine/Source/Runtime/Engine/Private/StaticMesh.cpp
+++ b/Engine/Source/Runtime/Engine/Private/StaticMesh.cpp
@@ -2521,14 +2521,19 @@ void FStaticMeshRenderData::GetResourceSizeEx(FResourceSizeEx& CumulativeResourc
 		const int32 VBSize = LODRenderData.VertexBuffers.StaticMeshVertexBuffer.GetResourceSize() +
 			LODRenderData.VertexBuffers.PositionVertexBuffer.GetStride()			* LODRenderData.VertexBuffers.PositionVertexBuffer.GetNumVertices() +
 			LODRenderData.VertexBuffers.ColorVertexBuffer.GetStride()				* LODRenderData.VertexBuffers.ColorVertexBuffer.GetNumVertices();
-		int32 IBSize = LODRenderData.IndexBuffer.GetAllocatedSize();
+		
+		int32 NumIndicies = LODRenderData.IndexBuffer.GetNumIndices();
 
 		if (LODRenderData.AdditionalIndexBuffers)
 		{
-			IBSize += LODRenderData.AdditionalIndexBuffers->WireframeIndexBuffer.GetAllocatedSize()
-			+ (RHISupportsTessellation(GShaderPlatformForFeatureLevel[GMaxRHIFeatureLevel]) ? LODRenderData.AdditionalIndexBuffers->AdjacencyIndexBuffer.GetAllocatedSize() : 0);
+			NumIndicies += LODRenderData.AdditionalIndexBuffers->ReversedDepthOnlyIndexBuffer.GetNumIndices();
+			NumIndicies += LODRenderData.AdditionalIndexBuffers->ReversedIndexBuffer.GetNumIndices();
+			NumIndicies += LODRenderData.AdditionalIndexBuffers->WireframeIndexBuffer.GetNumIndices();
+			NumIndicies += (RHISupportsTessellation(GShaderPlatformForFeatureLevel[GMaxRHIFeatureLevel]) ? LODRenderData.AdditionalIndexBuffers->AdjacencyIndexBuffer.GetNumIndices() : 0);
 		}
 
+		int32 IBSize = NumIndicies * (LODRenderData.IndexBuffer.Is32Bit() ? 4 : 2);
+
 		CumulativeResourceSize.AddUnknownMemoryBytes(VBSize + IBSize);
 		CumulativeResourceSize.AddUnknownMemoryBytes(LODRenderData.Sections.GetAllocatedSize());
 
