DR - Clang warning fix
#rnx

diff --git a/Engine/Source/Runtime/Engine/Private/InstancedStaticMesh.cpp b/Engine/Source/Runtime/Engine/Private/InstancedStaticMesh.cpp
--- a/Engine/Source/Runtime/Engine/Private/InstancedStaticMesh.cpp
+++ b/Engine/Source/Runtime/Engine/Private/InstancedStaticMesh.cpp
@@ -1008,8 +1008,8 @@ void FInstancedStaticMeshSceneProxy::GetDynamicRayTracingInstances(struct FRayTr
 		if (InstancedRenderData.Component->PerInstanceSMData.IsValidIndex(InstanceIdx))
 		{
 			const FInstancedStaticMeshInstanceData& InstanceData = InstancedRenderData.Component->PerInstanceSMData[InstanceIdx];
-			FMatrix LocalToWorld = InstancedRenderData.Component->GetComponentTransform().ToMatrixWithScale();
-			FMatrix InstanceTransform = InstanceData.Transform * LocalToWorld;
+			FMatrix ComponentLocalToWorld = InstancedRenderData.Component->GetComponentTransform().ToMatrixWithScale();
+			FMatrix InstanceTransform = InstanceData.Transform * ComponentLocalToWorld;
 
 			FRayTracingInstance RayTracingInstance;
 			RayTracingInstance.Geometry = &RenderData->LODResources[LOD].RayTracingGeometry;
