Compact flags for compiled material output

diff --git a/Engine/Source/Runtime/Engine/Private/Materials/HLSLMaterialTranslator.h b/Engine/Source/Runtime/Engine/Private/Materials/HLSLMaterialTranslator.h
--- a/Engine/Source/Runtime/Engine/Private/Materials/HLSLMaterialTranslator.h
+++ b/Engine/Source/Runtime/Engine/Private/Materials/HLSLMaterialTranslator.h
@@ -1043,11 +1043,14 @@ ResourcesString = TEXT("");
 				Errorf(TEXT("Only transparent or postprocess materials can read from scene depth."));
 			}
 
-			MaterialCompilationOutput.bUsesSceneDepthLookup = bUsesSceneDepth;
+			if (bUsesSceneDepth)
+			{
+				MaterialCompilationOutput.SetIsSceneTextureUsed(PPI_SceneDepth);
+			}
 
 			MaterialCompilationOutput.bUsesDistanceCullFade = bUsesDistanceCullFade;
 
-			if (MaterialCompilationOutput.bRequiresSceneColorCopy)
+			if (MaterialCompilationOutput.RequiresSceneColorCopy())
 			{
 				if (Domain != MD_Surface)
 				{
@@ -1309,16 +1312,16 @@ ResourcesString = TEXT("");
 			OutEnvironment.SetDefine(TEXT("NEEDS_WORLD_POSITION_EXCLUDING_SHADER_OFFSETS"), TEXT("1"));
 		}
 
-		if( bNeedsParticleSize )
+		if (bNeedsParticleSize)
 		{
 			OutEnvironment.SetDefine(TEXT("NEEDS_PARTICLE_SIZE"), TEXT("1"));
 		}
 
-		if( MaterialCompilationOutput.bNeedsSceneTextures )
+		if (MaterialCompilationOutput.bNeedsSceneTextures)
 		{
 			OutEnvironment.SetDefine(TEXT("NEEDS_SCENE_TEXTURES"), TEXT("1"));
 		}
-		if( MaterialCompilationOutput.bUsesEyeAdaptation )
+		if (MaterialCompilationOutput.bUsesEyeAdaptation)
 		{
 			OutEnvironment.SetDefine(TEXT("USES_EYE_ADAPTATION"), TEXT("1"));
 		}
@@ -4594,7 +4597,7 @@ protected:
 	void UseSceneTextureId(ESceneTextureId SceneTextureId, bool bTextureLookup)
 	{
 		MaterialCompilationOutput.bNeedsSceneTextures = true;
-		MaterialCompilationOutput.UsedSceneTextures |= (1ull << SceneTextureId);
+		MaterialCompilationOutput.SetIsSceneTextureUsed(SceneTextureId);
 
 		if(Material->GetMaterialDomain() == MD_DeferredDecal)
 		{
@@ -4657,25 +4660,7 @@ protected:
 			bUsesSceneDepth = true;
 		}
 
-		const bool bNeedsGBuffer = SceneTextureId == PPI_DiffuseColor 
-			|| SceneTextureId == PPI_SpecularColor
-			|| SceneTextureId == PPI_SubsurfaceColor
-			|| SceneTextureId == PPI_BaseColor
-			|| SceneTextureId == PPI_Specular
-			|| SceneTextureId == PPI_Metallic
-			|| SceneTextureId == PPI_WorldNormal
-			|| SceneTextureId == PPI_Opacity
-			|| SceneTextureId == PPI_Roughness
-			|| SceneTextureId == PPI_MaterialAO
-			|| SceneTextureId == PPI_DecalMask
-			|| SceneTextureId == PPI_ShadingModelColor
-			|| SceneTextureId == PPI_ShadingModelID
-			|| SceneTextureId == PPI_StoredBaseColor
-			|| SceneTextureId == PPI_StoredSpecular
-			|| SceneTextureId == PPI_Velocity;
-
-
-		MaterialCompilationOutput.bNeedsGBuffer = MaterialCompilationOutput.bNeedsGBuffer || bNeedsGBuffer;
+		const bool bNeedsGBuffer = MaterialCompilationOutput.NeedsGBuffer();
 
 		if (bNeedsGBuffer && IsForwardShadingEnabled(Platform))
 		{
@@ -4684,11 +4669,7 @@ protected:
 
 		if (SceneTextureId == PPI_Velocity)
 		{
-			if (Material->GetMaterialDomain() == MD_PostProcess)
-			{
-				MaterialCompilationOutput.bUsesVelocitySceneTexture = true;
-			}
-			else
+			if (Material->GetMaterialDomain() != MD_PostProcess)
 			{
 				Errorf(TEXT("Velocity scene textures are only available in post process materials."));
 			}
@@ -4720,7 +4701,7 @@ protected:
 			return INDEX_NONE;
 		}
 
-		MaterialCompilationOutput.bRequiresSceneColorCopy = true;
+		MaterialCompilationOutput.SetIsSceneTextureUsed(PPI_SceneColor);
 		AddEstimatedTextureSample();
 
 		int32 ScreenUVCode = GetScreenAlignedUV(Offset, ViewportUV, bUseOffset);
diff --git a/Engine/Source/Runtime/Engine/Private/Materials/MaterialShared.cpp b/Engine/Source/Runtime/Engine/Private/Materials/MaterialShared.cpp
--- a/Engine/Source/Runtime/Engine/Private/Materials/MaterialShared.cpp
+++ b/Engine/Source/Runtime/Engine/Private/Materials/MaterialShared.cpp
@@ -490,29 +490,40 @@ void FMaterialCompilationOutput::Serialize(FArchive& Ar)
 	UniformExpressionSet.Serialize(Ar);
 
 	Ar << UsedSceneTextures;
+
 #if WITH_EDITOR
-	Ar << NumUsedUVScalars;
-	Ar << NumUsedCustomInterpolatorScalars;
 	Ar << EstimatedNumTextureSamplesVS;
 	Ar << EstimatedNumTextureSamplesPS;
 	Ar << EstimatedNumVirtualTextureLookups;
+	Ar << NumUsedUVScalars;
+	Ar << NumUsedCustomInterpolatorScalars;
 #else
-	uint8 Tmp8; uint16 Tmp16;
-	Ar << Tmp8 << Tmp8;
-	Ar << Tmp16 << Tmp16 << Tmp16;
+	uint16 Tmp16 = 0;
+	uint8 Tmp8 = 0; 
+
+	Ar << Tmp16;
+	Ar << Tmp16;
+	Ar << Tmp16;
+	Ar << Tmp8;
+	Ar << Tmp8;
 #endif
 
-	Ar << bRequiresSceneColorCopy;
-	Ar << bNeedsSceneTextures;
-	Ar << bUsesEyeAdaptation;
-	Ar << bModifiesMeshPosition;
-	Ar << bUsesWorldPositionOffset;
-	Ar << bNeedsGBuffer;
-	Ar << bUsesGlobalDistanceField;
-	Ar << bUsesPixelDepthOffset;
-	Ar << bUsesSceneDepthLookup;
-	Ar << bUsesVelocitySceneTexture;
-	Ar << bUsesDistanceCullFade;
+	uint8 PackedFlags = 0;
+	PackedFlags |= (bUsesEyeAdaptation			<< 0);
+	PackedFlags |= (bModifiesMeshPosition		<< 1);
+	PackedFlags |= (bUsesWorldPositionOffset	<< 2);
+	PackedFlags |= (bUsesGlobalDistanceField	<< 3);
+	PackedFlags |= (bUsesPixelDepthOffset		<< 4);
+	PackedFlags |= (bUsesDistanceCullFade		<< 5);
+
+	Ar << PackedFlags;
+
+	bUsesEyeAdaptation			= (PackedFlags >> 0) & 1;
+	bModifiesMeshPosition		= (PackedFlags >> 1) & 1;
+	bUsesWorldPositionOffset	= (PackedFlags >> 2) & 1;
+	bUsesGlobalDistanceField	= (PackedFlags >> 3) & 1;
+	bUsesPixelDepthOffset		= (PackedFlags >> 4) & 1;
+	bUsesDistanceCullFade		= (PackedFlags >> 5) & 1;
 }
 
 void FMaterial::GetShaderMapId(EShaderPlatform Platform, FMaterialShaderMapId& OutId) const
diff --git a/Engine/Source/Runtime/Engine/Private/ShaderDerivedDataVersion.h b/Engine/Source/Runtime/Engine/Private/ShaderDerivedDataVersion.h
--- a/Engine/Source/Runtime/Engine/Private/ShaderDerivedDataVersion.h
+++ b/Engine/Source/Runtime/Engine/Private/ShaderDerivedDataVersion.h
@@ -10,6 +10,5 @@
 
 // In case of merge conflicts with DDC versions, you MUST generate a new GUID and set this new
 // guid as version
-
-#define GLOBALSHADERMAP_DERIVEDDATA_VER			TEXT("73D907EC61CB431EB3FEE5F78C520EF6")
-#define MATERIALSHADERMAP_DERIVEDDATA_VER		TEXT("AC3D6DA4A70D4E88B5C9F21CCEEF7AA1")
\ No newline at end of file
+#define GLOBALSHADERMAP_DERIVEDDATA_VER			TEXT("45EDF9DF531B4E7989EB97AF45C38B3E")
+#define MATERIALSHADERMAP_DERIVEDDATA_VER		TEXT("36C14CFABCD54E0E846080318C712DCE")
diff --git a/Engine/Source/Runtime/Engine/Public/MaterialShared.h b/Engine/Source/Runtime/Engine/Public/MaterialShared.h
--- a/Engine/Source/Runtime/Engine/Public/MaterialShared.h
+++ b/Engine/Source/Runtime/Engine/Public/MaterialShared.h
@@ -25,6 +25,7 @@
 #include "Misc/Optional.h"
 #include "Serialization/MemoryWriter.h"
 #include "Serialization/ArchiveProxy.h"
+#include "MaterialSceneTextureId.h"
 #include "VirtualTexturing.h"
 
 struct FExpressionInput;
@@ -117,7 +118,7 @@ enum EMaterialValueType
 	 * Any size float type by definition, but this is treated as a scalar which can auto convert (by replication) to any other size float vector.
 	 * Use this as the type for any scalar expressions.
 	 */
-	MCT_Float		          = 8|4|2|1,
+	MCT_Float                 = 8|4|2|1,
 	MCT_Texture2D	          = 1 << 4,
 	MCT_TextureCube	          = 1 << 5,
 	MCT_VolumeTexture         = 1 << 6,
@@ -532,79 +533,97 @@ public:
 	FMaterialCompilationOutput() :
 		UsedSceneTextures(0),
 #if WITH_EDITOR
-		NumUsedUVScalars(0),
-		NumUsedCustomInterpolatorScalars(0),
 		EstimatedNumTextureSamplesVS(0),
 		EstimatedNumTextureSamplesPS(0),
 		EstimatedNumVirtualTextureLookups(0),
+		NumUsedUVScalars(0),
+		NumUsedCustomInterpolatorScalars(0),
 #endif
-		bRequiresSceneColorCopy(false),
 		bNeedsSceneTextures(false),
 		bUsesEyeAdaptation(false),
 		bModifiesMeshPosition(false),
 		bUsesWorldPositionOffset(false),
-		bNeedsGBuffer(false),
 		bUsesGlobalDistanceField(false),
 		bUsesPixelDepthOffset(false),
-		bUsesSceneDepthLookup(false),
-		bUsesVelocitySceneTexture(false),
 		bUsesDistanceCullFade(false)
 	{}
 
 	ENGINE_API void Serialize(FArchive& Ar);
 
+	ENGINE_API bool IsSceneTextureUsed(ESceneTextureId TexId) const { return UsedSceneTextures & (1 << TexId); }
+	ENGINE_API void SetIsSceneTextureUsed(ESceneTextureId TexId) { UsedSceneTextures |= (1 << TexId); }
+
 	FUniformExpressionSet UniformExpressionSet;
 
 	/** Bitfield of the ESceneTextures used */
-	uint64 UsedSceneTextures;
+	uint32 UsedSceneTextures;
 
 #if WITH_EDITOR
-	/** Number of used custom UV scalars. */
-	uint8 NumUsedUVScalars;
-
-	/** Number of used custom vertex interpolation scalars. */
-	uint8 NumUsedCustomInterpolatorScalars;
-
 	/** Number of times SampleTexture is called, excludes custom nodes. */
 	uint16 EstimatedNumTextureSamplesVS;
 	uint16 EstimatedNumTextureSamplesPS;
 
-	/** Number of virtual texture lookups performed, excludes direct invocation in shaders (for example VT lightmaps) */
+        /** Number of virtual texture lookups performed, excludes direct invocation in shaders (for example VT lightmaps) */
 	uint16 EstimatedNumVirtualTextureLookups;
-#endif // WITH_EDITOR
 
-	/** Indicates whether the material uses scene color. */
-	bool bRequiresSceneColorCopy;
+	/** Number of used custom UV scalars. */
+	uint8 NumUsedUVScalars;
+
+	/** Number of used custom vertex interpolation scalars. */
+	uint8 NumUsedCustomInterpolatorScalars;
+#endif // WITH_EDITOR
 
 	/** true if the material needs the scenetexture lookups. */
-	bool bNeedsSceneTextures;
+	uint8 bNeedsSceneTextures : 1;
 
 	/** true if the material uses the EyeAdaptationLookup */
-	bool bUsesEyeAdaptation;
+	uint8 bUsesEyeAdaptation : 1;
 
 	/** true if the material modifies the the mesh position. */
-	bool bModifiesMeshPosition;
+	uint8 bModifiesMeshPosition : 1;
 
 	/** Whether the material uses world position offset. */
-	bool bUsesWorldPositionOffset;
-
-	/** true if the material uses any GBuffer textures */
-	bool bNeedsGBuffer;
+	uint8 bUsesWorldPositionOffset : 1;
 
 	/** true if material uses the global distance field */
-	bool bUsesGlobalDistanceField;
+	uint8 bUsesGlobalDistanceField : 1;
 
 	/** true if the material writes a pixel depth offset */
-	bool bUsesPixelDepthOffset;
+	uint8 bUsesPixelDepthOffset : 1;
+
+	/** true if the material uses distance cull fade */
+	uint8 bUsesDistanceCullFade : 1;
+
+	/** Indicates whether the material uses scene color. */
+	ENGINE_API bool RequiresSceneColorCopy() const { return IsSceneTextureUsed(PPI_SceneColor); }
+
+	/** true if the material uses any GBuffer textures */
+	ENGINE_API bool NeedsGBuffer() const
+	{
+		return
+			IsSceneTextureUsed(PPI_DiffuseColor) ||
+			IsSceneTextureUsed(PPI_SpecularColor) ||
+			IsSceneTextureUsed(PPI_SubsurfaceColor) ||
+			IsSceneTextureUsed(PPI_BaseColor) ||
+			IsSceneTextureUsed(PPI_Specular) ||
+			IsSceneTextureUsed(PPI_Metallic) ||
+			IsSceneTextureUsed(PPI_WorldNormal) ||
+			IsSceneTextureUsed(PPI_Opacity) ||
+			IsSceneTextureUsed(PPI_Roughness) ||
+			IsSceneTextureUsed(PPI_MaterialAO) ||
+			IsSceneTextureUsed(PPI_DecalMask) ||
+			IsSceneTextureUsed(PPI_ShadingModelColor) ||
+			IsSceneTextureUsed(PPI_ShadingModelID) ||
+			IsSceneTextureUsed(PPI_StoredBaseColor) ||
+			IsSceneTextureUsed(PPI_StoredSpecular) ||
+			IsSceneTextureUsed(PPI_Velocity);
+	}
 
 	/** true if the material uses the SceneDepth lookup */
-	bool bUsesSceneDepthLookup;
+	ENGINE_API bool UsesSceneDepthLookup() const { return IsSceneTextureUsed(PPI_SceneColor); }
 
 	/** true if the material uses the Velocity SceneTexture lookup */
-	bool bUsesVelocitySceneTexture;
-
-	/** true if the material uses distance cull fade */
-	bool bUsesDistanceCullFade;
+	ENGINE_API bool UsesVelocitySceneTexture() const { return IsSceneTextureUsed(PPI_Velocity); }
 };
 
 /** 
@@ -1048,16 +1067,16 @@ public:
 	const FString& GetFriendlyName() const { static FString T; return T; }
 	const FString& GetDebugDescription() const { static FString T; return T; }
 #endif
-	bool RequiresSceneColorCopy() const { return MaterialCompilationOutput.bRequiresSceneColorCopy; }
+	bool RequiresSceneColorCopy() const { return MaterialCompilationOutput.RequiresSceneColorCopy(); }
 	bool NeedsSceneTextures() const { return MaterialCompilationOutput.bNeedsSceneTextures; }
 	bool UsesGlobalDistanceField() const { return MaterialCompilationOutput.bUsesGlobalDistanceField; }
 	bool UsesWorldPositionOffset() const { return MaterialCompilationOutput.bUsesWorldPositionOffset; }
-	bool NeedsGBuffer() const { return MaterialCompilationOutput.bNeedsGBuffer; }
+	bool NeedsGBuffer() const { return MaterialCompilationOutput.NeedsGBuffer(); }
 	bool UsesEyeAdaptation() const { return MaterialCompilationOutput.bUsesEyeAdaptation; }
 	bool ModifiesMeshPosition() const { return MaterialCompilationOutput.bModifiesMeshPosition; }
 	bool UsesPixelDepthOffset() const { return MaterialCompilationOutput.bUsesPixelDepthOffset; }
-	bool UsesSceneDepthLookup() const { return MaterialCompilationOutput.bUsesSceneDepthLookup; }
-	bool UsesVelocitySceneTexture() const { return MaterialCompilationOutput.bUsesVelocitySceneTexture; }
+	bool UsesSceneDepthLookup() const { return MaterialCompilationOutput.UsesSceneDepthLookup(); }
+	bool UsesVelocitySceneTexture() const { return MaterialCompilationOutput.UsesVelocitySceneTexture(); }
 	bool UsesDistanceCullFade() const { return MaterialCompilationOutput.bUsesDistanceCullFade; }
 #if WITH_EDITOR
 	uint32 GetNumUsedUVScalars() const { return MaterialCompilationOutput.NumUsedUVScalars; }
