URuntimeVirtualTexture creates a null producer when no other producer is bound
This would be useful if the producer create/destroy is used to signal the dependent materials

diff --git a/Engine/Source/Runtime/Engine/Classes/VT/RuntimeVirtualTexture.h b/Engine/Source/Runtime/Engine/Classes/VT/RuntimeVirtualTexture.h
--- a/Engine/Source/Runtime/Engine/Classes/VT/RuntimeVirtualTexture.h
+++ b/Engine/Source/Runtime/Engine/Classes/VT/RuntimeVirtualTexture.h
@@ -3,40 +3,11 @@
 #pragma once
 
 #include "CoreMinimal.h"
-#include "Delegates/DelegateCombinations.h"
-#include "RenderResource.h"
 #include "UObject/ObjectMacros.h"
 #include "VirtualTexturing.h"
 #include "VT/RuntimeVirtualTextureEnum.h"
 #include "RuntimeVirtualTexture.generated.h"
 
-/** Runtime virtual texture FRenderResource */
-class FRuntimeVirtualTextureRenderResource : public FRenderResource
-{
-public:
-	FRuntimeVirtualTextureRenderResource(FVTProducerDescription const& InDesc, IVirtualTexture* InVirtualTextureProducer);
-
-	/** Getter for the virtual texture allocation */
-	IAllocatedVirtualTexture* GetAllocatedVirtualTexture() const { return AllocatedVirtualTexture; }
-
-protected:
-	//~ Begin FRenderResource Interface.
-	virtual void InitRHI() override;
-	virtual void ReleaseRHI() override;
-	//~ End FRenderResource Interface.
-
-	/** Allocate in the global virtual texture system. */
-	IAllocatedVirtualTexture* AcquireAllocatedVirtualTexture();
-	/** Release our virtual texture allocations  */
-	void ReleaseAllocatedVirtualTexture();
-
-private:
-	const FVTProducerDescription ProducerDesc;
-	IVirtualTexture* Producer;
-	FVirtualTextureProducerHandle ProducerHandle;
-	IAllocatedVirtualTexture* AllocatedVirtualTexture;
-};
-
 /** Runtime virtual texture UObject */
 UCLASS(ClassGroup = Rendering)
 class ENGINE_API URuntimeVirtualTexture : public UObject
@@ -109,20 +80,21 @@ public:
 	/** Release the resources for this object This will need to be called if our producer becomes stale and we aren't doing a full reinit with a new producer. */
 	void Release();
 
-	/** Getter for the associated virtual texture allocation. */
+	/** Getter for the associated virtual texture producer. Call on render thread only. */
+	FVirtualTextureProducerHandle GetProducerHandle() const;
+	/** Getter for the associated virtual texture allocation. Call on render thread only. */
 	IAllocatedVirtualTexture* GetAllocatedVirtualTexture() const;
 
 	/** Getter for the shader uniform parameters. */
 	FVector4 GetUniformParameter(int32 Index);
 
 protected:
-	/** Initialize the render resource. This kicks off render thread work. */
+	/** Initialize the render resources. This kicks off render thread work. */
 	void InitResource(IVirtualTexture* InProducer);
-	/** Release the render resource. This kicks off render thread work. */
-	void ReleaseResource();
+	/** Initialize the render resources with a null producer. This kicks off render thread work. */
+	void InitNullResource();
 
 	//~ Begin UObject Interface.
-	virtual void BeginDestroy() override;
 	virtual void GetAssetRegistryTags(TArray<FAssetRegistryTag>& OutTags) const override;
 #if WITH_EDITOR
 	virtual void PostEditChangeProperty(FPropertyChangedEvent& PropertyChangedEvent) override;
@@ -130,8 +102,8 @@ protected:
 	//~ End UObject Interface.
 
 private:
-	/** Render thread resource container */
-	FRuntimeVirtualTextureRenderResource* Resource;
+	/** Render thread resource container. */
+	class FRuntimeVirtualTextureRenderResource* Resource;
 
 	/** Material uniform parameters to support transform from world to UV coordinates. */
 	FVector4 WorldToUVTransformParameters[3];
diff --git a/Engine/Source/Runtime/Engine/Private/VT/RuntimeVirtualTexture.cpp b/Engine/Source/Runtime/Engine/Private/VT/RuntimeVirtualTexture.cpp
--- a/Engine/Source/Runtime/Engine/Private/VT/RuntimeVirtualTexture.cpp
+++ b/Engine/Source/Runtime/Engine/Private/VT/RuntimeVirtualTexture.cpp
@@ -7,72 +7,172 @@
 #include "VT/RuntimeVirtualTextureNotify.h"
 
 
-FRuntimeVirtualTextureRenderResource::FRuntimeVirtualTextureRenderResource(FVTProducerDescription const& InProducerDesc, IVirtualTexture* InVirtualTextureProducer)
-	: ProducerDesc(InProducerDesc)
-	, Producer(InVirtualTextureProducer)
-	, AllocatedVirtualTexture(nullptr)
+namespace
 {
-	check(InVirtualTextureProducer != nullptr);
-}
+	/** Null producer to use as placeholder when no producer has been set on a URuntimeVirtualTexture */
+	class FNullVirtualTextureProducer : public IVirtualTexture
+	{
+	public:
+		/** Get a producer description to use with the null producer. */
+		static void GetNullProducerDescription(FVTProducerDescription& OutDesc)
+		{
+			OutDesc.Dimensions = 2;
+			OutDesc.TileSize = 4;
+			OutDesc.TileBorderSize = 0;
+			OutDesc.BlockWidthInTiles = 1;
+			OutDesc.BlockHeightInTiles = 1;
+			OutDesc.MaxLevel = 1;
+			OutDesc.DepthInTiles = 1;
+			OutDesc.WidthInBlocks = 1;
+			OutDesc.HeightInBlocks = 1;
+			OutDesc.NumLayers = 0;
+		}
 
-void FRuntimeVirtualTextureRenderResource::InitRHI()
-{
-	ProducerHandle = GetRendererModule().RegisterVirtualTextureProducer(ProducerDesc, Producer);
+		//~ Begin IVirtualTexture Interface.
+		virtual FVTRequestPageResult RequestPageData(
+			const FVirtualTextureProducerHandle& ProducerHandle,
+			uint8 LayerMask,
+			uint8 vLevel,
+			uint32 vAddress,
+			EVTRequestPagePriority Priority
+		) override
+		{
+			return FVTRequestPageResult();
+		}
 
-	AcquireAllocatedVirtualTexture();
+		virtual IVirtualTextureFinalizer* ProducePageData(
+			FRHICommandListImmediate& RHICmdList,
+			ERHIFeatureLevel::Type FeatureLevel,
+			EVTProducePageFlags Flags,
+			const FVirtualTextureProducerHandle& ProducerHandle,
+			uint8 LayerMask,
+			uint8 vLevel,
+			uint32 vAddress,
+			uint64 RequestHandle,
+			const FVTProduceTargetLayer* TargetLayers
+		) override 
+		{
+			return nullptr; 
+		}
+		//~ End IVirtualTexture Interface.
+	};
 }
 
-void FRuntimeVirtualTextureRenderResource::ReleaseRHI()
+
+/** 
+ * Container for render thread resources created for a URuntimeVirtualTexture object. 
+ * Any access to the resources should be on the render thread only so that access is serialized with the Init()/Release() render thread tasks.
+ */
+class FRuntimeVirtualTextureRenderResource
 {
-	ReleaseAllocatedVirtualTexture();
+public:
+	FRuntimeVirtualTextureRenderResource()
+		: AllocatedVirtualTexture(nullptr)
+	{
+	}
 
-	GetRendererModule().ReleaseVirtualTextureProducer(ProducerHandle);
-	ProducerHandle = FVirtualTextureProducerHandle();
-}
+	/** Getter for the virtual texture producer. */
+	FVirtualTextureProducerHandle GetProducerHandle() const
+	{
+		check(IsInRenderingThread());
+		return ProducerHandle;
+	}
 
-IAllocatedVirtualTexture* FRuntimeVirtualTextureRenderResource::AcquireAllocatedVirtualTexture()
-{
-	check(IsInRenderingThread());
-	
-	if (AllocatedVirtualTexture == nullptr)
+	/** Getter for the virtual texture allocation. */
+	IAllocatedVirtualTexture* GetAllocatedVirtualTexture() const 
 	{
-		FAllocatedVTDescription VTDesc;
-		VTDesc.Dimensions = ProducerDesc.Dimensions;
-		VTDesc.TileSize = ProducerDesc.TileSize;
-		VTDesc.TileBorderSize = ProducerDesc.TileBorderSize;
-		VTDesc.NumLayers = ProducerDesc.NumLayers;
-		VTDesc.bPrivateSpace = true; // Dedicated page table allocation for runtime VTs
+		check(IsInRenderingThread());
+		return AllocatedVirtualTexture;
+	}
+
+	/** Queues up render thread work to create resources and also releases any old resources. */
+	void Init(FVTProducerDescription const& InDesc, IVirtualTexture* InVirtualTextureProducer)
+	{
+		FRuntimeVirtualTextureRenderResource* Resource = this;
 		
-		for (uint32 LayerIndex = 0u; LayerIndex < VTDesc.NumLayers; ++LayerIndex)
+		ENQUEUE_RENDER_COMMAND(FRuntimeVirtualTextureRenderResource_Init)(
+			[Resource, InDesc, InVirtualTextureProducer](FRHICommandList& RHICmdList)
 		{
-			VTDesc.ProducerHandle[LayerIndex] = ProducerHandle;
-			VTDesc.LocalLayerToProduce[LayerIndex] = LayerIndex;
-		}
-		
-		AllocatedVirtualTexture = GetRendererModule().AllocateVirtualTexture(VTDesc);
+			FVirtualTextureProducerHandle OldProducerHandle = Resource->ProducerHandle;
+			ReleaseVirtualTexture(Resource->AllocatedVirtualTexture);
+			Resource->ProducerHandle = GetRendererModule().RegisterVirtualTextureProducer(InDesc, InVirtualTextureProducer);
+			Resource->AllocatedVirtualTexture = AllocateVirtualTexture(InDesc, Resource->ProducerHandle);
+			// Release old producer after new one is created so that any destroy callbacks can access the new producer
+			GetRendererModule().ReleaseVirtualTextureProducer(OldProducerHandle);
+		});
 	}
 
-	return AllocatedVirtualTexture;
-}
-
-void FRuntimeVirtualTextureRenderResource::ReleaseAllocatedVirtualTexture()
-{
-	if (AllocatedVirtualTexture)
+	/** Queues up render thread work to release resources. */
+	void Release()
 	{
-		GetRendererModule().DestroyVirtualTexture(AllocatedVirtualTexture);
+		FVirtualTextureProducerHandle ProducerHandleToRelease = ProducerHandle;
+		ProducerHandle = FVirtualTextureProducerHandle();
+		IAllocatedVirtualTexture* AllocatedVirtualTextureToRelease = AllocatedVirtualTexture;
 		AllocatedVirtualTexture = nullptr;
+
+		ENQUEUE_RENDER_COMMAND(FRuntimeVirtualTextureRenderResource_Release)(
+			[ProducerHandleToRelease, AllocatedVirtualTextureToRelease](FRHICommandList& RHICmdList)
+		{
+			ReleaseVirtualTexture(AllocatedVirtualTextureToRelease);
+			GetRendererModule().ReleaseVirtualTextureProducer(ProducerHandleToRelease);
+		});
+	}
+
+protected:
+	/** Allocate in the global virtual texture system. */
+	static IAllocatedVirtualTexture* AllocateVirtualTexture(FVTProducerDescription const& InDesc, FVirtualTextureProducerHandle const& InProducerHandle)
+	{
+		IAllocatedVirtualTexture* OutAllocatedVirtualTexture = nullptr;
+
+		// Check for NumLayers avoids allocating for the null producer
+		if (InDesc.NumLayers > 0)
+		{
+			FAllocatedVTDescription VTDesc;
+			VTDesc.Dimensions = InDesc.Dimensions;
+			VTDesc.TileSize = InDesc.TileSize;
+			VTDesc.TileBorderSize = InDesc.TileBorderSize;
+			VTDesc.NumLayers = InDesc.NumLayers;
+			VTDesc.bPrivateSpace = true; // Dedicated page table allocation for runtime VTs
+
+			for (uint32 LayerIndex = 0u; LayerIndex < VTDesc.NumLayers; ++LayerIndex)
+			{
+				VTDesc.ProducerHandle[LayerIndex] = InProducerHandle;
+				VTDesc.LocalLayerToProduce[LayerIndex] = LayerIndex;
+			}
+
+			OutAllocatedVirtualTexture = GetRendererModule().AllocateVirtualTexture(VTDesc);
+		}
+
+		return OutAllocatedVirtualTexture;
 	}
-}
+
+	/** Release our virtual texture allocations  */
+	static void ReleaseVirtualTexture(IAllocatedVirtualTexture* InAllocatedVirtualTexture)
+	{
+		if (InAllocatedVirtualTexture != nullptr)
+		{
+			GetRendererModule().DestroyVirtualTexture(InAllocatedVirtualTexture);
+		}
+	}
+
+private:
+	FVirtualTextureProducerHandle ProducerHandle;
+	IAllocatedVirtualTexture* AllocatedVirtualTexture;
+};
+
 
 URuntimeVirtualTexture::URuntimeVirtualTexture(const FObjectInitializer& ObjectInitializer)
 	: Super(ObjectInitializer)
-	, Resource(nullptr)
 {
+	// Initialize the RHI resources with a null producer
+	Resource = new FRuntimeVirtualTextureRenderResource;
+	InitNullResource();
 }
 
 URuntimeVirtualTexture::~URuntimeVirtualTexture()
 {
-	check(Resource == nullptr);
+	Resource->Release();
+	delete Resource;
 }
 
 void URuntimeVirtualTexture::GetProducerDescription(FVTProducerDescription& OutDesc) const
@@ -124,9 +224,14 @@ int32 URuntimeVirtualTexture::GetEstimatedPhysicalTextureMemoryKb() const
 	return 0;
 }
 
+FVirtualTextureProducerHandle URuntimeVirtualTexture::GetProducerHandle() const
+{
+	return Resource->GetProducerHandle();
+}
+
 IAllocatedVirtualTexture* URuntimeVirtualTexture::GetAllocatedVirtualTexture() const
 {
-	return (Resource != nullptr) ? Resource->GetAllocatedVirtualTexture() : nullptr;
+	return Resource->GetAllocatedVirtualTexture();
 }
 
 FVector4 URuntimeVirtualTexture::GetUniformParameter(int32 Index)
@@ -143,40 +248,27 @@ void URuntimeVirtualTexture::Initialize(IVirtualTexture* InProducer, FTransform
 	WorldToUVTransformParameters[1] = BoxToWorld.GetUnitAxis(EAxis::X) * 1.f / BoxToWorld.GetScale3D().X;
 	WorldToUVTransformParameters[2] = BoxToWorld.GetUnitAxis(EAxis::Y) * 1.f / BoxToWorld.GetScale3D().Y;
 
-	ReleaseResource();
 	InitResource(InProducer);
 }
 
 void URuntimeVirtualTexture::Release()
 {
-	ReleaseResource();
+	InitNullResource();
 }
 
 void URuntimeVirtualTexture::InitResource(IVirtualTexture* InProducer)
 {
-	check(Resource == nullptr);
-	if (Resource == nullptr && InProducer != nullptr)
-	{
-		FVTProducerDescription Desc;
-		GetProducerDescription(Desc);
-		Resource = new FRuntimeVirtualTextureRenderResource(Desc, InProducer);
-		BeginInitResource(Resource);
-	}
-}
-
-void URuntimeVirtualTexture::ReleaseResource()
-{
-	if (Resource != nullptr)
-	{
-		BeginReleaseResource(Resource);
-		Resource = nullptr;
-	}
+	FVTProducerDescription Desc;
+	GetProducerDescription(Desc);
+	Resource->Init(Desc, InProducer);
 }
 
-void URuntimeVirtualTexture::BeginDestroy()
+void URuntimeVirtualTexture::InitNullResource()
 {
-	ReleaseResource();
-	Super::BeginDestroy();
+	FVTProducerDescription Desc;
+	FNullVirtualTextureProducer::GetNullProducerDescription(Desc);
+	FNullVirtualTextureProducer* Producer = new FNullVirtualTextureProducer;
+	Resource->Init(Desc, Producer);
 }
 
 void URuntimeVirtualTexture::GetAssetRegistryTags(TArray<FAssetRegistryTag>& OutTags) const
