Add ConditionalPostLoad to express dependency of UHierarchicalInstancedStaticMeshComponent on level lighting
Fixes race condition on load for foliage light maps

diff --git a/Engine/Source/Runtime/Engine/Classes/Components/HierarchicalInstancedStaticMeshComponent.h b/Engine/Source/Runtime/Engine/Classes/Components/HierarchicalInstancedStaticMeshComponent.h
--- a/Engine/Source/Runtime/Engine/Classes/Components/HierarchicalInstancedStaticMeshComponent.h
+++ b/Engine/Source/Runtime/Engine/Classes/Components/HierarchicalInstancedStaticMeshComponent.h
@@ -200,6 +200,7 @@ public:
 	virtual void Serialize(FArchive& Ar) override;
 	virtual void GetResourceSizeEx(FResourceSizeEx& CumulativeResourceSize) override;
 	virtual void PostDuplicate(bool bDuplicateForPIE) override;
+	virtual void PostLoad() override;
 	virtual FPrimitiveSceneProxy* CreateSceneProxy() override;
 	virtual FBoxSphereBounds CalcBounds(const FTransform& BoundTransform) const override;
 #if WITH_EDITOR
diff --git a/Engine/Source/Runtime/Engine/Private/HierarchicalInstancedStaticMesh.cpp b/Engine/Source/Runtime/Engine/Private/HierarchicalInstancedStaticMesh.cpp
--- a/Engine/Source/Runtime/Engine/Private/HierarchicalInstancedStaticMesh.cpp
+++ b/Engine/Source/Runtime/Engine/Private/HierarchicalInstancedStaticMesh.cpp
@@ -2074,6 +2074,22 @@ void UHierarchicalInstancedStaticMeshComponent::PostDuplicate(bool bDuplicateFor
 	}
 }
 
+void UHierarchicalInstancedStaticMeshComponent::PostLoad()
+{
+	// ConditionalPostLoad the level lightmap data since we expect it to be loaded before PostLoadPerInstanceData which is called from Super::PostLoad()
+	AActor* Owner = GetOwner();
+	if (Owner != nullptr)
+	{
+		ULevel* OwnerLevel = Owner->GetLevel();
+		if (OwnerLevel != nullptr && OwnerLevel->MapBuildData != nullptr)
+		{
+			OwnerLevel->MapBuildData->ConditionalPostLoad();
+		}
+	}
+
+	Super::PostLoad();
+}
+
 void UHierarchicalInstancedStaticMeshComponent::RemoveInstancesInternal(const int32* InstanceIndices, int32 Num)
 {
 	if (IsAsyncBuilding() && Num > 0)
