DR - Compile fox
#rnx

diff --git a/Engine/Source/Runtime/Engine/Private/SkeletalRenderGPUSkin.cpp b/Engine/Source/Runtime/Engine/Private/SkeletalRenderGPUSkin.cpp
--- a/Engine/Source/Runtime/Engine/Private/SkeletalRenderGPUSkin.cpp
+++ b/Engine/Source/Runtime/Engine/Private/SkeletalRenderGPUSkin.cpp
@@ -86,7 +86,7 @@ void FMorphVertexBuffer::InitDynamicRHI()
 	Flags = (EBufferUsageFlags)(Flags | BUF_ShaderResource);
 
 	VertexBufferRHI = RHICreateVertexBuffer(Size, Flags, CreateInfo);
-	bool bUsesSkinCache = RHISupportsComputeShaders(GMaxRHIShaderPlatform) && IsGPUSkinCacheAvailable() && GEnableGPUSkinCache;
+	bool bUsesSkinCache = RHISupportsComputeShaders(GMaxRHIShaderPlatform) && IsGPUSkinCacheAvailable(GMaxRHIShaderPlatform) && GEnableGPUSkinCache;
 	if (bUsesSkinCache)
 	{
 		SRVValue = RHICreateShaderResourceView(VertexBufferRHI, 4, PF_R32_FLOAT);
