Support for culling in ISM and HISM

diff --git a/Engine/Source/Runtime/Engine/Private/InstancedStaticMesh.cpp b/Engine/Source/Runtime/Engine/Private/InstancedStaticMesh.cpp
--- a/Engine/Source/Runtime/Engine/Private/InstancedStaticMesh.cpp
+++ b/Engine/Source/Runtime/Engine/Private/InstancedStaticMesh.cpp
@@ -44,11 +44,6 @@
 #include "UObject/EditorObjectVersion.h"
 
 
-static TAutoConsoleVariable<int32> CVarRayTracingRenderInstances(
-	TEXT("r.RayTracing.Instances"),
-	1,
-	TEXT("Include static mesh instances in the ray tracing acceleration structure"));
-
 const int32 InstancedStaticMeshMaxTexCoord = 8;
 
 IMPLEMENT_HIT_PROXY(HInstancedStaticMeshInstance, HHitProxy);
@@ -58,6 +53,37 @@ TAutoConsoleVariable<int32> CVarMinLOD(
 	-1,
 	TEXT("Used to discard the top LODs for performance evaluation. -1: Disable all effects of this cvar."));
 
+static TAutoConsoleVariable<int32> CVarRayTracingRenderInstances(
+	TEXT("r.RayTracing.Instances"),
+	1,
+	TEXT("Include static mesh instances in ray tracing effects (default = 1 (Instances enabled in ray tracing))"));
+
+static TAutoConsoleVariable<int32> CVarRayTracingRenderInstancesCulling(
+	TEXT("r.RayTracing.Instances.Culling"),
+	0,
+	TEXT("Enable culling for instances in ray tracing (default = 0 (Culling disabled))"));
+
+static TAutoConsoleVariable<float> CVarRayTracingInstancesCullClusterMaxRadiusMultiplier(
+	TEXT("r.RayTracing.Instances.CullClusterMaxRadiusMultiplier"),
+	20.0f, 
+	TEXT("Multiplier for the maximum instance size (default = 20cm)"));
+
+static TAutoConsoleVariable<float> CVarRayTracingInstancesCullClusterRadius(
+	TEXT("r.RayTracing.Instances.CullClusterRadius"),
+	10000.0f, // 100 m
+	TEXT("Ignore instances outside of this radius in ray tracing effects (default = 10000 (100m))"));
+
+static TAutoConsoleVariable<float> CVarRayTracingInstancesLowScaleThreshold(
+	TEXT("r.RayTracing.Instances.LowScaleRadiusThreshold"),
+	50.0f, // Instances with a radius smaller than this threshold get culled after CVarRayTracingInstancesLowScaleCullRadius
+	TEXT("Threshold that classifies instances as small (default = 50cm))"));
+
+static TAutoConsoleVariable<float> CVarRayTracingInstancesLowScaleCullRadius(
+	TEXT("r.RayTracing.Instances.LowScaleCullRadius"),
+	1000.0f, 
+	TEXT("Cull radius for small instances (default = 1000 (10m))"));
+
+
 /** InstancedStaticMeshInstance hit proxy */
 void HInstancedStaticMeshInstance::AddReferencedObjects(FReferenceCollector& Collector)
 {
@@ -998,39 +1024,208 @@ void FInstancedStaticMeshSceneProxy::GetDynamicRayTracingInstances(struct FRayTr
 	}
 
 	//#dxr_todo: select the appropriate LOD depending on Context.View
-	//#dxr_todo: add culling logic for instances that are far away
 	uint32 LOD = 0;
-
 	const int32 InstanceCount = InstancedRenderData.Component->PerInstanceSMData.Num();
 
-	for (int32 InstanceIdx = 0; InstanceIdx < InstanceCount; ++InstanceIdx)
+	if (CVarRayTracingRenderInstancesCulling.GetValueOnRenderThread() > 0 && RayTracingCullClusterInstances.Num() > 0)
 	{
-		if (InstancedRenderData.Component->PerInstanceSMData.IsValidIndex(InstanceIdx))
+		const float BVHCullRadius = CVarRayTracingInstancesCullClusterRadius.GetValueOnRenderThread();
+		const float BVHLowScaleThreshold = CVarRayTracingInstancesLowScaleThreshold.GetValueOnRenderThread();
+		const float BVHLowScaleRadius = CVarRayTracingInstancesLowScaleCullRadius.GetValueOnRenderThread();
+		const bool ApplyGeneralCulling = BVHCullRadius > 0.0f;
+		const bool ApplyLowScaleCulling = BVHLowScaleThreshold > 0.0f && BVHLowScaleRadius > 0.0f;
+		FMatrix ToWorld = InstancedRenderData.Component->GetComponentTransform().ToMatrixWithScale();
+
+		// Iterate over all culling clusters
+		for (int32 ClusterIdx = 0; ClusterIdx < RayTracingCullClusterBoundsMin.Num(); ++ClusterIdx)
 		{
-			const FInstancedStaticMeshInstanceData& InstanceData = InstancedRenderData.Component->PerInstanceSMData[InstanceIdx];
-			FMatrix ComponentLocalToWorld = InstancedRenderData.Component->GetComponentTransform().ToMatrixWithScale();
-			FMatrix InstanceTransform = InstanceData.Transform * ComponentLocalToWorld;
+			FVector VClusterBBoxSize = RayTracingCullClusterBoundsMax[ClusterIdx] - RayTracingCullClusterBoundsMin[ClusterIdx];
+			FVector VClusterCenter = 0.5f * (RayTracingCullClusterBoundsMax[ClusterIdx] + RayTracingCullClusterBoundsMin[ClusterIdx]);
+			FVector VToClusterCenter = VClusterCenter - Context.ReferenceView->ViewLocation;
+			float ClusterRadius = 0.5f * VClusterBBoxSize.Size();
+			float DistToClusterCenter = VToClusterCenter.Size();
 
-			FRayTracingInstance RayTracingInstance;
-			RayTracingInstance.Geometry = &RenderData->LODResources[LOD].RayTracingGeometry;
-			RayTracingInstance.InstanceTransforms.Add(InstanceTransform);
+			// Cull whole cluster if the bounding sphere is too far away
+			if ((DistToClusterCenter - ClusterRadius) > BVHCullRadius && ApplyGeneralCulling)
+			{
+				continue;
+			}
 
-			int32 SectionCount = InstancedRenderData.LODModels[LOD].Sections.Num();
+			TDoubleLinkedList< uint32 >* InstanceList = RayTracingCullClusterInstances[ClusterIdx];
 
-			for (int32 SectionIdx = 0; SectionIdx < SectionCount; ++SectionIdx)
+			// Unroll instances in the current cluster into the array
+			for (TDoubleLinkedList<uint32>::TDoubleLinkedListNode* InstancePtr = InstanceList->GetHead(); InstancePtr != nullptr; InstancePtr = InstancePtr->GetNextNode())
 			{
-				//#dxr_todo: so far we use the parent static mesh path to get material data
-				FMeshBatch MeshBatch;
-				FStaticMeshSceneProxy::GetMeshElement(LOD, 0, SectionIdx, 0, false, false, MeshBatch);
+				const uint32 Instance = InstancePtr->GetValue();
 
-				RayTracingInstance.Materials.Add(MeshBatch);
+				if (InstancedRenderData.Component->PerInstanceSMData.IsValidIndex(Instance))
+				{
+					const FInstancedStaticMeshInstanceData& InstanceData = InstancedRenderData.Component->PerInstanceSMData[Instance];
+					FMatrix InstanceTransform = InstanceData.Transform * ToWorld;
+					FVector InstanceLocation = InstanceTransform.TransformPosition({ 0.0f,0.0f,0.0f });
+					FVector VToInstanceCenter = Context.ReferenceView->ViewLocation - InstanceLocation;
+					float   DistanceToInstanceCenter = VToInstanceCenter.Size();
+
+					FVector VMin, VMax, VDiag;
+					InstancedRenderData.Component->GetLocalBounds(VMin, VMax);
+					VMin = InstanceTransform.TransformPosition(VMin);
+					VMax = InstanceTransform.TransformPosition(VMax);
+					VDiag = VMax - VMin;
+
+					float InstanceRadius = 0.5f * VDiag.Size();
+					float DistanceToInstanceStart = DistanceToInstanceCenter - InstanceRadius;
+
+					// Cull instance based on distance
+					if (DistanceToInstanceStart > BVHCullRadius && ApplyGeneralCulling)
+						continue;
+
+					// Special culling for small scale objects
+					if (InstanceRadius < BVHLowScaleThreshold && ApplyLowScaleCulling)
+					{
+						if (DistanceToInstanceStart > BVHLowScaleRadius)
+							continue;
+					}
+
+					FRayTracingInstance RayTracingInstance;
+					RayTracingInstance.Geometry = &RenderData->LODResources[LOD].RayTracingGeometry;
+					RayTracingInstance.InstanceTransforms.Add(InstanceTransform);
+
+					int32 SectionCount = InstancedRenderData.LODModels[LOD].Sections.Num();
+
+					for (int32 SectionIdx = 0; SectionIdx < SectionCount; ++SectionIdx)
+					{
+						//#dxr_todo: so far we use the parent static mesh path to get material data
+						FMeshBatch MeshBatch;
+						FStaticMeshSceneProxy::GetMeshElement(LOD, 0, SectionIdx, 0, false, false, MeshBatch);
+
+						RayTracingInstance.Materials.Add(MeshBatch);
+					}
+
+					RayTracingInstance.BuildInstanceMaskAndFlags();
+					OutRayTracingInstances.Add(RayTracingInstance);
+				}
 			}
+		}
+	}
+	else
+	{
+		// No culling
+		for (int32 InstanceIdx = 0; InstanceIdx < InstanceCount; ++InstanceIdx)
+		{
+			if (InstancedRenderData.Component->PerInstanceSMData.IsValidIndex(InstanceIdx))
+			{
+				const FInstancedStaticMeshInstanceData& InstanceData = InstancedRenderData.Component->PerInstanceSMData[InstanceIdx];
+				FMatrix ComponentLocalToWorld = InstancedRenderData.Component->GetComponentTransform().ToMatrixWithScale();
+				FMatrix InstanceTransform = InstanceData.Transform * ComponentLocalToWorld;
+
+				FRayTracingInstance RayTracingInstance;
+				RayTracingInstance.Geometry = &RenderData->LODResources[LOD].RayTracingGeometry;
+				RayTracingInstance.InstanceTransforms.Add(InstanceTransform);
 
-			RayTracingInstance.BuildInstanceMaskAndFlags();
-			OutRayTracingInstances.Add(RayTracingInstance);
+				int32 SectionCount = InstancedRenderData.LODModels[LOD].Sections.Num();
+
+				for (int32 SectionIdx = 0; SectionIdx < SectionCount; ++SectionIdx)
+				{
+					//#dxr_todo: so far we use the parent static mesh path to get material data
+					FMeshBatch MeshBatch;
+					FStaticMeshSceneProxy::GetMeshElement(LOD, 0, SectionIdx, 0, false, false, MeshBatch);
+
+					RayTracingInstance.Materials.Add(MeshBatch);
+				}
+
+				RayTracingInstance.BuildInstanceMaskAndFlags();
+				OutRayTracingInstances.Add(RayTracingInstance);
+			}
 		}
 	}
 }
+
+void FInstancedStaticMeshSceneProxy::SetupRayTracingCullClusters()
+{
+	//#dxr_todo: select the appropriate LOD depending on Context.View
+	int32 LOD = 0;
+	if (RenderData->LODResources.Num() > LOD && RenderData->LODResources[LOD].RayTracingGeometry.IsInitialized())
+	{
+		const float MaxClusterRadiusMultiplier = CVarRayTracingInstancesCullClusterMaxRadiusMultiplier.GetValueOnAnyThread();
+		const int32 Batches = GetNumMeshBatches();
+		const int32 InstanceCount = InstancedRenderData.Component->PerInstanceSMData.Num();
+		int32 ClusterIndex = 0;
+		FMatrix ComponentLocalToWorld = InstancedRenderData.Component->GetComponentTransform().ToMatrixWithScale();
+		float MaxInstanceRadius = 0.0f;
+
+		// Init first cluster
+		RayTracingCullClusterInstances.Add(new TDoubleLinkedList< uint32>());
+		RayTracingCullClusterBoundsMin.Add(FVector(MAX_FLT, MAX_FLT, MAX_FLT));
+		RayTracingCullClusterBoundsMax.Add(FVector(-MAX_FLT, -MAX_FLT, -MAX_FLT));
+
+		// Traverse instances to find maximum rarius
+		for (int32 Instance = 0; Instance < InstanceCount; ++Instance)
+		{
+			if (InstancedRenderData.Component->PerInstanceSMData.IsValidIndex(Instance))
+			{
+				const FInstancedStaticMeshInstanceData& InstanceData = InstancedRenderData.Component->PerInstanceSMData[Instance];
+				FMatrix InstanceTransform = InstanceData.Transform * ComponentLocalToWorld;
+				FVector VMin, VMax;
+
+				InstancedRenderData.Component->GetLocalBounds(VMin, VMax);
+				VMin = InstanceTransform.TransformPosition(VMin);
+				VMax = InstanceTransform.TransformPosition(VMax);
+
+				FVector VBBoxSize = VMax - VMin;
+
+				MaxInstanceRadius = FMath::Max(0.5f * VBBoxSize.Size(), MaxInstanceRadius);
+			}
+		}
+
+		float MaxClusterRadius = MaxInstanceRadius * MaxClusterRadiusMultiplier;
+
+		// Build clusters
+		for (int32 Instance = 0; Instance < InstanceCount; ++Instance)
+		{
+			if (InstancedRenderData.Component->PerInstanceSMData.IsValidIndex(Instance))
+			{
+				const FInstancedStaticMeshInstanceData& InstanceData = InstancedRenderData.Component->PerInstanceSMData[Instance];
+				FMatrix InstanceTransform = InstanceData.Transform * ComponentLocalToWorld;
+				FVector InstanceLocation = InstanceTransform.TransformPosition({ 0.0f,0.0f,0.0f });
+				FVector VMin = InstanceLocation - FVector(MaxInstanceRadius, MaxInstanceRadius, MaxInstanceRadius);
+				FVector VMax = InstanceLocation + FVector(MaxInstanceRadius, MaxInstanceRadius, MaxInstanceRadius);
+				bool bClusterFound = false;
+
+				// Try to find suitable cluster
+				for (int32 CandidateCluster = 0; CandidateCluster <= ClusterIndex; ++CandidateCluster)
+				{
+					// Build new candidate cluster bounds
+					FVector VCandidateMin = VMin.ComponentMin(RayTracingCullClusterBoundsMin[CandidateCluster]);
+					FVector VCandidateMax = VMax.ComponentMax(RayTracingCullClusterBoundsMax[CandidateCluster]);
+
+					FVector VCandidateBBoxSize = VCandidateMax - VCandidateMin;
+					float MaxCandidateRadius = 0.5f * VCandidateBBoxSize.Size();
+
+					// If new candidate is still small enough, update current cluster
+					if (MaxCandidateRadius <= MaxClusterRadius)
+					{
+						RayTracingCullClusterInstances[CandidateCluster]->AddTail(Instance);
+						RayTracingCullClusterBoundsMin[CandidateCluster] = VCandidateMin;
+						RayTracingCullClusterBoundsMax[CandidateCluster] = VCandidateMax;
+						bClusterFound = true;
+						break;
+					}
+				}
+
+				// if we couldn't add the instance to an existing cluster create a new one
+				if (!bClusterFound)
+				{
+					++ClusterIndex;
+					RayTracingCullClusterInstances.Add(new TDoubleLinkedList< uint32>());
+					RayTracingCullClusterInstances[ClusterIndex]->AddTail(Instance);
+					RayTracingCullClusterBoundsMin.Add(VMin);
+					RayTracingCullClusterBoundsMax.Add(VMax);
+				}
+			}
+		}
+	}
+}
+
 #endif
 
 
diff --git a/Engine/Source/Runtime/Engine/Private/InstancedStaticMesh.h b/Engine/Source/Runtime/Engine/Private/InstancedStaticMesh.h
--- a/Engine/Source/Runtime/Engine/Private/InstancedStaticMesh.h
+++ b/Engine/Source/Runtime/Engine/Private/InstancedStaticMesh.h
@@ -544,11 +544,22 @@ public:
 	{
 		bVFRequiresPrimitiveUniformBuffer = true;
 		SetupProxy(InComponent);
+
+#if RHI_RAYTRACING
+		SetupRayTracingCullClusters();
+#endif
 	}
 
 	~FInstancedStaticMeshSceneProxy()
 	{
 		InstancedRenderData.ReleaseResources(&GetScene( ), StaticMesh);
+
+#if RHI_RAYTRACING
+		for (int32 i = 0; i < RayTracingCullClusterInstances.Num(); ++i)
+		{
+			delete RayTracingCullClusterInstances[i];
+		}
+#endif
 	}
 
 	// FPrimitiveSceneProxy interface.
@@ -577,6 +588,8 @@ public:
 	}
 
 	virtual void GetDynamicRayTracingInstances(struct FRayTracingMaterialGatheringContext& Context, TArray<FRayTracingInstance>& OutRayTracingInstances) final override;
+	void SetupRayTracingCullClusters();
+
 #endif
 
 	virtual void GetLightRelevance(const FLightSceneProxy* LightSceneProxy, bool& bDynamic, bool& bRelevant, bool& bLightMapped, bool& bShadowMapped) const override;
@@ -629,6 +642,12 @@ protected:
 	FInstancingUserData UserData_SelectedInstances;
 	FInstancingUserData UserData_DeselectedInstances;
 
+#if RHI_RAYTRACING
+	TArray< FVector >						RayTracingCullClusterBoundsMin;
+	TArray< FVector >						RayTracingCullClusterBoundsMax;
+	TArray< TDoubleLinkedList< uint32 >* >	RayTracingCullClusterInstances;
+#endif
+
 	/** Common path for the Get*MeshElement functions */
 	void SetupInstancedMeshBatch(int32 LODIndex, int32 BatchIndex, FMeshBatch& OutMeshBatch) const;
 
