VT lightmap improvements
- When VT lightmaps are enabled, don't generate non-VT lightmaps by default (controlled by cvar)
- Allows VT lightmaps to avoid size restrictions imposed on regular lightmaps, should reduce number of lightmap sheets generated, and potentially improve batching

diff --git a/Engine/Source/Runtime/Engine/Classes/Engine/Texture.h b/Engine/Source/Runtime/Engine/Classes/Engine/Texture.h
--- a/Engine/Source/Runtime/Engine/Classes/Engine/Texture.h
+++ b/Engine/Source/Runtime/Engine/Classes/Engine/Texture.h
@@ -884,6 +884,7 @@ public:
 	 * Gets settings used to choose format for the given layer
 	 */
 	ENGINE_API void GetLayerFormatSettings(int32 LayerIndex, FTextureFormatSettings& OutSettings) const;
+	ENGINE_API void SetLayerFormatSettings(int32 LayerIndex, const FTextureFormatSettings& InSettings);
 
 	ENGINE_API void GetDefaultFormatSettings(FTextureFormatSettings& OutSettings) const;
 #endif
diff --git a/Engine/Source/Runtime/Engine/Private/LightMap.cpp b/Engine/Source/Runtime/Engine/Private/LightMap.cpp
--- a/Engine/Source/Runtime/Engine/Private/LightMap.cpp
+++ b/Engine/Source/Runtime/Engine/Private/LightMap.cpp
@@ -91,6 +91,14 @@ static TAutoConsoleVariable<int32> CVarVirtualTexturedLightMaps(
 	TEXT(" 1: Enabled."),
 	ECVF_ReadOnly);
 
+static TAutoConsoleVariable<int32> CVarIncludeNonVirtualTexturedLightMaps(
+	TEXT("r.IncludeNonVirtualTexturedLightmaps"),
+	0,
+	TEXT("If 'r.VirtualTexturedLightmaps' is enabled, controls whether non-VT lightmaps are generated/saved as well.\n") \
+	TEXT("Including non-VT lightmaps will constrain lightmap atlas size, which removes some of the benefit of VT lightmaps.\n") \
+	TEXT(" 0: Not included.\n") \
+	TEXT(" 1: Included."));
+
 static TAutoConsoleVariable<int32> CVarVTEnableLossyCompressLightmaps(
 	TEXT("r.VT.EnableLossyCompressLightmaps"),
 	0,
@@ -393,7 +401,7 @@ struct FLightMapPendingTexture : public FTextureLayout
 	ELightMapFlags					LightmapFlags;
 	// Optimization to quickly test if a new allocation won't fit
 	// Primarily of benefit to instanced mesh lightmaps
-	int32							UnallocatedTexels;
+	int64							UnallocatedTexels;
 	int32							NumOutstandingAsyncTasks;
 	bool							bUObjectsCreated;
 	int32							NumNonPower2Texels;
@@ -404,8 +412,8 @@ struct FLightMapPendingTexture : public FTextureLayout
 	bool							bHasRunPostEncode;
 	bool							bTexelDebuggingEnabled;
 
-	FLightMapPendingTexture(UWorld* InWorld, uint32 InSizeX,uint32 InSizeY)
-		: FTextureLayout(4, 4, InSizeX, InSizeY, /* PowerOfTwo */ true, /* Force2To1Aspect */ true, /* AlignByFour */ true) // Min size is 4x4 in case of block compression.
+	FLightMapPendingTexture(UWorld* InWorld, uint32 InSizeX,uint32 InSizeY, ETextureLayoutAspectRatio Aspect)
+		: FTextureLayout(4, 4, InSizeX, InSizeY, /* PowerOfTwo */ true, Aspect, /* AlignByFour */ true) // Min size is 4x4 in case of block compression.
 		, SkyOcclusionTexture(nullptr)
 		, AOMaterialMaskTexture(nullptr)
 		, ShadowMapTexture(nullptr)
@@ -423,7 +431,7 @@ struct FLightMapPendingTexture : public FTextureLayout
 		, bIsFinishedEncoding(false)
 		, bHasRunPostEncode(false)
 		, bTexelDebuggingEnabled(IsTexelDebuggingEnabled())
-	{}
+	{ FMemory::Memzero(Textures); }
 
 	~FLightMapPendingTexture()
 	{
@@ -436,6 +444,11 @@ struct FLightMapPendingTexture : public FTextureLayout
 	 */
 	void StartEncoding(ULevel* Unused, ITextureCompressorModule* UnusedCompressor);
 
+	void EncodeCoefficientTexture(int32 CoefficientIndex, UTexture* Texture, uint32 LayerIndex, const FColor& TextureColor, bool bEncodeVirtualTexture);
+	void EncodeSkyOcclusionTexture(UTexture* Texture, uint32 LayerIndex, const FColor& TextureColor);
+	void EncodeAOMaskTexture(UTexture* Texture, uint32 LayerIndex, const FColor& TextureColor);
+	void EncodeShadowMapTexture(const TArray<TArray<FFourDistanceFieldSamples>>& MipData, UTexture* Texture, uint32 LayerIndex);
+
 	/**
 	 * Call this function after the IsFinishedEncoding function returns true
 	 */
@@ -1153,41 +1166,47 @@ void FLightMapPendingTexture::CreateUObjects()
 {
 	check(IsInGameThread());
 	++GLightmapCounter;
-	if (NeedsSkyOcclusionTexture())
-	{
-		SkyOcclusionTexture = NewObject<ULightMapTexture2D>(Outer, GetSkyOcclusionTextureName(GLightmapCounter));
-	}
-
-	if (NeedsAOMaterialMaskTexture())
-	{
-		AOMaterialMaskTexture = NewObject<ULightMapTexture2D>(Outer, GetAOMaterialMaskTextureName(GLightmapCounter));
-	}
-
-	if (NeedsStaticShadowTexture())
+	
+	// Only build VT lightmaps if they are enabled
+	const bool bUseVirtualTextures = (CVarVirtualTexturedLightMaps.GetValueOnAnyThread() != 0) && UseVirtualTexturing(GMaxRHIFeatureLevel);
+	const bool bIncludeNonVirtualTextures = !bUseVirtualTextures || (CVarIncludeNonVirtualTexturedLightMaps.GetValueOnAnyThread() != 0);
+	
+	if (bIncludeNonVirtualTextures)
 	{
-		ShadowMapTexture = NewObject<UShadowMapTexture2D>(Outer, GetShadowTextureName(GLightmapCounter));
-	}
+		if (NeedsSkyOcclusionTexture())
+		{
+			SkyOcclusionTexture = NewObject<ULightMapTexture2D>(Outer, GetSkyOcclusionTextureName(GLightmapCounter));
+		}
 
-	// Encode and compress the coefficient textures.
-	for (uint32 CoefficientIndex = 0; CoefficientIndex < NUM_STORED_LIGHTMAP_COEF; CoefficientIndex += 2)
-	{
-		Textures[CoefficientIndex] = nullptr;
-		// Skip generating simple lightmaps if wanted.
-		static const auto CVarSupportLowQualityLightmaps = IConsoleManager::Get().FindTConsoleVariableDataInt(TEXT("r.SupportLowQualityLightmaps"));
-		const bool bAllowLowQualityLightMaps = (!CVarSupportLowQualityLightmaps) || (CVarSupportLowQualityLightmaps->GetValueOnAnyThread() != 0);
+		if (NeedsAOMaterialMaskTexture())
+		{
+			AOMaterialMaskTexture = NewObject<ULightMapTexture2D>(Outer, GetAOMaterialMaskTextureName(GLightmapCounter));
+		}
 
-		if ((!bAllowLowQualityLightMaps) && CoefficientIndex >= LQ_LIGHTMAP_COEF_INDEX)
+		if (NeedsStaticShadowTexture())
 		{
-			continue;
+			ShadowMapTexture = NewObject<UShadowMapTexture2D>(Outer, GetShadowTextureName(GLightmapCounter));
 		}
 
-		// Create the light-map texture for this coefficient.
-		auto Texture = NewObject<ULightMapTexture2D>(Outer, GetLightmapName(GLightmapCounter, CoefficientIndex));
-		Textures[CoefficientIndex] = Texture;
+		// Encode and compress the coefficient textures.
+		for (uint32 CoefficientIndex = 0; CoefficientIndex < NUM_STORED_LIGHTMAP_COEF; CoefficientIndex += 2)
+		{
+			Textures[CoefficientIndex] = nullptr;
+			// Skip generating simple lightmaps if wanted.
+			static const auto CVarSupportLowQualityLightmaps = IConsoleManager::Get().FindTConsoleVariableDataInt(TEXT("r.SupportLowQualityLightmaps"));
+			const bool bAllowLowQualityLightMaps = (!CVarSupportLowQualityLightmaps) || (CVarSupportLowQualityLightmaps->GetValueOnAnyThread() != 0);
+
+			if ((!bAllowLowQualityLightMaps) && CoefficientIndex >= LQ_LIGHTMAP_COEF_INDEX)
+			{
+				continue;
+			}
+
+			// Create the light-map texture for this coefficient.
+			auto Texture = NewObject<ULightMapTexture2D>(Outer, GetLightmapName(GLightmapCounter, CoefficientIndex));
+			Textures[CoefficientIndex] = Texture;
+		}
 	}
 
-	// Only build VT lightmaps if they are enabled
-	const bool bUseVirtualTextures = (CVarVirtualTexturedLightMaps.GetValueOnAnyThread() != 0) && UseVirtualTexturing(GMaxRHIFeatureLevel);
 	if (bUseVirtualTextures)
 	{
 		VirtualTexture = NewObject<ULightMapVirtualTexture2D>(Outer, GetVirtualTextureName(GLightmapCounter));
@@ -1222,6 +1241,10 @@ bool FLightMapPendingTexture::NeedsSkyOcclusionTexture() const
 {
 	if (bUObjectsCreated)
 	{
+		if (VirtualTexture)
+		{
+			return VirtualTexture->HasLayerForType(ELightMapVirtualTextureType::SkyOcclusion);
+		}
 		return SkyOcclusionTexture != nullptr;
 	}
 
@@ -1244,6 +1267,10 @@ bool FLightMapPendingTexture::NeedsAOMaterialMaskTexture() const
 {
 	if (bUObjectsCreated)
 	{
+		if (VirtualTexture)
+		{
+			return VirtualTexture->HasLayerForType(ELightMapVirtualTextureType::AOMaterialMask);
+		}
 		return AOMaterialMaskTexture != nullptr;
 	}
 
@@ -1260,6 +1287,10 @@ bool FLightMapPendingTexture::NeedsStaticShadowTexture() const
 {
 	if (bUObjectsCreated)
 	{
+		if (VirtualTexture)
+		{
+			return VirtualTexture->HasLayerForType(ELightMapVirtualTextureType::ShadowMask);
+		}
 		return ShadowMapTexture != nullptr;
 	}
 
@@ -1277,410 +1308,471 @@ bool FLightMapPendingTexture::NeedsStaticShadowTexture() const
 	return bResult;
 }
 
-void FLightMapPendingTexture::StartEncoding(ULevel* LightingScenario, ITextureCompressorModule* UnusedCompressor)
+void FLightMapPendingTexture::EncodeSkyOcclusionTexture(UTexture* Texture, uint32 LayerIndex, const FColor& TextureColor)
 {
-	if (!bUObjectsCreated)
-	{
-		check(IsInGameThread());
-		CreateUObjects();
-	}
+	const int32 NumMips = Texture->Source.GetNumMips();
 
-	FColor TextureColor;
-	if ( GVisualizeLightmapTextures )
-	{
-		TextureColor = FColor::MakeRandomColor();
-	}
+	FTextureFormatSettings FormatSettings;
+	FormatSettings.SRGB = false;
+	FormatSettings.CompressionNoAlpha = false;
+	FormatSettings.CompressionNone = !GCompressLightmaps;
+	Texture->SetLayerFormatSettings(LayerIndex, FormatSettings);
 
-	if (SkyOcclusionTexture != nullptr)
-	{
-		auto Texture = SkyOcclusionTexture;
-		
-		Texture->Source.Init2DWithMipChain(GetSizeX(), GetSizeY(), TSF_BGRA8);
-		Texture->MipGenSettings = TMGS_LeaveExistingMips;
-		int32 NumMips = Texture->Source.GetNumMips();
-		Texture->SRGB = false;
-		Texture->Filter	= GUseBilinearLightmaps ? TF_Default : TF_Nearest;
-		Texture->LODGroup = TEXTUREGROUP_Lightmap;
-		Texture->LightmapFlags = ELightMapFlags( LightmapFlags );
-		Texture->CompressionNoAlpha = false;
-		Texture->CompressionNone = !GCompressLightmaps;
-
-		int32 TextureSizeX = Texture->Source.GetSizeX();
-		int32 TextureSizeY = Texture->Source.GetSizeY();
+	int32 TextureSizeX = Texture->Source.GetSizeX();
+	int32 TextureSizeY = Texture->Source.GetSizeY();
 
-		int32 StartBottom = GetSizeX() * GetSizeY();
+	// Lock all mip levels.
+	FColor* MipData[MAX_TEXTURE_MIP_COUNT] = { 0 };
+	int8* MipCoverageData[MAX_TEXTURE_MIP_COUNT] = { 0 };
+	for (int32 MipIndex = 0; MipIndex < NumMips; ++MipIndex)
+	{
+		MipData[MipIndex] = (FColor*)Texture->Source.LockMip(0, LayerIndex, MipIndex);
 
-		// Lock all mip levels.
-		FColor* MipData[MAX_TEXTURE_MIP_COUNT] = {0};
-		int8* MipCoverageData[MAX_TEXTURE_MIP_COUNT] = {0};
-		for (int32 MipIndex = 0; MipIndex < NumMips; ++MipIndex)
-		{
-			MipData[MipIndex] = (FColor*)Texture->Source.LockMip(MipIndex);
+		const int32 MipSizeX = FMath::Max(1, TextureSizeX >> MipIndex);
+		const int32 MipSizeY = FMath::Max(1, TextureSizeY >> MipIndex);
+		MipCoverageData[MipIndex] = (int8*)FMemory::Malloc(MipSizeX * MipSizeY);
+	}
 
-			const int32 MipSizeX = FMath::Max( 1, TextureSizeX >> MipIndex );
-			const int32 MipSizeY = FMath::Max( 1, TextureSizeY >> MipIndex );
-			MipCoverageData[ MipIndex ] = (int8*)FMemory::Malloc( MipSizeX * MipSizeY );
-		}
+	// Create the uncompressed top mip-level.
+	FColor* TopMipData = MipData[0];
+	FMemory::Memzero(TopMipData, TextureSizeX * TextureSizeY * sizeof(FColor));
+	FMemory::Memzero(MipCoverageData[0], TextureSizeX * TextureSizeY);
 
-		// Create the uncompressed top mip-level.
-		FColor* TopMipData = MipData[0];
-		FMemory::Memzero( TopMipData, TextureSizeX * TextureSizeY * sizeof(FColor) );
-		FMemory::Memzero( MipCoverageData[0], TextureSizeX * TextureSizeY );
+	FIntRect TextureRect(MAX_int32, MAX_int32, MIN_int32, MIN_int32);
+	for (int32 AllocationIndex = 0; AllocationIndex < Allocations.Num(); AllocationIndex++)
+	{
+		auto& Allocation = Allocations[AllocationIndex];
 
-		FIntRect TextureRect( MAX_int32, MAX_int32, MIN_int32, MIN_int32 );
-		for(int32 AllocationIndex = 0;AllocationIndex < Allocations.Num();AllocationIndex++)
+		// Skip encoding of this texture if we were asked not to bother
+		if (!Allocation->bSkipEncoding)
 		{
-			auto& Allocation = Allocations[AllocationIndex];
-			// Link the light-map to the texture.
-			Allocation->LightMap->SkyOcclusionTexture = Texture;
-			
-			// Skip encoding of this texture if we were asked not to bother
-			if( !Allocation->bSkipEncoding )
-			{
-				TextureRect.Min.X = FMath::Min<int32>( TextureRect.Min.X, Allocation->OffsetX );
-				TextureRect.Min.Y = FMath::Min<int32>( TextureRect.Min.Y, Allocation->OffsetY );
-				TextureRect.Max.X = FMath::Max<int32>( TextureRect.Max.X, Allocation->OffsetX + Allocation->MappedRect.Width() );
-				TextureRect.Max.Y = FMath::Max<int32>( TextureRect.Max.Y, Allocation->OffsetY + Allocation->MappedRect.Height() );
+			TextureRect.Min.X = FMath::Min<int32>(TextureRect.Min.X, Allocation->OffsetX);
+			TextureRect.Min.Y = FMath::Min<int32>(TextureRect.Min.Y, Allocation->OffsetY);
+			TextureRect.Max.X = FMath::Max<int32>(TextureRect.Max.X, Allocation->OffsetX + Allocation->MappedRect.Width());
+			TextureRect.Max.Y = FMath::Max<int32>(TextureRect.Max.Y, Allocation->OffsetY + Allocation->MappedRect.Height());
 
-				// Copy the raw data for this light-map into the raw texture data array.
-				for(int32 Y = Allocation->MappedRect.Min.Y; Y < Allocation->MappedRect.Max.Y; ++Y)
+			// Copy the raw data for this light-map into the raw texture data array.
+			for (int32 Y = Allocation->MappedRect.Min.Y; Y < Allocation->MappedRect.Max.Y; ++Y)
+			{
+				for (int32 X = Allocation->MappedRect.Min.X; X < Allocation->MappedRect.Max.X; ++X)
 				{
-					for(int32 X = Allocation->MappedRect.Min.X; X < Allocation->MappedRect.Max.X; ++X)
-					{
-						const FLightMapCoefficients& SourceCoefficients = Allocation->RawData[Y * Allocation->TotalSizeX + X];
+					const FLightMapCoefficients& SourceCoefficients = Allocation->RawData[Y * Allocation->TotalSizeX + X];
 
-						int32 DestY = Y - Allocation->MappedRect.Min.Y + Allocation->OffsetY;
-						int32 DestX = X - Allocation->MappedRect.Min.X + Allocation->OffsetX;
+					int32 DestY = Y - Allocation->MappedRect.Min.Y + Allocation->OffsetY;
+					int32 DestX = X - Allocation->MappedRect.Min.X + Allocation->OffsetX;
 
-						FColor&	DestColor = TopMipData[DestY * TextureSizeX + DestX];
+					FColor&	DestColor = TopMipData[DestY * TextureSizeX + DestX];
 
-						DestColor.R = SourceCoefficients.SkyOcclusion[0];
-						DestColor.G = SourceCoefficients.SkyOcclusion[1];
-						DestColor.B = SourceCoefficients.SkyOcclusion[2];
-						DestColor.A = SourceCoefficients.SkyOcclusion[3];
+					DestColor.R = SourceCoefficients.SkyOcclusion[0];
+					DestColor.G = SourceCoefficients.SkyOcclusion[1];
+					DestColor.B = SourceCoefficients.SkyOcclusion[2];
+					DestColor.A = SourceCoefficients.SkyOcclusion[3];
 
-						int8& DestCoverage = MipCoverageData[0][ DestY * TextureSizeX + DestX ];
-						DestCoverage = SourceCoefficients.Coverage / 2;
-					}
+					int8& DestCoverage = MipCoverageData[0][DestY * TextureSizeX + DestX];
+					DestCoverage = SourceCoefficients.Coverage / 2;
 				}
 			}
 		}
+	}
 
-		GenerateLightmapMipsAndDilateColor(NumMips, TextureSizeX, TextureSizeY, TextureColor, MipData, MipCoverageData);
+	GenerateLightmapMipsAndDilateColor(NumMips, TextureSizeX, TextureSizeY, TextureColor, MipData, MipCoverageData);
 
-		// Unlock all mip levels.
-		for (int32 MipIndex = 0; MipIndex < NumMips; ++MipIndex)
-		{
-			Texture->Source.UnlockMip(MipIndex);
-			FMemory::Free( MipCoverageData[ MipIndex ] );
-		}
-
-	}
-	
-	if (AOMaterialMaskTexture != nullptr)
+	// Unlock all mip levels.
+	for (int32 MipIndex = 0; MipIndex < NumMips; ++MipIndex)
 	{
-		auto Texture = AOMaterialMaskTexture;
-		
-		Texture->Source.Init2DWithMipChain(GetSizeX(), GetSizeY(), TSF_G8);
-		Texture->MipGenSettings = TMGS_LeaveExistingMips;
-		int32 NumMips = Texture->Source.GetNumMips();
-		Texture->SRGB = false;
-		Texture->Filter	= GUseBilinearLightmaps ? TF_Default : TF_Nearest;
-		Texture->LODGroup = TEXTUREGROUP_Lightmap;
-		Texture->LightmapFlags = ELightMapFlags( LightmapFlags );
-		Texture->CompressionNoAlpha = false;
-		Texture->CompressionNone = !GCompressLightmaps;
-		// BC4
-		Texture->CompressionSettings = TC_Alpha;
+		Texture->Source.UnlockMip(0, LayerIndex, MipIndex);
+		FMemory::Free(MipCoverageData[MipIndex]);
+	}
+}
 
-		int32 TextureSizeX = Texture->Source.GetSizeX();
-		int32 TextureSizeY = Texture->Source.GetSizeY();
+void FLightMapPendingTexture::EncodeAOMaskTexture(UTexture* Texture, uint32 LayerIndex, const FColor& TextureColor)
+{
+	const int32 NumMips = Texture->Source.GetNumMips();
 
-		int32 StartBottom = GetSizeX() * GetSizeY();
+	FTextureFormatSettings FormatSettings;
+	FormatSettings.SRGB = false;
+	FormatSettings.CompressionNoAlpha = false;
+	FormatSettings.CompressionNone = !GCompressLightmaps;
+	FormatSettings.CompressionSettings = TC_Alpha; // BC4
+	Texture->SetLayerFormatSettings(LayerIndex, FormatSettings);
 
-		// Lock all mip levels.
-		uint8* MipData[MAX_TEXTURE_MIP_COUNT] = {0};
-		int8* MipCoverageData[MAX_TEXTURE_MIP_COUNT] = {0};
-		for (int32 MipIndex = 0; MipIndex < NumMips; ++MipIndex)
-		{
-			MipData[MipIndex] = (uint8*)Texture->Source.LockMip(MipIndex);
+	int32 TextureSizeX = Texture->Source.GetSizeX();
+	int32 TextureSizeY = Texture->Source.GetSizeY();
 
-			const int32 MipSizeX = FMath::Max( 1, TextureSizeX >> MipIndex );
-			const int32 MipSizeY = FMath::Max( 1, TextureSizeY >> MipIndex );
-			MipCoverageData[ MipIndex ] = (int8*)FMemory::Malloc( MipSizeX * MipSizeY );
-		}
+	// Lock all mip levels.
+	uint8* MipData[MAX_TEXTURE_MIP_COUNT] = { 0 };
+	int8* MipCoverageData[MAX_TEXTURE_MIP_COUNT] = { 0 };
+	for (int32 MipIndex = 0; MipIndex < NumMips; ++MipIndex)
+	{
+		MipData[MipIndex] = (uint8*)Texture->Source.LockMip(0, LayerIndex, MipIndex);
 
-		// Create the uncompressed top mip-level.
-		uint8* TopMipData = MipData[0];
-		FMemory::Memzero( TopMipData, TextureSizeX * TextureSizeY * sizeof(uint8) );
-		FMemory::Memzero( MipCoverageData[0], TextureSizeX * TextureSizeY );
+		const int32 MipSizeX = FMath::Max(1, TextureSizeX >> MipIndex);
+		const int32 MipSizeY = FMath::Max(1, TextureSizeY >> MipIndex);
+		MipCoverageData[MipIndex] = (int8*)FMemory::Malloc(MipSizeX * MipSizeY);
+	}
 
-		FIntRect TextureRect( MAX_int32, MAX_int32, MIN_int32, MIN_int32 );
-		for(int32 AllocationIndex = 0;AllocationIndex < Allocations.Num();AllocationIndex++)
+	// Create the uncompressed top mip-level.
+	uint8* TopMipData = MipData[0];
+	FMemory::Memzero(TopMipData, TextureSizeX * TextureSizeY * sizeof(uint8));
+	FMemory::Memzero(MipCoverageData[0], TextureSizeX * TextureSizeY);
+
+	FIntRect TextureRect(MAX_int32, MAX_int32, MIN_int32, MIN_int32);
+	for (int32 AllocationIndex = 0; AllocationIndex < Allocations.Num(); AllocationIndex++)
+	{
+		auto& Allocation = Allocations[AllocationIndex];
+
+		// Skip encoding of this texture if we were asked not to bother
+		if (!Allocation->bSkipEncoding)
 		{
-			auto& Allocation = Allocations[AllocationIndex];
-			// Link the light-map to the texture.
-			Allocation->LightMap->AOMaterialMaskTexture = Texture;
-			
-			// Skip encoding of this texture if we were asked not to bother
-			if( !Allocation->bSkipEncoding )
-			{
-				TextureRect.Min.X = FMath::Min<int32>( TextureRect.Min.X, Allocation->OffsetX );
-				TextureRect.Min.Y = FMath::Min<int32>( TextureRect.Min.Y, Allocation->OffsetY );
-				TextureRect.Max.X = FMath::Max<int32>( TextureRect.Max.X, Allocation->OffsetX + Allocation->MappedRect.Width() );
-				TextureRect.Max.Y = FMath::Max<int32>( TextureRect.Max.Y, Allocation->OffsetY + Allocation->MappedRect.Height() );
+			TextureRect.Min.X = FMath::Min<int32>(TextureRect.Min.X, Allocation->OffsetX);
+			TextureRect.Min.Y = FMath::Min<int32>(TextureRect.Min.Y, Allocation->OffsetY);
+			TextureRect.Max.X = FMath::Max<int32>(TextureRect.Max.X, Allocation->OffsetX + Allocation->MappedRect.Width());
+			TextureRect.Max.Y = FMath::Max<int32>(TextureRect.Max.Y, Allocation->OffsetY + Allocation->MappedRect.Height());
 
-				// Copy the raw data for this light-map into the raw texture data array.
-				for(int32 Y = Allocation->MappedRect.Min.Y; Y < Allocation->MappedRect.Max.Y; ++Y)
+			// Copy the raw data for this light-map into the raw texture data array.
+			for (int32 Y = Allocation->MappedRect.Min.Y; Y < Allocation->MappedRect.Max.Y; ++Y)
+			{
+				for (int32 X = Allocation->MappedRect.Min.X; X < Allocation->MappedRect.Max.X; ++X)
 				{
-					for(int32 X = Allocation->MappedRect.Min.X; X < Allocation->MappedRect.Max.X; ++X)
-					{
-						const FLightMapCoefficients& SourceCoefficients = Allocation->RawData[Y * Allocation->TotalSizeX + X];
+					const FLightMapCoefficients& SourceCoefficients = Allocation->RawData[Y * Allocation->TotalSizeX + X];
 
-						int32 DestY = Y - Allocation->MappedRect.Min.Y + Allocation->OffsetY;
-						int32 DestX = X - Allocation->MappedRect.Min.X + Allocation->OffsetX;
+					int32 DestY = Y - Allocation->MappedRect.Min.Y + Allocation->OffsetY;
+					int32 DestX = X - Allocation->MappedRect.Min.X + Allocation->OffsetX;
 
-						uint8& DestValue = TopMipData[DestY * TextureSizeX + DestX];
+					uint8& DestValue = TopMipData[DestY * TextureSizeX + DestX];
 
-						DestValue = SourceCoefficients.AOMaterialMask;
+					DestValue = SourceCoefficients.AOMaterialMask;
 
-						int8& DestCoverage = MipCoverageData[0][ DestY * TextureSizeX + DestX ];
-						DestCoverage = SourceCoefficients.Coverage / 2;
-					}
+					int8& DestCoverage = MipCoverageData[0][DestY * TextureSizeX + DestX];
+					DestCoverage = SourceCoefficients.Coverage / 2;
 				}
 			}
 		}
+	}
 
-		GenerateLightmapMipsAndDilateByte(NumMips, TextureSizeX, TextureSizeY, TextureColor.R, MipData, MipCoverageData);
-
-		// Unlock all mip levels.
-		for (int32 MipIndex = 0; MipIndex < NumMips; ++MipIndex)
-		{
-			Texture->Source.UnlockMip(MipIndex);
-			FMemory::Free( MipCoverageData[ MipIndex ] );
-		}
+	GenerateLightmapMipsAndDilateByte(NumMips, TextureSizeX, TextureSizeY, TextureColor.R, MipData, MipCoverageData);
 
+	// Unlock all mip levels.
+	for (int32 MipIndex = 0; MipIndex < NumMips; ++MipIndex)
+	{
+		Texture->Source.UnlockMip(0, LayerIndex, MipIndex);
+		FMemory::Free(MipCoverageData[MipIndex]);
 	}
+}
 
-	if (ShadowMapTexture != nullptr)
-	{
-		ShadowMapTexture->Filter = GUseBilinearLightmaps ? TF_Default : TF_Nearest;
-		// Signed distance field textures get stored in linear space, since they need more precision near .5.
-		ShadowMapTexture->SRGB = false;
-		ShadowMapTexture->LODGroup = TEXTUREGROUP_Shadowmap;
+void FLightMapPendingTexture::EncodeShadowMapTexture(const TArray<TArray<FFourDistanceFieldSamples>>& MipData, UTexture* Texture, uint32 LayerIndex)
+{
+	const int32 NumMips = Texture->Source.GetNumMips();
 
-		// TODO - Unify these?  Currently the values match
-		ShadowMapTexture->ShadowmapFlags = EShadowMapFlags(LightmapFlags);
+	FTextureFormatSettings FormatSettings;
+	FormatSettings.SRGB = false;
+	FormatSettings.CompressionNone = true;
+	Texture->SetLayerFormatSettings(LayerIndex, FormatSettings);
 
-		// Create the uncompressed top mip-level.
-		TArray< TArray<FFourDistanceFieldSamples> > MipData;
-		const int32 NumChannelsUsed = FLightMap2D::EncodeShadowTexture(LightingScenario, *this, ShadowMapTexture, MipData);
+	const ETextureSourceFormat SourceFormat = Texture->Source.GetFormat(LayerIndex);
+	check(SourceFormat == TSF_G8 || SourceFormat == TSF_BGRA8);
 
-		ShadowMapTexture->Source.Init2DWithMipChain(GetSizeX(), GetSizeY(), NumChannelsUsed == 1 ? TSF_G8 : TSF_BGRA8);
-		ShadowMapTexture->MipGenSettings = TMGS_LeaveExistingMips;
-		ShadowMapTexture->CompressionNone = true;
+	const int32 TextureSizeX = Texture->Source.GetSizeX();
+	const int32 TextureSizeY = Texture->Source.GetSizeY();
 
-		// Copy the mip-map data into the UShadowMapTexture2D's mip-map array.
-		for (int32 MipIndex = 0; MipIndex < MipData.Num(); MipIndex++)
-		{
-			uint8* DestMipData = (uint8*)ShadowMapTexture->Source.LockMip(MipIndex);
-			uint32 MipSizeX = FMath::Max<uint32>(1, GetSizeX() >> MipIndex);
-			uint32 MipSizeY = FMath::Max<uint32>(1, GetSizeY() >> MipIndex);
+	// Copy the mip-map data into the UShadowMapTexture2D's mip-map array.
+	for (int32 MipIndex = 0; MipIndex < MipData.Num(); MipIndex++)
+	{
+		uint8* DestMipData = (uint8*)Texture->Source.LockMip(0, LayerIndex, MipIndex);
+		uint32 MipSizeX = FMath::Max<uint32>(1, GetSizeX() >> MipIndex);
+		uint32 MipSizeY = FMath::Max<uint32>(1, GetSizeY() >> MipIndex);
 
-			for (uint32 Y = 0; Y < MipSizeY; Y++)
+		for (uint32 Y = 0; Y < MipSizeY; Y++)
+		{
+			for (uint32 X = 0; X < MipSizeX; X++)
 			{
-				for (uint32 X = 0; X < MipSizeX; X++)
-				{
-					const FFourDistanceFieldSamples& SourceSample = MipData[MipIndex][Y * MipSizeX + X];
+				const FFourDistanceFieldSamples& SourceSample = MipData[MipIndex][Y * MipSizeX + X];
 
-					if (NumChannelsUsed == 1)
-					{
-						DestMipData[Y * MipSizeX + X] = SourceSample.Samples[0].Distance;
-					}
-					else
-					{
-						((FColor*)DestMipData)[Y * MipSizeX + X] = FColor(SourceSample.Samples[0].Distance, SourceSample.Samples[1].Distance, SourceSample.Samples[2].Distance, SourceSample.Samples[3].Distance);
-					}
+				if (SourceFormat == TSF_G8)
+				{
+					DestMipData[Y * MipSizeX + X] = SourceSample.Samples[0].Distance;
+				}
+				else
+				{
+					((FColor*)DestMipData)[Y * MipSizeX + X] = FColor(SourceSample.Samples[0].Distance, SourceSample.Samples[1].Distance, SourceSample.Samples[2].Distance, SourceSample.Samples[3].Distance);
 				}
 			}
-
-			ShadowMapTexture->Source.UnlockMip(MipIndex);
 		}
+
+		Texture->Source.UnlockMip(0, LayerIndex, MipIndex);
 	}
+}
 
-	// Encode and compress the coefficient textures.
-	for(uint32 CoefficientIndex = 0; CoefficientIndex < NUM_STORED_LIGHTMAP_COEF; CoefficientIndex += 2)
-	{
-		auto Texture = Textures[CoefficientIndex];
-		if (Texture == nullptr)
+void FLightMapPendingTexture::EncodeCoefficientTexture(int32 CoefficientIndex, UTexture* Texture, uint32 LayerIndex, const FColor& TextureColor, bool bEncodeVirtualTexture)
+{
+	const int32 NumMips = Texture->Source.GetNumMips();
+	check(NumMips > 0);
+
+	FTextureFormatSettings FormatSettings;
+	FormatSettings.SRGB = false;
+	FormatSettings.CompressionNoAlpha = CoefficientIndex >= LQ_LIGHTMAP_COEF_INDEX;
+	FormatSettings.CompressionNone = !GCompressLightmaps;
+	//FormatSettings.bForcePVRTC4 = true;
+	Texture->SetLayerFormatSettings(LayerIndex, FormatSettings);
+
+	const int32 TextureSizeX = Texture->Source.GetSizeX();
+	const int32 TextureSizeY = Texture->Source.GetSizeY();
+
+	// Lock all mip levels.
+	FColor* MipData0[MAX_TEXTURE_MIP_COUNT] = { 0 };
+	FColor* MipData1[MAX_TEXTURE_MIP_COUNT] = { 0 };
+	int8* MipCoverageData0[MAX_TEXTURE_MIP_COUNT] = { 0 };
+	int8* MipCoverageData1[MAX_TEXTURE_MIP_COUNT] = { 0 };
+	{
+		const int32 StartBottom = GetSizeX() * GetSizeY();
+		for (int32 MipIndex = 0; MipIndex < NumMips; ++MipIndex)
 		{
-			continue;
+			const int32 MipSizeX = FMath::Max(1, TextureSizeX >> MipIndex);
+			const int32 MipSizeY = FMath::Max(1, TextureSizeY >> MipIndex);
+
+			MipData0[MipIndex] = (FColor*)Texture->Source.LockMip(0, LayerIndex, MipIndex);
+			MipCoverageData0[MipIndex] = (int8*)FMemory::Malloc(MipSizeX * MipSizeY);
+
+			if (bEncodeVirtualTexture)
+			{
+				// 2 coefficients are stored on adjacent VT layers
+				MipData1[MipIndex] = (FColor*)Texture->Source.LockMip(0, LayerIndex + 1, MipIndex);
+				MipCoverageData1[MipIndex] = (int8*)FMemory::Malloc(MipSizeX * MipSizeY);
+			}
+			else
+			{
+				// 2 coefficients are stored on the top/bottom halves of the same destination texture
+				MipData1[MipIndex] = MipData0[MipIndex] + StartBottom;
+				MipCoverageData1[MipIndex] = MipCoverageData0[MipIndex] + StartBottom;
+			}
 		}
+	}
 
-		Texture->Source.Init2DWithMipChain(GetSizeX(), GetSizeY() * 2, TSF_BGRA8);	// Top/bottom atlased
-		Texture->MipGenSettings = TMGS_LeaveExistingMips;
-		int32 NumMips = Texture->Source.GetNumMips();
-		check(NumMips > 0);
-		Texture->SRGB = 0;
-		Texture->Filter	= GUseBilinearLightmaps ? TF_Default : TF_Nearest;
-		Texture->LODGroup = TEXTUREGROUP_Lightmap;
-		Texture->LightmapFlags = ELightMapFlags( LightmapFlags );
-		Texture->CompressionNoAlpha = CoefficientIndex >= LQ_LIGHTMAP_COEF_INDEX;
-		Texture->CompressionNone = !GCompressLightmaps;
-		Texture->bForcePVRTC4 = true;
+	// Create the uncompressed top mip-level.
+	FMemory::Memzero(MipData0[0], TextureSizeX * TextureSizeY * sizeof(FColor));
+	FMemory::Memzero(MipCoverageData0[0], TextureSizeX * TextureSizeY);
+	if (bEncodeVirtualTexture)
+	{
+		FMemory::Memzero(MipData1[0], TextureSizeX * TextureSizeY * sizeof(FColor));
+		FMemory::Memzero(MipCoverageData1[0], TextureSizeX * TextureSizeY);
+	}
 
-		int32 TextureSizeX = Texture->Source.GetSizeX();
-		int32 TextureSizeY = Texture->Source.GetSizeY();
-		
-		int32 StartBottom = GetSizeX() * GetSizeY();
-		
-		// Lock all mip levels.
-		FColor* MipData[MAX_TEXTURE_MIP_COUNT] = {0};
-		int8* MipCoverageData[MAX_TEXTURE_MIP_COUNT] = {0};
-		for (int32 MipIndex = 0; MipIndex < NumMips; ++MipIndex)
+	for (int32 AllocationIndex = 0; AllocationIndex < Allocations.Num(); AllocationIndex++)
+	{
+		auto& Allocation = Allocations[AllocationIndex];
+		for (int k = 0; k < 2; k++)
 		{
-			MipData[MipIndex] = (FColor*)Texture->Source.LockMip(MipIndex);
+			Allocation->LightMap->ScaleVectors[CoefficientIndex + k] = FVector4(
+				Allocation->Scale[CoefficientIndex + k][0],
+				Allocation->Scale[CoefficientIndex + k][1],
+				Allocation->Scale[CoefficientIndex + k][2],
+				Allocation->Scale[CoefficientIndex + k][3]
+			);
 
-			const int32 MipSizeX = FMath::Max( 1, TextureSizeX >> MipIndex );
-			const int32 MipSizeY = FMath::Max( 1, TextureSizeY >> MipIndex );
-			MipCoverageData[ MipIndex ] = (int8*)FMemory::Malloc( MipSizeX * MipSizeY );
+			Allocation->LightMap->AddVectors[CoefficientIndex + k] = FVector4(
+				Allocation->Add[CoefficientIndex + k][0],
+				Allocation->Add[CoefficientIndex + k][1],
+				Allocation->Add[CoefficientIndex + k][2],
+				Allocation->Add[CoefficientIndex + k][3]
+			);
 		}
 
-		// Create the uncompressed top mip-level.
-		FColor* TopMipData = MipData[0];
-		FMemory::Memzero( TopMipData, TextureSizeX * TextureSizeY * sizeof(FColor) );
-		FMemory::Memzero( MipCoverageData[0], TextureSizeX * TextureSizeY );
-
-		
-		for(int32 AllocationIndex = 0;AllocationIndex < Allocations.Num();AllocationIndex++)
+		// Skip encoding of this texture if we were asked not to bother
+		if (!Allocation->bSkipEncoding)
 		{
-			auto& Allocation = Allocations[AllocationIndex];
-			// Link the light-map to the texture.
-			Allocation->LightMap->Textures[ CoefficientIndex / 2 ] = Texture;
-			for( int k = 0; k < 2; k++ )
-			{
-				Allocation->LightMap->ScaleVectors[ CoefficientIndex + k ] = FVector4(
-					Allocation->Scale[ CoefficientIndex + k ][0],
-					Allocation->Scale[ CoefficientIndex + k ][1],
-					Allocation->Scale[ CoefficientIndex + k ][2],
-					Allocation->Scale[ CoefficientIndex + k ][3]
-					);
-
-				Allocation->LightMap->AddVectors[ CoefficientIndex + k ] = FVector4(
-					Allocation->Add[ CoefficientIndex + k ][0],
-					Allocation->Add[ CoefficientIndex + k ][1],
-					Allocation->Add[ CoefficientIndex + k ][2],
-					Allocation->Add[ CoefficientIndex + k ][3]
-					);
-			}
+			FIntRect TextureRect(MAX_int32, MAX_int32, MIN_int32, MIN_int32);
+			TextureRect.Min.X = FMath::Min<int32>(TextureRect.Min.X, Allocation->OffsetX);
+			TextureRect.Min.Y = FMath::Min<int32>(TextureRect.Min.Y, Allocation->OffsetY);
+			TextureRect.Max.X = FMath::Max<int32>(TextureRect.Max.X, Allocation->OffsetX + Allocation->MappedRect.Width());
+			TextureRect.Max.Y = FMath::Max<int32>(TextureRect.Max.Y, Allocation->OffsetY + Allocation->MappedRect.Height());
 
-			// Skip encoding of this texture if we were asked not to bother
-			if( !Allocation->bSkipEncoding )
-			{
-				FIntRect TextureRect(MAX_int32, MAX_int32, MIN_int32, MIN_int32);
-				TextureRect.Min.X = FMath::Min<int32>( TextureRect.Min.X, Allocation->OffsetX );
-				TextureRect.Min.Y = FMath::Min<int32>( TextureRect.Min.Y, Allocation->OffsetY );
-				TextureRect.Max.X = FMath::Max<int32>( TextureRect.Max.X, Allocation->OffsetX + Allocation->MappedRect.Width() );
-				TextureRect.Max.Y = FMath::Max<int32>( TextureRect.Max.Y, Allocation->OffsetY + Allocation->MappedRect.Height() );
-	
-				NumNonPower2Texels += TextureRect.Width() * TextureRect.Height();
+			NumNonPower2Texels += TextureRect.Width() * TextureRect.Height();
 
-
-				// Copy the raw data for this light-map into the raw texture data array.
-				for(int32 Y = Allocation->MappedRect.Min.Y; Y < Allocation->MappedRect.Max.Y; ++Y)
+			// Copy the raw data for this light-map into the raw texture data array.
+			for (int32 Y = Allocation->MappedRect.Min.Y; Y < Allocation->MappedRect.Max.Y; ++Y)
+			{
+				for (int32 X = Allocation->MappedRect.Min.X; X < Allocation->MappedRect.Max.X; ++X)
 				{
-					for(int32 X = Allocation->MappedRect.Min.X; X < Allocation->MappedRect.Max.X; ++X)
-					{
-						const FLightMapCoefficients& SourceCoefficients = Allocation->RawData[Y * Allocation->TotalSizeX + X];
-						
-						int32 DestY = Y - Allocation->MappedRect.Min.Y + Allocation->OffsetY;
-						int32 DestX = X - Allocation->MappedRect.Min.X + Allocation->OffsetX;
+					const FLightMapCoefficients& SourceCoefficients = Allocation->RawData[Y * Allocation->TotalSizeX + X];
 
-						FColor&	DestColor = TopMipData[DestY * TextureSizeX + DestX];
-						int8&	DestCoverage = MipCoverageData[0][ DestY * TextureSizeX + DestX ];
+					int32 DestY = Y - Allocation->MappedRect.Min.Y + Allocation->OffsetY;
+					int32 DestX = X - Allocation->MappedRect.Min.X + Allocation->OffsetX;
 
-						FColor&	DestBottomColor = TopMipData[ StartBottom + DestX + DestY * TextureSizeX ];
-						int8&	DestBottomCoverage = MipCoverageData[0][ StartBottom + DestX + DestY * TextureSizeX ];
-						
-#if VISUALIZE_PACKING
-						if( X == Allocation->MappedRect.Min.X || Y == Allocation->MappedRect.Min.Y ||
-							X == Allocation->MappedRect.Max.X-1 || Y == Allocation->MappedRect.Max.Y-1 ||
-							X == Allocation->MappedRect.Min.X+1 || Y == Allocation->MappedRect.Min.Y+1 ||
-							X == Allocation->MappedRect.Max.X-2 || Y == Allocation->MappedRect.Max.Y-2 )
-						{
-							DestColor = FColor::Red;
-						}
-						else
-						{
-							DestColor = FColor::Green;
-						}
-#else
-						DestColor.R = SourceCoefficients.Coefficients[CoefficientIndex][0];
-						DestColor.G = SourceCoefficients.Coefficients[CoefficientIndex][1];
-						DestColor.B = SourceCoefficients.Coefficients[CoefficientIndex][2];
-						DestColor.A = SourceCoefficients.Coefficients[CoefficientIndex][3];
+					FColor&	DestColor = MipData0[0][DestY * TextureSizeX + DestX];
+					int8&	DestCoverage = MipCoverageData0[0][DestY * TextureSizeX + DestX];
 
-						DestBottomColor.R = SourceCoefficients.Coefficients[ CoefficientIndex + 1 ][0];
-						DestBottomColor.G = SourceCoefficients.Coefficients[ CoefficientIndex + 1 ][1];
-						DestBottomColor.B = SourceCoefficients.Coefficients[ CoefficientIndex + 1 ][2];
-						DestBottomColor.A = SourceCoefficients.Coefficients[ CoefficientIndex + 1 ][3];
+					FColor&	DestBottomColor = MipData1[0][DestX + DestY * TextureSizeX];
+					int8&	DestBottomCoverage = MipCoverageData1[0][DestX + DestY * TextureSizeX];
 
-						if ( GVisualizeLightmapTextures )
-						{
-							DestColor = TextureColor;
-						}
+#if VISUALIZE_PACKING
+					if (X == Allocation->MappedRect.Min.X || Y == Allocation->MappedRect.Min.Y ||
+						X == Allocation->MappedRect.Max.X - 1 || Y == Allocation->MappedRect.Max.Y - 1 ||
+						X == Allocation->MappedRect.Min.X + 1 || Y == Allocation->MappedRect.Min.Y + 1 ||
+						X == Allocation->MappedRect.Max.X - 2 || Y == Allocation->MappedRect.Max.Y - 2)
+					{
+						DestColor = FColor::Red;
+					}
+					else
+					{
+						DestColor = FColor::Green;
+					}
+#else // VISUALIZE_PACKING
+					DestColor.R = SourceCoefficients.Coefficients[CoefficientIndex][0];
+					DestColor.G = SourceCoefficients.Coefficients[CoefficientIndex][1];
+					DestColor.B = SourceCoefficients.Coefficients[CoefficientIndex][2];
+					DestColor.A = SourceCoefficients.Coefficients[CoefficientIndex][3];
+
+					DestBottomColor.R = SourceCoefficients.Coefficients[CoefficientIndex + 1][0];
+					DestBottomColor.G = SourceCoefficients.Coefficients[CoefficientIndex + 1][1];
+					DestBottomColor.B = SourceCoefficients.Coefficients[CoefficientIndex + 1][2];
+					DestBottomColor.A = SourceCoefficients.Coefficients[CoefficientIndex + 1][3];
+
+					if (GVisualizeLightmapTextures)
+					{
+						DestColor = TextureColor;
+					}
 
-						// uint8 -> int8
-						DestCoverage = DestBottomCoverage = SourceCoefficients.Coverage / 2;
-						if ( SourceCoefficients.Coverage > 0 )
-						{
-							NumLightmapMappedTexels++;
-						}
-						else
-						{
-							NumLightmapUnmappedTexels++;
-						}
+					// uint8 -> int8
+					DestCoverage = DestBottomCoverage = SourceCoefficients.Coverage / 2;
+					if (SourceCoefficients.Coverage > 0)
+					{
+						NumLightmapMappedTexels++;
+					}
+					else
+					{
+						NumLightmapUnmappedTexels++;
+					}
 
 #if WITH_EDITOR
-						if (bTexelDebuggingEnabled)
+					if (bTexelDebuggingEnabled)
+					{
+						int32 PaddedX = X;
+						int32 PaddedY = Y;
+						if (GLightmassDebugOptions.bPadMappings && (Allocation->PaddingType == LMPT_NormalPadding))
 						{
-							int32 PaddedX = X;
-							int32 PaddedY = Y;
-							if (GLightmassDebugOptions.bPadMappings && (Allocation->PaddingType == LMPT_NormalPadding))
+							if (Allocation->TotalSizeX - 2 > 0 && Allocation->TotalSizeY - 2 > 0)
 							{
-								if (Allocation->TotalSizeX - 2 > 0 && Allocation->TotalSizeY - 2 > 0)
-								{
-									PaddedX -= 1;
-									PaddedY -= 1;
-								}
+								PaddedX -= 1;
+								PaddedY -= 1;
 							}
+						}
 
-							if (Allocation->bDebug
-								&& PaddedX == GCurrentSelectedLightmapSample.LocalX
-								&& PaddedY == GCurrentSelectedLightmapSample.LocalY)
-							{
-								extern FColor GTexelSelectionColor;
-								DestColor = GTexelSelectionColor;
-							}
+						if (Allocation->bDebug
+							&& PaddedX == GCurrentSelectedLightmapSample.LocalX
+							&& PaddedY == GCurrentSelectedLightmapSample.LocalY)
+						{
+							extern FColor GTexelSelectionColor;
+							DestColor = GTexelSelectionColor;
 						}
-#endif
-#endif
 					}
+#endif // WITH_EDITOR
+#endif // !VISUALIZE_PACKING
 				}
+			}
 
 
-			}
 		}
+	}
 
+	GenerateLightmapMipsAndDilateColor(NumMips, TextureSizeX, TextureSizeY, TextureColor, MipData0, MipCoverageData0);
+	if (bEncodeVirtualTexture)
+	{
+		GenerateLightmapMipsAndDilateColor(NumMips, TextureSizeX, TextureSizeY, TextureColor, MipData1, MipCoverageData1);
+	}
 
-		GenerateLightmapMipsAndDilateColor(NumMips, TextureSizeX, TextureSizeY, TextureColor, MipData, MipCoverageData);
+	// Unlock all mip levels.
+	for (int32 MipIndex = 0; MipIndex < NumMips; ++MipIndex)
+	{
+		Texture->Source.UnlockMip(0, LayerIndex, MipIndex);
+		FMemory::Free(MipCoverageData0[MipIndex]);
+		if (bEncodeVirtualTexture)
+		{
+			Texture->Source.UnlockMip(0, LayerIndex + 1, MipIndex);
+			FMemory::Free(MipCoverageData1[MipIndex]);
+		}
+	}
+}
 
-		// Unlock all mip levels.
-		for (int32 MipIndex = 0; MipIndex < NumMips; ++MipIndex)
+void FLightMapPendingTexture::StartEncoding(ULevel* LightingScenario, ITextureCompressorModule* UnusedCompressor)
+{
+	if (!bUObjectsCreated)
+	{
+		check(IsInGameThread());
+		CreateUObjects();
+	}
+
+	// Create the uncompressed top mip-level.
+	TArray< TArray<FFourDistanceFieldSamples> > ShadowMipData;
+	int32 NumShadowChannelsUsed = 0;
+	if (NeedsStaticShadowTexture())
+	{
+		NumShadowChannelsUsed = FLightMap2D::EncodeShadowTexture(LightingScenario, *this, ShadowMipData);
+	}
+
+	FColor TextureColor;
+	if ( GVisualizeLightmapTextures )
+	{
+		TextureColor = FColor::MakeRandomColor();
+	}
+
+	const ETextureSourceFormat SkyOcclusionFormat = TSF_BGRA8;
+	const ETextureSourceFormat AOMaskFormat = TSF_G8;
+	const ETextureSourceFormat ShadowMapFormat = NumShadowChannelsUsed == 1 ? TSF_G8 : TSF_BGRA8;
+	const ETextureSourceFormat BaseFormat = TSF_BGRA8;
+
+	if (SkyOcclusionTexture != nullptr)
+	{
+		auto Texture = SkyOcclusionTexture;
+		Texture->Source.Init2DWithMipChain(GetSizeX(), GetSizeY(), SkyOcclusionFormat);
+		Texture->MipGenSettings = TMGS_LeaveExistingMips;
+		Texture->Filter	= GUseBilinearLightmaps ? TF_Default : TF_Nearest;
+		Texture->LODGroup = TEXTUREGROUP_Lightmap;
+		Texture->LightmapFlags = ELightMapFlags( LightmapFlags );
+		EncodeSkyOcclusionTexture(Texture, 0u, TextureColor);
+	}
+	
+	if (AOMaterialMaskTexture != nullptr)
+	{
+		auto Texture = AOMaterialMaskTexture;
+		Texture->Source.Init2DWithMipChain(GetSizeX(), GetSizeY(), AOMaskFormat);
+		Texture->MipGenSettings = TMGS_LeaveExistingMips;
+		Texture->Filter	= GUseBilinearLightmaps ? TF_Default : TF_Nearest;
+		Texture->LODGroup = TEXTUREGROUP_Lightmap;
+		Texture->LightmapFlags = ELightMapFlags( LightmapFlags );
+		EncodeAOMaskTexture(Texture, 0u, TextureColor);
+	}
+
+	if (ShadowMapTexture != nullptr)
+	{
+		ShadowMapTexture->Filter = GUseBilinearLightmaps ? TF_Default : TF_Nearest;
+		ShadowMapTexture->LODGroup = TEXTUREGROUP_Shadowmap;
+
+		// TODO - Unify these?  Currently the values match
+		ShadowMapTexture->ShadowmapFlags = EShadowMapFlags(LightmapFlags);
+
+		ShadowMapTexture->Source.Init2DWithMipChain(GetSizeX(), GetSizeY(), ShadowMapFormat);
+		ShadowMapTexture->MipGenSettings = TMGS_LeaveExistingMips;
+
+		EncodeShadowMapTexture(ShadowMipData, ShadowMapTexture, 0u);
+	}
+
+	// Encode and compress the coefficient textures.
+	for(uint32 CoefficientIndex = 0; CoefficientIndex < NUM_STORED_LIGHTMAP_COEF; CoefficientIndex += 2)
+	{
+		auto Texture = Textures[CoefficientIndex];
+		if (Texture == nullptr)
 		{
-			Texture->Source.UnlockMip(MipIndex);
-			FMemory::Free( MipCoverageData[ MipIndex ] );
+			continue;
 		}
+
+		Texture->Source.Init2DWithMipChain(GetSizeX(), GetSizeY() * 2, BaseFormat);	// Top/bottom atlased
+		Texture->MipGenSettings = TMGS_LeaveExistingMips;
+		Texture->Filter	= GUseBilinearLightmaps ? TF_Default : TF_Nearest;
+		Texture->LODGroup = TEXTUREGROUP_Lightmap;
+		Texture->LightmapFlags = ELightMapFlags( LightmapFlags );
+		Texture->bForcePVRTC4 = true;
+
+		EncodeCoefficientTexture(CoefficientIndex, Texture, 0u, TextureColor, false);
 	}
 
 	if (VirtualTexture)
@@ -1694,23 +1786,20 @@ void FLightMapPendingTexture::StartEncoding(ULevel* LightingScenario, ITextureCo
 
 		TArray<ETextureSourceFormat> LayerFormat;
 		LayerFormat.Init(TSF_Invalid, NumVirtualTextureLayers);
-		LayerFormat[0] = Textures[0]->Source.GetFormat();
-		LayerFormat[1] = Textures[0]->Source.GetFormat();
+		LayerFormat[0] = BaseFormat;
+		LayerFormat[1] = BaseFormat;
 
 		if (SkyOcclusionLayer != InvalidLayerId)
 		{
-			check(SkyOcclusionTexture);
-			LayerFormat[SkyOcclusionLayer] = SkyOcclusionTexture->Source.GetFormat();
+			LayerFormat[SkyOcclusionLayer] = SkyOcclusionFormat;
 		}
 		if (AOMaterialMaskLayer != InvalidLayerId)
 		{
-			check(AOMaterialMaskTexture);
-			LayerFormat[AOMaterialMaskLayer] = AOMaterialMaskTexture->Source.GetFormat();
+			LayerFormat[AOMaterialMaskLayer] = AOMaskFormat;
 		}
 		if (ShadowMaskLayer != InvalidLayerId)
 		{
-			check(ShadowMapTexture);
-			LayerFormat[ShadowMaskLayer] = ShadowMapTexture->Source.GetFormat();
+			LayerFormat[ShadowMaskLayer] = ShadowMapFormat;
 		}
 
 		VirtualTexture->Source.InitLayered2DWithMipChain(GetSizeX(), GetSizeY(), NumVirtualTextureLayers, LayerFormat.GetData());
@@ -1726,58 +1815,31 @@ void FLightMapPendingTexture::StartEncoding(ULevel* LightingScenario, ITextureCo
 		VirtualTexture->GetDefaultFormatSettings(DefaultFormatSettings);
 		VirtualTexture->LayerFormatSettings.Init(DefaultFormatSettings, NumVirtualTextureLayers);
 
-		for (int MipIndex = 0; MipIndex < VirtualTexture->Source.GetNumMips(); ++MipIndex)
+		EncodeCoefficientTexture(0, VirtualTexture, 0u, TextureColor, true);
+		if (SkyOcclusionLayer != InvalidLayerId)
 		{
-			TArray<uint8> SourceData;
-			Textures[0]->Source.GetMipData(SourceData, MipIndex);
-
-			// Copy top half of texture into layer0
-			uint8* Layer0Data = VirtualTexture->Source.LockMip(0, 0, MipIndex);
-			FMemory::Memcpy(Layer0Data, &SourceData[0], SourceData.Num() / 2);
-			VirtualTexture->Source.UnlockMip(0, 0, MipIndex);
-
-			// Bottom half of texture into layer1
-			uint8* Layer1Data = VirtualTexture->Source.LockMip(0, 1, MipIndex);
-			FMemory::Memcpy(Layer1Data, &SourceData[SourceData.Num() / 2], SourceData.Num() / 2);
-			VirtualTexture->Source.UnlockMip(0, 1, MipIndex);
-
-			if (SkyOcclusionLayer != InvalidLayerId)
-			{
-				FTextureFormatSettings& FormatSettings = VirtualTexture->LayerFormatSettings[SkyOcclusionLayer];
-				SkyOcclusionTexture->GetDefaultFormatSettings(FormatSettings);
-
-				SkyOcclusionTexture->Source.GetMipData(SourceData, MipIndex);
-				uint8* LayerData = VirtualTexture->Source.LockMip(0, SkyOcclusionLayer, MipIndex);
-				FMemory::Memcpy(LayerData, &SourceData[0], SourceData.Num());
-				VirtualTexture->Source.UnlockMip(0, SkyOcclusionLayer, MipIndex);
-			}
-			if (AOMaterialMaskLayer != InvalidLayerId)
-			{
-				FTextureFormatSettings& FormatSettings = VirtualTexture->LayerFormatSettings[AOMaterialMaskLayer];
-				AOMaterialMaskTexture->GetDefaultFormatSettings(FormatSettings);
-
-				AOMaterialMaskTexture->Source.GetMipData(SourceData, MipIndex);
-				uint8* LayerData = VirtualTexture->Source.LockMip(0, AOMaterialMaskLayer, MipIndex);
-				FMemory::Memcpy(LayerData, &SourceData[0], SourceData.Num());
-				VirtualTexture->Source.UnlockMip(0, AOMaterialMaskLayer, MipIndex);
-			}
-			if (ShadowMaskLayer != InvalidLayerId)
-			{
-				FTextureFormatSettings& FormatSettings = VirtualTexture->LayerFormatSettings[ShadowMaskLayer];
-				ShadowMapTexture->GetDefaultFormatSettings(FormatSettings);
-
-				ShadowMapTexture->Source.GetMipData(SourceData, MipIndex);
-				uint8* LayerData = VirtualTexture->Source.LockMip(0, ShadowMaskLayer, MipIndex);
-				FMemory::Memcpy(LayerData, &SourceData[0], SourceData.Num());
-				VirtualTexture->Source.UnlockMip(0, ShadowMaskLayer, MipIndex);
-			}
+			EncodeSkyOcclusionTexture(VirtualTexture, SkyOcclusionLayer, TextureColor);
+		}
+		if (AOMaterialMaskLayer != InvalidLayerId)
+		{
+			EncodeAOMaskTexture(VirtualTexture, AOMaterialMaskLayer, TextureColor);
+		}
+		if (ShadowMaskLayer != InvalidLayerId)
+		{
+			EncodeShadowMapTexture(ShadowMipData, VirtualTexture, ShadowMaskLayer);
 		}
 	}
 
+	// Link textures to allocations
 	for (int32 AllocationIndex = 0; AllocationIndex < Allocations.Num(); AllocationIndex++)
 	{
-		auto& Allocation = Allocations[AllocationIndex];
-		Allocation->LightMap->VirtualTexture = VirtualTexture;
+		FLightMapAllocation& Allocation = *Allocations[AllocationIndex];
+		Allocation.LightMap->SkyOcclusionTexture = SkyOcclusionTexture;
+		Allocation.LightMap->AOMaterialMaskTexture = AOMaterialMaskTexture;
+		Allocation.LightMap->ShadowMapTexture = ShadowMapTexture;
+		Allocation.LightMap->Textures[0] = Textures[0];
+		Allocation.LightMap->Textures[1] = Textures[2];
+		Allocation.LightMap->VirtualTexture = VirtualTexture;
 	}
 
 	bIsFinishedEncoding = true;
@@ -2271,8 +2333,20 @@ void FLightMap2D::EncodeTextures( UWorld* InWorld, bool bLightingSuccessful, boo
 #if WITH_EDITOR
 	if (bLightingSuccessful)
 	{
+		const bool bUseVirtualTextures = (CVarVirtualTexturedLightMaps.GetValueOnAnyThread() != 0) && UseVirtualTexturing(GMaxRHIFeatureLevel);
+		const bool bIncludeNonVirtualTextures = !bUseVirtualTextures || (CVarIncludeNonVirtualTexturedLightMaps.GetValueOnAnyThread() != 0);
+
 		GWarn->BeginSlowTask( NSLOCTEXT("LightMap2D", "BeginEncodingLightMapsTask", "Encoding light-maps"), false );
-		int32 PackedLightAndShadowMapTextureSize = InWorld->GetWorldSettings()->PackedLightAndShadowMapTextureSize;
+		int32 PackedLightAndShadowMapTextureSizeX = InWorld->GetWorldSettings()->PackedLightAndShadowMapTextureSize;
+		int32 PackedLightAndShadowMapTextureSizeY = PackedLightAndShadowMapTextureSizeX / 2;
+
+		if (!bIncludeNonVirtualTextures)
+		{
+			// If we exclusively using VT lightmaps, don't need to worry about max size of a given texture sheet
+			// Also just make square textures, don't need 2x1 aspect ratio
+			PackedLightAndShadowMapTextureSizeX = 32 * 1024; // FTextureLayout uses uint16 for size in some places
+			PackedLightAndShadowMapTextureSizeY = PackedLightAndShadowMapTextureSizeX;
+		}
 
 		// Reset the pending light-map size.
 		PendingLightMapSize = 0;
@@ -2345,8 +2419,8 @@ void FLightMap2D::EncodeTextures( UWorld* InWorld, bool bLightingSuccessful, boo
 
 			if (!Texture)
 			{
-				int32 NewTextureSizeX = PackedLightAndShadowMapTextureSize;
-				int32 NewTextureSizeY = PackedLightAndShadowMapTextureSize / 2;
+				int32 NewTextureSizeX = PackedLightAndShadowMapTextureSizeX;
+				int32 NewTextureSizeY = PackedLightAndShadowMapTextureSizeY;
 
 				// Assumes identically-sized allocations, fit into the smallest 2x1 rectangle
 				const int32 AllocationCountX = FMath::CeilToInt(FMath::Sqrt(FMath::DivideAndRoundUp(PendingGroup.Allocations.Num() * 2 * MaxHeight, MaxWidth)));
@@ -2360,12 +2434,17 @@ void FLightMap2D::EncodeTextures( UWorld* InWorld, bool bLightingSuccessful, boo
 					NewTextureSizeY = FMath::RoundUpToPowerOfTwo(AllocationSizeY);
 
 					// Force 2:1 aspect
-					NewTextureSizeX = FMath::Max(NewTextureSizeX, NewTextureSizeY * 2);
-					NewTextureSizeY = FMath::Max(NewTextureSizeY, NewTextureSizeX / 2);
+					if (bIncludeNonVirtualTextures)
+					{
+						NewTextureSizeX = FMath::Max(NewTextureSizeX, NewTextureSizeY * 2);
+						NewTextureSizeY = FMath::Max(NewTextureSizeY, NewTextureSizeX / 2);
+					}
 				}
 
 				// If there is no existing appropriate texture, create a new one.
-				Texture = new FLightMapPendingTexture(InWorld, NewTextureSizeX, NewTextureSizeY);
+				// If we have non-VT, need 2to1 aspect ratio to handle packing lightmap into top/bottom texture region
+				// With only VT, better to use square lightmaps (better page table usage)
+				Texture = new FLightMapPendingTexture(InWorld, NewTextureSizeX, NewTextureSizeY, bIncludeNonVirtualTextures ? ETextureLayoutAspectRatio::Force2To1 : ETextureLayoutAspectRatio::ForceSquare);
 				PendingTextures.Add(Texture);
 				Texture->Outer = PendingGroup.Outer;
 				Texture->Bounds = PendingGroup.Bounds;
@@ -2457,7 +2536,7 @@ void FLightMap2D::EncodeTextures( UWorld* InWorld, bool bLightingSuccessful, boo
 }
 
 #if WITH_EDITOR
-int32 FLightMap2D::EncodeShadowTexture(ULevel* LightingScenario, struct FLightMapPendingTexture& PendingTexture, UShadowMapTexture2D* Texture, TArray< TArray<FFourDistanceFieldSamples>>& MipData)
+int32 FLightMap2D::EncodeShadowTexture(ULevel* LightingScenario, struct FLightMapPendingTexture& PendingTexture, TArray< TArray<FFourDistanceFieldSamples>>& MipData)
 {
 	TArray<FFourDistanceFieldSamples>* TopMipData = new(MipData) TArray<FFourDistanceFieldSamples>();
 	TopMipData->Empty(PendingTexture.GetSizeX() * PendingTexture.GetSizeY());
@@ -2523,9 +2602,6 @@ int32 FLightMap2D::EncodeShadowTexture(ULevel* LightingScenario, struct FLightMa
 			}
 		}
 
-		// Link the shadow-map to the texture.
-		Allocation.LightMap->ShadowMapTexture = Texture;
-
 		// Free the shadow-map's raw data.
 		for (auto& ShadowMapPair : Allocation.ShadowMapData)
 		{
diff --git a/Engine/Source/Runtime/Engine/Private/ShadowMap.cpp b/Engine/Source/Runtime/Engine/Private/ShadowMap.cpp
--- a/Engine/Source/Runtime/Engine/Private/ShadowMap.cpp
+++ b/Engine/Source/Runtime/Engine/Private/ShadowMap.cpp
@@ -176,7 +176,7 @@ struct FShadowMapPendingTexture : FTextureLayout
 	 * Minimal initialization constructor.
 	 */
 	FShadowMapPendingTexture(uint32 InSizeX,uint32 InSizeY)
-		: FTextureLayout(4, 4, InSizeX, InSizeY, /* PowerOfTwo */ true, /* Force2To1Aspect */ false, /* AlignByFour */ true) // Min size is 4x4 in case of block compression.
+		: FTextureLayout(4, 4, InSizeX, InSizeY, /* PowerOfTwo */ true, /* Aspect */ ETextureLayoutAspectRatio::None, /* AlignByFour */ true) // Min size is 4x4 in case of block compression.
 		, Bounds(FBox(ForceInit))
 		, ShadowmapFlags(SMF_None)
 		, UnallocatedTexels(InSizeX * InSizeY)
diff --git a/Engine/Source/Runtime/Engine/Private/Texture.cpp b/Engine/Source/Runtime/Engine/Private/Texture.cpp
--- a/Engine/Source/Runtime/Engine/Private/Texture.cpp
+++ b/Engine/Source/Runtime/Engine/Private/Texture.cpp
@@ -1263,6 +1263,22 @@ void UTexture::GetLayerFormatSettings(int32 LayerIndex, FTextureFormatSettings&
 	}
 }
 
+void UTexture::SetLayerFormatSettings(int32 LayerIndex, const FTextureFormatSettings& InSettings)
+{
+	check(LayerIndex >= 0);
+	if (LayerIndex >= LayerFormatSettings.Num())
+	{
+		FTextureFormatSettings DefaultSettings;
+		GetDefaultFormatSettings(DefaultSettings);
+		LayerFormatSettings.Reserve(LayerIndex + 1);
+		while (LayerIndex >= LayerFormatSettings.Num())
+		{
+			LayerFormatSettings.Add(DefaultSettings);
+		}
+	}
+	LayerFormatSettings[LayerIndex] = InSettings;
+}
+
 #endif // #if WITH_EDITOR
 
 FName GetDefaultTextureFormatName( const ITargetPlatform* TargetPlatform, const UTexture* Texture, int32 LayerIndex, const FConfigFile& EngineSettings, bool bSupportDX11TextureFormats, bool bSupportCompressedVolumeTexture, int32 BlockSize )
diff --git a/Engine/Source/Runtime/Engine/Public/LightMap.h b/Engine/Source/Runtime/Engine/Public/LightMap.h
--- a/Engine/Source/Runtime/Engine/Public/LightMap.h
+++ b/Engine/Source/Runtime/Engine/Public/LightMap.h
@@ -305,7 +305,7 @@ public:
 	/**
 	 * Constructs mip maps for a single shadowmap texture.
 	 */
-	static int32 EncodeShadowTexture(ULevel* LightingScenario, struct FLightMapPendingTexture& PendingTexture, UShadowMapTexture2D* Texture, TArray< TArray<FFourDistanceFieldSamples>>& MipData);
+	static int32 EncodeShadowTexture(ULevel* LightingScenario, struct FLightMapPendingTexture& PendingTexture, TArray<TArray<FFourDistanceFieldSamples>>& MipData);
 #endif // WITH_EDITOR
 
 	/** Call to enable/disable status update of LightMap encoding */
diff --git a/Engine/Source/Runtime/Engine/Public/TextureLayout.h b/Engine/Source/Runtime/Engine/Public/TextureLayout.h
--- a/Engine/Source/Runtime/Engine/Public/TextureLayout.h
+++ b/Engine/Source/Runtime/Engine/Public/TextureLayout.h
@@ -8,6 +8,13 @@
 
 #include "CoreMinimal.h"
 
+enum class ETextureLayoutAspectRatio
+{
+	None,
+	Force2To1,
+	ForceSquare,
+};
+
 /**
  * An incremental texture space allocator.
  * For best results, add the elements ordered descending in size.
@@ -26,11 +33,11 @@ public:
 	 * @param	bInForce2To1Aspect - True if the texture size must have a 2:1 aspect.
 	 * @param	bInAlignByFour - True if the texture size must be a multiple of 4..
 	 */
-	FTextureLayout(uint32 MinSizeX, uint32 MinSizeY, uint32 MaxSizeX, uint32 MaxSizeY, bool bInPowerOfTwoSize = false, bool bInForce2To1Aspect = false, bool bInAlignByFour = true):
+	FTextureLayout(uint32 MinSizeX, uint32 MinSizeY, uint32 MaxSizeX, uint32 MaxSizeY, bool bInPowerOfTwoSize = false, ETextureLayoutAspectRatio InAspect = ETextureLayoutAspectRatio::None, bool bInAlignByFour = true):
 		SizeX(MinSizeX),
 		SizeY(MinSizeY),
+		AspectRatio(InAspect),
 		bPowerOfTwoSize(bInPowerOfTwoSize),
-		bForce2To1Aspect(bInForce2To1Aspect),
 		bAlignByFour(bInAlignByFour)
 	{
 		new(Nodes) FTextureLayoutNode(0, 0, MaxSizeX, MaxSizeY);
@@ -83,10 +90,21 @@ public:
 				SizeX = FMath::Max<uint32>(SizeX, FMath::RoundUpToPowerOfTwo(Node.MinX + ElementSizeX));
 				SizeY = FMath::Max<uint32>(SizeY, FMath::RoundUpToPowerOfTwo(Node.MinY + ElementSizeY));
 			
-				if (bForce2To1Aspect)
+				switch(AspectRatio)
 				{
+				case ETextureLayoutAspectRatio::None:
+					break;
+				case ETextureLayoutAspectRatio::ForceSquare:
+					SizeX = FMath::Max(SizeX, SizeY);
+					SizeY = SizeX;
+					break;
+				case ETextureLayoutAspectRatio::Force2To1:
 					SizeX = FMath::Max( SizeX, SizeY * 2 );
 					SizeY = FMath::Max( SizeY, SizeX / 2 );
+					break;
+				default:
+					checkNoEntry();
+					break;
 				}
 			}
 			else
@@ -200,8 +218,8 @@ private:
 
 	uint32 SizeX;
 	uint32 SizeY;
+	ETextureLayoutAspectRatio AspectRatio;
 	bool bPowerOfTwoSize;
-	bool bForce2To1Aspect;
 	bool bAlignByFour;
 	TArray<FTextureLayoutNode,TInlineAllocator<5> > Nodes;
 
@@ -266,10 +284,21 @@ private:
 					ExpectedSizeX = FMath::RoundUpToPowerOfTwo(ExpectedSizeX);
 					ExpectedSizeY = FMath::RoundUpToPowerOfTwo(ExpectedSizeY);
 
-					if (bForce2To1Aspect)
+					switch (AspectRatio)
 					{
-						ExpectedSizeX = FMath::Max( ExpectedSizeX, ExpectedSizeY * 2 );
-						ExpectedSizeY = FMath::Max( ExpectedSizeY, ExpectedSizeX / 2 );
+					case ETextureLayoutAspectRatio::None:
+						break;
+					case ETextureLayoutAspectRatio::ForceSquare:
+						ExpectedSizeX = FMath::Max(ExpectedSizeX, ExpectedSizeY);
+						ExpectedSizeY = ExpectedSizeX;
+						break;
+					case ETextureLayoutAspectRatio::Force2To1:
+						ExpectedSizeX = FMath::Max(ExpectedSizeX, ExpectedSizeY * 2);
+						ExpectedSizeY = FMath::Max(ExpectedSizeY, ExpectedSizeX / 2);
+						break;
+					default:
+						checkNoEntry();
+						break;
 					}
 				}
 				
diff --git a/Engine/Source/Runtime/Renderer/Private/RendererScene.cpp b/Engine/Source/Runtime/Renderer/Private/RendererScene.cpp
--- a/Engine/Source/Runtime/Renderer/Private/RendererScene.cpp
+++ b/Engine/Source/Runtime/Renderer/Private/RendererScene.cpp
@@ -1045,7 +1045,7 @@ FScene::FScene(UWorld* InWorld, bool bInRequiresHitProxies, bool bInIsEditorScen
 ,	ReflectionSceneData(InFeatureLevel)
 ,	IndirectLightingCache(InFeatureLevel)
 ,	DistanceFieldSceneData(GShaderPlatformForFeatureLevel[InFeatureLevel])
-,	PreshadowCacheLayout(0, 0, 0, 0, false, false)
+,	PreshadowCacheLayout(0, 0, 0, 0, false)
 ,	AtmosphericFog(NULL)
 ,	PrecomputedVisibilityHandler(NULL)
 ,	LightOctree(FVector::ZeroVector,HALF_WORLD_MAX)
diff --git a/Engine/Source/Runtime/Renderer/Private/ShadowSetup.cpp b/Engine/Source/Runtime/Renderer/Private/ShadowSetup.cpp
--- a/Engine/Source/Runtime/Renderer/Private/ShadowSetup.cpp
+++ b/Engine/Source/Runtime/Renderer/Private/ShadowSetup.cpp
@@ -1612,7 +1612,7 @@ void FSceneRenderer::UpdatePreshadowCache(FSceneRenderTargets& SceneContext)
 		{
 			// Initialize the texture layout if necessary
 			const FIntPoint PreshadowCacheBufferSize = SceneContext.GetPreShadowCacheTextureResolution();
-			Scene->PreshadowCacheLayout = FTextureLayout(1, 1, PreshadowCacheBufferSize.X, PreshadowCacheBufferSize.Y, false, false, false);
+			Scene->PreshadowCacheLayout = FTextureLayout(1, 1, PreshadowCacheBufferSize.X, PreshadowCacheBufferSize.Y, false, ETextureLayoutAspectRatio::None, false);
 		}
 
 		// Iterate through the cached preshadows, removing those that are not going to be rendered this frame
@@ -3700,7 +3700,7 @@ void FSceneRenderer::AllocatePerObjectShadowDepthTargets(FRHICommandListImmediat
 
 		int32 OriginalNumAtlases = SortedShadowsForShadowDepthPass.ShadowMapAtlases.Num();
 
-		FTextureLayout CurrentShadowLayout(1, 1, ShadowBufferResolution.X, ShadowBufferResolution.Y, false, false, false);
+		FTextureLayout CurrentShadowLayout(1, 1, ShadowBufferResolution.X, ShadowBufferResolution.Y, false, ETextureLayoutAspectRatio::None, false);
 		FPooledRenderTargetDesc ShadowMapDesc2D = FPooledRenderTargetDesc::Create2DDesc(ShadowBufferResolution, PF_ShadowDepth, FClearValueBinding::DepthOne, TexCreate_None, TexCreate_DepthStencilTargetable, false);
 		ShadowMapDesc2D.Flags |= GFastVRamConfig.ShadowPerObject;
 
@@ -3742,7 +3742,7 @@ void FSceneRenderer::AllocatePerObjectShadowDepthTargets(FRHICommandListImmediat
 				}
 				else
 				{
-					CurrentShadowLayout = FTextureLayout(1, 1, ShadowBufferResolution.X, ShadowBufferResolution.Y, false, false, false);
+					CurrentShadowLayout = FTextureLayout(1, 1, ShadowBufferResolution.X, ShadowBufferResolution.Y, false, ETextureLayoutAspectRatio::None, false);
 					SortedShadowsForShadowDepthPass.ShadowMapAtlases.AddDefaulted();
 
 					if (CurrentShadowLayout.AddElement(
@@ -3823,7 +3823,7 @@ const TCHAR* GetCSMRenderTargetName(int32 ShadowMapIndex)
 struct FLayoutAndAssignedShadows
 {
 	FLayoutAndAssignedShadows(int32 MaxTextureSize) :
-		TextureLayout(1, 1, MaxTextureSize, MaxTextureSize, false, false, false)
+		TextureLayout(1, 1, MaxTextureSize, MaxTextureSize, false, ETextureLayoutAspectRatio::None, false)
 	{}
 
 	FTextureLayout TextureLayout;
@@ -3898,7 +3898,7 @@ void FSceneRenderer::AllocateRSMDepthTargets(FRHICommandListImmediate& RHICmdLis
 		&& FeatureLevel >= ERHIFeatureLevel::SM5)
 	{
 		const int32 MaxTextureSize = 1 << (GMaxTextureMipCount - 1);
-		FTextureLayout ShadowLayout(1, 1, MaxTextureSize, MaxTextureSize, false, false, false);
+		FTextureLayout ShadowLayout(1, 1, MaxTextureSize, MaxTextureSize, false, ETextureLayoutAspectRatio::None, false);
 
 		for (int32 ShadowIndex = 0; ShadowIndex < RSMShadows.Num(); ShadowIndex++)
 		{
@@ -4031,7 +4031,7 @@ void FSceneRenderer::AllocateTranslucentShadowDepthTargets(FRHICommandListImmedi
 		// Start with an empty atlas for per-object shadows (don't allow packing object shadows into the CSM atlas atm)
 		SortedShadowsForShadowDepthPass.TranslucencyShadowMapAtlases.AddDefaulted();
 
-		FTextureLayout CurrentShadowLayout(1, 1, TranslucentShadowBufferResolution.X, TranslucentShadowBufferResolution.Y, false, false, false);
+		FTextureLayout CurrentShadowLayout(1, 1, TranslucentShadowBufferResolution.X, TranslucentShadowBufferResolution.Y, false, ETextureLayoutAspectRatio::None, false);
 
 		// Sort the projected shadows by resolution.
 		TranslucentShadows.Sort(FCompareFProjectedShadowInfoByResolution());
@@ -4054,7 +4054,7 @@ void FSceneRenderer::AllocateTranslucentShadowDepthTargets(FRHICommandListImmedi
 			}
 			else
 			{
-				CurrentShadowLayout = FTextureLayout(1, 1, TranslucentShadowBufferResolution.X, TranslucentShadowBufferResolution.Y, false, false, false);
+				CurrentShadowLayout = FTextureLayout(1, 1, TranslucentShadowBufferResolution.X, TranslucentShadowBufferResolution.Y, false, ETextureLayoutAspectRatio::None, false);
 				SortedShadowsForShadowDepthPass.TranslucencyShadowMapAtlases.AddDefaulted();
 
 				if (CurrentShadowLayout.AddElement(
