Don't mess with texture pad setting when converting VT
- No need to try to force on power of two padding here, if it wasn't already set we wouldn't have made it this far
- Previous code was broken, in that it would replace square power of two with regular (not square) power of two

diff --git a/Engine/Source/Runtime/Engine/Private/VT/VirtualTextureDataBuilder.cpp b/Engine/Source/Runtime/Engine/Private/VT/VirtualTextureDataBuilder.cpp
--- a/Engine/Source/Runtime/Engine/Private/VT/VirtualTextureDataBuilder.cpp
+++ b/Engine/Source/Runtime/Engine/Private/VT/VirtualTextureDataBuilder.cpp
@@ -861,17 +861,7 @@ void FVirtualTextureDataBuilder::BuildSourcePixels(const FTextureSourceData& Sou
 			// Leave original mip settings alone unless it's none at which point we will just generate them using a simple average
 			if (TBSettings.MipGenSettings == TMGS_NoMipmaps)
 			{
-				if (FMath::IsPowerOfTwo(SourceMips[0].SizeX) &&
-					FMath::IsPowerOfTwo(SourceMips[0].SizeY))
-				{
-					TBSettings.MipGenSettings = TMGS_SimpleAverage;
-				}
-				else
-				{
-					//checkf(false, TEXT("Non power of two textures cannot generate mips. Only existing mips can be used."));
-					TBSettings.MipGenSettings = TMGS_SimpleAverage;
-					TBSettings.PowerOfTwoMode = ETexturePowerOfTwoSetting::PadToPowerOfTwo;
-				}
+				TBSettings.MipGenSettings = TMGS_SimpleAverage;
 			}
 
 			// Use the texture compressor module to do all the hard work
