Fix some edge cases that could result in FMaterialRenderProxy being deleted without removing its VT producer callbacks

diff --git a/Engine/Source/Runtime/Engine/Private/Materials/MaterialShared.cpp b/Engine/Source/Runtime/Engine/Private/Materials/MaterialShared.cpp
--- a/Engine/Source/Runtime/Engine/Private/Materials/MaterialShared.cpp
+++ b/Engine/Source/Runtime/Engine/Private/Materials/MaterialShared.cpp
@@ -2301,10 +2301,9 @@ static void OnVirtualTextureDestroyed(const FVirtualTextureProducerHandle& InHan
 	MaterialProxy->CacheUniformExpressions(false);
 }
 
-static IAllocatedVirtualTexture* AllocateVTStack(const FMaterialRenderProxy* MaterialProxy, const FMaterialRenderContext& Context, const FUniformExpressionSet& UniformExpressionSet, const FMaterialVirtualTextureStack& VTStack)
+IAllocatedVirtualTexture* FMaterialRenderProxy::AllocateVTStack(const FMaterialRenderContext& Context, const FUniformExpressionSet& UniformExpressionSet, const FMaterialVirtualTextureStack& VTStack) const
 {
-	check(!VTStack.IsPreallocatedStack())
-
+	check(!VTStack.IsPreallocatedStack());
 	const uint32 NumLayers = VTStack.GetNumLayers();
 	if (NumLayers == 0u)
 	{
@@ -2333,13 +2332,14 @@ static IAllocatedVirtualTexture* AllocateVTStack(const FMaterialRenderProxy* Mat
 			const FVirtualTextureProducerHandle& ProducerHandle = VirtualTextureResourceForLayer->GetProducerHandle();
 			VTDesc.ProducerHandle[LayerIndex] = ProducerHandle;
 			VTDesc.LocalLayerToProduce[LayerIndex] = 0u;
-			GetRendererModule().AddVirtualTextureProducerDestroyedCallback(ProducerHandle, &OnVirtualTextureDestroyed, const_cast<FMaterialRenderProxy*>(MaterialProxy));
+			GetRendererModule().AddVirtualTextureProducerDestroyedCallback(ProducerHandle, &OnVirtualTextureDestroyed, const_cast<FMaterialRenderProxy*>(this));
 			bFoundValidLayer = true;
 		}
 	}
 
 	if (bFoundValidLayer)
 	{
+		HasVirtualTextureCallbacks = true;
 		return GetRendererModule().AllocateVirtualTexture(VTDesc);
 	}
 	return nullptr;
@@ -2375,6 +2375,8 @@ void FMaterialRenderProxy::EvaluateUniformExpressions(FUniformExpressionCache& O
 	OutUniformExpressionCache.ResetAllocatedVTs();
 	OutUniformExpressionCache.AllocatedVTs.Empty(UniformExpressionSet.VTStacks.Num());
 	OutUniformExpressionCache.OwnedAllocatedVTs.Empty(UniformExpressionSet.VTStacks.Num());
+	check(!HasVirtualTextureCallbacks);
+
 	for (int32 i = 0; i < UniformExpressionSet.VTStacks.Num(); ++i)
 	{
 		const FMaterialVirtualTextureStack& VTStack = UniformExpressionSet.VTStacks[i];
@@ -2385,7 +2387,7 @@ void FMaterialRenderProxy::EvaluateUniformExpressions(FUniformExpressionCache& O
 		}
 		else
 		{
-			AllocatedVT = AllocateVTStack(this, Context, UniformExpressionSet, VTStack);
+			AllocatedVT = AllocateVTStack(Context, UniformExpressionSet, VTStack);
 			if (AllocatedVT != nullptr)
 			{
 				OutUniformExpressionCache.OwnedAllocatedVTs.Add(AllocatedVT);
@@ -2446,10 +2448,7 @@ void FMaterialRenderProxy::CacheUniformExpressions(bool bRecreateUniformBuffer)
 	}
 	DeferredUniformExpressionCacheRequests.Add(this);
 
-	UMaterialInterface::IterateOverActiveFeatureLevels([&](ERHIFeatureLevel::Type InFeatureLevel)
-	{
-		InvalidateUniformExpressionCache(bRecreateUniformBuffer);
-	});
+	InvalidateUniformExpressionCache(bRecreateUniformBuffer);
 
 	if (!GDeferUniformExpressionCaching)
 	{
@@ -2473,14 +2472,19 @@ void FMaterialRenderProxy::CacheUniformExpressions_GameThread(bool bRecreateUnif
 void FMaterialRenderProxy::InvalidateUniformExpressionCache(bool bRecreateUniformBuffer)
 {
 	check(IsInRenderingThread());
+
+	if (HasVirtualTextureCallbacks)
+	{
+		GetRendererModule().RemoveAllVirtualTextureProducerDestroyedCallbacks(this);
+		HasVirtualTextureCallbacks = false;
+	}
+
 	for (int32 i = 0; i < ERHIFeatureLevel::Num; ++i)
 	{
 		UniformExpressionCache[i].bUpToDate = false;
 		UniformExpressionCache[i].CachedUniformExpressionShaderMap = nullptr;
 		UniformExpressionCache[i].ResetAllocatedVTs();
 
-		GetRendererModule().RemoveAllVirtualTextureProducerDestroyedCallbacks(this);
-
 		if (bRecreateUniformBuffer)
 		{
 			// This is required if the FMaterial is being recompiled (the uniform buffer layout will change).
@@ -2512,6 +2516,7 @@ FMaterialRenderProxy::FMaterialRenderProxy()
 	: SubsurfaceProfileRT(0)
 	, MarkedForGarbageCollection(0)
 	, DeletedFlag(0)
+	, HasVirtualTextureCallbacks(0)
 {
 }
 
@@ -2523,6 +2528,8 @@ FMaterialRenderProxy::~FMaterialRenderProxy()
 		ReleaseResource();
 	}
 
+	check(!HasVirtualTextureCallbacks);
+
 	DeletedFlag = 1;
 }
 
@@ -2549,6 +2556,17 @@ void FMaterialRenderProxy::ReleaseDynamicRHI()
 	FExternalTextureRegistry::Get().RemoveMaterialRenderProxyReference(this);
 }
 
+void FMaterialRenderProxy::ReleaseResource()
+{
+	ReleaseResourceFlag = true;
+	FRenderResource::ReleaseResource();
+	if (HasVirtualTextureCallbacks)
+	{
+		GetRendererModule().RemoveAllVirtualTextureProducerDestroyedCallbacks(this);
+		HasVirtualTextureCallbacks = false;
+	}
+}
+
 void FMaterialRenderProxy::UpdateDeferredCachedUniformExpressions()
 {
 	LLM_SCOPE(ELLMTag::Materials);
diff --git a/Engine/Source/Runtime/Engine/Public/MaterialShared.h b/Engine/Source/Runtime/Engine/Public/MaterialShared.h
--- a/Engine/Source/Runtime/Engine/Public/MaterialShared.h
+++ b/Engine/Source/Runtime/Engine/Public/MaterialShared.h
@@ -1946,6 +1946,7 @@ public:
 	// FRenderResource interface.
 	ENGINE_API virtual void InitDynamicRHI() override;
 	ENGINE_API virtual void ReleaseDynamicRHI() override;
+	ENGINE_API virtual void ReleaseResource() override;
 
 	ENGINE_API static const TSet<FMaterialRenderProxy*>& GetMaterialRenderProxyMap() 
 	{
@@ -1964,6 +1965,7 @@ public:
 	}
 
 private:
+	IAllocatedVirtualTexture* AllocateVTStack(const FMaterialRenderContext& Context, const FUniformExpressionSet& UniformExpressionSet, const FMaterialVirtualTextureStack& VTStack) const;
 
 	/** 0 if not set, game thread pointer, do not dereference, only for comparison */
 	const USubsurfaceProfile* SubsurfaceProfileRT;
@@ -1971,6 +1973,9 @@ private:
 	/** For tracking down a bug accessing a deleted proxy. */
 	mutable int8 MarkedForGarbageCollection : 1;
 	mutable int8 DeletedFlag : 1;
+	mutable int8 ReleaseResourceFlag : 1;
+	/** If any VT producer destroyed callbacks have been registered */
+	mutable int8 HasVirtualTextureCallbacks : 1;
 
 	/** 
 	 * Tracks all material render proxies in all scenes, can only be accessed on the rendering thread.
diff --git a/Engine/Source/Runtime/RenderCore/Public/RendererInterface.h b/Engine/Source/Runtime/RenderCore/Public/RendererInterface.h
--- a/Engine/Source/Runtime/RenderCore/Public/RendererInterface.h
+++ b/Engine/Source/Runtime/RenderCore/Public/RendererInterface.h
@@ -699,7 +699,7 @@ public:
 	virtual FVirtualTextureProducerHandle RegisterVirtualTextureProducer(const FVTProducerDescription& Desc, IVirtualTexture* Producer) = 0;
 	virtual void ReleaseVirtualTextureProducer(const FVirtualTextureProducerHandle& Handle) = 0;
 	virtual void AddVirtualTextureProducerDestroyedCallback(const FVirtualTextureProducerHandle& Handle, FVTProducerDestroyedFunction* Function, void* Baton) = 0;
-	virtual void RemoveAllVirtualTextureProducerDestroyedCallbacks(const void* Baton) = 0;
+	virtual uint32 RemoveAllVirtualTextureProducerDestroyedCallbacks(const void* Baton) = 0;
 
 	/**	Provided a list of packed virtual texture tile ids, let the VT system request them. Note this should be called as long as the tiles are needed.*/
 	virtual void RequestVirtualTextureTilesForRegion(IAllocatedVirtualTexture* AllocatedVT, const FVector2D& InScreenSpaceSize, const FIntRect& InTextureRegion, int32 InMipLevel) = 0;
diff --git a/Engine/Source/Runtime/Renderer/Private/RendererModule.h b/Engine/Source/Runtime/Renderer/Private/RendererModule.h
--- a/Engine/Source/Runtime/Renderer/Private/RendererModule.h
+++ b/Engine/Source/Runtime/Renderer/Private/RendererModule.h
@@ -102,7 +102,7 @@ public:
 	virtual FVirtualTextureProducerHandle RegisterVirtualTextureProducer(const FVTProducerDescription& Desc, IVirtualTexture* Producer) override;
 	virtual void ReleaseVirtualTextureProducer(const FVirtualTextureProducerHandle& Handle) override;
 	virtual void AddVirtualTextureProducerDestroyedCallback(const FVirtualTextureProducerHandle& Handle, FVTProducerDestroyedFunction* Function, void* Baton) override;
-	virtual void RemoveAllVirtualTextureProducerDestroyedCallbacks(const void* Baton) override;
+	virtual uint32 RemoveAllVirtualTextureProducerDestroyedCallbacks(const void* Baton) override;
 	virtual void RequestVirtualTextureTilesForRegion(IAllocatedVirtualTexture* AllocatedVT, const FVector2D& InScreenSpaceSize, const FIntRect& InTextureRegion, int32 InMipLevel) override;
 	virtual void LoadPendingVirtualTextureTiles(FRHICommandListImmediate& RHICmdList, ERHIFeatureLevel::Type FeatureLevel) override;
 	virtual void FlushVirtualTextureCache() override;
diff --git a/Engine/Source/Runtime/Renderer/Private/SceneRendering.cpp b/Engine/Source/Runtime/Renderer/Private/SceneRendering.cpp
--- a/Engine/Source/Runtime/Renderer/Private/SceneRendering.cpp
+++ b/Engine/Source/Runtime/Renderer/Private/SceneRendering.cpp
@@ -3509,9 +3509,9 @@ void FRendererModule::AddVirtualTextureProducerDestroyedCallback(const FVirtualT
 	FVirtualTextureSystem::Get().AddProducerDestroyedCallback(Handle, Function, Baton);
 }
 
-void FRendererModule::RemoveAllVirtualTextureProducerDestroyedCallbacks(const void* Baton)
+uint32 FRendererModule::RemoveAllVirtualTextureProducerDestroyedCallbacks(const void* Baton)
 {
-	FVirtualTextureSystem::Get().RemoveAllProducerDestroyedCallbacks(Baton);
+	return FVirtualTextureSystem::Get().RemoveAllProducerDestroyedCallbacks(Baton);
 }
 
 void FRendererModule::RequestVirtualTextureTilesForRegion(IAllocatedVirtualTexture* AllocatedVT, const FVector2D& InScreenSpaceSize, const FIntRect& InTextureRegion, int32 InMipLevel)
diff --git a/Engine/Source/Runtime/Renderer/Private/VT/VirtualTextureProducer.cpp b/Engine/Source/Runtime/Renderer/Private/VT/VirtualTextureProducer.cpp
--- a/Engine/Source/Runtime/Renderer/Private/VT/VirtualTextureProducer.cpp
+++ b/Engine/Source/Runtime/Renderer/Private/VT/VirtualTextureProducer.cpp
@@ -137,12 +137,13 @@ void FVirtualTextureProducerCollection::AddDestroyedCallback(const FVirtualTextu
 	}
 }
 
-void FVirtualTextureProducerCollection::RemoveAllCallbacks(const void* Baton)
+uint32 FVirtualTextureProducerCollection::RemoveAllCallbacks(const void* Baton)
 {
 	check(IsInRenderingThread());
 	check(Baton);
 
-	for (int32 CallbackIndex = 0u; CallbackIndex < Callbacks.Num(); ++CallbackIndex)
+	uint32 NumRemoved = 0u;
+	for (int32 CallbackIndex = CallbackList_Count; CallbackIndex < Callbacks.Num(); ++CallbackIndex)
 	{
 		FCallbackEntry& Callback = Callbacks[CallbackIndex];
 		if (Callback.Baton == Baton)
@@ -152,8 +153,10 @@ void FVirtualTextureProducerCollection::RemoveAllCallbacks(const void* Baton)
 			Callback.Baton = nullptr;
 			Callback.OwnerHandle = FVirtualTextureProducerHandle();
 			ReleaseCallback(CallbackIndex);
+			++NumRemoved;
 		}
 	}
+	return NumRemoved;
 }
 
 FVirtualTextureProducer* FVirtualTextureProducerCollection::FindProducer(const FVirtualTextureProducerHandle& Handle)
diff --git a/Engine/Source/Runtime/Renderer/Private/VT/VirtualTextureProducer.h b/Engine/Source/Runtime/Renderer/Private/VT/VirtualTextureProducer.h
--- a/Engine/Source/Runtime/Renderer/Private/VT/VirtualTextureProducer.h
+++ b/Engine/Source/Runtime/Renderer/Private/VT/VirtualTextureProducer.h
@@ -44,7 +44,7 @@ public:
 	void ReleaseProducer(FVirtualTextureSystem* System, const FVirtualTextureProducerHandle& Handle);
 	
 	void AddDestroyedCallback(const FVirtualTextureProducerHandle& Handle, FVTProducerDestroyedFunction* Function, void* Baton);
-	void RemoveAllCallbacks(const void* Baton);
+	uint32 RemoveAllCallbacks(const void* Baton);
 	void CallPendingCallbacks();
 
 	/**
diff --git a/Engine/Source/Runtime/Renderer/Private/VT/VirtualTextureSystem.cpp b/Engine/Source/Runtime/Renderer/Private/VT/VirtualTextureSystem.cpp
--- a/Engine/Source/Runtime/Renderer/Private/VT/VirtualTextureSystem.cpp
+++ b/Engine/Source/Runtime/Renderer/Private/VT/VirtualTextureSystem.cpp
@@ -481,9 +481,9 @@ void FVirtualTextureSystem::AddProducerDestroyedCallback(const FVirtualTexturePr
 	Producers.AddDestroyedCallback(Handle, Function, Baton);
 }
 
-void FVirtualTextureSystem::RemoveAllProducerDestroyedCallbacks(const void* Baton)
+uint32 FVirtualTextureSystem::RemoveAllProducerDestroyedCallbacks(const void* Baton)
 {
-	Producers.RemoveAllCallbacks(Baton);
+	return Producers.RemoveAllCallbacks(Baton);
 }
 
 FVirtualTextureProducer* FVirtualTextureSystem::FindProducer(const FVirtualTextureProducerHandle& Handle)
diff --git a/Engine/Source/Runtime/Renderer/Private/VT/VirtualTextureSystem.h b/Engine/Source/Runtime/Renderer/Private/VT/VirtualTextureSystem.h
--- a/Engine/Source/Runtime/Renderer/Private/VT/VirtualTextureSystem.h
+++ b/Engine/Source/Runtime/Renderer/Private/VT/VirtualTextureSystem.h
@@ -44,7 +44,7 @@ public:
 	FVirtualTextureProducerHandle RegisterProducer(const FVTProducerDescription& InDesc, IVirtualTexture* InProducer);
 	void ReleaseProducer(const FVirtualTextureProducerHandle& Handle);
 	void AddProducerDestroyedCallback(const FVirtualTextureProducerHandle& Handle, FVTProducerDestroyedFunction* Function, void* Baton);
-	void RemoveAllProducerDestroyedCallbacks(const void* Baton);
+	uint32 RemoveAllProducerDestroyedCallbacks(const void* Baton);
 	FVirtualTextureProducer* FindProducer(const FVirtualTextureProducerHandle& Handle);
 
 	FVirtualTextureSpace* AcquireSpace(const FVTSpaceDescription& InDesc, uint32 InSizeNeeded);
