Fix another case where material can be destroyed without removing VT callbacks

diff --git a/Engine/Source/Runtime/Engine/Private/Materials/MaterialShared.cpp b/Engine/Source/Runtime/Engine/Private/Materials/MaterialShared.cpp
--- a/Engine/Source/Runtime/Engine/Private/Materials/MaterialShared.cpp
+++ b/Engine/Source/Runtime/Engine/Private/Materials/MaterialShared.cpp
@@ -2528,7 +2528,12 @@ FMaterialRenderProxy::~FMaterialRenderProxy()
 		ReleaseResource();
 	}
 
-	check(!HasVirtualTextureCallbacks);
+	if (HasVirtualTextureCallbacks)
+	{
+		check(IsInRenderingThread());
+		GetRendererModule().RemoveAllVirtualTextureProducerDestroyedCallbacks(this);
+		HasVirtualTextureCallbacks = false;
+	}
 
 	DeletedFlag = 1;
 }
