Initially clear the SkinTangentIntermediate buffer. Previous code relied on typical (but not guaranteed) DX11 behavior where new allocations return zeroed memory. This resulted in artifacts on DX12 and Vulkan.

diff --git a/Engine/Source/Runtime/Engine/Private/GPUSkinCache.cpp b/Engine/Source/Runtime/Engine/Private/GPUSkinCache.cpp
--- a/Engine/Source/Runtime/Engine/Private/GPUSkinCache.cpp
+++ b/Engine/Source/Runtime/Engine/Private/GPUSkinCache.cpp
@@ -902,12 +902,17 @@ void FGPUSkinCache::DispatchUpdateSkinTangents(FRHICommandListImmediate& RHICmdL
 		uint32 NumIntsPerBuffer = DispatchData.NumTriangles * 3 * FGPUSkinCache::IntermediateAccumBufferNumInts;
 		CurrentStagingBufferIndex = (CurrentStagingBufferIndex + 1) % StagingBuffers.Num();
 		FRWBuffer& StagingBuffer = StagingBuffers[CurrentStagingBufferIndex];
-		if (StagingBuffers[CurrentStagingBufferIndex].NumBytes < NumIntsPerBuffer * sizeof(uint32))
+		if (StagingBuffer.NumBytes < NumIntsPerBuffer * sizeof(uint32))
 		{
 			StagingBuffer.Release();
-			StagingBuffer.Initialize(sizeof(int32), NumIntsPerBuffer, PF_R32_SINT, BUF_UnorderedAccess);
+			StagingBuffer.Initialize(sizeof(int32), NumIntsPerBuffer, PF_R32_SINT, BUF_UnorderedAccess, TEXT("SkinTangentIntermediate"));
 			RHICmdList.BindDebugLabelName(StagingBuffer.UAV, TEXT("SkinTangentIntermediate"));
-			SET_MEMORY_STAT(STAT_GPUSkinCache_TangentsIntermediateMemUsed, (NumIntsPerBuffer * sizeof(uint32)));
+
+			const uint32 MemSize = NumIntsPerBuffer * sizeof(uint32);
+			SET_MEMORY_STAT(STAT_GPUSkinCache_TangentsIntermediateMemUsed, MemSize);
+
+			// The UAV must be zero-filled. We leave it zeroed after each round (see RecomputeTangentsPerVertexPass.usf), so this is only needed on when the buffer is first created.
+			ClearUAV(RHICmdList, StagingBuffer.UAV, MemSize, 0);
 		}
 
 		// This code can be optimized by batched up and doing it with less Dispatch calls (costs more memory)
@@ -926,27 +931,27 @@ void FGPUSkinCache::DispatchUpdateSkinTangents(FRHICommandListImmediate& RHICmdL
 
 			FBaseRecomputeTangents* Shader = 0;
 
-            if (bFullPrecisionUV)
-            {
-                if (DispatchData.bExtraBoneInfluences)
-                {
-                    Shader = GAllowDupedVertsForRecomputeTangents ? (FBaseRecomputeTangents*)*ComputeShader100 : (FBaseRecomputeTangents*)*ComputeShader101;
-                }
-                else
-                {
-                    Shader = GAllowDupedVertsForRecomputeTangents ? (FBaseRecomputeTangents*)*ComputeShader000 : (FBaseRecomputeTangents*)*ComputeShader001;
-                }
-            }
-            else
+			if (bFullPrecisionUV)
+			{
+				if (DispatchData.bExtraBoneInfluences)
+				{
+					Shader = GAllowDupedVertsForRecomputeTangents ? (FBaseRecomputeTangents*)*ComputeShader100 : (FBaseRecomputeTangents*)*ComputeShader101;
+				}
+				else
+				{
+					Shader = GAllowDupedVertsForRecomputeTangents ? (FBaseRecomputeTangents*)*ComputeShader000 : (FBaseRecomputeTangents*)*ComputeShader001;
+				}
+			}
+			else
 			{
-                if (DispatchData.bExtraBoneInfluences)
-                {
-                    Shader = GAllowDupedVertsForRecomputeTangents ? (FBaseRecomputeTangents*)*ComputeShader110 : (FBaseRecomputeTangents*)*ComputeShader111;
-                }
-                else
-                {
-                    Shader = GAllowDupedVertsForRecomputeTangents ? (FBaseRecomputeTangents*)*ComputeShader010 : (FBaseRecomputeTangents*)*ComputeShader011;
-                }
+				if (DispatchData.bExtraBoneInfluences)
+				{
+					Shader = GAllowDupedVertsForRecomputeTangents ? (FBaseRecomputeTangents*)*ComputeShader110 : (FBaseRecomputeTangents*)*ComputeShader111;
+				}
+				else
+				{
+					Shader = GAllowDupedVertsForRecomputeTangents ? (FBaseRecomputeTangents*)*ComputeShader010 : (FBaseRecomputeTangents*)*ComputeShader011;
+				}
 			}
 
 			check(Shader);
@@ -962,12 +967,12 @@ void FGPUSkinCache::DispatchUpdateSkinTangents(FRHICommandListImmediate& RHICmdL
 
 			RHICmdList.TransitionResource(EResourceTransitionAccess::ERWNoBarrier, EResourceTransitionPipeline::EComputeToCompute, StagingBuffer.UAV.GetReference());
 
-            if (!GAllowDupedVertsForRecomputeTangents)
-            {
-                check(LodData.RenderSections[SectionIndex].DuplicatedVerticesBuffer.DupVertData.Num() && LodData.RenderSections[SectionIndex].DuplicatedVerticesBuffer.DupVertIndexData.Num());
-                DispatchData.DuplicatedIndices = LodData.RenderSections[SectionIndex].DuplicatedVerticesBuffer.DuplicatedVerticesIndexBuffer.VertexBufferSRV;
-                DispatchData.DuplicatedIndicesIndices = LodData.RenderSections[SectionIndex].DuplicatedVerticesBuffer.LengthAndIndexDuplicatedVerticesIndexBuffer.VertexBufferSRV;
-            }
+			if (!GAllowDupedVertsForRecomputeTangents)
+			{
+				check(LodData.RenderSections[SectionIndex].DuplicatedVerticesBuffer.DupVertData.Num() && LodData.RenderSections[SectionIndex].DuplicatedVerticesBuffer.DupVertIndexData.Num());
+				DispatchData.DuplicatedIndices = LodData.RenderSections[SectionIndex].DuplicatedVerticesBuffer.DuplicatedVerticesIndexBuffer.VertexBufferSRV;
+				DispatchData.DuplicatedIndicesIndices = LodData.RenderSections[SectionIndex].DuplicatedVerticesBuffer.LengthAndIndexDuplicatedVerticesIndexBuffer.VertexBufferSRV;
+			}
 
 			INC_DWORD_STAT_BY(STAT_GPUSkinCache_NumTrianglesForRecomputeTangents, NumTriangles);
 			Shader->SetParameters(RHICmdList, Entry, DispatchData, StagingBuffer);
diff --git a/Engine/Source/Runtime/Engine/Public/GPUSkinCache.h b/Engine/Source/Runtime/Engine/Public/GPUSkinCache.h
--- a/Engine/Source/Runtime/Engine/Public/GPUSkinCache.h
+++ b/Engine/Source/Runtime/Engine/Public/GPUSkinCache.h
@@ -158,11 +158,11 @@ public:
 		{
 			for (int32 Index = 0; Index < NUM_BUFFERS; ++Index)
 			{
-				RWBuffers[Index].Initialize(4, NumVertices * 3, PF_R32_FLOAT, BUF_Static);
+				RWBuffers[Index].Initialize(4, NumVertices * 3, PF_R32_FLOAT, BUF_Static, TEXT("Vertices"));
 			}
 			if (WithTangents)
 			{
-				Tangents.Initialize(8, NumVertices * 2, PF_R16G16B16A16_SNORM, BUF_Static);
+				Tangents.Initialize(8, NumVertices * 2, PF_R16G16B16A16_SNORM, BUF_Static, TEXT("Tangents"));
 			}
 		}
 
