Use virtual texture producer callbacks to notify material changes in runtime virtual texture

diff --git a/Engine/Source/Runtime/Engine/Private/Materials/MaterialShared.cpp b/Engine/Source/Runtime/Engine/Private/Materials/MaterialShared.cpp
--- a/Engine/Source/Runtime/Engine/Private/Materials/MaterialShared.cpp
+++ b/Engine/Source/Runtime/Engine/Private/Materials/MaterialShared.cpp
@@ -2285,20 +2285,33 @@ void FMaterialVirtualTextureStack::Serialize(FArchive& Ar)
 	FMaterialRenderProxy
 -----------------------------------------------------------------------------*/
 
-static IAllocatedVirtualTexture* GetPreallocatedVTStack(const FMaterialRenderContext& Context, const FUniformExpressionSet& UniformExpressionSet, const FMaterialVirtualTextureStack& VTStack)
+static void OnVirtualTextureDestroyed(const FVirtualTextureProducerHandle& InHandle, void* Baton)
+{
+	FMaterialRenderProxy* MaterialProxy = static_cast<FMaterialRenderProxy*>(Baton);
+
+	MaterialProxy->InvalidateUniformExpressionCache(false);
+	UMaterialInterface::IterateOverActiveFeatureLevels([&](ERHIFeatureLevel::Type InFeatureLevel)
+	{
+		MaterialProxy->UpdateUniformExpressionCacheIfNeeded(InFeatureLevel);
+	});
+}
+
+IAllocatedVirtualTexture* FMaterialRenderProxy::GetPreallocatedVTStack(const FMaterialRenderContext& Context, const FUniformExpressionSet& UniformExpressionSet, const FMaterialVirtualTextureStack& VTStack) const
 {
 	check(VTStack.IsPreallocatedStack())
 
 	URuntimeVirtualTexture const* Texture;
 	VTStack.GetTextureValue(Context, UniformExpressionSet, Texture);
 
-	return (Texture == nullptr) ? nullptr : Texture->GetAllocatedVirtualTexture();
-}
+	if (Texture == nullptr)
+	{
+		return nullptr;
+	}
 
-static void OnVirtualTextureDestroyed(const FVirtualTextureProducerHandle& InHandle, void* Baton)
-{
-	FMaterialRenderProxy* MaterialProxy = static_cast<FMaterialRenderProxy*>(Baton);
-	MaterialProxy->CacheUniformExpressions(false);
+	GetRendererModule().AddVirtualTextureProducerDestroyedCallback(Texture->GetProducerHandle(), &OnVirtualTextureDestroyed, const_cast<FMaterialRenderProxy*>(this));
+	HasVirtualTextureCallbacks = true;
+
+	return Texture->GetAllocatedVirtualTexture();
 }
 
 IAllocatedVirtualTexture* FMaterialRenderProxy::AllocateVTStack(const FMaterialRenderContext& Context, const FUniformExpressionSet& UniformExpressionSet, const FMaterialVirtualTextureStack& VTStack) const
diff --git a/Engine/Source/Runtime/Engine/Public/MaterialShared.h b/Engine/Source/Runtime/Engine/Public/MaterialShared.h
--- a/Engine/Source/Runtime/Engine/Public/MaterialShared.h
+++ b/Engine/Source/Runtime/Engine/Public/MaterialShared.h
@@ -1965,6 +1965,7 @@ public:
 	}
 
 private:
+	IAllocatedVirtualTexture* GetPreallocatedVTStack(const FMaterialRenderContext& Context, const FUniformExpressionSet& UniformExpressionSet, const FMaterialVirtualTextureStack& VTStack) const;
 	IAllocatedVirtualTexture* AllocateVTStack(const FMaterialRenderContext& Context, const FUniformExpressionSet& UniformExpressionSet, const FMaterialVirtualTextureStack& VTStack) const;
 
 	/** 0 if not set, game thread pointer, do not dereference, only for comparison */
