Remove slow NotifyMaterial calls for runtime virtual texture

diff --git a/Engine/Source/Runtime/Engine/Classes/VT/RuntimeVirtualTexturePlane.h b/Engine/Source/Runtime/Engine/Classes/VT/RuntimeVirtualTexturePlane.h
--- a/Engine/Source/Runtime/Engine/Classes/VT/RuntimeVirtualTexturePlane.h
+++ b/Engine/Source/Runtime/Engine/Classes/VT/RuntimeVirtualTexturePlane.h
@@ -51,9 +51,6 @@ private:
 	UPROPERTY(EditAnywhere, DuplicateTransient, Category = TransformFromBounds, meta = (DisplayName = "Source Actor"))
 	AActor* BoundsSourceActor = nullptr;
 
-	/** Flag used for deferred material notification after render state changes. */
-	bool bNotifyInNextTick;
-
 public:
 	/** Get the runtime virtual texture object set on this component */
 	URuntimeVirtualTexture* GetVirtualTexture() const { return VirtualTexture; }
@@ -67,15 +64,10 @@ public:
 #endif
 
 protected:
-	/** Apply any deferred material notifications. */
-	void NotifyMaterials();
-
 	//~ Begin UActorComponent Interface
 	virtual void CreateRenderState_Concurrent() override;
 	virtual void SendRenderTransform_Concurrent() override;
 	virtual void DestroyRenderState_Concurrent() override;
-	virtual void TickComponent(float DeltaTime, ELevelTick TickType, FActorComponentTickFunction* ThisTickFunction) override;
-	virtual void OnUnregister() override;
 	//~ End UActorComponent Interface
 
 public:
diff --git a/Engine/Source/Runtime/Engine/Private/VT/RuntimeVirtualTextureNotify.cpp b/Engine/Source/Runtime/Engine/Private/VT/RuntimeVirtualTextureNotify.cpp
--- a/Engine/Source/Runtime/Engine/Private/VT/RuntimeVirtualTextureNotify.cpp
+++ b/Engine/Source/Runtime/Engine/Private/VT/RuntimeVirtualTextureNotify.cpp
@@ -3,8 +3,6 @@
 #include "VT/RuntimeVirtualTextureNotify.h"
 
 #include "Components/PrimitiveComponent.h"
-#include "MaterialShared.h"
-#include "Materials/Material.h"
 #include "UObject/UObjectIterator.h"
 #include "VT/RuntimeVirtualTexture.h"
 #include "VT/RuntimeVirtualTexturePlane.h"
@@ -40,46 +38,4 @@ namespace RuntimeVirtualTexture
 	}
 
 #endif
-
-	void NotifyMaterials(URuntimeVirtualTexture const* VirtualTexture)
-	{
-		//todo [vt]:
-		// This operation is very slow! We should only do it during edition, but currently we do it once per virtual texture at runtime load/unload.
-		// Can we pre-calculate the list of materials to touch during cook? But even then touching here them will be extra work...
-		// Is there a way to set up dependencies so that we load the materials _after_ the virtual texture is allocated?
-		// Or maybe we should consider serializing the WorldToUVTransform in the URuntimeVirtualTexture and not depending on a URuntimeVirtualTextureComponent at runtime?
-
-		if (VirtualTexture == nullptr)
-		{
-			return;
-		}
-
-		TSet<UMaterial*> BaseMaterialsThatUseThisTexture;
-		for (TObjectIterator<UMaterialInterface> It; It; ++It)
-		{
-			UMaterialInterface* MaterialInterface = *It;
-
-			TArray<UObject*> Textures;
-			MaterialInterface->AppendReferencedTextures(Textures);
-
-			for (auto Texture : Textures)
-			{
-				if (Texture == VirtualTexture)
-				{
-					BaseMaterialsThatUseThisTexture.Add(MaterialInterface->GetMaterial());
-					break;
-				}
-			}
-		}
-
-		if (BaseMaterialsThatUseThisTexture.Num() > 0)
-		{
-			FMaterialUpdateContext UpdateContext(FMaterialUpdateContext::EOptions::SyncWithRenderingThread);
-			for (TSet<UMaterial*>::TConstIterator It(BaseMaterialsThatUseThisTexture); It; ++It)
-			{
-				(*It)->RecacheUniformExpressions(false);
-				UpdateContext.AddMaterial(*It);
-			}
-		}
-	}
 }
diff --git a/Engine/Source/Runtime/Engine/Private/VT/RuntimeVirtualTextureNotify.h b/Engine/Source/Runtime/Engine/Private/VT/RuntimeVirtualTextureNotify.h
--- a/Engine/Source/Runtime/Engine/Private/VT/RuntimeVirtualTextureNotify.h
+++ b/Engine/Source/Runtime/Engine/Private/VT/RuntimeVirtualTextureNotify.h
@@ -21,10 +21,4 @@ namespace RuntimeVirtualTexture
 	 */
 	void NotifyPrimitives(URuntimeVirtualTexture const* VirtualTexture);
 #endif
-
-	/**
-	 * Find materials referencing this virtual texture and re-cache the uniforms.
-	 * We need to do this after any operation that reallocates the virtual texture since the material caches info about the VT allocation in it's uniform buffer.
-	 */
-	void NotifyMaterials(URuntimeVirtualTexture const* VirtualTexture);
 }
\ No newline at end of file
diff --git a/Engine/Source/Runtime/Engine/Private/VT/RuntimeVirtualTexturePlane.cpp b/Engine/Source/Runtime/Engine/Private/VT/RuntimeVirtualTexturePlane.cpp
--- a/Engine/Source/Runtime/Engine/Private/VT/RuntimeVirtualTexturePlane.cpp
+++ b/Engine/Source/Runtime/Engine/Private/VT/RuntimeVirtualTexturePlane.cpp
@@ -26,7 +26,6 @@ ARuntimeVirtualTexturePlane::ARuntimeVirtualTexturePlane(const FObjectInitialize
 
 URuntimeVirtualTextureComponent::URuntimeVirtualTextureComponent(const FObjectInitializer& ObjectInitializer)
 	: Super(ObjectInitializer)
-	, bNotifyInNextTick(false)
 	, SceneProxy(nullptr)
 {
 	PrimaryComponentTick.bCanEverTick = true;
@@ -39,7 +38,6 @@ void URuntimeVirtualTextureComponent::CreateRenderState_Concurrent()
 	{
 		// This will modify the URuntimeVirtualTexture and allocate its VT
 		GetScene()->AddRuntimeVirtualTexture(this);
-		bNotifyInNextTick = true;
 	}
 
 	Super::CreateRenderState_Concurrent();
@@ -51,7 +49,6 @@ void URuntimeVirtualTextureComponent::SendRenderTransform_Concurrent()
 	{
 		// This will modify the URuntimeVirtualTexture and allocate its VT
 		GetScene()->AddRuntimeVirtualTexture(this);
-		bNotifyInNextTick = true;
 	}
 
 	Super::SendRenderTransform_Concurrent();
@@ -61,33 +58,10 @@ void URuntimeVirtualTextureComponent::DestroyRenderState_Concurrent()
 {
 	// This will modify the URuntimeVirtualTexture and free its VT
 	GetScene()->RemoveRuntimeVirtualTexture(this);
-	bNotifyInNextTick = true;
 
 	Super::DestroyRenderState_Concurrent();
 }
 
-void URuntimeVirtualTextureComponent::NotifyMaterials()
-{
-	if (bNotifyInNextTick)
-	{
-		// Notify materials of reallocation of the virtual texture caused by render state update.
-		// This is slow and ideally we will find a different approach.
-		RuntimeVirtualTexture::NotifyMaterials(VirtualTexture);
-		bNotifyInNextTick = false;
-	}
-}
-
-void URuntimeVirtualTextureComponent::TickComponent(float DeltaTime, ELevelTick TickType, FActorComponentTickFunction* ThisTickFunction)
-{
-	NotifyMaterials();
-}
-
-void URuntimeVirtualTextureComponent::OnUnregister()
-{
-	NotifyMaterials();
-	Super::OnUnregister();
-}
-
 #if WITH_EDITOR
 
 void URuntimeVirtualTextureComponent::SetRotation()
