Add specific logic for whether to interpret physical texture as sRGB in runtime virtual texture
Also tidied up some comments in various bits of runtime VT code

diff --git a/Engine/Source/Editor/LandscapeEditor/Private/LandscapeEdModePaintTools.cpp b/Engine/Source/Editor/LandscapeEditor/Private/LandscapeEdModePaintTools.cpp
--- a/Engine/Source/Editor/LandscapeEditor/Private/LandscapeEdModePaintTools.cpp
+++ b/Engine/Source/Editor/LandscapeEditor/Private/LandscapeEdModePaintTools.cpp
@@ -261,10 +261,10 @@ public:
 		ALandscape* Landscape = LandscapeInfo->LandscapeActor.Get();
 
 		// If we render to a runtime virtual texture then we need to flush here
-		//todo[vt]: This logic doesn't cover material overrides etc.
 		if (Landscape != nullptr && Landscape->RuntimeVirtualTextures.Num() > 0)
 		{
 			//todo[vt]: Only flush Bounds 
+			//todo[vt]: Only flush specific virtual textures
 			GetRendererModule().FlushVirtualTextureCache();
 		}
 	}
diff --git a/Engine/Source/Runtime/Engine/Classes/VT/RuntimeVirtualTexture.h b/Engine/Source/Runtime/Engine/Classes/VT/RuntimeVirtualTexture.h
--- a/Engine/Source/Runtime/Engine/Classes/VT/RuntimeVirtualTexture.h
+++ b/Engine/Source/Runtime/Engine/Classes/VT/RuntimeVirtualTexture.h
@@ -73,6 +73,8 @@ public:
 
 	/** Get virtual texture description based on the properties of this object. */
 	void GetProducerDescription(FVTProducerDescription& OutDesc) const;
+	/** Return true if the virtual texture layer should be sampled as sRGB */
+	bool IsLayerSRGB(int32 LayerIndex) const;
 
 	/** (Re)Initialize this object. Call this whenever we modify the producer or transform. */
 	void Initialize(IVirtualTexture* InProducer, FTransform const& BoxToWorld);
diff --git a/Engine/Source/Runtime/Engine/Private/Materials/MaterialUniformExpressions.cpp b/Engine/Source/Runtime/Engine/Private/Materials/MaterialUniformExpressions.cpp
--- a/Engine/Source/Runtime/Engine/Private/Materials/MaterialUniformExpressions.cpp
+++ b/Engine/Source/Runtime/Engine/Private/Materials/MaterialUniformExpressions.cpp
@@ -822,11 +822,11 @@ void FUniformExpressionSet::FillUniformBuffer(const FMaterialRenderContext& Mate
 					if (AllocatedVT != nullptr)
 					{
 						const int32 LayerIndex = UniformVirtualTextureExpressions[ExpressionIndex]->GetLayerIndex();
-						FRHIShaderResourceView* PhysicalViewRHI = AllocatedVT->GetPhysicalTextureView(LayerIndex, LayerIndex == 0); // todo need sRGB flag!
+						FRHIShaderResourceView* PhysicalViewRHI = AllocatedVT->GetPhysicalTextureView(LayerIndex, Texture->IsLayerSRGB(LayerIndex));
 						if (PhysicalViewRHI != nullptr)
 						{
 							*ResourceTablePhysicalTexturePtr = PhysicalViewRHI;
-							*ResourceTablePhysicalSamplerPtr = TStaticSamplerState<SF_AnisotropicLinear, AM_Clamp, AM_Clamp, AM_Clamp, 0, 8>::GetRHI(); // todo deal with sampler state
+							*ResourceTablePhysicalSamplerPtr = TStaticSamplerState<SF_Bilinear, AM_Clamp, AM_Clamp, AM_Clamp, 0, 8>::GetRHI();
 							bValidResources = true;
 						}
 					}
@@ -836,7 +836,7 @@ void FUniformExpressionSet::FillUniformBuffer(const FMaterialRenderContext& Mate
 			if (!bValidResources)
 			{
 				*ResourceTablePhysicalTexturePtr = GWhiteTextureWithSRV->ShaderResourceViewRHI;
-				*ResourceTablePhysicalSamplerPtr = TStaticSamplerState<SF_AnisotropicLinear, AM_Clamp, AM_Clamp, AM_Clamp, 0, 8>::GetRHI();
+				*ResourceTablePhysicalSamplerPtr = TStaticSamplerState<SF_Bilinear, AM_Clamp, AM_Clamp, AM_Clamp, 0, 8>::GetRHI();
 			}
 		}
 
diff --git a/Engine/Source/Runtime/Engine/Private/PrimitiveSceneProxy.cpp b/Engine/Source/Runtime/Engine/Private/PrimitiveSceneProxy.cpp
--- a/Engine/Source/Runtime/Engine/Private/PrimitiveSceneProxy.cpp
+++ b/Engine/Source/Runtime/Engine/Private/PrimitiveSceneProxy.cpp
@@ -219,7 +219,7 @@ FPrimitiveSceneProxy::FPrimitiveSceneProxy(const UPrimitiveComponent* InComponen
 				RuntimeVirtualTextures.Add(VirtualTexture);
 				RuntimeVirtualTextureMaterialTypes.Add(VirtualTexture->GetMaterialType());
 			
-				//todo[vt]: Only flush this specific virtual textures
+				//todo[vt]: Only flush this specific virtual texture
 				//todo[vt]: Only flush primitive bounds 
 				GetRendererModule().FlushVirtualTextureCache();
 			}
diff --git a/Engine/Source/Runtime/Engine/Private/StaticMeshRender.cpp b/Engine/Source/Runtime/Engine/Private/StaticMeshRender.cpp
--- a/Engine/Source/Runtime/Engine/Private/StaticMeshRender.cpp
+++ b/Engine/Source/Runtime/Engine/Private/StaticMeshRender.cpp
@@ -1057,7 +1057,7 @@ void FStaticMeshSceneProxy::DrawStaticElements(FStaticPrimitiveDrawInterface* PD
 		for (ERuntimeVirtualTextureMaterialType MaterialType : RuntimeVirtualTextureMaterialTypes)
 		{
 			// Use the minimal LOD which is usually (always?) what we would want.
-			//todo[vt]: Add user control for to allow use of a specific LOD in the virtual texture.
+			//todo[vt]: Add user control to allow use of a specific LOD in the virtual texture.
 			// Note that ComputeStaticMeshLOD() checks *all* mesh batches for LOD calculation (including these VT ones) and assumes that they are stored in LOD order.
 			// Adding another set of last LODs here should be OK but if we want to use a different LOD they would need adding in the mesh batches in the above loop to preserve LOD order.
 			const int32 LODIndex = NumLODs - 1;
diff --git a/Engine/Source/Runtime/Engine/Private/VT/RuntimeVirtualTexture.cpp b/Engine/Source/Runtime/Engine/Private/VT/RuntimeVirtualTexture.cpp
--- a/Engine/Source/Runtime/Engine/Private/VT/RuntimeVirtualTexture.cpp
+++ b/Engine/Source/Runtime/Engine/Private/VT/RuntimeVirtualTexture.cpp
@@ -212,6 +212,23 @@ void URuntimeVirtualTexture::GetProducerDescription(FVTProducerDescription& OutD
 	}
 }
 
+bool URuntimeVirtualTexture::IsLayerSRGB(int32 LayerIndex) const
+{
+	switch (MaterialType)
+	{
+	case ERuntimeVirtualTextureMaterialType::BaseColor:
+	case ERuntimeVirtualTextureMaterialType::BaseColor_Normal:
+	case ERuntimeVirtualTextureMaterialType::BaseColor_Normal_Specular:
+		// Only BaseColor layer is sRGB
+		return LayerIndex == 0;
+	default:
+		// Implement logic for any missing material types
+		check(false);
+	}
+
+	return false;
+}
+
 int32 URuntimeVirtualTexture::GetEstimatedPageTableTextureMemoryKb() const
 {
 	//todo[vt]: Estimate memory usage
diff --git a/Engine/Source/Runtime/Foliage/Private/InstancedFoliage.cpp b/Engine/Source/Runtime/Foliage/Private/InstancedFoliage.cpp
--- a/Engine/Source/Runtime/Foliage/Private/InstancedFoliage.cpp
+++ b/Engine/Source/Runtime/Foliage/Private/InstancedFoliage.cpp
@@ -1364,7 +1364,6 @@ void FFoliageStaticMesh::UpdateComponentSettings(const UFoliageType_InstancedSta
 		}
 		if (Component->TranslucencySortPriority != FoliageType->TranslucencySortPriority)
 		{
-			//todo[vt]: Trigger a virtual texture flush somewhere withing the update chain
 			Component->TranslucencySortPriority = FoliageType->TranslucencySortPriority;
 			bNeedsMarkRenderStateDirty = true;
 		}
diff --git a/Engine/Source/Runtime/Landscape/Private/LandscapeEdit.cpp b/Engine/Source/Runtime/Landscape/Private/LandscapeEdit.cpp
--- a/Engine/Source/Runtime/Landscape/Private/LandscapeEdit.cpp
+++ b/Engine/Source/Runtime/Landscape/Private/LandscapeEdit.cpp
@@ -597,6 +597,7 @@ void ULandscapeComponent::PostEditUndo()
 
 	if (GetLandscapeProxy()->RuntimeVirtualTextures.Num() > 0)
 	{
+		//todo[vt]: Only flush this specific virtual textures
 		//todo[vt]: Only flush Bounds 
 		GetRendererModule().FlushVirtualTextureCache();
 	}
diff --git a/Engine/Source/Runtime/Renderer/Private/VT/RuntimeVirtualTextureProducer.cpp b/Engine/Source/Runtime/Renderer/Private/VT/RuntimeVirtualTextureProducer.cpp
--- a/Engine/Source/Runtime/Renderer/Private/VT/RuntimeVirtualTextureProducer.cpp
+++ b/Engine/Source/Runtime/Renderer/Private/VT/RuntimeVirtualTextureProducer.cpp
@@ -17,9 +17,8 @@ FRuntimeVirtualTextureFinalizer::FRuntimeVirtualTextureFinalizer(FVTProducerDesc
 bool FRuntimeVirtualTextureFinalizer::IsReady()
 {
 	//todo[vt]: 
-	// Test if we have everything we need to render (shaders loaded etc).
-	// Current test for GPUScene.PrimitiveBuffer is a nasty thing to prevent a checkf triggering if no PrimitiveBuffer is bound. It feels like it requires too much knowledge of the renderer internals...
-	return Scene != nullptr && Scene->GetRenderScene() != nullptr && Scene->GetRenderScene()->GPUScene.PrimitiveBuffer.Buffer != nullptr;
+	// Test if we have everything we need to render (mips loaded, shaders loaded etc).
+	return Scene != nullptr && Scene->GetRenderScene() != nullptr;
 }
 
 void FRuntimeVirtualTextureFinalizer::AddTile(FTileEntry& Tile)
