Rename Plane->Volume for RuntimeVirtualTexture actors

diff --git a/Engine/Source/Editor/VirtualTexturingEditor/Classes/RuntimeVirtualTextureDetailsCustomization.h b/Engine/Source/Editor/VirtualTexturingEditor/Classes/RuntimeVirtualTextureDetailsCustomization.h
--- a/Engine/Source/Editor/VirtualTexturingEditor/Classes/RuntimeVirtualTextureDetailsCustomization.h
+++ b/Engine/Source/Editor/VirtualTexturingEditor/Classes/RuntimeVirtualTextureDetailsCustomization.h
@@ -37,7 +37,7 @@ private:
 };
 
 
-/** UI customization for ARuntimeVirtualTexturePlane */
+/** UI customization for URuntimeVirtualTextureComponent */
 class FRuntimeVirtualTextureComponentDetailsCustomization : public IDetailCustomization
 {
 public:
diff --git a/Engine/Source/Editor/VirtualTexturingEditor/Classes/RuntimeVirtualTexturePlaneActorFactory.h b/Engine/Source/Editor/VirtualTexturingEditor/Classes/RuntimeVirtualTexturePlaneActorFactory.h
deleted file mode 100644
diff --git a/Engine/Source/Editor/VirtualTexturingEditor/Classes/RuntimeVirtualTextureVolumeActorFactory.h b/Engine/Source/Editor/VirtualTexturingEditor/Classes/RuntimeVirtualTextureVolumeActorFactory.h
new file mode 100644
--- /dev/null
+++ b/Engine/Source/Editor/VirtualTexturingEditor/Classes/RuntimeVirtualTextureVolumeActorFactory.h
@@ -0,0 +1,18 @@
+// Copyright 1998-2019 Epic Games, Inc. All Rights Reserved.
+
+#pragma once
+
+#include "CoreMinimal.h"
+#include "ActorFactories/ActorFactory.h"
+#include "RuntimeVirtualTextureVolumeActorFactory.generated.h"
+
+/** Actor factory for ARuntimeVirtualTextureVolume */
+UCLASS(MinimalAPI, config=Editor, collapsecategories, hidecategories=Object)
+class URuntimeVirtualTextureVolumeActorFactory : public UActorFactory
+{
+	GENERATED_UCLASS_BODY()
+
+	//~ Begin UActorFactory Interface
+	virtual void PostSpawnActor(UObject* Asset, AActor* NewActor) override;
+	//~ End UActorFactory Interface
+};
diff --git a/Engine/Source/Editor/VirtualTexturingEditor/Private/RuntimeVirtualTextureDetailsCustomization.cpp b/Engine/Source/Editor/VirtualTexturingEditor/Private/RuntimeVirtualTextureDetailsCustomization.cpp
--- a/Engine/Source/Editor/VirtualTexturingEditor/Private/RuntimeVirtualTextureDetailsCustomization.cpp
+++ b/Engine/Source/Editor/VirtualTexturingEditor/Private/RuntimeVirtualTextureDetailsCustomization.cpp
@@ -2,12 +2,12 @@
 
 #include "RuntimeVirtualTextureDetailsCustomization.h"
 
+#include "Components/RuntimeVirtualTextureComponent.h"
 #include "DetailCategoryBuilder.h"
 #include "DetailLayoutBuilder.h"
 #include "DetailWidgetRow.h"
 #include "SResetToDefaultMenu.h"
 #include "VT/RuntimeVirtualTexture.h"
-#include "VT/RuntimeVirtualTexturePlane.h"
 #include "Widgets/Input/SButton.h"
 #include "Widgets/Layout/SBox.h"
 #include "Widgets/Layout/SWrapBox.h"
diff --git a/Engine/Source/Editor/VirtualTexturingEditor/Private/RuntimeVirtualTexturePlaneActorFactory.cpp b/Engine/Source/Editor/VirtualTexturingEditor/Private/RuntimeVirtualTexturePlaneActorFactory.cpp
deleted file mode 100644
diff --git a/Engine/Source/Editor/VirtualTexturingEditor/Private/RuntimeVirtualTextureVolumeActorFactory.cpp b/Engine/Source/Editor/VirtualTexturingEditor/Private/RuntimeVirtualTextureVolumeActorFactory.cpp
new file mode 100644
--- /dev/null
+++ b/Engine/Source/Editor/VirtualTexturingEditor/Private/RuntimeVirtualTextureVolumeActorFactory.cpp
@@ -0,0 +1,28 @@
+// Copyright 1998-2019 Epic Games, Inc. All Rights Reserved.
+
+#include "RuntimeVirtualTextureVolumeActorFactory.h"
+
+#include "VT/RuntimeVirtualTextureVolume.h"
+
+#define LOCTEXT_NAMESPACE "VirtualTexturingEditorModule"
+
+URuntimeVirtualTextureVolumeActorFactory::URuntimeVirtualTextureVolumeActorFactory(const FObjectInitializer& ObjectInitializer)
+	: Super(ObjectInitializer)
+{
+	DisplayName = LOCTEXT("VirtualTextureVolume_DisplayName", "Virtual Texture Volume");
+	NewActorClass = ARuntimeVirtualTextureVolume::StaticClass();
+	bShowInEditorQuickMenu = 1;
+}
+
+void URuntimeVirtualTextureVolumeActorFactory::PostSpawnActor(UObject* Asset, AActor* NewActor)
+{
+	FText ActorName = LOCTEXT("VirtualTextureVolume_DefaultActorName", "Virtual Texture");
+	NewActor->SetActorLabel(ActorName.ToString());
+
+	// Good default size to see object in editor
+	NewActor->SetActorScale3D(FVector(100.f, 100.f, 1.f));
+
+	Super::PostSpawnActor(Asset, NewActor);
+}
+
+#undef LOCTEXT_NAMESPACE
diff --git a/Engine/Source/Runtime/Engine/Classes/Components/RuntimeVirtualTextureComponent.h b/Engine/Source/Runtime/Engine/Classes/Components/RuntimeVirtualTextureComponent.h
new file mode 100644
--- /dev/null
+++ b/Engine/Source/Runtime/Engine/Classes/Components/RuntimeVirtualTextureComponent.h
@@ -0,0 +1,48 @@
+// Copyright 1998-2019 Epic Games, Inc. All Rights Reserved.
+
+#pragma once
+
+#include "CoreMinimal.h"
+#include "SceneComponent.h"
+#include "RuntimeVirtualTextureComponent.generated.h"
+
+class URuntimeVirtualTexture;
+
+/** Component used to place a URuntimeVirtualTexture in the world. */
+UCLASS(ClassGroup = Rendering, collapsecategories, hidecategories = (Activation, Collision, Cooking, Mobility, LOD, Object, Physics, Rendering), editinlinenew)
+class ENGINE_API URuntimeVirtualTextureComponent : public USceneComponent
+{
+	GENERATED_UCLASS_BODY()
+
+private:
+	/** The virtual texture object to use. */
+	UPROPERTY(EditAnywhere, DuplicateTransient, Category = VirtualTexture)
+	URuntimeVirtualTexture* VirtualTexture = nullptr;
+
+	/** Actor to copy the bounds from to set up the transform. */
+	UPROPERTY(EditAnywhere, DuplicateTransient, Category = TransformFromBounds, meta = (DisplayName = "Source Actor"))
+	AActor* BoundsSourceActor = nullptr;
+
+public:
+	/** Get the runtime virtual texture object set on this component */
+	URuntimeVirtualTexture* GetVirtualTexture() const { return VirtualTexture; }
+
+#if WITH_EDITOR
+	/** Copy the rotation from BoundsSourceActor to this component. Called by our UI details customization. */
+	void SetRotation();
+
+	/** Set this component transform to include the BoundsSourceActor bounds. Called by our UI details customization. */
+	void SetTransformToBounds();
+#endif
+
+protected:
+	//~ Begin UActorComponent Interface
+	virtual void CreateRenderState_Concurrent() override;
+	virtual void SendRenderTransform_Concurrent() override;
+	virtual void DestroyRenderState_Concurrent() override;
+	//~ End UActorComponent Interface
+
+public:
+	/** Scene proxy object. Managed by the scene but stored here. */
+	class FRuntimeVirtualTextureSceneProxy* SceneProxy;
+};
diff --git a/Engine/Source/Runtime/Engine/Classes/VT/RuntimeVirtualTexturePlane.h b/Engine/Source/Runtime/Engine/Classes/VT/RuntimeVirtualTexturePlane.h
deleted file mode 100644
diff --git a/Engine/Source/Runtime/Engine/Classes/VT/RuntimeVirtualTextureVolume.h b/Engine/Source/Runtime/Engine/Classes/VT/RuntimeVirtualTextureVolume.h
new file mode 100644
--- /dev/null
+++ b/Engine/Source/Runtime/Engine/Classes/VT/RuntimeVirtualTextureVolume.h
@@ -0,0 +1,33 @@
+// Copyright 1998-2019 Epic Games, Inc. All Rights Reserved.
+
+#pragma once
+
+#include "CoreMinimal.h"
+#include "GameFramework/Actor.h"
+#include "RuntimeVirtualTextureVolume.generated.h"
+
+/** Actor used to place a URuntimeVirtualTexture in the world. */
+UCLASS(hidecategories=(Actor, Collision, Cooking, Input, LOD, Replication), MinimalAPI)
+class ARuntimeVirtualTextureVolume : public AActor
+{
+	GENERATED_UCLASS_BODY()
+
+private:
+	/** Component that owns the runtime virtual texture. */
+	UPROPERTY(VisibleAnywhere, BlueprintReadOnly, Category = VirtualTexture, meta = (AllowPrivateAccess = "true"))
+	class URuntimeVirtualTextureComponent* VirtualTextureComponent;
+
+#if WITH_EDITORONLY_DATA
+	/** Box for visualizing virtual texture extents. */
+	UPROPERTY(Transient)
+	class UBoxComponent* Box = nullptr;
+#endif // WITH_EDITORONLY_DATA
+
+protected:
+	//~ Begin UObject Interface.
+	virtual bool NeedsLoadForServer() const override { return false; }
+	//~ End UObject Interface.
+	//~ Begin AActor Interface.
+	virtual bool IsLevelBoundsRelevant() const override { return false; }
+	//~ End AActor Interface
+};
diff --git a/Engine/Source/Runtime/Engine/Private/Components/RuntimeVirtualTextureComponent.cpp b/Engine/Source/Runtime/Engine/Private/Components/RuntimeVirtualTextureComponent.cpp
new file mode 100644
--- /dev/null
+++ b/Engine/Source/Runtime/Engine/Private/Components/RuntimeVirtualTextureComponent.cpp
@@ -0,0 +1,97 @@
+// Copyright 1998-2019 Epic Games, Inc. All Rights Reserved.
+
+#include "Components/RuntimeVirtualTextureComponent.h"
+
+URuntimeVirtualTextureComponent::URuntimeVirtualTextureComponent(const FObjectInitializer& ObjectInitializer)
+	: Super(ObjectInitializer)
+	, SceneProxy(nullptr)
+{
+	PrimaryComponentTick.bCanEverTick = true;
+	bTickInEditor = true;
+}
+
+void URuntimeVirtualTextureComponent::CreateRenderState_Concurrent()
+{
+	if (ShouldRender() && VirtualTexture != nullptr)
+	{
+		// This will modify the URuntimeVirtualTexture and allocate its VT
+		GetScene()->AddRuntimeVirtualTexture(this);
+	}
+
+	Super::CreateRenderState_Concurrent();
+}
+
+void URuntimeVirtualTextureComponent::SendRenderTransform_Concurrent()
+{
+	if (ShouldRender() && VirtualTexture != nullptr)
+	{
+		// This will modify the URuntimeVirtualTexture and allocate its VT
+		GetScene()->AddRuntimeVirtualTexture(this);
+	}
+
+	Super::SendRenderTransform_Concurrent();
+}
+
+void URuntimeVirtualTextureComponent::DestroyRenderState_Concurrent()
+{
+	// This will modify the URuntimeVirtualTexture and free its VT
+	GetScene()->RemoveRuntimeVirtualTexture(this);
+
+	Super::DestroyRenderState_Concurrent();
+}
+
+#if WITH_EDITOR
+
+void URuntimeVirtualTextureComponent::SetRotation()
+{
+	if (BoundsSourceActor != nullptr)
+	{
+		// Copy the source actor rotation and notify the parent actor
+		SetWorldRotation(BoundsSourceActor->GetTransform().GetRotation());
+		GetOwner()->PostEditMove(true);
+	}
+}
+
+void URuntimeVirtualTextureComponent::SetTransformToBounds()
+{
+	if (BoundsSourceActor != nullptr)
+	{
+		// Calculate the bounds in our local rotation space translated to the BoundsSourceActor center
+		const FQuat TargetRotation = GetComponentToWorld().GetRotation();
+		const FVector InitialPosition = BoundsSourceActor->GetComponentsBoundingBox().GetCenter();
+		const FVector InitialScale = FVector(0.5f, 0.5, 1.f);
+
+		FTransform LocalTransform;
+		LocalTransform.SetComponents(TargetRotation, InitialPosition, InitialScale);
+		FTransform WorldToLocal = LocalTransform.Inverse();
+
+		FBox BoundBox(ForceInit);
+		for (const UActorComponent* Component : BoundsSourceActor->GetComponents())
+		{
+			// Only gather visual components in the bounds calculation
+			const UPrimitiveComponent* PrimitiveComponent = Cast<const UPrimitiveComponent>(Component);
+			if (PrimitiveComponent != nullptr && PrimitiveComponent->IsRegistered())
+			{
+				const FTransform ComponentToActor = PrimitiveComponent->GetComponentTransform() * WorldToLocal;
+				FBoxSphereBounds LocalSpaceComponentBounds = PrimitiveComponent->CalcBounds(ComponentToActor);
+				BoundBox += LocalSpaceComponentBounds.GetBox();
+			}
+		}
+
+		// Create transform from bounds
+		FVector Origin;
+		FVector Extent;
+		BoundBox.GetCenterAndExtents(Origin, Extent);
+
+		Origin = LocalTransform.TransformPosition(Origin);
+
+		FTransform Transform;
+		Transform.SetComponents(TargetRotation, Origin, Extent);
+
+		// Apply final result and notify the parent actor
+		SetWorldTransform(Transform);
+		GetOwner()->PostEditMove(true);
+	}
+}
+
+#endif
diff --git a/Engine/Source/Runtime/Engine/Private/VT/RuntimeVirtualTextureNotify.cpp b/Engine/Source/Runtime/Engine/Private/VT/RuntimeVirtualTextureNotify.cpp
--- a/Engine/Source/Runtime/Engine/Private/VT/RuntimeVirtualTextureNotify.cpp
+++ b/Engine/Source/Runtime/Engine/Private/VT/RuntimeVirtualTextureNotify.cpp
@@ -3,9 +3,9 @@
 #include "VT/RuntimeVirtualTextureNotify.h"
 
 #include "Components/PrimitiveComponent.h"
+#include "Components/RuntimeVirtualTextureComponent.h"
 #include "UObject/UObjectIterator.h"
 #include "VT/RuntimeVirtualTexture.h"
-#include "VT/RuntimeVirtualTexturePlane.h"
 
 namespace RuntimeVirtualTexture
 {
diff --git a/Engine/Source/Runtime/Engine/Private/VT/RuntimeVirtualTexturePlane.cpp b/Engine/Source/Runtime/Engine/Private/VT/RuntimeVirtualTexturePlane.cpp
deleted file mode 100644
diff --git a/Engine/Source/Runtime/Engine/Private/VT/RuntimeVirtualTextureVolume.cpp b/Engine/Source/Runtime/Engine/Private/VT/RuntimeVirtualTextureVolume.cpp
new file mode 100644
--- /dev/null
+++ b/Engine/Source/Runtime/Engine/Private/VT/RuntimeVirtualTextureVolume.cpp
@@ -0,0 +1,24 @@
+// Copyright 1998-2019 Epic Games, Inc. All Rights Reserved.
+
+#include "VT/RuntimeVirtualTextureVolume.h"
+
+#include "Components/BoxComponent.h"
+#include "Components/RuntimeVirtualTextureComponent.h"
+
+ARuntimeVirtualTextureVolume::ARuntimeVirtualTextureVolume(const FObjectInitializer& ObjectInitializer)
+	: Super(ObjectInitializer)
+{
+	RootComponent = VirtualTextureComponent = CreateDefaultSubobject<URuntimeVirtualTextureComponent>(TEXT("VirtualTextureComponent"));
+
+#if WITH_EDITORONLY_DATA
+	// Add box for visualization of bounds
+	Box = CreateDefaultSubobject<UBoxComponent>(TEXT("Bounds"));
+	Box->SetBoxExtent(FVector(0.5f, 0.5f, 1.f), false);
+	Box->SetIsVisualizationComponent(true);
+	Box->SetCollisionEnabled(ECollisionEnabled::NoCollision);
+	Box->SetCanEverAffectNavigation(false);
+	Box->CanCharacterStepUpOn = ECanBeCharacterBase::ECB_No;
+	Box->SetGenerateOverlapEvents(false);
+	Box->SetupAttachment(VirtualTextureComponent);
+#endif
+}
diff --git a/Engine/Source/Runtime/Renderer/Private/RendererScene.cpp b/Engine/Source/Runtime/Renderer/Private/RendererScene.cpp
--- a/Engine/Source/Runtime/Renderer/Private/RendererScene.cpp
+++ b/Engine/Source/Runtime/Renderer/Private/RendererScene.cpp
@@ -29,6 +29,7 @@
 #include "GameFramework/WorldSettings.h"
 #include "Components/DecalComponent.h"
 #include "Components/ReflectionCaptureComponent.h"
+#include "Components/RuntimeVirtualTextureComponent.h"
 #include "ScenePrivateBase.h"
 #include "SceneCore.h"
 #include "PrimitiveSceneInfo.h"
@@ -54,7 +55,6 @@
 #include "DynamicShadowMapChannelBindingHelper.h"
 #include "GPUScene.h"
 #include "HAL/LowLevelMemTracker.h"
-#include "VT/RuntimeVirtualTexturePlane.h"
 #include "VT/RuntimeVirtualTextureSceneProxy.h"
 #if RHI_RAYTRACING
 #include "RayTracingDynamicGeometryCollection.h"
diff --git a/Engine/Source/Runtime/Renderer/Private/VT/RuntimeVirtualTextureSceneProxy.cpp b/Engine/Source/Runtime/Renderer/Private/VT/RuntimeVirtualTextureSceneProxy.cpp
--- a/Engine/Source/Runtime/Renderer/Private/VT/RuntimeVirtualTextureSceneProxy.cpp
+++ b/Engine/Source/Runtime/Renderer/Private/VT/RuntimeVirtualTextureSceneProxy.cpp
@@ -2,8 +2,8 @@
 
 #include "VT/RuntimeVirtualTextureSceneProxy.h"
 
+#include "Components/RuntimeVirtualTextureComponent.h"
 #include "VT/RuntimeVirtualTexture.h"
-#include "VT/RuntimeVirtualTexturePlane.h"
 #include "VT/RuntimeVirtualTextureProducer.h"
 
 
