Added missing includes

diff --git a/Engine/Source/Runtime/Engine/Private/Components/RuntimeVirtualTextureComponent.cpp b/Engine/Source/Runtime/Engine/Private/Components/RuntimeVirtualTextureComponent.cpp
--- a/Engine/Source/Runtime/Engine/Private/Components/RuntimeVirtualTextureComponent.cpp
+++ b/Engine/Source/Runtime/Engine/Private/Components/RuntimeVirtualTextureComponent.cpp
@@ -2,6 +2,10 @@
 
 #include "Components/RuntimeVirtualTextureComponent.h"
 
+#include "GameFramework/Actor.h"
+#include "PrimitiveComponent.h"
+#include "SceneInterface.h"
+
 URuntimeVirtualTextureComponent::URuntimeVirtualTextureComponent(const FObjectInitializer& ObjectInitializer)
 	: Super(ObjectInitializer)
 	, SceneProxy(nullptr)
