Added missing includes (2)

diff --git a/Engine/Source/Runtime/Engine/Private/Components/RuntimeVirtualTextureComponent.cpp b/Engine/Source/Runtime/Engine/Private/Components/RuntimeVirtualTextureComponent.cpp
--- a/Engine/Source/Runtime/Engine/Private/Components/RuntimeVirtualTextureComponent.cpp
+++ b/Engine/Source/Runtime/Engine/Private/Components/RuntimeVirtualTextureComponent.cpp
@@ -2,8 +2,8 @@
 
 #include "Components/RuntimeVirtualTextureComponent.h"
 
+#include "Components/PrimitiveComponent.h"
 #include "GameFramework/Actor.h"
-#include "PrimitiveComponent.h"
 #include "SceneInterface.h"
 
 URuntimeVirtualTextureComponent::URuntimeVirtualTextureComponent(const FObjectInitializer& ObjectInitializer)
