Ported MaterialShader and MeshMaterialShader to use new shader permutation system.

diff --git a/Engine/Source/Runtime/Engine/Private/Materials/MaterialShader.cpp b/Engine/Source/Runtime/Engine/Private/Materials/MaterialShader.cpp
--- a/Engine/Source/Runtime/Engine/Private/Materials/MaterialShader.cpp
+++ b/Engine/Source/Runtime/Engine/Private/Materials/MaterialShader.cpp
@@ -48,11 +48,9 @@ uint32 FMaterialShaderMap::NextCompilingId = 2;
  */
 TMap<TRefCountPtr<FMaterialShaderMap>, TArray<FMaterial*> > FMaterialShaderMap::ShaderMapsBeingCompiled;
 
-
-
-static inline bool ShouldCacheMaterialShader(const FMaterialShaderType* ShaderType, EShaderPlatform Platform, const FMaterial* Material)
+static inline bool ShouldCacheMaterialShader(const FMaterialShaderType* ShaderType, EShaderPlatform Platform, const FMaterial* Material, int32 PermutationId)
 {
-	return ShaderType->ShouldCache(Platform, Material) && Material->ShouldCache(Platform, ShaderType, nullptr);
+	return ShaderType->ShouldCompilePermutation(Platform, Material, PermutationId) && Material->ShouldCache(Platform, ShaderType, nullptr);
 }
 
 
@@ -919,6 +917,7 @@ void FMaterialShaderMapId::SetShaderDependencies(const TArray<FShaderType*>& Sha
  */
 FShaderCompileJob* FMaterialShaderType::BeginCompileShader(
 	uint32 ShaderMapId,
+	int32 PermutationId,
 	const FMaterial* Material,
 	FShaderCompilerEnvironment* MaterialEnvironment,
 	const FShaderPipelineType* ShaderPipeline,
@@ -928,7 +927,7 @@ FShaderCompileJob* FMaterialShaderType::BeginCompileShader(
 	const FString& DebugExtension
 	)
 {
-	FShaderCompileJob* NewJob = new FShaderCompileJob(ShaderMapId, nullptr, this, /* PermutationId = */ 0);
+	FShaderCompileJob* NewJob = new FShaderCompileJob(ShaderMapId, nullptr, this, PermutationId);
 
 	NewJob->Input.SharedEnvironment = MaterialEnvironment;
 	FShaderCompilerEnvironment& ShaderEnvironment = NewJob->Input.Environment;
@@ -942,7 +941,7 @@ FShaderCompileJob* FMaterialShaderType::BeginCompileShader(
 	Material->SetupExtaCompilationSettings(Platform, NewJob->Input.ExtraSettings);
 
 	// Allow the shader type to modify the compile environment.
-	SetupCompileEnvironment(Platform, Material, ShaderEnvironment);
+	SetupCompileEnvironment(Platform, Material, PermutationId, ShaderEnvironment);
 
 	// Compile the shader environment passed in with the shader type's source code.
 	::GlobalBeginCompileShader(
@@ -982,7 +981,7 @@ void FMaterialShaderType::BeginCompileShaderPipeline(
 	for (int32 Index = 0; Index < ShaderStages.Num(); ++Index)
 	{
 		auto* ShaderStage = ShaderStages[Index];
-		ShaderStage->BeginCompileShader(ShaderMapId, Material, MaterialEnvironment, ShaderPipeline, Platform, NewPipelineJob->StageJobs, DebugDescription, DebugExtension);
+		ShaderStage->BeginCompileShader(ShaderMapId, kUniqueShaderPermutationId, Material, MaterialEnvironment, ShaderPipeline, Platform, NewPipelineJob->StageJobs, DebugDescription, DebugExtension);
 	}
 
 	NewJobs.Add(NewPipelineJob);
@@ -1003,11 +1002,17 @@ FShader* FMaterialShaderType::FinishCompileShader(
 {
 	check(CurrentJob.bSucceeded);
 
-	FShaderType* SpecificType = CurrentJob.ShaderType->LimitShaderResourceToThisType() ? CurrentJob.ShaderType : NULL;
+	FShaderType* SpecificType = nullptr;
+	int32 SpecificPermutationId = 0;
+	if (CurrentJob.ShaderType->LimitShaderResourceToThisType())
+	{
+		SpecificType = CurrentJob.ShaderType;
+		SpecificPermutationId = CurrentJob.PermutationId;
+	}
 
 	// Reuse an existing resource with the same key or create a new one based on the compile output
 	// This allows FShaders to share compiled bytecode and RHI shader references
-	FShaderResource* Resource = FShaderResource::FindOrCreateShaderResource(CurrentJob.Output, SpecificType, /* PermutationId = */ 0);
+	FShaderResource* Resource = FShaderResource::FindOrCreateShaderResource(CurrentJob.Output, SpecificType, SpecificPermutationId);
 
 	if (ShaderPipelineType && !ShaderPipelineType->ShouldOptimizeUnusedOutputs(CurrentJob.Input.Target.GetPlatform()))
 	{
@@ -1016,12 +1021,12 @@ FShader* FMaterialShaderType::FinishCompileShader(
 	}
 
 	// Find a shader with the same key in memory
-	FShader* Shader = CurrentJob.ShaderType->FindShaderById(FShaderId(MaterialShaderMapHash, ShaderPipelineType, CurrentJob.VFType, CurrentJob.ShaderType, /* PermutationId = */ 0, CurrentJob.Input.Target));
+	FShader* Shader = CurrentJob.ShaderType->FindShaderById(FShaderId(MaterialShaderMapHash, ShaderPipelineType, CurrentJob.VFType, CurrentJob.ShaderType, SpecificPermutationId, CurrentJob.Input.Target));
 
 	// There was no shader with the same key so create a new one with the compile output, which will bind shader parameters
 	if (!Shader)
 	{
-		Shader = (*ConstructCompiledRef)(CompiledShaderInitializerType(this, CurrentJob.Output, Resource, UniformExpressionSet, MaterialShaderMapHash, ShaderPipelineType, nullptr, InDebugDescription));
+		Shader = (*ConstructCompiledRef)(CompiledShaderInitializerType(this, SpecificPermutationId, CurrentJob.Output, Resource, UniformExpressionSet, MaterialShaderMapHash, ShaderPipelineType, nullptr, InDebugDescription));
 		CurrentJob.Output.ParameterMap.VerifyBindingsAreComplete(GetName(), CurrentJob.Output.Target, CurrentJob.VFType);
 	}
 
@@ -1594,31 +1599,36 @@ void FMaterialShaderMap::Compile(
 			for(TLinkedList<FShaderType*>::TIterator ShaderTypeIt(FShaderType::GetTypeList());ShaderTypeIt;ShaderTypeIt.Next())
 			{
 				FMaterialShaderType* ShaderType = ShaderTypeIt->GetMaterialShaderType();
-				if (ShaderType &&  ShouldCacheMaterialShader(ShaderType, InPlatform, Material))
+				const int32 PermutationCount = ShaderType ? ShaderType->GetPermutationCount() : 0;
+				for (int32 PermutationId = 0; PermutationId < PermutationCount; ++PermutationId)
 				{
+					if (ShouldCacheMaterialShader(ShaderType, InPlatform, Material, PermutationId))
+					{
 #if WITH_EDITOR
-					// Verify that the shader map Id contains inputs for any shaders that will be put into this shader map
-					check(InShaderMapId.ContainsShaderType(ShaderType));
+						// Verify that the shader map Id contains inputs for any shaders that will be put into this shader map
+						check(InShaderMapId.ContainsShaderType(ShaderType, PermutationId));
 #endif
-					// Compile this material shader for this material.
-					TArray<FString> ShaderErrors;
-  
-					// Only compile the shader if we don't already have it
-					if (!HasShader(ShaderType, /* PermutationId = */ 0))
-					{
-						auto* Job = ShaderType->BeginCompileShader(
-							CompilingId,
-							Material,
-							MaterialEnvironment,
-							nullptr,
-							InPlatform,
-							NewJobs,
-							DebugDescription, DebugExtension
+						// Compile this material shader for this material.
+						TArray<FString> ShaderErrors;
+
+						// Only compile the shader if we don't already have it
+						if (!HasShader(ShaderType, PermutationId))
+						{
+							auto* Job = ShaderType->BeginCompileShader(
+								CompilingId,
+								PermutationId,
+								Material,
+								MaterialEnvironment,
+								nullptr,
+								InPlatform,
+								NewJobs,
+								DebugDescription, DebugExtension
 							);
-						check(!SharedShaderJobs.Find(ShaderType));
-						SharedShaderJobs.Add(ShaderType, Job);
+							check(!SharedShaderJobs.Find(ShaderType));
+							SharedShaderJobs.Add(ShaderType, Job);
+						}
+						NumShaders++;
 					}
-					NumShaders++;
 				}
 			}
   
@@ -1633,11 +1643,11 @@ void FMaterialShaderMap::Compile(
 					for (int32 Index = 0; Index < StageTypes.Num(); ++Index)
 					{
 						FMaterialShaderType* ShaderType = (FMaterialShaderType*)(StageTypes[Index]->GetMaterialShaderType());
-						if (ShaderType && ShouldCacheMaterialShader(ShaderType, InPlatform, Material))
+						if (ShaderType && ShouldCacheMaterialShader(ShaderType, InPlatform, Material, kUniqueShaderPermutationId))
 						{
 #if WITH_EDITOR
 							// Verify that the shader map Id contains inputs for any shaders that will be put into this shader map
-							check(InShaderMapId.ContainsShaderType(ShaderType));
+							check(InShaderMapId.ContainsShaderType(ShaderType, kUniqueShaderPermutationId));
 #endif
 							ShaderStagesToCompile.Add(ShaderType);
 						}
@@ -1746,8 +1756,8 @@ FShader* FMaterialShaderMap::ProcessCompilationResultsForSingleJob(FShaderCompil
 		check(Shader);
 		if (!ShaderPipeline)
 		{
-			check(!MeshShaderMap->HasShader(MeshMaterialShaderType, /* PermutationId = */ 0));
-			MeshShaderMap->AddShader(MeshMaterialShaderType, /* PermutationId = */ 0, Shader);
+			check(!MeshShaderMap->HasShader(MeshMaterialShaderType, CurrentJob.PermutationId));
+			MeshShaderMap->AddShader(MeshMaterialShaderType, CurrentJob.PermutationId, Shader);
 		}
 	}
 	else
@@ -1764,8 +1774,8 @@ FShader* FMaterialShaderMap::ProcessCompilationResultsForSingleJob(FShaderCompil
 		check(Shader);
 		if (!ShaderPipeline)
 		{
-			check(!HasShader(MaterialShaderType, /* PermutationId = */ 0));
-			AddShader(MaterialShaderType, /* PermutationId = */ 0, Shader);
+			check(!HasShader(MaterialShaderType, CurrentJob.PermutationId));
+			AddShader(MaterialShaderType, CurrentJob.PermutationId, Shader);
 		}
 	}
 
@@ -1881,7 +1891,7 @@ bool FMaterialShaderMap::ProcessCompilationResults(const TArray<FShaderCommonCom
 							for (int32 Index = 0; Index < StageTypes.Num(); ++Index)
 							{
 								FMeshMaterialShaderType* ShaderType = ((FMeshMaterialShaderType*)(StageTypes[Index]))->GetMeshMaterialShaderType();
-								FShader* Shader = MeshShaderMap->GetShader(ShaderType);
+								FShader* Shader = MeshShaderMap->GetShader(ShaderType, kUniqueShaderPermutationId);
 								check(Shader);
 								ShaderStages.Add(Shader);
 							}
@@ -1907,7 +1917,7 @@ bool FMaterialShaderMap::ProcessCompilationResults(const TArray<FShaderCommonCom
 						for (int32 Index = 0; Index < StageTypes.Num(); ++Index)
 						{
 							FMaterialShaderType* ShaderType = ((FMaterialShaderType*)(StageTypes[Index]))->GetMaterialShaderType();
-							FShader* Shader = GetShader(ShaderType);
+							FShader* Shader = GetShader(ShaderType, kUniqueShaderPermutationId);
 							check(Shader);
 							ShaderStages.Add(Shader);
 						}
@@ -1965,11 +1975,18 @@ bool FMaterialShaderMap::TryToAddToExistingCompilationTask(FMaterial* Material)
 	return false;
 }
 
-bool FMaterialShaderMap::IsMaterialShaderComplete(const FMaterial* Material, const FMaterialShaderType* ShaderType, const FShaderPipelineType* Pipeline, bool bSilent)
+bool FMaterialShaderMap::IsMaterialShaderComplete(
+	const FMaterial* Material,
+	const FMaterialShaderType* ShaderType,
+	const FShaderPipelineType* Pipeline,
+	int32 PermutationId,
+	bool bSilent) const
 {
-	// If we should cache this material, it's incomplete if the pipeline is missing or the shader itself is missing
-	if (ShouldCacheMaterialShader(ShaderType, GetShaderPlatform(), Material) &&
-		((Pipeline && !HasShaderPipeline(Pipeline)) || (!Pipeline && !HasShader((FShaderType*)ShaderType, /* PermutationId = */ 0))))
+	const bool bShouldCacheMaterialShader = ShouldCacheMaterialShader(ShaderType, Platform, Material, PermutationId);
+	const bool bMissingShaderPipeline = (Pipeline && !HasShaderPipeline(Pipeline));
+	const bool bMissingShader = (!Pipeline && !HasShader((FShaderType*)ShaderType, PermutationId));
+
+	if (bShouldCacheMaterialShader && (bMissingShaderPipeline || bMissingShader))
 	{
 		if (!bSilent)
 		{
@@ -1979,7 +1996,7 @@ bool FMaterialShaderMap::IsMaterialShaderComplete(const FMaterial* Material, con
 			}
 			else
 			{
-				UE_LOG(LogShaders, Warning, TEXT("Incomplete material %s, missing FMaterialShader %s."), *Material->GetFriendlyName(), ShaderType->GetName());
+				UE_LOG(LogShaders, Warning, TEXT("Incomplete material %s, missing FMaterialShader (%s, %d)."), *Material->GetFriendlyName(), ShaderType->GetName(), PermutationId);
 			}
 		}
 		return false;
@@ -2026,9 +2043,13 @@ bool FMaterialShaderMap::IsComplete(const FMaterial* Material, bool bSilent)
 	{
 		// Find this shader type in the material's shader map.
 		FMaterialShaderType* ShaderType = ShaderTypeIt->GetMaterialShaderType();
-		if (ShaderType && !IsMaterialShaderComplete(Material, ShaderType, nullptr, bSilent))
+		const int32 PermutationCount = ShaderType ? ShaderType->GetPermutationCount() : 0;
+		for (int32 PermutationId = 0; PermutationId < PermutationCount; ++PermutationId)
 		{
-			return false;
+			if (!IsMaterialShaderComplete(Material, ShaderType, nullptr, PermutationId, bSilent))
+			{
+				return false;
+			}
 		}
 	}
 
@@ -2045,7 +2066,7 @@ bool FMaterialShaderMap::IsComplete(const FMaterial* Material, bool bSilent)
 			for (int32 Index = 0; Index < StageTypes.Num(); ++Index)
 			{
 				auto* ShaderType = StageTypes[Index]->GetMaterialShaderType();
-				if (ShouldCacheMaterialShader(ShaderType, GetShaderPlatform(), Material))
+				if (ShouldCacheMaterialShader(ShaderType, GetShaderPlatform(), Material, kUniqueShaderPermutationId))
 				{
 					++NumShouldCache;
 				}
@@ -2060,7 +2081,7 @@ bool FMaterialShaderMap::IsComplete(const FMaterial* Material, bool bSilent)
 				for (int32 Index = 0; Index < StageTypes.Num(); ++Index)
 				{
 					auto* ShaderType = StageTypes[Index]->GetMaterialShaderType();
-					if (!IsMaterialShaderComplete(Material, ShaderType, Pipeline, bSilent))
+					if (!IsMaterialShaderComplete(Material, ShaderType, Pipeline, bSilent, kUniqueShaderPermutationId))
 					{
 						return false;
 					}
@@ -2094,13 +2115,17 @@ void FMaterialShaderMap::LoadMissingShadersFromMemory(const FMaterial* Material)
 	for (TLinkedList<FShaderType*>::TIterator ShaderTypeIt(FShaderType::GetTypeList());ShaderTypeIt;ShaderTypeIt.Next())
 	{
 		FMaterialShaderType* ShaderType = ShaderTypeIt->GetMaterialShaderType();
-		if (ShaderType && ShouldCacheMaterialShader(ShaderType, GetShaderPlatform(), Material) && !HasShader(ShaderType, /* PermutationId = */ 0))
+		const int32 PermutationCount = ShaderType ? ShaderType->GetPermutationCount() : 0;
+		for (int32 PermutationId = 0; PermutationId < PermutationCount; ++PermutationId)
 		{
-			FShaderId ShaderId(MaterialShaderMapHash, nullptr, nullptr, ShaderType, /** PermutationId = */ 0, FShaderTarget(ShaderType->GetFrequency(), GetShaderPlatform()));
-			FShader* FoundShader = ShaderType->FindShaderById(ShaderId);
-			if (FoundShader)
+			if (ShouldCacheMaterialShader(ShaderType, GetShaderPlatform(), Material, PermutationId) && !HasShader(ShaderType, PermutationId))
 			{
-				AddShader(ShaderType, /* PermutationId = */ 0, FoundShader);
+				FShaderId ShaderId(MaterialShaderMapHash, nullptr, nullptr, ShaderType, PermutationId, FShaderTarget(ShaderType->GetFrequency(), GetShaderPlatform()));
+				FShader* FoundShader = ShaderType->FindShaderById(ShaderId);
+				if (FoundShader)
+				{
+					AddShader(ShaderType, PermutationId, FoundShader);
+				}
 			}
 		}
 	}
@@ -2117,7 +2142,7 @@ void FMaterialShaderMap::LoadMissingShadersFromMemory(const FMaterial* Material)
 			for (const FShaderType* Shader : Stages)
 			{
 				FMaterialShaderType* ShaderType = (FMaterialShaderType*)Shader->GetMaterialShaderType();
-				if (ShaderType && ShouldCacheMaterialShader(ShaderType, GetShaderPlatform(), Material))
+				if (ShaderType && ShouldCacheMaterialShader(ShaderType, GetShaderPlatform(), Material, kUniqueShaderPermutationId))
 				{
 					++NumShaders;
 				}
@@ -2129,13 +2154,13 @@ void FMaterialShaderMap::LoadMissingShadersFromMemory(const FMaterial* Material)
 				for (auto* Shader : Stages)
 				{
 					FMaterialShaderType* ShaderType = (FMaterialShaderType*)Shader->GetMaterialShaderType();
-					if (!HasShader(ShaderType, /* PermutationId = */ 0))
+					if (!HasShader(ShaderType, kUniqueShaderPermutationId))
 					{
-						FShaderId ShaderId(MaterialShaderMapHash, PipelineType->ShouldOptimizeUnusedOutputs(GetShaderPlatform()) ? PipelineType : nullptr, nullptr, ShaderType, /** PermutationId = */ 0, FShaderTarget(ShaderType->GetFrequency(), GetShaderPlatform()));
+						FShaderId ShaderId(MaterialShaderMapHash, PipelineType->ShouldOptimizeUnusedOutputs(GetShaderPlatform()) ? PipelineType : nullptr, nullptr, ShaderType, kUniqueShaderPermutationId, FShaderTarget(ShaderType->GetFrequency(), GetShaderPlatform()));
 						FShader* FoundShader = ShaderType->FindShaderById(ShaderId);
 						if (FoundShader)
 						{
-							AddShader(ShaderType, /* PermutationId = */ 0, FoundShader);
+							AddShader(ShaderType, kUniqueShaderPermutationId, FoundShader);
 							ShadersForPipeline.Add(FoundShader);
 						}
 					}
@@ -2325,7 +2350,11 @@ void FMaterialShaderMap::FlushShadersByShaderType(FShaderType* ShaderType)
 
 	if (ShaderType->GetMaterialShaderType())
 	{
-		RemoveShaderTypePermutaion(ShaderType->GetMaterialShaderType(), /* PermutationId = */ 0);	
+		const int32 PermutationCount = ShaderType->GetPermutationCount();
+		for (int32 PermutationId = 0; PermutationId < PermutationCount; ++PermutationId)
+		{
+			RemoveShaderTypePermutaion(ShaderType->GetMaterialShaderType(), PermutationId);
+		}
 	}
 }
 
diff --git a/Engine/Source/Runtime/Engine/Private/Materials/MeshMaterialShader.cpp b/Engine/Source/Runtime/Engine/Private/Materials/MeshMaterialShader.cpp
--- a/Engine/Source/Runtime/Engine/Private/Materials/MeshMaterialShader.cpp
+++ b/Engine/Source/Runtime/Engine/Private/Materials/MeshMaterialShader.cpp
@@ -21,9 +21,9 @@ namespace MaterialMeshCookStats
 }
 #endif
 
-static inline bool ShouldCacheMeshShader(const FMeshMaterialShaderType* ShaderType, EShaderPlatform Platform, const FMaterial* Material, FVertexFactoryType* InVertexFactoryType)
+static inline bool ShouldCacheMeshShader(const FMeshMaterialShaderType* ShaderType, EShaderPlatform Platform, const FMaterial* Material, FVertexFactoryType* InVertexFactoryType, int32 PermutationId)
 {
-	return ShaderType->ShouldCache(Platform, Material, InVertexFactoryType) &&
+	return ShaderType->ShouldCompilePermutation(Platform, Material, InVertexFactoryType, PermutationId) &&
 		Material->ShouldCache(Platform, ShaderType, InVertexFactoryType) &&
 		InVertexFactoryType->ShouldCache(Platform, Material, ShaderType);
 }
@@ -36,6 +36,7 @@ static inline bool ShouldCacheMeshShader(const FMeshMaterialShaderType* ShaderTy
  */
 FShaderCompileJob* FMeshMaterialShaderType::BeginCompileShader(
 	uint32 ShaderMapId,
+	int32 PermutationId,
 	EShaderPlatform Platform,
 	const FMaterial* Material,
 	FShaderCompilerEnvironment* MaterialEnvironment,
@@ -43,11 +44,9 @@ FShaderCompileJob* FMeshMaterialShaderType::BeginCompileShader(
 	const FShaderPipelineType* ShaderPipeline,
 	TArray<FShaderCommonCompileJob*>& NewJobs,
 	FString DebugDescription,
-	FString DebugExtension
-
-	)
+	FString DebugExtension)
 {
-	FShaderCompileJob* NewJob = new FShaderCompileJob(ShaderMapId, VertexFactoryType, this, /* PermutationId = */ 0);
+	FShaderCompileJob* NewJob = new FShaderCompileJob(ShaderMapId, VertexFactoryType, this, PermutationId);
 
 	NewJob->Input.SharedEnvironment = MaterialEnvironment;
 	FShaderCompilerEnvironment& ShaderEnvironment = NewJob->Input.Environment;
@@ -66,7 +65,7 @@ FShaderCompileJob* FMeshMaterialShaderType::BeginCompileShader(
 	COOK_STAT(MaterialMeshCookStats::ShadersCompiled++);
 
 	// Allow the shader type to modify the compile environment.
-	SetupCompileEnvironment(Platform, Material, ShaderEnvironment);
+	SetupCompileEnvironment(Platform, Material, PermutationId, ShaderEnvironment);
 
 	bool bAllowDevelopmentShaderCompile = Material->GetAllowDevelopmentShaderCompile();
 
@@ -90,6 +89,7 @@ FShaderCompileJob* FMeshMaterialShaderType::BeginCompileShader(
 
 void FMeshMaterialShaderType::BeginCompileShaderPipeline(
 	uint32 ShaderMapId,
+	int32 PermutationId,
 	EShaderPlatform Platform,
 	const FMaterial* Material,
 	FShaderCompilerEnvironment* MaterialEnvironment,
@@ -98,8 +98,7 @@ void FMeshMaterialShaderType::BeginCompileShaderPipeline(
 	const TArray<FMeshMaterialShaderType*>& ShaderStages,
 	TArray<FShaderCommonCompileJob*>& NewJobs,
 	FString DebugDescription,
-	FString DebugExtension
-	)
+	FString DebugExtension)
 {
 	check(ShaderStages.Num() > 0);
 	check(ShaderPipeline);
@@ -110,7 +109,7 @@ void FMeshMaterialShaderType::BeginCompileShaderPipeline(
 	for (int32 Index = 0; Index < ShaderStages.Num(); ++Index)
 	{
 		auto* ShaderStage = ShaderStages[Index];
-		ShaderStage->BeginCompileShader(ShaderMapId, Platform, Material, MaterialEnvironment, VertexFactoryType, ShaderPipeline, NewPipelineJob->StageJobs, DebugDescription, DebugExtension);
+		ShaderStage->BeginCompileShader(ShaderMapId, PermutationId, Platform, Material, MaterialEnvironment, VertexFactoryType, ShaderPipeline, NewPipelineJob->StageJobs, DebugDescription, DebugExtension);
 	}
 
 	NewJobs.Add(NewPipelineJob);
@@ -152,7 +151,7 @@ FShader* FMeshMaterialShaderType::FinishCompileShader(
 
 	// Reuse an existing resource with the same key or create a new one based on the compile output
 	// This allows FShaders to share compiled bytecode and RHI shader references
-	FShaderResource* Resource = FShaderResource::FindOrCreateShaderResource(CurrentJob.Output, SpecificType, /* PermutationId = */ 0);
+	FShaderResource* Resource = FShaderResource::FindOrCreateShaderResource(CurrentJob.Output, SpecificType, CurrentJob.PermutationId);
 
 	if (ShaderPipelineType && !ShaderPipelineType->ShouldOptimizeUnusedOutputs(CurrentJob.Input.Target.GetPlatform()))
 	{
@@ -161,12 +160,12 @@ FShader* FMeshMaterialShaderType::FinishCompileShader(
 	}
 
 	// Find a shader with the same key in memory
-	FShader* Shader = CurrentJob.ShaderType->FindShaderById(FShaderId(MaterialShaderMapHash, ShaderPipelineType, CurrentJob.VFType, CurrentJob.ShaderType, /** PermutationId = */ 0, CurrentJob.Input.Target));
+	FShader* Shader = CurrentJob.ShaderType->FindShaderById(FShaderId(MaterialShaderMapHash, ShaderPipelineType, CurrentJob.VFType, CurrentJob.ShaderType, CurrentJob.PermutationId, CurrentJob.Input.Target));
 
 	// There was no shader with the same key so create a new one with the compile output, which will bind shader parameters
 	if (!Shader)
 	{
-		Shader = (*ConstructCompiledRef)(CompiledShaderInitializerType(this, CurrentJob.Output, Resource, UniformExpressionSet, MaterialShaderMapHash, InDebugDescription, ShaderPipelineType, CurrentJob.VFType));
+		Shader = (*ConstructCompiledRef)(CompiledShaderInitializerType(this, CurrentJob.PermutationId, CurrentJob.Output, Resource, UniformExpressionSet, MaterialShaderMapHash, InDebugDescription, ShaderPipelineType, CurrentJob.VFType));
 		ValidateAfterBind((FMeshMaterialShader*)Shader);
 		CurrentJob.Output.ParameterMap.VerifyBindingsAreComplete(GetName(), CurrentJob.Output.Target, CurrentJob.VFType);
 	}
@@ -200,36 +199,46 @@ uint32 FMeshMaterialShaderMap::BeginCompile(
 	TSet<FString> ShaderTypeNames;
 
 	// Iterate over all mesh material shader types.
-	TMap<FShaderType*, FShaderCompileJob*> SharedShaderJobs;
+	TMap<TShaderTypePermutation<const FShaderType>, FShaderCompileJob*> SharedShaderJobs;
 	for (TLinkedList<FShaderType*>::TIterator ShaderTypeIt(FShaderType::GetTypeList());ShaderTypeIt;ShaderTypeIt.Next())
 	{
 		FMeshMaterialShaderType* ShaderType = ShaderTypeIt->GetMeshMaterialShaderType();
-		if (ShaderType && ShouldCacheMeshShader(ShaderType, InPlatform, Material, VertexFactoryType))
+		if (!ShaderType)
+		{
+			continue;
+		}
+
+		for (int32 PermutationId = 0; PermutationId < ShaderType->GetPermutationCount(); ++PermutationId)
 		{
-			// Verify that the shader map Id contains inputs for any shaders that will be put into this shader map
+			if (ShouldCacheMeshShader(ShaderType, InPlatform, Material, VertexFactoryType, PermutationId))
+			{
+				// Verify that the shader map Id contains inputs for any shaders that will be put into this shader map
 #if WITH_EDITOR
-			check(InShaderMapId.ContainsVertexFactoryType(VertexFactoryType));
-			check(InShaderMapId.ContainsShaderType(ShaderType));
+				check(InShaderMapId.ContainsVertexFactoryType(VertexFactoryType));
+				check(InShaderMapId.ContainsShaderType(ShaderType, PermutationId));
 #endif
 
-			NumShadersPerVF++;
-			// only compile the shader if we don't already have it
-			if (!HasShader(ShaderType, /* PermutationId = */ 0))
-			{
-				// Compile this mesh material shader for this material and vertex factory type.
-				auto* Job = ShaderType->BeginCompileShader(
-					ShaderMapId,
-					InPlatform,
-					Material,
-					MaterialEnvironment,
-					VertexFactoryType,
-					nullptr,
-					NewJobs,
-					DebugDescription, 
-					DebugExtension
-					);
-				check(!SharedShaderJobs.Find(ShaderType));
-				SharedShaderJobs.Add(ShaderType, Job);
+				NumShadersPerVF++;
+				// only compile the shader if we don't already have it
+				if (!HasShader(ShaderType, PermutationId))
+				{
+					// Compile this mesh material shader for this material and vertex factory type.
+					auto* Job = ShaderType->BeginCompileShader(
+						ShaderMapId,
+						PermutationId,
+						InPlatform,
+						Material,
+						MaterialEnvironment,
+						VertexFactoryType,
+						nullptr,
+						NewJobs,
+						DebugDescription, 
+						DebugExtension
+						);
+					TShaderTypePermutation<const FShaderType> ShaderTypePermutation(ShaderType, PermutationId);
+					check(!SharedShaderJobs.Find(ShaderTypePermutation));
+					SharedShaderJobs.Add(ShaderTypePermutation, Job);
+				}
 			}
 		}
 	}
@@ -246,7 +255,7 @@ uint32 FMeshMaterialShaderMap::BeginCompile(
 			for (auto* Shader : StageTypes)
 			{
 				const FMeshMaterialShaderType* ShaderType = Shader->GetMeshMaterialShaderType();
-				if (ShouldCacheMeshShader(ShaderType, InPlatform, Material, VertexFactoryType))
+				if (ShouldCacheMeshShader(ShaderType, InPlatform, Material, VertexFactoryType, kUniqueShaderPermutationId))
 				{
 					++NumShaderStagesToCompile;
 				}
@@ -273,20 +282,21 @@ uint32 FMeshMaterialShaderMap::BeginCompile(
 						// Verify that the shader map Id contains inputs for any shaders that will be put into this shader map
 #if WITH_EDITOR
 						check(InShaderMapId.ContainsVertexFactoryType(VertexFactoryType));
-						check(InShaderMapId.ContainsShaderType(ShaderType));
+						check(InShaderMapId.ContainsShaderType(ShaderType, kUniqueShaderPermutationId));
 #endif
 						ShaderStagesToCompile.Add((FMeshMaterialShaderType*)ShaderType);
 					}
 
 					// Make a pipeline job with all the stages
-					FMeshMaterialShaderType::BeginCompileShaderPipeline(ShaderMapId, InPlatform, Material, MaterialEnvironment, VertexFactoryType, Pipeline, ShaderStagesToCompile, NewJobs, DebugDescription, DebugExtension);
+					FMeshMaterialShaderType::BeginCompileShaderPipeline(ShaderMapId, kUniqueShaderPermutationId, InPlatform, Material, MaterialEnvironment, VertexFactoryType, Pipeline, ShaderStagesToCompile, NewJobs, DebugDescription, DebugExtension);
 				}
 				else
 				{
 					// If sharing shaders amongst pipelines, add this pipeline as a dependency of an existing job
 					for (const FShaderType* ShaderType : StageTypes)
 					{
-						FShaderCompileJob** Job = SharedShaderJobs.Find(ShaderType);
+						TShaderTypePermutation<const FShaderType> ShaderTypePermutation(ShaderType, kUniqueShaderPermutationId);
+						FShaderCompileJob** Job = SharedShaderJobs.Find(ShaderTypePermutation);
 						checkf(Job, TEXT("Couldn't find existing shared job for mesh shader %s on pipeline %s!"), ShaderType->GetName(), Pipeline->GetName());
 						auto* SingleJob = (*Job)->GetSingleShaderJob();
 						auto& PipelinesToShare = SingleJob->SharingPipelines.FindOrAdd(VertexFactoryType);
@@ -306,14 +316,22 @@ uint32 FMeshMaterialShaderMap::BeginCompile(
 	return NumShadersPerVF;
 }
 
-inline bool FMeshMaterialShaderMap::IsMeshShaderComplete(const FMeshMaterialShaderMap* MeshShaderMap, EShaderPlatform Platform, const FMaterial* Material, const FMeshMaterialShaderType* ShaderType, const FShaderPipelineType* Pipeline, FVertexFactoryType* InVertexFactoryType, bool bSilent)
+inline bool FMeshMaterialShaderMap::IsMeshShaderComplete(
+	const FMeshMaterialShaderMap* MeshShaderMap,
+	EShaderPlatform Platform,
+	const FMaterial* Material,
+	const FMeshMaterialShaderType* ShaderType,
+	const FShaderPipelineType* Pipeline,
+	FVertexFactoryType* InVertexFactoryType,
+	int32 PermutationId,
+	bool bSilent)
 {
-	// If we should cache this shader then the map is empty IF
-	//		The shadermap is empty
-	//		OR If it doesn't have a pipeline and needs one
-	//		OR it's not in the shadermap
-	if (ShouldCacheMeshShader(ShaderType, Platform, Material, InVertexFactoryType) &&
-		(!MeshShaderMap || (Pipeline && !MeshShaderMap->HasShaderPipeline(Pipeline)) || (!Pipeline && !MeshShaderMap->HasShader((FShaderType*)ShaderType, /* PermutationId = */ 0))))
+	if (!ShouldCacheMeshShader(ShaderType, Platform, Material, InVertexFactoryType, PermutationId))
+	{
+		return true;
+	}
+
+	if (!MeshShaderMap || (Pipeline && !MeshShaderMap->HasShaderPipeline(Pipeline)) || (!Pipeline && !MeshShaderMap->HasShader((FShaderType*)ShaderType, PermutationId)))
 	{
 		if (!bSilent)
 		{
@@ -323,7 +341,7 @@ inline bool FMeshMaterialShaderMap::IsMeshShaderComplete(const FMeshMaterialShad
 			}
 			else
 			{
-				UE_LOG(LogShaders, Warning, TEXT("Incomplete material %s, missing %s from %s."), *Material->GetFriendlyName(), ShaderType->GetName(), InVertexFactoryType->GetName());
+				UE_LOG(LogShaders, Warning, TEXT("Incomplete material %s, missing (%s, %d) from %s."), *Material->GetFriendlyName(), ShaderType->GetName(), PermutationId, InVertexFactoryType->GetName());
 			}
 		}
 		return false;
@@ -338,16 +356,19 @@ bool FMeshMaterialShaderMap::IsComplete(
 	EShaderPlatform Platform,
 	const FMaterial* Material,
 	FVertexFactoryType* InVertexFactoryType,
-	bool bSilent
-	)
+	bool bSilent)
 {
 	// Iterate over all mesh material shader types.
 	for (TLinkedList<FShaderType*>::TIterator ShaderTypeIt(FShaderType::GetTypeList());ShaderTypeIt;ShaderTypeIt.Next())
 	{
 		FMeshMaterialShaderType* ShaderType = ShaderTypeIt->GetMeshMaterialShaderType();
-		if (ShaderType && !IsMeshShaderComplete(MeshShaderMap, Platform, Material, ShaderType, nullptr, InVertexFactoryType, bSilent))
+		const int32 PermutationCount = ShaderType ? ShaderType->GetPermutationCount() : 0;
+		for (int32 PermutationId = 0; PermutationId < PermutationCount; ++PermutationId)
 		{
-			return false;
+			if (!IsMeshShaderComplete(MeshShaderMap, Platform, Material, ShaderType, nullptr, InVertexFactoryType, PermutationId, bSilent))
+			{
+				return false;
+			}
 		}
 	}
 
@@ -365,7 +386,7 @@ bool FMeshMaterialShaderMap::IsComplete(
 			for (int32 Index = 0; Index < Stages.Num(); ++Index)
 			{
 				auto* ShaderType = Stages[Index]->GetMeshMaterialShaderType();
-				if (ShouldCacheMeshShader(ShaderType, Platform, Material, InVertexFactoryType))
+				if (ShouldCacheMeshShader(ShaderType, Platform, Material, InVertexFactoryType, kUniqueShaderPermutationId))
 				{
 					++NumShouldCache;
 				}
@@ -381,7 +402,7 @@ bool FMeshMaterialShaderMap::IsComplete(
 				for (int32 Index = 0; Index < Stages.Num(); ++Index)
 				{
 					auto* ShaderType = Stages[Index]->GetMeshMaterialShaderType();
-					if (ShaderType && !IsMeshShaderComplete(MeshShaderMap, Platform, Material, ShaderType, ShaderPipelineType, InVertexFactoryType, bSilent))
+					if (ShaderType && !IsMeshShaderComplete(MeshShaderMap, Platform, Material, ShaderType, ShaderPipelineType, InVertexFactoryType, kUniqueShaderPermutationId, bSilent))
 					{
 						return false;
 					}
@@ -402,14 +423,18 @@ void FMeshMaterialShaderMap::LoadMissingShadersFromMemory(
 	for (TLinkedList<FShaderType*>::TIterator ShaderTypeIt(FShaderType::GetTypeList());ShaderTypeIt;ShaderTypeIt.Next())
 	{
 		FMeshMaterialShaderType* ShaderType = ShaderTypeIt->GetMeshMaterialShaderType();
-		if (ShaderType && ShouldCacheMeshShader(ShaderType, InPlatform, Material, VertexFactoryType) && !HasShader((FShaderType*)ShaderType, /* PermutationId = */ 0))
+		const int32 PermutationCount = ShaderType ? ShaderType->GetPermutationCount() : 0;
+		for (int32 PermutationId = 0; PermutationId < PermutationCount; ++PermutationId)
 		{
-			const FShaderId ShaderId(MaterialShaderMapHash, nullptr, VertexFactoryType, (FShaderType*)ShaderType, /** PermutationId = */ 0, FShaderTarget(ShaderType->GetFrequency(), InPlatform));
-			FShader* FoundShader = ((FShaderType*)ShaderType)->FindShaderById(ShaderId);
-
-			if (FoundShader)
+			if (ShouldCacheMeshShader(ShaderType, InPlatform, Material, VertexFactoryType, PermutationId) && !HasShader((FShaderType*)ShaderType, PermutationId))
 			{
-				AddShader((FShaderType*)ShaderType, /* PermutationId = */ 0, FoundShader);
+				const FShaderId ShaderId(MaterialShaderMapHash, nullptr, VertexFactoryType, (FShaderType*)ShaderType, PermutationId, FShaderTarget(ShaderType->GetFrequency(), InPlatform));
+				FShader* FoundShader = ((FShaderType*)ShaderType)->FindShaderById(ShaderId);
+
+				if (FoundShader)
+				{
+					AddShader((FShaderType*)ShaderType, PermutationId, FoundShader);
+				}
 			}
 		}
 	}
@@ -426,7 +451,7 @@ void FMeshMaterialShaderMap::LoadMissingShadersFromMemory(
 			for (const FShaderType* Shader : Stages)
 			{
 				FMeshMaterialShaderType* ShaderType = (FMeshMaterialShaderType*)Shader->GetMeshMaterialShaderType();
-				if (ShaderType && ShouldCacheMeshShader(ShaderType, InPlatform, Material, VertexFactoryType))
+				if (ShaderType && ShouldCacheMeshShader(ShaderType, InPlatform, Material, VertexFactoryType, kUniqueShaderPermutationId))
 				{
 					++NumShaders;
 				}
@@ -442,13 +467,13 @@ void FMeshMaterialShaderMap::LoadMissingShadersFromMemory(
 				for (auto* Shader : Stages)
 				{
 					FMeshMaterialShaderType* ShaderType = (FMeshMaterialShaderType*)Shader->GetMeshMaterialShaderType();
-					if (!HasShader(ShaderType, /* PermutationId = */ 0))
+					if (!HasShader(ShaderType, kUniqueShaderPermutationId))
 					{
-						const FShaderId ShaderId(MaterialShaderMapHash, PipelineType->ShouldOptimizeUnusedOutputs(InPlatform) ? PipelineType : nullptr, VertexFactoryType, ShaderType, /** PermutationId = */ 0, FShaderTarget(ShaderType->GetFrequency(), InPlatform));
+						const FShaderId ShaderId(MaterialShaderMapHash, PipelineType->ShouldOptimizeUnusedOutputs(InPlatform) ? PipelineType : nullptr, VertexFactoryType, ShaderType, kUniqueShaderPermutationId, FShaderTarget(ShaderType->GetFrequency(), InPlatform));
 						FShader* FoundShader = ShaderType->FindShaderById(ShaderId);
 						if (FoundShader)
 						{
-							AddShader(ShaderType, /* PermutationId = */ 0, FoundShader);
+							AddShader(ShaderType, kUniqueShaderPermutationId, FoundShader);
 							ShadersForPipeline.Add(FoundShader);
 						}
 					}
@@ -473,7 +498,11 @@ void FMeshMaterialShaderMap::FlushShadersByShaderType(FShaderType* ShaderType)
 {
 	if (ShaderType->GetMeshMaterialShaderType())
 	{
-		RemoveShaderTypePermutaion(ShaderType, /* PermutationId = */ 0);
+		const int32 PermutationCount = ShaderType->GetPermutationCount();
+		for (int32 PermutationId = 0; PermutationId < PermutationCount; ++PermutationId)
+		{
+			RemoveShaderTypePermutaion(ShaderType, PermutationId);
+		}
 	}
 }
 
diff --git a/Engine/Source/Runtime/Engine/Private/ShaderCompiler/ShaderCompiler.cpp b/Engine/Source/Runtime/Engine/Private/ShaderCompiler/ShaderCompiler.cpp
--- a/Engine/Source/Runtime/Engine/Private/ShaderCompiler/ShaderCompiler.cpp
+++ b/Engine/Source/Runtime/Engine/Private/ShaderCompiler/ShaderCompiler.cpp
@@ -3752,7 +3752,7 @@ void FGlobalShaderTypeCompiler::BeginCompileShaderPipeline(EShaderPlatform Platf
 	for (int32 Index = 0; Index < ShaderStages.Num(); ++Index)
 	{
 		auto* ShaderStage = ShaderStages[Index];
-		BeginCompileShader(ShaderStage, /** PermutationId = */ 0, Platform, ShaderPipeline, NewPipelineJob->StageJobs);
+		BeginCompileShader(ShaderStage, kUniqueShaderPermutationId, Platform, ShaderPipeline, NewPipelineJob->StageJobs);
 	}
 
 	NewJobs.Add(NewPipelineJob);
@@ -3915,7 +3915,7 @@ void VerifyGlobalShaders(EShaderPlatform Platform, bool bLoadedFromCacheFile)
 				for (int32 Index = 0; Index < StageTypes.Num(); ++Index)
 				{
 					FGlobalShaderType* GlobalShaderType = ((FShaderType*)(StageTypes[Index]))->GetGlobalShaderType();
-					if (GlobalShaderType->ShouldCompilePermutation(Platform, /** PermutationId = */ 0))
+					if (GlobalShaderType->ShouldCompilePermutation(Platform, kUniqueShaderPermutationId))
 					{
 						ShaderStages.Add(GlobalShaderType);
 					}
@@ -3947,7 +3947,7 @@ void VerifyGlobalShaders(EShaderPlatform Platform, bool bLoadedFromCacheFile)
 						// If sharing shaders amongst pipelines, add this pipeline as a dependency of an existing individual job
 						for (const FShaderType* ShaderType : StageTypes)
 						{
-							TShaderTypePermutation<const FShaderType> ShaderTypePermutation(ShaderType, /* PermutationId = */ 0);
+							TShaderTypePermutation<const FShaderType> ShaderTypePermutation(ShaderType, kUniqueShaderPermutationId);
 
 							FShaderCompileJob** Job = SharedShaderJobs.Find(ShaderTypePermutation);
 							checkf(Job, TEXT("Couldn't find existing shared job for global shader %s on pipeline %s!"), ShaderType->GetName(), Pipeline->GetName());
@@ -4158,7 +4158,7 @@ bool IsGlobalShaderMapComplete(const TCHAR* TypeNameSubstring)
 					for (const FShaderType* Shader : Stages)
 					{
 						const FGlobalShaderType* GlobalShaderType = Shader->GetGlobalShaderType();
-						if (ShouldCacheGlobalShaderTypeName(GlobalShaderType, /* PermutationId = */ 0, TypeNameSubstring, Platform))
+						if (ShouldCacheGlobalShaderTypeName(GlobalShaderType, kUniqueShaderPermutationId, TypeNameSubstring, Platform))
 						{
 							++NumStagesNeeded;
 						}
@@ -4641,9 +4641,9 @@ void ProcessCompiledGlobalShaders(const TArray<FShaderCommonCompileJob*>& Compil
 					for (int32 Index = 0; Index < StageTypes.Num(); ++Index)
 					{
 						FGlobalShaderType* GlobalShaderType = ((FShaderType*)(StageTypes[Index]))->GetGlobalShaderType();
-						if (GlobalShaderType->ShouldCompilePermutation(Platform, /** PermutationId = */ 0))
+						if (GlobalShaderType->ShouldCompilePermutation(Platform, kUniqueShaderPermutationId))
 						{
-							FShader* Shader = GlobalShaderMap->GetShader(GlobalShaderType, /** PermutationId = */ 0);
+							FShader* Shader = GlobalShaderMap->GetShader(GlobalShaderType, kUniqueShaderPermutationId);
 							check(Shader);
 							ShaderStages.Add(Shader);
 						}
diff --git a/Engine/Source/Runtime/Engine/Public/MaterialShaderType.h b/Engine/Source/Runtime/Engine/Public/MaterialShaderType.h
--- a/Engine/Source/Runtime/Engine/Public/MaterialShaderType.h
+++ b/Engine/Source/Runtime/Engine/Public/MaterialShaderType.h
@@ -59,9 +59,13 @@ struct FMaterialShaderPermutationParameters
 	// Material to compile.
 	const FMaterial* Material;
 
-	FMaterialShaderPermutationParameters(EShaderPlatform InPlatform, const FMaterial* InMaterial)
+	/** Unique permutation identifier of the material shader type. */
+	const int32 PermutationId;
+
+	FMaterialShaderPermutationParameters(EShaderPlatform InPlatform, const FMaterial* InMaterial, int32 InPermutationId)
 		: Platform(InPlatform)
 		, Material(InMaterial)
+		, PermutationId(InPermutationId)
 	{
 	}
 };
@@ -79,6 +83,7 @@ public:
 
 		CompiledShaderInitializerType(
 			FShaderType* InType,
+			int32 InPermutationId,
 			const FShaderCompilerOutput& CompilerOutput,
 			FShaderResource* InResource,
 			const FUniformExpressionSet& InUniformExpressionSet,
@@ -87,16 +92,16 @@ public:
 			FVertexFactoryType* InVertexFactoryType,
 			const FString& InDebugDescription
 			)
-		: FGlobalShaderType::CompiledShaderInitializerType(InType,/** PermutationId = */ 0,CompilerOutput,InResource,InMaterialShaderMapHash,InShaderPipeline,InVertexFactoryType)
+		: FGlobalShaderType::CompiledShaderInitializerType(InType,InPermutationId,CompilerOutput,InResource,InMaterialShaderMapHash,InShaderPipeline,InVertexFactoryType)
 		, UniformExpressionSet(InUniformExpressionSet)
 		, DebugDescription(InDebugDescription)
 		{}
 	};
 
 	typedef FShader* (*ConstructCompiledType)(const CompiledShaderInitializerType&);
-	typedef bool (*ShouldCompilePermutationType)(EShaderPlatform,const FMaterial*);
+	typedef bool (*ShouldCompilePermutationType)(const FMaterialShaderPermutationParameters&);
 	typedef bool(*ValidateCompiledResultType)(EShaderPlatform, const TArray<FMaterial*>&, const FShaderParameterMap&, TArray<FString>&);
-	typedef void (*ModifyCompilationEnvironmentType)(EShaderPlatform, const FMaterial*, FShaderCompilerEnvironment&);
+	typedef void (*ModifyCompilationEnvironmentType)(const FMaterialShaderPermutationParameters&, FShaderCompilerEnvironment&);
 
 	FMaterialShaderType(
 		const TCHAR* InName,
@@ -129,6 +134,7 @@ public:
 	 */
 	class FShaderCompileJob* BeginCompileShader(
 		uint32 ShaderMapId,
+		int32 PermutationId,
 		const FMaterial* Material,
 		FShaderCompilerEnvironment* MaterialEnvironment,
 		const FShaderPipelineType* ShaderPipeline,
@@ -169,9 +175,9 @@ public:
 	 * @param Material - The material to check.
 	 * @return True if this shader type should be cached.
 	 */
-	bool ShouldCache(EShaderPlatform Platform,const FMaterial* Material) const
+	bool ShouldCompilePermutation(EShaderPlatform Platform, const FMaterial* Material, int32 PermutationId) const
 	{
-		return (*ShouldCompilePermutationRef)(Platform,Material);
+		return (*ShouldCompilePermutationRef)(FMaterialShaderPermutationParameters(Platform, Material, PermutationId));
 	}
 
 	/**
@@ -193,10 +199,10 @@ protected:
 	 * @param Platform - Platform to compile for.
 	 * @param Environment - The shader compile environment that the function modifies.
 	 */
-	void SetupCompileEnvironment(EShaderPlatform Platform, const FMaterial* Material, FShaderCompilerEnvironment& Environment)
+	void SetupCompileEnvironment(EShaderPlatform Platform, const FMaterial* Material, int32 PermutationId, FShaderCompilerEnvironment& Environment)
 	{
 		// Allow the shader type to modify its compile environment.
-		(*ModifyCompilationEnvironmentRef)(Platform, Material, Environment);
+		(*ModifyCompilationEnvironmentRef)(FMaterialShaderPermutationParameters(Platform, Material, PermutationId), Environment);
 	}
 
 private:
@@ -205,3 +211,57 @@ private:
 	ValidateCompiledResultType ValidateCompiledResultRef;
 	ModifyCompilationEnvironmentType ModifyCompilationEnvironmentRef;
 };
+
+/** DECLARE_MATERIAL_SHADER and IMPLEMENT_MATERIAL_SHADER setup a material shader class's boiler plate. They are meant to be used like so:
+ *
+ * class FMyMaterialShaderPS : public FMaterialShader
+ * {
+ *		// Setup the shader's boiler plate.
+ *		DECLARE_MATERIAL_SHADER(FMyMaterialShaderPS);
+ *
+ *		// Setup the shader's permutation domain. If no dimensions, can do FPermutationDomain = FShaderPermutationNone.
+ *		using FPermutationDomain = TShaderPermutationDomain<DIMENSIONS...>;
+ *
+ *		// ...
+ * };
+ *
+ * // Instantiates shader's global variable that will take care of compilation process of the shader. This needs imperatively to be
+ * done in a .cpp file regardless of whether FMyMaterialShaderPS is in a header or not.
+ * IMPLEMENT_MATERIAL_SHADER(FMyMaterialShaderPS, "/Engine/Private/MyShaderFile.usf", "MainPS", SF_Pixel);
+ */
+#define DECLARE_MATERIAL_SHADER(ShaderClass) \
+	public: \
+	using ShaderMetaType = FMaterialShaderType; \
+	\
+	static ShaderMetaType StaticType; \
+	\
+	static FShader* ConstructSerializedInstance() { return new ShaderClass(); } \
+	static FShader* ConstructCompiledInstance(const ShaderMetaType::CompiledShaderInitializerType& Initializer) \
+	{ return new ShaderClass(Initializer); } \
+	\
+	virtual uint32 GetTypeSize() const override { return sizeof(*this); } \
+	\
+	static void ModifyCompilationEnvironmentImpl( \
+		const FGlobalShaderPermutationParameters& Parameters, \
+		FShaderCompilerEnvironment& OutEnvironment) \
+	{ \
+		FPermutationDomain PermutationVector(Parameters.PermutationId); \
+		PermutationVector.ModifyCompilationEnvironment(OutEnvironment); \
+		ShaderClass::ModifyCompilationEnvironment(Parameters, OutEnvironment); \
+	}
+
+#define IMPLEMENT_MATERIAL_SHADER(ShaderClass,SourceFilename,FunctionName,Frequency) \
+	ShaderClass::ShaderMetaType ShaderClass::StaticType( \
+		TEXT(#ShaderClass), \
+		TEXT(SourceFilename), \
+		TEXT(FunctionName), \
+		Frequency, \
+		ShaderClass::FPermutationDomain::PermutationCount, \
+		ShaderClass::ConstructSerializedInstance, \
+		ShaderClass::ConstructCompiledInstance, \
+		ShaderClass::ModifyCompilationEnvironmentImpl, \
+		ShaderClass::ShouldCompilePermutation, \
+		ShaderClass::ValidateCompiledResult, \
+		ShaderClass::GetStreamOutElements, \
+		ShaderClass::GetRootParametersMetadata() \
+		)
diff --git a/Engine/Source/Runtime/Engine/Public/MaterialShared.h b/Engine/Source/Runtime/Engine/Public/MaterialShared.h
--- a/Engine/Source/Runtime/Engine/Public/MaterialShared.h
+++ b/Engine/Source/Runtime/Engine/Public/MaterialShared.h
@@ -810,7 +810,7 @@ public:
 	void AppendKeyString(FString& KeyString) const;
 
 	/** Returns true if the requested shader type is a dependency of this shader map Id. */
-	bool ContainsShaderType(const FShaderType* ShaderType) const;
+	bool ContainsShaderType(const FShaderType* ShaderType, int32 PermutationId) const;
 
 	/** Returns true if the requested shader type is a dependency of this shader map Id. */
 	bool ContainsShaderPipelineType(const FShaderPipelineType* ShaderPipelineType) const;
@@ -861,8 +861,7 @@ public:
 		EShaderPlatform Platform,
 		const FMaterial* Material,
 		FVertexFactoryType* InVertexFactoryType,
-		bool bSilent
-		);
+		bool bSilent);
 
 #if WITH_EDITOR
 	void LoadMissingShadersFromMemory(
@@ -890,7 +889,15 @@ private:
 	/** The vertex factory type these shaders are for. */
 	FVertexFactoryType* VertexFactoryType;
 
-	static bool IsMeshShaderComplete(const FMeshMaterialShaderMap* MeshShaderMap, EShaderPlatform Platform, const FMaterial* Material, const FMeshMaterialShaderType* ShaderType, const FShaderPipelineType* Pipeline, FVertexFactoryType* InVertexFactoryType, bool bSilent);
+	static bool IsMeshShaderComplete(
+		const FMeshMaterialShaderMap* MeshShaderMap,
+		EShaderPlatform Platform,
+		const FMaterial* Material,
+		const FMeshMaterialShaderType* ShaderType,
+		const FShaderPipelineType* Pipeline,
+		FVertexFactoryType* InVertexFactoryType,
+		int32 PermutationId,
+		bool bSilent);
 };
 
 /**
@@ -1193,7 +1200,12 @@ private:
 
 	FShader* ProcessCompilationResultsForSingleJob(class FShaderCompileJob* SingleJob, const FShaderPipelineType* ShaderPipeline, const FSHAHash& MaterialShaderMapHash);
 
-	bool IsMaterialShaderComplete(const FMaterial* Material, const FMaterialShaderType* ShaderType, const FShaderPipelineType* Pipeline, bool bSilent);
+	bool IsMaterialShaderComplete(
+		const FMaterial* Material,
+		const FMaterialShaderType* ShaderType,
+		const FShaderPipelineType* Pipeline,
+		int32 PermutationId,
+		bool bSilent) const;
 
 	/** Initializes OrderedMeshShaderMaps from the contents of MeshShaderMaps. */
 	void InitOrderedMeshShaderMaps();
@@ -1622,9 +1634,15 @@ public:
 	 * Note - Only implemented for FMeshMaterialShaderTypes
 	 */
 	template<typename ShaderType>
-	ShaderType* GetShader(FVertexFactoryType* VertexFactoryType, bool bFatalIfMissing = true) const
+	ShaderType* GetShader(FVertexFactoryType* VertexFactoryType, const typename ShaderType::FPermutationDomain& PermutationVector, bool bFatalIfMissing = true) const
+	{
+		return GetShader<ShaderType>(VertexFactoryType, PermutationVector.ToDimensionValueId(), bFatalIfMissing);
+	}
+
+	template <typename ShaderType>
+	ShaderType* GetShader(FVertexFactoryType* VertexFactoryType, int32 PermutationId = 0, bool bFatalIfMissing = true) const
 	{
-		return (ShaderType*)GetShader(&ShaderType::StaticType, VertexFactoryType, bFatalIfMissing);
+		return (ShaderType*)GetShader(&ShaderType::StaticType, VertexFactoryType, PermutationId, bFatalIfMissing);
 	}
 
 	ENGINE_API FShaderPipeline* GetShaderPipeline(class FShaderPipelineType* ShaderPipelineType, FVertexFactoryType* VertexFactoryType, bool bFatalIfNotFound = true) const;
@@ -1815,7 +1833,7 @@ private:
 	/**
 	 * Finds the shader matching the template type and the passed in vertex factory, asserts if not found.
 	 */
-	ENGINE_API FShader* GetShader(class FMeshMaterialShaderType* ShaderType, FVertexFactoryType* VertexFactoryType, bool bFatalIfMissing = true) const;
+	ENGINE_API FShader* GetShader(class FMeshMaterialShaderType* ShaderType, FVertexFactoryType* VertexFactoryType, int32 PermutationId, bool bFatalIfMissing = true) const;
 
 	void GetReferencedTexturesHash(EShaderPlatform Platform, FSHAHash& OutHash) const;
 
diff --git a/Engine/Source/Runtime/Engine/Public/MeshMaterialShaderType.h b/Engine/Source/Runtime/Engine/Public/MeshMaterialShaderType.h
--- a/Engine/Source/Runtime/Engine/Public/MeshMaterialShaderType.h
+++ b/Engine/Source/Runtime/Engine/Public/MeshMaterialShaderType.h
@@ -27,10 +27,14 @@ struct FMeshMaterialShaderPermutationParameters
 	// Type of vertex factory to compile.
 	const FVertexFactoryType* VertexFactoryType;
 
-	FMeshMaterialShaderPermutationParameters(EShaderPlatform InPlatform, const FMaterial* InMaterial, const FVertexFactoryType* InVertexFactoryType)
+	// Unique permutation identifier of the mesh material shader type.
+	const int32 PermutationId;
+
+	FMeshMaterialShaderPermutationParameters(EShaderPlatform InPlatform, const FMaterial* InMaterial, const FVertexFactoryType* InVertexFactoryType, const int32 InPermutationId)
 		: Platform(InPlatform)
 		, Material(InMaterial)
 		, VertexFactoryType(InVertexFactoryType)
+		, PermutationId(InPermutationId)
 	{
 	}
 };
@@ -46,6 +50,7 @@ public:
 		FVertexFactoryType* VertexFactoryType;
 		CompiledShaderInitializerType(
 			FShaderType* InType,
+			int32 PermutationId,
 			const FShaderCompilerOutput& CompilerOutput,
 			FShaderResource* InResource,
 			const FUniformExpressionSet& InUniformExpressionSet,
@@ -54,14 +59,14 @@ public:
 			const FShaderPipelineType* InShaderPipeline,
 			FVertexFactoryType* InVertexFactoryType
 			):
-			FMaterialShaderType::CompiledShaderInitializerType(InType,CompilerOutput,InResource,InUniformExpressionSet,InMaterialShaderMapHash,InShaderPipeline,InVertexFactoryType,InDebugDescription),
+			FMaterialShaderType::CompiledShaderInitializerType(InType,PermutationId,CompilerOutput,InResource,InUniformExpressionSet,InMaterialShaderMapHash,InShaderPipeline,InVertexFactoryType,InDebugDescription),
 			VertexFactoryType(InVertexFactoryType)
 		{}
 	};
 	typedef FShader* (*ConstructCompiledType)(const CompiledShaderInitializerType&);
-	typedef bool (*ShouldCompilePermutationType)(EShaderPlatform,const FMaterial*,const FVertexFactoryType* VertexFactoryType);
+	typedef bool (*ShouldCompilePermutationType)(const FMeshMaterialShaderPermutationParameters&);
 	typedef bool(*ValidateCompiledResultType)(EShaderPlatform, const TArray<FMaterial*>&, const FVertexFactoryType*, const FShaderParameterMap&, TArray<FString>&);
-	typedef void (*ModifyCompilationEnvironmentType)(EShaderPlatform, const FMaterial*, FShaderCompilerEnvironment&);
+	typedef void (*ModifyCompilationEnvironmentType)(const FMaterialShaderPermutationParameters&, FShaderCompilerEnvironment&);
 
 	FMeshMaterialShaderType(
 		const TCHAR* InName,
@@ -96,6 +101,7 @@ public:
 	 */
 	class FShaderCompileJob* BeginCompileShader(
 		uint32 ShaderMapId,
+		int32 PermutationId,
 		EShaderPlatform Platform,
 		const FMaterial* Material,
 		FShaderCompilerEnvironment* MaterialEnvironment,
@@ -108,6 +114,7 @@ public:
 
 	static void BeginCompileShaderPipeline(
 		uint32 ShaderMapId,
+		int32 PermutationId,
 		EShaderPlatform Platform,
 		const FMaterial* Material,
 		FShaderCompilerEnvironment* MaterialEnvironment,
@@ -139,9 +146,9 @@ public:
 	 * @param VertexFactoryType - The vertex factory type to check.
 	 * @return True if this shader type should be cached.
 	 */
-	bool ShouldCache(EShaderPlatform Platform,const FMaterial* Material,const FVertexFactoryType* VertexFactoryType) const
+	bool ShouldCompilePermutation(EShaderPlatform Platform,const FMaterial* Material,const FVertexFactoryType* VertexFactoryType, int32 PermutationId) const
 	{
-		return (*ShouldCompilePermutationRef)(Platform,Material,VertexFactoryType);
+		return (*ShouldCompilePermutationRef)(FMeshMaterialShaderPermutationParameters(Platform,Material,VertexFactoryType,PermutationId));
 	}
 
 	/**
@@ -164,10 +171,10 @@ protected:
 	 * @param Platform - Platform to compile for.
 	 * @param Environment - The shader compile environment that the function modifies.
 	 */
-	void SetupCompileEnvironment(EShaderPlatform Platform, const FMaterial* Material, FShaderCompilerEnvironment& Environment)
+	void SetupCompileEnvironment(EShaderPlatform Platform, const FMaterial* Material, int32 PermutationId, FShaderCompilerEnvironment& Environment)
 	{
 		// Allow the shader type to modify its compile environment.
-		(*ModifyCompilationEnvironmentRef)(Platform, Material, Environment);
+		(*ModifyCompilationEnvironmentRef)(FMaterialShaderPermutationParameters(Platform, Material, PermutationId), Environment);
 	}
 
 private:
@@ -177,3 +184,56 @@ private:
 	ModifyCompilationEnvironmentType ModifyCompilationEnvironmentRef;
 };
 
+/** DECLARE_MESH_MATERIAL_SHADER and IMPLEMENT_MESH_MATERIAL_SHADER setup a mesh material shader class's boiler plate. They are meant to be used like so:
+ *
+ * class FMyMeshMaterialShaderPS : public FMeshMaterialShader
+ * {
+ *		// Setup the shader's boiler plate.
+ *		DECLARE_MESH_MATERIAL_SHADER(FMyMeshMaterialShaderPS);
+ *
+ *		// Setup the shader's permutation domain. If no dimensions, can do FPermutationDomain = FShaderPermutationNone.
+ *		using FPermutationDomain = TShaderPermutationDomain<DIMENSIONS...>;
+ *
+ *		// ...
+ * };
+ *
+ * // Instantiates shader's global variable that will take care of compilation process of the shader. This needs imperatively to be
+ * done in a .cpp file regardless of whether FMyMeshMaterialShaderPS is in a header or not.
+ * IMPLEMENT_MESH_MATERIAL_SHADER(FMyMaterialShaderPS, "/Engine/Private/MyShaderFile.usf", "MainPS", SF_Pixel);
+ */
+#define DECLARE_MESH_MATERIAL_SHADER(ShaderClass) \
+	public: \
+	using ShaderMetaType = FMeshMaterialShaderType; \
+	\
+	static ShaderMetaType StaticType; \
+	\
+	static FShader* ConstructSerializedInstance() { return new ShaderClass(); } \
+	static FShader* ConstructCompiledInstance(const ShaderMetaType::CompiledShaderInitializerType& Initializer) \
+	{ return new ShaderClass(Initializer); } \
+	\
+	virtual uint32 GetTypeSize() const override { return sizeof(*this); } \
+	\
+	static void ModifyCompilationEnvironmentImpl( \
+		const FGlobalShaderPermutationParameters& Parameters, \
+		FShaderCompilerEnvironment& OutEnvironment) \
+	{ \
+		FPermutationDomain PermutationVector(Parameters.PermutationId); \
+		PermutationVector.ModifyCompilationEnvironment(OutEnvironment); \
+		ShaderClass::ModifyCompilationEnvironment(Parameters, OutEnvironment); \
+	}
+
+#define IMPLEMENT_MESH_MATERIAL_SHADER(ShaderClass,SourceFilename,FunctionName,Frequency) \
+	ShaderClass::ShaderMetaType ShaderClass::StaticType( \
+		TEXT(#ShaderClass), \
+		TEXT(SourceFilename), \
+		TEXT(FunctionName), \
+		Frequency, \
+		ShaderClass::FPermutationDomain::PermutationCount, \
+		ShaderClass::ConstructSerializedInstance, \
+		ShaderClass::ConstructCompiledInstance, \
+		ShaderClass::ModifyCompilationEnvironmentImpl, \
+		ShaderClass::ShouldCompilePermutation, \
+		ShaderClass::ValidateCompiledResult, \
+		ShaderClass::GetStreamOutElements, \
+		ShaderClass::GetRootParametersMetadata() \
+		)
diff --git a/Engine/Source/Runtime/RenderCore/Private/Shader.cpp b/Engine/Source/Runtime/RenderCore/Private/Shader.cpp
--- a/Engine/Source/Runtime/RenderCore/Private/Shader.cpp
+++ b/Engine/Source/Runtime/RenderCore/Private/Shader.cpp
@@ -1500,6 +1500,14 @@ FShaderPipelineType::FShaderPipelineType(
 	Stages.Add(InVertexShader);
 	AllStages[SF_Vertex] = InVertexShader;
 
+	for (uint32 FrequencyIndex = 0; FrequencyIndex < SF_NumStandardFrequencies; ++FrequencyIndex)
+	{
+		if (const FShaderType* ShaderType = AllStages[FrequencyIndex])
+		{
+			checkf(ShaderType->GetPermutationCount() == 1, TEXT("Shader '%s' has multiple shader permutations. Shader pipelines only support a single permutation."), ShaderType->GetName())
+		}
+	}
+
 	static uint32 TypeHashCounter = 0;
 	++TypeHashCounter;
 	HashIndex = TypeHashCounter;
diff --git a/Engine/Source/Runtime/RenderCore/Public/GlobalShader.h b/Engine/Source/Runtime/RenderCore/Public/GlobalShader.h
--- a/Engine/Source/Runtime/RenderCore/Public/GlobalShader.h
+++ b/Engine/Source/Runtime/RenderCore/Public/GlobalShader.h
@@ -253,15 +253,6 @@ inline TShaderMap<FGlobalShaderType>* GetGlobalShaderMap(ERHIFeatureLevel::Type
 	\
 	virtual uint32 GetTypeSize() const override { return sizeof(*this); } \
 	\
-	static bool ShouldCompilePermutationImpl(const FGlobalShaderPermutationParameters& Parameters) \
-	{ \
-		return ShaderClass::ShouldCompilePermutation(Parameters); \
-	} \
-	\
-	static bool ValidateCompiledResultImpl(EShaderPlatform Platform, const FShaderParameterMap& ParameterMap, TArray<FString>& OutErrors) \
-	{ \
-		return ShaderClass::ValidateCompiledResult(Platform, ParameterMap, OutErrors); \
-	} \
 	static void ModifyCompilationEnvironmentImpl( \
 		const FGlobalShaderPermutationParameters& Parameters, \
 		FShaderCompilerEnvironment& OutEnvironment) \
@@ -281,8 +272,8 @@ inline TShaderMap<FGlobalShaderType>* GetGlobalShaderMap(ERHIFeatureLevel::Type
 		ShaderClass::ConstructSerializedInstance, \
 		ShaderClass::ConstructCompiledInstance, \
 		ShaderClass::ModifyCompilationEnvironmentImpl, \
-		ShaderClass::ShouldCompilePermutationImpl, \
-		ShaderClass::ValidateCompiledResultImpl, \
+		ShaderClass::ShouldCompilePermutation, \
+		ShaderClass::ValidateCompiledResult, \
 		ShaderClass::GetStreamOutElements, \
 		ShaderClass::GetRootParametersMetadata() \
 		)
diff --git a/Engine/Source/Runtime/RenderCore/Public/Shader.h b/Engine/Source/Runtime/RenderCore/Public/Shader.h
--- a/Engine/Source/Runtime/RenderCore/Public/Shader.h
+++ b/Engine/Source/Runtime/RenderCore/Public/Shader.h
@@ -65,6 +65,8 @@ struct TShaderTypePermutation
 
 using FShaderPermutation = TShaderTypePermutation<FShaderType>;
 
+const int32 kUniqueShaderPermutationId = 0;
+
 template<typename MetaShaderType>
 FORCEINLINE uint32 GetTypeHash(const TShaderTypePermutation<MetaShaderType>& Var)
 {
diff --git a/Engine/Source/Runtime/Renderer/Private/PostProcess/PostProcessMaterial.cpp b/Engine/Source/Runtime/Renderer/Private/PostProcess/PostProcessMaterial.cpp
--- a/Engine/Source/Runtime/Renderer/Private/PostProcess/PostProcessMaterial.cpp
+++ b/Engine/Source/Runtime/Renderer/Private/PostProcess/PostProcessMaterial.cpp
@@ -67,23 +67,23 @@ public:
 	/**
 	  * Only compile these shaders for post processing domain materials
 	  */
-	static bool ShouldCompilePermutation(EShaderPlatform Platform, const FMaterial* Material)
-	{		
-		return ShouldCachePostProcessMaterial(MaterialTarget, Platform, Material)
+	static bool ShouldCompilePermutation(const FMaterialShaderPermutationParameters& Parameters)
+	{
+		return ShouldCachePostProcessMaterial(MaterialTarget, Parameters.Platform, Parameters.Material)
 			// compile mobile axis switched versions only for after tonemapper passes
-			&& (bSwitchVerticalAxis == false || (RHINeedsToSwitchVerticalAxis(Platform) && Material->GetBlendableLocation() == BL_AfterTonemapping));
+			&& (bSwitchVerticalAxis == false || (RHINeedsToSwitchVerticalAxis(Parameters.Platform) && Parameters.Material->GetBlendableLocation() == BL_AfterTonemapping));
 	}
 
-	static void ModifyCompilationEnvironment(EShaderPlatform Platform, const class FMaterial* Material, FShaderCompilerEnvironment& OutEnvironment)
+	static void ModifyCompilationEnvironment(const FMaterialShaderPermutationParameters& Parameters, FShaderCompilerEnvironment& OutEnvironment)
 	{
-		FMaterialShader::ModifyCompilationEnvironment(Platform, OutEnvironment);
+		FMaterialShader::ModifyCompilationEnvironment(Parameters, OutEnvironment);
 
 		OutEnvironment.SetDefine(TEXT("POST_PROCESS_MATERIAL"), 1);
 		OutEnvironment.SetDefine(TEXT("NEEDTOSWITCHVERTICLEAXIS"), bSwitchVerticalAxis);
 
 		if (MaterialTarget == EPostProcessMaterialTarget::Mobile)
 		{
-			OutEnvironment.SetDefine(TEXT("POST_PROCESS_MATERIAL_BEFORE_TONEMAP"), (Material->GetBlendableLocation() != BL_AfterTonemapping) ? 1 : 0);
+			OutEnvironment.SetDefine(TEXT("POST_PROCESS_MATERIAL_BEFORE_TONEMAP"), (Parameters.Material->GetBlendableLocation() != BL_AfterTonemapping) ? 1 : 0);
 		}
 	}
 	
@@ -133,21 +133,21 @@ public:
 	/**
 	  * Only compile these shaders for post processing domain materials
 	  */
-	static bool ShouldCompilePermutation(EShaderPlatform Platform, const FMaterial* Material)
+	static bool ShouldCompilePermutation(const FMaterialShaderPermutationParameters& Parameters)
 	{
-		return ShouldCachePostProcessMaterial(MaterialTarget, Platform, Material)
+		return ShouldCachePostProcessMaterial(MaterialTarget, Parameters.Platform, Parameters.Material)
 			// compile mobile axis switched versions only for after tonemapper passes
-			&& (bSwitchVerticalAxis == false || (RHINeedsToSwitchVerticalAxis(Platform) && Material->GetBlendableLocation() == BL_AfterTonemapping));
+			&& (bSwitchVerticalAxis == false || (RHINeedsToSwitchVerticalAxis(Parameters.Platform) && Parameters.Material->GetBlendableLocation() == BL_AfterTonemapping));
 	}
 
-	static void ModifyCompilationEnvironment(EShaderPlatform Platform, const class FMaterial* Material, FShaderCompilerEnvironment& OutEnvironment)
+	static void ModifyCompilationEnvironment(const FMaterialShaderPermutationParameters& Parameters, FShaderCompilerEnvironment& OutEnvironment)
 	{
-		FMaterialShader::ModifyCompilationEnvironment(Platform, OutEnvironment);
+		FMaterialShader::ModifyCompilationEnvironment(Parameters, OutEnvironment);
 
 		OutEnvironment.SetDefine(TEXT("POST_PROCESS_MATERIAL"), 1);
 		OutEnvironment.SetDefine(TEXT("POST_PROCESS_MATERIAL_UV_POLICY"), UVPolicy);
 
-		EBlendableLocation Location = EBlendableLocation(Material->GetBlendableLocation());
+		EBlendableLocation Location = EBlendableLocation(Parameters.Material->GetBlendableLocation());
 		OutEnvironment.SetDefine(TEXT("POST_PROCESS_MATERIAL_AFTER_TAA_UPSAMPLE"), (Location == BL_AfterTonemapping || Location == BL_ReplacingTonemapper) ? 1 : 0);
 
 		OutEnvironment.SetDefine(TEXT("NEEDTOSWITCHVERTICLEAXIS"), bSwitchVerticalAxis);
diff --git a/Engine/Source/Runtime/Renderer/Public/MaterialShader.h b/Engine/Source/Runtime/Renderer/Public/MaterialShader.h
--- a/Engine/Source/Runtime/Renderer/Public/MaterialShader.h
+++ b/Engine/Source/Runtime/Renderer/Public/MaterialShader.h
@@ -114,7 +114,7 @@ public:
 
 	typedef void (*ModifyCompilationEnvironmentType)(EShaderPlatform, const FMaterial*, FShaderCompilerEnvironment&);
 
-	static void ModifyCompilationEnvironment(EShaderPlatform Platform, const FMaterial* Material, FShaderCompilerEnvironment& OutEnvironment)
+	static void ModifyCompilationEnvironment(const FMaterialShaderPermutationParameters& Parameters, FShaderCompilerEnvironment& OutEnvironment)
 	{
 	}
 
