Renames confusing texture sampling input View MipBias to Apply View MipBias

diff --git a/Engine/Source/Runtime/Engine/Private/Materials/MaterialExpressions.cpp b/Engine/Source/Runtime/Engine/Private/Materials/MaterialExpressions.cpp
--- a/Engine/Source/Runtime/Engine/Private/Materials/MaterialExpressions.cpp
+++ b/Engine/Source/Runtime/Engine/Private/Materials/MaterialExpressions.cpp
@@ -1619,7 +1619,7 @@ FName UMaterialExpressionTextureSample::GetInputName(int32 InputIndex) const
 		IF_INPUT_RETURN(CoordinatesDY, TEXT("DDY(UVs)"));
 	}
 
-	IF_INPUT_RETURN(AutomaticViewMipBiasValue, TEXT("View MipBias"));
+	IF_INPUT_RETURN(AutomaticViewMipBiasValue, TEXT("Apply View MipBias"));
 
 	return TEXT("");
 }
