Refactored canvas triangle rendering by removing dead code and cleaning up data flow a bit.

diff --git a/Engine/Source/Runtime/Engine/Private/TriangleRendering.cpp b/Engine/Source/Runtime/Engine/Private/TriangleRendering.cpp
--- a/Engine/Source/Runtime/Engine/Private/TriangleRendering.cpp
+++ b/Engine/Source/Runtime/Engine/Private/TriangleRendering.cpp
@@ -19,214 +19,57 @@
 #include "EngineModule.h"
 #include "MeshPassProcessor.h"
 
-/** 
-* Vertex buffer
-*/
-class FMaterialTriangleVertexBuffer : public FVertexBuffer
-{
-public:
-	FVertexBuffer PositionBuffer;
-	FVertexBuffer TangentBuffer;
-	FVertexBuffer TexCoordBuffer;
-	FVertexBuffer ColorBuffer;
-
-	FShaderResourceViewRHIRef TangentBufferSRV;
-	FShaderResourceViewRHIRef TexCoordBufferSRV;
-	FShaderResourceViewRHIRef ColorBufferSRV;
-	FShaderResourceViewRHIRef PositionBufferSRV;
-
-	/** 
-	* Initialize the RHI for this rendering resource 
-	*/
-	virtual void InitRHI() override
-	{
-		FVector* PositionBufferData = nullptr;
-		uint32 PositionSize = 3 * sizeof(FVector);
-		// create vertex buffer
-		{
-			FRHIResourceCreateInfo CreateInfo;
-			void* Data = nullptr;
-			PositionBuffer.VertexBufferRHI = RHICreateAndLockVertexBuffer(PositionSize, BUF_Static | BUF_ShaderResource, CreateInfo, Data);
-			PositionBufferData = static_cast<FVector*>(Data);
-			if (RHISupportsManualVertexFetch(GMaxRHIShaderPlatform))
-			{
-				PositionBufferSRV = RHICreateShaderResourceView(PositionBuffer.VertexBufferRHI, sizeof(float), PF_R32_FLOAT);
-			}
-
-		}
-
-		FPackedNormal* TangentBufferData = nullptr;
-		uint32 TangentSize = 3 * 2 * sizeof(FPackedNormal);
-		// create vertex buffer
-		{
-			FRHIResourceCreateInfo CreateInfo;
-			void* Data = nullptr;
-			TangentBuffer.VertexBufferRHI = RHICreateAndLockVertexBuffer(TangentSize, BUF_Static | BUF_ShaderResource, CreateInfo, Data);
-			TangentBufferData = static_cast<FPackedNormal*>(Data);
-			if (RHISupportsManualVertexFetch(GMaxRHIShaderPlatform))
-			{
-				TangentBufferSRV = RHICreateShaderResourceView(TangentBuffer.VertexBufferRHI, sizeof(FPackedNormal), PF_R8G8B8A8_SNORM);
-			}
-		}
-
-		FVector2D* TexCoordBufferData = nullptr;
-		uint32 TexCoordSize = 3 * sizeof(FVector2D);
-		// create vertex buffer
-		{
-			FRHIResourceCreateInfo CreateInfo;
-			void* Data = nullptr;
-			TexCoordBuffer.VertexBufferRHI = RHICreateAndLockVertexBuffer(TexCoordSize, BUF_Static | BUF_ShaderResource, CreateInfo, Data);
-			TexCoordBufferData = static_cast<FVector2D*>(Data);
-			if (RHISupportsManualVertexFetch(GMaxRHIShaderPlatform))
-			{
-				TexCoordBufferSRV = RHICreateShaderResourceView(TexCoordBuffer.VertexBufferRHI, sizeof(FVector2D), PF_G32R32F);
-			}
-		}
-
-		uint32* ColorBufferData = nullptr;
-		uint32 ColorSize = 3 * sizeof(uint32);
-		// create vertex buffer
-		{
-			FRHIResourceCreateInfo CreateInfo;
-			void* Data = nullptr;
-			ColorBuffer.VertexBufferRHI = RHICreateAndLockVertexBuffer(ColorSize, BUF_Static | BUF_ShaderResource, CreateInfo, Data);
-			ColorBufferData = static_cast<uint32*>(Data);
-			if (RHISupportsManualVertexFetch(GMaxRHIShaderPlatform))
-			{
-				ColorBufferSRV = RHICreateShaderResourceView(ColorBuffer.VertexBufferRHI, sizeof(uint32), PF_R8G8B8A8);
-			}
-		}
-
-		// fill out the verts
-		PositionBufferData[0] = FVector(1, -1, 0);
-		PositionBufferData[1] = FVector(1, 1, 0);
-		PositionBufferData[2] = FVector(-1, -1, 0);
-
-		TexCoordBufferData[0] = FVector2D(1, 1);
-		TexCoordBufferData[1] = FVector2D(1, 0);
-		TexCoordBufferData[2] = FVector2D(0, 1);
-
-		for (int i = 0; i < 3; i++)
-		{
-			TangentBufferData[2 * i + 0] = FVector(1, 0, 0);
-			TangentBufferData[2 * i + 1] = FVector(0, 0, 1);
-			ColorBufferData[i] = FColor(255, 255, 255, 255).DWColor();
-		}
-
-		// Unlock the buffer.
-		RHIUnlockVertexBuffer(PositionBuffer.VertexBufferRHI);
-		RHIUnlockVertexBuffer(TangentBuffer.VertexBufferRHI);
-		RHIUnlockVertexBuffer(TexCoordBuffer.VertexBufferRHI);
-		RHIUnlockVertexBuffer(ColorBuffer.VertexBufferRHI);
-	}
-
-	virtual void ReleaseRHI() override
-	{
-		PositionBuffer.ReleaseRHI();
-		TangentBuffer.ReleaseRHI();
-		TexCoordBuffer.ReleaseRHI();
-		ColorBuffer.ReleaseRHI();
-
-		TangentBufferSRV.SafeRelease();
-		TexCoordBufferSRV.SafeRelease();
-		ColorBufferSRV.SafeRelease();
-		PositionBufferSRV.SafeRelease();
-	}
-};
-TGlobalResource<FMaterialTriangleVertexBuffer> GTriangleRendererVertexBuffer;
-
-/**
-* Vertex factory for rendering tiles.
-*/
-FCanvasTriangleRendererItem::FTriangleVertexFactory::FTriangleVertexFactory(ERHIFeatureLevel::Type InFeatureLevel)
+FCanvasTriangleRendererItem::FTriangleVertexFactory::FTriangleVertexFactory(
+	const FStaticMeshVertexBuffers* InVertexBuffers,
+	ERHIFeatureLevel::Type InFeatureLevel)
 	: FLocalVertexFactory(InFeatureLevel, "FTriangleVertexFactory")
-{
-	FLocalVertexFactory::FDataType VertexData;
-	VertexData.NumTexCoords = 1;
-
-	{
-		VertexData.LightMapCoordinateIndex = 0;
-		VertexData.TangentsSRV = GTriangleRendererVertexBuffer.TangentBufferSRV;
-		VertexData.TextureCoordinatesSRV = GTriangleRendererVertexBuffer.TexCoordBufferSRV;
-		VertexData.ColorComponentsSRV = GTriangleRendererVertexBuffer.ColorBufferSRV;
-		VertexData.PositionComponentSRV = GTriangleRendererVertexBuffer.PositionBufferSRV;
-	}
-		
-	{
-		// position
-		Data.PositionComponent = FVertexStreamComponent(
-			&GTriangleRendererVertexBuffer.PositionBuffer,
-			0,
-			sizeof(FVector),
-			VET_Float3,
-			EVertexStreamUsage::Default
-		);
-
-		// tangents
-		VertexData.TangentBasisComponents[0] = FVertexStreamComponent(&GTriangleRendererVertexBuffer.TangentBuffer, 0, 2 * sizeof(FPackedNormal), VET_PackedNormal, EVertexStreamUsage::ManualFetch);
-		VertexData.TangentBasisComponents[1] = FVertexStreamComponent(&GTriangleRendererVertexBuffer.TangentBuffer, sizeof(FPackedNormal), 2 * sizeof(FPackedNormal), VET_PackedNormal, EVertexStreamUsage::ManualFetch);
-
-		// color
-		VertexData.ColorComponent = FVertexStreamComponent(&GTriangleRendererVertexBuffer.ColorBuffer, 0, sizeof(FColor), VET_Color, EVertexStreamUsage::ManualFetch);
-		// UVs
-		VertexData.TextureCoordinates.Add(FVertexStreamComponent(
-			&GTriangleRendererVertexBuffer.TexCoordBuffer,
-			0,
-			sizeof(FVector2D),
-			VET_Float2,
-			EVertexStreamUsage::ManualFetch
-		));
-
-		FCanvasTriangleRendererItem::FTriangleVertexFactory* Self = this;
-		ENQUEUE_RENDER_COMMAND(FCanvasTriangleRendererItemInit)(
-			[Self, VertexData](FRHICommandListImmediate& RHICmdList)
-		{
-			// update the data
-			Self->SetData(VertexData);
-		});
-	}
-}
+	, VertexBuffers(InVertexBuffers)
+{}
 
-/**
- * Mesh used to render triangles.
- */
-FCanvasTriangleRendererItem::FTriangleMesh::FTriangleMesh(const FRawIndexBuffer* InIndexBuffer, const FTriangleVertexFactory* InVertexFactory)
+FCanvasTriangleRendererItem::FTriangleMesh::FTriangleMesh(
+	const FRawIndexBuffer* InIndexBuffer,
+	const FTriangleVertexFactory* InVertexFactory)
 	: IndexBuffer(InIndexBuffer)
 	, VertexFactory(InVertexFactory)
-{
+{}
 
+void FCanvasTriangleRendererItem::FTriangleVertexFactory::InitResource()
+{
+	FLocalVertexFactory::FDataType VertexData;
+	VertexBuffers->PositionVertexBuffer.BindPositionVertexBuffer(this, VertexData);
+	VertexBuffers->StaticMeshVertexBuffer.BindTangentVertexBuffer(this, VertexData);
+	VertexBuffers->StaticMeshVertexBuffer.BindPackedTexCoordVertexBuffer(this, VertexData);
+	VertexBuffers->StaticMeshVertexBuffer.BindLightMapVertexBuffer(this, VertexData, 0);
+	VertexBuffers->ColorVertexBuffer.BindColorVertexBuffer(this, VertexData);
+	SetData(VertexData);
+
+	FLocalVertexFactory::InitResource();
 }
 
-/** The mesh element. */
-FMeshBatch TriMeshElement;
-
 void FCanvasTriangleRendererItem::FTriangleMesh::InitRHI()
 {
-	FMeshBatchElement& BatchElement = TriMeshElement.Elements[0];
-	TriMeshElement.VertexFactory = VertexFactory;
-	BatchElement.IndexBuffer = IndexBuffer;
-	BatchElement.FirstIndex = 0;
-	BatchElement.NumPrimitives = 1;
-	BatchElement.MinVertexIndex = 0;
-	BatchElement.MaxVertexIndex = 2;
-	TriMeshElement.ReverseCulling = false;
-	TriMeshElement.bDisableBackfaceCulling = true;
-	TriMeshElement.Type = PT_TriangleList;
-	TriMeshElement.DepthPriorityGroup = SDPG_Foreground;
-	BatchElement.PrimitiveUniformBufferResource = &GIdentityPrimitiveUniformBuffer;
-}
-	
-void FCanvasTriangleRendererItem::FTriangleMesh::ReleaseRHI()
-{
+	MeshBatch.VertexFactory = VertexFactory;
+	MeshBatch.ReverseCulling = false;
+	MeshBatch.bDisableBackfaceCulling = true;
+	MeshBatch.Type = PT_TriangleList;
+	MeshBatch.DepthPriorityGroup = SDPG_Foreground;
+
+	FMeshBatchElement& MeshBatchElement = MeshBatch.Elements[0];
+	MeshBatchElement.IndexBuffer = IndexBuffer;
+	MeshBatchElement.FirstIndex = 0;
+	MeshBatchElement.NumPrimitives = 1;
+	MeshBatchElement.MinVertexIndex = 0;
+	MeshBatchElement.MaxVertexIndex = 2;
+	MeshBatchElement.PrimitiveUniformBufferResource = &GIdentityPrimitiveUniformBuffer;
 }
 
-void FCanvasTriangleRendererItem::InitTriangleBuffers(FLocalVertexFactory* VertexFactory, TArray<FTriangleInst>& Triangles, const FSceneView& View, bool bNeedsToSwitchVerticalAxis)
+void FCanvasTriangleRendererItem::FRenderData::InitTriangleMesh(const FSceneView& View, bool bNeedsToSwitchVerticalAxis)
 {
-	Data->StaticMeshVertexBuffers.PositionVertexBuffer.Init(Triangles.Num() * 3);
-	Data->StaticMeshVertexBuffers.StaticMeshVertexBuffer.Init(Triangles.Num() * 3, 1);
-	Data->StaticMeshVertexBuffers.ColorVertexBuffer.Init(Triangles.Num() * 3);
+	StaticMeshVertexBuffers.PositionVertexBuffer.Init(Triangles.Num() * 3);
+	StaticMeshVertexBuffers.StaticMeshVertexBuffer.Init(Triangles.Num() * 3, 1);
+	StaticMeshVertexBuffers.ColorVertexBuffer.Init(Triangles.Num() * 3);
 
-	Data->IndexBuffer.Indices.SetNum(Triangles.Num() * 3);
+	IndexBuffer.Indices.SetNum(Triangles.Num() * 3);
 
 	for (int32 i = 0; i < Triangles.Num(); i++)
 	{
@@ -235,61 +78,87 @@ void FCanvasTriangleRendererItem::InitTriangleBuffers(FLocalVertexFactory* Verte
 		/** The use of an index buffer here is actually necessary to workaround an issue with BaseVertexIndex, DrawPrimitive, and manual vertex fetch.
 		 *  In short, DrawIndexedPrimitive with StartIndex map SV_VertexId to the correct location, but DrawPrimitive with BaseVertexIndex will not.
 		 */
-		Data->IndexBuffer.Indices[StartIndex + 0] = StartIndex + 0;
-		Data->IndexBuffer.Indices[StartIndex + 1] = StartIndex + 1;
-		Data->IndexBuffer.Indices[StartIndex + 2] = StartIndex + 2;
+		IndexBuffer.Indices[StartIndex + 0] = StartIndex + 0;
+		IndexBuffer.Indices[StartIndex + 1] = StartIndex + 1;
+		IndexBuffer.Indices[StartIndex + 2] = StartIndex + 2;
 
 		const FCanvasUVTri& Tri = Triangles[i].Tri;
 
 		// create verts. Notice the order is (1, 0, 2)
 		if (bNeedsToSwitchVerticalAxis)
 		{
-			Data->StaticMeshVertexBuffers.PositionVertexBuffer.VertexPosition(StartIndex + 0) = FVector(Tri.V1_Pos.X, View.UnscaledViewRect.Height() - Tri.V1_Pos.Y, 0.0f);
-			Data->StaticMeshVertexBuffers.PositionVertexBuffer.VertexPosition(StartIndex + 1) = FVector(Tri.V0_Pos.X, View.UnscaledViewRect.Height() - Tri.V0_Pos.Y, 0.0f);
-			Data->StaticMeshVertexBuffers.PositionVertexBuffer.VertexPosition(StartIndex + 2) = FVector(Tri.V2_Pos.X, View.UnscaledViewRect.Height() - Tri.V2_Pos.Y, 0.0f);
+			StaticMeshVertexBuffers.PositionVertexBuffer.VertexPosition(StartIndex + 0) = FVector(Tri.V1_Pos.X, View.UnscaledViewRect.Height() - Tri.V1_Pos.Y, 0.0f);
+			StaticMeshVertexBuffers.PositionVertexBuffer.VertexPosition(StartIndex + 1) = FVector(Tri.V0_Pos.X, View.UnscaledViewRect.Height() - Tri.V0_Pos.Y, 0.0f);
+			StaticMeshVertexBuffers.PositionVertexBuffer.VertexPosition(StartIndex + 2) = FVector(Tri.V2_Pos.X, View.UnscaledViewRect.Height() - Tri.V2_Pos.Y, 0.0f);
 		}
 		else
 		{
-			Data->StaticMeshVertexBuffers.PositionVertexBuffer.VertexPosition(StartIndex + 0) = FVector(Tri.V1_Pos.X, Tri.V1_Pos.Y, 0.0f);
-			Data->StaticMeshVertexBuffers.PositionVertexBuffer.VertexPosition(StartIndex + 1) = FVector(Tri.V0_Pos.X, Tri.V0_Pos.Y, 0.0f);
-			Data->StaticMeshVertexBuffers.PositionVertexBuffer.VertexPosition(StartIndex + 2) = FVector(Tri.V2_Pos.X, Tri.V2_Pos.Y, 0.0f);
+			StaticMeshVertexBuffers.PositionVertexBuffer.VertexPosition(StartIndex + 0) = FVector(Tri.V1_Pos.X, Tri.V1_Pos.Y, 0.0f);
+			StaticMeshVertexBuffers.PositionVertexBuffer.VertexPosition(StartIndex + 1) = FVector(Tri.V0_Pos.X, Tri.V0_Pos.Y, 0.0f);
+			StaticMeshVertexBuffers.PositionVertexBuffer.VertexPosition(StartIndex + 2) = FVector(Tri.V2_Pos.X, Tri.V2_Pos.Y, 0.0f);
 		}
 
-		Data->StaticMeshVertexBuffers.StaticMeshVertexBuffer.SetVertexTangents(StartIndex + 0, FVector(1.0f, 0.0f, 0.0f), FVector(0.0f, 1.0f, 0.0f), FVector(0.0f, 0.0f, 1.0f));
-		Data->StaticMeshVertexBuffers.StaticMeshVertexBuffer.SetVertexTangents(StartIndex + 1, FVector(1.0f, 0.0f, 0.0f), FVector(0.0f, 1.0f, 0.0f), FVector(0.0f, 0.0f, 1.0f));
-		Data->StaticMeshVertexBuffers.StaticMeshVertexBuffer.SetVertexTangents(StartIndex + 2, FVector(1.0f, 0.0f, 0.0f), FVector(0.0f, 1.0f, 0.0f), FVector(0.0f, 0.0f, 1.0f));
+		StaticMeshVertexBuffers.StaticMeshVertexBuffer.SetVertexTangents(StartIndex + 0, FVector(1.0f, 0.0f, 0.0f), FVector(0.0f, 1.0f, 0.0f), FVector(0.0f, 0.0f, 1.0f));
+		StaticMeshVertexBuffers.StaticMeshVertexBuffer.SetVertexTangents(StartIndex + 1, FVector(1.0f, 0.0f, 0.0f), FVector(0.0f, 1.0f, 0.0f), FVector(0.0f, 0.0f, 1.0f));
+		StaticMeshVertexBuffers.StaticMeshVertexBuffer.SetVertexTangents(StartIndex + 2, FVector(1.0f, 0.0f, 0.0f), FVector(0.0f, 1.0f, 0.0f), FVector(0.0f, 0.0f, 1.0f));
 
-		Data->StaticMeshVertexBuffers.StaticMeshVertexBuffer.SetVertexUV(StartIndex + 0, 0, FVector2D(Tri.V1_UV.X, Tri.V1_UV.Y));
-		Data->StaticMeshVertexBuffers.StaticMeshVertexBuffer.SetVertexUV(StartIndex + 1, 0, FVector2D(Tri.V0_UV.X, Tri.V0_UV.Y));
-		Data->StaticMeshVertexBuffers.StaticMeshVertexBuffer.SetVertexUV(StartIndex + 2, 0, FVector2D(Tri.V2_UV.X, Tri.V2_UV.Y));
+		StaticMeshVertexBuffers.StaticMeshVertexBuffer.SetVertexUV(StartIndex + 0, 0, FVector2D(Tri.V1_UV.X, Tri.V1_UV.Y));
+		StaticMeshVertexBuffers.StaticMeshVertexBuffer.SetVertexUV(StartIndex + 1, 0, FVector2D(Tri.V0_UV.X, Tri.V0_UV.Y));
+		StaticMeshVertexBuffers.StaticMeshVertexBuffer.SetVertexUV(StartIndex + 2, 0, FVector2D(Tri.V2_UV.X, Tri.V2_UV.Y));
 
-		Data->StaticMeshVertexBuffers.ColorVertexBuffer.VertexColor(StartIndex + 0) = Tri.V1_Color.ToFColor(true);
-		Data->StaticMeshVertexBuffers.ColorVertexBuffer.VertexColor(StartIndex + 1) = Tri.V0_Color.ToFColor(true);
-		Data->StaticMeshVertexBuffers.ColorVertexBuffer.VertexColor(StartIndex + 2) = Tri.V2_Color.ToFColor(true);
+		StaticMeshVertexBuffers.ColorVertexBuffer.VertexColor(StartIndex + 0) = Tri.V1_Color.ToFColor(true);
+		StaticMeshVertexBuffers.ColorVertexBuffer.VertexColor(StartIndex + 1) = Tri.V0_Color.ToFColor(true);
+		StaticMeshVertexBuffers.ColorVertexBuffer.VertexColor(StartIndex + 2) = Tri.V2_Color.ToFColor(true);
 	}
 
-	FCanvasTriangleRendererItem::FRenderData* RenderData = Data;
-	ENQUEUE_RENDER_COMMAND(StaticMeshVertexBuffersLegacyInit)(
-		[VertexFactory, RenderData](FRHICommandListImmediate& RHICmdList)
-	{
-		RenderData->StaticMeshVertexBuffers.PositionVertexBuffer.InitResource();
-		RenderData->StaticMeshVertexBuffers.StaticMeshVertexBuffer.InitResource();
-		RenderData->StaticMeshVertexBuffers.ColorVertexBuffer.InitResource();
-
-		FLocalVertexFactory::FDataType RData;
-		RenderData->StaticMeshVertexBuffers.PositionVertexBuffer.BindPositionVertexBuffer(VertexFactory, RData);
-		RenderData->StaticMeshVertexBuffers.StaticMeshVertexBuffer.BindTangentVertexBuffer(VertexFactory, RData);
-		RenderData->StaticMeshVertexBuffers.StaticMeshVertexBuffer.BindPackedTexCoordVertexBuffer(VertexFactory, RData);
-		RenderData->StaticMeshVertexBuffers.StaticMeshVertexBuffer.BindLightMapVertexBuffer(VertexFactory, RData, 0);
-		RenderData->StaticMeshVertexBuffers.ColorVertexBuffer.BindColorVertexBuffer(VertexFactory, RData);
-		VertexFactory->SetData(RData);
-
-		RenderData->IndexBuffer.InitResource();
-		RenderData->VertexFactory.InitResource();
-		RenderData->TriMesh.InitResource();
-	});
+	StaticMeshVertexBuffers.PositionVertexBuffer.InitResource();
+	StaticMeshVertexBuffers.StaticMeshVertexBuffer.InitResource();
+	StaticMeshVertexBuffers.ColorVertexBuffer.InitResource();
+	IndexBuffer.InitResource();
+	VertexFactory.InitResource();
+	TriMesh.InitResource();
 };
 
+void FCanvasTriangleRendererItem::FRenderData::ReleaseTriangleMesh()
+{
+	TriMesh.ReleaseResource();
+	VertexFactory.ReleaseResource();
+	IndexBuffer.ReleaseResource();
+	StaticMeshVertexBuffers.PositionVertexBuffer.ReleaseResource();
+	StaticMeshVertexBuffers.StaticMeshVertexBuffer.ReleaseResource();
+	StaticMeshVertexBuffers.ColorVertexBuffer.ReleaseResource();
+}
+
+void FCanvasTriangleRendererItem::FRenderData::RenderTriangles(
+	FRHICommandListImmediate& RHICmdList,
+	FMeshPassProcessorRenderState& DrawRenderState,
+	const FSceneView& View,
+	bool bIsHitTesting,
+	bool bNeedsToSwitchVerticalAxis)
+{
+	check(IsInRenderingThread());
+
+	IRendererModule& RendererModule = GetRendererModule();
+
+	InitTriangleMesh(View, bNeedsToSwitchVerticalAxis);
+
+	SCOPED_DRAW_EVENTF(RHICmdList, CanvasDrawTriangles, *MaterialRenderProxy->GetMaterial(GMaxRHIFeatureLevel)->GetFriendlyName());
+
+	for (int32 TriIdx = 0; TriIdx < Triangles.Num(); TriIdx++)
+	{
+		const FTriangleInst& Tri = Triangles[TriIdx];
+
+		FMeshBatch& MeshBatch = TriMesh.MeshBatch;
+		MeshBatch.VertexFactory = &VertexFactory;
+		MeshBatch.MaterialRenderProxy = MaterialRenderProxy;
+		MeshBatch.Elements[0].FirstIndex = 3 * TriIdx;
+
+		GetRendererModule().DrawTileMesh(RHICmdList, DrawRenderState, View, MeshBatch, bIsHitTesting, Tri.HitProxyId);
+	}
+
+	ReleaseTriangleMesh();
+}
+
 bool FCanvasTriangleRendererItem::Render_RenderThread(FRHICommandListImmediate& RHICmdList, FMeshPassProcessorRenderState& DrawRenderState, const FCanvas* Canvas)
 {
 	float CurrentRealTime = 0.f;
@@ -305,21 +174,20 @@ bool FCanvasTriangleRendererItem::Render_RenderThread(FRHICommandListImmediate&
 
 	checkSlow(Data);
 
-	// current render target set for the canvas
 	const FRenderTarget* CanvasRenderTarget = Canvas->GetRenderTarget();
-	FSceneViewFamily* ViewFamily = new FSceneViewFamily(FSceneViewFamily::ConstructionValues(
+
+	TUniquePtr<const FSceneViewFamily> ViewFamily = MakeUnique<const FSceneViewFamily>(FSceneViewFamily::ConstructionValues(
 		CanvasRenderTarget,
 		nullptr,
 		FEngineShowFlags(ESFIM_Game))
 		.SetWorldTimes(CurrentWorldTime, DeltaWorldTime, CurrentRealTime)
 		.SetGammaCorrection(CanvasRenderTarget->GetDisplayGamma()));
 
-	FIntRect ViewRect(FIntPoint(0, 0), CanvasRenderTarget->GetSizeXY());
+	const FIntRect ViewRect(FIntPoint(0, 0), CanvasRenderTarget->GetSizeXY());
 
 	// make a temporary view
-
 	FSceneViewInitOptions ViewInitOptions;
-	ViewInitOptions.ViewFamily = ViewFamily;
+	ViewInitOptions.ViewFamily = ViewFamily.Get();
 	ViewInitOptions.SetViewRectangle(ViewRect);
 	ViewInitOptions.ViewOrigin = FVector::ZeroVector;
 	ViewInitOptions.ViewRotationMatrix = FMatrix::Identity;
@@ -327,38 +195,17 @@ bool FCanvasTriangleRendererItem::Render_RenderThread(FRHICommandListImmediate&
 	ViewInitOptions.BackgroundColor = FLinearColor::Black;
 	ViewInitOptions.OverlayColor = FLinearColor::White;
 
-	bool bNeedsToSwitchVerticalAxis = RHINeedsToSwitchVerticalAxis(Canvas->GetShaderPlatform()) && XOR(IsMobileHDR(), Canvas->GetAllowSwitchVerticalAxis());
-	FSceneView* View = new FSceneView(ViewInitOptions);
-
-	InitTriangleBuffers(&Data->VertexFactory, Data->Triangles, *View, bNeedsToSwitchVerticalAxis);
-
-	for (int32 TriIdx = 0; TriIdx < Data->Triangles.Num(); TriIdx++)
-	{
-		const FRenderData::FTriangleInst& Tri = Data->Triangles[TriIdx];
-
-		// update the FMeshBatch
-		FMeshBatch& TriMesh = Data->TriMesh.TriMeshElement;
-		TriMesh.VertexFactory = &Data->VertexFactory;
-		TriMesh.MaterialRenderProxy = Data->MaterialRenderProxy;
-		TriMesh.Elements[0].FirstIndex = 3 * TriIdx;
+	TUniquePtr<const FSceneView> View = MakeUnique<const FSceneView>(ViewInitOptions);
 
-		GetRendererModule().DrawTileMesh(RHICmdList, DrawRenderState, *View, TriMesh, Canvas->IsHitTesting(), Tri.HitProxyId);
-	}
+	const bool bNeedsToSwitchVerticalAxis = RHINeedsToSwitchVerticalAxis(Canvas->GetShaderPlatform()) && XOR(IsMobileHDR(), Canvas->GetAllowSwitchVerticalAxis());
 
-	Data->StaticMeshVertexBuffers.PositionVertexBuffer.ReleaseResource();
-	Data->StaticMeshVertexBuffers.StaticMeshVertexBuffer.ReleaseResource();
-	Data->StaticMeshVertexBuffers.ColorVertexBuffer.ReleaseResource();
-	Data->IndexBuffer.ReleaseResource();
-	Data->TriMesh.ReleaseResource();
-	Data->VertexFactory.ReleaseResource();
+	Data->RenderTriangles(RHICmdList, DrawRenderState, *View, Canvas->IsHitTesting(), bNeedsToSwitchVerticalAxis);
 
-	delete View->Family;
-	delete View;
 	if (Canvas->GetAllowedModes() & FCanvas::Allow_DeleteOnRender)
 	{
-		delete Data;
 		Data = nullptr;
 	}
+
 	return true;
 }
 
@@ -376,16 +223,17 @@ bool FCanvasTriangleRendererItem::Render_GameThread(const FCanvas* Canvas, FRend
 	}
 
 	checkSlow(Data);
-	// current render target set for the canvas
+
 	const FRenderTarget* CanvasRenderTarget = Canvas->GetRenderTarget();
-	FSceneViewFamily* ViewFamily = new FSceneViewFamily(FSceneViewFamily::ConstructionValues(
+
+	const FSceneViewFamily* ViewFamily = new FSceneViewFamily(FSceneViewFamily::ConstructionValues(
 		CanvasRenderTarget,
 		Canvas->GetScene(),
 		FEngineShowFlags(ESFIM_Game))
 		.SetWorldTimes(CurrentWorldTime, DeltaWorldTime, CurrentRealTime)
 		.SetGammaCorrection(CanvasRenderTarget->GetDisplayGamma()));
 
-	FIntRect ViewRect(FIntPoint(0, 0), CanvasRenderTarget->GetSizeXY());
+	const FIntRect ViewRect(FIntPoint(0, 0), CanvasRenderTarget->GetSizeXY());
 
 	// make a temporary view
 	FSceneViewInitOptions ViewInitOptions;
@@ -397,67 +245,30 @@ bool FCanvasTriangleRendererItem::Render_GameThread(const FCanvas* Canvas, FRend
 	ViewInitOptions.BackgroundColor = FLinearColor::Black;
 	ViewInitOptions.OverlayColor = FLinearColor::White;
 
-	FSceneView* View = new FSceneView(ViewInitOptions);
-
-	bool bNeedsToSwitchVerticalAxis = RHINeedsToSwitchVerticalAxis(Canvas->GetShaderPlatform()) && XOR(IsMobileHDR(),Canvas->GetAllowSwitchVerticalAxis());
+	const FSceneView* View = new FSceneView(ViewInitOptions);
 
-	struct FDrawTriangleParameters
-	{
-		FSceneView* View;
-		FRenderData* RenderData;
-		uint32 bIsHitTesting : 1;
-		uint32 AllowedCanvasModes;
-	};
-	FDrawTriangleParameters DrawTriangleParameters =
-	{
-		View,
-		Data,
-		(uint32)Canvas->IsHitTesting(),
-		Canvas->GetAllowedModes()
-	};
-
-	InitTriangleBuffers(&Data->VertexFactory, Data->Triangles, *View, bNeedsToSwitchVerticalAxis);
+	const bool bNeedsToSwitchVerticalAxis = RHINeedsToSwitchVerticalAxis(Canvas->GetShaderPlatform()) && XOR(IsMobileHDR(),Canvas->GetAllowSwitchVerticalAxis());
+	const bool bIsHitTesting = Canvas->IsHitTesting();
+	const bool bDeleteOnRender = Canvas->GetAllowedModes() & FCanvas::Allow_DeleteOnRender;
 
-	FDrawTriangleParameters Parameters = DrawTriangleParameters;
 	RenderScope.EnqueueRenderCommand(
-		[Parameters](FRHICommandListImmediate& RHICmdList)	
+		[LocalData = Data, View, bIsHitTesting, bNeedsToSwitchVerticalAxis](FRHICommandListImmediate& RHICmdList)
 	{
-		FMeshPassProcessorRenderState DrawRenderState(*Parameters.View);
+		FMeshPassProcessorRenderState DrawRenderState(*View);
 
 		// disable depth test & writes
 		DrawRenderState.SetDepthStencilState(TStaticDepthStencilState<false, CF_Always>::GetRHI());
 
-		SCOPED_DRAW_EVENT(RHICmdList, CanvasDrawTriangle);
-		for (int32 TriIdx = 0; TriIdx < Parameters.RenderData->Triangles.Num(); TriIdx++)
-		{
-			const FRenderData::FTriangleInst& Tri = Parameters.RenderData->Triangles[TriIdx];
-			// update the FMeshBatch
-			FMeshBatch& TriMesh = Parameters.RenderData->TriMesh.TriMeshElement;
-			TriMesh.VertexFactory = &Parameters.RenderData->VertexFactory;
-			TriMesh.MaterialRenderProxy = Parameters.RenderData->MaterialRenderProxy;
-			TriMesh.Elements[0].FirstIndex = 3 * TriIdx;
-
-			GetRendererModule().DrawTileMesh(RHICmdList, DrawRenderState, *Parameters.View, TriMesh, Parameters.bIsHitTesting, Tri.HitProxyId);
-		}
-
-		Parameters.RenderData->StaticMeshVertexBuffers.PositionVertexBuffer.ReleaseResource();
-		Parameters.RenderData->StaticMeshVertexBuffers.StaticMeshVertexBuffer.ReleaseResource();
-		Parameters.RenderData->StaticMeshVertexBuffers.ColorVertexBuffer.ReleaseResource();
-		Parameters.RenderData->IndexBuffer.ReleaseResource();
-		Parameters.RenderData->TriMesh.ReleaseResource();
-		Parameters.RenderData->VertexFactory.ReleaseResource();
+		LocalData->RenderTriangles(RHICmdList, DrawRenderState, *View, bIsHitTesting, bNeedsToSwitchVerticalAxis);
 
-		delete Parameters.View->Family;
-		delete Parameters.View;
-		if (Parameters.AllowedCanvasModes & FCanvas::Allow_DeleteOnRender)
-		{
-			delete Parameters.RenderData;
-		}
+		delete View->Family;
+		delete View;
 	});
 
 	if (Canvas->GetAllowedModes() & FCanvas::Allow_DeleteOnRender)
 	{
 		Data = nullptr;
 	}
+
 	return true;
-}
+}
\ No newline at end of file
diff --git a/Engine/Source/Runtime/Engine/Public/CanvasTypes.h b/Engine/Source/Runtime/Engine/Public/CanvasTypes.h
--- a/Engine/Source/Runtime/Engine/Public/CanvasTypes.h
+++ b/Engine/Source/Runtime/Engine/Public/CanvasTypes.h
@@ -1125,18 +1125,10 @@ public:
 		const FCanvas::FTransformEntry& InTransform = FCanvas::FTransformEntry(FMatrix::Identity),
 		bool bInFreezeTime = false)
 		// this data is deleted after rendering has completed
-		: Data(new FRenderData(InFeatureLevel, InMaterialRenderProxy, InTransform))
+		: Data(MakeShared<FRenderData>(InFeatureLevel, InMaterialRenderProxy, InTransform))
 		, bFreezeTime(bInFreezeTime)
 	{}
 
-	/**
-	* Destructor to delete data in case nothing rendered
-	*/
-	virtual ~FCanvasTriangleRendererItem()
-	{
-		delete Data;
-	}
-
 	/**
 	 * FCanvasTriangleRendererItem instance accessor
 	 *
@@ -1213,22 +1205,20 @@ private:
 	class FTriangleVertexFactory : public FLocalVertexFactory
 	{
 	public:
-		/** Default constructor. */
-		FTriangleVertexFactory(ERHIFeatureLevel::Type InFeatureLevel);
+		FTriangleVertexFactory(const FStaticMeshVertexBuffers* VertexBuffers, ERHIFeatureLevel::Type InFeatureLevel);
+		void InitResource() override;
+
+	private:
+		const FStaticMeshVertexBuffers* VertexBuffers;
 	};
 
-	/**
-	* Mesh used to render triangles.
-	*/
 	class FTriangleMesh : public FRenderResource
 	{
 	public:
 		FTriangleMesh(const FRawIndexBuffer* IndexBuffer, const FTriangleVertexFactory* VertexFactory);
 
-		/** The mesh element. */
-		FMeshBatch TriMeshElement;
+		FMeshBatch MeshBatch;
 		virtual void InitRHI() override;
-		virtual void ReleaseRHI() override;
 	private:
 		const FRawIndexBuffer* IndexBuffer;
 		const FTriangleVertexFactory* VertexFactory;
@@ -1236,35 +1226,15 @@ private:
 
 	class FRenderData
 	{
-		friend class FCanvasTriangleRendererItem;
-
-	private:
-		const FMaterialRenderProxy* MaterialRenderProxy;
-		FCanvas::FTransformEntry Transform;
-
-		/** The buffer containing vertex data. */
-		FRawIndexBuffer IndexBuffer;
-		FStaticMeshVertexBuffers StaticMeshVertexBuffers;
-		FTriangleVertexFactory VertexFactory;
-		FTriangleMesh TriMesh;
-
-		struct FTriangleInst
-		{
-			FCanvasUVTri Tri;
-			FHitProxyId HitProxyId;
-		};
-		TArray<FTriangleInst> Triangles;
-
 	public:
 		FRenderData(ERHIFeatureLevel::Type InFeatureLevel,
-			const FMaterialRenderProxy* InMaterialRenderProxy = NULL,
-			const FCanvas::FTransformEntry& InTransform = FCanvas::FTransformEntry(FMatrix::Identity))
+			const FMaterialRenderProxy* InMaterialRenderProxy,
+			const FCanvas::FTransformEntry& InTransform)
 			: MaterialRenderProxy(InMaterialRenderProxy)
 			, Transform(InTransform)
-			, VertexFactory(InFeatureLevel)
+			, VertexFactory(&StaticMeshVertexBuffers, InFeatureLevel)
 			, TriMesh(&IndexBuffer, &VertexFactory)
-		{
-		}
+		{}
 
 		FORCEINLINE int32 AddTriangle(const FCanvasUVTri& Tri, FHitProxyId HitProxyId)
 		{
@@ -1281,17 +1251,41 @@ private:
 		{
 			Triangles.Reserve(NumTriangles);
 		}
+
+		void RenderTriangles(
+			FRHICommandListImmediate& RHICmdList,
+			FMeshPassProcessorRenderState& DrawRenderState,
+			const FSceneView& View,
+			bool bIsHitTesting,
+			bool bNeedsToSwitchVerticalAxis);
+
+		const FMaterialRenderProxy* const MaterialRenderProxy;
+		const FCanvas::FTransformEntry Transform;
+
+	private:
+		void InitTriangleMesh(const FSceneView& View, bool bNeedsToSwitchVerticalAxis);
+		void ReleaseTriangleMesh();
+
+		FRawIndexBuffer IndexBuffer;
+		FStaticMeshVertexBuffers StaticMeshVertexBuffers;
+		FTriangleVertexFactory VertexFactory;
+		FTriangleMesh TriMesh;
+
+		struct FTriangleInst
+		{
+			FCanvasUVTri Tri;
+			FHitProxyId HitProxyId;
+		};
+		TArray<FTriangleInst> Triangles;
 	};
+
 	/**
-	* Render data which is allocated when a new FCanvasTriangleRendererItem is added for rendering.
-	* This data is only freed on the rendering thread once the item has finished rendering
-	*/
-	FRenderData* Data;
+	 * Render data which is allocated when a new FCanvasTriangleRendererItem is added for rendering.
+	 * This data is only freed on the rendering thread once the item has finished rendering
+	 */
+	TSharedPtr<FRenderData> Data;
 
 	const bool bFreezeTime;
-
-	typedef FRenderData::FTriangleInst FTriangleInst;
-	void InitTriangleBuffers(FLocalVertexFactory* VertexFactory, TArray<FTriangleInst>& Triangles, const FSceneView& View, bool bNeedsToSwitchVerticalAxis);
 };
 
 /**
