Fixed MaterialShader and MeshMaterialShader macros.

diff --git a/Engine/Source/Runtime/Engine/Public/MaterialShaderType.h b/Engine/Source/Runtime/Engine/Public/MaterialShaderType.h
--- a/Engine/Source/Runtime/Engine/Public/MaterialShaderType.h
+++ b/Engine/Source/Runtime/Engine/Public/MaterialShaderType.h
@@ -242,7 +242,7 @@ private:
 	virtual uint32 GetTypeSize() const override { return sizeof(*this); } \
 	\
 	static void ModifyCompilationEnvironmentImpl( \
-		const FGlobalShaderPermutationParameters& Parameters, \
+		const FMaterialShaderPermutationParameters& Parameters, \
 		FShaderCompilerEnvironment& OutEnvironment) \
 	{ \
 		FPermutationDomain PermutationVector(Parameters.PermutationId); \
@@ -262,6 +262,5 @@ private:
 		ShaderClass::ModifyCompilationEnvironmentImpl, \
 		ShaderClass::ShouldCompilePermutation, \
 		ShaderClass::ValidateCompiledResult, \
-		ShaderClass::GetStreamOutElements, \
-		ShaderClass::GetRootParametersMetadata() \
+		ShaderClass::GetStreamOutElements \
 		)
diff --git a/Engine/Source/Runtime/Engine/Public/MeshMaterialShaderType.h b/Engine/Source/Runtime/Engine/Public/MeshMaterialShaderType.h
--- a/Engine/Source/Runtime/Engine/Public/MeshMaterialShaderType.h
+++ b/Engine/Source/Runtime/Engine/Public/MeshMaterialShaderType.h
@@ -214,7 +214,7 @@ private:
 	virtual uint32 GetTypeSize() const override { return sizeof(*this); } \
 	\
 	static void ModifyCompilationEnvironmentImpl( \
-		const FGlobalShaderPermutationParameters& Parameters, \
+		const FMaterialShaderPermutationParameters& Parameters, \
 		FShaderCompilerEnvironment& OutEnvironment) \
 	{ \
 		FPermutationDomain PermutationVector(Parameters.PermutationId); \
@@ -234,6 +234,5 @@ private:
 		ShaderClass::ModifyCompilationEnvironmentImpl, \
 		ShaderClass::ShouldCompilePermutation, \
 		ShaderClass::ValidateCompiledResult, \
-		ShaderClass::GetStreamOutElements, \
-		ShaderClass::GetRootParametersMetadata() \
+		ShaderClass::GetStreamOutElements \
 		)
diff --git a/Engine/Source/Runtime/Renderer/Public/MaterialShader.h b/Engine/Source/Runtime/Renderer/Public/MaterialShader.h
--- a/Engine/Source/Runtime/Renderer/Public/MaterialShader.h
+++ b/Engine/Source/Runtime/Renderer/Public/MaterialShader.h
@@ -100,6 +100,7 @@ inline FArchive& operator<<(FArchive& Ar, FDebugUniformExpressionSet& DebugExpre
 /** Base class of all shaders that need material parameters. */
 class RENDERER_API FMaterialShader : public FShader
 {
+	DECLARE_SHADER_TYPE(FMaterialShader, Material);
 public:
 	static FName UniformBufferLayoutName;
 
diff --git a/Engine/Source/Runtime/Renderer/Public/MeshMaterialShader.h b/Engine/Source/Runtime/Renderer/Public/MeshMaterialShader.h
--- a/Engine/Source/Runtime/Renderer/Public/MeshMaterialShader.h
+++ b/Engine/Source/Runtime/Renderer/Public/MeshMaterialShader.h
@@ -32,6 +32,7 @@ public:
 /** Base class of all shaders that need material and vertex factory parameters. */
 class RENDERER_API FMeshMaterialShader : public FMaterialShader
 {
+	DECLARE_SHADER_TYPE(FMeshMaterialShader, MeshMaterial);
 public:
 	FMeshMaterialShader() {}
 
