VT stacks can't be shared across expression scopes, fixes crash when trying to send VT output to both material attribute and runtime VT output

diff --git a/Engine/Source/Runtime/Engine/Private/Materials/HLSLMaterialTranslator.h b/Engine/Source/Runtime/Engine/Private/Materials/HLSLMaterialTranslator.h
--- a/Engine/Source/Runtime/Engine/Private/Materials/HLSLMaterialTranslator.h
+++ b/Engine/Source/Runtime/Engine/Private/Materials/HLSLMaterialTranslator.h
@@ -144,6 +144,7 @@ struct FShaderCodeChunk
 
 struct FMaterialVTStackEntry
 {
+	TArray<FShaderCodeChunk>* Scope;
 	uint64 CoordinateHash;
 	uint64 MipValue0Hash;
 	uint64 MipValue1Hash;
@@ -3828,7 +3829,8 @@ protected:
 		const uint64 MipValue0Hash = GetParameterHash(MipValue0Index);
 		const uint64 MipValue1Hash = GetParameterHash(MipValue1Index);
 
-		uint64 Hash = CityHash128to64({ CoordinatHash, MipValue0Hash });
+		uint64 Hash = CityHash128to64({ (uint64)CurrentScopeChunks, CoordinatHash });
+		Hash = CityHash128to64({ Hash, MipValue0Hash });
 		Hash = CityHash128to64({ Hash, MipValue1Hash });
 		Hash = CityHash128to64({ Hash, (uint64)MipValueMode });
 		Hash = CityHash128to64({ Hash, (uint64)AddressU });
@@ -3836,12 +3838,14 @@ protected:
 		Hash = CityHash128to64({ Hash, (uint64)(AspectRatio * 1000.0f) });
 		Hash = CityHash128to64({ Hash, (uint64)PreallocatedStackTextureIndex });
 
+
 		// First check to see if we have an existing VTStack that matches this key, that can still fit another layer
 		for (int32 Index = VTStackHash.First(Hash); VTStackHash.IsValid(Index); Index = VTStackHash.Next(Index))
 		{
 			const FMaterialVirtualTextureStack& Stack = MaterialCompilationOutput.UniformExpressionSet.VTStacks[Index];
 			const FMaterialVTStackEntry& Entry = VTStacks[Index];
 			if (!Stack.AreLayersFull() &&
+				Entry.Scope == CurrentScopeChunks &&
 				Entry.CoordinateHash == CoordinatHash &&
 				Entry.MipValue0Hash == MipValue0Hash &&
 				Entry.MipValue1Hash == MipValue1Hash &&
@@ -3851,16 +3855,6 @@ protected:
 				Entry.AspectRatio == AspectRatio &&
 				Entry.PreallocatedStackTextureIndex == PreallocatedStackTextureIndex)
 			{
-				if (Entry.DebugCoordinateIndex != CoordinateIndex ||
-					Entry.DebugMipValue0Index != MipValue0Index ||
-					Entry.DebugMipValue1Index != MipValue1Index)
-				{
-					FString UVs = GetParameterCode(CoordinateIndex);
-					FString CheckUVs = GetParameterCode(Entry.DebugCoordinateIndex);
-					FString Mip0 = GetParameterCode(MipValue0Index, TEXT(""));
-					FString CheckMip0 = GetParameterCode(Entry.DebugMipValue0Index, TEXT(""));
-					int a = 0;
-				}
 				return Index;
 			}
 		}
@@ -3869,6 +3863,7 @@ protected:
 		const int32 StackIndex = VTStacks.AddDefaulted();
 		VTStackHash.Add(Hash, StackIndex);
 		FMaterialVTStackEntry& Entry = VTStacks[StackIndex];
+		Entry.Scope = CurrentScopeChunks;
 		Entry.CoordinateHash = CoordinatHash;
 		Entry.MipValue0Hash = MipValue0Hash;
 		Entry.MipValue1Hash = MipValue1Hash;
