Material shader permutation vector fixes

diff --git a/Engine/Source/Runtime/Engine/Private/Materials/MaterialShader.cpp b/Engine/Source/Runtime/Engine/Private/Materials/MaterialShader.cpp
--- a/Engine/Source/Runtime/Engine/Private/Materials/MaterialShader.cpp
+++ b/Engine/Source/Runtime/Engine/Private/Materials/MaterialShader.cpp
@@ -1595,7 +1595,7 @@ void FMaterialShaderMap::Compile(
 			}
   
 			// Iterate over all material shader types.
-			TMap<FShaderType*, FShaderCompileJob*> SharedShaderJobs;
+			TMap<TShaderTypePermutation<const FShaderType>, FShaderCompileJob*> SharedShaderJobs;
 			for(TLinkedList<FShaderType*>::TIterator ShaderTypeIt(FShaderType::GetTypeList());ShaderTypeIt;ShaderTypeIt.Next())
 			{
 				FMaterialShaderType* ShaderType = ShaderTypeIt->GetMaterialShaderType();
@@ -1606,7 +1606,7 @@ void FMaterialShaderMap::Compile(
 					{
 #if WITH_EDITOR
 						// Verify that the shader map Id contains inputs for any shaders that will be put into this shader map
-						check(InShaderMapId.ContainsShaderType(ShaderType, PermutationId));
+						check(InShaderMapId.ContainsShaderType(ShaderType, kUniqueShaderPermutationId));
 #endif
 						// Compile this material shader for this material.
 						TArray<FString> ShaderErrors;
@@ -1624,8 +1624,10 @@ void FMaterialShaderMap::Compile(
 								NewJobs,
 								DebugDescription, DebugExtension
 							);
-							check(!SharedShaderJobs.Find(ShaderType));
-							SharedShaderJobs.Add(ShaderType, Job);
+
+							TShaderTypePermutation<const FShaderType> ShaderTypePermutation(ShaderType, PermutationId);
+							check(!SharedShaderJobs.Find(ShaderTypePermutation));
+							SharedShaderJobs.Add(ShaderTypePermutation, Job);
 						}
 						NumShaders++;
 					}
@@ -1673,7 +1675,8 @@ void FMaterialShaderMap::Compile(
 							// If sharing shaders amongst pipelines, add this pipeline as a dependency of an existing job
 							for (const FShaderType* ShaderType : StageTypes)
 							{
-								FShaderCompileJob** Job = SharedShaderJobs.Find(ShaderType);
+								TShaderTypePermutation<const FShaderType> ShaderTypePermutation(ShaderType, kUniqueShaderPermutationId);
+								FShaderCompileJob** Job = SharedShaderJobs.Find(ShaderTypePermutation);
 								checkf(Job, TEXT("Couldn't find existing shared job for material shader %s on pipeline %s!"), ShaderType->GetName(), Pipeline->GetName());
 								auto* SingleJob = (*Job)->GetSingleShaderJob();
 								auto& PipelinesToShare = SingleJob->SharingPipelines.FindOrAdd(nullptr);
diff --git a/Engine/Source/Runtime/Engine/Private/Materials/MeshMaterialShader.cpp b/Engine/Source/Runtime/Engine/Private/Materials/MeshMaterialShader.cpp
--- a/Engine/Source/Runtime/Engine/Private/Materials/MeshMaterialShader.cpp
+++ b/Engine/Source/Runtime/Engine/Private/Materials/MeshMaterialShader.cpp
@@ -215,7 +215,7 @@ uint32 FMeshMaterialShaderMap::BeginCompile(
 				// Verify that the shader map Id contains inputs for any shaders that will be put into this shader map
 #if WITH_EDITOR
 				check(InShaderMapId.ContainsVertexFactoryType(VertexFactoryType));
-				check(InShaderMapId.ContainsShaderType(ShaderType, PermutationId));
+				check(InShaderMapId.ContainsShaderType(ShaderType, kUniqueShaderPermutationId));
 #endif
 
 				NumShadersPerVF++;
diff --git a/Engine/Source/Runtime/Engine/Public/MaterialShaderType.h b/Engine/Source/Runtime/Engine/Public/MaterialShaderType.h
--- a/Engine/Source/Runtime/Engine/Public/MaterialShaderType.h
+++ b/Engine/Source/Runtime/Engine/Public/MaterialShaderType.h
@@ -125,7 +125,6 @@ public:
 		checkf(FPaths::GetExtension(InSourceFilename) == TEXT("usf"),
 			TEXT("Incorrect virtual shader path extension for material shader '%s': Only .usf files should be compiled."),
 			InSourceFilename);
-		check(InTotalPermutationCount == 1);
 	}
 
 	/**
diff --git a/Engine/Source/Runtime/Engine/Public/MeshMaterialShaderType.h b/Engine/Source/Runtime/Engine/Public/MeshMaterialShaderType.h
--- a/Engine/Source/Runtime/Engine/Public/MeshMaterialShaderType.h
+++ b/Engine/Source/Runtime/Engine/Public/MeshMaterialShaderType.h
@@ -90,7 +90,6 @@ public:
 		checkf(FPaths::GetExtension(InSourceFilename) == TEXT("usf"),
 			TEXT("Incorrect virtual shader path extension for mesh material shader '%s': Only .usf files should be compiled."),
 			InSourceFilename);
-		check(InTotalPermutationCount == 1);
 	}
 
 	/**
