DR - The issue is that GameThreadShaderMap and the RenderingThreadShaderMap can get out of sync. And one case of this is when DefaultMaterials are forced though the !bApplyCompletedShaderMapForRendering path because bSynchronousCompile is set. Although it is noted that "if !bApplyCompletedShaderMapForRendering, RenderingThreadShaderMap must be set elsewhere to match up with the new value of GameThreadShaderMap" the Result is actually a dangling reference of RenderingThreadShaderMap. Therefore I removed bApplyCompletedShaderMapForRendering and force the synchronization of the RenderingThreadShaderMap via the SetGameThreadShaderMap interface each and every time it is changed.

diff --git a/Engine/Source/Developer/MaterialBaking/Private/ExportMaterialProxy.h b/Engine/Source/Developer/MaterialBaking/Private/ExportMaterialProxy.h
--- a/Engine/Source/Developer/MaterialBaking/Private/ExportMaterialProxy.h
+++ b/Engine/Source/Developer/MaterialBaking/Private/ExportMaterialProxy.h
@@ -270,7 +270,7 @@ public:
 			break;
 		};
 
-		CacheShaders(ResourceId, GMaxRHIShaderPlatform, true);
+		CacheShaders(ResourceId, GMaxRHIShaderPlatform);
 	}
 
 	/** This override is required otherwise the shaders aren't ready for use when the surface is rendered resulting in a blank image */
diff --git a/Engine/Source/Developer/MaterialUtilities/Private/MaterialUtilities.cpp b/Engine/Source/Developer/MaterialUtilities/Private/MaterialUtilities.cpp
--- a/Engine/Source/Developer/MaterialUtilities/Private/MaterialUtilities.cpp
+++ b/Engine/Source/Developer/MaterialUtilities/Private/MaterialUtilities.cpp
@@ -450,7 +450,7 @@ public:
 			break;
 		};
 		
-		CacheShaders(ResourceId, GMaxRHIShaderPlatform, true);
+		CacheShaders(ResourceId, GMaxRHIShaderPlatform);
 	}
 
 	virtual bool IsUsedWithStaticLighting() const { return true; }
diff --git a/Engine/Source/Editor/MaterialEditor/Private/MaterialEditor.cpp b/Engine/Source/Editor/MaterialEditor/Private/MaterialEditor.cpp
--- a/Engine/Source/Editor/MaterialEditor/Private/MaterialEditor.cpp
+++ b/Engine/Source/Editor/MaterialEditor/Private/MaterialEditor.cpp
@@ -4353,7 +4353,7 @@ FMatExpressionPreview* FMaterialEditor::GetExpressionPreview(UMaterialExpression
 			bNewlyCreated = true;
 			Preview = new FMatExpressionPreview(MaterialExpression);
 			ExpressionPreviews.Add(Preview);
-			Preview->CacheShaders(GMaxRHIShaderPlatform, true);
+			Preview->CacheShaders(GMaxRHIShaderPlatform);
 		}
 		return Preview;
 	}
diff --git a/Engine/Source/Editor/UnrealEd/Private/Lightmass/LightmassRender.cpp b/Engine/Source/Editor/UnrealEd/Private/Lightmass/LightmassRender.cpp
--- a/Engine/Source/Editor/UnrealEd/Private/Lightmass/LightmassRender.cpp
+++ b/Engine/Source/Editor/UnrealEd/Private/Lightmass/LightmassRender.cpp
@@ -262,7 +262,7 @@ public:
 
 				// Override with a special usage so we won't re-use the shader map used by the material for rendering
 				ResourceId.Usage = GetShaderMapUsage();
-				CacheShaders(ResourceId, GMaxRHIShaderPlatform, true);
+				CacheShaders(ResourceId, GMaxRHIShaderPlatform);
 			}
 		}
 	}
diff --git a/Engine/Source/Runtime/Engine/Classes/Materials/Material.h b/Engine/Source/Runtime/Engine/Classes/Materials/Material.h
--- a/Engine/Source/Runtime/Engine/Classes/Materials/Material.h
+++ b/Engine/Source/Runtime/Engine/Classes/Materials/Material.h
@@ -1483,7 +1483,7 @@ private:
 	void CacheResourceShadersForCooking(EShaderPlatform Platform, TArray<FMaterialResource*>& OutCachedMaterialResources, const ITargetPlatform* TargetPlatform = nullptr);
 
 	/** Caches shader maps for an array of material resources. */
-	void CacheShadersForResources(EShaderPlatform ShaderPlatform, const TArray<FMaterialResource*>& ResourcesToCache, bool bApplyCompletedShaderMapForRendering, const ITargetPlatform* TargetPlatform = nullptr);
+	void CacheShadersForResources(EShaderPlatform ShaderPlatform, const TArray<FMaterialResource*>& ResourcesToCache, const ITargetPlatform* TargetPlatform = nullptr);
 
 	/**
 	 * If there is some texture reference used by a TextureProperty node in any expressions, this function
diff --git a/Engine/Source/Runtime/Engine/Classes/Materials/MaterialInstance.h b/Engine/Source/Runtime/Engine/Classes/Materials/MaterialInstance.h
--- a/Engine/Source/Runtime/Engine/Classes/Materials/MaterialInstance.h
+++ b/Engine/Source/Runtime/Engine/Classes/Materials/MaterialInstance.h
@@ -673,7 +673,7 @@ protected:
 	void CacheResourceShadersForRendering();
 
 	/** Caches shader maps for an array of material resources. */
-	void CacheShadersForResources(EShaderPlatform ShaderPlatform, const TArray<FMaterialResource*>& ResourcesToCache, bool bApplyCompletedShaderMapForRendering, const ITargetPlatform* TargetPlatform = nullptr);
+	void CacheShadersForResources(EShaderPlatform ShaderPlatform, const TArray<FMaterialResource*>& ResourcesToCache, const ITargetPlatform* TargetPlatform = nullptr);
 
 	/** 
 	 * Copies over material instance parameters from the base material given a material interface.
diff --git a/Engine/Source/Runtime/Engine/Private/DebugViewModeMaterialProxy.cpp b/Engine/Source/Runtime/Engine/Private/DebugViewModeMaterialProxy.cpp
--- a/Engine/Source/Runtime/Engine/Private/DebugViewModeMaterialProxy.cpp
+++ b/Engine/Source/Runtime/Engine/Private/DebugViewModeMaterialProxy.cpp
@@ -77,7 +77,7 @@ FDebugViewModeMaterialProxy::FDebugViewModeMaterialProxy(
 
 		ResourceId.Usage = Usage;
 
-		CacheShaders(ResourceId, GMaxRHIShaderPlatform, true);
+		CacheShaders(ResourceId, GMaxRHIShaderPlatform);
 	}
 	else
 	{
diff --git a/Engine/Source/Runtime/Engine/Private/Materials/Material.cpp b/Engine/Source/Runtime/Engine/Private/Materials/Material.cpp
--- a/Engine/Source/Runtime/Engine/Private/Materials/Material.cpp
+++ b/Engine/Source/Runtime/Engine/Private/Materials/Material.cpp
@@ -2867,7 +2867,7 @@ void UMaterial::CacheResourceShadersForRendering(bool bRegenerateId)
 			ResourcesToCache.Reset();
 			check(MaterialResources[LocalActiveQL][FeatureLevel]);
 			ResourcesToCache.Add(MaterialResources[LocalActiveQL][FeatureLevel]);
-			CacheShadersForResources(ShaderPlatform, ResourcesToCache, true);
+			CacheShadersForResources(ShaderPlatform, ResourcesToCache);
 		}
 
 		FString AdditionalFormatToCache = GCompileMaterialsForShaderFormatCVar->GetString();
@@ -2929,7 +2929,7 @@ void UMaterial::CacheResourceShadersForCooking(EShaderPlatform ShaderPlatform, T
 		}
 	}
 
-	CacheShadersForResources(ShaderPlatform, ResourcesToCache, false, TargetPlatform);
+	CacheShadersForResources(ShaderPlatform, ResourcesToCache, TargetPlatform);
 
 	for (int32 ResourceIndex = 0; ResourceIndex < ResourcesToCache.Num(); ResourceIndex++)
 	{
@@ -2937,14 +2937,14 @@ void UMaterial::CacheResourceShadersForCooking(EShaderPlatform ShaderPlatform, T
 	}
 }
 
-void UMaterial::CacheShadersForResources(EShaderPlatform ShaderPlatform, const TArray<FMaterialResource*>& ResourcesToCache, bool bApplyCompletedShaderMapForRendering, const ITargetPlatform* TargetPlatform)
+void UMaterial::CacheShadersForResources(EShaderPlatform ShaderPlatform, const TArray<FMaterialResource*>& ResourcesToCache, const ITargetPlatform* TargetPlatform)
 {
 	RebuildExpressionTextureReferences();
 
 	for (int32 ResourceIndex = 0; ResourceIndex < ResourcesToCache.Num(); ResourceIndex++)
 	{
 		FMaterialResource* CurrentResource = ResourcesToCache[ResourceIndex];
-		const bool bSuccess = CurrentResource->CacheShaders(ShaderPlatform, bApplyCompletedShaderMapForRendering, TargetPlatform);
+		const bool bSuccess = CurrentResource->CacheShaders(ShaderPlatform, TargetPlatform);
 
 		if (!bSuccess)
 		{
diff --git a/Engine/Source/Runtime/Engine/Private/Materials/MaterialInstance.cpp b/Engine/Source/Runtime/Engine/Private/Materials/MaterialInstance.cpp
--- a/Engine/Source/Runtime/Engine/Private/Materials/MaterialInstance.cpp
+++ b/Engine/Source/Runtime/Engine/Private/Materials/MaterialInstance.cpp
@@ -2613,7 +2613,7 @@ void UMaterialInstance::CacheResourceShadersForRendering()
 			ResourcesToCache.Reset();
 			check(StaticPermutationMaterialResources[LocalActiveQL][FeatureLevel]);
 			ResourcesToCache.Add(StaticPermutationMaterialResources[LocalActiveQL][FeatureLevel]);
-			CacheShadersForResources(ShaderPlatform, ResourcesToCache, true);
+			CacheShadersForResources(ShaderPlatform, ResourcesToCache);
 		}
 	}
 
@@ -2664,7 +2664,7 @@ void UMaterialInstance::CacheResourceShadersForCooking(EShaderPlatform ShaderPla
 			}
 		}
 
-		CacheShadersForResources(ShaderPlatform, ResourcesToCache, false, TargetPlatform);
+		CacheShadersForResources(ShaderPlatform, ResourcesToCache, TargetPlatform);
 
 		for (int32 ResourceIndex = 0; ResourceIndex < ResourcesToCache.Num(); ResourceIndex++)
 		{
@@ -2673,7 +2673,7 @@ void UMaterialInstance::CacheResourceShadersForCooking(EShaderPlatform ShaderPla
 	}
 }
 
-void UMaterialInstance::CacheShadersForResources(EShaderPlatform ShaderPlatform, const TArray<FMaterialResource*>& ResourcesToCache, bool bApplyCompletedShaderMapForRendering, const ITargetPlatform* TargetPlatform)
+void UMaterialInstance::CacheShadersForResources(EShaderPlatform ShaderPlatform, const TArray<FMaterialResource*>& ResourcesToCache, const ITargetPlatform* TargetPlatform)
 {
 	UMaterial* BaseMaterial = GetMaterial();
 	BaseMaterial->CacheExpressionTextureReferences();
@@ -2689,7 +2689,7 @@ void UMaterialInstance::CacheShadersForResources(EShaderPlatform ShaderPlatform,
 		FMaterialShaderMapId ShaderMapId;
 		CurrentResource->GetShaderMapId(ShaderPlatform, ShaderMapId);
 
-		const bool bSuccess = CurrentResource->CacheShaders(ShaderMapId, ShaderPlatform, bApplyCompletedShaderMapForRendering, TargetPlatform);
+		const bool bSuccess = CurrentResource->CacheShaders(ShaderMapId, ShaderPlatform, TargetPlatform);
 
 		if (!bSuccess)
 		{
diff --git a/Engine/Source/Runtime/Engine/Private/Materials/MaterialShader.cpp b/Engine/Source/Runtime/Engine/Private/Materials/MaterialShader.cpp
--- a/Engine/Source/Runtime/Engine/Private/Materials/MaterialShader.cpp
+++ b/Engine/Source/Runtime/Engine/Private/Materials/MaterialShader.cpp
@@ -1402,12 +1402,6 @@ void FMaterialShaderMap::LoadForRemoteRecompile(FArchive& Ar, EShaderPlatform Sh
 
 							MaterialResource->SetGameThreadShaderMap(LoadedShaderMap);
 							MaterialResource->RegisterInlineShaderMap(false);
-
-							ENQUEUE_RENDER_COMMAND(FSetShaderMapOnMaterialResources)(
-								[MaterialResource, LoadedShaderMap](FRHICommandListImmediate& RHICmdList)
-								{
-									MaterialResource->SetRenderingThreadShaderMap(LoadedShaderMap);
-								});
 						}
 					}
 				}
@@ -1432,8 +1426,7 @@ void FMaterialShaderMap::Compile(
 	TRefCountPtr<FShaderCompilerEnvironment> MaterialEnvironment,
 	const FMaterialCompilationOutput& InMaterialCompilationOutput,
 	EShaderPlatform InPlatform,
-	bool bSynchronousCompile,
-	bool bApplyCompletedShaderMapForRendering)
+	bool bSynchronousCompile)
 {
 	if (FPlatformProperties::RequiresCookedData())
 	{
@@ -1703,7 +1696,7 @@ void FMaterialShaderMap::Compile(
 			const bool bRecreateComponentRenderStateOnCompletion = Material->IsPersistent();
 
 			// Note: using Material->IsPersistent() to detect whether this is a preview material which should have higher priority over background compiling
-			GShaderCompilingManager->AddJobs(NewJobs, bApplyCompletedShaderMapForRendering && !bSynchronousCompile, bSynchronousCompile || !Material->IsPersistent(), bRecreateComponentRenderStateOnCompletion);
+			GShaderCompilingManager->AddJobs(NewJobs, bSynchronousCompile || !Material->IsPersistent(), bRecreateComponentRenderStateOnCompletion);
   
 			// Compile the shaders for this shader map now if the material is not deferring and deferred compiles are not enabled globally
 			if (bSynchronousCompile)
diff --git a/Engine/Source/Runtime/Engine/Private/Materials/MaterialShared.cpp b/Engine/Source/Runtime/Engine/Private/Materials/MaterialShared.cpp
--- a/Engine/Source/Runtime/Engine/Private/Materials/MaterialShared.cpp
+++ b/Engine/Source/Runtime/Engine/Private/Materials/MaterialShared.cpp
@@ -970,7 +970,7 @@ void FMaterial::SerializeInlineShaderMap(FArchive& Ar)
 			{
 				TRefCountPtr<FMaterialShaderMap> LoadedShaderMap = new FMaterialShaderMap();
 				LoadedShaderMap->Serialize(Ar, true, bCooked && Ar.IsLoading());
-				GameThreadShaderMap = LoadedShaderMap;
+				SetGameThreadShaderMap(LoadedShaderMap);
 			}
 		}
 	}
@@ -1039,9 +1039,9 @@ void FMaterial::ReleaseShaderMap()
 		
 		FMaterial* Material = this;
 		ENQUEUE_RENDER_COMMAND(ReleaseShaderMap)(
-			[Material](FRHICommandList& RHICmdList)
+		[Material](FRHICommandList& RHICmdList)
 		{
-			Material->SetRenderingThreadShaderMap(nullptr);
+			Material->RenderingThreadShaderMap = nullptr;
 		});
 	}
 }
@@ -1767,18 +1767,18 @@ void FMaterial::SetupMaterialEnvironment(
  * Caches the material shaders for this material with no static parameters on the given platform.
  * This is used by material resources of UMaterials.
  */
-bool FMaterial::CacheShaders(EShaderPlatform Platform, bool bApplyCompletedShaderMapForRendering, const ITargetPlatform* TargetPlatform)
+bool FMaterial::CacheShaders(EShaderPlatform Platform, const ITargetPlatform* TargetPlatform)
 {
 	FMaterialShaderMapId NoStaticParametersId;
 	GetShaderMapId(Platform, NoStaticParametersId);
-	return CacheShaders(NoStaticParametersId, Platform, bApplyCompletedShaderMapForRendering, TargetPlatform);
+	return CacheShaders(NoStaticParametersId, Platform, TargetPlatform);
 }
 
 /**
  * Caches the material shaders for the given static parameter set and platform.
  * This is used by material resources of UMaterialInstances.
  */
-bool FMaterial::CacheShaders(const FMaterialShaderMapId& ShaderMapId, EShaderPlatform Platform, bool bApplyCompletedShaderMapForRendering, const ITargetPlatform* TargetPlatform)
+bool FMaterial::CacheShaders(const FMaterialShaderMapId& ShaderMapId, EShaderPlatform Platform, const ITargetPlatform* TargetPlatform)
 {
 	bool bSucceeded = false;
 	UE_CLOG(!ShaderMapId.IsValid(), LogMaterial, Warning, TEXT("Invalid shader map ID caching shaders for '%s', will use default material."), *GetFriendlyName());
@@ -1798,7 +1798,7 @@ bool FMaterial::CacheShaders(const FMaterialShaderMapId& ShaderMapId, EShaderPla
 		// Re-use an identical shader map in memory if possible, removing the reference to the inlined shader map
 		if (ExistingShaderMap)
 		{
-			GameThreadShaderMap = ExistingShaderMap;
+			SetGameThreadShaderMap(ExistingShaderMap);
 		}
 		else if (GameThreadShaderMap)
 		{
@@ -1809,19 +1809,20 @@ bool FMaterial::CacheShaders(const FMaterialShaderMapId& ShaderMapId, EShaderPla
 	else
 	{
 #if WITH_EDITOR
-		// Find the material's cached shader map.
-		GameThreadShaderMap = FMaterialShaderMap::FindId(ShaderMapId, Platform);
+		TRefCountPtr<FMaterialShaderMap> ShaderMap = FMaterialShaderMap::FindId(ShaderMapId, Platform);
 
 		// On-the-fly view shaders are not using ddc currently, as their shadermap is not persistent.
 		// See FMaterialShaderMap::ProcessCompilationResults().
 		if  (GetMaterialShaderMapUsage() != EMaterialShaderMapUsage::DebugViewMode)
 		{
 			// Attempt to load from the derived data cache if we are uncooked
-			if ((!GameThreadShaderMap || !GameThreadShaderMap->IsComplete(this, true)) && !FPlatformProperties::RequiresCookedData())
+			if ((!ShaderMap || !ShaderMap->IsComplete(this, true)) && !FPlatformProperties::RequiresCookedData())
 			{
-				FMaterialShaderMap::LoadFromDerivedDataCache(this, ShaderMapId, Platform, GameThreadShaderMap);
+				FMaterialShaderMap::LoadFromDerivedDataCache(this, ShaderMapId, Platform, ShaderMap);
 			}
 		}
+
+		SetGameThreadShaderMap(ShaderMap);
 #endif // WITH_EDITOR
 	}
 
@@ -1849,7 +1850,7 @@ bool FMaterial::CacheShaders(const FMaterialShaderMapId& ShaderMapId, EShaderPla
 		OutstandingCompileShaderMapIds.AddUnique(GameThreadShaderMap->GetCompilingId());
 #endif // WITH_EDITOR
 		// Reset the shader map so the default material will be used until the compile finishes.
-		GameThreadShaderMap = nullptr;
+		SetGameThreadShaderMap(nullptr);
 		bSucceeded = true;
 	}
 	else if (!GameThreadShaderMap || !(bAssumeShaderMapIsComplete || GameThreadShaderMap->IsComplete(this, !bLogShaderMapFailInfo)))
@@ -1880,7 +1881,7 @@ bool FMaterial::CacheShaders(const FMaterialShaderMapId& ShaderMapId, EShaderPla
 			}
 
 			// Reset the shader map so the default material will be used.
-			GameThreadShaderMap = nullptr;
+			SetGameThreadShaderMap(nullptr);
 		}
 		else
 		{
@@ -1895,14 +1896,16 @@ bool FMaterial::CacheShaders(const FMaterialShaderMapId& ShaderMapId, EShaderPla
 			}
 			UE_LOG(LogMaterial, Display, TEXT("%s cached shader map for material %s, compiling. %s"),ShaderMapCondition,*GetFriendlyName(), IsSpecialEngineMaterial() ? TEXT("Is special engine material.") : TEXT("") );
 
+			TRefCountPtr<FMaterialShaderMap> ShaderMap;
+
 			// If there's no cached shader map for this material, compile a new one.
 			// This is just kicking off the async compile, GameThreadShaderMap will not be complete yet
-			bSucceeded = BeginCompileShaderMap(ShaderMapId, Platform, GameThreadShaderMap, bApplyCompletedShaderMapForRendering, TargetPlatform);
+			bSucceeded = BeginCompileShaderMap(ShaderMapId, Platform, ShaderMap, TargetPlatform);
 
 			if (!bSucceeded)
 			{
 				// If it failed to compile the material, reset the shader map so the material isn't used.
-				GameThreadShaderMap = nullptr;
+				SetGameThreadShaderMap(nullptr);
 
 #if WITH_EDITOR
 				if (IsDefaultMaterial())
@@ -1918,6 +1921,10 @@ bool FMaterial::CacheShaders(const FMaterialShaderMapId& ShaderMapId, EShaderPla
 				}
 #endif // WITH_EDITOR
 			}
+			else
+			{
+				SetGameThreadShaderMap(ShaderMap);
+			}
 		}
 	}
 	else
@@ -1930,15 +1937,6 @@ bool FMaterial::CacheShaders(const FMaterialShaderMapId& ShaderMapId, EShaderPla
 #endif // WITH_EDITOR
 	}
 
-	// Enqueue the final shader map update
-	FMaterial* Material = this;
-	ENQUEUE_RENDER_COMMAND(UpdateRenderThreadShaderMap)(
-		[Material](FRHICommandListImmediate& RHICmdList)
-		{	
-			Material->RenderingThreadShaderMap = Material->GameThreadShaderMap;
-		}
-	);
-
 	return bSucceeded;
 }
 
@@ -1953,8 +1951,7 @@ bool FMaterial::CacheShaders(const FMaterialShaderMapId& ShaderMapId, EShaderPla
 bool FMaterial::BeginCompileShaderMap(
 	const FMaterialShaderMapId& ShaderMapId, 
 	EShaderPlatform Platform, 
-	TRefCountPtr<FMaterialShaderMap>& OutShaderMap, 
-	bool bApplyCompletedShaderMapForRendering,
+	TRefCountPtr<FMaterialShaderMap>& OutShaderMap,
 	const ITargetPlatform* TargetPlatform)
 {
 #if WITH_EDITORONLY_DATA
@@ -1983,7 +1980,7 @@ bool FMaterial::BeginCompileShaderMap(
 		MaterialEnvironment->IncludeVirtualPathToContentsMap.Add(TEXT("/Engine/Generated/Material.ush"), MaterialShaderCode);
 
 		// Compile the shaders for the material.
-		NewShaderMap->Compile(this, ShaderMapId, MaterialEnvironment, NewCompilationOutput, Platform, bSynchronousCompile, bApplyCompletedShaderMapForRendering);
+		NewShaderMap->Compile(this, ShaderMapId, MaterialEnvironment, NewCompilationOutput, Platform, bSynchronousCompile);
 
 		if (bSynchronousCompile)
 		{
@@ -3010,7 +3007,7 @@ void FMaterial::UpdateEditorLoadedMaterialResources(EShaderPlatform InShaderPlat
 		FMaterial* CurrentMaterial = *It;
 		if (!CurrentMaterial->GetGameThreadShaderMap() || !CurrentMaterial->GetGameThreadShaderMap()->IsComplete(CurrentMaterial, true))
 		{
-			CurrentMaterial->CacheShaders(InShaderPlatform, true);
+			CurrentMaterial->CacheShaders(InShaderPlatform);
 		}
 	}
 }
@@ -3954,8 +3951,6 @@ void SetShaderMapsOnMaterialResources_RenderThread(FRHICommandListImmediate& RHI
 
 				if (Material && MaterialsToUpdate.Contains(Material))
 				{
-					// Materials used as async fallbacks can't be updated through this mechanism and should have been updated synchronously earlier
-					check(!Material->RequiresSynchronousCompilation());
 					MaterialProxy->CacheUniformExpressions(true);
 					bFoundAnyInitializedMaterials = true;
 
@@ -3974,16 +3969,9 @@ void SetShaderMapsOnMaterialResources_RenderThread(FRHICommandListImmediate& RHI
 
 void SetShaderMapsOnMaterialResources(const TMap<FMaterial*, FMaterialShaderMap*>& MaterialsToUpdate)
 {
-	for (TMap<FMaterial*, FMaterialShaderMap*>::TConstIterator It(MaterialsToUpdate); It; ++It)
-	{
-		FMaterial* Material = It.Key();
-		check(!Material->RequiresSynchronousCompilation());
-	}
-
-	FMaterialsToUpdateMap InMaterialsToUpdate = MaterialsToUpdate;
 	ENQUEUE_RENDER_COMMAND(FSetShaderMapOnMaterialResources)(
-		[InMaterialsToUpdate](FRHICommandListImmediate& RHICmdList)
-		{
-			SetShaderMapsOnMaterialResources_RenderThread(RHICmdList, InMaterialsToUpdate);
-		});
+	[InMaterialsToUpdate = MaterialsToUpdate](FRHICommandListImmediate& RHICmdList)
+	{
+		SetShaderMapsOnMaterialResources_RenderThread(RHICmdList, InMaterialsToUpdate);
+	});
 }
diff --git a/Engine/Source/Runtime/Engine/Private/ShaderCompiler/ShaderCompiler.cpp b/Engine/Source/Runtime/Engine/Private/ShaderCompiler/ShaderCompiler.cpp
--- a/Engine/Source/Runtime/Engine/Private/ShaderCompiler/ShaderCompiler.cpp
+++ b/Engine/Source/Runtime/Engine/Private/ShaderCompiler/ShaderCompiler.cpp
@@ -1680,7 +1680,7 @@ bool FShaderCompilingManager::GetDumpShaderDebugInfo() const
 	return GDumpShaderDebugInfo != 0;
 }
 
-void FShaderCompilingManager::AddJobs(TArray<FShaderCommonCompileJob*>& NewJobs, bool bApplyCompletedShaderMapForRendering, bool bOptimizeForLowLatency, bool bRecreateComponentRenderStateOnCompletion)
+void FShaderCompilingManager::AddJobs(TArray<FShaderCommonCompileJob*>& NewJobs, bool bOptimizeForLowLatency, bool bRecreateComponentRenderStateOnCompletion)
 {
 	check(!FPlatformProperties::RequiresCookedData());
 
@@ -1721,7 +1721,6 @@ void FShaderCompilingManager::AddJobs(TArray<FShaderCommonCompileJob*>& NewJobs,
 	{
 		NewJobs[JobIndex]->bOptimizeForLowLatency = bOptimizeForLowLatency;
 		FShaderMapCompileResults& ShaderMapInfo = ShaderMapJobs.FindOrAdd(NewJobs[JobIndex]->Id);
-		ShaderMapInfo.bApplyCompletedShaderMapForRendering = bApplyCompletedShaderMapForRendering;
 		ShaderMapInfo.bRecreateComponentRenderStateOnCompletion = bRecreateComponentRenderStateOnCompletion;
 		auto* PipelineJob = NewJobs[JobIndex]->GetShaderPipelineJob();
 		if (PipelineJob)
@@ -1965,7 +1964,6 @@ void FShaderCompilingManager::ProcessCompiledShaderMaps(
 	// Keeps shader maps alive as they are passed from the shader compiler and applied to the owning FMaterial
 	TArray<TRefCountPtr<FMaterialShaderMap> > LocalShaderMapReferences;
 	TMap<FMaterial*, FMaterialShaderMap*> MaterialsToUpdate;
-	TMap<FMaterial*, FMaterialShaderMap*> MaterialsToApplyToScene;
 
 	// Process compiled shader maps in FIFO order, in case a shader map has been enqueued multiple times,
 	// Which can happen if a material is edited while a background compile is going on
@@ -2088,11 +2086,6 @@ void FShaderCompilingManager::ProcessCompiledShaderMaps(
 							if (CompletedShaderMap->IsComplete(Material, true))
 							{
 								MaterialsToUpdate.Add(Material, CompletedShaderMap);
-								// Note: if !CompileResults.bApplyCompletedShaderMapForRendering, RenderingThreadShaderMap must be set elsewhere to match up with the new value of GameThreadShaderMap
-								if (CompileResults.bApplyCompletedShaderMapForRendering)
-								{
-									MaterialsToApplyToScene.Add(Material, CompletedShaderMap);
-								}
 							}
 
 							if (GShowShaderWarnings && Errors.Num() > 0)
@@ -2145,19 +2138,17 @@ void FShaderCompilingManager::ProcessCompiledShaderMaps(
 
 	if (MaterialsToUpdate.Num() > 0)
 	{
+		SetShaderMapsOnMaterialResources(MaterialsToUpdate);
+
 		for (TMap<FMaterial*, FMaterialShaderMap*>::TConstIterator It(MaterialsToUpdate); It; ++It)
 		{
 			FMaterial* Material = It.Key();
 			FMaterialShaderMap* ShaderMap = It.Value();
 			check(!ShaderMap || ShaderMap->IsValidForRendering());
 
-			Material->SetGameThreadShaderMap(It.Value());
+			Material->GameThreadShaderMap = It.Value();
 		}
 
-		const TSet<FSceneInterface*>& AllocatedScenes = GetRendererModule().GetAllocatedScenes();
-
-		SetShaderMapsOnMaterialResources(MaterialsToApplyToScene);
-
 		for (TMap<FMaterial*, FMaterialShaderMap*>::TIterator It(MaterialsToUpdate); It; ++It)
 		{
 			FMaterial* Material = It.Key();
@@ -2418,7 +2409,7 @@ bool FShaderCompilingManager::HandlePotentialRetryOnError(TMap<int32, FShaderMap
 				}
 
 				// Send all the shaders from this shader map through the compiler again
-				AddJobs(Results.FinishedJobs, Results.bApplyCompletedShaderMapForRendering, true, Results.bRecreateComponentRenderStateOnCompletion);
+				AddJobs(Results.FinishedJobs, true, Results.bRecreateComponentRenderStateOnCompletion);
 			}
 		}
 
@@ -3965,7 +3956,7 @@ void VerifyGlobalShaders(EShaderPlatform Platform, bool bLoadedFromCacheFile)
 
 	if (GlobalShaderJobs.Num() > 0)
 	{
-		GShaderCompilingManager->AddJobs(GlobalShaderJobs, true, true, false);
+		GShaderCompilingManager->AddJobs(GlobalShaderJobs, true, false);
 
 		const bool bAllowAsynchronousGlobalShaderCompiling =
 			// OpenGL requires that global shader maps are compiled before attaching
diff --git a/Engine/Source/Runtime/Engine/Public/MaterialShared.h b/Engine/Source/Runtime/Engine/Public/MaterialShared.h
--- a/Engine/Source/Runtime/Engine/Public/MaterialShared.h
+++ b/Engine/Source/Runtime/Engine/Public/MaterialShared.h
@@ -951,8 +951,7 @@ public:
 		TRefCountPtr<FShaderCompilerEnvironment> MaterialEnvironment,
 		const FMaterialCompilationOutput& InMaterialCompilationOutput,
 		EShaderPlatform Platform,
-		bool bSynchronousCompile,
-		bool bApplyCompletedShaderMapForRendering
+		bool bSynchronousCompile
 		);
 
 #if WITH_EDITOR
@@ -1367,13 +1366,13 @@ public:
 	 * Caches the material shaders for this material with no static parameters on the given platform.
 	 * This is used by material resources of UMaterials.
 	 */
-	ENGINE_API bool CacheShaders(EShaderPlatform Platform, bool bApplyCompletedShaderMapForRendering, const ITargetPlatform* TargetPlatform = nullptr);
+	ENGINE_API bool CacheShaders(EShaderPlatform Platform, const ITargetPlatform* TargetPlatform = nullptr);
 
 	/**
 	 * Caches the material shaders for the given static parameter set and platform.
 	 * This is used by material resources of UMaterialInstances.
 	 */
-	ENGINE_API bool CacheShaders(const FMaterialShaderMapId& ShaderMapId, EShaderPlatform Platform, bool bApplyCompletedShaderMapForRendering, const ITargetPlatform* TargetPlatform = nullptr);
+	ENGINE_API bool CacheShaders(const FMaterialShaderMapId& ShaderMapId, EShaderPlatform Platform, const ITargetPlatform* TargetPlatform = nullptr);
 
 	/**
 	 * Should the shader for this material with the given platform, shader type and vertex 
@@ -1598,11 +1597,16 @@ public:
 		return GameThreadShaderMap; 
 	}
 
-	/** Note: SetRenderingThreadShaderMap must also be called with the same value, but from the rendering thread. */
 	void SetGameThreadShaderMap(FMaterialShaderMap* InMaterialShaderMap)
 	{
 		checkSlow(IsInGameThread() || IsInAsyncLoadingThread());
 		GameThreadShaderMap = InMaterialShaderMap;
+
+		FMaterial* Material = this;
+		ENQUEUE_RENDER_COMMAND(SetGameThreadShaderMap)([Material](FRHICommandListImmediate& RHICmdList)
+		{
+			Material->RenderingThreadShaderMap = Material->GameThreadShaderMap;
+		});
 	}
 
 	void SetInlineShaderMap(FMaterialShaderMap* InMaterialShaderMap)
@@ -1611,6 +1615,12 @@ public:
 		GameThreadShaderMap = InMaterialShaderMap;
 		bContainsInlineShaders = true;
 		bLoadedCookedShaderMapId = true;
+
+		FMaterial* Material = this;
+		ENQUEUE_RENDER_COMMAND(SetInlineShaderMap)([Material](FRHICommandListImmediate& RHICmdList)
+		{
+			Material->RenderingThreadShaderMap = Material->GameThreadShaderMap;
+		});
 	}
 
 	ENGINE_API class FMaterialShaderMap* GetRenderingThreadShaderMap() const;
@@ -1820,7 +1830,6 @@ private:
 		const FMaterialShaderMapId& ShaderMapId,
 		EShaderPlatform Platform, 
 		TRefCountPtr<class FMaterialShaderMap>& OutShaderMap, 
-		bool bApplyCompletedShaderMapForRendering, 
 		const ITargetPlatform* TargetPlatform = nullptr);
 
 	/** Populates OutEnvironment with defines needed to compile shaders for this material. */
diff --git a/Engine/Source/Runtime/Engine/Public/ShaderCompiler.h b/Engine/Source/Runtime/Engine/Public/ShaderCompiler.h
--- a/Engine/Source/Runtime/Engine/Public/ShaderCompiler.h
+++ b/Engine/Source/Runtime/Engine/Public/ShaderCompiler.h
@@ -349,13 +349,11 @@ struct FShaderMapCompileResults
 	FShaderMapCompileResults() :
 		NumJobsQueued(0),
 		bAllJobsSucceeded(true),
-		bApplyCompletedShaderMapForRendering(true),
 		bRecreateComponentRenderStateOnCompletion(false)
 	{}
 
 	int32 NumJobsQueued;
 	bool bAllJobsSucceeded;
-	bool bApplyCompletedShaderMapForRendering;
 	bool bRecreateComponentRenderStateOnCompletion;
 	TArray<FShaderCommonCompileJob*> FinishedJobs;
 };
@@ -554,7 +552,7 @@ public:
 	 * Adds shader jobs to be asynchronously compiled. 
 	 * FinishCompilation or ProcessAsyncResults must be used to get the results.
 	 */
-	ENGINE_API void AddJobs(TArray<FShaderCommonCompileJob*>& NewJobs, bool bApplyCompletedShaderMapForRendering, bool bOptimizeForLowLatency, bool bRecreateComponentRenderStateOnCompletion);
+	ENGINE_API void AddJobs(TArray<FShaderCommonCompileJob*>& NewJobs, bool bOptimizeForLowLatency, bool bRecreateComponentRenderStateOnCompletion);
 
 	/**
 	* Removes all outstanding compile jobs for the passed shader maps.
