Fix for material shader using permutation vector

diff --git a/Engine/Source/Runtime/Engine/Private/Materials/MaterialShader.cpp b/Engine/Source/Runtime/Engine/Private/Materials/MaterialShader.cpp
--- a/Engine/Source/Runtime/Engine/Private/Materials/MaterialShader.cpp
+++ b/Engine/Source/Runtime/Engine/Private/Materials/MaterialShader.cpp
@@ -1021,12 +1021,12 @@ FShader* FMaterialShaderType::FinishCompileShader(
 	}
 
 	// Find a shader with the same key in memory
-	FShader* Shader = CurrentJob.ShaderType->FindShaderById(FShaderId(MaterialShaderMapHash, ShaderPipelineType, CurrentJob.VFType, CurrentJob.ShaderType, SpecificPermutationId, CurrentJob.Input.Target));
+	FShader* Shader = CurrentJob.ShaderType->FindShaderById(FShaderId(MaterialShaderMapHash, ShaderPipelineType, CurrentJob.VFType, CurrentJob.ShaderType, CurrentJob.PermutationId, CurrentJob.Input.Target));
 
 	// There was no shader with the same key so create a new one with the compile output, which will bind shader parameters
 	if (!Shader)
 	{
-		Shader = (*ConstructCompiledRef)(CompiledShaderInitializerType(this, SpecificPermutationId, CurrentJob.Output, Resource, UniformExpressionSet, MaterialShaderMapHash, ShaderPipelineType, nullptr, InDebugDescription));
+		Shader = (*ConstructCompiledRef)(CompiledShaderInitializerType(this, CurrentJob.PermutationId, CurrentJob.Output, Resource, UniformExpressionSet, MaterialShaderMapHash, ShaderPipelineType, nullptr, InDebugDescription));
 		CurrentJob.Output.ParameterMap.VerifyBindingsAreComplete(GetName(), CurrentJob.Output.Target, CurrentJob.VFType);
 	}
 
