DR - Fix unity compile fix
#rnx

diff --git a/Engine/Source/Runtime/Engine/Private/Materials/MaterialShared.cpp b/Engine/Source/Runtime/Engine/Private/Materials/MaterialShared.cpp
--- a/Engine/Source/Runtime/Engine/Private/Materials/MaterialShared.cpp
+++ b/Engine/Source/Runtime/Engine/Private/Materials/MaterialShared.cpp
@@ -2282,7 +2282,7 @@ void FMaterialVirtualTextureStack::Serialize(FArchive& Ar)
 	FMaterialRenderProxy
 -----------------------------------------------------------------------------*/
 
-static void OnVirtualTextureDestroyed(const FVirtualTextureProducerHandle& InHandle, void* Baton)
+static void OnVirtualTextureDestroyedCB(const FVirtualTextureProducerHandle& InHandle, void* Baton)
 {
 	FMaterialRenderProxy* MaterialProxy = static_cast<FMaterialRenderProxy*>(Baton);
 
@@ -2305,7 +2305,7 @@ IAllocatedVirtualTexture* FMaterialRenderProxy::GetPreallocatedVTStack(const FMa
 		return nullptr;
 	}
 
-	GetRendererModule().AddVirtualTextureProducerDestroyedCallback(Texture->GetProducerHandle(), &OnVirtualTextureDestroyed, const_cast<FMaterialRenderProxy*>(this));
+	GetRendererModule().AddVirtualTextureProducerDestroyedCallback(Texture->GetProducerHandle(), &OnVirtualTextureDestroyedCB, const_cast<FMaterialRenderProxy*>(this));
 	HasVirtualTextureCallbacks = true;
 
 	return Texture->GetAllocatedVirtualTexture();
@@ -2342,7 +2342,7 @@ IAllocatedVirtualTexture* FMaterialRenderProxy::AllocateVTStack(const FMaterialR
 			const FVirtualTextureProducerHandle& ProducerHandle = VirtualTextureResourceForLayer->GetProducerHandle();
 			VTDesc.ProducerHandle[LayerIndex] = ProducerHandle;
 			VTDesc.LocalLayerToProduce[LayerIndex] = 0u;
-			GetRendererModule().AddVirtualTextureProducerDestroyedCallback(ProducerHandle, &OnVirtualTextureDestroyed, const_cast<FMaterialRenderProxy*>(this));
+			GetRendererModule().AddVirtualTextureProducerDestroyedCallback(ProducerHandle, &OnVirtualTextureDestroyedCB, const_cast<FMaterialRenderProxy*>(this));
 			bFoundValidLayer = true;
 		}
 	}
