Serialize whether material has runtime virtual texture output to avoid expensive runtime check
#rnx

diff --git a/Engine/Source/Runtime/Engine/Classes/Materials/MaterialExpressionRuntimeVirtualTextureOutput.h b/Engine/Source/Runtime/Engine/Classes/Materials/MaterialExpressionRuntimeVirtualTextureOutput.h
--- a/Engine/Source/Runtime/Engine/Classes/Materials/MaterialExpressionRuntimeVirtualTextureOutput.h
+++ b/Engine/Source/Runtime/Engine/Classes/Materials/MaterialExpressionRuntimeVirtualTextureOutput.h
@@ -38,11 +38,6 @@ class UMaterialExpressionRuntimeVirtualTextureOutput : public UMaterialExpressio
 	FExpressionInput Opacity;
 
 public:
-	//~ Begin UObject Interface
-	// Only needed for (bad, slow) implementation of FMaterialResource::HasRuntimeVirtualTextureOutput()
-	virtual bool NeedsLoadForClient() const override { return true; }
-	//~ End UObject Interface
-
 #if WITH_EDITOR
 	//~ Begin UMaterialExpression Interface
 	virtual int32 Compile(class FMaterialCompiler* Compiler, int32 OutputIndex) override;
diff --git a/Engine/Source/Runtime/Engine/Private/Materials/HLSLMaterialTranslator.h b/Engine/Source/Runtime/Engine/Private/Materials/HLSLMaterialTranslator.h
--- a/Engine/Source/Runtime/Engine/Private/Materials/HLSLMaterialTranslator.h
+++ b/Engine/Source/Runtime/Engine/Private/Materials/HLSLMaterialTranslator.h
@@ -1326,6 +1326,11 @@ ResourcesString = TEXT("");
 		{
 			OutEnvironment.SetDefine(TEXT("USES_EYE_ADAPTATION"), TEXT("1"));
 		}
+
+		if (MaterialCompilationOutput.bHasRuntimeVirtualTextureOutput)
+		{
+			OutEnvironment.SetDefine(TEXT("VIRTUAL_TEXTURE_OUTPUT"), 1);
+		}
 		
 		// @todo MetalMRT: Remove this hack and implement proper atmospheric-fog solution for Metal MRT...
 		OutEnvironment.SetDefine(TEXT("MATERIAL_ATMOSPHERIC_FOG"), !IsMetalMRTPlatform(InPlatform) ? bUsesAtmosphericFog : 0);
@@ -6521,6 +6526,14 @@ protected:
 		return INDEX_NONE;
 	}
 
+	virtual int32 VirtualTextureOutput() override
+	{
+		MaterialCompilationOutput.bHasRuntimeVirtualTextureOutput = true;
+		
+		// return value is not used
+		return INDEX_NONE;
+	}
+
 #if HANDLE_CUSTOM_OUTPUTS_AS_MATERIAL_ATTRIBUTES
 	/** Used to translate code for custom output attributes such as ClearCoatBottomNormal */
 	void GenerateCustomAttributeCode(int32 OutputIndex, int32 OutputCode, EMaterialValueType OutputType, FString& DisplayName)
diff --git a/Engine/Source/Runtime/Engine/Private/Materials/MaterialExpressions.cpp b/Engine/Source/Runtime/Engine/Private/Materials/MaterialExpressions.cpp
--- a/Engine/Source/Runtime/Engine/Private/Materials/MaterialExpressions.cpp
+++ b/Engine/Source/Runtime/Engine/Private/Materials/MaterialExpressions.cpp
@@ -1924,6 +1924,7 @@ int32 UMaterialExpressionRuntimeVirtualTextureOutput::Compile(class FMaterialCom
 		CodeInput = Opacity.IsConnected() ? Opacity.Compile(Compiler) : Compiler->Constant(1.f);
 	}
 
+	Compiler->VirtualTextureOutput();
 	return Compiler->CustomOutput(this, OutputIndex, CodeInput);
 }
 
diff --git a/Engine/Source/Runtime/Engine/Private/Materials/MaterialShared.cpp b/Engine/Source/Runtime/Engine/Private/Materials/MaterialShared.cpp
--- a/Engine/Source/Runtime/Engine/Private/Materials/MaterialShared.cpp
+++ b/Engine/Source/Runtime/Engine/Private/Materials/MaterialShared.cpp
@@ -509,21 +509,23 @@ void FMaterialCompilationOutput::Serialize(FArchive& Ar)
 #endif
 
 	uint8 PackedFlags = 0;
-	PackedFlags |= (bUsesEyeAdaptation			<< 0);
-	PackedFlags |= (bModifiesMeshPosition		<< 1);
-	PackedFlags |= (bUsesWorldPositionOffset	<< 2);
-	PackedFlags |= (bUsesGlobalDistanceField	<< 3);
-	PackedFlags |= (bUsesPixelDepthOffset		<< 4);
-	PackedFlags |= (bUsesDistanceCullFade		<< 5);
+	PackedFlags |= (bUsesEyeAdaptation				<< 0);
+	PackedFlags |= (bModifiesMeshPosition			<< 1);
+	PackedFlags |= (bUsesWorldPositionOffset		<< 2);
+	PackedFlags |= (bUsesGlobalDistanceField		<< 3);
+	PackedFlags |= (bUsesPixelDepthOffset			<< 4);
+	PackedFlags |= (bUsesDistanceCullFade			<< 5);
+	PackedFlags |= (bHasRuntimeVirtualTextureOutput	<< 6);
 
 	Ar << PackedFlags;
 
-	bUsesEyeAdaptation			= (PackedFlags >> 0) & 1;
-	bModifiesMeshPosition		= (PackedFlags >> 1) & 1;
-	bUsesWorldPositionOffset	= (PackedFlags >> 2) & 1;
-	bUsesGlobalDistanceField	= (PackedFlags >> 3) & 1;
-	bUsesPixelDepthOffset		= (PackedFlags >> 4) & 1;
-	bUsesDistanceCullFade		= (PackedFlags >> 5) & 1;
+	bUsesEyeAdaptation				= (PackedFlags >> 0) & 1;
+	bModifiesMeshPosition			= (PackedFlags >> 1) & 1;
+	bUsesWorldPositionOffset		= (PackedFlags >> 2) & 1;
+	bUsesGlobalDistanceField		= (PackedFlags >> 3) & 1;
+	bUsesPixelDepthOffset			= (PackedFlags >> 4) & 1;
+	bUsesDistanceCullFade			= (PackedFlags >> 5) & 1;
+	bHasRuntimeVirtualTextureOutput	= (PackedFlags >> 6) & 1;
 }
 
 void FMaterial::GetShaderMapId(EShaderPlatform Platform, FMaterialShaderMapId& OutId) const
@@ -852,6 +854,11 @@ bool FMaterial::MaterialUsesSceneDepthLookup_GameThread() const
 	return GameThreadShaderMap.GetReference() ? GameThreadShaderMap->UsesSceneDepthLookup() : false;
 }
 
+bool FMaterial::HasRuntimeVirtualTextureOutput_RenderThread() const
+{
+	return RenderingThreadShaderMap ? RenderingThreadShaderMap->HasRuntimeVirtualTextureOutput() : false;
+}
+
 FMaterialShaderMap* FMaterial::GetRenderingThreadShaderMap() const 
 { 
 	check(IsInParallelRenderingThread());
@@ -1379,7 +1386,8 @@ uint32 FMaterialResource::GetStencilCompare() const
 
 bool FMaterialResource::HasRuntimeVirtualTextureOutput() const
 {
-	//todo[vt]: This might be expensive. Can we cache the result once on material init/update?
+	// Slow check used only for ShouldCompilePermutation() calls. 
+	// Runtime calls from render thread can use faster FMaterial::HasRuntimeVirtualTextureOutput_RenderThread()
 	for (UMaterialExpression* Expression : Material->Expressions)
 	{
 		if (Expression->IsA(UMaterialExpressionRuntimeVirtualTextureOutput::StaticClass()))
diff --git a/Engine/Source/Runtime/Engine/Public/MaterialCompiler.h b/Engine/Source/Runtime/Engine/Public/MaterialCompiler.h
--- a/Engine/Source/Runtime/Engine/Public/MaterialCompiler.h
+++ b/Engine/Source/Runtime/Engine/Public/MaterialCompiler.h
@@ -294,6 +294,7 @@ public:
 
 	virtual int32 CustomExpression(class UMaterialExpressionCustom* Custom, TArray<int32>& CompiledInputs) = 0;
 	virtual int32 CustomOutput(class UMaterialExpressionCustomOutput* Custom, int32 OutputIndex, int32 OutputCode) = 0;
+	virtual int32 VirtualTextureOutput() = 0;
 
 	virtual int32 DDX(int32 X) = 0;
 	virtual int32 DDY(int32 X) = 0;
@@ -529,6 +530,8 @@ public:
 
 	virtual int32 CustomExpression(class UMaterialExpressionCustom* Custom, TArray<int32>& CompiledInputs) override { return Compiler->CustomExpression(Custom,CompiledInputs); }
 	virtual int32 CustomOutput(class UMaterialExpressionCustomOutput* Custom, int32 OutputIndex, int32 OutputCode) override{ return Compiler->CustomOutput(Custom, OutputIndex, OutputCode); }
+	virtual int32 VirtualTextureOutput() override { return Compiler->VirtualTextureOutput(); }
+
 	virtual int32 DDX(int32 X) override { return Compiler->DDX(X); }
 	virtual int32 DDY(int32 X) override { return Compiler->DDY(X); }
 
diff --git a/Engine/Source/Runtime/Engine/Public/MaterialShared.h b/Engine/Source/Runtime/Engine/Public/MaterialShared.h
--- a/Engine/Source/Runtime/Engine/Public/MaterialShared.h
+++ b/Engine/Source/Runtime/Engine/Public/MaterialShared.h
@@ -548,7 +548,8 @@ public:
 		bUsesWorldPositionOffset(false),
 		bUsesGlobalDistanceField(false),
 		bUsesPixelDepthOffset(false),
-		bUsesDistanceCullFade(false)
+		bUsesDistanceCullFade(false),
+		bHasRuntimeVirtualTextureOutput(false)
 	{}
 
 	ENGINE_API void Serialize(FArchive& Ar);
@@ -597,6 +598,9 @@ public:
 	/** true if the material uses distance cull fade */
 	uint8 bUsesDistanceCullFade : 1;
 
+	/** true if the material supports virtual texture output */
+	uint8 bHasRuntimeVirtualTextureOutput : 1;
+
 	/** Indicates whether the material uses scene color. */
 	ENGINE_API bool RequiresSceneColorCopy() const { return IsSceneTextureUsed(PPI_SceneColor); }
 
@@ -1087,6 +1091,7 @@ public:
 	bool UsesSceneDepthLookup() const { return MaterialCompilationOutput.UsesSceneDepthLookup(); }
 	bool UsesVelocitySceneTexture() const { return MaterialCompilationOutput.UsesVelocitySceneTexture(); }
 	bool UsesDistanceCullFade() const { return MaterialCompilationOutput.bUsesDistanceCullFade; }
+	bool HasRuntimeVirtualTextureOutput() const { return MaterialCompilationOutput.bHasRuntimeVirtualTextureOutput; }
 #if WITH_EDITOR
 	uint32 GetNumUsedUVScalars() const { return MaterialCompilationOutput.NumUsedUVScalars; }
 	uint32 GetNumUsedCustomInterpolatorScalars() const { return MaterialCompilationOutput.NumUsedCustomInterpolatorScalars; }
@@ -1588,6 +1593,9 @@ public:
 	ENGINE_API bool MaterialUsesSceneDepthLookup_RenderThread() const;
 	ENGINE_API bool MaterialUsesSceneDepthLookup_GameThread() const;
 
+	/** Does the material have a runtime virtual texture output node. */
+	ENGINE_API bool HasRuntimeVirtualTextureOutput_RenderThread() const;
+
 	/** Note: This function is only intended for use in deciding whether or not shader permutations are required before material translation occurs. */
 	ENGINE_API bool MaterialMayModifyMeshPosition() const;
 
diff --git a/Engine/Source/Runtime/Renderer/Private/VT/RuntimeVirtualTextureRender.cpp b/Engine/Source/Runtime/Renderer/Private/VT/RuntimeVirtualTextureRender.cpp
--- a/Engine/Source/Runtime/Renderer/Private/VT/RuntimeVirtualTextureRender.cpp
+++ b/Engine/Source/Runtime/Renderer/Private/VT/RuntimeVirtualTextureRender.cpp
@@ -26,16 +26,8 @@ namespace RuntimeVirtualTexture
 
 		static bool ShouldCompilePermutation(const FMeshMaterialShaderPermutationParameters& Parameters)
 		{
-			return IsFeatureLevelSupported(Parameters.Platform, ERHIFeatureLevel::SM4) && (Parameters.Material->GetMaterialDomain() == MD_RuntimeVirtualTexture || Parameters.Material->HasRuntimeVirtualTextureOutput());
-		}
-
-		static void ModifyCompilationEnvironment(const FMaterialShaderPermutationParameters& Parameters, FShaderCompilerEnvironment& OutEnvironment)
-		{
-			FMeshMaterialShader::ModifyCompilationEnvironment(Parameters, OutEnvironment);
-			if (Parameters.Material->HasRuntimeVirtualTextureOutput())
-			{
-				OutEnvironment.SetDefine(TEXT("VIRTUAL_TEXTURE_OUTPUT"), 1);
-			}
+			return IsFeatureLevelSupported(Parameters.Platform, ERHIFeatureLevel::SM4) && 
+				(Parameters.Material->GetMaterialDomain() == MD_RuntimeVirtualTexture || Parameters.Material->HasRuntimeVirtualTextureOutput());
 		}
 
 		FShader_VirtualTextureMaterialDraw()
@@ -237,8 +229,7 @@ namespace RuntimeVirtualTexture
 			const FMaterial& Material = MeshBatch.MaterialRenderProxy->GetMaterialWithFallback(FeatureLevel, FallbackMaterialRenderProxyPtr);
 			const FMaterialRenderProxy& MaterialRenderProxy = FallbackMaterialRenderProxyPtr ? *FallbackMaterialRenderProxyPtr : *MeshBatch.MaterialRenderProxy;
 
-			//todo[vt]: Find alternative to this slow material validation (maybe move it to mesh batch creation time and fallback to default vt material there?)
-			if (Material.GetMaterialDomain() == MD_RuntimeVirtualTexture || Material.HasRuntimeVirtualTextureOutput())
+			if (Material.GetMaterialDomain() == MD_RuntimeVirtualTexture || Material.HasRuntimeVirtualTextureOutput_RenderThread())
 			{
 				switch ((ERuntimeVirtualTextureMaterialType)MeshBatch.RuntimeVirtualTextureMaterialType)
 				{
