Added a material node: RuntimeVirtualTextureReplace
This allows a material to change shading behaviour when rendering to runtime virtual texture
#rnx

diff --git a/Engine/Shaders/Private/Common.ush b/Engine/Shaders/Private/Common.ush
--- a/Engine/Shaders/Private/Common.ush
+++ b/Engine/Shaders/Private/Common.ush
@@ -1587,6 +1587,17 @@ bool GetRayTracingQualitySwitch()
 #endif
 }
 
+// Used for the RuntimeVirtualTextureReplace material expression
+// Returns true if rendering to a runtime virtual texture page
+bool GetRuntimeVirtualTextureOutputSwitch()
+{
+#if VIRTUAL_TEXTURE_PAGE_RENDER
+	return true;
+#else
+	return false;
+#endif
+}
+
 #if FEATURE_LEVEL >= FEATURE_LEVEL_SM4
 struct FWriteToSliceGeometryOutput
 {
diff --git a/Engine/Source/Runtime/Engine/Classes/Materials/MaterialExpressionRuntimeVirtualTextureReplace.h b/Engine/Source/Runtime/Engine/Classes/Materials/MaterialExpressionRuntimeVirtualTextureReplace.h
new file mode 100644
--- /dev/null
+++ b/Engine/Source/Runtime/Engine/Classes/Materials/MaterialExpressionRuntimeVirtualTextureReplace.h
@@ -0,0 +1,30 @@
+// Copyright 1998-2019 Epic Games, Inc. All Rights Reserved.
+
+#pragma once
+
+#include "CoreMinimal.h"
+#include "Materials/MaterialExpression.h"
+#include "UObject/ObjectMacros.h"
+#include "MaterialExpressionRuntimeVirtualTextureReplace.generated.h"
+
+/** Material output expression to switch logic for a path that renders to runtime virtual texture pages. */
+UCLASS(collapsecategories, hidecategories = Object)
+class UMaterialExpressionRuntimeVirtualTextureReplace : public UMaterialExpression
+{
+	GENERATED_UCLASS_BODY()
+
+	/** Used by the default rendering passes. */
+	UPROPERTY()
+	FExpressionInput Default;
+
+	/** Used by the pass that renders to a runtime virtual texture page. */
+	UPROPERTY()
+	FExpressionInput VirtualTextureOutput;
+
+	//~ Begin UMaterialExpression Interface
+#if WITH_EDITOR
+	virtual int32 Compile(class FMaterialCompiler* Compiler, int32 OutputIndex) override;
+	virtual void GetCaption(TArray<FString>& OutCaptions) const override;
+#endif
+	//~ End UMaterialExpression Interface
+};
diff --git a/Engine/Source/Runtime/Engine/Private/Materials/HLSLMaterialTranslator.h b/Engine/Source/Runtime/Engine/Private/Materials/HLSLMaterialTranslator.h
--- a/Engine/Source/Runtime/Engine/Private/Materials/HLSLMaterialTranslator.h
+++ b/Engine/Source/Runtime/Engine/Private/Materials/HLSLMaterialTranslator.h
@@ -5925,6 +5925,17 @@ protected:
 
 	virtual int32 MaterialProxyReplace(int32 Realtime, int32 MaterialProxy) override { return Realtime; }
 
+	virtual int32 VirtualTextureOutputReplace(int32 Default, int32 VirtualTexture) override
+	{
+		if (Default == INDEX_NONE || VirtualTexture == INDEX_NONE)
+		{
+			return INDEX_NONE;
+		}
+
+		EMaterialValueType ResultType = GetArithmeticResultType(Default, VirtualTexture);
+		return AddCodeChunk(ResultType, TEXT("(GetRuntimeVirtualTextureOutputSwitch() ? (%s) : (%s))"), *GetParameterCode(VirtualTexture), *GetParameterCode(Default));
+	}
+
 	virtual int32 ObjectOrientation() override
 	{ 
 		return AddInlinedCodeChunk(MCT_Float3,TEXT("GetObjectOrientation(Parameters.PrimitiveId)"));
diff --git a/Engine/Source/Runtime/Engine/Private/Materials/MaterialExpressions.cpp b/Engine/Source/Runtime/Engine/Private/Materials/MaterialExpressions.cpp
--- a/Engine/Source/Runtime/Engine/Private/Materials/MaterialExpressions.cpp
+++ b/Engine/Source/Runtime/Engine/Private/Materials/MaterialExpressions.cpp
@@ -159,6 +159,7 @@
 #include "Materials/MaterialExpressionRotator.h"
 #include "Materials/MaterialExpressionRound.h"
 #include "Materials/MaterialExpressionRuntimeVirtualTextureOutput.h"
+#include "Materials/MaterialExpressionRuntimeVirtualTextureReplace.h"
 #include "Materials/MaterialExpressionRuntimeVirtualTextureSample.h"
 #include "Materials/MaterialExpressionSaturate.h"
 #include "Materials/MaterialExpressionSceneColor.h"
@@ -1874,16 +1875,16 @@ UMaterialExpressionRuntimeVirtualTextureOutput::UMaterialExpressionRuntimeVirtua
 	// Structure to hold one-time initialization
 	struct FConstructorStatics
 	{
-		FText NAME_Texture;
+		FText NAME_VirtualTexture;
 		FConstructorStatics()
-			: NAME_Texture(LOCTEXT("Texture", "Texture"))
+			: NAME_VirtualTexture(LOCTEXT("VirtualTexture", "VirtualTexture"))
 		{
 		}
 	};
 	static FConstructorStatics ConstructorStatics;
 
 #if WITH_EDITORONLY_DATA
-	MenuCategories.Add(ConstructorStatics.NAME_Texture);
+	MenuCategories.Add(ConstructorStatics.NAME_VirtualTexture);
 #endif
 
 #if WITH_EDITOR
@@ -1956,16 +1957,16 @@ UMaterialExpressionRuntimeVirtualTextureSample::UMaterialExpressionRuntimeVirtua
 	// Structure to hold one-time initialization
 	struct FConstructorStatics
 	{
-		FText NAME_Texture;
+		FText NAME_VirtualTexture;
 		FConstructorStatics()
-			: NAME_Texture(LOCTEXT("Texture", "Texture"))
+			: NAME_VirtualTexture(LOCTEXT("VirtualTexture", "VirtualTexture"))
 		{
 		}
 	};
 	static FConstructorStatics ConstructorStatics;
 
 #if WITH_EDITORONLY_DATA
-	MenuCategories.Add(ConstructorStatics.NAME_Texture);
+	MenuCategories.Add(ConstructorStatics.NAME_VirtualTexture);
 #endif
 
 #if WITH_EDITOR
@@ -2131,6 +2132,51 @@ void UMaterialExpressionRuntimeVirtualTextureSample::GetCaption(TArray<FString>&
 
 #endif // WITH_EDITOR
 
+UMaterialExpressionRuntimeVirtualTextureReplace::UMaterialExpressionRuntimeVirtualTextureReplace(const FObjectInitializer& ObjectInitializer)
+	: Super(ObjectInitializer)
+{
+	// Structure to hold one-time initialization
+	struct FConstructorStatics
+	{
+		FText NAME_VirtualTexture;
+		FConstructorStatics()
+			: NAME_VirtualTexture(LOCTEXT("VirtualTexture", "VirtualTexture"))
+		{
+		}
+	};
+	static FConstructorStatics ConstructorStatics;
+
+#if WITH_EDITORONLY_DATA
+	MenuCategories.Add(ConstructorStatics.NAME_VirtualTexture);
+#endif
+}
+
+#if WITH_EDITOR
+
+int32 UMaterialExpressionRuntimeVirtualTextureReplace::Compile(class FMaterialCompiler* Compiler, int32 OutputIndex)
+{
+	if (!Default.GetTracedInput().Expression)
+	{
+		return Compiler->Errorf(TEXT("Missing RuntimeVirtualTextureReplace input 'Default'"));
+	}
+
+	if (!VirtualTextureOutput.GetTracedInput().Expression)
+	{
+		return Compiler->Errorf(TEXT("Missing RuntimeVirtualTextureReplace input 'VirtualTextureOutput'"));
+	}
+
+	int32 Arg1 = Default.Compile(Compiler);
+	int32 Arg2 = VirtualTextureOutput.Compile(Compiler);
+	return Compiler->VirtualTextureOutputReplace(Arg1, Arg2);
+}
+
+void UMaterialExpressionRuntimeVirtualTextureReplace::GetCaption(TArray<FString>& OutCaptions) const
+{
+	OutCaptions.Add(TEXT("RuntimeVirtualTextureReplace"));
+}
+
+#endif // WITH_EDITOR
+
 UMaterialExpressionAdd::UMaterialExpressionAdd(const FObjectInitializer& ObjectInitializer)
 	: Super(ObjectInitializer)
 {
diff --git a/Engine/Source/Runtime/Engine/Public/MaterialCompiler.h b/Engine/Source/Runtime/Engine/Public/MaterialCompiler.h
--- a/Engine/Source/Runtime/Engine/Public/MaterialCompiler.h
+++ b/Engine/Source/Runtime/Engine/Public/MaterialCompiler.h
@@ -285,6 +285,7 @@ public:
 	virtual int32 ShadowReplace(int32 Default, int32 Shadow) = 0;
 	virtual int32 RayTracingQualitySwitchReplace(int32 Normal, int32 RayTraced) = 0;
 	virtual int32 MaterialProxyReplace(int32 Realtime, int32 MaterialProxy) = 0;
+	virtual int32 VirtualTextureOutputReplace(int32 Default, int32 VirtualTexture) = 0;
 
 	virtual int32 ObjectOrientation() = 0;
 	virtual int32 RotateAboutAxis(int32 NormalizedRotationAxisAndAngleIndex, int32 PositionOnAxisIndex, int32 PositionIndex) = 0;
@@ -519,6 +520,8 @@ public:
 	virtual int32 ShadowReplace(int32 Default, int32 Shadow) override { return Compiler->ShadowReplace(Default, Shadow); }
 	virtual int32 RayTracingQualitySwitchReplace(int32 Normal, int32 RayTraced) override { return Compiler->RayTracingQualitySwitchReplace(Normal, RayTraced); }
 	virtual int32 MaterialProxyReplace(int32 Realtime, int32 MaterialProxy) override { return Realtime; }
+	virtual int32 VirtualTextureOutputReplace(int32 Default, int32 VirtualTexture) override { return Compiler->VirtualTextureOutputReplace(Default, VirtualTexture); }
+	
 	virtual int32 ObjectOrientation() override { return Compiler->ObjectOrientation(); }
 	virtual int32 RotateAboutAxis(int32 NormalizedRotationAxisAndAngleIndex, int32 PositionOnAxisIndex, int32 PositionIndex) override
 	{
diff --git a/Engine/Source/Runtime/Renderer/Private/VT/RuntimeVirtualTextureRender.cpp b/Engine/Source/Runtime/Renderer/Private/VT/RuntimeVirtualTextureRender.cpp
--- a/Engine/Source/Runtime/Renderer/Private/VT/RuntimeVirtualTextureRender.cpp
+++ b/Engine/Source/Runtime/Renderer/Private/VT/RuntimeVirtualTextureRender.cpp
@@ -30,6 +30,12 @@ namespace RuntimeVirtualTexture
 				(Parameters.Material->GetMaterialDomain() == MD_RuntimeVirtualTexture || Parameters.Material->HasRuntimeVirtualTextureOutput());
 		}
 
+		static void ModifyCompilationEnvironment(const FMaterialShaderPermutationParameters& Parameters, FShaderCompilerEnvironment& OutEnvironment)
+		{
+			FMeshMaterialShader::ModifyCompilationEnvironment(Parameters, OutEnvironment);
+			OutEnvironment.SetDefine(TEXT("VIRTUAL_TEXTURE_PAGE_RENDER"), 1);
+		}
+
 		FShader_VirtualTextureMaterialDraw()
 		{}
 
