Added VirtualTextureFeatureSwitch material node to change behaviour depending on VT support availability
(Could use feature and platform level switches in a material but this should be simpler to maintain)
#rnx

diff --git a/Engine/Source/Runtime/Engine/Classes/Materials/MaterialExpressionRuntimeVirtualTextureReplace.h b/Engine/Source/Runtime/Engine/Classes/Materials/MaterialExpressionRuntimeVirtualTextureReplace.h
--- a/Engine/Source/Runtime/Engine/Classes/Materials/MaterialExpressionRuntimeVirtualTextureReplace.h
+++ b/Engine/Source/Runtime/Engine/Classes/Materials/MaterialExpressionRuntimeVirtualTextureReplace.h
@@ -24,6 +24,9 @@ class UMaterialExpressionRuntimeVirtualTextureReplace : public UMaterialExpressi
 	//~ Begin UMaterialExpression Interface
 #if WITH_EDITOR
 	virtual int32 Compile(class FMaterialCompiler* Compiler, int32 OutputIndex) override;
+	virtual bool IsResultMaterialAttributes(int32 OutputIndex) override;
+	virtual uint32 GetInputType(int32 InputIndex) override { return MCT_Unknown; }
+	virtual uint32 GetOutputType(int32 OutputIndex) override { return MCT_Unknown; }
 	virtual void GetCaption(TArray<FString>& OutCaptions) const override;
 #endif
 	//~ End UMaterialExpression Interface
diff --git a/Engine/Source/Runtime/Engine/Classes/Materials/MaterialExpressionVirtualTextureFeatureSwitch.h b/Engine/Source/Runtime/Engine/Classes/Materials/MaterialExpressionVirtualTextureFeatureSwitch.h
new file mode 100644
--- /dev/null
+++ b/Engine/Source/Runtime/Engine/Classes/Materials/MaterialExpressionVirtualTextureFeatureSwitch.h
@@ -0,0 +1,33 @@
+// Copyright 1998-2019 Epic Games, Inc. All Rights Reserved.
+
+#pragma once
+
+#include "CoreMinimal.h"
+#include "Materials/MaterialExpression.h"
+#include "UObject/ObjectMacros.h"
+#include "MaterialExpressionVirtualTextureFeatureSwitch.generated.h"
+
+/** Material output expression to switch logic according to whether virtual texturing is supported on this project, platform and feature level. */
+UCLASS(collapsecategories, hidecategories = Object)
+class UMaterialExpressionVirtualTextureFeatureSwitch : public UMaterialExpression
+{
+	GENERATED_UCLASS_BODY()
+
+	/** Used if virtual texture is not supported. */
+	UPROPERTY()
+	FExpressionInput No;
+
+	/** Used if virtual texture is supported. */
+	UPROPERTY()
+	FExpressionInput Yes;
+
+	//~ Begin UMaterialExpression Interface
+#if WITH_EDITOR
+	virtual int32 Compile(class FMaterialCompiler* Compiler, int32 OutputIndex) override;
+	virtual bool IsResultMaterialAttributes(int32 OutputIndex) override;
+	virtual uint32 GetInputType(int32 InputIndex) override { return MCT_Unknown; }
+	virtual uint32 GetOutputType(int32 OutputIndex) override { return MCT_Unknown; }
+	virtual void GetCaption(TArray<FString>& OutCaptions) const override;
+#endif
+	//~ End UMaterialExpression Interface
+};
diff --git a/Engine/Source/Runtime/Engine/Private/Materials/MaterialExpressions.cpp b/Engine/Source/Runtime/Engine/Private/Materials/MaterialExpressions.cpp
--- a/Engine/Source/Runtime/Engine/Private/Materials/MaterialExpressions.cpp
+++ b/Engine/Source/Runtime/Engine/Private/Materials/MaterialExpressions.cpp
@@ -161,6 +161,7 @@
 #include "Materials/MaterialExpressionRuntimeVirtualTextureOutput.h"
 #include "Materials/MaterialExpressionRuntimeVirtualTextureReplace.h"
 #include "Materials/MaterialExpressionRuntimeVirtualTextureSample.h"
+#include "Materials/MaterialExpressionVirtualTextureFeatureSwitch.h"
 #include "Materials/MaterialExpressionSaturate.h"
 #include "Materials/MaterialExpressionSceneColor.h"
 #include "Materials/MaterialExpressionSceneDepth.h"
@@ -2170,6 +2171,18 @@ int32 UMaterialExpressionRuntimeVirtualTextureReplace::Compile(class FMaterialCo
 	return Compiler->VirtualTextureOutputReplace(Arg1, Arg2);
 }
 
+bool UMaterialExpressionRuntimeVirtualTextureReplace::IsResultMaterialAttributes(int32 OutputIndex)
+{
+	for (FExpressionInput* ExpressionInput : GetInputs())
+	{
+		if (ExpressionInput->GetTracedInput().Expression && !ExpressionInput->Expression->ContainsInputLoop() && ExpressionInput->Expression->IsResultMaterialAttributes(ExpressionInput->OutputIndex))
+		{
+			return true;
+		}
+	}
+	return false;
+}
+
 void UMaterialExpressionRuntimeVirtualTextureReplace::GetCaption(TArray<FString>& OutCaptions) const
 {
 	OutCaptions.Add(TEXT("RuntimeVirtualTextureReplace"));
@@ -2177,6 +2190,66 @@ void UMaterialExpressionRuntimeVirtualTextureReplace::GetCaption(TArray<FString>
 
 #endif // WITH_EDITOR
 
+UMaterialExpressionVirtualTextureFeatureSwitch::UMaterialExpressionVirtualTextureFeatureSwitch(const FObjectInitializer& ObjectInitializer)
+	: Super(ObjectInitializer)
+{
+	// Structure to hold one-time initialization
+	struct FConstructorStatics
+	{
+		FText NAME_VirtualTexture;
+		FConstructorStatics()
+			: NAME_VirtualTexture(LOCTEXT("VirtualTexture", "VirtualTexture"))
+		{
+		}
+	};
+	static FConstructorStatics ConstructorStatics;
+
+#if WITH_EDITORONLY_DATA
+	MenuCategories.Add(ConstructorStatics.NAME_VirtualTexture);
+#endif
+}
+
+#if WITH_EDITOR
+
+int32 UMaterialExpressionVirtualTextureFeatureSwitch::Compile(class FMaterialCompiler* Compiler, int32 OutputIndex)
+{
+	if (!Yes.GetTracedInput().Expression)
+	{
+		return Compiler->Errorf(TEXT("Missing VirtualTextureFeatureSwitch input 'Yes'"));
+	}
+
+	if (!No.GetTracedInput().Expression)
+	{
+		return Compiler->Errorf(TEXT("Missing VirtualTextureFeatureSwitch input 'No'"));
+	}
+
+	if (UseVirtualTexturing(Compiler->GetFeatureLevel(), Compiler->GetTargetPlatform()))
+	{
+		return Yes.Compile(Compiler);
+	}
+	
+	return No.Compile(Compiler);
+}
+
+bool UMaterialExpressionVirtualTextureFeatureSwitch::IsResultMaterialAttributes(int32 OutputIndex)
+{
+	for (FExpressionInput* ExpressionInput : GetInputs())
+	{
+		if (ExpressionInput->GetTracedInput().Expression && !ExpressionInput->Expression->ContainsInputLoop() && ExpressionInput->Expression->IsResultMaterialAttributes(ExpressionInput->OutputIndex))
+		{
+			return true;
+		}
+	}
+	return false;
+}
+
+void UMaterialExpressionVirtualTextureFeatureSwitch::GetCaption(TArray<FString>& OutCaptions) const
+{
+	OutCaptions.Add(TEXT("VirtualTextureFeatureSwitch"));
+}
+
+#endif // WITH_EDITOR
+
 UMaterialExpressionAdd::UMaterialExpressionAdd(const FObjectInitializer& ObjectInitializer)
 	: Super(ObjectInitializer)
 {
