Missing file from previous submit

diff --git a/Engine/Source/Runtime/Engine/Private/Materials/HLSLMaterialTranslator.h b/Engine/Source/Runtime/Engine/Private/Materials/HLSLMaterialTranslator.h
--- a/Engine/Source/Runtime/Engine/Private/Materials/HLSLMaterialTranslator.h
+++ b/Engine/Source/Runtime/Engine/Private/Materials/HLSLMaterialTranslator.h
@@ -2874,6 +2874,8 @@ protected:
 			{MEVP_TemporalSampleCount, MCT_Float1, TEXT("View.TemporalAAParams.y"), nullptr},
 			{MEVP_TemporalSampleIndex, MCT_Float1, TEXT("View.TemporalAAParams.x"), nullptr},
 			{MEVP_TemporalSampleOffset, MCT_Float2, TEXT("View.TemporalAAParams.zw"), nullptr},
+			{MEVP_RuntimeVirtualTextureOutputLevel, MCT_Float1, TEXT("View.VirtualTextureParams.x"), nullptr},
+			{MEVP_RuntimeVirtualTextureOutputDerivative, MCT_Float2, TEXT("View.VirtualTextureParams.zw"), nullptr},
 		};
 		static_assert((sizeof(ViewPropertyMetaArray) / sizeof(ViewPropertyMetaArray[0])) == MEVP_MAX, "incoherency between EMaterialExposedViewProperty and ViewPropertyMetaArray");
 
