Fix lightmap texture properties for non-VT case (causing bad SRGB value, resulting in wrong lightmap color)

diff --git a/Engine/Source/Runtime/Engine/Private/Texture.cpp b/Engine/Source/Runtime/Engine/Private/Texture.cpp
--- a/Engine/Source/Runtime/Engine/Private/Texture.cpp
+++ b/Engine/Source/Runtime/Engine/Private/Texture.cpp
@@ -1266,17 +1266,28 @@ void UTexture::GetLayerFormatSettings(int32 LayerIndex, FTextureFormatSettings&
 void UTexture::SetLayerFormatSettings(int32 LayerIndex, const FTextureFormatSettings& InSettings)
 {
 	check(LayerIndex >= 0);
-	if (LayerIndex >= LayerFormatSettings.Num())
+	if (LayerIndex == 0 && LayerFormatSettings.Num() == 0)
 	{
-		FTextureFormatSettings DefaultSettings;
-		GetDefaultFormatSettings(DefaultSettings);
-		LayerFormatSettings.Reserve(LayerIndex + 1);
-		while (LayerIndex >= LayerFormatSettings.Num())
+		// Apply layer0 settings directly to texture properties
+		CompressionSettings = InSettings.CompressionSettings;
+		CompressionNone = InSettings.CompressionNone;
+		CompressionNoAlpha = InSettings.CompressionNoAlpha;
+		SRGB = InSettings.SRGB;
+	}
+	else
+	{
+		if (LayerIndex >= LayerFormatSettings.Num())
 		{
-			LayerFormatSettings.Add(DefaultSettings);
+			FTextureFormatSettings DefaultSettings;
+			GetDefaultFormatSettings(DefaultSettings);
+			LayerFormatSettings.Reserve(LayerIndex + 1);
+			while (LayerIndex >= LayerFormatSettings.Num())
+			{
+				LayerFormatSettings.Add(DefaultSettings);
+			}
 		}
+		LayerFormatSettings[LayerIndex] = InSettings;
 	}
-	LayerFormatSettings[LayerIndex] = InSettings;
 }
 
 #endif // #if WITH_EDITOR
