Virtual textures are resized to be a least 1 tile in size when converted
Generally no good reason to have VTs this small, but can be useful when mass-converting existing textures/materials to VT

diff --git a/Engine/Source/Runtime/Engine/Private/VT/VirtualTextureDataBuilder.cpp b/Engine/Source/Runtime/Engine/Private/VT/VirtualTextureDataBuilder.cpp
--- a/Engine/Source/Runtime/Engine/Private/VT/VirtualTextureDataBuilder.cpp
+++ b/Engine/Source/Runtime/Engine/Private/VT/VirtualTextureDataBuilder.cpp
@@ -257,6 +257,7 @@ FVirtualTextureDataBuilder::FVirtualTextureDataBuilder(FVirtualTextureBuiltData
 	, SizeInBlocksY(0)
 	, BlockSizeX(0)
 	, BlockSizeY(0)
+	, BlockSizeScale(1)
 	, SizeX(0)
 	, SizeY(0)
 {
@@ -277,6 +278,8 @@ void FVirtualTextureDataBuilder::Build(const FTextureSourceData& InSourceData, c
 
 	SettingsPerLayer.AddUninitialized(NumLayers);
 	FMemory::Memcpy(&SettingsPerLayer[0], InSettingsPerLayer, sizeof(FTextureBuildSettings) * NumLayers);
+	const FTextureBuildSettings& BuildSettingsLayer0 = SettingsPerLayer[0];
+	const int32 TileSize = BuildSettingsLayer0.VirtualTextureTileSize;
 
 	BlockSizeX = InSourceData.BlockSizeX;
 	BlockSizeY = InSourceData.BlockSizeY;
@@ -301,17 +304,23 @@ void FVirtualTextureDataBuilder::Build(const FTextureSourceData& InSourceData, c
 		break;
 	}
 
-	SizeInBlocksX = InSourceData.SizeInBlocksX;
-	SizeInBlocksY = InSourceData.SizeInBlocksY;
-	SizeX = BlockSizeX * SizeInBlocksX;
-	SizeY = BlockSizeY * SizeInBlocksY;
-
 	// We require VT blocks (UDIM pages) to be PoT, but multi block textures may have full logical dimension that's not PoT
 	check(FMath::IsPowerOfTwo(BlockSizeX));
 	check(FMath::IsPowerOfTwo(BlockSizeY));
 
-	const FTextureBuildSettings& BuildSettingsLayer0 = SettingsPerLayer[0];
-	const int32 TileSize = BuildSettingsLayer0.VirtualTextureTileSize;
+	// Ensure block size is at least 1 tile, while preserving aspect ratio
+	BlockSizeScale = 1;
+	while (BlockSizeX < TileSize || BlockSizeY < TileSize)
+	{
+		BlockSizeX *= 2;
+		BlockSizeY *= 2;
+		BlockSizeScale *= 2;
+	}
+
+	SizeInBlocksX = InSourceData.SizeInBlocksX;
+	SizeInBlocksY = InSourceData.SizeInBlocksY;
+	SizeX = BlockSizeX * SizeInBlocksX;
+	SizeY = BlockSizeY * SizeInBlocksY;
 
 	//NOTE: OutData may point to a previously build data so it is important to
 	//properly initialize all fields and not assume this is a freshly constructed object
@@ -870,14 +879,39 @@ void FVirtualTextureDataBuilder::BuildSourcePixels(const FTextureSourceData& Sou
 
 			// Use the texture compressor module to do all the hard work
 			TArray<FCompressedImage2D> CompressedMips;
-			if (!Compressor->BuildTexture(SourceMips, *CompositeSourceMips, TBSettings, CompressedMips))
+			bool bBuildTextureResult = false;
+			if (BlockSizeScale == 1)
 			{
-				check(false);
+				bBuildTextureResult = Compressor->BuildTexture(SourceMips, *CompositeSourceMips, TBSettings, CompressedMips);
 			}
+			else
+			{
+				TArray<FImage> ScaledSourceMips;
+				TArray<FImage> ScaledCompositeMips;
+				ScaledSourceMips.Reserve(SourceMips.Num());
+				ScaledCompositeMips.Reserve(CompositeSourceMips->Num());
+				for (const FImage& SrcMip : SourceMips)
+				{
+					FImage* ScaledMip = new(ScaledSourceMips) FImage;
+					SrcMip.ResizeTo(*ScaledMip, SrcMip.SizeX * BlockSizeScale, SrcMip.SizeY * BlockSizeScale, SrcMip.Format, SrcMip.GammaSpace);
+				}
+
+				for (const FImage& SrcMip : *CompositeSourceMips)
+				{
+					FImage* ScaledMip = new(ScaledCompositeMips) FImage;
+					SrcMip.ResizeTo(*ScaledMip, SrcMip.SizeX * BlockSizeScale, SrcMip.SizeY * BlockSizeScale, SrcMip.Format, SrcMip.GammaSpace);
+				}
+
+				bBuildTextureResult = Compressor->BuildTexture(ScaledSourceMips, ScaledCompositeMips, TBSettings, CompressedMips);
+			}
+
+			check(bBuildTextureResult);
 
 			// Get size of block from Compressor output, since it may have been padded/adjusted
 			BlockData.SizeX = CompressedMips[0].SizeX;
 			BlockData.SizeY = CompressedMips[0].SizeY;
+			check(BlockData.SizeX << BlockData.MipBias == BlockSizeX);
+			check(BlockData.SizeY << BlockData.MipBias == BlockSizeY);
 
 			const uint32 BlockSize = FMath::Max(BlockData.SizeX, BlockData.SizeY);
 			if (NumBlocks == 1u)
diff --git a/Engine/Source/Runtime/Engine/Private/VT/VirtualTextureDataBuilder.h b/Engine/Source/Runtime/Engine/Private/VT/VirtualTextureDataBuilder.h
--- a/Engine/Source/Runtime/Engine/Private/VT/VirtualTextureDataBuilder.h
+++ b/Engine/Source/Runtime/Engine/Private/VT/VirtualTextureDataBuilder.h
@@ -89,6 +89,7 @@ private:
 	int32 SizeInBlocksY;
 	int32 BlockSizeX;
 	int32 BlockSizeY;
+	int32 BlockSizeScale;
 	int32 SizeX;
 	int32 SizeY;
 
diff --git a/Engine/Source/Runtime/ImageCore/Private/ImageCore.cpp b/Engine/Source/Runtime/ImageCore/Private/ImageCore.cpp
--- a/Engine/Source/Runtime/ImageCore/Private/ImageCore.cpp
+++ b/Engine/Source/Runtime/ImageCore/Private/ImageCore.cpp
@@ -239,6 +239,48 @@ static void CopyImage(const FImage& SrcImage, FImage& DestImage)
 	}
 }
 
+static FLinearColor SampleImage(const FLinearColor* Pixels, int Width, int Height, float X, float Y)
+{
+	const int32 TexelX0 = FMath::FloorToInt(X);
+	const int32 TexelY0 = FMath::FloorToInt(Y);
+	const int32 TexelX1 = FMath::Min(TexelX0 + 1, Width - 1);
+	const int32 TexelY1 = FMath::Min(TexelY0 + 1, Height - 1);
+	checkSlow(TexelX0 >= 0 && TexelX0 < Width);
+	checkSlow(TexelY0 >= 0 && TexelY0 < Width);
+
+	const float FracX1 = FMath::Frac(X);
+	const float FracY1 = FMath::Frac(Y);
+	const float FracX0 = 1.0f - FracX1;
+	const float FracY0 = 1.0f - FracY1;
+	const FLinearColor& Color00 = Pixels[TexelY0 * Width + TexelX0];
+	const FLinearColor& Color01 = Pixels[TexelY1 * Width + TexelX0];
+	const FLinearColor& Color10 = Pixels[TexelY0 * Width + TexelX1];
+	const FLinearColor& Color11 = Pixels[TexelY1 * Width + TexelX1];
+	return
+		Color00 * (FracX0 * FracY0) +
+		Color01 * (FracX0 * FracY1) +
+		Color10 * (FracX1 * FracY0) +
+		Color11 * (FracX1 * FracY1);
+}
+
+static void ResizeImage(const FImage& SrcImage, FImage& DestImage)
+{
+	const FLinearColor* SrcPixels = SrcImage.AsRGBA32F();
+	FLinearColor* DestPixels = DestImage.AsRGBA32F();
+	const float DestToSrcScaleX = (float)SrcImage.SizeX / (float)DestImage.SizeX;
+	const float DestToSrcScaleY = (float)SrcImage.SizeY / (float)DestImage.SizeY;
+
+	for (int32 DestY = 0; DestY < DestImage.SizeY; ++DestY)
+	{
+		const float SrcY = (float)DestY * DestToSrcScaleY;
+		for (int32 DestX = 0; DestX < DestImage.SizeX; ++DestX)
+		{
+			const float SrcX = (float)DestX * DestToSrcScaleX;
+			const FLinearColor Color = SampleImage(SrcPixels, SrcImage.SizeX, SrcImage.SizeY, SrcX, SrcY);
+			DestPixels[DestY * DestImage.SizeX + DestX] = Color;
+		}
+	}
+}
 
 /* FImage constructors
  *****************************************************************************/
@@ -301,6 +343,41 @@ void FImage::CopyTo(FImage& DestImage, ERawImageFormat::Type DestFormat, EGammaS
 	CopyImage(*this, DestImage);
 }
 
+void FImage::ResizeTo(FImage& DestImage, int32 DestSizeX, int32 DestSizeY, ERawImageFormat::Type DestFormat, EGammaSpace DestGammaSpace) const
+{
+	check(NumSlices == 1); // only support 1 slice for now
+
+	FImage TempSrcImage;
+	const FImage* SrcImagePtr = this;
+	if (Format != ERawImageFormat::RGBA32F)
+	{
+		CopyTo(TempSrcImage, ERawImageFormat::RGBA32F, EGammaSpace::Linear);
+		SrcImagePtr = &TempSrcImage;
+	}
+
+	if (DestFormat == ERawImageFormat::RGBA32F)
+	{
+		DestImage.SizeX = DestSizeX;
+		DestImage.SizeY = DestSizeY;
+		DestImage.NumSlices = 1;
+		DestImage.Format = DestFormat;
+		DestImage.GammaSpace = DestGammaSpace;
+		InitImageStorage(DestImage);
+		ResizeImage(*SrcImagePtr, DestImage);
+	}
+	else
+	{
+		FImage TempDestImage;
+		TempDestImage.SizeX = DestSizeX;
+		TempDestImage.SizeY = DestSizeY;
+		TempDestImage.NumSlices = 1;
+		TempDestImage.Format = ERawImageFormat::RGBA32F;
+		TempDestImage.GammaSpace = DestGammaSpace;
+		InitImageStorage(TempDestImage);
+		ResizeImage(*SrcImagePtr, TempDestImage);
+		TempDestImage.CopyTo(DestImage, DestFormat, DestGammaSpace);
+	}
+}
 
 int32 FImage::GetBytesPerPixel() const
 {
diff --git a/Engine/Source/Runtime/ImageCore/Public/ImageCore.h b/Engine/Source/Runtime/ImageCore/Public/ImageCore.h
--- a/Engine/Source/Runtime/ImageCore/Public/ImageCore.h
+++ b/Engine/Source/Runtime/ImageCore/Public/ImageCore.h
@@ -91,6 +91,18 @@ public:
 	 */
 	IMAGECORE_API void CopyTo(FImage& DestImage, ERawImageFormat::Type DestFormat, EGammaSpace DestGammaSpace) const;
 
+	/**
+	 * Copies and resizes the image to a destination image with the specified size and format.
+	 * Resize is done using bilinear filtering
+	 *
+	 * @param DestImage - The destination image.
+	 * @param DestSizeX - Width of the resized image
+	 * @param DestSizeY - Height of the resized image
+	 * @param DestFormat - The destination image format.
+	 * @param DestSRGB - Whether the destination image is in SRGB format.
+	 */
+	IMAGECORE_API void ResizeTo(FImage& DestImage, int32 DestSizeX, int32 DestSizeY, ERawImageFormat::Type DestFormat, EGammaSpace DestGammaSpace) const;
+
 	/**
 	 * Gets the number of bytes per pixel.
 	 *
