Add tooltip text to UMaterialExpressionSetMaterialAttributes to hopefully make unintuitive behavior more clear

diff --git a/Engine/Source/Runtime/Engine/Classes/Materials/MaterialExpressionSetMaterialAttributes.h b/Engine/Source/Runtime/Engine/Classes/Materials/MaterialExpressionSetMaterialAttributes.h
--- a/Engine/Source/Runtime/Engine/Classes/Materials/MaterialExpressionSetMaterialAttributes.h
+++ b/Engine/Source/Runtime/Engine/Classes/Materials/MaterialExpressionSetMaterialAttributes.h
@@ -42,6 +42,7 @@ class UMaterialExpressionSetMaterialAttributes : public UMaterialExpression
 	virtual uint32 GetInputType(int32 InputIndex) override;
 	virtual bool IsInputConnectionRequired(int32 InputIndex) const override {return true;}
 	virtual bool IsResultMaterialAttributes(int32 OutputIndex) override {return true;}
+	virtual void GetExpressionToolTip(TArray<FString>& OutToolTip) override;
 #endif
 	//~ End UMaterialExpression Interface
 };
diff --git a/Engine/Source/Runtime/Engine/Private/Materials/MaterialExpressions.cpp b/Engine/Source/Runtime/Engine/Private/Materials/MaterialExpressions.cpp
--- a/Engine/Source/Runtime/Engine/Private/Materials/MaterialExpressions.cpp
+++ b/Engine/Source/Runtime/Engine/Private/Materials/MaterialExpressions.cpp
@@ -5532,6 +5532,11 @@ uint32 UMaterialExpressionSetMaterialAttributes::GetInputType(int32 InputIndex)
 	return InputType;
 }
 
+void UMaterialExpressionSetMaterialAttributes::GetExpressionToolTip(TArray<FString>& OutToolTip)
+{
+	ConvertToMultilineToolTip(TEXT("Allows assigning values to specific inputs on a material attributes pin. Any unconnected inputs will be unchanged."), 40, OutToolTip);
+}
+
 void UMaterialExpressionSetMaterialAttributes::PreEditChange(UProperty* PropertyAboutToChange)
 {
 	// Backup attribute array so we can re-connect pins
