DR - the map build data holding the Lightmap info cannot be duplatec by the automatic duplcation code as it holds references into the original level via UIDs etc. Therfore we null it out and require the rebuild of lighting after duplication.

diff --git a/Engine/Source/Runtime/Engine/Classes/Engine/Level.h b/Engine/Source/Runtime/Engine/Classes/Engine/Level.h
--- a/Engine/Source/Runtime/Engine/Classes/Engine/Level.h
+++ b/Engine/Source/Runtime/Engine/Classes/Engine/Level.h
@@ -481,7 +481,7 @@ public:
 	 * Registry for data from the map build.  This is stored in a separate package from the level to speed up saving / autosaving. 
 	 * ReleaseRenderingResources must be called before changing what is referenced, to update the rendering thread state.
 	 */
-	UPROPERTY()
+	UPROPERTY(NonPIEDuplicateTransient)
 	UMapBuildDataRegistry* MapBuildData;
 
 	/** Level offset at time when lighting was built */
