Reflection capture component should honor detail mode setting in the same way other components do.

diff --git a/Engine/Source/Runtime/Engine/Private/Components/ReflectionCaptureComponent.cpp b/Engine/Source/Runtime/Engine/Private/Components/ReflectionCaptureComponent.cpp
--- a/Engine/Source/Runtime/Engine/Private/Components/ReflectionCaptureComponent.cpp
+++ b/Engine/Source/Runtime/Engine/Private/Components/ReflectionCaptureComponent.cpp
@@ -738,7 +738,7 @@ void UReflectionCaptureComponent::CreateRenderState_Concurrent()
 
 	UpdatePreviewShape();
 
-	if (ShouldRender())
+	if (ShouldComponentAddToScene() && ShouldRender())
 	{
 		GetWorld()->Scene->AddReflectionCapture(this);
 	}
@@ -752,7 +752,7 @@ void UReflectionCaptureComponent::SendRenderTransform_Concurrent()
 	{
 		UpdatePreviewShape();
 
-		if (ShouldRender())
+		if (ShouldComponentAddToScene() && ShouldRender())
 		{
 			GetWorld()->Scene->UpdateReflectionCaptureTransform(this);
 		}
