Hopefully fix static analysis warnings

diff --git a/Engine/Source/Runtime/Engine/Private/LightMap.cpp b/Engine/Source/Runtime/Engine/Private/LightMap.cpp
--- a/Engine/Source/Runtime/Engine/Private/LightMap.cpp
+++ b/Engine/Source/Runtime/Engine/Private/LightMap.cpp
@@ -1311,6 +1311,7 @@ bool FLightMapPendingTexture::NeedsStaticShadowTexture() const
 void FLightMapPendingTexture::EncodeSkyOcclusionTexture(UTexture* Texture, uint32 LayerIndex, const FColor& TextureColor)
 {
 	const int32 NumMips = Texture->Source.GetNumMips();
+	check(NumMips > 0);
 
 	FTextureFormatSettings FormatSettings;
 	FormatSettings.SRGB = false;
@@ -1388,6 +1389,7 @@ void FLightMapPendingTexture::EncodeSkyOcclusionTexture(UTexture* Texture, uint3
 void FLightMapPendingTexture::EncodeAOMaskTexture(UTexture* Texture, uint32 LayerIndex, const FColor& TextureColor)
 {
 	const int32 NumMips = Texture->Source.GetNumMips();
+	check(NumMips > 0);
 
 	FTextureFormatSettings FormatSettings;
 	FormatSettings.SRGB = false;
@@ -1463,6 +1465,8 @@ void FLightMapPendingTexture::EncodeAOMaskTexture(UTexture* Texture, uint32 Laye
 void FLightMapPendingTexture::EncodeShadowMapTexture(const TArray<TArray<FFourDistanceFieldSamples>>& MipData, UTexture* Texture, uint32 LayerIndex)
 {
 	const int32 NumMips = Texture->Source.GetNumMips();
+	check(NumMips > 0);
+	check(NumMips == MipData.Num());
 
 	FTextureFormatSettings FormatSettings;
 	FormatSettings.SRGB = false;
