Fix crash when VT connected to multiple custom output pins.  Not safe to use address of local variable as key

diff --git a/Engine/Source/Runtime/Engine/Private/Materials/HLSLMaterialTranslator.h b/Engine/Source/Runtime/Engine/Private/Materials/HLSLMaterialTranslator.h
--- a/Engine/Source/Runtime/Engine/Private/Materials/HLSLMaterialTranslator.h
+++ b/Engine/Source/Runtime/Engine/Private/Materials/HLSLMaterialTranslator.h
@@ -144,7 +144,7 @@ struct FShaderCodeChunk
 
 struct FMaterialVTStackEntry
 {
-	TArray<FShaderCodeChunk>* Scope;
+	uint64 ScopeID;
 	uint64 CoordinateHash;
 	uint64 MipValue0Hash;
 	uint64 MipValue1Hash;
@@ -174,6 +174,8 @@ protected:
 	TArray<FMaterialParameterInfo> ParameterOwnerStack;
 	/** The code chunks corresponding to the currently compiled property or custom output. */
 	TArray<FShaderCodeChunk>* CurrentScopeChunks;
+	uint64 CurrentScopeID;
+	uint64 NextTempScopeID;
 
 	// List of Shared pixel properties. Used to share generated code
 	bool SharedPixelProperties[CompiledMP_MAX];
@@ -318,6 +320,9 @@ public:
 		const ITargetPlatform* InTargetPlatform = nullptr) //if InTargetPlatform is nullptr, we use the current active
 	:	ShaderFrequency(SF_Pixel)
 	,	MaterialProperty(MP_EmissiveColor)
+	,	CurrentScopeChunks(nullptr)
+	,	CurrentScopeID(0u)
+	,	NextTempScopeID(SF_NumFrequencies)
 	,	Material(InMaterial)
 	,	MaterialCompilationOutput(InMaterialCompilationOutput)
 	,	StaticParameters(InStaticParameters)
@@ -430,6 +435,20 @@ public:
 		FunctionStacks[Frequency].Empty();
 	}
 
+	void AssignTempScope(TArray<FShaderCodeChunk>& InScope)
+	{
+		CurrentScopeChunks = &InScope;
+		CurrentScopeID = NextTempScopeID++;
+	}
+
+	void AssignShaderFrequencyScope(EShaderFrequency InShaderFrequency)
+	{
+		check(InShaderFrequency < SF_NumFrequencies);
+		check(InShaderFrequency < NextTempScopeID);
+		CurrentScopeChunks = &SharedPropertyCodeChunks[InShaderFrequency];
+		CurrentScopeID = (uint64)InShaderFrequency;
+	}
+
 	void GatherCustomVertexInterpolators(TArray<UMaterialExpression*> Expressions)
 	{
 		for (UMaterialExpression* Expression : Expressions)
@@ -437,7 +456,7 @@ public:
 			if (UMaterialExpressionVertexInterpolator* Interpolator = Cast<UMaterialExpressionVertexInterpolator>(Expression))
 			{
 				TArray<FShaderCodeChunk> CustomExpressionChunks;
-				CurrentScopeChunks = &CustomExpressionChunks;
+				AssignTempScope(CustomExpressionChunks);
 
 				// Errors are appended to a temporary pool as it's not known at this stage which interpolators are required
 				CompileErrorsSink = &Interpolator->CompileErrors;
@@ -569,7 +588,7 @@ public:
 						MaterialProperty = MP_MAX; // Indicates we're not compiling any material property.
 						ShaderFrequency = SF_Pixel;
 						TArray<FShaderCodeChunk> CustomExpressionChunks;
-						CurrentScopeChunks = &CustomExpressionChunks; //-V506
+						AssignTempScope(CustomExpressionChunks);
 						CustomOutput->Compile(this, Index);
 					}
 
@@ -942,7 +961,7 @@ ResourcesString = TEXT("");
 						FunctionStacks[ShaderFrequency].Add(FMaterialFunctionCompileState(nullptr));
 
 						CustomExpressionChunks.Empty();
-						CurrentScopeChunks = &CustomExpressionChunks;
+						AssignTempScope(CustomExpressionChunks);
 						int32 Result = Material->CompileCustomAttribute(Attribute.AttributeID, this);
 
 						// Consider attribute used if varies from default value
@@ -2362,8 +2381,7 @@ protected:
 		}
 
 		bCompilingPreviousFrame = bUsePreviousFrameTime;
-
-		CurrentScopeChunks = &SharedPropertyCodeChunks[ShaderFrequency];
+		AssignShaderFrequencyScope(ShaderFrequency);
 	}
 
 	virtual void PushMaterialAttribute(const FGuid& InAttributeID) override
@@ -3836,7 +3854,7 @@ protected:
 		const uint64 MipValue0Hash = GetParameterHash(MipValue0Index);
 		const uint64 MipValue1Hash = GetParameterHash(MipValue1Index);
 
-		uint64 Hash = CityHash128to64({ (uint64)CurrentScopeChunks, CoordinatHash });
+		uint64 Hash = CityHash128to64({ CurrentScopeID, CoordinatHash });
 		Hash = CityHash128to64({ Hash, MipValue0Hash });
 		Hash = CityHash128to64({ Hash, MipValue1Hash });
 		Hash = CityHash128to64({ Hash, (uint64)MipValueMode });
@@ -3852,7 +3870,7 @@ protected:
 			const FMaterialVirtualTextureStack& Stack = MaterialCompilationOutput.UniformExpressionSet.VTStacks[Index];
 			const FMaterialVTStackEntry& Entry = VTStacks[Index];
 			if (!Stack.AreLayersFull() &&
-				Entry.Scope == CurrentScopeChunks &&
+				Entry.ScopeID == CurrentScopeID &&
 				Entry.CoordinateHash == CoordinatHash &&
 				Entry.MipValue0Hash == MipValue0Hash &&
 				Entry.MipValue1Hash == MipValue1Hash &&
@@ -3870,7 +3888,7 @@ protected:
 		const int32 StackIndex = VTStacks.AddDefaulted();
 		VTStackHash.Add(Hash, StackIndex);
 		FMaterialVTStackEntry& Entry = VTStacks[StackIndex];
-		Entry.Scope = CurrentScopeChunks;
+		Entry.ScopeID = CurrentScopeID;
 		Entry.CoordinateHash = CoordinatHash;
 		Entry.MipValue0Hash = MipValue0Hash;
 		Entry.MipValue1Hash = MipValue1Hash;
