Fix FColoredTexturedMaterialRenderProxy's vtable not being exported (UE-75886).

diff --git a/Engine/Source/Runtime/Engine/Public/MaterialShared.h b/Engine/Source/Runtime/Engine/Public/MaterialShared.h
--- a/Engine/Source/Runtime/Engine/Public/MaterialShared.h
+++ b/Engine/Source/Runtime/Engine/Public/MaterialShared.h
@@ -2051,7 +2051,7 @@ public:
 /**
  * An material render proxy which overrides the material's Color vector and Texture parameter (mixed together).
  */
-class FColoredTexturedMaterialRenderProxy : public FColoredMaterialRenderProxy
+class ENGINE_VTABLE FColoredTexturedMaterialRenderProxy : public FColoredMaterialRenderProxy
 {
 public:
 
