From git #5913
- Fix material parameters sometimes not updating if large numeric value is assigned
- Setting 'value - 1.0f' in attempt to force inequality to be true will fail for very large floating point values, where precision is > 1

diff --git a/Engine/Source/Runtime/Engine/Private/Materials/MaterialInstance.cpp b/Engine/Source/Runtime/Engine/Private/Materials/MaterialInstance.cpp
--- a/Engine/Source/Runtime/Engine/Private/Materials/MaterialInstance.cpp
+++ b/Engine/Source/Runtime/Engine/Private/Materials/MaterialInstance.cpp
@@ -3411,18 +3411,18 @@ void UMaterialInstance::SetVectorParameterValueInternal(const FMaterialParameter
 
 	FVectorParameterValue* ParameterValue = GameThread_FindParameterByName(VectorParameterValues, ParameterInfo);
 
+	bool bForceUpdate = false;
 	if(!ParameterValue)
 	{
 		// If there's no element for the named parameter in array yet, add one.
 		ParameterValue = new(VectorParameterValues) FVectorParameterValue;
 		ParameterValue->ParameterInfo = ParameterInfo;
 		ParameterValue->ExpressionGUID.Invalidate();
-		// Force an update on first use
-		ParameterValue->ParameterValue.B = Value.B - 1.f;
+		bForceUpdate = true;
 	}
 
 	// Don't enqueue an update if it isn't needed
-	if (ParameterValue->ParameterValue != Value)
+	if (bForceUpdate || ParameterValue->ParameterValue != Value)
 	{
 		ParameterValue->ParameterValue = Value;
 		// Update the material instance data in the rendering thread.
@@ -3456,18 +3456,18 @@ void UMaterialInstance::SetScalarParameterValueInternal(const FMaterialParameter
 
 	FScalarParameterValue* ParameterValue = GameThread_FindParameterByName(ScalarParameterValues, ParameterInfo);
 
+	bool bForceUpdate = false;
 	if(!ParameterValue)
 	{
 		// If there's no element for the named parameter in array yet, add one.
 		ParameterValue = new(ScalarParameterValues) FScalarParameterValue;
 		ParameterValue->ParameterInfo = ParameterInfo;
 		ParameterValue->ExpressionGUID.Invalidate();
-		// Force an update on first use
-		ParameterValue->ParameterValue = Value - 1.f;
+		bForceUpdate = true;
 	}
 
 	// Don't enqueue an update if it isn't needed
-	if (ParameterValue->ParameterValue != Value)
+	if (bForceUpdate || ParameterValue->ParameterValue != Value)
 	{
 		ParameterValue->ParameterValue = Value;
 		// Update the material instance data in the rendering thread.
@@ -3516,18 +3516,18 @@ void UMaterialInstance::SetTextureParameterValueInternal(const FMaterialParamete
 
 	FTextureParameterValue* ParameterValue = GameThread_FindParameterByName(TextureParameterValues, ParameterInfo);
 
+	bool bForceUpdate = false;
 	if(!ParameterValue)
 	{
 		// If there's no element for the named parameter in array yet, add one.
 		ParameterValue = new(TextureParameterValues) FTextureParameterValue;
 		ParameterValue->ParameterInfo = ParameterInfo;
 		ParameterValue->ExpressionGUID.Invalidate();
-		// Force an update on first use
-		ParameterValue->ParameterValue = Value == GEngine->DefaultDiffuseTexture ? NULL : GEngine->DefaultDiffuseTexture;
+		bForceUpdate = true;
 	}
 
 	// Don't enqueue an update if it isn't needed
-	if (ParameterValue->ParameterValue != Value)
+	if (bForceUpdate || ParameterValue->ParameterValue != Value)
 	{
 		// set as an ensure, because it is somehow possible to accidentally pass non-textures into here via blueprints...
 		if (Value && ensureMsgf(Value->IsA(UTexture::StaticClass()), TEXT("Expecting a UTexture! Value='%s' class='%s'"), *Value->GetName(), *Value->GetClass()->GetName()))
@@ -3546,19 +3546,19 @@ void UMaterialInstance::SetFontParameterValueInternal(const FMaterialParameterIn
 
 	FFontParameterValue* ParameterValue = GameThread_FindParameterByName(FontParameterValues, ParameterInfo);
 
+	bool bForceUpdate = false;
 	if(!ParameterValue)
 	{
 			// If there's no element for the named parameter in array yet, add one.
 			ParameterValue = new(FontParameterValues) FFontParameterValue;
 			ParameterValue->ParameterInfo = ParameterInfo;
 			ParameterValue->ExpressionGUID.Invalidate();
-			// Force an update on first use
-			ParameterValue->FontValue == GEngine->GetTinyFont() ? NULL : GEngine->GetTinyFont();
-			ParameterValue->FontPage = FontPage - 1;
+			bForceUpdate = true;
 	}
 
 	// Don't enqueue an update if it isn't needed
-	if (ParameterValue->FontValue != FontValue ||
+	if (bForceUpdate ||
+		ParameterValue->FontValue != FontValue ||
 		ParameterValue->FontPage != FontPage)
 	{
 		ParameterValue->FontValue = FontValue;
