Fix bad signed/unsigned logic causing crash when adjusting RVT mip level removal
#rnx

diff --git a/Engine/Source/Runtime/Engine/Classes/VT/RuntimeVirtualTexture.h b/Engine/Source/Runtime/Engine/Classes/VT/RuntimeVirtualTexture.h
--- a/Engine/Source/Runtime/Engine/Classes/VT/RuntimeVirtualTexture.h
+++ b/Engine/Source/Runtime/Engine/Classes/VT/RuntimeVirtualTexture.h
@@ -41,7 +41,7 @@ protected:
 	int32 TileBorderSize = 2; // 4
 
 	/** Number of low mips to cut from the virtual texture. This can reduce peak virtual texture update cost but will also increase the probability of mip shimmering. */
-	UPROPERTY(EditAnywhere, BlueprintReadOnly, Category = Size, meta = (UIMin = "0", UIMax = "6", DisplayName = "Number of low mips to remove from the virtual texture"))
+	UPROPERTY(EditAnywhere, BlueprintReadOnly, Category = Size, meta = (UIMin = "0", UIMax = "5", DisplayName = "Number of low mips to remove from the virtual texture"))
 	int32 RemoveLowMips = 0;
 
 	/** Enable usage of the virtual texture. This option is intended only for debugging and visualization of the scene without virtual textures. It isn't serialized. */
diff --git a/Engine/Source/Runtime/Engine/Private/VT/RuntimeVirtualTexture.cpp b/Engine/Source/Runtime/Engine/Private/VT/RuntimeVirtualTexture.cpp
--- a/Engine/Source/Runtime/Engine/Private/VT/RuntimeVirtualTexture.cpp
+++ b/Engine/Source/Runtime/Engine/Private/VT/RuntimeVirtualTexture.cpp
@@ -183,7 +183,7 @@ void URuntimeVirtualTexture::GetProducerDescription(FVTProducerDescription& OutD
 	OutDesc.TileBorderSize = GetTileBorderSize();
 	OutDesc.BlockWidthInTiles = GetWidth() / GetTileSize();
 	OutDesc.BlockHeightInTiles = GetHeight() / GetTileSize();
-	OutDesc.MaxLevel = FMath::Max(FMath::CeilLogTwo(FMath::Max(OutDesc.BlockWidthInTiles, OutDesc.BlockHeightInTiles)) - RemoveLowMips, 1u);
+	OutDesc.MaxLevel = FMath::Max((int32)FMath::CeilLogTwo(FMath::Max(OutDesc.BlockWidthInTiles, OutDesc.BlockHeightInTiles)) - RemoveLowMips, 1);
 	OutDesc.DepthInTiles = 1;
 	OutDesc.WidthInBlocks = 1;
 	OutDesc.HeightInBlocks = 1;
