DR - Merge fixes
#rnx

diff --git a/Engine/Plugins/Runtime/AR/Microsoft/HoloLensAR/Source/HoloLensTargetPlatform/Private/HoloLensTargetPlatform.cpp b/Engine/Plugins/Runtime/AR/Microsoft/HoloLensAR/Source/HoloLensTargetPlatform/Private/HoloLensTargetPlatform.cpp
--- a/Engine/Plugins/Runtime/AR/Microsoft/HoloLensAR/Source/HoloLensTargetPlatform/Private/HoloLensTargetPlatform.cpp
+++ b/Engine/Plugins/Runtime/AR/Microsoft/HoloLensAR/Source/HoloLensTargetPlatform/Private/HoloLensTargetPlatform.cpp
@@ -100,10 +100,9 @@ const FPlatformAudioCookOverrides* FHoloLensTargetPlatform::GetAudioCompressionS
 	return nullptr;
 }
 
-void FHoloLensTargetPlatform::GetTextureFormats(const UTexture* InTexture, TArray<FName>& OutFormats) const
+void FHoloLensTargetPlatform::GetTextureFormats(const UTexture* InTexture, TArray< TArray<FName> >& OutFormats) const
 {
-	FName TextureFormatName = GetDefaultTextureFormatName(this, InTexture, EngineSettings, false);
-	OutFormats.Add(TextureFormatName);
+	GetDefaultTextureFormatNamePerLayer(OutFormats.AddDefaulted_GetRef(), this, InTexture, EngineSettings, false);
 }
 
 void FHoloLensTargetPlatform::GetAllTextureFormats(TArray<FName>& OutFormats) const
diff --git a/Engine/Plugins/Runtime/AR/Microsoft/HoloLensAR/Source/HoloLensTargetPlatform/Private/HoloLensTargetPlatform.h b/Engine/Plugins/Runtime/AR/Microsoft/HoloLensAR/Source/HoloLensTargetPlatform/Private/HoloLensTargetPlatform.h
--- a/Engine/Plugins/Runtime/AR/Microsoft/HoloLensAR/Source/HoloLensTargetPlatform/Private/HoloLensTargetPlatform.h
+++ b/Engine/Plugins/Runtime/AR/Microsoft/HoloLensAR/Source/HoloLensTargetPlatform/Private/HoloLensTargetPlatform.h
@@ -72,7 +72,7 @@ public:
 
 	virtual const class FStaticMeshLODSettings& GetStaticMeshLODSettings() const override { return StaticMeshLODSettings; }
 
-	virtual void GetTextureFormats(const UTexture* InTexture, TArray<FName>& OutFormats) const override;
+	virtual void GetTextureFormats(const UTexture* InTexture, TArray< TArray<FName> >& OutFormats) const override;
 
 	virtual const UTextureLODSettings& GetTextureLODSettings() const override { return *TextureLODSettings; }
 
diff --git a/Engine/Source/Runtime/Engine/Private/GeometryCollectionRendering.cpp b/Engine/Source/Runtime/Engine/Private/GeometryCollectionRendering.cpp
--- a/Engine/Source/Runtime/Engine/Private/GeometryCollectionRendering.cpp
+++ b/Engine/Source/Runtime/Engine/Private/GeometryCollectionRendering.cpp
@@ -9,7 +9,7 @@ void FGeometryCollectionVertexFactoryShaderParameters::GetElementShaderBindings(
 	const FSceneInterface* Scene,
 	const FSceneView* View,
 	const FMeshMaterialShader* Shader,
-	bool bShaderRequiresPositionOnlyStream,
+	const EVertexInputStreamType InputStreamType,
 	ERHIFeatureLevel::Type FeatureLevel,
 	const FVertexFactory* VertexFactory,
 	const FMeshBatchElement& BatchElement,
@@ -24,7 +24,7 @@ void FGeometryCollectionVertexFactoryShaderParameters::GetElementShaderBindings(
 	VertexFactoryUniformBuffer = LocalVertexFactory->GetUniformBuffer();
 
 
-	FLocalVertexFactoryShaderParametersBase::GetElementShaderBindingsBase(Scene, View, Shader, bShaderRequiresPositionOnlyStream, FeatureLevel, VertexFactory, BatchElement, VertexFactoryUniformBuffer, ShaderBindings, VertexStreams);
+	FLocalVertexFactoryShaderParametersBase::GetElementShaderBindingsBase(Scene, View, Shader, InputStreamType, FeatureLevel, VertexFactory, BatchElement, VertexFactoryUniformBuffer, ShaderBindings, VertexStreams);
 
 
 	// We only want to set the SRV parameters if we support manual vertex fetch.
diff --git a/Engine/Source/Runtime/Engine/Public/GeometryCollectionRendering.h b/Engine/Source/Runtime/Engine/Public/GeometryCollectionRendering.h
--- a/Engine/Source/Runtime/Engine/Public/GeometryCollectionRendering.h
+++ b/Engine/Source/Runtime/Engine/Public/GeometryCollectionRendering.h
@@ -28,7 +28,7 @@ public:
 		const FSceneInterface* Scene,
 		const FSceneView* View,
 		const FMeshMaterialShader* Shader,
-		bool bShaderRequiresPositionOnlyStream,
+		const EVertexInputStreamType InputStreamType,
 		ERHIFeatureLevel::Type FeatureLevel,
 		const FVertexFactory* VertexFactory,
 		const FMeshBatchElement& BatchElement,
diff --git a/Engine/Source/Runtime/MRMesh/Private/MRMeshComponent.cpp b/Engine/Source/Runtime/MRMesh/Private/MRMeshComponent.cpp
--- a/Engine/Source/Runtime/MRMesh/Private/MRMeshComponent.cpp
+++ b/Engine/Source/Runtime/MRMesh/Private/MRMeshComponent.cpp
@@ -411,7 +411,7 @@ private:
 							Mesh.MaterialRenderProxy = MaterialProxy;
 
 							FDynamicPrimitiveUniformBuffer& DynamicPrimitiveUniformBuffer = Collector.AllocateOneFrameResource<FDynamicPrimitiveUniformBuffer>();
-							DynamicPrimitiveUniformBuffer.Set(GetLocalToWorld(), GetLocalToWorld(), InfiniteBounds, InfiniteBounds, true, false, UseEditorDepthTest(), false);
+							DynamicPrimitiveUniformBuffer.Set(GetLocalToWorld(), GetLocalToWorld(), InfiniteBounds, InfiniteBounds, true, false, DrawsVelocity(), false);
 							BatchElement.PrimitiveUniformBufferResource = &DynamicPrimitiveUniformBuffer.UniformBuffer;
 
 							BatchElement.FirstIndex = 0;
diff --git a/Engine/Source/Runtime/NullDrv/Public/NullRHI.h b/Engine/Source/Runtime/NullDrv/Public/NullRHI.h
--- a/Engine/Source/Runtime/NullDrv/Public/NullRHI.h
+++ b/Engine/Source/Runtime/NullDrv/Public/NullRHI.h
@@ -477,6 +477,7 @@ public:
 	{
 
 	}
+	virtual void RHIAdvanceFrameForGetViewportBackBuffer(FRHIViewport* Viewport) final override
 	{
 
 	}
