DR - More fixes
#rnx

diff --git a/Engine/Plugins/FX/Niagara/Source/NiagaraVertexFactories/Public/NiagaraDrawIndirect.h b/Engine/Plugins/FX/Niagara/Source/NiagaraVertexFactories/Public/NiagaraDrawIndirect.h
--- a/Engine/Plugins/FX/Niagara/Source/NiagaraVertexFactories/Public/NiagaraDrawIndirect.h
+++ b/Engine/Plugins/FX/Niagara/Source/NiagaraVertexFactories/Public/NiagaraDrawIndirect.h
@@ -55,8 +55,8 @@ public:
 	FNiagaraDrawIndirectArgsGenCS(const ShaderMetaType::CompiledShaderInitializerType& Initializer);
 
 	bool Serialize(FArchive& Ar);
-	void SetOutput(FRHICommandList& RHICmdList, FUnorderedAccessViewRHIParamRef DrawIndirectArgsUAV, FUnorderedAccessViewRHIParamRef InstanceCountsUAV);
-	void SetParameters(FRHICommandList& RHICmdList, FShaderResourceViewRHIParamRef TaskInfosBuffer, int32 NumArgGenTasks, int32 NumInstanceCountClearTasks);
+	void SetOutput(FRHICommandList& RHICmdList, FRHIUnorderedAccessView* DrawIndirectArgsUAV, FRHIUnorderedAccessView* InstanceCountsUAV);
+	void SetParameters(FRHICommandList& RHICmdList, FRHIShaderResourceView* TaskInfosBuffer, int32 NumArgGenTasks, int32 NumInstanceCountClearTasks);
 	void UnbindBuffers(FRHICommandList& RHICmdList);
 
 protected:
diff --git a/Engine/Source/Runtime/Engine/Public/GeometryCollectionRendering.h b/Engine/Source/Runtime/Engine/Public/GeometryCollectionRendering.h
--- a/Engine/Source/Runtime/Engine/Public/GeometryCollectionRendering.h
+++ b/Engine/Source/Runtime/Engine/Public/GeometryCollectionRendering.h
@@ -66,9 +66,9 @@ public:
 	// Data includes what we need for transform and everything in local vertex factory too
 	struct FDataType : public FLocalVertexFactory::FDataType
 	{
-		FShaderResourceViewRHIParamRef InstanceTransformSRV = nullptr;
-		FShaderResourceViewRHIParamRef InstancePrevTransformSRV = nullptr;
-		FShaderResourceViewRHIParamRef InstanceBoneMapSRV = nullptr;
+		FRHIShaderResourceView* InstanceTransformSRV = nullptr;
+		FRHIShaderResourceView* InstancePrevTransformSRV = nullptr;
+		FRHIShaderResourceView* InstanceBoneMapSRV = nullptr;
 	};
 
 	//
@@ -144,27 +144,27 @@ public:
 		return ShaderFrequency == SF_Vertex ? new FGeometryCollectionVertexFactoryShaderParameters() : NULL;
 	}
 
-	inline void SetInstanceTransformSRV(FShaderResourceViewRHIParamRef InstanceTransformSRV)
+	inline void SetInstanceTransformSRV(FRHIShaderResourceView* InstanceTransformSRV)
 	{
 		Data.InstanceTransformSRV = InstanceTransformSRV;
 	}
 	
-	inline const FShaderResourceViewRHIParamRef GetInstanceTransformSRV() const
+	inline const FRHIShaderResourceView* GetInstanceTransformSRV() const
 	{
 		return Data.InstanceTransformSRV;
 	}
 
-	inline void SetInstancePrevTransformSRV(FShaderResourceViewRHIParamRef InstancePrevTransformSRV)
+	inline void SetInstancePrevTransformSRV(FRHIShaderResourceView* InstancePrevTransformSRV)
 	{
 		Data.InstancePrevTransformSRV = InstancePrevTransformSRV;
 	}
 
-	inline const FShaderResourceViewRHIParamRef GetInstancePrevTransformSRV() const
+	inline const FRHIShaderResourceView* GetInstancePrevTransformSRV() const
 	{
 		return Data.InstancePrevTransformSRV;
 	}
 
-	inline const FShaderResourceViewRHIParamRef GetInstanceBoneMapSRV() const
+	inline const FRHIShaderResourceView* GetInstanceBoneMapSRV() const
 	{
 		return Data.InstanceBoneMapSRV;
 	}
