DR - Fix
#rnx

diff --git a/Engine/Source/Runtime/Engine/Public/GeometryCollectionRendering.h b/Engine/Source/Runtime/Engine/Public/GeometryCollectionRendering.h
--- a/Engine/Source/Runtime/Engine/Public/GeometryCollectionRendering.h
+++ b/Engine/Source/Runtime/Engine/Public/GeometryCollectionRendering.h
@@ -149,7 +149,7 @@ public:
 		Data.InstanceTransformSRV = InstanceTransformSRV;
 	}
 	
-	inline const FRHIShaderResourceView* GetInstanceTransformSRV() const
+	inline FRHIShaderResourceView* GetInstanceTransformSRV() const
 	{
 		return Data.InstanceTransformSRV;
 	}
@@ -159,12 +159,12 @@ public:
 		Data.InstancePrevTransformSRV = InstancePrevTransformSRV;
 	}
 
-	inline const FRHIShaderResourceView* GetInstancePrevTransformSRV() const
+	inline FRHIShaderResourceView* GetInstancePrevTransformSRV() const
 	{
 		return Data.InstancePrevTransformSRV;
 	}
 
-	inline const FRHIShaderResourceView* GetInstanceBoneMapSRV() const
+	inline FRHIShaderResourceView* GetInstanceBoneMapSRV() const
 	{
 		return Data.InstanceBoneMapSRV;
 	}
