Fix support for crunch VT compression when source texture format isn't BGRA8
- This can happen when manually overriding compression settings for HDR source textures

diff --git a/Engine/Source/Runtime/Engine/Private/VT/VirtualTextureDataBuilder.cpp b/Engine/Source/Runtime/Engine/Private/VT/VirtualTextureDataBuilder.cpp
--- a/Engine/Source/Runtime/Engine/Private/VT/VirtualTextureDataBuilder.cpp
+++ b/Engine/Source/Runtime/Engine/Private/VT/VirtualTextureDataBuilder.cpp
@@ -503,53 +503,16 @@ void FVirtualTextureDataBuilder::BuildTiles(const TArray<FVTSourceTileEntry>& Ti
 	FThreadSafeBool bCompressionError = false;
 	EPixelFormat CompressedFormat = PF_Unknown;
 
-	// Don't want platform specific swizzling for VT tile data, this tends to add extra padding for textures with odd dimensions
-	// (VT physical tiles generally not power-of-2 after adding border)
-	FName TextureFormatName = BuildSettingsForLayer.TextureFormatName;
-	FString BaseTextureFormatName = TextureFormatName.ToString();
-	if (BaseTextureFormatName.StartsWith(TEXT("PS4_")))
-	{
-		TextureFormatName = *BaseTextureFormatName.Replace(TEXT("PS4_"), TEXT(""));
-	}
-	else if (BaseTextureFormatName.StartsWith("XBOXONE_"))
-	{
-		TextureFormatName = *BaseTextureFormatName.Replace(TEXT("XBOXONE_"), TEXT(""));
-	}
-	else if (BaseTextureFormatName.StartsWith("SWITCH_"))
-	{
-		TextureFormatName = *BaseTextureFormatName.Replace(TEXT("SWITCH_"), TEXT(""));
-	}
-
-	// We handle AutoDXT specially here since otherwise the texture format compressor would choose a DXT format for every tile
-	// individually. Causing tiles in the same VT to use different formats which we don't allow.
-	static FName NameDXT1(TEXT("DXT1"));
-	static FName NameDXT5(TEXT("DXT5"));
-	static FName NameAutoDXT(TEXT("AutoDXT"));
-	if (TextureFormatName == NameAutoDXT)
-	{
-		if (LayerData.bHasAlpha)
-		{
-			TextureFormatName = NameDXT5;
-		}
-		else
-		{
-			TextureFormatName = NameDXT1;
-		}
-	}
-
 #if WITH_CRUNCH_COMPRESSION
-	const bool bUseCrunch = BuildSettingsLayer0.bVirtualTextureEnableCompressCrunch &&
-		BuildSettingsLayer0.LossyCompressionAmount != TLCA_None &&
-		CrunchCompression::IsValidFormat(TextureFormatName);
-	if (bUseCrunch)
+	if (LayerData.bUseCrunch)
 	{
 		check(LayerData.ImageFormat == ERawImageFormat::BGRA8);
 
 		FCrunchEncodeParameters CrunchParameters;
 		CrunchParameters.ImageWidth = PhysicalTileSize;
 		CrunchParameters.ImageHeight = PhysicalTileSize;
-		CrunchParameters.bIsGammaCorrected = BuildSettingsForLayer.GetGammaSpace() != EGammaSpace::Linear;
-		CrunchParameters.OutputFormat = TextureFormatName;
+		CrunchParameters.bIsGammaCorrected = LayerData.GammaSpace != EGammaSpace::Linear;
+		CrunchParameters.OutputFormat = LayerData.TextureFormatName;
 		switch (BuildSettingsLayer0.LossyCompressionAmount)
 		{
 		case TLCA_Lowest: CrunchParameters.CompressionAmmount = 0.0f; break;
@@ -573,6 +536,8 @@ void FVirtualTextureDataBuilder::BuildTiles(const TArray<FVTSourceTileEntry>& Ti
 		{
 			const FTextureSourceBlockData& Block = SourceBlocks[Tile.BlockIndex];
 			const FImage& SourceMip = Block.MipsPerLayer[LayerIndex][Tile.MipIndexInBlock];
+			check(SourceMip.Format == LayerData.ImageFormat);
+
 			FPixelDataRectangle SourceData(LayerData.SourceFormat,
 				SourceMip.SizeX,
 				SourceMip.SizeY,
@@ -601,13 +566,15 @@ void FVirtualTextureDataBuilder::BuildTiles(const TArray<FVTSourceTileEntry>& Ti
 
 		if (CrunchCompression::Encode(CrunchParameters, GeneratedData.CodecPayload, GeneratedData.TilePayload))
 		{
+			static FName NameDXT1(TEXT("DXT1"));
+			static FName NameDXT5(TEXT("DXT5"));
 			static FName NameBC4(TEXT("BC4"));
 			static FName NameBC5(TEXT("BC5"));
 			GeneratedData.Codec = EVirtualTextureCodec::Crunch;
-			if (TextureFormatName == NameDXT1) CompressedFormat = PF_DXT1;
-			else if (TextureFormatName == NameDXT5) CompressedFormat = PF_DXT5;
-			else if (TextureFormatName == NameBC4) CompressedFormat = PF_BC4;
-			else if (TextureFormatName == NameBC5) CompressedFormat = PF_BC5;
+			if (LayerData.TextureFormatName == NameDXT1) CompressedFormat = PF_DXT1;
+			else if (LayerData.TextureFormatName == NameDXT5) CompressedFormat = PF_DXT5;
+			else if (LayerData.TextureFormatName == NameBC4) CompressedFormat = PF_BC4;
+			else if (LayerData.TextureFormatName == NameBC5) CompressedFormat = PF_BC5;
 			else CompressedFormat = PF_Unknown;
 		}
 		else
@@ -623,7 +590,7 @@ void FVirtualTextureDataBuilder::BuildTiles(const TArray<FVTSourceTileEntry>& Ti
 		// as these settings were already baked into the SourcePixels.
 		FTextureBuildSettings TBSettings;
 		TBSettings.MaxTextureResolution = TNumericLimits<uint32>::Max();
-		TBSettings.TextureFormatName = TextureFormatName;
+		TBSettings.TextureFormatName = LayerData.TextureFormatName;
 		TBSettings.bSRGB = BuildSettingsForLayer.bSRGB;
 		TBSettings.bUseLegacyGamma = BuildSettingsForLayer.bUseLegacyGamma;
 		TBSettings.MipGenSettings = TMGS_NoMipmaps;
@@ -638,6 +605,8 @@ void FVirtualTextureDataBuilder::BuildTiles(const TArray<FVTSourceTileEntry>& Ti
 
 			const FTextureSourceBlockData& Block = SourceBlocks[Tile.BlockIndex];
 			const FImage& SourceMip = Block.MipsPerLayer[LayerIndex][Tile.MipIndexInBlock];
+			check(SourceMip.Format == LayerData.ImageFormat);
+
 			FPixelDataRectangle SourceData(LayerData.SourceFormat,
 				SourceMip.SizeX,
 				SourceMip.SizeY,
@@ -1063,6 +1032,75 @@ void FVirtualTextureDataBuilder::BuildSourcePixels(const FTextureSourceData& Sou
 			}
 		}
 	}
+
+	for (int32 LayerIndex = 0; LayerIndex < NumLayers; ++LayerIndex)
+	{
+		const FTextureBuildSettings& BuildSettingsForLayer = SettingsPerLayer[LayerIndex];
+		FVirtualTextureSourceLayerData& LayerData = SourceLayers[LayerIndex];
+
+		// Don't want platform specific swizzling for VT tile data, this tends to add extra padding for textures with odd dimensions
+		// (VT physical tiles generally not power-of-2 after adding border)
+		FName TextureFormatName = BuildSettingsForLayer.TextureFormatName;
+		FString BaseTextureFormatName = TextureFormatName.ToString();
+		if (BaseTextureFormatName.StartsWith(TEXT("PS4_")))
+		{
+			TextureFormatName = *BaseTextureFormatName.Replace(TEXT("PS4_"), TEXT(""));
+		}
+		else if (BaseTextureFormatName.StartsWith("XBOXONE_"))
+		{
+			TextureFormatName = *BaseTextureFormatName.Replace(TEXT("XBOXONE_"), TEXT(""));
+		}
+		else if (BaseTextureFormatName.StartsWith("SWITCH_"))
+		{
+			TextureFormatName = *BaseTextureFormatName.Replace(TEXT("SWITCH_"), TEXT(""));
+		}
+
+		// We handle AutoDXT specially here since otherwise the texture format compressor would choose a DXT format for every tile
+		// individually. Causing tiles in the same VT to use different formats which we don't allow.
+		static FName NameDXT1(TEXT("DXT1"));
+		static FName NameDXT5(TEXT("DXT5"));
+		static FName NameAutoDXT(TEXT("AutoDXT"));
+		if (TextureFormatName == NameAutoDXT)
+		{
+			if (LayerData.bHasAlpha)
+			{
+				TextureFormatName = NameDXT5;
+			}
+			else
+			{
+				TextureFormatName = NameDXT1;
+			}
+		}
+
+		bool bUseCrunch = false;
+#if WITH_CRUNCH_COMPRESSION
+		bUseCrunch = SettingsPerLayer[0].bVirtualTextureEnableCompressCrunch &&
+			SettingsPerLayer[0].LossyCompressionAmount != TLCA_None &&
+			CrunchCompression::IsValidFormat(TextureFormatName);
+#endif // WITH_CRUNCH_COMPRESSION
+		LayerData.TextureFormatName = TextureFormatName;
+		LayerData.bUseCrunch = bUseCrunch;
+
+		if (bUseCrunch && LayerData.ImageFormat != ERawImageFormat::BGRA8)
+		{
+			// Crunch input data must be in BGRA8 format
+			LayerData.FormatName = "BGRA8";
+			LayerData.PixelFormat = PF_B8G8R8A8;
+			LayerData.SourceFormat = TSF_BGRA8;
+			LayerData.ImageFormat = ERawImageFormat::BGRA8;
+			LayerData.GammaSpace = BuildSettingsForLayer.GetGammaSpace();
+
+			FImage ConvertedMip;
+			for(FTextureSourceBlockData& SourceBlockData : SourceBlocks)
+			{
+				for (FImage& SourceMip : SourceBlockData.MipsPerLayer[LayerIndex])
+				{
+					SourceMip.CopyTo(ConvertedMip, ERawImageFormat::BGRA8, LayerData.GammaSpace);
+					SourceMip = MoveTemp(ConvertedMip);
+				}
+			}
+		}
+	}
 }
 
 void FVirtualTextureDataBuilder::FreeSourcePixels()
diff --git a/Engine/Source/Runtime/Engine/Private/VT/VirtualTextureDataBuilder.h b/Engine/Source/Runtime/Engine/Private/VT/VirtualTextureDataBuilder.h
--- a/Engine/Source/Runtime/Engine/Private/VT/VirtualTextureDataBuilder.h
+++ b/Engine/Source/Runtime/Engine/Private/VT/VirtualTextureDataBuilder.h
@@ -36,9 +36,11 @@ struct FVirtualTextureSourceLayerData
 	ETextureSourceFormat SourceFormat;
 	EPixelFormat PixelFormat;
 	FName FormatName;
+	FName TextureFormatName;
 
 	EGammaSpace GammaSpace;
 	bool bHasAlpha;
+	bool bUseCrunch;
 };
 
 /**
