- VertexInterpolator material expressions work correctly when creating lightmass/proxy materials
- Lightmass/Proxy replace switches don't attempt to evaluate non-active branch; this fixes problem when non-active branch has compile errors in the current context, which causes compile to fail

diff --git a/Engine/Source/Developer/MaterialBaking/Private/ExportMaterialProxy.h b/Engine/Source/Developer/MaterialBaking/Private/ExportMaterialProxy.h
--- a/Engine/Source/Developer/MaterialBaking/Private/ExportMaterialProxy.h
+++ b/Engine/Source/Developer/MaterialBaking/Private/ExportMaterialProxy.h
@@ -214,8 +214,7 @@ struct FExportMaterialCompiler : public FProxyMaterialCompiler
 		}
 	}
 
-	virtual int32 LightmassReplace(int32 Realtime, int32 Lightmass) override { return Realtime; }
-	virtual int32 MaterialProxyReplace(int32 Realtime, int32 MaterialProxy) override { return MaterialProxy; }
+	virtual EMaterialCompilerType GetCompilerType() const override { return EMaterialCompilerType::MaterialProxy; }
 };
 
 class FExportMaterialProxy : public FMaterial, public FMaterialRenderProxy
diff --git a/Engine/Source/Developer/MaterialUtilities/Private/MaterialUtilities.cpp b/Engine/Source/Developer/MaterialUtilities/Private/MaterialUtilities.cpp
--- a/Engine/Source/Developer/MaterialUtilities/Private/MaterialUtilities.cpp
+++ b/Engine/Source/Developer/MaterialUtilities/Private/MaterialUtilities.cpp
@@ -392,8 +392,7 @@ struct FExportMaterialCompiler : public FProxyMaterialCompiler
 		}
 	}
 	
-	virtual int32 LightmassReplace(int32 Realtime, int32 Lightmass) override { return Realtime; }
-	virtual int32 MaterialProxyReplace(int32 Realtime, int32 MaterialProxy) override { return MaterialProxy; }
+	virtual EMaterialCompilerType GetCompilerType() const override { return EMaterialCompilerType::MaterialProxy; }
 };
 
 
diff --git a/Engine/Source/Editor/UnrealEd/Private/Lightmass/LightmassRender.cpp b/Engine/Source/Editor/UnrealEd/Private/Lightmass/LightmassRender.cpp
--- a/Engine/Source/Editor/UnrealEd/Private/Lightmass/LightmassRender.cpp
+++ b/Engine/Source/Editor/UnrealEd/Private/Lightmass/LightmassRender.cpp
@@ -206,11 +206,9 @@ struct FLightmassMaterialCompiler : public FProxyMaterialCompiler
 		return Compiler->Constant(0.0f);
 	}
 
-	virtual int32 LightmassReplace(int32 Realtime, int32 Lightmass) override { return Lightmass; }
-
 	virtual int32 GIReplace(int32 Direct, int32 StaticIndirect, int32 DynamicIndirect) override { return StaticIndirect; }
 
-	virtual int32 MaterialProxyReplace(int32 Realtime, int32 MaterialProxy) override { return Realtime; }
+	virtual EMaterialCompilerType GetCompilerType() const override { return EMaterialCompilerType::Lightmass; }
 
 #if WITH_EDITOR
 	virtual int32 MaterialBakingWorldPosition() override
diff --git a/Engine/Source/Runtime/Engine/Private/Materials/HLSLMaterialTranslator.h b/Engine/Source/Runtime/Engine/Private/Materials/HLSLMaterialTranslator.h
--- a/Engine/Source/Runtime/Engine/Private/Materials/HLSLMaterialTranslator.h
+++ b/Engine/Source/Runtime/Engine/Private/Materials/HLSLMaterialTranslator.h
@@ -6005,8 +6005,6 @@ protected:
 		return ResultIdx;
 	}
 
-	virtual int32 LightmassReplace(int32 Realtime, int32 Lightmass) override { return Realtime; }
-
 	virtual int32 GIReplace(int32 Direct, int32 StaticIndirect, int32 DynamicIndirect) override 
 	{ 
 		if(Direct == INDEX_NONE || DynamicIndirect == INDEX_NONE)
@@ -6041,8 +6039,6 @@ protected:
 		return AddCodeChunk(ResultType, TEXT("(GetRayTracingQualitySwitch() ? (%s) : (%s))"), *GetParameterCode(RayTraced), *GetParameterCode(Normal));
 	}
 
-	virtual int32 MaterialProxyReplace(int32 Realtime, int32 MaterialProxy) override { return Realtime; }
-
 	virtual int32 VirtualTextureOutputReplace(int32 Default, int32 VirtualTexture) override
 	{
 		if (Default == INDEX_NONE || VirtualTexture == INDEX_NONE)
diff --git a/Engine/Source/Runtime/Engine/Private/Materials/MaterialExpressions.cpp b/Engine/Source/Runtime/Engine/Private/Materials/MaterialExpressions.cpp
--- a/Engine/Source/Runtime/Engine/Private/Materials/MaterialExpressions.cpp
+++ b/Engine/Source/Runtime/Engine/Private/Materials/MaterialExpressions.cpp
@@ -13150,9 +13150,7 @@ int32 UMaterialExpressionLightmassReplace::Compile(class FMaterialCompiler* Comp
 	}
 	else
 	{
-		int32 Arg1 = Realtime.Compile(Compiler);
-		int32 Arg2 = Lightmass.Compile(Compiler);
-		return Compiler->LightmassReplace(Arg1, Arg2);
+		return Compiler->IsLightmassCompiler() ? Lightmass.Compile(Compiler) : Realtime.Compile(Compiler);
 	}
 }
 
@@ -13250,9 +13248,7 @@ int32 UMaterialExpressionMaterialProxyReplace::Compile(class FMaterialCompiler*
 	}
 	else
 	{
-		int32 Arg1 = Realtime.Compile(Compiler);
-		int32 Arg2 = MaterialProxy.Compile(Compiler);
-		return Compiler->MaterialProxyReplace(Arg1, Arg2);
+		return Compiler->IsMaterialProxyCompiler() ? MaterialProxy.Compile(Compiler) : Realtime.Compile(Compiler);
 	}
 }
 
@@ -15469,7 +15465,12 @@ int32 UMaterialExpressionVertexInterpolator::Compile(class FMaterialCompiler* Co
 {
 	if (Input.GetTracedInput().Expression)
 	{
-		if (InterpolatorIndex == INDEX_NONE || CompileErrors.Num() > 0)
+		if (Compiler->IsVertexInterpolatorBypass())
+		{
+			// Certain types of compilers don't support vertex interpolators, just evaluate the input directly in that case
+			return Input.Compile(Compiler);
+		}
+		else if (InterpolatorIndex == INDEX_NONE || CompileErrors.Num() > 0)
 		{
 			// Now this node is confirmed part of the graph, append all errors from the input compilation
 			check(CompileErrors.Num() == CompileErrorExpressions.Num());
@@ -15505,6 +15506,8 @@ int32 UMaterialExpressionVertexInterpolator::CompileInput(class FMaterialCompile
 	InterpolatedType = MCT_Unknown;
 	InterpolatorOffset = INDEX_NONE;
 
+	ensure(!Compiler->IsVertexInterpolatorBypass());
+
 	CompileErrors.Empty();
 	CompileErrorExpressions.Empty();
 
diff --git a/Engine/Source/Runtime/Engine/Public/MaterialCompiler.h b/Engine/Source/Runtime/Engine/Public/MaterialCompiler.h
--- a/Engine/Source/Runtime/Engine/Public/MaterialCompiler.h
+++ b/Engine/Source/Runtime/Engine/Public/MaterialCompiler.h
@@ -36,6 +36,14 @@ enum class EVirtualTextureUnpackType
 	HeightR8G8,
 };
 
+/** What type of compiler is this? Used by material expressions that select input based on compile context */
+enum class EMaterialCompilerType
+{
+	Standard, /** Standard HLSL translator */
+	Lightmass, /** Lightmass proxy compiler */
+	MaterialProxy, /** Flat material proxy compiler */
+};
+
 /** 
  * The interface used to translate material expressions into executable code. 
  * Note: Most member functions should be pure virtual to force a FProxyMaterialCompiler override!
@@ -72,6 +80,16 @@ public:
 
 	virtual int32 CallExpression(FMaterialExpressionKey ExpressionKey,FMaterialCompiler* InCompiler) = 0;
 
+	virtual EMaterialCompilerType GetCompilerType() const { return EMaterialCompilerType::Standard; }
+	inline bool IsMaterialProxyCompiler() const { return GetCompilerType() == EMaterialCompilerType::MaterialProxy; }
+	inline bool IsLightmassCompiler() const { return GetCompilerType() == EMaterialCompilerType::Lightmass; }
+
+	inline bool IsVertexInterpolatorBypass() const
+	{
+		const EMaterialCompilerType Type = GetCompilerType();
+		return Type == EMaterialCompilerType::MaterialProxy || Type == EMaterialCompilerType::Lightmass;
+	}
+
 	virtual EMaterialValueType GetType(int32 Code) = 0;
 
 	virtual EMaterialQualityLevel::Type GetQualityLevel() = 0;
@@ -280,11 +298,9 @@ public:
 	virtual int32 LightmapUVs() = 0;
 	virtual int32 PrecomputedAOMask()  = 0;
 
-	virtual int32 LightmassReplace(int32 Realtime, int32 Lightmass) = 0;
 	virtual int32 GIReplace(int32 Direct, int32 StaticIndirect, int32 DynamicIndirect) = 0;
 	virtual int32 ShadowReplace(int32 Default, int32 Shadow) = 0;
 	virtual int32 RayTracingQualitySwitchReplace(int32 Normal, int32 RayTraced) = 0;
-	virtual int32 MaterialProxyReplace(int32 Realtime, int32 MaterialProxy) = 0;
 	virtual int32 VirtualTextureOutputReplace(int32 Default, int32 VirtualTexture) = 0;
 
 	virtual int32 ObjectOrientation() = 0;
@@ -517,11 +533,9 @@ public:
 	virtual int32 LightmapUVs() override { return Compiler->LightmapUVs(); }
 	virtual int32 PrecomputedAOMask() override { return Compiler->PrecomputedAOMask(); }
 
-	virtual int32 LightmassReplace(int32 Realtime, int32 Lightmass) override { return Realtime; }
 	virtual int32 GIReplace(int32 Direct, int32 StaticIndirect, int32 DynamicIndirect) override { return Compiler->GIReplace(Direct, StaticIndirect, DynamicIndirect); }
 	virtual int32 ShadowReplace(int32 Default, int32 Shadow) override { return Compiler->ShadowReplace(Default, Shadow); }
 	virtual int32 RayTracingQualitySwitchReplace(int32 Normal, int32 RayTraced) override { return Compiler->RayTracingQualitySwitchReplace(Normal, RayTraced); }
-	virtual int32 MaterialProxyReplace(int32 Realtime, int32 MaterialProxy) override { return Realtime; }
 	virtual int32 VirtualTextureOutputReplace(int32 Default, int32 VirtualTexture) override { return Compiler->VirtualTextureOutputReplace(Default, VirtualTexture); }
 	
 	virtual int32 ObjectOrientation() override { return Compiler->ObjectOrientation(); }
