Add material shader stage switch, allows customizing graph for vertex vs pixel shader

diff --git a/Engine/Source/Runtime/Engine/Classes/Materials/Material.h b/Engine/Source/Runtime/Engine/Classes/Materials/Material.h
--- a/Engine/Source/Runtime/Engine/Classes/Materials/Material.h
+++ b/Engine/Source/Runtime/Engine/Classes/Materials/Material.h
@@ -1698,12 +1698,16 @@ protected:
 	 *	@Param	InFeatureLevel			Optional feature level - if supplied, only walk FeatureLevelSwitch branches according to it.
 	 *	@Param	InQuality				Optional quality switch - if supplied, only walk QualitySwitch branches according to it.
 	 *	@Param	InShadingPath			Optional shading path switch - if supplied, only walk ShadingPathSwitch branches according to it.
+	 *	@Param	InShaderFrequency		Optional shader frequency - if supplied, only walk ShaderFrequencySwitch branches according to it.
 	 *
 	 *	@return	bool					true if successful, false if not.
 	 */
 	ENGINE_API virtual bool RecursiveGetExpressionChain(UMaterialExpression* InExpression, TArray<FExpressionInput*>& InOutProcessedInputs, 
 		TArray<UMaterialExpression*>& OutExpressions, struct FStaticParameterSet* InStaticParameterSet,
-		ERHIFeatureLevel::Type InFeatureLevel = ERHIFeatureLevel::Num, EMaterialQualityLevel::Type InQuality = EMaterialQualityLevel::Num, ERHIShadingPath::Type InShadingPath = ERHIShadingPath::Num);
+		ERHIFeatureLevel::Type InFeatureLevel = ERHIFeatureLevel::Num,
+		EMaterialQualityLevel::Type InQuality = EMaterialQualityLevel::Num,
+		ERHIShadingPath::Type InShadingPath = ERHIShadingPath::Num,
+		EShaderFrequency InShaderFrequency = SF_NumFrequencies);
 
 	/**
 	*	Recursively update the bRealtimePreview for each expression based on whether it is connected to something that is time-varying.
diff --git a/Engine/Source/Runtime/Engine/Classes/Materials/MaterialExpressionShaderStageSwitch.h b/Engine/Source/Runtime/Engine/Classes/Materials/MaterialExpressionShaderStageSwitch.h
new file mode 100644
--- /dev/null
+++ b/Engine/Source/Runtime/Engine/Classes/Materials/MaterialExpressionShaderStageSwitch.h
@@ -0,0 +1,36 @@
+// Copyright 1998-2019 Epic Games, Inc. All Rights Reserved.
+
+
+#pragma once
+
+#include "CoreMinimal.h"
+#include "UObject/ObjectMacros.h"
+#include "MaterialExpressionIO.h"
+#include "Materials/MaterialExpression.h"
+#include "MaterialExpressionShaderStageSwitch.generated.h"
+
+UCLASS(collapsecategories, hidecategories = Object, MinimalAPI)
+class UMaterialExpressionShaderStageSwitch : public UMaterialExpression
+{
+	GENERATED_UCLASS_BODY()
+
+	static inline bool ShouldUsePixelShaderInput(EShaderFrequency InShaderFrequency)
+	{
+		// Compute is considered pixel shader here, as the majority of material behavior that depends on pixel shader also allows compute shader
+		return (InShaderFrequency == SF_Pixel || InShaderFrequency == SF_Compute);
+	}
+
+	UPROPERTY(meta = (RequiredInput = "true", ToolTip = "Input will be used in pixel (or compute) shader stages"))
+	FExpressionInput PixelShader;
+
+	UPROPERTY(meta = (RequiredInput = "true", ToolTip = "Input will be in vertex (or tessellation) shader stages"))
+	FExpressionInput VertexShader;
+
+	//~ Begin UMaterialExpression Interface
+#if WITH_EDITOR
+	virtual int32 Compile(class FMaterialCompiler* Compiler, int32 OutputIndex) override;
+	virtual void GetCaption(TArray<FString>& OutCaptions) const override;
+	virtual void GetExpressionToolTip(TArray<FString>& OutToolTip) override;
+#endif
+	//~ End UMaterialExpression Interface
+};
diff --git a/Engine/Source/Runtime/Engine/Private/Materials/Material.cpp b/Engine/Source/Runtime/Engine/Private/Materials/Material.cpp
--- a/Engine/Source/Runtime/Engine/Private/Materials/Material.cpp
+++ b/Engine/Source/Runtime/Engine/Private/Materials/Material.cpp
@@ -25,6 +25,7 @@
 #include "Materials/MaterialExpressionQualitySwitch.h"
 #include "Materials/MaterialExpressionFeatureLevelSwitch.h"
 #include "Materials/MaterialExpressionShadingPathSwitch.h"
+#include "Materials/MaterialExpressionShaderStageSwitch.h"
 #include "Materials/MaterialExpressionMakeMaterialAttributes.h"
 #include "Materials/MaterialExpressionParameter.h"
 #include "Materials/MaterialExpressionScalarParameter.h"
@@ -5052,7 +5053,8 @@ bool UMaterial::GetExpressionsInPropertyChain(EMaterialProperty InProperty,
 	if (StartingExpression->Expression)
 	{
 		ProcessedInputs.AddUnique(StartingExpression);
-		RecursiveGetExpressionChain(StartingExpression->Expression, ProcessedInputs, OutExpressions, InStaticParameterSet, InFeatureLevel, InQuality, InShadingPath);
+		const EShaderFrequency ShaderFrequency = FMaterialAttributeDefinitionMap::GetShaderFrequency(InProperty);
+		RecursiveGetExpressionChain(StartingExpression->Expression, ProcessedInputs, OutExpressions, InStaticParameterSet, InFeatureLevel, InQuality, InShadingPath, ShaderFrequency);
 	}
 	return true;
 }
@@ -5192,7 +5194,7 @@ bool UMaterial::GetTexturesInPropertyChain(EMaterialProperty InProperty, TArray<
 
 bool UMaterial::RecursiveGetExpressionChain(UMaterialExpression* InExpression, TArray<FExpressionInput*>& InOutProcessedInputs, 
 	TArray<UMaterialExpression*>& OutExpressions, struct FStaticParameterSet* InStaticParameterSet,
-	ERHIFeatureLevel::Type InFeatureLevel, EMaterialQualityLevel::Type InQuality, ERHIShadingPath::Type InShadingPath)
+	ERHIFeatureLevel::Type InFeatureLevel, EMaterialQualityLevel::Type InQuality, ERHIShadingPath::Type InShadingPath, EShaderFrequency InShaderFrequency)
 {
 	OutExpressions.AddUnique(InExpression);
 	TArray<FExpressionInput*> Inputs;
@@ -5201,6 +5203,7 @@ bool UMaterial::RecursiveGetExpressionChain(UMaterialExpression* InExpression, T
 	UMaterialExpressionQualitySwitch* QualitySwitchExp;
 	UMaterialExpressionShadingPathSwitch* ShadingPathSwitchExp;
 	UMaterialExpressionMakeMaterialAttributes* MakeMaterialAttributesExp;
+	UMaterialExpressionShaderStageSwitch* ShaderStageSwitchExp;
 
 	if (InFeatureLevel != ERHIFeatureLevel::Num && (FeatureLevelSwitchExp = Cast<UMaterialExpressionFeatureLevelSwitch>(InExpression)) != nullptr)
 	{
@@ -5235,6 +5238,17 @@ bool UMaterial::RecursiveGetExpressionChain(UMaterialExpression* InExpression, T
 			Inputs.Add(&ShadingPathSwitchExp->Default);
 		}
 	}
+	else if (InShaderFrequency != SF_NumFrequencies && (ShaderStageSwitchExp = Cast<UMaterialExpressionShaderStageSwitch>(InExpression)) != nullptr)
+	{
+		if (UMaterialExpressionShaderStageSwitch::ShouldUsePixelShaderInput(InShaderFrequency))
+		{
+			Inputs.Add(&ShaderStageSwitchExp->PixelShader);
+		}
+		else
+		{
+			Inputs.Add(&ShaderStageSwitchExp->VertexShader);
+		}
+	}
 	else if (InFeatureLevel <= ERHIFeatureLevel::ES3_1 && (MakeMaterialAttributesExp = Cast<UMaterialExpressionMakeMaterialAttributes>(InExpression)) != nullptr)
 	{
 		// Follow only mobile-relevant inputs
@@ -5306,7 +5320,7 @@ bool UMaterial::RecursiveGetExpressionChain(UMaterialExpression* InExpression, T
 					if (bProcessInput == true)
 					{
 						InOutProcessedInputs.Add(InnerInput);
-						RecursiveGetExpressionChain(InnerInput->Expression, InOutProcessedInputs, OutExpressions, InStaticParameterSet, InFeatureLevel, InQuality);
+						RecursiveGetExpressionChain(InnerInput->Expression, InOutProcessedInputs, OutExpressions, InStaticParameterSet, InFeatureLevel, InQuality, InShadingPath, InShaderFrequency);
 					}
 				}
 			}
diff --git a/Engine/Source/Runtime/Engine/Private/Materials/MaterialExpressions.cpp b/Engine/Source/Runtime/Engine/Private/Materials/MaterialExpressions.cpp
--- a/Engine/Source/Runtime/Engine/Private/Materials/MaterialExpressions.cpp
+++ b/Engine/Source/Runtime/Engine/Private/Materials/MaterialExpressions.cpp
@@ -216,6 +216,7 @@
 #include "Materials/MaterialExpressionMaterialLayerOutput.h"
 #include "Materials/MaterialExpressionCurveAtlasRowParameter.h"
 #include "Materials/MaterialExpressionMapARPassthroughCameraUV.h"
+#include "Materials/MaterialExpressionShaderStageSwitch.h"
 #include "Materials/MaterialUniformExpressions.h"
 #include "EditorSupportDelegates.h"
 #include "MaterialCompiler.h"
@@ -13212,6 +13213,58 @@ void UMaterialExpressionShadowReplace::GetExpressionToolTip(TArray<FString>& Out
 }
 #endif // WITH_EDITOR
 
+//
+//	UMaterialExpressionShaderStageSwitch
+//
+UMaterialExpressionShaderStageSwitch::UMaterialExpressionShaderStageSwitch(const FObjectInitializer& ObjectInitializer)
+	: Super(ObjectInitializer)
+{
+	// Structure to hold one-time initialization
+	struct FConstructorStatics
+	{
+		FText NAME_Utility;
+		FConstructorStatics()
+			: NAME_Utility(LOCTEXT("Utility", "Utility"))
+		{
+		}
+	};
+	static FConstructorStatics ConstructorStatics;
+
+#if WITH_EDITORONLY_DATA
+	MenuCategories.Add(ConstructorStatics.NAME_Utility);
+#endif
+}
+
+#if WITH_EDITOR
+int32 UMaterialExpressionShaderStageSwitch::Compile(class FMaterialCompiler* Compiler, int32 OutputIndex)
+{
+	if (!PixelShader.GetTracedInput().Expression)
+	{
+		return Compiler->Errorf(TEXT("Missing input PixelShader"));
+	}
+	else if (!VertexShader.GetTracedInput().Expression)
+	{
+		return Compiler->Errorf(TEXT("Missing input VertexShader"));
+	}
+	else
+	{
+		const EShaderFrequency ShaderFrequency = Compiler->GetCurrentShaderFrequency();
+		return ShouldUsePixelShaderInput(ShaderFrequency) ? PixelShader.Compile(Compiler) : VertexShader.Compile(Compiler);
+	}
+}
+
+
+void UMaterialExpressionShaderStageSwitch::GetCaption(TArray<FString>& OutCaptions) const
+{
+	OutCaptions.Add(TEXT("Shader Stage Switch"));
+}
+
+void UMaterialExpressionShaderStageSwitch::GetExpressionToolTip(TArray<FString>& OutToolTip)
+{
+	ConvertToMultilineToolTip(TEXT("Allows material to define specialized behavior for certain shader stages."), 40, OutToolTip);
+}
+#endif // WITH_EDITOR
+
 
 //
 //	UMaterialExpressionMaterialProxy
