Corrent gamma and sRGB behavior for render targets

diff --git a/Engine/Source/Runtime/Core/Public/UObject/RenderingObjectVersion.h b/Engine/Source/Runtime/Core/Public/UObject/RenderingObjectVersion.h
--- a/Engine/Source/Runtime/Core/Public/UObject/RenderingObjectVersion.h
+++ b/Engine/Source/Runtime/Core/Public/UObject/RenderingObjectVersion.h
@@ -87,6 +87,9 @@ struct CORE_API FRenderingObjectVersion
 		// Lightmaps replace ULightMapVirtualTexture (non-UTexture derived class) with ULightMapVirtualTexture2D (derived from UTexture)
 		VirtualTexturedLightmapsV2,
 
+		// UTextureRenderTarget2D now explicitly allows users to create sRGB or non-sRGB type targets
+		ExplicitSRGBSetting,
+
 		// -----<new versions can be added above this line>-------------------------------------------------
 		VersionPlusOne,
 		LatestVersion = VersionPlusOne - 1
diff --git a/Engine/Source/Runtime/Engine/Classes/Engine/TextureRenderTarget2D.h b/Engine/Source/Runtime/Engine/Classes/Engine/TextureRenderTarget2D.h
--- a/Engine/Source/Runtime/Engine/Classes/Engine/TextureRenderTarget2D.h
+++ b/Engine/Source/Runtime/Engine/Classes/Engine/TextureRenderTarget2D.h
@@ -25,6 +25,8 @@ enum ETextureRenderTargetFormat
 	RTF_RG8,
 	/** RGBA channels, 8 bit per channel fixed point, range [0, 1]. */
 	RTF_RGBA8,
+	/** RGBA channels, 8 bit per channel fixed point, range [0, 1]. RGB is encoded with sRGB gamma curve. A is always stored as linear. */
+	RTF_RGBA8_SRGB,
 	/** R channel, 16 bit per channel floating point, range [-65504, 65504] */
 	RTF_R16f,
 	/** RG channels, 16 bit per channel floating point, range [-65504, 65504] */
@@ -45,18 +47,19 @@ inline EPixelFormat GetPixelFormatFromRenderTargetFormat(ETextureRenderTargetFor
 {
 	switch (RTFormat)
 	{
-	case RTF_R8: return PF_G8; break;
-	case RTF_RG8: return PF_R8G8; break;
-	case RTF_RGBA8: return PF_B8G8R8A8; break;
-
-	case RTF_R16f: return PF_R16F; break;
-	case RTF_RG16f: return PF_G16R16F; break;
-	case RTF_RGBA16f: return PF_FloatRGBA; break;
-
-	case RTF_R32f: return PF_R32_FLOAT; break;
-	case RTF_RG32f: return PF_G32R32F; break;
-	case RTF_RGBA32f: return PF_A32B32G32R32F; break;
-	case RTF_RGB10A2: return PF_A2B10G10R10; break;
+	case RTF_R8: return PF_G8;
+	case RTF_RG8: return PF_R8G8;
+	case RTF_RGBA8: return PF_B8G8R8A8;
+	case RTF_RGBA8_SRGB: return PF_B8G8R8A8;
+
+	case RTF_R16f: return PF_R16F;
+	case RTF_RG16f: return PF_G16R16F;
+	case RTF_RGBA16f: return PF_FloatRGBA;
+
+	case RTF_R32f: return PF_R32_FLOAT;
+	case RTF_RG32f: return PF_G32R32F;
+	case RTF_RGBA32f: return PF_A32B32G32R32F;
+	case RTF_RGB10A2: return PF_A2B10G10R10;
 	}
 
 	ensureMsgf(false, TEXT("Unhandled ETextureRenderTargetFormat entry %u"), (uint32)RTFormat);
@@ -114,7 +117,7 @@ class UTextureRenderTarget2D : public UTextureRenderTarget
 	 * Use the smallest format that has enough precision and range for what you are doing.
 	 */
 	UPROPERTY(EditAnywhere, BlueprintReadOnly, Category=TextureRenderTarget2D, AssetRegistrySearchable)
-	TEnumAsByte<enum ETextureRenderTargetFormat> RenderTargetFormat;
+	TEnumAsByte<ETextureRenderTargetFormat> RenderTargetFormat;
 
 	/** Whether to support Mip maps for this render target texture */
 	UPROPERTY(EditAnywhere, BlueprintReadOnly, Category = TextureRenderTarget2D, AssetRegistrySearchable)
@@ -205,6 +208,18 @@ class UTextureRenderTarget2D : public UTextureRenderTarget
 		}
 	}
 
+	bool IsSRGB() const
+	{
+		if (OverrideFormat == PF_Unknown)
+		{
+			return RenderTargetFormat == RTF_RGBA8_SRGB;
+		}
+		else
+		{
+			return !bForceLinearGamma;
+		}
+	}
+
 private:
 	int32	NumMips;
 };
diff --git a/Engine/Source/Runtime/Engine/Private/TextureRenderTarget2D.cpp b/Engine/Source/Runtime/Engine/Private/TextureRenderTarget2D.cpp
--- a/Engine/Source/Runtime/Engine/Private/TextureRenderTarget2D.cpp
+++ b/Engine/Source/Runtime/Engine/Private/TextureRenderTarget2D.cpp
@@ -191,13 +191,10 @@ void UTextureRenderTarget2D::PostEditChangeProperty(FPropertyChangedEvent& Prope
 	SizeX = FMath::Clamp<int32>(SizeX - (SizeX % GPixelFormats[Format].BlockSizeX),1,MaxSize);
 	SizeY = FMath::Clamp<int32>(SizeY - (SizeY % GPixelFormats[Format].BlockSizeY),1,MaxSize);
 
-	// Always set SRGB back to 'on'; it will be turned off again in the call to Super::PostEditChangeProperty below if necessary
-	if (PropertyChangedEvent.Property)
-	{
-		SRGB = true;
-	}
-
 	Super::PostEditChangeProperty(PropertyChangedEvent);
+
+    // SRGB may have been changed by Super, reset it since we prefer to honor explicit user choice
+	SRGB = IsSRGB();
 }
 #endif // WITH_EDITOR
 
@@ -211,6 +208,30 @@ void UTextureRenderTarget2D::Serialize(FArchive& Ar)
 	{
 		RenderTargetFormat = bHDR_DEPRECATED ? RTF_RGBA16f : RTF_RGBA8;
 	}
+
+	if (Ar.CustomVer(FRenderingObjectVersion::GUID) < FRenderingObjectVersion::ExplicitSRGBSetting)
+	{
+		float DisplayGamme = 2.2f;
+
+		if (TargetGamma > KINDA_SMALL_NUMBER * 10.0f)
+		{
+			DisplayGamme = TargetGamma;
+		}
+
+		EPixelFormat Format = GetFormat();
+		if (Format == PF_FloatRGB || Format == PF_FloatRGBA || bForceLinearGamma)
+		{
+			DisplayGamme = 1.0f;
+		}
+
+		// This is odd behavior to apply the sRGB gamma correction when target gamma is not 1.0f, but this
+		// is to maintain old behavior and users won't have to change content.
+		if (RenderTargetFormat == RTF_RGBA8 && FMath::Abs(DisplayGamme - 1.0f) > KINDA_SMALL_NUMBER)
+		{
+			RenderTargetFormat = RTF_RGBA8_SRGB;
+			SRGB = true;
+		}
+	}
 }
 
 void UTextureRenderTarget2D::PostLoad()
@@ -428,15 +449,8 @@ void FTextureRenderTarget2DResource::InitDynamicRHI()
 {
 	if( TargetSizeX > 0 && TargetSizeY > 0 )
 	{
-		bool bUseSRGB=true;
-		// if render target gamma used was 1.0 then disable SRGB for the static texture
-		if( FMath::Abs(GetDisplayGamma() - 1.0f) < KINDA_SMALL_NUMBER )
-		{
-			bUseSRGB = false;
-		}
-
 		// Create the RHI texture. Only one mip is used and the texture is targetable for resolve.
-		uint32 TexCreateFlags = bUseSRGB ? TexCreate_SRGB : 0;
+		uint32 TexCreateFlags = Owner->IsSRGB() ? TexCreate_SRGB : 0;
 		TexCreateFlags |= Owner->bGPUSharedFlag ? TexCreate_Shared : 0;
 		FRHIResourceCreateInfo CreateInfo = FRHIResourceCreateInfo(FClearValueBinding(ClearColor));
 
diff --git a/Engine/Source/Runtime/UMG/Private/Slate/SRetainerWidget.cpp b/Engine/Source/Runtime/UMG/Private/Slate/SRetainerWidget.cpp
--- a/Engine/Source/Runtime/UMG/Private/Slate/SRetainerWidget.cpp
+++ b/Engine/Source/Runtime/UMG/Private/Slate/SRetainerWidget.cpp
@@ -157,8 +157,7 @@ void SRetainerWidget::Construct(const FArguments& InArgs)
 
 	UTextureRenderTarget2D* RenderTarget = NewObject<UTextureRenderTarget2D>();
 	RenderTarget->ClearColor = FLinearColor::Transparent;
-	RenderTarget->OverrideFormat = PF_B8G8R8A8;
-	RenderTarget->bForceLinearGamma = false;
+	RenderTarget->RenderTargetFormat = RTF_RGBA8_SRGB;
 
 	RenderingResources->RenderTarget = RenderTarget;
 	SurfaceBrush.SetResourceObject(RenderTarget);
