Fix minor data versioning.

diff --git a/Engine/Source/Runtime/Engine/Private/TextureRenderTarget2D.cpp b/Engine/Source/Runtime/Engine/Private/TextureRenderTarget2D.cpp
--- a/Engine/Source/Runtime/Engine/Private/TextureRenderTarget2D.cpp
+++ b/Engine/Source/Runtime/Engine/Private/TextureRenderTarget2D.cpp
@@ -212,14 +212,13 @@ void UTextureRenderTarget2D::Serialize(FArchive& Ar)
 	if (Ar.CustomVer(FRenderingObjectVersion::GUID) < FRenderingObjectVersion::ExplicitSRGBSetting)
 	{
 		float DisplayGamme = 2.2f;
+		EPixelFormat Format = GetFormat();
 
 		if (TargetGamma > KINDA_SMALL_NUMBER * 10.0f)
 		{
 			DisplayGamme = TargetGamma;
 		}
-
-		EPixelFormat Format = GetFormat();
-		if (Format == PF_FloatRGB || Format == PF_FloatRGBA || bForceLinearGamma)
+		else if (Format == PF_FloatRGB || Format == PF_FloatRGBA || bForceLinearGamma)
 		{
 			DisplayGamme = 1.0f;
 		}
