DR - Fix shutdown order on dynamic resolution queries

diff --git a/Engine/Source/Runtime/Engine/Private/DynamicResolution.cpp b/Engine/Source/Runtime/Engine/Private/DynamicResolution.cpp
--- a/Engine/Source/Runtime/Engine/Private/DynamicResolution.cpp
+++ b/Engine/Source/Runtime/Engine/Private/DynamicResolution.cpp
@@ -622,6 +622,14 @@ public:
 
 		// Wait for all queries to land.
 		HandLandedQueriesToHeuristic(/* bWait = */ true);
+
+		for (InFlightFrameQueries& Query : InFlightFrames)
+		{
+			Query.BeginDynamicResolutionQuery.ReleaseQuery();
+			Query.BeginFrameQuery.ReleaseQuery();
+			Query.EndDynamicResolutionQuery.ReleaseQuery();
+			Query.EndFrameQuery.ReleaseQuery();
+		}
 	}
 
 
