Log error message to console when cooking/starting PIE session for material instances that have invalid texture parameters assigned

diff --git a/Engine/Source/Runtime/Engine/Classes/Materials/MaterialInstance.h b/Engine/Source/Runtime/Engine/Classes/Materials/MaterialInstance.h
--- a/Engine/Source/Runtime/Engine/Classes/Materials/MaterialInstance.h
+++ b/Engine/Source/Runtime/Engine/Classes/Materials/MaterialInstance.h
@@ -598,6 +598,8 @@ public:
 	ENGINE_API virtual void LogMaterialsAndTextures(FOutputDevice& Ar, int32 Indent) const override;
 #endif
 
+	ENGINE_API void ValidateTextureOverrides(ERHIFeatureLevel::Type InFeatureLevel) const;
+
 	/**
 	 *	Returns all the Guids related to this material. For material instances, this includes the parent hierarchy.
 	 *  Used for versioning as parent changes don't update the child instance Guids.
diff --git a/Engine/Source/Runtime/Engine/Private/Materials/MaterialInstance.cpp b/Engine/Source/Runtime/Engine/Private/Materials/MaterialInstance.cpp
--- a/Engine/Source/Runtime/Engine/Private/Materials/MaterialInstance.cpp
+++ b/Engine/Source/Runtime/Engine/Private/Materials/MaterialInstance.cpp
@@ -1160,7 +1160,79 @@ void UMaterialInstance::LogMaterialsAndTextures(FOutputDevice& Ar, int32 Indent)
 		}
 	}
 }
-#endif
+#endif // !(UE_BUILD_SHIPPING || UE_BUILD_TEST)
+
+void UMaterialInstance::ValidateTextureOverrides(ERHIFeatureLevel::Type InFeatureLevel) const
+{
+	const UMaterial* Material = GetMaterial();
+	const FMaterialResource* CurrentResource = Material->GetMaterialResource(InFeatureLevel);
+
+	const TArray<TRefCountPtr<FMaterialUniformExpressionTexture>>* TextureExpressions[] =
+	{
+		&CurrentResource->GetUniform2DTextureExpressions(),
+		&CurrentResource->GetUniformCubeTextureExpressions(),
+		&CurrentResource->GetUniformVolumeTextureExpressions(),
+		&CurrentResource->GetUniformVirtualTextureExpressions(),
+	};
+
+	FString MaterialName;
+	GetName(MaterialName);
+
+	for (int32 TypeIndex = 0; TypeIndex < ARRAY_COUNT(TextureExpressions); ++TypeIndex)
+	{
+		for (FMaterialUniformExpressionTexture* TextureExpression : *TextureExpressions[TypeIndex])
+		{
+			FMaterialUniformExpressionTextureParameter* ParameterExpression = TextureExpression->GetTextureParameterUniformExpression();
+			if (ParameterExpression)
+			{
+				UTexture* Texture = nullptr;
+				const FMaterialParameterInfo& ParameterInfo = ParameterExpression->GetParameterInfo();
+				if (GetTextureParameterValue(ParameterInfo, Texture, true))
+				{
+					check(Texture);
+					const EMaterialValueType TextureType = Texture->GetMaterialType();
+					switch (TypeIndex)
+					{
+					case 0:
+						if (!(TextureType & (MCT_Texture2D | MCT_TextureExternal | MCT_TextureVirtual)))
+						{
+							UE_LOG(LogMaterial, Error, TEXT("MaterialInstance \"%s\" parameter '%s' assigned texture \"%s\" has invalid type, required 2D texture"), *MaterialName, *ParameterInfo.Name.ToString(), *Texture->GetName());
+						}
+						else if (TextureType & MCT_TextureVirtual)
+						{
+							UE_LOG(LogMaterial, Error, TEXT("MaterialInstance \"%s\" parameter '%s' assigned texture \"%s\" requires non-virtual texture"), *MaterialName, *ParameterInfo.Name.ToString(), *Texture->GetName());
+						}
+						break;
+					case 1:
+						if (!(TextureType & MCT_TextureCube))
+						{
+							UE_LOG(LogMaterial, Error, TEXT("MaterialInstance \"%s\" parameter '%s' assigned texture \"%s\" has invalid type, required Cube texture"), *MaterialName, *ParameterInfo.Name.ToString(), *Texture->GetName());
+						}
+						break;
+					case 2:
+						if (!(TextureType & MCT_VolumeTexture))
+						{
+							UE_LOG(LogMaterial, Error, TEXT("MaterialInstance \"%s\" parameter '%s' assigned texture \"%s\" has invalid type, required Volume texture"), *MaterialName, *ParameterInfo.Name.ToString(), *Texture->GetName());
+						}
+						break;
+					case 3:
+						if (!(TextureType & (MCT_Texture2D | MCT_TextureExternal | MCT_TextureVirtual)))
+						{
+							UE_LOG(LogMaterial, Error, TEXT("MaterialInstance \"%s\" parameter '%s' assigned texture \"%s\" has invalid type, required 2D texture"), *MaterialName, *ParameterInfo.Name.ToString(), *Texture->GetName());
+						}
+						else if (!(TextureType & MCT_TextureVirtual))
+						{
+							UE_LOG(LogMaterial, Error, TEXT("MaterialInstance \"%s\" parameter '%s' assigned texture \"%s\" requires virtual texture"), *MaterialName, *ParameterInfo.Name.ToString(), *Texture->GetName());
+						}
+						break;
+					default:
+						checkNoEntry();
+					}
+				}
+			}
+		}
+	}
+}
 
 void UMaterialInstance::GetUsedTexturesAndIndices(TArray<UTexture*>& OutTextures, TArray< TArray<int32> >& OutIndices, EMaterialQualityLevel::Type QualityLevel, ERHIFeatureLevel::Type FeatureLevel) const
 {
@@ -3140,6 +3212,11 @@ void UMaterialInstance::Serialize(FArchive& Ar)
 		SaveShaderStableKeys(Ar.CookingTarget());
 	}
 #endif
+
+	if (Ar.IsSaving() && Ar.IsCooking())
+	{
+		ValidateTextureOverrides(GMaxRHIFeatureLevel);
+	}
 }
 
 void UMaterialInstance::PostLoad()
diff --git a/Engine/Source/Runtime/Engine/Private/Materials/MaterialUniformExpressions.h b/Engine/Source/Runtime/Engine/Private/Materials/MaterialUniformExpressions.h
--- a/Engine/Source/Runtime/Engine/Private/Materials/MaterialUniformExpressions.h
+++ b/Engine/Source/Runtime/Engine/Private/Materials/MaterialUniformExpressions.h
@@ -359,9 +359,9 @@ public:
 		else
 		{
 			OutValue = NULL;
-		const bool bOverrideValuesOnly = !AreExperimentalMaterialLayersEnabled();
-			if(!MaterialInterface->GetTextureParameterValue(ParameterInfo,OutValue,bOverrideValuesOnly))
-		{
+			const bool bOverrideValuesOnly = !AreExperimentalMaterialLayersEnabled();
+			if (!MaterialInterface->GetTextureParameterValue(ParameterInfo, OutValue, bOverrideValuesOnly))
+			{
 				OutValue = GetIndexedTexture<UTexture>(Material, TextureIndex);
 			}
 		}
@@ -377,6 +377,11 @@ public:
 		return ParameterInfo.Name;
 	}
 
+	const FMaterialParameterInfo& GetParameterInfo() const
+	{
+		return ParameterInfo;
+	}
+
 	virtual bool IsIdentical(const FMaterialUniformExpression* OtherExpression) const
 	{
 		if (GetType() != OtherExpression->GetType())
diff --git a/Engine/Source/Runtime/Engine/Private/World.cpp b/Engine/Source/Runtime/Engine/Private/World.cpp
--- a/Engine/Source/Runtime/Engine/Private/World.cpp
+++ b/Engine/Source/Runtime/Engine/Private/World.cpp
@@ -799,6 +799,28 @@ void UWorld::PostDuplicate(bool bDuplicateForPIE)
 			Obj->PostEditChange();
 		}
 	}
+
+	if (bDuplicateForPIE)
+	{
+		// When PIE begins, check/log any problems with textures assigned to material instances
+		for (TObjectIterator<UPrimitiveComponent> It; It; ++It)
+		{
+			UPrimitiveComponent* Component = *It;
+			AActor* Owner = Component->GetOwner();
+			if (Owner != nullptr && !Owner->HasAnyFlags(RF_ClassDefaultObject) && Owner->IsInLevel(PersistentLevel))
+			{
+				TArray<UMaterialInterface*> Materials;
+				Component->GetUsedMaterials(Materials);
+				for (UMaterialInterface* Material : Materials)
+				{
+					if (UMaterialInstance* MaterialInstance = Cast<UMaterialInstance>(Material))
+					{
+						MaterialInstance->ValidateTextureOverrides(FeatureLevel);
+					}
+				}
+			}
+		}
+	}
 #endif // WITH_EDITOR
 }
 
