Fixed memory leak in Instanced Static Mesh reported from UDN.

diff --git a/Engine/Source/Runtime/Engine/Private/InstancedStaticMesh.cpp b/Engine/Source/Runtime/Engine/Private/InstancedStaticMesh.cpp
--- a/Engine/Source/Runtime/Engine/Private/InstancedStaticMesh.cpp
+++ b/Engine/Source/Runtime/Engine/Private/InstancedStaticMesh.cpp
@@ -709,10 +709,8 @@ void FPerInstanceRenderData::UpdateFromPreallocatedData(FStaticMeshInstanceData&
 
 	InOther.SetAllowCPUAccess(InstanceBuffer.RequireCPUAccess);
 
-	FStaticMeshInstanceData* NewInstanceData = new FStaticMeshInstanceData();
-	FMemory::Memswap(&InOther, NewInstanceData, sizeof(FStaticMeshInstanceData));
-
-	InstanceBuffer_GameThread = MakeShared<FStaticMeshInstanceData, ESPMode::ThreadSafe>(*NewInstanceData);
+	InstanceBuffer_GameThread = MakeShared<FStaticMeshInstanceData, ESPMode::ThreadSafe>();
+	FMemory::Memswap(&InOther, InstanceBuffer_GameThread.Get(), sizeof(FStaticMeshInstanceData));
 
 	typedef TSharedPtr<FStaticMeshInstanceData, ESPMode::ThreadSafe> FStaticMeshInstanceDataPtr;
 
