Add some keywords to new shader stage switch expression

diff --git a/Engine/Source/Runtime/Engine/Classes/Materials/MaterialExpressionShaderStageSwitch.h b/Engine/Source/Runtime/Engine/Classes/Materials/MaterialExpressionShaderStageSwitch.h
--- a/Engine/Source/Runtime/Engine/Classes/Materials/MaterialExpressionShaderStageSwitch.h
+++ b/Engine/Source/Runtime/Engine/Classes/Materials/MaterialExpressionShaderStageSwitch.h
@@ -31,6 +31,7 @@ class UMaterialExpressionShaderStageSwitch : public UMaterialExpression
 	virtual int32 Compile(class FMaterialCompiler* Compiler, int32 OutputIndex) override;
 	virtual void GetCaption(TArray<FString>& OutCaptions) const override;
 	virtual void GetExpressionToolTip(TArray<FString>& OutToolTip) override;
+	virtual FText GetKeywords() const override { return FText::FromString(TEXT("vertex pixel shader")); }
 #endif
 	//~ End UMaterialExpression Interface
 };
