Add SubUV properties material expression, allowsdirect  access to values used to implement particle subUV animation

diff --git a/Engine/Source/Runtime/Engine/Classes/Materials/MaterialExpressionParticleSubUVProperties.h b/Engine/Source/Runtime/Engine/Classes/Materials/MaterialExpressionParticleSubUVProperties.h
new file mode 100644
--- /dev/null
+++ b/Engine/Source/Runtime/Engine/Classes/Materials/MaterialExpressionParticleSubUVProperties.h
@@ -0,0 +1,26 @@
+// Copyright 1998-2019 Epic Games, Inc. All Rights Reserved.
+
+
+#pragma once
+
+#include "CoreMinimal.h"
+#include "UObject/ObjectMacros.h"
+#include "Materials/MaterialExpression.h"
+#include "MaterialExpressionParticleSubUVProperties.generated.h"
+
+UCLASS(MinimalAPI)
+class UMaterialExpressionParticleSubUVProperties : public UMaterialExpression
+{
+	GENERATED_UCLASS_BODY()
+
+	//~ Begin UMaterialExpression Interface
+#if WITH_EDITOR
+	virtual int32 Compile(class FMaterialCompiler* Compiler, int32 OutputIndex) override;
+	virtual void GetCaption(TArray<FString>& OutCaptions) const override;
+	virtual void GetExpressionToolTip(TArray<FString>& OutToolTip) override;
+#endif
+	//~ End UMaterialExpression Interface
+};
+
+
+
diff --git a/Engine/Source/Runtime/Engine/Private/Materials/HLSLMaterialTranslator.h b/Engine/Source/Runtime/Engine/Private/Materials/HLSLMaterialTranslator.h
--- a/Engine/Source/Runtime/Engine/Private/Materials/HLSLMaterialTranslator.h
+++ b/Engine/Source/Runtime/Engine/Private/Materials/HLSLMaterialTranslator.h
@@ -3560,6 +3560,29 @@ protected:
 		return ParticleSubUV;
 	}
 
+	virtual int32 ParticleSubUVProperty(int32 PropertyIndex) override
+	{
+		int32 Result = INDEX_NONE;
+		switch (PropertyIndex)
+		{
+		case 0:
+			Result = AddCodeChunk(MCT_Float2, TEXT("Parameters.Particle.SubUVCoords[0].xy"));
+			break;
+		case 1:
+			Result = AddCodeChunk(MCT_Float2, TEXT("Parameters.Particle.SubUVCoords[1].xy"));
+			break;
+		case 2:
+			Result = AddCodeChunk(MCT_Float, TEXT("Parameters.Particle.SubUVLerp"));
+			break;
+		default:
+			checkNoEntry();
+			break;
+		}
+
+		bUsesParticleSubUVs = true;
+		return Result;
+	}
+
 	virtual int32 ParticleColor() override
 	{
 		if (ShaderFrequency != SF_Vertex && ShaderFrequency != SF_Pixel && ShaderFrequency != SF_Compute)
diff --git a/Engine/Source/Runtime/Engine/Private/Materials/MaterialExpressions.cpp b/Engine/Source/Runtime/Engine/Private/Materials/MaterialExpressions.cpp
--- a/Engine/Source/Runtime/Engine/Private/Materials/MaterialExpressions.cpp
+++ b/Engine/Source/Runtime/Engine/Private/Materials/MaterialExpressions.cpp
@@ -186,6 +186,7 @@
 #include "Materials/MaterialExpressionTextureProperty.h"
 #include "Materials/MaterialExpressionTextureSample.h"
 #include "Materials/MaterialExpressionParticleSubUV.h"
+#include "Materials/MaterialExpressionParticleSubUVProperties.h"
 #include "Materials/MaterialExpressionTextureSampleParameter.h"
 #include "Materials/MaterialExpressionTextureObjectParameter.h"
 #include "Materials/MaterialExpressionTextureSampleParameter2D.h"
@@ -8268,6 +8269,56 @@ void UMaterialExpressionParticleSubUV::GetCaption(TArray<FString>& OutCaptions)
 }
 #endif // WITH_EDITOR
 
+//
+//	MaterialExpressionParticleSubUVProperties
+//
+UMaterialExpressionParticleSubUVProperties::UMaterialExpressionParticleSubUVProperties(const FObjectInitializer& ObjectInitializer)
+	: Super(ObjectInitializer)
+{
+	// Structure to hold one-time initialization
+	struct FConstructorStatics
+	{
+		FText NAME_Particles;
+		FText NAME_Coordinates;
+		FConstructorStatics()
+			: NAME_Particles(LOCTEXT("Particles", "Particles"))
+			, NAME_Coordinates(LOCTEXT("Coordinates", "Coordinates"))
+		{
+		}
+	};
+	static FConstructorStatics ConstructorStatics;
+
+#if WITH_EDITORONLY_DATA
+	MenuCategories.Add(ConstructorStatics.NAME_Particles);
+	MenuCategories.Add(ConstructorStatics.NAME_Coordinates);
+
+	bShaderInputData = true;
+	bShowOutputNameOnPin = true;
+
+	Outputs.Reset();
+	Outputs.Add(FExpressionOutput(TEXT("TextureCoordinate0"), 1, 1, 1, 0, 0));
+	Outputs.Add(FExpressionOutput(TEXT("TextureCoordinate1"), 1, 1, 1, 0, 0));
+	Outputs.Add(FExpressionOutput(TEXT("Blend")));
+#endif // WITH_EDITORONLY_DATA
+}
+
+#if WITH_EDITOR
+int32 UMaterialExpressionParticleSubUVProperties::Compile(class FMaterialCompiler* Compiler, int32 OutputIndex)
+{
+	return Compiler->ParticleSubUVProperty(OutputIndex);
+}
+
+void UMaterialExpressionParticleSubUVProperties::GetExpressionToolTip(TArray<FString>& OutToolTip)
+{
+	ConvertToMultilineToolTip(TEXT("Provides direct access to properties used to implement particle UV frame animation."), 40, OutToolTip);
+}
+
+void UMaterialExpressionParticleSubUVProperties::GetCaption(TArray<FString>& OutCaptions) const
+{
+	OutCaptions.Add(TEXT("Particle SubUV Properties"));
+}
+#endif // WITH_EDITOR
+
 //
 //	MaterialExpressionParticleMacroUV
 //
diff --git a/Engine/Source/Runtime/Engine/Public/MaterialCompiler.h b/Engine/Source/Runtime/Engine/Public/MaterialCompiler.h
--- a/Engine/Source/Runtime/Engine/Public/MaterialCompiler.h
+++ b/Engine/Source/Runtime/Engine/Public/MaterialCompiler.h
@@ -176,6 +176,7 @@ public:
 	virtual int32 ActorWorldPosition() = 0;
 	virtual int32 ParticleMacroUV() = 0;
 	virtual int32 ParticleSubUV(int32 TextureIndex, EMaterialSamplerType SamplerType, bool bBlend) = 0;
+	virtual int32 ParticleSubUVProperty(int32 PropertyIndex) = 0;
 	virtual int32 ParticleColor() = 0;
 	virtual int32 ParticlePosition() = 0;
 	virtual int32 ParticleRadius() = 0;
@@ -445,6 +446,7 @@ public:
 	virtual int32 ActorWorldPosition() override { return Compiler->ActorWorldPosition(); }
 	virtual int32 ParticleMacroUV() override { return Compiler->ParticleMacroUV(); }
 	virtual int32 ParticleSubUV(int32 TextureIndex, EMaterialSamplerType SamplerType, bool bBlend) override { return Compiler->ParticleSubUV(TextureIndex, SamplerType, bBlend); }
+	virtual int32 ParticleSubUVProperty(int32 PropertyIndex) override { return Compiler->ParticleSubUVProperty(PropertyIndex); }
 	virtual int32 ParticleColor() override { return Compiler->ParticleColor(); }
 	virtual int32 ParticlePosition() override { return Compiler->ParticlePosition(); }
 	virtual int32 ParticleRadius() override { return Compiler->ParticleRadius(); }
