No more crash on changing cubemap resolution for skylight in the editor.

diff --git a/Engine/Source/Runtime/Engine/Private/Components/SkyLightComponent.cpp b/Engine/Source/Runtime/Engine/Private/Components/SkyLightComponent.cpp
--- a/Engine/Source/Runtime/Engine/Private/Components/SkyLightComponent.cpp
+++ b/Engine/Source/Runtime/Engine/Private/Components/SkyLightComponent.cpp
@@ -74,6 +74,8 @@ void FSkyTextureCubeResource::InitRHI()
 	if (GetFeatureLevel() >= ERHIFeatureLevel::SM4 || GSupportsRenderTargetFormat_PF_FloatRGBA)
 	{
 		FRHIResourceCreateInfo CreateInfo;
+		
+		checkf(FMath::IsPowerOfTwo(Size), TEXT("Size of SkyTextureCube must be a power of two; size is %d"), Size);
 		TextureCubeRHI = RHICreateTextureCube(Size, Format, NumMips, 0, CreateInfo);
 		TextureRHI = TextureCubeRHI;
 
@@ -473,21 +475,18 @@ void USkyLightComponent::DestroyRenderState_Concurrent()
 void USkyLightComponent::PreEditChange(UProperty* PropertyAboutToChange)
 {
 	Super::PreEditChange(PropertyAboutToChange);
-
 	PreEditCubemapResolution = CubemapResolution;
 }
 
 void USkyLightComponent::PostEditChangeProperty(FPropertyChangedEvent& PropertyChangedEvent)
 {
-	Super::PostEditChangeProperty(PropertyChangedEvent);
-
-	if (PropertyChangedEvent.Property && 
-		PropertyChangedEvent.Property->GetFName() == GET_MEMBER_NAME_CHECKED(USkyLightComponent, CubemapResolution) &&
-		PropertyChangedEvent.ChangeType != EPropertyChangeType::Interactive)
+	if (PropertyChangedEvent.Property && PropertyChangedEvent.Property->GetFName() == GET_MEMBER_NAME_CHECKED(USkyLightComponent, CubemapResolution))
 	{
+		// Simply rounds the cube map size to nearest power of two. Occasionally checks for out of video mem.
 		SanitizeCubemapSize();
 	}
 
+	Super::PostEditChangeProperty(PropertyChangedEvent);
 	SetCaptureIsDirty();
 }
 
