Works arround cook on fly server crash in UMaterialInstance::ValidateTextureOverrides()

Repro steps:
* Run QAGame TM-GI-RayTracing -filehostip=<Your dev IP> on your console
* Run on PC QAGame -run=Cook -cookonthefly -TargetPlatform=<YourConsole>

diff --git a/Engine/Source/Runtime/Engine/Private/Materials/MaterialInstance.cpp b/Engine/Source/Runtime/Engine/Private/Materials/MaterialInstance.cpp
--- a/Engine/Source/Runtime/Engine/Private/Materials/MaterialInstance.cpp
+++ b/Engine/Source/Runtime/Engine/Private/Materials/MaterialInstance.cpp
@@ -1167,6 +1167,11 @@ void UMaterialInstance::ValidateTextureOverrides(ERHIFeatureLevel::Type InFeatur
 	const UMaterial* Material = GetMaterial();
 	const FMaterialResource* CurrentResource = Material->GetMaterialResource(InFeatureLevel);
 
+	if (!CurrentResource)
+	{
+		return;
+	}
+
 	const TArray<TRefCountPtr<FMaterialUniformExpressionTexture>>* TextureExpressions[] =
 	{
 		&CurrentResource->GetUniform2DTextureExpressions(),
