Update render settings UI with properly rounded values for VT tile/border size

diff --git a/Engine/Source/Runtime/Engine/Private/RendererSettings.cpp b/Engine/Source/Runtime/Engine/Private/RendererSettings.cpp
--- a/Engine/Source/Runtime/Engine/Private/RendererSettings.cpp
+++ b/Engine/Source/Runtime/Engine/Private/RendererSettings.cpp
@@ -149,6 +149,16 @@ void URendererSettings::PostEditChangeProperty(FPropertyChangedEvent& PropertyCh
 			}
 		}
 
+		if (PropertyChangedEvent.Property->GetFName() == GET_MEMBER_NAME_CHECKED(URendererSettings, VirtualTextureTileSize))
+		{
+			VirtualTextureTileSize = FMath::RoundUpToPowerOfTwo(VirtualTextureTileSize);
+		}
+
+		if (PropertyChangedEvent.Property->GetFName() == GET_MEMBER_NAME_CHECKED(URendererSettings, VirtualTextureTileBorderSize))
+		{
+			VirtualTextureTileBorderSize = FMath::RoundUpToPowerOfTwo(VirtualTextureTileBorderSize);
+		}
+
 		ExportValuesToConsoleVariables(PropertyChangedEvent.Property);
 
 		if (PropertyChangedEvent.Property->GetFName() == GET_MEMBER_NAME_CHECKED(URendererSettings, ReflectionCaptureResolution) && 
