Print warning to console when user attempts to enable VT streaming for non-power-2 texture

diff --git a/Engine/Source/Runtime/Engine/Private/Texture2D.cpp b/Engine/Source/Runtime/Engine/Private/Texture2D.cpp
--- a/Engine/Source/Runtime/Engine/Private/Texture2D.cpp
+++ b/Engine/Source/Runtime/Engine/Private/Texture2D.cpp
@@ -552,7 +552,11 @@ void UTexture2D::PostEditChangeProperty(FPropertyChangedEvent& PropertyChangedEv
 		// Force NPT textures to have no mipmaps.
 		MipGenSettings = TMGS_NoMipmaps;
 		NeverStream = true;
-		VirtualTextureStreaming = false;
+		if (VirtualTextureStreaming)
+		{
+			UE_LOG(LogTexture, Warning, TEXT("VirtualTextureStreaming not supported for \"%s\", texture size is not a power-of-2"), *GetName());
+			VirtualTextureStreaming = false;
+		}
 	}
 
 	// Make sure settings are correct for LUT textures.
