From 6910510:
Bug fix where non-UAV textures crash using GenerateMips functionality. Ensures render targets for mip generation have UAV access.

diff --git a/Engine/Source/Runtime/Engine/Private/TextureRenderTarget2D.cpp b/Engine/Source/Runtime/Engine/Private/TextureRenderTarget2D.cpp
--- a/Engine/Source/Runtime/Engine/Private/TextureRenderTarget2D.cpp
+++ b/Engine/Source/Runtime/Engine/Private/TextureRenderTarget2D.cpp
@@ -455,7 +455,7 @@ void FTextureRenderTarget2DResource::InitDynamicRHI()
 
 		if (Owner->bAutoGenerateMips)
 		{
-			TexCreateFlags |= TexCreate_GenerateMipCapable;
+			TexCreateFlags |= (TexCreate_GenerateMipCapable | TexCreate_UAV);
 		}
 
 		if (Owner->bCanCreateUAV)
