When loading FMaterialShaderMap, only allocate mesh shader maps for used vertex factory types, rather than all VF types
- Saves some memory
- A bit faster, as we don't need to iterate/allocate all VF types

diff --git a/Engine/Source/Runtime/Engine/Private/Materials/MaterialShader.cpp b/Engine/Source/Runtime/Engine/Private/Materials/MaterialShader.cpp
--- a/Engine/Source/Runtime/Engine/Private/Materials/MaterialShader.cpp
+++ b/Engine/Source/Runtime/Engine/Private/Materials/MaterialShader.cpp
@@ -2486,20 +2486,6 @@ void FMaterialShaderMap::Serialize(FArchive& Ar, bool bInlineShaderResources, bo
 	{
 		MeshShaderMaps.Empty();
 
-		for(TLinkedList<FVertexFactoryType*>::TIterator VertexFactoryTypeIt(FVertexFactoryType::GetTypeList());VertexFactoryTypeIt;VertexFactoryTypeIt.Next())
-		{
-			FVertexFactoryType* VertexFactoryType = *VertexFactoryTypeIt;
-			check(VertexFactoryType);
-
-			if (VertexFactoryType->IsUsedWithMaterials())
-			{
-				MeshShaderMaps.Add(new FMeshMaterialShaderMap(GetShaderPlatform(), VertexFactoryType));
-			}
-		}
-
-		// Initialize OrderedMeshShaderMaps from the new contents of MeshShaderMaps.
-		InitOrderedMeshShaderMaps();
-
 		// Material shaders
 		TShaderMap<FMaterialShaderType>::SerializeInline(Ar, bInlineShaderResources, false, bLoadedByCookedMaterial);
 
@@ -2507,18 +2493,22 @@ void FMaterialShaderMap::Serialize(FArchive& Ar, bool bInlineShaderResources, bo
 		int32 NumMeshShaderMaps = 0;
 		Ar << NumMeshShaderMaps;
 
+		MeshShaderMaps.Empty(NumMeshShaderMaps);
+		OrderedMeshShaderMaps.Empty(FVertexFactoryType::GetNumVertexFactoryTypes());
+		OrderedMeshShaderMaps.AddZeroed(FVertexFactoryType::GetNumVertexFactoryTypes());
+
 		for (int32 VFIndex = 0; VFIndex < NumMeshShaderMaps; VFIndex++)
 		{
 			FVertexFactoryType* VFType = nullptr;
-
 			Ar << VFType;
-
 			// Not handling missing vertex factory types on cooked data
 			// The cooker binary and running binary are assumed to be on the same code version
 			check(VFType);
-			FMeshMaterialShaderMap* MeshShaderMap = OrderedMeshShaderMaps[VFType->GetId()];
-			check(MeshShaderMap);
+
+			FMeshMaterialShaderMap* MeshShaderMap = new FMeshMaterialShaderMap(GetShaderPlatform(), VFType);
 			MeshShaderMap->SerializeInline(Ar, bInlineShaderResources, false, bLoadedByCookedMaterial);
+			MeshShaderMaps.Add(MeshShaderMap);
+			OrderedMeshShaderMaps[VFType->GetId()] = MeshShaderMap;
 		}
 
 		bool bCooked;
