PR #5999: UE-76055 Fix the Fog toggle while in Unlit View Mode (Contributed by [REDACTED_PERSON])

diff --git a/Engine/Source/Runtime/Engine/Private/ShowFlags.cpp b/Engine/Source/Runtime/Engine/Private/ShowFlags.cpp
--- a/Engine/Source/Runtime/Engine/Private/ShowFlags.cpp
+++ b/Engine/Source/Runtime/Engine/Private/ShowFlags.cpp
@@ -421,9 +421,14 @@ void EngineShowFlagOverride(EShowFlagInitMode ShowFlagInitMode, EViewModeIndex V
 			EngineShowFlags.SetBrushes(true);
 		}
 
+		if (ViewModeIndex == VMI_Unlit)
+		{
+			EngineShowFlags.SetLighting(false);
+			EngineShowFlags.AtmosphericFog = 0;
+		}
+
 		if( ViewModeIndex == VMI_Wireframe ||
 			ViewModeIndex == VMI_BrushWireframe ||
-			ViewModeIndex == VMI_Unlit ||
 			ViewModeIndex == VMI_StationaryLightOverlap ||
 			ViewModeIndex == VMI_ShaderComplexity ||
 			ViewModeIndex == VMI_QuadOverdraw ||
@@ -436,7 +441,7 @@ void EngineShowFlagOverride(EShowFlagInitMode ShowFlagInitMode, EViewModeIndex V
 		{
 			EngineShowFlags.SetLighting(false);
 			EngineShowFlags.AtmosphericFog = 0;
-			EngineShowFlags.Fog = (ViewModeIndex == VMI_Unlit); // Only unlit uses it 
+			EngineShowFlags.Fog = 0;
 		}
 
 		if( ViewModeIndex == VMI_Lit ||
