Fixed "-Wshadow" warnings of older clang version for Android.

#rnx

diff --git a/Engine/Source/Runtime/Engine/Private/GPUSkinVertexFactory.cpp b/Engine/Source/Runtime/Engine/Private/GPUSkinVertexFactory.cpp
--- a/Engine/Source/Runtime/Engine/Private/GPUSkinVertexFactory.cpp
+++ b/Engine/Source/Runtime/Engine/Private/GPUSkinVertexFactory.cpp
@@ -236,10 +236,10 @@ bool FGPUBaseSkinVertexFactory::FShaderDataType::UpdateBoneData(FRHICommandListI
 			if (!bUseSkinCache && DeferSkeletalLockAndFillToRHIThread())
 			{
 				FRHIVertexBuffer* VertexBuffer = CurrentBoneBuffer->VertexBufferRHI;
-				RHICmdList.EnqueueLambda([VertexBuffer, VectorArraySize, &ReferenceToLocalMatrices, &BoneMap](FRHICommandListImmediate& RHICmdList)
+				RHICmdList.EnqueueLambda([VertexBuffer, VectorArraySize, &ReferenceToLocalMatrices, &BoneMap](FRHICommandListImmediate& InRHICmdList)
 				{
 					QUICK_SCOPE_CYCLE_COUNTER(STAT_FRHICommandUpdateBoneBuffer_Execute);
-					FMatrix3x4* ChunkMatrices = (FMatrix3x4*)RHICmdList.LockVertexBuffer(VertexBuffer, 0, VectorArraySize, RLM_WriteOnly);
+					FMatrix3x4* LambdaChunkMatrices = (FMatrix3x4*)InRHICmdList.LockVertexBuffer(VertexBuffer, 0, VectorArraySize, RLM_WriteOnly);
 					//FMatrix3x4 is sizeof() == 48
 					// PLATFORM_CACHE_LINE_SIZE (128) / 48 = 2.6
 					//  sizeof(FMatrix) == 64
@@ -254,11 +254,11 @@ bool FGPUBaseSkinVertexFactory::FShaderDataType::UpdateBoneData(FRHICommandListI
 						FPlatformMisc::Prefetch(ReferenceToLocalMatrices.GetData() + RefToLocalIdx + PreFetchStride);
 						FPlatformMisc::Prefetch(ReferenceToLocalMatrices.GetData() + RefToLocalIdx + PreFetchStride, PLATFORM_CACHE_LINE_SIZE);
 
-						FMatrix3x4& BoneMat = ChunkMatrices[BoneIdx];
+						FMatrix3x4& BoneMat = LambdaChunkMatrices[BoneIdx];
 						const FMatrix& RefToLocal = ReferenceToLocalMatrices[RefToLocalIdx];
 						RefToLocal.To3x4MatrixTranspose((float*)BoneMat.M);
 					}
-					RHICmdList.UnlockVertexBuffer(VertexBuffer);
+					InRHICmdList.UnlockVertexBuffer(VertexBuffer);
 				});
 
 				RHICmdList.RHIThreadFence(true);
@@ -924,11 +924,11 @@ bool FGPUBaseSkinAPEXClothVertexFactory::ClothShaderType::UpdateClothSimulData(F
 				{
 					QUICK_SCOPE_CYCLE_COUNTER(STAT_FRHICommandUpdateBoneBuffer_Execute);
 					float* RESTRICT Data = (float* RESTRICT)InRHICmdList.LockVertexBuffer(VertexBuffer, 0, VectorArraySize, RLM_WriteOnly);
-					uint32 NumSimulVerts = InSimulPositions.Num();
-					check(NumSimulVerts > 0 && NumSimulVerts <= MAX_APEXCLOTH_VERTICES_FOR_VB);
+					uint32 LambdaNumSimulVerts = InSimulPositions.Num();
+					check(LambdaNumSimulVerts > 0 && LambdaNumSimulVerts <= MAX_APEXCLOTH_VERTICES_FOR_VB);
 					float* RESTRICT Pos = (float* RESTRICT) &InSimulPositions[0].X;
 					float* RESTRICT Normal = (float* RESTRICT) &InSimulNormals[0].X;
-					for (uint32 Index = 0; Index < NumSimulVerts; Index++)
+					for (uint32 Index = 0; Index < LambdaNumSimulVerts; Index++)
 					{
 						FPlatformMisc::Prefetch(Pos + PLATFORM_CACHE_LINE_SIZE);
 						FPlatformMisc::Prefetch(Normal + PLATFORM_CACHE_LINE_SIZE);
diff --git a/Engine/Source/Runtime/SlateRHIRenderer/Private/SlateRHIRenderingPolicy.cpp b/Engine/Source/Runtime/SlateRHIRenderer/Private/SlateRHIRenderingPolicy.cpp
--- a/Engine/Source/Runtime/SlateRHIRenderer/Private/SlateRHIRenderingPolicy.cpp
+++ b/Engine/Source/Runtime/SlateRHIRenderer/Private/SlateRHIRenderingPolicy.cpp
@@ -30,6 +30,7 @@
 extern void UpdateNoiseTextureParameters(FViewUniformShaderParameters& ViewUniformShaderParameters);
 
 DECLARE_CYCLE_STAT(TEXT("Update Buffers RT"), STAT_SlateUpdateBufferRTTime, STATGROUP_Slate);
+DECLARE_CYCLE_STAT(TEXT("Update Buffers RT"), STAT_SlateUpdateBufferRTTimeLambda, STATGROUP_Slate);
 DECLARE_CYCLE_STAT(TEXT("PreFill Buffers RT"), STAT_SlatePreFullBufferRTTime, STATGROUP_Slate);
 DECLARE_DWORD_COUNTER_STAT(TEXT("Num Layers"), STAT_SlateNumLayers, STATGROUP_Slate);
 DECLARE_DWORD_COUNTER_STAT(TEXT("Num Batches"), STAT_SlateNumBatches, STATGROUP_Slate);
@@ -196,7 +197,8 @@ void FSlateRHIRenderingPolicy::UpdateVertexAndIndexBuffers(FRHICommandListImmedi
 
 			RHICmdList.EnqueueLambda([VertexBufferRHI, IndexBufferRHI, &InBatchData, bAbsoluteIndices](FRHICommandListImmediate& InRHICmdList)
 			{
-				SCOPE_CYCLE_COUNTER(STAT_SlateUpdateBufferRTTime);
+				// Use different cycle counter to avoid -Wshadow warning on older clang versions for Android
+				SCOPE_CYCLE_COUNTER(STAT_SlateUpdateBufferRTTimeLambda);
 
 				const int32 NumBatchedVertices = InBatchData.GetNumBatchedVertices();
 				const int32 NumBatchedIndices = InBatchData.GetNumBatchedIndices();
diff --git a/Engine/Source/Runtime/SlateRHIRenderer/Private/SlateUpdatableBuffer.cpp b/Engine/Source/Runtime/SlateRHIRenderer/Private/SlateUpdatableBuffer.cpp
--- a/Engine/Source/Runtime/SlateRHIRenderer/Private/SlateUpdatableBuffer.cpp
+++ b/Engine/Source/Runtime/SlateRHIRenderer/Private/SlateUpdatableBuffer.cpp
@@ -4,6 +4,7 @@
 #include "RenderingThread.h"
 
 DECLARE_CYCLE_STAT(TEXT("UpdateInstanceBuffer Time"), STAT_SlateUpdateInstanceBuffer, STATGROUP_Slate);
+DECLARE_CYCLE_STAT(TEXT("UpdateInstanceBuffer Time"), STAT_SlateUpdateInstanceBufferLambda, STATGROUP_Slate);
 
 FSlateUpdatableInstanceBuffer::FSlateUpdatableInstanceBuffer( int32 InitialInstanceCount )
 	: FreeBufferIndex(0)
@@ -80,7 +81,8 @@ void FSlateUpdatableInstanceBuffer::UpdateRenderingData_RenderThread(FRHICommand
 		FRHIVertexBuffer* VertexBufferRHI = InstanceBufferResource.VertexBufferRHI;
 		RHICmdList.EnqueueLambda([VertexBufferRHI, &InstanceData = RenderThreadBufferData](FRHICommandListImmediate& InRHICmdList)
 		{
-			SCOPE_CYCLE_COUNTER(STAT_SlateUpdateInstanceBuffer);
+			// Use different cycle counter to avoid -Wshadow warning on older clang versions for Android
+			SCOPE_CYCLE_COUNTER(STAT_SlateUpdateInstanceBufferLambda);
 
 			int32 RequiredVertexBufferSize = InstanceData.Num() * sizeof(FVector4);
 			uint8* InstanceBufferData = (uint8*)InRHICmdList.LockVertexBuffer(VertexBufferRHI, 0, RequiredVertexBufferSize, RLM_WriteOnly);
