Don't allow IES textures to be converted to VT

diff --git a/Engine/Source/Runtime/Engine/Classes/Engine/TextureLightProfile.h b/Engine/Source/Runtime/Engine/Classes/Engine/TextureLightProfile.h
--- a/Engine/Source/Runtime/Engine/Classes/Engine/TextureLightProfile.h
+++ b/Engine/Source/Runtime/Engine/Classes/Engine/TextureLightProfile.h
@@ -12,6 +12,10 @@ class UTextureLightProfile : public UTexture2D
 {
 	GENERATED_UCLASS_BODY()
 
+#if WITH_EDITOR
+	virtual void PostEditChangeProperty(FPropertyChangedEvent& PropertyChangedEvent) override;
+#endif
+
 	/** Light brightness in Candelas, imported from IES profile, <= 0 if the profile is used for masking only. Use with InverseSquareFalloff. */
 	UPROPERTY(EditAnywhere, Category=TextureLightProfile, AssetRegistrySearchable)
 	float Brightness;
diff --git a/Engine/Source/Runtime/Engine/Private/TextureLightProfile.cpp b/Engine/Source/Runtime/Engine/Private/TextureLightProfile.cpp
--- a/Engine/Source/Runtime/Engine/Private/TextureLightProfile.cpp
+++ b/Engine/Source/Runtime/Engine/Private/TextureLightProfile.cpp
@@ -17,3 +17,14 @@ UTextureLightProfile::UTextureLightProfile(const FObjectInitializer& ObjectIniti
 	Brightness = -1;
 }
 
+#if WITH_EDITOR
+void UTextureLightProfile::PostEditChangeProperty(FPropertyChangedEvent& PropertyChangedEvent)
+{
+	if (VirtualTextureStreaming)
+	{
+		UE_LOG(LogTexture, Warning, TEXT("VirtualTextureStreaming not supported for IES textures"));
+		VirtualTextureStreaming = false;
+	}
+	Super::PostEditChangeProperty(PropertyChangedEvent);
+}
+#endif // WITH_EDITOR
