Fix crash when generating thumbnail image for VTs with extreme aspect ratios

diff --git a/Engine/Source/Runtime/Engine/Private/Texture2D.cpp b/Engine/Source/Runtime/Engine/Private/Texture2D.cpp
--- a/Engine/Source/Runtime/Engine/Private/Texture2D.cpp
+++ b/Engine/Source/Runtime/Engine/Private/Texture2D.cpp
@@ -1925,10 +1925,11 @@ void FVirtualTexture2DResource::InitializeEditorResources(IVirtualTexture* InVir
 
 		// Choose a mip level for the thumbnail texture to ensure proper size
 		const uint32 MaxMipLevel = VTData->GetNumMips() - 1u;
+		const uint32 MaxTextureSize = FMath::Min<uint32>(GetMax2DTextureDimension(), 1024u);
 		uint32 MipLevel = 0;
 		uint32 MipWidth = GetSizeX();
 		uint32 MipHeight = GetSizeY();
-		while ((MipWidth > 128u && MipHeight > 128u) && MipLevel < MaxMipLevel)
+		while (((MipWidth > 128u && MipHeight > 128u) || MipWidth > MaxTextureSize || MipHeight > MaxTextureSize) && MipLevel < MaxMipLevel)
 		{
 			++MipLevel;
 			MipWidth = FMath::Max(MipWidth / 2u, 1u);
