Trigger a material compilation error if a RVT output node is used in a material with RVT domain.
Simply overriding material output might be confusing for users, so this is more explicit.

diff --git a/Engine/Source/Runtime/Engine/Private/Materials/HLSLMaterialTranslator.h b/Engine/Source/Runtime/Engine/Private/Materials/HLSLMaterialTranslator.h
--- a/Engine/Source/Runtime/Engine/Private/Materials/HLSLMaterialTranslator.h
+++ b/Engine/Source/Runtime/Engine/Private/Materials/HLSLMaterialTranslator.h
@@ -6697,8 +6697,17 @@ protected:
 
 	virtual int32 VirtualTextureOutput() override
 	{
-		MaterialCompilationOutput.bHasRuntimeVirtualTextureOutput = true;
-		
+		if (Material->GetMaterialDomain() == MD_RuntimeVirtualTexture)
+		{
+			// RuntimeVirtualTextureOutput would priority over the output material attributes here
+			// But that could be considered confusing for the user so we error instead
+			Errorf(TEXT("RuntimeVirtualTextureOutput nodes are not used when the Material Domain is set to 'Virtual Texture'"));
+		}
+		else
+		{
+			MaterialCompilationOutput.bHasRuntimeVirtualTextureOutput = true;
+		}
+
 		// return value is not used
 		return INDEX_NONE;
 	}
