Fixed high crash count when compiling shaders and the drive runs out of space. Also fixed log manager crash when HandleOutputFileCorrupted is triggered in some out of drive space situations.

diff --git a/Engine/Source/Runtime/Engine/Private/ShaderCompiler/ShaderCompiler.cpp b/Engine/Source/Runtime/Engine/Private/ShaderCompiler/ShaderCompiler.cpp
--- a/Engine/Source/Runtime/Engine/Private/ShaderCompiler/ShaderCompiler.cpp
+++ b/Engine/Source/Runtime/Engine/Private/ShaderCompiler/ShaderCompiler.cpp
@@ -348,7 +348,7 @@ namespace SCWErrorCode
 
 	void HandleOutputFileCorrupted(const TCHAR* Filename, int64 ExpectedSize, int64 ActualSize)
 	{
-		ModalErrorOrLog(FString::Printf(TEXT("Output file corrupted (expected %I64d bytes, but only got %I64d): %s"), Filename, ExpectedSize, ActualSize));
+		ModalErrorOrLog(FString::Printf(TEXT("Output file corrupted (expected %I64d bytes, but only got %I64d): %s"), ExpectedSize, ActualSize, Filename));
 	}
 }
 
@@ -1197,7 +1197,7 @@ void FShaderCompileThreadRunnable::WriteNewTasks()
 				uint64 TotalDiskSpace = 0;
 				uint64 FreeDiskSpace = 0;
 				FPlatformMisc::GetDiskTotalAndFreeSpace(TransferFileName, TotalDiskSpace, FreeDiskSpace);
-				UE_LOG(LogShaderCompilers, Fatal, TEXT("Could not write the shader compiler transfer filename to '%s' (Free Disk Space: %llu."), *TransferFileName, FreeDiskSpace);
+				UE_LOG(LogShaderCompilers, Error, TEXT("Could not write the shader compiler transfer filename to '%s' (Free Disk Space: %llu."), *TransferFileName, FreeDiskSpace);
 			}
 			delete TransferFile;
 
@@ -1208,7 +1208,7 @@ void FShaderCompileThreadRunnable::WriteNewTasks()
 				uint64 TotalDiskSpace = 0;
 				uint64 FreeDiskSpace = 0;
 				FPlatformMisc::GetDiskTotalAndFreeSpace(TransferFileName, TotalDiskSpace, FreeDiskSpace);
-				UE_LOG(LogShaderCompilers, Fatal, TEXT("Could not rename the shader compiler transfer filename to '%s' from '%s' (Free Disk Space: %llu)."), *ProperTransferFileName, *TransferFileName, FreeDiskSpace);
+				UE_LOG(LogShaderCompilers, Error, TEXT("Could not rename the shader compiler transfer filename to '%s' from '%s' (Free Disk Space: %llu)."), *ProperTransferFileName, *TransferFileName, FreeDiskSpace);
 			}
 		}
 	}
