Fix for Cvar marked up in RendererSettings "was not found or is set to read-only"

diff --git a/Engine/Source/Runtime/Engine/Private/DeveloperSettings.cpp b/Engine/Source/Runtime/Engine/Private/DeveloperSettings.cpp
--- a/Engine/Source/Runtime/Engine/Private/DeveloperSettings.cpp
+++ b/Engine/Source/Runtime/Engine/Private/DeveloperSettings.cpp
@@ -167,7 +167,13 @@ void UDeveloperSettings::ExportValuesToConsoleVariables(UProperty* PropertyThatC
 			}
 			else
 			{
-				UE_LOG(LogInit, Warning, TEXT("CVar named '%s' marked up in %s was not found or is set to read-only"), *CVarName, *GetClass()->GetName());
+				// Reduce the amount of log spam for read-only properties. 
+				// We assume that if property requires restart it is very likely it needs to stay read-only and therefore no need to log a warning.
+				static const FName ConfigRestartRequiredKey = TEXT("ConfigRestartRequired");
+				if (!PropertyThatChanged->GetBoolMetaData(ConfigRestartRequiredKey))
+				{
+					UE_LOG(LogInit, Warning, TEXT("CVar named '%s' marked up in %s was not found or is set to read-only"), *CVarName, *GetClass()->GetName());
+				}
 			}
 		}
 	}
