Copying //UE4/Dev-Volumetrics to Dev-Rendering (//UE4/Dev-Rendering)

diff --git a/Engine/Source/Runtime/Engine/Classes/Components/SkyAtmosphereComponent.h b/Engine/Source/Runtime/Engine/Classes/Components/SkyAtmosphereComponent.h
new file mode 100644
--- /dev/null
+++ b/Engine/Source/Runtime/Engine/Classes/Components/SkyAtmosphereComponent.h
@@ -0,0 +1,199 @@
+// Copyright 1998-2019 Epic Games, Inc. All Rights Reserved.
+
+#pragma once
+
+#include "CoreMinimal.h"
+#include "Components/SceneComponent.h"
+#include "GameFramework/Info.h"
+#include "Misc/Guid.h"
+#include "RenderResource.h"
+
+#include "SkyAtmosphereComponent.generated.h"
+
+
+USTRUCT(BlueprintType)
+struct FTentDistribution
+{
+	GENERATED_USTRUCT_BODY()
+
+	UPROPERTY(EditAnywhere, BlueprintReadOnly, interp, Category = "Tent", meta = (UIMin = 0.0, UIMax = 60.0))
+	float TipAltitude;
+	
+	UPROPERTY(EditAnywhere, BlueprintReadOnly, interp, Category = "Tent", meta = (UIMin = 0.0, UIMax = 1.0, ClampMin = 0.0, SliderExponent = 4.0))
+	float TipValue;
+	
+	UPROPERTY(EditAnywhere, BlueprintReadOnly, interp, Category = "Tent", meta = (UIMin = 0.01, UIMax = 20.0, ClampMin = 0.0))
+	float Width;
+};
+
+
+/**
+ * 
+ */
+UCLASS(ClassGroup = Rendering, collapsecategories, hidecategories = (Object, Mobility, Activation, "Components|Activation"), editinlinenew, meta = (BlueprintSpawnableComponent), MinimalAPI)
+class USkyAtmosphereComponent : public USceneComponent
+{
+	GENERATED_UCLASS_BODY()
+
+	~USkyAtmosphereComponent();
+
+
+
+	/** The planet radius. (kilometers from the center to the ground level). */
+	UPROPERTY(EditAnywhere, BlueprintReadOnly, interp, Category = "Planet", meta = (DisplayName = "Ground Radius", UIMin = 6000.0, UIMax = 7000.0, ClampMin = 100.0, ClampMax = 10000.0))
+	float BottomRadius;
+
+	/** The ground albedo that will tint the astmophere when the sun light will bounce on it. Only taken into account when MultiScattering>0.0. */
+	UPROPERTY(EditAnywhere, BlueprintReadOnly, interp, Category = "Planet", meta = (HideAlphaChannel))
+	FColor GroundAlbedo;
+
+
+
+	/** The planet radius. (kilometers from the center to the ground level). */
+	UPROPERTY(EditAnywhere, BlueprintReadOnly, interp, Category = "Atmosphere", meta = (UIMin = 10.0, UIMax = 200.0, ClampMin = 10.0))
+	float AtmosphereHeight;
+
+	/** Render multi scattering as if sun light would bounce around in the atmosphere. This is achieved using a dual scattering approach. */
+	UPROPERTY(EditAnywhere, BlueprintReadOnly, interp, Category = "Atmosphere", meta = (DisplayName = "MultiScattering", UIMin = 0.0, UIMax = 1.0, ClampMin = 0.0, ClampMax = 2.0))
+	float MultiScatteringFactor;
+	
+
+
+	/** Rayleigh scattering coefficient scale.*/
+	UPROPERTY(EditAnywhere, BlueprintReadOnly, interp, Category = "Atmosphere - Raleigh", meta = (UIMin = 0.0, UIMax = 2.0, ClampMin = 0.0, SliderExponent = 4.0))
+	float RayleighScatteringScale;
+
+	/** The Rayleigh scattering coefficients resulting from molecules in the air at an altitude of 0 kilometer. */
+	UPROPERTY(EditAnywhere, BlueprintReadOnly, interp, Category = "Atmosphere - Raleigh", meta=(HideAlphaChannel))
+	FColor RayleighScattering;
+
+	/** The altitude in kilometer at which Rayleigh scattering effect is reduced to 40%.*/
+	UPROPERTY(EditAnywhere, BlueprintReadOnly, interp, Category = "Atmosphere - Raleigh", meta = (UIMin = 0.01, UIMax = 20.0, ClampMin = 0.1, SliderExponent = 5.0))
+	float RayleighExponentialDistribution;
+
+
+
+	/** Mie scattering coefficient scale.*/
+	UPROPERTY(EditAnywhere, BlueprintReadOnly, interp, Category = "Atmosphere - Mie", meta = (UIMin = 0.0, UIMax = 5.0, ClampMin = 0.0, SliderExponent = 4.0))
+	float MieScatteringScale;
+
+	/** The Mie scattering coefficients resulting from particles in the air at an altitude of 0 kilometer. As it becomes higher, light will be scattered more. */
+	UPROPERTY(EditAnywhere, BlueprintReadOnly, interp, Category = "Atmosphere - Mie", meta = (HideAlphaChannel))
+	FColor MieScattering;
+
+	/** Mie absorption coefficient scale.*/
+	UPROPERTY(EditAnywhere, BlueprintReadOnly, interp, Category = "Atmosphere - Mie", meta = (UIMin = 0.0, UIMax = 5.0, ClampMin = 0.0, SliderExponent = 4.0))
+	float MieAbsorptionScale;
+
+	/** The Mie absorption coefficients resulting from particles in the air at an altitude of 0 kilometer. As it becomes higher, light will be absorbed more. */
+	UPROPERTY(EditAnywhere, BlueprintReadOnly, interp, Category = "Atmosphere - Mie", meta = (HideAlphaChannel))
+	FColor MieAbsorption;
+	
+	/** A value of 0 mean light is uniformly scattered. A value closer to 1 means lights will scatter more forward, resulting in halos around light sources. */
+	UPROPERTY(EditAnywhere, BlueprintReadOnly, interp, Category = "Atmosphere - Mie", meta = (UIMin = 0.0, UIMax = 0.999, ClampMin = 0.0, ClampMax = 0.999))
+	float MieAnisotropy;
+
+	/** The altitude in kilometer at which Mie effects are reduced to 40%.*/
+	UPROPERTY(EditAnywhere, BlueprintReadOnly, interp, Category = "Atmosphere - Mie", meta = (UIMin = 0.01, UIMax = 10.0, ClampMin = 0.1, SliderExponent = 5.0))
+	float MieExponentialDistribution;
+
+
+
+	/** Absorption coefficients for another atmosphere layer. Density increase from 0 to 1 between 10 to 25km and decreases from 1 to 0 between 25 to 40km. This approximates ozone molecules distribution in the Earth atmosphere. */
+	UPROPERTY(EditAnywhere, BlueprintReadOnly, interp, Category = "Atmosphere - Absorption", meta = (DisplayName = "Absorption Scale", UIMin = 0.0, UIMax = 0.2, ClampMin = 0.0, SliderExponent = 3.0))
+	float OtherAbsorptionScale;
+
+	/** Absorption coefficients for another atmosphere layer. Density increase from 0 to 1 between 10 to 25km and decreases from 1 to 0 between 25 to 40km. The default values represents ozone molecules absorption in the Earth atmosphere. */
+	UPROPERTY(EditAnywhere, BlueprintReadOnly, interp, Category = "Atmosphere - Absorption", meta = (DisplayName = "Absorption", HideAlphaChannel))
+	FColor OtherAbsorption;
+
+	/** Represents the altitude based tent distribution of absorption particles in the atmosphere. */
+	UPROPERTY(EditAnywhere, BlueprintReadOnly, Category = "Atmosphere - Absorption", meta = (DisplayName = "Tent Distribution"))
+	FTentDistribution OtherTentDistribution;
+
+	
+
+	/** Scales the luminance of pixels representing the sky, i.e. not belonging to any surface. */
+	UPROPERTY(EditAnywhere, BlueprintReadOnly, interp, Category = "Art direction", meta = (HideAlphaChannel))
+	FLinearColor SkyLuminanceFactor;
+
+	/** Makes the aerial perspective look thicker by scaling distances from view to surfaces (opaque and translucent). */
+	UPROPERTY(EditAnywhere, BlueprintReadOnly, interp, Category = "Art direction", meta = (UIMin = 0.0, UIMax = 3.0, ClampMin = 0.0, SliderExponent = 2.0))
+	float AerialPespectiveViewDistanceScale;
+
+
+
+protected:
+	//~ Begin UActorComponent Interface.
+	virtual void CreateRenderState_Concurrent() override;
+	//virtual void SendRenderTransform_Concurrent() override;
+	virtual void DestroyRenderState_Concurrent() override;
+	//~ End UActorComponent Interface.
+
+public:
+
+	//~ Begin UObject Interface. 
+//	virtual void PostLoad() override;
+//	virtual bool IsPostLoadThreadSafe() const override;
+//	virtual void BeginDestroy() override;
+	virtual void PostInterpChange(UProperty* PropertyThatChanged) override;
+	virtual void Serialize(FArchive& Ar) override;
+#if WITH_EDITOR
+	virtual void PostEditChangeProperty(FPropertyChangedEvent& PropertyChangedEvent) override;
+#endif // WITH_EDITOR
+	//~ End UObject Interface
+
+	//~ Begin UActorComponent Interface.
+#if WITH_EDITOR
+	virtual void CheckForErrors() override;
+#endif // WITH_EDITOR
+	//~ End UActorComponent Interface.
+
+	FGuid GetStaticLightingBuiltGuid() const { return bStaticLightingBuiltGUID; }
+
+private:
+
+	/** Add this component to the render scene */
+	void AddToRenderScene() const;
+
+
+	/**
+	 * GUID used to associate a atmospheric component with precomputed lighting/shadowing information across levels.
+	 * The GUID changes whenever the atmospheric properties change, e.g. LUTs.
+	 */
+	UPROPERTY()
+	FGuid bStaticLightingBuiltGUID;
+	/**
+	 * Validate static lighting GUIDs and update as appropriate.
+	 */
+	void ValidateStaticLightingGUIDs();
+	/**
+	 * Update static lighting GUIDs.
+	 */
+	void UpdateStaticLightingGUIDs();
+};
+
+
+/**
+ *	A placeable actor that simulates sky and light scattering in the atmosphere.
+ *	@see TODO address to the documentation.
+ */
+UCLASS(showcategories = (Movement, Rendering, "Utilities|Transformation", "Input|MouseInput", "Input|TouchInput"), ClassGroup = Fog, hidecategories = (Info, Object, Input), MinimalAPI)
+class ASkyAtmosphere : public AInfo
+{
+	GENERATED_UCLASS_BODY()
+
+private:
+
+	UPROPERTY(BlueprintReadOnly, VisibleAnywhere, Category = Atmosphere, meta = (AllowPrivateAccess = "true"))
+		class USkyAtmosphereComponent* SkyAtmosphereComponent;
+
+#if WITH_EDITORONLY_DATA
+	/** Arrow component to indicate default sun rotation */
+	UPROPERTY()
+		class UArrowComponent* ArrowComponent;
+#endif
+
+public:
+
+};
diff --git a/Engine/Source/Runtime/Engine/Classes/Engine/MapBuildDataRegistry.h b/Engine/Source/Runtime/Engine/Classes/Engine/MapBuildDataRegistry.h
--- a/Engine/Source/Runtime/Engine/Classes/Engine/MapBuildDataRegistry.h
+++ b/Engine/Source/Runtime/Engine/Classes/Engine/MapBuildDataRegistry.h
@@ -272,6 +272,12 @@ public:
 	}
 };
 
+class FSkyAtmosphereMapBuildData
+{
+public:
+	bool bDummy = false;
+};
+
 UCLASS(MinimalAPI)
 class UMapBuildDataRegistry : public UObject
 {
@@ -338,6 +344,17 @@ public:
 	ENGINE_API const FReflectionCaptureMapBuildData* GetReflectionCaptureBuildData(FGuid CaptureId) const;
 	ENGINE_API FReflectionCaptureMapBuildData* GetReflectionCaptureBuildData(FGuid CaptureId);
 
+	/**
+	 * Allocates a new FAtmosphericFogMapBuildData from the registry.
+	 * Warning: Further allocations will invalidate the returned reference.
+	 */
+	ENGINE_API FSkyAtmosphereMapBuildData& FindOrAllocateSkyAtmosphereBuildData(const FGuid& Guid);
+	/**
+	 * @returns pointer to the AtmosphericFogBuildData, nullptr if built data has not been built yet.
+	 */
+	ENGINE_API const FSkyAtmosphereMapBuildData* GetSkyAtmosphereBuildData(const FGuid& Guid) const;
+	ENGINE_API void ClearSkyAtmosphereBuildData();
+
 	ENGINE_API void InvalidateStaticLighting(UWorld* World, bool bRecreateRenderState = true, const TSet<FGuid>* ResourcesToKeep = nullptr);
 	ENGINE_API void InvalidateReflectionCaptures(const TSet<FGuid>* ResourcesToKeep = nullptr);
 
@@ -363,6 +380,7 @@ private:
 	TMap<FGuid, FPrecomputedVolumetricLightmapData*> LevelPrecomputedVolumetricLightmapBuildData;
 	TMap<FGuid, FLightComponentMapBuildData> LightBuildData;
 	TMap<FGuid, FReflectionCaptureMapBuildData> ReflectionCaptureBuildData;
+	TMap<FGuid, FSkyAtmosphereMapBuildData> SkyAtmosphereBuildData;
 
 	bool bSetupResourceClusters;
 	TArray<FLightmapResourceCluster> LightmapResourceClusters;
diff --git a/Engine/Source/Runtime/Engine/Classes/Materials/MaterialExpressionSkyAtmosphereLightDirection.h b/Engine/Source/Runtime/Engine/Classes/Materials/MaterialExpressionSkyAtmosphereLightDirection.h
new file mode 100644
--- /dev/null
+++ b/Engine/Source/Runtime/Engine/Classes/Materials/MaterialExpressionSkyAtmosphereLightDirection.h
@@ -0,0 +1,28 @@
+// Copyright 1998-2019 Epic Games, Inc. All Rights Reserved.
+
+
+#pragma once
+
+#include "CoreMinimal.h"
+#include "UObject/ObjectMacros.h"
+#include "Materials/MaterialExpression.h"
+#include "MaterialExpressionSkyAtmosphereLightDirection.generated.h"
+
+UCLASS()
+class UMaterialExpressionSkyAtmosphereLightDirection : public UMaterialExpression
+{
+	GENERATED_UCLASS_BODY()
+
+	/** Index of the atmosphere light to sample. */
+	UPROPERTY(EditAnywhere, BlueprintReadWrite, Category = MaterialExpressionTextureCoordinate, meta = (UIMin = "0", UIMax = "1", ClampMin = "0", ClampMax = "1"))
+	int32 LightIndex;
+
+	//~ Begin UMaterialExpression Interface
+#if WITH_EDITOR
+	virtual int32 Compile(class FMaterialCompiler* Compiler, int32 OutputIndex) override;
+	virtual void GetCaption(TArray<FString>& OutCaptions) const override;
+#endif
+	//~ End UMaterialExpression Interface
+};
+
+
diff --git a/Engine/Source/Runtime/Engine/Classes/Materials/MaterialExpressionSkyAtmosphereLightIlluminance.h b/Engine/Source/Runtime/Engine/Classes/Materials/MaterialExpressionSkyAtmosphereLightIlluminance.h
new file mode 100644
--- /dev/null
+++ b/Engine/Source/Runtime/Engine/Classes/Materials/MaterialExpressionSkyAtmosphereLightIlluminance.h
@@ -0,0 +1,77 @@
+// Copyright 1998-2019 Epic Games, Inc. All Rights Reserved.
+
+
+#pragma once
+
+#include "CoreMinimal.h"
+#include "UObject/ObjectMacros.h"
+#include "Materials/MaterialExpression.h"
+#include "MaterialExpressionSkyAtmosphereLightIlluminance.generated.h"
+
+UCLASS()
+class UMaterialExpressionSkyAtmosphereLightIlluminance : public UMaterialExpression
+{
+	GENERATED_UCLASS_BODY()
+
+	/** Index of the atmosphere light to sample. */
+	UPROPERTY(EditAnywhere, BlueprintReadWrite, Category = MaterialExpressionTextureCoordinate, meta = (UIMin = "0", UIMax = "1", ClampMin = "0", ClampMax = "1"))
+	int32 LightIndex;
+
+	//~ Begin UMaterialExpression Interface
+#if WITH_EDITOR
+	virtual int32 Compile(class FMaterialCompiler* Compiler, int32 OutputIndex) override;
+	virtual void GetCaption(TArray<FString>& OutCaptions) const override;
+#endif
+	//~ End UMaterialExpression Interface
+};
+
+
+
+UCLASS()
+class UMaterialExpressionSkyAtmosphereLightDiskLuminance : public UMaterialExpression
+{
+	GENERATED_UCLASS_BODY()
+
+	/** Index of the atmosphere light to sample. */
+	UPROPERTY(EditAnywhere, BlueprintReadWrite, Category = MaterialExpressionTextureCoordinate, meta = (UIMin = "0", UIMax = "1", ClampMin = "0", ClampMax = "1"))
+	int32 LightIndex;
+
+	//~ Begin UMaterialExpression Interface
+#if WITH_EDITOR
+	virtual int32 Compile(class FMaterialCompiler* Compiler, int32 OutputIndex) override;
+	virtual void GetCaption(TArray<FString>& OutCaptions) const override;
+#endif
+	//~ End UMaterialExpression Interface
+};
+
+
+
+UCLASS()
+class UMaterialExpressionSkyAtmosphereAerialPerspective : public UMaterialExpression
+{
+	GENERATED_UCLASS_BODY()
+
+	//~ Begin UMaterialExpression Interface
+#if WITH_EDITOR
+	virtual int32 Compile(class FMaterialCompiler* Compiler, int32 OutputIndex) override;
+	virtual void GetCaption(TArray<FString>& OutCaptions) const override;
+#endif
+	//~ End UMaterialExpression Interface
+};
+
+
+
+UCLASS()
+class UMaterialExpressionSkyAtmosphereDistantLightScatteredLuminance : public UMaterialExpression
+{
+	GENERATED_UCLASS_BODY()
+
+	//~ Begin UMaterialExpression Interface
+#if WITH_EDITOR
+	virtual int32 Compile(class FMaterialCompiler* Compiler, int32 OutputIndex) override;
+	virtual void GetCaption(TArray<FString>& OutCaptions) const override;
+#endif
+	//~ End UMaterialExpression Interface
+};
+
+
diff --git a/Engine/Source/Runtime/Engine/Private/Components/DirectionalLightComponent.cpp b/Engine/Source/Runtime/Engine/Private/Components/DirectionalLightComponent.cpp
--- a/Engine/Source/Runtime/Engine/Private/Components/DirectionalLightComponent.cpp
+++ b/Engine/Source/Runtime/Engine/Private/Components/DirectionalLightComponent.cpp
@@ -387,6 +387,11 @@ public:
 		return SunDiscOuterSpaceLuminance;
 	}
 
+	virtual FLinearColor GetTransmittanceFactor() const override
+	{
+		return AtmosphereTransmittanceFactor;
+	}
+
 	virtual float GetSunLightHalfApexAngleRadian() const override
 	{
 		return 0.5f * LightSourceAngle * PI / 180.0f; // LightSourceAngle is apex angle (angular diameter) in degree
diff --git a/Engine/Source/Runtime/Engine/Private/Components/LightComponent.cpp b/Engine/Source/Runtime/Engine/Private/Components/LightComponent.cpp
--- a/Engine/Source/Runtime/Engine/Private/Components/LightComponent.cpp
+++ b/Engine/Source/Runtime/Engine/Private/Components/LightComponent.cpp
@@ -277,6 +277,7 @@ FLightSceneProxy::FLightSceneProxy(const ULightComponent* InLightComponent)
 	, bCastModulatedShadows(false)
 	, bUseWholeSceneCSMForMovableObjects(false)
 	, bTiledDeferredLightingSupported(false)
+	, AtmosphereSunLightIndex(InLightComponent->GetAtmosphereSunLightIndex())
 	, LightType(InLightComponent->GetLightType())	
 	, LightingChannelMask(GetLightingChannelMaskForStruct(InLightComponent->LightingChannels))
 	, StatId(InLightComponent->GetStatID(true))
diff --git a/Engine/Source/Runtime/Engine/Private/Components/SkyAtmosphereComponent.cpp b/Engine/Source/Runtime/Engine/Private/Components/SkyAtmosphereComponent.cpp
new file mode 100644
--- /dev/null
+++ b/Engine/Source/Runtime/Engine/Private/Components/SkyAtmosphereComponent.cpp
@@ -0,0 +1,296 @@
+// Copyright 1998-2019 Epic Games, Inc. All Rights Reserved.
+
+#include "Components/SkyAtmosphereComponent.h"
+
+#include "Atmosphere/AtmosphericFogComponent.h"
+#include "Components/ArrowComponent.h"
+#include "Components/BillboardComponent.h"
+#include "Engine/MapBuildDataRegistry.h"
+#include "Internationalization/Text.h"
+#include "Logging/MessageLog.h"
+#include "Logging/TokenizedMessage.h"
+#include "Misc/MapErrors.h"
+#include "Misc/UObjectToken.h"
+#include "UObject/UObjectIterator.h"
+
+#if WITH_EDITOR
+#include "ObjectEditorUtils.h"
+#endif
+
+#define LOCTEXT_NAMESPACE "SkyAtmosphereComponent"
+
+//#pragma optimize( "", off )
+
+
+
+/*=============================================================================
+	USkyAtmosphereComponent implementation.
+=============================================================================*/
+
+USkyAtmosphereComponent::USkyAtmosphereComponent(const FObjectInitializer& ObjectInitializer)
+	: Super(ObjectInitializer)
+{
+	// All distance here are in kilometer and scattering/absorptions coefficient in 1/kilometers.
+	const float EarthBottomRadius = 6360.0f;
+	const float EarthTopRadius = 6420.0f;
+	const float EarthRayleighScaleHeight = 8.0f;
+	const float EarthMieScaleHeight = 1.2f;
+	
+	// Default: Earth like atmosphere
+	BottomRadius = EarthBottomRadius;
+	AtmosphereHeight = EarthTopRadius - EarthBottomRadius;
+	GroundAlbedo = FColor(0.0f, 0.0f, 0.0f);
+
+	// FLoat to a u8 rgb + float length can lose some precision but it is better UI wise.
+	const FLinearColor RayleightScatteringRaw = FLinearColor(0.005802f, 0.013558f, 0.033100f);
+	RayleighScattering = (RayleightScatteringRaw * (1.0f / RayleightScatteringRaw.B)).ToFColor(false);
+	RayleighScatteringScale = RayleightScatteringRaw.B;
+	RayleighExponentialDistribution = EarthRayleighScaleHeight;
+
+	MieScattering = FColor(FColor::White);
+	MieScatteringScale = 0.003996f;
+	MieAbsorption = FColor(FColor::White);
+	MieAbsorptionScale = 0.000444f;
+	MieAnisotropy = 0.8f;
+	MieExponentialDistribution = EarthMieScaleHeight;
+
+	// Absorption tent distribution representing ozone distribution in Earth atmosphere.
+	const FLinearColor OtherAbsorptionRaw = FLinearColor(0.000650f, 0.001881f, 0.000085f);
+	OtherAbsorptionScale = OtherAbsorptionRaw.G;
+	OtherAbsorption = (OtherAbsorptionRaw * (1.0f / OtherAbsorptionRaw.G)).ToFColor(false);
+	OtherTentDistribution.TipAltitude = 25.0f;
+	OtherTentDistribution.TipValue    =  1.0f;
+	OtherTentDistribution.Width       = 15.0f;
+
+	SkyLuminanceFactor = FLinearColor(FLinearColor::White);
+	MultiScatteringFactor = 1.0f;
+	AerialPespectiveViewDistanceScale = 1.0f;
+
+	ValidateStaticLightingGUIDs();
+}
+
+USkyAtmosphereComponent::~USkyAtmosphereComponent()
+{
+}
+
+static bool SkyAtmosphereComponentStaticLightingBuilt(const USkyAtmosphereComponent* Component)
+{
+	UMapBuildDataRegistry* Registry = Component->GetOwner() && Component->GetOwner()->GetLevel() ? Component->GetOwner()->GetLevel()->GetOrCreateMapBuildData() : nullptr;
+	const FSkyAtmosphereMapBuildData* SkyAtmosphereFogBuildData = Registry ? Registry->GetSkyAtmosphereBuildData(Component->GetStaticLightingBuiltGuid()) : nullptr;
+	class FSceneInterface* Scene = Component->GetWorld()->Scene;
+
+	// Only require building if there is a Sky or Sun light requiring lighting builds, i.e. non movable.
+	const bool StaticLightingDependsOnAtmosphere = Scene->HasSkyLightRequiringLightingBuild() || Scene->HasAtmosphereLightRequiringLightingBuild();
+	// Built data is available or static lighting does not depend any sun/sky components.
+	return (SkyAtmosphereFogBuildData != nullptr && StaticLightingDependsOnAtmosphere) || !StaticLightingDependsOnAtmosphere;
+}
+
+void USkyAtmosphereComponent::AddToRenderScene() const
+{
+	if (ShouldComponentAddToScene() && ShouldRender() && IsRegistered() && (GetOuter() == NULL || !GetOuter()->HasAnyFlags(RF_ClassDefaultObject)))
+	{
+		GetWorld()->Scene->AddSkyAtmosphere(this, SkyAtmosphereComponentStaticLightingBuilt(this));
+	}
+}
+
+void USkyAtmosphereComponent::CreateRenderState_Concurrent()
+{
+	Super::CreateRenderState_Concurrent();
+	// If one day we need to look up lightmass built data, lookup it up here using the guid from the correct MapBuildData.
+	AddToRenderScene();
+}
+
+void USkyAtmosphereComponent::DestroyRenderState_Concurrent()
+{
+	Super::DestroyRenderState_Concurrent();
+	GetWorld()->Scene->RemoveSkyAtmosphere(this);
+
+	// Search for new fog component
+	for (TObjectIterator<USkyAtmosphereComponent> It; It; ++It)
+	{
+		USkyAtmosphereComponent* Component = *It;
+		checkSlow(Component);
+		if (Component != this && Component->IsRegistered())
+		{
+			Component->AddToRenderScene();
+			break;
+		}
+	}
+}
+
+void USkyAtmosphereComponent::ValidateStaticLightingGUIDs()
+{
+	// Validate light guids.
+	if (!bStaticLightingBuiltGUID.IsValid())
+	{
+		UpdateStaticLightingGUIDs();
+	}
+}
+
+void USkyAtmosphereComponent::UpdateStaticLightingGUIDs()
+{
+	bStaticLightingBuiltGUID = FGuid::NewGuid();
+}
+
+#if WITH_EDITOR
+
+void USkyAtmosphereComponent::CheckForErrors()
+{
+	AActor* Owner = GetOwner();
+	if (Owner && bVisible)
+	{
+		UWorld* ThisWorld = Owner->GetWorld();
+		bool bMultipleFound = false;
+		bool bLegacyAtmosphericFogFound = false;
+
+		if (ThisWorld)
+		{
+			for (TObjectIterator<USkyAtmosphereComponent> ComponentIt; ComponentIt; ++ComponentIt)
+			{
+				USkyAtmosphereComponent* Component = *ComponentIt;
+
+				if (Component != this
+					&& !Component->IsPendingKill()
+					&& Component->bVisible
+					&& Component->GetOwner()
+					&& ThisWorld->ContainsActor(Component->GetOwner())
+					&& !Component->GetOwner()->IsPendingKill())
+				{
+					bMultipleFound = true;
+					break;
+				}
+			}
+			for (TObjectIterator<UAtmosphericFogComponent> ComponentIt; ComponentIt; ++ComponentIt)
+			{
+				UAtmosphericFogComponent* Component = *ComponentIt;
+
+				if (!Component->IsPendingKill()
+					&& Component->bVisible
+					&& Component->GetOwner()
+					&& ThisWorld->ContainsActor(Component->GetOwner())
+					&& !Component->GetOwner()->IsPendingKill())
+				{
+					bLegacyAtmosphericFogFound = true;
+					break;
+				}
+			}
+		}
+
+		if (bMultipleFound)
+		{
+			FMessageLog("MapCheck").Error()
+				->AddToken(FUObjectToken::Create(Owner))
+				->AddToken(FTextToken::Create(LOCTEXT("MapCheck_Message_MultipleSkyAtmosphere", "Multiple sky atmosphere are active, only one can be enabled per world.")))
+				->AddToken(FMapErrorToken::Create(FMapErrors::MultipleSkyAtmospheres));
+		}
+		if (bLegacyAtmosphericFogFound)
+		{
+			FMessageLog("MapCheck").Error()
+				->AddToken(FUObjectToken::Create(Owner))
+				->AddToken(FTextToken::Create(LOCTEXT("MapCheck_Message_MultipleSkyAtmosphereType", "A SkyAtmosphere and a legacy AtmosphericFog components are both active, we recommend to have only one enabled per world.")))
+				->AddToken(FMapErrorToken::Create(FMapErrors::MultipleSkyAtmosphereTypes));
+		}
+	}
+}
+
+void USkyAtmosphereComponent::PostEditChangeProperty(FPropertyChangedEvent& PropertyChangedEvent)
+{
+	Super::PostEditChangeProperty(PropertyChangedEvent);
+
+	// If any properties have been changed in the atmosphere category, it means the sky look will change and lighting needs to be rebuild.
+	const FName CategoryName = FObjectEditorUtils::GetCategoryFName(PropertyChangedEvent.Property);
+	if (CategoryName == FName(TEXT("Planet")) ||
+		CategoryName == FName(TEXT("Atmosphere")) ||
+		CategoryName == FName(TEXT("Atmosphere - Raleigh")) ||
+		CategoryName == FName(TEXT("Atmosphere - Mie")) ||
+		CategoryName == FName(TEXT("Atmosphere - Absorption")) ||
+		CategoryName == FName(TEXT("Art direction")))
+	{
+		UMapBuildDataRegistry* Registry = GetOwner() && GetOwner()->GetLevel() ? GetOwner()->GetLevel()->GetOrCreateMapBuildData() : nullptr;
+		if (SkyAtmosphereComponentStaticLightingBuilt(this))
+		{
+			// If we have changed an atmosphere property and the lighyting has already been built, we need to ask for a rebuild by updating the static lighting GUIDs.
+			UpdateStaticLightingGUIDs();
+		}
+	}
+}
+
+#endif // WITH_EDITOR
+
+void USkyAtmosphereComponent::PostInterpChange(UProperty* PropertyThatChanged)
+{
+	Super::PostInterpChange(PropertyThatChanged);
+	MarkRenderStateDirty();
+}
+
+void USkyAtmosphereComponent::Serialize(FArchive& Ar)
+{
+	Super::Serialize(Ar);
+	Ar << bStaticLightingBuiltGUID;
+}
+
+
+/*=============================================================================
+	ASkyAtmosphere implementation.
+=============================================================================*/
+
+#if WITH_EDITOR
+#include "ObjectEditorUtils.h"
+#endif
+
+ASkyAtmosphere::ASkyAtmosphere(const FObjectInitializer& ObjectInitializer)
+	: Super(ObjectInitializer)
+{
+	SkyAtmosphereComponent = CreateDefaultSubobject<USkyAtmosphereComponent>(TEXT("SkyAtmosphereComponent"));
+	RootComponent = SkyAtmosphereComponent;
+
+#if WITH_EDITORONLY_DATA
+	ArrowComponent = CreateEditorOnlyDefaultSubobject<UArrowComponent>(TEXT("ArrowComponent"));
+
+	if (!IsRunningCommandlet())
+	{
+		// Structure to hold one-time initialization
+		struct FConstructorStatics
+		{
+			ConstructorHelpers::FObjectFinderOptional<UTexture2D> SkyAtmosphereTextureObject;
+			FName ID_SkyAtmosphere;
+			FText NAME_SkyAtmosphere;
+			FConstructorStatics()
+				: SkyAtmosphereTextureObject(TEXT("/Engine/EditorResources/S_ExpoHeightFog"))
+				, ID_SkyAtmosphere(TEXT("Fog"))
+				, NAME_SkyAtmosphere(NSLOCTEXT("SpriteCategory", "Fog", "Fog"))
+			{
+			}
+		};
+		static FConstructorStatics ConstructorStatics;
+
+		if (GetSpriteComponent())
+		{
+			GetSpriteComponent()->Sprite = ConstructorStatics.SkyAtmosphereTextureObject.Get();
+			GetSpriteComponent()->RelativeScale3D = FVector(0.5f, 0.5f, 0.5f);
+			GetSpriteComponent()->SpriteInfo.Category = ConstructorStatics.ID_SkyAtmosphere;
+			GetSpriteComponent()->SpriteInfo.DisplayName = ConstructorStatics.NAME_SkyAtmosphere;
+			GetSpriteComponent()->SetupAttachment(SkyAtmosphereComponent);
+		}
+
+		if (ArrowComponent)
+		{
+			ArrowComponent->ArrowColor = FColor(150, 200, 255);
+
+			ArrowComponent->bTreatAsASprite = true;
+			ArrowComponent->SpriteInfo.Category = ConstructorStatics.ID_SkyAtmosphere;
+			ArrowComponent->SpriteInfo.DisplayName = ConstructorStatics.NAME_SkyAtmosphere;
+			ArrowComponent->SetupAttachment(SkyAtmosphereComponent);
+			ArrowComponent->bLightAttachment = true;
+			ArrowComponent->bIsScreenSizeScaled = true;
+		}
+	}
+#endif // WITH_EDITORONLY_DATA
+
+	PrimaryActorTick.bCanEverTick = true;
+	bHidden = false;
+}
+
+#undef LOCTEXT_NAMESPACE
+
+
diff --git a/Engine/Source/Runtime/Engine/Private/Components/SkyLightComponent.cpp b/Engine/Source/Runtime/Engine/Private/Components/SkyLightComponent.cpp
--- a/Engine/Source/Runtime/Engine/Private/Components/SkyLightComponent.cpp
+++ b/Engine/Source/Runtime/Engine/Private/Components/SkyLightComponent.cpp
@@ -186,6 +186,7 @@ FSkyLightSceneProxy::FSkyLightSceneProxy(const USkyLightComponent* InLightCompon
 	, IsDirtyImportanceSamplingData(true)
 #endif
 	, LightColor(FLinearColor(InLightComponent->LightColor) * InLightComponent->Intensity)
+	, bMovable(InLightComponent->IsMovable())
 {
 	const FSHVectorRGB3* InIrradianceEnvironmentMap = &InLightComponent->IrradianceEnvironmentMap;
 	const FSHVectorRGB3* BlendDestinationIrradianceEnvironmentMap = &InLightComponent->BlendDestinationIrradianceEnvironmentMap;
diff --git a/Engine/Source/Runtime/Engine/Private/MapBuildData.cpp b/Engine/Source/Runtime/Engine/Private/MapBuildData.cpp
--- a/Engine/Source/Runtime/Engine/Private/MapBuildData.cpp
+++ b/Engine/Source/Runtime/Engine/Private/MapBuildData.cpp
@@ -37,6 +37,12 @@ FArchive& operator<<(FArchive& Ar, FMeshMapBuildData& MeshMapBuildData)
 	return Ar;
 }
 
+FArchive& operator<<(FArchive& Ar, FSkyAtmosphereMapBuildData& Data)
+{
+	//Ar << Data.Dummy; // No serialisation needed
+	return Ar;
+}
+
 ULevel* UWorld::GetActiveLightingScenario() const
 {
 	for (int32 LevelIndex = 0; LevelIndex < Levels.Num(); LevelIndex++)
@@ -380,6 +386,11 @@ void UMapBuildDataRegistry::Serialize(FArchive& Ar)
 		{
 			Ar << ReflectionCaptureBuildData;
 		}
+
+		if (Ar.CustomVer(FRenderingObjectVersion::GUID) >= FRenderingObjectVersion::SkyAtmosphereStaticLightingVersioning)
+		{
+			Ar << SkyAtmosphereBuildData;
+		}
 	}
 }
 
@@ -615,6 +626,23 @@ FReflectionCaptureMapBuildData* UMapBuildDataRegistry::GetReflectionCaptureBuild
 	return ReflectionCaptureBuildData.Find(CaptureId);
 }
 
+FSkyAtmosphereMapBuildData& UMapBuildDataRegistry::FindOrAllocateSkyAtmosphereBuildData(const FGuid& Guid)
+{
+	check(Guid.IsValid());
+	return SkyAtmosphereBuildData.FindOrAdd(Guid);
+}
+
+const FSkyAtmosphereMapBuildData* UMapBuildDataRegistry::GetSkyAtmosphereBuildData(const FGuid& Guid) const
+{
+	check(Guid.IsValid());
+	return SkyAtmosphereBuildData.Find(Guid);
+}
+
+void UMapBuildDataRegistry::ClearSkyAtmosphereBuildData()
+{
+	SkyAtmosphereBuildData.Empty();
+}
+
 void UMapBuildDataRegistry::InvalidateStaticLighting(UWorld* World, bool bRecreateRenderState, const TSet<FGuid>* ResourcesToKeep)
 {
 	TUniquePtr<FGlobalComponentRecreateRenderStateContext> RecreateContext;
diff --git a/Engine/Source/Runtime/Engine/Private/Materials/HLSLMaterialTranslator.h b/Engine/Source/Runtime/Engine/Private/Materials/HLSLMaterialTranslator.h
--- a/Engine/Source/Runtime/Engine/Private/Materials/HLSLMaterialTranslator.h
+++ b/Engine/Source/Runtime/Engine/Private/Materials/HLSLMaterialTranslator.h
@@ -278,6 +278,8 @@ protected:
 	uint32 bNeedsSceneTexturePostProcessInputs : 1;
 	/** true if any atmospheric fog expressions are used */
 	uint32 bUsesAtmosphericFog : 1;
+	/** true if any SkyAtmosphere expressions are used */
+	uint32 bUsesSkyAtmosphere : 1;
 	/** true if the material reads vertex color in the pixel shader. */
 	uint32 bUsesVertexColor : 1;
 	/** true if the material reads particle color in the pixel shader. */
@@ -358,6 +360,7 @@ public:
 	,	bNeedsParticleSize(false)
 	,	bNeedsSceneTexturePostProcessInputs(false)
 	,	bUsesAtmosphericFog(false)
+	,	bUsesSkyAtmosphere(false)
 	,	bUsesVertexColor(false)
 	,	bUsesParticleColor(false)
 	,	bUsesParticleTransform(false)
@@ -1387,6 +1390,7 @@ ResourcesString = TEXT("");
 		
 		// @todo MetalMRT: Remove this hack and implement proper atmospheric-fog solution for Metal MRT...
 		OutEnvironment.SetDefine(TEXT("MATERIAL_ATMOSPHERIC_FOG"), !IsMetalMRTPlatform(InPlatform) ? bUsesAtmosphericFog : 0);
+		OutEnvironment.SetDefine(TEXT("MATERIAL_SKY_ATMOSPHERE"), bUsesSkyAtmosphere);
 		OutEnvironment.SetDefine(TEXT("INTERPOLATE_VERTEX_COLOR"), bUsesVertexColor);
 		OutEnvironment.SetDefine(TEXT("NEEDS_PARTICLE_COLOR"), bUsesParticleColor); 
 		OutEnvironment.SetDefine(TEXT("NEEDS_PARTICLE_TRANSFORM"), bUsesParticleTransform);
@@ -6552,6 +6556,42 @@ protected:
 
 	}
 
+	virtual int32 SkyAtmosphereLightIlluminance(int32 LightIndex) override
+	{
+		bUsesSkyAtmosphere = true;
+		return AddCodeChunk(MCT_Float3, TEXT("MaterialExpressionSkyAtmosphereLightIlluminance(Parameters, %d)"), LightIndex);
+	}
+
+	virtual int32 SkyAtmosphereLightDirection(int32 LightIndex) override
+	{
+		bUsesSkyAtmosphere = true;
+		return AddCodeChunk(MCT_Float3, TEXT("MaterialExpressionSkyAtmosphereLightDirection(Parameters, %d)"), LightIndex);
+	}
+
+	virtual int32 SkyAtmosphereLightDiskLuminance(int32 LightIndex) override
+	{
+		bUsesSkyAtmosphere = true;
+		return AddCodeChunk(MCT_Float3, TEXT("MaterialExpressionSkyAtmosphereLightDiskLuminance(Parameters, %d)"), LightIndex);
+	}
+
+	virtual int32 SkyAtmosphereViewLuminance() override
+	{
+		bUsesSkyAtmosphere = true;
+		return AddCodeChunk(MCT_Float3, TEXT("MaterialExpressionSkyAtmosphereViewLuminance(Parameters)"));
+	}
+
+	virtual int32 SkyAtmosphereAerialPerspective() override
+	{
+		bUsesSkyAtmosphere = true;
+		return AddCodeChunk(MCT_Float4, TEXT("MaterialExpressionSkyAtmosphereAerialPerspective(Parameters)"));
+	}
+
+	virtual int32 SkyAtmosphereDistantLightScatteredLuminance() override
+	{
+		bUsesSkyAtmosphere = true;
+		return AddCodeChunk(MCT_Float3, TEXT("MaterialExpressionSkyAtmosphereDistantLightScatteredLuminance(Parameters)"));
+	}
+
 	virtual int32 CustomPrimitiveData(int32 OutputIndex, EMaterialValueType Type) override
 	{
 		check(OutputIndex < FCustomPrimitiveData::NumCustomPrimitiveDataFloats);
diff --git a/Engine/Source/Runtime/Engine/Private/Materials/MaterialExpressions.cpp b/Engine/Source/Runtime/Engine/Private/Materials/MaterialExpressions.cpp
--- a/Engine/Source/Runtime/Engine/Private/Materials/MaterialExpressions.cpp
+++ b/Engine/Source/Runtime/Engine/Private/Materials/MaterialExpressions.cpp
@@ -214,6 +214,9 @@
 #include "Materials/MaterialExpressionClearCoatNormalCustomOutput.h"
 #include "Materials/MaterialExpressionAtmosphericLightVector.h"
 #include "Materials/MaterialExpressionAtmosphericLightColor.h"
+#include "Materials/MaterialExpressionSkyAtmosphereLightIlluminance.h"
+#include "Materials/MaterialExpressionSkyAtmosphereLightDirection.h"
+#include "Materials/MaterialExpressionSkyAtmosphereViewLuminance.h"
 #include "Materials/MaterialExpressionMaterialLayerOutput.h"
 #include "Materials/MaterialExpressionCurveAtlasRowParameter.h"
 #include "Materials/MaterialExpressionMapARPassthroughCameraUV.h"
@@ -15742,6 +15745,210 @@ void UMaterialExpressionAtmosphericLightColor::GetCaption(TArray<FString>& OutCa
 }
 #endif // WITH_EDITOR
 
+///////////////////////////////////////////////////////////////////////////////
+// UMaterialExpressionSkyAtmosphereLightIlluminance
+///////////////////////////////////////////////////////////////////////////////
+UMaterialExpressionSkyAtmosphereLightIlluminance::UMaterialExpressionSkyAtmosphereLightIlluminance(const FObjectInitializer& ObjectInitializer)
+	: Super(ObjectInitializer)
+{
+	// Structure to hold one-time initialization
+	struct FConstructorStatics
+	{
+		FText NAME_Sky;
+		FConstructorStatics()
+			: NAME_Sky(LOCTEXT("Sky", "Sky"))
+		{
+		}
+	};
+	static FConstructorStatics ConstructorStatics;
+
+#if WITH_EDITORONLY_DATA
+	MenuCategories.Add(ConstructorStatics.NAME_Sky);
+#endif
+}
+
+#if WITH_EDITOR
+int32 UMaterialExpressionSkyAtmosphereLightIlluminance::Compile(class FMaterialCompiler* Compiler, int32 OutputIndex)
+{
+	return Compiler->SkyAtmosphereLightIlluminance(LightIndex);
+}
+
+void UMaterialExpressionSkyAtmosphereLightIlluminance::GetCaption(TArray<FString>& OutCaptions) const
+{
+	OutCaptions.Add(FString::Printf(TEXT("SkyAtmosphereLightIlluminance[%i]"), LightIndex));
+}
+#endif // WITH_EDITOR
+
+///////////////////////////////////////////////////////////////////////////////
+// UMaterialExpressionSkyAtmosphereLightDirection
+///////////////////////////////////////////////////////////////////////////////
+UMaterialExpressionSkyAtmosphereLightDirection::UMaterialExpressionSkyAtmosphereLightDirection(const FObjectInitializer& ObjectInitializer)
+	: Super(ObjectInitializer)
+{
+	// Structure to hold one-time initialization
+	struct FConstructorStatics
+	{
+		FText NAME_Sky;
+		FConstructorStatics()
+			: NAME_Sky(LOCTEXT("Sky", "Sky"))
+		{
+		}
+	};
+	static FConstructorStatics ConstructorStatics;
+
+#if WITH_EDITORONLY_DATA
+	MenuCategories.Add(ConstructorStatics.NAME_Sky);
+#endif
+}
+
+#if WITH_EDITOR
+int32 UMaterialExpressionSkyAtmosphereLightDirection::Compile(class FMaterialCompiler* Compiler, int32 OutputIndex)
+{
+	return Compiler->SkyAtmosphereLightDirection(LightIndex);
+}
+
+void UMaterialExpressionSkyAtmosphereLightDirection::GetCaption(TArray<FString>& OutCaptions) const
+{
+	OutCaptions.Add(FString::Printf(TEXT("SkyAtmosphereLightDirection[%i]"), LightIndex));
+}
+#endif // WITH_EDITOR
+
+///////////////////////////////////////////////////////////////////////////////
+// UMaterialExpressionSkyAtmosphereLightDiskLuminance
+///////////////////////////////////////////////////////////////////////////////
+UMaterialExpressionSkyAtmosphereLightDiskLuminance::UMaterialExpressionSkyAtmosphereLightDiskLuminance(const FObjectInitializer& ObjectInitializer)
+	: Super(ObjectInitializer)
+{
+	// Structure to hold one-time initialization
+	struct FConstructorStatics
+	{
+		FText NAME_Sky;
+		FConstructorStatics()
+			: NAME_Sky(LOCTEXT("Sky", "Sky"))
+		{
+		}
+	};
+	static FConstructorStatics ConstructorStatics;
+
+#if WITH_EDITORONLY_DATA
+	MenuCategories.Add(ConstructorStatics.NAME_Sky);
+#endif
+}
+
+#if WITH_EDITOR
+int32 UMaterialExpressionSkyAtmosphereLightDiskLuminance::Compile(class FMaterialCompiler* Compiler, int32 OutputIndex)
+{
+	return Compiler->SkyAtmosphereLightDiskLuminance(LightIndex);
+}
+
+void UMaterialExpressionSkyAtmosphereLightDiskLuminance::GetCaption(TArray<FString>& OutCaptions) const
+{
+	OutCaptions.Add(FString::Printf(TEXT("SkyAtmosphereLightDiskLuminance[%i]"), LightIndex));
+}
+#endif // WITH_EDITOR
+
+///////////////////////////////////////////////////////////////////////////////
+// UMaterialExpressionSkyAtmosphereViewLuminance
+///////////////////////////////////////////////////////////////////////////////
+UMaterialExpressionSkyAtmosphereViewLuminance::UMaterialExpressionSkyAtmosphereViewLuminance(const FObjectInitializer& ObjectInitializer)
+	: Super(ObjectInitializer)
+{
+	// Structure to hold one-time initialization
+	struct FConstructorStatics
+	{
+		FText NAME_Sky;
+		FConstructorStatics()
+			: NAME_Sky(LOCTEXT("Sky", "Sky"))
+		{
+		}
+	};
+	static FConstructorStatics ConstructorStatics;
+
+#if WITH_EDITORONLY_DATA
+	MenuCategories.Add(ConstructorStatics.NAME_Sky);
+#endif
+}
+
+#if WITH_EDITOR
+int32 UMaterialExpressionSkyAtmosphereViewLuminance::Compile(class FMaterialCompiler* Compiler, int32 OutputIndex)
+{
+	return Compiler->SkyAtmosphereViewLuminance();
+}
+
+void UMaterialExpressionSkyAtmosphereViewLuminance::GetCaption(TArray<FString>& OutCaptions) const
+{
+	OutCaptions.Add(TEXT("SkyAtmosphereViewLuminance"));
+}
+#endif // WITH_EDITOR
+
+///////////////////////////////////////////////////////////////////////////////
+// UMaterialExpressionSkyAtmosphereAerialPerpsective
+///////////////////////////////////////////////////////////////////////////////
+UMaterialExpressionSkyAtmosphereAerialPerspective::UMaterialExpressionSkyAtmosphereAerialPerspective(const FObjectInitializer& ObjectInitializer)
+	: Super(ObjectInitializer)
+{
+	// Structure to hold one-time initialization
+	struct FConstructorStatics
+	{
+		FText NAME_Sky;
+		FConstructorStatics()
+			: NAME_Sky(LOCTEXT("Sky", "Sky"))
+		{
+		}
+	};
+	static FConstructorStatics ConstructorStatics;
+
+#if WITH_EDITORONLY_DATA
+	MenuCategories.Add(ConstructorStatics.NAME_Sky);
+#endif
+}
+
+#if WITH_EDITOR
+int32 UMaterialExpressionSkyAtmosphereAerialPerspective::Compile(class FMaterialCompiler* Compiler, int32 OutputIndex)
+{
+	return Compiler->SkyAtmosphereAerialPerspective();
+}
+
+void UMaterialExpressionSkyAtmosphereAerialPerspective::GetCaption(TArray<FString>& OutCaptions) const
+{
+	OutCaptions.Add(TEXT("SkyAtmosphereAerialPerspective"));
+}
+#endif // WITH_EDITOR
+
+///////////////////////////////////////////////////////////////////////////////
+// UMaterialExpressionSkyAtmosphereAerialPerpsective
+///////////////////////////////////////////////////////////////////////////////
+UMaterialExpressionSkyAtmosphereDistantLightScatteredLuminance::UMaterialExpressionSkyAtmosphereDistantLightScatteredLuminance(const FObjectInitializer& ObjectInitializer)
+	: Super(ObjectInitializer)
+{
+	// Structure to hold one-time initialization
+	struct FConstructorStatics
+	{
+		FText NAME_Sky;
+		FConstructorStatics()
+			: NAME_Sky(LOCTEXT("Sky", "Sky"))
+		{
+		}
+	};
+	static FConstructorStatics ConstructorStatics;
+
+#if WITH_EDITORONLY_DATA
+	MenuCategories.Add(ConstructorStatics.NAME_Sky);
+#endif
+}
+
+#if WITH_EDITOR
+int32 UMaterialExpressionSkyAtmosphereDistantLightScatteredLuminance::Compile(class FMaterialCompiler* Compiler, int32 OutputIndex)
+{
+	return Compiler->SkyAtmosphereDistantLightScatteredLuminance();
+}
+
+void UMaterialExpressionSkyAtmosphereDistantLightScatteredLuminance::GetCaption(TArray<FString>& OutCaptions) const
+{
+	OutCaptions.Add(TEXT("SkyAtmosphereDistantLightScatteredLuminance"));
+}
+#endif // WITH_EDITOR
+
 ///////////////////////////////////////////////////////////////////////////////
 // UMaterialExpressionPreSkinnedPosition
 ///////////////////////////////////////////////////////////////////////////////
diff --git a/Engine/Source/Runtime/Engine/Private/SceneManagement.cpp b/Engine/Source/Runtime/Engine/Private/SceneManagement.cpp
--- a/Engine/Source/Runtime/Engine/Private/SceneManagement.cpp
+++ b/Engine/Source/Runtime/Engine/Private/SceneManagement.cpp
@@ -735,6 +735,18 @@ FViewUniformShaderParameters::FViewUniformShaderParameters()
 	PreIntegratedBRDF = GWhiteTexture->TextureRHI;
 	PreIntegratedBRDFSampler = TStaticSamplerState<SF_Bilinear, AM_Clamp, AM_Clamp, AM_Clamp>::GetRHI();
 
+	TransmittanceLutTexture = GWhiteTexture->TextureRHI;
+	TransmittanceLutTextureSampler = TStaticSamplerState<SF_Bilinear, AM_Clamp, AM_Clamp, AM_Clamp>::GetRHI();
+
+	SkyViewLutTexture = GBlackTexture->TextureRHI;
+	SkyViewLutTextureSampler = TStaticSamplerState<SF_Bilinear, AM_Clamp, AM_Clamp, AM_Clamp>::GetRHI();
+
+	DistantSkyLightLutTexture = GBlackTexture->TextureRHI;
+	DistantSkyLightLutTextureSampler = TStaticSamplerState<SF_Bilinear, AM_Clamp, AM_Clamp, AM_Clamp>::GetRHI();
+
+	CameraAerialPerspectiveVolume = GBlackAlpha1VolumeTexture->TextureRHI;
+	CameraAerialPerspectiveVolumeSampler = TStaticSamplerState<SF_Bilinear, AM_Clamp, AM_Clamp, AM_Clamp>::GetRHI();
+
 	PrimitiveSceneData = GIdentityPrimitiveBuffer.PrimitiveSceneDataBufferSRV;
 	LightmapSceneData = GIdentityPrimitiveBuffer.LightmapSceneDataBufferSRV;
 }
diff --git a/Engine/Source/Runtime/Engine/Public/EngineDefines.h b/Engine/Source/Runtime/Engine/Public/EngineDefines.h
--- a/Engine/Source/Runtime/Engine/Public/EngineDefines.h
+++ b/Engine/Source/Runtime/Engine/Public/EngineDefines.h
@@ -17,6 +17,9 @@
 // Whether lightmass generates FSHVector2 or FSHVector3. Linked with VER_UE4_INDIRECT_LIGHTING_SH3
 #define NUM_INDIRECT_LIGHTING_SH_COEFFICIENTS 9
 
+// The number of lights to consider for sky/atmospheric light scattering
+#define NUM_ATMOSPHERE_LIGHTS 2
+
 /*-----------------------------------------------------------------------------
 	Size of the world.
 -----------------------------------------------------------------------------*/
diff --git a/Engine/Source/Runtime/Engine/Public/MaterialCompiler.h b/Engine/Source/Runtime/Engine/Public/MaterialCompiler.h
--- a/Engine/Source/Runtime/Engine/Public/MaterialCompiler.h
+++ b/Engine/Source/Runtime/Engine/Public/MaterialCompiler.h
@@ -335,6 +335,12 @@ public:
 	virtual int32 EyeAdaptation() = 0;
 	virtual int32 AtmosphericLightVector() = 0;
 	virtual int32 AtmosphericLightColor() = 0;
+	virtual int32 SkyAtmosphereLightIlluminance(int32 LightIndex) = 0;
+	virtual int32 SkyAtmosphereLightDirection(int32 LightIndex) = 0;
+	virtual int32 SkyAtmosphereLightDiskLuminance(int32 LightIndex) = 0;
+	virtual int32 SkyAtmosphereViewLuminance() = 0;
+	virtual int32 SkyAtmosphereAerialPerspective() = 0;
+	virtual int32 SkyAtmosphereDistantLightScatteredLuminance() = 0;
 	virtual int32 CustomPrimitiveData(int32 OutputIndex, EMaterialValueType Type) = 0;
 	virtual int32 ShadingModel(EMaterialShadingModel InSelectedShadingModel) = 0;
 
@@ -604,6 +610,36 @@ public:
 	{
 		return Compiler->AtmosphericLightColor();
 	}
+
+	virtual int32 SkyAtmosphereLightIlluminance(int32 LightIndex) override
+	{
+		return Compiler->SkyAtmosphereLightIlluminance(LightIndex);
+	}
+
+	virtual int32 SkyAtmosphereLightDirection(int32 LightIndex) override
+	{
+		return Compiler->SkyAtmosphereLightDirection(LightIndex);
+	}
+
+	virtual int32 SkyAtmosphereLightDiskLuminance(int32 LightIndex) override
+	{
+		return Compiler->SkyAtmosphereLightDiskLuminance(LightIndex);
+	}
+
+	virtual int32 SkyAtmosphereViewLuminance() override
+	{
+		return Compiler->SkyAtmosphereViewLuminance();
+	}
+
+	virtual int32 SkyAtmosphereAerialPerspective() override
+	{
+		return Compiler->SkyAtmosphereAerialPerspective();
+	}
+
+	virtual int32 SkyAtmosphereDistantLightScatteredLuminance() override
+	{
+		return Compiler->SkyAtmosphereDistantLightScatteredLuminance();
+	}
 	
 	virtual int32 CustomPrimitiveData(int32 OutputIndex, EMaterialValueType Type) override
 	{
diff --git a/Engine/Source/Runtime/Engine/Public/Rendering/SkyAtmosphereCommonData.cpp b/Engine/Source/Runtime/Engine/Public/Rendering/SkyAtmosphereCommonData.cpp
new file mode 100644
--- /dev/null
+++ b/Engine/Source/Runtime/Engine/Public/Rendering/SkyAtmosphereCommonData.cpp
@@ -0,0 +1,148 @@
+// Copyright 1998-2019 Epic Games, Inc. All Rights Reserved.
+
+/*=============================================================================
+	SkyAtmosphereCommonData.cpp
+=============================================================================*/
+
+#include "SkyAtmosphereCommonData.h"
+
+#include "Components/SkyAtmosphereComponent.h"
+
+//#pragma optimize( "", off )
+
+
+FAtmosphereSetup::FAtmosphereSetup(const USkyAtmosphereComponent& SkyAtmosphereComponent)
+{
+	// For now we convert color as if they were linear because this is how scattering coefficient should be considered for instance. Coefficients are not a perceptual value.
+	auto LinearFColorToLinearColor = [](const FColor& c) {return FLinearColor(float(c.R) / 255.0f, float(c.G) / 255.0f, float(c.B) / 255.0f, float(c.A) / 255.0f); };
+
+	// Convert Tent distribution to linear curve coefficients.
+	auto TentToCoefficients = [](const FTentDistribution& Tent, float& LayerWidth, float& LinTerm0, float&  LinTerm1, float& ConstTerm0, float& ConstTerm1)
+	{
+		if (Tent.Width > 0.0f && Tent.TipValue > 0.0f)
+		{
+			const float px = Tent.TipAltitude;
+			const float py = Tent.TipValue;
+			const float slope = Tent.TipValue / Tent.Width;
+			LayerWidth = px;
+			LinTerm0 = slope;
+			LinTerm1 = -slope;
+			ConstTerm0 = py - px * LinTerm0;
+			ConstTerm1 = py - px * LinTerm1;
+		}
+		else
+		{
+			LayerWidth = 0.0f;
+			LinTerm0 = 0.0f;
+			LinTerm1 = 0.0f;
+			ConstTerm0 = 0.0f;
+			ConstTerm1 = 0.0f;
+		}
+	};
+
+	BottomRadius = SkyAtmosphereComponent.BottomRadius;
+	TopRadius = SkyAtmosphereComponent.BottomRadius + SkyAtmosphereComponent.AtmosphereHeight;
+	GroundAlbedo = FLinearColor(SkyAtmosphereComponent.GroundAlbedo);
+	MultiScatteringFactor = SkyAtmosphereComponent.MultiScatteringFactor;
+
+	RayleighDensityExpScale = -1.0f / SkyAtmosphereComponent.RayleighExponentialDistribution;
+	RayleighScattering = LinearFColorToLinearColor(SkyAtmosphereComponent.RayleighScattering) * SkyAtmosphereComponent.RayleighScatteringScale;
+
+	MieScattering = LinearFColorToLinearColor(SkyAtmosphereComponent.MieScattering) * SkyAtmosphereComponent.MieScatteringScale;
+	MieAbsorption = LinearFColorToLinearColor(SkyAtmosphereComponent.MieAbsorption) * SkyAtmosphereComponent.MieAbsorptionScale;
+	MieExtinction = MieScattering + MieAbsorption;
+	MiePhaseG = SkyAtmosphereComponent.MieAnisotropy;
+	MieDensityExpScale = -1.0f / SkyAtmosphereComponent.MieExponentialDistribution;
+
+	AbsorptionExtinction = LinearFColorToLinearColor(SkyAtmosphereComponent.OtherAbsorption) * SkyAtmosphereComponent.OtherAbsorptionScale;
+	TentToCoefficients(SkyAtmosphereComponent.OtherTentDistribution, AbsorptionDensity0LayerWidth, AbsorptionDensity0LinearTerm, AbsorptionDensity1LinearTerm, AbsorptionDensity0ConstantTerm, AbsorptionDensity1ConstantTerm);
+}
+
+FLinearColor FAtmosphereSetup::GetTransmittanceAtGroundLevel(const FVector& SunDirection) const
+{
+	// The following code is from SkyAtmosphere.usf and has been converted to lambda functions. 
+	// It compute transmittance from the origin towards a sun direction. 
+
+	auto RayIntersectSphere = [&](FVector RayOrigin, FVector RayDirection, FVector SphereOrigin, float SphereRadius)
+	{
+		FVector LocalPosition = RayOrigin - SphereOrigin;
+		float LocalPositionSqr = FVector::DotProduct(LocalPosition, LocalPosition);
+
+		FVector QuadraticCoef;
+		QuadraticCoef.X = FVector::DotProduct(RayDirection, RayDirection);
+		QuadraticCoef.Y = 2.0f * FVector::DotProduct(RayDirection, LocalPosition);
+		QuadraticCoef.Z = LocalPositionSqr - SphereRadius * SphereRadius;
+
+		float Discriminant = QuadraticCoef.Y * QuadraticCoef.Y - 4.0f * QuadraticCoef.X * QuadraticCoef.Z;
+
+		// Only continue if the ray intersects the sphere
+		FVector2D Intersections = { -1.0f, -1.0f };
+		if (Discriminant >= 0)
+		{
+			float SqrtDiscriminant = sqrt(Discriminant);
+			Intersections.X = (-QuadraticCoef.Y - 1.0f * SqrtDiscriminant) / (2 * QuadraticCoef.X);
+			Intersections.Y = (-QuadraticCoef.Y + 1.0f * SqrtDiscriminant) / (2 * QuadraticCoef.X);
+		}
+		return Intersections;
+	};
+
+	// Nearest intersection of ray r,mu with sphere boundary
+	auto raySphereIntersectNearest = [&](FVector RayOrigin, FVector RayDirection, FVector SphereOrigin, float SphereRadius)
+	{
+		FVector2D sol = RayIntersectSphere(RayOrigin, RayDirection, SphereOrigin, SphereRadius);
+		const float sol0 = sol.X;
+		const float sol1 = sol.Y;
+		if (sol0 < 0.0f && sol1 < 0.0f)
+		{
+			return -1.0f;
+		}
+		if (sol0 < 0.0f)
+		{
+			return FMath::Max(0.0f, sol1);
+		}
+		else if (sol1 < 0.0f)
+		{
+			return FMath::Max(0.0f, sol0);
+		}
+		return FMath::Max(0.0f, FMath::Min(sol0, sol1));
+	};
+
+	auto OpticalDepth = [&](FVector RayOrigin, FVector RayDirection)
+	{
+		float TMax = raySphereIntersectNearest(RayOrigin, RayDirection, FVector(0.0f, 0.0f, 0.0f), TopRadius);
+
+		FLinearColor OpticalDepthRGB = FLinearColor(ForceInitToZero);
+		FVector VectorZero = FVector(ForceInitToZero);
+		if (TMax > 0.0f)
+		{
+			const float SampleCount = 15.0f;
+			const float SampleStep = 1.0f / 15.0f;
+			const float SampleLength = SampleStep * TMax;
+			for (float SampleT = 0.0f; SampleT < 1.0f; SampleT += SampleStep)
+			{
+				FVector Pos = RayOrigin + RayDirection * (TMax * SampleT);
+				const float viewHeight = (FVector::Distance(Pos, VectorZero) - BottomRadius);
+
+				const float densityMie = FMath::Max(0.0f, FMath::Exp(MieDensityExpScale * viewHeight));
+				const float densityRay = FMath::Max(0.0f, FMath::Exp(RayleighDensityExpScale * viewHeight));
+				const float densityOzo = FMath::Clamp(viewHeight < AbsorptionDensity0LayerWidth ?
+					AbsorptionDensity0LinearTerm * viewHeight + AbsorptionDensity0ConstantTerm :
+					AbsorptionDensity1LinearTerm * viewHeight + AbsorptionDensity1ConstantTerm,
+					0.0f, 1.0f);
+
+				FLinearColor SampleExtinction = densityMie * MieExtinction + densityRay * RayleighScattering + densityOzo * AbsorptionExtinction;
+				OpticalDepthRGB += SampleLength * SampleExtinction;
+			}
+		}
+
+		return OpticalDepthRGB;
+	};
+
+	// Assuming camera is along Z on (0,0,earthRadius + 500m)
+	const FVector WorldPos = FVector(0.0f, 0.0f, BottomRadius + 0.5);
+	const FVector WorldDir = SunDirection;
+	FLinearColor OpticalDepthRGB = OpticalDepth(WorldPos, WorldDir);
+	return FLinearColor(FMath::Exp(-OpticalDepthRGB.R), FMath::Exp(-OpticalDepthRGB.G), FMath::Exp(-OpticalDepthRGB.B));
+}
+
+
diff --git a/Engine/Source/Runtime/Engine/Public/SceneInterface.h b/Engine/Source/Runtime/Engine/Public/SceneInterface.h
--- a/Engine/Source/Runtime/Engine/Public/SceneInterface.h
+++ b/Engine/Source/Runtime/Engine/Public/SceneInterface.h
@@ -10,6 +10,7 @@
 
 class AWorldSettings;
 class FAtmosphericFogSceneInfo;
+class FSkyAtmosphereRenderSceneInfo;
 class FMaterial;
 class FMaterialShaderMap;
 class FPrimitiveSceneInfo;
@@ -103,6 +104,10 @@ public:
 	virtual void AddInvisibleLight(ULightComponent* Light) = 0;
 	virtual void SetSkyLight(FSkyLightSceneProxy* Light) = 0;
 	virtual void DisableSkyLight(FSkyLightSceneProxy* Light) = 0;
+
+	virtual bool HasSkyLightRequiringLightingBuild() const = 0;
+	virtual bool HasAtmosphereLightRequiringLightingBuild() const = 0;
+
 	/** 
 	 * Adds a new decal component to the scene
 	 * 
@@ -240,6 +245,23 @@ public:
 	 */
 	virtual FAtmosphericFogSceneInfo* GetAtmosphericFogSceneInfo() = 0;
 
+	/**
+	 * Adds the unique sky atmosphere component to the scene
+	 *
+	 * @param SkyAtmosphereComponent - component to add
+	 */
+	virtual void AddSkyAtmosphere(const class USkyAtmosphereComponent* SkyAtmosphereComponent, bool bStaticLightingBuilt) = 0;
+	/**
+	 * Removes the unique sky atmosphere component to the scene
+	 *
+	 * @param SkyAtmosphereComponent - component to remove
+	 */
+	virtual void RemoveSkyAtmosphere(const class USkyAtmosphereComponent* SkyAtmosphereComponent) = 0;
+	/**
+	 * Returns the scene's unique FSkyAtmosphereRenderSceneInfo if it exists
+	 */
+	virtual FSkyAtmosphereRenderSceneInfo* GetSkyAtmosphereSceneInfo() = 0;
+
 	/**
 	 * Adds a wind source component to the scene.
 	 * @param WindComponent - The component to add.
diff --git a/Engine/Source/Runtime/Engine/Public/SceneManagement.h b/Engine/Source/Runtime/Engine/Public/SceneManagement.h
--- a/Engine/Source/Runtime/Engine/Public/SceneManagement.h
+++ b/Engine/Source/Runtime/Engine/Public/SceneManagement.h
@@ -1052,6 +1052,8 @@ public:
 	FLinearColor OcclusionTint;
 	int32 SamplesPerPixel;
 
+	bool IsMovable() { return bMovable; }
+
 #if RHI_RAYTRACING
 	bool IsDirtyImportanceSamplingData;
 	bool ShouldRebuildCdf() const;
@@ -1085,6 +1087,7 @@ public:
 
 private:
 	FLinearColor LightColor;
+	const uint8 bMovable : 1;
 };
 
 
@@ -1359,8 +1362,10 @@ public:
 	// Atmosphere / Fog related functions.
 
 	inline bool IsUsedAsAtmosphereSunLight() const { return bUsedAsAtmosphereSunLight; }
+	inline uint8 GetAtmosphereSunLightIndex() const { return AtmosphereSunLightIndex; }
 	virtual void SetAtmosphereRelatedProperties(FLinearColor TransmittanceFactor, FLinearColor SunOuterSpaceLuminance) {}
 	virtual FLinearColor GetOuterSpaceLuminance() const { return FLinearColor::White; }
+	virtual FLinearColor GetTransmittanceFactor() const { return FLinearColor::White; }
 	static float GetSunOnEarthHalfApexAngleRadian() 
 	{ 
 		const float SunOnEarthApexAngleDegree = 0.545f;	// Apex angle == angular diameter
@@ -1512,6 +1517,9 @@ protected:
 	/** Whether the light supports rendering in tiled deferred pass */
 	uint8 bTiledDeferredLightingSupported : 1;
 
+	/** The index of the atmospheric light. Multiple lights can be considered when computing the sky/atmospheric scattering. */
+	const uint8 AtmosphereSunLightIndex;
+
 	/** The light type (ELightComponentType) */
 	const uint8 LightType;
 
diff --git a/Engine/Source/Runtime/Engine/Public/SceneView.h b/Engine/Source/Runtime/Engine/Public/SceneView.h
--- a/Engine/Source/Runtime/Engine/Public/SceneView.h
+++ b/Engine/Source/Runtime/Engine/Public/SceneView.h
@@ -654,7 +654,6 @@ enum ETranslucencyVolumeCascade
 	VIEW_UNIFORM_BUFFER_MEMBER_EX(float, DemosaicVposOffset, EShaderPrecisionModifier::Half) \
 	VIEW_UNIFORM_BUFFER_MEMBER(FVector, IndirectLightingColorScale) \
 	VIEW_UNIFORM_BUFFER_MEMBER_EX(float, HDR32bppEncodingMode, EShaderPrecisionModifier::Half) \
-	VIEW_UNIFORM_BUFFER_MEMBER(FVector, AtmosphericFogSunDirection) \
 	VIEW_UNIFORM_BUFFER_MEMBER_EX(float, AtmosphericFogSunPower, EShaderPrecisionModifier::Half) \
 	VIEW_UNIFORM_BUFFER_MEMBER_EX(float, AtmosphericFogPower, EShaderPrecisionModifier::Half) \
 	VIEW_UNIFORM_BUFFER_MEMBER_EX(float, AtmosphericFogDensityScale, EShaderPrecisionModifier::Half) \
@@ -666,11 +665,12 @@ enum ETranslucencyVolumeCascade
 	VIEW_UNIFORM_BUFFER_MEMBER_EX(float, AtmosphericFogStartDistance, EShaderPrecisionModifier::Half) \
 	VIEW_UNIFORM_BUFFER_MEMBER_EX(float, AtmosphericFogDistanceOffset, EShaderPrecisionModifier::Half) \
 	VIEW_UNIFORM_BUFFER_MEMBER_EX(float, AtmosphericFogSunDiscScale, EShaderPrecisionModifier::Half) \
-	VIEW_UNIFORM_BUFFER_MEMBER_EX(float, AtmosphericFogSunDiscHalfApexAngleRadian, EShaderPrecisionModifier::Half) \
-	VIEW_UNIFORM_BUFFER_MEMBER_EX(FLinearColor, AtmosphericFogSunDiscLuminance, EShaderPrecisionModifier::Half) \
+	VIEW_UNIFORM_BUFFER_MEMBER_ARRAY(FVector4, AtmosphereLightDirection, [NUM_ATMOSPHERE_LIGHTS]) \
+	VIEW_UNIFORM_BUFFER_MEMBER_ARRAY(FLinearColor, AtmosphereLightColor, [NUM_ATMOSPHERE_LIGHTS]) \
+	VIEW_UNIFORM_BUFFER_MEMBER_ARRAY(FLinearColor, AtmosphereLightDiscLuminance, [NUM_ATMOSPHERE_LIGHTS]) \
+	VIEW_UNIFORM_BUFFER_MEMBER_ARRAY(float, AtmosphereLightDiscCosHalfApexAngle, [NUM_ATMOSPHERE_LIGHTS]) \
 	VIEW_UNIFORM_BUFFER_MEMBER(uint32, AtmosphericFogRenderMask) \
 	VIEW_UNIFORM_BUFFER_MEMBER(uint32, AtmosphericFogInscatterAltitudeSampleNum) \
-	VIEW_UNIFORM_BUFFER_MEMBER(FLinearColor, AtmosphericFogSunColor) \
 	VIEW_UNIFORM_BUFFER_MEMBER(FVector, NormalCurvatureToRoughnessScaleBias) \
 	VIEW_UNIFORM_BUFFER_MEMBER(float, RenderingReflectionCaptureMask) \
 	VIEW_UNIFORM_BUFFER_MEMBER(FLinearColor, AmbientCubemapTint) \
@@ -707,7 +707,16 @@ enum ETranslucencyVolumeCascade
 	VIEW_UNIFORM_BUFFER_MEMBER(float, EyeToPixelSpreadAngle) \
 	VIEW_UNIFORM_BUFFER_MEMBER(FMatrix, WorldToVirtualTexture) \
 	VIEW_UNIFORM_BUFFER_MEMBER(FVector4, VirtualTextureParams) \
-	VIEW_UNIFORM_BUFFER_MEMBER_ARRAY(FVector4, XRPassthroughCameraUVs, [2])
+	VIEW_UNIFORM_BUFFER_MEMBER_ARRAY(FVector4, XRPassthroughCameraUVs, [2]) \
+	VIEW_UNIFORM_BUFFER_MEMBER(FVector4, SkyViewLutSizeAndInvSize) \
+	VIEW_UNIFORM_BUFFER_MEMBER(float, SkyAtmosphereBottomRadius) \
+	VIEW_UNIFORM_BUFFER_MEMBER(float, SkyAtmosphereTopRadius) \
+	VIEW_UNIFORM_BUFFER_MEMBER(float, SkyAtmosphereCameraAerialPerspectiveVolumeStartDepth) \
+	VIEW_UNIFORM_BUFFER_MEMBER(float, SkyAtmosphereCameraAerialPerspectiveVolumeDepthResolution) \
+	VIEW_UNIFORM_BUFFER_MEMBER(float, SkyAtmosphereCameraAerialPerspectiveVolumeDepthResolutionInv) \
+	VIEW_UNIFORM_BUFFER_MEMBER(float, SkyAtmosphereCameraAerialPerspectiveVolumeDepthSliceLength) \
+	VIEW_UNIFORM_BUFFER_MEMBER(float, SkyAtmosphereCameraAerialPerspectiveVolumeDepthSliceLengthInv) \
+	VIEW_UNIFORM_BUFFER_MEMBER(float, SkyAtmosphereApplyCameraAerialPerspectiveVolume) \
 
 #define VIEW_UNIFORM_BUFFER_MEMBER(type, identifier) \
 	SHADER_PARAMETER(type, identifier)
@@ -779,6 +788,15 @@ BEGIN_GLOBAL_SHADER_PARAMETER_STRUCT_WITH_CONSTRUCTOR(FViewUniformShaderParamete
 	SHADER_PARAMETER_SRV(StructuredBuffer<float4>, PrimitiveSceneData)
 	SHADER_PARAMETER_SRV(StructuredBuffer<float4>, LightmapSceneData)
 
+	SHADER_PARAMETER_TEXTURE(Texture2D, TransmittanceLutTexture)
+	SHADER_PARAMETER_SAMPLER(SamplerState, TransmittanceLutTextureSampler)
+	SHADER_PARAMETER_TEXTURE(Texture2D, SkyViewLutTexture)
+	SHADER_PARAMETER_SAMPLER(SamplerState, SkyViewLutTextureSampler)
+	SHADER_PARAMETER_TEXTURE(Texture2D, DistantSkyLightLutTexture)
+	SHADER_PARAMETER_SAMPLER(SamplerState, DistantSkyLightLutTextureSampler)
+	SHADER_PARAMETER_TEXTURE(Texture3D, CameraAerialPerspectiveVolume)
+	SHADER_PARAMETER_SAMPLER(SamplerState, CameraAerialPerspectiveVolumeSampler)
+
 END_GLOBAL_SHADER_PARAMETER_STRUCT()
 
 /** Copy of the view uniform shader parameters associated with a view for instanced stereo. */
diff --git a/Engine/Source/Runtime/OpenGLDrv/Private/OpenGLCommands.cpp b/Engine/Source/Runtime/OpenGLDrv/Private/OpenGLCommands.cpp
--- a/Engine/Source/Runtime/OpenGLDrv/Private/OpenGLCommands.cpp
+++ b/Engine/Source/Runtime/OpenGLDrv/Private/OpenGLCommands.cpp
@@ -1104,7 +1104,7 @@ void FOpenGLDynamicRHI::RHISetShaderSampler(FRHIGeometryShader* GeometryShaderRH
 
 void FOpenGLDynamicRHI::RHISetShaderTexture(FRHIComputeShader* ComputeShaderRHI,uint32 TextureIndex, FRHITexture* NewTextureRHI)
 {
-	check(GMaxRHIFeatureLevel >= ERHIFeatureLevel::SM5);
+	check(FOpenGL::SupportsComputeShaders());
 	VERIFY_GL_SCOPE();
 	FOpenGLTextureBase* NewTexture = GetOpenGLTextureFromRHITexture(NewTextureRHI);
 	ensureMsgf((int32)TextureIndex < FOpenGL::GetMaxComputeTextureImageUnits(), TEXT("Using more compute texture units (%d) than allowed (%d)!"), TextureIndex, FOpenGL::GetMaxComputeTextureImageUnits());
@@ -1192,7 +1192,7 @@ void FOpenGLDynamicRHI::RHISetShaderUniformBuffer(FRHIPixelShader* PixelShaderRH
 void FOpenGLDynamicRHI::RHISetShaderUniformBuffer(FRHIComputeShader* ComputeShaderRHI,uint32 BufferIndex, FRHIUniformBuffer* BufferRHI)
 {
 	VERIFY_GL_SCOPE();
-	check(GMaxRHIFeatureLevel >= ERHIFeatureLevel::SM5);
+	check(FOpenGL::SupportsComputeShaders());
 	PendingState.BoundUniformBuffers[SF_Compute][BufferIndex] = BufferRHI;
 	PendingState.DirtyUniformBuffers[SF_Compute] |= 1 << BufferIndex;
 }
diff --git a/Engine/Source/Runtime/OpenGLDrv/Private/OpenGLDevice.cpp b/Engine/Source/Runtime/OpenGLDrv/Private/OpenGLDevice.cpp
--- a/Engine/Source/Runtime/OpenGLDrv/Private/OpenGLDevice.cpp
+++ b/Engine/Source/Runtime/OpenGLDrv/Private/OpenGLDevice.cpp
@@ -1050,6 +1050,7 @@ static void InitRHICapabilitiesForGL()
 			GLenum InternalFormatRGBA16 = FOpenGL::GetTextureHalfFloatInternalFormat();
 			GLenum PixelTypeRGBA16 = FOpenGL::GetTextureHalfFloatPixelType();
 			SetupTextureFormat(PF_FloatRGBA, FOpenGLTextureFormat(InternalFormatRGBA16, InternalFormatRGBA16, GL_RGBA, PixelTypeRGBA16, false, false));
+			SetupTextureFormat(PF_FloatRGB, FOpenGLTextureFormat(GL_RGB16F, GL_RGB16F, GL_RGB, GL_HALF_FLOAT, false, false));
 #else
 			SetupTextureFormat(PF_FloatRGBA, FOpenGLTextureFormat(GL_RGBA, GL_RGBA, GL_RGBA, GL_HALF_FLOAT_OES, false, false));
 #endif
diff --git a/Engine/Source/Runtime/RenderCore/Private/RenderTargetPool.cpp b/Engine/Source/Runtime/RenderCore/Private/RenderTargetPool.cpp
--- a/Engine/Source/Runtime/RenderCore/Private/RenderTargetPool.cpp
+++ b/Engine/Source/Runtime/RenderCore/Private/RenderTargetPool.cpp
@@ -569,7 +569,7 @@ Done:
 		if((Desc.TargetableFlags & TexCreate_UAV) && !bDeferTextureAllocation)
 		{
 			// The render target desc is invalid if a UAV is requested with an RHI that doesn't support the high-end feature level.
-			check(GMaxRHIFeatureLevel == ERHIFeatureLevel::SM5);
+			check(GMaxRHIFeatureLevel == ERHIFeatureLevel::SM5 || GMaxRHIFeatureLevel == ERHIFeatureLevel::ES3_1);
 			Found->RenderTargetItem.MipUAVs.Reserve(Desc.NumMips);
 			for (uint32 MipLevel = 0; MipLevel < Desc.NumMips; MipLevel++)
 			{
diff --git a/Engine/Source/Runtime/RenderCore/Private/RenderUtils.cpp b/Engine/Source/Runtime/RenderCore/Private/RenderUtils.cpp
--- a/Engine/Source/Runtime/RenderCore/Private/RenderUtils.cpp
+++ b/Engine/Source/Runtime/RenderCore/Private/RenderUtils.cpp
@@ -244,8 +244,8 @@ class FBlackVolumeTextureResourceBulkDataInterface : public FResourceBulkDataInt
 public:
 
 	/** Default constructor. */
-	FBlackVolumeTextureResourceBulkDataInterface()
-		: Color(0)
+	FBlackVolumeTextureResourceBulkDataInterface(uint8 Alpha)
+		: Color(0, 0, 0, Alpha)
 	{
 	}
 
@@ -281,7 +281,7 @@ private:
 /**
  * A class representing a 1x1x1 black volume texture.
  */
-template <EPixelFormat PixelFormat>
+template <EPixelFormat PixelFormat, uint8 Alpha>
 class FBlackVolumeTexture : public FTexture
 {
 public:
@@ -294,7 +294,7 @@ public:
 		if (GSupportsTexture3D)
 		{
 			// Create the texture.
-			FBlackVolumeTextureResourceBulkDataInterface BlackTextureBulkData;
+			FBlackVolumeTextureResourceBulkDataInterface BlackTextureBulkData(Alpha);
 			FRHIResourceCreateInfo CreateInfo(&BlackTextureBulkData);
 			CreateInfo.DebugName = TEXT("BlackVolumeTexture");
 			FTexture3DRHIRef Texture3D = RHICreateTexture3D(1,1,1,PixelFormat,1,TexCreate_ShaderResource,CreateInfo);
@@ -303,7 +303,7 @@ public:
 		else
 		{
 			// Create a texture, even though it's not a volume texture
-			FBlackVolumeTextureResourceBulkDataInterface BlackTextureBulkData;
+			FBlackVolumeTextureResourceBulkDataInterface BlackTextureBulkData(Alpha);
 			FRHIResourceCreateInfo CreateInfo(&BlackTextureBulkData);
 			FTexture2DRHIRef Texture2D = RHICreateTexture2D(1, 1, PixelFormat, 1, 1, TexCreate_ShaderResource, CreateInfo);
 			TextureRHI = Texture2D;
@@ -332,10 +332,11 @@ public:
 };
 
 /** Global black volume texture resource. */
-FTexture* GBlackVolumeTexture = new TGlobalResource<FBlackVolumeTexture<PF_B8G8R8A8>>();
+FTexture* GBlackVolumeTexture = new TGlobalResource<FBlackVolumeTexture<PF_B8G8R8A8, 0>>();
+FTexture* GBlackAlpha1VolumeTexture = new TGlobalResource<FBlackVolumeTexture<PF_B8G8R8A8, 255>>();
 
 /** Global black volume texture resource. */
-FTexture* GBlackUintVolumeTexture = new TGlobalResource<FBlackVolumeTexture<PF_R8G8B8A8_UINT>>();
+FTexture* GBlackUintVolumeTexture = new TGlobalResource<FBlackVolumeTexture<PF_R8G8B8A8_UINT, 0>>();
 
 class FBlackArrayTexture : public FTexture
 {
@@ -346,7 +347,7 @@ public:
 		if (GetFeatureLevel() >= ERHIFeatureLevel::SM4)
 		{
 			// Create the texture RHI.
-			FBlackVolumeTextureResourceBulkDataInterface BlackTextureBulkData;
+			FBlackVolumeTextureResourceBulkDataInterface BlackTextureBulkData(0);
 			FRHIResourceCreateInfo CreateInfo(&BlackTextureBulkData);
 			CreateInfo.DebugName = TEXT("BlackArrayTexture");
 			FTexture2DArrayRHIRef TextureArray = RHICreateTexture2DArray(1, 1, 1, PF_B8G8R8A8, 1, 1, TexCreate_ShaderResource, CreateInfo);
diff --git a/Engine/Source/Runtime/RenderCore/Public/PixelShaderUtils.h b/Engine/Source/Runtime/RenderCore/Public/PixelShaderUtils.h
--- a/Engine/Source/Runtime/RenderCore/Public/PixelShaderUtils.h
+++ b/Engine/Source/Runtime/RenderCore/Public/PixelShaderUtils.h
@@ -37,13 +37,20 @@ struct RENDERCORE_API FPixelShaderUtils
 		const TShaderMap<FGlobalShaderType>* GlobalShaderMap,
 		const TShaderClass* PixelShader,
 		const typename TShaderClass::FParameters& Parameters,
-		const FIntRect& Viewport)
+		const FIntRect& Viewport,
+		FRHIBlendState* BlendState = nullptr,
+		FRHIRasterizerState* RasterizerState = nullptr,
+		FRHIDepthStencilState* DepthStencilState = nullptr)
 	{
 		check(PixelShader);
 		RHICmdList.SetViewport(Viewport.Min.X, Viewport.Min.Y, 0.0f, Viewport.Max.X, Viewport.Max.Y, 1.0f);
 		
 		FGraphicsPipelineStateInitializer GraphicsPSOInit;
 		InitFullscreenPipelineState(RHICmdList, GlobalShaderMap, PixelShader, /* out */ GraphicsPSOInit);
+		GraphicsPSOInit.BlendState = BlendState ? BlendState : GraphicsPSOInit.BlendState;
+		GraphicsPSOInit.RasterizerState = RasterizerState ? RasterizerState : GraphicsPSOInit.RasterizerState;
+		GraphicsPSOInit.DepthStencilState = DepthStencilState ? DepthStencilState : GraphicsPSOInit.DepthStencilState;
+
 		SetGraphicsPipelineState(RHICmdList, GraphicsPSOInit);
 
 		SetShaderParameters(RHICmdList, PixelShader, PixelShader->GetPixelShader(), Parameters);
@@ -59,7 +66,10 @@ struct RENDERCORE_API FPixelShaderUtils
 		FRDGEventName&& PassName,
 		const TShaderClass* PixelShader,
 		typename TShaderClass::FParameters* Parameters,
-		const FIntRect& Viewport)
+		const FIntRect& Viewport,
+		FRHIBlendState* BlendState = nullptr,
+		FRHIRasterizerState* RasterizerState = nullptr,
+		FRHIDepthStencilState* DepthStencilState = nullptr)
 	{
 		check(PixelShader);
 		ClearUnusedGraphResources(PixelShader, Parameters);
@@ -68,9 +78,10 @@ struct RENDERCORE_API FPixelShaderUtils
 			Forward<FRDGEventName>(PassName),
 			Parameters,
 			ERDGPassFlags::Raster,
-			[Parameters, GlobalShaderMap, PixelShader, Viewport](FRHICommandList& RHICmdList)
+			[Parameters, GlobalShaderMap, PixelShader, Viewport, BlendState, RasterizerState, DepthStencilState](FRHICommandList& RHICmdList)
 		{
-			FPixelShaderUtils::DrawFullscreenPixelShader(RHICmdList, GlobalShaderMap, PixelShader, *Parameters, Viewport);
+			FPixelShaderUtils::DrawFullscreenPixelShader(RHICmdList, GlobalShaderMap, PixelShader, *Parameters, Viewport, 
+				BlendState, RasterizerState, DepthStencilState);
 		});
 	}
 };
diff --git a/Engine/Source/Runtime/RenderCore/Public/RenderUtils.h b/Engine/Source/Runtime/RenderCore/Public/RenderUtils.h
--- a/Engine/Source/Runtime/RenderCore/Public/RenderUtils.h
+++ b/Engine/Source/Runtime/RenderCore/Public/RenderUtils.h
@@ -133,6 +133,9 @@ extern RENDERCORE_API class FTexture* GBlackArrayTexture;
 /** A global black volume texture. */
 extern RENDERCORE_API class FTexture* GBlackVolumeTexture;
 
+/** A global black volume texture, with alpha=1. */
+extern RENDERCORE_API class FTexture* GBlackAlpha1VolumeTexture;
+
 /** A global black volume texture<uint>  */
 extern RENDERCORE_API class FTexture* GBlackUintVolumeTexture;
 
diff --git a/Engine/Source/Runtime/RenderCore/Public/RendererInterface.h b/Engine/Source/Runtime/RenderCore/Public/RendererInterface.h
--- a/Engine/Source/Runtime/RenderCore/Public/RendererInterface.h
+++ b/Engine/Source/Runtime/RenderCore/Public/RendererInterface.h
@@ -251,7 +251,7 @@ public:
 		}
 
 		return Extent.X != 0 && NumMips != 0 && NumSamples >=1 && NumSamples <=16 && Format != PF_Unknown
-			&& ((TargetableFlags & TexCreate_UAV) == 0 || GMaxRHIFeatureLevel == ERHIFeatureLevel::SM5);
+			&& ((TargetableFlags & TexCreate_UAV) == 0 || GMaxRHIFeatureLevel == ERHIFeatureLevel::SM5 || GMaxRHIFeatureLevel == ERHIFeatureLevel::ES3_1);
 	}
 
 	FIntVector GetSize() const
diff --git a/Engine/Source/Runtime/Renderer/Private/BasePassRendering.cpp b/Engine/Source/Runtime/Renderer/Private/BasePassRendering.cpp
--- a/Engine/Source/Runtime/Renderer/Private/BasePassRendering.cpp
+++ b/Engine/Source/Runtime/Renderer/Private/BasePassRendering.cpp
@@ -11,6 +11,7 @@
 #include "MeshPassProcessor.inl"
 #include "EditorPrimitivesRendering.h"
 #include "TranslucentRendering.h"
+#include "Rendering/SkyAtmosphereCommonData.h"
 
 // Changing this causes a full shader recompile
 static TAutoConsoleVariable<int32> CVarSelectiveBasePassOutputs(
@@ -505,6 +506,7 @@ void SetupSharedBasePassParameters(
 	SetupReflectionUniformParameters(View, SharedParameters.Reflection);
 	SetupFogUniformParameters(View, SharedParameters.Fog);
 	SetupPlanarReflectionUniformParameters(View, ReflectionSceneProxy, SharedParameters.PlanarReflection);
+	SetupSkyAtmosphereBasePassSharedUniformShaderParameters(View, SharedParameters.SkyAtmosphere);
 
 	const IPooledRenderTarget* PooledRT = GetSubsufaceProfileTexture_RT(RHICmdList);
 
diff --git a/Engine/Source/Runtime/Renderer/Private/BasePassRendering.h b/Engine/Source/Runtime/Renderer/Private/BasePassRendering.h
--- a/Engine/Source/Runtime/Renderer/Private/BasePassRendering.h
+++ b/Engine/Source/Runtime/Renderer/Private/BasePassRendering.h
@@ -52,6 +52,7 @@ BEGIN_GLOBAL_SHADER_PARAMETER_STRUCT(FSharedBasePassUniformParameters,)
 	SHADER_PARAMETER_STRUCT(FReflectionUniformParameters, Reflection)
 	SHADER_PARAMETER_STRUCT(FPlanarReflectionUniformParameters, PlanarReflection) // Single global planar reflection for the forward pass.
 	SHADER_PARAMETER_STRUCT(FFogUniformParameters, Fog)
+	SHADER_PARAMETER_STRUCT(FSkyAtmosphereBasePassSharedUniformShaderParameters, SkyAtmosphere)
 	SHADER_PARAMETER_TEXTURE(Texture2D, SSProfilesTexture)
 END_GLOBAL_SHADER_PARAMETER_STRUCT()
 
diff --git a/Engine/Source/Runtime/Renderer/Private/DeferredShadingRenderer.cpp b/Engine/Source/Runtime/Renderer/Private/DeferredShadingRenderer.cpp
--- a/Engine/Source/Runtime/Renderer/Private/DeferredShadingRenderer.cpp
+++ b/Engine/Source/Runtime/Renderer/Private/DeferredShadingRenderer.cpp
@@ -7,6 +7,7 @@
 #include "DeferredShadingRenderer.h"
 #include "VelocityRendering.h"
 #include "AtmosphereRendering.h"
+#include "SkyAtmosphereRendering.h"
 #include "ScenePrivate.h"
 #include "ScreenRendering.h"
 #include "PostProcess/SceneFilterRendering.h"
@@ -823,10 +824,20 @@ void FDeferredShadingSceneRenderer::Render(FRHICommandListImmediate& RHICmdList)
 
 	PrepareViewRectsForRendering();
 
-	if (Scene->SunLight && Scene->HasAtmosphericFog())
+	if (Scene->HasSkyAtmosphere())
+	{
+		for (int32 LightIndex = 0; LightIndex < NUM_ATMOSPHERE_LIGHTS; ++LightIndex)
+		{
+			if (Scene->AtmosphereLights[LightIndex])
+			{
+				Scene->GetSkyAtmosphereSceneInfo()->PrepareSunLightProxy(*Scene->AtmosphereLights[LightIndex]);
+			}
+		}
+	}
+	else if (Scene->AtmosphereLights[0] && Scene->HasAtmosphericFog())
 	{
 		// Only one atmospheric light at one time.
-		Scene->GetAtmosphericFogSceneInfo()->PrepareSunLightProxy(*Scene->SunLight);
+		Scene->GetAtmosphericFogSceneInfo()->PrepareSunLightProxy(*Scene->AtmosphereLights[0]);
 	}
 
 	SCOPED_NAMED_EVENT(FDeferredShadingSceneRenderer_Render, FColor::Emerald);
@@ -1114,7 +1125,12 @@ void FDeferredShadingSceneRenderer::Render(FRHICommandListImmediate& RHICmdList)
 		FTaskGraphInterface::Get().WaitUntilTasksComplete(UpdateViewCustomDataEvents, ENamedThreads::GetRenderThread());
 	}
 
-	
+	// Generate the Sky/Atmosphere look up tables
+	const bool bShouldRenderSkyAtmosphere = ShouldRenderSkyAtmosphere(Scene->GetSkyAtmosphereSceneInfo(), Scene->GetShaderPlatform());
+	if (bShouldRenderSkyAtmosphere)
+	{
+		RenderSkyAtmosphereLookUpTables(RHICmdList);
+	}
 
 	checkSlow(RHICmdList.IsOutsideRenderPass());
 
@@ -1746,6 +1762,12 @@ void FDeferredShadingSceneRenderer::Render(FRHICommandListImmediate& RHICmdList)
 		}
 	}
 
+	// Draw the sky atmosphere
+	if (bShouldRenderSkyAtmosphere)
+	{
+		RenderSkyAtmosphere(RHICmdList);
+	}
+
 	checkSlow(RHICmdList.IsOutsideRenderPass());
 
 	GRenderTargetPool.AddPhaseEvent(TEXT("Fog"));
@@ -1805,6 +1827,13 @@ void FDeferredShadingSceneRenderer::Render(FRHICommandListImmediate& RHICmdList)
 
 	checkSlow(RHICmdList.IsOutsideRenderPass());
 
+	if (bShouldRenderSkyAtmosphere)
+	{
+		// Debug the sky atmosphere. Critically rendered before translucency to avoid emissive leaking over visualization by writing depth. 
+		// Alternative: render in post process chain as VisualizeHDR.
+		RenderDebugSkyAtmosphere(RHICmdList);
+	}
+
 	GRenderTargetPool.AddPhaseEvent(TEXT("Translucency"));
 
 	// Draw translucency.
diff --git a/Engine/Source/Runtime/Renderer/Private/DeferredShadingRenderer.h b/Engine/Source/Runtime/Renderer/Private/DeferredShadingRenderer.h
--- a/Engine/Source/Runtime/Renderer/Private/DeferredShadingRenderer.h
+++ b/Engine/Source/Runtime/Renderer/Private/DeferredShadingRenderer.h
@@ -192,6 +192,9 @@ private:
 	/** Renders the scene's atmosphere. */
 	void RenderAtmosphere(FRHICommandListImmediate& RHICmdList, const FLightShaftsOutput& LightShaftsOutput);
 
+	/** Render debug visualizations about the sky atmosphere into the scene render target.*/
+	void RenderDebugSkyAtmosphere(FRHICommandListImmediate& RHICmdList);
+
 	void RenderDiffuseIndirectAndAmbientOcclusion(FRHICommandListImmediate& RHICmdList);
 
 	/** Renders sky lighting and reflections that can be done in a deferred pass. */
diff --git a/Engine/Source/Runtime/Renderer/Private/FogRendering.cpp b/Engine/Source/Runtime/Renderer/Private/FogRendering.cpp
--- a/Engine/Source/Runtime/Renderer/Private/FogRendering.cpp
+++ b/Engine/Source/Runtime/Renderer/Private/FogRendering.cpp
@@ -70,13 +70,11 @@ void SetupFogUniformParameters(const FViewInfo& View, FFogUniformParameters& Out
 	// Volumetric Fog
 	{
 		FRHITexture* IntegratedLightScatteringTexture = nullptr;
-
 		if (View.VolumetricFogResources.IntegratedLightScattering)
 		{
 			IntegratedLightScatteringTexture = View.VolumetricFogResources.IntegratedLightScattering->GetRenderTargetItem().ShaderResourceTexture;
 		}
-
-		SetBlack3DIfNull(IntegratedLightScatteringTexture);
+		SetBlackAlpha13DIfNull(IntegratedLightScatteringTexture);
 
 		const bool bApplyVolumetricFog = View.VolumetricFogResources.IntegratedLightScattering != NULL;
 		OutParameters.ApplyVolumetricFog = bApplyVolumetricFog ? 1.0f : 0.0f;
@@ -336,12 +334,13 @@ void FSceneRenderer::InitFogConstants()
 				View.DirectionalInscatteringExponent = FogInfo.DirectionalInscatteringExponent;
 				View.DirectionalInscatteringStartDistance = FogInfo.DirectionalInscatteringStartDistance;
 				View.InscatteringLightDirection = FVector(0);
-				if (Scene->SunLight)
+				FLightSceneInfo* SunLight = Scene->AtmosphereLights[0];	// Fog only takes into account a single atmosphere light with index 0.
+				if (SunLight)
 				{
-					View.InscatteringLightDirection = -Scene->SunLight->Proxy->GetDirection();
-					View.DirectionalInscatteringColor = FogInfo.DirectionalInscatteringColor * Scene->SunLight->Proxy->GetColor().ComputeLuminance();
+					View.InscatteringLightDirection = -SunLight->Proxy->GetDirection();
+					View.DirectionalInscatteringColor = FogInfo.DirectionalInscatteringColor * SunLight->Proxy->GetColor().ComputeLuminance();
 				}
-				View.bUseDirectionalInscattering = Scene->SunLight!=nullptr;
+				View.bUseDirectionalInscattering = SunLight != nullptr;
 			}
 		}
 	}
diff --git a/Engine/Source/Runtime/Renderer/Private/MobileBasePassRendering.cpp b/Engine/Source/Runtime/Renderer/Private/MobileBasePassRendering.cpp
--- a/Engine/Source/Runtime/Renderer/Private/MobileBasePassRendering.cpp
+++ b/Engine/Source/Runtime/Renderer/Private/MobileBasePassRendering.cpp
@@ -157,6 +157,8 @@ void SetupMobileBasePassUniformParameters(
 {
 	SetupFogUniformParameters(View, BasePassParameters.Fog);
 
+	SetupSkyAtmosphereBasePassSharedUniformShaderParameters(View, BasePassParameters.SkyAtmosphere);
+
 	const FScene* Scene = View.Family->Scene ? View.Family->Scene->GetRenderScene() : nullptr;
 	const FPlanarReflectionSceneProxy* ReflectionSceneProxy = Scene ? Scene->GetForwardPassGlobalPlanarReflection() : nullptr;
 	SetupPlanarReflectionUniformParameters(View, ReflectionSceneProxy, BasePassParameters.PlanarReflection);
@@ -189,6 +191,10 @@ void SetupMobileDirectionalLightUniformParameters(
 	if (Light)
 	{
 		Params.DirectionalLightColor = Light->Proxy->GetColor() / PI;
+		if (Light->Proxy->IsUsedAsAtmosphereSunLight())
+		{
+			Params.DirectionalLightColor *= Light->Proxy->GetTransmittanceFactor();
+		}
 		Params.DirectionalLightDirectionAndShadowTransition = FVector4(-Light->Proxy->GetDirection(), 0.f);
 
 		const FVector2D FadeParams = Light->Proxy->GetDirectionalLightDistanceFadeParameters(FeatureLevel, Light->IsPrecomputedLightingValid(), SceneView.MaxShadowCascades);
diff --git a/Engine/Source/Runtime/Renderer/Private/MobileBasePassRendering.h b/Engine/Source/Runtime/Renderer/Private/MobileBasePassRendering.h
--- a/Engine/Source/Runtime/Renderer/Private/MobileBasePassRendering.h
+++ b/Engine/Source/Runtime/Renderer/Private/MobileBasePassRendering.h
@@ -22,11 +22,13 @@
 #include "FogRendering.h"
 #include "PlanarReflectionRendering.h"
 #include "BasePassRendering.h"
+#include "SkyAtmosphereRendering.h"
 
 BEGIN_GLOBAL_SHADER_PARAMETER_STRUCT(FMobileBasePassUniformParameters, )
 	SHADER_PARAMETER_STRUCT(FFogUniformParameters, Fog)
 	SHADER_PARAMETER_STRUCT(FPlanarReflectionUniformParameters, PlanarReflection) // Single global planar reflection for the forward pass.
 	SHADER_PARAMETER_STRUCT(FMobileSceneTextureUniformParameters, SceneTextures)
+	SHADER_PARAMETER_STRUCT(FSkyAtmosphereBasePassSharedUniformShaderParameters, SkyAtmosphere)
 END_GLOBAL_SHADER_PARAMETER_STRUCT()
 
 extern void SetupMobileBasePassUniformParameters(
diff --git a/Engine/Source/Runtime/Renderer/Private/MobileShadingRenderer.cpp b/Engine/Source/Runtime/Renderer/Private/MobileShadingRenderer.cpp
--- a/Engine/Source/Runtime/Renderer/Private/MobileShadingRenderer.cpp
+++ b/Engine/Source/Runtime/Renderer/Private/MobileShadingRenderer.cpp
@@ -46,6 +46,7 @@
 #include "GPUScene.h"
 #include "MaterialSceneTextureId.h"
 #include "DebugViewModeRendering.h"
+#include "SkyAtmosphereRendering.h"
 #include "VisualizeTexture.h"
 
 uint32 GetShadowQuality();
@@ -235,6 +236,13 @@ void FMobileSceneRenderer::InitViews(FRHICommandListImmediate& RHICmdList)
 
 	PreVisibilityFrameSetup(RHICmdList);
 	ComputeViewVisibility(RHICmdList, BasePassDepthStencilAccess, ViewCommandsPerView, DynamicIndexBuffer, DynamicVertexBuffer, DynamicReadBuffer);
+
+	// Initialise Sky/View resources before the view global uniform buffer is built.
+	if (Scene && ShouldRenderSkyAtmosphere(Scene->GetSkyAtmosphereSceneInfo(), Scene->GetShaderPlatform()))
+	{
+		InitSkyAtmosphereForViews(RHICmdList);
+	}
+
 	PostVisibilityFrameSetup(ILCTaskData);
 
 	const bool bDynamicShadows = ViewFamily.EngineShowFlags.DynamicShadows;
@@ -307,6 +315,17 @@ void FMobileSceneRenderer::Render(FRHICommandListImmediate& RHICmdList)
 
 	PrepareViewRectsForRendering();
 
+	if (Scene->HasSkyAtmosphere())
+	{
+		for (int32 LightIndex = 0; LightIndex < NUM_ATMOSPHERE_LIGHTS; ++LightIndex)
+		{
+			if (Scene->AtmosphereLights[LightIndex])
+			{
+				Scene->GetSkyAtmosphereSceneInfo()->PrepareSunLightProxy(*Scene->AtmosphereLights[LightIndex]);
+			}
+		}
+	}
+
 	CSV_SCOPED_TIMING_STAT_EXCLUSIVE(RenderOther);
 	QUICK_SCOPE_CYCLE_COUNTER(STAT_FMobileSceneRenderer_Render);
 	//FRHICommandListExecutor::GetImmediateCommandList().ImmediateFlush(EImmediateFlushType::DispatchToRHIThread);
@@ -379,6 +398,13 @@ void FMobileSceneRenderer::Render(FRHICommandListImmediate& RHICmdList)
 	}
 
 	const bool bGammaSpace = !IsMobileHDR();
+
+	// Generate the Sky/Atmosphere look up tables
+	const bool bShouldRenderSkyAtmosphere = ShouldRenderSkyAtmosphere(Scene->GetSkyAtmosphereSceneInfo(), Scene->GetShaderPlatform());
+	if (bShouldRenderSkyAtmosphere)
+	{
+		RenderSkyAtmosphereLookUpTables(RHICmdList);
+	}
 	
 	// Custom depth
 	if (!bGammaSpace)
@@ -521,6 +547,12 @@ void FMobileSceneRenderer::Render(FRHICommandListImmediate& RHICmdList)
 
 	// scene depth is read only and can be fetched
 	RHICmdList.NextSubpass();
+
+	// Draw the sky atmosphere
+	if (bShouldRenderSkyAtmosphere)
+	{
+		RenderSkyAtmosphere(RHICmdList);
+	}
 		
 	// Split if we need to render translucency in a separate render pass
 	if (bRequiresTranslucencyPass)
diff --git a/Engine/Source/Runtime/Renderer/Private/RendererScene.cpp b/Engine/Source/Runtime/Renderer/Private/RendererScene.cpp
--- a/Engine/Source/Runtime/Renderer/Private/RendererScene.cpp
+++ b/Engine/Source/Runtime/Renderer/Private/RendererScene.cpp
@@ -37,6 +37,7 @@
 #include "LightSceneInfo.h"
 #include "LightMapRendering.h"
 #include "AtmosphereRendering.h"
+#include "SkyAtmosphereRendering.h"
 #include "BasePassRendering.h"
 #include "MobileBasePassRendering.h"
 #include "LightPropagationVolume.h"
@@ -1060,12 +1061,12 @@ FScene::FScene(UWorld* InWorld, bool bInRequiresHitProxies, bool bInIsEditorScen
 ,	bPathTracingNeedsInvalidation(true)
 ,	SkyLight(NULL)
 ,	SimpleDirectionalLight(NULL)
-,	SunLight(NULL)
 ,	ReflectionSceneData(InFeatureLevel)
 ,	IndirectLightingCache(InFeatureLevel)
 ,	DistanceFieldSceneData(GShaderPlatformForFeatureLevel[InFeatureLevel])
 ,	PreshadowCacheLayout(0, 0, 0, 0, false)
 ,	AtmosphericFog(NULL)
+,	SkyAtmosphere(NULL)
 ,	PrecomputedVisibilityHandler(NULL)
 ,	LightOctree(FVector::ZeroVector,HALF_WORLD_MAX)
 ,	PrimitiveOctree(FVector::ZeroVector,HALF_WORLD_MAX)
@@ -1089,6 +1090,7 @@ FScene::FScene(UWorld* InWorld, bool bInRequiresHitProxies, bool bInIsEditorScen
 ,	SceneFrameNumber(0)
 {
 	FMemory::Memzero(MobileDirectionalLights);
+	FMemory::Memzero(AtmosphereLights);
 
 	check(World);
 	World->Scene = this;
@@ -1806,11 +1808,7 @@ void FScene::AddLightSceneInfo_RenderThread(FLightSceneInfo* LightSceneInfo)
 		AssignAvailableShadowMapChannelForLight(LightSceneInfo);
 	}
 
-	if (LightSceneInfo->Proxy->IsUsedAsAtmosphereSunLight() &&
-		(!SunLight || LightSceneInfo->Proxy->GetColor().ComputeLuminance() > SunLight->Proxy->GetColor().ComputeLuminance()) ) // choose brightest sun light...
-	{
-		SunLight = LightSceneInfo;
-	}
+	ProcessAtmosphereLightAddition_RenderThread(LightSceneInfo);
 
 	// Add the light to the scene.
 	LightSceneInfo->AddToScene();
@@ -1942,6 +1940,21 @@ void FScene::DisableSkyLight(FSkyLightSceneProxy* LightProxy)
 	});
 }
 
+bool FScene::HasSkyLightRequiringLightingBuild() const
+{
+	return SkyLight != nullptr && !SkyLight->IsMovable();
+}
+
+bool FScene::HasAtmosphereLightRequiringLightingBuild() const
+{
+	bool AnySunLightNotMovable = false;
+	for (uint8 Index = 0; Index < NUM_ATMOSPHERE_LIGHTS; ++Index)
+	{
+		AnySunLightNotMovable |= AtmosphereLights[Index] != nullptr && !AtmosphereLights[Index]->Proxy->IsMovable();
+	}
+	return AnySunLightNotMovable;
+}
+
 void FScene::AddOrRemoveDecal_RenderThread(FDeferredDecalProxy* Proxy, bool bAdd)
 {
 	if(bAdd)
@@ -2794,21 +2807,7 @@ void FScene::RemoveLightSceneInfo_RenderThread(FLightSceneInfo* LightSceneInfo)
 		    }
 		}
 
-		if (LightSceneInfo == SunLight)
-		{
-			SunLight = NULL;
-			// Search for new sun light...
-			for (TSparseArray<FLightSceneInfoCompact>::TConstIterator It(Lights); It; ++It)
-			{
-				const FLightSceneInfoCompact& LightInfo = *It;
-				if (LightInfo.LightSceneInfo != LightSceneInfo
-					&& LightInfo.LightSceneInfo->Proxy->bUsedAsAtmosphereSunLight
-					&& (!SunLight || SunLight->Proxy->GetColor().ComputeLuminance() < LightInfo.LightSceneInfo->Proxy->GetColor().ComputeLuminance()) )
-				{
-					SunLight = LightInfo.LightSceneInfo;
-				}
-			}
-		}
+		ProcessAtmosphereLightRemoval_RenderThread(LightSceneInfo);
 
 		// Remove the light from the scene.
 		LightSceneInfo->RemoveFromScene();
@@ -3586,6 +3585,44 @@ void FScene::OnLevelAddedToWorld_RenderThread(FName InLevelName)
 	}
 }
 
+void FScene::ProcessAtmosphereLightAddition_RenderThread(FLightSceneInfo* LightSceneInfo)
+{
+	if (LightSceneInfo->Proxy->IsUsedAsAtmosphereSunLight())
+	{
+		const uint8 Index = LightSceneInfo->Proxy->GetAtmosphereSunLightIndex();
+		if (!AtmosphereLights[Index] ||																								// Set it if null
+			LightSceneInfo->Proxy->GetColor().ComputeLuminance() > AtmosphereLights[Index]->Proxy->GetColor().ComputeLuminance())	// Or choose the brightest sun light
+		{
+			AtmosphereLights[Index] = LightSceneInfo;
+		}
+	}
+}
+
+void FScene::ProcessAtmosphereLightRemoval_RenderThread(FLightSceneInfo* LightSceneInfo)
+{
+	// When a light has its intensity or index changed, it will be removed first, then re-added. So we only need to check the index of the removed light.
+	const uint8 Index = LightSceneInfo->Proxy->GetAtmosphereSunLightIndex();
+	if (AtmosphereLights[Index] == LightSceneInfo)
+	{
+		AtmosphereLights[Index] = nullptr;
+		float SelectedLightLuminance = 0.0f;
+
+		for (TSparseArray<FLightSceneInfoCompact>::TConstIterator It(Lights); It; ++It)
+		{
+			const FLightSceneInfoCompact& LightInfo = *It;
+			float LightLuminance = LightInfo.LightSceneInfo->Proxy->GetColor().ComputeLuminance();
+
+			if (LightInfo.LightSceneInfo != LightSceneInfo
+				&& LightInfo.LightSceneInfo->Proxy->IsUsedAsAtmosphereSunLight() && LightInfo.LightSceneInfo->Proxy->GetAtmosphereSunLightIndex() == Index
+				&& (!AtmosphereLights[Index] || SelectedLightLuminance < LightLuminance))
+			{
+				AtmosphereLights[Index] = LightInfo.LightSceneInfo;
+				SelectedLightLuminance = LightLuminance;
+			}
+		}
+	}
+}
+
 void FScene::OnLevelRemovedFromWorld(UWorld* InWorld, bool bIsLightingScenario)
 {
 	if (bIsLightingScenario)
@@ -3648,6 +3685,8 @@ public:
 	virtual void AddInvisibleLight(ULightComponent* Light) override {}
 	virtual void SetSkyLight(FSkyLightSceneProxy* Light) override {}
 	virtual void DisableSkyLight(FSkyLightSceneProxy* Light) override {}
+	virtual bool HasSkyLightRequiringLightingBuild() const { return false; }
+	virtual bool HasAtmosphereLightRequiringLightingBuild() const { return false; }
 
 	virtual void AddDecal(UDecalComponent*) override {}
 	virtual void RemoveDecal(UDecalComponent*) override {}
@@ -3665,6 +3704,11 @@ public:
 	virtual void RemoveAtmosphericFog(class UAtmosphericFogComponent* FogComponent) override {}
 	virtual void RemoveAtmosphericFogResource_RenderThread(FRenderResource* FogResource) override {}
 	virtual FAtmosphericFogSceneInfo* GetAtmosphericFogSceneInfo() override { return NULL; }
+
+	virtual void AddSkyAtmosphere(const class USkyAtmosphereComponent* SkyAtmosphereComponent, bool bStaticLightingBuilt) override {}
+	virtual void RemoveSkyAtmosphere(const class USkyAtmosphereComponent* SkyAtmosphereComponent) override {}
+	virtual FSkyAtmosphereRenderSceneInfo* GetSkyAtmosphereSceneInfo() override { return NULL; }
+
 	virtual void AddWindSource(class UWindDirectionalSourceComponent* WindComponent) override {}
 	virtual void RemoveWindSource(class UWindDirectionalSourceComponent* WindComponent) override {}
 	virtual const TArray<class FWindSourceSceneProxy*>& GetWindSources_RenderThread() const override
diff --git a/Engine/Source/Runtime/Renderer/Private/ScenePrivate.h b/Engine/Source/Runtime/Renderer/Private/ScenePrivate.h
--- a/Engine/Source/Runtime/Renderer/Private/ScenePrivate.h
+++ b/Engine/Source/Runtime/Renderer/Private/ScenePrivate.h
@@ -2502,8 +2502,8 @@ public:
 	/** For the mobile renderer, the first directional light in each lighting channel. */
 	FLightSceneInfo* MobileDirectionalLights[NUM_LIGHTING_CHANNELS];
 
-	/** The sun light for atmospheric effect, if any. */
-	FLightSceneInfo* SunLight;
+	/** The light sources for atmospheric effects, if any. */
+	FLightSceneInfo* AtmosphereLights[NUM_ATMOSPHERE_LIGHTS];
 
 	/** The decals in the scene. */
 	TSparseArray<FDeferredDecalProxy*> Decals;
@@ -2556,6 +2556,9 @@ public:
 	/** The atmospheric fog components in the scene. */
 	FAtmosphericFogSceneInfo* AtmosphericFog;
 
+	/** The sky/atmosphere components of the scene. */
+	FSkyAtmosphereRenderSceneInfo* SkyAtmosphere;
+
 	/** The wind sources in the scene. */
 	TArray<class FWindSourceSceneProxy*> WindSources;
 
@@ -2646,6 +2649,8 @@ public:
 	virtual void AddInvisibleLight(ULightComponent* Light) override;
 	virtual void SetSkyLight(FSkyLightSceneProxy* Light) override;
 	virtual void DisableSkyLight(FSkyLightSceneProxy* Light) override;
+	virtual bool HasSkyLightRequiringLightingBuild() const override;
+	virtual bool HasAtmosphereLightRequiringLightingBuild() const override;
 	virtual void AddDecal(UDecalComponent* Component) override;
 	virtual void RemoveDecal(UDecalComponent* Component) override;
 	virtual void UpdateDecalTransform(UDecalComponent* Decal) override;
@@ -2680,6 +2685,11 @@ public:
 	virtual void RemoveAtmosphericFog(UAtmosphericFogComponent* FogComponent) override;
 	virtual void RemoveAtmosphericFogResource_RenderThread(FRenderResource* FogResource) override;
 	virtual FAtmosphericFogSceneInfo* GetAtmosphericFogSceneInfo() override { return AtmosphericFog; }
+
+	virtual void AddSkyAtmosphere(const class USkyAtmosphereComponent* SkyAtmosphereComponent, bool bStaticLightingBuilt) override;
+	virtual void RemoveSkyAtmosphere(const class USkyAtmosphereComponent* SkyAtmosphereComponent) override;
+	virtual FSkyAtmosphereRenderSceneInfo* GetSkyAtmosphereSceneInfo() override { return SkyAtmosphere; }
+
 	virtual void AddWindSource(UWindDirectionalSourceComponent* WindComponent) override;
 	virtual void RemoveWindSource(UWindDirectionalSourceComponent* WindComponent) override;
 	virtual const TArray<FWindSourceSceneProxy*>& GetWindSources_RenderThread() const override;
@@ -2698,6 +2708,10 @@ public:
 	{
 		return (AtmosphericFog != NULL); // Use default value when Sun Light is not existing
 	}
+	bool HasSkyAtmosphere() const
+	{
+		return (SkyAtmosphere != NULL);
+	}
 
 	/**
 	 * Retrieves the lights interacting with the passed in primitive and adds them to the out array.
@@ -2920,6 +2934,11 @@ private:
 
 	void UpdateLightTransform_RenderThread(FLightSceneInfo* LightSceneInfo, const struct FUpdateLightTransformParameters& Parameters);
 
+	/**
+	 * Deletes the internal AtmosphericFog scene info and operates required operations.
+	 */
+	void DeleteAtmosphericFogSceneInfo();
+
 	/** 
 	* Updates the contents of the given reflection capture by rendering the scene. 
 	* This must be called on the game thread.
@@ -2956,6 +2975,9 @@ private:
 	 */
 	void OnLevelAddedToWorld_RenderThread(FName InLevelName);
 
+	void ProcessAtmosphereLightRemoval_RenderThread(FLightSceneInfo* LightSceneInfo);
+	void ProcessAtmosphereLightAddition_RenderThread(FLightSceneInfo* LightSceneInfo);
+
 private:
 	/** 
 	 * The number of visible lights in the scene
