Handle disabling linear gamma when using SRGB format.

diff --git a/Engine/Source/Runtime/Engine/Private/TextureRenderTarget2D.cpp b/Engine/Source/Runtime/Engine/Private/TextureRenderTarget2D.cpp
--- a/Engine/Source/Runtime/Engine/Private/TextureRenderTarget2D.cpp
+++ b/Engine/Source/Runtime/Engine/Private/TextureRenderTarget2D.cpp
@@ -172,6 +172,8 @@ void UTextureRenderTarget2D::UpdateResourceImmediate(bool bClearRenderTarget/*=t
 #if WITH_EDITOR
 void UTextureRenderTarget2D::PostEditChangeProperty(FPropertyChangedEvent& PropertyChangedEvent)
 {
+	const FName PropertyName = (PropertyChangedEvent.Property != nullptr) ? PropertyChangedEvent.Property->GetFName() : NAME_None;
+
 	EPixelFormat Format = GetFormat();
 
 	const int32 WarnSize = 2048; 
@@ -196,6 +198,18 @@ void UTextureRenderTarget2D::PostEditChangeProperty(FPropertyChangedEvent& Prope
 	SizeX = FMath::Clamp<int32>(SizeX - (SizeX % GPixelFormats[Format].BlockSizeX),1,MaxSize);
 	SizeY = FMath::Clamp<int32>(SizeY - (SizeY % GPixelFormats[Format].BlockSizeY),1,MaxSize);
 
+	if (PropertyName == GET_MEMBER_NAME_CHECKED(UTextureRenderTarget2D, RenderTargetFormat))
+	{
+		if (RenderTargetFormat == RTF_RGBA8_SRGB)
+		{
+			bForceLinearGamma = false;
+		}
+		else
+		{
+			bForceLinearGamma = true;
+		}
+	}
+
 	Super::PostEditChangeProperty(PropertyChangedEvent);
 
     // SRGB may have been changed by Super, reset it since we prefer to honor explicit user choice
