Fixed particle system dynamic parameter being invalid in instanced stereo

diff --git a/Engine/Source/Runtime/Engine/Private/Particles/ParticleSystemRender.cpp b/Engine/Source/Runtime/Engine/Private/Particles/ParticleSystemRender.cpp
--- a/Engine/Source/Runtime/Engine/Private/Particles/ParticleSystemRender.cpp
+++ b/Engine/Source/Runtime/Engine/Private/Particles/ParticleSystemRender.cpp
@@ -2544,7 +2544,13 @@ void FDynamicMeshEmitterData::GetInstanceData(void* InstanceData, void* DynamicP
 				CurrentInstanceVertexDynParam->DynamicValue[1] = DynamicParameterValue.Y;
 				CurrentInstanceVertexDynParam->DynamicValue[2] = DynamicParameterValue.Z;
 				CurrentInstanceVertexDynParam->DynamicValue[3] = DynamicParameterValue.W;
-				TempDynamicParameterVert += DynamicParameterVertexStride;
+
+				for (uint32 Factor = 1; Factor < InstanceFactor; Factor++)
+				{
+					FMemory::Memcpy(TempDynamicParameterVert + DynamicParameterVertexStride * Factor, TempDynamicParameterVert, DynamicParameterVertexStride);
+				}
+
+				TempDynamicParameterVert += DynamicParameterVertexStride * InstanceFactor;
 			}
 		}
 		
