Testing certain warning codes (truncation, narrowing, etc) and treating them as errors. Fixed existing warnings in handwritten shaders and from material genererated shaders.

diff --git a/Engine/Shaders/Private/ComposeSeparateTranslucency.usf b/Engine/Shaders/Private/ComposeSeparateTranslucency.usf
--- a/Engine/Shaders/Private/ComposeSeparateTranslucency.usf
+++ b/Engine/Shaders/Private/ComposeSeparateTranslucency.usf
@@ -12,7 +12,7 @@ void MainPS(
 	float4 SvPosition : SV_POSITION,
 	out float4 OutColor : SV_Target0)
 {
-	float2 BufferUV = SvPosition * View.BufferSizeAndInvSize.zw;
+	float2 BufferUV = SvPosition.xy * View.BufferSizeAndInvSize.zw;
 	
 	float4 SceneColorSample = SceneColor.SampleLevel(SceneColorSampler, BufferUV, 0);
 	float4 SeparateTranslucencySample = SeparateTranslucency.SampleLevel(SeparateTranslucencySampler, BufferUV, 0);
diff --git a/Engine/Shaders/Private/HZB.usf b/Engine/Shaders/Private/HZB.usf
--- a/Engine/Shaders/Private/HZB.usf
+++ b/Engine/Shaders/Private/HZB.usf
@@ -181,7 +181,7 @@ void HZBBuildCS(
 
 void HZBBuildPS(float4 SvPosition : SV_POSITION, out float4 OutColor : SV_Target0)
 {
-	float2 BufferUV = SvPosition * DispatchThreadIdToBufferUV.xy + DispatchThreadIdToBufferUV.zw;
+	float2 BufferUV = SvPosition.xy * DispatchThreadIdToBufferUV.xy + DispatchThreadIdToBufferUV.zw;
 	float4 DeviceZ = Gather4(ParentTextureMip, ParentTextureMipSampler, BufferUV);
 
 	float FurthestDeviceZ = min(min(DeviceZ.x, DeviceZ.y), min(DeviceZ.z, DeviceZ.w));
diff --git a/Engine/Shaders/Private/MeshParticleVertexFactory.ush b/Engine/Shaders/Private/MeshParticleVertexFactory.ush
--- a/Engine/Shaders/Private/MeshParticleVertexFactory.ush
+++ b/Engine/Shaders/Private/MeshParticleVertexFactory.ush
@@ -465,7 +465,7 @@ float4 VertexFactoryGetWorldPosition(FVertexFactoryInput Input, FVertexFactoryIn
 
 float3 VertexFactoryGetWorldNormal(FVertexFactoryInput Input, FVertexFactoryIntermediates Intermediates)
 {
-	return TangentBias(Input.TangentZ);
+	return TangentBias(Input.TangentZ.xyz);
 }
 
 float4 VertexFactoryGetRasterizedWorldPosition(FVertexFactoryInput Input, FVertexFactoryIntermediates Intermediates, float4 InWorldPosition)
diff --git a/Engine/Shaders/Private/PostProcessMaterialShaders.usf b/Engine/Shaders/Private/PostProcessMaterialShaders.usf
--- a/Engine/Shaders/Private/PostProcessMaterialShaders.usf
+++ b/Engine/Shaders/Private/PostProcessMaterialShaders.usf
@@ -169,7 +169,7 @@ void MainPS(
 
 	// can be optimized
 	float4 SvPosition = Input.Position;
-	float2 ViewportUV = (SvPosition - PostProcessOutput_ViewportMin) * PostProcessOutput_ViewportSizeInverse;
+	float2 ViewportUV = (SvPosition.xy - PostProcessOutput_ViewportMin) * PostProcessOutput_ViewportSizeInverse;
 
 #if NUM_TEX_COORD_INTERPOLATORS
 	UNROLL
diff --git a/Engine/Shaders/Private/SimpleElementVolumeTexturePreviewPixelShader.usf b/Engine/Shaders/Private/SimpleElementVolumeTexturePreviewPixelShader.usf
--- a/Engine/Shaders/Private/SimpleElementVolumeTexturePreviewPixelShader.usf
+++ b/Engine/Shaders/Private/SimpleElementVolumeTexturePreviewPixelShader.usf
@@ -152,7 +152,7 @@ void GetTexCoordAndIncrement(in float2 TextureCoordinate, out int NumSteps, out
 	const float MaxSteps = 128;
 
 	//View Local->World transform
-	float3x3 ViewMatrix = TraceViewMatrix; // GetTraceViewMatrix();
+	float3x3 ViewMatrix = (float3x3)TraceViewMatrix; // GetTraceViewMatrix();
 
 	//Allow FOV to be specified in Degrees via settings
 	const float FOV = 90;
diff --git a/Engine/Shaders/Private/SubsurfaceBurleyNormalized.usf b/Engine/Shaders/Private/SubsurfaceBurleyNormalized.usf
--- a/Engine/Shaders/Private/SubsurfaceBurleyNormalized.usf
+++ b/Engine/Shaders/Private/SubsurfaceBurleyNormalized.usf
@@ -616,7 +616,7 @@ float4 BurleyNormalizedSS(float2 BufferUV)
 
 				//step 4: accumulate radiance from the diffusion profile rR(r)
 				//make sure the DiffuseMeanFreePath is not zero and in m.
-				float3 DiffusionProfile = GetDiffuseReflectProfileWithDiffuseMeanFreePath(BurleyParameter.SurfaceAlbedo, BurleyParameter.DiffuseMeanFreePath, S3D, RadiusSampled);
+				float3 DiffusionProfile = GetDiffuseReflectProfileWithDiffuseMeanFreePath(BurleyParameter.SurfaceAlbedo.xyz, BurleyParameter.DiffuseMeanFreePath.xyz, S3D.xyz, RadiusSampled);
 				RadianceAccumulated.xyz += (DiffusionProfile / BurleySampleInfo.Pdf) * (SampledRadianceAndDepth.xyz*TintColor*Mask);
 
 #ifdef USE_BILATERAL_FILTERING
diff --git a/Engine/Shaders/Private/VolumetricFogLightFunction.usf b/Engine/Shaders/Private/VolumetricFogLightFunction.usf
--- a/Engine/Shaders/Private/VolumetricFogLightFunction.usf
+++ b/Engine/Shaders/Private/VolumetricFogLightFunction.usf
@@ -42,7 +42,7 @@ void Main(
 	float3 AbsoluteWorldPosition = SvPositionToWorld(SvPosition);
 	*/
 
-	float2 ScreenPosition = (SvPosition * LightFunctionTexelSize - .5f) * float2(2, -2);
+	float2 ScreenPosition = (SvPosition.xy * LightFunctionTexelSize - .5f) * float2(2, -2);
 	float3 AbsoluteWorldPosition = GetWorldPositionForLightFunctionTexturePosition(ScreenPosition);
 
 	float4 LightVector = mul(float4(AbsoluteWorldPosition, 1), LightFunctionWorldToLight);
diff --git a/Engine/Shaders/Private/WideCustomResolveShaders.usf b/Engine/Shaders/Private/WideCustomResolveShaders.usf
--- a/Engine/Shaders/Private/WideCustomResolveShaders.usf
+++ b/Engine/Shaders/Private/WideCustomResolveShaders.usf
@@ -64,7 +64,7 @@ Texture2DMS<float4,MSAA_SAMPLE_COUNT> Tex; // Input MSAA color
 // Box filter
 float3 resolve_bspline(uint2 pos)
 {
-	return (Tex.Load(pos, 0) + Tex.Load(pos, 1) + Tex.Load(pos, 2) + Tex.Load(pos, 3)) * .25;
+	return (Tex.Load(pos, 0) + Tex.Load(pos, 1) + Tex.Load(pos, 2) + Tex.Load(pos, 3)).xyz * .25;
 }
 # elif WIDE_RESOLVE_WIDTH == 1
 #  include "WideCustomResolve_Wide.ush"
diff --git a/Engine/Source/Developer/Windows/ShaderFormatD3D/Private/D3DShaderCompiler.cpp b/Engine/Source/Developer/Windows/ShaderFormatD3D/Private/D3DShaderCompiler.cpp
--- a/Engine/Source/Developer/Windows/ShaderFormatD3D/Private/D3DShaderCompiler.cpp
+++ b/Engine/Source/Developer/Windows/ShaderFormatD3D/Private/D3DShaderCompiler.cpp
@@ -280,6 +280,12 @@ static FString D3D11CreateShaderCompileCommandLine(
 		FXCCommandline += FString(TEXT(" /Gec"));
 	}
 
+	if (CompileFlags & D3DCOMPILE_WARNINGS_ARE_ERRORS)
+	{
+		CompileFlags &= ~D3DCOMPILE_WARNINGS_ARE_ERRORS;
+		FXCCommandline += FString(TEXT(" /WX"));
+	}
+
 	switch (CompileFlags & SHADER_OPTIMIZATION_LEVEL_MASK)
 	{
 		case D3D10_SHADER_OPTIMIZATION_LEVEL2:
@@ -1562,6 +1568,29 @@ static bool CompileAndProcessD3DShader(FString& PreprocessedShaderSource, const
 	return SUCCEEDED(Result);
 }
 
+void CheckForWarningsAsErrors(const FString& CurrentError) 
+{	
+	FString ErrorCode = CurrentError.Mid(CurrentError.Find("X") + 1, 4);
+	
+	// todo: Create a flag to enable/disable this.
+	bool bTreatAsWarningsInstead = false;
+
+	// Add your warnings codes to the list. https://docs.microsoft.com/en-us/windows/win32/direct3dhlsl/hlsl-errors-and-warnings
+	if (ErrorCode == "3206" || // Implicit truncation.
+		ErrorCode == "3205" || // Narrowing conversion.
+		ErrorCode == "3200")   // Type Mismatch
+	{
+		if (bTreatAsWarningsInstead) 
+		{
+			UE_LOG(LogD3D11ShaderCompiler, Warning, TEXT("%s"), *CurrentError)
+		}
+		else 
+		{
+			UE_LOG(LogD3D11ShaderCompiler, Fatal, TEXT("%s"), *CurrentError);
+		}
+	}
+}
+
 void CompileD3DShader(const FShaderCompilerInput& Input,FShaderCompilerOutput& Output, FShaderCompilerDefinitions& AdditionalDefines, const FString& WorkingDirectory)
 {
 	FString PreprocessedShaderSource;
@@ -1694,6 +1723,9 @@ void CompileD3DShader(const FShaderCompilerInput& Input,FShaderCompilerOutput& O
 		CompileFlags |= TranslateCompilerFlagD3D11((ECompilerFlags)Input.Environment.CompilerFlags[FlagIndex]);
 	}
 
+	// Currently not using this flag because it captures all warnings as errors.
+	//CompileFlags |= D3DCOMPILE_WARNINGS_ARE_ERRORS;
+
 	TArray<FString> FilteredErrors;
 	if (!CompileAndProcessD3DShader(PreprocessedShaderSource, CompilerPath, CompileFlags, Input, EntryPointName, ShaderProfile, false, FilteredErrors, Output))
 	{
@@ -1725,6 +1757,9 @@ void CompileD3DShader(const FShaderCompilerInput& Input,FShaderCompilerOutput& O
 			NewError.StrippedErrorMessage = CurrentError;
 		}
 		Output.Errors.Add(NewError);
+
+		// Todo: Add a flag to enable/disable this check.
+		CheckForWarningsAsErrors(CurrentError);
 	}
 
 	if (Input.ExtraSettings.bExtractShaderSource)
diff --git a/Engine/Source/Runtime/Engine/Private/Materials/HLSLMaterialTranslator.h b/Engine/Source/Runtime/Engine/Private/Materials/HLSLMaterialTranslator.h
--- a/Engine/Source/Runtime/Engine/Private/Materials/HLSLMaterialTranslator.h
+++ b/Engine/Source/Runtime/Engine/Private/Materials/HLSLMaterialTranslator.h
@@ -907,11 +907,11 @@ public:
 			Chunk[MP_DiffuseColor]					= Material->CompilePropertyAndSetMaterialProperty(MP_DiffuseColor			,this);
 			Chunk[MP_SpecularColor]					= Material->CompilePropertyAndSetMaterialProperty(MP_SpecularColor			,this);
 			Chunk[MP_BaseColor]						= Material->CompilePropertyAndSetMaterialProperty(MP_BaseColor				,this);
-			Chunk[MP_Metallic]						= Material->CompilePropertyAndSetMaterialProperty(MP_Metallic				,this);
-			Chunk[MP_Specular]						= Material->CompilePropertyAndSetMaterialProperty(MP_Specular				,this);
-			Chunk[MP_Roughness]						= Material->CompilePropertyAndSetMaterialProperty(MP_Roughness				,this);
+			Chunk[MP_Metallic]						= ForceCast(Material->CompilePropertyAndSetMaterialProperty(MP_Metallic		,this), MCT_Float1);
+			Chunk[MP_Specular]						= ForceCast(Material->CompilePropertyAndSetMaterialProperty(MP_Specular		,this), MCT_Float1);
+			Chunk[MP_Roughness]						= ForceCast(Material->CompilePropertyAndSetMaterialProperty(MP_Roughness    ,this), MCT_Float1);
 			Chunk[MP_Opacity]						= Material->CompilePropertyAndSetMaterialProperty(MP_Opacity				,this);
-			Chunk[MP_OpacityMask]					= Material->CompilePropertyAndSetMaterialProperty(MP_OpacityMask			,this);
+			Chunk[MP_OpacityMask]					= ForceCast(Material->CompilePropertyAndSetMaterialProperty(MP_OpacityMask  ,this), MCT_Float1);
 			Chunk[MP_WorldPositionOffset]			= Material->CompilePropertyAndSetMaterialProperty(MP_WorldPositionOffset	,this);
 			Chunk[MP_WorldDisplacement]				= Material->CompilePropertyAndSetMaterialProperty(MP_WorldDisplacement		,this);
 			Chunk[MP_TessellationMultiplier]		= Material->CompilePropertyAndSetMaterialProperty(MP_TessellationMultiplier	,this);			
@@ -947,8 +947,8 @@ public:
 
 			Chunk[MP_CustomData0]					= Material->CompilePropertyAndSetMaterialProperty(MP_CustomData0		,this);
 			Chunk[MP_CustomData1]					= Material->CompilePropertyAndSetMaterialProperty(MP_CustomData1		,this);
-			Chunk[MP_AmbientOcclusion]				= Material->CompilePropertyAndSetMaterialProperty(MP_AmbientOcclusion	,this);
-
+			Chunk[MP_AmbientOcclusion]				= ForceCast(Material->CompilePropertyAndSetMaterialProperty(MP_AmbientOcclusion	,this), MCT_Float1);
+			
 			if (IsTranslucentBlendMode(BlendMode))
 			{
 				int32 UserRefraction = ForceCast(Material->CompilePropertyAndSetMaterialProperty(MP_Refraction, this), MCT_Float1);
@@ -5717,7 +5717,7 @@ protected:
 		}
 		else
 		{
-			return AddCodeChunk(GetParameterType(A),TEXT("max(%s,%s)"),*GetParameterCode(A),*CoerceParameter(B,GetParameterType(A)));
+			return AddCodeChunk(GetParameterType(B),TEXT("max(%s,%s)"), *CoerceParameter(A, GetParameterType(B)), *GetParameterCode(B));
 		}
 	}
 
