DR - check(InFlightFrames.Num()==0) on exit because InFlightFrames.Empty(); was missed during the merge. As emtying the Arrays will also release the Queries we don't need to do that manually.

diff --git a/Engine/Source/Runtime/Engine/Private/DynamicResolution.cpp b/Engine/Source/Runtime/Engine/Private/DynamicResolution.cpp
--- a/Engine/Source/Runtime/Engine/Private/DynamicResolution.cpp
+++ b/Engine/Source/Runtime/Engine/Private/DynamicResolution.cpp
@@ -632,13 +632,9 @@ public:
 		// Wait for all queries to land.
 		HandLandedQueriesToHeuristic(/* bWait = */ true);
 
-		for (InFlightFrameQueries& Query : InFlightFrames)
-		{
-			Query.BeginDynamicResolutionQuery.ReleaseQuery();
-			Query.BeginFrameQuery.ReleaseQuery();
-			Query.EndDynamicResolutionQuery.ReleaseQuery();
-			Query.EndFrameQuery.ReleaseQuery();
-		}
+		// Reset the frame properties
+		InFlightFrames.Empty();
+		QueryPool = nullptr;
 	}
 
 
