Reflection capture not switching to use cube map. This was fixed by checking if a cube map property was added or changed and then re-creating the new reflection probe data.

diff --git a/Engine/Source/Runtime/Engine/Private/Components/ReflectionCaptureComponent.cpp b/Engine/Source/Runtime/Engine/Private/Components/ReflectionCaptureComponent.cpp
--- a/Engine/Source/Runtime/Engine/Private/Components/ReflectionCaptureComponent.cpp
+++ b/Engine/Source/Runtime/Engine/Private/Components/ReflectionCaptureComponent.cpp
@@ -968,6 +968,13 @@ bool UReflectionCaptureComponent::CanEditChange(const UProperty* Property) const
 
 void UReflectionCaptureComponent::PostEditChangeProperty(FPropertyChangedEvent& PropertyChangedEvent)
 {
+	if (PropertyChangedEvent.GetPropertyName() == GET_MEMBER_NAME_CHECKED(UReflectionCaptureComponent, Cubemap) ||
+		PropertyChangedEvent.GetPropertyName() == GET_MEMBER_NAME_CHECKED(UReflectionCaptureComponent, SourceCubemapAngle) ||
+		PropertyChangedEvent.GetPropertyName() == GET_MEMBER_NAME_CHECKED(UReflectionCaptureComponent, ReflectionSourceType))
+	{
+		MarkDirtyForRecapture();
+	}
+
 	Super::PostEditChangeProperty(PropertyChangedEvent);
 }
 
