Add support for volumetric lightmap (VLM) sub level streaming

* Basic assumption: most of the data is contained in small area of volume - which is the case most of the time (1% - 2% of the volume contains 90%+ of the total brick data).
* Only the most detailed level of bricks are streamed - they are assigned to each sub level and streamed in and out with the level. The top level bricks are assigned to the persistent level so that we have something to display (and fewer discontinuities) when the sub levels are streamed out. However runtime dynamic stitching is not implemented so there might be seams in the area of streamed out sub levels - which is probably ok since the level is not needed.
* With SM4 going away, VLM sub level streaming GPU path requires SM5. All platforms having < SM5 will fall back to CPU interpolation path.
* It is recommended to enable bUseFastVoxelization=true and bUseEmbreeInstancing=true in BaseLightmass.ini with this change to speed up VLM generation especially if you have large worlds. Data management has been changed so that we can handle 4GB+ VLM data as long as you have enough video memory - however lighting build time could be not practical so use with care.
* Special notice: we have never supported the workflow of sharing sub levels between multiple worlds, or open each sub level separately to bake, and we're unlikely to support them in the future. It is possible that you happen to get some results, but they are not expected to work.

diff --git a/Engine/Source/Editor/UnrealEd/Private/Lightmass/ImportVolumetricLightmap.cpp b/Engine/Source/Editor/UnrealEd/Private/Lightmass/ImportVolumetricLightmap.cpp
--- a/Engine/Source/Editor/UnrealEd/Private/Lightmass/ImportVolumetricLightmap.cpp
+++ b/Engine/Source/Editor/UnrealEd/Private/Lightmass/ImportVolumetricLightmap.cpp
@@ -22,7 +22,7 @@ DEFINE_LOG_CATEGORY_STATIC(LogVolumetricLightmapImport, Log, All);
 
 FVector GDebugVoxelPosition = FVector::ZeroVector;
 
-void CopyBrickToAtlasVolumeTexture(int32 FormatSize, FIntVector AtlasSize, FIntVector BrickMin, FIntVector BrickSize, const uint8* SourceData, uint8* DestData)
+void CopyBrickToAtlasVolumeTexture(int32 FormatSize, FIntVector AtlasSize, FIntVector BrickMin, FIntVector BrickSize, const uint8* RESTRICT SourceData, uint8* RESTRICT DestData)
 {
 	const int32 SourcePitch = BrickSize.X * FormatSize;
 	const int32 Pitch = AtlasSize.X * FormatSize;
@@ -43,6 +43,39 @@ void CopyBrickToAtlasVolumeTexture(int32 FormatSize, FIntVector AtlasSize, FIntV
 	}
 }
 
+void CopyBetweenAtlasVolumeTextures(
+	int32 FormatSize,
+	FIntVector BrickSize,
+	FIntVector SourceAtlasSize,
+	FIntVector SourceBrickMin,
+	const TArray<uint8>& SourceData,
+	FIntVector DestAtlasSize,
+	FIntVector DestBrickMin,
+	TArray<uint8>& DestData)
+{
+	const int32 DestPitch = DestAtlasSize.X * FormatSize;
+	const int32 DestDepthPitch = DestAtlasSize.X * DestAtlasSize.Y * FormatSize;
+	const int32 SourcePitch = SourceAtlasSize.X * FormatSize;
+	const int32 SourceDepthPitch = SourceAtlasSize.X * SourceAtlasSize.Y * FormatSize;
+
+	// Copy each row into the correct position in the global volume texture
+	for (int32 ZIndex = 0; ZIndex < BrickSize.Z; ZIndex++)
+	{
+		const int32 DestZIndex = (DestBrickMin.Z + ZIndex) * DestDepthPitch + DestBrickMin.X * FormatSize;
+		const int32 SourceZIndex = (SourceBrickMin.Z + ZIndex) * SourceDepthPitch + SourceBrickMin.X * FormatSize;
+
+		for (int32 YIndex = 0; YIndex < BrickSize.Y; YIndex++)
+		{
+			const int32 DestIndex = DestZIndex + (DestBrickMin.Y + YIndex) * DestPitch;
+			const int32 SourceIndex = SourceZIndex + (SourceBrickMin.Y + YIndex) * SourcePitch;
+			for (int32 XIndex = 0; XIndex < BrickSize.X * FormatSize; XIndex++)
+			{
+				FMemory::Memcpy((uint8*)&DestData[DestIndex], (const uint8*)&SourceData[SourceIndex], BrickSize.X * FormatSize);
+			}
+		}
+	}
+}
+
 int32 ComputeLinearVoxelIndex(FIntVector VoxelCoordinate, FIntVector VolumeDimensions)
 {
 	return (VoxelCoordinate.Z * VolumeDimensions.Y + VoxelCoordinate.Y) * VolumeDimensions.X + VoxelCoordinate.X;
@@ -50,6 +83,7 @@ int32 ComputeLinearVoxelIndex(FIntVector VoxelCoordinate, FIntVector VolumeDimen
 
 struct FImportedVolumetricLightmapBrick
 {
+	FGuid IntersectingLevelGuid;
 	FIntVector IndirectionTexturePosition;
 	int32 TreeDepth;
 	float AverageClosestGeometryDistance;
@@ -156,6 +190,7 @@ void FLightmassProcessor::ImportIrradianceTasks(bool& bGenerateSkyShadowing, TAr
 			{
 				NewTaskData.Bricks.AddDefaulted();
 				FImportedVolumetricLightmapBrick& NewBrick = NewTaskData.Bricks.Last();
+				Swarm.ReadChannel(Channel, &NewBrick.IntersectingLevelGuid, sizeof(NewBrick.IntersectingLevelGuid));
 				Swarm.ReadChannel(Channel, &NewBrick.IndirectionTexturePosition, sizeof(NewBrick.IndirectionTexturePosition));
 				Swarm.ReadChannel(Channel, &NewBrick.TreeDepth, sizeof(NewBrick.TreeDepth));
 				Swarm.ReadChannel(Channel, &NewBrick.AverageClosestGeometryDistance, sizeof(NewBrick.AverageClosestGeometryDistance));
@@ -857,9 +892,10 @@ int32 TrimBricksForMemoryLimit(
 	const int32 PaddedBrickSize = VolumetricLightmapSettings.BrickSize + 1;
 	TArray<const FImportedVolumetricLightmapBrick*>& HighestDensityBricks = BricksByDepth[VolumetricLightmapSettings.MaxRefinementLevels - 1];
 
-	const int32 BrickSizeBytes = VoxelSizeBytes * PaddedBrickSize * PaddedBrickSize * PaddedBrickSize;
-	const int32 MaxBrickBytes = MaximumBrickMemoryMb * 1024 * 1024;
-	const int32 NumBricksBudgeted = FMath::DivideAndRoundUp(MaxBrickBytes, BrickSizeBytes);
+	const uint64 BrickSizeBytes = VoxelSizeBytes * PaddedBrickSize * PaddedBrickSize * PaddedBrickSize;
+	const uint64 MaxBrickBytes = MaximumBrickMemoryMb * 1024 * 1024;
+	check(FMath::DivideAndRoundUp(MaxBrickBytes, BrickSizeBytes) <= 0x7FFFFFFFull);
+	const int32 NumBricksBudgeted = (int32)FMath::DivideAndRoundUp(MaxBrickBytes, BrickSizeBytes);
 	const int32 NumBricksToRemove = FMath::Clamp<int32>(NumBricksBeforeTrimming - NumBricksBudgeted, 0, HighestDensityBricks.Num());
 
 	if (NumBricksToRemove > 0)
@@ -886,11 +922,12 @@ void BuildIndirectionTexture(
 	int32 MaxBricksInLayoutOneDim,
 	FIntVector BrickLayoutDimensions,
 	int32 IndirectionTextureDataStride,
-	FPrecomputedVolumetricLightmapData& CurrentLevelData)
+	FPrecomputedVolumetricLightmapData& CurrentLevelData,
+	bool bOnlyBuildForPersistentLevel)
 {
 	const int32 BrickSizeLog2 = FMath::FloorLog2(VolumetricLightmapSettings.BrickSize);
 
-	int32 BrickStartAllocation = 0;
+	int32 BrickAllocation = 0;
 
 	for (int32 CurrentDepth = 0; CurrentDepth < VolumetricLightmapSettings.MaxRefinementLevels; CurrentDepth++)
 	{
@@ -899,7 +936,19 @@ void BuildIndirectionTexture(
 		for (int32 BrickIndex = 0; BrickIndex < BricksAtCurrentDepth.Num(); BrickIndex++)
 		{
 			const FImportedVolumetricLightmapBrick& Brick = *BricksAtCurrentDepth[BrickIndex];
-			const FIntVector BrickLayoutPosition = ComputeBrickLayoutPosition(BrickStartAllocation + BrickIndex, BrickLayoutDimensions);
+
+			if (CurrentDepth == VolumetricLightmapSettings.MaxRefinementLevels - 1)
+			{
+				const FGuid PersistentLevelGuid = FGuid(0, 0, 0, 0);
+
+				// Skip non-intersecting, bottom detailed bricks for persistent level
+				if (bOnlyBuildForPersistentLevel && Brick.IntersectingLevelGuid != PersistentLevelGuid)
+				{
+					continue;
+				}
+			}
+
+			const FIntVector BrickLayoutPosition = ComputeBrickLayoutPosition(BrickAllocation, BrickLayoutDimensions);
 			check(BrickLayoutPosition.X < MaxBricksInLayoutOneDim && BrickLayoutPosition.Y < MaxBricksInLayoutOneDim && BrickLayoutPosition.Z < MaxBricksInLayoutOneDim);
 			checkSlow(IndirectionTextureDataStride == sizeof(uint8) * 4);
 
@@ -923,9 +972,9 @@ void BuildIndirectionTexture(
 					}
 				}
 			}
-		}
 
-		BrickStartAllocation += BricksAtCurrentDepth.Num();
+			BrickAllocation++;
+		}
 	}
 }
 
@@ -945,6 +994,7 @@ void FLightmassProcessor::ImportVolumetricLightmap()
 	const int32 PaddedBrickSize = VolumetricLightmapSettings.BrickSize + 1;
 	const int32 MaxBricksInLayoutOneDim = 1 << 8;
 	const int32 MaxBrickTextureLayoutSize = PaddedBrickSize * MaxBricksInLayoutOneDim;
+	const FGuid PersistentLevelGuid = FGuid(0, 0, 0, 0);
 
 	TArray<FImportedVolumetricLightmapTaskData> TaskDataArray;
 	TaskDataArray.Empty(Exporter->VolumetricLightmapTaskGuids.Num());
@@ -972,7 +1022,7 @@ void FLightmassProcessor::ImportVolumetricLightmap()
 
 	ULevel* StorageLevel = System.LightingScenario ? System.LightingScenario : System.GetWorld()->PersistentLevel;
 	UMapBuildDataRegistry* Registry = StorageLevel->GetOrCreateMapBuildData();
-	FPrecomputedVolumetricLightmapData& CurrentLevelData = Registry->AllocateLevelPrecomputedVolumetricLightmapBuildData(StorageLevel->LevelBuildDataId);
+	FPrecomputedVolumetricLightmapData CurrentLevelData;
 
 	{
 		CurrentLevelData.InitializeOnImport(FBox(VolumetricLightmapSettings.VolumeMin, VolumetricLightmapSettings.VolumeMin + VolumetricLightmapSettings.VolumeSize), BrickSize);
@@ -1024,7 +1074,7 @@ void FLightmassProcessor::ImportVolumetricLightmap()
 		CurrentLevelData.IndirectionTexture.Resize(TotalIndirectionTextureSize * IndirectionTextureDataStride);
 	}
 
-	BuildIndirectionTexture(BricksByDepth, VolumetricLightmapSettings, MaxBricksInLayoutOneDim, BrickLayoutDimensions, IndirectionTextureDataStride, CurrentLevelData);
+	BuildIndirectionTexture(BricksByDepth, VolumetricLightmapSettings, MaxBricksInLayoutOneDim, BrickLayoutDimensions, IndirectionTextureDataStride, CurrentLevelData, false);
 
 	TArray<Lightmass::FIrradianceVoxelImportProcessingData> VoxelImportProcessingData;
 
@@ -1047,11 +1097,9 @@ void FLightmassProcessor::ImportVolumetricLightmap()
 			CurrentLevelData.BrickData.SHCoefficients[i].Resize(TotalBrickDataSize * Stride);
 		}
 
-		CurrentLevelData.BrickData.LQLightColor.Data.Empty(TotalBrickDataSize * GPixelFormats[CurrentLevelData.BrickData.LQLightColor.Format].BlockBytes);
-		CurrentLevelData.BrickData.LQLightColor.Data.AddZeroed(TotalBrickDataSize * GPixelFormats[CurrentLevelData.BrickData.LQLightColor.Format].BlockBytes);
+		CurrentLevelData.BrickData.LQLightColor.Resize(TotalBrickDataSize * GPixelFormats[CurrentLevelData.BrickData.LQLightColor.Format].BlockBytes);
 
-		CurrentLevelData.BrickData.LQLightDirection.Data.Empty(TotalBrickDataSize * GPixelFormats[CurrentLevelData.BrickData.LQLightDirection.Format].BlockBytes);
-		CurrentLevelData.BrickData.LQLightDirection.Data.AddZeroed(TotalBrickDataSize * GPixelFormats[CurrentLevelData.BrickData.LQLightDirection.Format].BlockBytes);
+		CurrentLevelData.BrickData.LQLightDirection.Resize(TotalBrickDataSize * GPixelFormats[CurrentLevelData.BrickData.LQLightDirection.Format].BlockBytes);
 
 		VoxelImportProcessingData.Empty(TotalBrickDataSize);
 		VoxelImportProcessingData.AddZeroed(TotalBrickDataSize);
@@ -1216,72 +1264,361 @@ void FLightmassProcessor::ImportVolumetricLightmap()
 		BrickStartAllocation += BricksAtCurrentDepth.Num();
 	}
 
-	CurrentLevelData.FinalizeImport();
+	TMap<FGuid, int32> SubLevelTotalBricks;
+	TMap<FGuid, int32> SubLevelBrickAllocator;
+	TMap<FGuid, FIntVector> SubLevelBrickLayoutDimensions;
 
-	float ImportTime = FPlatformTime::Seconds() - StartTime;
-	UE_LOG(LogVolumetricLightmapImport, Log,  TEXT("Imported Volumetric Lightmap in %.3fs"), ImportTime);
-	UE_LOG(LogVolumetricLightmapImport, Log,  TEXT("     Indirection Texture %ux%ux%u = %.1fMb"), 
-		CurrentLevelData.IndirectionTextureDimensions.X,
-		CurrentLevelData.IndirectionTextureDimensions.Y,
-		CurrentLevelData.IndirectionTextureDimensions.Z,
-		CurrentLevelData.IndirectionTexture.Data.Num() / 1024.0f / 1024.0f);
+	{
+		for (int32 CurrentDepth = 0; CurrentDepth < VolumetricLightmapSettings.MaxRefinementLevels; CurrentDepth++)
+		{
+			const TArray<const FImportedVolumetricLightmapBrick*>& BricksAtCurrentDepth = BricksByDepth[CurrentDepth];
 
-	int32 BrickDataSize = CurrentLevelData.BrickData.AmbientVector.Data.Num() + CurrentLevelData.BrickData.SkyBentNormal.Data.Num() + CurrentLevelData.BrickData.DirectionalLightShadowing.Data.Num()
-			+ CurrentLevelData.BrickData.LQLightColor.Data.Num() + CurrentLevelData.BrickData.LQLightColor.Data.Num();
+			for (int32 BrickIndex = 0; BrickIndex < BricksAtCurrentDepth.Num(); BrickIndex++)
+			{
+				const FImportedVolumetricLightmapBrick& Brick = *BricksAtCurrentDepth[BrickIndex];
 
-	for (int32 i = 0; i < ARRAY_COUNT(CurrentLevelData.BrickData.SHCoefficients); i++)
-	{
-		BrickDataSize += CurrentLevelData.BrickData.SHCoefficients[i].Data.Num();
-	}
+				FGuid LevelGuid;
 
-	int32 TotalNumBricks = 0;
+				if (CurrentDepth == VolumetricLightmapSettings.MaxRefinementLevels - 1)
+				{
+					LevelGuid = Brick.IntersectingLevelGuid;
+				}
+				else
+				{
+					// Top level bricks are put into persistent level
+					LevelGuid = PersistentLevelGuid;
+				}
 
-	for (int32 CurrentDepth = 0; CurrentDepth < VolumetricLightmapSettings.MaxRefinementLevels; CurrentDepth++)
-	{
-		const TArray<const FImportedVolumetricLightmapBrick*>& BricksAtCurrentDepth = BricksByDepth[CurrentDepth];
-		TotalNumBricks += BricksAtCurrentDepth.Num();
-	}
+				if (!SubLevelTotalBricks.Contains(LevelGuid))
+				{
+					SubLevelTotalBricks.Add(LevelGuid, 0);
+					SubLevelBrickAllocator.Add(LevelGuid, 0);
+					SubLevelBrickLayoutDimensions.Add(LevelGuid, FIntVector(0));
+				}
 
-	const int32 ActualBrickSizeBytes = BrickDataSize / TotalNumBricks;
+				SubLevelTotalBricks[LevelGuid]++;
+			}
+		}
 
-	FString TrimmedString;
-	
-	if (NumBottomLevelBricksTrimmedByInterpolationError > 0)
-	{
-		TrimmedString += FString::Printf(TEXT(" (trimmed %.1fMb due to %f MinBrickError)"),
-			NumBottomLevelBricksTrimmedByInterpolationError * ActualBrickSizeBytes / 1024.0f / 1024.0f,
-			VolumetricLightmapSettings.MinBrickError);
+		for (auto Pair : SubLevelTotalBricks)
+		{
+			bool bPersistentLevel = Pair.Key == FGuid();
+
+			FIntVector& SubLevelBrickLayoutDimension = SubLevelBrickLayoutDimensions[Pair.Key];
+			int32 SubLevelBrickTextureLinearAllocator = SubLevelTotalBricks[Pair.Key];
+			SubLevelBrickLayoutDimension.X = FMath::Min(SubLevelBrickTextureLinearAllocator, MaxBricksInLayoutOneDim);
+			SubLevelBrickTextureLinearAllocator = FMath::DivideAndRoundUp(SubLevelBrickTextureLinearAllocator, SubLevelBrickLayoutDimension.X);
+			SubLevelBrickLayoutDimension.Y = FMath::Min(SubLevelBrickTextureLinearAllocator, MaxBricksInLayoutOneDim);
+			SubLevelBrickTextureLinearAllocator = FMath::DivideAndRoundUp(SubLevelBrickTextureLinearAllocator, SubLevelBrickLayoutDimension.Y);
+			SubLevelBrickLayoutDimension.Z = FMath::Min(SubLevelBrickTextureLinearAllocator, MaxBricksInLayoutOneDim);
+
+			{
+				ULevel* SubLevelStorageLevel = System.LightingScenario ? System.LightingScenario : FindLevel(Pair.Key);
+				UMapBuildDataRegistry* SubLevelRegistry = SubLevelStorageLevel->GetOrCreateMapBuildData();
+				FPrecomputedVolumetricLightmapData& SubLevelData = SubLevelRegistry->AllocateLevelPrecomputedVolumetricLightmapBuildData(FindLevel(Pair.Key)->LevelBuildDataId);
+
+				SubLevelData.InitializeOnImport(FBox(VolumetricLightmapSettings.VolumeMin, VolumetricLightmapSettings.VolumeMin + VolumetricLightmapSettings.VolumeSize), BrickSize);
+				{
+					SubLevelData.BrickData.AmbientVector.Format = PF_FloatR11G11B10;
+					SubLevelData.BrickData.SkyBentNormal.Format = PF_B8G8R8A8;
+					SubLevelData.BrickData.DirectionalLightShadowing.Format = PF_G8;
+
+					for (int32 i = 0; i < ARRAY_COUNT(SubLevelData.BrickData.SHCoefficients); i++)
+					{
+						SubLevelData.BrickData.SHCoefficients[i].Format = PF_B8G8R8A8;
+					}
+
+					SubLevelData.BrickData.LQLightColor.Format = PF_FloatR11G11B10;
+					SubLevelData.BrickData.LQLightDirection.Format = PF_B8G8R8A8;
+				}
+
+				SubLevelData.BrickDataDimensions = SubLevelBrickLayoutDimension * PaddedBrickSize;
+				const int32 TotalBrickDataSize = SubLevelData.BrickDataDimensions.X * SubLevelData.BrickDataDimensions.Y * SubLevelData.BrickDataDimensions.Z;
+
+				SubLevelData.BrickData.AmbientVector.Resize(TotalBrickDataSize * GPixelFormats[SubLevelData.BrickData.AmbientVector.Format].BlockBytes);
+
+				if (bGenerateSkyShadowing)
+				{
+					SubLevelData.BrickData.SkyBentNormal.Resize(TotalBrickDataSize * GPixelFormats[SubLevelData.BrickData.SkyBentNormal.Format].BlockBytes);
+				}
+
+				SubLevelData.BrickData.DirectionalLightShadowing.Resize(TotalBrickDataSize * GPixelFormats[SubLevelData.BrickData.DirectionalLightShadowing.Format].BlockBytes);
+
+				for (int32 i = 0; i < ARRAY_COUNT(SubLevelData.BrickData.SHCoefficients); i++)
+				{
+					const int32 Stride = GPixelFormats[SubLevelData.BrickData.SHCoefficients[i].Format].BlockBytes;
+					SubLevelData.BrickData.SHCoefficients[i].Resize(TotalBrickDataSize * Stride);
+				}
+
+				SubLevelData.BrickData.LQLightColor.Resize(TotalBrickDataSize * GPixelFormats[SubLevelData.BrickData.LQLightColor.Format].BlockBytes);
+
+				SubLevelData.BrickData.LQLightDirection.Resize(TotalBrickDataSize * GPixelFormats[SubLevelData.BrickData.LQLightDirection.Format].BlockBytes);
+			}
+		}
+
+		BrickStartAllocation = 0;
+
+		for (int32 CurrentDepth = 0; CurrentDepth < VolumetricLightmapSettings.MaxRefinementLevels; CurrentDepth++)
+		{
+			const TArray<const FImportedVolumetricLightmapBrick*>& BricksAtCurrentDepth = BricksByDepth[CurrentDepth];
+
+			for (int32 BrickIndex = 0; BrickIndex < BricksAtCurrentDepth.Num(); BrickIndex++)
+			{
+				const FImportedVolumetricLightmapBrick& Brick = *BricksAtCurrentDepth[BrickIndex];
+
+				FGuid LevelGuid;
+
+				if (CurrentDepth == VolumetricLightmapSettings.MaxRefinementLevels - 1)
+				{
+					LevelGuid = Brick.IntersectingLevelGuid;
+				}
+				else
+				{
+					// Top level bricks are put into persistent level
+					LevelGuid = PersistentLevelGuid;
+				}
+
+				ULevel* SubLevelStorageLevel = System.LightingScenario ? System.LightingScenario : FindLevel(LevelGuid);
+				UMapBuildDataRegistry* SubLevelRegistry = SubLevelStorageLevel->GetOrCreateMapBuildData();
+				FPrecomputedVolumetricLightmapData& SubLevelData = *SubLevelRegistry->GetLevelPrecomputedVolumetricLightmapBuildData(FindLevel(LevelGuid)->LevelBuildDataId);
+
+				const FIntVector BrickLayoutPosition = ComputeBrickLayoutPosition(BrickStartAllocation + BrickIndex, BrickLayoutDimensions) * PaddedBrickSize;
+				const FIntVector SubLevelBrickLayoutPosition = ComputeBrickLayoutPosition(SubLevelBrickAllocator[LevelGuid], SubLevelBrickLayoutDimensions[LevelGuid]) * PaddedBrickSize;
+
+				CopyBetweenAtlasVolumeTextures(
+					GPixelFormats[CurrentLevelData.BrickData.AmbientVector.Format].BlockBytes,
+					FIntVector(PaddedBrickSize),
+					CurrentLevelData.BrickDataDimensions,
+					BrickLayoutPosition,
+					CurrentLevelData.BrickData.AmbientVector.Data,
+					SubLevelBrickLayoutDimensions[LevelGuid] * PaddedBrickSize,
+					SubLevelBrickLayoutPosition,
+					SubLevelData.BrickData.AmbientVector.Data);
+
+				for (int32 i = 0; i < ARRAY_COUNT(Brick.SHCoefficients); i++)
+				{
+					CopyBetweenAtlasVolumeTextures(
+						GPixelFormats[CurrentLevelData.BrickData.SHCoefficients[i].Format].BlockBytes,
+						FIntVector(PaddedBrickSize),
+						CurrentLevelData.BrickDataDimensions,
+						BrickLayoutPosition,
+						CurrentLevelData.BrickData.SHCoefficients[i].Data,
+						SubLevelBrickLayoutDimensions[LevelGuid] * PaddedBrickSize,
+						SubLevelBrickLayoutPosition,
+						SubLevelData.BrickData.SHCoefficients[i].Data);
+				}
+
+				CopyBetweenAtlasVolumeTextures(
+					GPixelFormats[CurrentLevelData.BrickData.LQLightColor.Format].BlockBytes,
+					FIntVector(PaddedBrickSize),
+					CurrentLevelData.BrickDataDimensions,
+					BrickLayoutPosition,
+					CurrentLevelData.BrickData.LQLightColor.Data,
+					SubLevelBrickLayoutDimensions[LevelGuid] * PaddedBrickSize,
+					SubLevelBrickLayoutPosition,
+					SubLevelData.BrickData.LQLightColor.Data);
+
+				CopyBetweenAtlasVolumeTextures(
+					GPixelFormats[CurrentLevelData.BrickData.LQLightDirection.Format].BlockBytes,
+					FIntVector(PaddedBrickSize),
+					CurrentLevelData.BrickDataDimensions,
+					BrickLayoutPosition,
+					CurrentLevelData.BrickData.LQLightDirection.Data,
+					SubLevelBrickLayoutDimensions[LevelGuid] * PaddedBrickSize,
+					SubLevelBrickLayoutPosition,
+					SubLevelData.BrickData.LQLightDirection.Data);
+
+				if (bGenerateSkyShadowing)
+				{
+					CopyBetweenAtlasVolumeTextures(
+						GPixelFormats[CurrentLevelData.BrickData.SkyBentNormal.Format].BlockBytes,
+						FIntVector(PaddedBrickSize),
+						CurrentLevelData.BrickDataDimensions,
+						BrickLayoutPosition,
+						CurrentLevelData.BrickData.SkyBentNormal.Data,
+						SubLevelBrickLayoutDimensions[LevelGuid] * PaddedBrickSize,
+						SubLevelBrickLayoutPosition,
+						SubLevelData.BrickData.SkyBentNormal.Data);
+				}
+
+				CopyBetweenAtlasVolumeTextures(
+					GPixelFormats[CurrentLevelData.BrickData.DirectionalLightShadowing.Format].BlockBytes,
+					FIntVector(PaddedBrickSize),
+					CurrentLevelData.BrickDataDimensions,
+					BrickLayoutPosition,
+					CurrentLevelData.BrickData.DirectionalLightShadowing.Data,
+					SubLevelBrickLayoutDimensions[LevelGuid] * PaddedBrickSize,
+					SubLevelBrickLayoutPosition,
+					SubLevelData.BrickData.DirectionalLightShadowing.Data);
+
+				SubLevelBrickAllocator[LevelGuid]++;
+
+				if (LevelGuid != PersistentLevelGuid)
+				{
+					SubLevelData.SubLevelBrickPositions.Add(Brick.IndirectionTexturePosition);
+					check(SubLevelData.SubLevelBrickPositions.Num() == SubLevelBrickAllocator[LevelGuid]);
+				}
+			}
+
+			BrickStartAllocation += BricksAtCurrentDepth.Num();
+		}
+
+		{
+			const int32 TotalIndirectionTextureSize = CurrentLevelData.IndirectionTextureDimensions.X * CurrentLevelData.IndirectionTextureDimensions.Y * CurrentLevelData.IndirectionTextureDimensions.Z;
+			CurrentLevelData.IndirectionTexture.Resize(TotalIndirectionTextureSize * IndirectionTextureDataStride);
+			BuildIndirectionTexture(BricksByDepth, VolumetricLightmapSettings, MaxBricksInLayoutOneDim, SubLevelBrickLayoutDimensions[FGuid()], IndirectionTextureDataStride, CurrentLevelData, true);
+
+			ULevel* SubLevelStorageLevel = System.LightingScenario ? System.LightingScenario : FindLevel(PersistentLevelGuid);
+			UMapBuildDataRegistry* SubLevelRegistry = SubLevelStorageLevel->GetOrCreateMapBuildData();
+			FPrecomputedVolumetricLightmapData& PersistentLevelData = *SubLevelRegistry->GetLevelPrecomputedVolumetricLightmapBuildData(FindLevel(PersistentLevelGuid)->LevelBuildDataId);
+
+			PersistentLevelData.InitializeOnImport(FBox(VolumetricLightmapSettings.VolumeMin, VolumetricLightmapSettings.VolumeMin + VolumetricLightmapSettings.VolumeSize), BrickSize);
+			PersistentLevelData.IndirectionTexture = CurrentLevelData.IndirectionTexture;
+			PersistentLevelData.IndirectionTextureDimensions = CurrentLevelData.IndirectionTextureDimensions;
+		}
+
+		{
+			for (int32 CurrentDepth = 0; CurrentDepth < VolumetricLightmapSettings.MaxRefinementLevels; CurrentDepth++)
+			{
+				const TArray<const FImportedVolumetricLightmapBrick*>& BricksAtCurrentDepth = BricksByDepth[CurrentDepth];
+
+				for (int32 BrickIndex = 0; BrickIndex < BricksAtCurrentDepth.Num(); BrickIndex++)
+				{
+					const FImportedVolumetricLightmapBrick& Brick = *BricksAtCurrentDepth[BrickIndex];
+
+					FGuid LevelGuid;
+
+					if (CurrentDepth == VolumetricLightmapSettings.MaxRefinementLevels - 1)
+					{
+						LevelGuid = Brick.IntersectingLevelGuid;
+					}
+					else
+					{
+						// Top level bricks are put into persistent level
+						LevelGuid = PersistentLevelGuid;
+					}
+
+					if (LevelGuid != PersistentLevelGuid)
+					{
+						ULevel* SubLevelStorageLevel = System.LightingScenario ? System.LightingScenario : FindLevel(LevelGuid);
+						UMapBuildDataRegistry* SubLevelRegistry = SubLevelStorageLevel->GetOrCreateMapBuildData();
+						FPrecomputedVolumetricLightmapData& SubLevelData = *SubLevelRegistry->GetLevelPrecomputedVolumetricLightmapBuildData(FindLevel(LevelGuid)->LevelBuildDataId);
+
+						FIntVector IndirectionBrickOffset;
+						int32 IndirectionBrickSize;
+						const FVector IndirectionDataSourceCoordinate = FVector(Brick.IndirectionTexturePosition) + FVector(0.5f, 0.5f, 0.5f);
+						SampleIndirectionTexture(IndirectionDataSourceCoordinate, CurrentLevelData.IndirectionTextureDimensions, CurrentLevelData.IndirectionTexture.Data.GetData(), IndirectionBrickOffset, IndirectionBrickSize);
+
+						SubLevelData.IndirectionTextureOriginalValues.Add(FColor((uint8)IndirectionBrickOffset.X, (uint8)IndirectionBrickOffset.Y, (uint8)IndirectionBrickOffset.Z, (uint8)IndirectionBrickSize));
+					}
+				}
+			}
+
+		}
 	}
 
-	if (NumBottomLevelBricksTrimmedForMemoryLimit > 0)
 	{
-		TrimmedString += FString::Printf(TEXT(" (trimmed %.1fMb due to %.1fMb MaximumBrickMemoryMb)"),
-			NumBottomLevelBricksTrimmedForMemoryLimit * ActualBrickSizeBytes / 1024.0f / 1024.0f,
-			MaximumBrickMemoryMb);
-	}
+		/**
+		 * Statistics
+		 */
 
-	UE_LOG(LogVolumetricLightmapImport, Log,  TEXT("     BrickData %ux%ux%u = %.1fMb%s"), 
-		CurrentLevelData.BrickDataDimensions.X,
-		CurrentLevelData.BrickDataDimensions.Y,
-		CurrentLevelData.BrickDataDimensions.Z,
-		BrickDataSize / 1024.0f / 1024.0f,
-		*TrimmedString);
+		float ImportTime = FPlatformTime::Seconds() - StartTime;
+		UE_LOG(LogVolumetricLightmapImport, Log, TEXT("Imported Volumetric Lightmap in %.3fs"), ImportTime);
+		UE_LOG(LogVolumetricLightmapImport, Log, TEXT("     Indirection Texture %ux%ux%u = %.1fMb"),
+			CurrentLevelData.IndirectionTextureDimensions.X,
+			CurrentLevelData.IndirectionTextureDimensions.Y,
+			CurrentLevelData.IndirectionTextureDimensions.Z,
+			CurrentLevelData.IndirectionTexture.Data.Num() / 1024.0f / 1024.0f);
 
-	UE_LOG(LogVolumetricLightmapImport, Log, TEXT("     Bricks at Level"));
+		uint64 BrickDataSize = 0;
 
-	const float TotalVolume = VolumetricLightmapSettings.VolumeSize.X * VolumetricLightmapSettings.VolumeSize.Y * VolumetricLightmapSettings.VolumeSize.Z;
+		for (auto Pair : SubLevelTotalBricks)
+		{
+			ULevel* SubLevelStorageLevel = System.LightingScenario ? System.LightingScenario : FindLevel(Pair.Key);
+			UMapBuildDataRegistry* SubLevelRegistry = SubLevelStorageLevel->GetOrCreateMapBuildData();
+			FPrecomputedVolumetricLightmapData& SubLevelData = *SubLevelRegistry->GetLevelPrecomputedVolumetricLightmapBuildData(FindLevel(Pair.Key)->LevelBuildDataId);
+
+			SubLevelData.FinalizeImport();
+
+			check(SubLevelData.IndirectionTextureOriginalValues.Num() == SubLevelData.SubLevelBrickPositions.Num());
+
+			BrickDataSize += 
+				SubLevelData.BrickData.AmbientVector.Data.Num() + 
+				SubLevelData.BrickData.SHCoefficients[0].Data.Num() + 
+				SubLevelData.BrickData.SHCoefficients[1].Data.Num() + 
+				SubLevelData.BrickData.SHCoefficients[2].Data.Num() + 
+				SubLevelData.BrickData.SHCoefficients[3].Data.Num() + 
+				SubLevelData.BrickData.SHCoefficients[4].Data.Num() + 
+				SubLevelData.BrickData.SHCoefficients[5].Data.Num() + 
+				SubLevelData.BrickData.SkyBentNormal.Data.Num() + 
+				SubLevelData.BrickData.DirectionalLightShadowing.Data.Num() +
+				SubLevelData.BrickData.LQLightColor.Data.Num() +
+				SubLevelData.BrickData.LQLightDirection.Data.Num() +
+				SubLevelData.SubLevelBrickPositions.Num() * SubLevelData.SubLevelBrickPositions.GetTypeSize() +
+				SubLevelData.IndirectionTextureOriginalValues.Num() * SubLevelData.IndirectionTextureOriginalValues.GetTypeSize();
+		}
 
-	for (int32 CurrentDepth = 0; CurrentDepth < VolumetricLightmapSettings.MaxRefinementLevels; CurrentDepth++)
-	{
-		const int32 DetailCellsPerCurrentLevelBrick = 1 << ((VolumetricLightmapSettings.MaxRefinementLevels - CurrentDepth) * BrickSizeLog2);
-		const FVector CurrentDepthBrickSize = DetailCellSize * DetailCellsPerCurrentLevelBrick;
-		const TArray<const FImportedVolumetricLightmapBrick*>& BricksAtCurrentDepth = BricksByDepth[CurrentDepth];
-		const float CurrentDepthBrickVolume = CurrentDepthBrickSize.X * CurrentDepthBrickSize.Y * CurrentDepthBrickSize.Z;
+		int32 TotalNumBricks = 0;
+
+		for (int32 CurrentDepth = 0; CurrentDepth < VolumetricLightmapSettings.MaxRefinementLevels; CurrentDepth++)
+		{
+			const TArray<const FImportedVolumetricLightmapBrick*>& BricksAtCurrentDepth = BricksByDepth[CurrentDepth];
+			TotalNumBricks += BricksAtCurrentDepth.Num();
+		}
+
+		const int32 ActualBrickSizeBytes = BrickDataSize / TotalNumBricks;
 
-		UE_LOG(LogVolumetricLightmapImport, Log, TEXT("         %u: %.1f%% covering %.1f%% of volume"), 
-			CurrentDepth, 
-			100.0f * BricksAtCurrentDepth.Num() / TotalNumBricks,
-			100.0f * BricksAtCurrentDepth.Num() * CurrentDepthBrickVolume / TotalVolume);
+		FString TrimmedString;
+
+		if (NumBottomLevelBricksTrimmedByInterpolationError > 0)
+		{
+			TrimmedString += FString::Printf(TEXT(" (trimmed %.1fMb due to %f MinBrickError)"),
+				NumBottomLevelBricksTrimmedByInterpolationError * ActualBrickSizeBytes / 1024.0f / 1024.0f,
+				VolumetricLightmapSettings.MinBrickError);
+		}
+
+		if (NumBottomLevelBricksTrimmedForMemoryLimit > 0)
+		{
+			TrimmedString += FString::Printf(TEXT(" (trimmed %.1fMb due to %.1fMb MaximumBrickMemoryMb)"),
+				NumBottomLevelBricksTrimmedForMemoryLimit * ActualBrickSizeBytes / 1024.0f / 1024.0f,
+				MaximumBrickMemoryMb);
+		}
+
+		UE_LOG(LogVolumetricLightmapImport, Log, TEXT("     BrickData (all levels) %.1fMb%s"),
+			BrickDataSize / 1024.0f / 1024.0f,
+			*TrimmedString);
+
+		UE_LOG(LogVolumetricLightmapImport, Log, TEXT("     Bricks at depth"));
+
+		const float TotalVolume = VolumetricLightmapSettings.VolumeSize.X * VolumetricLightmapSettings.VolumeSize.Y * VolumetricLightmapSettings.VolumeSize.Z;
+
+		for (int32 CurrentDepth = 0; CurrentDepth < VolumetricLightmapSettings.MaxRefinementLevels; CurrentDepth++)
+		{
+			const int32 DetailCellsPerCurrentLevelBrick = 1 << ((VolumetricLightmapSettings.MaxRefinementLevels - CurrentDepth) * BrickSizeLog2);
+			const FVector CurrentDepthBrickSize = DetailCellSize * DetailCellsPerCurrentLevelBrick;
+			const TArray<const FImportedVolumetricLightmapBrick*>& BricksAtCurrentDepth = BricksByDepth[CurrentDepth];
+			const float CurrentDepthBrickVolume = CurrentDepthBrickSize.X * CurrentDepthBrickSize.Y * CurrentDepthBrickSize.Z;
+
+			UE_LOG(LogVolumetricLightmapImport, Log, TEXT("         %u: %.1f%% covering %.1f%% of volume"),
+				CurrentDepth,
+				100.0f * BricksAtCurrentDepth.Num() / TotalNumBricks,
+				100.0f * BricksAtCurrentDepth.Num() * CurrentDepthBrickVolume / TotalVolume);
+		}
+
+		UE_LOG(LogVolumetricLightmapImport, Log, TEXT("     Bricks in each level"));
+
+		for (auto Pair : SubLevelTotalBricks)
+		{
+			FString LevelName = FindLevel(Pair.Key)->GetOuter()->GetName();
+			if (FindLevel(Pair.Key)->IsPersistentLevel())
+			{
+				UE_LOG(LogVolumetricLightmapImport, Log, TEXT("         %s \t\t\t %d bricks \t %.1f%% (Persistent Level)"), *LevelName, Pair.Value, 100.0f * Pair.Value / TotalNumBricks);
+			}
+			else
+			{
+				UE_LOG(LogVolumetricLightmapImport, Log, TEXT("         %s \t\t\t %d bricks \t %.1f%%"), *LevelName, Pair.Value, 100.0f * Pair.Value / TotalNumBricks);
+			}
+		}
 	}
 }
 
diff --git a/Engine/Source/Editor/UnrealEd/Private/Lightmass/Lightmass.cpp b/Engine/Source/Editor/UnrealEd/Private/Lightmass/Lightmass.cpp
--- a/Engine/Source/Editor/UnrealEd/Private/Lightmass/Lightmass.cpp
+++ b/Engine/Source/Editor/UnrealEd/Private/Lightmass/Lightmass.cpp
@@ -58,7 +58,7 @@
 
 extern FSwarmDebugOptions GSwarmDebugOptions;
 
-DEFINE_LOG_CATEGORY_STATIC(LogLightmassSolver, Error, All);
+DEFINE_LOG_CATEGORY_STATIC(LogLightmassSolver, Warning, All);
 /**
  * If false (default behavior), Lightmass is launched automatically when a lighting build starts.
  * If true, it must be launched manually (e.g. through a debugger).
@@ -2147,24 +2147,26 @@ void FLightmassExporter::SetVolumetricLightmapSettings(Lightmass::FVolumetricLig
 	const float TargetDetailCellSize = WorldInfoSettings.VolumetricLightmapDetailCellSize;
 
 	FIntVector FullGridSize(
-		FMath::TruncToInt(RequiredVolumeSize.X / TargetDetailCellSize) + 1,
-		FMath::TruncToInt(RequiredVolumeSize.Y / TargetDetailCellSize) + 1,
-		FMath::TruncToInt(RequiredVolumeSize.Z / TargetDetailCellSize) + 1);
+		FMath::TruncToInt(2 * ImportanceExtent.X / TargetDetailCellSize) + 1,
+		FMath::TruncToInt(2 * ImportanceExtent.Y / TargetDetailCellSize) + 1,
+		FMath::TruncToInt(2 * ImportanceExtent.Z / TargetDetailCellSize) + 1);
 
-	if (FullGridSize.GetMax() > 800)
+	// Make sure size of indirection texture does not exceed INT_MAX
+	const int32 MaxGridDimension = FMath::Pow(2048000000 / 4, 1.0f / 3) * OutSettings.BrickSize;
+	if (FullGridSize.GetMax() > MaxGridDimension)
 	{
 		UE_LOG(LogLightmassSolver, Warning, 
 			TEXT("Volumetric lightmap grid size is too large which can cause potential crashes or long build time, clamping from (%d, %d, %d) to (%d, %d, %d)"),
 			FullGridSize.X,
 			FullGridSize.Y,
 			FullGridSize.Z,
-			FMath::Min(FullGridSize.X, 800),
-			FMath::Min(FullGridSize.Y, 800),
-			FMath::Min(FullGridSize.Z, 800)
+			FMath::Min(FullGridSize.X, MaxGridDimension),
+			FMath::Min(FullGridSize.Y, MaxGridDimension),
+			FMath::Min(FullGridSize.Z, MaxGridDimension)
 			);
-		FullGridSize.X = FMath::Min(FullGridSize.X, 800);
-		FullGridSize.Y = FMath::Min(FullGridSize.Y, 800);
-		FullGridSize.Z = FMath::Min(FullGridSize.Z, 800);
+		FullGridSize.X = FMath::Min(FullGridSize.X, MaxGridDimension);
+		FullGridSize.Y = FMath::Min(FullGridSize.Y, MaxGridDimension);
+		FullGridSize.Z = FMath::Min(FullGridSize.Z, MaxGridDimension);
 	}
 
 	const int32 BrickSizeLog2 = FMath::FloorLog2(OutSettings.BrickSize);
@@ -4505,7 +4507,7 @@ UStaticMesh* FLightmassProcessor::FindStaticMesh(FGuid& Guid)
 	return NULL;
 }
 
-ULevel* FLightmassProcessor::FindLevel(FGuid& Guid)
+ULevel* FLightmassProcessor::FindLevel(const FGuid& Guid)
 {
 	if (Exporter)
 	{
diff --git a/Engine/Source/Editor/UnrealEd/Private/Lightmass/Lightmass.h b/Engine/Source/Editor/UnrealEd/Private/Lightmass/Lightmass.h
--- a/Engine/Source/Editor/UnrealEd/Private/Lightmass/Lightmass.h
+++ b/Engine/Source/Editor/UnrealEd/Private/Lightmass/Lightmass.h
@@ -573,7 +573,7 @@ protected:
 	 */
 	UStaticMesh* FindStaticMesh(FGuid& Guid);
 
-	ULevel* FindLevel(FGuid& Guid);
+	ULevel* FindLevel(const FGuid& Guid);
 
 	/**
 	 *	Import light map data from the given channel.
diff --git a/Engine/Source/Programs/UnrealLightmass/Private/ImportExport/Exporter.cpp b/Engine/Source/Programs/UnrealLightmass/Private/ImportExport/Exporter.cpp
--- a/Engine/Source/Programs/UnrealLightmass/Private/ImportExport/Exporter.cpp
+++ b/Engine/Source/Programs/UnrealLightmass/Private/ImportExport/Exporter.cpp
@@ -491,6 +491,7 @@ namespace Lightmass
 
 			for (int32 BrickIndex = 0; BrickIndex < NumBricks; BrickIndex++)
 			{
+				Swarm->Write(&TaskData.BrickData[BrickIndex].IntersectingLevelGuid, sizeof(TaskData.BrickData[BrickIndex].IntersectingLevelGuid));
 				Swarm->Write(&TaskData.BrickData[BrickIndex].IndirectionTexturePosition, sizeof(TaskData.BrickData[BrickIndex].IndirectionTexturePosition));
 				Swarm->Write(&TaskData.BrickData[BrickIndex].TreeDepth, sizeof(TaskData.BrickData[BrickIndex].TreeDepth));
 				Swarm->Write(&TaskData.BrickData[BrickIndex].AverageClosestGeometryDistance, sizeof(TaskData.BrickData[BrickIndex].AverageClosestGeometryDistance));
diff --git a/Engine/Source/Programs/UnrealLightmass/Private/Lighting/AdaptiveVolumetricLightmap.cpp b/Engine/Source/Programs/UnrealLightmass/Private/Lighting/AdaptiveVolumetricLightmap.cpp
--- a/Engine/Source/Programs/UnrealLightmass/Private/Lighting/AdaptiveVolumetricLightmap.cpp
+++ b/Engine/Source/Programs/UnrealLightmass/Private/Lighting/AdaptiveVolumetricLightmap.cpp
@@ -119,7 +119,7 @@ bool PointUnderTriangle(FVector Point, FTriangle Triangle)
 	return BaryCentricCoordinate.X >= 0 && BaryCentricCoordinate.Y >= 0 && BaryCentricCoordinate.Z >= 0 && PointOnTriangleZ > Point.Z;
 }
 
-bool FStaticLightingSystem::DoesVoxelIntersectSceneGeometry(const FBox& CellBounds) const
+bool FStaticLightingSystem::DoesVoxelIntersectSceneGeometry(const FBox& CellBounds, FGuid& OutIntersectingLevelGuid) const
 {
 	const float Child2dTriangleArea = .5f * CellBounds.GetSize().X * CellBounds.GetSize().Y / (VolumetricLightmapSettings.BrickSize * VolumetricLightmapSettings.BrickSize);
 	const float SurfaceLightmapDensityThreshold = .5f * VolumetricLightmapSettings.SurfaceLightmapMinTexelsPerVoxelAxis * VolumetricLightmapSettings.SurfaceLightmapMinTexelsPerVoxelAxis / Child2dTriangleArea;
@@ -129,8 +129,18 @@ bool FStaticLightingSystem::DoesVoxelIntersectSceneGeometry(const FBox& CellBoun
 
 	if (Scene.GeneralSettings.bUseFastVoxelization)
 	{
-		if (VoxelizationSurfaceAggregateMesh->IntersectBox(ExpandedCellBoundsSurfaceGeometry) || VoxelizationVolumeAggregateMesh->IntersectBox(ExpandedCellBoundsVolumeGeometry))
+		const FStaticLightingMesh* Mesh = nullptr;
+		Mesh = VoxelizationSurfaceAggregateMesh->IntersectBox(ExpandedCellBoundsSurfaceGeometry);
+		if (Mesh != nullptr)
 		{
+			OutIntersectingLevelGuid = Mesh->LevelGuid;
+			return true;
+		}
+
+		Mesh = VoxelizationVolumeAggregateMesh->IntersectBox(ExpandedCellBoundsVolumeGeometry);
+		if (Mesh != nullptr)
+		{
+			OutIntersectingLevelGuid = Mesh->LevelGuid;
 			return true;
 		}
 
@@ -140,17 +150,28 @@ bool FStaticLightingSystem::DoesVoxelIntersectSceneGeometry(const FBox& CellBoun
 
 			if (MeshInstance->StaticMesh->VoxelizationMesh != nullptr)
 			{
+				auto TransformBox = [](const FBox& Box, const FMatrix& Transform) -> const FBox {
+					FVector TransformedCorners[8];
+					for (int32 Corner = 0; Corner < 8; Corner++)
+					{
+						TransformedCorners[Corner] = Transform.TransformPosition(Box.GetExtrema(Corner));
+					}
+					return FBox(TransformedCorners, ARRAY_COUNT(TransformedCorners));
+				};
+
 				if (MeshInstance->Mapping->GetVolumeMapping() == nullptr)
 				{
-					if (MeshInstance->StaticMesh->VoxelizationMesh->IntersectBox(ExpandedCellBoundsSurfaceGeometry.TransformBy(MeshInstance->WorldToLocal)))
+					if (MeshInstance->StaticMesh->VoxelizationMesh->IntersectBox(ExpandedCellBoundsSurfaceGeometry.TransformBy(MeshInstance->WorldToLocal)) != nullptr)
 					{
+						OutIntersectingLevelGuid = MeshInstance->LevelGuid;
 						return true;
 					}
 				}
 				else
 				{
-					if (MeshInstance->StaticMesh->VoxelizationMesh->IntersectBox(ExpandedCellBoundsVolumeGeometry.TransformBy(MeshInstance->WorldToLocal)))
+					if (MeshInstance->StaticMesh->VoxelizationMesh->IntersectBox(ExpandedCellBoundsVolumeGeometry.TransformBy(MeshInstance->WorldToLocal)) != nullptr)
 					{
+						OutIntersectingLevelGuid = MeshInstance->LevelGuid;
 						return true;
 					}
 				}
@@ -222,6 +243,7 @@ bool FStaticLightingSystem::DoesVoxelIntersectSceneGeometry(const FBox& CellBoun
 
 								if (IntersectTriangleAndAABB(Triangle, ExpandedCellBounds))
 								{
+									OutIntersectingLevelGuid = CurrentMesh->LevelGuid;
 									return true;
 								}
 							}
@@ -235,7 +257,7 @@ bool FStaticLightingSystem::DoesVoxelIntersectSceneGeometry(const FBox& CellBoun
 	}
 }
 
-bool FStaticLightingSystem::ShouldRefineVoxel(int32 TreeDepth, const FBox& CellBounds, const TArray<FVector>& VoxelTestPositions, bool bDebugThisVoxel) const
+bool FStaticLightingSystem::ShouldRefineVoxel(int32 TreeDepth, const FBox& CellBounds, const TArray<FVector>& VoxelTestPositions, bool bDebugThisVoxel, FGuid& OutIntersectingLevelGuid) const
 {
 	const bool bCellInsideImportanceVolume = Scene.IsBoxInImportanceVolume(CellBounds);
 
@@ -282,7 +304,7 @@ bool FStaticLightingSystem::ShouldRefineVoxel(int32 TreeDepth, const FBox& CellB
 		return true;
 	}
 
-	bool bVoxelIntersectsScene = DoesVoxelIntersectSceneGeometry(CellBounds);
+	bool bVoxelIntersectsScene = DoesVoxelIntersectSceneGeometry(CellBounds, OutIntersectingLevelGuid);
 	
 	if (!bVoxelIntersectsScene)
 	{
@@ -405,6 +427,7 @@ bool FStaticLightingSystem::ShouldRefineVoxel(int32 TreeDepth, const FBox& CellB
 
 struct FIrradianceBrickBuildData
 {
+	FGuid IntersectingLevelGuid;
 	FIntVector LocalCellCoordinate;
 	int32 TreeDepth;
 	bool bHasChildren;
@@ -423,11 +446,13 @@ void FStaticLightingSystem::RecursivelyBuildBrickTree(
 	bool bCoveringDebugPosition,
 	const FBox& TopLevelCellBounds, 
 	const TArray<FVector>& VoxelTestPositions,
+	const FGuid& IntersectingLevelGuid,
 	TArray<FIrradianceBrickBuildData>& OutBrickBuildData)
 {
 	if (StartCellIndex == 0)
 	{
 		FIrradianceBrickBuildData NewBuildData;
+		NewBuildData.IntersectingLevelGuid = IntersectingLevelGuid;
 		NewBuildData.LocalCellCoordinate = LocalCellCoordinate;
 		NewBuildData.TreeDepth = TreeDepth;
 		NewBuildData.bDebugBrick = bCoveringDebugPosition;
@@ -466,7 +491,8 @@ void FStaticLightingSystem::RecursivelyBuildBrickTree(
 						const FBox CellBounds(ChildCellPosition, ChildCellPosition + WorldChildCellSize);
 
 						const bool bChildCoveringDebugPosition = bDebugVolumetricLightmapCell && CellBounds.IsInside(DebugWorldPosition);
-						const bool bSubdivideCell = ShouldRefineVoxel(TreeDepth + 1, CellBounds, VoxelTestPositions, bChildCoveringDebugPosition);
+						FGuid ChildIntersectingLevelGuid;
+						const bool bSubdivideCell = ShouldRefineVoxel(TreeDepth + 1, CellBounds, VoxelTestPositions, bChildCoveringDebugPosition, ChildIntersectingLevelGuid);
 
 						if (bSubdivideCell)
 						{
@@ -477,7 +503,7 @@ void FStaticLightingSystem::RecursivelyBuildBrickTree(
 								Y * DetailCellsPerChildLevelBrick, 
 								Z * DetailCellsPerChildLevelBrick);
 
-							RecursivelyBuildBrickTree(0, NumCellsPerBrick, LocalCellCoordinate + LocalChildCellCoordinate, TreeDepth + 1, bChildCoveringDebugPosition, TopLevelCellBounds, VoxelTestPositions, OutBrickBuildData);
+							RecursivelyBuildBrickTree(0, NumCellsPerBrick, LocalCellCoordinate + LocalChildCellCoordinate, TreeDepth + 1, bChildCoveringDebugPosition, TopLevelCellBounds, VoxelTestPositions, ChildIntersectingLevelGuid, OutBrickBuildData);
 						}
 					}
 				}
@@ -782,7 +808,7 @@ void FStaticLightingSystem::CalculateAdaptiveVolumetricLightmap(int32 TaskIndex)
 	GenerateVoxelTestPositions(VoxelTestPositions);
 
 	TArray<FIrradianceBrickBuildData> BrickBuildData;
-	RecursivelyBuildBrickTree(StartCellIndex, NumCells, FIntVector::ZeroValue, 0, bCoveringDebugPosition, TopLevelBounds, VoxelTestPositions, BrickBuildData);
+	RecursivelyBuildBrickTree(StartCellIndex, NumCells, FIntVector::ZeroValue, 0, bCoveringDebugPosition, TopLevelBounds, VoxelTestPositions, FGuid(), BrickBuildData);
 
 	MappingContext.Stats.VolumetricLightmapVoxelizationTime += FPlatformTime::Seconds() - StartTime;
 
@@ -796,6 +822,11 @@ void FStaticLightingSystem::CalculateAdaptiveVolumetricLightmap(int32 TaskIndex)
 		DataLink->Element.BrickData.Empty(BrickBuildData.Num());
 		DataLink->Element.BrickData.AddDefaulted(BrickBuildData.Num());
 
+		for (int32 BrickIndex = 0; BrickIndex < BrickBuildData.Num(); BrickIndex++)
+		{
+			DataLink->Element.BrickData[BrickIndex].IntersectingLevelGuid = BrickBuildData[BrickIndex].IntersectingLevelGuid;
+		}
+
 		// Calculate lighting for all bricks
 		for (int32 BrickIndex = 0; BrickIndex < BrickBuildData.Num(); BrickIndex++)
 		{
diff --git a/Engine/Source/Programs/UnrealLightmass/Private/Lighting/Collision.cpp b/Engine/Source/Programs/UnrealLightmass/Private/Lighting/Collision.cpp
--- a/Engine/Source/Programs/UnrealLightmass/Private/Lighting/Collision.cpp
+++ b/Engine/Source/Programs/UnrealLightmass/Private/Lighting/Collision.cpp
@@ -520,14 +520,25 @@ bool FDefaultAggregateMesh::IntersectLightRay(
 	return ClosestIntersection.bIntersects;
 }
 
-bool FDefaultAggregateMesh::IntersectBox(const FBox Box)  const
+const FStaticLightingMesh* FDefaultAggregateMesh::IntersectBox(const FBox Box)  const
 {
 	FLightRay UnusedRay;
 	FStaticLightingAggregateMeshDataProvider kDOPDataProvider(this, UnusedRay);
 	TkDOPBoxCollisionCheck<const FStaticLightingAggregateMeshDataProvider, uint32> kDOPCheck(
 		Box,
 		kDOPDataProvider);
-	return kDopTree.BoxCheck(kDOPCheck);
+
+	kDopTree.BoxCheck(kDOPCheck);
+
+	if (kDOPCheck.Item != -1)
+	{
+		const FTriangleSOAPayload& Payload = TrianglePayloads[kDOPCheck.Item];
+		return Payload.MeshInfo->Mesh;
+	}
+	else
+	{
+		return nullptr;
+	}
 }
 
 } //namespace Lightmass
diff --git a/Engine/Source/Programs/UnrealLightmass/Private/Lighting/Collision.h b/Engine/Source/Programs/UnrealLightmass/Private/Lighting/Collision.h
--- a/Engine/Source/Programs/UnrealLightmass/Private/Lighting/Collision.h
+++ b/Engine/Source/Programs/UnrealLightmass/Private/Lighting/Collision.h
@@ -248,7 +248,7 @@ public:
 		class FCoherentRayCache& CoherentRayCache,
 		FLightRayIntersection& Intersection) const override;
 
-	bool IntersectBox(const FBox Box) const;
+	const FStaticLightingMesh* IntersectBox(const FBox Box) const;
 
 private:
 
diff --git a/Engine/Source/Programs/UnrealLightmass/Private/Lighting/LightingSystem.h b/Engine/Source/Programs/UnrealLightmass/Private/Lighting/LightingSystem.h
--- a/Engine/Source/Programs/UnrealLightmass/Private/Lighting/LightingSystem.h
+++ b/Engine/Source/Programs/UnrealLightmass/Private/Lighting/LightingSystem.h
@@ -1348,6 +1348,8 @@ struct FPrecomputedVisibilityData
 
 struct FIrradianceBrickData
 {
+	FGuid IntersectingLevelGuid;
+
 	/** Position in the global indirection texture.  Used for mapping brick positions back to world space. */
 	FIntVector IndirectionTexturePosition;
 
@@ -2222,6 +2224,7 @@ private:
 		bool bCoveringDebugPosition,
 		const FBox& TopLevelCellBounds, 
 		const TArray<FVector>& VoxelTestPositions,
+		const FGuid& IntersectingLevelGuid,
 		TArray<struct FIrradianceBrickBuildData>& OutBrickBuildData);
 
 	void ProcessVolumetricLightmapBrickTask(class FVolumetricLightmapBrickTaskDescription* Task);
@@ -2232,9 +2235,9 @@ private:
 
 	void CalculateAdaptiveVolumetricLightmap(int32 TaskIndex);
 
-	bool DoesVoxelIntersectSceneGeometry(const FBox& CellBounds) const;
+	bool DoesVoxelIntersectSceneGeometry(const FBox& CellBounds, FGuid& OutIntersectingLevelGuid) const;
 
-	bool ShouldRefineVoxel(int32 TreeDepth, const FBox& AABB, const TArray<FVector>& VoxelTestPositions, bool bDebugThisVoxel) const;
+	bool ShouldRefineVoxel(int32 TreeDepth, const FBox& AABB, const TArray<FVector>& VoxelTestPositions, bool bDebugThisVoxel, FGuid& OutIntersectingLevelGuid) const;
 
 	/** Computes a shadow depth map for a stationary light. */
 	void CalculateStaticShadowDepthMap(FGuid LightGuid);
diff --git a/Engine/Source/Programs/UnrealLightmass/Private/LightmassCore/Math/LMkDOP.h b/Engine/Source/Programs/UnrealLightmass/Private/LightmassCore/Math/LMkDOP.h
--- a/Engine/Source/Programs/UnrealLightmass/Private/LightmassCore/Math/LMkDOP.h
+++ b/Engine/Source/Programs/UnrealLightmass/Private/LightmassCore/Math/LMkDOP.h
@@ -1007,7 +1007,10 @@ struct TkDOPNode
 				Triangle.Vertices[2].Z = LmVectorGetComponent(TriangleSOA.Positions[2].Z, SubIndex);
 
 				if (BoxTriangleIntersectionInternal::IntersectTriangleAndAABB(Triangle, Check.Box))
+				{
+					Check.Item = TriangleSOA.Payload[SubIndex];
 					return true;
+				}
 			}
 		}
 
@@ -1552,12 +1555,14 @@ template <typename COLL_DATA_PROVIDER, typename KDOP_IDX_TYPE> struct TkDOPBoxCo
 	public TkDOPCollisionCheck<COLL_DATA_PROVIDER, KDOP_IDX_TYPE>
 {
 	const FBox Box;
+	int32 Item;
 
 	TkDOPBoxCollisionCheck(
 		const FBox& InBox,
 		const COLL_DATA_PROVIDER& InCollDataProvider) :
 		TkDOPCollisionCheck<COLL_DATA_PROVIDER, KDOP_IDX_TYPE>(InCollDataProvider),
-		Box(InBox)
+		Box(InBox),
+		Item(-1)
 	{}
 };
 
diff --git a/Engine/Source/Programs/UnrealLightmass/Public/ImportExport.h b/Engine/Source/Programs/UnrealLightmass/Public/ImportExport.h
--- a/Engine/Source/Programs/UnrealLightmass/Public/ImportExport.h
+++ b/Engine/Source/Programs/UnrealLightmass/Public/ImportExport.h
@@ -77,7 +77,7 @@ namespace Lightmass
 
 	static const int32 LM_TEXTUREMAPPING_VERSION		= 1;
 	static const int32 LM_VOLUMESAMPLES_VERSION			= 1;
-	static const int32 LM_VOLUMETRICLIGHTMAP_VERSION		= 2;
+	static const int32 LM_VOLUMETRICLIGHTMAP_VERSION		= 3;
 	static const int32 LM_PRECOMPUTEDVISIBILITY_VERSION	= 1;
 	static const int32 LM_VOLUMEDEBUGOUTPUT_VERSION		= 1;
 	static const int32 LM_DOMINANTSHADOW_VERSION		= 1;
diff --git a/Engine/Source/Runtime/Core/Public/UObject/RenderingObjectVersion.h b/Engine/Source/Runtime/Core/Public/UObject/RenderingObjectVersion.h
--- a/Engine/Source/Runtime/Core/Public/UObject/RenderingObjectVersion.h
+++ b/Engine/Source/Runtime/Core/Public/UObject/RenderingObjectVersion.h
@@ -92,6 +92,8 @@ struct CORE_API FRenderingObjectVersion
 
 		SkyAtmosphereStaticLightingVersioning,
 
+		VolumetricLightmapStreaming,
+
 		// -----<new versions can be added above this line>-------------------------------------------------
 		VersionPlusOne,
 		LatestVersion = VersionPlusOne - 1
diff --git a/Engine/Source/Runtime/Engine/Private/Level.cpp b/Engine/Source/Runtime/Engine/Private/Level.cpp
--- a/Engine/Source/Runtime/Engine/Private/Level.cpp
+++ b/Engine/Source/Runtime/Engine/Private/Level.cpp
@@ -1776,7 +1776,7 @@ void ULevel::InitializeRenderingResources()
 
 		if (!PrecomputedVolumetricLightmap->IsAddedToScene())
 		{
-			PrecomputedVolumetricLightmap->AddToScene(OwningWorld->Scene, EffectiveMapBuildData, LevelBuildDataId);
+			PrecomputedVolumetricLightmap->AddToScene(OwningWorld->Scene, EffectiveMapBuildData, LevelBuildDataId, IsPersistentLevel());
 		}
 
 		if (OwningWorld->Scene && EffectiveMapBuildData)
diff --git a/Engine/Source/Runtime/Engine/Private/MapBuildData.cpp b/Engine/Source/Runtime/Engine/Private/MapBuildData.cpp
--- a/Engine/Source/Runtime/Engine/Private/MapBuildData.cpp
+++ b/Engine/Source/Runtime/Engine/Private/MapBuildData.cpp
@@ -60,26 +60,26 @@ ULevel* UWorld::GetActiveLightingScenario() const
 
 void UWorld::PropagateLightingScenarioChange()
 {
-	for (FActorIterator It(this); It; ++It)
+	for (ULevel* Level : GetLevels())
 	{
-		TInlineComponentArray<USceneComponent*> Components;
-		(*It)->GetComponents(Components);
+		Level->ReleaseRenderingResources();
+		Level->InitializeRenderingResources();
 
-		for (int32 ComponentIndex = 0; ComponentIndex < Components.Num(); ComponentIndex++)
+		for (UModelComponent* ModelComponent : Level->ModelComponents)
 		{
-			USceneComponent* CurrentComponent = Components[ComponentIndex];
-			CurrentComponent->PropagateLightingScenarioChange();
+			ModelComponent->PropagateLightingScenarioChange();
 		}
 	}
 
-	for (ULevel* Level : GetLevels())
+	for (FActorIterator It(this); It; ++It)
 	{
-		Level->ReleaseRenderingResources();
-		Level->InitializeRenderingResources();
+		TInlineComponentArray<USceneComponent*> Components;
+		(*It)->GetComponents(Components);
 
-		for (UModelComponent* ModelComponent : Level->ModelComponents)
+		for (int32 ComponentIndex = 0; ComponentIndex < Components.Num(); ComponentIndex++)
 		{
-			ModelComponent->PropagateLightingScenarioChange();
+			USceneComponent* CurrentComponent = Components[ComponentIndex];
+			CurrentComponent->PropagateLightingScenarioChange();
 		}
 	}
 
diff --git a/Engine/Source/Runtime/Engine/Private/PrecomputedVolumetricLightmap.cpp b/Engine/Source/Runtime/Engine/Private/PrecomputedVolumetricLightmap.cpp
--- a/Engine/Source/Runtime/Engine/Private/PrecomputedVolumetricLightmap.cpp
+++ b/Engine/Source/Runtime/Engine/Private/PrecomputedVolumetricLightmap.cpp
@@ -13,6 +13,8 @@
 #include "Engine/MapBuildDataRegistry.h"
 #include "Interfaces/ITargetPlatform.h"
 #include "UObject/MobileObjectVersion.h"
+#include "Containers/DynamicRHIResourceArray.h"
+#include "RenderGraphUtils.h"
 
 DECLARE_MEMORY_STAT(TEXT("Volumetric Lightmap"),STAT_VolumetricLightmapBuildData,STATGROUP_MapBuildData);
 
@@ -27,10 +29,123 @@ void FVolumetricLightmapDataLayer::CreateTexture(FIntVector Dimensions)
 		Dimensions.Z, 
 		Format,
 		1,
-		TexCreate_ShaderResource | TexCreate_DisableAutoDefrag,
+		TexCreate_ShaderResource | TexCreate_DisableAutoDefrag | TexCreate_UAV,
 		CreateInfo);
 }
 
+void FVolumetricLightmapDataLayer::CreateTargetTexture(FIntVector Dimensions)
+{
+	FRHIResourceCreateInfo CreateInfo;
+
+	Texture = RHICreateTexture3D(
+		Dimensions.X,
+		Dimensions.Y,
+		Dimensions.Z,
+		Format,
+		1,
+		TexCreate_ShaderResource | TexCreate_DisableAutoDefrag | TexCreate_UAV,
+		CreateInfo);
+}
+
+void FVolumetricLightmapDataLayer::CreateUAV()
+{
+	check(Texture);
+
+	UAV = RHICreateUnorderedAccessView(Texture);
+}
+
+struct FVolumetricLightmapBrickTextureSet
+{
+	FIntVector BrickDataDimensions;
+
+	FVolumetricLightmapDataLayer AmbientVector;
+	FVolumetricLightmapDataLayer SHCoefficients[6];
+	FVolumetricLightmapDataLayer SkyBentNormal;
+	FVolumetricLightmapDataLayer DirectionalLightShadowing;
+
+	template<class VolumetricLightmapBrickDataType> // Can be either FVolumetricLightmapBrickData or FVolumetricLightmapBrickTextureSet
+	void Initialize(FIntVector InBrickDataDimensions, VolumetricLightmapBrickDataType& BrickData)
+	{
+		BrickDataDimensions = InBrickDataDimensions;
+
+		AmbientVector.Format = BrickData.AmbientVector.Format;
+		SkyBentNormal.Format = BrickData.SkyBentNormal.Format;
+		DirectionalLightShadowing.Format = BrickData.DirectionalLightShadowing.Format;
+
+		for (int32 i = 0; i < ARRAY_COUNT(SHCoefficients); i++)
+		{
+			SHCoefficients[i].Format = BrickData.SHCoefficients[i].Format;
+		}
+
+		AmbientVector.CreateTargetTexture(BrickDataDimensions);
+		AmbientVector.CreateUAV();
+
+		for (int32 i = 0; i < ARRAY_COUNT(SHCoefficients); i++)
+		{
+			SHCoefficients[i].CreateTargetTexture(BrickDataDimensions);
+			SHCoefficients[i].CreateUAV();
+		}
+
+		if (BrickData.SkyBentNormal.Texture.IsValid())
+		{
+			SkyBentNormal.CreateTargetTexture(BrickDataDimensions);
+			SkyBentNormal.CreateUAV();
+		}
+
+		DirectionalLightShadowing.CreateTargetTexture(BrickDataDimensions);
+		DirectionalLightShadowing.CreateUAV();
+	}
+
+	void Release()
+	{
+		AmbientVector.Texture.SafeRelease();
+		for (int32 i = 0; i < ARRAY_COUNT(SHCoefficients); i++)
+		{
+			SHCoefficients[i].Texture.SafeRelease();
+		}
+		SkyBentNormal.Texture.SafeRelease();
+		DirectionalLightShadowing.Texture.SafeRelease();
+
+		AmbientVector.UAV.SafeRelease();
+		for (int32 i = 0; i < ARRAY_COUNT(SHCoefficients); i++)
+		{
+			SHCoefficients[i].UAV.SafeRelease();
+		}
+		SkyBentNormal.UAV.SafeRelease();
+		DirectionalLightShadowing.UAV.SafeRelease();
+	}
+};
+
+class ENGINE_API FVolumetricLightmapBrickAtlas : public FRenderResource
+{
+public:
+	FVolumetricLightmapBrickAtlas();
+
+	FVolumetricLightmapBrickTextureSet TextureSet;
+
+	virtual void ReleaseRHI() override;
+
+	struct Allocation
+	{
+		// The data being allocated, as an identifier for the entry
+		class FPrecomputedVolumetricLightmapData* Data = nullptr;
+
+		int32 Size = 0;
+		int32 StartOffset = 0;
+	};
+
+	TArray<Allocation> Allocations;
+
+	void Insert(int32 Index, FPrecomputedVolumetricLightmapData* Data);
+	void Remove(FPrecomputedVolumetricLightmapData* Data);
+
+private:
+	bool bInitialized;
+	int32 PaddedBrickSize;
+};
+
+TGlobalResource<FVolumetricLightmapBrickAtlas> GVolumetricLightmapBrickAtlas = TGlobalResource<FVolumetricLightmapBrickAtlas>();
+
 FArchive& operator<<(FArchive& Ar,FVolumetricLightmapDataLayer& Layer)
 {
 	Ar << Layer.Data;
@@ -60,6 +175,7 @@ FArchive& operator<<(FArchive& Ar,FVolumetricLightmapDataLayer& Layer)
 FArchive& operator<<(FArchive& Ar,FPrecomputedVolumetricLightmapData& Volume)
 {
 	Ar.UsingCustomVersion(FMobileObjectVersion::GUID);
+	Ar.UsingCustomVersion(FRenderingObjectVersion::GUID);
 
 	Ar << Volume.Bounds;
 	Ar << Volume.IndirectionTextureDimensions;
@@ -94,14 +210,22 @@ FArchive& operator<<(FArchive& Ar,FPrecomputedVolumetricLightmapData& Volume)
 		}
 	}
 
+	if (Ar.CustomVer(FRenderingObjectVersion::GUID) >= FRenderingObjectVersion::VolumetricLightmapStreaming)
+	{
+		Ar << Volume.SubLevelBrickPositions;
+		Ar << Volume.IndirectionTextureOriginalValues;
+	}
+
 	if (Ar.IsLoading())
 	{
-		if (GMaxRHIFeatureLevel >= ERHIFeatureLevel::SM4)
+		if (GMaxRHIFeatureLevel >= ERHIFeatureLevel::SM4 && !GIsEditor)
 		{
 			// drop LQ data for SM4+
 			Volume.BrickData.DiscardLowQualityLayers();
 		}
 
+		Volume.bTransient = false;
+
 		const SIZE_T VolumeBytes = Volume.GetAllocatedBytes();
 		INC_DWORD_STAT_BY(STAT_VolumetricLightmapBuildData, VolumeBytes);
 	}
@@ -146,12 +270,21 @@ int32 FVolumetricLightmapBrickData::GetMinimumVoxelSize() const
 }
 
 FPrecomputedVolumetricLightmapData::FPrecomputedVolumetricLightmapData()
+	: Bounds(EForceInit::ForceInitToZero)
+	, bTransient(true)
+	, IndirectionTextureDimensions(EForceInit::ForceInitToZero)
+	, BrickSize(0)
+	, BrickDataDimensions(EForceInit::ForceInitToZero)
+	, BrickDataBaseOffsetInAtlas(0)
 {}
 
 FPrecomputedVolumetricLightmapData::~FPrecomputedVolumetricLightmapData()
 {
-	const SIZE_T VolumeBytes = GetAllocatedBytes();
-	DEC_DWORD_STAT_BY(STAT_VolumetricLightmapBuildData, VolumeBytes);
+	if (!bTransient)
+	{
+		const SIZE_T VolumeBytes = GetAllocatedBytes();
+		DEC_DWORD_STAT_BY(STAT_VolumetricLightmapBuildData, VolumeBytes);
+	}
 }
 
 /** */
@@ -163,45 +296,386 @@ void FPrecomputedVolumetricLightmapData::InitializeOnImport(const FBox& NewBound
 
 void FPrecomputedVolumetricLightmapData::FinalizeImport()
 {
+	bTransient = false;
 	const SIZE_T VolumeBytes = GetAllocatedBytes();
 	INC_DWORD_STAT_BY(STAT_VolumetricLightmapBuildData, VolumeBytes);
 }
 
 ENGINE_API void FPrecomputedVolumetricLightmapData::InitRHI()
 {
-	IndirectionTexture.CreateTexture(IndirectionTextureDimensions);
-	BrickData.AmbientVector.CreateTexture(BrickDataDimensions);
+	if (FeatureLevel >= ERHIFeatureLevel::SM5)
+	{
+		if (IndirectionTextureDimensions.GetMax() > 0)
+		{
+			IndirectionTexture.CreateTexture(IndirectionTextureDimensions);
+		}
+
+		if (BrickDataDimensions.GetMax() > 0)
+		{
+			BrickData.AmbientVector.CreateTexture(BrickDataDimensions);
+
+			for (int32 i = 0; i < ARRAY_COUNT(BrickData.SHCoefficients); i++)
+			{
+				BrickData.SHCoefficients[i].CreateTexture(BrickDataDimensions);
+			}
+
+			if (BrickData.SkyBentNormal.Data.Num() > 0)
+			{
+				BrickData.SkyBentNormal.CreateTexture(BrickDataDimensions);
+			}
+
+			BrickData.DirectionalLightShadowing.CreateTexture(BrickDataDimensions);
+		}
+
+		GVolumetricLightmapBrickAtlas.Insert(INT_MAX, this);
+
+		// It is now safe to release the brick data used for upload. They will stay in GPU memory until UMapBuildDataRegistry::BeginDestroy().
+		BrickData.ReleaseRHI();
+	}
+}
+
+ENGINE_API void FPrecomputedVolumetricLightmapData::InitRHIForSubLevelResources()
+{
+	if (SubLevelBrickPositions.Num() > 0)
+	{
+		SubLevelBrickPositions.SetAllowCPUAccess(true);
+		IndirectionTextureOriginalValues.SetAllowCPUAccess(true);
+
+		{
+			FRHIResourceCreateInfo CreateInfo(&SubLevelBrickPositions);
+			SubLevelBrickPositionsBuffer = RHICreateVertexBuffer(SubLevelBrickPositions.Num() * SubLevelBrickPositions.GetTypeSize(), BUF_Static | BUF_ShaderResource, CreateInfo);
+			SubLevelBrickPositionsSRV = RHICreateShaderResourceView(SubLevelBrickPositionsBuffer, sizeof(uint32), PF_R32_UINT);
+		}
+
+		{
+			FRHIResourceCreateInfo CreateInfo(&IndirectionTextureOriginalValues);
+			IndirectionTextureOriginalValuesBuffer = RHICreateVertexBuffer(IndirectionTextureOriginalValues.Num() * IndirectionTextureOriginalValues.GetTypeSize(), BUF_Static | BUF_ShaderResource, CreateInfo);
+			IndirectionTextureOriginalValuesSRV = RHICreateShaderResourceView(IndirectionTextureOriginalValuesBuffer, sizeof(FColor), PF_R8G8B8A8_UINT);
+		}
+	}
+}
+
+ENGINE_API void FPrecomputedVolumetricLightmapData::ReleaseRHI()
+{
+	if (FeatureLevel >= ERHIFeatureLevel::SM5)
+	{
+		{
+			IndirectionTexture.Texture.SafeRelease();
+			IndirectionTexture.UAV.SafeRelease();
+			BrickData.ReleaseRHI();
+		}
+
+		GVolumetricLightmapBrickAtlas.Remove(this);
+	}
+}
 
-	for (int32 i = 0; i < ARRAY_COUNT(BrickData.SHCoefficients); i++)
+ENGINE_API void FPrecomputedVolumetricLightmapData::ReleaseRHIForSubLevelResources()
+{
 	{
-		BrickData.SHCoefficients[i].CreateTexture(BrickDataDimensions);
+		SubLevelBrickPositionsBuffer.SafeRelease();
+		SubLevelBrickPositionsSRV.SafeRelease();
+
+		IndirectionTextureOriginalValuesBuffer.SafeRelease();
+		IndirectionTextureOriginalValuesSRV.SafeRelease();
 	}
+}
+
+ENGINE_API void FPrecomputedVolumetricLightmapData::HandleDataMovementInAtlas(int32 OldOffset, int32 NewOffset)
+{
+	BrickDataBaseOffsetInAtlas = NewOffset;
+
+	FRHICommandListImmediate& RHICmdList = FRHICommandListExecutor::GetImmediateCommandList();
+
+	if (IndirectionTextureDimensions.GetMax() > 0)
+	{
+		const int32 PaddedBrickSize = BrickSize + 1;
+		int32 NumBricks = BrickDataDimensions.X * BrickDataDimensions.Y * BrickDataDimensions.Z / (PaddedBrickSize * PaddedBrickSize * PaddedBrickSize);
 
-	if (BrickData.SkyBentNormal.Data.Num() > 0)
+		for (FPrecomputedVolumetricLightmapData* SceneData : SceneDataAdded)
+		{
+			TShaderMap<FGlobalShaderType>* GlobalShaderMap = GetGlobalShaderMap(FeatureLevel);
+
+			TShaderMapRef<FMoveWholeIndirectionTextureCS> ComputeShader(GlobalShaderMap);
+
+			FMoveWholeIndirectionTextureCS::FParameters Parameters;
+			Parameters.NumBricks = NumBricks;
+			Parameters.StartPosInOldVolume = OldOffset;
+			Parameters.StartPosInNewVolume = BrickDataBaseOffsetInAtlas;
+			Parameters.IndirectionTexture = SceneData->IndirectionTexture.UAV;
+
+			FComputeShaderUtils::Dispatch(RHICmdList, *ComputeShader, Parameters,
+				FIntVector(
+					FMath::DivideAndRoundUp(IndirectionTextureDimensions.X, 4),
+					FMath::DivideAndRoundUp(IndirectionTextureDimensions.Y, 4),
+					FMath::DivideAndRoundUp(IndirectionTextureDimensions.Z, 4))
+			);
+		}
+	}
+	else
 	{
-		BrickData.SkyBentNormal.CreateTexture(BrickDataDimensions);
+		InitRHIForSubLevelResources();
+
+		for (FPrecomputedVolumetricLightmapData* SceneData : SceneDataAdded)
+		{
+			if (SceneData->IndirectionTexture.Texture.IsValid())
+			{
+				TShaderMap<FGlobalShaderType>* GlobalShaderMap = GetGlobalShaderMap(FeatureLevel);
+
+				TShaderMapRef<FPatchIndirectionTextureCS> ComputeShader(GlobalShaderMap);
+
+				int32 NumBricks = SubLevelBrickPositions.Num();
+
+				FPatchIndirectionTextureCS::FParameters Parameters;
+				Parameters.NumBricks = NumBricks;
+				Parameters.StartPosInNewVolume = BrickDataBaseOffsetInAtlas;
+				Parameters.IndirectionTexture = SceneData->IndirectionTexture.UAV;
+				Parameters.SubLevelBrickPositions = SubLevelBrickPositionsSRV;
+
+				FComputeShaderUtils::Dispatch(RHICmdList, *ComputeShader, Parameters, FIntVector(FMath::DivideAndRoundUp(NumBricks, 64), 1, 1));
+			}
+		}
+
+		ReleaseRHIForSubLevelResources();
 	}
+}
+
+FIntVector ComputeBrickLayoutPosition(int32 BrickLayoutAllocation, FIntVector BrickLayoutDimensions)
+{
+	const FIntVector BrickPosition(
+		BrickLayoutAllocation % BrickLayoutDimensions.X,
+		(BrickLayoutAllocation / BrickLayoutDimensions.X) % BrickLayoutDimensions.Y,
+		BrickLayoutAllocation / (BrickLayoutDimensions.X * BrickLayoutDimensions.Y));
 
-	BrickData.DirectionalLightShadowing.CreateTexture(BrickDataDimensions);
+	return BrickPosition;
 }
 
-ENGINE_API void FPrecomputedVolumetricLightmapData::ReleaseRHI()
+ENGINE_API void FPrecomputedVolumetricLightmapData::AddToSceneData(FPrecomputedVolumetricLightmapData* SceneData)
+{
+	check(SceneDataAdded.Find(SceneData) == INDEX_NONE);
+
+	SceneDataAdded.Add(SceneData);
+
+	FRHICommandListImmediate& RHICmdList = FRHICommandListExecutor::GetImmediateCommandList();
+
+	if (IndirectionTextureDimensions.GetMax() > 0)
+	{
+		// Copy parameters from the persistent level VLM
+		SceneData->Bounds = Bounds;
+		SceneData->BrickSize = BrickSize;
+		SceneData->BrickDataDimensions = BrickDataDimensions;
+
+		SceneData->IndirectionTexture.Format = IndirectionTexture.Format;
+		SceneData->IndirectionTextureDimensions = IndirectionTextureDimensions;
+
+		if (FeatureLevel >= ERHIFeatureLevel::SM5)
+		{
+			// GPU Path
+
+			if (GIsEditor)
+			{
+				// If we are in editor, or whenever a persistent level can be assigned to two Worlds (which is only possible when you modify the engine code),
+				// make a copy of the indirection texture
+				SceneData->IndirectionTexture.CreateTargetTexture(IndirectionTextureDimensions);
+				FRHICopyTextureInfo CopyInfo;
+				RHICmdList.CopyTexture(IndirectionTexture.Texture, SceneData->IndirectionTexture.Texture, CopyInfo);
+			}
+			else
+			{
+				// Steal the indirection texture. When the sublevels are unloaded the values will be restored.
+				SceneData->IndirectionTexture.Texture = IndirectionTexture.Texture;
+			}
+
+			SceneData->IndirectionTexture.CreateUAV();
+
+			const int32 PaddedBrickSize = BrickSize + 1;
+			int32 NumBricks = BrickDataDimensions.X * BrickDataDimensions.Y * BrickDataDimensions.Z / (PaddedBrickSize * PaddedBrickSize * PaddedBrickSize);
+
+			TShaderMap<FGlobalShaderType>* GlobalShaderMap = GetGlobalShaderMap(FeatureLevel);
+
+			TShaderMapRef<FMoveWholeIndirectionTextureCS> ComputeShader(GlobalShaderMap);
+
+			FMoveWholeIndirectionTextureCS::FParameters Parameters;
+			Parameters.NumBricks = NumBricks;
+			Parameters.StartPosInOldVolume = 0;
+			Parameters.StartPosInNewVolume = BrickDataBaseOffsetInAtlas;
+			Parameters.IndirectionTexture = SceneData->IndirectionTexture.UAV;
+
+			FComputeShaderUtils::Dispatch(RHICmdList, *ComputeShader, Parameters,
+				FIntVector(
+					FMath::DivideAndRoundUp(IndirectionTextureDimensions.X, 4),
+					FMath::DivideAndRoundUp(IndirectionTextureDimensions.Y, 4),
+					FMath::DivideAndRoundUp(IndirectionTextureDimensions.Z, 4))
+			);
+		}
+		else
+		{
+			// CPU Path
+			SceneData->IndirectionTexture.Data = IndirectionTexture.Data;
+			SceneData->CPUSubLevelIndirectionTable.Empty();
+			SceneData->CPUSubLevelIndirectionTable.AddZeroed(IndirectionTextureDimensions.X * IndirectionTextureDimensions.Y * IndirectionTextureDimensions.Z);
+			SceneData->CPUSubLevelBrickDataList.Empty();
+			SceneData->CPUSubLevelBrickDataList.Add(this);
+		}
+	}
+	else
+	{
+		if (FeatureLevel >= ERHIFeatureLevel::SM5)
+		{
+			// GPU Path
+			if (SceneData->IndirectionTexture.Texture.IsValid())
+			{
+				InitRHIForSubLevelResources();
+
+				TShaderMap<FGlobalShaderType>* GlobalShaderMap = GetGlobalShaderMap(FeatureLevel);
+
+				TShaderMapRef<FPatchIndirectionTextureCS> ComputeShader(GlobalShaderMap);
+
+				int32 NumBricks = SubLevelBrickPositions.Num();
+
+				FPatchIndirectionTextureCS::FParameters Parameters;
+				Parameters.NumBricks = NumBricks;
+				Parameters.StartPosInNewVolume = BrickDataBaseOffsetInAtlas;
+				Parameters.IndirectionTexture = SceneData->IndirectionTexture.UAV;
+				Parameters.SubLevelBrickPositions = SubLevelBrickPositionsSRV;
+
+				FComputeShaderUtils::Dispatch(RHICmdList, *ComputeShader, Parameters, FIntVector(FMath::DivideAndRoundUp(NumBricks, 64), 1, 1));
+
+				ReleaseRHIForSubLevelResources();
+			}
+		}
+		else
+		{
+			// CPU Path
+			if (SceneData->IndirectionTexture.Data.Num() > 0)
+			{
+				int32 Index = SceneData->CPUSubLevelBrickDataList.Add(this);
+				check(Index < UINT8_MAX);
+				uint8 Value = (uint8)Index;
+
+				for (int32 BrickIndex = 0; BrickIndex < SubLevelBrickPositions.Num(); BrickIndex++)
+				{
+					const FIntVector BrickLayoutPosition = ComputeBrickLayoutPosition(BrickIndex, BrickDataDimensions);
+
+					const FIntVector IndirectionDestDataCoordinate = SubLevelBrickPositions[BrickIndex];
+					const int32 IndirectionDestDataIndex =
+						((IndirectionDestDataCoordinate.Z * SceneData->IndirectionTextureDimensions.Y) + IndirectionDestDataCoordinate.Y) *
+						SceneData->IndirectionTextureDimensions.X + IndirectionDestDataCoordinate.X;
+
+					{
+						const int32 IndirectionTextureDataStride = GPixelFormats[SceneData->IndirectionTexture.Format].BlockBytes;
+						uint8* IndirectionVoxelPtr = (uint8*)&SceneData->IndirectionTexture.Data[IndirectionDestDataIndex * IndirectionTextureDataStride];
+						*(IndirectionVoxelPtr + 0) = BrickLayoutPosition.X;
+						*(IndirectionVoxelPtr + 1) = BrickLayoutPosition.Y;
+						*(IndirectionVoxelPtr + 2) = BrickLayoutPosition.Z;
+						*(IndirectionVoxelPtr + 3) = 1;
+					}
+
+					{
+						SceneData->CPUSubLevelIndirectionTable[IndirectionDestDataIndex] = Value;
+					}
+				}
+			}
+		}
+	}
+
+	SceneData->BrickDataDimensions = GVolumetricLightmapBrickAtlas.TextureSet.BrickDataDimensions;
+	SceneData->BrickData.AmbientVector = GVolumetricLightmapBrickAtlas.TextureSet.AmbientVector;
+	for (int32 i = 0; i < ARRAY_COUNT(SceneData->BrickData.SHCoefficients); i++)
+	{
+		SceneData->BrickData.SHCoefficients[i] = GVolumetricLightmapBrickAtlas.TextureSet.SHCoefficients[i];
+	}
+	SceneData->BrickData.SkyBentNormal = GVolumetricLightmapBrickAtlas.TextureSet.SkyBentNormal;
+	SceneData->BrickData.DirectionalLightShadowing = GVolumetricLightmapBrickAtlas.TextureSet.DirectionalLightShadowing;
+}
+
+ENGINE_API void FPrecomputedVolumetricLightmapData::RemoveFromSceneData(FPrecomputedVolumetricLightmapData* SceneData, int32 PersistentLevelBrickDataBaseOffset)
 {
-	IndirectionTexture.Texture.SafeRelease();
-	BrickData.AmbientVector.Texture.SafeRelease();
+	check(SceneDataAdded.Find(SceneData) != INDEX_NONE);
 
-	for (int32 i = 0; i < ARRAY_COUNT(BrickData.SHCoefficients); i++)
+	SceneDataAdded.Remove(SceneData);
+
+	FRHICommandListImmediate& RHICmdList = FRHICommandListExecutor::GetImmediateCommandList();
+
+	if (IndirectionTextureDimensions.GetMax() > 0)
 	{
-		BrickData.SHCoefficients[i].Texture.SafeRelease();
+		// Do nothing, as when a VLM data with indirection texture is being destroyed, the persistent level is going away
+	}
+	else
+	{
+		if (FeatureLevel >= ERHIFeatureLevel::SM5)
+		{
+			if (SceneData->IndirectionTexture.Texture.IsValid())
+			{
+				SCOPED_DRAW_EVENT(RHICmdList, RemoveSubLevelBricksCS);
+
+				InitRHIForSubLevelResources();
+
+				TShaderMap<FGlobalShaderType>* GlobalShaderMap = GetGlobalShaderMap(FeatureLevel);
+
+				TShaderMapRef<FRemoveSubLevelBricksCS> ComputeShader(GlobalShaderMap);
+
+				FRemoveSubLevelBricksCS::FParameters Parameters;
+				Parameters.NumBricks = SubLevelBrickPositions.Num();
+				Parameters.IndirectionTexture = SceneData->IndirectionTexture.UAV;
+				Parameters.SubLevelBrickPositions = SubLevelBrickPositionsSRV;
+				Parameters.IndirectionTextureOriginalValues = IndirectionTextureOriginalValuesSRV;
+				Parameters.PersistentLevelBrickDataBaseOffset = PersistentLevelBrickDataBaseOffset;
+
+				FComputeShaderUtils::Dispatch(RHICmdList, *ComputeShader, Parameters, FIntVector(FMath::DivideAndRoundUp(SubLevelBrickPositions.Num(), 64), 1, 1));
+
+				ReleaseRHIForSubLevelResources();
+			}
+		}
+		else
+		{
+			// CPU Path
+			if (SceneData->IndirectionTexture.Data.Num() > 0)
+			{
+				SceneData->CPUSubLevelBrickDataList.Remove(this);
+
+				for (int32 BrickIndex = 0; BrickIndex < SubLevelBrickPositions.Num(); BrickIndex++)
+				{
+					const FColor OriginalValue = IndirectionTextureOriginalValues[BrickIndex];
+
+					const FIntVector IndirectionDestDataCoordinate = SubLevelBrickPositions[BrickIndex];
+					const int32 IndirectionDestDataIndex =
+						((IndirectionDestDataCoordinate.Z * SceneData->IndirectionTextureDimensions.Y) + IndirectionDestDataCoordinate.Y) *
+						SceneData->IndirectionTextureDimensions.X + IndirectionDestDataCoordinate.X;
+
+					{
+						const int32 IndirectionTextureDataStride = GPixelFormats[SceneData->IndirectionTexture.Format].BlockBytes;
+						uint8* IndirectionVoxelPtr = (uint8*)&SceneData->IndirectionTexture.Data[IndirectionDestDataIndex * IndirectionTextureDataStride];
+						*(IndirectionVoxelPtr + 0) = OriginalValue.R;
+						*(IndirectionVoxelPtr + 1) = OriginalValue.G;
+						*(IndirectionVoxelPtr + 2) = OriginalValue.B;
+						*(IndirectionVoxelPtr + 3) = 1;
+					}
+
+					{
+						SceneData->CPUSubLevelIndirectionTable[IndirectionDestDataIndex] = 0;
+					}
+				}
+			}
+		}
 	}
 
-	BrickData.SkyBentNormal.Texture.SafeRelease();
-	BrickData.DirectionalLightShadowing.Texture.SafeRelease();
+	SceneData->BrickDataDimensions = GVolumetricLightmapBrickAtlas.TextureSet.BrickDataDimensions;
+	SceneData->BrickData.AmbientVector = GVolumetricLightmapBrickAtlas.TextureSet.AmbientVector;
+	for (int32 i = 0; i < ARRAY_COUNT(SceneData->BrickData.SHCoefficients); i++)
+	{
+		SceneData->BrickData.SHCoefficients[i] = GVolumetricLightmapBrickAtlas.TextureSet.SHCoefficients[i];
+	}
+	SceneData->BrickData.SkyBentNormal = GVolumetricLightmapBrickAtlas.TextureSet.SkyBentNormal;
+	SceneData->BrickData.DirectionalLightShadowing = GVolumetricLightmapBrickAtlas.TextureSet.DirectionalLightShadowing;
 }
 
 SIZE_T FPrecomputedVolumetricLightmapData::GetAllocatedBytes() const
 {
-	return IndirectionTexture.DataSize + BrickData.GetAllocatedBytes();
+	return
+		IndirectionTexture.DataSize + 
+		BrickData.GetAllocatedBytes() + 
+		SubLevelBrickPositions.Num() * SubLevelBrickPositions.GetTypeSize() +
+		IndirectionTextureOriginalValues.Num() * IndirectionTextureOriginalValues.GetTypeSize();
 }
 
 
@@ -215,7 +689,7 @@ FPrecomputedVolumetricLightmap::~FPrecomputedVolumetricLightmap()
 {
 }
 
-void FPrecomputedVolumetricLightmap::AddToScene(FSceneInterface* Scene, UMapBuildDataRegistry* Registry, FGuid LevelBuildDataId)
+void FPrecomputedVolumetricLightmap::AddToScene(FSceneInterface* Scene, UMapBuildDataRegistry* Registry, FGuid LevelBuildDataId, bool bIsPersistentLevel)
 {
 	check(!bAddedToScene);
 
@@ -237,7 +711,7 @@ void FPrecomputedVolumetricLightmap::AddToScene(FSceneInterface* Scene, UMapBuil
 			{
 				Volume->SetData(NewData, Scene);
 			});
-		Scene->AddPrecomputedVolumetricLightmap(this);
+		Scene->AddPrecomputedVolumetricLightmap(this, bIsPersistentLevel);
 	}
 }
 
@@ -260,9 +734,16 @@ void FPrecomputedVolumetricLightmap::SetData(FPrecomputedVolumetricLightmapData*
 {
 	Data = NewData;
 
-	if (Data && RHISupportsVolumeTextures(Scene->GetFeatureLevel()))
+	if (Data)
 	{
-		Data->InitResource();
+		Data->FeatureLevel = Scene->GetFeatureLevel();
+		Data->IndirectionTexture.bNeedsCPUAccess = GIsEditor;
+		Data->BrickData.SetNeedsCPUAccess(GIsEditor);
+
+		if (Scene->GetFeatureLevel() >= ERHIFeatureLevel::SM5)
+		{
+			Data->InitResource();
+		}
 	}
 }
 
@@ -304,6 +785,28 @@ void SampleIndirectionTexture(
 	OutIndirectionBrickSize = *(IndirectionVoxelPtr + 3);
 }
 
+void SampleIndirectionTextureWithSubLevel(
+	FVector IndirectionDataSourceCoordinate,
+	FIntVector IndirectionTextureDimensions,
+	const uint8* IndirectionTextureData,
+	const TArray<uint8>& CPUSubLevelIndirectionTable,
+	FIntVector& OutIndirectionBrickOffset,
+	int32& OutIndirectionBrickSize,
+	int32& OutSubLevelIndex)
+{
+	SampleIndirectionTexture(IndirectionDataSourceCoordinate, IndirectionTextureDimensions, IndirectionTextureData, OutIndirectionBrickOffset, OutIndirectionBrickSize);
+
+	FIntVector IndirectionDataCoordinateInt(IndirectionDataSourceCoordinate);
+
+	IndirectionDataCoordinateInt.X = FMath::Clamp<int32>(IndirectionDataCoordinateInt.X, 0, IndirectionTextureDimensions.X - 1);
+	IndirectionDataCoordinateInt.Y = FMath::Clamp<int32>(IndirectionDataCoordinateInt.Y, 0, IndirectionTextureDimensions.Y - 1);
+	IndirectionDataCoordinateInt.Z = FMath::Clamp<int32>(IndirectionDataCoordinateInt.Z, 0, IndirectionTextureDimensions.Z - 1);
+
+	const int32 IndirectionDataIndex = ((IndirectionDataCoordinateInt.Z * IndirectionTextureDimensions.Y) + IndirectionDataCoordinateInt.Y) * IndirectionTextureDimensions.X + IndirectionDataCoordinateInt.X;
+
+	OutSubLevelIndex = CPUSubLevelIndirectionTable[IndirectionDataIndex];
+}
+
 FVector ComputeBrickTextureCoordinate(
 	FVector IndirectionDataSourceCoordinate,
 	FIntVector IndirectionBrickOffset, 
@@ -316,3 +819,264 @@ FVector ComputeBrickTextureCoordinate(
 	FVector BrickTextureCoordinate = FVector(IndirectionBrickOffset * PaddedBrickSize) + FractionalIndirectionDataCoordinate * BrickSize;
 	return BrickTextureCoordinate;
 }
+
+IMPLEMENT_GLOBAL_SHADER(FRemoveSubLevelBricksCS, "/Engine/Private/VolumetricLightmapStreaming.usf", "RemoveSubLevelBricksCS", SF_Compute);
+IMPLEMENT_GLOBAL_SHADER(FCopyResidentBricksCS, "/Engine/Private/VolumetricLightmapStreaming.usf", "CopyResidentBricksCS", SF_Compute);
+IMPLEMENT_GLOBAL_SHADER(FCopyResidentBrickSHCoefficientsCS, "/Engine/Private/VolumetricLightmapStreaming.usf", "CopyResidentBrickSHCoefficientsCS", SF_Compute);
+IMPLEMENT_GLOBAL_SHADER(FPatchIndirectionTextureCS, "/Engine/Private/VolumetricLightmapStreaming.usf", "PatchIndirectionTextureCS", SF_Compute);
+IMPLEMENT_GLOBAL_SHADER(FMoveWholeIndirectionTextureCS, "/Engine/Private/VolumetricLightmapStreaming.usf", "MoveWholeIndirectionTextureCS", SF_Compute);
+
+FVolumetricLightmapBrickAtlas::FVolumetricLightmapBrickAtlas()
+	: bInitialized(false)
+	, PaddedBrickSize(5)
+{
+}
+
+template<class VolumetricLightmapBrickDataType>
+void CopyDataIntoAtlas(FRHICommandList& RHICmdList, ERHIFeatureLevel::Type FeatureLevel, int32 SrcOffset, int32 DestOffset, int32 NumBricks, const VolumetricLightmapBrickDataType& SrcData, FVolumetricLightmapBrickTextureSet DestTextureSet)
+{
+	{
+		TShaderMap<FGlobalShaderType>* GlobalShaderMap = GetGlobalShaderMap(FeatureLevel);
+
+		FCopyResidentBricksCS::FPermutationDomain PermutationVector;
+		PermutationVector.Set<FCopyResidentBricksCS::FHasSkyBentNormal>(SrcData.SkyBentNormal.Texture.IsValid());
+
+		TShaderMapRef<FCopyResidentBricksCS> ComputeShader(GlobalShaderMap, PermutationVector);
+
+		FCopyResidentBricksCS::FParameters Parameters;
+
+		Parameters.StartPosInOldVolume = SrcOffset;
+		Parameters.StartPosInNewVolume = DestOffset;
+
+		Parameters.AmbientVector = SrcData.AmbientVector.Texture;
+		Parameters.SkyBentNormal = SrcData.SkyBentNormal.Texture;
+		Parameters.DirectionalLightShadowing = SrcData.DirectionalLightShadowing.Texture;
+
+		Parameters.OutAmbientVector = DestTextureSet.AmbientVector.UAV;
+		Parameters.OutSkyBentNormal = DestTextureSet.SkyBentNormal.UAV;
+		Parameters.OutDirectionalLightShadowing = DestTextureSet.DirectionalLightShadowing.UAV;
+
+		FComputeShaderUtils::Dispatch(RHICmdList, *ComputeShader, Parameters, FIntVector(NumBricks, 1, 1));
+	}
+
+	for (int32 i = 0; i < ARRAY_COUNT(SrcData.SHCoefficients); i++)
+	{
+		TShaderMap<FGlobalShaderType>* GlobalShaderMap = GetGlobalShaderMap(FeatureLevel);
+
+		TShaderMapRef<FCopyResidentBrickSHCoefficientsCS> ComputeShader(GlobalShaderMap);
+
+		FCopyResidentBrickSHCoefficientsCS::FParameters Parameters;
+
+		Parameters.StartPosInOldVolume = SrcOffset;
+		Parameters.StartPosInNewVolume = DestOffset;
+
+		Parameters.SHCoefficients = SrcData.SHCoefficients[i].Texture;
+		Parameters.OutSHCoefficients = DestTextureSet.SHCoefficients[i].UAV;
+
+		FComputeShaderUtils::Dispatch(RHICmdList, *ComputeShader, Parameters, FIntVector(NumBricks, 1, 1));
+	}
+}
+
+void FVolumetricLightmapBrickAtlas::Insert(int32 Index, FPrecomputedVolumetricLightmapData* Data)
+{
+	check(!Allocations.FindByPredicate([Data](const Allocation& Other) { return Other.Data == Data; }));
+
+	if (!bInitialized)
+	{
+		FeatureLevel = ERHIFeatureLevel::SM5;
+		check(Data->BrickSize > 0);
+		PaddedBrickSize = Data->BrickSize + 1;
+		TextureSet.Initialize(Data->BrickDataDimensions, Data->BrickData);
+		bInitialized = true;
+	}
+	else
+	{
+		check(TextureSet.AmbientVector.Format == Data->BrickData.AmbientVector.Format);
+		for (int32 i = 0; i < ARRAY_COUNT(TextureSet.SHCoefficients); i++)
+		{
+			check(TextureSet.SHCoefficients[i].Format == Data->BrickData.SHCoefficients[i].Format);
+		}
+		check(TextureSet.SkyBentNormal.Format == Data->BrickData.SkyBentNormal.Format);
+		check(TextureSet.DirectionalLightShadowing.Format == Data->BrickData.DirectionalLightShadowing.Format);
+	}
+
+	int32 NumTotalBricks = 0;
+
+	for (const auto& Allocation : Allocations)
+	{
+		NumTotalBricks += Allocation.Size;
+	}
+
+	NumTotalBricks += Data->BrickDataDimensions.X * Data->BrickDataDimensions.Y * Data->BrickDataDimensions.Z / (PaddedBrickSize * PaddedBrickSize * PaddedBrickSize);
+
+	TArray<Allocation> NewAllocations;
+	FVolumetricLightmapBrickTextureSet NewTextureSet;
+
+	{
+		const int32 MaxBricksInLayoutOneDim = 1 << 8;
+		int32 BrickTextureLinearAllocator = NumTotalBricks;
+		FIntVector BrickLayoutDimensions;
+		BrickLayoutDimensions.X = FMath::Min(BrickTextureLinearAllocator, MaxBricksInLayoutOneDim);
+		BrickTextureLinearAllocator = FMath::DivideAndRoundUp(BrickTextureLinearAllocator, BrickLayoutDimensions.X);
+		BrickLayoutDimensions.Y = FMath::Min(BrickTextureLinearAllocator, MaxBricksInLayoutOneDim);
+		BrickTextureLinearAllocator = FMath::DivideAndRoundUp(BrickTextureLinearAllocator, BrickLayoutDimensions.Y);
+		BrickLayoutDimensions.Z = FMath::Min(BrickTextureLinearAllocator, MaxBricksInLayoutOneDim);
+		FIntVector BrickDataDimensions = BrickLayoutDimensions * PaddedBrickSize;
+
+		NewTextureSet.Initialize(BrickDataDimensions, TextureSet);
+	}
+
+	FRHICommandListImmediate& RHICmdList = FRHICommandListExecutor::GetImmediateCommandList();
+
+	// Dry run to handle persistent level data movement properly
+	{
+		int32 BrickStartAllocation = 0;
+
+		// Copy old allocations
+		for (int32 AllocationIndex = 0; AllocationIndex < Index && AllocationIndex < Allocations.Num(); AllocationIndex++)
+		{
+			BrickStartAllocation += Allocations[AllocationIndex].Size;
+		}
+
+		// Insert new allocation
+		{
+			int32 NumBricks = Data->BrickDataDimensions.X * Data->BrickDataDimensions.Y * Data->BrickDataDimensions.Z / (PaddedBrickSize * PaddedBrickSize * PaddedBrickSize);
+
+			BrickStartAllocation += NumBricks;
+		}
+
+		// Copy the rest of allocations
+		for (int32 AllocationIndex = Index; AllocationIndex < Allocations.Num(); AllocationIndex++)
+		{
+			if (Allocations[AllocationIndex].Data->IndirectionTextureDimensions.GetMax() > 0)
+			{
+				Allocations[AllocationIndex].Data->HandleDataMovementInAtlas(Allocations[AllocationIndex].StartOffset, BrickStartAllocation);
+			}
+			BrickStartAllocation += Allocations[AllocationIndex].Size;
+		}
+	}
+
+	{
+		int32 BrickStartAllocation = 0;
+
+		// Copy old allocations
+		for (int32 AllocationIndex = 0; AllocationIndex < Index && AllocationIndex < Allocations.Num(); AllocationIndex++)
+		{
+			CopyDataIntoAtlas(RHICmdList, FeatureLevel, Allocations[AllocationIndex].StartOffset, BrickStartAllocation, Allocations[AllocationIndex].Size, TextureSet, NewTextureSet);
+
+			NewAllocations.Add(Allocation{ Allocations[AllocationIndex].Data, Allocations[AllocationIndex].Size, BrickStartAllocation });
+			BrickStartAllocation += Allocations[AllocationIndex].Size;
+		}
+
+		// Insert new allocation
+		{
+			int32 NumBricks = Data->BrickDataDimensions.X * Data->BrickDataDimensions.Y * Data->BrickDataDimensions.Z / (PaddedBrickSize * PaddedBrickSize * PaddedBrickSize);
+
+			CopyDataIntoAtlas(RHICmdList, FeatureLevel, 0, BrickStartAllocation, NumBricks, Data->BrickData, NewTextureSet);
+
+			NewAllocations.Add(Allocation{ Data, NumBricks, BrickStartAllocation });
+			Data->BrickDataBaseOffsetInAtlas = BrickStartAllocation;
+			BrickStartAllocation += NumBricks;
+		}
+
+		// Copy the rest of allocations
+		for (int32 AllocationIndex = Index; AllocationIndex < Allocations.Num(); AllocationIndex++)
+		{
+			CopyDataIntoAtlas(RHICmdList, FeatureLevel, Allocations[AllocationIndex].StartOffset, BrickStartAllocation, Allocations[AllocationIndex].Size, TextureSet, NewTextureSet);
+
+			NewAllocations.Add(Allocation{ Allocations[AllocationIndex].Data, Allocations[AllocationIndex].Size, BrickStartAllocation });
+			// Handle the sub level data movements
+			if (Allocations[AllocationIndex].Data->IndirectionTextureDimensions.GetMax() == 0)
+			{
+				Allocations[AllocationIndex].Data->HandleDataMovementInAtlas(Allocations[AllocationIndex].StartOffset, BrickStartAllocation);
+			}
+			BrickStartAllocation += Allocations[AllocationIndex].Size;
+		}
+	}
+
+	// Replace with new allcations
+	Allocations = NewAllocations;
+	TextureSet = NewTextureSet; // <-- Old texture references are released here
+}
+
+void FVolumetricLightmapBrickAtlas::Remove(FPrecomputedVolumetricLightmapData* Data)
+{
+	UE_LOG(LogTemp, Warning, TEXT("VLM Atlas try to remove %llX"), Data);
+	Allocation* AllocationEntry = Allocations.FindByPredicate([Data](const Allocation& Other) { return Other.Data == Data; });
+	if (!AllocationEntry)
+	{
+		return;
+	}
+
+	int32 Index = (int32)(AllocationEntry - Allocations.GetData());
+
+	int32 NumTotalBricks = 0;
+
+	for (const auto& Allocation : Allocations)
+	{
+		if (Allocation.Data != Data)
+		{
+			NumTotalBricks += Allocation.Size;
+		}
+	}
+
+	TArray<Allocation> NewAllocations;
+	FVolumetricLightmapBrickTextureSet NewTextureSet;
+
+	if (NumTotalBricks > 0)
+	{
+		{
+			const int32 MaxBricksInLayoutOneDim = 1 << 8;
+			int32 BrickTextureLinearAllocator = NumTotalBricks;
+			FIntVector BrickLayoutDimensions;
+			BrickLayoutDimensions.X = FMath::Min(BrickTextureLinearAllocator, MaxBricksInLayoutOneDim);
+			BrickTextureLinearAllocator = FMath::DivideAndRoundUp(BrickTextureLinearAllocator, BrickLayoutDimensions.X);
+			BrickLayoutDimensions.Y = FMath::Min(BrickTextureLinearAllocator, MaxBricksInLayoutOneDim);
+			BrickTextureLinearAllocator = FMath::DivideAndRoundUp(BrickTextureLinearAllocator, BrickLayoutDimensions.Y);
+			BrickLayoutDimensions.Z = FMath::Min(BrickTextureLinearAllocator, MaxBricksInLayoutOneDim);
+			FIntVector BrickDataDimensions = BrickLayoutDimensions * PaddedBrickSize;
+
+			NewTextureSet.Initialize(BrickDataDimensions, TextureSet);
+		}
+
+		FRHICommandListImmediate& RHICmdList = FRHICommandListExecutor::GetImmediateCommandList();
+
+		{
+			int32 BrickStartAllocation = 0;
+
+			// Copy old allocations
+			for (int32 AllocationIndex = 0; AllocationIndex < Index && AllocationIndex < Allocations.Num(); AllocationIndex++)
+			{
+				CopyDataIntoAtlas(RHICmdList, FeatureLevel, Allocations[AllocationIndex].StartOffset, BrickStartAllocation, Allocations[AllocationIndex].Size, TextureSet, NewTextureSet);
+
+				NewAllocations.Add(Allocation{ Allocations[AllocationIndex].Data, Allocations[AllocationIndex].Size, BrickStartAllocation });
+				BrickStartAllocation += Allocations[AllocationIndex].Size;
+			}
+
+			// Skip the allocation being deleted
+
+			// Copy the rest of allocations
+			for (int32 AllocationIndex = Index + 1; AllocationIndex < Allocations.Num(); AllocationIndex++)
+			{
+				CopyDataIntoAtlas(RHICmdList, FeatureLevel, Allocations[AllocationIndex].StartOffset, BrickStartAllocation, Allocations[AllocationIndex].Size, TextureSet, NewTextureSet);
+
+				NewAllocations.Add(Allocation{ Allocations[AllocationIndex].Data, Allocations[AllocationIndex].Size, BrickStartAllocation });
+				Allocations[AllocationIndex].Data->HandleDataMovementInAtlas(Allocations[AllocationIndex].StartOffset, BrickStartAllocation);
+				BrickStartAllocation += Allocations[AllocationIndex].Size;
+			}
+		}
+	}
+	else
+	{
+		bInitialized = false;
+	}
+
+	// Replace with new allcations
+	Allocations = NewAllocations;
+	TextureSet = NewTextureSet; // <-- Old texture references are released here
+}
+
+void FVolumetricLightmapBrickAtlas::ReleaseRHI()
+{
+	TextureSet.Release();
+}
diff --git a/Engine/Source/Runtime/Engine/Public/PrecomputedVolumetricLightmap.h b/Engine/Source/Runtime/Engine/Public/PrecomputedVolumetricLightmap.h
--- a/Engine/Source/Runtime/Engine/Public/PrecomputedVolumetricLightmap.h
+++ b/Engine/Source/Runtime/Engine/Public/PrecomputedVolumetricLightmap.h
@@ -14,6 +14,8 @@
 #include "Math/PackedVector.h"
 #include "RHI.h"
 #include "RenderResource.h"
+#include "ShaderParameterStruct.h"
+#include "GlobalShader.h"
 
 class FSceneInterface;
 
@@ -23,7 +25,8 @@ public:
 
 	FVolumetricLightmapDataLayer() :
 		DataSize(0),
-		Format(PF_Unknown)
+		Format(PF_Unknown),
+		bNeedsCPUAccess(false)
 	{}
 
 	friend FArchive& operator<<(FArchive& Ar, FVolumetricLightmapDataLayer& Volume);
@@ -40,8 +43,7 @@ public:
 
 	virtual void Discard() override
 	{
-		//temporarily remove until the semantics around lighting scenarios are fixed.
-		if (!GIsEditor)
+		if (!bNeedsCPUAccess)
 		{
 			Data.Empty();
 		}
@@ -50,17 +52,22 @@ public:
 	void Resize(int32 NewSize)
 	{
 		Data.Empty(NewSize);
-		Data.AddUninitialized(NewSize);
+		Data.AddZeroed(NewSize);
 		DataSize = NewSize;
 	}
 
 	ENGINE_API void CreateTexture(FIntVector Dimensions);
+	ENGINE_API void CreateTargetTexture(FIntVector Dimensions);
+	ENGINE_API void CreateUAV();
 
 	TArray<uint8> Data;
 	// Stored redundantly for stats after Data has been discarded
 	int32 DataSize;
 	EPixelFormat Format;
 	FTexture3DRHIRef Texture;
+	FUnorderedAccessViewRHIRef UAV;
+
+	bool bNeedsCPUAccess;
 };
 
 class FVolumetricLightmapBrickData
@@ -76,6 +83,29 @@ public:
 
 	ENGINE_API int32 GetMinimumVoxelSize() const;
 
+	ENGINE_API void ReleaseRHI()
+	{
+		AmbientVector.Texture.SafeRelease();
+
+		for (int32 i = 0; i < ARRAY_COUNT(SHCoefficients); i++)
+		{
+			SHCoefficients[i].Texture.SafeRelease();
+		}
+
+		SkyBentNormal.Texture.SafeRelease();
+		DirectionalLightShadowing.Texture.SafeRelease();
+
+		AmbientVector.UAV.SafeRelease();
+
+		for (int32 i = 0; i < ARRAY_COUNT(SHCoefficients); i++)
+		{
+			SHCoefficients[i].UAV.SafeRelease();
+		}
+
+		SkyBentNormal.UAV.SafeRelease();
+		DirectionalLightShadowing.UAV.SafeRelease();
+	}
+
 	SIZE_T GetAllocatedBytes() const
 	{
 		SIZE_T NumBytes = AmbientVector.DataSize + SkyBentNormal.DataSize + DirectionalLightShadowing.DataSize;
@@ -89,6 +119,19 @@ public:
 		return NumBytes;
 	}
 
+	void SetNeedsCPUAccess(bool InAccess)
+	{
+		AmbientVector.bNeedsCPUAccess = InAccess;
+
+		for (int32 i = 0; i < ARRAY_COUNT(SHCoefficients); i++)
+		{
+			SHCoefficients[i].bNeedsCPUAccess = InAccess;
+		}
+
+		SkyBentNormal.bNeedsCPUAccess = InAccess;
+		DirectionalLightShadowing.bNeedsCPUAccess = InAccess;
+	}
+
 	// discard the layers used for low quality lightmap (LQ includes direct lighting from stationary lights).
 	void DiscardLowQualityLayers()
 	{
@@ -106,7 +149,7 @@ class FPrecomputedVolumetricLightmapData : public FRenderResource
 public:
 
 	ENGINE_API FPrecomputedVolumetricLightmapData();
-	virtual ~FPrecomputedVolumetricLightmapData();
+	ENGINE_API virtual ~FPrecomputedVolumetricLightmapData();
 
 	friend FArchive& operator<<(FArchive& Ar, FPrecomputedVolumetricLightmapData& Volume);
 	friend FArchive& operator<<(FArchive& Ar, FPrecomputedVolumetricLightmapData*& Volume);
@@ -117,6 +160,13 @@ public:
 	ENGINE_API virtual void InitRHI() override;
 	ENGINE_API virtual void ReleaseRHI() override;
 
+	ENGINE_API void InitRHIForSubLevelResources();
+	ENGINE_API void ReleaseRHIForSubLevelResources();
+
+	ENGINE_API void HandleDataMovementInAtlas(int32 OldOffset, int32 NewOffset);
+	ENGINE_API void AddToSceneData(FPrecomputedVolumetricLightmapData* SceneData);
+	ENGINE_API void RemoveFromSceneData(FPrecomputedVolumetricLightmapData* SceneData, int32 PersistentLevelBrickDataBaseOffset);
+
 	SIZE_T GetAllocatedBytes() const;
 
 	const FBox& GetBounds() const
@@ -126,6 +176,8 @@ public:
 
 	FBox Bounds;
 
+	bool bTransient;
+
 	FIntVector IndirectionTextureDimensions;
 	FVolumetricLightmapDataLayer IndirectionTexture;
 
@@ -133,6 +185,29 @@ public:
 	FIntVector BrickDataDimensions;
 	FVolumetricLightmapBrickData BrickData;
 
+	/**
+	 * Position data for sub level streaming
+	 */
+	// Brick positions in the persistent level's indirection texture
+	TResourceArray<FIntVector> SubLevelBrickPositions;
+	TResourceArray<FColor> IndirectionTextureOriginalValues;
+
+	/**
+	 * Runtime data for sub level streaming
+	 */
+	FVertexBufferRHIRef SubLevelBrickPositionsBuffer;
+	FShaderResourceViewRHIRef SubLevelBrickPositionsSRV;
+
+	FVertexBufferRHIRef IndirectionTextureOriginalValuesBuffer;
+	FShaderResourceViewRHIRef IndirectionTextureOriginalValuesSRV;
+
+	int32 BrickDataBaseOffsetInAtlas;
+	TArray<FPrecomputedVolumetricLightmapData*> SceneDataAdded;
+
+	// CPU indirection table for mobile path
+	TArray<uint8> CPUSubLevelIndirectionTable;
+	TArray<FPrecomputedVolumetricLightmapData*> CPUSubLevelBrickDataList;
+
 private:
 
 	friend class FPrecomputedVolumetricLightmap;
@@ -147,9 +222,9 @@ class FPrecomputedVolumetricLightmap
 public:
 
 	ENGINE_API FPrecomputedVolumetricLightmap();
-	~FPrecomputedVolumetricLightmap();
+	ENGINE_API ~FPrecomputedVolumetricLightmap();
 
-	ENGINE_API void AddToScene(class FSceneInterface* Scene, class UMapBuildDataRegistry* Registry, FGuid LevelBuildDataId);
+	ENGINE_API void AddToScene(class FSceneInterface* Scene, class UMapBuildDataRegistry* Registry, FGuid LevelBuildDataId, bool bIsPersistentLevel);
 
 	ENGINE_API void RemoveFromScene(FSceneInterface* Scene);
 	
@@ -318,8 +393,126 @@ extern ENGINE_API void SampleIndirectionTexture(
 	FIntVector& OutIndirectionBrickOffset,
 	int32& OutIndirectionBrickSize);
 
+extern ENGINE_API void SampleIndirectionTextureWithSubLevel(
+	FVector IndirectionDataSourceCoordinate,
+	FIntVector IndirectionTextureDimensions,
+	const uint8* IndirectionTextureData,
+	const TArray<uint8>& CPUSubLevelIndirectionTable,
+	FIntVector& OutIndirectionBrickOffset,
+	int32& OutIndirectionBrickSize,
+	int32& OutSubLevelIndex);
+
 extern ENGINE_API FVector ComputeBrickTextureCoordinate(
 	FVector IndirectionDataSourceCoordinate,
 	FIntVector IndirectionBrickOffset,
 	int32 IndirectionBrickSize,
 	int32 BrickSize);
+
+class FRemoveSubLevelBricksCS : public FGlobalShader
+{
+	DECLARE_GLOBAL_SHADER(FRemoveSubLevelBricksCS);
+	SHADER_USE_PARAMETER_STRUCT(FRemoveSubLevelBricksCS, FGlobalShader);
+
+	static bool ShouldCompilePermutation(const FGlobalShaderPermutationParameters& Parameters)
+	{
+		static const auto AllowStaticLightingVar = IConsoleManager::Get().FindTConsoleVariableDataInt(TEXT("r.AllowStaticLighting"));
+
+		return IsFeatureLevelSupported(Parameters.Platform, ERHIFeatureLevel::SM5) && AllowStaticLightingVar->GetValueOnAnyThread() != 0;
+	}
+
+	BEGIN_SHADER_PARAMETER_STRUCT(FParameters, )
+		SHADER_PARAMETER(uint32, NumBricks)
+		SHADER_PARAMETER(uint32, PersistentLevelBrickDataBaseOffset)
+		SHADER_PARAMETER_SRV(Buffer<uint>, SubLevelBrickPositions)
+		SHADER_PARAMETER_SRV(Buffer<uint4>, IndirectionTextureOriginalValues)
+		SHADER_PARAMETER_UAV(RWTexture3D<uint4>, IndirectionTexture)
+	END_SHADER_PARAMETER_STRUCT()
+};
+
+class FCopyResidentBricksCS : public FGlobalShader
+{
+	DECLARE_GLOBAL_SHADER(FCopyResidentBricksCS);
+	SHADER_USE_PARAMETER_STRUCT(FCopyResidentBricksCS, FGlobalShader);
+
+	class FHasSkyBentNormal : SHADER_PERMUTATION_BOOL("HAS_SKY_BENT_NORMAL");
+	using FPermutationDomain = TShaderPermutationDomain<FHasSkyBentNormal>;
+
+	static bool ShouldCompilePermutation(const FGlobalShaderPermutationParameters& Parameters)
+	{
+		static const auto AllowStaticLightingVar = IConsoleManager::Get().FindTConsoleVariableDataInt(TEXT("r.AllowStaticLighting"));
+
+		return IsFeatureLevelSupported(Parameters.Platform, ERHIFeatureLevel::SM5) && AllowStaticLightingVar->GetValueOnAnyThread() != 0;
+	}
+
+	BEGIN_SHADER_PARAMETER_STRUCT(FParameters, )
+		SHADER_PARAMETER(uint32, StartPosInOldVolume)
+		SHADER_PARAMETER(uint32, StartPosInNewVolume)
+		SHADER_PARAMETER_TEXTURE(Texture3D, AmbientVector)
+		SHADER_PARAMETER_TEXTURE(Texture3D, SkyBentNormal)
+		SHADER_PARAMETER_TEXTURE(Texture3D, DirectionalLightShadowing)
+
+		SHADER_PARAMETER_UAV(RWTexture3D<float3>, OutAmbientVector)
+		SHADER_PARAMETER_UAV(RWTexture3D<float4>, OutSkyBentNormal)
+		SHADER_PARAMETER_UAV(RWTexture3D<float>, OutDirectionalLightShadowing)
+	END_SHADER_PARAMETER_STRUCT()
+};
+
+class FCopyResidentBrickSHCoefficientsCS : public FGlobalShader
+{
+	DECLARE_GLOBAL_SHADER(FCopyResidentBrickSHCoefficientsCS);
+	SHADER_USE_PARAMETER_STRUCT(FCopyResidentBrickSHCoefficientsCS, FGlobalShader);
+
+	static bool ShouldCompilePermutation(const FGlobalShaderPermutationParameters& Parameters)
+	{
+		static const auto AllowStaticLightingVar = IConsoleManager::Get().FindTConsoleVariableDataInt(TEXT("r.AllowStaticLighting"));
+
+		return IsFeatureLevelSupported(Parameters.Platform, ERHIFeatureLevel::SM5) && AllowStaticLightingVar->GetValueOnAnyThread() != 0;
+	}
+
+	BEGIN_SHADER_PARAMETER_STRUCT(FParameters, )
+		SHADER_PARAMETER(uint32, StartPosInOldVolume)
+		SHADER_PARAMETER(uint32, StartPosInNewVolume)
+		SHADER_PARAMETER_TEXTURE(Texture3D, SHCoefficients)
+		SHADER_PARAMETER_UAV(RWTexture3D<float4>, OutSHCoefficients)
+	END_SHADER_PARAMETER_STRUCT()
+};
+
+class FPatchIndirectionTextureCS : public FGlobalShader
+{
+	DECLARE_GLOBAL_SHADER(FPatchIndirectionTextureCS);
+	SHADER_USE_PARAMETER_STRUCT(FPatchIndirectionTextureCS, FGlobalShader);
+
+	static bool ShouldCompilePermutation(const FGlobalShaderPermutationParameters& Parameters)
+	{
+		static const auto AllowStaticLightingVar = IConsoleManager::Get().FindTConsoleVariableDataInt(TEXT("r.AllowStaticLighting"));
+
+		return IsFeatureLevelSupported(Parameters.Platform, ERHIFeatureLevel::SM5) && AllowStaticLightingVar->GetValueOnAnyThread() != 0;
+	}
+
+	BEGIN_SHADER_PARAMETER_STRUCT(FParameters, )
+		SHADER_PARAMETER(uint32, NumBricks)
+		SHADER_PARAMETER(uint32, StartPosInNewVolume)
+		SHADER_PARAMETER_SRV(Buffer<uint>, SubLevelBrickPositions)
+		SHADER_PARAMETER_UAV(RWTexture3D<uint4>, IndirectionTexture)
+	END_SHADER_PARAMETER_STRUCT()
+};
+
+class FMoveWholeIndirectionTextureCS : public FGlobalShader
+{
+	DECLARE_GLOBAL_SHADER(FMoveWholeIndirectionTextureCS);
+	SHADER_USE_PARAMETER_STRUCT(FMoveWholeIndirectionTextureCS, FGlobalShader);
+
+	static bool ShouldCompilePermutation(const FGlobalShaderPermutationParameters& Parameters)
+	{
+		static const auto AllowStaticLightingVar = IConsoleManager::Get().FindTConsoleVariableDataInt(TEXT("r.AllowStaticLighting"));
+
+		return IsFeatureLevelSupported(Parameters.Platform, ERHIFeatureLevel::SM5) && AllowStaticLightingVar->GetValueOnAnyThread() != 0;
+	}
+
+	BEGIN_SHADER_PARAMETER_STRUCT(FParameters, )
+		SHADER_PARAMETER(uint32, NumBricks)
+		SHADER_PARAMETER(uint32, StartPosInOldVolume)
+		SHADER_PARAMETER(uint32, StartPosInNewVolume)
+		SHADER_PARAMETER_UAV(RWTexture3D<uint4>, IndirectionTexture)
+	END_SHADER_PARAMETER_STRUCT()
+};
diff --git a/Engine/Source/Runtime/Engine/Public/SceneInterface.h b/Engine/Source/Runtime/Engine/Public/SceneInterface.h
--- a/Engine/Source/Runtime/Engine/Public/SceneInterface.h
+++ b/Engine/Source/Runtime/Engine/Public/SceneInterface.h
@@ -170,7 +170,7 @@ public:
 	virtual void RemovePrecomputedLightVolume(const class FPrecomputedLightVolume* Volume) {}
 
 	virtual bool HasPrecomputedVolumetricLightmap_RenderThread() const { return false; }
-	virtual void AddPrecomputedVolumetricLightmap(const class FPrecomputedVolumetricLightmap* Volume) {}
+	virtual void AddPrecomputedVolumetricLightmap(const class FPrecomputedVolumetricLightmap* Volume, bool bIsPersistentLevel) {}
 	virtual void RemovePrecomputedVolumetricLightmap(const class FPrecomputedVolumetricLightmap* Volume) {}
 
 	/** Add a runtime virtual texture object to the scene. */
diff --git a/Engine/Source/Runtime/Renderer/Private/LightMapRendering.cpp b/Engine/Source/Runtime/Renderer/Private/LightMapRendering.cpp
--- a/Engine/Source/Runtime/Renderer/Private/LightMapRendering.cpp
+++ b/Engine/Source/Runtime/Renderer/Private/LightMapRendering.cpp
@@ -171,18 +171,21 @@ void InterpolateVolumetricLightmap(
 	SCOPE_CYCLE_COUNTER(STAT_InterpolateVolumetricLightmapOnCPU);
 
 	checkSlow(VolumetricLightmapSceneData.HasData());
-	const FPrecomputedVolumetricLightmapData& VolumetricLightmapData = *VolumetricLightmapSceneData.GetLevelVolumetricLightmap()->Data;
+	const FPrecomputedVolumetricLightmapData& GlobalVolumetricLightmapData = *VolumetricLightmapSceneData.GetLevelVolumetricLightmap()->Data;
 	
-	const FVector IndirectionDataSourceCoordinate = ComputeIndirectionCoordinate(LookupPosition, VolumetricLightmapData.GetBounds(), VolumetricLightmapData.IndirectionTextureDimensions);
+	const FVector IndirectionDataSourceCoordinate = ComputeIndirectionCoordinate(LookupPosition, GlobalVolumetricLightmapData.GetBounds(), GlobalVolumetricLightmapData.IndirectionTextureDimensions);
 
-	check(VolumetricLightmapData.IndirectionTexture.Data.Num() > 0);
-	checkSlow(GPixelFormats[VolumetricLightmapData.IndirectionTexture.Format].BlockBytes == sizeof(uint8) * 4);
-	const int32 NumIndirectionTexels = VolumetricLightmapData.IndirectionTextureDimensions.X * VolumetricLightmapData.IndirectionTextureDimensions.Y * VolumetricLightmapData.IndirectionTextureDimensions.Z;
-	check(VolumetricLightmapData.IndirectionTexture.Data.Num() * VolumetricLightmapData.IndirectionTexture.Data.GetTypeSize() == NumIndirectionTexels * sizeof(uint8) * 4);
+	check(GlobalVolumetricLightmapData.IndirectionTexture.Data.Num() > 0);
+	checkSlow(GPixelFormats[GlobalVolumetricLightmapData.IndirectionTexture.Format].BlockBytes == sizeof(uint8) * 4);
+	const int32 NumIndirectionTexels = GlobalVolumetricLightmapData.IndirectionTextureDimensions.X * GlobalVolumetricLightmapData.IndirectionTextureDimensions.Y * GlobalVolumetricLightmapData.IndirectionTextureDimensions.Z;
+	check(GlobalVolumetricLightmapData.IndirectionTexture.Data.Num() * GlobalVolumetricLightmapData.IndirectionTexture.Data.GetTypeSize() == NumIndirectionTexels * sizeof(uint8) * 4);
 	
 	FIntVector IndirectionBrickOffset;
 	int32 IndirectionBrickSize;
-	SampleIndirectionTexture(IndirectionDataSourceCoordinate, VolumetricLightmapData.IndirectionTextureDimensions, VolumetricLightmapData.IndirectionTexture.Data.GetData(), IndirectionBrickOffset, IndirectionBrickSize);
+	int32 SubLevelIndex;
+	SampleIndirectionTextureWithSubLevel(IndirectionDataSourceCoordinate, GlobalVolumetricLightmapData.IndirectionTextureDimensions, GlobalVolumetricLightmapData.IndirectionTexture.Data.GetData(), GlobalVolumetricLightmapData.CPUSubLevelIndirectionTable, IndirectionBrickOffset, IndirectionBrickSize, SubLevelIndex);
+
+	const FPrecomputedVolumetricLightmapData& VolumetricLightmapData = *GlobalVolumetricLightmapData.CPUSubLevelBrickDataList[SubLevelIndex];
 
 	const FVector BrickTextureCoordinate = ComputeBrickTextureCoordinate(IndirectionDataSourceCoordinate, IndirectionBrickOffset, IndirectionBrickSize, VolumetricLightmapData.BrickSize);
 
diff --git a/Engine/Source/Runtime/Renderer/Private/PrimitiveSceneInfo.cpp b/Engine/Source/Runtime/Renderer/Private/PrimitiveSceneInfo.cpp
--- a/Engine/Source/Runtime/Renderer/Private/PrimitiveSceneInfo.cpp
+++ b/Engine/Source/Runtime/Renderer/Private/PrimitiveSceneInfo.cpp
@@ -507,7 +507,7 @@ void FPrimitiveSceneInfo::AddToScene(FRHICommandListImmediate& RHICmdList, bool
 	// Create an indirect lighting cache uniform buffer if we attaching a primitive that may require it, as it may be stored inside a cached mesh command.
 	if (IsIndirectLightingCacheAllowed(Scene->GetFeatureLevel())
 		&& Proxy->WillEverBeLit()
-		&& ((Proxy->HasStaticLighting() && Proxy->NeedsUnbuiltPreviewLighting()) || (Proxy->IsMovable() && Proxy->GetIndirectLightingCacheQuality() != ILCQ_Off)))
+		&& ((Proxy->HasStaticLighting() && Proxy->NeedsUnbuiltPreviewLighting()) || (Proxy->IsMovable() && Proxy->GetIndirectLightingCacheQuality() != ILCQ_Off) || Proxy->GetLightmapType() == ELightmapType::ForceVolumetric))
 	{
 		if (!IndirectLightingCacheUniformBuffer)
 		{
@@ -1006,7 +1006,7 @@ void FPrimitiveSceneInfo::UpdateIndirectLightingCacheBuffer()
 	{
 		QUICK_SCOPE_CYCLE_COUNTER(STAT_UpdateIndirectLightingCacheBuffer);
 
-		if (!RHISupportsVolumeTextures(Scene->GetFeatureLevel())
+		if (Scene->GetFeatureLevel() < ERHIFeatureLevel::SM5
 			&& Scene->VolumetricLightmapSceneData.HasData()
 			&& (Proxy->IsMovable() || Proxy->NeedsUnbuiltPreviewLighting() || Proxy->GetLightmapType() == ELightmapType::ForceVolumetric)
 			&& Proxy->WillEverBeLit())
diff --git a/Engine/Source/Runtime/Renderer/Private/RendererScene.cpp b/Engine/Source/Runtime/Renderer/Private/RendererScene.cpp
--- a/Engine/Source/Runtime/Renderer/Private/RendererScene.cpp
+++ b/Engine/Source/Runtime/Renderer/Private/RendererScene.cpp
@@ -1312,7 +1312,7 @@ void FScene::UpdatePrimitiveTransform_RenderThread(FRHICommandListImmediate& RHI
 	PrimitiveSceneProxy->SetTransform(LocalToWorld, WorldBounds, LocalBounds, AttachmentRootPosition);
 	PrimitiveTransforms[PrimitiveSceneInfo->PackedIndex] = LocalToWorld;
 
-	if (!RHISupportsVolumeTextures(GetFeatureLevel())
+	if (GetFeatureLevel() < ERHIFeatureLevel::SM5
 		&& (PrimitiveSceneProxy->IsMovable() || PrimitiveSceneProxy->NeedsUnbuiltPreviewLighting() || PrimitiveSceneProxy->GetLightmapType() == ELightmapType::ForceVolumetric))
 	{
 		PrimitiveSceneInfo->MarkIndirectLightingCacheBufferDirty();
@@ -2398,14 +2398,28 @@ void FScene::RemovePrecomputedLightVolume(const FPrecomputedLightVolume* Volume)
 		});
 }
 
-void FVolumetricLightmapSceneData::AddLevelVolume(const FPrecomputedVolumetricLightmap* InVolume, EShadingPath ShadingPath)
+void FVolumetricLightmapSceneData::AddLevelVolume(const FPrecomputedVolumetricLightmap* InVolume, EShadingPath ShadingPath, bool bIsPersistentLevel)
 {
 	LevelVolumetricLightmaps.Add(InVolume);
+
+	if (bIsPersistentLevel)
+	{
+		PersistentLevelVolumetricLightmap = InVolume;
+	}
+
+	InVolume->Data->AddToSceneData(&GlobalVolumetricLightmapData);
 }
 
 void FVolumetricLightmapSceneData::RemoveLevelVolume(const FPrecomputedVolumetricLightmap* InVolume)
 {
 	LevelVolumetricLightmaps.Remove(InVolume);
+
+	InVolume->Data->RemoveFromSceneData(&GlobalVolumetricLightmapData, PersistentLevelVolumetricLightmap ? PersistentLevelVolumetricLightmap->Data->BrickDataBaseOffsetInAtlas : 0);
+
+	if (PersistentLevelVolumetricLightmap == InVolume)
+	{
+		PersistentLevelVolumetricLightmap = nullptr;
+	}
 }
 
 bool FScene::HasPrecomputedVolumetricLightmap_RenderThread() const
@@ -2413,12 +2427,12 @@ bool FScene::HasPrecomputedVolumetricLightmap_RenderThread() const
 	return VolumetricLightmapSceneData.HasData();
 }
 
-void FScene::AddPrecomputedVolumetricLightmap(const FPrecomputedVolumetricLightmap* Volume)
+void FScene::AddPrecomputedVolumetricLightmap(const FPrecomputedVolumetricLightmap* Volume, bool bIsPersistentLevel)
 {
 	FScene* Scene = this;
 
 	ENQUEUE_RENDER_COMMAND(AddVolumeCommand)
-		([Scene, Volume](FRHICommandListImmediate& RHICmdList) 
+		([Scene, Volume, bIsPersistentLevel](FRHICommandListImmediate& RHICmdList)
 		{
 #if !(UE_BUILD_SHIPPING || UE_BUILD_TEST)
 			if (Volume && Scene->GetShadingPath() == EShadingPath::Mobile)
@@ -2430,7 +2444,7 @@ void FScene::AddPrecomputedVolumetricLightmap(const FPrecomputedVolumetricLightm
 				}
 			}
 #endif
-		Scene->VolumetricLightmapSceneData.AddLevelVolume(Volume, Scene->GetShadingPath());
+		Scene->VolumetricLightmapSceneData.AddLevelVolume(Volume, Scene->GetShadingPath(), bIsPersistentLevel);
 		});
 }
 
diff --git a/Engine/Source/Runtime/Renderer/Private/ScenePrivate.h b/Engine/Source/Runtime/Renderer/Private/ScenePrivate.h
--- a/Engine/Source/Runtime/Renderer/Private/ScenePrivate.h
+++ b/Engine/Source/Runtime/Renderer/Private/ScenePrivate.h
@@ -46,6 +46,7 @@
 #include "LightMapDensityRendering.h"
 #include "VolumetricFogShared.h"
 #include "DebugViewModeRendering.h"
+#include "PrecomputedVolumetricLightmap.h"
 #if RHI_RAYTRACING
 #include "RayTracing/RayTracingIESLightProfiles.h"
 #include "Halton.h"
@@ -1554,17 +1555,21 @@ class FVolumetricLightmapSceneData
 {
 public:
 
+	FVolumetricLightmapSceneData() { GlobalVolumetricLightmap.Data = &GlobalVolumetricLightmapData; }
 	bool HasData() const { return LevelVolumetricLightmaps.Num() > 0; }
-	void AddLevelVolume(const class FPrecomputedVolumetricLightmap* InVolume, EShadingPath ShadingPath);
+	void AddLevelVolume(const class FPrecomputedVolumetricLightmap* InVolume, EShadingPath ShadingPath, bool bIsPersistentLevel);
 	void RemoveLevelVolume(const class FPrecomputedVolumetricLightmap* InVolume);
 	const FPrecomputedVolumetricLightmap* GetLevelVolumetricLightmap() const 
 	{ 
-		return LevelVolumetricLightmaps.Num() > 0 ? LevelVolumetricLightmaps.Last() : NULL; 
+		return &GlobalVolumetricLightmap;
 	}
 
 	TMap<FVector, FVolumetricLightmapInterpolation> CPUInterpolationCache;
 
+	FPrecomputedVolumetricLightmapData GlobalVolumetricLightmapData;
 private:
+	FPrecomputedVolumetricLightmap GlobalVolumetricLightmap;
+	const FPrecomputedVolumetricLightmap* PersistentLevelVolumetricLightmap = nullptr;
 	TArray<const FPrecomputedVolumetricLightmap*> LevelVolumetricLightmaps;
 };
 
@@ -2679,7 +2684,7 @@ public:
 	virtual void AddPrecomputedLightVolume(const class FPrecomputedLightVolume* Volume) override;
 	virtual void RemovePrecomputedLightVolume(const class FPrecomputedLightVolume* Volume) override;
 	virtual bool HasPrecomputedVolumetricLightmap_RenderThread() const override;
-	virtual void AddPrecomputedVolumetricLightmap(const class FPrecomputedVolumetricLightmap* Volume) override;
+	virtual void AddPrecomputedVolumetricLightmap(const class FPrecomputedVolumetricLightmap* Volume, bool bIsPersistentLevel) override;
 	virtual void RemovePrecomputedVolumetricLightmap(const class FPrecomputedVolumetricLightmap* Volume) override;
 	virtual void AddRuntimeVirtualTexture(class URuntimeVirtualTextureComponent* Component) override;
 	virtual void RemoveRuntimeVirtualTexture(class URuntimeVirtualTextureComponent* Component) override;
diff --git a/Engine/Source/Runtime/Renderer/Private/VisualizeVolumetricLightmap.cpp b/Engine/Source/Runtime/Renderer/Private/VisualizeVolumetricLightmap.cpp
--- a/Engine/Source/Runtime/Renderer/Private/VisualizeVolumetricLightmap.cpp
+++ b/Engine/Source/Runtime/Renderer/Private/VisualizeVolumetricLightmap.cpp
@@ -206,10 +206,10 @@ void FDeferredShadingSceneRenderer::VisualizeVolumetricLightmap(FRHICommandListI
 				PixelShader->SetParameters(RHICmdList, View);
 
 				int32 BrickSize = VolumetricLightmapData->BrickSize;
-				int32 NumQuads = VolumetricLightmapData->IndirectionTextureDimensions.X * VolumetricLightmapData->IndirectionTextureDimensions.Y * VolumetricLightmapData->IndirectionTextureDimensions.Z * BrickSize * BrickSize * BrickSize;
+				uint32 NumQuads = VolumetricLightmapData->IndirectionTextureDimensions.X * VolumetricLightmapData->IndirectionTextureDimensions.Y * VolumetricLightmapData->IndirectionTextureDimensions.Z * BrickSize * BrickSize * BrickSize;
 
 				RHICmdList.SetStreamSource(0, NULL, 0);
-				RHICmdList.DrawIndexedPrimitive(GVisualizeQuadIndexBuffer.IndexBufferRHI, 0, 0, 4 * GQuadsPerVisualizeInstance, 0, 2 * GQuadsPerVisualizeInstance, FMath::DivideAndRoundUp(NumQuads, GQuadsPerVisualizeInstance));
+				RHICmdList.DrawIndexedPrimitive(GVisualizeQuadIndexBuffer.IndexBufferRHI, 0, 0, 4 * GQuadsPerVisualizeInstance, 0, 2 * GQuadsPerVisualizeInstance, FMath::DivideAndRoundUp(FMath::Min(NumQuads, 0x7FFFFFFFu / 4), (uint32)GQuadsPerVisualizeInstance));
 			}
 		}
 		RHICmdList.EndRenderPass();
