Fix missing shader when r.AllowStaticLighting=0

diff --git a/Engine/Source/Runtime/Engine/Private/PrecomputedVolumetricLightmap.cpp b/Engine/Source/Runtime/Engine/Private/PrecomputedVolumetricLightmap.cpp
--- a/Engine/Source/Runtime/Engine/Private/PrecomputedVolumetricLightmap.cpp
+++ b/Engine/Source/Runtime/Engine/Private/PrecomputedVolumetricLightmap.cpp
@@ -691,6 +691,13 @@ FPrecomputedVolumetricLightmap::~FPrecomputedVolumetricLightmap()
 
 void FPrecomputedVolumetricLightmap::AddToScene(FSceneInterface* Scene, UMapBuildDataRegistry* Registry, FGuid LevelBuildDataId, bool bIsPersistentLevel)
 {
+	static const auto AllowStaticLightingVar = IConsoleManager::Get().FindTConsoleVariableDataInt(TEXT("r.AllowStaticLighting"));
+
+	if (AllowStaticLightingVar->GetValueOnAnyThread() == 0)
+	{
+		return;
+	}
+
 	check(!bAddedToScene);
 
 	FPrecomputedVolumetricLightmapData* NewData = NULL;
@@ -717,6 +724,13 @@ void FPrecomputedVolumetricLightmap::AddToScene(FSceneInterface* Scene, UMapBuil
 
 void FPrecomputedVolumetricLightmap::RemoveFromScene(FSceneInterface* Scene)
 {
+	static const auto AllowStaticLightingVar = IConsoleManager::Get().FindTConsoleVariableDataInt(TEXT("r.AllowStaticLighting"));
+
+	if (AllowStaticLightingVar->GetValueOnAnyThread() == 0)
+	{
+		return;
+	}
+
 	if (bAddedToScene)
 	{
 		bAddedToScene = false;
