Removed hack to override eye adaptation metric.

diff --git a/Engine/Source/Runtime/Engine/Private/SceneView.cpp b/Engine/Source/Runtime/Engine/Private/SceneView.cpp
--- a/Engine/Source/Runtime/Engine/Private/SceneView.cpp
+++ b/Engine/Source/Runtime/Engine/Private/SceneView.cpp
@@ -1755,25 +1755,6 @@ void FSceneView::EndFinalPostprocessSettings(const FSceneViewInitOptions& ViewIn
 {
 	const auto SceneViewFeatureLevel = GetFeatureLevel();
 
-#if !(UE_BUILD_SHIPPING || UE_BUILD_TEST)
-	static const auto CVar = IConsoleManager::Get().FindTConsoleVariableDataInt(TEXT("r.EyeAdaptation.MethodOveride"));
-	if (CVar->GetValueOnGameThread() == -2)
-	{
-		// seemed to be good setting for Paragon, we might want to remove or adjust this later on
-		FinalPostProcessSettings.AutoExposureMethod = AEM_Basic;
-		FinalPostProcessSettings.AutoExposureBias = -0.6f;
-		FinalPostProcessSettings.AutoExposureMaxBrightness = 2.f;
-		FinalPostProcessSettings.AutoExposureMinBrightness = 0.05f;
-		FinalPostProcessSettings.AutoExposureSpeedDown = 1.f;
-		FinalPostProcessSettings.AutoExposureSpeedUp = 3.f;
-		if (CVarDefaultAutoExposureExtendDefaultLuminanceRange.GetValueOnGameThread())
-		{
-			FinalPostProcessSettings.AutoExposureMinBrightness = LuminanceToEV100(FinalPostProcessSettings.AutoExposureMinBrightness);
-			FinalPostProcessSettings.AutoExposureMaxBrightness = LuminanceToEV100(FinalPostProcessSettings.AutoExposureMaxBrightness);
-		}
-	}
-#endif
-
 	// will be deprecated soon, use the new asset LightPropagationVolumeBlendable instead
 	{
 		FLightPropagationVolumeSettings& Dest = FinalPostProcessSettings.BlendableManager.GetSingleFinalData<FLightPropagationVolumeSettings>();
