attempt to fix cis

diff --git a/Engine/Source/Runtime/Engine/Public/PrecomputedVolumetricLightmap.h b/Engine/Source/Runtime/Engine/Public/PrecomputedVolumetricLightmap.h
--- a/Engine/Source/Runtime/Engine/Public/PrecomputedVolumetricLightmap.h
+++ b/Engine/Source/Runtime/Engine/Public/PrecomputedVolumetricLightmap.h
@@ -16,6 +16,7 @@
 #include "RenderResource.h"
 #include "ShaderParameterStruct.h"
 #include "GlobalShader.h"
+#include "Containers/DynamicRHIResourceArray.h"
 
 class FSceneInterface;
 
