Fix incremental non-unity build

diff --git a/Engine/Source/Runtime/Engine/Private/PrecomputedVolumetricLightmap.cpp b/Engine/Source/Runtime/Engine/Private/PrecomputedVolumetricLightmap.cpp
--- a/Engine/Source/Runtime/Engine/Private/PrecomputedVolumetricLightmap.cpp
+++ b/Engine/Source/Runtime/Engine/Private/PrecomputedVolumetricLightmap.cpp
@@ -13,7 +13,6 @@
 #include "Engine/MapBuildDataRegistry.h"
 #include "Interfaces/ITargetPlatform.h"
 #include "UObject/MobileObjectVersion.h"
-#include "Containers/DynamicRHIResourceArray.h"
 #include "RenderGraphUtils.h"
 
 DECLARE_MEMORY_STAT(TEXT("Volumetric Lightmap"),STAT_VolumetricLightmapBuildData,STATGROUP_MapBuildData);
diff --git a/Engine/Source/Runtime/Engine/Public/PrecomputedVolumetricLightmap.h b/Engine/Source/Runtime/Engine/Public/PrecomputedVolumetricLightmap.h
--- a/Engine/Source/Runtime/Engine/Public/PrecomputedVolumetricLightmap.h
+++ b/Engine/Source/Runtime/Engine/Public/PrecomputedVolumetricLightmap.h
@@ -9,14 +9,13 @@
 #include "CoreMinimal.h"
 #include "Misc/Guid.h"
 #include "Math/SHMath.h"
-#include "Containers/ResourceArray.h"
+#include "Containers/DynamicRHIResourceArray.h"
 #include "PixelFormat.h"
 #include "Math/PackedVector.h"
 #include "RHI.h"
 #include "RenderResource.h"
 #include "ShaderParameterStruct.h"
 #include "GlobalShader.h"
-#include "Containers/DynamicRHIResourceArray.h"
 
 class FSceneInterface;
 
