Temp workaround for missing FCopyResidentBricksCS caused by ordering issue between WorldContext.World()->InitWorld(); and GShaderCompilingManager->ProcessAsyncResults(false, true);

diff --git a/Engine/Source/Runtime/Engine/Private/PrecomputedVolumetricLightmap.cpp b/Engine/Source/Runtime/Engine/Private/PrecomputedVolumetricLightmap.cpp
--- a/Engine/Source/Runtime/Engine/Private/PrecomputedVolumetricLightmap.cpp
+++ b/Engine/Source/Runtime/Engine/Private/PrecomputedVolumetricLightmap.cpp
@@ -14,6 +14,8 @@
 #include "Interfaces/ITargetPlatform.h"
 #include "UObject/MobileObjectVersion.h"
 #include "RenderGraphUtils.h"
+// FIXME: temp fix for ordering issue between WorldContext.World()->InitWorld(); and GShaderCompilingManager->ProcessAsyncResults(false, true); in UnrealEngine.cpp
+#include "ShaderCompiler.h"
 
 DECLARE_MEMORY_STAT(TEXT("Volumetric Lightmap"),STAT_VolumetricLightmapBuildData,STATGROUP_MapBuildData);
 
@@ -697,6 +699,12 @@ void FPrecomputedVolumetricLightmap::AddToScene(FSceneInterface* Scene, UMapBuil
 		return;
 	}
 
+	// FIXME: temp fix for ordering issue between WorldContext.World()->InitWorld(); and GShaderCompilingManager->ProcessAsyncResults(false, true); in UnrealEngine.cpp
+	if (GShaderCompilingManager)
+	{
+		GShaderCompilingManager->ProcessAsyncResults(false, true);
+	}
+
 	check(!bAddedToScene);
 
 	FPrecomputedVolumetricLightmapData* NewData = NULL;
