Fixed crash in cook on the fly server.
ULevel::InitializeRenderingResources->FPrecomputedVolumetricLightmap::AddToScene->FPrecomputedVolumetricLightmap::SetData->Data->InitResource(); was accessing global shader map, which was nullptr as cook on the fly could be cooking for some other platform than PC.

diff --git a/Engine/Source/Runtime/Engine/Private/Level.cpp b/Engine/Source/Runtime/Engine/Private/Level.cpp
--- a/Engine/Source/Runtime/Engine/Private/Level.cpp
+++ b/Engine/Source/Runtime/Engine/Private/Level.cpp
@@ -1759,7 +1759,7 @@ void ULevel::InitializeRenderingResources()
 {
 	// OwningWorld can be NULL when InitializeRenderingResources is called during undo, where a transient ULevel is created to allow undoing level move operations
 	// At the point at which Pre/PostEditChange is called on that transient ULevel, it is not part of any world and therefore should not have its rendering resources initialized
-	if (OwningWorld && bIsVisible)
+	if (OwningWorld && bIsVisible && FApp::CanEverRender())
 	{
 		ULevel* ActiveLightingScenario = OwningWorld->GetActiveLightingScenario();
 		UMapBuildDataRegistry* EffectiveMapBuildData = MapBuildData;
@@ -1788,14 +1788,17 @@ void ULevel::InitializeRenderingResources()
 
 void ULevel::ReleaseRenderingResources()
 {
-	if (OwningWorld && PrecomputedLightVolume)
+	if (OwningWorld && FApp::CanEverRender())
 	{
-		PrecomputedLightVolume->RemoveFromScene(OwningWorld->Scene);
-	}
+		if (PrecomputedLightVolume)
+		{
+			PrecomputedLightVolume->RemoveFromScene(OwningWorld->Scene);
+		}
 
-	if (OwningWorld && PrecomputedVolumetricLightmap)
-	{
-		PrecomputedVolumetricLightmap->RemoveFromScene(OwningWorld->Scene);
+		if (PrecomputedVolumetricLightmap)
+		{
+			PrecomputedVolumetricLightmap->RemoveFromScene(OwningWorld->Scene);
+		}
 	}
 }
 
