Fix typed UAV loads in DX11.0

diff --git a/Engine/Shaders/Private/VolumetricLightmapStreaming.usf b/Engine/Shaders/Private/VolumetricLightmapStreaming.usf
--- a/Engine/Shaders/Private/VolumetricLightmapStreaming.usf
+++ b/Engine/Shaders/Private/VolumetricLightmapStreaming.usf
@@ -101,6 +101,8 @@ void PatchIndirectionTextureCS(uint3 BrickID : SV_DispatchThreadID)
 	IndirectionTexture[BrickPos] = uint4(NewBrickPos, 1);
 }
 
+Texture3D<uint4> OldIndirectionTexture;
+
 [numthreads(4, 4, 4)]
 void MoveWholeIndirectionTextureCS(uint3 IndirectionTextureCellID : SV_DispatchThreadID) 
 {
@@ -109,11 +111,15 @@ void MoveWholeIndirectionTextureCS(uint3 IndirectionTextureCellID : SV_DispatchT
 	
 	if (any(IndirectionTextureCellID >= IndirectionTextureDimenstions)) return;
 
-	uint4 IndirectionTextureOriginalValue = IndirectionTexture[IndirectionTextureCellID];
+	uint4 IndirectionTextureOriginalValue = OldIndirectionTexture[IndirectionTextureCellID];
 	uint OriginalBrickLinearAddress = IndirectionTextureOriginalValue.x + IndirectionTextureOriginalValue.y * 256 + IndirectionTextureOriginalValue.z * 256 * 256;
 	
 	if (OriginalBrickLinearAddress >= StartPosInOldVolume && OriginalBrickLinearAddress < StartPosInOldVolume + NumBricks)
 	{
 		IndirectionTexture[IndirectionTextureCellID] = uint4(ComputeBrickLayoutPosition(OriginalBrickLinearAddress - StartPosInOldVolume + StartPosInNewVolume, uint3(256, 256, 256)), IndirectionTextureOriginalValue.w);
 	}
+	else
+	{
+		IndirectionTexture[IndirectionTextureCellID] = IndirectionTextureOriginalValue;
+	}
 }
diff --git a/Engine/Source/Runtime/Engine/Private/PrecomputedVolumetricLightmap.cpp b/Engine/Source/Runtime/Engine/Private/PrecomputedVolumetricLightmap.cpp
--- a/Engine/Source/Runtime/Engine/Private/PrecomputedVolumetricLightmap.cpp
+++ b/Engine/Source/Runtime/Engine/Private/PrecomputedVolumetricLightmap.cpp
@@ -423,11 +423,16 @@ ENGINE_API void FPrecomputedVolumetricLightmapData::HandleDataMovementInAtlas(in
 
 			TShaderMapRef<FMoveWholeIndirectionTextureCS> ComputeShader(GlobalShaderMap);
 
+			FVolumetricLightmapDataLayer NewIndirectionTexture = SceneData->IndirectionTexture;
+			NewIndirectionTexture.CreateTargetTexture(IndirectionTextureDimensions);
+			NewIndirectionTexture.CreateUAV();
+
 			FMoveWholeIndirectionTextureCS::FParameters Parameters;
 			Parameters.NumBricks = NumBricks;
 			Parameters.StartPosInOldVolume = OldOffset;
 			Parameters.StartPosInNewVolume = BrickDataBaseOffsetInAtlas;
-			Parameters.IndirectionTexture = SceneData->IndirectionTexture.UAV;
+			Parameters.OldIndirectionTexture = SceneData->IndirectionTexture.Texture;
+			Parameters.IndirectionTexture = NewIndirectionTexture.UAV;
 
 			FComputeShaderUtils::Dispatch(RHICmdList, *ComputeShader, Parameters,
 				FIntVector(
@@ -435,6 +440,8 @@ ENGINE_API void FPrecomputedVolumetricLightmapData::HandleDataMovementInAtlas(in
 					FMath::DivideAndRoundUp(IndirectionTextureDimensions.Y, 4),
 					FMath::DivideAndRoundUp(IndirectionTextureDimensions.Z, 4))
 			);
+
+			SceneData->IndirectionTexture = NewIndirectionTexture;
 		}
 	}
 	else
@@ -497,22 +504,6 @@ ENGINE_API void FPrecomputedVolumetricLightmapData::AddToSceneData(FPrecomputedV
 		{
 			// GPU Path
 
-			if (GIsEditor)
-			{
-				// If we are in editor, or whenever a persistent level can be assigned to two Worlds (which is only possible when you modify the engine code),
-				// make a copy of the indirection texture
-				SceneData->IndirectionTexture.CreateTargetTexture(IndirectionTextureDimensions);
-				FRHICopyTextureInfo CopyInfo;
-				RHICmdList.CopyTexture(IndirectionTexture.Texture, SceneData->IndirectionTexture.Texture, CopyInfo);
-			}
-			else
-			{
-				// Steal the indirection texture. When the sublevels are unloaded the values will be restored.
-				SceneData->IndirectionTexture.Texture = IndirectionTexture.Texture;
-			}
-
-			SceneData->IndirectionTexture.CreateUAV();
-
 			const int32 PaddedBrickSize = BrickSize + 1;
 			int32 NumBricks = BrickDataDimensions.X * BrickDataDimensions.Y * BrickDataDimensions.Z / (PaddedBrickSize * PaddedBrickSize * PaddedBrickSize);
 
@@ -520,11 +511,16 @@ ENGINE_API void FPrecomputedVolumetricLightmapData::AddToSceneData(FPrecomputedV
 
 			TShaderMapRef<FMoveWholeIndirectionTextureCS> ComputeShader(GlobalShaderMap);
 
+			FVolumetricLightmapDataLayer NewIndirectionTexture = SceneData->IndirectionTexture;
+			NewIndirectionTexture.CreateTargetTexture(IndirectionTextureDimensions);
+			NewIndirectionTexture.CreateUAV();
+
 			FMoveWholeIndirectionTextureCS::FParameters Parameters;
 			Parameters.NumBricks = NumBricks;
 			Parameters.StartPosInOldVolume = 0;
 			Parameters.StartPosInNewVolume = BrickDataBaseOffsetInAtlas;
-			Parameters.IndirectionTexture = SceneData->IndirectionTexture.UAV;
+			Parameters.OldIndirectionTexture = IndirectionTexture.Texture;
+			Parameters.IndirectionTexture = NewIndirectionTexture.UAV;
 
 			FComputeShaderUtils::Dispatch(RHICmdList, *ComputeShader, Parameters,
 				FIntVector(
@@ -532,6 +528,14 @@ ENGINE_API void FPrecomputedVolumetricLightmapData::AddToSceneData(FPrecomputedV
 					FMath::DivideAndRoundUp(IndirectionTextureDimensions.Y, 4),
 					FMath::DivideAndRoundUp(IndirectionTextureDimensions.Z, 4))
 			);
+
+			SceneData->IndirectionTexture = NewIndirectionTexture;
+
+			if (!GIsEditor)
+			{
+				// Steal the indirection texture. When the sublevels are unloaded the values will be restored.
+				IndirectionTexture = SceneData->IndirectionTexture;
+			}
 		}
 		else
 		{
diff --git a/Engine/Source/Runtime/Engine/Public/PrecomputedVolumetricLightmap.h b/Engine/Source/Runtime/Engine/Public/PrecomputedVolumetricLightmap.h
--- a/Engine/Source/Runtime/Engine/Public/PrecomputedVolumetricLightmap.h
+++ b/Engine/Source/Runtime/Engine/Public/PrecomputedVolumetricLightmap.h
@@ -513,6 +513,7 @@ class FMoveWholeIndirectionTextureCS : public FGlobalShader
 		SHADER_PARAMETER(uint32, NumBricks)
 		SHADER_PARAMETER(uint32, StartPosInOldVolume)
 		SHADER_PARAMETER(uint32, StartPosInNewVolume)
+		SHADER_PARAMETER_TEXTURE(Texture3D<uint4>, OldIndirectionTexture)
 		SHADER_PARAMETER_UAV(RWTexture3D<uint4>, IndirectionTexture)
 	END_SHADER_PARAMETER_STRUCT()
 };
