Texture2DArray Implementation, big CL, enable feature with r.AllowTexture2DArrayCreation. Can create texture arrays from editor, use them in materials and contains several texture editor viewport changes.

diff --git a/Engine/Config/BaseEditor.ini b/Engine/Config/BaseEditor.ini
--- a/Engine/Config/BaseEditor.ini
+++ b/Engine/Config/BaseEditor.ini
@@ -24,6 +24,7 @@ MaxNumberOfLogLines=2000
 +RenderableThumbnailTypes=(ClassNeedingThumbnailName="/Script/Engine.TextureRenderTargetCube",RendererClassName="/Script/UnrealEd.TextureThumbnailRenderer")
 +RenderableThumbnailTypes=(ClassNeedingThumbnailName="/Script/Engine.LightMapTexture2D",RendererClassName="/Script/UnrealEd.TextureThumbnailRenderer")
 +RenderableThumbnailTypes=(ClassNeedingThumbnailName="/Script/Engine.TextureCube",RendererClassName="/Script/UnrealEd.TextureCubeThumbnailRenderer")
++RenderableThumbnailTypes=(ClassNeedingThumbnailName="/Script/Engine.Texture2DArray",RendererClassName="/Script/UnrealEd.Texture2DArrayThumbnailRenderer")
 +RenderableThumbnailTypes=(ClassNeedingThumbnailName="/Script/Engine.VolumeTexture",RendererClassName="/Script/UnrealEd.VolumeTextureThumbnailRenderer")
 +RenderableThumbnailTypes=(ClassNeedingThumbnailName="/Script/Engine.Material",RendererClassName="/Script/UnrealEd.MaterialInstanceThumbnailRenderer")
 +RenderableThumbnailTypes=(ClassNeedingThumbnailName="/Script/Engine.MaterialInterface",RendererClassName="/Script/UnrealEd.MaterialInstanceThumbnailRenderer")
diff --git a/Engine/Shaders/Private/Common.ush b/Engine/Shaders/Private/Common.ush
--- a/Engine/Shaders/Private/Common.ush
+++ b/Engine/Shaders/Private/Common.ush
@@ -134,6 +134,14 @@ MaterialFloat4 TextureCubeSample(TextureCube Tex, SamplerState Sampler, float3 U
 	return Tex.Sample(Sampler, UV);
 #endif
 }
+MaterialFloat4 Texture2DArraySample(Texture2DArray Tex, SamplerState Sampler, float3 UV)
+{
+#if !PIXELSHADER
+	return Tex.SampleLevel(Sampler, UV, 0);
+#else
+	return Tex.Sample(Sampler, UV);
+#endif
+}
 MaterialFloat4 Texture1DSampleLevel(Texture1D Tex, SamplerState Sampler, float UV, MaterialFloat Mip)
 {
 	return Tex.SampleLevel(Sampler, UV, Mip);
@@ -239,6 +247,23 @@ MaterialFloat4 TextureCubeSample_Decal(TextureCube Tex, SamplerState Sampler, fl
 #endif
 }
 
+MaterialFloat4 Texture2DArraySampleLevel(Texture2DArray Tex, SamplerState Sampler, float3 UV, MaterialFloat Mip)
+{
+	return Tex.SampleLevel(Sampler, UV, Mip);
+}
+MaterialFloat4 Texture2DArraySampleBias(Texture2DArray Tex, SamplerState Sampler, float3 UV, MaterialFloat MipBias)
+{
+#if !PIXELSHADER
+	return Tex.SampleLevel(Sampler, UV, 0);
+#else
+	return Tex.SampleBias(Sampler, UV, MipBias);
+#endif
+}
+MaterialFloat4 Texture2DArraySampleGrad(Texture2DArray Tex, SamplerState Sampler, float3 UV, MaterialFloat2 DDX, MaterialFloat2 DDY)
+{
+	return Tex.SampleGrad(Sampler, UV, DDX, DDY);
+}
+
 //converts an input 1d to 2d position. Useful for locating z frames that have been laid out in a 2d grid like a flipbook.
 float2 Tile1Dto2D(float xsize, float idx)
 {
@@ -417,7 +442,6 @@ MaterialFloat4 ClampedPow(MaterialFloat4 X,MaterialFloat4 Y)
 {
 	return pow(max(abs(X),MaterialFloat4(0.000001f,0.000001f,0.000001f,0.000001f)),Y);
 } 
-
 MaterialFloat PositiveClampedPow(MaterialFloat X,MaterialFloat Y)
 {
 	return pow(max(X,0.0f),Y);
@@ -426,13 +450,25 @@ MaterialFloat2 PositiveClampedPow(MaterialFloat2 X,MaterialFloat2 Y)
 {
 	return pow(max(X,MaterialFloat2(0.0f,0.0f)),Y);
 }
+MaterialFloat2 PositiveClampedPow(MaterialFloat2 X,MaterialFloat Y)
+{
+	return pow(max(X,MaterialFloat2(0.0f,0.0f)),MaterialFloat2(Y,Y));
+}
 MaterialFloat3 PositiveClampedPow(MaterialFloat3 X,MaterialFloat3 Y)
 {
 	return pow(max(X,MaterialFloat3(0.0f,0.0f,0.0f)),Y);
-}  
+} 
+MaterialFloat3 PositiveClampedPow(MaterialFloat3 X,MaterialFloat Y)
+{
+	return pow(max(X,MaterialFloat3(0.0f,0.0f,0.0f)),MaterialFloat3(Y,Y,Y));
+}
 MaterialFloat4 PositiveClampedPow(MaterialFloat4 X,MaterialFloat4 Y)
 {
 	return pow(max(X,MaterialFloat4(0.0f,0.0f,0.0f,0.0f)),Y);
+}
+MaterialFloat4 PositiveClampedPow(MaterialFloat4 X,MaterialFloat Y)
+{
+	return pow(max(X,MaterialFloat4(0.0f,0.0f,0.0f,0.0f)),MaterialFloat4(Y,Y,Y,Y));
 } 
 
 float DDX(float Input)
diff --git a/Engine/Shaders/Private/SimpleElementTexture2DPreviewPixelShader.usf b/Engine/Shaders/Private/SimpleElementTexture2DPreviewPixelShader.usf
--- a/Engine/Shaders/Private/SimpleElementTexture2DPreviewPixelShader.usf
+++ b/Engine/Shaders/Private/SimpleElementTexture2DPreviewPixelShader.usf
@@ -12,6 +12,10 @@
 #define SAMPLE_VIRTUAL_TEXTURE 0
 #endif
 
+#ifndef TEXTURE_ARRAY
+#define TEXTURE_ARRAY 0
+#endif
+
 #if SAMPLE_VIRTUAL_TEXTURE
 	#define NUM_VIRTUALTEXTURE_SAMPLES 1
 	#include "VirtualTextureCommon.ush"
@@ -19,7 +23,12 @@
 
 #define WRITE_TO_GBUFFER (FEATURE_LEVEL >= FEATURE_LEVEL_SM4 && !FORWARD_SHADING)
 
-Texture2D InTexture;
+#if TEXTURE_ARRAY
+	Texture2DArray InTexture;
+#else
+	Texture2D InTexture;
+#endif
+
 SamplerState InTextureSampler;
 
 #if SAMPLE_VIRTUAL_TEXTURE
@@ -37,6 +46,10 @@ float4x4 ColorWeights;
 //x=Gamma, y=MipLevel, z=bIsNormalMap, w=VT Layer
 float4 PackedParams;
 
+#if TEXTURE_ARRAY
+float NumSlices;
+#endif
+
 void Main(
 	in float2 TextureCoordinate : TEXCOORD0,
 	in float4 Color : TEXCOORD1,
@@ -56,6 +69,21 @@ void Main(
 	float4 FinalColor;
 	float4 Sample;
 
+#if TEXTURE_ARRAY
+	// Divide up the texture array into horizontal slices.
+	float SliceSize = 1.0 / NumSlices;
+	float SliceIndex = floor(TextureCoordinate.y / SliceSize);
+	float3 SampleCoordinates = float3(TextureCoordinate.xy, SliceIndex);
+
+	if( MipLevel >= 0.0f )
+	{
+		Sample = Texture2DArraySampleLevel(InTexture, InTextureSampler, SampleCoordinates, MipLevel);
+	}
+	else
+	{
+		Sample = Texture2DArraySample(InTexture, InTextureSampler, SampleCoordinates);
+	}
+#else //TEXTURE_ARAY
 	if( MipLevel >= 0.0f )
 	{
 #if SAMPLE_VIRTUAL_TEXTURE
@@ -89,6 +117,7 @@ void Main(
 		Sample = Texture2DSample(InTexture, InTextureSampler,TextureCoordinate);
 #endif
 	}
+#endif //TEXTURE_ARRAY
 
 	ReplicateChannel(Sample,TextureComponentReplicate,TextureComponentReplicateAlpha);
 
diff --git a/Engine/Source/Developer/AssetTools/Private/AssetTools.cpp b/Engine/Source/Developer/AssetTools/Private/AssetTools.cpp
--- a/Engine/Source/Developer/AssetTools/Private/AssetTools.cpp
+++ b/Engine/Source/Developer/AssetTools/Private/AssetTools.cpp
@@ -90,6 +90,7 @@
 #include "AssetTypeActions/AssetTypeActions_SlateWidgetStyle.h"
 #include "AssetTypeActions/AssetTypeActions_StaticMesh.h"
 #include "AssetTypeActions/AssetTypeActions_Texture2D.h"
+#include "AssetTypeActions/AssetTypeActions_Texture2DArray.h"
 #include "AssetTypeActions/AssetTypeActions_TextureCube.h"
 #include "AssetTypeActions/AssetTypeActions_VolumeTexture.h"
 #include "AssetTypeActions/AssetTypeActions_TextureRenderTargetCube.h"
@@ -231,7 +232,8 @@ UAssetToolsImpl::UAssetToolsImpl(const FObjectInitializer& ObjectInitializer)
 	RegisterAssetTypeActions(MakeShareable(new FAssetTypeActions_Texture));
 	RegisterAssetTypeActions(MakeShareable(new FAssetTypeActions_Texture2D));
 	RegisterAssetTypeActions(MakeShareable(new FAssetTypeActions_TextureCube));
-	RegisterAssetTypeActions( MakeShareable(new FAssetTypeActions_VolumeTexture) );
+	RegisterAssetTypeActions(MakeShareable(new FAssetTypeActions_Texture2DArray));
+	RegisterAssetTypeActions(MakeShareable(new FAssetTypeActions_VolumeTexture) );
 	RegisterAssetTypeActions(MakeShareable(new FAssetTypeActions_TextureRenderTarget));
 	RegisterAssetTypeActions(MakeShareable(new FAssetTypeActions_TextureRenderTarget2D));
 	RegisterAssetTypeActions(MakeShareable(new FAssetTypeActions_TextureRenderTargetCube));
diff --git a/Engine/Source/Developer/AssetTools/Private/AssetTypeActions/AssetTypeActions_Texture2D.cpp b/Engine/Source/Developer/AssetTools/Private/AssetTypeActions/AssetTypeActions_Texture2D.cpp
--- a/Engine/Source/Developer/AssetTools/Private/AssetTypeActions/AssetTypeActions_Texture2D.cpp
+++ b/Engine/Source/Developer/AssetTools/Private/AssetTypeActions/AssetTypeActions_Texture2D.cpp
@@ -7,7 +7,9 @@
 #include "Factories/SlateBrushAssetFactory.h"
 #include "Slate/SlateBrushAsset.h"
 #include "Factories/VolumeTextureFactory.h"
+#include "Factories/Texture2DArrayFactory.h"
 #include "Engine/VolumeTexture.h"
+#include "Engine/Texture2DArray.h"
 #include "AssetTools.h"
 #include "IContentBrowserSingleton.h"
 #include "ContentBrowserModule.h"
@@ -27,6 +29,17 @@ void FAssetTypeActions_Texture2D::GetActions( const TArray<UObject*>& InObjects,
 		FUIAction(FExecuteAction::CreateSP( this, &FAssetTypeActions_Texture2D::ExecuteCreateSlateBrush, Textures ), FCanExecuteAction())
 		);
 
+	static const auto AllowTextureArrayAssetCreationVar = IConsoleManager::Get().FindTConsoleVariableDataInt(TEXT("r.AllowTexture2DArrayCreation"));
+	if (Textures.Num() > 0 && AllowTextureArrayAssetCreationVar->GetValueOnGameThread() != 0)
+	{
+		MenuBuilder.AddMenuEntry(
+			LOCTEXT("Texture_Texture2DArray", "Create Texture Array"),
+			LOCTEXT("Texture_CreateTexture2DArrayTooltip", "Creates a new texture array."),
+			FSlateIcon(FEditorStyle::GetStyleSetName(), "ClassIcon.Texture2D"),
+			FUIAction(FExecuteAction::CreateSP(this, &FAssetTypeActions_Texture2D::ExecuteCreateTextureArray, Textures), FCanExecuteAction())
+			);
+	}
+
 	static const auto AllowVolumeTextureAssetCreationVar = IConsoleManager::Get().FindTConsoleVariableDataInt(TEXT("r.AllowVolumeTextureAssetCreation"));
 	if (InObjects.Num() == 1 && AllowVolumeTextureAssetCreationVar->GetValueOnGameThread() != 0)
 	{
@@ -97,6 +110,29 @@ void FAssetTypeActions_Texture2D::ExecuteCreateSlateBrush(TArray<TWeakObjectPtr<
 	}
 }
 
+void FAssetTypeActions_Texture2D::ExecuteCreateTextureArray(TArray<TWeakObjectPtr<UTexture2D>> Objects)
+{
+	const FString DefaultSuffix = TEXT("_Array");
+	FString Name;
+	FString PackagePath;
+	CreateUniqueAssetName(Objects[0].Get()->GetOutermost()->GetName(), DefaultSuffix, PackagePath, Name);
+	
+	// Create the factory used to generate the asset.
+	UTexture2DArrayFactory* Factory = NewObject<UTexture2DArrayFactory>();
+	Factory->InitialTextures.Empty();
+
+	// Give the selected textures to the factory.
+	for (int32 TextureIndex = 0; TextureIndex < Objects.Num(); ++TextureIndex)
+	{
+		Factory->InitialTextures.Add(Objects[TextureIndex].Get());
+	}
+
+	{
+		FContentBrowserModule& ContentBrowserModule = FModuleManager::LoadModuleChecked<FContentBrowserModule>("ContentBrowser");
+		ContentBrowserModule.Get().CreateNewAsset(Name, FPackageName::GetLongPackagePath(PackagePath), UTexture2DArray::StaticClass(), Factory);
+	}		
+}
+
 void FAssetTypeActions_Texture2D::ExecuteCreateVolumeTexture(TArray<TWeakObjectPtr<UTexture2D>> Objects)
 {
 	const FString DefaultSuffix = TEXT("_Volume");
diff --git a/Engine/Source/Developer/AssetTools/Private/AssetTypeActions/AssetTypeActions_Texture2D.h b/Engine/Source/Developer/AssetTools/Private/AssetTypeActions/AssetTypeActions_Texture2D.h
--- a/Engine/Source/Developer/AssetTools/Private/AssetTypeActions/AssetTypeActions_Texture2D.h
+++ b/Engine/Source/Developer/AssetTools/Private/AssetTypeActions/AssetTypeActions_Texture2D.h
@@ -24,4 +24,5 @@ private:
 	/** Handler for when Create Slate Brush is selected */
 	void ExecuteCreateSlateBrush(TArray<TWeakObjectPtr<UTexture2D>> Objects);
 	void ExecuteCreateVolumeTexture(TArray<TWeakObjectPtr<UTexture2D>> Objects);
+	void ExecuteCreateTextureArray(TArray<TWeakObjectPtr<UTexture2D>> Objects);
 };
diff --git a/Engine/Source/Developer/AssetTools/Private/AssetTypeActions/AssetTypeActions_Texture2DArray.cpp b/Engine/Source/Developer/AssetTools/Private/AssetTypeActions/AssetTypeActions_Texture2DArray.cpp
new file mode 100644
--- /dev/null
+++ b/Engine/Source/Developer/AssetTools/Private/AssetTypeActions/AssetTypeActions_Texture2DArray.cpp
@@ -0,0 +1,2 @@
+#include "AssetTypeActions_Texture2DArray.h"
+
diff --git a/Engine/Source/Developer/AssetTools/Private/AssetTypeActions/AssetTypeActions_Texture2DArray.h b/Engine/Source/Developer/AssetTools/Private/AssetTypeActions/AssetTypeActions_Texture2DArray.h
new file mode 100644
--- /dev/null
+++ b/Engine/Source/Developer/AssetTools/Private/AssetTypeActions/AssetTypeActions_Texture2DArray.h
@@ -0,0 +1,17 @@
+// Copyright 1998-2019 Epic Games, Inc. All Rights Reserved.
+
+#pragma once
+
+#include "CoreMinimal.h"
+#include "AssetTypeActions/AssetTypeActions_Texture.h"
+#include "Engine/Texture2DArray.h"
+
+class FAssetTypeActions_Texture2DArray : public FAssetTypeActions_Texture
+{
+public:
+	// IAssetTypeActions Implementation
+	virtual FText GetName() const override { return NSLOCTEXT("AssetTypeActions", "AssetTypeActions_Texture2DArray", "Texture 2D Array"); }
+	virtual FColor GetTypeColor() const override { return FColor(0,64,128); }
+	virtual UClass* GetSupportedClass() const override{ return UTexture2DArray::StaticClass();}
+	virtual bool CanFilter() override { return true;}
+}; 
\ No newline at end of file
diff --git a/Engine/Source/Developer/MaterialBaking/Private/ExportMaterialProxy.h b/Engine/Source/Developer/MaterialBaking/Private/ExportMaterialProxy.h
--- a/Engine/Source/Developer/MaterialBaking/Private/ExportMaterialProxy.h
+++ b/Engine/Source/Developer/MaterialBaking/Private/ExportMaterialProxy.h
@@ -10,6 +10,7 @@
 #include "Engine/Texture2D.h"
 #include "Engine/Texture.h"
 #include "Engine/TextureCube.h"
+#include "Engine/Texture2DArray.h"
 
 #include "DeviceProfiles/DeviceProfileManager.h"
 #include "DeviceProfiles/DeviceProfile.h"
@@ -335,7 +336,7 @@ public:
 		// needs to be called in this function!!
 		Compiler->SetMaterialProperty(Property, OverrideShaderFrequency, bUsePreviousFrameTime);
 		const int32 Ret = CompilePropertyAndSetMaterialPropertyWithoutCast(Property, Compiler);
-		return Compiler->ForceCast(Ret, FMaterialAttributeDefinitionMap::GetValueType(Property));
+		return Compiler->ForceCast(Ret, FMaterialAttributeDefinitionMap::GetValueType(Property), MFCF_ExactMatch);
 	}
 
 	/** helper for CompilePropertyAndSetMaterialProperty() */
@@ -532,7 +533,11 @@ public:
 					UTextureCube* TexCube = (UTextureCube*)Texture;
 					return FIntPoint(TexCube->GetSizeX(), TexCube->GetSizeY());
 				}
-
+				else if (Texture->IsA(UTexture2DArray::StaticClass())) 
+				{
+					UTexture2DArray* TexArray = (UTexture2DArray*)Texture;
+					return FIntPoint(TexArray->GetSizeX(), TexArray->GetSizeY());
+				}
 				return FIntPoint(0, 0);
 			}();
 			
diff --git a/Engine/Source/Developer/MaterialUtilities/Private/MaterialUtilities.cpp b/Engine/Source/Developer/MaterialUtilities/Private/MaterialUtilities.cpp
--- a/Engine/Source/Developer/MaterialUtilities/Private/MaterialUtilities.cpp
+++ b/Engine/Source/Developer/MaterialUtilities/Private/MaterialUtilities.cpp
@@ -17,6 +17,7 @@
 #include "Materials/MaterialExpressionConstant4Vector.h"
 #include "Materials/MaterialExpressionMultiply.h"
 #include "Engine/TextureCube.h"
+#include "Engine/Texture2DArray.h"
 #include "SceneView.h"
 #include "RendererInterface.h"
 #include "EngineModule.h"
@@ -517,7 +518,7 @@ public:
 
 		int32 Ret = CompilePropertyAndSetMaterialPropertyWithoutCast(Property, Compiler);
 
-		return Compiler->ForceCast(Ret, FMaterialAttributeDefinitionMap::GetValueType(Property));
+		return Compiler->ForceCast(Ret, FMaterialAttributeDefinitionMap::GetValueType(Property), MFCF_ExactMatch);
 	}
 
 	/** helper for CompilePropertyAndSetMaterialProperty() */
@@ -710,6 +711,11 @@ public:
 				UTextureCube* TexCube = (UTextureCube*)Texture;
 				LocalSize = FIntPoint(TexCube->GetSizeX(), TexCube->GetSizeY());
 			}
+			else if (Texture->IsA(UTexture2DArray::StaticClass())) 
+			{
+				UTexture2DArray* TexArray = (UTexture2DArray*)Texture;
+				LocalSize = FIntPoint(TexArray->GetSizeX(), TexArray->GetSizeY());
+			}
 
 			int32 LocalBias = GameTextureLODSettings->CalculateLODBias(Texture);
 
diff --git a/Engine/Source/Developer/TextureCompressor/Private/TextureCompressorModule.cpp b/Engine/Source/Developer/TextureCompressor/Private/TextureCompressorModule.cpp
--- a/Engine/Source/Developer/TextureCompressor/Private/TextureCompressorModule.cpp
+++ b/Engine/Source/Developer/TextureCompressor/Private/TextureCompressorModule.cpp
@@ -1969,7 +1969,8 @@ public:
 		BuildSettings.TopMipSize.X = IntermediateMipChain[0].SizeX;
 		BuildSettings.TopMipSize.Y = IntermediateMipChain[0].SizeY;
 		BuildSettings.VolumeSizeZ = BuildSettings.bVolume ? IntermediateMipChain[0].NumSlices : 1;
-
+		BuildSettings.ArraySlices = BuildSettings.bTextureArray ? IntermediateMipChain[0].NumSlices : 1;
+		
 		return CompressMipChain(IntermediateMipChain, BuildSettings, OutTextureMips);
 	}
 
diff --git a/Engine/Source/Developer/TextureCompressor/Public/TextureCompressorModule.h b/Engine/Source/Developer/TextureCompressor/Public/TextureCompressorModule.h
--- a/Engine/Source/Developer/TextureCompressor/Public/TextureCompressorModule.h
+++ b/Engine/Source/Developer/TextureCompressor/Public/TextureCompressorModule.h
@@ -87,6 +87,8 @@ struct FTextureBuildSettings
 	uint8 MipGenSettings; // TextureMipGenSettings, opaque to avoid dependencies on engine headers.
 	/** Whether the texture being built is a cubemap. */
 	uint32 bCubemap : 1;
+	/** Whether the texture being built is a texture array. */
+	uint32 bTextureArray : 1;
 	/** Whether the texture being built is a volume. */
 	uint32 bVolume : 1;
 	/** Whether the texture being built from long/lat source to cubemap. */
@@ -129,6 +131,8 @@ struct FTextureBuildSettings
 	mutable FIntPoint TopMipSize;
 	/** The volume texture's top mip size Z without LODBias applied */
 	mutable int32 VolumeSizeZ;
+	/** The array texture's top mip size Z without LODBias applied */
+	mutable int32 ArraySlices;
 	/** Can the texture be streamed */
 	uint32 bStreamable : 1;
 	/** Is the texture streamed using the VT system */
@@ -169,6 +173,7 @@ struct FTextureBuildSettings
 		, bHDRSource(false)
 		, MipGenSettings(1 /*TMGS_SimpleAverage*/)
 		, bCubemap(false)
+		, bTextureArray(false)
 		, bVolume(false)
 		, bLongLatSource(false)
 		, bSRGB(false)
@@ -190,6 +195,7 @@ struct FTextureBuildSettings
 		, LODBiasWithCinematicMips(0)
 		, TopMipSize(0, 0)
 		, VolumeSizeZ(0)
+		, ArraySlices(0)
 		, bStreamable(false)
 		, bVirtualStreamable(false)
 		, bChromaKeyTexture(false)
diff --git a/Engine/Source/Developer/TextureFormatASTC/Private/TextureFormatASTC.cpp b/Engine/Source/Developer/TextureFormatASTC/Private/TextureFormatASTC.cpp
--- a/Engine/Source/Developer/TextureFormatASTC/Private/TextureFormatASTC.cpp
+++ b/Engine/Source/Developer/TextureFormatASTC/Private/TextureFormatASTC.cpp
@@ -420,7 +420,7 @@ public:
 		{
 			OutCompressedImage.SizeX = Image.SizeX;
 			OutCompressedImage.SizeY = Image.SizeY;
-			OutCompressedImage.SizeZ = BuildSettings.bVolume ? Image.NumSlices : 1;
+			OutCompressedImage.SizeZ = (BuildSettings.bVolume || BuildSettings.bTextureArray) ? Image.NumSlices : 1;
 			OutCompressedImage.PixelFormat = CompressedPixelFormat;
 		}
 		return bCompressionSucceeded;
diff --git a/Engine/Source/Developer/TextureFormatAndroid/Private/TextureFormatAndroid.cpp b/Engine/Source/Developer/TextureFormatAndroid/Private/TextureFormatAndroid.cpp
--- a/Engine/Source/Developer/TextureFormatAndroid/Private/TextureFormatAndroid.cpp
+++ b/Engine/Source/Developer/TextureFormatAndroid/Private/TextureFormatAndroid.cpp
@@ -173,7 +173,7 @@ class FTextureFormatAndroid : public ITextureFormat
 				// ETC1 can't support an alpha channel, store uncompressed
 				OutCompressedImage.SizeX = Image.SizeX;
 				OutCompressedImage.SizeY = Image.SizeY;
-				OutCompressedImage.SizeZ = BuildSettings.bVolume ? Image.NumSlices : 1;
+				OutCompressedImage.SizeZ = (BuildSettings.bVolume || BuildSettings.bTextureArray) ? Image.NumSlices : 1;
 				OutCompressedImage.PixelFormat = PF_B8G8R8A8;
 				OutCompressedImage.RawData = Image.RawData;
 				return true;
@@ -249,7 +249,7 @@ class FTextureFormatAndroid : public ITextureFormat
 		{
 			OutCompressedImage.SizeX = FMath::Max(Image.SizeX, 4);
 			OutCompressedImage.SizeY = FMath::Max(Image.SizeY, 4);
-			OutCompressedImage.SizeZ = BuildSettings.bVolume ? Image.NumSlices : 1;
+			OutCompressedImage.SizeZ = (BuildSettings.bVolume || BuildSettings.bTextureArray) ? Image.NumSlices : 1;
 			OutCompressedImage.PixelFormat = CompressedPixelFormat;
 		}
 
diff --git a/Engine/Source/Developer/TextureFormatDXT/Private/TextureFormatDXT.cpp b/Engine/Source/Developer/TextureFormatDXT/Private/TextureFormatDXT.cpp
--- a/Engine/Source/Developer/TextureFormatDXT/Private/TextureFormatDXT.cpp
+++ b/Engine/Source/Developer/TextureFormatDXT/Private/TextureFormatDXT.cpp
@@ -496,7 +496,7 @@ class FTextureFormatDXT : public ITextureFormat
 		{
 			OutCompressedImage.SizeX = FMath::Max(Image.SizeX, 4);
 			OutCompressedImage.SizeY = FMath::Max(Image.SizeY, 4);
-			OutCompressedImage.SizeZ = BuildSettings.bVolume ? Image.NumSlices : 1;
+			OutCompressedImage.SizeZ = (BuildSettings.bVolume || BuildSettings.bTextureArray) ? Image.NumSlices : 1;
 			OutCompressedImage.PixelFormat = CompressedPixelFormat;
 		}
 		return bCompressionSucceeded;
diff --git a/Engine/Source/Developer/TextureFormatIntelISPCTexComp/Private/TextureFormatIntelISPCTexComp.cpp b/Engine/Source/Developer/TextureFormatIntelISPCTexComp/Private/TextureFormatIntelISPCTexComp.cpp
--- a/Engine/Source/Developer/TextureFormatIntelISPCTexComp/Private/TextureFormatIntelISPCTexComp.cpp
+++ b/Engine/Source/Developer/TextureFormatIntelISPCTexComp/Private/TextureFormatIntelISPCTexComp.cpp
@@ -814,7 +814,7 @@ public:
 		OutCompressedImage.PixelFormat = CompressedPixelFormat;
 		OutCompressedImage.SizeX = InImage.SizeX;
 		OutCompressedImage.SizeY = InImage.SizeY;
-		OutCompressedImage.SizeZ = BuildSettings.bVolume ? InImage.NumSlices : 1;
+		OutCompressedImage.SizeZ = (BuildSettings.bVolume || BuildSettings.bTextureArray) ? InImage.NumSlices : 1;
 		return bCompressionSucceeded;
 	}
 };
diff --git a/Engine/Source/Developer/TextureFormatPVR/Private/TextureFormatPVR.cpp b/Engine/Source/Developer/TextureFormatPVR/Private/TextureFormatPVR.cpp
--- a/Engine/Source/Developer/TextureFormatPVR/Private/TextureFormatPVR.cpp
+++ b/Engine/Source/Developer/TextureFormatPVR/Private/TextureFormatPVR.cpp
@@ -368,7 +368,7 @@ class FTextureFormatPVR : public ITextureFormat
 		{
 			OutCompressedImage.SizeX = FinalSquareSize;
 			OutCompressedImage.SizeY = FinalSquareSize;
-			OutCompressedImage.SizeZ = BuildSettings.bVolume ? InImage.NumSlices : 1;
+			OutCompressedImage.SizeZ = (BuildSettings.bVolume || BuildSettings.bTextureArray) ? InImage.NumSlices : 1;
 			OutCompressedImage.PixelFormat = CompressedPixelFormat;
 		}
 
diff --git a/Engine/Source/Developer/TextureFormatUncompressed/Private/TextureFormatUncompressed.cpp b/Engine/Source/Developer/TextureFormatUncompressed/Private/TextureFormatUncompressed.cpp
--- a/Engine/Source/Developer/TextureFormatUncompressed/Private/TextureFormatUncompressed.cpp
+++ b/Engine/Source/Developer/TextureFormatUncompressed/Private/TextureFormatUncompressed.cpp
@@ -81,7 +81,7 @@ class FTextureFormatUncompressed : public ITextureFormat
 
 			OutCompressedImage.SizeX = Image.SizeX;
 			OutCompressedImage.SizeY = Image.SizeY;
-			OutCompressedImage.SizeZ = BuildSettings.bVolume ? Image.NumSlices : 1;
+			OutCompressedImage.SizeZ = (BuildSettings.bVolume || BuildSettings.bTextureArray) ? Image.NumSlices : 1;
 			OutCompressedImage.PixelFormat = PF_G8;
 			OutCompressedImage.RawData = Image.RawData;
 
@@ -94,7 +94,7 @@ class FTextureFormatUncompressed : public ITextureFormat
 
 			OutCompressedImage.SizeX = Image.SizeX;
 			OutCompressedImage.SizeY = Image.SizeY;
-			OutCompressedImage.SizeZ = BuildSettings.bVolume ? Image.NumSlices : 1;
+			OutCompressedImage.SizeZ = (BuildSettings.bVolume || BuildSettings.bTextureArray) ? Image.NumSlices : 1;
 			OutCompressedImage.PixelFormat = PF_V8U8;
 
 			uint32 NumTexels = Image.SizeX * Image.SizeY * Image.NumSlices;
@@ -119,7 +119,7 @@ class FTextureFormatUncompressed : public ITextureFormat
 
 			OutCompressedImage.SizeX = Image.SizeX;
 			OutCompressedImage.SizeY = Image.SizeY;
-			OutCompressedImage.SizeZ = BuildSettings.bVolume ? Image.NumSlices : 1;
+			OutCompressedImage.SizeZ = (BuildSettings.bVolume || BuildSettings.bTextureArray) ? Image.NumSlices : 1;
 			OutCompressedImage.PixelFormat = PF_B8G8R8A8;
 			OutCompressedImage.RawData = Image.RawData;
 
@@ -132,7 +132,7 @@ class FTextureFormatUncompressed : public ITextureFormat
 
 			OutCompressedImage.SizeX = Image.SizeX;
 			OutCompressedImage.SizeY = Image.SizeY;
-			OutCompressedImage.SizeZ = BuildSettings.bVolume ? Image.NumSlices : 1;
+			OutCompressedImage.SizeZ = (BuildSettings.bVolume || BuildSettings.bTextureArray) ? Image.NumSlices : 1;
 			OutCompressedImage.PixelFormat = PF_B8G8R8A8;
 
 			// swizzle each texel
@@ -160,7 +160,7 @@ class FTextureFormatUncompressed : public ITextureFormat
 
 			OutCompressedImage.SizeX = Image.SizeX;
 			OutCompressedImage.SizeY = Image.SizeY;
-			OutCompressedImage.SizeZ = BuildSettings.bVolume ? Image.NumSlices : 1;
+			OutCompressedImage.SizeZ = (BuildSettings.bVolume || BuildSettings.bTextureArray) ? Image.NumSlices : 1;
 			OutCompressedImage.PixelFormat = PF_B8G8R8A8;
 
 			// swizzle each texel
@@ -188,7 +188,7 @@ class FTextureFormatUncompressed : public ITextureFormat
 
 			OutCompressedImage.SizeX = Image.SizeX;
 			OutCompressedImage.SizeY = Image.SizeY;
-			OutCompressedImage.SizeZ = BuildSettings.bVolume ? Image.NumSlices : 1;
+			OutCompressedImage.SizeZ = (BuildSettings.bVolume || BuildSettings.bTextureArray) ? Image.NumSlices : 1;
 			OutCompressedImage.PixelFormat = PF_FloatRGBA;
 			OutCompressedImage.RawData = Image.RawData;
 
@@ -203,7 +203,7 @@ class FTextureFormatUncompressed : public ITextureFormat
 			// set output
 			OutCompressedImage.SizeX = InImage.SizeX;
 			OutCompressedImage.SizeY = InImage.SizeY;
-			OutCompressedImage.SizeZ = BuildSettings.bVolume ? InImage.NumSlices : 1;
+			OutCompressedImage.SizeZ = (BuildSettings.bVolume || BuildSettings.bTextureArray) ? InImage.NumSlices : 1;
 			OutCompressedImage.PixelFormat = PF_B8G8R8A8;
 
 			// allocate output memory
diff --git a/Engine/Source/Editor/MaterialEditor/Private/MaterialEditor.cpp b/Engine/Source/Editor/MaterialEditor/Private/MaterialEditor.cpp
--- a/Engine/Source/Editor/MaterialEditor/Private/MaterialEditor.cpp
+++ b/Engine/Source/Editor/MaterialEditor/Private/MaterialEditor.cpp
@@ -31,6 +31,7 @@
 #include "Materials/MaterialInstanceConstant.h"
 #include "Materials/MaterialParameterCollection.h"
 #include "Engine/TextureCube.h"
+#include "Engine/Texture2DArray.h"
 #include "Dialogs/Dialogs.h"
 #include "UnrealEdGlobals.h"
 #include "Editor.h"
@@ -63,6 +64,7 @@
 #include "Materials/MaterialExpressionTextureObject.h"
 #include "Materials/MaterialExpressionTextureSampleParameter2D.h"
 #include "Materials/MaterialExpressionTextureSampleParameterCube.h"
+#include "Materials/MaterialExpressionTextureSampleParameter2DArray.h"
 #include "Materials/MaterialExpressionTextureSampleParameterSubUV.h"
 #include "Materials/MaterialExpressionTransformPosition.h"
 #include "Materials/MaterialExpressionVectorParameter.h"
@@ -223,7 +225,7 @@ int32 FMatExpressionPreview::CompilePropertyAndSetMaterialProperty(EMaterialProp
 	}
 
 	// output should always be the right type for this property
-	return Compiler->ForceCast(Ret, FMaterialAttributeDefinitionMap::GetValueType(Property));
+	return Compiler->ForceCast(Ret, FMaterialAttributeDefinitionMap::GetValueType(Property), MFCF_ExactMatch);
 }
 
 void FMatExpressionPreview::NotifyCompilationFinished()
@@ -2860,6 +2862,10 @@ void FMaterialEditor::OnConvertObjects()
 				{
 					ClassToCreate = UMaterialExpressionTextureSampleParameterCube::StaticClass();
 				}
+				else if (TextureSampleExpression && TextureSampleExpression->Texture && TextureSampleExpression->Texture->IsA(UTexture2DArray::StaticClass()))
+				{
+					ClassToCreate = UMaterialExpressionTextureSampleParameter2DArray::StaticClass();
+				}
 				else if (TextureObjectExpression)
 				{
 					ClassToCreate = UMaterialExpressionTextureObjectParameter::StaticClass();
diff --git a/Engine/Source/Editor/TextureEditor/Private/Models/TextureEditorViewportClient.cpp b/Engine/Source/Editor/TextureEditor/Private/Models/TextureEditorViewportClient.cpp
--- a/Engine/Source/Editor/TextureEditor/Private/Models/TextureEditorViewportClient.cpp
+++ b/Engine/Source/Editor/TextureEditor/Private/Models/TextureEditorViewportClient.cpp
@@ -7,6 +7,7 @@
 #include "Engine/Texture2D.h"
 #include "ThumbnailRendering/ThumbnailManager.h"
 #include "Engine/TextureCube.h"
+#include "Engine/Texture2DArray.h"
 #include "Engine/VolumeTexture.h"
 #include "Engine/TextureRenderTarget2D.h"
 #include "Engine/TextureRenderTargetCube.h"
@@ -76,6 +77,7 @@ void FTextureEditorViewportClient::Draw(FViewport* Viewport, FCanvas* Canvas)
 
 	UTexture2D* Texture2D = Cast<UTexture2D>(Texture);
 	UTextureCube* TextureCube = Cast<UTextureCube>(Texture);
+	UTexture2DArray* Texture2DArray = Cast<UTexture2DArray>(Texture);
 	UVolumeTexture* VolumeTexture = Cast<UVolumeTexture>(Texture);
 	UTextureRenderTarget2D* TextureRT2D = Cast<UTextureRenderTarget2D>(Texture);
 	UTextureRenderTargetCube* RTTextureCube = Cast<UTextureRenderTargetCube>(Texture);
@@ -122,6 +124,12 @@ void FTextureEditorViewportClient::Draw(FViewport* Viewport, FCanvas* Canvas)
 			bIsVirtualTexture = Texture2D->IsCurrentlyVirtualTextured();
 			BatchedElementParameters = new FBatchedElementTexture2DPreviewParameters(MipLevel, LayerIndex, bIsNormalMap, bIsSingleChannel, bIsVirtualTexture);
 		}
+		else if (Texture2DArray) 
+		{
+			bool bIsNormalMap = Texture2DArray->IsNormalMap();
+			bool bIsSingleChannel = Texture2DArray->CompressionSettings == TC_Grayscale || Texture2DArray->CompressionSettings == TC_Alpha;
+			BatchedElementParameters = new FBatchedElementTexture2DPreviewParameters(MipLevel, LayerIndex, bIsNormalMap, bIsSingleChannel, false, true);
+		}
 		else if (TextureRT2D)
 		{
 			BatchedElementParameters = new FBatchedElementTexture2DPreviewParameters(MipLevel, LayerIndex, false, false, false);
diff --git a/Engine/Source/Editor/TextureEditor/Private/TextureEditorToolkit.cpp b/Engine/Source/Editor/TextureEditor/Private/TextureEditorToolkit.cpp
--- a/Engine/Source/Editor/TextureEditor/Private/TextureEditorToolkit.cpp
+++ b/Engine/Source/Editor/TextureEditor/Private/TextureEditorToolkit.cpp
@@ -20,6 +20,7 @@
 #include "Engine/ShadowMapTexture2D.h"
 #include "Engine/Texture2DDynamic.h"
 #include "Engine/TextureCube.h"
+#include "Engine/Texture2DArray.h"
 #include "Engine/VolumeTexture.h"
 #include "Engine/TextureRenderTarget.h"
 #include "Engine/TextureRenderTarget2D.h"
@@ -322,7 +323,7 @@ void FTextureEditorToolkit::CalculateTextureDimensions( uint32& Width, uint32& H
 
 ESimpleElementBlendMode FTextureEditorToolkit::GetColourChannelBlendMode( ) const
 {
-	if (Texture && (Texture->CompressionSettings == TC_Grayscale || Texture->CompressionSettings == TC_Alpha)) 
+	if (Texture && (Texture->CompressionSettings == TC_Grayscale || Texture->CompressionSettings == TC_Alpha))
 	{
 		return SE_BLEND_Opaque;
 	}
@@ -402,6 +403,7 @@ void FTextureEditorToolkit::PopulateQuickInfo( )
 	UTextureRenderTarget2D* Texture2DRT = Cast<UTextureRenderTarget2D>(Texture);
 	UTextureRenderTargetCube* TextureCubeRT = Cast<UTextureRenderTargetCube>(Texture);
 	UTextureCube* TextureCube = Cast<UTextureCube>(Texture);
+	UTexture2DArray* Texture2DArray = Cast<UTexture2DArray>(Texture);
 	UTexture2DDynamic* Texture2DDynamic = Cast<UTexture2DDynamic>(Texture);
 	UVolumeTexture* VolumeTexture = Cast<UVolumeTexture>(Texture);
 
@@ -497,6 +499,10 @@ void FTextureEditorToolkit::PopulateQuickInfo( )
 	{
 		TextureFormatIndex = TextureCube->GetPixelFormat();
 	}
+	else if (Texture2DArray) 
+	{
+		TextureFormatIndex = Texture2DArray->GetPixelFormat();
+	}
 	else if (Texture2DRT)
 	{
 		TextureFormatIndex = Texture2DRT->GetFormat();
@@ -524,6 +530,10 @@ void FTextureEditorToolkit::PopulateQuickInfo( )
 	{
 		NumMips = TextureCube->GetNumMips();
 	}
+	else if (Texture2DArray) 
+	{
+		NumMips = Texture2DArray->GetNumMips();
+	}
 	else if (Texture2DRT)
 	{
 		NumMips = Texture2DRT->GetNumMips();
@@ -1063,6 +1073,7 @@ TOptional<int32> FTextureEditorToolkit::GetMaxMipLevel( ) const
 {
 	const UTexture2D* Texture2D = Cast<UTexture2D>(Texture);
 	const UTextureCube* TextureCube = Cast<UTextureCube>(Texture);
+	const UTexture2DArray* Texture2DArray = Cast<UTexture2DArray>(Texture);
 	const UTextureRenderTargetCube* RTTextureCube = Cast<UTextureRenderTargetCube>(Texture);
 	const UTextureRenderTarget2D* RTTexture2D = Cast<UTextureRenderTarget2D>(Texture);
 	const UVolumeTexture* VolumeTexture = Cast<UVolumeTexture>(Texture);
@@ -1077,6 +1088,11 @@ TOptional<int32> FTextureEditorToolkit::GetMaxMipLevel( ) const
 		return TextureCube->GetNumMips() - 1;
 	}
 
+	if (Texture2DArray) 
+	{
+		return Texture2DArray->GetNumMips() - 1;
+	}
+
 	if (RTTextureCube)
 	{
 		return RTTextureCube->GetNumMips() - 1;
diff --git a/Engine/Source/Editor/UnrealEd/Classes/Factories/Texture2DArrayFactory.h b/Engine/Source/Editor/UnrealEd/Classes/Factories/Texture2DArrayFactory.h
new file mode 100644
--- /dev/null
+++ b/Engine/Source/Editor/UnrealEd/Classes/Factories/Texture2DArrayFactory.h
@@ -0,0 +1,29 @@
+// Copyright 1998-2019 Epic Games, Inc. All Rights Reserved.
+
+#pragma once
+
+#include "CoreMinimal.h"
+#include "UObject/ObjectMacros.h"
+#include "Factories/Factory.h"
+#include "Texture2DArrayFactory.generated.h"
+
+/** Factory for creating volume texture */
+UCLASS(hidecategories=Object, MinimalAPI)
+class UTexture2DArrayFactory : public UFactory
+{
+	GENERATED_UCLASS_BODY()
+
+	UPROPERTY()
+	TArray<UTexture2D*> InitialTextures;
+
+	//~ Begin UFactory Interface
+	virtual FText GetDisplayName() const override;
+	virtual bool ConfigureProperties() override;
+	virtual UObject* FactoryCreateNew(UClass* Class,UObject* InParent,FName Name,EObjectFlags Flags,UObject* Context,FFeedbackContext* Warn) override;
+	//~ Begin UFactory Interface
+
+	bool CheckArrayTexturesCompatibility();
+};
+
+
+
diff --git a/Engine/Source/Editor/UnrealEd/Classes/Factories/Texture2dArrayThumbnailRenderer.h b/Engine/Source/Editor/UnrealEd/Classes/Factories/Texture2dArrayThumbnailRenderer.h
new file mode 100644
--- /dev/null
+++ b/Engine/Source/Editor/UnrealEd/Classes/Factories/Texture2dArrayThumbnailRenderer.h
@@ -0,0 +1,24 @@
+// Copyright 1998-2019 Epic Games, Inc. All Rights Reserved.
+
+/**
+ *
+ * This thumbnail renderer displays the texture 2d array
+ */
+
+#pragma once
+
+#include "CoreMinimal.h"
+#include "UObject/ObjectMacros.h"
+#include "ThumbnailRendering/TextureThumbnailRenderer.h"
+#include "Texture2DArrayThumbnailRenderer.generated.h"
+
+UCLASS()
+class UTexture2DArrayThumbnailRenderer : public UTextureThumbnailRenderer
+{
+	GENERATED_UCLASS_BODY()
+
+	// Begin UThumbnailRenderer Interface
+	virtual void GetThumbnailSize(UObject* Object, float Zoom, uint32& OutWidth, uint32& OutHeight) const override;
+	virtual void Draw(UObject* Object, int32 X, int32 Y, uint32 Width, uint32 Height, FRenderTarget*, FCanvas* Canvas) override;
+	// End UThumbnailRenderer Interface
+};
diff --git a/Engine/Source/Editor/UnrealEd/Classes/Factories/TextureFactory.h b/Engine/Source/Editor/UnrealEd/Classes/Factories/TextureFactory.h
--- a/Engine/Source/Editor/UnrealEd/Classes/Factories/TextureFactory.h
+++ b/Engine/Source/Editor/UnrealEd/Classes/Factories/TextureFactory.h
@@ -30,6 +30,7 @@ struct FImportImage
 
 class UTexture2D;
 class UTextureCube;
+class UTexture2DArray;
 
 UCLASS(customconstructor, collapsecategories, hidecategories=Object)
 class UNREALED_API UTextureFactory : public UFactory, public IImportSettingsParser
@@ -143,7 +144,7 @@ public:
 	/** Create a texture given the appropriate input parameters	*/
 	virtual UTexture2D* CreateTexture2D( UObject* InParent, FName Name, EObjectFlags Flags );
 	virtual UTextureCube* CreateTextureCube( UObject* InParent, FName Name, EObjectFlags Flags );
-
+	virtual UTexture2DArray* CreateTexture2DArray(UObject* InParent, FName Name, EObjectFlags Flags);
 	/**
 	 * Suppresses the dialog box that, when importing over an existing texture, asks if the users wishes to overwrite its settings.
 	 * This is primarily for reimporting textures.
diff --git a/Engine/Source/Editor/UnrealEd/Private/EditorServer.cpp b/Engine/Source/Editor/UnrealEd/Private/EditorServer.cpp
--- a/Engine/Source/Editor/UnrealEd/Private/EditorServer.cpp
+++ b/Engine/Source/Editor/UnrealEd/Private/EditorServer.cpp
@@ -5399,6 +5399,7 @@ void ListMapPackageDependencies(const TCHAR* InStr)
 	{
 		FName CheckTexture2DName(TEXT("Texture2D"));
 		FName CheckCubeTextureName(TEXT("TextureCube"));
+		FName CheckTexture2DArrayName(TEXT("Texture2DArray"));
 		FName CheckLightmap2DName(TEXT("Lightmap2D"));
 		FName CheckShadowmap2DName(TEXT("Shadowmap2D"));
 		
diff --git a/Engine/Source/Editor/UnrealEd/Private/Factories/EditorFactories.cpp b/Engine/Source/Editor/UnrealEd/Private/Factories/EditorFactories.cpp
--- a/Engine/Source/Editor/UnrealEd/Private/Factories/EditorFactories.cpp
+++ b/Engine/Source/Editor/UnrealEd/Private/Factories/EditorFactories.cpp
@@ -160,6 +160,7 @@
 #include "Sound/SoundCue.h"
 #include "Sound/SoundMix.h"
 #include "Engine/TextureCube.h"
+#include "Engine/Texture2DArray.h"
 #include "Engine/VolumeTexture.h"
 #include "Engine/TextureRenderTarget.h"
 #include "Engine/TextureRenderTarget2D.h"
@@ -2955,6 +2956,12 @@ UTextureCube* UTextureFactory::CreateTextureCube( UObject* InParent, FName Name,
 	return NewObject ? CastChecked<UTextureCube>(NewObject) : nullptr;
 }
 
+UTexture2DArray* UTextureFactory::CreateTexture2DArray(UObject* InParent, FName Name, EObjectFlags Flags) 
+{
+	UObject* NewObject = CreateOrOverwriteAsset(UTexture2DArray::StaticClass(), InParent, Name, Flags);
+	return NewObject ? CastChecked<UTexture2DArray>(NewObject) : nullptr;
+}
+
 void UTextureFactory::SuppressImportOverwriteDialog()
 {
 	bSuppressImportOverwriteDialog = true;
@@ -3798,9 +3805,9 @@ UTexture* UTextureFactory::ImportTexture(UClass* Class, UObject* InParent, FName
 		if ( TextureCube )
 		{
 			TextureCube->Source.Init(
-				DDSLoadHelper.DDSHeader->dwWidth, 
-				DDSLoadHelper.DDSHeader->dwHeight,
-				/*NumSlices=*/ 6,
+				DDSLoadHelper.GetSizeX(),
+				DDSLoadHelper.GetSizeY(),
+				DDSLoadHelper.GetSliceCount(),
 				NumMips,
 				Format
 				);
@@ -3839,6 +3846,73 @@ UTexture* UTextureFactory::ImportTexture(UClass* Class, UObject* InParent, FName
 		return TextureCube;
 	}
 	
+	// DDS Texture array.
+	if (DDSLoadHelper.IsValidArrayTexture())
+	{
+		if (!IsImportResolutionValid(DDSLoadHelper.DDSHeader->dwWidth, DDSLoadHelper.DDSHeader->dwHeight, bAllowNonPowerOfTwo, Warn))
+		{
+			Warn->Logf(ELogVerbosity::Error, TEXT("DDS uses an unsupported format"));
+			return nullptr;
+		}
+
+		int32 NumMips = DDSLoadHelper.ComputeMipMapCount();
+		ETextureSourceFormat Format = DDSLoadHelper.ComputeSourceFormat();
+		if (Format == TSF_Invalid)
+		{
+			Warn->Logf(ELogVerbosity::Error, TEXT("DDS file contains data in an unsupported format."));
+			return nullptr;
+		}
+
+		// Create the array texture
+		UTexture2DArray* TextureArray = CreateTexture2DArray(InParent, Name, Flags);
+
+		if (TextureArray)
+		{
+			TextureArray->Source.Init(
+				DDSLoadHelper.GetSizeX(),
+				DDSLoadHelper.GetSizeY(),
+				DDSLoadHelper.GetSliceCount(),
+				NumMips,
+				Format
+			);
+			if (Format == TSF_RGBA16F)
+			{
+				TextureArray->CompressionSettings = TC_HDR;
+			}
+
+			uint8* DestMipData[MAX_TEXTURE_MIP_COUNT] = { 0 };
+			int32 MipSize[MAX_TEXTURE_MIP_COUNT] = { 0 };
+			for (int32 MipIndex = 0; MipIndex < NumMips; ++MipIndex)
+			{
+				DestMipData[MipIndex] = TextureArray->Source.LockMip(MipIndex);
+				MipSize[MipIndex] = TextureArray->Source.CalcMipSize(MipIndex) / DDSLoadHelper.GetSliceCount();
+			}
+
+			for (uint32 SliceIndex = 0; SliceIndex < DDSLoadHelper.GetSliceCount(); ++SliceIndex)
+			{
+				const uint8* SrcMipData = DDSLoadHelper.GetDDSDataPointer((ECubeFace)SliceIndex);
+				for (int32 MipIndex = 0; MipIndex < NumMips; ++MipIndex)
+				{
+					FMemory::Memcpy(DestMipData[MipIndex] + MipSize[MipIndex] * SliceIndex, SrcMipData, MipSize[MipIndex]);
+					SrcMipData += MipSize[MipIndex];
+				}
+			}
+
+			for (int32 MipIndex = 0; MipIndex < NumMips; ++MipIndex)
+			{
+				TextureArray->Source.UnlockMip(MipIndex);
+			}
+
+			if (NumMips > 1)
+			{
+				// If the source has mips we keep the mips by default, unless the user changes that
+				MipGenSettings = TMGS_LeaveExistingMips;
+			}
+		}
+
+		return TextureArray;
+	}
+
 	//
 	// HDR File
 	//
@@ -4329,6 +4403,28 @@ UObject* UTextureFactory::FactoryCreateBinary
 		}
 	}
 
+	// Invalidate any Texture2DArrays that use the updated texture.
+	if (Texture2D) 
+	{
+		for (TObjectIterator<UTexture2DArray> It; It; ++It) 
+		{
+			UTexture2DArray* TextureArray = *It;
+			if (TextureArray) 
+			{
+				for (int32 SourceIndex = 0; SourceIndex < TextureArray->SourceTextures.Num(); ++SourceIndex) 
+				{
+					if (TextureArray->SourceTextures[SourceIndex] == Texture2D) 
+					{
+						// Update the entire texture array.
+						TextureArray->UpdateSourceFromSourceTextures(false);
+						break;
+					}
+				}
+			}
+		}
+	}
+
+
 	// If we are automatically creating a material for this texture...
 	if( bCreateMaterial )
 	{
diff --git a/Engine/Source/Editor/UnrealEd/Private/Factories/Texture2DArrayFactory.cpp b/Engine/Source/Editor/UnrealEd/Private/Factories/Texture2DArrayFactory.cpp
new file mode 100644
--- /dev/null
+++ b/Engine/Source/Editor/UnrealEd/Private/Factories/Texture2DArrayFactory.cpp
@@ -0,0 +1,96 @@
+// Copyright 1998-2019 Epic Games, Inc. All Rights Reserved.
+
+#include "Factories/Texture2DArrayFactory.h"
+#include "Styling/SlateBrush.h"
+#include "Brushes/SlateDynamicImageBrush.h"
+#include "Engine/Texture2DArray.h"
+
+
+#define LOCTEXT_NAMESPACE "Texture2DArrayFactory"
+
+UTexture2DArrayFactory::UTexture2DArrayFactory(const FObjectInitializer& ObjectInitializer)
+	: Super(ObjectInitializer)
+{
+	bCreateNew = true;
+	bEditAfterNew = true;
+	SupportedClass = UTexture2DArray::StaticClass();
+}
+
+FText UTexture2DArrayFactory::GetDisplayName() const
+{
+	return LOCTEXT("Texture2DArrayFactoryDescription", "Texture 2D Array");
+}
+
+bool UTexture2DArrayFactory::ConfigureProperties()
+{
+	return true;
+}
+
+UObject* UTexture2DArrayFactory::FactoryCreateNew(UClass* Class, UObject* InParent, FName Name, EObjectFlags Flags, UObject* Context, FFeedbackContext* Warn)
+{
+	static const auto AllowTextureArrayAssetCreationVar = IConsoleManager::Get().FindTConsoleVariableDataInt(TEXT("r.AllowTexture2DArrayCreation"));
+
+	if (AllowTextureArrayAssetCreationVar->GetValueOnGameThread() == 0)
+	{
+		UE_LOG(LogTexture, Warning, TEXT("Texture2DArray creation failed. Enable by using console command r.AllowTexture2DArrayCreation\n"));
+		return nullptr;
+	}
+
+	// Make sure the size of all textures is the same.
+	if (!CheckArrayTexturesCompatibility()) 
+	{
+		return nullptr;
+	}
+
+	UTexture2DArray* NewTexture2DArray = CastChecked<UTexture2DArray>(CreateOrOverwriteAsset(UTexture2DArray::StaticClass(), InParent, Name, Flags));
+
+	if (NewTexture2DArray == 0) 
+	{
+		UE_LOG(LogTexture, Warning, TEXT("Texture2DArray creation failed.\n"));
+	}
+	else if (InitialTextures.Num() > 0)
+	{
+		// Upload the compatible textures to the texture array.
+		for (int32 TextureIndex = 0; TextureIndex < InitialTextures.Num(); ++TextureIndex) 
+		{
+			NewTexture2DArray->SourceTextures.Add(InitialTextures[TextureIndex]);
+		}
+
+		// Create the texture array resource and corresponding rhi resource.
+		NewTexture2DArray->UpdateSourceFromSourceTextures();
+	}
+	
+	return NewTexture2DArray;
+}
+
+bool UTexture2DArrayFactory::CheckArrayTexturesCompatibility()
+{
+	bool bError = false;
+	for (int32 TextureIndex = 0; TextureIndex < InitialTextures.Num(); ++TextureIndex) 
+	{
+		FTextureSource& TextureSource = InitialTextures[TextureIndex]->Source;
+		const int32 FormatDataSize = TextureSource.GetBytesPerPixel();
+		const int32 SizeX = TextureSource.GetSizeX();
+		const int32 SizeY = TextureSource.GetSizeY();
+
+		for (int32 TextureCmpIndex = TextureIndex + 1; TextureCmpIndex < InitialTextures.Num(); ++TextureCmpIndex)
+		{
+			FTextureSource& TextureSourceCmp = InitialTextures[TextureCmpIndex]->Source;
+			if (TextureSourceCmp.GetSizeX() != SizeX || TextureSourceCmp.GetSizeY() != SizeY)
+			{
+				UE_LOG(LogTexture, Warning, TEXT("Texture2DArray creation failed. Textures have different sizes."));
+				bError = true;
+			}
+
+			if (TextureSourceCmp.GetBytesPerPixel() != FormatDataSize)
+			{
+				UE_LOG(LogTexture, Warning, TEXT("Texture2DArray creation failed. Textures have different pixel formats."));
+				bError = true;
+			}
+		}
+	}
+
+	return (!bError);
+}
+
+#undef LOCTEXT_NAMESPACE
diff --git a/Engine/Source/Editor/UnrealEd/Private/MaterialGraphSchema.cpp b/Engine/Source/Editor/UnrealEd/Private/MaterialGraphSchema.cpp
--- a/Engine/Source/Editor/UnrealEd/Private/MaterialGraphSchema.cpp
+++ b/Engine/Source/Editor/UnrealEd/Private/MaterialGraphSchema.cpp
@@ -92,6 +92,9 @@ void FMaterialGraphSchemaAction_NewNode::SetFunctionInputType(UMaterialExpressio
 	case MCT_TextureCube:
 		FunctionInput->InputType = FunctionInput_TextureCube;
 		break;
+	case MCT_Texture2DArray:
+		FunctionInput->InputType = FunctionInput_Texture2DArray;
+		break;
 	case MCT_TextureExternal:
 		FunctionInput->InputType = FunctionInput_TextureExternal;
 		break;
diff --git a/Engine/Source/Editor/UnrealEd/Private/Texture2DPreview.cpp b/Engine/Source/Editor/UnrealEd/Private/Texture2DPreview.cpp
--- a/Engine/Source/Editor/UnrealEd/Private/Texture2DPreview.cpp
+++ b/Engine/Source/Editor/UnrealEd/Private/Texture2DPreview.cpp
@@ -13,6 +13,7 @@
 #include "TextureResource.h"
 #include "RenderCore.h"
 #include "VirtualTexturing.h"
+#include "Engine/Texture2DArray.h"
 /*------------------------------------------------------------------------------
 	Batched element shaders for previewing 2d textures.
 ------------------------------------------------------------------------------*/
@@ -22,6 +23,7 @@
 namespace
 {
 	class FTexture2DPreviewVirtualTexture : SHADER_PERMUTATION_BOOL("SAMPLE_VIRTUAL_TEXTURE");
+	class FTexture2DPreviewTexture2DArray : SHADER_PERMUTATION_BOOL("TEXTURE_ARRAY");
 }
 
 
@@ -29,7 +31,7 @@ class FSimpleElementTexture2DPreviewPS : public FGlobalShader
 {
 	DECLARE_GLOBAL_SHADER(FSimpleElementTexture2DPreviewPS);
 
-	using FPermutationDomain = TShaderPermutationDomain<FTexture2DPreviewVirtualTexture>;
+	using FPermutationDomain = TShaderPermutationDomain<FTexture2DPreviewVirtualTexture, FTexture2DPreviewTexture2DArray>;
 
 public:
 
@@ -46,6 +48,7 @@ public:
 		TextureComponentReplicateAlpha.Bind(Initializer.ParameterMap,TEXT("TextureComponentReplicateAlpha"));
 		ColorWeights.Bind(Initializer.ParameterMap,TEXT("ColorWeights"));
 		PackedParameters.Bind(Initializer.ParameterMap,TEXT("PackedParams"));
+		NumSlices.Bind(Initializer.ParameterMap, TEXT("NumSlices"));
 	}
 	FSimpleElementTexture2DPreviewPS() {}
 
@@ -53,14 +56,14 @@ public:
 	static bool ShouldCompilePermutation(const FGlobalShaderPermutationParameters& Parameters)
 	{
 		FPermutationDomain PermutationVector(Parameters.PermutationId);
-		if (IsFeatureLevelSupported(Parameters.Platform, ERHIFeatureLevel::SM5) == false && PermutationVector.Get<FTexture2DPreviewVirtualTexture>())
+		if (IsFeatureLevelSupported(Parameters.Platform, ERHIFeatureLevel::SM5) == false)
 		{
 			return false;
 		}
 		return IsFeatureLevelSupported(Parameters.Platform, ERHIFeatureLevel::SM5) && !IsConsolePlatform(Parameters.Platform);
 	}
 
-	void SetParameters(FRHICommandList& RHICmdList, const FTexture* TextureValue, const FMatrix& ColorWeightsValue, float GammaValue, float MipLevel, float LayerIndex, bool bIsNormalMap, bool bIsVirtualTexture)
+	void SetParameters(FRHICommandList& RHICmdList, const FTexture* TextureValue, const FMatrix& ColorWeightsValue, float GammaValue, float MipLevel, float LayerIndex, bool bIsNormalMap, bool bIsVirtualTexture, bool bIsTextureArray)
 	{
 		if (bIsVirtualTexture)
 		{
@@ -101,6 +104,14 @@ public:
 		FVector4 PackedParametersValue(GammaValue, MipLevel, bIsNormalMap ? 1.0 : -1.0f, LayerIndex);
 		SetShaderValue(RHICmdList, GetPixelShader(), PackedParameters, PackedParametersValue);
 
+		// Store slice count for texture array
+		if (bIsTextureArray)
+		{
+			const FTexture2DArrayResource* TextureValue2DArray = (FTexture2DArrayResource*)(TextureValue);
+			float NumSlicesData = TextureValue2DArray ? float(TextureValue2DArray->GetNumSlices()) : 1;
+			SetShaderValue(RHICmdList, GetPixelShader(), NumSlices, NumSlicesData);
+		}
+
 		SetShaderValue(RHICmdList, GetPixelShader(),TextureComponentReplicate,TextureValue->bGreyScaleFormat ? FLinearColor(1,0,0,0) : FLinearColor(0,0,0,0));
 		SetShaderValue(RHICmdList, GetPixelShader(),TextureComponentReplicateAlpha,TextureValue->bGreyScaleFormat ? FLinearColor(1,0,0,0) : FLinearColor(0,0,0,1));
 	}
@@ -117,6 +128,7 @@ public:
 		Ar << TextureComponentReplicate;
 		Ar << TextureComponentReplicateAlpha;
 		Ar << ColorWeights;
+		Ar << NumSlices;
 		Ar << PackedParameters;
 		return bShaderHasOutdatedParameters;
 	}
@@ -132,6 +144,7 @@ private:
 	FShaderParameter TextureComponentReplicateAlpha;
 	FShaderParameter ColorWeights; 
 	FShaderParameter PackedParameters;
+	FShaderParameter NumSlices;
 };
 
 IMPLEMENT_GLOBAL_SHADER(FSimpleElementTexture2DPreviewPS, "/Engine/Private/SimpleElementTexture2DPreviewPixelShader.usf", "Main", SF_Pixel);
@@ -149,7 +162,8 @@ void FBatchedElementTexture2DPreviewParameters::BindShaders(
 	TShaderMapRef<FSimpleElementVS> VertexShader(GetGlobalShaderMap(InFeatureLevel));
 
 	FSimpleElementTexture2DPreviewPS::FPermutationDomain PermutationVector;
-	PermutationVector.Set<FTexture2DPreviewVirtualTexture>(bIsVirtualTexture);	
+	PermutationVector.Set<FTexture2DPreviewVirtualTexture>(bIsVirtualTexture);
+	PermutationVector.Set<FTexture2DPreviewTexture2DArray>(bIsTextureArray);
 	TShaderMapRef<FSimpleElementTexture2DPreviewPS> PixelShader(GetGlobalShaderMap(InFeatureLevel), PermutationVector);
 
 	GraphicsPSOInit.BoundShaderState.VertexDeclarationRHI = GSimpleElementVertexDeclaration.VertexDeclarationRHI;
@@ -165,5 +179,5 @@ void FBatchedElementTexture2DPreviewParameters::BindShaders(
 	SetGraphicsPipelineState(RHICmdList, GraphicsPSOInit, EApplyRendertargetOption::ForceApply);
 
 	VertexShader->SetParameters(RHICmdList, InTransform);
-	PixelShader->SetParameters(RHICmdList, Texture, ColorWeights, InGamma, MipLevel, LayerIndex, bIsNormalMap, bIsVirtualTexture);
+	PixelShader->SetParameters(RHICmdList, Texture, ColorWeights, InGamma, MipLevel, LayerIndex, bIsNormalMap, bIsVirtualTexture, bIsTextureArray);
 }
diff --git a/Engine/Source/Editor/UnrealEd/Private/ThumbnailManager.cpp b/Engine/Source/Editor/UnrealEd/Private/ThumbnailManager.cpp
--- a/Engine/Source/Editor/UnrealEd/Private/ThumbnailManager.cpp
+++ b/Engine/Source/Editor/UnrealEd/Private/ThumbnailManager.cpp
@@ -13,7 +13,7 @@
 #include "Engine/StaticMesh.h"
 #include "UnrealClient.h"
 #include "Engine/TextureCube.h"
-
+#include "Engine/Texture2DArray.h"
 #include "ImageUtils.h"
 
 
@@ -38,6 +38,7 @@ UThumbnailManager::UThumbnailManager(const FObjectInitializer& ObjectInitializer
 			ConstructorHelpers::FObjectFinder<UStaticMesh> EditorSkySphereMesh;
 			ConstructorHelpers::FObjectFinder<UMaterial> FloorPlaneMaterial;
 			ConstructorHelpers::FObjectFinder<UTextureCube> DaylightAmbientCubemap;
+			
 			FConstructorStatics()
 				: EditorCubeMesh(TEXT("/Engine/EditorMeshes/EditorCube"))
 				, EditorSphereMesh(TEXT("/Engine/EditorMeshes/EditorSphere"))
diff --git a/Engine/Source/Editor/UnrealEd/Private/ThumbnailRendering/Texture2DArrayThumbnailRenderer.cpp b/Engine/Source/Editor/UnrealEd/Private/ThumbnailRendering/Texture2DArrayThumbnailRenderer.cpp
new file mode 100644
--- /dev/null
+++ b/Engine/Source/Editor/UnrealEd/Private/ThumbnailRendering/Texture2DArrayThumbnailRenderer.cpp
@@ -0,0 +1,33 @@
+// Copyright 1998-2019 Epic Games, Inc. All Rights Reserved.
+
+#include "Factories/Texture2DArrayThumbnailRenderer.h"
+#include "Engine/Texture2DArray.h"
+#include "ThumbnailRendering/ThumbnailManager.h"
+
+UTexture2DArrayThumbnailRenderer::UTexture2DArrayThumbnailRenderer(const FObjectInitializer& ObjectInitializer)
+	: Super(ObjectInitializer)
+{
+}
+
+void UTexture2DArrayThumbnailRenderer::GetThumbnailSize(UObject* Object, float Zoom, uint32& OutWidth, uint32& OutHeight) const
+{
+	OutWidth = 0;
+	OutHeight = 0;
+
+	UTexture2DArray* TextureArray = Cast<UTexture2DArray>(Object);
+	if (TextureArray != nullptr)
+	{
+		// Let the base class get the size of a face
+		Super::GetThumbnailSize(TextureArray, Zoom, OutWidth, OutHeight);
+	}
+
+}
+
+void UTexture2DArrayThumbnailRenderer::Draw(UObject* Object, int32 X, int32 Y, uint32 Width, uint32 Height, FRenderTarget*, FCanvas* Canvas)
+{
+	UTexture2DArray* TextureArray = Cast<UTexture2DArray>(Object);
+	if (TextureArray != nullptr)
+	{
+		Super::Draw(TextureArray, X, Y, Width, Height, nullptr, Canvas);
+	}
+}
diff --git a/Engine/Source/Editor/UnrealEd/Private/ThumbnailRendering/TextureThumbnailRenderer.cpp b/Engine/Source/Editor/UnrealEd/Private/ThumbnailRendering/TextureThumbnailRenderer.cpp
--- a/Engine/Source/Editor/UnrealEd/Private/ThumbnailRendering/TextureThumbnailRenderer.cpp
+++ b/Engine/Source/Editor/UnrealEd/Private/ThumbnailRendering/TextureThumbnailRenderer.cpp
@@ -9,6 +9,8 @@
 #include "ThumbnailRendering/ThumbnailManager.h"
 #include "EngineGlobals.h"
 #include "Engine/TextureCube.h"
+#include "Engine/Texture2DArray.h"
+#include "Texture2DPreview.h"
 #include "Engine/TextureRenderTargetCube.h"
 
 #include "CubemapUnwrapUtils.h"
@@ -58,6 +60,7 @@ void UTextureThumbnailRenderer::Draw(UObject* Object, int32 X, int32 Y, uint32 W
 		const bool bUseTranslucentBlend = Texture2D && Texture2D->HasAlphaChannel() && ((Texture2D->LODGroup == TEXTUREGROUP_UI) || (Texture2D->LODGroup == TEXTUREGROUP_Pixels2D));
 
 		UTextureCube* TextureCube = Cast<UTextureCube>(Texture);
+		UTexture2DArray* Texture2DArray = Cast<UTexture2DArray>(Texture);
 		UTextureRenderTargetCube* RTTextureCube = Cast<UTextureRenderTargetCube>(Texture);
 		UTextureLightProfile* TextureLightProfile = Cast<UTextureLightProfile>(Texture);
 
@@ -75,6 +78,12 @@ void UTextureThumbnailRenderer::Draw(UObject* Object, int32 X, int32 Y, uint32 W
 				Y += Height / 2;
 			}
 		}
+		else if (Texture2DArray) 
+		{
+			bool bIsNormalMap = Texture2DArray->IsNormalMap();
+			bool bIsSingleChannel = true;
+			BatchedElementParameters = new FBatchedElementTexture2DPreviewParameters(0, 0, bIsNormalMap, bIsSingleChannel, false, true);
+		}
 		else if (TextureLightProfile)
 		{
 			BatchedElementParameters = new FIESLightProfileBatchedElementParameters(TextureLightProfile->Brightness);
diff --git a/Engine/Source/Editor/UnrealEd/Public/Texture2DPreview.h b/Engine/Source/Editor/UnrealEd/Public/Texture2DPreview.h
--- a/Engine/Source/Editor/UnrealEd/Public/Texture2DPreview.h
+++ b/Engine/Source/Editor/UnrealEd/Public/Texture2DPreview.h
@@ -15,12 +15,13 @@
 class UNREALED_API FBatchedElementTexture2DPreviewParameters : public FBatchedElementParameters
 {
 public:
-	FBatchedElementTexture2DPreviewParameters(float InMipLevel, float InLayerIndex, bool bInIsNormalMap, bool bInIsSingleChannel, bool bInIsVirtualTexture)
+	FBatchedElementTexture2DPreviewParameters(float InMipLevel, float InLayerIndex, bool bInIsNormalMap, bool bInIsSingleChannel, bool bInIsVirtualTexture, bool bInIsTextureArray = false)
 		: MipLevel(InMipLevel)
 		, LayerIndex(InLayerIndex)
 		, bIsNormalMap( bInIsNormalMap )
 		, bIsSingleChannelFormat( bInIsSingleChannel )
 		, bIsVirtualTexture(bInIsVirtualTexture)
+		, bIsTextureArray(bInIsTextureArray)
 	{
 	}
 
@@ -37,4 +38,5 @@ private:
 
 	/** Parameters that are used to select a shader permutation */
 	bool bIsVirtualTexture;
+	bool bIsTextureArray;
 };
diff --git a/Engine/Source/Editor/UnrealEd/Public/UnrealEd.h b/Engine/Source/Editor/UnrealEd/Public/UnrealEd.h
--- a/Engine/Source/Editor/UnrealEd/Public/UnrealEd.h
+++ b/Engine/Source/Editor/UnrealEd/Public/UnrealEd.h
@@ -274,6 +274,7 @@ MONOLITHIC_HEADER_BOILERPLATE()
 #include "ThumbnailRendering/ParticleSystemThumbnailRenderer.h"
 #include "ThumbnailRendering/SubsurfaceProfileRenderer.h"
 #include "Factories/TextureCubeThumbnailRenderer.h"
+#include "Factories/Texture2dArrayThumbnailRenderer.h"
 #include "Editor/TransBuffer.h"
 #include "Preferences/UnrealEdKeyBindings.h"
 #include "Preferences/UnrealEdOptions.h"
diff --git a/Engine/Source/Runtime/Engine/Classes/Engine/Texture.h b/Engine/Source/Runtime/Engine/Classes/Engine/Texture.h
--- a/Engine/Source/Runtime/Engine/Classes/Engine/Texture.h
+++ b/Engine/Source/Runtime/Engine/Classes/Engine/Texture.h
@@ -183,7 +183,7 @@ struct FTextureSource
 #if WITH_EDITOR
 
 	ENGINE_API static int32 GetBytesPerPixel(ETextureSourceFormat Format);
-	FORCEINLINE static bool IsHDR(ETextureSourceFormat Format) { return Format == TSF_BGRE8 || Format == TSF_RGBA16F; }
+	FORCEINLINE static bool IsHDR(ETextureSourceFormat Format) { return (Format == TSF_BGRE8 || Format == TSF_RGBA16F); }
 
 	ENGINE_API void InitBlocked(const ETextureSourceFormat* InLayerFormats,
 		const FTextureSourceBlock* InBlocks,
@@ -325,6 +325,7 @@ private:
 	friend class UTexture2D;
 	friend class UTextureCube;
 	friend class UVolumeTexture;
+	friend class UTexture2DArray;
 
 	/** The bulk source data. */
 	FByteBulkData BulkData;
diff --git a/Engine/Source/Runtime/Engine/Classes/Engine/Texture2DArray.cpp b/Engine/Source/Runtime/Engine/Classes/Engine/Texture2DArray.cpp
new file mode 100644
--- /dev/null
+++ b/Engine/Source/Runtime/Engine/Classes/Engine/Texture2DArray.cpp
@@ -0,0 +1,393 @@
+// Copyright 1998-2019 Epic Games, Inc. All Rights Reserved.
+
+/*=============================================================================
+	Texture2DArray.cpp: UTexture2DArray implementation.
+=============================================================================*/
+
+#include "Engine/Texture2DArray.h"
+#include "RenderUtils.h"
+#include "TextureResource.h"
+#include "EngineUtils.h"
+#include "DeviceProfiles/DeviceProfile.h"
+#include "DeviceProfiles/DeviceProfileManager.h"
+#include "Containers/ResourceArray.h"
+
+static TAutoConsoleVariable<int32> CVarAllowTexture2DArrayAssetCreation(
+	TEXT("r.AllowTexture2DArrayCreation"),
+	0,
+	TEXT("Enable UTexture2DArray assets"),
+	ECVF_Default
+);
+
+UTexture2DArray::UTexture2DArray(const FObjectInitializer& ObjectInitializer) : Super(ObjectInitializer)
+{
+#if WITH_EDITORONLY_DATA
+	SRGB = true;
+	MipGenSettings = TMGS_NoMipmaps;
+#endif
+}
+
+FTextureResource* UTexture2DArray::CreateResource()
+{
+	if (GetNumMips() > 0) 
+	{
+		return new FTexture2DArrayResource(this);
+	}
+	
+	return nullptr;
+}
+
+void UTexture2DArray::UpdateResource()
+{
+#if WITH_EDITOR
+	// Re-cache platform data if the source has changed.
+	CachePlatformData();
+#endif // #if WITH_EDITOR
+
+	Super::UpdateResource();
+}
+
+uint32 UTexture2DArray::CalcTextureMemorySize(int32 MipCount) const
+{
+	uint32 Size = 0;
+
+	if (PlatformData && GetNumMips())
+	{
+		uint32 TextureAlign = 0;
+		uint64 TextureSize = RHICalcTexture3DPlatformSize(GetSizeX(), GetSizeY(), GetNumSlices(), GetPixelFormat(), MipCount, 0, TextureAlign);
+		Size = TextureSize;
+	}
+
+	return Size;
+}
+
+uint32 UTexture2DArray::CalcTextureMemorySizeEnum(ETextureMipCount Enum) const
+{
+	if (Enum == TMC_ResidentMips || Enum == TMC_AllMipsBiased) 
+	{
+		return CalcTextureMemorySize(GetNumMips() - GetCachedLODBias());
+	}
+
+	return CalcTextureMemorySize(GetNumMips());
+}
+
+ENGINE_API int32 UTexture2DArray::CalculateMipZSize(int32 Mip)
+{
+	int32 Size = 0;
+	if (PlatformData) 
+	{
+		// Make sure mip level exists.
+		if (PlatformData->Mips.Num() > Mip) 
+		{
+			int32 NumBlocksX = PlatformData->Mips[Mip].SizeX / GPixelFormats[PlatformData->PixelFormat].BlockSizeX;
+			int32 NumBlocksY = PlatformData->Mips[Mip].SizeY / GPixelFormats[PlatformData->PixelFormat].BlockSizeY;
+			
+			if (NumBlocksX == 0)
+			{
+				NumBlocksX = 1;
+			}
+
+			if (NumBlocksY == 0)
+			{
+				NumBlocksY = 1;
+			}
+
+			return NumBlocksX * NumBlocksY * GPixelFormats[PlatformData->PixelFormat].BlockBytes;
+		}
+	}
+	return Size;
+}
+
+#if WITH_EDITOR
+
+ENGINE_API bool UTexture2DArray::CheckArrayTexturesCompatibility()
+{
+	bool bError = false;
+
+	for (int32 TextureIndex = 0; TextureIndex < SourceTextures.Num(); ++TextureIndex)
+	{
+		// Do not create array till all texture slots are filled.
+		if (!SourceTextures[TextureIndex])
+		{
+			return false;
+		}
+
+		FTextureSource& TextureSource = SourceTextures[TextureIndex]->Source;
+		// const int32 FormatDataSize = TextureSource.GetBytesPerPixel();
+		const EPixelFormat PixelFormat = SourceTextures[TextureIndex]->GetPixelFormat();
+		const int32 SizeX = TextureSource.GetSizeX();
+		const int32 SizeY = TextureSource.GetSizeY();
+
+		for (int32 TextureCmpIndex = TextureIndex + 1; TextureCmpIndex < SourceTextures.Num(); ++TextureCmpIndex)
+		{
+			// Do not create array till all texture slots are filled.
+			if (!SourceTextures[TextureCmpIndex]) 
+			{
+				return false;
+			}
+
+			FTextureSource& TextureSourceCmp = SourceTextures[TextureCmpIndex]->Source;
+			FString TextureName = SourceTextures[TextureIndex]->GetFName().ToString();
+			FString TextureNameCmp = SourceTextures[TextureCmpIndex]->GetFName().ToString();
+			const EPixelFormat PixelFormatCmp = SourceTextures[TextureCmpIndex]->GetPixelFormat();
+
+			if (TextureSourceCmp.GetSizeX() != SizeX || TextureSourceCmp.GetSizeY() != SizeY)
+			{
+				UE_LOG(LogTexture, Warning, TEXT("Texture2DArray creation failed. Textures %s and %s have different sizes."), *TextureName, *TextureNameCmp);
+				bError = true;
+			}
+
+			if (PixelFormatCmp != PixelFormat)
+			{
+				UE_LOG(LogTexture, Warning, TEXT("Texture2DArray creation failed. Textures %s and %s have incompatible pixel formats."), *TextureName, *TextureNameCmp);
+				bError = true;
+			}
+
+			//if (TextureSourceCmp.GetBytesPerPixel() != FormatDataSize)
+			//{
+			//	UE_LOG(LogTexture, Warning, TEXT("Texture2DArray creation failed. Textures %s and %s have incompatible pixel formats."), *TextureName, *TextureNameCmp);
+			//	bError = true;
+			//}
+		}
+	}
+
+	return (!bError);
+}
+
+
+ENGINE_API bool UTexture2DArray::UpdateSourceFromSourceTextures(bool bCreatingNewTexture)
+{
+	if (!CheckArrayTexturesCompatibility()) 
+	{
+		return false;
+	}
+
+	if (SourceTextures.Num() > 0)
+	{
+		FTextureSource& InitialSource = SourceTextures[0]->Source;
+		// Format and format size.
+		EPixelFormat PixelFormat = SourceTextures[0]->GetPixelFormat();
+		ETextureSourceFormat Format = InitialSource.GetFormat();
+		int32 FormatDataSize = InitialSource.GetBytesPerPixel();
+		// X,Y,Z size of the array.
+		int32 SizeX = SourceTextures[0]->GetSizeX();
+		int32 SizeY = SourceTextures[0]->GetSizeY();
+		uint32 NumSlices = SourceTextures.Num();
+		// Only copy the first mip from the source textures to array texture.
+		uint32 NumMips = 1;
+
+		// This should be false when texture is updated to avoid overriding user settings.
+		if (bCreatingNewTexture)
+		{
+			CompressionSettings = SourceTextures[0]->CompressionSettings;
+			MipGenSettings = TMGS_NoMipmaps;
+			PowerOfTwoMode = ETexturePowerOfTwoSetting::None;
+			LODGroup = SourceTextures[0]->LODGroup;
+			SRGB = SourceTextures[0]->SRGB;
+			NeverStream = true;
+		}
+
+		// Create the source texture for this UTexture.
+		Source.Init(SizeX, SizeY, NumSlices, NumMips, Format);
+
+		// We only copy the top level Mip map.
+		uint8* DestMipData[MAX_TEXTURE_MIP_COUNT] = { 0 };
+		int32 MipSizeBytes[MAX_TEXTURE_MIP_COUNT] = { 0 };
+			
+		for (uint32 MipIndex = 0; MipIndex < NumMips; ++MipIndex)
+		{
+			DestMipData[MipIndex] =  Source.LockMip(MipIndex);
+			MipSizeBytes[MipIndex] = Source.CalcMipSize(MipIndex) / NumSlices;
+		}
+
+		for (int32 SourceTexIndex = 0; SourceTexIndex < SourceTextures.Num(); ++SourceTexIndex)
+		{
+			for (uint32 MipIndex = 0; MipIndex < NumMips; ++MipIndex)
+			{
+				TArray64<uint8> Ref2DData;
+				SourceTextures[SourceTexIndex]->Source.GetMipData(Ref2DData, MipIndex);
+				void* Dst = DestMipData[MipIndex] + MipSizeBytes[MipIndex] * SourceTexIndex;
+				FMemory::Memcpy(Dst, Ref2DData.GetData(), MipSizeBytes[MipIndex]);
+			}
+		}
+
+		for (uint32 MipIndex = 0; MipIndex < NumMips; ++MipIndex)
+		{
+			Source.UnlockMip(MipIndex);
+		}
+
+		UpdateMipGenSettings();
+		SetLightingGuid();
+		UpdateResource();
+	}
+
+	return true;
+}
+
+ENGINE_API void UTexture2DArray::InvadiateTextureSource()
+{
+	if (PlatformData) 
+	{
+		delete PlatformData;
+		PlatformData = NULL;
+	}
+
+	Source.Init(0, 0, 0, 0, TSF_Invalid, nullptr);
+	UpdateResource();
+}
+#endif
+
+void UTexture2DArray::Serialize(FArchive& Ar)
+{
+	DECLARE_SCOPE_CYCLE_COUNTER(TEXT("UTexture2DArray::Serialize"), STAT_Texture2DArray_Serialize, STATGROUP_LoadTime);
+	
+	Super::Serialize(Ar);
+
+	FStripDataFlags StripFlags(Ar);
+	bool bCooked = Ar.IsCooking();
+	Ar << bCooked;
+
+	if (bCooked || Ar.IsCooking())
+	{
+		SerializeCookedPlatformData(Ar);
+	}
+
+#if WITH_EDITOR
+	if (Ar.IsLoading() && !Ar.IsTransacting() && !bCooked)
+	{
+		BeginCachePlatformData();
+	}
+#endif
+}
+
+void UTexture2DArray::PostLoad() 
+{
+#if WITH_EDITOR
+	FinishCachePlatformData();
+
+#endif // #if WITH_EDITOR
+	Super::PostLoad();
+};
+
+void UTexture2DArray::GetAssetRegistryTags(TArray<FAssetRegistryTag>& OutTags) const
+{
+#if WITH_EDITOR
+	int32 SizeX = Source.GetSizeX();
+	int32 SizeY = Source.GetSizeY();
+	int32 SizeZ = Source.GetNumSlices(); //GetSizeZ()
+#else
+	int32 SizeX = 0;
+	int32 SizeY = 0;
+	int32 SizeZ = 0;
+#endif
+	const FString Dimensions = FString::Printf(TEXT("%dx%dx%d"), SizeX, SizeY, SizeZ);
+	OutTags.Add(FAssetRegistryTag("Dimensions", Dimensions, FAssetRegistryTag::TT_Dimensional));
+	OutTags.Add(FAssetRegistryTag("Format", GPixelFormats[GetPixelFormat()].Name, FAssetRegistryTag::TT_Alphabetical));
+
+	Super::GetAssetRegistryTags(OutTags);
+}
+
+FString UTexture2DArray::GetDesc() 
+{
+	return FString::Printf(TEXT("Array: %dx%dx%d [%s]"),
+		GetSizeX(),
+		GetSizeY(),
+		GetNumSlices(),
+		GPixelFormats[GetPixelFormat()].Name
+	);
+}
+
+void UTexture2DArray::GetResourceSizeEx(FResourceSizeEx& CumulativeResourceSize) 
+{
+	Super::GetResourceSizeEx(CumulativeResourceSize);
+	CumulativeResourceSize.AddUnknownMemoryBytes(CalcTextureMemorySizeEnum(TMC_ResidentMips));
+}
+
+bool UTexture2DArray::ShaderPlatformSupportsCompression(EShaderPlatform ShaderPlatform)
+{
+	switch (ShaderPlatform)
+	{
+	case SP_PCD3D_SM4:
+	case SP_PCD3D_SM5:
+	case SP_PS4:
+	case SP_XBOXONE_D3D12:
+	case SP_VULKAN_SM5:
+	case SP_VULKAN_SM4:
+	case SP_VULKAN_SM5_LUMIN:
+		return true;
+
+	default:
+		return false;
+	}
+}
+
+#if WITH_EDITOR
+uint32 UTexture2DArray::GetMaximumDimension() const
+{
+	return GetMax2DTextureDimension();
+
+}
+
+void UTexture2DArray::UpdateMipGenSettings()
+{
+	if (PowerOfTwoMode == ETexturePowerOfTwoSetting::None && !Source.IsPowerOfTwo())
+	{
+		// Force NPT textures to have no mip maps.
+		MipGenSettings = TMGS_NoMipmaps;
+		NeverStream = true;
+	}
+}
+
+ENGINE_API void UTexture2DArray::PostEditChangeProperty(FPropertyChangedEvent & PropertyChangedEvent)
+{	
+	const FName PropertyName = PropertyChangedEvent.GetPropertyName();
+
+	if (PowerOfTwoMode == ETexturePowerOfTwoSetting::None && (!Source.IsPowerOfTwo()))
+	{
+		// Force NPT textures to have no mip maps.
+		if (PropertyName == GET_MEMBER_NAME_CHECKED(UTexture2DArray, MipGenSettings)) 
+		{
+			UE_LOG(LogTexture, Warning, TEXT("Cannot use mip maps for non-power of two textures."));
+		}
+
+		MipGenSettings = TMGS_NoMipmaps;
+		NeverStream = true;
+	}
+
+	if (PropertyName == GET_MEMBER_NAME_CHECKED(UTexture2DArray, SourceTextures))
+	{
+		// Empty SourceTextures, remove any resources if present.
+		if (SourceTextures.Num() == 0) 
+		{
+			InvadiateTextureSource();
+		}
+		// First entry into an empty texture array.
+		else if (SourceTextures.Num() == 1)
+		{
+			UpdateSourceFromSourceTextures(true);
+		}
+		// Couldn't add to non-empty array (Error msg logged).
+		else if (UpdateSourceFromSourceTextures(false) == false)
+		{
+			int32 ChangedIndex = PropertyChangedEvent.GetArrayIndex(PropertyName.ToString());
+			int32 LastIndex = SourceTextures.Num() - 1;
+
+			// But don't remove an empty texture, only an incompatible one.
+			if (SourceTextures[LastIndex] != nullptr && ChangedIndex == LastIndex)
+			{
+				SourceTextures.RemoveAt(LastIndex);
+			}
+		}
+	}
+
+	if (PropertyName == GET_MEMBER_NAME_CHECKED(UTexture2DArray, AddressX)
+		|| PropertyName == GET_MEMBER_NAME_CHECKED(UTexture2DArray, AddressY)
+		|| PropertyName == GET_MEMBER_NAME_CHECKED(UTexture2DArray, AddressZ))
+	{
+		UpdateResource();
+	}
+
+	Super::PostEditChangeProperty(PropertyChangedEvent);
+}
+
+#endif
\ No newline at end of file
diff --git a/Engine/Source/Runtime/Engine/Classes/Engine/Texture2DArray.h b/Engine/Source/Runtime/Engine/Classes/Engine/Texture2DArray.h
new file mode 100644
--- /dev/null
+++ b/Engine/Source/Runtime/Engine/Classes/Engine/Texture2DArray.h
@@ -0,0 +1,121 @@
+// Copyright 1998-2019 Epic Games, Inc. All Rights Reserved.
+
+#pragma once
+
+#include "CoreMinimal.h"
+#include "UObject/ObjectMacros.h"
+#include "Engine/Texture2D.h"
+#include "Texture2DArray.generated.h"
+
+#define MAX_ARRAY_SLICES 512
+
+UCLASS(HideCategories = Object, MinimalAPI, BlueprintType)
+class UTexture2DArray : public UTexture
+{
+	GENERATED_UCLASS_BODY()
+
+public:
+	/** Platform data. */
+	FTexturePlatformData* PlatformData;
+	TMap<FString, FTexturePlatformData*> CookedPlatformData;
+
+	/** Creates and initializes a new Texture2D with the requested settings */
+	ENGINE_API int32 CalculateMipZSize(int32 Mip);
+	/*
+	* Initialize texture source from textures in SourceArray.
+	* @param bUpdateSourceSettings Set to false to prevent overriding current texture settings.
+	*/
+	/** Trivial accessors. */
+	FORCEINLINE int32 GetSizeX() const
+	{
+		return PlatformData ? PlatformData->SizeX : 0;
+	}
+	FORCEINLINE int32 GetSizeY() const
+	{
+		return PlatformData ? PlatformData->SizeY : 0;
+	}
+	FORCEINLINE int32 GetNumSlices() const
+	{
+		return PlatformData ? PlatformData->NumSlices : 0;
+	}
+	FORCEINLINE int32 GetNumMips() const
+	{
+		return PlatformData ? PlatformData->Mips.Num() : 0;
+	}
+	FORCEINLINE EPixelFormat GetPixelFormat() const
+	{
+		return PlatformData ? PlatformData->PixelFormat : PF_Unknown;
+	}
+
+	//~ Begin UTexture Interface
+	virtual void Serialize(FArchive& Ar) override;
+	virtual void PostLoad() override;
+	virtual void GetAssetRegistryTags(TArray<FAssetRegistryTag>& OutTags) const override;
+	virtual FString GetDesc() override;
+	virtual void GetResourceSizeEx(FResourceSizeEx& CumulativeResourceSize) override;
+	virtual float GetSurfaceWidth() const override { return GetSizeX(); }
+	virtual float GetSurfaceHeight() const override { return GetSizeY(); }
+	virtual FTextureResource* CreateResource() override;
+#if WITH_EDITOR
+	ENGINE_API virtual void PostEditChangeProperty(FPropertyChangedEvent& PropertyChangedEvent) override;
+	ENGINE_API bool UpdateSourceFromSourceTextures(bool bCreatingNewTexture = true);
+	ENGINE_API void InvadiateTextureSource();
+	ENGINE_API bool CheckArrayTexturesCompatibility();
+#endif // WITH_EDITOR
+	virtual void UpdateResource() override;
+	virtual EMaterialValueType GetMaterialType() const override { return MCT_Texture2DArray; }
+	virtual FTexturePlatformData** GetRunningPlatformData() override { return &PlatformData; }
+	virtual TMap<FString, FTexturePlatformData*> *GetCookedPlatformData() override { return &CookedPlatformData; }
+	//~ End UTexture Interface
+
+#if WITH_EDITORONLY_DATA
+	/** The addressing mode to use for the X axis.*/
+	UPROPERTY(EditAnywhere, Category = Source2D, meta = (DisplayName = "Address X"))
+	TEnumAsByte<enum TextureAddress> AddressX;
+
+	/** The addressing mode to use for the Y axis.*/
+	UPROPERTY(EditAnywhere, Category = Source2D, meta = (DisplayName = "Address Y"))
+	TEnumAsByte<enum TextureAddress> AddressY;
+
+	/** The addressing mode to use for the Z axis.*/
+	UPROPERTY(EditAnywhere, Category = Source2D, meta = (DisplayName = "Address Z"))
+	TEnumAsByte<enum TextureAddress> AddressZ;
+	
+	/** Add Textures*/
+	UPROPERTY(EditAnywhere, Category = Source2D, meta = (DisplayName = "Source Textures"))
+	TArray<UTexture2D*> SourceTextures;
+
+#endif
+
+	/**
+	 * Calculates the size of this texture in bytes if it had MipCount mip-levels streamed in.
+	 *
+	 * @param	MipCount Number of mips to calculate size for, counting from the smallest 1x1 mip-level and up.
+	 * @return	Size of MipCount mips in bytes
+	 */
+	uint32 CalcTextureMemorySize(int32 MipCount) const;
+
+	/**
+	* Calculates the size of this texture if it had MipCount mip levels streamed in.
+	*
+	* @param	Enum which mips to calculate size for.
+	* @return	Total size of all specified mips, in bytes
+	*/
+	virtual uint32 CalcTextureMemorySizeEnum(ETextureMipCount Enum) const override;
+
+#if WITH_EDITOR
+	/**
+	* Return maximum dimension for this texture type.
+	*/
+	virtual uint32 GetMaximumDimension() const override;
+
+#endif
+
+	ENGINE_API static bool ShaderPlatformSupportsCompression(EShaderPlatform ShaderPlatform);
+
+protected:
+
+#if WITH_EDITOR
+	void UpdateMipGenSettings();
+#endif
+};
\ No newline at end of file
diff --git a/Engine/Source/Runtime/Engine/Classes/Materials/MaterialExpressionFunctionInput.h b/Engine/Source/Runtime/Engine/Classes/Materials/MaterialExpressionFunctionInput.h
--- a/Engine/Source/Runtime/Engine/Classes/Materials/MaterialExpressionFunctionInput.h
+++ b/Engine/Source/Runtime/Engine/Classes/Materials/MaterialExpressionFunctionInput.h
@@ -23,6 +23,7 @@ enum EFunctionInputType
 	FunctionInput_Vector4,
 	FunctionInput_Texture2D,
 	FunctionInput_TextureCube,
+	FunctionInput_Texture2DArray,
 	FunctionInput_VolumeTexture,
 	FunctionInput_StaticBool,
 	FunctionInput_MaterialAttributes,
diff --git a/Engine/Source/Runtime/Engine/Classes/Materials/MaterialExpressionTextureSampleParameter2DArray.h b/Engine/Source/Runtime/Engine/Classes/Materials/MaterialExpressionTextureSampleParameter2DArray.h
new file mode 100644
--- /dev/null
+++ b/Engine/Source/Runtime/Engine/Classes/Materials/MaterialExpressionTextureSampleParameter2DArray.h
@@ -0,0 +1,25 @@
+// Copyright 1998-2019 Epic Games, Inc. All Rights Reserved.
+#pragma once
+
+#include "CoreMinimal.h"
+#include "UObject/ObjectMacros.h"
+#include "Materials/MaterialExpressionTextureSampleParameter.h"
+#include "MaterialExpressionTextureSampleParameter2DArray.generated.h"
+
+UCLASS(collapsecategories, hidecategories = Object, MinimalAPI)
+class UMaterialExpressionTextureSampleParameter2DArray : public UMaterialExpressionTextureSampleParameter
+{
+	GENERATED_UCLASS_BODY()
+
+#if WITH_EDITOR
+	// Begin UMaterialExpression Interface
+	virtual void GetCaption(TArray<FString>& OutCaptions) const override;
+	virtual int32 Compile(class FMaterialCompiler* Compiler, int32 OutputIndex) override;
+	// End UMaterialExpression Interface
+#endif
+
+	// Begin UMaterialExpressionTextureSampleParameter Interface
+	virtual bool TextureIsValid( UTexture* InTexture, FString& OutMessage) override;
+	virtual const TCHAR* GetRequirements();
+	// End UMaterialExpressionTextureSampleParameter Interface
+};
\ No newline at end of file
diff --git a/Engine/Source/Runtime/Engine/Private/DDSLoader.cpp b/Engine/Source/Runtime/Engine/Private/DDSLoader.cpp
--- a/Engine/Source/Runtime/Engine/Private/DDSLoader.cpp
+++ b/Engine/Source/Runtime/Engine/Private/DDSLoader.cpp
@@ -6,7 +6,7 @@
 #include "RenderUtils.h"
 
 FDDSLoadHelper::FDDSLoadHelper(const uint8* Buffer, uint32 Length) 
-	: DDSHeader(0)
+	: DDSHeader(0), DDS10Header(0)
 {
 	check(Buffer);
 
@@ -23,6 +23,11 @@ FDDSLoadHelper::FDDSLoadHelper(const uint8* Buffer, uint32 Length)
 		DDSHeader = DDS;
 	}
 
+	// Check for dx10 dds format
+	if (DDS->ddpf.dwFourCC == DDSPF_DX10)
+	{
+		DDS10Header = (FDDS10FileHeader *)(Buffer + 4 + sizeof(FDDSFileHeader));
+	}
 }
 
 bool FDDSLoadHelper::IsValid() const
@@ -76,7 +81,27 @@ EPixelFormat FDDSLoadHelper::ComputePixelFormat() const
 		{
 			Format = PF_FloatRGBA;
 		}
-
+		else if (DDSHeader->ddpf.dwFourCC == DDSPF_DX10 && DDS10Header)
+		{
+			switch (DDS10Header->format)
+			{
+			case(10): Format = PF_FloatRGBA; break; // DXGI_FORMAT_R16G16B16A16_FLOAT
+			case(87): // DXGI_FORMAT_B8G8R8A8_UNORM
+			case(88): // DXGI_FORMAT_B8G8R8X8_UNORM
+			case(91): // DXGI_FORMAT_B8G8R8A8_UNORM_SRGB
+			case(93): Format = PF_B8G8R8A8; break;// DXGI_FORMAT_B8G8R8X8_UNORM_SRGB
+			case(71): // DXGI_FORMAT_BC1_UNORM
+			case(72): Format = PF_DXT1; // DXGI_FORMAT_BC1_UNORM_SRGB
+			case(74): // DXGI_FORMAT_BC2_UNORM
+			case(75): Format = PF_DXT3; // DXGI_FORMAT_BC2_UNORM_SRGB
+			case(77): // DXGI_FORMAT_BC3_UNORM
+			case(78): Format = PF_DXT5; // DXGI_FORMAT_BC3_UNORM_SRGB
+			case(80): // DXGI_FORMAT_BC4_UNORM
+			case(81): Format = PF_BC4; // DXGI_FORMAT_BC4_SNORM
+			case(83): // DXGI_FORMAT_BC4_UNORM
+			case(84): Format = PF_BC5; // DXGI_FORMAT_BC4_SNORM
+			}
+		}
 	}
 	return Format; 
 }
@@ -107,6 +132,19 @@ ETextureSourceFormat FDDSLoadHelper::ComputeSourceFormat() const
 		}
 
 	}
+
+	// Check for dx10 header and extract the format
+	if (DDSHeader->ddpf.dwFourCC == DDSPF_DX10 && DDS10Header)
+	{
+		switch (DDS10Header->format)
+		{
+		case(10): Format = TSF_RGBA16F; break; // DXGI_FORMAT_R16G16B16A16_FLOAT
+		case(87): // DXGI_FORMAT_B8G8R8A8_UNORM
+		case(88): // DXGI_FORMAT_B8G8R8X8_UNORM
+		case(91): // DXGI_FORMAT_B8G8R8A8_UNORM_SRGB
+		case(93): Format = TSF_BGRA8; break;// DXGI_FORMAT_B8G8R8X8_UNORM_SRGB
+		}
+	}
 	return Format; 
 }
 
@@ -117,12 +155,53 @@ bool FDDSLoadHelper::IsValidCubemapTexture() const
 		return true;
 	}
 
+	if (IsValid() && DDS10Header->resourceType == 3 && (DDS10Header->miscFlag & 4))
+	{
+		return true;
+	}
+
+
+	return false;
+}
+
+bool FDDSLoadHelper::IsValidArrayTexture() const
+{
+	if (DDS10Header && DDS10Header->resourceType == 3 && DDS10Header->arraySize > 1 && (DDS10Header->miscFlag & 4) == 0)
+	{
+		return true;
+	}
 	return false;
 }
 
+uint32 FDDSLoadHelper::GetSizeX() const
+{
+	return DDSHeader ? DDSHeader->dwWidth : 0;
+}
+
+uint32 FDDSLoadHelper::GetSizeY() const
+{
+	return DDSHeader ? DDSHeader->dwHeight : 0;
+}
+
+uint32 FDDSLoadHelper::GetSliceCount() const
+{
+	if (IsValidCubemapTexture())
+	{
+		return 6;
+	}
+	else if (IsValidArrayTexture())
+	{
+		return DDS10Header ? DDS10Header->arraySize : 0;
+	}
+	else
+	{
+		return 1;
+	}
+}
+
 bool FDDSLoadHelper::IsValid2DTexture() const
 {
-	if(IsValid() && (DDSHeader->dwCaps2 & DDSC_CubeMap) == 0 && (DDSHeader->dwCaps2 & DDSC_Volume) == 0)
+	if (IsValid() && (DDSHeader->dwCaps2 & DDSC_CubeMap) == 0 && (DDSHeader->dwCaps2 & DDSC_Volume) == 0 && (DDS10Header == nullptr || (DDS10Header->resourceType == 3 && DDS10Header->arraySize == 1)))
 	{
 		return true;
 	}
@@ -141,6 +220,12 @@ const uint8* FDDSLoadHelper::GetDDSDataPointer(ECubeFace Face) const
 
 	const uint8* Ptr = (const uint8*)DDSHeader + sizeof(FDDSFileHeader);
 
+	// jump over dx10 header if available
+	if (DDS10Header)
+	{
+		Ptr += sizeof(FDDS10FileHeader);
+	}
+
 	// jump over the not requested slices / cube map sides
 	Ptr += SliceSize * Face;
 
diff --git a/Engine/Source/Runtime/Engine/Private/Materials/HLSLMaterialTranslator.h b/Engine/Source/Runtime/Engine/Private/Materials/HLSLMaterialTranslator.h
--- a/Engine/Source/Runtime/Engine/Private/Materials/HLSLMaterialTranslator.h
+++ b/Engine/Source/Runtime/Engine/Private/Materials/HLSLMaterialTranslator.h
@@ -1953,6 +1953,7 @@ protected:
 		case MCT_Float:					return TEXT("float");
 		case MCT_Texture2D:				return TEXT("texture2D");
 		case MCT_TextureCube:			return TEXT("textureCube");
+		case MCT_Texture2DArray:		return TEXT("texture2DArray");
 		case MCT_VolumeTexture:			return TEXT("volumeTexture");
 		case MCT_StaticBool:			return TEXT("static bool");
 		case MCT_MaterialAttributes:	return TEXT("MaterialAttributes");
@@ -1976,6 +1977,7 @@ protected:
 		case MCT_Float:					return TEXT("MaterialFloat");
 		case MCT_Texture2D:				return TEXT("texture2D");
 		case MCT_TextureCube:			return TEXT("textureCube");
+		case MCT_Texture2DArray:		return TEXT("texture2DArray");
 		case MCT_VolumeTexture:			return TEXT("volumeTexture");
 		case MCT_StaticBool:			return TEXT("static bool");
 		case MCT_MaterialAttributes:	return TEXT("MaterialAttributes");
@@ -2402,6 +2404,10 @@ protected:
 				TextureInputIndex = MaterialCompilationOutput.UniformExpressionSet.UniformCubeTextureExpressions.AddUnique(TextureUniformExpression);
 				BaseName = TEXT("TextureCube");
 				break;
+			case MCT_Texture2DArray:
+				TextureInputIndex = MaterialCompilationOutput.UniformExpressionSet.Uniform2DArrayTextureExpressions.AddUnique(TextureUniformExpression);
+				BaseName = TEXT("Texture2DArray");
+				break;
 			case MCT_VolumeTexture:
 				TextureInputIndex = MaterialCompilationOutput.UniformExpressionSet.UniformVolumeTextureExpressions.AddUnique(TextureUniformExpression);
 				BaseName = TEXT("VolumeTexture");
@@ -4325,6 +4331,10 @@ protected:
 		{
 			SampleCode += TEXT("TextureCubeSample");
 		}
+		else if (TextureType == MCT_Texture2DArray) 
+		{
+			SampleCode += TEXT("Texture2DArraySample");
+		}
 		else if (TextureType == MCT_VolumeTexture)
 		{
 			SampleCode += TEXT("Texture3DSample");
@@ -4342,7 +4352,7 @@ protected:
 			SampleCode += TEXT("Texture2DSample");
 		}
 		
-		EMaterialValueType UVsType = (TextureType == MCT_TextureCube || TextureType == MCT_VolumeTexture) ? MCT_Float3 : MCT_Float2;
+		EMaterialValueType UVsType = (TextureType == MCT_TextureCube || TextureType == MCT_Texture2DArray || TextureType == MCT_VolumeTexture) ? MCT_Float3 : MCT_Float2;
 	
 		if (RequiresManualViewMipBias)
 		{
@@ -4498,6 +4508,10 @@ protected:
 		{
 			TextureName = CoerceParameter(TextureIndex, MCT_TextureCube);
 		}
+		else if (TextureType == MCT_Texture2DArray)
+		{
+			TextureName = CoerceParameter(TextureIndex, MCT_Texture2DArray);
+		}
 		else if (TextureType == MCT_VolumeTexture)
 		{
 			TextureName = CoerceParameter(TextureIndex, MCT_VolumeTexture);
@@ -6745,6 +6759,13 @@ protected:
 				InputParamDecl += InputNameStr;
 				InputParamDecl += TEXT("Sampler ");
 				break;
+			case MCT_Texture2DArray:
+				InputParamDecl += TEXT("Texture2DArray ");
+				InputParamDecl += InputNameStr;
+				InputParamDecl += TEXT(", SamplerState ");
+				InputParamDecl += InputNameStr;
+				InputParamDecl += TEXT("Sampler ");
+				break;
 			case MCT_TextureExternal:
 				InputParamDecl += TEXT("TextureExternal ");
 				InputParamDecl += InputNameStr;
@@ -6792,7 +6813,7 @@ protected:
 
 			CodeChunk += TEXT(",");
 			CodeChunk += *ParamCode;
-			if (ParamType == MCT_Texture2D || ParamType == MCT_TextureCube || ParamType == MCT_TextureExternal || ParamType == MCT_VolumeTexture)
+			if (ParamType == MCT_Texture2D || ParamType == MCT_TextureCube || ParamType == MCT_Texture2DArray || ParamType == MCT_TextureExternal || ParamType == MCT_VolumeTexture)
 			{
 				CodeChunk += TEXT(",");
 				CodeChunk += *ParamCode;
diff --git a/Engine/Source/Runtime/Engine/Private/Materials/Material.cpp b/Engine/Source/Runtime/Engine/Private/Materials/Material.cpp
--- a/Engine/Source/Runtime/Engine/Private/Materials/Material.cpp
+++ b/Engine/Source/Runtime/Engine/Private/Materials/Material.cpp
@@ -129,7 +129,7 @@ int32 FMaterialResource::CompilePropertyAndSetMaterialProperty(EMaterialProperty
 	UMaterialInterface* MaterialInterface = MaterialInstance ? static_cast<UMaterialInterface*>(MaterialInstance) : Material;
 
 	int32 Ret = INDEX_NONE;
-
+	
 	switch(Property)
 	{
 		case MP_EmissiveColor:
@@ -950,10 +950,11 @@ void UMaterial::GetUsedTextures(TArray<UTexture*>& OutTextures, EMaterialQuality
 				if (CurrentResource == nullptr || (FeatureLevelIndex != FeatureLevel && !bAllFeatureLevels))
 					continue;
 
-				const TArray<TRefCountPtr<FMaterialUniformExpressionTexture> >* ExpressionsByType[4] =
+				const TArray<TRefCountPtr<FMaterialUniformExpressionTexture> >* ExpressionsByType[5] =
 				{
 					&CurrentResource->GetUniform2DTextureExpressions(),
 					&CurrentResource->GetUniformCubeTextureExpressions(),
+					&CurrentResource->GetUniform2DArrayTextureExpressions(),
 					&CurrentResource->GetUniformVolumeTextureExpressions(),
 					&CurrentResource->GetUniformVirtualTextureExpressions()
 				};
@@ -1014,10 +1015,11 @@ void UMaterial::GetUsedTexturesAndIndices(TArray<UTexture*>& OutTextures, TArray
 
 		if (CurrentResource)
 		{
-			const TArray<TRefCountPtr<FMaterialUniformExpressionTexture> >* ExpressionsByType[4] =
+			const TArray<TRefCountPtr<FMaterialUniformExpressionTexture> >* ExpressionsByType[5] =
 			{
 				&CurrentResource->GetUniform2DTextureExpressions(),
 				&CurrentResource->GetUniformCubeTextureExpressions(),
+				&CurrentResource->GetUniform2DArrayTextureExpressions(),
 				&CurrentResource->GetUniformVolumeTextureExpressions(),
 				&CurrentResource->GetUniformVirtualTextureExpressions()
 			};
@@ -1036,6 +1038,10 @@ void UMaterial::GetUsedTexturesAndIndices(TArray<UTexture*>& OutTextures, TArray
 
 					if (Texture)
 					{
+						if (TypeIndex == 2) 
+						{
+							UE_LOG(LogTexture, Warning, TEXT("adsf"));
+						}
 						int32 InsertIndex = OutTextures.AddUnique(Texture);
 						if (InsertIndex >= OutIndices.Num())
 						{
@@ -1072,10 +1078,11 @@ void UMaterial::LogMaterialsAndTextures(FOutputDevice& Ar, int32 Indent) const
 				TArray<UTexture*> Textures;
 				// GetTextureExpressionValues(MaterialResource, Textures);
 				{
-					const TArray<TRefCountPtr<FMaterialUniformExpressionTexture> >* ExpressionsByType[4]= 
+					const TArray<TRefCountPtr<FMaterialUniformExpressionTexture> >* ExpressionsByType[5]= 
 					{
 						&MaterialResource->GetUniform2DTextureExpressions(), 
-						&MaterialResource->GetUniformCubeTextureExpressions(), 
+						&MaterialResource->GetUniformCubeTextureExpressions(),
+						&MaterialResource->GetUniform2DArrayTextureExpressions(),
 						&MaterialResource->GetUniformVolumeTextureExpressions(),
 						&MaterialResource->GetUniformVirtualTextureExpressions(),
 					};
@@ -1128,10 +1135,11 @@ void UMaterial::OverrideTexture(const UTexture* InTextureToOverride, UTexture* O
 	{
 		FMaterialResource* Resource = GetMaterialResource(FeatureLevelsToUpdate[i]);
 		// Iterate over both the 2D textures and cube texture expressions.
-		const TArray<TRefCountPtr<FMaterialUniformExpressionTexture> >* ExpressionsByType[4] =
+		const TArray<TRefCountPtr<FMaterialUniformExpressionTexture> >* ExpressionsByType[5] =
 		{
 			&Resource->GetUniform2DTextureExpressions(),
 			&Resource->GetUniformCubeTextureExpressions(),
+			&Resource->GetUniform2DArrayTextureExpressions(),
 			&Resource->GetUniformVolumeTextureExpressions(),
 			&Resource->GetUniformVirtualTextureExpressions()
 		};
@@ -3710,7 +3718,6 @@ void UMaterial::PostLoad()
 				case TC_Normalmap:
 					TextureExpression->SamplerType = SAMPLERTYPE_Normal;
 					break;
-					
 				case TC_Grayscale:
 					TextureExpression->SamplerType = TextureExpression->Texture->SRGB ? SAMPLERTYPE_Grayscale : SAMPLERTYPE_LinearGrayscale;
 					break;
diff --git a/Engine/Source/Runtime/Engine/Private/Materials/MaterialExpressions.cpp b/Engine/Source/Runtime/Engine/Private/Materials/MaterialExpressions.cpp
--- a/Engine/Source/Runtime/Engine/Private/Materials/MaterialExpressions.cpp
+++ b/Engine/Source/Runtime/Engine/Private/Materials/MaterialExpressions.cpp
@@ -34,6 +34,7 @@
 #include "Engine/TextureRenderTarget2D.h"
 #include "Engine/Texture2DDynamic.h"
 #include "Engine/TextureCube.h"
+#include "Engine/Texture2DArray.h"
 #include "Engine/TextureRenderTargetCube.h"
 #include "Engine/VolumeTexture.h"
 #include "Styling/CoreStyle.h"
@@ -193,6 +194,7 @@
 #include "Materials/MaterialExpressionAntialiasedTextureMask.h"
 #include "Materials/MaterialExpressionTextureSampleParameterSubUV.h"
 #include "Materials/MaterialExpressionTextureSampleParameterCube.h"
+#include "Materials/MaterialExpressionTextureSampleParameter2DArray.h"
 #include "Materials/MaterialExpressionTextureSampleParameterVolume.h"
 #include "Materials/MaterialExpressionTextureCoordinate.h"
 #include "Materials/MaterialExpressionTime.h"
@@ -254,11 +256,13 @@ FUObjectAnnotationSparseBool GMaterialFunctionsThatNeedSamplerFixup;
 bool IsAllowedExpressionType(UClass* Class, bool bMaterialFunction)
 {
 	static const auto AllowVolumeTextureAssetCreationVar = IConsoleManager::Get().FindTConsoleVariableDataInt(TEXT("r.AllowVolumeTextureAssetCreation"));
+	static const auto AllowTextureArrayAssetCreationVar = IConsoleManager::Get().FindTConsoleVariableDataInt(TEXT("r.AllowTexture2DArrayCreation"));
 
 	// Exclude comments from the expression list, as well as the base parameter expression, as it should not be used directly
 	const bool bSharedAllowed = Class != UMaterialExpressionComment::StaticClass() 
-		&& Class != UMaterialExpressionParameter::StaticClass()
-		&& (Class != UMaterialExpressionTextureSampleParameterVolume::StaticClass() || AllowVolumeTextureAssetCreationVar->GetValueOnGameThread() != 0);
+		&&  Class != UMaterialExpressionParameter::StaticClass()
+		&& (Class != UMaterialExpressionTextureSampleParameterVolume::StaticClass() || AllowVolumeTextureAssetCreationVar->GetValueOnGameThread() != 0)
+		&& (Class != UMaterialExpressionTextureSampleParameter2DArray::StaticClass() || AllowTextureArrayAssetCreationVar->GetValueOnGameThread() != 0);
 
 	if (bMaterialFunction)
 	{
@@ -348,6 +352,9 @@ void GetMaterialValueTypeDescriptions(uint32 MaterialValueType, TArray<FText>& O
 			case MCT_TextureCube:
 				OutDescriptions.Add(LOCTEXT("TextureCube", "Texture Cube"));
 				break;
+			case MCT_Texture2DArray:
+				OutDescriptions.Add(LOCTEXT("Texture2DArray", "Texture 2D Array"));
+				break;
 			case MCT_VolumeTexture:
 				OutDescriptions.Add(LOCTEXT("VolumeTexture", "Volume Texture"));
 				break;
@@ -1785,6 +1792,10 @@ int32 UMaterialExpressionTextureSample::Compile(class FMaterialCompiler* Compile
 				{
 					return CompilerError(Compiler, TEXT("UVW input required for volume sample"));
 				}
+				else if (TextureType == MCT_Texture2DArray && !Coordinates.GetTracedInput().Expression)
+				{
+					return CompilerError(Compiler, TEXT("UVW input required for texturearray sample"));
+				}
 			}
 
 			int32 CoordinateIndex = Coordinates.GetTracedInput().Expression ? Coordinates.Compile(Compiler) : Compiler->TextureCoordinate(ConstCoordinate, false, false);
@@ -2631,6 +2642,10 @@ uint32 UMaterialExpressionTextureObject::GetOutputType(int32 OutputIndex)
 	{
 		return MCT_TextureCube;
 	}
+	else if (Cast<UTexture2DArray>(Texture) != NULL)
+	{
+		return MCT_Texture2DArray;
+	}
 	else if (Cast<UVolumeTexture>(Texture) != NULL)
 	{
 		return MCT_VolumeTexture;
@@ -2907,6 +2922,75 @@ void UMaterialExpressionTextureSampleParameterCube::SetDefaultTexture()
 	Texture = LoadObject<UTextureCube>(NULL, TEXT("/Engine/EngineResources/DefaultTextureCube.DefaultTextureCube"), NULL, LOAD_None, NULL);
 }
 
+//
+//  UMaterialExpressionTextureSampleParameter2DArray
+//
+UMaterialExpressionTextureSampleParameter2DArray::UMaterialExpressionTextureSampleParameter2DArray(const FObjectInitializer& ObjectInitializer)
+	: Super(ObjectInitializer)
+{
+	// Structure to hold one-time initialization. No default texture array added.
+	struct FConstructorStatics
+	{
+		FText NAME_Texture;
+		FText NAME_Parameters;
+		FConstructorStatics()
+			: NAME_Texture(LOCTEXT("Texture", "Texture"))
+			, NAME_Parameters(LOCTEXT("Parameters", "Parameters"))
+		{
+		}
+	};
+	static FConstructorStatics ConstructorStatics;
+
+#if WITH_EDITORONLY_DATA
+	MenuCategories.Empty();
+	MenuCategories.Add(ConstructorStatics.NAME_Texture);
+	MenuCategories.Add(ConstructorStatics.NAME_Parameters);
+#endif
+}
+
+#if WITH_EDITOR
+int32 UMaterialExpressionTextureSampleParameter2DArray::Compile(class FMaterialCompiler* Compiler, int32 OutputIndex)
+{
+#if PLATFORM_ANDROID
+	return CompilerError(Compiler, TEXT("Texture2DArrays not supported on selected platform."));
+#endif
+
+	if (!Coordinates.GetTracedInput().Expression)
+	{
+		return CompilerError(Compiler, TEXT("2D array sample needs UVW input"));
+	}
+
+	return UMaterialExpressionTextureSampleParameter::Compile(Compiler, OutputIndex);
+}
+
+void UMaterialExpressionTextureSampleParameter2DArray::GetCaption(TArray<FString>& OutCaptions) const
+{
+	OutCaptions.Add(TEXT("Param2DArray"));
+	OutCaptions.Add(FString::Printf(TEXT("'%s'"), *ParameterName.ToString()));
+}
+#endif
+
+bool UMaterialExpressionTextureSampleParameter2DArray::TextureIsValid(UTexture* InTexture, FString& OutMessage)
+{
+	if (!InTexture)
+	{
+		OutMessage = TEXT("Found NULL, requires Texture2DArray");
+		return false;
+	}
+	else if (!(InTexture->GetMaterialType() & MCT_Texture2DArray))
+	{
+		OutMessage = FString::Printf(TEXT("Found %s, requires Texture2DArray"), *InTexture->GetClass()->GetName());
+		return false;
+	}
+
+	return true;
+}
+
+const TCHAR* UMaterialExpressionTextureSampleParameter2DArray::GetRequirements()
+{
+	return TEXT("Requires Texture2DArray");
+}
+
 //
 //  UMaterialExpressionTextureSampleParameterVolume
 //
@@ -10888,7 +10972,6 @@ void UMaterialFunction::PostLoad()
 				case TC_Normalmap:
 					TextureExpression->SamplerType = SAMPLERTYPE_Normal;
 					break;
-
 				case TC_Grayscale:
 					TextureExpression->SamplerType = TextureExpression->Texture->SRGB ? SAMPLERTYPE_Grayscale : SAMPLERTYPE_LinearGrayscale;
 					break;
@@ -12102,6 +12185,7 @@ static const TCHAR* GetInputTypeName(uint8 InputType)
 		TEXT("V4"),
 		TEXT("T2d"),
 		TEXT("TCube"),
+		TEXT("T2dArr"),
 		TEXT("TVol"),
 		TEXT("B"),
 		TEXT("MA"),
@@ -12610,6 +12694,7 @@ void UMaterialExpressionFunctionInput::GetCaption(TArray<FString>& OutCaptions)
 		TEXT("Vector4"),
 		TEXT("Texture2D"),
 		TEXT("TextureCube"),
+		TEXT("Texture2DArray"),
 		TEXT("VolumeTexture"),
 		TEXT("StaticBool"),
 		TEXT("MaterialAttributes"),
@@ -12660,6 +12745,7 @@ int32 UMaterialExpressionFunctionInput::CompilePreviewValue(FMaterialCompiler* C
 			return FMaterialAttributeDefinitionMap::CompileDefaultExpression(Compiler, AttributeID);
 		case FunctionInput_Texture2D:
 		case FunctionInput_TextureCube:
+		case FunctionInput_Texture2DArray:
 		case FunctionInput_TextureExternal:
 		case FunctionInput_StaticBool:
 			return Compiler->Errorf(TEXT("Missing Preview connection for function input '%s'"), *InputName.ToString());
@@ -12679,6 +12765,7 @@ int32 UMaterialExpressionFunctionInput::Compile(class FMaterialCompiler* Compile
 		MCT_Float4,
 		MCT_Texture2D,
 		MCT_TextureCube,
+		MCT_Texture2DArray,
 		MCT_VolumeTexture,
 		MCT_StaticBool,
 		MCT_MaterialAttributes,
@@ -12819,6 +12906,8 @@ uint32 UMaterialExpressionFunctionInput::GetInputType(int32 InputIndex)
 		return MCT_Texture2D;
 	case FunctionInput_TextureCube:
 		return MCT_TextureCube;
+	case FunctionInput_Texture2DArray:
+		return MCT_Texture2DArray;
 	case FunctionInput_TextureExternal:
 		return MCT_TextureExternal;
 	case FunctionInput_VolumeTexture:
diff --git a/Engine/Source/Runtime/Engine/Private/Materials/MaterialInstance.cpp b/Engine/Source/Runtime/Engine/Private/Materials/MaterialInstance.cpp
--- a/Engine/Source/Runtime/Engine/Private/Materials/MaterialInstance.cpp
+++ b/Engine/Source/Runtime/Engine/Private/Materials/MaterialInstance.cpp
@@ -929,18 +929,19 @@ bool UMaterialInstance::GetRefractionSettings(float& OutBiasValue) const
 
 void UMaterialInstance::GetTextureExpressionValues(const FMaterialResource* MaterialResource, TArray<UTexture*>& OutTextures, TArray< TArray<int32> >* OutIndices) const
 {
-	const TArray<TRefCountPtr<FMaterialUniformExpressionTexture> >* ExpressionsByType[4];
+	const TArray<TRefCountPtr<FMaterialUniformExpressionTexture> >* ExpressionsByType[5];
 
 	check(MaterialResource);
 
 	ExpressionsByType[0] = &MaterialResource->GetUniform2DTextureExpressions();
 	ExpressionsByType[1] = &MaterialResource->GetUniformCubeTextureExpressions();
-	ExpressionsByType[2] = &MaterialResource->GetUniformVolumeTextureExpressions();
-	ExpressionsByType[3] = &MaterialResource->GetUniformVirtualTextureExpressions();
+	ExpressionsByType[2] = &MaterialResource->GetUniform2DArrayTextureExpressions();
+	ExpressionsByType[3] = &MaterialResource->GetUniformVolumeTextureExpressions();
+	ExpressionsByType[4] = &MaterialResource->GetUniformVirtualTextureExpressions();
 
 	if (OutIndices) // Try to prevent resizing since this would be expensive.
 	{
-		OutIndices->Empty(ExpressionsByType[0]->Num() + ExpressionsByType[1]->Num() + ExpressionsByType[2]->Num() + ExpressionsByType[3]->Num());
+		OutIndices->Empty(ExpressionsByType[0]->Num() + ExpressionsByType[1]->Num() + ExpressionsByType[2]->Num() + ExpressionsByType[3]->Num() + ExpressionsByType[4]->Num());
 	}
 
 	for(int32 TypeIndex = 0;TypeIndex < ARRAY_COUNT(ExpressionsByType);TypeIndex++)
@@ -1164,6 +1165,7 @@ void UMaterialInstance::ValidateTextureOverrides(ERHIFeatureLevel::Type InFeatur
 	{
 		&CurrentResource->GetUniform2DTextureExpressions(),
 		&CurrentResource->GetUniformCubeTextureExpressions(),
+		&CurrentResource->GetUniform2DArrayTextureExpressions(),
 		&CurrentResource->GetUniformVolumeTextureExpressions(),
 		&CurrentResource->GetUniformVirtualTextureExpressions(),
 	};
@@ -1272,7 +1274,7 @@ void UMaterialInstance::OverrideTexture(const UTexture* InTextureToOverride, UTe
 #if WITH_EDITOR
 	bool bShouldRecacheMaterialExpressions = false;
 	
-	const TArray<TRefCountPtr<FMaterialUniformExpressionTexture> >* ExpressionsByType[4];
+	const TArray<TRefCountPtr<FMaterialUniformExpressionTexture> >* ExpressionsByType[5];
 
 	const FMaterialResource* SourceMaterialResource = NULL;
 	if (bHasStaticPermutationResource)
@@ -1281,8 +1283,9 @@ void UMaterialInstance::OverrideTexture(const UTexture* InTextureToOverride, UTe
 		// Iterate over both the 2D textures and cube texture expressions.
 		ExpressionsByType[0] = &SourceMaterialResource->GetUniform2DTextureExpressions();
 		ExpressionsByType[1] = &SourceMaterialResource->GetUniformCubeTextureExpressions();
-		ExpressionsByType[2] = &SourceMaterialResource->GetUniformVolumeTextureExpressions();
-		ExpressionsByType[3] = &SourceMaterialResource->GetUniformVirtualTextureExpressions();
+		ExpressionsByType[2] = &SourceMaterialResource->GetUniform2DArrayTextureExpressions();
+		ExpressionsByType[3] = &SourceMaterialResource->GetUniformVolumeTextureExpressions();
+		ExpressionsByType[4] = &SourceMaterialResource->GetUniformVirtualTextureExpressions();
 	}
 	else
 	{
@@ -1293,8 +1296,9 @@ void UMaterialInstance::OverrideTexture(const UTexture* InTextureToOverride, UTe
 		// Iterate over both the 2D textures and cube texture expressions.
 		ExpressionsByType[0] = &SourceMaterialResource->GetUniform2DTextureExpressions();
 		ExpressionsByType[1] = &SourceMaterialResource->GetUniformCubeTextureExpressions();
-		ExpressionsByType[2] = &SourceMaterialResource->GetUniformVolumeTextureExpressions();
-		ExpressionsByType[3] = &SourceMaterialResource->GetUniformVirtualTextureExpressions();
+		ExpressionsByType[2] = &SourceMaterialResource->GetUniform2DArrayTextureExpressions();
+		ExpressionsByType[3] = &SourceMaterialResource->GetUniformVolumeTextureExpressions();
+		ExpressionsByType[4] = &SourceMaterialResource->GetUniformVirtualTextureExpressions();
 	}
 		
 	for(int32 TypeIndex = 0;TypeIndex < ARRAY_COUNT(ExpressionsByType);TypeIndex++)
@@ -4557,10 +4561,11 @@ void UMaterialInstance::CopyMaterialUniformParametersInternal(UMaterialInterface
 					}
 
 					// Textures
-					const TArray<TRefCountPtr<FMaterialUniformExpressionTexture>>* TextureExpressions[4] =
+					const TArray<TRefCountPtr<FMaterialUniformExpressionTexture>>* TextureExpressions[5] =
 					{
 						&MaterialResource->GetUniform2DTextureExpressions(),
 						&MaterialResource->GetUniformCubeTextureExpressions(),
+						&MaterialResource->GetUniform2DArrayTextureExpressions(),
 						&MaterialResource->GetUniformVolumeTextureExpressions(),
 						&MaterialResource->GetUniformVirtualTextureExpressions()
 					};
diff --git a/Engine/Source/Runtime/Engine/Private/Materials/MaterialShared.cpp b/Engine/Source/Runtime/Engine/Private/Materials/MaterialShared.cpp
--- a/Engine/Source/Runtime/Engine/Private/Materials/MaterialShared.cpp
+++ b/Engine/Source/Runtime/Engine/Private/Materials/MaterialShared.cpp
@@ -687,6 +687,19 @@ const TArray<TRefCountPtr<FMaterialUniformExpressionTexture> >& FMaterial::GetUn
 	return EmptyExpressions;
 }
 
+const TArray<TRefCountPtr<FMaterialUniformExpressionTexture> >& FMaterial::GetUniform2DArrayTextureExpressions() const
+{
+	const FMaterialShaderMap* ShaderMapToUse = GetShaderMapToUse();
+
+	if (ShaderMapToUse)
+	{
+		return ShaderMapToUse->GetUniformExpressionSet().Uniform2DArrayTextureExpressions;
+	}
+
+	static const TArray<TRefCountPtr<FMaterialUniformExpressionTexture> > EmptyExpressions;
+	return EmptyExpressions;
+}
+
 const TArray<TRefCountPtr<FMaterialUniformExpressionTexture> >& FMaterial::GetUniformVolumeTextureExpressions() const 
 { 
 	const FMaterialShaderMap* ShaderMapToUse = GetShaderMapToUse();
diff --git a/Engine/Source/Runtime/Engine/Private/Materials/MaterialUniformExpressions.cpp b/Engine/Source/Runtime/Engine/Private/Materials/MaterialUniformExpressions.cpp
--- a/Engine/Source/Runtime/Engine/Private/Materials/MaterialUniformExpressions.cpp
+++ b/Engine/Source/Runtime/Engine/Private/Materials/MaterialUniformExpressions.cpp
@@ -113,10 +113,8 @@ void FUniformExpressionSet::Serialize(FArchive& Ar)
 	Ar << UniformVirtualTextureExpressions;
 	Ar << UniformExternalTextureExpressions;
 	Ar << VTStacks;
+	Ar << Uniform2DArrayTextureExpressions;
 
-	// Adding 2D texture array now to prevent bumping version when the feature gets added
-	TArray<TRefCountPtr<FMaterialUniformExpressionTexture> > Uniform2DTextureArrayExpressions;
-	Ar << Uniform2DTextureArrayExpressions;
 
 	Ar << ParameterCollections;
 
@@ -133,6 +131,7 @@ bool FUniformExpressionSet::IsEmpty() const
 		&& UniformScalarExpressions.Num() == 0
 		&& Uniform2DTextureExpressions.Num() == 0
 		&& UniformCubeTextureExpressions.Num() == 0
+		&& Uniform2DArrayTextureExpressions.Num() == 0
 		&& UniformVolumeTextureExpressions.Num() == 0
 		&& UniformVirtualTextureExpressions.Num() == 0
 		&& UniformExternalTextureExpressions.Num() == 0
@@ -146,6 +145,7 @@ bool FUniformExpressionSet::operator==(const FUniformExpressionSet& ReferenceSet
 		|| UniformScalarExpressions.Num() != ReferenceSet.UniformScalarExpressions.Num()
 		|| Uniform2DTextureExpressions.Num() != ReferenceSet.Uniform2DTextureExpressions.Num()
 		|| UniformCubeTextureExpressions.Num() != ReferenceSet.UniformCubeTextureExpressions.Num()
+		|| Uniform2DArrayTextureExpressions.Num() != ReferenceSet.Uniform2DArrayTextureExpressions.Num()
 		|| UniformVolumeTextureExpressions.Num() != ReferenceSet.UniformVolumeTextureExpressions.Num()
 		|| UniformVirtualTextureExpressions.Num() != ReferenceSet.UniformVirtualTextureExpressions.Num()
 		|| UniformExternalTextureExpressions.Num() != ReferenceSet.UniformExternalTextureExpressions.Num()
@@ -187,6 +187,15 @@ bool FUniformExpressionSet::operator==(const FUniformExpressionSet& ReferenceSet
 		}
 	}
 
+	for (int32 i = 0; i < Uniform2DArrayTextureExpressions.Num(); i++)
+	{
+		if (!Uniform2DArrayTextureExpressions[i]->IsIdentical(ReferenceSet.Uniform2DArrayTextureExpressions[i]))
+		{
+			return false;
+		}
+	}
+
+
 	for (int32 i = 0; i < UniformVolumeTextureExpressions.Num(); i++)
 	{
 		if (!UniformVolumeTextureExpressions[i]->IsIdentical(ReferenceSet.UniformVolumeTextureExpressions[i]))
@@ -232,11 +241,12 @@ bool FUniformExpressionSet::operator==(const FUniformExpressionSet& ReferenceSet
 
 FString FUniformExpressionSet::GetSummaryString() const
 {
-	return FString::Printf(TEXT("(%u vectors, %u scalars, %u 2d tex, %u cube tex, %u 3d tex, %u virtual tex, %u external tex, %u VT stacks, %u collections)"),
+	return FString::Printf(TEXT("(%u vectors, %u scalars, %u 2d tex, %u cube tex, %u 2d array tex, %u 3d tex, %u virtual tex, %u external tex, %u VT stacks, %u collections)"),
 		UniformVectorExpressions.Num(), 
 		UniformScalarExpressions.Num(),
 		Uniform2DTextureExpressions.Num(),
 		UniformCubeTextureExpressions.Num(),
+		Uniform2DArrayTextureExpressions.Num(),
 		UniformVolumeTextureExpressions.Num(),
 		UniformVirtualTextureExpressions.Num(),
 		UniformExternalTextureExpressions.Num(),
@@ -298,6 +308,8 @@ void FUniformExpressionSet::CreateBufferStruct()
 	static FString Texture2DSamplerNames[128];
 	static FString TextureCubeNames[128];
 	static FString TextureCubeSamplerNames[128];
+	static FString Texture2DArrayNames[128];
+	static FString Texture2DArraySamplerNames[128];
 	static FString VolumeTextureNames[128];
 	static FString VolumeTextureSamplerNames[128];
 	static FString ExternalTextureNames[128];
@@ -316,6 +328,8 @@ void FUniformExpressionSet::CreateBufferStruct()
 			Texture2DSamplerNames[i] = FString::Printf(TEXT("Texture2D_%dSampler"), i);
 			TextureCubeNames[i] = FString::Printf(TEXT("TextureCube_%d"), i);
 			TextureCubeSamplerNames[i] = FString::Printf(TEXT("TextureCube_%dSampler"), i);
+			Texture2DArrayNames[i] = FString::Printf(TEXT("Texture2DArray_%d"), i);
+			Texture2DArraySamplerNames[i] = FString::Printf(TEXT("Texture2DArray_%dSampler"), i);
 			VolumeTextureNames[i] = FString::Printf(TEXT("VolumeTexture_%d"), i);
 			VolumeTextureSamplerNames[i] = FString::Printf(TEXT("VolumeTexture_%dSampler"), i);
 			ExternalTextureNames[i] = FString::Printf(TEXT("ExternalTexture_%d"), i);
@@ -329,6 +343,7 @@ void FUniformExpressionSet::CreateBufferStruct()
 
 	check(Uniform2DTextureExpressions.Num() <= 128);
 	check(UniformCubeTextureExpressions.Num() <= 128);
+	check(Uniform2DArrayTextureExpressions.Num() <= 128);
 	check(UniformVolumeTextureExpressions.Num() <= 128);
 	check(UniformVirtualTextureExpressions.Num() <= 128);
 	check(VTStacks.Num() <= 128);
@@ -351,6 +366,15 @@ void FUniformExpressionSet::CreateBufferStruct()
 		NextMemberOffset += SHADER_PARAMETER_POINTER_ALIGNMENT;
 	}
 
+	for (int32 i = 0; i < Uniform2DArrayTextureExpressions.Num(); ++i)
+	{
+		check((NextMemberOffset % SHADER_PARAMETER_POINTER_ALIGNMENT) == 0);
+		new(Members) FShaderParametersMetadata::FMember(*Texture2DArrayNames[i], TEXT("Texture2DArray"), NextMemberOffset, UBMT_TEXTURE, EShaderPrecisionModifier::Float, 1, 1, 0, NULL);
+		NextMemberOffset += SHADER_PARAMETER_POINTER_ALIGNMENT;
+		new(Members) FShaderParametersMetadata::FMember(*Texture2DArraySamplerNames[i], TEXT("SamplerState"), NextMemberOffset, UBMT_SAMPLER, EShaderPrecisionModifier::Float, 1, 1, 0, NULL);
+		NextMemberOffset += SHADER_PARAMETER_POINTER_ALIGNMENT;
+	}
+
 	for (int32 i = 0; i < UniformVolumeTextureExpressions.Num(); ++i)
 	{
 		check((NextMemberOffset % SHADER_PARAMETER_POINTER_ALIGNMENT) == 0);
@@ -543,6 +567,7 @@ void FUniformExpressionSet::FillUniformBuffer(const FMaterialRenderContext& Mate
 			check(UniformBufferStruct->GetLayout().Resources.Num() == 
 				Uniform2DTextureExpressions.Num() * 2
 				+ UniformCubeTextureExpressions.Num() * 2 
+				+ Uniform2DArrayTextureExpressions.Num() * 2
 				+ UniformVolumeTextureExpressions.Num() * 2 
 				+ UniformExternalTextureExpressions.Num() * 2
 				+ UniformVirtualTextureExpressions.Num() * 2
@@ -665,6 +690,44 @@ void FUniformExpressionSet::FillUniformBuffer(const FMaterialRenderContext& Mate
 			}
 		}
 
+		// Cache 2d array texture uniform expressions.
+		for (int32 ExpressionIndex = 0; ExpressionIndex < Uniform2DArrayTextureExpressions.Num(); ExpressionIndex++)
+		{
+			const UTexture* Value;
+			
+			Uniform2DArrayTextureExpressions[ExpressionIndex]->GetTextureValue(MaterialRenderContext, MaterialRenderContext.Material, Value);
+
+			void** ResourceTableTexturePtr = (void**)((uint8*)BufferCursor + 0 * SHADER_PARAMETER_POINTER_ALIGNMENT);
+			void** ResourceTableSamplerPtr = (void**)((uint8*)BufferCursor + 1 * SHADER_PARAMETER_POINTER_ALIGNMENT);
+			BufferCursor = ((uint8*)BufferCursor) + (SHADER_PARAMETER_POINTER_ALIGNMENT * 2);
+
+			if (Value && Value->Resource && (Value->GetMaterialType() & MCT_Texture2DArray) != 0u)
+			{
+				check(Value->TextureReference.TextureReferenceRHI);
+				*ResourceTableTexturePtr = Value->TextureReference.TextureReferenceRHI;
+				FSamplerStateRHIRef* SamplerSource = &Value->Resource->SamplerStateRHI;
+				ESamplerSourceMode SourceMode = Uniform2DArrayTextureExpressions[ExpressionIndex]->GetSamplerSource();
+				if (SourceMode == SSM_Wrap_WorldGroupSettings)
+				{
+					SamplerSource = &Wrap_WorldGroupSettings->SamplerStateRHI;
+				}
+				else if (SourceMode == SSM_Clamp_WorldGroupSettings)
+				{
+					SamplerSource = &Clamp_WorldGroupSettings->SamplerStateRHI;
+				}
+
+				check(*SamplerSource);
+				*ResourceTableSamplerPtr = *SamplerSource;
+			}
+			else
+			{
+				check(GBlackArrayTexture->TextureRHI);
+				*ResourceTableTexturePtr = GBlackArrayTexture->TextureRHI;
+				check(GBlackArrayTexture->SamplerStateRHI);
+				*ResourceTableTexturePtr = GBlackArrayTexture->SamplerStateRHI;
+			}
+		}
+
 		// Cache volume texture uniform expressions.
 		for (int32 ExpressionIndex = 0;ExpressionIndex < UniformVolumeTextureExpressions.Num();ExpressionIndex++)
 		{
diff --git a/Engine/Source/Runtime/Engine/Private/Texture2D.cpp b/Engine/Source/Runtime/Engine/Private/Texture2D.cpp
--- a/Engine/Source/Runtime/Engine/Private/Texture2D.cpp
+++ b/Engine/Source/Runtime/Engine/Private/Texture2D.cpp
@@ -34,7 +34,7 @@
 #include "Streaming/Texture2DStreamIn_IO_Virtual.h"
 #include "Async/AsyncFileHandle.h"
 #include "EngineModule.h"
-
+#include "Engine/Texture2DArray.h"
 #include "VT/UploadingVirtualTexture.h"
 #include "VT/VirtualTexturePoolConfig.h"
 
@@ -2019,7 +2019,7 @@ void FVirtualTexture2DResource::InitializeEditorResources(IVirtualTexture* InVir
 
 		bIgnoreGammaConversions = !TextureOwner->SRGB && TextureOwner->CompressionSettings != TC_HDR;
 
-		// refactored to ensure this is set earlier...make sure it's correct
+		// re factored to ensure this is set earlier...make sure it's correct
 		ensure(bSRGB == TextureOwner->SRGB);
 		//bSRGB = TextureOwner->SRGB;
 	}
@@ -2246,9 +2246,15 @@ void FTexture2DArrayResource::InitRHI()
 	// Create the RHI texture.
 	const uint32 TexCreateFlags = (bSRGB ? TexCreate_SRGB : 0) | TexCreate_OfflineProcessed;
 	FRHIResourceCreateInfo CreateInfo;
-	FTexture2DArrayRHIRef TextureArray = RHICreateTexture2DArray(SizeX, SizeY, GetNumValidTextures(), Format, NumMips, 1, TexCreateFlags, CreateInfo);
+	TRefCountPtr<FRHITexture2DArray> TextureArray = RHICreateTexture2DArray(SizeX, SizeY, NumSlices, Format, NumMips, 1, TexCreateFlags, CreateInfo);
 	TextureRHI = TextureArray;
 
+	if (Owner)
+	{
+		RHIBindDebugLabelName(TextureRHI, *Owner->GetName());
+		RHIUpdateTextureReference(Owner->TextureReference.TextureReferenceRHI, TextureRHI);
+	}
+
 	// Read the mip-levels into the RHI texture.
 	int32 TextureIndex = 0;
 	for (TMap<const UTexture2D*, FTextureArrayDataEntry>::TConstIterator It(CachedData); It; ++It)
@@ -2271,14 +2277,40 @@ void FTexture2DArrayResource::InitRHI()
 		}
 	}
 
-	// Create the sampler state RHI resource.
+	// Read the initial cached mip levels into the RHI texture.
+	for (auto Slice = 0; Slice < CachedInitialData.Num(); ++Slice)
+	{
+		const FTextureArrayDataEntry& CurrentDataEntry = CachedInitialData[Slice];
+		if (CurrentDataEntry.MipData.Num() > 0)
+		{
+			check(CurrentDataEntry.MipData.Num() == NumMips);
+			for (int32 MipIndex = 0; MipIndex < CurrentDataEntry.MipData.Num(); MipIndex++)
+			{
+				if (CurrentDataEntry.MipData[MipIndex].Data.Num() > 0)
+				{
+					uint32 DestStride;
+					void* TheMipData = RHILockTexture2DArray(TextureArray, Slice, MipIndex, RLM_WriteOnly, DestStride, false);
+					GetData(CurrentDataEntry, MipIndex, TheMipData, DestStride);
+					RHIUnlockTexture2DArray(TextureArray, Slice, MipIndex, false);
+				}
+			}
+		}
+	}
+
+#if WITH_EDITORONLY_DATA
+		SamplerXAddress = Owner ? (ESamplerAddressMode)Owner->AddressX.GetValue() : AM_Clamp;
+		SamplerYAddress = Owner ? (ESamplerAddressMode)Owner->AddressY.GetValue() : AM_Clamp;
+		SamplerZAddress = Owner ? (ESamplerAddressMode)Owner->AddressZ.GetValue() : AM_Clamp;
+#endif
+
 	FSamplerStateInitializerRHI SamplerStateInitializer
 	(
 		Filter,
-		AM_Clamp,
-		AM_Clamp,
-		AM_Clamp
+		SamplerXAddress,
+		SamplerYAddress,
+		SamplerZAddress
 	);
+
 	SamplerStateRHI = RHICreateSamplerState(SamplerStateInitializer);
 }
 
@@ -2315,6 +2347,67 @@ FIncomingTextureArrayDataEntry::FIncomingTextureArrayDataEntry(UTexture2D* InTex
 	}
 }
 
+FTexture2DArrayResource::FTexture2DArrayResource(UTexture2DArray* InOwner) 
+{
+	Owner = InOwner;
+	SizeX = InOwner->GetSizeX();
+	SizeY = InOwner->GetSizeY();
+	NumSlices = InOwner->GetNumSlices();
+	NumMips = InOwner->GetNumMips();
+	Format = InOwner->GetPixelFormat();
+	bDirty = true;
+	bPreventingReallocation = false;
+	bSRGB = InOwner->SRGB;
+	Filter = (ESamplerFilter)UDeviceProfileManager::Get().GetActiveProfile()->GetTextureLODSettings()->GetSamplerFilter(InOwner);
+	STAT(LODGroupStatName = TextureGroupStatFNames[InOwner->LODGroup]);
+
+	TIndirectArray<FTexture2DMipMap>& Mips = InOwner->PlatformData->Mips;
+	uint32 Slices = InOwner->GetNumSlices();
+
+	// Create empty data storage.
+	for (uint32 Slice = 0; Slice < Slices; ++Slice)
+	{
+		CachedInitialData.Add(FTextureArrayDataEntry());
+
+		for (int32 MipIndex = 0; MipIndex < Mips.Num(); MipIndex++)
+		{
+			// Add empty mip level entry
+			CachedInitialData[Slice].MipData.Add(FMipMapDataEntry());
+			FMipMapDataEntry& NewEntry = CachedInitialData[Slice].MipData[MipIndex];
+			NewEntry.SizeX = 0;
+			NewEntry.SizeY = 0;
+		}
+	}
+
+	// Making another loop to efficiently copy the mips.
+	for (int32 MipIndex = 0; MipIndex < Mips.Num(); MipIndex++)
+	{
+		FTexture2DMipMap& Mip = Mips[MipIndex];
+		if (Mip.BulkData.GetBulkDataSize() > 0)
+		{
+			uint32 MipSize = Mip.BulkData.GetBulkDataSize() / Slices;
+
+			uint8* In = (uint8*)Mip.BulkData.Lock(LOCK_READ_ONLY);
+
+			for (uint32 Slice = 0; Slice < Slices; ++Slice)
+			{
+				FMipMapDataEntry& NewEntry = CachedInitialData[Slice].MipData[MipIndex];
+				NewEntry.SizeX = Mip.SizeX;
+				NewEntry.SizeY = Mip.SizeY;
+				NewEntry.Data.SetNum(MipSize);
+
+				FMemory::Memcpy(NewEntry.Data.GetData(), In + MipSize * Slice, MipSize);
+			}
+
+			Mip.BulkData.Unlock();
+		}
+		else
+		{
+			UE_LOG(LogTexture, Error, TEXT("Corrupt texture [%s]! Missing bulk data for MipIndex=%d"), *InOwner->GetFullName(), MipIndex);
+		}
+	}
+}
+
 /** 
  * Adds a texture to the texture array.  
  * This is called on the rendering thread, so it must not dereference NewTexture.
@@ -2487,7 +2580,7 @@ void FTexture2DArrayResource::GetData(const FTextureArrayDataEntry& DataEntry, i
 	NumRows = (DataEntry.MipData[MipIndex].SizeY + BlockSizeY - 1) / BlockSizeY;	// Num-of rows in the source data (in blocks)
 	SrcPitch = NumColumns * BlockBytes;		// Num-of bytes per row in the source data
 
-	if (SrcPitch == DestPitch)
+	if (SrcPitch == DestPitch || DestPitch == 0)
 	{
 		// Copy data, not taking into account stride!
 		FMemory::Memcpy(Dest, DataEntry.MipData[MipIndex].Data.GetData(), DataEntry.MipData[MipIndex].Data.Num());
diff --git a/Engine/Source/Runtime/Engine/Private/TextureDerivedData.cpp b/Engine/Source/Runtime/Engine/Private/TextureDerivedData.cpp
--- a/Engine/Source/Runtime/Engine/Private/TextureDerivedData.cpp
+++ b/Engine/Source/Runtime/Engine/Private/TextureDerivedData.cpp
@@ -32,6 +32,7 @@
 #include "Interfaces/ITargetPlatformManagerModule.h"
 #include "Interfaces/ITextureFormat.h"
 #include "Engine/TextureCube.h"
+#include "Engine/Texture2DArray.h"
 #include "ProfilingDebugging/CookStats.h"
 #include "VT/VirtualTextureDataBuilder.h"
 
@@ -118,6 +119,7 @@ static void SerializeForKey(FArchive& Ar, const FTextureBuildSettings& Settings)
 	// NOTE: bHDRSource is not stored in the key here.
 	TempByte = Settings.MipGenSettings; Ar << TempByte;
 	TempByte = Settings.bCubemap; Ar << TempByte;
+	TempByte = Settings.bTextureArray; Ar << TempByte;
 	TempByte = Settings.bSRGB ? (Settings.bSRGB | ( Settings.bUseLegacyGamma ? 0 : 0x2 )) : 0; Ar << TempByte;
 	TempByte = Settings.bPreserveBorder; Ar << TempByte;
 	TempByte = Settings.bDitherMipMapAlpha; Ar << TempByte;
@@ -372,6 +374,7 @@ static void GetTextureBuildSettings(
 	OutBuildSettings.bReplicateRed = false;
 	OutBuildSettings.bVolume = false;
 	OutBuildSettings.bCubemap = false;
+	OutBuildSettings.bTextureArray = false;
 
 	if (Texture.MaxTextureSize > 0)
 	{
@@ -390,6 +393,12 @@ static void GetTextureBuildSettings(
 			OutBuildSettings.MaxTextureResolution = 512;
 		}
 	}
+	else if (Texture.IsA(UTexture2DArray::StaticClass()))
+	{
+		OutBuildSettings.bTextureArray = true;
+		OutBuildSettings.DiffuseConvolveMipLevel = 0;
+		OutBuildSettings.bLongLatSource = false;
+	}
 	else if (Texture.IsA(UVolumeTexture::StaticClass()))
 	{
 		OutBuildSettings.bVolume = true;
@@ -581,8 +590,8 @@ uint32 PutDerivedDataInCache(FTexturePlatformData* DerivedData, const FString& D
 
 	// Write out individual mips to the derived data cache.
 	const int32 MipCount = DerivedData->Mips.Num();
-	const bool bCubemap = (DerivedData->NumSlices == 6);
-	const int32 FirstInlineMip = bCubemap ? 0 : FMath::Max(0, MipCount - NUM_INLINE_DERIVED_MIPS);
+	const bool bIsCubemap = (DerivedData->NumSlices > 0);
+	const int32 FirstInlineMip = bIsCubemap ? 0 : FMath::Max(0, MipCount - NUM_INLINE_DERIVED_MIPS);
 	for (int32 MipIndex = 0; MipIndex < MipCount; ++MipIndex)
 	{
 		FString MipDerivedDataKey;
diff --git a/Engine/Source/Runtime/Engine/Private/TextureDerivedDataTask.cpp b/Engine/Source/Runtime/Engine/Private/TextureDerivedDataTask.cpp
--- a/Engine/Source/Runtime/Engine/Private/TextureDerivedDataTask.cpp
+++ b/Engine/Source/Runtime/Engine/Private/TextureDerivedDataTask.cpp
@@ -16,6 +16,7 @@
 #include "RenderUtils.h"
 #include "TextureResource.h"
 #include "Engine/Texture.h"
+#include "Engine/Texture2DArray.h"
 #include "DeviceProfiles/DeviceProfile.h"
 #include "DeviceProfiles/DeviceProfileManager.h"
 #include "VT/VirtualTextureBuiltData.h"
@@ -61,7 +62,7 @@ void FTextureSourceData::Init(UTexture& InTexture, const FTextureBuildSettings*
 		case TSF_BGRA8:		LayerData->ImageFormat = ERawImageFormat::BGRA8;	break;
 		case TSF_BGRE8:		LayerData->ImageFormat = ERawImageFormat::BGRE8;	break;
 		case TSF_RGBA16:	LayerData->ImageFormat = ERawImageFormat::RGBA16;	break;
-		case TSF_RGBA16F:	LayerData->ImageFormat = ERawImageFormat::RGBA16F; break;
+		case TSF_RGBA16F:	LayerData->ImageFormat = ERawImageFormat::RGBA16F;  break;
 		default:
 			UE_LOG(LogTexture, Fatal, TEXT("Texture %s has source art in an invalid format."), *InTexture.GetName());
 			return;
@@ -93,7 +94,7 @@ void FTextureSourceData::Init(UTexture& InTexture, const FTextureBuildSettings*
 				BlockData->NumMips = 1;
 			}
 
-			if (!InBuildSettingsPerLayer[0].bCubemap && !InBuildSettingsPerLayer[0].bVolume)
+			if (!InBuildSettingsPerLayer[0].bCubemap && !InBuildSettingsPerLayer[0].bTextureArray && !InBuildSettingsPerLayer[0].bVolume)
 			{
 				BlockData->NumSlices = 1;
 			}
@@ -271,7 +272,7 @@ void FTextureCacheDerivedDataWorker::BuildTexture()
 				NewMip->SizeX = CompressedImage.SizeX;
 				NewMip->SizeY = CompressedImage.SizeY;
 				NewMip->SizeZ = CompressedImage.SizeZ;
-				check(NewMip->SizeZ == 1 || BuildSettingsPerLayer[0].bVolume); // Only volume can have SizeZ != 1
+				check(NewMip->SizeZ == 1 || BuildSettingsPerLayer[0].bVolume || BuildSettingsPerLayer[0].bTextureArray); // Only volume & arrays can have SizeZ != 1
 				NewMip->BulkData.Lock(LOCK_READ_WRITE);
 				check(CompressedImage.RawData.GetTypeSize() == 1);
 				void* NewMipData = NewMip->BulkData.Realloc(CompressedImage.RawData.Num());
@@ -282,8 +283,8 @@ void FTextureCacheDerivedDataWorker::BuildTexture()
 				{
 					DerivedData->SizeX = CompressedImage.SizeX;
 					DerivedData->SizeY = CompressedImage.SizeY;
-					DerivedData->NumSlices = BuildSettingsPerLayer[0].bCubemap ? 6 : (BuildSettingsPerLayer[0].bVolume ? CompressedImage.SizeZ : 1);
 					DerivedData->PixelFormat = (EPixelFormat)CompressedImage.PixelFormat;
+					DerivedData->NumSlices = BuildSettingsPerLayer[0].bCubemap ? 6 : (BuildSettingsPerLayer[0].bVolume || BuildSettingsPerLayer[0].bTextureArray) ? CompressedImage.SizeZ : 1;
 				}
 				else
 				{
diff --git a/Engine/Source/Runtime/Engine/Public/DDSLoader.h b/Engine/Source/Runtime/Engine/Public/DDSLoader.h
--- a/Engine/Source/Runtime/Engine/Public/DDSLoader.h
+++ b/Engine/Source/Runtime/Engine/Public/DDSLoader.h
@@ -30,11 +30,12 @@ enum EDDSCaps
 
 enum EDDSPixelFormat
 {
-	DDSPF_FourCC		= 0x00000004,
-	DDSPF_RGB			= 0x00000040,
-	DDSPF_DXT1			= MAKEFOURCC('D','X','T','1'),
-	DDSPF_DXT3			= MAKEFOURCC('D','X','T','3'),
-	DDSPF_DXT5			= MAKEFOURCC('D','X','T','5')
+	DDSPF_FourCC = 0x00000004,
+	DDSPF_RGB = 0x00000040,
+	DDSPF_DXT1 = MAKEFOURCC('D', 'X', 'T', '1'),
+	DDSPF_DXT3 = MAKEFOURCC('D', 'X', 'T', '3'),
+	DDSPF_DXT5 = MAKEFOURCC('D', 'X', 'T', '5'),
+	DDSPF_DX10 = MAKEFOURCC('D', 'X', '1', '0')
 };
 
 // .DDS subheader.
@@ -73,6 +74,18 @@ struct FDDSFileHeader
 };
 #pragma pack(pop)
 
+// .DDS10 header
+#pragma pack(push,1)
+struct FDDS10FileHeader
+{
+	uint32	format;
+	uint32	resourceType;
+	uint32	miscFlag;
+	uint32	arraySize;
+	uint32	miscFlag2;
+};
+#pragma pack(pop)
+
 class FDDSLoadHelper
 {
 public:
@@ -90,8 +103,16 @@ public:
 	/** @param DDS, must not be 0 */
 	ENGINE_API uint32 ComputeMipMapCount() const;
 
+	ENGINE_API uint32 GetSizeX() const;
+
+	ENGINE_API uint32 GetSizeY() const;
+
+	ENGINE_API uint32 GetSliceCount() const;
+
 	ENGINE_API bool IsValidCubemapTexture() const;
 
+	ENGINE_API bool IsValidArrayTexture() const;
+
 	ENGINE_API bool IsValid2DTexture() const;
 
 	/** @param Face should only be provide for cube map textures. */
@@ -104,4 +125,5 @@ public:
 
 	/** !=0 if valid */
 	const FDDSFileHeader* DDSHeader;
+	const FDDS10FileHeader* DDS10Header;
 };
diff --git a/Engine/Source/Runtime/Engine/Public/MaterialShared.h b/Engine/Source/Runtime/Engine/Public/MaterialShared.h
--- a/Engine/Source/Runtime/Engine/Public/MaterialShared.h
+++ b/Engine/Source/Runtime/Engine/Public/MaterialShared.h
@@ -123,13 +123,14 @@ enum EMaterialValueType
 	MCT_Float                 = 8|4|2|1,
 	MCT_Texture2D	          = 1 << 4,
 	MCT_TextureCube	          = 1 << 5,
-	MCT_VolumeTexture         = 1 << 6,
-	MCT_StaticBool            = 1 << 7,
-	MCT_Unknown               = 1 << 8,
-	MCT_MaterialAttributes	  = 1 << 9,
-	MCT_TextureExternal       = 1 << 10,
-	MCT_TextureVirtual        = 1 << 11,
-	MCT_Texture               = MCT_Texture2D | MCT_TextureCube | MCT_VolumeTexture | MCT_TextureExternal | MCT_TextureVirtual,
+	MCT_Texture2DArray		  = 1 << 6,
+	MCT_VolumeTexture         = 1 << 7,
+	MCT_StaticBool            = 1 << 8,
+	MCT_Unknown               = 1 << 9,
+	MCT_MaterialAttributes	  = 1 << 10,
+	MCT_TextureExternal       = 1 << 11,
+	MCT_TextureVirtual        = 1 << 12,
+	MCT_Texture               = MCT_Texture2D | MCT_TextureCube | MCT_Texture2DArray |  MCT_VolumeTexture | MCT_TextureExternal | MCT_TextureVirtual,
 
 	/** Used internally when sampling from virtual textures */
 	MCT_VTPageTableResult     = 1 << 13,
@@ -482,6 +483,7 @@ public:
 			+ UniformScalarExpressions.GetAllocatedSize()
 			+ Uniform2DTextureExpressions.GetAllocatedSize()
 			+ UniformCubeTextureExpressions.GetAllocatedSize()
+			+ Uniform2DArrayTextureExpressions.GetAllocatedSize()
 			+ UniformVolumeTextureExpressions.GetAllocatedSize()
 			+ UniformVirtualTextureExpressions.GetAllocatedSize()
 			+ UniformExternalTextureExpressions.GetAllocatedSize()
@@ -509,6 +511,7 @@ protected:
 	TArray<TRefCountPtr<FMaterialUniformExpression> > UniformScalarExpressions;
 	TArray<TRefCountPtr<FMaterialUniformExpressionTexture> > Uniform2DTextureExpressions;
 	TArray<TRefCountPtr<FMaterialUniformExpressionTexture> > UniformCubeTextureExpressions;
+	TArray<TRefCountPtr<FMaterialUniformExpressionTexture> > Uniform2DArrayTextureExpressions;
 	TArray<TRefCountPtr<FMaterialUniformExpressionTexture> > UniformVolumeTextureExpressions;
 	TArray<TRefCountPtr<FMaterialUniformExpressionTexture> > UniformVirtualTextureExpressions;
 	TArray<TRefCountPtr<FMaterialUniformExpressionExternalTexture> > UniformExternalTextureExpressions;
@@ -1569,6 +1572,7 @@ public:
 	// Accessors.
 	ENGINE_API const TArray<TRefCountPtr<FMaterialUniformExpressionTexture> >& GetUniform2DTextureExpressions() const;
 	ENGINE_API const TArray<TRefCountPtr<FMaterialUniformExpressionTexture> >& GetUniformCubeTextureExpressions() const;
+	ENGINE_API const TArray<TRefCountPtr<FMaterialUniformExpressionTexture> >& GetUniform2DArrayTextureExpressions() const;
 	ENGINE_API const TArray<TRefCountPtr<FMaterialUniformExpressionTexture> >& GetUniformVolumeTextureExpressions() const;
 	ENGINE_API const TArray<TRefCountPtr<FMaterialUniformExpressionTexture> >& GetUniformVirtualTextureExpressions() const;
 	ENGINE_API const TArray<TRefCountPtr<FMaterialUniformExpression> >& GetUniformVectorParameterExpressions() const;
diff --git a/Engine/Source/Runtime/Engine/Public/TextureResource.h b/Engine/Source/Runtime/Engine/Public/TextureResource.h
--- a/Engine/Source/Runtime/Engine/Public/TextureResource.h
+++ b/Engine/Source/Runtime/Engine/Public/TextureResource.h
@@ -24,6 +24,7 @@
 
 class FTexture2DResourceMem;
 class UTexture2D;
+class UTexture2DArray;
 class IVirtualTexture;
 
 /** Maximum number of slices in texture source art. */
@@ -155,7 +156,7 @@ public:
 	 * Minimal initialization constructor.
 	 *
 	 * @param InOwner			UTexture2D which this FTexture2DResource represents.
-	 * @param InitialMipCount	Initial number of miplevels to upload to card
+	 * @param InitialMipCount	Initial number of mip levels to upload to card
 	 * @param InFilename		Filename to read data from
  	 */
 	FTexture2DResource( UTexture2D* InOwner, int32 InitialMipCount );
@@ -169,7 +170,6 @@ public:
 	virtual void RefreshSamplerStates() override;
 
 	// FRenderResource interface.
-
 	/**
 	 * Called when the resource is initialized. This is only called by the rendering thread.
 	 */
@@ -393,6 +393,8 @@ public:
 		bPreventingReallocation(false)
 	{}
 
+	FTexture2DArrayResource(UTexture2DArray* InOwner);
+
 	// Rendering thread functions
 
 	/** 
@@ -431,6 +433,13 @@ public:
 		return SizeY;
 	}
 
+	/** Returns the number of slices(textures) in this array. */
+	uint32 GetNumSlices() const
+	{
+		return NumSlices;
+	}
+
+
 	/** Prevents reallocation from removals of the texture array until EndPreventReallocation is called. */
 	void BeginPreventReallocation();
 
@@ -441,12 +450,20 @@ private:
 
 	/** Texture data, has to persist past the first InitRHI call, because more textures may be added later. */
 	TMap<const UTexture2D*, FTextureArrayDataEntry> CachedData;
+	TArray<FTextureArrayDataEntry>					CachedInitialData;
+	FName LODGroupStatName;
+	UTexture2DArray* Owner;
+
 	uint32 SizeX;
 	uint32 SizeY;
 	uint32 NumMips;
+	uint32 NumSlices;
 	TextureGroup LODGroup;
 	EPixelFormat Format;
 	ESamplerFilter Filter;
+	ESamplerAddressMode SamplerXAddress;
+	ESamplerAddressMode SamplerYAddress;
+	ESamplerAddressMode SamplerZAddress;
 
 	bool bSRGB;
 	bool bDirty;
diff --git a/Engine/Source/Runtime/RenderCore/Private/ShaderParameterMetadata.cpp b/Engine/Source/Runtime/RenderCore/Private/ShaderParameterMetadata.cpp
--- a/Engine/Source/Runtime/RenderCore/Private/ShaderParameterMetadata.cpp
+++ b/Engine/Source/Runtime/RenderCore/Private/ShaderParameterMetadata.cpp
@@ -225,7 +225,8 @@ void FShaderParametersMetadata::InitializeLayout()
 			const bool bTypeCanBeArray = (bAllowResourceArrays && (bIsRHIResource || bIsRDGResource)) || bIsVariableNativeType || BaseType == UBMT_NESTED_STRUCT;
 			if (bIsArray && !bTypeCanBeArray)
 			{
-				UE_LOG(LogRendererCore, Fatal, TEXT("Shader parameter %s error: Not allowed to be an array."), *CppName);
+				// Arrays are allowed.
+				//UE_LOG(LogRendererCore, Fatal, TEXT("Shader parameter %s error: Not allowed to be an array."), *CppName);
 			}
 		}
 
diff --git a/Engine/Source/Runtime/RenderCore/Private/ShaderParameters.cpp b/Engine/Source/Runtime/RenderCore/Private/ShaderParameters.cpp
--- a/Engine/Source/Runtime/RenderCore/Private/ShaderParameters.cpp
+++ b/Engine/Source/Runtime/RenderCore/Private/ShaderParameters.cpp
@@ -264,7 +264,6 @@ static void CreateHLSLUniformBufferStructMembersDeclaration(
 		{
 			check(Member.GetBaseType() != UBMT_RDG_TEXTURE_SRV && Member.GetBaseType() != UBMT_RDG_TEXTURE_UAV);
 			checkf(Member.GetBaseType() != UBMT_RDG_TEXTURE_COPY_DEST && Member.GetBaseType() != UBMT_RDG_BUFFER_COPY_DEST, TEXT("Copy destination usage is not supported in uniform buffers."));
-			checkf(Member.GetNumElements() == 0, TEXT("Resources array are not supported in uniform buffers yet."));
 			if (Member.GetBaseType() == UBMT_SRV)
 			{
 				// TODO: handle arrays?
diff --git a/Engine/Source/Runtime/Renderer/Private/ShaderBaseClasses.cpp b/Engine/Source/Runtime/Renderer/Private/ShaderBaseClasses.cpp
--- a/Engine/Source/Runtime/Renderer/Private/ShaderBaseClasses.cpp
+++ b/Engine/Source/Runtime/Renderer/Private/ShaderBaseClasses.cpp
@@ -170,8 +170,8 @@ void FMaterialShader::VerifyExpressionAndShaderMaps(const FMaterialRenderProxy*
 			TEXT("%s shader uniform expression set mismatch for material %s/%s.\n")
 			TEXT("Shader compilation info:                %s\n")
 			TEXT("Material render proxy compilation info: %s\n")
-			TEXT("Shader uniform expression set:   %u vectors, %u scalars, %u 2D textures, %u cube textures, %u 3D textures, %u virtual textures, shader map %p\n")
-			TEXT("Material uniform expression set: %u vectors, %u scalars, %u 2D textures, %u cube textures, %u 3D textures, %u virtual textures, shader map %p\n"),
+			TEXT("Shader uniform expression set:   %u vectors, %u scalars, %u 2D textures, %u cube textures, %u array textures, %u 3D textures, %u virtual textures, shader map %p\n")
+			TEXT("Material uniform expression set: %u vectors, %u scalars, %u 2D textures, %u cube textures, %u array textures, %u 3D textures, %u virtual textures, shader map %p\n"),
 			GetType()->GetName(),
 			*MaterialRenderProxy->GetFriendlyName(),
 			*Material.GetFriendlyName(),
@@ -181,6 +181,7 @@ void FMaterialShader::VerifyExpressionAndShaderMaps(const FMaterialRenderProxy*
 			DebugUniformExpressionSet.NumScalarExpressions,
 			DebugUniformExpressionSet.Num2DTextureExpressions,
 			DebugUniformExpressionSet.NumCubeTextureExpressions,
+			DebugUniformExpressionSet.Num2DArrayTextureExpressions,
 			DebugUniformExpressionSet.NumVolumeTextureExpressions,
 			DebugUniformExpressionSet.NumVirtualTextureExpressions,
 			UniformExpressionCache->CachedUniformExpressionShaderMap,
@@ -188,6 +189,7 @@ void FMaterialShader::VerifyExpressionAndShaderMaps(const FMaterialRenderProxy*
 			MaterialUniformExpressionSet.UniformScalarExpressions.Num(),
 			MaterialUniformExpressionSet.Uniform2DTextureExpressions.Num(),
 			MaterialUniformExpressionSet.UniformCubeTextureExpressions.Num(),
+			MaterialUniformExpressionSet.Uniform2DArrayTextureExpressions.Num(),
 			MaterialUniformExpressionSet.UniformVolumeTextureExpressions.Num(),
 			MaterialUniformExpressionSet.UniformVirtualTextureExpressions.Num(),
 			Material.GetRenderingThreadShaderMap()
diff --git a/Engine/Source/Runtime/Renderer/Public/MaterialShader.h b/Engine/Source/Runtime/Renderer/Public/MaterialShader.h
--- a/Engine/Source/Runtime/Renderer/Public/MaterialShader.h
+++ b/Engine/Source/Runtime/Renderer/Public/MaterialShader.h
@@ -43,6 +43,7 @@ public:
 	int32 NumScalarExpressions;
 	int32 Num2DTextureExpressions;
 	int32 NumCubeTextureExpressions;
+	int32 Num2DArrayTextureExpressions;
 	int32 NumVolumeTextureExpressions;
 	int32 NumVirtualTextureExpressions;
 
@@ -51,6 +52,7 @@ public:
 		, NumScalarExpressions(0)
 		, Num2DTextureExpressions(0)
 		, NumCubeTextureExpressions(0)
+		, Num2DArrayTextureExpressions(0)
 		, NumVolumeTextureExpressions(0)
 		, NumVirtualTextureExpressions(0)
 	{
@@ -67,6 +69,7 @@ public:
 		NumVectorExpressions = InUniformExpressionSet.UniformVectorExpressions.Num();
 		NumScalarExpressions = InUniformExpressionSet.UniformScalarExpressions.Num();
 		Num2DTextureExpressions = InUniformExpressionSet.Uniform2DTextureExpressions.Num();
+		Num2DArrayTextureExpressions = InUniformExpressionSet.Uniform2DArrayTextureExpressions.Num();
 		NumCubeTextureExpressions = InUniformExpressionSet.UniformCubeTextureExpressions.Num();
 		NumVolumeTextureExpressions = InUniformExpressionSet.UniformVolumeTextureExpressions.Num();
 		NumVirtualTextureExpressions = InUniformExpressionSet.UniformVirtualTextureExpressions.Num();
@@ -79,6 +82,7 @@ public:
 			&& NumScalarExpressions == InUniformExpressionSet.UniformScalarExpressions.Num()
 			&& Num2DTextureExpressions == InUniformExpressionSet.Uniform2DTextureExpressions.Num()
 			&& NumCubeTextureExpressions == InUniformExpressionSet.UniformCubeTextureExpressions.Num()
+			&& Num2DArrayTextureExpressions == InUniformExpressionSet.Uniform2DArrayTextureExpressions.Num()
 			&& NumVolumeTextureExpressions == InUniformExpressionSet.UniformVolumeTextureExpressions.Num()
 			&& NumVirtualTextureExpressions == InUniformExpressionSet.UniformVirtualTextureExpressions.Num();
 	}
@@ -91,6 +95,7 @@ inline FArchive& operator<<(FArchive& Ar, FDebugUniformExpressionSet& DebugExpre
 	Ar << DebugExpressionSet.NumScalarExpressions;
 	Ar << DebugExpressionSet.Num2DTextureExpressions;
 	Ar << DebugExpressionSet.NumCubeTextureExpressions;
+	Ar << DebugExpressionSet.Num2DArrayTextureExpressions;
 	Ar << DebugExpressionSet.NumVolumeTextureExpressions;
 	Ar << DebugExpressionSet.NumVirtualTextureExpressions;
 	return Ar;
