Remove VLM debug warning

diff --git a/Engine/Source/Runtime/Engine/Private/PrecomputedVolumetricLightmap.cpp b/Engine/Source/Runtime/Engine/Private/PrecomputedVolumetricLightmap.cpp
--- a/Engine/Source/Runtime/Engine/Private/PrecomputedVolumetricLightmap.cpp
+++ b/Engine/Source/Runtime/Engine/Private/PrecomputedVolumetricLightmap.cpp
@@ -1051,7 +1051,6 @@ void FVolumetricLightmapBrickAtlas::Insert(int32 Index, FPrecomputedVolumetricLi
 
 void FVolumetricLightmapBrickAtlas::Remove(FPrecomputedVolumetricLightmapData* Data)
 {
-	UE_LOG(LogTemp, Warning, TEXT("VLM Atlas try to remove %llX"), Data);
 	Allocation* AllocationEntry = Allocations.FindByPredicate([Data](const Allocation& Other) { return Other.Data == Data; });
 	if (!AllocationEntry)
 	{
