DR - Fix missing breaks

diff --git a/Engine/Source/Runtime/Engine/Private/DDSLoader.cpp b/Engine/Source/Runtime/Engine/Private/DDSLoader.cpp
--- a/Engine/Source/Runtime/Engine/Private/DDSLoader.cpp
+++ b/Engine/Source/Runtime/Engine/Private/DDSLoader.cpp
@@ -91,15 +91,15 @@ EPixelFormat FDDSLoadHelper::ComputePixelFormat() const
 			case(91): // DXGI_FORMAT_B8G8R8A8_UNORM_SRGB
 			case(93): Format = PF_B8G8R8A8; break;// DXGI_FORMAT_B8G8R8X8_UNORM_SRGB
 			case(71): // DXGI_FORMAT_BC1_UNORM
-			case(72): Format = PF_DXT1; // DXGI_FORMAT_BC1_UNORM_SRGB
+			case(72): Format = PF_DXT1; break; // DXGI_FORMAT_BC1_UNORM_SRGB
 			case(74): // DXGI_FORMAT_BC2_UNORM
-			case(75): Format = PF_DXT3; // DXGI_FORMAT_BC2_UNORM_SRGB
+			case(75): Format = PF_DXT3; break; // DXGI_FORMAT_BC2_UNORM_SRGB
 			case(77): // DXGI_FORMAT_BC3_UNORM
-			case(78): Format = PF_DXT5; // DXGI_FORMAT_BC3_UNORM_SRGB
+			case(78): Format = PF_DXT5; break; // DXGI_FORMAT_BC3_UNORM_SRGB
 			case(80): // DXGI_FORMAT_BC4_UNORM
-			case(81): Format = PF_BC4; // DXGI_FORMAT_BC4_SNORM
+			case(81): Format = PF_BC4; break; // DXGI_FORMAT_BC4_SNORM
 			case(83): // DXGI_FORMAT_BC4_UNORM
-			case(84): Format = PF_BC5; // DXGI_FORMAT_BC4_SNORM
+			case(84): Format = PF_BC5; break; // DXGI_FORMAT_BC4_SNORM
 			}
 		}
 	}
