Fix missing SkyBentNormal texture in the brick atlas

diff --git a/Engine/Source/Runtime/Engine/Private/PrecomputedVolumetricLightmap.cpp b/Engine/Source/Runtime/Engine/Private/PrecomputedVolumetricLightmap.cpp
--- a/Engine/Source/Runtime/Engine/Private/PrecomputedVolumetricLightmap.cpp
+++ b/Engine/Source/Runtime/Engine/Private/PrecomputedVolumetricLightmap.cpp
@@ -948,6 +948,14 @@ void FVolumetricLightmapBrickAtlas::Insert(int32 Index, FPrecomputedVolumetricLi
 		}
 		check(TextureSet.SkyBentNormal.Format == Data->BrickData.SkyBentNormal.Format);
 		check(TextureSet.DirectionalLightShadowing.Format == Data->BrickData.DirectionalLightShadowing.Format);
+
+		// If the incoming BrickData has sky bent normal, also create one in atlas
+		// TODO: release SkyBentNormal if no brick data in the atlas uses it
+		if (!TextureSet.SkyBentNormal.Texture.IsValid() && Data->BrickData.SkyBentNormal.Texture.IsValid())
+		{
+			TextureSet.SkyBentNormal.CreateTargetTexture(TextureSet.BrickDataDimensions);
+			TextureSet.SkyBentNormal.CreateUAV();
+		}
 	}
 
 	int32 NumTotalBricks = 0;
