Fix FComponentRecreateRenderStateContext primitive info invalidation

diff --git a/Engine/Source/Runtime/Engine/Private/Components/ActorComponent.cpp b/Engine/Source/Runtime/Engine/Private/Components/ActorComponent.cpp
--- a/Engine/Source/Runtime/Engine/Private/Components/ActorComponent.cpp
+++ b/Engine/Source/Runtime/Engine/Private/Components/ActorComponent.cpp
@@ -132,13 +132,31 @@ FGlobalComponentRecreateRenderStateContext::FGlobalComponentRecreateRenderStateC
 	// recreate render state for all components.
 	for (UActorComponent* Component : TObjectRange<UActorComponent>())
 	{
-		ComponentContexts.Add(new FComponentRecreateRenderStateContext(Component));
+		ComponentContexts.Add(new FComponentRecreateRenderStateContext(Component, &ScenesToUpdateAllPrimitiveSceneInfos));
 	}
+
+	UpdateAllPrimitiveSceneInfos();
 }
 
 FGlobalComponentRecreateRenderStateContext::~FGlobalComponentRecreateRenderStateContext()
 {
 	ComponentContexts.Empty();
+
+	UpdateAllPrimitiveSceneInfos();
+}
+
+void FGlobalComponentRecreateRenderStateContext::UpdateAllPrimitiveSceneInfos()
+{
+	ENQUEUE_RENDER_COMMAND(UpdateAllPrimitiveSceneInfosCmd)(
+		[ScenesToUpdateAllPrimitiveSceneInfos = MoveTemp(ScenesToUpdateAllPrimitiveSceneInfos)](FRHICommandListImmediate& RHICmdList)
+		{
+			for (FSceneInterface* Scene : ScenesToUpdateAllPrimitiveSceneInfos)
+			{
+				Scene->UpdateAllPrimitiveSceneInfos(RHICmdList);
+			}
+		});
+
+	check(ScenesToUpdateAllPrimitiveSceneInfos.Num() == 0);
 }
 
 // Create Physics global delegate
diff --git a/Engine/Source/Runtime/Engine/Public/ComponentRecreateRenderStateContext.h b/Engine/Source/Runtime/Engine/Public/ComponentRecreateRenderStateContext.h
--- a/Engine/Source/Runtime/Engine/Public/ComponentRecreateRenderStateContext.h
+++ b/Engine/Source/Runtime/Engine/Public/ComponentRecreateRenderStateContext.h
@@ -12,8 +12,12 @@ class FComponentRecreateRenderStateContext
 private:
 	/** Pointer to component we are recreating render state for */
 	UActorComponent* Component;
+
+	TSet<FSceneInterface*>* ScenesToUpdateAllPrimitiveSceneInfos;
+
 public:
-	FComponentRecreateRenderStateContext(UActorComponent* InComponent)
+	FComponentRecreateRenderStateContext(UActorComponent* InComponent, TSet<FSceneInterface*>* InScenesToUpdateAllPrimitiveSceneInfos = nullptr)
+		: ScenesToUpdateAllPrimitiveSceneInfos(InScenesToUpdateAllPrimitiveSceneInfos)
 	{
 		check(InComponent);
 		checkf(!InComponent->IsUnreachable(), TEXT("%s"), *InComponent->GetFullName());
@@ -22,6 +26,23 @@ public:
 		{
 			InComponent->DestroyRenderState_Concurrent();
 			Component = InComponent;
+
+			if (ScenesToUpdateAllPrimitiveSceneInfos == nullptr)
+			{
+				// If no batching is available (this RecreateRenderStateContext is not created by a FGlobalComponentRecreateRenderStateContext), issue one update per component
+				ENQUEUE_RENDER_COMMAND(UpdateAllPrimitiveSceneInfosCmd)([InComponent](FRHICommandListImmediate& RHICmdList) {
+					if (InComponent->GetScene())
+						InComponent->GetScene()->UpdateAllPrimitiveSceneInfos(RHICmdList);
+					});
+			}
+			else
+			{
+				if (InComponent->GetScene())
+				{
+					// Try to batch the updates inside FGlobalComponentRecreateRenderStateContext
+					ScenesToUpdateAllPrimitiveSceneInfos->Add(InComponent->GetScene());
+				}
+			}
 		}
 		else
 		{
@@ -34,6 +55,23 @@ public:
 		if (Component && !Component->IsRenderStateCreated() && Component->IsRegistered())
 		{
 			Component->CreateRenderState_Concurrent();
+
+			if (ScenesToUpdateAllPrimitiveSceneInfos == nullptr)
+			{
+				UActorComponent* InComponent = Component;
+				ENQUEUE_RENDER_COMMAND(UpdateAllPrimitiveSceneInfosCmd)([InComponent](FRHICommandListImmediate& RHICmdList) {
+					if (InComponent->GetScene())
+						InComponent->GetScene()->UpdateAllPrimitiveSceneInfos(RHICmdList);
+					});
+			}
+			else
+			{
+				if (Component->GetScene())
+				{
+					// Try to batch the updates inside FGlobalComponentRecreateRenderStateContext
+					ScenesToUpdateAllPrimitiveSceneInfos->Add(Component->GetScene());
+				}
+			}
 		}
 	}
 };
@@ -53,4 +91,8 @@ public:
 private:
 	/** The recreate contexts for the individual components. */
 	TIndirectArray<FComponentRecreateRenderStateContext> ComponentContexts;
+
+	TSet<FSceneInterface*> ScenesToUpdateAllPrimitiveSceneInfos;
+
+	void UpdateAllPrimitiveSceneInfos();
 };
