Fix non-unity build

diff --git a/Engine/Source/Runtime/Engine/Public/ComponentRecreateRenderStateContext.h b/Engine/Source/Runtime/Engine/Public/ComponentRecreateRenderStateContext.h
--- a/Engine/Source/Runtime/Engine/Public/ComponentRecreateRenderStateContext.h
+++ b/Engine/Source/Runtime/Engine/Public/ComponentRecreateRenderStateContext.h
@@ -5,6 +5,7 @@
 #include "CoreMinimal.h"
 #include "Containers/IndirectArray.h"
 #include "Components/ActorComponent.h"
+#include "SceneInterface.h"
 
 /** Destroys render state for a component and then recreates it when this object is destroyed */
 class FComponentRecreateRenderStateContext
