Enable DXC by default for Metal shader compilation where available (Mac & Windows Editor - there's no Linux build for DXC yet).

diff --git a/Engine/Source/Runtime/Engine/Private/ShaderCompiler/ShaderCompiler.cpp b/Engine/Source/Runtime/Engine/Private/ShaderCompiler/ShaderCompiler.cpp
--- a/Engine/Source/Runtime/Engine/Private/ShaderCompiler/ShaderCompiler.cpp
+++ b/Engine/Source/Runtime/Engine/Private/ShaderCompiler/ShaderCompiler.cpp
@@ -245,10 +245,10 @@ static TAutoConsoleVariable<int32> CVarD3DForceDXC(
 
 static TAutoConsoleVariable<int32> CVarMetalForceDXC(
 	TEXT("r.Metal.ForceDXC"),
-	0,
+	(PLATFORM_WINDOWS || PLATFORM_MAC || PLATFORM_IOS),
 	TEXT("Forces DirectX Shader Compiler (DXC) to be used for all Metal shaders instead of hlslcc.\n")
-	TEXT(" 0: Disable (default)\n")
-	TEXT(" 1: Force new compiler for all shaders"),
+	TEXT(" 0: Disable\n")
+	TEXT(" 1: Use new compiler for all shaders (Default)"),
 	ECVF_ReadOnly);
 
 int32 GCreateShadersOnLoad = 0;
@@ -3246,7 +3246,7 @@ void GlobalBeginCompileShader(
 			}
 		}
 	}
-
+	
 	if (IsOpenGLPlatform((EShaderPlatform)Target.Platform) && 
 		IsMobilePlatform((EShaderPlatform)Target.Platform))
 	{
