Fix EngineTest failing on Lens_LensFlare test

diff --git a/Engine/Source/Runtime/Engine/Private/UpdateTextureShaders.cpp b/Engine/Source/Runtime/Engine/Private/UpdateTextureShaders.cpp
--- a/Engine/Source/Runtime/Engine/Private/UpdateTextureShaders.cpp
+++ b/Engine/Source/Runtime/Engine/Private/UpdateTextureShaders.cpp
@@ -10,6 +10,8 @@ IMPLEMENT_SHADER_TYPE(,FCopyTexture2DCS,TEXT("/Engine/Private/UpdateTextureShade
 IMPLEMENT_SHADER_TYPE(template<>, TCopyTexture2DCS<1u>, TEXT("/Engine/Private/UpdateTextureShaders.usf"), TEXT("TCopyTexture2DCS"), SF_Compute);
 IMPLEMENT_SHADER_TYPE(template<>, TCopyTexture2DCS<2u>, TEXT("/Engine/Private/UpdateTextureShaders.usf"), TEXT("TCopyTexture2DCS"), SF_Compute);
 IMPLEMENT_SHADER_TYPE(template<>, TCopyTexture2DCS<4u>, TEXT("/Engine/Private/UpdateTextureShaders.usf"), TEXT("TCopyTexture2DCS"), SF_Compute);
+typedef TCopyTexture2DCS< 3u, float > FCopyTexture2DCS_Float3;
+IMPLEMENT_SHADER_TYPE(template<>, FCopyTexture2DCS_Float3, TEXT("/Engine/Private/UpdateTextureShaders.usf"), TEXT("TCopyTexture2DCS"), SF_Compute);
 
 IMPLEMENT_SHADER_TYPE(template<>, TUpdateTexture2DSubresouceCS<1u>, TEXT("/Engine/Private/UpdateTextureShaders.usf"), TEXT("TUpdateTexture2DSubresourceCS"), SF_Compute);
 IMPLEMENT_SHADER_TYPE(template<>, TUpdateTexture2DSubresouceCS<2u>, TEXT("/Engine/Private/UpdateTextureShaders.usf"), TEXT("TUpdateTexture2DSubresourceCS"), SF_Compute);
diff --git a/Engine/Source/Runtime/Engine/Public/UpdateTextureShaders.h b/Engine/Source/Runtime/Engine/Public/UpdateTextureShaders.h
--- a/Engine/Source/Runtime/Engine/Public/UpdateTextureShaders.h
+++ b/Engine/Source/Runtime/Engine/Public/UpdateTextureShaders.h
@@ -167,7 +167,7 @@ public:
 	FShaderParameter DestPosSize;
 };
 
-template<uint32 NumComponents>
+template<uint32 NumComponents, typename ComponentType = uint32>
 class TCopyTexture2DCS : public FGlobalShader
 {
 	DECLARE_SHADER_TYPE(TCopyTexture2DCS, Global);
@@ -190,18 +190,34 @@ public:
 		return bShaderHasOutdatedParameters;
 	}
 
-	static void ModifyCompilationEnvironment(const FGlobalShaderPermutationParameters& Parameters, FShaderCompilerEnvironment& OutEnvironment)
+	static const TCHAR* GetTypename(uint32 Dummy)
 	{
-		const TCHAR* Type = nullptr;
-		static_assert(NumComponents == 1u || NumComponents == 2u || NumComponents == 4u, "Invalid NumComponents");
 		switch (NumComponents)
 		{
-		case 1u: Type = TEXT("uint"); break;
-		case 2u: Type = TEXT("uint2"); break;
-		case 4u: Type = TEXT("uint4"); break;
+		case 1u: return TEXT("uint"); break;
+		case 2u: return TEXT("uint2"); break;
+		case 3u: return TEXT("uint3"); break;
+		case 4u: return TEXT("uint4"); break;
 		}
+		return nullptr;
+	}
 
-		OutEnvironment.SetDefine(TEXT("COMPONENT_TYPE"), Type);
+	static const TCHAR* GetTypename(float Dummy)
+	{
+		switch (NumComponents)
+		{
+		case 1u: return TEXT("float"); break;
+		case 2u: return TEXT("float2"); break;
+		case 3u: return TEXT("float3"); break;
+		case 4u: return TEXT("float4"); break;
+		}
+		return nullptr;
+	}
+
+	static void ModifyCompilationEnvironment(const FGlobalShaderPermutationParameters& Parameters, FShaderCompilerEnvironment& OutEnvironment)
+	{
+		static_assert(NumComponents >= 1u && NumComponents <= 4u, "Invalid NumComponents");
+		OutEnvironment.SetDefine(TEXT("COMPONENT_TYPE"), GetTypename((ComponentType)0));
 	}
 
 	static bool ShouldCompilePermutation(const FGlobalShaderPermutationParameters& Parameters)
