Copying

//UE4/Dev-Rendering/Engine/Source/Runtime/Engine/Private/Texture.cpp
to //UE4/Dev-RenderPlat-Staging/Engine/Source/Runtime/Engine/Private/Texture.cpp
for CL: 8639348

diff --git a/Engine/Source/Runtime/Engine/Private/Texture.cpp b/Engine/Source/Runtime/Engine/Private/Texture.cpp
--- a/Engine/Source/Runtime/Engine/Private/Texture.cpp
+++ b/Engine/Source/Runtime/Engine/Private/Texture.cpp
@@ -338,7 +338,7 @@ void UTexture::Serialize(FArchive& Ar)
 	{
 		if (UseVirtualTexturing(GMaxRHIFeatureLevel, Ar.CookingTarget()) == false)
 		{
-			UE_LOG(LogTexture, Warning, TEXT("%s is marked for virtual streaming but virtual texture streaming is not available."), *GetPathName());
+			UE_LOG(LogTexture, Display, TEXT("%s is marked for virtual streaming but virtual texture streaming is not available."), *GetPathName());
 		}
 	}
 
