DRS - Disable Metal DXC
#rnx

diff --git a/Engine/Source/Runtime/Engine/Private/ShaderCompiler/ShaderCompiler.cpp b/Engine/Source/Runtime/Engine/Private/ShaderCompiler/ShaderCompiler.cpp
--- a/Engine/Source/Runtime/Engine/Private/ShaderCompiler/ShaderCompiler.cpp
+++ b/Engine/Source/Runtime/Engine/Private/ShaderCompiler/ShaderCompiler.cpp
@@ -245,7 +245,7 @@ static TAutoConsoleVariable<int32> CVarD3DForceDXC(
 
 static TAutoConsoleVariable<int32> CVarMetalForceDXC(
 	TEXT("r.Metal.ForceDXC"),
-	(PLATFORM_WINDOWS || PLATFORM_MAC || PLATFORM_IOS),
+	0,
 	TEXT("Forces DirectX Shader Compiler (DXC) to be used for all Metal shaders instead of hlslcc.\n")
 	TEXT(" 0: Disable\n")
 	TEXT(" 1: Use new compiler for all shaders (Default)"),
