Fixing PreExposure on post processed before tonemapping.

diff --git a/Engine/Shaders/Private/PostProcessMaterialShaders.usf b/Engine/Shaders/Private/PostProcessMaterialShaders.usf
--- a/Engine/Shaders/Private/PostProcessMaterialShaders.usf
+++ b/Engine/Shaders/Private/PostProcessMaterialShaders.usf
@@ -199,6 +199,10 @@ void MainPS(
 	const float Alpha = 1.0f;
 #endif
 	OutColor = float4(GetMaterialEmissive(PixelMaterialInputs), Alpha );
+
+#if !MATERIALBLENDING_MODULATE && USE_PREEXPOSURE && !POST_PROCESS_MATERIAL_AFTER_TAA_UPSAMPLE
+	OutColor.xyz *= View.PreExposure;
+#endif
 } 
 
 void MainPS_VideoOverlay(
diff --git a/Engine/Source/Runtime/Engine/Private/Materials/HLSLMaterialTranslator.h b/Engine/Source/Runtime/Engine/Private/Materials/HLSLMaterialTranslator.h
--- a/Engine/Source/Runtime/Engine/Private/Materials/HLSLMaterialTranslator.h
+++ b/Engine/Source/Runtime/Engine/Private/Materials/HLSLMaterialTranslator.h
@@ -4876,11 +4876,20 @@ protected:
 
 		if (FeatureLevel >= ERHIFeatureLevel::SM5)
 		{
-			return AddCodeChunk(
+			int32 LookUp = AddCodeChunk(
 				MCT_Float4,
 				TEXT("SceneTextureLookup(%s, %d, %s)"),
 				*CoerceParameter(BufferUV, MCT_Float2), (int)SceneTextureId, bFiltered ? TEXT("true") : TEXT("false")
 				);
+
+			if (SceneTextureId == PPI_PostProcessInput0 && Material->GetMaterialDomain() == MD_PostProcess && Material->GetBlendableLocation() != BL_AfterTonemapping)
+			{
+				return AddInlinedCodeChunk(MCT_Float4, TEXT("(float4(View.OneOverPreExposure.xxx, 1) * %s)"), *CoerceParameter(LookUp, MCT_Float4));
+			}
+			else
+			{
+				return LookUp;
+			}
 		}
 		else // mobile
 		{
