Reproduced CL9318716: Fixes pre-exposure of ComposurePostMoves automation tests

diff --git a/Engine/Source/Runtime/Engine/Classes/Materials/MaterialExpressionViewProperty.h b/Engine/Source/Runtime/Engine/Classes/Materials/MaterialExpressionViewProperty.h
--- a/Engine/Source/Runtime/Engine/Classes/Materials/MaterialExpressionViewProperty.h
+++ b/Engine/Source/Runtime/Engine/Classes/Materials/MaterialExpressionViewProperty.h
@@ -35,6 +35,8 @@ enum EMaterialExposedViewProperty
 	MEVP_RuntimeVirtualTextureOutputLevel UMETA(DisplayName = "Virtual Texture Output Level"),
 	/** World space derivatives for Runtime Virtual Texture Output. */
 	MEVP_RuntimeVirtualTextureOutputDerivative UMETA(DisplayName = "Virtual Texture Output Derivative"),
+	/** Pre Exposure */
+	MEVP_PreExposure UMETA(DisplayName = "Pre-Exposure"),
 
 	MEVP_MAX,
 };
diff --git a/Engine/Source/Runtime/Engine/Private/Materials/HLSLMaterialTranslator.h b/Engine/Source/Runtime/Engine/Private/Materials/HLSLMaterialTranslator.h
--- a/Engine/Source/Runtime/Engine/Private/Materials/HLSLMaterialTranslator.h
+++ b/Engine/Source/Runtime/Engine/Private/Materials/HLSLMaterialTranslator.h
@@ -3138,6 +3138,7 @@ protected:
 			{MEVP_TemporalSampleOffset, MCT_Float2, TEXT("View.TemporalAAParams.zw"), nullptr},
 			{MEVP_RuntimeVirtualTextureOutputLevel, MCT_Float1, TEXT("View.VirtualTextureParams.x"), nullptr},
 			{MEVP_RuntimeVirtualTextureOutputDerivative, MCT_Float2, TEXT("View.VirtualTextureParams.zw"), nullptr},
+			{MEVP_PreExposure, MCT_Float1, TEXT("View.PreExposure.x"), TEXT("View.OneOverPreExposure.x")},
 		};
 		static_assert((sizeof(ViewPropertyMetaArray) / sizeof(ViewPropertyMetaArray[0])) == MEVP_MAX, "incoherency between EMaterialExposedViewProperty and ViewPropertyMetaArray");
 
