Merging //UE4/Main@9388041 to Dev-RenderPlat-Staging

diff --git a/Engine/Source/Runtime/Engine/Private/Collision/SceneQuery.cpp b/Engine/Source/Runtime/Engine/Private/Collision/SceneQuery.cpp
--- a/Engine/Source/Runtime/Engine/Private/Collision/SceneQuery.cpp
+++ b/Engine/Source/Runtime/Engine/Private/Collision/SceneQuery.cpp
@@ -25,10 +25,8 @@ float DebugLineLifetime = 2.f;
 #include "PhysicsEngine/PhysXSupport.h"
 #include "PhysicsEngine/CollisionAnalyzerCapture.h"
 
-#if INCLUDE_CHAOS
 #include "ChaosSolversModule.h"
 #include "Physics/Experimental/ChaosInterfaceWrapper.h"
-#endif
 
 CSV_DEFINE_CATEGORY(SceneQuery, false);
 
@@ -118,9 +116,9 @@ struct FRaycastSQAdditionalInputs
 	}
 };
 
-void LowLevelRaycast(FPhysScene& Scene, const FVector& Start, const FVector& Dir, float DeltaMag, FPhysicsHitCallback<FHitRaycast>& HitBuffer, EHitFlags OutputFlags, FQueryFlags QueryFlags, const FCollisionFilterData& Filter, const FQueryFilterData& QueryFilterData, ICollisionQueryFilterCallbackBase* QueryCallback);
-void LowLevelSweep(FPhysScene& Scene, const FPhysicsGeometry& Geom, const FTransform& StartTM, const FVector& Dir, float DeltaMag, FPhysicsHitCallback<FHitSweep>& HitBuffer, EHitFlags OutputFlags, FQueryFlags QueryFlags, const FCollisionFilterData& Filter, const FQueryFilterData& QueryFilterData, ICollisionQueryFilterCallbackBase* QueryCallback);
-void LowLevelOverlap(FPhysScene& Scene, const FPhysicsGeometry& Geom, const FTransform& GeomPose, FPhysicsHitCallback<FHitOverlap>& HitBuffer, FQueryFlags QueryFlags, const FCollisionFilterData& Filter, const FQueryFilterData& QueryFilterData, ICollisionQueryFilterCallbackBase* QueryCallback);
+void LowLevelRaycast(FPhysScene& Scene, const FVector& Start, const FVector& Dir, float DeltaMag, FPhysicsHitCallback<FHitRaycast>& HitBuffer, EHitFlags OutputFlags, FQueryFlags QueryFlags, const FCollisionFilterData& Filter, const FQueryFilterData& QueryFilterData, ICollisionQueryFilterCallbackBase* QueryCallback, const FQueryDebugParams& DebugParams);
+void LowLevelSweep(FPhysScene& Scene, const FPhysicsGeometry& Geom, const FTransform& StartTM, const FVector& Dir, float DeltaMag, FPhysicsHitCallback<FHitSweep>& HitBuffer, EHitFlags OutputFlags, FQueryFlags QueryFlags, const FCollisionFilterData& Filter, const FQueryFilterData& QueryFilterData, ICollisionQueryFilterCallbackBase* QueryCallback, const FQueryDebugParams& DebugParams);
+void LowLevelOverlap(FPhysScene& Scene, const FPhysicsGeometry& Geom, const FTransform& GeomPose, FPhysicsHitCallback<FHitOverlap>& HitBuffer, FQueryFlags QueryFlags, const FCollisionFilterData& Filter, const FQueryFilterData& QueryFilterData, ICollisionQueryFilterCallbackBase* QueryCallback, const FQueryDebugParams& DebugParams);
 
 template <typename InHitType, ESweepOrRay InGeometryQuery, ESingleMultiOrTest InSingleMultiOrTest>
 struct TSQTraits
@@ -164,7 +162,11 @@ struct TSQTraits
 	static typename TEnableIf<T == ESweepOrRay::Raycast, void>::Type SceneTrace(FPhysScene& Scene, const TGeomInputs& GeomInputs, const FVector& Dir, float DeltaMag, const FTransform& StartTM, THitBuffer& HitBuffer, EHitFlags OutputFlags, EQueryFlags QueryFlags, const FCollisionFilterData& FilterData, const FCollisionQueryParams& Params, ICollisionQueryFilterCallbackBase* QueryCallback)
 	{
 		FQueryFilterData QueryFilterData = MakeQueryFilterData(FilterData, QueryFlags, Params);
-		LowLevelRaycast(Scene, StartTM.GetLocation(), Dir, DeltaMag, HitBuffer, OutputFlags, QueryFlags, FilterData, QueryFilterData, QueryCallback);	//todo(ocohen): namespace?
+		FQueryDebugParams DebugParams;
+#if !(UE_BUILD_TEST || UE_BUILD_SHIPPING) && WITH_CHAOS
+		DebugParams.bDebugQuery = Params.bDebugQuery;
+#endif
+		LowLevelRaycast(Scene, StartTM.GetLocation(), Dir, DeltaMag, HitBuffer, OutputFlags, QueryFlags, FilterData, QueryFilterData, QueryCallback, DebugParams);	//todo(ocohen): namespace?
 	}
 
 	//SceneTrace - sweep
@@ -172,7 +174,11 @@ struct TSQTraits
 	static typename TEnableIf<T == ESweepOrRay::Sweep, void>::Type SceneTrace(FPhysScene& Scene, const TGeomInputs& GeomInputs, const FVector& Dir, float DeltaMag, const FTransform& StartTM, THitBuffer& HitBuffer, EHitFlags OutputFlags, EQueryFlags QueryFlags, const FCollisionFilterData& FilterData, const FCollisionQueryParams& Params, ICollisionQueryFilterCallbackBase* QueryCallback)
 	{
 		FQueryFilterData QueryFilterData = MakeQueryFilterData(FilterData, QueryFlags, Params);
-		LowLevelSweep(Scene, *GeomInputs.GetGeometry(), StartTM, Dir, DeltaMag, HitBuffer, OutputFlags, QueryFlags, FilterData, QueryFilterData, QueryCallback);	//todo(ocohen): namespace?
+		FQueryDebugParams DebugParams;
+#if !(UE_BUILD_TEST || UE_BUILD_SHIPPING) && WITH_CHAOS
+		DebugParams.bDebugQuery = Params.bDebugQuery;
+#endif
+		LowLevelSweep(Scene, *GeomInputs.GetGeometry(), StartTM, Dir, DeltaMag, HitBuffer, OutputFlags, QueryFlags, FilterData, QueryFilterData, QueryCallback, DebugParams);	//todo(ocohen): namespace?
 	}
 
 	static void ResetOutHits(TArray<FHitResult>& OutHits, const FVector& Start, const FVector& End)
@@ -303,7 +309,6 @@ struct TSQTraits
 	constexpr static bool IsSweep() { return GeometryQuery == ESweepOrRay::Sweep;  }
 };
 
-#if INCLUDE_CHAOS
 struct FScopeHelper
 {
 	FScopeHelper()
@@ -325,7 +330,6 @@ struct FScopeHelper
 	}
 	FRWLock* MLock;
 };
-#endif
 
 template <typename Traits, typename TGeomInputs>
 bool TSceneCastCommon(const UWorld* World, typename Traits::TOutHits& OutHits, const TGeomInputs& GeomInputs, const FVector Start, const FVector End, ECollisionChannel TraceChannel, const struct FCollisionQueryParams& Params, const struct FCollisionResponseParams& ResponseParams, const struct FCollisionObjectQueryParams& ObjectParams)
@@ -373,10 +377,7 @@ bool TSceneCastCommon(const UWorld* World, typename Traits::TOutHits& OutHits, c
 		// Enable scene locks, in case they are required
 		FPhysScene& PhysScene = *World->GetPhysicsScene();
 		
-		
-#if INCLUDE_CHAOS
 		FScopeHelper ChaosLockedScope;
-#endif
 
 		FScopedSceneReadLock SceneLocks(PhysScene);
 		{
@@ -569,7 +570,7 @@ bool GeomOverlapMultiImp(const UWorld* World, const FPhysicsGeometry& Geom, cons
 				FScopedSQHitchRepeater<TRemoveReference<decltype(OverlapBuffer)>::Type> HitchRepeater(OverlapBuffer, QueryCallback, FHitchDetectionInfo(GeomPose, TraceChannel, Params));
 				do
 				{
-					LowLevelOverlap(PhysScene, Geom, GeomPose, HitchRepeater.GetBuffer(), QueryFlags, Filter, MakeQueryFilterData(Filter, QueryFlags, Params), &QueryCallback);
+					LowLevelOverlap(PhysScene, Geom, GeomPose, HitchRepeater.GetBuffer(), QueryFlags, Filter, MakeQueryFilterData(Filter, QueryFlags, Params), &QueryCallback, FQueryDebugParams());
 				} while(HitchRepeater.RepeatOnHitch());
 
 				if(GetHasBlock(OverlapBuffer) && InfoType != EQueryInfo::GatherAll)	//just want true or false so don't bother gathering info
@@ -623,12 +624,8 @@ bool FGenericPhysicsInterface::GeomOverlapBlockingTest(const UWorld* World, cons
 
 	TArray<FOverlapResult> Overlaps;	//needed only for template shared code
 	FTransform GeomTransform(Rot, Pos);
-#if WITH_PHYSX && !WITH_CHAOS_NEEDS_TO_BE_FIXED
-	FPhysXShapeAdapter Adaptor(GeomTransform.GetRotation(), CollisionShape);
+	FPhysicsShapeAdapter Adaptor(GeomTransform.GetRotation(), CollisionShape);
 	return GeomOverlapMultiImp<EQueryInfo::IsBlocking>(World, Adaptor.GetGeometry(), CollisionShape, Adaptor.GetGeomPose(GeomTransform.GetTranslation()), Overlaps, TraceChannel, Params, ResponseParams, ObjectParams);
-#else
-	return false;
-#endif
 }
 
 bool FGenericPhysicsInterface::GeomOverlapAnyTest(const UWorld* World, const struct FCollisionShape& CollisionShape, const FVector& Pos, const FQuat& Rot, ECollisionChannel TraceChannel, const struct FCollisionQueryParams& Params, const struct FCollisionResponseParams& ResponseParams, const struct FCollisionObjectQueryParams& ObjectParams)
@@ -639,12 +636,8 @@ bool FGenericPhysicsInterface::GeomOverlapAnyTest(const UWorld* World, const str
 
 	TArray<FOverlapResult> Overlaps;	//needed only for template shared code
 	FTransform GeomTransform(Rot, Pos);
-#if WITH_PHYSX && !WITH_CHAOS_NEEDS_TO_BE_FIXED
-	FPhysXShapeAdapter Adaptor(GeomTransform.GetRotation(), CollisionShape);
+	FPhysicsShapeAdapter Adaptor(GeomTransform.GetRotation(), CollisionShape);
 	return GeomOverlapMultiImp<EQueryInfo::GatherAll>(World, Adaptor.GetGeometry(), CollisionShape, Adaptor.GetGeomPose(GeomTransform.GetTranslation()), Overlaps, TraceChannel, Params, ResponseParams, ObjectParams);
-#else
-	return false;
-#endif
 }
 
 template<>
@@ -654,12 +647,8 @@ bool FGenericPhysicsInterface::GeomOverlapMulti(const UWorld* World, const FPhys
 	SCOPE_CYCLE_COUNTER(STAT_Collision_GeomOverlapMultiple);
 	CSV_SCOPED_TIMING_STAT(SceneQuery, GeomOverlapMultiple);
 
-#if WITH_PHYSX && !WITH_CHAOS_NEEDS_TO_BE_FIXED
 	FTransform GeomTransform(InRotation, InPosition);
 	return GeomOverlapMultiImp<EQueryInfo::GatherAll>(World, InGeom.GetGeometry(), InGeom, GeomTransform, OutOverlaps, TraceChannel, Params, ResponseParams, ObjectParams);
-#else
-	return false;
-#endif
 }
 
 template<>
@@ -670,10 +659,6 @@ bool FGenericPhysicsInterface::GeomOverlapMulti(const UWorld* World, const FColl
 	CSV_SCOPED_TIMING_STAT(SceneQuery, GeomOverlapMultiple);
 
 	FTransform GeomTransform(InRotation, InPosition);
-#if WITH_PHYSX && !WITH_CHAOS_NEEDS_TO_BE_FIXED
-	FPhysXShapeAdapter Adaptor(GeomTransform.GetRotation(), InGeom);
+	FPhysicsShapeAdapter Adaptor(GeomTransform.GetRotation(), InGeom);
 	return GeomOverlapMultiImp<EQueryInfo::GatherAll>(World, Adaptor.GetGeometry(), InGeom, Adaptor.GetGeomPose(GeomTransform.GetTranslation()), OutOverlaps, TraceChannel, Params, ResponseParams, ObjectParams);
-#else
-	return false;
-#endif
 }
\ No newline at end of file
diff --git a/Engine/Source/Runtime/Engine/Private/Collision/SceneQueryLowLevel.cpp b/Engine/Source/Runtime/Engine/Private/Collision/SceneQueryLowLevel.cpp
--- a/Engine/Source/Runtime/Engine/Private/Collision/SceneQueryLowLevel.cpp
+++ b/Engine/Source/Runtime/Engine/Private/Collision/SceneQueryLowLevel.cpp
@@ -1,7 +1,6 @@
 // Copyright 1998-2019 Epic Games, Inc. All Rights Reserved.
 
 #if WITH_PHYSX
-
 #include "PhysXPublic.h"
 #endif
 #include "Physics/PhysicsInterfaceDeclares.h"
@@ -16,12 +15,10 @@
 
 #include "PhysTestSerializer.h"
 
-#if INCLUDE_CHAOS
 #include "SQAccelerator.h"
 #include "SQVerifier.h"
 #include "PBDRigidsSolver.h"
 #include "Chaos/PBDRigidsEvolutionGBF.h"
-#endif
 
 int32 ForceStandardSQ = 0;
 FAutoConsoleVariableRef CVarForceStandardSQ(TEXT("p.ForceStandardSQ"), ForceStandardSQ, TEXT("If enabled, we force the standard scene query even if custom SQ structure is enabled"));
@@ -42,7 +39,7 @@ void FinalizeCapture(FPhysTestSerializer& Serializer)
 	{
 		Serializer.Serialize(TEXT("SQCapture"));
 	}
-#if INCLUDE_CHAOS && WITH_PHYSX
+#if WITH_PHYSX
 	if (ReplaySQs)
 	{
 		const bool bReplaySuccess = SQComparisonHelper(Serializer);
@@ -64,14 +61,28 @@ constexpr int32 ReplaySQs = 0;
 void FinalizeCapture(FPhysTestSerializer& Serializer) {}
 #endif
 
-void LowLevelRaycast(FPhysScene& Scene, const FVector& Start, const FVector& Dir, float DeltaMag, FPhysicsHitCallback<FHitRaycast>& HitBuffer, EHitFlags OutputFlags, FQueryFlags QueryFlags, const FCollisionFilterData& Filter, const FQueryFilterData& QueryFilterData, ICollisionQueryFilterCallbackBase* QueryCallback)
+void LowLevelRaycast(FPhysScene& Scene, const FVector& Start, const FVector& Dir, float DeltaMag, FPhysicsHitCallback<FHitRaycast>& HitBuffer, EHitFlags OutputFlags, FQueryFlags QueryFlags, const FCollisionFilterData& Filter, const FQueryFilterData& QueryFilterData, ICollisionQueryFilterCallbackBase* QueryCallback, const FQueryDebugParams& DebugParams)
 {
 #if !defined(PHYSICS_INTERFACE_PHYSX) || !PHYSICS_INTERFACE_PHYSX
-	if (const auto& SolverAccelerationStructure = Scene.GetScene().SolverAccelerationStructure)
+	if (const auto& SolverAccelerationStructure = Scene.GetScene().GetSpacialAcceleration())
 	{
-		FChaosSQAccelerator SQAccelerator(*SolverAccelerationStructure.Get());
-		//ISQAccelerator* SQAccelerator = Scene.GetSQAccelerator();
-		SQAccelerator.Raycast(Start, Dir, DeltaMag, HitBuffer, OutputFlags, QueryFilterData, *QueryCallback);
+		FChaosSQAccelerator SQAccelerator(*SolverAccelerationStructure);
+		if (SerializeSQs && IsInGameThread())
+		{
+			FPhysTestSerializer Serializer;
+			Serializer.SetPhysicsData(*Scene.GetSolver()->GetEvolution());
+			FSQCapture& RaycastCapture = Serializer.CaptureSQ();
+			RaycastCapture.StartCaptureChaosRaycast(*Scene.GetSolver()->GetEvolution(), Start, Dir, DeltaMag, OutputFlags, QueryFilterData, Filter, *QueryCallback);
+			SQAccelerator.Raycast(Start, Dir, DeltaMag, HitBuffer, OutputFlags, QueryFilterData, *QueryCallback, DebugParams);
+			RaycastCapture.EndCaptureChaosRaycast(HitBuffer);
+
+			FinalizeCapture(Serializer);
+		}
+		else
+		{
+			//ISQAccelerator* SQAccelerator = Scene.GetSQAccelerator();
+			SQAccelerator.Raycast(Start, Dir, DeltaMag, HitBuffer, OutputFlags, QueryFilterData, *QueryCallback, DebugParams);
+		}
 	}
 #else
 	if (SerializeSQs | ReplaySQs)
@@ -92,14 +103,28 @@ void LowLevelRaycast(FPhysScene& Scene, const FVector& Start, const FVector& Dir
 #endif
 }
 
-void LowLevelSweep(FPhysScene& Scene, const FPhysicsGeometry& QueryGeom, const FTransform& StartTM, const FVector& Dir, float DeltaMag, FPhysicsHitCallback<FHitSweep>& HitBuffer, EHitFlags OutputFlags, FQueryFlags QueryFlags, const FCollisionFilterData& Filter, const FQueryFilterData& QueryFilterData, ICollisionQueryFilterCallbackBase* QueryCallback)
+void LowLevelSweep(FPhysScene& Scene, const FPhysicsGeometry& QueryGeom, const FTransform& StartTM, const FVector& Dir, float DeltaMag, FPhysicsHitCallback<FHitSweep>& HitBuffer, EHitFlags OutputFlags, FQueryFlags QueryFlags, const FCollisionFilterData& Filter, const FQueryFilterData& QueryFilterData, ICollisionQueryFilterCallbackBase* QueryCallback, const FQueryDebugParams& DebugParams)
 {
 #if !defined(PHYSICS_INTERFACE_PHYSX) || !PHYSICS_INTERFACE_PHYSX
-	if (const auto& SolverAccelerationStructure = Scene.GetScene().SolverAccelerationStructure)
+	if (const auto& SolverAccelerationStructure = Scene.GetScene().GetSpacialAcceleration())
 	{
-		FChaosSQAccelerator SQAccelerator(*SolverAccelerationStructure.Get());
-		//ISQAccelerator* SQAccelerator = Scene.GetSQAccelerator();
-		SQAccelerator.Sweep(QueryGeom, StartTM, Dir, DeltaMag, HitBuffer, OutputFlags, QueryFilterData, *QueryCallback);
+		FChaosSQAccelerator SQAccelerator(*SolverAccelerationStructure);
+		if (SerializeSQs && IsInGameThread())
+		{
+			FPhysTestSerializer Serializer;
+			Serializer.SetPhysicsData(*Scene.GetSolver()->GetEvolution());
+			FSQCapture& SweepCapture = Serializer.CaptureSQ();
+			SweepCapture.StartCaptureChaosSweep(*Scene.GetSolver()->GetEvolution(), QueryGeom, StartTM, Dir, DeltaMag, OutputFlags, QueryFilterData, Filter, *QueryCallback);
+			SQAccelerator.Sweep(QueryGeom, StartTM, Dir, DeltaMag, HitBuffer, OutputFlags, QueryFilterData, *QueryCallback);
+			SweepCapture.EndCaptureChaosSweep(HitBuffer);
+
+			FinalizeCapture(Serializer);
+		}
+		else
+		{
+			//ISQAccelerator* SQAccelerator = Scene.GetSQAccelerator();
+			SQAccelerator.Sweep(QueryGeom, StartTM, Dir, DeltaMag, HitBuffer, OutputFlags, QueryFilterData, *QueryCallback);
+		}
 	}
 #else
 	if (SerializeSQs | ReplaySQs)
@@ -120,14 +145,28 @@ void LowLevelSweep(FPhysScene& Scene, const FPhysicsGeometry& QueryGeom, const F
 #endif
 }
 
-void LowLevelOverlap(FPhysScene& Scene, const FPhysicsGeometry& QueryGeom, const FTransform& GeomPose, FPhysicsHitCallback<FHitOverlap>& HitBuffer, FQueryFlags QueryFlags, const FCollisionFilterData& Filter, const FQueryFilterData& QueryFilterData, ICollisionQueryFilterCallbackBase* QueryCallback)
+void LowLevelOverlap(FPhysScene& Scene, const FPhysicsGeometry& QueryGeom, const FTransform& GeomPose, FPhysicsHitCallback<FHitOverlap>& HitBuffer, FQueryFlags QueryFlags, const FCollisionFilterData& Filter, const FQueryFilterData& QueryFilterData, ICollisionQueryFilterCallbackBase* QueryCallback, const FQueryDebugParams& DebugParams)
 {
 #if !defined(PHYSICS_INTERFACE_PHYSX) || !PHYSICS_INTERFACE_PHYSX
-	if (const auto& SolverAccelerationStructure = Scene.GetScene().SolverAccelerationStructure)
+	if (const auto& SolverAccelerationStructure = Scene.GetScene().GetSpacialAcceleration())
 	{
-		FChaosSQAccelerator SQAccelerator(*SolverAccelerationStructure.Get());
-		//ISQAccelerator* SQAccelerator = Scene.GetSQAccelerator();
-		SQAccelerator.Overlap(QueryGeom, GeomPose, HitBuffer, QueryFilterData, *QueryCallback);
+		FChaosSQAccelerator SQAccelerator(*SolverAccelerationStructure);
+		if (SerializeSQs && IsInGameThread())
+		{
+			FPhysTestSerializer Serializer;
+			Serializer.SetPhysicsData(*Scene.GetSolver()->GetEvolution());
+			FSQCapture& OverlapCapture = Serializer.CaptureSQ();
+			OverlapCapture.StartCaptureChaosOverlap(*Scene.GetSolver()->GetEvolution(), QueryGeom, GeomPose, QueryFilterData, Filter, *QueryCallback);
+			SQAccelerator.Overlap(QueryGeom, GeomPose, HitBuffer, QueryFilterData, *QueryCallback);
+			OverlapCapture.EndCaptureChaosOverlap(HitBuffer);
+
+			FinalizeCapture(Serializer);
+		}
+		else
+		{
+			//ISQAccelerator* SQAccelerator = Scene.GetSQAccelerator();
+			SQAccelerator.Overlap(QueryGeom, GeomPose, HitBuffer, QueryFilterData, *QueryCallback);
+		}
 	}
 #else
 	if (SerializeSQs | ReplaySQs)
diff --git a/Engine/Source/Runtime/Engine/Private/Collision/WorldCollision.cpp b/Engine/Source/Runtime/Engine/Private/Collision/WorldCollision.cpp
--- a/Engine/Source/Runtime/Engine/Private/Collision/WorldCollision.cpp
+++ b/Engine/Source/Runtime/Engine/Private/Collision/WorldCollision.cpp
@@ -16,6 +16,10 @@
 #include "PhysXPublic.h"
 #include "Physics/PhysicsInterfaceTypes.h"
 
+#if WITH_CHAOS
+#include "Chaos/ImplicitObject.h"
+#endif
+
 using namespace PhysicsInterfaceTypes;
 
 DEFINE_LOG_CATEGORY(LogCollision);
@@ -409,11 +413,19 @@ bool UWorld::ComponentSweepMulti(TArray<struct FHitResult>& OutHits, class UPrim
 			check(Shape.IsValid());
 			ECollisionShapeType ShapeType = FPhysicsInterface::GetShapeType(Shape);
 
+#if WITH_CHAOS
+			if (!Shape.GetGeometry().IsConvex())
+			{
+				//We skip complex shapes. Should this respect complex as simple?
+				continue;
+			}
+#else
 			if(ShapeType == ECollisionShapeType::Heightfield || ShapeType == ECollisionShapeType::Trimesh)
 			{
 				//We skip complex shapes. Should this respect complex as simple?
 				continue;
 			}
+#endif
 
 			// Calc shape global pose
 			const FTransform ShapeLocalTransform = FPhysicsInterface::GetLocalTransform(Shape);
diff --git a/Engine/Source/Runtime/Engine/Private/Components/PrimitiveComponent.cpp b/Engine/Source/Runtime/Engine/Private/Components/PrimitiveComponent.cpp
--- a/Engine/Source/Runtime/Engine/Private/Components/PrimitiveComponent.cpp
+++ b/Engine/Source/Runtime/Engine/Private/Components/PrimitiveComponent.cpp
@@ -2049,11 +2049,13 @@ bool UPrimitiveComponent::MoveComponentImpl( const FVector& Delta, const FQuat&
 #endif 
 			UWorld* const MyWorld = GetWorld();
 
+			static const FName TraceTagName = TEXT("MoveComponent");
 			const bool bForceGatherOverlaps = !ShouldCheckOverlapFlagToQueueOverlaps(*this);
 			FComponentQueryParams Params(SCENE_QUERY_STAT(MoveComponent), Actor);
 			FCollisionResponseParams ResponseParam;
 			InitSweepCollisionParams(Params, ResponseParam);
 			Params.bIgnoreTouches |= !(GetGenerateOverlapEvents() || bForceGatherOverlaps);
+			Params.TraceTag = TraceTagName;
 			bool const bHadBlockingHit = MyWorld->ComponentSweepMulti(Hits, this, TraceStart, TraceEnd, InitialRotationQuat, Params);
 
 			if (Hits.Num() > 0)
diff --git a/Engine/Source/Runtime/Engine/Private/Components/SkeletalMeshComponent.cpp b/Engine/Source/Runtime/Engine/Private/Components/SkeletalMeshComponent.cpp
--- a/Engine/Source/Runtime/Engine/Private/Components/SkeletalMeshComponent.cpp
+++ b/Engine/Source/Runtime/Engine/Private/Components/SkeletalMeshComponent.cpp
@@ -179,7 +179,7 @@ USkeletalMeshComponent::USkeletalMeshComponent(const FObjectInitializer& ObjectI
 	ClothTickFunction.EndTickGroup = TG_PostPhysics;
 	ClothTickFunction.bCanEverTick = true;
 
-#if WITH_APEX_CLOTHING
+#if WITH_APEX_CLOTHING || WITH_CHAOS_CLOTHING
 	ClothMaxDistanceScale = 1.0f;
 	bResetAfterTeleport = true;
 	TeleportDistanceThreshold = 300.0f;
@@ -196,7 +196,7 @@ USkeletalMeshComponent::USkeletalMeshComponent(const FObjectInitializer& ObjectI
 	bBindClothToMasterComponent = false;
 	bClothingSimulationSuspended = false;
 
-#endif//#if WITH_APEX_CLOTHING
+#endif//#if WITH_APEX_CLOTHING || WITH_CHAOS_CLOTHING
 
 	MassMode = EClothMassMode::Density;
 	UniformMass = 1.f;
@@ -233,18 +233,15 @@ USkeletalMeshComponent::USkeletalMeshComponent(const FObjectInitializer& ObjectI
 	ResetRootBodyIndex();
 
 	TArray<IClothingSimulationFactoryClassProvider*> ClassProviders = IModularFeatures::Get().GetModularFeatureImplementations<IClothingSimulationFactoryClassProvider>(IClothingSimulationFactoryClassProvider::FeatureName);
-	if(ClassProviders.Num() > 0)
+
+	ClothingSimulationFactory = nullptr;
+	for (int32 Index = ClassProviders.Num() - 1; Index >= 0 && !*ClothingSimulationFactory; --Index)  // We use the last provider in the list so plugins/modules can override ours
 	{
-		// We use the last provider in the list so plugins/modules can override ours
-		IClothingSimulationFactoryClassProvider* Provider = ClassProviders.Last();
+		IClothingSimulationFactoryClassProvider* const Provider = ClassProviders[Index];
 		check(Provider);
 
 		ClothingSimulationFactory = Provider->GetDefaultSimulationFactoryClass();
 	}
-	else
-	{
-		ClothingSimulationFactory = nullptr;
-	}
 
 	ClothingSimulation = nullptr;
 	ClothingSimulationContext = nullptr;
@@ -532,16 +529,15 @@ void USkeletalMeshComponent::OnRegister()
 		SetComponentTickEnabled(false);
 	}
 
-#if WITH_APEX_CLOTHING
+#if WITH_APEX_CLOTHING || WITH_CHAOS_CLOTHING
 	
 	// If we don't have a valid simulation factory - check to see if we have an available default to use instead
 	if(*ClothingSimulationFactory == nullptr)
 	{
 		TArray<IClothingSimulationFactoryClassProvider*> ClassProviders = IModularFeatures::Get().GetModularFeatureImplementations<IClothingSimulationFactoryClassProvider>(IClothingSimulationFactoryClassProvider::FeatureName);
-		if(ClassProviders.Num() > 0)
+		for (int32 Index = ClassProviders.Num() - 1; Index >= 0 && !*ClothingSimulationFactory; --Index)  // We use the last provider in the list so plugins/modules can override ours
 		{
-			// We use the last provider in the list so plugins/modules can override ours
-			IClothingSimulationFactoryClassProvider* Provider = ClassProviders.Last();
+			IClothingSimulationFactoryClassProvider* const Provider = ClassProviders[Index];
 			check(Provider);
 
 			ClothingSimulationFactory = Provider->GetDefaultSimulationFactoryClass();
@@ -802,7 +798,7 @@ bool USkeletalMeshComponent::InitializeAnimScriptInstance(bool bForceReinit, boo
 
 bool USkeletalMeshComponent::IsWindEnabled() const
 {
-#if WITH_APEX_CLOTHING
+#if WITH_APEX_CLOTHING || WITH_CHAOS_CLOTHING
 	// Wind is enabled in game worlds
 	return GetWorld() && GetWorld()->IsGameWorld();
 #else
@@ -2542,9 +2538,9 @@ FBoxSphereBounds USkeletalMeshComponent::CalcBounds(const FTransform& LocalToWor
 			NewBounds = NewBounds + FBoxSphereBounds(ComponentLocation, FVector(1.0f), 1.0f);
 		}
 
-#if WITH_APEX_CLOTHING
+#if WITH_APEX_CLOTHING || WITH_CHAOS_CLOTHING
 		AddClothingBounds(NewBounds, LocalToWorld);
-#endif// #if WITH_APEX_CLOTHING
+#endif// #if WITH_APEX_CLOTHING || WITH_CHAOS_CLOTHING
 
 		bCachedLocalBoundsUpToDate = true;
 		CachedWorldSpaceBounds = NewBounds;
@@ -2599,7 +2595,7 @@ void USkeletalMeshComponent::SetSkeletalMesh(USkeletalMesh* InSkelMesh, bool bRe
 
 		InitAnim(bReinitPose);
 
-#if WITH_APEX_CLOTHING
+#if WITH_APEX_CLOTHING || WITH_CHAOS_CLOTHING
 		RecreateClothingActors();
 #endif
 	}
diff --git a/Engine/Source/Runtime/Engine/Private/Components/SkinnedMeshComponent.cpp b/Engine/Source/Runtime/Engine/Private/Components/SkinnedMeshComponent.cpp
--- a/Engine/Source/Runtime/Engine/Private/Components/SkinnedMeshComponent.cpp
+++ b/Engine/Source/Runtime/Engine/Private/Components/SkinnedMeshComponent.cpp
@@ -1457,6 +1457,11 @@ void USkinnedMeshComponent::SetSkeletalMesh(USkeletalMesh* InSkelMesh, bool bRei
 
 		SkeletalMesh = InSkelMesh;
 
+		for (TWeakObjectPtr<USkinnedMeshComponent>& SlavePoseComponent : SlavePoseComponents)
+		{
+			SlavePoseComponent->UpdateMasterBoneMap();
+		}
+
 		// Don't init anim state if not registered
 		if (IsRegistered())
 		{
diff --git a/Engine/Source/Runtime/Engine/Private/LevelTick.cpp b/Engine/Source/Runtime/Engine/Private/LevelTick.cpp
--- a/Engine/Source/Runtime/Engine/Private/LevelTick.cpp
+++ b/Engine/Source/Runtime/Engine/Private/LevelTick.cpp
@@ -42,7 +42,7 @@
 #include "TimerManager.h"
 #include "Camera/CameraPhotography.h"
 #include "HAL/LowLevelMemTracker.h"
-#if WITH_CHAOS
+#if ENABLE_COLLISION_ANALYZER
 #include "PhysicsEngine/CollisionAnalyzerCapture.h"
 #endif
 
@@ -1241,7 +1241,6 @@ public:
 #if ENABLE_COLLISION_ANALYZER
 #include "ICollisionAnalyzer.h"
 #include "CollisionAnalyzerModule.h"
-extern bool GCollisionAnalyzerIsRecording;
 #endif // ENABLE_COLLISION_ANALYZER
 
 
diff --git a/Engine/Source/Runtime/Engine/Private/PhysicsEngine/BodyInstance.cpp b/Engine/Source/Runtime/Engine/Private/PhysicsEngine/BodyInstance.cpp
--- a/Engine/Source/Runtime/Engine/Private/PhysicsEngine/BodyInstance.cpp
+++ b/Engine/Source/Runtime/Engine/Private/PhysicsEngine/BodyInstance.cpp
@@ -37,6 +37,12 @@
 #include "PxShape.h"
 #endif // WITH_PHYSX
 
+#if WITH_CHAOS
+#include "Chaos/ImplicitObject.h"
+#endif
+
+
+
 #define LOCTEXT_NAMESPACE "BodyInstance"
 
 #include "Components/ModelComponent.h"
@@ -2538,7 +2544,7 @@ void FBodyInstance::UpdateMassProperties()
 #if WITH_PHYSX
 	if (FPhysicsInterface::IsValid(ActorHandle) && FPhysicsInterface::IsRigidBody(ActorHandle))
 	{
-		FPhysicsCommand::ExecuteWrite(ActorHandle, [&](const FPhysicsActorHandle& Actor)
+		FPhysicsCommand::ExecuteWrite(ActorHandle, [&](FPhysicsActorHandle& Actor)
 		{
 			check(FPhysicsInterface::IsValid(Actor));
 
@@ -3086,11 +3092,18 @@ bool FBodyInstance::OverlapTestForBodiesImpl(const FVector& Pos, const FQuat& Ro
 
 		for(const FPhysicsShapeHandle& Shape : TargetShapes)
 		{
+#if WITH_CHAOS
+			if (!Shape.GetGeometry().IsConvex())
+			{
+				continue;	//we skip complex shapes - should this respect ComplexAsSimple?
+			}
+#else
 			ECollisionShapeType ShapeType = FPhysicsInterface::GetShapeType(Shape);
 			if(ShapeType == ECollisionShapeType::Heightfield || ShapeType == ECollisionShapeType::Trimesh)
 			{
 				continue;	//we skip complex shapes - should this respect ComplexAsSimple?
 			}
+#endif
 
 			// Calc shape global pose
 			FTransform PShapeGlobalPose = FPhysicsInterface::GetLocalTransform(Shape) * PTestGlobalPose;
@@ -3211,11 +3224,18 @@ bool FBodyInstance::OverlapMulti(TArray<struct FOverlapResult>& InOutOverlaps, c
 					continue;
 				}
 
+#if WITH_CHAOS
+				if (!ShapeRef.GetGeometry().IsConvex())
+				{
+					continue;	//we skip complex shapes - should this respect ComplexAsSimple?
+				}
+#else
 				ECollisionShapeType ShapeType = FPhysicsInterface::GetShapeType(ShapeRef);
-				if(ShapeType == ECollisionShapeType::Heightfield || ShapeType == ECollisionShapeType::Trimesh)
+				if (ShapeType == ECollisionShapeType::Heightfield || ShapeType == ECollisionShapeType::Trimesh)
 				{
 					continue;	//we skip complex shapes - should this respect ComplexAsSimple?
 				}
+#endif
 
 				// Calc shape global pose
 				const FTransform LocalTransform = FPhysicsInterface::GetLocalTransform(ShapeRef);
diff --git a/Engine/Source/Runtime/Engine/Private/PhysicsEngine/CollisionAnalyzerCapture.h b/Engine/Source/Runtime/Engine/Private/PhysicsEngine/CollisionAnalyzerCapture.h
--- a/Engine/Source/Runtime/Engine/Private/PhysicsEngine/CollisionAnalyzerCapture.h
+++ b/Engine/Source/Runtime/Engine/Private/PhysicsEngine/CollisionAnalyzerCapture.h
@@ -2,6 +2,8 @@
 
 #pragma once
 
+#include "Collision.h"
+
 #if ENABLE_COLLISION_ANALYZER
 
 #include "ICollisionAnalyzer.h"
@@ -9,8 +11,11 @@
 #include "Physics/PhysicsInterfaceCore.h"
 #include "Physics/PhysicsInterfaceUtils.h"
 
-bool bSkipCapture = false;
-bool GCollisionAnalyzerIsRecording = false;
+#include "Chaos/Capsule.h"
+#include "Chaos/Sphere.h"
+
+extern bool bSkipCapture;
+extern bool GCollisionAnalyzerIsRecording;
 
 /** Util to convert from PhysX shape and rotation to unreal shape enum, dimension vector and rotation */
 inline static void ConvertGeometryCollection(const FPhysicsGeometryCollection& GeomCollection, FQuat& InOutRot, ECAQueryShape::Type& OutQueryShape, FVector& OutDims)
diff --git a/Engine/Source/Runtime/Engine/Private/PhysicsEngine/CollisionQueryFilterCallback.cpp b/Engine/Source/Runtime/Engine/Private/PhysicsEngine/CollisionQueryFilterCallback.cpp
--- a/Engine/Source/Runtime/Engine/Private/PhysicsEngine/CollisionQueryFilterCallback.cpp
+++ b/Engine/Source/Runtime/Engine/Private/PhysicsEngine/CollisionQueryFilterCallback.cpp
@@ -12,10 +12,8 @@
 #include "PhysXInterfaceWrapper.h"
 #endif
 
-#if INCLUDE_CHAOS
 #include "ChaosInterfaceWrapperCore.h"
 #include "Physics/Experimental/ChaosInterfaceWrapper.h"
-#endif
 
 ECollisionQueryHitType FCollisionQueryFilterCallback::CalcQueryHitType(const FCollisionFilterData& QueryFilter, const FCollisionFilterData& ShapeFilter, bool bPreFilter)
 {
@@ -110,7 +108,6 @@ ECollisionQueryHitType FCollisionQueryFilterCallback::PreFilterImp(const FCollis
 }
 #endif
 
-#if INCLUDE_CHAOS
 ECollisionQueryHitType FCollisionQueryFilterCallback::PreFilterImp(const FCollisionFilterData& FilterData, const Chaos::TPerShapeData<float,3>& Shape, const Chaos::TGeometryParticle<float,3>& Actor)
 {
 	SCOPE_CYCLE_COUNTER(STAT_Collision_PreFilter);
@@ -130,7 +127,6 @@ ECollisionQueryHitType FCollisionQueryFilterCallback::PreFilterImp(const FCollis
 
 	return PreFilterImp(FilterData, ShapeFilter, ComponentID, BodyInstance);
 }
-#endif
 
 ECollisionQueryHitType FCollisionQueryFilterCallback::PreFilterImp(const FCollisionFilterData& FilterData, const FCollisionFilterData& ShapeFilter, uint32 ComponentID, const FBodyInstance* BodyInstance)
 {
@@ -271,7 +267,6 @@ ECollisionQueryHitType FCollisionQueryFilterCallback::PostFilterImp(const FColli
 }
 #endif
 
-#if INCLUDE_CHAOS
 ECollisionQueryHitType FCollisionQueryFilterCallback::PostFilterImp(const FCollisionFilterData& FilterData, const ChaosInterface::FQueryHit& Hit)
 {
 	// Unused in non-sweeps
@@ -285,8 +280,6 @@ ECollisionQueryHitType FCollisionQueryFilterCallback::PostFilterImp(const FColli
 
 	return PostFilterImp(FilterData, bIsOverlap);
 }
-#endif
-
 
 #if WITH_PHYSX
 
diff --git a/Engine/Source/Runtime/Engine/Private/PhysicsEngine/Experimental/ChaosDerivedData.cpp b/Engine/Source/Runtime/Engine/Private/PhysicsEngine/Experimental/ChaosDerivedData.cpp
--- a/Engine/Source/Runtime/Engine/Private/PhysicsEngine/Experimental/ChaosDerivedData.cpp
+++ b/Engine/Source/Runtime/Engine/Private/PhysicsEngine/Experimental/ChaosDerivedData.cpp
@@ -21,7 +21,7 @@ const TCHAR* FChaosDerivedDataCooker::GetPluginName() const
 
 const TCHAR* FChaosDerivedDataCooker::GetVersionString() const
 {
-	return TEXT("19C69FC43DDA4F058B28C21F08D623F2");
+	return TEXT("UN4SF7NU5G9O0X2LHNHLBG0R0CIS3TON");
 }
 
 FString FChaosDerivedDataCooker::GetPluginSpecificCacheKeySuffix() const
@@ -99,8 +99,9 @@ void FChaosDerivedDataCooker::BuildTriangleMeshes(TArray<TUniquePtr<Chaos::TTria
 	FinalIndices.Reserve(InParams.TriangleMeshDesc.Indices.Num() * 3);
 	for(const FTriIndices& Tri : InParams.TriangleMeshDesc.Indices)
 	{
-		FinalIndices.Add(Tri.v0);
+		//question: It seems like unreal triangles are CW, but couldn't find confirmation for this
 		FinalIndices.Add(Tri.v1);
+		FinalIndices.Add(Tri.v0);
 		FinalIndices.Add(Tri.v2);
 	}
 
@@ -123,8 +124,19 @@ void FChaosDerivedDataCooker::BuildTriangleMeshes(TArray<TUniquePtr<Chaos::TTria
 
 	for(int32 TriangleIndex = 0; TriangleIndex < NumTriangles; ++TriangleIndex)
 	{
+		// Only add this triangle if it is valid
 		const int32 BaseIndex = TriangleIndex * 3;
-		Triangles.Add(Chaos::TVector<int32, 3>(FinalIndices[BaseIndex], FinalIndices[BaseIndex + 1], FinalIndices[BaseIndex + 2]));
+		const bool bIsValidTriangle = Chaos::TConvexBuilder<Precision>::IsValidTriangle(
+			FinalVerts[FinalIndices[BaseIndex]],
+			FinalVerts[FinalIndices[BaseIndex + 1]],
+			FinalVerts[FinalIndices[BaseIndex + 2]]);
+
+		// TODO: Figure out a proper way to handle this. Could these edges get sewn together? Is this important?
+		//if (ensureMsgf(bIsValidTriangle, TEXT("FChaosDerivedDataCooker::BuildTriangleMeshes(): Trimesh attempted cooked with invalid triangle!")));
+		if (bIsValidTriangle)
+		{
+			Triangles.Add(Chaos::TVector<int32, 3>(FinalIndices[BaseIndex], FinalIndices[BaseIndex + 1], FinalIndices[BaseIndex + 2]));
+		}
 	}
 
 	OutTriangleMeshes.Emplace(new Chaos::TTriangleMeshImplicitObject<Precision>(MoveTemp(TriMeshParticles), MoveTemp(Triangles)));
@@ -137,7 +149,7 @@ template void FChaosDerivedDataCooker::BuildTriangleMeshes(TArray<TUniquePtr<Cha
 template<typename Precision>
 void FChaosDerivedDataCooker::BuildConvexMeshes(TArray<TUniquePtr<Chaos::TImplicitObject<Precision, 3>>>& OutConvexMeshes, const FCookBodySetupInfo& InParams)
 {
-	auto BuildConvexFromVerts = [](TArray<TUniquePtr<Chaos::TImplicitObject<Precision, 3>>>& OutConvexes, const TArray<TArray<FVector>>& InMeshVerts)
+	auto BuildConvexFromVerts = [](TArray<TUniquePtr<Chaos::TImplicitObject<Precision, 3>>>& OutConvexes, const TArray<TArray<FVector>>& InMeshVerts, const bool bMirrored)
 	{
 		for(const TArray<FVector>& HullVerts : InMeshVerts)
 		{
@@ -149,7 +161,8 @@ void FChaosDerivedDataCooker::BuildConvexMeshes(TArray<TUniquePtr<Chaos::TImplic
 
 			for(int32 VertIndex = 0; VertIndex < NumHullVerts; ++VertIndex)
 			{
-				ConvexParticles.X(VertIndex) = HullVerts[VertIndex];
+				const FVector& HullVert = HullVerts[VertIndex];
+				ConvexParticles.X(VertIndex) = FVector(bMirrored ? -HullVert.X : HullVert.X, HullVert.Y, HullVert.Z);
 			}
 
 			OutConvexes.Emplace(new Chaos::TConvex<Precision, 3>(ConvexParticles));
@@ -158,12 +171,12 @@ void FChaosDerivedDataCooker::BuildConvexMeshes(TArray<TUniquePtr<Chaos::TImplic
 
 	if(InParams.bCookNonMirroredConvex)
 	{
-		BuildConvexFromVerts(OutConvexMeshes, InParams.NonMirroredConvexVertices);
+		BuildConvexFromVerts(OutConvexMeshes, InParams.NonMirroredConvexVertices, false);
 	}
 
 	if(InParams.bCookMirroredConvex)
 	{
-		BuildConvexFromVerts(OutConvexMeshes, InParams.MirroredConvexVertices);
+		BuildConvexFromVerts(OutConvexMeshes, InParams.MirroredConvexVertices, true);
 	}
 
 }
diff --git a/Engine/Source/Runtime/Engine/Private/PhysicsEngine/Experimental/ChaosInterfaceUtils.cpp b/Engine/Source/Runtime/Engine/Private/PhysicsEngine/Experimental/ChaosInterfaceUtils.cpp
new file mode 100644
--- /dev/null
+++ b/Engine/Source/Runtime/Engine/Private/PhysicsEngine/Experimental/ChaosInterfaceUtils.cpp
@@ -0,0 +1,270 @@
+// Copyright 1998-2019 Epic Games, Inc. All Rights Reserved.
+
+#include "Physics/Experimental/ChaosInterfaceUtils.h"
+
+#include "Chaos/Box.h"
+#include "Chaos/Capsule.h"
+#include "Chaos/GeometryParticles.h"
+#include "Chaos/ImplicitObject.h"
+#include "Chaos/ImplicitObjectScaled.h"
+#include "Chaos/Levelset.h"
+#include "Chaos/Sphere.h"
+#include "Chaos/TriangleMesh.h"
+#include "Chaos/TriangleMeshImplicitObject.h"
+#include "Chaos/UniformGrid.h"
+
+#include "PhysicsEngine/BoxElem.h"
+#include "PhysicsEngine/ConvexElem.h"
+#include "PhysicsEngine/SphylElem.h"
+#include "PhysicsEngine/SphereElem.h"
+
+#include "Physics/PhysicsInterfaceTypes.h"
+#include "PhysicsEngine/AggregateGeom.h"
+
+#if WITH_PHYSX
+#include "foundation/PxSimpleTypes.h"
+#include "geometry/PxConvexMesh.h"
+#include "geometry/PxTriangleMesh.h"
+#endif
+
+#define FORCE_ANALYTICS 0
+
+namespace ChaosInterface
+{
+	template<class PHYSX_MESH>
+	TArray<Chaos::TVector<int32, 3>> GetMeshElements(const PHYSX_MESH* PhysXMesh)
+	{
+		check(false);
+	}
+
+#if WITH_PHYSX
+
+	template<>
+	TArray<Chaos::TVector<int32, 3>> GetMeshElements(const physx::PxConvexMesh* PhysXMesh)
+	{
+		TArray<Chaos::TVector<int32, 3>> CollisionMeshElements;
+#if !WITH_CHAOS_NEEDS_TO_BE_FIXED
+		int32 offset = 0;
+		int32 NbPolygons = static_cast<int32>(PhysXMesh->getNbPolygons());
+		for (int32 i = 0; i < NbPolygons; i++)
+		{
+			physx::PxHullPolygon Poly;
+			bool status = PhysXMesh->getPolygonData(i, Poly);
+			const auto Indices = PhysXMesh->getIndexBuffer() + Poly.mIndexBase;
+
+			for (int32 j = 2; j < static_cast<int32>(Poly.mNbVerts); j++)
+			{
+				CollisionMeshElements.Add(Chaos::TVector<int32, 3>(Indices[offset], Indices[offset + j], Indices[offset + j - 1]));
+			}
+		}
+#endif
+		return CollisionMeshElements;
+	}
+
+	template<>
+	TArray<Chaos::TVector<int32, 3>> GetMeshElements(const physx::PxTriangleMesh* PhysXMesh)
+	{
+		TArray<Chaos::TVector<int32, 3>> CollisionMeshElements;
+		const auto MeshFlags = PhysXMesh->getTriangleMeshFlags();
+		for (int32 j = 0; j < static_cast<int32>(PhysXMesh->getNbTriangles()); ++j)
+		{
+			if (MeshFlags | physx::PxTriangleMeshFlag::e16_BIT_INDICES)
+			{
+				const physx::PxU16* Indices = reinterpret_cast<const physx::PxU16*>(PhysXMesh->getTriangles());
+				CollisionMeshElements.Add(Chaos::TVector<int32, 3>(Indices[3 * j], Indices[3 * j + 1], Indices[3 * j + 2]));
+			}
+			else
+			{
+				const physx::PxU32* Indices = reinterpret_cast<const physx::PxU32*>(PhysXMesh->getTriangles());
+				CollisionMeshElements.Add(Chaos::TVector<int32, 3>(Indices[3 * j], Indices[3 * j + 1], Indices[3 * j + 2]));
+			}
+		}
+		return CollisionMeshElements;
+	}
+
+	template<class PHYSX_MESH>
+	TUniquePtr<Chaos::TImplicitObject<float, 3>> ConvertPhysXMeshToLevelset(const PHYSX_MESH* PhysXMesh, const FVector& Scale)
+	{
+#if WITH_CHAOS && !WITH_CHAOS_NEEDS_TO_BE_FIXED
+		TArray<Chaos::TVector<int32, 3>> CollisionMeshElements = GetMeshElements(PhysXMesh);
+		Chaos::TParticles<float, 3> CollisionMeshParticles;
+		CollisionMeshParticles.AddParticles(PhysXMesh->getNbVertices());
+		for (uint32 j = 0; j < CollisionMeshParticles.Size(); ++j)
+		{
+			const auto& Vertex = PhysXMesh->getVertices()[j];
+			CollisionMeshParticles.X(j) = Scale * Chaos::TVector<float, 3>(Vertex.x, Vertex.y, Vertex.z);
+		}
+		Chaos::TBox<float, 3> BoundingBox(CollisionMeshParticles.X(0), CollisionMeshParticles.X(0));
+		for (uint32 j = 1; j < CollisionMeshParticles.Size(); ++j)
+		{
+			BoundingBox.GrowToInclude(CollisionMeshParticles.X(j));
+		}
+#if FORCE_ANALYTICS
+		return TUniquePtr<Chaos::TImplicitObject<float, 3>>(new Chaos::TBox<float, 3>(BoundingBox));
+#else
+		int32 MaxAxisSize = 10;
+		int32 MaxAxis;
+		const auto Extents = BoundingBox.Extents();
+		if (Extents[0] > Extents[1] && Extents[0] > Extents[2])
+		{
+			MaxAxis = 0;
+		}
+		else if (Extents[1] > Extents[2])
+		{
+			MaxAxis = 1;
+		}
+		else
+		{
+			MaxAxis = 2;
+		}
+		Chaos::TVector<int32, 3> Counts(MaxAxisSize * Extents[0] / Extents[MaxAxis], MaxAxisSize * Extents[1] / Extents[MaxAxis], MaxAxisSize * Extents[2] / Extents[MaxAxis]);
+		Counts[0] = Counts[0] < 1 ? 1 : Counts[0];
+		Counts[1] = Counts[1] < 1 ? 1 : Counts[1];
+		Counts[2] = Counts[2] < 1 ? 1 : Counts[2];
+		Chaos::TUniformGrid<float, 3> Grid(BoundingBox.Min(), BoundingBox.Max(), Counts, 1);
+		Chaos::TTriangleMesh<float> CollisionMesh(MoveTemp(CollisionMeshElements));
+		return TUniquePtr<Chaos::TImplicitObject<float, 3>>(new Chaos::TLevelSet<float, 3>(Grid, CollisionMeshParticles, CollisionMesh));
+#endif
+
+#else
+		return TUniquePtr<Chaos::TImplicitObject<float, 3>>();
+#endif // !WITH_CHAOS_NEEDS_TO_BE_FIXED
+
+	}
+
+#endif
+
+
+	void CreateGeometry(const FGeometryAddParams& InParams, TArray<TUniquePtr<Chaos::TImplicitObject<float, 3>>>& OutGeoms, Chaos::TShapesArray<float, 3>& OutShapes)
+	{
+		const FVector& Scale = InParams.Scale;
+		TArray<TUniquePtr<Chaos::TImplicitObject<float, 3>>>& Geoms = OutGeoms;
+		Chaos::TShapesArray<float, 3>& Shapes = OutShapes;
+
+		auto NewShapeHelper = [&InParams](Chaos::TSerializablePtr<Chaos::TImplicitObject<float, 3>> InGeom, bool bComplexShape = false)
+		{
+			auto NewShape = Chaos::TPerShapeData<float, 3>::CreatePerShapeData();
+			NewShape->Geometry = InGeom;
+			NewShape->QueryData = bComplexShape ? InParams.CollisionData.CollisionFilterData.QueryComplexFilter : InParams.CollisionData.CollisionFilterData.QuerySimpleFilter;
+			NewShape->SimData = InParams.CollisionData.CollisionFilterData.SimFilter;
+			return NewShape;
+		};
+
+		for (uint32 i = 0; i < static_cast<uint32>(InParams.Geometry->SphereElems.Num()); ++i)
+		{
+			const FKSphereElem ScaledSphereElem = InParams.Geometry->SphereElems[i].GetFinalScaled(Scale, InParams.LocalTransform);
+			const float UseRadius = FMath::Max(ScaledSphereElem.Radius, KINDA_SMALL_NUMBER);
+
+			auto ImplicitSphere = MakeUnique<Chaos::TSphere<float, 3>>(Chaos::TVector<float, 3>(0.f, 0.f, 0.f), UseRadius);
+			auto NewShape = NewShapeHelper(MakeSerializable(ImplicitSphere));
+			Shapes.Emplace(MoveTemp(NewShape));
+			Geoms.Add(MoveTemp(ImplicitSphere));
+
+		}
+		for (uint32 i = 0; i < static_cast<uint32>(InParams.Geometry->BoxElems.Num()); ++i)
+		{
+			const FKBoxElem ScaledBoxElem = InParams.Geometry->BoxElems[i].GetFinalScaled(Scale, InParams.LocalTransform);
+			const FTransform& BoxTransform = ScaledBoxElem.GetTransform();
+			Chaos::TVector<float, 3> HalfExtents = Chaos::TVector<float, 3>(ScaledBoxElem.X * 0.5f, ScaledBoxElem.Y * 0.5f, ScaledBoxElem.Z * 0.5f);
+
+			HalfExtents.X = FMath::Max(HalfExtents.X, KINDA_SMALL_NUMBER);
+			HalfExtents.Y = FMath::Max(HalfExtents.Y, KINDA_SMALL_NUMBER);
+			HalfExtents.Z = FMath::Max(HalfExtents.Z, KINDA_SMALL_NUMBER);
+
+			//todo: properly support transform on all geometries, for now just fixing a specific broken case with a translated aabb
+	//		ensure(BoxTransform.GetRotation() == FQuat::Identity);
+			TUniquePtr<Chaos::TBox<float, 3>> Box;
+			auto ImplicitBox = MakeUnique<Chaos::TBox<float, 3>>(BoxTransform.GetTranslation() - HalfExtents, BoxTransform.GetTranslation() + HalfExtents);
+			auto NewShape = NewShapeHelper(MakeSerializable(ImplicitBox));
+			Shapes.Emplace(MoveTemp(NewShape));
+			Geoms.Add(MoveTemp(ImplicitBox));
+		}
+		for (uint32 i = 0; i < static_cast<uint32>(InParams.Geometry->SphylElems.Num()); ++i)
+		{
+			const FKSphylElem& UnscaledSphyl = InParams.Geometry->SphylElems[i];
+
+			const FKSphylElem ScaledSphylElem = UnscaledSphyl.GetFinalScaled(Scale, InParams.LocalTransform);
+			float HalfHeight = FMath::Max(ScaledSphylElem.Length * 0.5f, KINDA_SMALL_NUMBER);
+			const float Radius = FMath::Max(ScaledSphylElem.Radius, KINDA_SMALL_NUMBER);
+
+			if (HalfHeight < KINDA_SMALL_NUMBER)
+			{
+				//not a capsule just use a sphere
+				auto ImplicitSphere = MakeUnique<Chaos::TSphere<float, 3>>(Chaos::TVector<float, 3>(0), Radius);
+				auto NewShape = NewShapeHelper(MakeSerializable(ImplicitSphere));
+				Shapes.Emplace(MoveTemp(NewShape));
+				Geoms.Add(MoveTemp(ImplicitSphere));
+
+			}
+			else
+			{
+				Chaos::TVector<float, 3> HalfExtents(0, 0, HalfHeight);
+
+				auto ImplicitCapsule = MakeUnique<Chaos::TCapsule<float>>(-HalfExtents, HalfExtents, Radius);
+				auto NewShape = NewShapeHelper(MakeSerializable(ImplicitCapsule));
+				Shapes.Emplace(MoveTemp(NewShape));
+				Geoms.Add(MoveTemp(ImplicitCapsule));
+			}
+		}
+#if 0
+		for (uint32 i = 0; i < static_cast<uint32>(InParams.Geometry->TaperedCapsuleElems.Num()); ++i)
+		{
+			ensure(FMath::IsNearlyEqual(Scale[0], Scale[1]) && FMath::IsNearlyEqual(Scale[1], Scale[2]));
+			const auto& TCapsule = InParams.Geometry->TaperedCapsuleElems[i];
+			if (TCapsule.Length == 0)
+			{
+				Chaos::TSphere<float, 3> * ImplicitSphere = new Chaos::TSphere<float, 3>(-half_extents, TCapsule.Radius * Scale[0]);
+				if (PhysicsProxy) PhysicsProxy->ImplicitObjects_GameThread.Add(ImplicitSphere);
+				else if (OutOptShapes) OutOptShapes->Add({ ImplicitSphere,true,true,InActor });
+			}
+			else
+			{
+				Chaos::TVector<float, 3> half_extents(0, 0, TCapsule.Length / 2 * Scale[0]);
+				auto ImplicitCylinder = MakeUnique<Chaos::TCylinder<float>>(-half_extents, half_extents, TCapsule.Radius * Scale[0]);
+				if (PhysicsProxy) PhysicsProxy->ImplicitObjects_GameThread.Add(MoveTemp(ImplicitSphere));
+				else if (OutOptShapes) OutOptShapes->Add({ ImplicitSphere,true,true,InActor });
+
+				auto ImplicitSphereA = MakeUnique<Chaos::TSphere<float, 3>>(-half_extents, TCapsule.Radius * Scale[0]);
+				if (PhysicsProxy) PhysicsProxy->ImplicitObjects_GameThread.Add(MoveTemp(ImplicitSphereA));
+				else if (OutOptShapes) OutOptShapes->Add({ ImplicitSphereA,true,true,InActor });
+
+				auto ImplicitSphereB = MakeUnique<Chaos::TSphere<float, 3>>(half_extents, TCapsule.Radius * Scale[0]);
+				if (PhysicsProxy) PhysicsProxy->ImplicitObjects_GameThread.Add(MoveTemp(ImplicitSphereB));
+				else if (OutOptShapes) OutOptShapes->Add({ ImplicitSphereB,true,true,InActor });
+			}
+		}
+#endif
+#if WITH_CHAOS && !PHYSICS_INTERFACE_PHYSX
+		for (uint32 i = 0; i < static_cast<uint32>(InParams.Geometry->ConvexElems.Num()); ++i)
+		{
+			const FKConvexElem& CollisionBody = InParams.Geometry->ConvexElems[i];
+			if (const auto& ConvexImplicit = CollisionBody.GetChaosConvexMesh())
+			{
+				auto Implicit = MakeUnique<Chaos::TImplicitObjectScaled<float, 3>>(MakeSerializable(ConvexImplicit), Scale);
+				auto NewShape = NewShapeHelper(MakeSerializable(Implicit));
+				Shapes.Emplace(MoveTemp(NewShape));
+				Geoms.Add(MoveTemp(Implicit));
+			}
+		}
+
+		for (const auto& ChaosTriMesh : InParams.ChaosTriMeshes)
+		{
+			auto Implicit = MakeUnique<Chaos::TImplicitObjectScaled<float, 3>>(MakeSerializable(ChaosTriMesh), Scale);
+			auto NewShape = NewShapeHelper(MakeSerializable(Implicit), true);
+			Shapes.Emplace(MoveTemp(NewShape));
+			Geoms.Add(MoveTemp(Implicit));
+		}
+#endif
+#if WITH_PHYSX && PHYSICS_INTERFACE_PHYSX
+		for (const auto& PhysXMesh : InParams.TriMeshes)
+		{
+			auto Implicit = ConvertPhysXMeshToLevelset(PhysXMesh, Scale);
+			auto NewShape = NewShapeHelper(MakeSerializable(Implicit), true);
+			Shapes.Emplace(MoveTemp(NewShape));
+			Geoms.Add(MoveTemp(Implicit));
+
+		}
+#endif
+	}
+
+}
\ No newline at end of file
diff --git a/Engine/Source/Runtime/Engine/Private/PhysicsEngine/Experimental/ChaosInterfaceWrapper.cpp b/Engine/Source/Runtime/Engine/Private/PhysicsEngine/Experimental/ChaosInterfaceWrapper.cpp
new file mode 100644
--- /dev/null
+++ b/Engine/Source/Runtime/Engine/Private/PhysicsEngine/Experimental/ChaosInterfaceWrapper.cpp
@@ -0,0 +1,34 @@
+// Copyright 1998-2019 Epic Games, Inc. All Rights Reserved.
+
+#include "Physics/Experimental/ChaosInterfaceWrapper.h"
+#include "Physics/Experimental/PhysScene_Chaos.h"
+#include "Chaos/ParticleHandle.h"
+#include "PhysxUserData.h"
+
+
+namespace ChaosInterface
+{
+	FBodyInstance* GetUserData(const Chaos::TGeometryParticle<float,3>& Actor)
+	{
+		void* UserData = Actor.UserData();
+		return UserData ? FPhysxUserData::Get<FBodyInstance>(Actor.UserData()) : nullptr;
+	}
+
+	UPhysicalMaterial* GetUserData(const FPhysTypeDummy& Material)
+	{
+		return nullptr;
+	}
+
+#if WITH_CHAOS
+	FScopedSceneReadLock::FScopedSceneReadLock(FPhysScene_ChaosInterface& SceneIn)
+		: Scene(SceneIn.GetScene())
+	{
+		Scene.ExternalDataLock.ReadLock();
+	}
+
+	FScopedSceneReadLock::~FScopedSceneReadLock()
+	{
+		Scene.ExternalDataLock.ReadUnlock();
+	}
+#endif
+}
diff --git a/Engine/Source/Runtime/Engine/Private/PhysicsEngine/Experimental/PhysInterface_Chaos.cpp b/Engine/Source/Runtime/Engine/Private/PhysicsEngine/Experimental/PhysInterface_Chaos.cpp
--- a/Engine/Source/Runtime/Engine/Private/PhysicsEngine/Experimental/PhysInterface_Chaos.cpp
+++ b/Engine/Source/Runtime/Engine/Private/PhysicsEngine/Experimental/PhysInterface_Chaos.cpp
@@ -4,11 +4,13 @@
 
 #include "Physics/Experimental/PhysInterface_Chaos.h"
 #include "Physics/Experimental/PhysScene_Chaos.h"
+#include "Physics/Experimental/ChaosInterfaceUtils.h"
 #include "Physics/PhysicsInterfaceTypes.h"
 #include "PhysicsEngine/AggregateGeom.h"
 
 #include "Chaos/Box.h"
 #include "Chaos/Cylinder.h"
+#include "Chaos/TaperedCylinder.h"
 #include "Chaos/Capsule.h"
 #include "Chaos/ImplicitObjectTransformed.h"
 #include "Chaos/ImplicitObjectUnion.h"
@@ -16,6 +18,7 @@
 #include "Chaos/PBDRigidParticles.h"
 #include "Chaos/Sphere.h"
 #include "Chaos/Matrix.h"
+#include "Chaos/MassProperties.h"
 #include "PhysicsSolver.h"
 #include "Templates/UniquePtr.h"
 #include "ChaosSolversModule.h"
@@ -25,8 +28,7 @@
 #include "Chaos/Convex.h"
 #include "Chaos/GeometryQueries.h"
 #include "Chaos/Plane.h"
-
-
+#include "ChaosCheck.h"
 
 #include "Async/ParallelFor.h"
 #include "Components/PrimitiveComponent.h"
@@ -34,7 +36,6 @@
 #include "Collision/CollisionConversions.h"
 #include "PhysicsInterfaceUtilsCore.h"
 
-
 #if WITH_PHYSX
 #include "geometry/PxConvexMesh.h"
 #include "geometry/PxTriangleMesh.h"
@@ -42,10 +43,6 @@
 #include "extensions/PxMassProperties.h"
 #endif
 
-
-
-#define FORCE_ANALYTICS 0
-
 DEFINE_STAT(STAT_TotalPhysicsTime);
 DEFINE_STAT(STAT_NumCloths);
 DEFINE_STAT(STAT_NumClothVerts);
@@ -80,6 +77,16 @@ DECLARE_DWORD_COUNTER_STAT(TEXT("(ASync) Mobile Bodies"), STAT_NumMobileBodiesAs
 DECLARE_DWORD_COUNTER_STAT(TEXT("(ASync) Static Bodies"), STAT_NumStaticBodiesAsync, STATGROUP_Physics);
 DECLARE_DWORD_COUNTER_STAT(TEXT("(ASync) Shapes"), STAT_NumShapesAsync, STATGROUP_Physics);
 
+ECollisionShapeType GetGeometryType(const Chaos::TPerShapeData<float, 3>& Shape)
+{
+	return GetType(*Shape.Geometry);
+}
+
+const Chaos::TImplicitObject<float, 3>& FPhysicsShapeReference_Chaos::GetGeometry() const
+{
+	check(IsValid()); return *Shape->Geometry;
+}
+
 FPhysInterface_Chaos::FPhysInterface_Chaos(const AWorldSettings* Settings) 
 {
 
@@ -93,18 +100,22 @@ FPhysInterface_Chaos::~FPhysInterface_Chaos()
 // Interface functions
 void FPhysInterface_Chaos::CreateActor(const FActorCreationParams& InParams, FPhysicsActorHandle& Handle)
 {
+	LLM_SCOPE(ELLMTag::Chaos);
+
 	// Create the new particle
 	if (InParams.bStatic)
 	{
-		Handle = new Chaos::TGeometryParticle<float, 3>();
+		Handle = Chaos::TGeometryParticle<float, 3>::CreateParticle().Release();	//todo: should BodyInstance use a unique ptr to manage this memory?
 	}
 	else if (InParams.BodyInstance && InParams.BodyInstance->ShouldInstanceSimulatingPhysics())
 	{
-		Handle = new Chaos::TPBDRigidParticle<float, 3>();
+		TUniquePtr<Chaos::TPBDRigidParticle<float, 3>> RigidHandle = Chaos::TPBDRigidParticle<float, 3>::CreateParticle();
+		RigidHandle->SetGravityEnabled(InParams.bEnableGravity);
+		Handle = RigidHandle.Release();	//todo: should BodyInstance use a unique ptr to manage this memory?
 	}
 	else
 	{
-		Handle = new Chaos::TKinematicGeometryParticle<float, 3>();
+		Handle = Chaos::TKinematicGeometryParticle<float, 3>::CreateParticle().Release(); //todo: should BodyInstance use a unique ptr to manage this memory?
 	}
 
 	// Set up the new particle's game-thread data. This will be sent to physics-thread when
@@ -116,12 +127,21 @@ void FPhysInterface_Chaos::CreateActor(const FActorCreationParams& InParams, FPh
 
 void FPhysInterface_Chaos::AddActorToSolver(FPhysicsActorHandle& Handle, Chaos::FPhysicsSolver* Solver, Chaos::IDispatcher* Dispatcher)
 {
+	LLM_SCOPE(ELLMTag::Chaos);
+
 	Solver->RegisterObject(Handle);
 }
 
 void FPhysInterface_Chaos::ReleaseActor(FPhysicsActorHandle& Handle, FPhysScene* InScene, bool bNeverDerferRelease)
 {
-	check(Handle);
+	if(!Handle)
+	{
+		UE_LOG(LogChaos, Warning, TEXT("Attempting to release an actor with a null handle"));
+		CHAOS_ENSURE(false);
+
+		return;
+	}
+
 	if (InScene)
 	{
 		InScene->GetScene().RemoveActorFromAccelerationStructure(Handle);
@@ -135,7 +155,7 @@ void FPhysInterface_Chaos::ReleaseActor(FPhysicsActorHandle& Handle, FPhysScene*
 
 void FPhysInterface_Chaos::RemoveActorFromSolver(FPhysicsActorHandle& Handle, Chaos::FPhysicsSolver* Solver, Chaos::IDispatcher* Dispatcher)
 {
-	if (Solver)
+	if (Solver && Handle->Proxy)
 	{
 		Solver->UnregisterObject(Handle);
 	}
@@ -157,7 +177,7 @@ void FPhysInterface_Chaos::AddActorToAggregate_AssumesLocked(const FPhysicsAggre
 int32 FPhysInterface_Chaos::GetNumShapes(const FPhysicsActorHandle& InHandle)
 {
 	// #todo : Implement
-	return 1;
+	return InHandle->ShapesArray().Num();
 }
 
 void FPhysInterface_Chaos::ReleaseShape(const FPhysicsShapeHandle& InShape)
@@ -221,6 +241,13 @@ bool FPhysInterface_Chaos::IsInScene(const FPhysicsActorHandle& InActorReference
 
 FPhysScene* FPhysInterface_Chaos::GetCurrentScene(const FPhysicsActorHandle& InHandle)
 {
+	if(!InHandle)
+	{
+		UE_LOG(LogChaos, Warning, TEXT("Attempting to get the current scene for a null handle."));
+		CHAOS_ENSURE(false);
+		return nullptr;
+	}
+
 	if (IPhysicsProxyBase* Proxy = InHandle->Proxy)
 	{
 		Chaos::FPhysicsSolver* Solver = Proxy->GetSolver();
@@ -377,8 +404,12 @@ FTransform FPhysInterface_Chaos::GetComTransformLocal_AssumesLocked(const FPhysi
 
 FVector FPhysInterface_Chaos::GetLocalInertiaTensor_AssumesLocked(const FPhysicsActorHandle& InActorReference)
 {
-	// #todo : Implement
-	return FVector(1);
+	if (Chaos::TPBDRigidParticle<float, 3 >* RigidParticle = InActorReference->AsDynamic())
+	{
+		const Chaos::PMatrix<float, 3, 3> & Tensor = RigidParticle->I();
+		return FVector(Tensor.M[0][0], Tensor.M[1][1], Tensor.M[2][2]) ;
+	}
+	return FVector::ZeroVector;
 }
 
 FBox FPhysInterface_Chaos::GetBounds_AssumesLocked(const FPhysicsActorHandle& InActorReference)
@@ -434,12 +465,25 @@ void FPhysInterface_Chaos::AddRadialImpulse_AssumesLocked(const FPhysicsActorHan
 
 bool FPhysInterface_Chaos::IsGravityEnabled_AssumesLocked(const FPhysicsActorHandle& InActorReference)
 {
-    // @todo(mlentine): Gravity is system wide currently. This should change.
-    return true;
+	if (Chaos::TPBDRigidParticle<float, 3 >* RigidParticle = InActorReference->AsDynamic())
+	{
+		return RigidParticle->IsGravityEnabled();
+	}
+	return false;
 }
 void FPhysInterface_Chaos::SetGravityEnabled_AssumesLocked(const FPhysicsActorHandle& InActorReference, bool bEnabled)
 {
-    // @todo(mlentine): Gravity is system wide currently. This should change.
+	if (Chaos::TPBDRigidParticle<float, 3 >* RigidParticle = InActorReference->AsDynamic())
+	{
+		RigidParticle->SetGravityEnabled(bEnabled);
+		FPhysicsCommand::ExecuteWrite(InActorReference, [&](const FPhysicsActorHandle& Actor)
+		{
+			// todo : This is currently synced in FSingleParticlePhysicsProxy<Chaos::TPBDRigidParticle<float, 3>>::PushToPhysicsState. 
+			//        Ideally this would execute a write command to the gravity forces on the physics thread. However,
+			//        the Actor.Handle() does not have access to the Evolution, so the PerParticleGravityForces are not accessible. 
+			//        This will need to be fixed. 
+		});
+	}
 }
 
 float FPhysInterface_Chaos::GetSleepEnergyThreshold_AssumesLocked(const FPhysicsActorHandle& InActorReference)
@@ -450,7 +494,7 @@ void FPhysInterface_Chaos::SetSleepEnergyThreshold_AssumesLocked(const FPhysicsA
 {
 }
 
-void FPhysInterface_Chaos::SetMass_AssumesLocked(const FPhysicsActorHandle& InActorReference, float InMass)
+void FPhysInterface_Chaos::SetMass_AssumesLocked(FPhysicsActorHandle& InActorReference, float InMass)
 {
 	if (Chaos::TPBDRigidParticle<float, 3 >* RigidParticle = InActorReference->AsDynamic())
 	{
@@ -466,13 +510,13 @@ void FPhysInterface_Chaos::SetMass_AssumesLocked(const FPhysicsActorHandle& InAc
 	}
 }
 
-void FPhysInterface_Chaos::SetMassSpaceInertiaTensor_AssumesLocked(const FPhysicsActorHandle& InActorReference, const FVector& InTensor)
+void FPhysInterface_Chaos::SetMassSpaceInertiaTensor_AssumesLocked(FPhysicsActorHandle& InActorReference, const FVector& InTensor)
 {
 	if (Chaos::TPBDRigidParticle<float, 3 >* RigidParticle = InActorReference->AsDynamic())
 	{
 		if(ensure(!FMath::IsNearlyZero(InTensor.X)) && ensure(!FMath::IsNearlyZero(InTensor.Y)) && ensure(!FMath::IsNearlyZero(InTensor.Z)) )
 		{
-			RigidParticle->SetI(Chaos::PMatrix<float,3,3>(InTensor.X, InTensor.Y, InTensor.Z));
+			RigidParticle->SetI(Chaos::PMatrix<float, 3, 3>(InTensor.X, InTensor.Y, InTensor.Z));
 			RigidParticle->SetInvI(Chaos::PMatrix<float, 3, 3>(1./InTensor.X, 1./InTensor.Y, 1./InTensor.Z));
 		}
 	}
@@ -717,10 +761,141 @@ bool FPhysInterface_Chaos::IsBroken(const FPhysicsConstraintReference_Chaos& InC
 	return true;
 }
 
+enum class EPhysicsInterfaceScopedLockType : uint8
+{
+	Read,
+	Write
+};
+
+struct FScopedSceneLock_Chaos
+{
+	FScopedSceneLock_Chaos(FPhysicsActorHandle const * InActorHandle, EPhysicsInterfaceScopedLockType InLockType)
+		: LockType(InLockType)
+	{
+		Scene = GetSceneForActor(InActorHandle);
+		LockScene();
+	}
+
+	FScopedSceneLock_Chaos(FPhysicsActorHandle const * InActorHandleA, FPhysicsActorHandle const * InActorHandleB, EPhysicsInterfaceScopedLockType InLockType)
+		: LockType(InLockType)
+	{
+		FPhysScene_ChaosInterface* SceneA = GetSceneForActor(InActorHandleA);
+		FPhysScene_ChaosInterface* SceneB = GetSceneForActor(InActorHandleB);
+
+		if(SceneA == SceneB)
+		{
+			Scene = SceneA;
+		}
+		else if(!SceneA || !SceneB)
+		{
+			Scene = SceneA ? SceneA : SceneB;
+		}
+		else
+		{
+			UE_LOG(LogPhysics, Warning, TEXT("Attempted to aquire a physics scene lock for two paired actors that were not in the same scene. Skipping lock"));
+		}
+
+		LockScene();
+	}
+
+	FScopedSceneLock_Chaos(FPhysicsConstraintHandle const * InHandle, EPhysicsInterfaceScopedLockType InLockType)
+		: Scene(nullptr)
+		, LockType(InLockType)
+	{
+		UE_LOG(LogPhysics, Warning, TEXT("Constraint instance attempted scene lock, Constraints currently unimplemented"));
+	}
+
+	FScopedSceneLock_Chaos(USkeletalMeshComponent* InSkelMeshComp, EPhysicsInterfaceScopedLockType InLockType)
+		: LockType(InLockType)
+	{
+		Scene = nullptr;
+
+		if(InSkelMeshComp)
+		{
+			for(FBodyInstance* BI : InSkelMeshComp->Bodies)
+			{
+				Scene = GetSceneForActor(&BI->GetPhysicsActorHandle());
+				if(Scene)
+				{
+					break;
+				}
+			}
+		}
+
+		LockScene();
+	}
+
+	FScopedSceneLock_Chaos(FPhysScene_ChaosInterface* InScene, EPhysicsInterfaceScopedLockType InLockType)
+		: Scene(InScene)
+		, LockType(InLockType)
+	{
+		LockScene();
+	}
+
+	~FScopedSceneLock_Chaos()
+	{
+		UnlockScene();
+	}
+
+private:
+
+	void LockScene()
+	{
+		if(!Scene)
+		{
+			return;
+		}
+
+		switch(LockType)
+		{
+		case EPhysicsInterfaceScopedLockType::Read:
+			Scene->GetScene().ExternalDataLock.ReadLock();
+			break;
+		case EPhysicsInterfaceScopedLockType::Write:
+			Scene->GetScene().ExternalDataLock.WriteLock();
+			break;
+		}
+	}
+
+	void UnlockScene()
+	{
+		if(!Scene)
+		{
+			return;
+		}
+
+		switch(LockType)
+		{
+		case EPhysicsInterfaceScopedLockType::Read:
+			Scene->GetScene().ExternalDataLock.ReadUnlock();
+			break;
+		case EPhysicsInterfaceScopedLockType::Write:
+			Scene->GetScene().ExternalDataLock.WriteUnlock();
+			break;
+		}
+	}
+
+	FPhysScene_ChaosInterface* GetSceneForActor(FPhysicsActorHandle const * InActorHandle)
+	{
+		FBodyInstance* ActorInstance = FPhysicsUserData_Chaos::Get<FBodyInstance>((*InActorHandle)->UserData());
+
+		if(ActorInstance)
+		{
+			return ActorInstance->GetPhysicsScene();
+		}
+
+		return nullptr;
+	}
+
+	FPhysScene_ChaosInterface* Scene;
+	EPhysicsInterfaceScopedLockType LockType;
+};
+
 bool FPhysInterface_Chaos::ExecuteOnUnbrokenConstraintReadOnly(const FPhysicsConstraintReference_Chaos& InConstraintRef, TFunctionRef<void(const FPhysicsConstraintReference_Chaos&)> Func)
 {
     if (!IsBroken(InConstraintRef))
     {
+		FScopedSceneLock_Chaos SceneLock(&InConstraintRef, EPhysicsInterfaceScopedLockType::Read);
         Func(InConstraintRef);
         return true;
     }
@@ -731,114 +906,134 @@ bool FPhysInterface_Chaos::ExecuteOnUnbrokenConstraintReadWrite(const FPhysicsCo
 {
     if (!IsBroken(InConstraintRef))
     {
+		FScopedSceneLock_Chaos SceneLock(&InConstraintRef, EPhysicsInterfaceScopedLockType::Write);
         Func(InConstraintRef);
         return true;
     }
     return false;
 }
 
-template<class PHYSX_MESH>
-TArray<Chaos::TVector<int32, 3>> GetMeshElements(const PHYSX_MESH* PhysXMesh)
+bool FPhysInterface_Chaos::ExecuteRead(const FPhysicsActorHandle& InActorReference, TFunctionRef<void(const FPhysicsActorHandle& Actor)> InCallable)
 {
-	check(false);
-}
+	if(InActorReference)
+	{
+		FScopedSceneLock_Chaos SceneLock(&InActorReference, EPhysicsInterfaceScopedLockType::Read);
 
-#if WITH_PHYSX
+		InCallable(InActorReference);
+		return true;
+	}
+	return false;
+}
 
-template<>
-TArray<Chaos::TVector<int32, 3>> GetMeshElements(const physx::PxConvexMesh* PhysXMesh)
+bool FPhysInterface_Chaos::ExecuteRead(USkeletalMeshComponent* InMeshComponent, TFunctionRef<void()> InCallable)
 {
-	TArray<Chaos::TVector<int32, 3>> CollisionMeshElements;
-#if !WITH_CHAOS_NEEDS_TO_BE_FIXED
-	int32 offset = 0;
-	int32 NbPolygons = static_cast<int32>(PhysXMesh->getNbPolygons());
-	for (int32 i = 0; i < NbPolygons; i++)
-	{
-		PxHullPolygon Poly;
-		bool status = PhysXMesh->getPolygonData(i, Poly);
-		const auto Indices = PhysXMesh->getIndexBuffer() + Poly.mIndexBase;
+	FScopedSceneLock_Chaos SceneLock(InMeshComponent, EPhysicsInterfaceScopedLockType::Read);
+	InCallable();
+	return true;
+}
 
-		for (int32 j = 2; j < static_cast<int32>(Poly.mNbVerts); j++)
-		{
-			CollisionMeshElements.Add(Chaos::TVector<int32, 3>(Indices[offset], Indices[offset + j], Indices[offset + j - 1]));
-		}
-	}
-#endif
-	return CollisionMeshElements;
+bool FPhysInterface_Chaos::ExecuteRead(const FPhysicsActorHandle& InActorReferenceA, const FPhysicsActorHandle& InActorReferenceB, TFunctionRef<void(const FPhysicsActorHandle& ActorA, const FPhysicsActorHandle& ActorB)> InCallable)
+{
+	FScopedSceneLock_Chaos SceneLock(&InActorReferenceA, &InActorReferenceB, EPhysicsInterfaceScopedLockType::Read);
+	InCallable(InActorReferenceA, InActorReferenceB);
+	return true;
 }
 
-template<>
-TArray<Chaos::TVector<int32, 3>> GetMeshElements(const physx::PxTriangleMesh* PhysXMesh)
+bool FPhysInterface_Chaos::ExecuteRead(const FPhysicsConstraintReference_Chaos& InConstraintRef, TFunctionRef<void(const FPhysicsConstraintReference_Chaos& Constraint)> InCallable)
 {
-	TArray<Chaos::TVector<int32, 3>> CollisionMeshElements;
-	const auto MeshFlags = PhysXMesh->getTriangleMeshFlags();
-	for (int32 j = 0; j < static_cast<int32>(PhysXMesh->getNbTriangles()); ++j)
+	if(InConstraintRef.IsValid())
 	{
-		if (MeshFlags | PxTriangleMeshFlag::e16_BIT_INDICES)
-		{
-			const PxU16* Indices = reinterpret_cast<const PxU16*>(PhysXMesh->getTriangles());
-			CollisionMeshElements.Add(Chaos::TVector<int32, 3>(Indices[3 * j], Indices[3 * j + 1], Indices[3 * j + 2]));
-		}
-		else
-		{
-			const PxU32* Indices = reinterpret_cast<const PxU32*>(PhysXMesh->getTriangles());
-			CollisionMeshElements.Add(Chaos::TVector<int32, 3>(Indices[3 * j], Indices[3 * j + 1], Indices[3 * j + 2]));
-		}
+		FScopedSceneLock_Chaos SceneLock(&InConstraintRef, EPhysicsInterfaceScopedLockType::Read);
+		InCallable(InConstraintRef);
+		return true;
 	}
-	return CollisionMeshElements;
+
+	return false;
 }
 
-template<class PHYSX_MESH>
-TUniquePtr<Chaos::TImplicitObject<float, 3>> ConvertPhysXMeshToLevelset(const PHYSX_MESH* PhysXMesh, const FVector& Scale)
+bool FPhysInterface_Chaos::ExecuteRead(FPhysScene* InScene, TFunctionRef<void()> InCallable)
 {
-#if !WITH_CHAOS_NEEDS_TO_BE_FIXED
-    TArray<Chaos::TVector<int32, 3>> CollisionMeshElements = GetMeshElements(PhysXMesh);
-	Chaos::TParticles<float, 3> CollisionMeshParticles;
-	CollisionMeshParticles.AddParticles(PhysXMesh->getNbVertices());
-	for (uint32 j = 0; j < CollisionMeshParticles.Size(); ++j)
+	if(InScene)
 	{
-		const auto& Vertex = PhysXMesh->getVertices()[j];
-		CollisionMeshParticles.X(j) = Scale * Chaos::TVector<float, 3>(Vertex.x, Vertex.y, Vertex.z);
+		FScopedSceneLock_Chaos SceneLock(InScene, EPhysicsInterfaceScopedLockType::Read);
+		InCallable();
+		return true;
 	}
-	Chaos::TBox<float, 3> BoundingBox(CollisionMeshParticles.X(0), CollisionMeshParticles.X(0));
-	for (uint32 j = 1; j < CollisionMeshParticles.Size(); ++j)
+
+	return false;
+}
+
+bool FPhysInterface_Chaos::ExecuteWrite(const FPhysicsActorHandle& InActorReference, TFunctionRef<void(const FPhysicsActorHandle& Actor)> InCallable)
+{
+	//why do we have a write that takes in a const handle?
+	if(InActorReference)
 	{
-		BoundingBox.GrowToInclude(CollisionMeshParticles.X(j));
+		FScopedSceneLock_Chaos SceneLock(&InActorReference, EPhysicsInterfaceScopedLockType::Write);
+		InCallable(InActorReference);
+		return true;
 	}
-#if FORCE_ANALYTICS
-	return TUniquePtr<Chaos::TImplicitObject<float, 3>>(new Chaos::TBox<float, 3>(BoundingBox));
-#else
-	int32 MaxAxisSize = 10;
-	int32 MaxAxis;
-	const auto Extents = BoundingBox.Extents();
-	if (Extents[0] > Extents[1] && Extents[0] > Extents[2])
+	return false;
+}
+
+bool FPhysInterface_Chaos::ExecuteWrite(FPhysicsActorHandle& InActorReference, TFunctionRef<void(FPhysicsActorHandle& Actor)> InCallable)
+{
+	if(InActorReference)
 	{
-		MaxAxis = 0;
+		FScopedSceneLock_Chaos SceneLock(&InActorReference, EPhysicsInterfaceScopedLockType::Write);
+		InCallable(InActorReference);
+		return true;
 	}
-	else if (Extents[1] > Extents[2])
+
+	return false;
+}
+
+bool FPhysInterface_Chaos::ExecuteWrite(USkeletalMeshComponent* InMeshComponent, TFunctionRef<void()> InCallable)
+{
+	FScopedSceneLock_Chaos SceneLock(InMeshComponent, EPhysicsInterfaceScopedLockType::Write);
+	InCallable();
+	return true;
+}
+
+bool FPhysInterface_Chaos::ExecuteWrite(const FPhysicsActorHandle& InActorReferenceA, const FPhysicsActorHandle& InActorReferenceB, TFunctionRef<void(const FPhysicsActorHandle& ActorA, const FPhysicsActorHandle& ActorB)> InCallable)
+{
+	FScopedSceneLock_Chaos SceneLock(&InActorReferenceA, &InActorReferenceB, EPhysicsInterfaceScopedLockType::Write);
+	InCallable(InActorReferenceA, InActorReferenceB);
+	return true;
+}
+
+bool FPhysInterface_Chaos::ExecuteWrite(const FPhysicsConstraintReference_Chaos& InConstraintRef, TFunctionRef<void(const FPhysicsConstraintReference_Chaos& Constraint)> InCallable)
+{
+	if(InConstraintRef.IsValid())
 	{
-		MaxAxis = 1;
+		FScopedSceneLock_Chaos SceneLock(&InConstraintRef, EPhysicsInterfaceScopedLockType::Write);
+		InCallable(InConstraintRef);
+		return true;
 	}
-	else
+
+	return false;
+}
+
+bool FPhysInterface_Chaos::ExecuteWrite(FPhysScene* InScene, TFunctionRef<void()> InCallable)
+{
+	if(InScene)
 	{
-		MaxAxis = 2;
+		FScopedSceneLock_Chaos SceneLock(InScene, EPhysicsInterfaceScopedLockType::Write);
+		InCallable();
+		return true;
 	}
-    Chaos::TVector<int32, 3> Counts(MaxAxisSize * Extents[0] / Extents[MaxAxis], MaxAxisSize * Extents[1] / Extents[MaxAxis], MaxAxisSize * Extents[2] / Extents[MaxAxis]);
-    Counts[0] = Counts[0] < 1 ? 1 : Counts[0];
-    Counts[1] = Counts[1] < 1 ? 1 : Counts[1];
-    Counts[2] = Counts[2] < 1 ? 1 : Counts[2];
-    Chaos::TUniformGrid<float, 3> Grid(BoundingBox.Min(), BoundingBox.Max(), Counts, 1);
-	Chaos::TTriangleMesh<float> CollisionMesh(MoveTemp(CollisionMeshElements));
-	return TUniquePtr<Chaos::TImplicitObject<float, 3>>(new Chaos::TLevelSet<float, 3>(Grid, CollisionMeshParticles, CollisionMesh));
-#endif
 
-#else
-	return TUniquePtr<Chaos::TImplicitObject<float, 3>>();
-#endif // !WITH_CHAOS_NEEDS_TO_BE_FIXED
+	return false;
+}
 
+void FPhysInterface_Chaos::ExecuteShapeWrite(FBodyInstance* InInstance, FPhysicsShapeHandle& InShape, TFunctionRef<void(FPhysicsShapeHandle& InShape)> InCallable)
+{
+	if(InInstance && InShape.IsValid())
+	{
+		FScopedSceneLock_Chaos SceneLock(&InInstance->GetPhysicsActorHandle(), EPhysicsInterfaceScopedLockType::Write);
+		InCallable(InShape);
+	}
 }
 
-#endif
 
 FPhysicsShapeHandle FPhysInterface_Chaos::CreateShape(physx::PxGeometry* InGeom, bool bSimulation, bool bQuery, UPhysicalMaterial* InSimpleMaterial, TArray<UPhysicalMaterial*>* InComplexMaterials)
 {
@@ -871,152 +1066,25 @@ const FBodyInstance* FPhysInterface_Chaos::ShapeToOriginalBodyInstance(const FBo
 	return TargetInstance;
 }
 
-	
+
+
 void FPhysInterface_Chaos::AddGeometry(FPhysicsActorHandle& InActor, const FGeometryAddParams& InParams, TArray<FPhysicsShapeHandle>* OutOptShapes)
 {
-	const FVector& Scale = InParams.Scale;
 	TArray<TUniquePtr<Chaos::TImplicitObject<float, 3>>> Geoms;
 	Chaos::TShapesArray<float, 3> Shapes;
+	ChaosInterface::CreateGeometry(InParams, Geoms, Shapes);
 
-	auto NewShapeHelper = [&InParams](const Chaos::TImplicitObject<float, 3>* InGeom, bool bComplexShape = false)
-	{
-		auto NewShape = MakeUnique<Chaos::TPerShapeData<float, 3>>();
-		NewShape->Geometry = InGeom;
-		NewShape->QueryData = bComplexShape ? InParams.CollisionData.CollisionFilterData.QueryComplexFilter : InParams.CollisionData.CollisionFilterData.QuerySimpleFilter;
-		NewShape->SimData = InParams.CollisionData.CollisionFilterData.SimFilter;
-		return NewShape;
-	};
-
-	if (InActor) 
+#if WITH_CHAOS
+	if (InActor && Geoms.Num())
 	{
-		for (uint32 i = 0; i < static_cast<uint32>(InParams.Geometry->SphereElems.Num()); ++i)
-		{
-			const FKSphereElem ScaledSphereElem = InParams.Geometry->SphereElems[i].GetFinalScaled(Scale, InParams.LocalTransform);
-			const float UseRadius = FMath::Max(ScaledSphereElem.Radius, KINDA_SMALL_NUMBER);
-
-			auto ImplicitSphere = MakeUnique<Chaos::TSphere<float, 3>>(Chaos::TVector<float, 3>(0.f, 0.f, 0.f), UseRadius);
-			auto NewShape = NewShapeHelper(ImplicitSphere.Get());
-			Shapes.Emplace(MoveTemp(NewShape));
-
-			if (OutOptShapes) OutOptShapes->Add({NewShape.Get(),true,true,InActor});
-			Geoms.Add(MoveTemp(ImplicitSphere));
-
-		}
-		for (uint32 i = 0; i < static_cast<uint32>(InParams.Geometry->BoxElems.Num()); ++i)
-		{
-			const auto& Box = InParams.Geometry->BoxElems[i];
-			Chaos::TVector<float, 3> half_extents = Scale * Chaos::TVector<float, 3>(Box.X / 2.f, Box.Y / 2.f, Box.Z / 2.f);
-			auto ImplicitBox = MakeUnique<Chaos::TBox<float, 3>>(-half_extents, half_extents);
-			auto NewShape = NewShapeHelper(ImplicitBox.Get());
-			Shapes.Emplace(MoveTemp(NewShape));
-
-			if (OutOptShapes) OutOptShapes->Add({NewShape.Get(),true,true,InActor});
-			Geoms.Add(MoveTemp(ImplicitBox));
-
-		}
-		for (uint32 i = 0; i < static_cast<uint32>(InParams.Geometry->SphylElems.Num()); ++i)
+		if (OutOptShapes)
 		{
-			const FKSphylElem& UnscaledSphyl = InParams.Geometry->SphylElems[i];
-
-			const FKSphylElem ScaledSphylElem = UnscaledSphyl.GetFinalScaled(Scale, InParams.LocalTransform);
-			float HalfHeight = FMath::Max(ScaledSphylElem.Length * 0.5f, KINDA_SMALL_NUMBER);
-			const float Radius = FMath::Max(ScaledSphylElem.Radius, KINDA_SMALL_NUMBER);
-
-			if (HalfHeight < KINDA_SMALL_NUMBER)
+			for (auto& Shape : Shapes)
 			{
-				//not a capsule just use a sphere
-				auto ImplicitSphere = MakeUnique<Chaos::TSphere<float, 3>>(Chaos::TVector<float, 3>(0), Radius);
-				auto NewShape = NewShapeHelper(ImplicitSphere.Get());
-				Shapes.Emplace(MoveTemp(NewShape));
-
-				if (OutOptShapes) OutOptShapes->Add({NewShape.Get(),true,true,InActor});
-				Geoms.Add(MoveTemp(ImplicitSphere));
-
-			}
-			else
-			{
-				Chaos::TVector<float, 3> HalfExtents(0, 0, HalfHeight);
-
-				auto ImplicitCapsule = MakeUnique<Chaos::TCapsule<float>>(-HalfExtents, HalfExtents, Radius);
-				auto NewShape = NewShapeHelper(ImplicitCapsule.Get());
-				Shapes.Emplace(MoveTemp(NewShape));
-
-				if (OutOptShapes) OutOptShapes->Add({ NewShape.Get(),true,true,InActor });
-				Geoms.Add(MoveTemp(ImplicitCapsule));
+				OutOptShapes->Add({ Shape.Get(), true, true, InActor });
 			}
 		}
-#if 0
-		for (uint32 i = 0; i < static_cast<uint32>(InParams.Geometry->TaperedCapsuleElems.Num()); ++i)
-		{
-			ensure(FMath::IsNearlyEqual(Scale[0], Scale[1]) && FMath::IsNearlyEqual(Scale[1], Scale[2]));
-			const auto& TCapsule = InParams.Geometry->TaperedCapsuleElems[i];
-			if (TCapsule.Length == 0)
-			{
-				Chaos::TSphere<float, 3> * ImplicitSphere = new Chaos::TSphere<float, 3>(-half_extents, TCapsule.Radius * Scale[0]);
-				if (PhysicsProxy) PhysicsProxy->ImplicitObjects_GameThread.Add(ImplicitSphere);
-				else if (OutOptShapes) OutOptShapes->Add({ImplicitSphere,true,true,InActor});
-			}
-			else
-			{
-				Chaos::TVector<float, 3> half_extents(0, 0, TCapsule.Length / 2 * Scale[0]);
-				auto ImplicitCylinder = MakeUnique<Chaos::TCylinder<float>>(-half_extents, half_extents, TCapsule.Radius * Scale[0]);
-				if (PhysicsProxy) PhysicsProxy->ImplicitObjects_GameThread.Add(MoveTemp(ImplicitSphere));
-				else if (OutOptShapes) OutOptShapes->Add({ImplicitSphere,true,true,InActor});
-
-				auto ImplicitSphereA = MakeUnique<Chaos::TSphere<float, 3>>(-half_extents, TCapsule.Radius * Scale[0]);
-				if (PhysicsProxy) PhysicsProxy->ImplicitObjects_GameThread.Add(MoveTemp(ImplicitSphereA));
-				else if (OutOptShapes) OutOptShapes->Add({ImplicitSphereA,true,true,InActor});
-
-				auto ImplicitSphereB = MakeUnique<Chaos::TSphere<float, 3>>(half_extents, TCapsule.Radius * Scale[0]);
-				if (PhysicsProxy) PhysicsProxy->ImplicitObjects_GameThread.Add(MoveTemp(ImplicitSphereB));
-				else if (OutOptShapes) OutOptShapes->Add({ImplicitSphereB,true,true,InActor});
-			}
-		}
-#endif
-
-		for (uint32 i = 0; i < static_cast<uint32>(InParams.Geometry->ConvexElems.Num()); ++i)
-		{
-			const FKConvexElem& CollisionBody = InParams.Geometry->ConvexElems[i];
-			if (const auto& ConvexImplicit = CollisionBody.GetChaosConvexMesh())
-			{
-				auto Implicit = MakeUnique<Chaos::TImplicitObjectScaled<float, 3>>(MakeSerializable(ConvexImplicit), Scale);
-				auto NewShape = NewShapeHelper(Implicit.Get());
-				Shapes.Emplace(MoveTemp(NewShape));
-
-				if (OutOptShapes) OutOptShapes->Add({ NewShape.Get(),true,true,InActor });
-				Geoms.Add(MoveTemp(Implicit));
-			}
-		}
-
-		for (const auto& ChaosTriMesh : InParams.ChaosTriMeshes)
-		{
-			auto Implicit = MakeUnique<Chaos::TImplicitObjectScaled<float, 3>>(MakeSerializable(ChaosTriMesh), Scale);
-			auto NewShape = NewShapeHelper(Implicit.Get(), true);
-			Shapes.Emplace(MoveTemp(NewShape));
-
-			if (OutOptShapes) OutOptShapes->Add({ NewShape.Get(),true,true,InActor });
-			Geoms.Add(MoveTemp(Implicit));
-		}
-	}
-	else
-	{
-#if WITH_PHYSX
-		for (const auto& PhysXMesh : InParams.TriMeshes)
-		{
-			auto Implicit = ConvertPhysXMeshToLevelset(PhysXMesh, Scale);
-			auto NewShape = NewShapeHelper(Implicit.Get(), true);
-			Shapes.Emplace(MoveTemp(NewShape));
 
-			if (OutOptShapes) OutOptShapes->Add({NewShape.Get(),true,true,InActor});
-			Geoms.Add(MoveTemp(Implicit));
-
-		}
-#endif
-	}
-
-#if WITH_CHAOS
-	if (InActor)
-	{
 		//todo: we should not be creating unique geometry per actor
 		InActor->SetGeometry(MakeUnique<Chaos::TImplicitObjectUnion<float, 3>>(MoveTemp(Geoms)));
 		InActor->SetShapesArray(MoveTemp(Shapes));
@@ -1122,20 +1190,75 @@ bool FPhysInterface_Chaos::IsShapeType(const FPhysicsShapeReference_Chaos& InSha
     {
         return true;
     }
-    // Other than sphere and box the basic types do not correlate so we return false
+    if (InType == ECollisionShapeType::Plane && InShapeRef.Shape->Geometry->GetType() == Chaos::ImplicitObjectType::Plane)
+    {
+        return true;
+    }
+    if (InType == ECollisionShapeType::Capsule && InShapeRef.Shape->Geometry->GetType() == Chaos::ImplicitObjectType::Capsule)
+    {
+        return true;
+    }
+    if (InType == ECollisionShapeType::Convex && InShapeRef.Shape->Geometry->GetType() == Chaos::ImplicitObjectType::Convex)
+    {
+        return true;
+    }
+    if (InType == ECollisionShapeType::Trimesh && InShapeRef.Shape->Geometry->GetType() == Chaos::ImplicitObjectType::TriangleMesh)
+    {
+        return true;
+    }
+    if (InType == ECollisionShapeType::Heightfield && InShapeRef.Shape->Geometry->GetType() == Chaos::ImplicitObjectType::HeightField)
+    {
+        return true;
+    }
+    if (InType == ECollisionShapeType::Scaled && InShapeRef.Shape->Geometry->GetType() == Chaos::ImplicitObjectType::Scaled)
+    {
+        return true;
+    }
+
     return false;
 }
 
-ECollisionShapeType FPhysInterface_Chaos::GetShapeType(const FPhysicsShapeReference_Chaos& InShapeRef)
+ECollisionShapeType FPhysInterface_Chaos::GetShapeType(const FPhysicsShapeReference_Chaos& InShapeRef, bool bGetInnerType)
 {
-    if (InShapeRef.Shape->Geometry->GetType() == Chaos::ImplicitObjectType::Box)
+	Chaos::ImplicitObjectType Type = InShapeRef.Shape->Geometry->GetType(true);
+	if (bGetInnerType && Type == Chaos::ImplicitObjectType::Scaled)
+	{
+		Type = InShapeRef.Shape->Geometry->GetObject<Chaos::TImplicitObjectScaled<float, 3>>()->GetUnscaledObject()->GetType(true);
+	}
+
+    if (Type == Chaos::ImplicitObjectType::Box)
     {
         return ECollisionShapeType::Box;
     }
-    if (InShapeRef.Shape->Geometry->GetType() == Chaos::ImplicitObjectType::Sphere)
+    if (Type == Chaos::ImplicitObjectType::Sphere)
     {
         return ECollisionShapeType::Sphere;
     }
+    if (Type == Chaos::ImplicitObjectType::Plane)
+    {
+        return ECollisionShapeType::Plane;
+    }
+    if (Type == Chaos::ImplicitObjectType::Capsule)
+    {
+        return ECollisionShapeType::Capsule;
+    }
+    if (Type == Chaos::ImplicitObjectType::Convex)
+    {
+        return ECollisionShapeType::Convex;
+    }
+    if (Type == Chaos::ImplicitObjectType::TriangleMesh)
+    {
+        return ECollisionShapeType::Trimesh;
+    }
+    if (Type == Chaos::ImplicitObjectType::HeightField)
+    {
+        return ECollisionShapeType::Heightfield;
+    }
+    if (Type == Chaos::ImplicitObjectType::Scaled)
+    {
+        return ECollisionShapeType::Scaled;
+    }
+
     return ECollisionShapeType::None;
 }
 
@@ -1188,21 +1311,195 @@ void FinishSceneStat()
 
 #if WITH_PHYSX
 
+void CalculateMassPropertiesOfImplicitType(
+	TArray< Chaos::TMassProperties<float, 3> > & MassProperties,
+	float & TotalMass,
+	const Chaos::TRigidTransform<float, 3> & WorldTransform,
+	const Chaos::TImplicitObject<float, 3>* ImplicitObject, 
+	float InDensityKGPerCM)
+{
+	// WIP
+	// @todo : Support center of mass offsets.
+	// @todo : Support Mass space alignment. 
+
+	if (ImplicitObject)
+	{
+		Chaos::TVector<float, 3> Scale = WorldTransform.GetScale3D();
+		Chaos::PMatrix<float, 3, 3> ScaleM(Scale.X, Scale.Y, Scale.Z);
+
+		if (ImplicitObject->GetType() == Chaos::ImplicitObjectType::Sphere)
+		{
+			const Chaos::TSphere<float, 3> * Sphere = ImplicitObject->template GetObject<Chaos::TSphere<float, 3>>();
+			float Mass = Sphere->GetVolume()*InDensityKGPerCM;
+
+			Chaos::TMassProperties<float, 3> MassProperty;
+			MassProperty.CenterOfMass = Chaos::TVector<float, 3>(0);
+			MassProperty.Volume = Sphere->GetVolume();
+			MassProperty.InertiaTensor = Sphere->GetInertiaTensor(Mass)*ScaleM;
+			MassProperty.RotationOfMass = Chaos::TRotation<float, 3>::FromIdentity();
+			//MassProperty.RotationOfMass = Chaos::TransformToLocalSpace<float,3>(MassProperty.InertiaTensor);
+
+			MassProperties.Add(MassProperty);
+			TotalMass += Mass;
+		}
+		else if (ImplicitObject->GetType() == Chaos::ImplicitObjectType::Box)
+		{
+			const Chaos::TBox<float, 3> * Box = ImplicitObject->template GetObject<Chaos::TBox<float, 3>>();
+			float Mass = Box->GetVolume()*InDensityKGPerCM;
+
+			Chaos::TMassProperties<float, 3> MassProperty;
+			MassProperty.CenterOfMass = Chaos::TVector<float,3>(0);
+			MassProperty.Volume = Box->GetVolume();
+			MassProperty.InertiaTensor = Box->GetInertiaTensor(Mass); // What's the box? Scale!
+			MassProperty.RotationOfMass = Chaos::TRotation<float, 3>::FromIdentity();
+			//MassProperty.RotationOfMass = Chaos::TransformToLocalSpace<float, 3>(MassProperty.InertiaTensor);
+
+			MassProperties.Add(MassProperty);
+			TotalMass += Mass;
+		}
+		else if (ImplicitObject->GetType() == Chaos::ImplicitObjectType::Capsule)
+		{
+			const Chaos::TCapsule<float> * Capsule = ImplicitObject->template GetObject<Chaos::TCapsule<float>>();
+			float Mass = Capsule->GetVolume()*InDensityKGPerCM;
+
+			Chaos::TMassProperties<float, 3> MassProperty;
+			MassProperty.CenterOfMass = Chaos::TVector<float, 3>(0);
+			MassProperty.Volume = Capsule->GetVolume();
+			MassProperty.InertiaTensor = Capsule->GetInertiaTensor(Mass)*ScaleM;
+			MassProperty.RotationOfMass = Chaos::TRotation<float, 3>::FromIdentity();
+			//MassProperty.RotationOfMass = Chaos::TransformToLocalSpace<float, 3>(MassProperty.InertiaTensor);
+
+			MassProperties.Add(MassProperty);
+			TotalMass += Mass;
+		}
+		else if (ImplicitObject->GetType() == Chaos::ImplicitObjectType::Cylinder)
+		{
+			const Chaos::TCylinder<float> * Cylinder = ImplicitObject->template GetObject<Chaos::TCylinder<float>>();
+			float Mass = Cylinder->GetVolume()*InDensityKGPerCM;
+
+			Chaos::TMassProperties<float, 3> MassProperty;
+			MassProperty.CenterOfMass = Chaos::TVector<float, 3>(0);
+			MassProperty.Volume = Cylinder->GetVolume();
+			MassProperty.InertiaTensor = Cylinder->GetInertiaTensor(Mass)*ScaleM;
+			MassProperty.RotationOfMass = Chaos::TRotation<float, 3>::FromIdentity();
+			//MassProperty.RotationOfMass = Chaos::TransformToLocalSpace<float, 3>(MassProperty.InertiaTensor);
+
+			MassProperties.Add(MassProperty);
+			TotalMass += Mass;
+		}
+		else if (ImplicitObject->GetType() == Chaos::ImplicitObjectType::TaperedCylinder)
+		{
+			const Chaos::TTaperedCylinder<float> * TaperedCylinder = ImplicitObject->template GetObject<Chaos::TTaperedCylinder<float>>();
+			float Mass = TaperedCylinder->GetVolume()*InDensityKGPerCM;
+
+			Chaos::TMassProperties<float, 3> MassProperty;
+			MassProperty.CenterOfMass = Chaos::TVector<float, 3>(0);
+			MassProperty.Volume = TaperedCylinder->GetVolume();
+			MassProperty.InertiaTensor = TaperedCylinder->GetInertiaTensor(Mass)*ScaleM;
+			MassProperty.RotationOfMass = Chaos::TRotation<float, 3>::FromIdentity();
+			//MassProperty.RotationOfMass = Chaos::TransformToLocalSpace<float, 3>(MassProperty.InertiaTensor);
+
+			MassProperties.Add(MassProperty);
+			TotalMass += Mass;
+		}
+		else if (ImplicitObject->GetType() == Chaos::ImplicitObjectType::Convex)
+		{
+			// @todo: until this is actually used by anything that matters, just keep it simple. 
+			const Chaos::TConvex<float,3> * Convex = ImplicitObject->template GetObject<Chaos::TConvex<float,3>>();
+			float Mass = Convex->BoundingBox().GetVolume()*InDensityKGPerCM;
+
+			Chaos::TMassProperties<float, 3> MassProperty;
+			MassProperty.CenterOfMass = Chaos::TVector<float, 3>(0);
+			MassProperty.Volume = Convex->BoundingBox().GetVolume();
+			MassProperty.InertiaTensor = Chaos::TBox<float, 3>::GetInertiaTensor(Mass, Convex->BoundingBox().Extents());
+			MassProperty.RotationOfMass = Chaos::TRotation<float, 3>::FromIdentity();
+			//MassProperty.RotationOfMass = Chaos::TransformToLocalSpace<float, 3>(MassProperty.InertiaTensor);
+
+			MassProperties.Add(MassProperty);
+			TotalMass += Mass;
+		}
+		else if (ImplicitObject->GetType() == Chaos::ImplicitObjectType::TriangleMesh)
+		{
+			// @todo: until this is actually used by anything that matters, just keep it simple. 
+			const Chaos::TTriangleMeshImplicitObject<float> * TriangleMeshImplicitObject = ImplicitObject->template GetObject<Chaos::TTriangleMeshImplicitObject<float>>();
+			float Mass = TriangleMeshImplicitObject->BoundingBox().GetVolume()*InDensityKGPerCM;
+
+			Chaos::TMassProperties<float, 3> MassProperty;
+			MassProperty.CenterOfMass = Chaos::TVector<float, 3>(0);
+			MassProperty.Volume = TriangleMeshImplicitObject->BoundingBox().GetVolume();
+			MassProperty.InertiaTensor = Chaos::TBox<float, 3>::GetInertiaTensor(Mass, TriangleMeshImplicitObject->BoundingBox().Extents());
+			MassProperty.RotationOfMass = Chaos::TRotation<float, 3>::FromIdentity();
+			//MassProperty.RotationOfMass = Chaos::TransformToLocalSpace<float, 3>(MassProperty.InertiaTensor);
+
+			MassProperties.Add(MassProperty);
+			TotalMass += Mass;
+		}
+		else if (ImplicitObject->GetType() == Chaos::ImplicitObjectType::Scaled)
+		{
+			const Chaos::TImplicitObjectScaled<float, 3> * ScaledObject = ImplicitObject->template GetObject<Chaos::TImplicitObjectScaled<float, 3>>();
+			if (ensureMsgf(ScaledObject->GetUnscaledObject(), TEXT("Error : Null internal MObject type within TImplicitObjectScaled.")))
+			{
+				Chaos::TRigidTransform<float, 3> ScaledWorldTransform(WorldTransform.GetTranslation(), WorldTransform.GetRotation(), ScaledObject->GetScale());
+				CalculateMassPropertiesOfImplicitType(MassProperties, TotalMass, ScaledWorldTransform, ScaledObject->GetUnscaledObject(), InDensityKGPerCM);
+			}
+		}
+		else if (ImplicitObject->GetType() == Chaos::ImplicitObjectType::Union)
+		{
+			const Chaos::TImplicitObjectUnion<float,3> * ImplicitUnion = ImplicitObject->template GetObject<Chaos::TImplicitObjectUnion<float,3>>();
+
+			for (const TUniquePtr < Chaos::TImplicitObject<float,3> > & ImplicitSubObject : ImplicitUnion->GetObjects())
+			{
+				CalculateMassPropertiesOfImplicitType(MassProperties, TotalMass, WorldTransform, ImplicitSubObject.Get(), InDensityKGPerCM);
+			}
+		}
+		else
+		{
+			// @todo : Enable when ready.
+			//ensureMsgf(false, TEXT("Unsupported implicit object type (%s)."), *(ImplicitObject->GetTypeName().ToString()) );
+		}
+	}
+}
+
 void FPhysInterface_Chaos::CalculateMassPropertiesFromShapeCollection(physx::PxMassProperties& OutProperties, const TArray<FPhysicsShapeHandle>& InShapes, float InDensityKGPerCM)
 {
-	// #todo : Implement
-    //OutProperties.inertiaTensor = physx::PxMat33();
-    //OutProperties.inertiaTensor(0, 0) = Inertia.M[0][0];
-    //OutProperties.inertiaTensor(1, 1) = Inertia.M[1][1];
-    //OutProperties.inertiaTensor(2, 2) = Inertia.M[2][2];
-    //OutProperties.mass = LocalParticles.M(Index);
+	float TotalMass = 0;
+	TArray< Chaos::TMassProperties<float,3> > MassProperties;
+
+	for (const FPhysicsShapeHandle& ShapeHandle : InShapes)
+	{
+		if (const Chaos::TPerShapeData<float, 3>* Shape = ShapeHandle.Shape)
+		{
+			if (const Chaos::TImplicitObject<float, 3> * ImplicitObject = Shape->Geometry.Get())
+			{
+				FTransform WorldTransform(ShapeHandle.ActorRef->R(), ShapeHandle.ActorRef->X());
+				CalculateMassPropertiesOfImplicitType(MassProperties, TotalMass, WorldTransform, ImplicitObject, InDensityKGPerCM);
+			}
+		}
+	}
+
+
+	Chaos::PMatrix<float, 3, 3> Tensor;
+	if (MassProperties.Num())
+	{
+		Tensor = Chaos::Combine<float, 3>(MassProperties).InertiaTensor;
+	}
+	else 
+	{
+		// @todo : Add support for all types, but for now just hard code a unit sphere tensor {r:50cm} if the type was not processed
+		Tensor = Chaos::PMatrix<float, 3, 3>(5.24e5, 5.24e5, 5.24e5);
+		TotalMass = 523.f;
+	}
+
+	float Mat[] = { Tensor.M[0][0],Tensor.M[0][1],Tensor.M[0][2],Tensor.M[1][0],Tensor.M[1][1],Tensor.M[1][2],Tensor.M[2][0],Tensor.M[2][1],Tensor.M[2][2] };
+	OutProperties.inertiaTensor = PxMat33(Mat);
+	OutProperties.mass = TotalMass;
 }
 
 #endif
 
 bool FPhysInterface_Chaos::LineTrace_Geom(FHitResult& OutHit, const FBodyInstance* InInstance, const FVector& WorldStart, const FVector& WorldEnd, bool bTraceComplex, bool bExtractPhysMaterial)
 {
-	// Need an instandce to trace against
+	// Need an instance to trace against
 	check(InInstance);
 
 	OutHit.TraceStart = WorldStart;
@@ -1370,9 +1667,13 @@ bool FPhysInterface_Chaos::Sweep_Geom(FHitResult& OutHit, const FBodyInstance* I
 								Hit.Actor = ShapeRef.ActorRef;
 								Hit.WorldPosition = WorldPosition;
 								Hit.WorldNormal = WorldNormal;
+								Hit.FaceIndex = FaceIdx;
+								if (!HadInitialOverlap(Hit))
+								{
+									Hit.FaceIndex = FindFaceIndex(Hit, Dir);
+								}
 
 								FTransform StartTransform(InStart);
-								Hit.FaceIndex = FindFaceIndex(Hit, Dir);
 								ConvertQueryImpactHit(OwnerComponentInst->GetWorld(), Hit, OutHit, DeltaMag, QueryFilter, InStart, InEnd, nullptr, StartTransform, false, false);
 								bSweepHit = true;
 							}
@@ -1458,7 +1759,7 @@ bool FPhysInterface_Chaos::Overlap_Geom(const FBodyInstance* InBodyInstance, con
 bool FPhysInterface_Chaos::Overlap_Geom(const FBodyInstance* InBodyInstance, const FCollisionShape& InCollisionShape, const FQuat& InShapeRotation, const FTransform& InShapeTransform, FMTDResult* OutOptResult)
 {
 	FPhysicsShapeAdapter Adaptor(InShapeRotation, InCollisionShape);
-	return Overlap_GeomInternal(InBodyInstance, Adaptor.GetGeometry(), Adaptor.GetGeometryPose(InShapeTransform.GetTranslation()), OutOptResult);
+	return Overlap_GeomInternal(InBodyInstance, Adaptor.GetGeometry(), Adaptor.GetGeomPose(InShapeTransform.GetTranslation()), OutOptResult);
 }
 
 bool FPhysInterface_Chaos::GetSquaredDistanceToBody(const FBodyInstance* InInstance, const FVector& InPoint, float& OutDistanceSquared, FVector* OutOptPointOnBody)
@@ -1492,7 +1793,8 @@ bool FPhysInterface_Chaos::GetSquaredDistanceToBody(const FBodyInstance* InInsta
 				continue;
 			}
 
-			ECollisionShapeType GeomType = FPhysicsInterface::GetShapeType(Shape);
+			bool bGetInnerType = true;
+			ECollisionShapeType GeomType = FPhysicsInterface::GetShapeType(Shape, bGetInnerType);
 
 			if (GeomType == ECollisionShapeType::Trimesh)
 			{
@@ -1607,7 +1909,7 @@ const FPhysicsGeometry& FPhysicsShapeAdapter_Chaos::GetGeometry() const
 	return *Geometry;
 }
 
-FTransform FPhysicsShapeAdapter_Chaos::GetGeometryPose(const FVector& Pos) const
+FTransform FPhysicsShapeAdapter_Chaos::GetGeomPose(const FVector& Pos) const
 {
 	return FTransform(GeometryRotation, Pos);
 }
