Merging //UE4/Main@9388041 to Dev-RenderPlat-Staging

diff --git a/Engine/Source/Runtime/Engine/Private/PhysicsEngine/Experimental/PhysScene_Chaos.cpp b/Engine/Source/Runtime/Engine/Private/PhysicsEngine/Experimental/PhysScene_Chaos.cpp
--- a/Engine/Source/Runtime/Engine/Private/PhysicsEngine/Experimental/PhysScene_Chaos.cpp
+++ b/Engine/Source/Runtime/Engine/Private/PhysicsEngine/Experimental/PhysScene_Chaos.cpp
@@ -1,7 +1,5 @@
 // Copyright 1998-2019 Epic Games, Inc. All Rights Reserved.
 
-#if INCLUDE_CHAOS
-
 #include "Physics/Experimental/PhysScene_Chaos.h"
 
 #include "PhysicsSolver.h"
@@ -31,10 +29,10 @@
 #include "PhysicsProxy/StaticMeshPhysicsProxy.h"
 #include "Chaos/UniformGrid.h"
 #include "Chaos/BoundingVolume.h"
-#include "Chaos/ISpatialAcceleration.h"
 #include "Chaos/Framework/DebugSubstep.h"
 #include "Chaos/PBDSpringConstraints.h"
 #include "Chaos/PerParticleGravity.h"
+#include "PBDRigidActiveParticlesBuffer.h"
 
 #if !UE_BUILD_SHIPPING
 #include "Engine/World.h"
@@ -347,6 +345,8 @@ FPhysScene_Chaos::FPhysScene_Chaos(AActor* InSolverActor)
 	, bIsWorldPaused(false)
 #endif
 {
+	LLM_SCOPE(ELLMTag::Chaos);
+
 	ChaosModule = FModuleManager::Get().GetModulePtr<FChaosSolversModule>("ChaosSolvers");
 	check(ChaosModule);
 
@@ -419,6 +419,8 @@ bool FPhysScene_Chaos::IsTickable() const
 
 void FPhysScene_Chaos::Tick(float DeltaTime)
 {
+	LLM_SCOPE(ELLMTag::Chaos);
+
 #if WITH_EDITOR
 	// Check the editor pause status and update this object's single-step counter.
 	// This check cannot be moved to IsTickable() since this is a test/update operation
@@ -442,20 +444,6 @@ void FPhysScene_Chaos::Tick(float DeltaTime)
 
 	UE_LOG(LogFPhysScene_ChaosSolver, Verbose, TEXT("FPhysScene_Chaos::Tick(%3.5f)"), SafeDelta);
 	Solver->AdvanceSolverBy(SafeDelta);
-
-	Solver->ForEachPhysicsProxyParallel([](auto* Object)
-	{
-		// #BGallagher TODO Just use one side of the buffer for single-thread tick
-		Object->BufferPhysicsResults();
-		Object->FlipBuffer();
-	});
-
-	Solver->ForEachPhysicsProxy([](auto* Object)
-	{
-		Object->PullFromPhysicsState();
-	});
-
-
 }
 
 Chaos::FPhysicsSolver* FPhysScene_Chaos::GetSolver() const
@@ -585,11 +573,15 @@ void FPhysScene_Chaos::AddObject(UPrimitiveComponent* Component, FFieldSystemPhy
 
 void FPhysScene_Chaos::RemoveActorFromAccelerationStructure(FPhysicsActorHandle& Actor)
 {
-	if (SolverAccelerationStructure)
+#if WITH_CHAOS
+	if (GetSpacialAcceleration())
 	{
+		ExternalDataLock.WriteLock();
 		Chaos::TAccelerationStructureHandle<float, 3> AccelerationHandle(Actor);
-		SolverAccelerationStructure->RemoveElement(AccelerationHandle);
+		GetSpacialAcceleration()->RemoveElement(AccelerationHandle);
+		ExternalDataLock.WriteUnlock();
 	}
+#endif
 }
 
 template<typename ObjectType>
@@ -743,6 +735,7 @@ void FPhysScene_Chaos::RemoveObject(FFieldSystemPhysicsProxy* InObject)
 #endif
 }
 
+#if XGE_FIXED
 void FPhysScene_Chaos::UnregisterEvent(const Chaos::EEventType& EventID)
 {
 	check(IsInGameThread());
@@ -774,6 +767,7 @@ void FPhysScene_Chaos::UnregisterEventHandler(const Chaos::EEventType& EventID,
 		});
 	}
 }
+#endif // XGE_FIXED
 
 void FPhysScene_Chaos::Shutdown()
 {
@@ -805,6 +799,36 @@ void FPhysScene_Chaos::AddReferencedObjects(FReferenceCollector& Collector)
 #endif
 }
 
+const Chaos::ISpatialAcceleration<Chaos::TAccelerationStructureHandle<float, 3>, float, 3>* FPhysScene_Chaos::GetSpacialAcceleration() const
+{
+	if(GetDispatcher() && GetDispatcher()->GetMode() == Chaos::EThreadingMode::SingleThread)
+	{
+		return GetSolver()->GetEvolution()->GetSpatialAcceleration();
+	}
+
+	return SolverAccelerationStructure.Get();
+}
+
+Chaos::ISpatialAcceleration<Chaos::TAccelerationStructureHandle<float, 3>, float, 3>* FPhysScene_Chaos::GetSpacialAcceleration()
+{
+	if(GetDispatcher() && GetDispatcher()->GetMode() == Chaos::EThreadingMode::SingleThread)
+	{
+		return GetSolver()->GetEvolution()->GetSpatialAcceleration();
+	}
+
+	return SolverAccelerationStructure.Get();
+}
+
+void FPhysScene_Chaos::CopySolverAccelerationStructure()
+{
+	if(SceneSolver)
+	{
+		ExternalDataLock.WriteLock();
+		SceneSolver->GetEvolution()->UpdateExternalAccelerationStructure(SolverAccelerationStructure);
+		ExternalDataLock.WriteUnlock();
+	}
+}
+
 #if CHAOS_WITH_PAUSABLE_SOLVER
 void FPhysScene_Chaos::OnUpdateWorldPause()
 {
@@ -896,8 +920,6 @@ FPhysScene_ChaosInterface::FPhysScene_ChaosInterface(const AWorldSettings* InSet
 #if TODO_FIX_REFERENCES_TO_ADDARRAY
 	Scene.GetSolver()->GetEvolution()->GetParticles().AddArray(&BodyInstances);
 #endif
-	// @todo: What do we do with MGravity?
-	Scene.GetSolver()->GetEvolution()->AddForceFunction(Chaos::Utilities::GetRigidsGravityFunction(Chaos::TVector<float, 3>(0, 0, -1), 980.f));
 
 	Scene.GetSolver()->PhysSceneHack = this;
 
@@ -978,7 +1000,24 @@ void FPhysScene_ChaosInterface::AddCustomPhysics_AssumesLocked(FBodyInstance* Bo
 
 void FPhysScene_ChaosInterface::AddForce_AssumesLocked(FBodyInstance* BodyInstance, const FVector& Force, bool bAllowSubstepping, bool bAccelChange)
 {
-	// #todo : Implement
+	FPhysicsActorHandle& Handle = BodyInstance->GetPhysicsActorHandle();
+	if (FPhysicsInterface::IsValid(Handle))
+	{
+		if (Chaos::TPBDRigidParticle<float, 3>* Rigid = Handle->AsDynamic())
+		{
+			const Chaos::TVector<float, 3> CurrentForce = Rigid->ExternalForce();
+			if (bAccelChange)
+			{
+				const float Mass = Rigid->M();
+				const Chaos::TVector<float, 3> TotalAcceleration = CurrentForce + (Force * Mass);
+				Rigid->SetExternalForce(TotalAcceleration);
+			}
+			else
+			{
+				Rigid->SetExternalForce(CurrentForce + Force);
+			}
+		}
+	}
 }
 
 void FPhysScene_ChaosInterface::AddForceAtPosition_AssumesLocked(FBodyInstance* BodyInstance, const FVector& Force, const FVector& Position, bool bAllowSubstepping, bool bIsLocalForce /*= false*/)
@@ -1107,76 +1146,59 @@ void FPhysScene_ChaosInterface::StartFrame()
 	FChaosSolversModule* SolverModule = FChaosSolversModule::GetModule();
 	checkSlow(SolverModule);
 
-	Chaos::IDispatcher* Dispatcher = SolverModule->GetDispatcher();
-
-	switch(Dispatcher->GetMode())
-	{
-	case EChaosThreadingMode::SingleThread:
-	{
-		OnPhysScenePreTick.Broadcast(this, MDeltaTime);
-		OnPhysSceneStep.Broadcast(this, MDeltaTime);
-
-		// Here we can directly tick the scene. Single threaded mode doesn't buffer any commands
-		// that would require pumping here - everything is done on demand.
-		Scene.Tick(MDeltaTime);
-	}
-		break;
-	case EChaosThreadingMode::TaskGraph:
+	if (Chaos::IDispatcher* Dispatcher = SolverModule->GetDispatcher())
 	{
-		check(!CompletionEvent.GetReference())
-
-		OnPhysScenePreTick.Broadcast(this, MDeltaTime);
-		OnPhysSceneStep.Broadcast(this, MDeltaTime);
-
+		for (auto * Solver : SolverModule->GetSolvers())
 		{
-			const TArray<FPhysicsSolver*>& SolverList = SolverModule->GetSolvers();
-			TArray<FPhysicsSolver*> ActiveSolvers;
-			ActiveSolvers.Reserve(SolverList.Num());
-
-			// #BG calculate active solver list once as we dispatch our first task
-			for (FPhysicsSolver* Solver : SolverList)
-			{
-				if (Solver->HasActiveParticles())
-				{
-					ActiveSolvers.Add(Solver);
-				}
-			}
+			Solver->PushPhysicsState(Dispatcher);
+		}
 
-			const int32 NumActiveSolvers = ActiveSolvers.Num();
+		switch (Dispatcher->GetMode())
+		{
+		case EChaosThreadingMode::SingleThread:
+		{
+			OnPhysScenePreTick.Broadcast(this, MDeltaTime);
+			OnPhysSceneStep.Broadcast(this, MDeltaTime);
 
-			PhysicsParallelFor(NumActiveSolvers, [&](int32 Index)
-			{
-				FPhysicsSolver* Solver = ActiveSolvers[Index];
-				Solver->UpdatePhysicsThreadStructures();
-			});
+			// Here we can directly tick the scene. Single threaded mode doesn't buffer any commands
+			// that would require pumping here - everything is done on demand.
+			Scene.Tick(MDeltaTime);
 		}
+		break;
+		case EChaosThreadingMode::TaskGraph:
+		{
+			check(!CompletionEvent.GetReference())
 
-		FGraphEventRef SimulationCompleteEvent = FGraphEvent::CreateGraphEvent();
+				OnPhysScenePreTick.Broadcast(this, MDeltaTime);
+			OnPhysSceneStep.Broadcast(this, MDeltaTime);
 
-		// Need to fire off a parallel task to handle running physics commands and
-		// ticking the scene while the engine continues on until TG_EndPhysics
-		// (this should happen in TG_StartPhysics)
-		PhysicsTickTask = TGraphTask<FPhysicsTickTask>::CreateTask(nullptr, ENamedThreads::GameThread).ConstructAndDispatchWhenReady(SimulationCompleteEvent, MDeltaTime);
+			FGraphEventRef SimulationCompleteEvent = FGraphEvent::CreateGraphEvent();
 
-		// Setup post simulate tasks
-		if(PhysicsTickTask.GetReference())
-		{
-			FGraphEventArray PostSimPrerequisites;
-			PostSimPrerequisites.Add(SimulationCompleteEvent);
+			// Need to fire off a parallel task to handle running physics commands and
+			// ticking the scene while the engine continues on until TG_EndPhysics
+			// (this should happen in TG_StartPhysics)
+			PhysicsTickTask = TGraphTask<FPhysicsTickTask>::CreateTask(nullptr, ENamedThreads::GameThread).ConstructAndDispatchWhenReady(SimulationCompleteEvent, MDeltaTime);
+
+			// Setup post simulate tasks
+			if (PhysicsTickTask.GetReference())
+			{
+				FGraphEventArray PostSimPrerequisites;
+				PostSimPrerequisites.Add(SimulationCompleteEvent);
 
-			DECLARE_CYCLE_STAT(TEXT("FDelegateGraphTask.CompletePhysicsSimulation"), STAT_FDelegateGraphTask_CompletePhysicsSimulation, STATGROUP_TaskGraphTasks);
+				DECLARE_CYCLE_STAT(TEXT("FDelegateGraphTask.CompletePhysicsSimulation"), STAT_FDelegateGraphTask_CompletePhysicsSimulation, STATGROUP_TaskGraphTasks);
 
-			// Completion event runs in parallel and will flip out our buffers, gamethread work can be done in EndFrame (Called by world after this completion event finishes)
-			CompletionEvent = FDelegateGraphTask::CreateAndDispatchWhenReady(FDelegateGraphTask::FDelegate::CreateRaw(this, &FPhysScene_ChaosInterface::CompleteSceneSimulation), GET_STATID(STAT_FDelegateGraphTask_CompletePhysicsSimulation), &PostSimPrerequisites, ENamedThreads::GameThread, ENamedThreads::AnyHiPriThreadHiPriTask);
+				// Completion event runs in parallel and will flip out our buffers, gamethread work can be done in EndFrame (Called by world after this completion event finishes)
+				CompletionEvent = FDelegateGraphTask::CreateAndDispatchWhenReady(FDelegateGraphTask::FDelegate::CreateRaw(this, &FPhysScene_ChaosInterface::CompleteSceneSimulation), GET_STATID(STAT_FDelegateGraphTask_CompletePhysicsSimulation), &PostSimPrerequisites, ENamedThreads::GameThread, ENamedThreads::AnyHiPriThreadHiPriTask);
+			}
 		}
-	}
 		break;
 
-	// No action for dedicated thread, the module will sync independently from the scene in
-	// this case. (See FChaosSolversModule::SyncTask and FPhysicsThreadSyncCaller)
-	case EChaosThreadingMode::DedicatedThread:
-	default:
-		break;
+		// No action for dedicated thread, the module will sync independently from the scene in
+		// this case. (See FChaosSolversModule::SyncTask and FPhysicsThreadSyncCaller)
+		case EChaosThreadingMode::DedicatedThread:
+		default:
+			break;
+		}
 	}
 }
 
@@ -1195,7 +1217,7 @@ void FPhysScene_ChaosInterface::EndFrame(ULineBatchComponent* InLineBatcher)
 	{
 	case EChaosThreadingMode::SingleThread:
 	{
-		SyncBodies();
+		SyncBodies(Scene.GetSolver());
 		Scene.GetSolver()->SyncEvents_GameThread();
 
 		OnPhysScenePostTick.Broadcast(this);
@@ -1208,16 +1230,16 @@ void FPhysScene_ChaosInterface::EndFrame(ULineBatchComponent* InLineBatcher)
 		CompletionEvent = nullptr;
 		PhysicsTickTask = nullptr;
 
+		//flush queue so we can merge the two threads
+		Dispatcher->Execute();
+
 		// Flip the buffers over to the game thread and sync
 		{
 			SCOPE_CYCLE_COUNTER(STAT_FlipResults);
 
 			//update external SQ structure
 			//for now just copy the whole thing, stomping any changes that came from GT
-			if (Chaos::FPhysicsSolver* SceneSolver = Scene.GetSolver())
-			{
-				Scene.SolverAccelerationStructure = SceneSolver->GetEvolution()->GetCollisionConstraints().CreateExternalAccelerationStructure();
-			}
+			Scene.CopySolverAccelerationStructure();
 
 			const TArray<FPhysicsSolver*>& SolverList = SolverModule->GetSolvers();
 			TArray<FPhysicsSolver*> ActiveSolvers;
@@ -1234,22 +1256,6 @@ void FPhysScene_ChaosInterface::EndFrame(ULineBatchComponent* InLineBatcher)
 
 			const int32 NumActiveSolvers = ActiveSolvers.Num();
 
-			PhysicsParallelFor(NumActiveSolvers, [&](int32 Index)
-			{
-				FPhysicsSolver* Solver = ActiveSolvers[Index];
-
-				Solver->ForEachPhysicsProxy([](auto* Object)
-				{
-					Object->FlipBuffer();
-				});
-
-				Solver->ForEachPhysicsProxy([](auto* Object)
-				{
-					Object->PullFromPhysicsState();
-				});
-
-			});
-
 			for (FPhysicsSolver* Solver : ActiveSolvers)
 			{
 				SyncBodies(Solver);
@@ -1271,7 +1277,11 @@ void FPhysScene_ChaosInterface::EndFrame(ULineBatchComponent* InLineBatcher)
 
 void FPhysScene_ChaosInterface::WaitPhysScenes()
 {
-
+	if (CompletionEvent && !CompletionEvent->IsComplete())
+	{
+		QUICK_SCOPE_CYCLE_COUNTER(STAT_FPhysScene_WaitPhysScenes);
+		FTaskGraphInterface::Get().WaitUntilTaskCompletes(CompletionEvent, ENamedThreads::GameThread);
+	}
 }
 
 FGraphEventRef FPhysScene_ChaosInterface::GetCompletionEvent()
@@ -1331,99 +1341,54 @@ void FPhysScene_ChaosInterface::KillVisualDebugger()
 
 }
 
-// bool FPhysScene_ChaosInterface::ExecPxVis(uint32 SceneType, const TCHAR* Cmd, FOutputDevice* Ar)
-// {
-// 	return true;
-// }
-// 
-// bool FPhysScene_ChaosInterface::ExecApexVis(uint32 SceneType, const TCHAR* Cmd, FOutputDevice* Ar)
-// {
-// 	return true;
-// }
-
-
-void FPhysScene_ChaosInterface::SyncBodies()
+void FPhysScene_ChaosInterface::SyncBodies(Chaos::FPhysicsSolver* Solver)
 {
+	DECLARE_SCOPE_CYCLE_COUNTER(TEXT("SyncBodies"), STAT_SyncBodies, STATGROUP_Physics);
 	TArray<FPhysScenePendingComponentTransform_Chaos> PendingTransforms;
 
-#if TODO_REIMPLEMENT_GET_RIGID_PARTICLES
-	for(uint32 Index = 0; Index < Scene.GetSolver()->GetRigidParticles().Size(); ++Index)
-	{
-		if(BodyInstances[Index])
-		{
-			Chaos::TRigidTransform<float, 3> NewTransform(Scene.GetSolver()->GetRigidParticles().X(Index), Scene.GetSolver()->GetRigidParticles().R(Index));
-			FPhysScenePendingComponentTransform_Chaos NewEntry(BodyInstances[Index]->OwnerComponent.Get(), NewTransform);
-			PendingTransforms.Add(NewEntry);
-		}
-	}
-#endif
+	Chaos::FPBDRigidActiveParticlesBufferAccessor Accessor(Solver->GetActiveParticlesBuffer());
 
-	for(FPhysScenePendingComponentTransform_Chaos& Entry : PendingTransforms)
+	const Chaos::FPBDRigidActiveParticlesBufferOut* ActiveParticleBuffer = Accessor.GetSolverOutData();
+	for (Chaos::TGeometryParticle<float, 3>* ActiveParticle : ActiveParticleBuffer->ActiveGameThreadParticles)
 	{
-		UPrimitiveComponent* OwnerComponent = Entry.OwningComp.Get();
-		if(OwnerComponent != nullptr)
+		if (IPhysicsProxyBase * ProxyBase = ActiveParticle->Proxy)
 		{
-			AActor* Owner = OwnerComponent->GetOwner();
-
-			if(!Entry.NewTransform.EqualsNoScale(OwnerComponent->GetComponentTransform()))
-			{
-				const FVector MoveBy = Entry.NewTransform.GetLocation() - OwnerComponent->GetComponentTransform().GetLocation();
-				const FQuat NewRotation = Entry.NewTransform.GetRotation();
-
-				OwnerComponent->MoveComponent(MoveBy, NewRotation, false, NULL, MOVECOMP_SkipPhysicsMove);
-			}
-
-			if(Owner != NULL && !Owner->IsPendingKill())
+			if (ProxyBase->GetType() == EPhysicsProxyType::SingleRigidParticleType)
 			{
-				Owner->CheckStillInWorld();
-			}
-		}
-	}
-}
+				FSingleParticlePhysicsProxy< Chaos::TPBDRigidParticle<float, 3> > * Proxy = static_cast<FSingleParticlePhysicsProxy< Chaos::TPBDRigidParticle<float, 3> >*>(ProxyBase);
+				Proxy->PullFromPhysicsState();
 
-void FPhysScene_ChaosInterface::SyncBodies(Chaos::FPhysicsSolver* Solver)
-{
-	TArray<FPhysScenePendingComponentTransform_Chaos> PendingTransforms;
+				if (FBodyInstance* BodyInstance = FPhysicsUserData::Get<FBodyInstance>(ActiveParticle->UserData()))
+				{
+					if (BodyInstance->InstanceBodyIndex == INDEX_NONE && BodyInstance->OwnerComponent.IsValid())
+					{
+						UPrimitiveComponent* OwnerComponent = BodyInstance->OwnerComponent.Get();
+						if (OwnerComponent != nullptr)
+						{
+							AActor* Owner = OwnerComponent->GetOwner();
 
-	Solver->ForEachPhysicsProxy([this, Solver, &PendingTransforms](auto* Object)
-	{
-		FBodyInstance* BodyInstance = FPhysicsUserData::Get<FBodyInstance>(Object->GetUserData());
-		if (BodyInstance)
-		{
-			if (BodyInstance->InstanceBodyIndex == INDEX_NONE && BodyInstance->OwnerComponent.IsValid())
-			{
-				Chaos::TRigidTransform<float, 3> NewTransform(Object->GetTransform());
-				FPhysScenePendingComponentTransform_Chaos NewEntry(BodyInstance->OwnerComponent.Get(), NewTransform);
-				PendingTransforms.Add(NewEntry);
-			}
-		}
-	});
+							Chaos::TRigidTransform<float, 3> NewTransform(ActiveParticle->X(), ActiveParticle->R());
 
-	for(FPhysScenePendingComponentTransform_Chaos& Entry : PendingTransforms)
-	{
-		UPrimitiveComponent* OwnerComponent = Entry.OwningComp.Get();
-		if(OwnerComponent != nullptr)
-		{
-			AActor* Owner = OwnerComponent->GetOwner();
-
-			if(!Entry.NewTransform.EqualsNoScale(OwnerComponent->GetComponentTransform()))
-			{
-				const FVector MoveBy = Entry.NewTransform.GetLocation() - OwnerComponent->GetComponentTransform().GetLocation();
-				const FQuat NewRotation = Entry.NewTransform.GetRotation();
+							if (!NewTransform.EqualsNoScale(OwnerComponent->GetComponentTransform()))
+							{
+								const FVector MoveBy = NewTransform.GetLocation() - OwnerComponent->GetComponentTransform().GetLocation();
+								const FQuat NewRotation = NewTransform.GetRotation();
 
-				OwnerComponent->MoveComponent(MoveBy, NewRotation, false, NULL, MOVECOMP_SkipPhysicsMove);
-			}
+								OwnerComponent->MoveComponent(MoveBy, NewRotation, false, NULL, MOVECOMP_SkipPhysicsMove);
+							}
 
-			if(Owner != NULL && !Owner->IsPendingKill())
-			{
-				Owner->CheckStillInWorld();
+							if (Owner != NULL && !Owner->IsPendingKill())
+							{
+								Owner->CheckStillInWorld();
+							}
+						}
+					}
+				}
 			}
 		}
 	}
 }
 
-
-
 FPhysicsConstraintReference_Chaos 
 FPhysScene_ChaosInterface::AddSpringConstraint(const TArray< TPair<FPhysicsActorHandle, FPhysicsActorHandle> >& Constraint)
 {
@@ -1442,6 +1407,7 @@ void FPhysScene_ChaosInterface::CompleteSceneSimulation(ENamedThreads::Type Curr
 
 	// Cache our results to the threaded buffer.
 	{
+		LLM_SCOPE(ELLMTag::Chaos);
 		SCOPE_CYCLE_COUNTER(STAT_BufferPhysicsResults);
 
 		FChaosSolversModule* Module = FChaosSolversModule::GetModule();
@@ -1467,10 +1433,9 @@ void FPhysScene_ChaosInterface::CompleteSceneSimulation(ENamedThreads::Type Curr
 		{
 			FPhysicsSolver* Solver = ActiveSolvers[Index];
 
-			Solver->ForEachPhysicsProxyParallel([](auto* Object)
-			{
-				Object->BufferPhysicsResults();
-			});
+			Solver->GetActiveParticlesBuffer()->CaptureSolverData(Solver);
+			Solver->BufferPhysicsResults();
+			Solver->FlipBuffers();
 		});
 	}
 }
@@ -1478,5 +1443,3 @@ void FPhysScene_ChaosInterface::CompleteSceneSimulation(ENamedThreads::Type Curr
 TSharedPtr<IPhysicsReplicationFactory> FPhysScene_ChaosInterface::PhysicsReplicationFactory;
 
 #endif // WITH_CHAOS
-
-#endif // INCLUDE_CHAOS
diff --git a/Engine/Source/Runtime/Engine/Private/PhysicsEngine/ImmediatePhysics/ImmediatePhysicsChaos/ImmediatePhysicsActorHandle_Chaos.cpp b/Engine/Source/Runtime/Engine/Private/PhysicsEngine/ImmediatePhysics/ImmediatePhysicsChaos/ImmediatePhysicsActorHandle_Chaos.cpp
--- a/Engine/Source/Runtime/Engine/Private/PhysicsEngine/ImmediatePhysics/ImmediatePhysicsChaos/ImmediatePhysicsActorHandle_Chaos.cpp
+++ b/Engine/Source/Runtime/Engine/Private/PhysicsEngine/ImmediatePhysics/ImmediatePhysicsChaos/ImmediatePhysicsActorHandle_Chaos.cpp
@@ -2,91 +2,199 @@
 
 #include "Physics/ImmediatePhysics/ImmediatePhysicsChaos/ImmediatePhysicsActorHandle_Chaos.h"
 
-#if INCLUDE_CHAOS
-
 #include "Chaos/Box.h"
 #include "Chaos/Capsule.h"
-#include "Chaos/GeometryParticles.h"
+#include "Chaos/MassProperties.h"
 #include "Chaos/ParticleHandle.h"
 #include "Chaos/PBDRigidParticles.h"
 #include "Chaos/PBDRigidsEvolutionGBF.h"
 #include "Chaos/Sphere.h"
 
+#include "Physics/Experimental/ChaosInterfaceUtils.h"
+
 #include "PhysicsEngine/BodyInstance.h"
 #include "PhysicsEngine/BodySetup.h"
 
+extern int32 ImmediatePhysicsDisableCollisions;
+
 namespace ImmediatePhysics_Chaos
 {
 	//
 	// Utils
 	//
 
-	bool CreateGeometry(FBodyInstance* BodyInstance, const FVector& Scale, TUniquePtr<Chaos::TImplicitObject<FReal, Dimensions>>& OutGeom, TArray<TUniquePtr<Chaos::TPerShapeData<float, 3>>>& OutShapes)
+	template<typename T, int d>
+	Chaos::PMatrix<T, d, d> CalculateInertia_Solid(T Mass, const FKSphereElem& SphereElem)
 	{
-		// @todo(ccaulfield): copied from FPhysInterface_Chaos::AddGeometry. Merge implementations.
-		TArray<TUniquePtr<Chaos::TImplicitObject<float, 3>>> Geoms;
-		TArray<TUniquePtr<Chaos::TPerShapeData<float, 3>>> Shapes;
+		return Chaos::PMatrix<T, d, d>(
+			((T)2 / (T)5) * Mass * SphereElem.Radius * SphereElem.Radius,
+			((T)2 / (T)5) * Mass * SphereElem.Radius * SphereElem.Radius,
+			((T)2 / (T)5) * Mass * SphereElem.Radius * SphereElem.Radius
+			);
+	}
 
-		FBodyCollisionData CollisionData;
-		BodyInstance->BuildBodyFilterData(CollisionData.CollisionFilterData);
-		BodyInstance->BuildBodyCollisionFlags(CollisionData.CollisionFlags, BodyInstance->GetCollisionEnabled(), BodyInstance->BodySetup->GetCollisionTraceFlag() == CTF_UseComplexAsSimple);
+	template<typename T, int d>
+	Chaos::PMatrix<T, d, d> CalculateInertia_Solid(T Mass, const FKSphylElem& SphylElem)
+	{
+		return Chaos::PMatrix<T, d, d>(
+			((T)1 / (T)12) * Mass * ((T)3 * SphylElem.Radius + SphylElem.Length),
+			((T)1 / (T)12) * Mass * ((T)3 * SphylElem.Radius + SphylElem.Length),
+			((T)1 / (T)2) * Mass * SphylElem.Radius * SphylElem.Radius
+			);
+	}
 
-		FKAggregateGeom* Geometry = &BodyInstance->BodySetup->AggGeom;
-		FVector ActorToCoMTranslation = BodyInstance->GetMassSpaceLocal().GetTranslation();
+	template<typename T, int d>
+	Chaos::PMatrix<T, d, d> CalculateInertia_Solid(T Mass, const FKBoxElem& BoxElem)
+	{
+		return Chaos::PMatrix<T, d, d>(
+			((T)1 / (T)12) * Mass * (BoxElem.Y * BoxElem.Y + BoxElem.Z * BoxElem.Z),
+			((T)1 / (T)12) * Mass * (BoxElem.Z * BoxElem.Z + BoxElem.X * BoxElem.X),
+			((T)1 / (T)12) * Mass * (BoxElem.X * BoxElem.X + BoxElem.Y * BoxElem.Y)
+			);
+	}
 
-		auto NewShapeHelper = [CollisionData](const Chaos::TImplicitObject<float, 3>* InGeom, bool bComplexShape = false)
-		{
-			auto NewShape = MakeUnique<Chaos::TPerShapeData<float, 3>>();
-			NewShape->Geometry = InGeom;
-			NewShape->QueryData = bComplexShape ? CollisionData.CollisionFilterData.QueryComplexFilter : CollisionData.CollisionFilterData.QuerySimpleFilter;
-			NewShape->SimData = CollisionData.CollisionFilterData.SimFilter;
-			return NewShape;
-		};
-
-		for (uint32 i = 0; i < static_cast<uint32>(Geometry->SphereElems.Num()); ++i)
+	template<typename T, int d>
+	void CalculateMassProperties(T Density, const FVector& Scale, const FTransform& LocalTransform, const FKAggregateGeom& AggGeom, T& OutMass, Chaos::TMassProperties<T, d>& OutMassProperties)
+	{
+		using namespace Chaos;
+		TArray<TMassProperties<T, d>> AllMassProperties;
+		T AllMass = 0;
+
+		for (uint32 i = 0; i < static_cast<uint32>(AggGeom.SphereElems.Num()); ++i)
 		{
-			ensure(FMath::IsNearlyEqual(Scale[0], Scale[1]) && FMath::IsNearlyEqual(Scale[1], Scale[2]));
-			const auto& CollisionSphere = Geometry->SphereElems[i];
-			auto ImplicitSphere = MakeUnique<Chaos::TSphere<float, 3>>(Chaos::TVector<float, 3>(0.f, 0.f, 0.f), CollisionSphere.Radius * Scale[0]);
-			auto NewShape = NewShapeHelper(ImplicitSphere.Get());
-			Shapes.Emplace(MoveTemp(NewShape));
-			Geoms.Add(MoveTemp(ImplicitSphere));
+			const FKSphereElem ScaledSphereElem = AggGeom.SphereElems[i].GetFinalScaled(Scale, LocalTransform);
+
+			T Volume = ScaledSphereElem.GetVolume(FVector::OneVector);
+			T Mass = Density * Volume;
+
+			TMassProperties<T, d> MassProperties;
+			MassProperties.CenterOfMass = LocalTransform.GetTranslation() + ScaledSphereElem.Center;
+			MassProperties.RotationOfMass = TRotation<T, d>::FromIdentity();
+			MassProperties.Volume = Volume;
+			MassProperties.InertiaTensor = CalculateInertia_Solid<T, d>(Mass, ScaledSphereElem);
 
+			AllMassProperties.Add(MassProperties);
+			AllMass += Mass;
 		}
-		for (uint32 i = 0; i < static_cast<uint32>(Geometry->BoxElems.Num()); ++i)
+		for (uint32 i = 0; i < static_cast<uint32>(AggGeom.BoxElems.Num()); ++i)
 		{
-			const auto& Box = Geometry->BoxElems[i];
-			Chaos::TVector<float, 3> half_extents = Scale * Chaos::TVector<float, 3>(Box.X / 2.f, Box.Y / 2.f, Box.Z / 2.f);
-			auto ImplicitBox = MakeUnique<Chaos::TBox<float, 3>>(-half_extents, half_extents);
-			auto NewShape = NewShapeHelper(ImplicitBox.Get());
-			Shapes.Emplace(MoveTemp(NewShape));
-			Geoms.Add(MoveTemp(ImplicitBox));
+			const auto& BoxElem = AggGeom.BoxElems[i];
 
+			T Volume = BoxElem.GetVolume(Scale);
+			T Mass = Density * Volume;
+
+			TMassProperties<T, d> MassProperties;
+			MassProperties.CenterOfMass = LocalTransform.GetTranslation() + BoxElem.Center;
+			MassProperties.RotationOfMass = LocalTransform.GetRotation() * TRotation<T, d>(FQuat(BoxElem.Rotation));
+			MassProperties.Volume = Volume;
+			MassProperties.InertiaTensor = CalculateInertia_Solid<T, d>(Mass, BoxElem);
+
+			AllMassProperties.Add(MassProperties);
+			AllMass += Mass;
 		}
-		for (uint32 i = 0; i < static_cast<uint32>(Geometry->SphylElems.Num()); ++i)
+		for (uint32 i = 0; i < static_cast<uint32>(AggGeom.SphylElems.Num()); ++i)
 		{
-			ensure(FMath::IsNearlyEqual(Scale[0], Scale[1]) && FMath::IsNearlyEqual(Scale[1], Scale[2]));
-			const auto& Sphyl = Geometry->SphylElems[i];
-			if (Sphyl.Length == 0)
+			const FKSphylElem& UnscaledSphyl = AggGeom.SphylElems[i];
+			const FKSphylElem ScaledSphylElem = UnscaledSphyl.GetFinalScaled(Scale, LocalTransform);
+			float HalfHeight = FMath::Max(ScaledSphylElem.Length * 0.5f, KINDA_SMALL_NUMBER);
+			const float Radius = FMath::Max(ScaledSphylElem.Radius, KINDA_SMALL_NUMBER);
+			if (HalfHeight < KINDA_SMALL_NUMBER)
 			{
-				auto ImplicitSphere = MakeUnique<Chaos::TSphere<float, 3>>(Chaos::TVector<float, 3>(0), Sphyl.Radius * Scale[0]);
-				auto NewShape = NewShapeHelper(ImplicitSphere.Get());
-				Shapes.Emplace(MoveTemp(NewShape));
-				Geoms.Add(MoveTemp(ImplicitSphere));
+				//not a capsule just use a sphere
+				const FKSphereElem ScaledSphereElem = FKSphereElem(Radius);
+
+				T Volume = ScaledSphereElem.GetVolume(FVector::OneVector);
+				T Mass = Density * Volume;
+
+				TMassProperties<T, d> MassProperties;
+				MassProperties.CenterOfMass = LocalTransform.GetTranslation() + ScaledSphereElem.Center;
+				MassProperties.RotationOfMass = TRotation<T, d>::FromIdentity();
+				MassProperties.Volume = Volume;
+				MassProperties.InertiaTensor = CalculateInertia_Solid<T, d>(Mass, ScaledSphereElem);
 
+				AllMassProperties.Add(MassProperties);
+				AllMass += Mass;
 			}
 			else
 			{
-				Chaos::TVector<float, 3> HalfExtents(0, 0, Sphyl.Length / 2 * Scale[0]);
-				FTransform SphylTransform = Sphyl.GetTransform();
-				SphylTransform.SetTranslation(SphylTransform.GetTranslation() - ActorToCoMTranslation);
-
-				auto ImplicitCapsule = MakeUnique<Chaos::TCapsule<float>>(SphylTransform.TransformPosition(-HalfExtents), SphylTransform.TransformPosition(HalfExtents), Sphyl.Radius * Scale[0]);
-				auto NewShape = NewShapeHelper(ImplicitCapsule.Get());
-				Shapes.Emplace(MoveTemp(NewShape));
-				Geoms.Add(MoveTemp(ImplicitCapsule));
+				T Volume = ScaledSphylElem.GetVolume(FVector::OneVector);
+				T Mass = Density * Volume;
+
+				TMassProperties<T, d> MassProperties;
+				MassProperties.CenterOfMass = LocalTransform.GetTranslation() + ScaledSphylElem.Center;
+				MassProperties.RotationOfMass = LocalTransform.GetRotation() * TRotation<T, d>(FQuat(ScaledSphylElem.Rotation));
+				MassProperties.Volume = Volume;
+				MassProperties.InertiaTensor = CalculateInertia_Solid<T, d>(Mass, ScaledSphylElem);
+
+				AllMassProperties.Add(MassProperties);
+				AllMass += Mass;
 			}
 		}
+#if WITH_CHAOS && !PHYSICS_INTERFACE_PHYSX
+		for (uint32 i = 0; i < static_cast<uint32>(AggGeom.ConvexElems.Num()); ++i)
+		{
+			const FKConvexElem& CollisionBody = AggGeom.ConvexElems[i];
+			if (const auto& ConvexImplicit = CollisionBody.GetChaosConvexMesh())
+			{
+				// @todo(ccaulfield): calculate inertia of convex
+			}
+		}
+
+		// @todo(ccaulfield): tri meshes...
+		//for (const auto& ChaosTriMesh : InParams.ChaosTriMeshes)
+		//{
+		//}
+#endif
+		static bool bModeThanOne = false;
+		if (AllMassProperties.Num() > 1)
+		{
+			bModeThanOne = true;
+		}
+
+		OutMassProperties = Combine(AllMassProperties);
+		OutMass = AllMass;
+	}
+
+
+	bool CreateGeometry(FBodyInstance* BodyInstance, const FVector& Scale, float& OutMass, Chaos::TVector<float, 3>& OutInertia, Chaos::TRigidTransform<float, 3>& OutCoMTransform, TUniquePtr<Chaos::TImplicitObject<FReal, Dimensions>>& OutGeom, TArray<TUniquePtr<Chaos::TPerShapeData<float, 3>>>& OutShapes)
+	{
+		UBodySetup* BodySetup = BodyInstance->BodySetup.Get();
+
+#if WITH_CHAOS && !PHYSICS_INTERFACE_PHYSX
+		float Density = 1.e-3f;	// 1g/cm3
+		Chaos::TMassProperties<float, 3> MassProperties;
+		CalculateMassProperties<float, 3>(Density, Scale, FTransform::Identity, BodySetup->AggGeom, OutMass, MassProperties);
+		OutInertia = Chaos::TVector<float, 3>(MassProperties.InertiaTensor.M[0][0], MassProperties.InertiaTensor.M[1][1], MassProperties.InertiaTensor.M[2][2]);
+		OutCoMTransform = FTransform(MassProperties.RotationOfMass, MassProperties.CenterOfMass);
+#else
+		OutMass = BodyInstance->GetBodyMass();
+		OutInertia = BodyInstance->GetBodyInertiaTensor();
+		OutCoMTransform = BodyInstance->GetMassSpaceLocal();
+#endif
+
+		FBodyCollisionData BodyCollisionData;
+		BodyInstance->BuildBodyFilterData(BodyCollisionData.CollisionFilterData);
+		BodyInstance->BuildBodyCollisionFlags(BodyCollisionData.CollisionFlags, BodyInstance->GetCollisionEnabled(), BodyInstance->BodySetup->GetCollisionTraceFlag() == CTF_UseComplexAsSimple);
+
+		FGeometryAddParams AddParams;
+		AddParams.bDoubleSided = BodySetup->bDoubleSidedGeometry;
+		AddParams.CollisionData = BodyCollisionData;
+		AddParams.CollisionTraceType = BodySetup->GetCollisionTraceFlag();
+		AddParams.Scale = Scale;
+		//AddParams.SimpleMaterial = SimpleMaterial;
+		//AddParams.ComplexMaterials = TArrayView<UPhysicalMaterial*>(ComplexMaterials);
+		AddParams.LocalTransform = Chaos::TRigidTransform<float, 3>(-OutCoMTransform.GetTranslation(), OutCoMTransform.GetRotation().Inverse());
+		AddParams.Geometry = &BodySetup->AggGeom;
+#if WITH_PHYSX
+		AddParams.TriMeshes = TArrayView<PxTriangleMesh*>(BodySetup->TriMeshes);
+#endif
+#if WITH_CHAOS
+		AddParams.ChaosTriMeshes = MakeArrayView(BodySetup->ChaosTriMeshes);
+#endif
+
+		TArray<TUniquePtr<Chaos::TImplicitObject<float, 3>>> Geoms;
+		TArray<TUniquePtr<Chaos::TPerShapeData<float, 3>>, TInlineAllocator<1>> Shapes;
+		ChaosInterface::CreateGeometry(AddParams, Geoms, Shapes);
 
 		if (Geoms.Num() == 0)
 		{
@@ -102,7 +210,10 @@ namespace ImmediatePhysics_Chaos
 			OutGeom = MakeUnique<Chaos::TImplicitObjectUnion<float, 3>>(MoveTemp(Geoms));
 		}
 
-		OutShapes = MoveTemp(Shapes);
+		for (auto& Shape : Shapes)
+		{
+			OutShapes.Emplace(MoveTemp(Shape));
+		}
 
 		return true;
 	}
@@ -118,8 +229,13 @@ namespace ImmediatePhysics_Chaos
 		using namespace Chaos;
 
 		// @todo(ccaulfield): Scale
-		if (CreateGeometry(BodyInstance, FVector::OneVector, Geometry, Shapes))
+		float Mass = 0;
+		TVector<float, 3> Inertia = TVector<float, 3>::OneVector;
+		TRigidTransform<float, 3> CoMTransform = TRigidTransform<float, 3>::Identity;
+		if (CreateGeometry(BodyInstance, FVector::OneVector, Mass, Inertia, CoMTransform, Geometry, Shapes))
 		{
+			ActorToCoMTransform = CoMTransform;
+
 			switch (ActorType)
 			{
 			case EActorType::StaticActor:
@@ -135,8 +251,6 @@ namespace ImmediatePhysics_Chaos
 
 			if (ParticleHandle != nullptr)
 			{
-				ActorToCoMTranslation = BodyInstance->GetMassSpaceLocal().GetTranslation();
-
 				SetWorldTransform(Transform);
 
 				ParticleHandle->SetGeometry(MakeSerializable(Geometry));
@@ -149,9 +263,7 @@ namespace ImmediatePhysics_Chaos
 
 				if (auto* Dynamic = ParticleHandle->AsDynamic())
 				{
-					float Mass = BodyInstance->GetBodyMass();
 					float MassInv = (Mass > 0.0f) ? 1.0f / Mass : 0.0f;
-					FVector Inertia = BodyInstance->GetBodyInertiaTensor();
 					FVector InertiaInv = (Mass > 0.0f) ? Inertia.Reciprocal() : FVector::ZeroVector;
 					Dynamic->SetM(Mass);
 					Dynamic->SetInvM(MassInv);
@@ -178,6 +290,16 @@ namespace ImmediatePhysics_Chaos
 		return { ParticleHandle };
 	}
 
+	Chaos::TGeometryParticleHandle<FReal, Dimensions>* FActorHandle::GetParticle()
+	{
+		return ParticleHandle;
+	}
+
+	const Chaos::TGeometryParticleHandle<FReal, Dimensions>* FActorHandle::GetParticle() const
+	{
+		return ParticleHandle;
+	}
+
 	void FActorHandle::SetEnabled(bool bEnabled)
 	{
 		if (auto* Dynamic = ParticleHandle->AsDynamic())
@@ -188,7 +310,7 @@ namespace ImmediatePhysics_Chaos
 
 	void FActorHandle::SetWorldTransform(const FTransform& WorldTM)
 	{
-		FTransform ParticleTransform = FTransform(WorldTM.GetRotation(), WorldTM.TransformPosition(ActorToCoMTranslation));
+		FTransform ParticleTransform = ActorToCoMTransform * WorldTM;
 
 		ParticleHandle->SetX(ParticleTransform.GetTranslation());
 		ParticleHandle->SetR(ParticleTransform.GetRotation());
@@ -227,8 +349,13 @@ namespace ImmediatePhysics_Chaos
 #endif
 		if (TKinematicGeometryParticleHandle<FReal, Dimensions>* KinematicParticleHandle = ParticleHandle->AsKinematic())
 		{
-			ParticleHandle->SetX(WorldTM.GetTranslation());
-			ParticleHandle->SetR(WorldTM.GetRotation());
+			FTransform ParticleTransform = ActorToCoMTransform * WorldTM;
+
+			FMatrix AxesW = FRotationMatrix::Make(WorldTM.GetRotation());
+			FMatrix AxesP = FRotationMatrix::Make(ParticleTransform.GetRotation());
+
+			ParticleHandle->SetX(ParticleTransform.GetTranslation());
+			ParticleHandle->SetR(ParticleTransform.GetRotation());
 		}
 	}
 
@@ -249,7 +376,7 @@ namespace ImmediatePhysics_Chaos
 	FTransform FActorHandle::GetWorldTransform() const
 	{
 		FTransform ParticleTransform = FTransform(Handle()->R(), Handle()->X());
-		return FTransform(ParticleTransform.GetRotation(), ParticleTransform.GetTranslation() - ParticleTransform.TransformVector(ActorToCoMTranslation));
+		return ActorToCoMTransform.GetRelativeTransformReverse(ParticleTransform);
 	}
 
 	void FActorHandle::SetLinearVelocity(const FVector& NewLinearVelocity)
@@ -413,11 +540,19 @@ namespace ImmediatePhysics_Chaos
 		return FLT_MAX;
 	}
 
-	const FVector& FActorHandle::GetCoMTranslation() const
+	const FTransform& FActorHandle::GetLocalCoMTransform() const
 	{
-		return ActorToCoMTranslation;
+		return ActorToCoMTransform;
 	}
 
-}
+	int32 FActorHandle::GetLevel() const
+	{
+		return Level;
+	}
 
-#endif // INCLUDE_CHAOS
+	void FActorHandle::SetLevel(int32 InLevel)
+	{
+		Level = InLevel;
+	}
+
+}
diff --git a/Engine/Source/Runtime/Engine/Private/PhysicsEngine/ImmediatePhysics/ImmediatePhysicsChaos/ImmediatePhysicsJointHandle_Chaos.cpp b/Engine/Source/Runtime/Engine/Private/PhysicsEngine/ImmediatePhysics/ImmediatePhysicsChaos/ImmediatePhysicsJointHandle_Chaos.cpp
--- a/Engine/Source/Runtime/Engine/Private/PhysicsEngine/ImmediatePhysics/ImmediatePhysicsChaos/ImmediatePhysicsJointHandle_Chaos.cpp
+++ b/Engine/Source/Runtime/Engine/Private/PhysicsEngine/ImmediatePhysics/ImmediatePhysicsChaos/ImmediatePhysicsJointHandle_Chaos.cpp
@@ -1,30 +1,111 @@
 // Copyright 1998-2019 Epic Games, Inc. All Rights Reserved.
 
 #include "Physics/ImmediatePhysics/ImmediatePhysicsChaos/ImmediatePhysicsJointHandle_Chaos.h"
-
-#if INCLUDE_CHAOS
-
 #include "Physics/ImmediatePhysics/ImmediatePhysicsChaos/ImmediatePhysicsActorHandle_Chaos.h"
 
-#include "Chaos/PBD6DJointConstraints.h"
+#include "Chaos/PBDJointConstraints.h"
 
 #include "PhysicsEngine/ConstraintInstance.h"
 
+static_assert((int32)Chaos::EJointMotionType::Free == (int32)EAngularConstraintMotion::ACM_Free, "Chaos::EJointMotionType and EAngularConstraintMotion mismatch");
+static_assert((int32)Chaos::EJointMotionType::Limited == (int32)EAngularConstraintMotion::ACM_Limited, "Chaos::EJointMotionType and EAngularConstraintMotion mismatch");
+static_assert((int32)Chaos::EJointMotionType::Locked == (int32)EAngularConstraintMotion::ACM_Locked, "Chaos::EJointMotionType and EAngularConstraintMotion mismatch");
 
 namespace ImmediatePhysics_Chaos
 {
+	float ChaosImmediate_JointStiffness = 1.0f;
+	FAutoConsoleVariableRef CVarJointStiffness(TEXT("p.Chaos.ImmPhys.JointStiffness"), ChaosImmediate_JointStiffness, TEXT("Joint solver stiffness."));
+
+	float ChaosImmediate_MaxDriveStiffness = 2000.0f;
+	FAutoConsoleVariableRef CVarMaxDriveStiffness(TEXT("p.Chaos.ImmPhys.MaxDriveStiffness"), ChaosImmediate_MaxDriveStiffness, TEXT("The value of drive stiffness per unit mass that equates to full stiffness in the solver."));
+
+	int32 ChaosImmediate_ScaleDriveStiffnessByMass = 0;
+	FAutoConsoleVariableRef CVarScaleDriveStiffnessByMass(TEXT("p.Chaos.ImmPhys.ScaleDriveStiffnessByMass"), ChaosImmediate_ScaleDriveStiffnessByMass, TEXT("If true, converted stiffness is multiplied by inertia."));
+
+	// Convert the UE (PhysX) drive spring stiffness to a joint stiffness [0,1] value for use in the solver.
+	float ConvertAngularDriveStiffness(float InStiffness, const Chaos::TVector<float, 3>& IIa, const Chaos::TVector<float, 3>& IIb)
+	{
+		if (ChaosImmediate_MaxDriveStiffness > 0)
+		{
+			float IIMax = 1.0f;
+			if (ChaosImmediate_ScaleDriveStiffnessByMass)
+			{
+				IIMax = IIa.Min() + IIb.Min();
+			}
+			float Stiffness = (InStiffness / ChaosImmediate_MaxDriveStiffness) * IIMax;
+			return FMath::Clamp(Stiffness, 0.0f, 1.0f);
+		}
+		return 0;
+	}
+
+	float ConvertAngularDriveDamping(float InDamping, const Chaos::TVector<float, 3>& IIa, const Chaos::TVector<float, 3>& IIb)
+	{
+		return 0;
+	}
 
 	FJointHandle::FJointHandle(FChaosConstraintContainer* InConstraints, FConstraintInstance* ConstraintInstance, FActorHandle* Actor1, FActorHandle* Actor2)
-		: Constraints(InConstraints)
+		: ActorHandles({ Actor1, Actor2 })
+		, Constraints(InConstraints)
 	{
-		// BodyInstance/PhysX has the constraint locations in actor-space, but we need them in Centre-of-Mass space
+		using namespace Chaos;
+
+		// BodyInstance/PhysX has the constraint locations in actor-space, but we need them in Center-of-Mass space
 		// @todo(ccaulfield): support CoM in Chaos particles
 		FTransform ConstraintFrame1 = ConstraintInstance->GetRefFrame(EConstraintFrame::Frame1);
 		FTransform ConstraintFrame2 = ConstraintInstance->GetRefFrame(EConstraintFrame::Frame2);
-		ConstraintFrame1.SetTranslation(ConstraintFrame1.GetTranslation() - Actor1->GetCoMTranslation());
-		ConstraintFrame2.SetTranslation(ConstraintFrame2.GetTranslation() - Actor2->GetCoMTranslation());
 
-		ConstraintHandle = Constraints->AddConstraint({ Actor1->ParticleHandle, Actor2->ParticleHandle }, { ConstraintFrame1, ConstraintFrame2 });
+		TPBDJointSettings<float, 3> ConstraintSettings;
+		ConstraintSettings.ConstraintFrames = 
+		{ 
+			ConstraintFrame1.GetRelativeTransform(Actor1->GetLocalCoMTransform()),
+			ConstraintFrame2.GetRelativeTransform(Actor2->GetLocalCoMTransform())
+		};
+		ConstraintSettings.Motion.Stiffness = ChaosImmediate_JointStiffness;
+		ConstraintSettings.Motion.LinearMotionTypes = 
+		{
+			static_cast<EJointMotionType>(ConstraintInstance->GetLinearXMotion()),
+			static_cast<EJointMotionType>(ConstraintInstance->GetLinearYMotion()),
+			static_cast<EJointMotionType>(ConstraintInstance->GetLinearZMotion()),
+		};
+		ConstraintSettings.Motion.LinearLimit = ConstraintInstance->GetLinearLimit();
+		ConstraintSettings.Motion.AngularMotionTypes =
+		{
+			static_cast<EJointMotionType>(ConstraintInstance->GetAngularTwistMotion()),
+			static_cast<EJointMotionType>(ConstraintInstance->GetAngularSwing1Motion()),
+			static_cast<EJointMotionType>(ConstraintInstance->GetAngularSwing2Motion())
+		};
+		ConstraintSettings.Motion.AngularLimits =
+		{
+			FMath::DegreesToRadians(ConstraintInstance->GetAngularTwistLimit()),
+			FMath::DegreesToRadians(ConstraintInstance->GetAngularSwing1Limit()),
+			FMath::DegreesToRadians(ConstraintInstance->GetAngularSwing2Limit())
+		};
+
+		// @todo(ccaulfield): Remove one of these
+		ConstraintSettings.Motion.AngularDriveTarget = FQuat(ConstraintInstance->ProfileInstance.AngularDrive.OrientationTarget);
+		ConstraintSettings.Motion.AngularDriveTargetAngles = TVector<float, 3>(
+			FMath::DegreesToRadians(ConstraintInstance->ProfileInstance.AngularDrive.OrientationTarget.Roll), 
+			FMath::DegreesToRadians(ConstraintInstance->ProfileInstance.AngularDrive.OrientationTarget.Pitch),
+			FMath::DegreesToRadians(ConstraintInstance->ProfileInstance.AngularDrive.OrientationTarget.Yaw));
+
+		if (ConstraintInstance->ProfileInstance.AngularDrive.AngularDriveMode == EAngularDriveMode::SLERP)
+		{
+			ConstraintSettings.Motion.bAngularSLerpDriveEnabled = ConstraintInstance->ProfileInstance.AngularDrive.SlerpDrive.bEnablePositionDrive;
+			ConstraintSettings.Motion.bAngularTwistDriveEnabled = false;
+			ConstraintSettings.Motion.bAngularSwingDriveEnabled = false;
+			
+		}
+		else
+		{
+			ConstraintSettings.Motion.bAngularSLerpDriveEnabled = false;
+			ConstraintSettings.Motion.bAngularTwistDriveEnabled = ConstraintInstance->ProfileInstance.AngularDrive.TwistDrive.bEnablePositionDrive;
+			ConstraintSettings.Motion.bAngularSwingDriveEnabled = ConstraintInstance->ProfileInstance.AngularDrive.SwingDrive.bEnablePositionDrive;
+		}
+
+		ConstraintSettings.Motion.AngularDriveStiffness = ConvertAngularDriveStiffness(ConstraintInstance->ProfileInstance.AngularDrive.TwistDrive.Stiffness, Actor1->GetInverseInertia(), Actor2->GetInverseInertia());
+		ConstraintSettings.Motion.AngularDriveDamping = ConvertAngularDriveDamping(ConstraintInstance->ProfileInstance.AngularDrive.TwistDrive.Damping, Actor1->GetInverseInertia(), Actor2->GetInverseInertia());
+
+		ConstraintHandle = Constraints->AddConstraint({ Actor1->ParticleHandle, Actor2->ParticleHandle }, ConstraintSettings);
 	}
 
 	FJointHandle::~FJointHandle()
@@ -32,6 +113,28 @@ namespace ImmediatePhysics_Chaos
 		ConstraintHandle->RemoveConstraint();
 	}
 
-}
+	typename FJointHandle::FChaosConstraintHandle* FJointHandle::GetConstraint()
+	{
+		return ConstraintHandle;
+	}
+	
+	const typename FJointHandle::FChaosConstraintHandle* FJointHandle::GetConstraint() const
+	{
+		return ConstraintHandle;
+	}
+
+	const Chaos::TVector<FActorHandle*, 2>& FJointHandle::GetActorHandles()
+	{
+		return ActorHandles;
+	}
+
+	const Chaos::TVector<const FActorHandle*, 2>& FJointHandle::GetActorHandles() const
+	{
+		return reinterpret_cast<const Chaos::TVector<const FActorHandle*, 2>&>(ActorHandles);
+	}
 
-#endif
+	void FJointHandle::UpdateLevels()
+	{
+		ConstraintHandle->SetParticleLevels({ ActorHandles[0]->GetLevel(), ActorHandles[1]->GetLevel() });
+	}
+}
diff --git a/Engine/Source/Runtime/Engine/Private/PhysicsEngine/ImmediatePhysics/ImmediatePhysicsChaos/ImmediatePhysicsSimulation_Chaos.cpp b/Engine/Source/Runtime/Engine/Private/PhysicsEngine/ImmediatePhysics/ImmediatePhysicsChaos/ImmediatePhysicsSimulation_Chaos.cpp
--- a/Engine/Source/Runtime/Engine/Private/PhysicsEngine/ImmediatePhysics/ImmediatePhysicsChaos/ImmediatePhysicsSimulation_Chaos.cpp
+++ b/Engine/Source/Runtime/Engine/Private/PhysicsEngine/ImmediatePhysics/ImmediatePhysicsChaos/ImmediatePhysicsSimulation_Chaos.cpp
@@ -1,23 +1,119 @@
 // Copyright 1998-2019 Epic Games, Inc. All Rights Reserved.
 
 #include "Physics/ImmediatePhysics/ImmediatePhysicsChaos/ImmediatePhysicsSimulation_Chaos.h"
-
-#if INCLUDE_CHAOS
-
 #include "Physics/ImmediatePhysics/ImmediatePhysicsChaos/ImmediatePhysicsActorHandle_Chaos.h"
 #include "Physics/ImmediatePhysics/ImmediatePhysicsChaos/ImmediatePhysicsJointHandle_Chaos.h"
 #include "Physics/ImmediatePhysics/ImmediatePhysicsStats.h"
 
+#include "Chaos/ChaosDebugDraw.h"
+#include "Chaos/DebugDrawQueue.h"
 #include "Chaos/ParticleHandle.h"
-#include "Chaos/PBD6DJointConstraints.h"
+#include "Chaos/PBDCollisionConstraint.h"
+#include "Chaos/PBDJointConstraints.h"
 #include "Chaos/PBDRigidParticles.h"
 #include "Chaos/PBDRigidsEvolutionGBF.h"
+#include "ChaosLog.h"
 
 #include "PhysicsEngine/BodyInstance.h"
 #include "PhysicsEngine/ConstraintInstance.h"
 
+// @todo(ccaulfield): remove when finished
+float ChaosImmediate_Evolution_DeltaTime = 0.03f;
+int32 ChaosImmediate_Evolution_Iterations = 10;
+int32 ChaosImmediate_Collision_Enabled = 1;
+int32 ChaosImmediate_Collision_ApplyEnabled = 0;
+int32 ChaosImmediate_Collision_PushOutIterations = 5;
+int32 ChaosImmediate_Collision_PushOutPairIterations = 2;
+float ChaosImmediate_Collision_Thickness = 0;
+
+FAutoConsoleVariableRef CVarChaosImmPhysDeltaTime(TEXT("p.Chaos.ImmPhys.DeltaTime"), ChaosImmediate_Evolution_DeltaTime, TEXT("Override chaos immediate physics delta time if non-zero"));
+FAutoConsoleVariableRef CVarChaosImmPhysIterations(TEXT("p.Chaos.ImmPhys.Iterations"), ChaosImmediate_Evolution_Iterations, TEXT("Number of constraint solver loops in immediate physics"));
+FAutoConsoleVariableRef CVarImmediatePhysicsDisableCollisions(TEXT("p.Chaos.ImmPhys.EnableCollisions"), ChaosImmediate_Collision_Enabled, TEXT("Enable/Disable collisions in Immediate Physics."));
+FAutoConsoleVariableRef CVarCollisionDisableApply(TEXT("p.Chaos.ImmPhys.CollisionEnableApply"), ChaosImmediate_Collision_ApplyEnabled, TEXT("Enable/Disable the Apply() (velocity, friction, restitution) method for collisions"));
+FAutoConsoleVariableRef CVarCollisionPushOutIterations(TEXT("p.Chaos.ImmPhys.CollisionPushOutIterations"), ChaosImmediate_Collision_PushOutIterations, TEXT("Set the ApplyPushOut() (position correction) iteration count (0 to disable)"));
+FAutoConsoleVariableRef CVarCollisionPushOutPairIterations(TEXT("p.Chaos.ImmPhys.CollisionPushOutPairIterations"), ChaosImmediate_Collision_PushOutPairIterations, TEXT("Set the ApplyPushOut() internal pair interations (position correction) iteration count (0 to disable)"));
+FAutoConsoleVariableRef CVarChaosImmPhysThickness(TEXT("p.Chaos.ImmPhys.CollisionThickness"), ChaosImmediate_Collision_Thickness, TEXT("ChaosImmediateThickness"));
+
+float ChaosImmediate_Joint_SwingTwistAngleTolerance = 1.0e-6f;
+int32 ChaosImmediate_Joint_EnableTwistLimits = 1;
+int32 ChaosImmediate_Joint_EnableSwingLimits = 1;
+int32 ChaosImmediate_Joint_EnableDrives = 1;
+int32 ChaosImmediate_Joint_FreezeIterations = 0;
+int32 ChaosImmediate_Joint_FrozenIterations = 1;
+float ChaosImmediate_Joint_PBDDriveStiffness = 0.0f;
+float ChaosImmediate_Joint_PBDMinParentMassRatio = 0.5f;
+float ChaosImmediate_Joint_PBDMaxInertiaRatio = 5.0f;
+FAutoConsoleVariableRef CVarSwingTwistAngleTolerance(TEXT("p.Chaos.ImmPhys.Joint.SwingTwistAngleTolerance"), ChaosImmediate_Joint_SwingTwistAngleTolerance, TEXT("SwingTwistAngleTolerance."));
+FAutoConsoleVariableRef CVarEnableTwistLimits(TEXT("p.Chaos.ImmPhys.Joint.EnableTwistLimits"), ChaosImmediate_Joint_EnableTwistLimits, TEXT("EnableTwistLimits."));
+FAutoConsoleVariableRef CVarEnableSwingLimits(TEXT("p.Chaos.ImmPhys.Joint.EnableSwingLimits"), ChaosImmediate_Joint_EnableSwingLimits, TEXT("EnableSwingLimits."));
+FAutoConsoleVariableRef CVarEnableDrives(TEXT("p.Chaos.ImmPhys.Joint.EnableDrives"), ChaosImmediate_Joint_EnableDrives, TEXT("EnableDrives."));
+FAutoConsoleVariableRef CVarFreezeIterations(TEXT("p.Chaos.ImmPhys.Joint.FreezeIterations"), ChaosImmediate_Joint_FreezeIterations, TEXT("FreezeIterations."));
+FAutoConsoleVariableRef CVarFrozenIterations(TEXT("p.Chaos.ImmPhys.Joint.FrozenIterations"), ChaosImmediate_Joint_FrozenIterations, TEXT("FrozenIterations."));
+FAutoConsoleVariableRef CVarDriveStiffness(TEXT("p.Chaos.ImmPhys.Joint.PBDDriveStiffness"), ChaosImmediate_Joint_PBDDriveStiffness, TEXT("6Dof joint drive stiffness override (if > 0)."));
+FAutoConsoleVariableRef CVarMinParentMassRatio(TEXT("p.Chaos.ImmPhys.Joint.MinParentMassRatio"), ChaosImmediate_Joint_PBDMinParentMassRatio, TEXT("6Dof joint PBDMinParentMassRatio (if > 0)"));
+FAutoConsoleVariableRef CVarMaxInertiaRatio(TEXT("p.Chaos.ImmPhys.Joint.MaxInertiaRatio"), ChaosImmediate_Joint_PBDMaxInertiaRatio, TEXT("6Dof joint PBDMaxInertiaRatio (if > 0)"));
+
+// DebugDraw CVars
+#if UE_BUILD_DEBUG
+int32 ChaosImmediate_DebugDrawParticles = 1;
+int32 ChaosImmediate_DebugDrawShapes = 1;
+int32 ChaosImmediate_DebugDrawCollisions = 3;
+int32 ChaosImmediate_DebugDrawJoints = 1;
+#else
+int32 ChaosImmediate_DebugDrawParticles = 0;
+int32 ChaosImmediate_DebugDrawShapes = 0;
+int32 ChaosImmediate_DebugDrawCollisions = 0;
+int32 ChaosImmediate_DebugDrawJoints = 0;
+#endif
+
+FAutoConsoleVariableRef CVarDebugDrawParticles(TEXT("p.Chaos.ImmPhys.DebugDrawParticles"), ChaosImmediate_DebugDrawParticles, TEXT("Draw Particle Transforms (0 = never; 1 = end of frame; 2 = begin and end of frame; 3 = each evolution phase)."));
+FAutoConsoleVariableRef CVarDebugDrawShapes(TEXT("p.Chaos.ImmPhys.DebugDrawShapes"), ChaosImmediate_DebugDrawShapes, TEXT("Draw Shapes (0 = never; 1 = end of frame; 2 = begin and end of frame; 3 = each evolution phase."));
+FAutoConsoleVariableRef CVarDebugDrawCollisions(TEXT("p.Chaos.ImmPhys.DebugDrawCollisions"), ChaosImmediate_DebugDrawCollisions, TEXT("Draw Collisions (0 = never; 1 = end of frame; 2 = begin and end of frame; 3 = each evolution phase)"));
+FAutoConsoleVariableRef CVarDebugDrawJoints(TEXT("p.Chaos.ImmPhys.DebugDrawJoints"), ChaosImmediate_DebugDrawJoints, TEXT("Draw Joints. (0 = never; 1 = end of frame; 2 = begin and end of frame; 3 = each evolution phase; 4 = each Apply step)."));
+
 namespace ImmediatePhysics_Chaos
 {
+	template<typename T, int d>
+	bool ShouldIgnoreCollisionConstraint(
+		const Chaos::TGeometryParticleHandle<T, d>* ParticleA,
+		const Chaos::TGeometryParticleHandle<T, d>* ParticleB,
+		const TMap<const Chaos::TGeometryParticleHandle<T, d>*, TSet<const Chaos::TGeometryParticleHandle<T, d>*>>& IgnoreSetMap)
+	{
+		using namespace Chaos;
+
+		if (!ChaosImmediate_Collision_Enabled)
+		{
+			return true;
+		}
+		if (const TSet<const TGeometryParticleHandle<T, d>*>* IgnoreSet = IgnoreSetMap.Find(ParticleA))
+		{
+			return IgnoreSet->Contains(ParticleB);
+		}
+		if (const TSet<const TGeometryParticleHandle<T, d>*>* IgnoreSet = IgnoreSetMap.Find(ParticleB))
+		{
+			return IgnoreSet->Contains(ParticleA);
+		}
+		return false;
+	}
+
+	template<typename T, int d>
+	void FilterCollisionConstraints(
+		TArray<Chaos::TRigidBodyContactConstraint<T, d>>& Constraints, 
+		const TMap<const Chaos::TGeometryParticleHandle<T, d>*, TSet<const Chaos::TGeometryParticleHandle<T, d>*>>& IgnoreSetMap)
+	{
+		using namespace Chaos;
+
+		for (int ConstraintIndex = 0; ConstraintIndex < Constraints.Num(); ++ConstraintIndex)
+		{
+			if (ShouldIgnoreCollisionConstraint(Constraints[ConstraintIndex].Particle, Constraints[ConstraintIndex].Levelset, IgnoreSetMap))
+			{
+				Constraints.RemoveAtSwap(ConstraintIndex);
+				--ConstraintIndex;
+				continue;
+			}
+		}
+	}
+
 
 	//
 	//
@@ -25,26 +121,85 @@ namespace ImmediatePhysics_Chaos
 
 	FSimulation::FSimulation()
 		: NumActiveActorHandles(0)
-		, Gravity(FVector::ZeroVector)
+		, SimulationSpaceTransform(FTransform::Identity)
 	{
 		using namespace Chaos;
 
-		const int DefaultSolverIterations = 3;
 		Particles = MakeUnique<TPBDRigidsSOAs<FReal, Dimensions>>();
-		Joints = MakeUnique<TPBD6DJointConstraints<FReal, Dimensions>>();
-		JointsRule = MakeUnique<TPBDConstraintIslandRule<TPBD6DJointConstraints<FReal, Dimensions>, FReal, Dimensions>>(*Joints);
-		Evolution = MakeUnique<TPBDRigidsEvolutionGBF<FReal, Dimensions>>(*Particles.Get(), DefaultSolverIterations);
+		Joints = MakeUnique<TPBDJointConstraints<FReal, Dimensions>>(TPBDJointSolverSettings<FReal, Dimensions>());
+		JointsRule = MakeUnique<TPBDConstraintIslandRule<TPBDJointConstraints<FReal, Dimensions>, FReal, Dimensions>>(*Joints);
+		Evolution = MakeUnique<TPBDRigidsEvolutionGBF<FReal, Dimensions>>(*Particles.Get(), 1);
+		TPBDCollisionConstraint<FReal, Dimensions>& Collisions = Evolution->GetCollisionConstraints();
 
 		Evolution->AddConstraintRule(JointsRule.Get());
 
-		Evolution->AddForceFunction([this](TTransientPBDRigidParticleHandle<FReal, Dimensions>& Particle, const FReal Dt)
+		Collisions.SetPostComputeCallback(
+			[this](TArray<TRigidBodyContactConstraint<FReal, Dimensions>>& Constraints)
+			{
+				FilterCollisionConstraints(Constraints, IgnoreCollisionParticlePairTable);
+			});
+		
+#if CHAOS_DEBUG_DRAW
+		Evolution->SetPostIntegrateCallback(
+			[this]()
+			{
+				DebugDrawElementsAll(3, 3, 0.3f);
+			});
+		Evolution->SetPostApplyCallback(
+			[this](int32 Island)
 			{
-				Particle.F() += Gravity * Particle.M();
+				DebugDrawElementsIsland(Island, 3, 3, 0.3f);
 			});
+		Evolution->SetPostApplyPushOutCallback(
+			[this](int32 Island)
+			{
+				DebugDrawElementsIsland(Island, 3, 3, 0.3f);
+			});
+		Collisions.SetPostApplyCallback(
+			[this](const float Dt, const TArray<TPBDCollisionConstraintHandle<float, 3>*>& InConstraintHandles)
+			{
+				if (ChaosImmediate_DebugDrawCollisions == 4)
+				{
+					DebugDraw::DrawCollisions(SimulationSpaceTransform, InConstraintHandles, 0.3f);
+				}
+			});
+		Collisions.SetPostApplyPushOutCallback(
+			[this](const float Dt, const TArray<TPBDCollisionConstraintHandle<float, 3>*>& InConstraintHandles, bool bRequiresAnotherIteration)
+			{
+				if (ChaosImmediate_DebugDrawCollisions == 4)
+				{
+					DebugDraw::DrawCollisions(SimulationSpaceTransform, InConstraintHandles, 0.3f);
+				}
+			});
+		Joints->SetPreApplyCallback(
+			[this](const float Dt, const TArray<TPBDJointConstraintHandle<float, 3>*>& InConstraintHandles)
+			{
+				if (ChaosImmediate_DebugDrawJoints == 4)
+				{
+					DebugDraw::DrawJointConstraints(SimulationSpaceTransform, InConstraintHandles, 0.3f);
+				}
+			});
+		Joints->SetPostApplyCallback(
+			[this](const float Dt, const TArray<TPBDJointConstraintHandle<float, 3>*>& InConstraintHandles)
+			{
+				if (ChaosImmediate_DebugDrawJoints == 4)
+				{
+					DebugDraw::DrawJointConstraints(SimulationSpaceTransform, InConstraintHandles, 1.0f);
+				}
+			});
+#endif
 	}
 
 	FSimulation::~FSimulation()
 	{
+		using namespace Chaos;
+
+		Evolution->GetCollisionConstraints().ClearPostComputeCallback();
+		Evolution->GetCollisionConstraints().ClearPostApplyCallback();
+		Evolution->GetCollisionConstraints().ClearPostApplyPushOutCallback();
+
+		//Evolution->RemoveConstraintRule(JointsRule.Get());
+
 		for (FActorHandle* ActorHandle : ActorHandles)
 		{
 			delete ActorHandle;
@@ -129,27 +284,95 @@ namespace ImmediatePhysics_Chaos
 
 	void FSimulation::SetIgnoreCollisionPairTable(const TArray<FIgnorePair>& InIgnoreCollisionPairTable)
 	{
-#if IMMEDIATEPHYSICS_CHAOS_TODO
-		IgnoreCollisionPairTable.Reset();
-		for (const FIgnorePair& IgnorePair : InIgnoreTable)
+		using namespace Chaos;
+
+		IgnoreCollisionParticlePairTable.Reset();
+		for (const FIgnorePair& IgnorePair : InIgnoreCollisionPairTable)
 		{
-			TSet<FActorHandle*>& IgnoreActorsA = IgnoreCollisionPairTable.FindOrAdd(IgnorePair.A);
-			IgnoreActorsA.Add(IgnorePair.B);
+			if(!IgnorePair.A || !IgnorePair.B)
+			{
+				continue;
+			}
+
+			const TGeometryParticleHandle<FReal, Dimensions>* ParticleA = IgnorePair.A->GetParticle();
+			const TGeometryParticleHandle<FReal, Dimensions>* ParticleB = IgnorePair.B->GetParticle();
 
-			TSet<FActorHandle*>& IgnoreActorsB = IgnoreCollisionPairTable.FindOrAdd(IgnorePair.B);
-			IgnoreActorsB.Add(IgnorePair.A);
+			IgnoreCollisionParticlePairTable.FindOrAdd(ParticleA).Add(ParticleB); 
+			IgnoreCollisionParticlePairTable.FindOrAdd(ParticleB).Add(ParticleA);
 		}
-		bRecreateIterationCache = true;
-#endif
 	}
 
 	void FSimulation::SetIgnoreCollisionActors(const TArray<FActorHandle*>& InIgnoreCollisionActors)
 	{
-#if IMMEDIATEPHYSICS_CHAOS_TODO
-		IgnoreCollisionActors.Reset();
-		IgnoreCollisionActors.Append(InIgnoreCollisionActors);
-		bRecreateIterationCache = true;
-#endif
+		using namespace Chaos;
+
+		for (FActorHandle* ActorHandle : InIgnoreCollisionActors)
+		{
+			TPBDRigidParticleHandle<FReal, Dimensions>* Particle = ActorHandle->GetParticle()->AsDynamic();
+			if (Particle != nullptr)
+			{
+				Particle->SetCollisionGroup(INDEX_NONE);
+			}
+		}
+	}
+
+	void FSimulation::ConditionConstraints()
+	{
+		// Assign levels to actors based on connection-distance to a non-dynamic actor
+		// This is used to sort constraints and set parent/child relationship on the constrained particles in a constraint
+		// @todo(ccaulfield): this should use the constraint graph and should only update when constraint connectivity changes or particles change type
+		using namespace Chaos;
+
+		TMap<FActorHandle*, TArray<FJointHandle*>> ActorJoints;
+		TArray<FActorHandle*> ActorQueue;
+		ActorQueue.Reserve(ActorHandles.Num());
+
+		// Reset all actor levels
+		for (FActorHandle* ActorHandle : ActorHandles)
+		{
+			ActorHandle->SetLevel(INDEX_NONE);
+			ActorJoints.Emplace(ActorHandle);
+
+			if (!ActorHandle->GetParticle()->AsDynamic())
+			{
+				ActorHandle->SetLevel(0);
+				ActorQueue.Add(ActorHandle);
+			}
+		}
+
+		// Build a list of joints per actor
+		for (FJointHandle* JointHandle : JointHandles)
+		{
+			TVector<FActorHandle*, 2> JointActorHandles = JointHandle->GetActorHandles();
+			ActorJoints[JointActorHandles[0]].Add(JointHandle);
+			ActorJoints[JointActorHandles[1]].Add(JointHandle);
+		}
+
+		// Breadth-first assign level
+		for (int32 ActorQueueIndex = 0; ActorQueueIndex < ActorQueue.Num(); ++ActorQueueIndex)
+		{
+			FActorHandle* ActorHandle = ActorQueue[ActorQueueIndex];
+			for (FJointHandle* JointHandle : ActorJoints[ActorHandle])
+			{
+				TVector<FActorHandle*, 2> JointActorHandles = JointHandle->GetActorHandles();
+				if (JointActorHandles[0]->GetLevel() == INDEX_NONE)
+				{
+					JointActorHandles[0]->SetLevel(ActorHandle->GetLevel() + 1);
+					ActorQueue.Add(JointActorHandles[0]);
+				}
+				if (JointActorHandles[1]->GetLevel() == INDEX_NONE)
+				{
+					JointActorHandles[1]->SetLevel(ActorHandle->GetLevel() + 1);
+					ActorQueue.Add(JointActorHandles[1]);
+				}
+			}
+		}
+
+		// Update constraint levels
+		for (FJointHandle* JointHandle : JointHandles)
+		{
+			JointHandle->UpdateLevels();
+		}
 	}
 
 	DECLARE_CYCLE_STAT(TEXT("FSimulation::Simulate_Chaos"), STAT_ImmediateSimulate_Chaos, STATGROUP_ImmediatePhysics);
@@ -157,11 +380,105 @@ namespace ImmediatePhysics_Chaos
 	void FSimulation::Simulate(float DeltaTime, const FVector& InGravity)
 	{
 		SCOPE_CYCLE_COUNTER(STAT_ImmediateSimulate_Chaos);
+		using namespace Chaos;
+
+		// TEMP: overrides
+		{
+			if (ChaosImmediate_Evolution_DeltaTime > 0)
+			{
+				DeltaTime = ChaosImmediate_Evolution_DeltaTime;
+			}
+			UE_LOG(LogChaosJoint, Verbose, TEXT("Simulate dt = %f"), DeltaTime);
+
+			// Overrides
+			TPBDJointSolverSettings<float, 3> JointsSettings = Joints->GetSettings();
+			JointsSettings.SwingTwistAngleTolerance = ChaosImmediate_Joint_SwingTwistAngleTolerance;
+			JointsSettings.PBDMinParentMassRatio = ChaosImmediate_Joint_PBDMinParentMassRatio;
+			JointsSettings.PBDMaxInertiaRatio = ChaosImmediate_Joint_PBDMaxInertiaRatio;
+			JointsSettings.FreezeIterations = ChaosImmediate_Joint_FreezeIterations;
+			JointsSettings.FrozenIterations = ChaosImmediate_Joint_FrozenIterations;
+			JointsSettings.bEnableTwistLimits = ChaosImmediate_Joint_EnableTwistLimits > 0;
+			JointsSettings.bEnableSwingLimits = ChaosImmediate_Joint_EnableSwingLimits > 0;
+			JointsSettings.bEnableDrives = ChaosImmediate_Joint_EnableDrives > 0;
+			JointsSettings.PBDDriveStiffness = ChaosImmediate_Joint_PBDDriveStiffness;
+			Joints->SetSettings(JointsSettings);
+
+			Evolution->SetNumIterations(ChaosImmediate_Evolution_Iterations);
 
-		Gravity = InGravity;
+			Evolution->GetCollisionConstraints().SetThickness(ChaosImmediate_Collision_Thickness);
+			Evolution->GetCollisionConstraints().SetVelocitySolveEnabled(ChaosImmediate_Collision_ApplyEnabled != 0);
+			Evolution->GetCollisionConstraints().SetPushOutPairIterations(ChaosImmediate_Collision_PushOutPairIterations);
+			Evolution->GetCollisionConstraintsRule().SetPushOutIterations(ChaosImmediate_Collision_PushOutIterations);
+		}
+
+		DebugDrawElementsAll(2, 2, 0.7f);
 
+		ConditionConstraints();
+
+		Evolution->GetGravityForces().SetAcceleration(InGravity);
 		Evolution->AdvanceOneTimeStep(DeltaTime);
+		Evolution->EndFrame(DeltaTime);
+
+		DebugDrawElementsAll(1, 3, 1.0f);
+	}
+
+	void FSimulation::DebugDrawElementsAll(const int32 MinDebugLevel, const int32 MaxDebugLevel, const float ColorScale)
+	{
+#if CHAOS_DEBUG_DRAW
+		using namespace Chaos;
+		if (FDebugDrawQueue::IsDebugDrawingEnabled())
+		{
+			if ((ChaosImmediate_DebugDrawParticles >= MinDebugLevel) && (ChaosImmediate_DebugDrawParticles <= MaxDebugLevel))
+			{
+				DebugDraw::DrawParticleTransforms(SimulationSpaceTransform, Evolution->GetParticles().GetAllParticlesView(), ColorScale);
+			}
+			if ((ChaosImmediate_DebugDrawShapes >= MinDebugLevel) && (ChaosImmediate_DebugDrawShapes <= MaxDebugLevel))
+			{
+				DebugDraw::DrawParticleShapes(SimulationSpaceTransform, Evolution->GetParticles().GetAllParticlesView(), ColorScale);
+			}
+			if ((ChaosImmediate_DebugDrawCollisions >= MinDebugLevel) && (ChaosImmediate_DebugDrawCollisions <= MaxDebugLevel))
+			{
+				DebugDraw::DrawCollisions(SimulationSpaceTransform, Evolution->GetCollisionConstraints(), ColorScale);
+			}
+			if ((ChaosImmediate_DebugDrawJoints >= MinDebugLevel) && (ChaosImmediate_DebugDrawJoints <= MaxDebugLevel))
+			{
+				DebugDraw::DrawJointConstraints(SimulationSpaceTransform, *Joints, ColorScale);
+			}
+		}
+#endif
 	}
-}
 
-#endif // INCLUDE_CHAOS
\ No newline at end of file
+	void FSimulation::DebugDrawElementsIsland(const int32 Island, const int32 MinDebugLevel, const int32 MaxDebugLevel, const float ColorScale)
+	{
+#if CHAOS_DEBUG_DRAW
+		using namespace Chaos;
+		if (FDebugDrawQueue::IsDebugDrawingEnabled())
+		{
+			if ((ChaosImmediate_DebugDrawParticles >= MinDebugLevel) && (ChaosImmediate_DebugDrawParticles <= MaxDebugLevel))
+			{
+				DebugDraw::DrawParticleTransforms(SimulationSpaceTransform, Evolution->GetIslandParticles(Island), ColorScale);
+			}
+			if ((ChaosImmediate_DebugDrawShapes >= MinDebugLevel) && (ChaosImmediate_DebugDrawShapes <= MaxDebugLevel))
+			{
+				DebugDraw::DrawParticleShapes(SimulationSpaceTransform, Evolution->GetIslandParticles(Island), ColorScale);
+			}
+			if ((ChaosImmediate_DebugDrawCollisions >= MinDebugLevel) && (ChaosImmediate_DebugDrawCollisions <= MaxDebugLevel))
+			{
+				Evolution->GetCollisionConstraintsRule().VisitIslandConstraints(Island,
+					[this, ColorScale](const TArray<TPBDCollisionConstraintHandle<float, 3>*>& ConstraintHandles)
+					{
+						DebugDraw::DrawCollisions(SimulationSpaceTransform, ConstraintHandles, ColorScale);
+					});
+			}
+			if ((ChaosImmediate_DebugDrawJoints >= MinDebugLevel) && (ChaosImmediate_DebugDrawJoints <= MaxDebugLevel))
+			{
+				JointsRule->VisitIslandConstraints(Island,
+					[this, ColorScale](const TArray<TPBDJointConstraintHandle<float, 3>*>& ConstraintHandles)
+					{
+						DebugDraw::DrawJointConstraints(SimulationSpaceTransform, *Joints, ColorScale);
+					});
+			}
+		}
+#endif
+	}
+}
diff --git a/Engine/Source/Runtime/Engine/Private/PhysicsEngine/PhysAnim.cpp b/Engine/Source/Runtime/Engine/Private/PhysicsEngine/PhysAnim.cpp
--- a/Engine/Source/Runtime/Engine/Private/PhysicsEngine/PhysAnim.cpp
+++ b/Engine/Source/Runtime/Engine/Private/PhysicsEngine/PhysAnim.cpp
@@ -525,12 +525,15 @@ void USkeletalMeshComponent::UpdateKinematicBonesToAnim(const TArray<FTransform>
 	}
 #endif
 
+#if !(WITH_CHAOS) // TODO Deferred update not implemented yet in Chaos.
+
 	// If we are only using bodies for physics, don't need to move them right away, can defer until simulation (unless told not to)
 	if (DeferralAllowed == EAllowKinematicDeferral::AllowDeferral && (bDeferKinematicBoneUpdate || BodyInstance.GetCollisionEnabled() == ECollisionEnabled::PhysicsOnly))
 	{
 		PhysScene->MarkForPreSimKinematicUpdate(this, Teleport, bNeedsSkinning);
 		return;
 	}
+#endif
 
 #if !(UE_BUILD_SHIPPING || UE_BUILD_TEST)
 	// If desired, draw the skeleton at the point where we pass it to the physics.
@@ -586,9 +589,10 @@ void USkeletalMeshComponent::UpdateKinematicBonesToAnim(const TArray<FTransform>
 						continue;
 					}
 					FPhysicsActorHandle& ActorHandle = BodyInst->ActorHandle;
-					const bool bIsRigidBody = FPhysicsInterface::IsRigidBody(ActorHandle);
 
-					if (FPhysicsInterface::IsValid(ActorHandle) && bIsRigidBody && (bTeleport || !BodyInst->IsInstanceSimulatingPhysics()))	//If we have a body and it's kinematic, or we are teleporting a simulated body
+					if (FPhysicsInterface::IsValid(ActorHandle) &&
+						!FPhysicsInterface::IsStatic(ActorHandle) &&
+						(bTeleport || !BodyInst->IsInstanceSimulatingPhysics()))	//If we have a body and it's kinematic, or we are teleporting a simulated body
 					{
 						const int32 BoneIndex = BodyInst->InstanceBoneIndex;
 
diff --git a/Engine/Source/Runtime/Engine/Private/PhysicsEngine/PhysLevel.cpp b/Engine/Source/Runtime/Engine/Private/PhysicsEngine/PhysLevel.cpp
--- a/Engine/Source/Runtime/Engine/Private/PhysicsEngine/PhysLevel.cpp
+++ b/Engine/Source/Runtime/Engine/Private/PhysicsEngine/PhysLevel.cpp
@@ -15,10 +15,7 @@
 #include "Modules/ModuleManager.h"
 #include "PhysicsInitialization.h"
 #include "PhysicsEngine/PhysicsSettings.h"
-
-#if INCLUDE_CHAOS
 #include "ChaosSolversModule.h"
-#endif
 
 #if WITH_PHYSX
 	#include "PhysicsEngine/PhysXSupport.h"
@@ -41,7 +38,6 @@ FPhysicsDelegates::FOnPhysSceneInit FPhysicsDelegates::OnPhysSceneInit;
 FPhysicsDelegates::FOnPhysSceneTerm FPhysicsDelegates::OnPhysSceneTerm;
 FPhysicsDelegates::FOnPhysDispatchNotifications FPhysicsDelegates::OnPhysDispatchNotifications;
 
-#if INCLUDE_CHAOS
 /**
  *  Chaos is external to engine but utilizes IChaosSettingsProvider to take settings
  *  From external callers, this implementation allows Chaos to request settings from
@@ -97,8 +93,6 @@ private:
 
 static FEngineChaosSettingsProvider GEngineChaosSettingsProvider;
 
-#endif
-
 //////////////////////////////////////////////////////////////////////////
 // UWORLD
 //////////////////////////////////////////////////////////////////////////
@@ -258,7 +252,6 @@ FName FEndPhysicsTickFunction::DiagnosticContext(bool bDetailed)
 //////// GAME-LEVEL RIGID BODY PHYSICS STUFF ///////
 void PostEngineInitialize()
 {
-#if INCLUDE_CHAOS
 	FChaosSolversModule* ChaosModule = FChaosSolversModule::GetModule();
 
 	if(ChaosModule)
@@ -268,7 +261,6 @@ void PostEngineInitialize()
 		ChaosModule->SetSettingsProvider(&GEngineChaosSettingsProvider);
 		ChaosModule->OnSettingsChanged();
 	}
-#endif
 }
 
 bool InitGamePhys()
diff --git a/Engine/Source/Runtime/Engine/Private/PhysicsEngine/PhysScene_PhysX.cpp b/Engine/Source/Runtime/Engine/Private/PhysicsEngine/PhysScene_PhysX.cpp
--- a/Engine/Source/Runtime/Engine/Private/PhysicsEngine/PhysScene_PhysX.cpp
+++ b/Engine/Source/Runtime/Engine/Private/PhysicsEngine/PhysScene_PhysX.cpp
@@ -24,9 +24,7 @@
 #include "PhysicsInterfaceDeclaresCore.h"
 #if !WITH_CHAOS_NEEDS_TO_BE_FIXED
 
-#if INCLUDE_CHAOS
 #include "SQAccelerator.h"
-#endif
 
 #if WITH_PHYSX
 #include "PhysXPublic.h"
@@ -2255,8 +2253,9 @@ void ListAwakeRigidBodiesFromScene(bool bIncludeKinematic, PxScene* PhysXScene,
 void FPhysScene_PhysX::SerializeForTesting(FArchive& Ar)
 {
 	FPhysTestSerializer Serializer;
+	Chaos::FChaosArchive ChaosAr(Ar);
 	Serializer.SetPhysicsData(*GetPxScene());
-	Serializer.Serialize(Ar);
+	Serializer.Serialize(ChaosAr);
 }
 #endif // WITH_PHYSX
 
diff --git a/Engine/Source/Runtime/Engine/Private/PhysicsEngine/PhysUtils.cpp b/Engine/Source/Runtime/Engine/Private/PhysicsEngine/PhysUtils.cpp
--- a/Engine/Source/Runtime/Engine/Private/PhysicsEngine/PhysUtils.cpp
+++ b/Engine/Source/Runtime/Engine/Private/PhysicsEngine/PhysUtils.cpp
@@ -16,6 +16,9 @@
 #include "PhysicsEngine/BodySetup.h"
 #include "Physics/PhysicsInterfaceCore.h"
 #include "PhysXSupportCore.h"
+#include "PhysicsSolver.h"
+#include "Chaos/PBDRigidsEvolutionGBF.h"
+#include "Chaos/ChaosArchive.h"
 
 /** Returns false if ModelToHulls operation should halt because of vertex count overflow. */
 static bool AddConvexPrim(FKAggregateGeom* OutGeom, TArray<FPlane> &Planes, UModel* InModel)
@@ -468,8 +471,44 @@ bool FPhysScene::ExecApexVis(const TCHAR* Cmd, FOutputDevice* Ar)
 
 
 #if WITH_CHAOS
-bool FPhysicsInterface::ExecPhysCommands(const TCHAR* Cmd, FOutputDevice* Ar, UWorld* InWorld)
+bool FPhysicsInterface::ExecPhysCommands(const TCHAR* Cmd, FOutputDevice* OutputDevice, UWorld* InWorld)
 {
+
+#if CHAOS_MEMORY_TRACKING
+	if (FParse::Command(&Cmd, TEXT("ChaosMemoryDistribution")))
+	{
+		//
+		// NOTE: This is an awful, awful way to do this.
+		//
+
+		// Make an archive and serialze the whole scene.
+		// TODO: Don't do this!
+		FArchive BaseAr;
+		BaseAr.SetIsLoading(false);
+		BaseAr.SetIsSaving(true);
+		Chaos::FChaosArchive Ar(BaseAr);
+		FPhysScene* PhysScene = InWorld->GetPhysicsScene();
+		Chaos::FPhysicsSolver* Solver = PhysScene->GetSolver();
+		Chaos::TPBDRigidsEvolutionGBF<float, 3>* Evolution = Solver->GetEvolution();
+		Evolution->Serialize(Ar);
+		TUniquePtr<Chaos::FChaosArchiveContext> ArchiveContext = Ar.StealContext();
+
+		// Grab the memory tracking map from the archive context
+		const TMap<FName, Chaos::FChaosArchiveSectionData>& MemoryMap = ArchiveContext->SectionMap;
+
+		OutputDevice->Logf(TEXT("Chaos serialized memory distribution:"));
+		int64 TotalBytes = 0;
+		for (auto It : MemoryMap)
+		{
+			const FName SectionName = It.Key;
+			const Chaos::FChaosArchiveSectionData SectionData = It.Value;
+			TotalBytes += SectionData.SizeExclusive;
+			OutputDevice->Logf(TEXT("%s ~ count: %d, bytes: %d, megabytes: %f"), *SectionName.ToString(), SectionData.Count, SectionData.SizeExclusive, (float)SectionData.SizeExclusive * 10e-7f);
+		}
+		OutputDevice->Logf(TEXT("Total bytes: %d, megabytes: %f"), TotalBytes, (float)TotalBytes * 10e-7f);
+	}
+#endif
+
     return false;
 }
 #else
diff --git a/Engine/Source/Runtime/Engine/Private/PhysicsEngine/PhysicsInterfacePhysX.cpp b/Engine/Source/Runtime/Engine/Private/PhysicsEngine/PhysicsInterfacePhysX.cpp
--- a/Engine/Source/Runtime/Engine/Private/PhysicsEngine/PhysicsInterfacePhysX.cpp
+++ b/Engine/Source/Runtime/Engine/Private/PhysicsEngine/PhysicsInterfacePhysX.cpp
@@ -14,9 +14,7 @@
 
 #if !WITH_CHAOS_NEEDS_TO_BE_FIXED
 
-#if INCLUDE_CHAOS
 #include "SQAccelerator.h"
-#endif
 
 #if WITH_PHYSX
 #include "PhysXPublic.h"
@@ -668,6 +666,19 @@ bool FPhysicsCommand_PhysX::ExecuteWrite(FPhysScene_PhysX* InScene, TFunctionRef
 	return false;
 }
 
+bool FPhysicsCommand_PhysX::ExecuteWrite(FPhysicsActorHandle_PhysX& InHandle, TFunctionRef<void(FPhysicsActorHandle_PhysX& Actor)> InCallable)
+{
+	if(InHandle.IsValid())
+	{
+		FPhysicsInterfaceScopedLock_PhysX ScopeLock(&InHandle, EPhysicsInterfaceScopedLockType::Write);
+		InCallable(InHandle);
+
+		return true;
+	}
+
+	return false;
+}
+
 void FPhysicsCommand_PhysX::ExecuteShapeWrite(FBodyInstance* InInstance, FPhysicsShapeHandle_PhysX& InShape, TFunctionRef<void(FPhysicsShapeHandle_PhysX& InShape)> InCallable)
 {
 	if(InShape.IsValid())
diff --git a/Engine/Source/Runtime/Engine/Private/SkeletalRenderGPUSkin.cpp b/Engine/Source/Runtime/Engine/Private/SkeletalRenderGPUSkin.cpp
--- a/Engine/Source/Runtime/Engine/Private/SkeletalRenderGPUSkin.cpp
+++ b/Engine/Source/Runtime/Engine/Private/SkeletalRenderGPUSkin.cpp
@@ -333,6 +333,7 @@ void FSkeletalMeshObjectGPUSkin::UpdateDynamicData_RenderThread(FGPUSkinCache* G
 
 #if RHI_RAYTRACING
 	bRequireRecreatingRayTracingGeometry = (DynamicData == nullptr || // Newly created
+		RayTracingGeometry.Initializer.PositionVertexBuffer == nullptr ||
 		(DynamicData != nullptr && DynamicData->LODIndex != InDynamicData->LODIndex)); // LOD level changed
 #endif
 
@@ -574,7 +575,7 @@ void FSkeletalMeshObjectGPUSkin::ProcessUpdatedDynamicData(FGPUSkinCache* GPUSki
 			// Create a uniform buffer from the bone transforms.
 			TArray<FMatrix>& ReferenceToLocalMatrices = DynamicData->ReferenceToLocal;
 			bool bNeedFence = ShaderData.UpdateBoneData(RHICmdList, ReferenceToLocalMatrices, Section.BoneMap, RevisionNumber, false, FeatureLevel, bUseSkinCache);
-#if WITH_APEX_CLOTHING
+#if WITH_APEX_CLOTHING || WITH_CHAOS_CLOTHING
 			// Update uniform buffer for APEX cloth simulation mesh positions and normals
 			if( bClothFactory )
 			{				
@@ -587,7 +588,7 @@ void FSkeletalMeshObjectGPUSkin::ProcessUpdatedDynamicData(FGPUSkinCache* GPUSki
 					bNeedFence = ClothShaderData.UpdateClothSimulData(RHICmdList, DynamicData->ClothingSimData[ActorIdx].Positions, DynamicData->ClothingSimData[ActorIdx].Normals, FrameNumberToPrepare, FeatureLevel) || bNeedFence;
 				}
 			}
-#endif // WITH_APEX_CLOTHING
+#endif // WITH_APEX_CLOTHING || WITH_CHAOS_CLOTHING
 			// Try to use the GPU skinning cache if possible
 			if (bUseSkinCache)
 			{
diff --git a/Engine/Source/Runtime/Engine/Public/CollisionQueryParams.h b/Engine/Source/Runtime/Engine/Public/CollisionQueryParams.h
--- a/Engine/Source/Runtime/Engine/Public/CollisionQueryParams.h
+++ b/Engine/Source/Runtime/Engine/Public/CollisionQueryParams.h
@@ -83,6 +83,10 @@ struct ENGINE_API FCollisionQueryParams
 		RETURN_QUICK_DECLARE_CYCLE_STAT(UnknownSceneQuery, STATGROUP_Collision);
 	}
 
+#if !(UE_BUILD_TEST || UE_BUILD_SHIPPING)
+	bool bDebugQuery;
+#endif
+
 private:
 
 	/** Tracks whether the IgnoreComponents list is verified unique. */
@@ -149,6 +153,9 @@ public:
 		bIgnoreBlocks = false;
 		bIgnoreTouches = false;
 		StatId = GetUnknownStatId();
+#if !(UE_BUILD_TEST || UE_BUILD_SHIPPING)
+		bDebugQuery = false;
+#endif
 	}
 
 	FCollisionQueryParams()
@@ -164,6 +171,9 @@ public:
 		bIgnoreBlocks = false;
 		bIgnoreTouches = false;
 		StatId = GetUnknownStatId();
+#if !(UE_BUILD_TEST || UE_BUILD_SHIPPING)
+		bDebugQuery = false;
+#endif
 	}
 
 	FCollisionQueryParams(FName InTraceTag, bool bInTraceComplex=false, const AActor* InIgnoreActor=NULL)
diff --git a/Engine/Source/Runtime/Engine/Public/Physics/Experimental/ChaosInterfaceUtils.h b/Engine/Source/Runtime/Engine/Public/Physics/Experimental/ChaosInterfaceUtils.h
new file mode 100644
--- /dev/null
+++ b/Engine/Source/Runtime/Engine/Public/Physics/Experimental/ChaosInterfaceUtils.h
@@ -0,0 +1,17 @@
+// Copyright 1998-2019 Epic Games, Inc. All Rights Reserved.
+#pragma once
+
+#include "Chaos/Declares.h"
+#include "Containers/ContainersFwd.h"
+#include "Containers/ContainerAllocationPolicies.h"
+#include "Templates/UniquePtr.h"
+
+struct FGeometryAddParams;
+
+namespace ChaosInterface
+{
+	/**
+	 * Create the Chaos Geometry based on the geometry parameters.
+	 */
+	void CreateGeometry(const FGeometryAddParams& InParams, TArray<TUniquePtr<Chaos::TImplicitObject<float, 3>>>& OutGeoms, TArray<TUniquePtr<Chaos::TPerShapeData<float, 3>>, TInlineAllocator<1>>& OutShapes);
+}
diff --git a/Engine/Source/Runtime/Engine/Public/Physics/Experimental/ChaosInterfaceWrapper.h b/Engine/Source/Runtime/Engine/Public/Physics/Experimental/ChaosInterfaceWrapper.h
--- a/Engine/Source/Runtime/Engine/Public/Physics/Experimental/ChaosInterfaceWrapper.h
+++ b/Engine/Source/Runtime/Engine/Public/Physics/Experimental/ChaosInterfaceWrapper.h
@@ -1,28 +1,37 @@
 // Copyright 1998-2019 Epic Games, Inc. All Rights Reserved.
 #pragma once
 
-#if INCLUDE_CHAOS
 #include "ChaosInterfaceWrapperCore.h"
+#include "Chaos/ParticleHandleFwd.h"
+#include "PhysicsInterfaceUtilsCore.h"
 #include "CollisionQueryFilterCallbackCore.h"
-#include "Chaos/ParticleHandle.h"
-#include "PhysxUserData.h"
 //todo: move this include into an impl header
 
+class FPhysScene_Chaos;
 class UPhysicalMaterial;
 struct FBodyInstance;
 struct FCollisionFilterData;
 struct FCollisionQueryParams;
 
+namespace Chaos
+{
+	template <class T, int d>
+	class TPerShapeData;
+}
+
 namespace ChaosInterface
 {
 
-// Needed by low level SQ calls. Right now there's no specific locking for LLI
-// #PHYS2 update as locking becomes necessary
+// Needed by low level SQ calls.
+#if WITH_CHAOS
 struct FScopedSceneReadLock
 {
-	FScopedSceneReadLock(FPhysScene& Scene)
-	{}
+	FScopedSceneReadLock(FPhysScene_ChaosInterface& SceneIn);
+	~FScopedSceneReadLock();
+
+	FPhysScene_Chaos& Scene;
 };
+#endif
 
 inline FQueryFilterData MakeQueryFilterData(const FCollisionFilterData& FilterData, EQueryFlags QueryFlags, const FCollisionQueryParams& Params)
 {
@@ -33,20 +42,12 @@ inline FQueryFilterData MakeQueryFilterData(const FCollisionFilterData& FilterDa
 #endif
 }
 
-inline UPhysicalMaterial* GetUserData(const FPhysTypeDummy& Material)
-{
-	return nullptr;
-}
+FBodyInstance* GetUserData(const Chaos::TGeometryParticle<float, 3>& Actor);
 
-inline FBodyInstance* GetUserData(const Chaos::TGeometryParticle<float,3>& Actor)
-{
-	void* UserData = Actor.UserData();
-	return UserData ? FPhysxUserData::Get<FBodyInstance>(Actor.UserData()) : nullptr;
-}
+UPhysicalMaterial* GetUserData(const FPhysTypeDummy& Material);
 
 }
-void LowLevelRaycast(FPhysScene& Scene, const FVector& Start, const FVector& Dir, float DeltaMag, FPhysicsHitCallback<FHitRaycast>& HitBuffer, EHitFlags OutputFlags, FQueryFlags QueryFlags, const FCollisionFilterData& Filter, const FQueryFilterData& QueryFilterData, ICollisionQueryFilterCallbackBase* QueryCallback);
-void LowLevelSweep(FPhysScene& Scene, const FPhysicsGeometry& Geom, const FTransform& StartTM, const FVector& Dir, float DeltaMag, FPhysicsHitCallback<FHitSweep>& HitBuffer, EHitFlags OutputFlags, FQueryFlags QueryFlags, const FCollisionFilterData& Filter, const FQueryFilterData& QueryFilterData, ICollisionQueryFilterCallbackBase* QueryCallback);
-void LowLevelOverlap(FPhysScene& Scene, const FPhysicsGeometry& Geom, const FTransform& GeomPose, FPhysicsHitCallback<FHitOverlap>& HitBuffer, FQueryFlags QueryFlags, const FCollisionFilterData& Filter, const FQueryFilterData& QueryFilterData, ICollisionQueryFilterCallbackBase* QueryCallback);
 
-#endif // INCLUDE_CHAOS
\ No newline at end of file
+void LowLevelRaycast(FPhysScene& Scene, const FVector& Start, const FVector& Dir, float DeltaMag, FPhysicsHitCallback<FHitRaycast>& HitBuffer, EHitFlags OutputFlags, FQueryFlags QueryFlags, const FCollisionFilterData& Filter, const FQueryFilterData& QueryFilterData, ICollisionQueryFilterCallbackBase* QueryCallback, const FQueryDebugParams& DebugParams = FQueryDebugParams());
+void LowLevelSweep(FPhysScene& Scene, const FPhysicsGeometry& Geom, const FTransform& StartTM, const FVector& Dir, float DeltaMag, FPhysicsHitCallback<FHitSweep>& HitBuffer, EHitFlags OutputFlags, FQueryFlags QueryFlags, const FCollisionFilterData& Filter, const FQueryFilterData& QueryFilterData, ICollisionQueryFilterCallbackBase* QueryCallback, const FQueryDebugParams& DebugParams = FQueryDebugParams());
+void LowLevelOverlap(FPhysScene& Scene, const FPhysicsGeometry& Geom, const FTransform& GeomPose, FPhysicsHitCallback<FHitOverlap>& HitBuffer, FQueryFlags QueryFlags, const FCollisionFilterData& Filter, const FQueryFilterData& QueryFilterData, ICollisionQueryFilterCallbackBase* QueryCallback, const FQueryDebugParams& DebugParams = FQueryDebugParams());
diff --git a/Engine/Source/Runtime/Engine/Public/Physics/Experimental/PhysInterface_Chaos.h b/Engine/Source/Runtime/Engine/Public/Physics/Experimental/PhysInterface_Chaos.h
--- a/Engine/Source/Runtime/Engine/Public/Physics/Experimental/PhysInterface_Chaos.h
+++ b/Engine/Source/Runtime/Engine/Public/Physics/Experimental/PhysInterface_Chaos.h
@@ -6,17 +6,8 @@
 
 #include "ChaosInterfaceWrapper.h"
 #include "Engine/Engine.h"
-#include "Chaos/ArrayCollection.h"
-#include "Chaos/ArrayCollectionArray.h"
 #include "Chaos/Declares.h"
-#include "Chaos/Capsule.h"
-#include "Chaos/ImplicitObject.h"
-#include "Chaos/Pair.h"
-#include "Chaos/Transform.h"
-#include "PhysicsEngine/ConstraintDrives.h"
 #include "PhysicsEngine/ConstraintTypes.h"
-#include "PhysicsPublic.h"
-#include "PhysicsReplication.h"
 #include "PhysicsInterfaceWrapperShared.h"
 #include "Physics/PhysicsInterfaceDeclares.h"
 #include "Physics/PhysicsInterfaceTypes.h"
@@ -32,6 +23,10 @@ class FPhysInterface_Chaos;
 struct FBodyInstance;
 struct FPhysxUserData;
 class IPhysicsReplicationFactory;
+struct FConstraintDrive;
+struct FLinearDriveConstraint;
+struct FAngularDriveConstraint;
+
 
 class AWorldSettings;
 
@@ -56,6 +51,21 @@ namespace Chaos
 
 	template <typename T, int>
 	class TConvex;
+
+	template <typename T>
+	class TCapsule;
+
+	template <typename T, int>
+	class TBox;
+
+	template <typename T, int>
+	class TSphere;
+
+	template <typename T, int>
+	class TPerShapeData;
+
+	template <typename T, int d>
+	using TShapesArray = TArray<TUniquePtr<TPerShapeData<T, d>>, TInlineAllocator<1>>;
 }
 
 class FPhysicsConstraintReference_Chaos
@@ -89,7 +99,7 @@ public:
 	bool Equals(const FPhysicsShapeReference_Chaos& Other) const { return Shape == Other.Shape; }
     bool operator==(const FPhysicsShapeReference_Chaos& Other) const { return Equals(Other); }
 	
-	const Chaos::TImplicitObject<float, 3>& GetGeometry() const { check(IsValid()); return *Shape->Geometry; }
+	const Chaos::TImplicitObject<float, 3>& GetGeometry() const;
 
 
 	Chaos::TPerShapeData<float, 3>* Shape;
@@ -104,7 +114,7 @@ public:
 	FPhysicsShapeAdapter_Chaos(const FQuat& Rot, const FCollisionShape& CollisionShape);
 
 	const FPhysicsGeometry& GetGeometry() const;
-	FTransform GetGeometryPose(const FVector& Pos) const;
+	FTransform GetGeomPose(const FVector& Pos) const;
 	const FQuat& GetGeomOrientation() const;
 
 private:
@@ -269,8 +279,8 @@ public:
 	static float GetSleepEnergyThreshold_AssumesLocked(const FPhysicsActorHandle& InActorReference);
 	static void SetSleepEnergyThreshold_AssumesLocked(const FPhysicsActorHandle& InActorReference, float InEnergyThreshold);
 
-	static void SetMass_AssumesLocked(const FPhysicsActorHandle& InHandle, float InMass);
-	static void SetMassSpaceInertiaTensor_AssumesLocked(const FPhysicsActorHandle& InHandle, const FVector& InTensor);
+	static void SetMass_AssumesLocked(FPhysicsActorHandle& InHandle, float InMass);
+	static void SetMassSpaceInertiaTensor_AssumesLocked(FPhysicsActorHandle& InHandle, const FVector& InTensor);
 	static void SetComLocalPose_AssumesLocked(const FPhysicsActorHandle& InHandle, const FTransform& InComLocalPose);
 
 	static float GetStabilizationEnergyThreshold_AssumesLocked(const FPhysicsActorHandle& InHandle);
@@ -332,113 +342,20 @@ public:
     /////////////////////////////////////////////
 
     // Interface needed for cmd
-    static bool ExecuteRead(const FPhysicsActorHandle& InActorReference, TFunctionRef<void(const FPhysicsActorHandle& Actor)> InCallable)
-    {
-		if (InActorReference)
-		{
-			InCallable(InActorReference);
-			return true;
-		}
-		return false;
-    }
-
-    static bool ExecuteRead(USkeletalMeshComponent* InMeshComponent, TFunctionRef<void()> InCallable)
-    {
-        InCallable();
-        return true;
-    }
-
-    static bool ExecuteRead(const FPhysicsActorHandle& InActorReferenceA, const FPhysicsActorHandle& InActorReferenceB, TFunctionRef<void(const FPhysicsActorHandle& ActorA, const FPhysicsActorHandle& ActorB)> InCallable)
-    {
-		InCallable(InActorReferenceA, InActorReferenceB);
-		return true;
-    }
-
-    static bool ExecuteRead(const FPhysicsConstraintReference_Chaos& InConstraintRef, TFunctionRef<void(const FPhysicsConstraintReference_Chaos& Constraint)> InCallable)
-    {
-		if(InConstraintRef.IsValid())
-		{
-			InCallable(InConstraintRef);
-			return true;
-		}
-
-		return false;
-    }
-
-    static bool ExecuteRead(FPhysScene* InScene, TFunctionRef<void()> InCallable)
-    {
-		if(InScene)
-		{
-			InCallable();
-			return true;
-		}
-
-		return false;
-    }
+    static bool ExecuteRead(const FPhysicsActorHandle& InActorReference, TFunctionRef<void(const FPhysicsActorHandle& Actor)> InCallable);
+    static bool ExecuteRead(USkeletalMeshComponent* InMeshComponent, TFunctionRef<void()> InCallable);
+    static bool ExecuteRead(const FPhysicsActorHandle& InActorReferenceA, const FPhysicsActorHandle& InActorReferenceB, TFunctionRef<void(const FPhysicsActorHandle& ActorA, const FPhysicsActorHandle& ActorB)> InCallable);
+    static bool ExecuteRead(const FPhysicsConstraintReference_Chaos& InConstraintRef, TFunctionRef<void(const FPhysicsConstraintReference_Chaos& Constraint)> InCallable);
+    static bool ExecuteRead(FPhysScene* InScene, TFunctionRef<void()> InCallable);
 
-    static bool ExecuteWrite(const FPhysicsActorHandle& InActorReference, TFunctionRef<void(const FPhysicsActorHandle& Actor)> InCallable)
-    {
-		//why do we have a write that takes in a const handle?
-		if (InActorReference)
-		{
-			InCallable(InActorReference);
-			return true;
-		}
-		return false;
-    }
+    static bool ExecuteWrite(const FPhysicsActorHandle& InActorReference, TFunctionRef<void(const FPhysicsActorHandle& Actor)> InCallable);
+	static bool ExecuteWrite(FPhysicsActorHandle& InActorReference, TFunctionRef<void(FPhysicsActorHandle& Actor)> InCallable);
+    static bool ExecuteWrite(USkeletalMeshComponent* InMeshComponent, TFunctionRef<void()> InCallable);
+    static bool ExecuteWrite(const FPhysicsActorHandle& InActorReferenceA, const FPhysicsActorHandle& InActorReferenceB, TFunctionRef<void(const FPhysicsActorHandle& ActorA, const FPhysicsActorHandle& ActorB)> InCallable);
+    static bool ExecuteWrite(const FPhysicsConstraintReference_Chaos& InConstraintRef, TFunctionRef<void(const FPhysicsConstraintReference_Chaos& Constraint)> InCallable);
+    static bool ExecuteWrite(FPhysScene* InScene, TFunctionRef<void()> InCallable);
 
-	static bool ExecuteWrite(FPhysicsActorHandle& InActorReference, TFunctionRef<void(FPhysicsActorHandle& Actor)> InCallable)
-	{
-		if (InActorReference)
-		{
-			InCallable(InActorReference);
-			return true;
-		}
-
-		return false;
-	}
-
-    static bool ExecuteWrite(USkeletalMeshComponent* InMeshComponent, TFunctionRef<void()> InCallable)
-    {
-        InCallable();
-        return true;
-    }
-
-    static bool ExecuteWrite(const FPhysicsActorHandle& InActorReferenceA, const FPhysicsActorHandle& InActorReferenceB, TFunctionRef<void(const FPhysicsActorHandle& ActorA, const FPhysicsActorHandle& ActorB)> InCallable)
-    {
-		InCallable(InActorReferenceA, InActorReferenceB);
-		return true;
-    }
-
-    static bool ExecuteWrite(const FPhysicsConstraintReference_Chaos& InConstraintRef, TFunctionRef<void(const FPhysicsConstraintReference_Chaos& Constraint)> InCallable)
-    {
-		if(InConstraintRef.IsValid())
-		{
-			InCallable(InConstraintRef);
-			return true;
-		}
-
-		return false;
-    }
-	
-    static bool ExecuteWrite(FPhysScene* InScene, TFunctionRef<void()> InCallable)
-    {
-		if(InScene)
-		{
-			InCallable();
-			return true;
-		}
-
-		return false;
-    }
-
-    static void ExecuteShapeWrite(FBodyInstance* InInstance, FPhysicsShapeHandle& InShape, TFunctionRef<void(FPhysicsShapeHandle& InShape)> InCallable)
-    {
-		if(InInstance && InShape.IsValid())
-        {
-			InCallable(InShape);
-		}
-    }
+    static void ExecuteShapeWrite(FBodyInstance* InInstance, FPhysicsShapeHandle& InShape, TFunctionRef<void(FPhysicsShapeHandle& InShape)> InCallable);
 
 	// Misc
 
@@ -451,6 +368,7 @@ public:
 	// Shape interface functions
 	static FPhysicsShapeHandle CreateShape(physx::PxGeometry* InGeom, bool bSimulation = true, bool bQuery = true, UPhysicalMaterial* InSimpleMaterial = nullptr, TArray<UPhysicalMaterial*>* InComplexMaterials = nullptr);
 	
+	static void CreateGeometry(const FGeometryAddParams& InParams, TArray<TUniquePtr<Chaos::TImplicitObject<float, 3>>>& OutGeoms, Chaos::TShapesArray<float, 3>& OutShapes, TArray<FPhysicsShapeHandle>* OutOptShapes);
 	static void AddGeometry(FPhysicsActorHandle& InActor, const FGeometryAddParams& InParams, TArray<FPhysicsShapeHandle>* OutOptShapes = nullptr);
 	static FPhysicsShapeHandle CloneShape(const FPhysicsShapeHandle& InShape);
 	static FPhysicsGeometryCollection_Chaos GetGeometryCollection(const FPhysicsShapeHandle& InShape);
@@ -460,7 +378,7 @@ public:
 	static bool IsSimulationShape(const FPhysicsShapeHandle& InShape);
 	static bool IsQueryShape(const FPhysicsShapeHandle& InShape);
 	static bool IsShapeType(const FPhysicsShapeHandle& InShape, ECollisionShapeType InType);
-	static ECollisionShapeType GetShapeType(const FPhysicsShapeHandle& InShape);
+	static ECollisionShapeType GetShapeType(const FPhysicsShapeHandle& InShape, bool GetInnerType = false);
 	static FTransform GetLocalTransform(const FPhysicsShapeHandle& InShape);
     static void* GetUserData(const FPhysicsShapeHandle& InShape) { return nullptr; }
 
@@ -502,10 +420,7 @@ FORCEINLINE ECollisionShapeType GetType(const Chaos::TImplicitObject<float, 3>&
 	return ECollisionShapeType::None;
 }
 */
-FORCEINLINE ECollisionShapeType GetGeometryType(const Chaos::TPerShapeData<float, 3>& Shape)
-{
-	return GetType(*Shape.Geometry);
-}
+ECollisionShapeType GetGeometryType(const Chaos::TPerShapeData<float, 3>& Shape);
 
 /*
 FORCEINLINE float GetRadius(const Chaos::TCapsule<float>& Capsule)
diff --git a/Engine/Source/Runtime/Engine/Public/Physics/Experimental/PhysScene_Chaos.h b/Engine/Source/Runtime/Engine/Public/Physics/Experimental/PhysScene_Chaos.h
--- a/Engine/Source/Runtime/Engine/Public/Physics/Experimental/PhysScene_Chaos.h
+++ b/Engine/Source/Runtime/Engine/Public/Physics/Experimental/PhysScene_Chaos.h
@@ -2,8 +2,6 @@
 
 #pragma once
 
-#if INCLUDE_CHAOS
-
 #include "CoreMinimal.h"
 #include "Tickable.h"
 #include "Physics/PhysScene.h"
@@ -12,14 +10,20 @@
 #include "GameFramework/Actor.h"
 #include "PhysicsPublic.h"
 #include "PhysInterface_Chaos.h"
+#include "Physics/PhysicsInterfaceUtils.h"
 #include "Chaos/Transform.h"
+#include "Chaos/ArrayCollectionArray.h"
 #include "Chaos/Declares.h"
 #include "PhysicsProxy/SingleParticlePhysicsProxyFwd.h"
+#include "Framework/Threading.h"
 
 #ifndef CHAOS_WITH_PAUSABLE_SOLVER
 #define CHAOS_WITH_PAUSABLE_SOLVER 1
 #endif
 
+// Currently compilation issue with Incredibuild when including headers required by event template functions
+#define XGE_FIXED 0
+
 class UPrimitiveComponent;
 
 class AdvanceOneTimeStepTask;
@@ -38,9 +42,8 @@ class IPhysicsProxyBase;
 namespace Chaos
 {
 	class FPhysicsProxy;
-	class IDispatcher;
 
-	enum EEventType;
+	enum EEventType : int32;
 
 	template<typename PayloadType, typename HandlerType>
 	class TRawEventHandler;
@@ -50,6 +53,8 @@ namespace Chaos
 
 	template <typename TPayload, typename T, int d>
 	class ISpatialAcceleration;
+
+	class IDispatcher;
 }
 
 /**
@@ -111,6 +116,7 @@ public:
 
 	void RemoveActorFromAccelerationStructure(FPhysicsActorHandle& Actor);
 
+#if XGE_FIXED
 	template<typename PayloadType>
 	void RegisterEvent(const Chaos::EEventType& EventID, TFunction<void(const Chaos::FPBDRigidsSolver* Solver, PayloadType& EventData)> InLambda);
 	void UnregisterEvent(const Chaos::EEventType& EventID);
@@ -118,6 +124,7 @@ public:
 	template<typename PayloadType, typename HandlerType>
 	void RegisterEventHandler(const Chaos::EEventType& EventID, HandlerType* Handler, typename Chaos::TRawEventHandler<PayloadType, HandlerType>::FHandlerFunction Func);
 	void UnregisterEventHandler(const Chaos::EEventType& EventID, const void* Handler);
+#endif // XGE_FIXED
 
 	void Shutdown();
 
@@ -144,9 +151,18 @@ public:
 		return PhysicsProxyPtr ? *PhysicsProxyPtr : nullptr;
 	}
 
-	TUniquePtr<Chaos::ISpatialAcceleration<Chaos::TAccelerationStructureHandle<float, 3>, float, 3>> SolverAccelerationStructure;
+	const Chaos::ISpatialAcceleration<Chaos::TAccelerationStructureHandle<float, 3>, float, 3>* GetSpacialAcceleration() const;
+	Chaos::ISpatialAcceleration<Chaos::TAccelerationStructureHandle<float, 3>, float, 3>* GetSpacialAcceleration();
+	
+	/**
+	 * Copies the acceleration structure out of the solver, does no thread safety checking so ensure calls
+	 * to this are made at appropriate sync points if required
+	 */
+	void CopySolverAccelerationStructure();
 
 private:
+	TUniquePtr<Chaos::ISpatialAcceleration<Chaos::TAccelerationStructureHandle<float, 3>, float, 3>> SolverAccelerationStructure;
+
 #if CHAOS_WITH_PAUSABLE_SOLVER
 	/** Callback that checks the status of the world settings for this scene before pausing/unpausing its solver. */
 	void OnUpdateWorldPause();
@@ -190,8 +206,16 @@ private:
 	// Cache the state of the game pause in order to avoid sending extraneous commands to the solver.
 	bool bIsWorldPaused;
 #endif
+
+	/** Scene lock object for external threads (non-physics) */
+	Chaos::FPhysicsSceneGuard ExternalDataLock;
+
+	// Allow other code to obtain read-locks when needed
+	friend struct FScopedSceneReadLock;
+	friend struct FScopedSceneLock_Chaos;
 };
 
+#if XGE_FIXED
 template<typename PayloadType>
 void FPhysScene_Chaos::RegisterEvent(const Chaos::EEventType& EventID, TFunction<void(const Chaos::FPBDRigidsSolver* Solver, PayloadType& EventData)> InLambda)
 {
@@ -225,6 +249,7 @@ void FPhysScene_Chaos::RegisterEventHandler(const Chaos::EEventType& EventID, Ha
 		});
 	}
 }
+#endif // XGE_FIXED
 
 class UWorld;
 class AWorldSettings;
@@ -315,6 +340,7 @@ public:
 	ENGINE_API bool ExecPxVis(uint32 SceneType, const TCHAR* Cmd, FOutputDevice* Ar);
 	ENGINE_API bool ExecApexVis(uint32 SceneType, const TCHAR* Cmd, FOutputDevice* Ar);
 
+#if XGE_FIXED
 	template<typename PayloadType>
 	void RegisterEvent(const Chaos::EEventType& EventID, TFunction<void(const Chaos::FPBDRigidsSolver* Solver, PayloadType& EventData)> InLambda)
 	{
@@ -334,10 +360,10 @@ public:
 	{
 		Scene.UnregisterEventHandler(EventID, Handler);
 	}
+#endif // XGE_FIXED
 
 private:
 
-	void SyncBodies();
 	void SyncBodies(Chaos::FPhysicsSolver* Solver);
 
 	void SetKinematicTransform(FPhysicsActorHandle& InActorReference, const Chaos::TRigidTransform<float, 3>& NewTransform)
@@ -347,16 +373,6 @@ private:
 		// InActorReference.GetPhysicsProxy()->SetKinematicData(...)
 	}
 
-	void Lock()
-	{
-		MCriticalSection.Lock();
-	}
-
-	void Unlock()
-	{
-		MCriticalSection.Unlock();
-	}
-
 	void EnableCollisionPair(const TTuple<int32, int32>& CollisionPair)
 	{
 		// #todo : Implement
@@ -390,8 +406,6 @@ private:
 
 	FPhysScene_Chaos Scene;
 
-	// @todo(mlentine): Locking is very heavy handed right now; need to make less so.
-	FCriticalSection MCriticalSection;
 	float MDeltaTime;
 	//Body Instances
 	Chaos::TArrayCollectionArray<FBodyInstance*> BodyInstances;
@@ -403,7 +417,4 @@ private:
 	FGraphEventRef CompletionEvent;
 	FGraphEventRef PhysicsTickTask;
 };
-
-#endif
-
 #endif
diff --git a/Engine/Source/Runtime/Engine/Public/Physics/ImmediatePhysics/ImmediatePhysicsChaos/ImmediatePhysicsActorHandle_Chaos.h b/Engine/Source/Runtime/Engine/Public/Physics/ImmediatePhysics/ImmediatePhysicsChaos/ImmediatePhysicsActorHandle_Chaos.h
--- a/Engine/Source/Runtime/Engine/Public/Physics/ImmediatePhysics/ImmediatePhysicsChaos/ImmediatePhysicsActorHandle_Chaos.h
+++ b/Engine/Source/Runtime/Engine/Public/Physics/ImmediatePhysics/ImmediatePhysicsChaos/ImmediatePhysicsActorHandle_Chaos.h
@@ -2,8 +2,6 @@
 
 #pragma once
 
-#if INCLUDE_CHAOS
-
 #include "Physics/ImmediatePhysics/ImmediatePhysicsChaos/ImmediatePhysicsCore_Chaos.h"
 
 #include "Engine/EngineTypes.h"
@@ -110,8 +108,13 @@ namespace ImmediatePhysics_Chaos
 		float GetMaxContactImpulse() const;
 
 		/** Get the actor-space centre of mass offset */
-		// @todo(ccaulfield): should be FTransform
-		const FVector& GetCoMTranslation() const;
+		const FTransform& GetLocalCoMTransform() const;
+
+		Chaos::TGeometryParticleHandle<FReal, Dimensions>* GetParticle();
+		const Chaos::TGeometryParticleHandle<FReal, Dimensions>* GetParticle() const;
+
+		int32 GetLevel() const;
+		void SetLevel(int32 InLevel);
 
 	private:
 		friend struct FSimulation;
@@ -125,9 +128,8 @@ namespace ImmediatePhysics_Chaos
 		Chaos::TGeometryParticleHandle<FReal, Dimensions>* ParticleHandle;
 		TUniquePtr<Chaos::TImplicitObject<float, 3>> Geometry;
 		TArray<TUniquePtr<Chaos::TPerShapeData<float, 3>>> Shapes;
-		FVector ActorToCoMTranslation;
+		FTransform ActorToCoMTransform;
+		int32 Level;
 	};
 
 }
-
-#endif // INCLUDE_CHAOS
\ No newline at end of file
diff --git a/Engine/Source/Runtime/Engine/Public/Physics/ImmediatePhysics/ImmediatePhysicsChaos/ImmediatePhysicsCore_Chaos.h b/Engine/Source/Runtime/Engine/Public/Physics/ImmediatePhysics/ImmediatePhysicsChaos/ImmediatePhysicsCore_Chaos.h
--- a/Engine/Source/Runtime/Engine/Public/Physics/ImmediatePhysics/ImmediatePhysicsChaos/ImmediatePhysicsCore_Chaos.h
+++ b/Engine/Source/Runtime/Engine/Public/Physics/ImmediatePhysics/ImmediatePhysicsChaos/ImmediatePhysicsCore_Chaos.h
@@ -2,8 +2,6 @@
 
 #pragma once
 
-#if INCLUDE_CHAOS
-
 #include "Physics/ImmediatePhysics/ImmediatePhysicsChaos/ImmediatePhysicsDeclares_Chaos.h"
 #include "Physics/ImmediatePhysics/ImmediatePhysicsShared/ImmediatePhysicsCore.h"
 
@@ -16,8 +14,6 @@ namespace Chaos
 	template<typename P, typename T, int D> class TPBDConstraintIslandRule;
 	template<typename T, int D> class TPBDJointConstraintHandle;
 	template<typename T, int D> class TPBDJointConstraints;
-	template<typename T, int D> class TPBD6DJointConstraintHandle;
-	template<typename T, int D> class TPBD6DJointConstraints;
 	template<typename T, int D> struct TPBDRigidParticleParameters;
 	template<typename T, int D> class TPBDRigidsEvolutionGBF;
 	template<typename T, int D> class TPBDRigidsSOAs;
@@ -36,7 +32,5 @@ namespace ImmediatePhysics_Chaos
 struct FBodyInstance;
 struct FConstraintInstance;
 
-#endif
-
 // Used to define out code that still has to be implemented to match PhysX
 #define IMMEDIATEPHYSICS_CHAOS_TODO 0
