DRPS - Copy from 9389853
#rnx

diff --git a/Engine/Source/Runtime/Engine/Private/UnrealEngine.cpp b/Engine/Source/Runtime/Engine/Private/UnrealEngine.cpp
--- a/Engine/Source/Runtime/Engine/Private/UnrealEngine.cpp
+++ b/Engine/Source/Runtime/Engine/Private/UnrealEngine.cpp
@@ -2363,7 +2363,6 @@ void UEngine::InitializeObjectReferences()
 	{
 		// Materials that are needed in-game if debug viewmodes are allowed
 		LoadSpecialMaterial(TEXT("WireframeMaterialName"), WireframeMaterialName, WireframeMaterial, true);
-		LoadSpecialMaterial(TEXT("HairDefaultMaterialName"), HairDefaultMaterialName, HairDefaultMaterial, true);
 		LoadSpecialMaterial(TEXT("HairDebugMaterialName"), HairDebugMaterialName, HairDebugMaterial, true);
 		LoadSpecialMaterial(TEXT("LevelColorationLitMaterialName"), LevelColorationLitMaterialName, LevelColorationLitMaterial, true);
 		LoadSpecialMaterial(TEXT("LevelColorationUnlitMaterialName"), LevelColorationUnlitMaterialName, LevelColorationUnlitMaterial, true);
@@ -2382,6 +2381,7 @@ void UEngine::InitializeObjectReferences()
 	LoadSpecialMaterial(TEXT("RemoveSurfaceMaterialName"), RemoveSurfaceMaterialName.ToString(), RemoveSurfaceMaterial, false);
 
 	// these one's are needed both editor and standalone 
+	LoadSpecialMaterial(TEXT("HairDefaultMaterialName"), HairDefaultMaterialName, HairDefaultMaterial, false);
 	LoadSpecialMaterial(TEXT("DebugMeshMaterialName"), DebugMeshMaterialName.ToString(), DebugMeshMaterial, false);
 	LoadSpecialMaterial(TEXT("EmissiveMeshMaterialName"), EmissiveMeshMaterialName.ToString(), EmissiveMeshMaterial, false);
 	LoadSpecialMaterial(TEXT("InvalidLightmapSettingsMaterialName"), InvalidLightmapSettingsMaterialName.ToString(), InvalidLightmapSettingsMaterial, false);
