Fix crash from the new SK build refactor
#rnx

diff --git a/Engine/Source/Runtime/Engine/Private/SkeletalMesh.cpp b/Engine/Source/Runtime/Engine/Private/SkeletalMesh.cpp
--- a/Engine/Source/Runtime/Engine/Private/SkeletalMesh.cpp
+++ b/Engine/Source/Runtime/Engine/Private/SkeletalMesh.cpp
@@ -713,23 +713,30 @@ void USkeletalMesh::InitResources()
 					
 					//In Editor we want to make sure the data is in sync between UserSectionsData and LODModel Sections
 					const FSkelMeshSourceSectionUserData& SectionUserData = ImportLODModel.UserSectionsData.FindChecked(ImportSection.OriginalDataSectionIndex);
-					bool bDataInSync = SectionUserData.bDisabled == ImportSection.bDisabled &&
+					bool bImportDataInSync = SectionUserData.bDisabled == ImportSection.bDisabled &&
 						SectionUserData.bCastShadow == ImportSection.bCastShadow &&
-						SectionUserData.bRecomputeTangent == ImportSection.bRecomputeTangent &&
-						SectionUserData.CorrespondClothAssetIndex == ImportSection.CorrespondClothAssetIndex &&
-						SectionUserData.ClothingData.AssetGuid == ImportSection.ClothingData.AssetGuid &&
-						SectionUserData.ClothingData.AssetLodIndex == ImportSection.ClothingData.AssetLodIndex;
-					check(bDataInSync);
+						SectionUserData.bRecomputeTangent == ImportSection.bRecomputeTangent;
+					//Check the cloth only for parent section, since chunked section should not have cloth
+					if (bImportDataInSync && ImportSection.ChunkedParentSectionIndex == INDEX_NONE)
+					{
+						bImportDataInSync = SectionUserData.CorrespondClothAssetIndex == ImportSection.CorrespondClothAssetIndex &&
+							SectionUserData.ClothingData.AssetGuid == ImportSection.ClothingData.AssetGuid &&
+							SectionUserData.ClothingData.AssetLodIndex == ImportSection.ClothingData.AssetLodIndex;
+					}
 					
 					//In Editor we want to make sure the data is in sync between UserSectionsData and RenderSections
 					const FSkelMeshRenderSection& RenderSection = RenderLODModel.RenderSections[SectionIndex];
-					bDataInSync = SectionUserData.bDisabled == RenderSection.bDisabled &&
+					bool bRenderDataInSync = SectionUserData.bDisabled == RenderSection.bDisabled &&
 						SectionUserData.bCastShadow == RenderSection.bCastShadow &&
 						SectionUserData.bRecomputeTangent == RenderSection.bRecomputeTangent &&
 						SectionUserData.CorrespondClothAssetIndex == RenderSection.CorrespondClothAssetIndex &&
 						SectionUserData.ClothingData.AssetGuid == RenderSection.ClothingData.AssetGuid &&
 						SectionUserData.ClothingData.AssetLodIndex == RenderSection.ClothingData.AssetLodIndex;
-					check(bDataInSync);
+
+					if (!bImportDataInSync || !bRenderDataInSync)
+					{
+						UE_ASSET_LOG(LogSkeletalMesh, Fatal, this, TEXT("Data out of sync in lod %d. bImportDataInSync=%d, bRenderDataInSync=%d"), LODIndex, bImportDataInSync, bRenderDataInSync);
+					}
 				}
 			}
 		}
