Add the ability to stop a demo after a CSV profile completes, with -csvDemoStopAfterProfile

diff --git a/Engine/Source/Runtime/Engine/Private/DemoNetDriver.cpp b/Engine/Source/Runtime/Engine/Private/DemoNetDriver.cpp
--- a/Engine/Source/Runtime/Engine/Private/DemoNetDriver.cpp
+++ b/Engine/Source/Runtime/Engine/Private/DemoNetDriver.cpp
@@ -42,6 +42,8 @@
 
 DEFINE_LOG_CATEGORY( LogDemo );
 
+#define DEMO_CSV_PROFILING_HELPERS_ENABLED (CSV_PROFILER && (!UE_BUILD_SHIPPING))
+
 CSV_DECLARE_CATEGORY_MODULE_EXTERN(CORE_API, Basic);
 
 static TAutoConsoleVariable<float> CVarDemoRecordHz( TEXT( "demo.RecordHz" ), 8, TEXT( "Maximum number of demo frames recorded per second" ) );
@@ -3668,13 +3670,21 @@ void UDemoNetDriver::AddUserToReplay(const FString& UserString)
 	}
 }
 
-#if (CSV_PROFILER && (!UE_BUILD_SHIPPING))
+#if DEMO_CSV_PROFILING_HELPERS_ENABLED
 struct FCsvDemoSettings
 {
+	FCsvDemoSettings() 
+		: bCaptureCsv(false)
+		, StartTime(-1)
+		, EndTime(-1)
+		, FrameCount(0)
+		, bStopAfterProfile(false)
+	{}
 	bool bCaptureCsv;
 	int32 StartTime;
 	int32 EndTime;
 	int32 FrameCount;
+	bool bStopAfterProfile;
 };
 
 static FCsvDemoSettings GetCsvDemoSettings()
@@ -3687,16 +3697,15 @@ static FCsvDemoSettings GetCsvDemoSettings()
 		{
 			Settings.EndTime = -1.0;
 		}
-
 		if (!FParse::Value(FCommandLine::Get(), TEXT("-csvdemoframecount="), Settings.FrameCount))
 		{
 			Settings.FrameCount = -1;
 		}
 	}
-
+	Settings.bStopAfterProfile = FParse::Param(FCommandLine::Get(), TEXT("csvDemoStopAfterProfile"));
 	return Settings;
 }
-#endif // (CSV_PROFILER && (!UE_BUILD_SHIPPING))
+#endif // DEMO_CSV_PROFILING_HELPERS_ENABLED
 
 class FDemoNetDriverReplayPlaylistHelper
 {
@@ -3719,9 +3728,9 @@ void UDemoNetDriver::TickDemoPlayback(float DeltaSeconds)
 		return;
 	}
 
-#if (CSV_PROFILER && (!UE_BUILD_SHIPPING))
+#if DEMO_CSV_PROFILING_HELPERS_ENABLED
+	static FCsvDemoSettings CsvDemoSettings = GetCsvDemoSettings();
 	{
-		static FCsvDemoSettings CsvDemoSettings = GetCsvDemoSettings();
 		if (CsvDemoSettings.bCaptureCsv)
 		{
 			bool bDoCapture = IsPlaying()
@@ -3741,7 +3750,7 @@ void UDemoNetDriver::TickDemoPlayback(float DeltaSeconds)
 			}
 		}
 	}
-#endif // (CSV_PROFILER && (!UE_BUILD_SHIPPING))
+#endif // DEMO_CSV_PROFILING_HELPERS_ENABLED
 
 	if (!IsPlaying())
 	{
@@ -3797,7 +3806,16 @@ void UDemoNetDriver::TickDemoPlayback(float DeltaSeconds)
 	// Use AtEnd() of the archive instead of checking DemoCurrentTime/DemoTotalTime, because the DemoCurrentTime may never catch up to DemoTotalTime.
 	if (FArchive* const StreamingArchive = ReplayStreamer->GetStreamingArchive())
 	{
-		const bool bIsAtEnd = StreamingArchive->AtEnd() && (PlaybackPackets.Num() == 0 || (DemoCurrentTime + DeltaSeconds >= DemoTotalTime));
+		bool bIsAtEnd = StreamingArchive->AtEnd() && (PlaybackPackets.Num() == 0 || (DemoCurrentTime + DeltaSeconds >= DemoTotalTime));
+#if DEMO_CSV_PROFILING_HELPERS_ENABLED
+		bool bCsvIsCapturing = FCsvProfiler::Get()->IsCapturing();
+		static bool bCsvProfilingEnabledPreviousTick = bCsvIsCapturing;
+	    if (CsvDemoSettings.bStopAfterProfile && !bCsvIsCapturing && bCsvProfilingEnabledPreviousTick)
+	    {
+			bIsAtEnd = true;
+	    }
+		bCsvProfilingEnabledPreviousTick = bCsvIsCapturing;
+#endif
 		if (!ReplayStreamer->IsLive() && bIsAtEnd)
 		{
 			OnDemoFinishPlaybackDelegate.Broadcast();
