[NO NON-PII COMMIT MESSAGE]

diff --git a/Engine/Source/Runtime/Engine/Private/SoundNodeRandom.cpp b/Engine/Source/Runtime/Engine/Private/SoundNodeRandom.cpp
--- a/Engine/Source/Runtime/Engine/Private/SoundNodeRandom.cpp
+++ b/Engine/Source/Runtime/Engine/Private/SoundNodeRandom.cpp
@@ -175,6 +175,10 @@ int32 USoundNodeRandom::DetermineAmountOfBranchesToPreselect()
 		// use the minimum of either:
 		AmountOfBranchesToPreselect = FMath::Min(PreselectAtLevelLoad, OverrideForAmountOfBranchesToPreselect);
 	}
+	else if (PreselectAtLevelLoad > 0)
+	{
+		AmountOfBranchesToPreselect = PreselectAtLevelLoad;
+	}
 	else
 	{
 		// Otherwise, just use the override:
