this undo specifically reverts the soundcue/soundnode changes while retaining the settings fields

diff --git a/Engine/Plugins/Runtime/SoundMod/Source/SoundMod/Private/SoundNodeModPlayer.cpp b/Engine/Plugins/Runtime/SoundMod/Source/SoundMod/Private/SoundNodeModPlayer.cpp
--- a/Engine/Plugins/Runtime/SoundMod/Source/SoundMod/Private/SoundNodeModPlayer.cpp
+++ b/Engine/Plugins/Runtime/SoundMod/Source/SoundMod/Private/SoundNodeModPlayer.cpp
@@ -22,7 +22,7 @@ void USoundNodeModPlayer::Serialize(FArchive& Ar)
 		}
 		else if (Ar.IsSaving())
 		{
-			USoundMod* HardReference = (ShouldHardReferenceAsset(Ar.CookingTarget()) ? SoundMod : nullptr);
+			USoundMod* HardReference = (ShouldHardReferenceAsset() ? SoundMod : nullptr);
 			Ar << HardReference;
 		}
 	}
diff --git a/Engine/Source/Runtime/Engine/Classes/Sound/SoundCue.h b/Engine/Source/Runtime/Engine/Classes/Sound/SoundCue.h
--- a/Engine/Source/Runtime/Engine/Classes/Sound/SoundCue.h
+++ b/Engine/Source/Runtime/Engine/Classes/Sound/SoundCue.h
@@ -110,7 +110,7 @@ protected:
 	float SubtitlePriority;
 
 private:
-	float MaxAudibleDistance;	
+	float MaxAudibleDistance;
 
 public:
 	/* Indicates whether attenuation should use the Attenuation Overrides or the Attenuation Settings asset */
@@ -122,8 +122,6 @@ public:
 	uint8 bExcludeFromRandomNodeBranchCulling : 1;
 
 private:
-	UPROPERTY()
-	int32 CookedQualityIndex = INDEX_NONE;
 	
 	/** Whether a sound has play when silent enabled (i.e. for a sound cue, if any sound wave player has it enabled). */
 	UPROPERTY()
@@ -241,9 +239,6 @@ public:
 	static void StaticAudioQualityChanged(int32 NewQualityLevel);
 
 	FORCEINLINE static int32 GetCachedQualityLevel() { return CachedQualityLevel; }
-	
-	/** Set the Quality level that the Cue was cooked at, called by the SoundQualityNodes */
-	int32 GetCookedQualityIndex() const { return CookedQualityIndex; }
 
 	/** Call to cache any values which need to be computed from the sound cue graph. e.g. MaxDistance, Duration, etc. */
 	void CacheAggregateValues();
@@ -305,7 +300,5 @@ private:
 
 	/** Ptr to interface to sound cue editor operations. */
 	static TSharedPtr<ISoundCueAudioEditor> SoundCueAudioEditor;
-
-	FCriticalSection EditorOnlyCs;
 #endif // WITH_EDITOR
 };
diff --git a/Engine/Source/Runtime/Engine/Classes/Sound/SoundNodeAssetReferencer.h b/Engine/Source/Runtime/Engine/Classes/Sound/SoundNodeAssetReferencer.h
--- a/Engine/Source/Runtime/Engine/Classes/Sound/SoundNodeAssetReferencer.h
+++ b/Engine/Source/Runtime/Engine/Classes/Sound/SoundNodeAssetReferencer.h
@@ -20,7 +20,7 @@ public:
 	virtual void LoadAsset(bool bAddToRoot = false) PURE_VIRTUAL(USoundNodeAssetReferencer::LoadAsset,);
 	virtual void ClearAssetReferences() PURE_VIRTUAL(USoundNodeAssetReferencer::ClearAssetReferences, );
 
-	bool ShouldHardReferenceAsset(const class ITargetPlatform*) const;
+	bool ShouldHardReferenceAsset() const;
 
 #if WITH_EDITOR
 	virtual void PostEditImport() override;
diff --git a/Engine/Source/Runtime/Engine/Classes/Sound/SoundNodeQualityLevel.h b/Engine/Source/Runtime/Engine/Classes/Sound/SoundNodeQualityLevel.h
--- a/Engine/Source/Runtime/Engine/Classes/Sound/SoundNodeQualityLevel.h
+++ b/Engine/Source/Runtime/Engine/Classes/Sound/SoundNodeQualityLevel.h
@@ -24,12 +24,13 @@ class USoundNodeQualityLevel : public USoundNode
 
 public:
 
+#if WITH_EDITOR
 	//~ Begin UObject Interface
 	virtual void PostLoad() override;
 	//~ End UObject Interface
+#endif
 
 	//~ Begin USoundNode Interface.
-	virtual void Serialize(FArchive& Ar) override;	
 	virtual void ParseNodes( FAudioDevice* AudioDevice, const UPTRINT NodeWaveInstanceHash, FActiveSound& ActiveSound, const FSoundParseParameters& ParseParams, TArray<FWaveInstance*>& WaveInstances ) override;
 	virtual int32 GetMaxChildNodes() const override;
 	virtual int32 GetMinChildNodes() const override;
@@ -40,15 +41,8 @@ public:
 #endif
 	//~ End USoundNode Interface.
 
-	// A Property to indicate which quality this node was cooked with. (INDEX_NONE if not cooked, or unset).
-	UPROPERTY()
-	int32 CookedQualityLevelIndex = INDEX_NONE;
-
 #if WITH_EDITOR
 	void ReconcileNode(bool bReconstructNode);
-
-	// Critical section to protect Child nodes array while serializing
-	FCriticalSection EditorOnlyCs;
 #endif
 
 };
diff --git a/Engine/Source/Runtime/Engine/Private/SoundCue.cpp b/Engine/Source/Runtime/Engine/Private/SoundCue.cpp
--- a/Engine/Source/Runtime/Engine/Private/SoundCue.cpp
+++ b/Engine/Source/Runtime/Engine/Private/SoundCue.cpp
@@ -25,16 +25,13 @@
 #include "DSP/Dsp.h"
 #if WITH_EDITOR
 #include "Kismet2/BlueprintEditorUtils.h"
+#include "Sound/AudioSettings.h"
 #include "SoundCueGraph/SoundCueGraphNode.h"
 #include "SoundCueGraph/SoundCueGraph.h"
 #include "SoundCueGraph/SoundCueGraphNode_Root.h"
 #include "SoundCueGraph/SoundCueGraphSchema.h"
 #endif // WITH_EDITOR
 
-#include "Interfaces/ITargetPlatform.h"
-#include "AudioCompressionSettings.h"
-#include "Sound/AudioSettings.h"
-
 /*-----------------------------------------------------------------------------
 	USoundCue implementation.
 -----------------------------------------------------------------------------*/
@@ -51,7 +48,6 @@ USoundCue::USoundCue(const FObjectInitializer& ObjectInitializer)
 	VolumeMultiplier = 0.75f;
 	PitchMultiplier = 1.0f;
 	SubtitlePriority = DEFAULT_SUBTITLE_PRIORITY;
-	CookedQualityIndex = INDEX_NONE;
 }
 
 #if WITH_EDITOR
@@ -110,23 +106,7 @@ void USoundCue::Serialize(FStructuredArchive::FRecord Record)
 		CacheAggregateValues();
 	}
 
-#if WITH_EDITOR
-	// If we are cooking, record our cooked quality before serialize and then undo it.
-	if (UnderlyingArchive.IsCooking() && UnderlyingArchive.IsSaving() && UnderlyingArchive.CookingTarget())
-	{
-		if (const FPlatformAudioCookOverrides* AudioCookOverrides = UnderlyingArchive.CookingTarget()->GetAudioCompressionSettings())
-		{
-			FScopeLock Lock(&EditorOnlyCs);
-			CookedQualityIndex = AudioCookOverrides->SoundCueCookQualityIndex;
-			Super::Serialize(Record);
-			CookedQualityIndex = INDEX_NONE;
-		}
-	}
-	else
-#endif //WITH_EDITOR
-	{
-		Super::Serialize(Record);
-	}
+	Super::Serialize(Record);
 
 	if (UnderlyingArchive.UE4Ver() >= VER_UE4_COOKED_ASSETS_IN_EDITOR_SUPPORT)
 	{
@@ -171,17 +151,6 @@ void USoundCue::PostLoad()
 	}
 	else
 #endif // WITH_EDITOR
-
-	// Warn if the Quality index is set to something that we can't support.
-	UE_CLOG(USoundCue::GetCachedQualityLevel() != CookedQualityIndex && CookedQualityIndex != INDEX_NONE, LogAudio, Warning,
-		TEXT("'%s' is ingoring Quality Setting '%s'(%d) as it was cooked with '%s'(%d)"),
-		*GetFullNameSafe(this),
-		*GetDefault<UAudioSettings>()->FindQualityNameByIndex(USoundCue::GetCachedQualityLevel()),
-		USoundCue::GetCachedQualityLevel(),
-		*GetDefault<UAudioSettings>()->FindQualityNameByIndex(CookedQualityIndex),
-		CookedQualityIndex
-	);
-
 	if (GEngine && *GEngine->GameUserSettingsClass)
 	{
 		EvaluateNodes(false);
@@ -255,14 +224,7 @@ void USoundCue::EvaluateNodes(bool bAddToRoot)
 		}
 	};
 
-	// Only Evaluate nodes if we haven't been cooked, as cooked builds will hard-ref all SoundAssetReferences.	
-	UE_CLOG(CookedQualityIndex == INDEX_NONE, LogAudio, Verbose, TEXT("'%s', DOING EvaluateNodes as we are *NOT* cooked"), *GetName());
-	UE_CLOG(CookedQualityIndex != INDEX_NONE, LogAudio, Verbose, TEXT("'%s', SKIPPING EvaluateNodes as we *ARE* cooked"), *GetName());
-
-	if (CookedQualityIndex == INDEX_NONE)
-	{		
-		EvaluateNodes_Internal(FirstNode);
-	}
+	EvaluateNodes_Internal(FirstNode);
 }
 
 float USoundCue::FindMaxDistanceInternal() const
diff --git a/Engine/Source/Runtime/Engine/Private/SoundNodeAssetReferencer.cpp b/Engine/Source/Runtime/Engine/Private/SoundNodeAssetReferencer.cpp
--- a/Engine/Source/Runtime/Engine/Private/SoundNodeAssetReferencer.cpp
+++ b/Engine/Source/Runtime/Engine/Private/SoundNodeAssetReferencer.cpp
@@ -5,31 +5,11 @@
 #include "Sound/SoundNodeRandom.h"
 #include "Sound/SoundCue.h"
 #include "UObject/FrameworkObjectVersion.h"
-#include "Interfaces/ITargetPlatform.h"
-#include "AudioCompressionSettings.h"
-#include "Sound/AudioSettings.h"
 
 #define ASYNC_LOAD_RANDOMIZED_SOUNDS 1
 
-bool USoundNodeAssetReferencer::ShouldHardReferenceAsset(const ITargetPlatform* TargetPlatform) const
+bool USoundNodeAssetReferencer::ShouldHardReferenceAsset() const
 {
-	if (TargetPlatform)
-	{
-		if (const FPlatformAudioCookOverrides* CookOverrides = TargetPlatform->GetAudioCompressionSettings())
-		{			
-			// If the Quality nodes are cooked, everything is hard refs.
-			if (CookOverrides->SoundCueCookQualityIndex != INDEX_NONE)
-			{
-				UE_LOG(LogAudio, Verbose, TEXT("HARD reffing '%s:%s', as we are cooking using quality '%s'"), 
-					*GetNameSafe(GetOuter()),
-					*GetName(),
-					*GetDefault<UAudioSettings>()->FindQualityNameByIndex(CookOverrides->SoundCueCookQualityIndex)
-				)
-				return true;
-			}
-		}
-	}
-
 	bool bShouldHardReference = true;
 
 	if (USoundCue* Cue = Cast<USoundCue>(GetOuter()))
@@ -72,12 +52,6 @@ bool USoundNodeAssetReferencer::ShouldHardReferenceAsset(const ITargetPlatform*
 		
 	}
 
-	UE_LOG(LogAudio, Verbose, TEXT("%s reffing '%s:%s'."),
-		bShouldHardReference ? TEXT("HARD") : TEXT("SOFT"),
-		*GetNameSafe(GetOuter()),
-		*GetName()
-	);
-
 	return bShouldHardReference;
 }
 
diff --git a/Engine/Source/Runtime/Engine/Private/SoundNodeQualityLevel.cpp b/Engine/Source/Runtime/Engine/Private/SoundNodeQualityLevel.cpp
--- a/Engine/Source/Runtime/Engine/Private/SoundNodeQualityLevel.cpp
+++ b/Engine/Source/Runtime/Engine/Private/SoundNodeQualityLevel.cpp
@@ -13,9 +13,14 @@
 #include "Settings/LevelEditorPlaySettings.h"
 #include "Editor.h"
 #endif
-#include "Interfaces/ITargetPlatform.h"
 
 #if WITH_EDITOR
+void USoundNodeQualityLevel::PostLoad()
+{
+	Super::PostLoad();
+
+	ReconcileNode(false);
+}
 
 void USoundNodeQualityLevel::ReconcileNode(bool bReconstructNode)
 {
@@ -42,24 +47,6 @@ FText USoundNodeQualityLevel::GetInputPinName(int32 PinIndex) const
 }
 #endif
 
-void USoundNodeQualityLevel::PostLoad()
-{
-	Super::PostLoad();
-#if WITH_EDITOR
-	ReconcileNode(false);
-#endif //WITH_EDITOR
-
-	UE_CLOG(Cast<USoundCue>(GetOuter()) && Cast<USoundCue>(GetOuter())->GetCookedQualityIndex() != CookedQualityLevelIndex && CookedQualityLevelIndex != INDEX_NONE,
-		LogAudio,
-		Warning, TEXT("'%s' has been cooked with multiple quality levels. '%s'(%d) vs '%s'(%d)"),
-		*GetFullNameSafe(this),
-		*GetDefault<UAudioSettings>()->FindQualityNameByIndex(USoundCue::GetCachedQualityLevel()),
-		USoundCue::GetCachedQualityLevel(),
-		*GetDefault<UAudioSettings>()->FindQualityNameByIndex(CookedQualityLevelIndex),
-		CookedQualityLevelIndex
-	);
-}
-
 void USoundNodeQualityLevel::PrimeChildWavePlayers(bool bRecurse)
 {
 	// If we're able to retrieve a valid cached quality level for this sound cue,
@@ -115,62 +102,11 @@ void USoundNodeQualityLevel::ParseNodes( FAudioDevice* AudioDevice, const UPTRIN
 		QualityLevel = USoundCue::GetCachedQualityLevel();
 	}
 #else
-	
-	int32 QualityLevel = USoundCue::GetCachedQualityLevel();
-	
-	// If CookedQualityLevelIndex has been set, we will have a *single* quality level.
-	if (CookedQualityLevelIndex >= 0)
-	{	
-		// Remap to index 0 (as all other levels have been removed by cooker).
-		QualityLevel = 0;
-	}
-#endif	
-	
+	const int32 QualityLevel = USoundCue::GetCachedQualityLevel();
+#endif
+
 	if (ChildNodes.IsValidIndex(QualityLevel) && ChildNodes[QualityLevel])
 	{
 		ChildNodes[QualityLevel]->ParseNodes( AudioDevice, GetNodeWaveInstanceHash(NodeWaveInstanceHash, ChildNodes[QualityLevel], QualityLevel), ActiveSound, ParseParams, WaveInstances );
 	}
 }
-
-void USoundNodeQualityLevel::Serialize(FArchive& Ar)
-{
-#if WITH_EDITOR
-
-	if (Ar.IsCooking() && Ar.IsSaving() && Ar.CookingTarget() )
-	{			
-		if (const FPlatformAudioCookOverrides* AudioCookOverrides = Ar.CookingTarget()->GetAudioCompressionSettings())
-		{
-			// Prevent any other thread saving this class while we are modifying the ChildNode array.
-			FScopeLock Lock(&EditorOnlyCs);
-			
-			// Are we doing any cook quality level optimization for this Node?			
-			if (ChildNodes.IsValidIndex(AudioCookOverrides->SoundCueCookQualityIndex))
-			{
-				// Set our cook quality, as we serialize. 
-				CookedQualityLevelIndex = AudioCookOverrides->SoundCueCookQualityIndex;
-
-				// Move out all nodes.
-				TArray<USoundNode*> ChildNodesBackup;
-				ChildNodesBackup = MoveTemp(ChildNodes);
-				
-				// Put *just* the node we care about in our child array to be serialized by the Super
-				ChildNodes.Add(ChildNodesBackup[CookedQualityLevelIndex]);
-
-				// Call base serialize that will walk all properties and serialize them.
-				Super::Serialize(Ar);
-
-				// Return to our original state. (careful the cook only variables don't leak out).
-				ChildNodes = MoveTemp(ChildNodesBackup);
-				CookedQualityLevelIndex = INDEX_NONE;
-
-				// We are done.
-				return;
-			}	
-		}		
-	}
-		
-#endif //WITH_EDITOR	
-
-	// ... in all other cases, we just call the super.
-	Super::Serialize(Ar);
-}
\ No newline at end of file
diff --git a/Engine/Source/Runtime/Engine/Private/SoundNodeWavePlayer.cpp b/Engine/Source/Runtime/Engine/Private/SoundNodeWavePlayer.cpp
--- a/Engine/Source/Runtime/Engine/Private/SoundNodeWavePlayer.cpp
+++ b/Engine/Source/Runtime/Engine/Private/SoundNodeWavePlayer.cpp
@@ -22,7 +22,7 @@ void USoundNodeWavePlayer::Serialize(FArchive& Ar)
 		}
 		else if (Ar.IsSaving())
 		{
-			USoundWave* HardReference = (ShouldHardReferenceAsset(Ar.CookingTarget()) ? SoundWave : nullptr);
+			USoundWave* HardReference = (ShouldHardReferenceAsset() ? SoundWave : nullptr);
 			Ar << HardReference;
 		}
 	}
@@ -45,7 +45,6 @@ void USoundNodeWavePlayer::LoadAsset(bool bAddToRoot)
 			const FString LongPackageName = SoundWaveAssetPtr.GetLongPackageName();
 			if (!LongPackageName.IsEmpty())
 			{
-				UE_LOG(LogAudio, VeryVerbose, TEXT(" '%s:%s', Async loading..."), *GetNameSafe(GetOuter()), *GetName());
 				bAsyncLoading = true;
 				LoadPackageAsync(LongPackageName, FLoadPackageAsyncDelegate::CreateUObject(this, &USoundNodeWavePlayer::OnSoundWaveLoaded, bAddToRoot));
 			}
@@ -61,12 +60,6 @@ void USoundNodeWavePlayer::LoadAsset(bool bAddToRoot)
 	}
 	else
 	{
-		UE_LOG(LogAudio, VeryVerbose, TEXT("'%s:%s', DOING SYNCRONOUS loading... %s"),
-			*GetNameSafe(GetOuter()),
-			*GetName(),
-			SoundWaveAssetPtr.IsPending() ? TEXT("pending.") : TEXT("not-pending.")
-		);
-
 		SoundWave = SoundWaveAssetPtr.LoadSynchronous();
 		if (SoundWave)
 		{
