Delegates for FUserActivityTracking::OnActivityChanged are now unregistered when FEngineSessionManager and FEditorSessionSummaryWriter shutdown.

diff --git a/Engine/Source/Runtime/Engine/Private/Analytics/EditorSessionSummaryWriter.cpp b/Engine/Source/Runtime/Engine/Private/Analytics/EditorSessionSummaryWriter.cpp
--- a/Engine/Source/Runtime/Engine/Private/Analytics/EditorSessionSummaryWriter.cpp
+++ b/Engine/Source/Runtime/Engine/Private/Analytics/EditorSessionSummaryWriter.cpp
@@ -170,13 +170,11 @@ void FEditorSessionSummaryWriter::Tick(float DeltaTime)
 void FEditorSessionSummaryWriter::Shutdown()
 {
 	FCoreDelegates::OnHandleSystemError.RemoveAll(this);
-	FCoreDelegates::ApplicationHasReactivatedDelegate.RemoveAll(this);
-	FCoreDelegates::ApplicationWillDeactivateDelegate.RemoveAll(this);
-	FCoreDelegates::ApplicationWillEnterBackgroundDelegate.RemoveAll(this);
-	FCoreDelegates::ApplicationHasEnteredForegroundDelegate.RemoveAll(this);
 	FCoreDelegates::ApplicationWillTerminateDelegate.RemoveAll(this);
 	FCoreDelegates::IsVanillaProductChanged.RemoveAll(this);
 
+	FUserActivityTracking::OnActivityChanged.RemoveAll(this);
+
 	if (CurrentSession != nullptr)
 	{
 		if (!CurrentSession->bIsTerminating && !CurrentSession->bCrashed)
diff --git a/Engine/Source/Runtime/Engine/Private/EngineSessionManager.cpp b/Engine/Source/Runtime/Engine/Private/EngineSessionManager.cpp
--- a/Engine/Source/Runtime/Engine/Private/EngineSessionManager.cpp
+++ b/Engine/Source/Runtime/Engine/Private/EngineSessionManager.cpp
@@ -243,6 +243,8 @@ void FEngineSessionManager::Shutdown()
 	FCoreDelegates::ApplicationWillTerminateDelegate.RemoveAll(this);
 	FCoreDelegates::IsVanillaProductChanged.RemoveAll(this);
 
+	FUserActivityTracking::OnActivityChanged.RemoveAll(this);
+
 	if (!CurrentSession.bIsTerminating) // Skip Slate if terminating, since we can't guarantee which thread called us.
 	{
 		FSlateApplication::Get().GetOnModalLoopTickEvent().RemoveAll(this);
