Comment out RepairChaosActor calls and remove spawning from InitializeNewWorld.

diff --git a/Engine/Source/Runtime/Engine/Private/World.cpp b/Engine/Source/Runtime/Engine/Private/World.cpp
--- a/Engine/Source/Runtime/Engine/Private/World.cpp
+++ b/Engine/Source/Runtime/Engine/Private/World.cpp
@@ -960,7 +960,7 @@ void UWorld::PostLoad()
 	RepairWorldSettings();
 #endif
 #if INCLUDE_CHAOS
-	RepairChaosActors();
+	//RepairChaosActors();
 #endif
 
 	for (auto It = StreamingLevels.CreateIterator(); It; ++It)
@@ -1411,7 +1411,7 @@ void UWorld::InitWorld(const InitializationValues IVS)
 	RepairWorldSettings();
 #endif
 #if INCLUDE_CHAOS
-	RepairChaosActors();
+	//RepairChaosActors();
 #endif
 
 	// initialize DefaultPhysicsVolume for the world
@@ -1635,13 +1635,13 @@ void UWorld::InitializeNewWorld(const InitializationValues IVS)
 #endif
 
 #if INCLUDE_CHAOS
-	FChaosSolversModule* ChaosModule = FModuleManager::Get().GetModulePtr<FChaosSolversModule>("ChaosSolvers");
+	/*FChaosSolversModule* ChaosModule = FModuleManager::Get().GetModulePtr<FChaosSolversModule>("ChaosSolvers");
 	check(ChaosModule);
 	FActorSpawnParameters ChaosSpawnInfo;
 	ChaosSpawnInfo.SpawnCollisionHandlingOverride = ESpawnActorCollisionHandlingMethod::AlwaysSpawn;
 	ChaosSpawnInfo.Name = TEXT("DefaultChaosActor");
 	SpawnActor(ChaosModule->GetSolverActorClass(), nullptr, nullptr, ChaosSpawnInfo);
-	check(PhysicsScene_Chaos);
+	check(PhysicsScene_Chaos);*/
 #endif
 
 	// Initialize the world
