update debug exec cmds to have new "au" prefix instead of the now depricated "Audio."

diff --git a/Engine/Source/Runtime/Engine/Private/AudioDebug.cpp b/Engine/Source/Runtime/Engine/Private/AudioDebug.cpp
--- a/Engine/Source/Runtime/Engine/Private/AudioDebug.cpp
+++ b/Engine/Source/Runtime/Engine/Private/AudioDebug.cpp
@@ -178,8 +178,8 @@ namespace
 		}
 	}
 
-	FAutoConsoleCommandWithWorld DumpActiveSounds(TEXT("Audio.DumpActiveSounds"), TEXT("Outputs data about all the currently active sounds."), FConsoleCommandWithWorldDelegate::CreateStatic(&HandleDumpActiveSounds), ECVF_Cheat);
-	FAutoConsoleCommandWithWorld ClearMutesAndSolos(TEXT("Audio.ClearMutesAndSolos"), TEXT("Clears any solo-ing/mute-ing sounds"), FConsoleCommandWithWorldDelegate::CreateStatic(&HandleClearMutesAndSolos), ECVF_Cheat);
+	FAutoConsoleCommandWithWorld DumpActiveSounds(TEXT("au.DumpActiveSounds"), TEXT("Outputs data about all the currently active sounds."), FConsoleCommandWithWorldDelegate::CreateStatic(&HandleDumpActiveSounds), ECVF_Cheat);
+	FAutoConsoleCommandWithWorld ClearMutesAndSolos(TEXT("au.ClearMutesAndSolos"), TEXT("Clears any solo-ing/mute-ing sounds"), FConsoleCommandWithWorldDelegate::CreateStatic(&HandleClearMutesAndSolos), ECVF_Cheat);
 
 } // namespace <>
 
