Fix case where virtual texture file cache has more line reservation requests then it can handle

diff --git a/Engine/Source/Runtime/Engine/Private/FileCache/FileCache.cpp b/Engine/Source/Runtime/Engine/Private/FileCache/FileCache.cpp
--- a/Engine/Source/Runtime/Engine/Private/FileCache/FileCache.cpp
+++ b/Engine/Source/Runtime/Engine/Private/FileCache/FileCache.cpp
@@ -473,14 +473,7 @@ IMemoryReadStreamRef FFileCacheHandle::ReadData(FGraphEventArray& OutCompletionE
 		LineToRequest.SetNum(TotalNumSlots, false);
 	}
 
-	int32 NumSlotsNeeded = 0;
-	for (CacheLineID LineID = StartLine; LineID.Get() <= EndLine.Get(); ++LineID)
-	{
-		if (LineID.Get() >= NumSlots || !LineToSlot[LineID.Get()].IsValid())
-		{
-			++NumSlotsNeeded;
-		}
-	}
+	int32 NumSlotsNeeded = EndLine.Get() + 1 - StartLine.Get();
 
 	FFileCache& Cache = GetCache();
 
diff --git a/Engine/Source/Runtime/Engine/Private/VT/VirtualTextureChunkManager.cpp b/Engine/Source/Runtime/Engine/Private/VT/VirtualTextureChunkManager.cpp
--- a/Engine/Source/Runtime/Engine/Private/VT/VirtualTextureChunkManager.cpp
+++ b/Engine/Source/Runtime/Engine/Private/VT/VirtualTextureChunkManager.cpp
@@ -16,13 +16,12 @@
 #include "VirtualTextureChunkDDCCache.h"
 #endif
 
-static int32 NumTranscodeRequests = 128;
+static int32 NumTranscodeRequests = 32;
 static FAutoConsoleVariableRef CVarNumTranscodeRequests(
 	TEXT("r.VT.NumTranscodeRequests"),
 	NumTranscodeRequests,
-	TEXT("Number of transcode request that can be in flight. default 128\n")
-	, ECVF_ReadOnly
-);
+	TEXT("Number of transcode request that can be in flight. default 32\n"),
+	ECVF_Default);
 
 static FVirtualTextureChunkStreamingManager* VirtualTexturePageStreamingManager = nullptr;
 struct FVirtualTextureChunkStreamingManager& FVirtualTextureChunkStreamingManager::Get()
