Fixed crash when deleting a material assigned to a BSP.

diff --git a/Engine/Source/Runtime/Engine/Classes/Components/ModelComponent.h b/Engine/Source/Runtime/Engine/Classes/Components/ModelComponent.h
--- a/Engine/Source/Runtime/Engine/Classes/Components/ModelComponent.h
+++ b/Engine/Source/Runtime/Engine/Classes/Components/ModelComponent.h
@@ -117,7 +117,7 @@ public:
 	virtual void Serialize(FArchive& Ar) override;
 	virtual void PostLoad() override;
 #if WITH_EDITOR
-	virtual void PostEditUndo() override;
+	virtual void PostEditChangeProperty(FPropertyChangedEvent& PropertyChangedEvent) override;
 #endif // WITH_EDITOR
 	virtual bool IsNameStableForNetworking() const override;
 	static void AddReferencedObjects(UObject* InThis, FReferenceCollector& Collector);
@@ -192,4 +192,3 @@ protected:
 };
 
 
-
diff --git a/Engine/Source/Runtime/Engine/Private/Components/ModelComponent.cpp b/Engine/Source/Runtime/Engine/Private/Components/ModelComponent.cpp
--- a/Engine/Source/Runtime/Engine/Private/Components/ModelComponent.cpp
+++ b/Engine/Source/Runtime/Engine/Private/Components/ModelComponent.cpp
@@ -309,15 +309,26 @@ void UModelComponent::PostLoad()
 }
 
 #if WITH_EDITOR
-void UModelComponent::PostEditUndo()
+void UModelComponent::PostEditChangeProperty(FPropertyChangedEvent& PropertyChangedEvent)
 {
 	// Rebuild the component's render data after applying a transaction to it.
 	ULevel* Level = GetTypedOuter<ULevel>();
-	if(ensure(Level))
+	if (ensure(Level))
 	{
 		Level->InvalidateModelSurface();
+		if (PropertyChangedEvent.ChangeType == EPropertyChangeType::Redirected)
+		{
+			// Normally, we can wait until the next tick to update the model surfaces. However, when force-deleting a material (signaled by the
+			// Redirected change type), UActorComponent::PostEditChangeProperty() will try to re-register the component, which ends up using render
+			// resources such as the index buffer. We need to call CommitModelSurfaces() now to make sure these resources are rebuilt taking into
+			// account the replacement material. This is not optimal, since it will be called multiple times, but it only happens when deleting a
+			// material assigned to a BSP component.
+			// This is a stopgap measure until we bite the bullet and refactor the update code so that resources aren't used before they are 
+			// recreated by the normal level rebuilding process.
+			Level->CommitModelSurfaces();
+		}
 	}
-	Super::PostEditUndo();
+	Super::PostEditChangeProperty(PropertyChangedEvent);
 }
 #endif // WITH_EDITOR
 
