Fix crash when editing and saving a default material

When FMaterialEditor::UpdateOriginalMaterial() duplicates the original material (due to hitting save in the material editor), the PostLoad() callback inside UMaterial recompiles the shader. The bUsedAsSpecialEngineMaterial flag hasn't been set yet on the copy (it's set after StaticDuplicateObject() returns), so the copy may not be usable with all the object types that the original material was used for. When this is the case, e.g. when the material is assigned to a spline mesh, we trigger an assert in FMaterial::GetShaderPipeline().

FMaterialEditor::UpdateOriginalMaterial() uses UMaterial::ForceNoCompilationInPostLoad() to disable compiling during duplication, but the flag that's set by that function was not actually used anywhere. The check for GAllowCompilationInPostLoad seems to have been accidentally removed in CL 2044510.

diff --git a/Engine/Source/Runtime/Engine/Private/Materials/Material.cpp b/Engine/Source/Runtime/Engine/Private/Materials/Material.cpp
--- a/Engine/Source/Runtime/Engine/Private/Materials/Material.cpp
+++ b/Engine/Source/Runtime/Engine/Private/Materials/Material.cpp
@@ -3999,7 +3999,7 @@ void UMaterial::PostLoad()
 		}
 #endif
 		//Don't compile shaders in post load for dev overhead materials.
-		if (FApp::CanEverRender() && !bIsMaterialEditorStatsMaterial)
+		if (FApp::CanEverRender() && !bIsMaterialEditorStatsMaterial && GAllowCompilationInPostLoad)
 		{
 			// Before caching shader resources we have to make sure all referenced textures have been post loaded
 			// as we depend on their resources being valid.
