Add option to redirect ShaderDebugInfo (r.OverrideShaderDebugDir)

diff --git a/Engine/Source/Runtime/Engine/Private/ShaderCompiler/ShaderCompiler.cpp b/Engine/Source/Runtime/Engine/Private/ShaderCompiler/ShaderCompiler.cpp
--- a/Engine/Source/Runtime/Engine/Private/ShaderCompiler/ShaderCompiler.cpp
+++ b/Engine/Source/Runtime/Engine/Private/ShaderCompiler/ShaderCompiler.cpp
@@ -266,6 +266,15 @@ static FAutoConsoleVariableRef CVarCreateShadersOnLoad(
 	TEXT("Whether to create shaders on load, which can reduce hitching, but use more memory.  Otherwise they will be created as needed.")
 );
 
+
+static TAutoConsoleVariable<FString> CVarShaderOverrideDebugDir(
+	TEXT("r.OverrideShaderDebugDir"),
+	"",
+	TEXT("Override output location of shader debug files\n")
+	TEXT("Empty: use default location Saved\\ShaderDebugInfo.\n"),
+	ECVF_ReadOnly);
+
+
 extern bool CompileShaderPipeline(const IShaderFormat* Compiler, FName Format, FShaderPipelineCompileJob* PipelineJob, const FString& Dir);
 
 #if ENABLE_COOK_STATS
@@ -1572,6 +1581,11 @@ FShaderCompilingManager::FShaderCompilingManager() :
 	AbsoluteShaderBaseWorkingDirectory = AbsoluteBaseDirectory + TEXT("/");
 
 	FString AbsoluteDebugInfoDirectory = IFileManager::Get().ConvertToAbsolutePathForExternalAppForWrite(*(FPaths::ProjectSavedDir() / TEXT("ShaderDebugInfo")));
+	const FString OverrideShaderDebugDir = CVarShaderOverrideDebugDir.GetValueOnAnyThread();
+	if (!OverrideShaderDebugDir.IsEmpty())
+	{
+		AbsoluteDebugInfoDirectory = OverrideShaderDebugDir;
+	}
 	FPaths::NormalizeDirectoryName(AbsoluteDebugInfoDirectory);
 	AbsoluteShaderDebugInfoDirectory = AbsoluteDebugInfoDirectory;
 
