Fix CHAOS_ENSURE to allow use as condition. Change empty convex ensure to CHAOS_ENSURE to prevent breaking cooks.

diff --git a/Engine/Source/Runtime/Engine/Private/PhysicsEngine/BodySetup.cpp b/Engine/Source/Runtime/Engine/Private/PhysicsEngine/BodySetup.cpp
--- a/Engine/Source/Runtime/Engine/Private/PhysicsEngine/BodySetup.cpp
+++ b/Engine/Source/Runtime/Engine/Private/PhysicsEngine/BodySetup.cpp
@@ -21,6 +21,7 @@
 #include "UObject/UObjectIterator.h"
 #include "UObject/PropertyPortFlags.h"
 #include "Components/SplineMeshComponent.h"
+#include "ChaosCheck.h"
 
 #include "PhysXCookHelper.h"
 
@@ -696,7 +697,7 @@ void UBodySetup::FinishCreatingPhysicsMeshes_Chaos(FChaosDerivedDataReader<float
 		{
 			FKConvexElem& ConvexElem = AggGeom.ConvexElems[ElementIndex];
 
-			if (ensure(InReader.ConvexImplicitObjects[ElementIndex]->IsValidGeometry()))
+			if (CHAOS_ENSURE(InReader.ConvexImplicitObjects[ElementIndex]->IsValidGeometry()))
 			{
 				ConvexElem.SetChaosConvexMesh(MoveTemp(InReader.ConvexImplicitObjects[ElementIndex]));
 
diff --git a/Engine/Source/Runtime/Experimental/Chaos/Public/ChaosCheck.h b/Engine/Source/Runtime/Experimental/Chaos/Public/ChaosCheck.h
--- a/Engine/Source/Runtime/Experimental/Chaos/Public/ChaosCheck.h
+++ b/Engine/Source/Runtime/Experimental/Chaos/Public/ChaosCheck.h
@@ -8,6 +8,6 @@
 #define CHAOS_CHECK(Condition) check(Condition)
 #define CHAOS_ENSURE(Condition) ensure(Condition)
 #else
-#define CHAOS_CHECK(Condition)
-#define CHAOS_ENSURE(Condition)
+#define CHAOS_CHECK(Condition) (!!(Condition))
+#define CHAOS_ENSURE(Condition) (!!(Condition))
 #endif
