Changed CSM Shadows to have a fade out region for the last cascade to prevent hard shadow line in lower quality scalability settings.

diff --git a/Engine/Source/Runtime/Engine/Private/Components/DirectionalLightComponent.cpp b/Engine/Source/Runtime/Engine/Private/Components/DirectionalLightComponent.cpp
--- a/Engine/Source/Runtime/Engine/Private/Components/DirectionalLightComponent.cpp
+++ b/Engine/Source/Runtime/Engine/Private/Components/DirectionalLightComponent.cpp
@@ -779,12 +779,15 @@ private:
 		{
 			SplitFar += FadeExtension;
 		}
-		else if (bIsRayTracedCascade)
+		else 
 		{
-			FadeExtension *= FMath::Clamp(CVarRtdfFarTransitionScale.GetValueOnAnyThread(), 0.0f, 1.0f);
-			// For the ray-traced cascade, we want to fade out to avoid a hard line.
-			// Since there is no further cascade, extending the far makes less sensse as it affects performance by extending the shadow range.
-			// Thus, we instead move the fade plane closer.
+			if (bIsRayTracedCascade)
+			{
+				FadeExtension *= FMath::Clamp(CVarRtdfFarTransitionScale.GetValueOnAnyThread(), 0.0f, 1.0f);
+			}
+			// For the last cascade, we want to fade out to avoid a hard line, since there is no further cascade to overlap with, 
+			// extending the far makes little sensse as extending the shadow range would be counter intuitive and affect performance. 
+			// Thus, move the fade plane closer:
 			FadePlane -= FadeExtension;
 		}
 
