Fix crash when we compute the derivedatakey, the archive do not support FSoftObjectPath, since the derive data key is a writer serialize only, I serialize a string instead
#rnx

diff --git a/Engine/Source/Editor/Persona/Private/PersonaMeshDetails.cpp b/Engine/Source/Editor/Persona/Private/PersonaMeshDetails.cpp
--- a/Engine/Source/Editor/Persona/Private/PersonaMeshDetails.cpp
+++ b/Engine/Source/Editor/Persona/Private/PersonaMeshDetails.cpp
@@ -237,6 +237,11 @@ private:
 
 	bool IsNeedApplyLODChange() const
 	{
+		if (!BuildAvailable)
+		{
+			return true;
+		}
+
 		TSharedPtr<IPersonaToolkit> SharedToolkit = PersonaToolkit.Pin();
 		if (SharedToolkit.IsValid())
 		{
diff --git a/Engine/Source/Runtime/Engine/Private/SkeletalMesh.cpp b/Engine/Source/Runtime/Engine/Private/SkeletalMesh.cpp
--- a/Engine/Source/Runtime/Engine/Private/SkeletalMesh.cpp
+++ b/Engine/Source/Runtime/Engine/Private/SkeletalMesh.cpp
@@ -3837,14 +3837,17 @@ FGuid FSkeletalMeshLODInfo::ComputeDeriveDataCacheKey(const FSkeletalMeshLODGrou
 	Ar << BonesToRemove;
 	Ar << BonesToPrioritize;
 	Ar << WeightOfPrioritization;
+
+	//TODO: Ask the derivedata key of the UObject reference by FSoftObjectPath. So if someone change the UObject, this LODs will get dirty
+	//and will be rebuild.
 	if (BakePose != nullptr)
 	{
-		FSoftObjectPath BakePosePath(BakePose);
+		FString BakePosePath = BakePose->GetFullName();
 		Ar << BakePosePath;
 	}
 	if (BakePoseOverride != nullptr)
 	{
-		FSoftObjectPath BakePoseOverridePath(BakePoseOverride);
+		FString BakePoseOverridePath = BakePoseOverride->GetFullName();
 		Ar << BakePoseOverridePath;
 	}
 	FArchive_Serialize_BitfieldBool(Ar, bAllowCPUAccess);
