Improve performance of SortActorsHierarchy

diff --git a/Engine/Source/Runtime/Engine/Private/Level.cpp b/Engine/Source/Runtime/Engine/Private/Level.cpp
--- a/Engine/Source/Runtime/Engine/Private/Level.cpp
+++ b/Engine/Source/Runtime/Engine/Private/Level.cpp
@@ -763,133 +763,68 @@ void ULevel::UpdateLevelComponents(bool bRerunConstructionScripts)
 	IncrementalUpdateComponents( 0, bRerunConstructionScripts );
 }
 
-namespace FLevelSortUtils
+/**
+ *	Sorts actors such that parent actors will appear before children actors in the list
+ *	Stable sort
+ */
+static void SortActorsHierarchy(TArray<AActor*>& Actors, UObject* Level)
 {
-	void AddToListSafe(AActor* TestActor, TArray<AActor*>& List)
+	const double StartTime = FPlatformTime::Seconds();
+
+	TMap<AActor*, int32> DepthMap;
+	TArray<AActor*, TInlineAllocator<10>> VisitedActors;
+
+	TFunction<int32(AActor*)> CalcAttachDepth = [&DepthMap, &VisitedActors, &CalcAttachDepth](AActor* Actor)
 	{
-		if (TestActor)
+		int32 Depth = 0;
+		if (int32* FoundDepth = DepthMap.Find(Actor))
 		{
-			const bool bAlreadyAdded = List.Contains(TestActor);
-			if (bAlreadyAdded)
+			Depth = *FoundDepth;
+		}
+		else
+		{
+			if (AActor* ParentActor = Actor->GetAttachParentActor())
 			{
-				FString ListItemDesc;
-				for (int32 Idx = 0; Idx < List.Num(); Idx++)
+				if (VisitedActors.Contains(ParentActor))
 				{
-					if (Idx > 0)
+					FString VisitedActorLoop;
+					for (AActor* VisitedActor : VisitedActors)
 					{
-						ListItemDesc += TEXT(", ");
+						VisitedActorLoop += VisitedActor->GetName() + TEXT(" -> ");
 					}
+					VisitedActorLoop += Actor->GetName();
 
-					ListItemDesc += GetNameSafe(List[Idx]);
+					UE_LOG(LogLevel, Warning, TEXT("Found loop in attachment hierarchy: %s"), *VisitedActorLoop);
+					// Once we find a loop, depth is mostly meaningless, so we'll treat the "end" of the loop as 0
+				}
+				else
+				{
+					VisitedActors.Add(Actor);
+					Depth = CalcAttachDepth(ParentActor) + 1;
 				}
-
-				UE_LOG(LogLevel, Warning, TEXT("Found a cycle in actor's parent chain: %s"), *ListItemDesc);
-			}
-			else
-			{
-				List.Add(TestActor);
 			}
+			DepthMap.Add(Actor, Depth);
 		}
-	}
-
-	// Finds list of parents from an entry in ParentMap, returns them in provided array and removes from map
-	// Logs an error when cycle is found
-	void FindAndRemoveParentChain(TMap<AActor*, AActor*>& ParentMap, TArray<AActor*>& ParentChain)
-	{
-		check(ParentMap.Num());
-		
-		// seed from first entry
-		TMap<AActor*, AActor*>::TIterator It(ParentMap);
-		ParentChain.Add(It.Key());
-		ParentChain.Add(It.Value());
-		It.RemoveCurrent();
-
-		// fill chain's parent nodes
-		bool bLoop = true;
-		while (bLoop)
-		{
-			AActor* MapValue = nullptr;
-			bLoop = ParentMap.RemoveAndCopyValue(ParentChain.Last(), MapValue);
-			AddToListSafe(MapValue, ParentChain);
-		}
-
-		// find chain's child nodes, ignore cycle detection since it would've triggered already from previous loop
-		for (AActor* const* MapKey = ParentMap.FindKey(ParentChain[0]); MapKey; MapKey = ParentMap.FindKey(ParentChain[0]))
-		{
-			AActor* MapValue = nullptr;
-			ParentMap.RemoveAndCopyValue((AActor*)MapKey, MapValue);
-			ParentChain.Insert(MapValue, 0);
-		}
-	}
-
-	struct FDepthSort
-	{
-		TMap<AActor*, int32>* DepthMap;
-
-		bool operator()(AActor* A, AActor* B) const
-		{
-			const int32 DepthA = A ? DepthMap->FindRef(A) : MAX_int32;
-			const int32 DepthB = B ? DepthMap->FindRef(B) : MAX_int32;
-			return DepthA < DepthB;
-		}
+		return Depth;
 	};
-}
 
-/**
- *	Sorts actors such that parent actors will appear before children actors in the list
- *	Stable sort
- */
-static void SortActorsHierarchy(TArray<AActor*>& Actors, UObject* Level)
-{
-	const double StartTime = FPlatformTime::Seconds();
-
-	// Precalculate parent map to avoid processing cycles during sort
-	TMap<AActor*, AActor*> ParentMap;
-	for (int32 Idx = 0; Idx < Actors.Num(); Idx++)
+	for (AActor* Actor : Actors)
 	{
-		if (Actors[Idx])
+		if (Actor)
 		{
-			AActor* ParentActor = Actors[Idx]->GetAttachParentActor();
-			if (ParentActor)
-			{
-				ParentMap.Add(Actors[Idx], ParentActor);
-			}
+			CalcAttachDepth(Actor);
+			VisitedActors.Reset();
 		}
 	}
 
-	if (ParentMap.Num())
+	auto DepthSorter = [&DepthMap](AActor* A, AActor* B)
 	{
-		TMap<AActor*, int32> DepthMap;
-		FLevelSortUtils::FDepthSort DepthSorter;
-		DepthSorter.DepthMap = &DepthMap;
-
-		TArray<AActor*> ParentChain;
-		while (ParentMap.Num())
-		{
-			ParentChain.Reset();
-			FLevelSortUtils::FindAndRemoveParentChain(ParentMap, ParentChain);
-
-			// Topmost parent in found parent chain might have its parent already removed
-			// so we need to use it's stored depth as a base depth for whole chain
-			int32 ParentChainStartDepth = 0;
-			if (ParentChain.Num() > 0)
-			{
-				if (int32* StartDepthPtr = DepthMap.Find(ParentChain.Last()))
-				{
-					ParentChainStartDepth = *StartDepthPtr;
-				}
-			}
-
-			for (int32 Idx = 0; Idx < ParentChain.Num(); Idx++)
-			{
-				DepthMap.Add(ParentChain[Idx], ParentChainStartDepth + ParentChain.Num() - Idx - 1);
-			}
-		}
+		const int32 DepthA = A ? DepthMap.FindRef(A) : MAX_int32;
+		const int32 DepthB = B ? DepthMap.FindRef(B) : MAX_int32;
+		return DepthA < DepthB;
+	};
 
-		// Unfortunately TArray.StableSort assumes no null entries in the array
-		// So it forces me to use internal unrestricted version
-		StableSortInternal(Actors.GetData(), Actors.Num(), DepthSorter);
-	}
+	StableSortInternal(Actors.GetData(), Actors.Num(), DepthSorter);
 
 	const float ElapsedTime = (float)(FPlatformTime::Seconds() - StartTime);
 	if (ElapsedTime > 1.0f)
