Merging 9412940, 9574419, and 9575864 from Dev-Anim
Inertial blending for Linked Anim Blueprints

diff --git a/Engine/Source/Editor/AnimGraph/Classes/AnimationGraph.h b/Engine/Source/Editor/AnimGraph/Classes/AnimationGraph.h
--- a/Engine/Source/Editor/AnimGraph/Classes/AnimationGraph.h
+++ b/Engine/Source/Editor/AnimGraph/Classes/AnimationGraph.h
@@ -5,6 +5,7 @@
 #include "CoreMinimal.h"
 #include "UObject/ObjectMacros.h"
 #include "EdGraph/EdGraph.h"
+#include "Animation/AnimClassInterface.h"
 #include "AnimationGraph.generated.h"
 
 class UEdGraphPin;
@@ -19,5 +20,9 @@ class UAnimationGraph : public UEdGraph
 
 	/** Delegate fired when a pin's default value is changed */
 	FOnPinDefaultValueChanged OnPinDefaultValueChanged;
+
+	/** Blending options for this graph (only valid for linked anim graphs with inertialization) */
+	UPROPERTY(EditAnywhere, Category = GraphBlending, meta = (ShowOnlyInnerProperties))
+	FAnimGraphBlendOptions BlendOptions;
 };
 
diff --git a/Engine/Source/Editor/AnimGraph/Private/AnimGraphDetails.cpp b/Engine/Source/Editor/AnimGraph/Private/AnimGraphDetails.cpp
--- a/Engine/Source/Editor/AnimGraph/Private/AnimGraphDetails.cpp
+++ b/Engine/Source/Editor/AnimGraph/Private/AnimGraphDetails.cpp
@@ -16,6 +16,7 @@
 #include "AnimGraphNode_Root.h"
 #include "ScopedTransaction.h"
 #include "SKismetInspector.h"
+#include "AnimationGraph.h"
 #include "AnimationGraphSchema.h"
 #include "Kismet2/BlueprintEditorUtils.h"
 #include "Widgets/Input/SComboButton.h"
@@ -44,6 +45,11 @@ void FAnimGraphDetails::CustomizeDetails(IDetailLayoutBuilder& DetailLayout)
 
 	Graph = CastChecked<UEdGraph>(Objects[0].Get());
 
+	if(IsInterface())
+	{
+		DetailLayout.HideCategory("GraphBlending");
+	}
+
 	bool const bIsStateMachine = !Graph->GetOuter()->IsA(UAnimBlueprint::StaticClass());
 
 	if(Objects.Num() > 1 || bIsStateMachine)
@@ -53,7 +59,6 @@ void FAnimGraphDetails::CustomizeDetails(IDetailLayoutBuilder& DetailLayout)
 		return;
 	}
 
-	TSharedPtr<IAnimationBlueprintEditor> AnimBlueprintEditor = AnimBlueprintEditorPtr.Pin();
 	const bool bIsDefaultGraph = Graph->GetFName() == UEdGraphSchema_K2::GN_AnimGraph;
 
 	if(!Graph->bAllowDeletion && !bIsDefaultGraph)
@@ -266,8 +271,7 @@ void FAnimGraphDetails::CustomizeDetails(IDetailLayoutBuilder& DetailLayout)
 		];
 	}
 
-	bool bIsInterface = AnimBlueprintEditor->GetBlueprintObj()->BlueprintType == BPTYPE_Interface;
-	if(bIsInterface)
+	if(IsInterface())
 	{
 		UAnimationGraphSchema::AutoArrangeInterfaceGraph(*Graph);
 	}
@@ -275,13 +279,10 @@ void FAnimGraphDetails::CustomizeDetails(IDetailLayoutBuilder& DetailLayout)
 
 FReply FAnimGraphDetails::OnAddNewInputPoseClicked()
 {
-	TSharedPtr<IAnimationBlueprintEditor> AnimBlueprintEditor = AnimBlueprintEditorPtr.Pin();
-
 	EK2NewNodeFlags NewNodeOperation = EK2NewNodeFlags::None;
 	FVector2D NewNodePosition(0.0f, 0.0f);
 
-	bool bIsInterface = AnimBlueprintEditor->GetBlueprintObj()->BlueprintType == BPTYPE_Interface;
-	if(!bIsInterface)
+	if(!IsInterface())
 	{
 		NewNodePosition = UAnimationGraphSchema::GetPositionForNewLinkedInputPoseNode(*Graph);
 	}
@@ -373,6 +374,12 @@ TSharedRef< ITableRow > FAnimGraphDetails::MakeGroupViewWidget( TSharedPtr<FText
 		];
 }
 
+bool FAnimGraphDetails::IsInterface() const
+{
+	TSharedPtr<IAnimationBlueprintEditor> AnimBlueprintEditor = AnimBlueprintEditorPtr.Pin();
+	return AnimBlueprintEditor->GetBlueprintObj()->BlueprintType == BPTYPE_Interface;
+}
+
 void FAnimGraphDetails::RefreshGroupSource()
 {
 	TSharedPtr<IAnimationBlueprintEditor> AnimBlueprintEditor = AnimBlueprintEditorPtr.Pin();
diff --git a/Engine/Source/Editor/AnimGraph/Private/AnimGraphDetails.h b/Engine/Source/Editor/AnimGraph/Private/AnimGraphDetails.h
--- a/Engine/Source/Editor/AnimGraph/Private/AnimGraphDetails.h
+++ b/Engine/Source/Editor/AnimGraph/Private/AnimGraphDetails.h
@@ -35,9 +35,6 @@ public:
 	virtual void CustomizeDetails(IDetailLayoutBuilder& DetailLayout) override;
 
 private:
-	/** Helper function to get the root node of this graph */
-	UAnimGraphNode_Root* GetRoot() const;
-
 	/** UI handlers */
 	FReply OnAddNewInputPoseClicked();
 	EVisibility OnGetNewInputPoseTextVisibility(TWeakPtr<SWidget> WeakInputsHeaderWidget) const;
@@ -47,8 +44,8 @@ private:
 	void OnGroupSelectionChanged(TSharedPtr<FText> ProposedSelection, ESelectInfo::Type SelectInfo);
 	TSharedRef<ITableRow> MakeGroupViewWidget(TSharedPtr<FText> Item, const TSharedRef< STableViewBase >& OwnerTable);
 
-	/** Set the layer's group */
-	void SetAnimationGraphLayerGroup(const FText& InGroupName);
+	/** Helper function to check if this is the anim graph details for an anim bp interface */
+	bool IsInterface() const;
 
 	/** Refresh the displayed groups */
 	void RefreshGroupSource();
diff --git a/Engine/Source/Editor/KismetCompiler/Private/AnimBlueprintCompiler.cpp b/Engine/Source/Editor/KismetCompiler/Private/AnimBlueprintCompiler.cpp
--- a/Engine/Source/Editor/KismetCompiler/Private/AnimBlueprintCompiler.cpp
+++ b/Engine/Source/Editor/KismetCompiler/Private/AnimBlueprintCompiler.cpp
@@ -48,6 +48,7 @@
 #include "AnimGraphNode_LinkedAnimGraph.h"
 #include "AnimGraphNode_Slot.h"
 #include "AnimationEditorUtils.h"
+#include "AnimationGraph.h"
 
 #include "AnimBlueprintPostCompileValidation.h" 
 #include "AnimGraphNode_LinkedInputPose.h"
@@ -2238,6 +2239,7 @@ void FAnimBlueprintCompilerContext::CleanAndSanitizeClass(UBlueprintGeneratedCla
 	NewAnimBlueprintClass->LinkedAnimLayerNodeProperties.Empty();
 	NewAnimBlueprintClass->EvaluateGraphExposedInputs.Empty();
 	NewAnimBlueprintClass->GraphAssetPlayerInformation.Empty();
+	NewAnimBlueprintClass->GraphBlendOptions.Empty();
 
 	// Copy over runtime data from the blueprint to the class
 	NewAnimBlueprintClass->TargetSkeleton = AnimBlueprint->TargetSkeleton;
@@ -2268,6 +2270,34 @@ void FAnimBlueprintCompilerContext::FinishCompilingClass(UClass* Class)
 	{
 		AnimBlueprintGeneratedClass->SyncGroupNames.Add(GroupInfo.Name);
 	}
+
+	// Add graph blend options to class if blend values were actually customized
+	auto AddBlendOptions = [AnimBlueprintGeneratedClass](UEdGraph* Graph)
+	{
+		UAnimationGraph* AnimGraph = Cast<UAnimationGraph>(Graph);
+		if (AnimGraph && (AnimGraph->BlendOptions.BlendInTime >= 0.0f || AnimGraph->BlendOptions.BlendOutTime >= 0.0f))
+		{
+			AnimBlueprintGeneratedClass->GraphBlendOptions.Add(AnimGraph->GetFName(), AnimGraph->BlendOptions);
+		}
+	};
+
+
+	for (UEdGraph* Graph : Blueprint->FunctionGraphs)
+	{
+		AddBlendOptions(Graph);
+	}
+
+	for (FBPInterfaceDescription& InterfaceDesc : Blueprint->ImplementedInterfaces)
+	{
+		if (InterfaceDesc.Interface->IsChildOf<UAnimLayerInterface>())
+		{
+			for (UEdGraph* Graph : InterfaceDesc.Graphs)
+			{
+				AddBlendOptions(Graph);
+			}
+		}
+	}
+
 	Super::FinishCompilingClass(Class);
 }
 
diff --git a/Engine/Source/Runtime/Engine/Classes/Animation/AnimBlueprintGeneratedClass.h b/Engine/Source/Runtime/Engine/Classes/Animation/AnimBlueprintGeneratedClass.h
--- a/Engine/Source/Runtime/Engine/Classes/Animation/AnimBlueprintGeneratedClass.h
+++ b/Engine/Source/Runtime/Engine/Classes/Animation/AnimBlueprintGeneratedClass.h
@@ -218,6 +218,10 @@ class ENGINE_API UAnimBlueprintGeneratedClass : public UBlueprintGeneratedClass,
 	UPROPERTY()
 	TMap<FName, FGraphAssetPlayerInformation> GraphAssetPlayerInformation;
 
+	// Per layer graph blending options
+	UPROPERTY()
+	TMap<FName, FAnimGraphBlendOptions> GraphBlendOptions;
+
 public:
 
 	virtual const TArray<FBakedAnimationStateMachine>& GetBakedStateMachines() const override { return BakedStateMachines; }
@@ -236,6 +240,7 @@ public:
 	virtual const TArray<FExposedValueHandler>& GetExposedValueHandlers() const { return EvaluateGraphExposedInputs; }
 	virtual const TArray<FAnimBlueprintFunction>& GetAnimBlueprintFunctions() const override { return AnimBlueprintFunctions; }
 	virtual const TMap<FName, FGraphAssetPlayerInformation>& GetGraphAssetPlayerInformation() const override { return GraphAssetPlayerInformation; }
+	virtual const TMap<FName, FAnimGraphBlendOptions>& GetGraphBlendOptions() const override { return GraphBlendOptions; }
 	
 public:
 #if WITH_EDITORONLY_DATA
diff --git a/Engine/Source/Runtime/Engine/Classes/Animation/AnimClassData.h b/Engine/Source/Runtime/Engine/Classes/Animation/AnimClassData.h
--- a/Engine/Source/Runtime/Engine/Classes/Animation/AnimClassData.h
+++ b/Engine/Source/Runtime/Engine/Classes/Animation/AnimClassData.h
@@ -78,6 +78,10 @@ public:
 	UPROPERTY()
 	TArray<FExposedValueHandler> EvaluateGraphExposedInputs;
 
+	// Per layer graph blending options
+	UPROPERTY()
+	TMap<FName, FAnimGraphBlendOptions> GraphBlendOptions;
+
 public:
 	// IAnimClassInterface interface
 	virtual const TArray<FBakedAnimationStateMachine>& GetBakedStateMachines() const override { return BakedStateMachines; }
@@ -96,7 +100,7 @@ public:
 	virtual int32 GetSyncGroupIndex(FName SyncGroupName) const override { return SyncGroupNames.IndexOfByKey(SyncGroupName); }
 	virtual const TArray<FExposedValueHandler>& GetExposedValueHandlers() const { return EvaluateGraphExposedInputs; }
 	virtual const TMap<FName, FGraphAssetPlayerInformation>& GetGraphAssetPlayerInformation() const { return GraphNameAssetPlayers; }
-
+	virtual const TMap<FName, FAnimGraphBlendOptions>& GetGraphBlendOptions() const { return GraphBlendOptions; }
 #if WITH_EDITOR
 	void CopyFrom(IAnimClassInterface* AnimClass)
 	{
@@ -116,6 +120,7 @@ public:
 		SyncGroupNames = AnimClass->GetSyncGroupNames();
 		EvaluateGraphExposedInputs = AnimClass->GetExposedValueHandlers();
 		GraphNameAssetPlayers = AnimClass->GetGraphAssetPlayerInformation();
+		GraphBlendOptions = AnimClass->GetGraphBlendOptions();
 	}
 #endif // WITH_EDITOR
 };
diff --git a/Engine/Source/Runtime/Engine/Classes/Animation/AnimClassInterface.h b/Engine/Source/Runtime/Engine/Classes/Animation/AnimClassInterface.h
--- a/Engine/Source/Runtime/Engine/Classes/Animation/AnimClassInterface.h
+++ b/Engine/Source/Runtime/Engine/Classes/Animation/AnimClassInterface.h
@@ -105,6 +105,26 @@ struct FGraphAssetPlayerInformation
 	TArray<int32> PlayerNodeIndices;
 };
 
+/** Blending options for anim graph layers */
+USTRUCT()
+struct FAnimGraphBlendOptions
+{
+	GENERATED_USTRUCT_BODY()
+
+	/** Time to blend this graph in. Only valid for linked anim graphs. Must use inertialization node to blend. */
+	UPROPERTY(EditAnywhere, Category = GraphBlending)
+	float BlendInTime;       
+
+	/** Time to blend this graph out. Only valid for linked anim graphs. Must use inertialization node to blend. */
+	UPROPERTY(EditAnywhere, Category = GraphBlending)
+	float BlendOutTime;
+
+	FAnimGraphBlendOptions()
+		: BlendInTime(-1.0f)
+		, BlendOutTime(-1.0f)
+	{}
+};
+
 
 UINTERFACE()
 class ENGINE_API UAnimClassInterface : public UInterface
@@ -134,6 +154,7 @@ public:
 	virtual const TMap<FName, FCachedPoseIndices>& GetOrderedSavedPoseNodeIndicesMap() const = 0;
 	virtual const TArray<FAnimBlueprintFunction>& GetAnimBlueprintFunctions() const = 0;
 	virtual const TMap<FName, FGraphAssetPlayerInformation>& GetGraphAssetPlayerInformation() const = 0;
+	virtual const TMap<FName, FAnimGraphBlendOptions>& GetGraphBlendOptions() const = 0;
 	virtual USkeleton* GetTargetSkeleton() const = 0;
 
 	virtual int32 GetSyncGroupIndex(FName SyncGroupName) const = 0;
diff --git a/Engine/Source/Runtime/Engine/Classes/Animation/AnimNode_LinkedAnimGraph.h b/Engine/Source/Runtime/Engine/Classes/Animation/AnimNode_LinkedAnimGraph.h
--- a/Engine/Source/Runtime/Engine/Classes/Animation/AnimNode_LinkedAnimGraph.h
+++ b/Engine/Source/Runtime/Engine/Classes/Animation/AnimNode_LinkedAnimGraph.h
@@ -71,10 +71,14 @@ public:
 	void CacheBonesSubGraph_AnyThread(const FAnimationCacheBonesContext& Context);
 
 protected:
+
 	virtual void OnInitializeAnimInstance(const FAnimInstanceProxy* InProxy, const UAnimInstance* InAnimInstance) override;
 	virtual bool NeedsOnInitializeAnimInstance() const override { return true; }
 	// End of FAnimNode_Base interface
 
+	// Inertial blending duration to request next update
+	float PendingBlendDuration;
+
 	// Re-create the linked instances for this node
 	void ReinitializeLinkedAnimInstance(const UAnimInstance* InOwningAnimInstance, UAnimInstance* InNewAnimInstance = nullptr);
 
@@ -96,6 +100,9 @@ protected:
 	/** Helper function for finding function inputs when linking/unlinking */
 	int32 FindFunctionInputIndex(const FAnimBlueprintFunction& AnimBlueprintFunction, const FName& InInputName);
 
+	/** Request a blend when the active instance changes */
+	void RequestBlend(const IAnimClassInterface* PriorAnimBPClass, const IAnimClassInterface* NewAnimBPClass);
+
 	friend class UAnimInstance;
 };
 
diff --git a/Engine/Source/Runtime/Engine/Private/Animation/AnimNode_LinkedAnimGraph.cpp b/Engine/Source/Runtime/Engine/Private/Animation/AnimNode_LinkedAnimGraph.cpp
--- a/Engine/Source/Runtime/Engine/Private/Animation/AnimNode_LinkedAnimGraph.cpp
+++ b/Engine/Source/Runtime/Engine/Private/Animation/AnimNode_LinkedAnimGraph.cpp
@@ -3,12 +3,36 @@
 #include "Animation/AnimNode_LinkedAnimGraph.h"
 #include "Animation/AnimClassInterface.h"
 #include "Animation/AnimInstanceProxy.h"
+#include "Animation/AnimNode_Inertialization.h"
 #include "Animation/AnimNode_LinkedInputPose.h"
 #include "Animation/AnimNode_Root.h"
 
+static float GetBlendDuration(const IAnimClassInterface* PriorAnimBPClass, const IAnimClassInterface* NewAnimBPClass, FName Layer)
+{
+	const FAnimGraphBlendOptions* PriorBlendOptions = PriorAnimBPClass ? PriorAnimBPClass->GetGraphBlendOptions().Find(Layer) : nullptr;
+	const FAnimGraphBlendOptions* NewBlendOptions = NewAnimBPClass ? NewAnimBPClass->GetGraphBlendOptions().Find(Layer) : nullptr;
+
+	float BlendOutTime = PriorBlendOptions ? PriorBlendOptions->BlendOutTime : -1.0f;
+	float BlendInTime = NewBlendOptions ? NewBlendOptions->BlendInTime : -1.0f;
+
+	if (BlendInTime < 0.0f)
+	{
+		return BlendOutTime;
+	}
+
+	if (BlendOutTime < 0.0f)
+	{
+		return BlendInTime;
+	}
+
+	return FMath::Min(BlendInTime, BlendOutTime);
+}
+
 FAnimNode_LinkedAnimGraph::FAnimNode_LinkedAnimGraph()
 	: InstanceClass(nullptr)
 	, Tag(NAME_None)
+	, LinkedRoot(nullptr)
+	, PendingBlendDuration(-1.0f)
 {
 }
 
@@ -87,6 +111,18 @@ void FAnimNode_LinkedAnimGraph::Update_AnyThread(const FAnimationUpdateContext&
 		// subsequent nodes get properly updated
 		InputPoses[0].Update(InContext);
 	}
+
+	// Consume pending inertial blend request
+	if (PendingBlendDuration >= 0.0f)
+	{
+		FAnimNode_Inertialization* InertializationNode = InContext.GetAncestor<FAnimNode_Inertialization>();
+		if (InertializationNode)
+		{
+			InertializationNode->Request(PendingBlendDuration);
+		}
+
+		PendingBlendDuration = -1.0f;
+	}
 }
 
 void FAnimNode_LinkedAnimGraph::Evaluate_AnyThread(FPoseContext& Output)
@@ -173,6 +209,8 @@ void FAnimNode_LinkedAnimGraph::ReinitializeLinkedAnimInstance(const UAnimInstan
 {
 	UAnimInstance* InstanceToRun = GetTargetInstance<UAnimInstance>();
 
+	IAnimClassInterface* PriorAnimBPClass = InstanceToRun ? IAnimClassInterface::GetFromClass(InstanceToRun->GetClass()) : nullptr;
+
 	if(*InstanceClass || InNewAnimInstance)
 	{
 		USkeletalMeshComponent* MeshComp = InOwningAnimInstance->GetSkelMeshComponent();
@@ -207,6 +245,10 @@ void FAnimNode_LinkedAnimGraph::ReinitializeLinkedAnimInstance(const UAnimInstan
 		}
 
 		InitializeProperties(InOwningAnimInstance, InstanceToRun->GetClass());
+
+		IAnimClassInterface* NewAnimBPClass = IAnimClassInterface::GetFromClass(InstanceToRun->GetClass());
+
+		RequestBlend(PriorAnimBPClass, NewAnimBPClass);
 	}
 	else if(InstanceToRun)
 	{
@@ -356,3 +398,8 @@ int32 FAnimNode_LinkedAnimGraph::FindFunctionInputIndex(const FAnimBlueprintFunc
 
 	return INDEX_NONE;
 }
+
+void FAnimNode_LinkedAnimGraph::RequestBlend(const IAnimClassInterface* PriorAnimBPClass, const IAnimClassInterface* NewAnimBPClass)
+{
+	PendingBlendDuration = GetBlendDuration(PriorAnimBPClass, NewAnimBPClass, GetDynamicLinkFunctionName());
+}
diff --git a/Engine/Source/Runtime/Engine/Private/Animation/AnimNode_LinkedAnimLayer.cpp b/Engine/Source/Runtime/Engine/Private/Animation/AnimNode_LinkedAnimLayer.cpp
--- a/Engine/Source/Runtime/Engine/Private/Animation/AnimNode_LinkedAnimLayer.cpp
+++ b/Engine/Source/Runtime/Engine/Private/Animation/AnimNode_LinkedAnimLayer.cpp
@@ -1,6 +1,7 @@
 // Copyright 1998-2019 Epic Games, Inc. All Rights Reserved.
 
 #include "Animation/AnimNode_LinkedAnimLayer.h"
+#include "Animation/AnimClassInterface.h"
 
 FName FAnimNode_LinkedAnimLayer::GetDynamicLinkFunctionName() const
 {
@@ -33,6 +34,8 @@ void FAnimNode_LinkedAnimLayer::InitializeSelfLayer(const UAnimInstance* SelfAni
 {
 	UAnimInstance* CurrentTarget = GetTargetInstance<UAnimInstance>();
 
+	IAnimClassInterface* PriorAnimBPClass = CurrentTarget ? IAnimClassInterface::GetFromClass(CurrentTarget->GetClass()) : nullptr;
+
 	USkeletalMeshComponent* MeshComp = SelfAnimInstance->GetSkelMeshComponent();
 	check(MeshComp);
 
@@ -54,7 +57,12 @@ void FAnimNode_LinkedAnimLayer::InitializeSelfLayer(const UAnimInstance* SelfAni
 	// Link before we call InitializeAnimation() so we propgate the call to linked input poses
 	DynamicLink(const_cast<UAnimInstance*>(SelfAnimInstance));
 
-	InitializeProperties(SelfAnimInstance, SelfAnimInstance->GetClass());
+	UClass* SelfClass = SelfAnimInstance->GetClass();
+	InitializeProperties(SelfAnimInstance, SelfClass);
+
+	IAnimClassInterface* NewAnimBPClass = IAnimClassInterface::GetFromClass(SelfClass);
+
+	RequestBlend(PriorAnimBPClass, NewAnimBPClass);
 }
 
 void FAnimNode_LinkedAnimLayer::SetLinkedLayerInstance(const UAnimInstance* InOwningAnimInstance, UAnimInstance* InNewLinkedInstance)
