Set default log level to Display for LogAudio & move a few items around to avoid spam/warning incorrectly

diff --git a/Engine/Source/Runtime/Engine/Private/AudioStreamingCache.cpp b/Engine/Source/Runtime/Engine/Private/AudioStreamingCache.cpp
--- a/Engine/Source/Runtime/Engine/Private/AudioStreamingCache.cpp
+++ b/Engine/Source/Runtime/Engine/Private/AudioStreamingCache.cpp
@@ -550,7 +550,7 @@ void FAudioChunkCache::BlockForAllPendingLoads() const
 		if (bLoadInProgress)
 		{
 			float TimeSinceStarted = FPlatformTime::Seconds() - TimeStarted;
-			UE_LOG(LogAudio, Display, TEXT("Waited %f seconds for async audio chunk loads."), TimeSinceStarted);
+			UE_LOG(LogAudio, Log, TEXT("Waited %f seconds for async audio chunk loads."), TimeSinceStarted);
 			FPlatformProcess::Sleep(0.0f);
 		}
 
diff --git a/Engine/Source/Runtime/Engine/Private/AudioThread.cpp b/Engine/Source/Runtime/Engine/Private/AudioThread.cpp
--- a/Engine/Source/Runtime/Engine/Private/AudioThread.cpp
+++ b/Engine/Source/Runtime/Engine/Private/AudioThread.cpp
@@ -557,11 +557,11 @@ void FAudioCommandFence::Wait(bool bProcessGameThreadTasks) const
 					double LongestTaskTimeMs;
 					FAudioThread::GetLongestTaskInfo(LongestTask, LongestTaskTimeMs);
 
-					UE_LOG(LogAudio, Warning, TEXT("Waited %.2f ms for audio thread. (Current Task: %s, Longest task: %s %.2f ms)"), ThisTime * 1000.0f, *CurrentTask, *LongestTask, LongestTaskTimeMs);
+					UE_LOG(LogAudio, Display, TEXT("Waited %.2f ms for audio thread. (Current Task: %s, Longest task: %s %.2f ms)"), ThisTime * 1000.0f, *CurrentTask, *LongestTask, LongestTaskTimeMs);
 				}
 				else
 				{
-					UE_LOG(LogAudio, Warning,  TEXT("Waited %f ms for audio thread."), ThisTime * 1000.0f);
+					UE_LOG(LogAudio, Display,  TEXT("Waited %f ms for audio thread."), ThisTime * 1000.0f);
 				}
 			}
 		} while (!bDone);
diff --git a/Engine/Source/Runtime/Engine/Private/SampleBufferIO.cpp b/Engine/Source/Runtime/Engine/Private/SampleBufferIO.cpp
--- a/Engine/Source/Runtime/Engine/Private/SampleBufferIO.cpp
+++ b/Engine/Source/Runtime/Engine/Private/SampleBufferIO.cpp
@@ -452,7 +452,7 @@ namespace Audio
 		SerializedWavData.Reset();
 		SerializeWaveFile(SerializedWavData, (const uint8*)CurrentBuffer.GetData(), CurrentBuffer.GetNumSamples() * sizeof(int16), CurrentBuffer.GetNumChannels(), CurrentBuffer.GetSampleRate());
 
-		UE_LOG(LogAudio, Display, TEXT("Serializing %d sample file (%d bytes) to sound asset at %s"), CurrentBuffer.GetNumSamples(), SerializedWavData.Num(), *AbsoluteFilePath);
+		UE_LOG(LogAudio, Verbose, TEXT("Serializing %d sample file (%d bytes) to sound asset at %s"), CurrentBuffer.GetNumSamples(), SerializedWavData.Num(), *AbsoluteFilePath);
 
 		// TODO: Check to see if we need to call USoundWave::FreeResources here.
 
@@ -486,7 +486,7 @@ namespace Audio
 		// Clamp buffer to prevent wraparound when serializing:
 		CurrentBuffer.Clamp(0.9999f);
 		SerializeWaveFile(SerializedWavData, (const uint8*) CurrentBuffer.GetData(), CurrentBuffer.GetNumSamples() * sizeof(int16), CurrentBuffer.GetNumChannels(), CurrentBuffer.GetSampleRate());
-		UE_LOG(LogAudio, Display, TEXT("Serializing %d sample file (%d bytes) to %s"), CurrentBuffer.GetNumSamples(), SerializedWavData.Num(), *AbsoluteFilePath);
+		UE_LOG(LogAudio, Verbose, TEXT("Serializing %d sample file (%d bytes) to %s"), CurrentBuffer.GetNumSamples(), SerializedWavData.Num(), *AbsoluteFilePath);
 		if (SerializedWavData.Num() == 0)
 		{
 			UE_LOG(LogAudio, Error, TEXT("Wave serialize operation failed: failiure in SerializeWaveFile"));
@@ -502,7 +502,7 @@ namespace Audio
 			if (FileHandle)
 			{
 				int32 NumChunks = FMath::CeilToInt((float)SerializedWavData.Num() / ChunkSize);
-				UE_LOG(LogAudio, Display, TEXT("Writing wav file in %d chunks..."), NumChunks);
+				UE_LOG(LogAudio, Verbose, TEXT("Writing wav file in %d chunks..."), NumChunks);
 				for (int32 ChunkIndex = 0; ChunkIndex < NumChunks; ChunkIndex++)
 				{
 					int32 BufferIndex = ChunkIndex * ChunkSize;
@@ -536,7 +536,7 @@ namespace Audio
 				return;
 			}
 		}
-		UE_LOG(LogAudio, Display, TEXT("Succeeded in writing wav file."));
+		UE_LOG(LogAudio, Verbose, TEXT("Succeeded in writing wav file."));
 		CurrentState = ESoundWavePCMWriterState::Suceeded;
 	}
 
diff --git a/Engine/Source/Runtime/Engine/Public/Audio.h b/Engine/Source/Runtime/Engine/Public/Audio.h
--- a/Engine/Source/Runtime/Engine/Public/Audio.h
+++ b/Engine/Source/Runtime/Engine/Public/Audio.h
@@ -17,7 +17,7 @@
 #include "Sound/SoundSourceBusSend.h"
 #include "IAudioExtensionPlugin.h"
 
-ENGINE_API DECLARE_LOG_CATEGORY_EXTERN(LogAudio, Warning, All);
+ENGINE_API DECLARE_LOG_CATEGORY_EXTERN(LogAudio, Display, All);
 
 // Special log category used for temporary programmer debugging code of audio
 ENGINE_API DECLARE_LOG_CATEGORY_EXTERN(LogAudioDebug, Display, All);
