#ue4 - Fix debug vis from "Net.RepMovement.DrawDebug" lasting forever.

diff --git a/Engine/Source/Runtime/Engine/Private/ActorReplication.cpp b/Engine/Source/Runtime/Engine/Private/ActorReplication.cpp
--- a/Engine/Source/Runtime/Engine/Private/ActorReplication.cpp
+++ b/Engine/Source/Runtime/Engine/Private/ActorReplication.cpp
@@ -183,7 +183,7 @@ void AActor::OnRep_ReplicatedMovement()
 #if !(UE_BUILD_SHIPPING || UE_BUILD_TEST)
 		if (CVarDrawDebugRepMovement->GetInt() > 0)
 		{
-			DrawDebugCapsule(GetWorld(), LocalRepMovement.Location, GetSimpleCollisionHalfHeight(), GetSimpleCollisionRadius(), LocalRepMovement.Rotation.Quaternion(), FColor(100, 255, 100), true, 1.f);
+			DrawDebugCapsule(GetWorld(), LocalRepMovement.Location, GetSimpleCollisionHalfHeight(), GetSimpleCollisionRadius(), LocalRepMovement.Rotation.Quaternion(), FColor(100, 255, 100), false, 1.f);
 		}
 #endif
 
