#ue4 - Call UpdateCharacterStateBeforeMovement() and UpdateCharacterStateAfterMovement() for simulated proxies during simulation. However the base implementations avoid the crouch state changes in those functions because those are replicated from the server. This is more for overrides to be able to run custom behavior to match the server simulation.

diff --git a/Engine/Source/Runtime/Engine/Private/Components/CharacterMovementComponent.cpp b/Engine/Source/Runtime/Engine/Private/Components/CharacterMovementComponent.cpp
--- a/Engine/Source/Runtime/Engine/Private/Components/CharacterMovementComponent.cpp
+++ b/Engine/Source/Runtime/Engine/Private/Components/CharacterMovementComponent.cpp
@@ -1775,6 +1775,8 @@ void UCharacterMovementComponent::SimulateMovement(float DeltaSeconds)
 			}
 		}
 
+		UpdateCharacterStateBeforeMovement(DeltaSeconds);
+
 		if (MovementMode != MOVE_None)
 		{
 			//TODO: Also ApplyAccumulatedForces()?
@@ -1873,6 +1875,8 @@ void UCharacterMovementComponent::SimulateMovement(float DeltaSeconds)
 			UE_LOG(LogCharacterMovement, Verbose, TEXT("Proxy %s SKIPPING simulate movement"), *GetNameSafe(CharacterOwner));
 		}
 
+		UpdateCharacterStateAfterMovement(DeltaSeconds);
+
 		// consume path following requested velocity
 		bHasRequestedVelocity = false;
 
@@ -2824,24 +2828,32 @@ void UCharacterMovementComponent::UnCrouch(bool bClientSimulation)
 
 void UCharacterMovementComponent::UpdateCharacterStateBeforeMovement(float DeltaSeconds)
 {
-	// Check for a change in crouch state. Players toggle crouch by changing bWantsToCrouch.
-	const bool bIsCrouching = IsCrouching();
-	if (bIsCrouching && (!bWantsToCrouch || !CanCrouchInCurrentState()))
-	{
-		UnCrouch(false);
-	}
-	else if (!bIsCrouching && bWantsToCrouch && CanCrouchInCurrentState())
+	// Proxies get replicated crouch state.
+	if (CharacterOwner->Role != ROLE_SimulatedProxy)
 	{
-		Crouch(false);
+		// Check for a change in crouch state. Players toggle crouch by changing bWantsToCrouch.
+		const bool bIsCrouching = IsCrouching();
+		if (bIsCrouching && (!bWantsToCrouch || !CanCrouchInCurrentState()))
+		{
+			UnCrouch(false);
+		}
+		else if (!bIsCrouching && bWantsToCrouch && CanCrouchInCurrentState())
+		{
+			Crouch(false);
+		}
 	}
 }
 
 void UCharacterMovementComponent::UpdateCharacterStateAfterMovement(float DeltaSeconds)
 {
-	// Uncrouch if no longer allowed to be crouched
-	if (IsCrouching() && !CanCrouchInCurrentState())
+	// Proxies get replicated crouch state.
+	if (CharacterOwner->Role != ROLE_SimulatedProxy)
 	{
-		UnCrouch(false);
+		// Uncrouch if no longer allowed to be crouched
+		if (IsCrouching() && !CanCrouchInCurrentState())
+		{
+			UnCrouch(false);
+		}
 	}
 }
 
