Don't track retirements for 100% reliable connections

diff --git a/Engine/Source/Runtime/Engine/Private/DataReplication.cpp b/Engine/Source/Runtime/Engine/Private/DataReplication.cpp
--- a/Engine/Source/Runtime/Engine/Private/DataReplication.cpp
+++ b/Engine/Source/Runtime/Engine/Private/DataReplication.cpp
@@ -362,7 +362,9 @@ void FObjectReplicator::InitRecentProperties(uint8* Source)
 		return;
 	}
 
-	bSupportsFastArrayDelta = !!GSupportsFastArrayDelta;
+	const bool bIsReplay = Connection->InternalAck;
+
+	bSupportsFastArrayDelta = !!GSupportsFastArrayDelta && !bIsReplay;
 
 	if (FSendingRepState* SendingRepState = RepState->GetSendingRepState())
 	{
@@ -373,9 +375,7 @@ void FObjectReplicator::InitRecentProperties(uint8* Source)
 		const uint16 NumLifetimeCustomDeltaProperties = FNetSerializeCB::GetNumLifetimeCustomDeltaProperties(LocalRepLayout);
 		SendingRepState->RecentCustomDeltaState.SetNum(NumLifetimeCustomDeltaProperties);
 
-		const bool bIsRecordingReplay = Connection->InternalAck;
-
-		if (bIsRecordingReplay)
+		if (bIsReplay)
 		{
 			SendingRepState->CDOCustomDeltaState.SetNum(NumLifetimeCustomDeltaProperties);
 			SendingRepState->CheckpointCustomDeltaState.SetNum(NumLifetimeCustomDeltaProperties);
@@ -391,7 +391,7 @@ void FObjectReplicator::InitRecentProperties(uint8* Source)
 
 			SendCustomDeltaProperty(UseObject, CustomDeltaProperty, DeltaState, NewState, OldState);
 
-			if (bIsRecordingReplay)
+			if (bIsReplay)
 			{
 				// Store the initial delta state in case we need it for when we're asked to resend all data since channel was first opened (bResendAllDataSinceOpen)
 				SendingRepState->CDOCustomDeltaState[CustomDeltaProperty] = NewState;
@@ -575,12 +575,16 @@ void FObjectReplicator::StartReplicating(class UActorChannel * InActorChannel)
 
 	if (ConnectionNetDriver->IsServer() || ConnectionNetDriver->MaySendProperties())
 	{
-		if (FSendingRepState* SendingRepState = RepState.IsValid() ? RepState->GetSendingRepState() : nullptr)
+		// We don't need to handle retirement if our connection is reliable.
+		if (!Connection->InternalAck)
 		{
-			// Allocate retirement list.
-			// SetNum now constructs, so this is safe
+			if (FSendingRepState * SendingRepState = RepState.IsValid() ? RepState->GetSendingRepState() : nullptr)
+			{
+				// Allocate retirement list.
+				// SetNum now constructs, so this is safe
 
-			SendingRepState->Retirement.SetNum(ObjectClass->ClassReps.Num());
+				SendingRepState->Retirement.SetNum(ObjectClass->ClassReps.Num());
+			}
 		}
 
 		const UWorld* const World = ConnectionNetDriver->GetWorld();
@@ -1455,17 +1459,22 @@ void FObjectReplicator::ReplicateCustomDeltaProperties( FNetBitWriter & Bunch, F
 			continue;
 		}
 
-		// Get info.
-		FPropertyRetirement& Retire = SendingRepState->Retirement[CustomDeltaProperty];
+		FPropertyRetirement** LastNext = nullptr;
 
-		// Update Retirement records with this new state so we can handle packet drops.
-		// LastNext will be pointer to the last "Next" pointer in the list (so pointer to a pointer)
-		FPropertyRetirement** LastNext = UpdateAckedRetirements(Retire, Connection->OutAckPacketId, Object);
+		if (!Connection->InternalAck)
+		{
+			// Get info.
+			FPropertyRetirement& Retire = SendingRepState->Retirement[CustomDeltaProperty];
+
+			// Update Retirement records with this new state so we can handle packet drops.
+			// LastNext will be pointer to the last "Next" pointer in the list (so pointer to a pointer)
+			LastNext = UpdateAckedRetirements(Retire, Connection->OutAckPacketId, Object);
 
-		check(LastNext != nullptr);
-		check(*LastNext == nullptr);
+			check(LastNext != nullptr);
+			check(*LastNext == nullptr);
 
-		ValidateRetirementHistory(Retire, Object);
+			ValidateRetirementHistory(Retire, Object);
+		}
 
 		//-----------------------------------------
 		//	Do delta serialization on dynamic properties
@@ -1477,10 +1486,13 @@ void FObjectReplicator::ReplicateCustomDeltaProperties( FNetBitWriter & Bunch, F
 			continue;
 		}
 
-		*LastNext = new FPropertyRetirement();
+		if (!Connection->InternalAck)
+		{
+			*LastNext = new FPropertyRetirement();
 
-		// Remember what the old state was at this point in time.  If we get a nak, we will need to revert back to this.
-		(*LastNext)->DynamicState = OldState;		
+			// Remember what the old state was at this point in time.  If we get a nak, we will need to revert back to this.
+			(*LastNext)->DynamicState = OldState;
+		}
 
 		// Save NewState into the RecentCustomDeltaState array (old state is a reference into our RecentCustomDeltaState map)
 		OldState = NewState; 
