Change Network Granular Logging to use scopes and output in a cleaner format. Add more scopes.

diff --git a/Engine/Source/Runtime/Engine/Private/DataChannel.cpp b/Engine/Source/Runtime/Engine/Private/DataChannel.cpp
--- a/Engine/Source/Runtime/Engine/Private/DataChannel.cpp
+++ b/Engine/Source/Runtime/Engine/Private/DataChannel.cpp
@@ -250,23 +250,25 @@ void UChannel::BeginDestroy()
 
 void UChannel::Serialize(FArchive& Ar)
 {
-	Super::Serialize(Ar);
+	GRANULAR_NETWORK_MEMORY_TRACKING_INIT(Ar, "UChannel::Serialize");
+
+	GRANULAR_NETWORK_MEMORY_TRACKING_TRACK("Super", Super::Serialize(Ar));
 
 	if (Ar.IsCountingMemory())
 	{
 		if (InRec)
 		{
-			InRec->CountMemory(Ar);
+			GRANULAR_NETWORK_MEMORY_TRACKING_TRACK("InRec", InRec->CountMemory(Ar));
 		}
 
 		if (OutRec)
 		{
-			OutRec->CountMemory(Ar);
+			GRANULAR_NETWORK_MEMORY_TRACKING_TRACK("OutRec", OutRec->CountMemory(Ar));
 		}
 
 		if (InPartialBunch)
 		{
-			InPartialBunch->CountMemory(Ar);
+			GRANULAR_NETWORK_MEMORY_TRACKING_TRACK("InPartialBunch", InPartialBunch->CountMemory(Ar));
 		}
 	}
 }
@@ -3107,12 +3109,12 @@ void UActorChannel::AddReferencedObjects(UObject* InThis, FReferenceCollector& C
 
 void UActorChannel::Serialize(FArchive& Ar)
 {
-	Super::Serialize(Ar);
+	GRANULAR_NETWORK_MEMORY_TRACKING_INIT(Ar, "UActorChannel::Serialize");
+
+	GRANULAR_NETWORK_MEMORY_TRACKING_TRACK("UActorChannel::Super", Super::Serialize(Ar));
 
 	if (Ar.IsCountingMemory())
 	{
-		GRANULAR_NETWORK_MEMORY_TRACKING_INIT(Ar, "UActorChannel::Serialize");
-
 		GRANULAR_NETWORK_MEMORY_TRACKING_TRACK("ReplicationMap",
 			ReplicationMap.CountBytes(Ar);
 
diff --git a/Engine/Source/Runtime/Engine/Private/DemoNetDriver.cpp b/Engine/Source/Runtime/Engine/Private/DemoNetDriver.cpp
--- a/Engine/Source/Runtime/Engine/Private/DemoNetDriver.cpp
+++ b/Engine/Source/Runtime/Engine/Private/DemoNetDriver.cpp
@@ -39,6 +39,7 @@
 #include "Stats/Stats2.h"
 #include "Engine/ChildConnection.h"
 #include "Net/ReplayPlaylistTracker.h"
+#include "Net/NetworkGranularMemoryLogging.h"
 
 DEFINE_LOG_CATEGORY( LogDemo );
 
@@ -6566,7 +6567,9 @@ void UDemoPendingNetGame::LoadMapCompleted(UEngine* Engine, FWorldContext& Conte
 
 void UDemoNetDriver::Serialize(FArchive& Ar)
 {
-	Super::Serialize(Ar);
+	GRANULAR_NETWORK_MEMORY_TRACKING_INIT(Ar, "UDemoNetDriver::Serialize");
+
+	GRANULAR_NETWORK_MEMORY_TRACKING_TRACK("Super", Super::Serialize(Ar));
 
 	if (Ar.IsCountingMemory())
 	{
@@ -6576,98 +6579,118 @@ void UDemoNetDriver::Serialize(FArchive& Ar)
 		//		QueuedReplayTasks.
 		//		DemoURL
 
-		DeletedNetStartupActors.CountBytes(Ar);
-
-		for (FString& ActorString : DeletedNetStartupActors)
-		{
-			Ar << ActorString;
-		}
-
-		DeletedNetStartupActorGUIDs.CountBytes(Ar);
-
-		// The map for RollbackNetStartupActors may have already been serialized,
-		// However, that won't capture non-property members or properly count them.
-		for (const auto& RollbackNetStartupActorPair : RollbackNetStartupActors)
-		{
-			RollbackNetStartupActorPair.Value.CountBytes(Ar);
-		}
 
+		GRANULAR_NETWORK_MEMORY_TRACKING_TRACK("DeletedNetStartupActors",
+			DeletedNetStartupActors.CountBytes(Ar);
+			for (FString& ActorString : DeletedNetStartupActors)
+			{
+				Ar << ActorString;
+			}
+		);
 
-		ExternalDataToObjectMap.CountBytes(Ar);
+		GRANULAR_NETWORK_MEMORY_TRACKING_TRACK("DeletedNetStartupActorGUIDs", DeletedNetStartupActorGUIDs.CountBytes(Ar));
 
-		for (const auto& ExternalDataToObjectPair : ExternalDataToObjectMap)
-		{
-			ExternalDataToObjectPair.Value.CountBytes(Ar);
-		}
+		GRANULAR_NETWORK_MEMORY_TRACKING_TRACK("RollbackNetStartupActorsValues",
+			// The map for RollbackNetStartupActors may have already been serialized,
+			// However, that won't capture non-property members or properly count them.
+			for (const auto& RollbackNetStartupActorPair : RollbackNetStartupActors)
+			{
+				RollbackNetStartupActorPair.Value.CountBytes(Ar);
+			}
+		);
 
-		PlaybackPackets.CountBytes(Ar);
 
-		for (const FPlaybackPacket& Packet : PlaybackPackets)
-		{
-			Packet.CountBytes(Ar);
-		}
+		GRANULAR_NETWORK_MEMORY_TRACKING_TRACK("ExternalDataToObjectMap",
+			ExternalDataToObjectMap.CountBytes(Ar);
+			for (const auto& ExternalDataToObjectPair : ExternalDataToObjectMap)
+			{
+				ExternalDataToObjectPair.Value.CountBytes(Ar);
+			}
+		);
 
-		UniqueStreamingLevels.CountBytes(Ar);
-		NewStreamingLevelsThisFrame.CountBytes(Ar);
-		NonQueuedGUIDsForScrubbing.CountBytes(Ar);
-		QueuedReplayTasks.CountBytes(Ar);
-		
-		Ar << DemoSessionID;
+		GRANULAR_NETWORK_MEMORY_TRACKING_TRACK("PlaybackPackets",
+			PlaybackPackets.CountBytes(Ar);
+			for (const FPlaybackPacket& Packet : PlaybackPackets)
+			{
+				Packet.CountBytes(Ar);
+			}
+		);
+
+		GRANULAR_NETWORK_MEMORY_TRACKING_TRACK("UniqueStreamingLevels", UniqueStreamingLevels.CountBytes(Ar));
+		GRANULAR_NETWORK_MEMORY_TRACKING_TRACK("NewStreamingLevelsThisFrame", NewStreamingLevelsThisFrame.CountBytes(Ar));
+		GRANULAR_NETWORK_MEMORY_TRACKING_TRACK("NonQueuedGUIDsForScrubbing", NonQueuedGUIDsForScrubbing.CountBytes(Ar));
+		GRANULAR_NETWORK_MEMORY_TRACKING_TRACK("QueuedReplayTasks", QueuedReplayTasks.CountBytes(Ar));
+		GRANULAR_NETWORK_MEMORY_TRACKING_TRACK("DemoSessionID", DemoSessionID.CountBytes(Ar));
+		GRANULAR_NETWORK_MEMORY_TRACKING_TRACK("PlaybackDemoHeader", PlaybackDemoHeader.CountBytes(Ar));
+		GRANULAR_NETWORK_MEMORY_TRACKING_TRACK("PrioritizedActors", PrioritizedActors.CountBytes(Ar));
+
+		GRANULAR_NETWORK_MEMORY_TRACKING_TRACK("LevelNamesAndTimes",
+			LevelNamesAndTimes.CountBytes(Ar);
+			for (const FLevelNameAndTime& LevelNameAndTime : LevelNamesAndTimes)
+			{
+				LevelNameAndTime.CountBytes(Ar);
+			}
+		);
 
-		PlaybackDemoHeader.CountBytes(Ar);
+		GRANULAR_NETWORK_MEMORY_TRACKING_TRACK("LevelInternals", LevelIntervals.CountBytes(Ar));
+		GRANULAR_NETWORK_MEMORY_TRACKING_TRACK("TrackedRewindActorsByGUID", TrackedRewindActorsByGUID.CountBytes(Ar));
 
-		PrioritizedActors.CountBytes(Ar);
+		GRANULAR_NETWORK_MEMORY_TRACKING_TRACK("AllLevelStatuses",
+			AllLevelStatuses.CountBytes(Ar);
+			for (const FLevelStatus& LevelStatus : AllLevelStatuses)
+			{
+				LevelStatus.CountBytes(Ar);
+			}
+		);
 
-		LevelNamesAndTimes.CountBytes(Ar);
-		for (const FLevelNameAndTime& LevelNameAndTime : LevelNamesAndTimes)
-		{
-			LevelNameAndTime.CountBytes(Ar);
-		}
+		GRANULAR_NETWORK_MEMORY_TRACKING_TRACK("LevelStatusesByName",
+			LevelStatusesByName.CountBytes(Ar);
+			for (const auto& LevelStatusNamePair : LevelStatusesByName)
+			{
+				LevelStatusNamePair.Key.CountBytes(Ar);
+			}
+		);
 
-		LevelIntervals.CountBytes(Ar);
-		TrackedRewindActorsByGUID.CountBytes(Ar);
-		AllLevelStatuses.CountBytes(Ar);
-		for (const FLevelStatus& LevelStatus : AllLevelStatuses)
-		{
-			LevelStatus.CountBytes(Ar);
-		}
+		GRANULAR_NETWORK_MEMORY_TRACKING_TRACK("LevelStatusIndexByLevel", LevelStatusIndexByLevel.CountBytes(Ar));
+		GRANULAR_NETWORK_MEMORY_TRACKING_TRACK("SeenLevelStatuses", SeenLevelStatuses.CountBytes(Ar));
+		GRANULAR_NETWORK_MEMORY_TRACKING_TRACK("LevelsPendingFastForward", LevelsPendingFastForward.CountBytes(Ar));
+		GRANULAR_NETWORK_MEMORY_TRACKING_TRACK("ObjectsWithExternalData", ObjectsWithExternalData.CountBytes(Ar));
+		GRANULAR_NETWORK_MEMORY_TRACKING_TRACK("CheckpointSaveContext", CheckpointSaveContext.CountBytes(Ar));
 
-		LevelStatusesByName.CountBytes(Ar);
-		for (const auto& LevelStatusNamePair : LevelStatusesByName)
-		{
-			LevelStatusNamePair.Key.CountBytes(Ar);
-		}
 
-		LevelStatusIndexByLevel.CountBytes(Ar);
-		SeenLevelStatuses.CountBytes(Ar);
-		LevelsPendingFastForward.CountBytes(Ar);
-		ObjectsWithExternalData.CountBytes(Ar);
-		CheckpointSaveContext.CountBytes(Ar);
-		QueuedPacketsBeforeTravel.CountBytes(Ar);
-		for (const FQueuedDemoPacket& QueuedPacket : QueuedPacketsBeforeTravel)
-		{
-			QueuedPacket.CountBytes(Ar);
-		}
+		GRANULAR_NETWORK_MEMORY_TRACKING_TRACK("QueuedPacketsBeforeTravel",
+			QueuedPacketsBeforeTravel.CountBytes(Ar);
+			for (const FQueuedDemoPacket& QueuedPacket : QueuedPacketsBeforeTravel)
+			{
+				QueuedPacket.CountBytes(Ar);
+			}
+		);
 	}
 }
 
 void UDemoNetConnection::Serialize(FArchive& Ar)
 {
-	Super::Serialize(Ar);
+	GRANULAR_NETWORK_MEMORY_TRACKING_INIT(Ar, "UDemoNetConnection::Serialize");
+
+	GRANULAR_NETWORK_MEMORY_TRACKING_TRACK("Super", Super::Serialize(Ar));
 
 	if (Ar.IsCountingMemory())
 	{
-		QueuedDemoPackets.CountBytes(Ar);
-		for (const FQueuedDemoPacket& QueuedPacket : QueuedDemoPackets)
-		{
-			QueuedPacket.CountBytes(Ar);
-		}
-
-		QueuedCheckpointPackets.CountBytes(Ar);
-		for (const FQueuedDemoPacket& QueuedPacket : QueuedCheckpointPackets)
-		{
-			QueuedPacket.CountBytes(Ar);
-		}
+		GRANULAR_NETWORK_MEMORY_TRACKING_TRACK("QueuedDemoPackets",
+			QueuedDemoPackets.CountBytes(Ar);
+			for (const FQueuedDemoPacket& QueuedPacket : QueuedDemoPackets)
+			{
+				QueuedPacket.CountBytes(Ar);
+			}
+		);
+		
+		GRANULAR_NETWORK_MEMORY_TRACKING_TRACK("QueuedCheckpointPackets",
+			QueuedCheckpointPackets.CountBytes(Ar);
+			for (const FQueuedDemoPacket& QueuedPacket : QueuedCheckpointPackets)
+			{
+				QueuedPacket.CountBytes(Ar);
+			}
+		);
 	}
 }
 
diff --git a/Engine/Source/Runtime/Engine/Private/Net/NetworkGranularMemoryLogging.cpp b/Engine/Source/Runtime/Engine/Private/Net/NetworkGranularMemoryLogging.cpp
--- a/Engine/Source/Runtime/Engine/Private/Net/NetworkGranularMemoryLogging.cpp
+++ b/Engine/Source/Runtime/Engine/Private/Net/NetworkGranularMemoryLogging.cpp
@@ -7,6 +7,8 @@
 #include "Serialization/ArchiveCountMem.h"
 #include "EngineLogs.h"
 #include "HAL/IConsoleManager.h"
+#include "Containers/Ticker.h"
+#include "Templates/Sorting.h"
 
 namespace GranularNetworkMemoryTrackingPrivate
 {
@@ -16,34 +18,229 @@ namespace GranularNetworkMemoryTrackingPrivate
 		TEXT("When enabled, Obj List will print out highly detailed information about Network Memory Usage")
 	);
 
-	FHelper::FHelper(FArchive& InAr, FString&& InScopeName) :
+	struct FNetworkMemoryTrackingScope
+	{
+		FNetworkMemoryTrackingScope(const FString& ScopeName)
+			: ScopeName(ScopeName)
+		{
+		}
+
+		void GenerateRows(const FString& Prefix, TArray<FString>& OutRows)
+		{
+			const FString NewPrefix = (Prefix.IsEmpty()) ? ScopeName : (Prefix + TEXT("::") + ScopeName);
+
+			SubScopes.KeySort(TLess<FString>());
+			for (auto It = SubScopes.CreateIterator(); It; ++It)
+			{
+				It.Value()->GenerateRows(NewPrefix, OutRows);
+			}
+
+			struct FField
+			{
+				const FString* Name;
+				uint64 Bytes;
+			};
+
+			auto FieldLessThan = [](const FField& LHS, const FField& RHS)
+			{
+				return *(LHS.Name) < *(RHS.Name);
+			};
+
+			TArray<FField> FieldValues;
+			FieldValues.Reserve(Fields.Num() + 1);
+			FieldValues.Add({ &NewPrefix, TotalBytes });
+
+			for (auto ConstIt = Fields.CreateConstIterator(); ConstIt; ++ConstIt)
+			{
+				FieldValues.Add({ &ConstIt.Key(), ConstIt.Value() });
+			}
+
+			// Keep our total scope data in place, but sort the rest of the fields.
+			Sort(FieldValues.GetData() + 1, FieldValues.Num() - 1, FieldLessThan);
+
+			FString ReportRow = FString::Printf(TEXT("%s\r\n%s"),
+				*FString::JoinBy(FieldValues, TEXT(","), [](const FField& Field) { return *Field.Name; }),
+				*FString::JoinBy(FieldValues, TEXT(","), [](const FField& Field) { return FString::Printf(TEXT("%llu"), Field.Bytes); }));
+
+			OutRows.Emplace(ReportRow);
+		}
+
+		FNetworkMemoryTrackingScope* FindOrAddScope(const FString& NewScopeName)
+		{
+			TUniquePtr<FNetworkMemoryTrackingScope>& FoundScope = SubScopes.FindOrAdd(NewScopeName);
+			if (!FoundScope.IsValid())
+			{
+				FoundScope = MakeUnique<FNetworkMemoryTrackingScope>(NewScopeName);
+			}
+
+			return FoundScope.Get();
+		}
+
+		void AddBytesToField(const FString& FieldName, const uint64 Bytes)
+		{
+			Fields.FindOrAdd(FieldName) += Bytes;
+			TotalBytes += Bytes;
+		}
+
+		uint64 GetTotalBytes() const
+		{
+			return TotalBytes;
+		}
+
+		const FString ScopeName;
+
+	private:
+
+		TMap<FString, TUniquePtr<FNetworkMemoryTrackingScope>> SubScopes;
+		TMap<FString, uint64> Fields;
+		uint64 TotalBytes = 0u;
+	};
+
+	struct FNetworkMemoryTrackingScopeStack
+	{
+		static TUniquePtr<FNetworkMemoryTrackingScopeStack>& Get()
+		{
+			static TUniquePtr<FNetworkMemoryTrackingScopeStack> Stack;
+
+			if (!Stack.IsValid())
+			{
+				Stack.Reset(new FNetworkMemoryTrackingScopeStack());
+			}
+
+			return Stack;
+		}
+
+		void PushScope(const FString& ScopeName)
+		{
+			if (CurrentScope)
+			{
+				CurrentScope = CurrentScope->FindOrAddScope(ScopeName);
+			}
+			else if (TUniquePtr<FNetworkMemoryTrackingScope>* ExistingScope = TopLevelScopes.Find(ScopeName))
+			{
+				CurrentScope = ExistingScope->Get();
+			}
+			else
+			{
+				 CurrentScope = TopLevelScopes.Add(ScopeName, MakeUnique<FNetworkMemoryTrackingScope>(ScopeName)).Get();
+			}
+
+			ScopeStack.Push(CurrentScope);
+		}
+
+		void PopScope()
+		{
+			ScopeStack.Pop();
+			CurrentScope = (ScopeStack.Num() > 0) ? ScopeStack.Top() : nullptr;
+		}
+
+		void TrackWork(const FString& WorkName, const uint64 WorkBytes)
+		{
+			static const FString UnknownScopeName(TEXT("UNKNOWN"));
+
+			auto GetCorrectScope = [this]() -> FNetworkMemoryTrackingScope &
+			{
+				if (CurrentScope)
+				{
+					return *CurrentScope;
+				}
+
+				if (TUniquePtr<FNetworkMemoryTrackingScope> * UnknownScope = TopLevelScopes.Find(UnknownScopeName))
+				{
+					return **UnknownScope;
+				}
+
+				return *TopLevelScopes.Emplace(UnknownScopeName, MakeUnique<FNetworkMemoryTrackingScope>(UnknownScopeName));
+			};
+
+			GetCorrectScope().AddBytesToField(WorkName, WorkBytes);
+		}
+
+	private:
+
+		FNetworkMemoryTrackingScopeStack() :
+			TickHandle(FTicker::GetCoreTicker().AddTicker(TEXT("NetworkGranularMemoryLogging::FNetworkMemoryTrackingScopeStack"), 0, &FNetworkMemoryTrackingScopeStack::OnTick))
+		{
+		}
+
+		static bool OnTick(float DeltaTime)
+		{
+			TUniquePtr<FNetworkMemoryTrackingScopeStack>& ScopeStack = Get();
+			ScopeStack->TopLevelScopes.KeySort(TLess<FString>());
+
+			const FString EmptyPrefix(TEXT(""));
+			TArray<FString> Rows;
+			TArray<FString> TopLevelKBRows;
+
+			
+			for (auto It = ScopeStack->TopLevelScopes.CreateIterator(); It; ++It)
+			{
+				It.Value()->GenerateRows(EmptyPrefix, Rows);
+				TopLevelKBRows.Add(FString::Printf(TEXT("%s KB\r\n%d"), *It.Key(), static_cast<double>(It.Value()->GetTotalBytes()) / double(1024.f)));
+			}
+
+			Rows.Append(TopLevelKBRows);
+
+			// TODO: Replace \r\n with some platform specific newline macro.
+			UE_LOG(LogNet, Warning, TEXT("\r\n%s"), *FString::Join(Rows, TEXT("\r\n\r\n")));
+
+			FTicker::GetCoreTicker().RemoveTicker(ScopeStack->TickHandle);
+			ScopeStack.Reset();
+			return true;
+		}
+
+		FDelegateHandle TickHandle;
+		TMap<FString, TUniquePtr<FNetworkMemoryTrackingScope>> TopLevelScopes;
+		FNetworkMemoryTrackingScope* CurrentScope = nullptr;
+		TArray<FNetworkMemoryTrackingScope*> ScopeStack;
+	};
+
+	static const bool ShouldTrackMemory(const FArchive& Ar)
+	{
+		return CVarUseGranularNetworkTracking.GetValueOnAnyThread() && Ar.IsCountingMemory() && FString(TEXT("FArchiveCountMem")).Equals(Ar.GetArchiveName());
+	}
+
+	FScopeMarker::FScopeMarker(FArchive& InAr, FString&& InScopeName) :
 		Ar(InAr),
 		ScopeName(InScopeName),
-		bShouldTrack(CVarUseGranularNetworkTracking.GetValueOnAnyThread() && FString(TEXT("FArchiveCountMem")).Equals(Ar.GetArchiveName()))
+		ScopeStack(ShouldTrackMemory(Ar) ? FNetworkMemoryTrackingScopeStack::Get().Get() : nullptr)
+	{
+		if (ScopeStack)
+		{
+			ScopeStack->PushScope(ScopeName);
+		}
+	}
+
+	FScopeMarker::~FScopeMarker()
 	{
+		if (ScopeStack)
+		{
+			ScopeStack->PopScope();
+		}
 	}
 
-	void FHelper::BeginWork()
+	void FScopeMarker::BeginWork()
 	{
-		if (bShouldTrack)
+		if (ScopeStack)
 		{
 			PreWorkPos = ((FArchiveCountMem&)Ar).GetMax();
 		}
 	}
 
-	void FHelper::EndWork(const FString& WorkName)
+	void FScopeMarker::EndWork(const FString& WorkName)
 	{
-		if (bShouldTrack)
+		if (ScopeStack)
 		{
 			LogCustomWork(WorkName, ((FArchiveCountMem&)Ar).GetMax() - PreWorkPos);
 		}
 	}
 
-	void FHelper::LogCustomWork(const FString& WorkName, const uint64 Bytes) const
+	void FScopeMarker::LogCustomWork(const FString& WorkName, const uint64 Bytes) const
 	{
-		if (bShouldTrack)
+		if (ScopeStack)
 		{
-			UE_LOG(LogNet, Log, TEXT("%s: %s is %llu bytes"), *ScopeName, *WorkName, Bytes);
+			// UE_LOG(LogNet, Log, TEXT("%s: %s is %llu bytes"), *ScopeName, *WorkName, Bytes);
+			ScopeStack->TrackWork(WorkName, Bytes);
 		}
 	}
 };
diff --git a/Engine/Source/Runtime/Engine/Private/NetConnection.cpp b/Engine/Source/Runtime/Engine/Private/NetConnection.cpp
--- a/Engine/Source/Runtime/Engine/Private/NetConnection.cpp
+++ b/Engine/Source/Runtime/Engine/Private/NetConnection.cpp
@@ -592,12 +592,17 @@ bool UNetConnection::IsEncryptionEnabled() const
 
 void UNetConnection::Serialize( FArchive& Ar )
 {
-	UObject::Serialize( Ar );
-	Ar << PackageMap;
-	for (UChannel* Channel : Channels)
-	{
-		Ar << Channel;
-	}
+	GRANULAR_NETWORK_MEMORY_TRACKING_INIT(Ar, "UNetConnection::Serialize");
+
+	GRANULAR_NETWORK_MEMORY_TRACKING_TRACK("UObject", UObject::Serialize(Ar));
+
+	GRANULAR_NETWORK_MEMORY_TRACKING_TRACK("PackageName", Ar << PackageMap;);
+	GRANULAR_NETWORK_MEMORY_TRACKING_TRACK("Channels",
+		for (UChannel* Channel : Channels)
+		{
+			Ar << Channel;
+		}
+	);
 
 	if (Ar.IsCountingMemory())
 	{
@@ -608,7 +613,7 @@ void UNetConnection::Serialize( FArchive& Ar )
 		//		Histogram data.
 		// These are probably insignificant, though.
 
-		GRANULAR_NETWORK_MEMORY_TRACKING_INIT(Ar, "UNetConnection::Serialize");
+		GRANULAR_NETWORK_MEMORY_TRACKING_TRACK("Channels", Channels.CountBytes(Ar));
 
 		GRANULAR_NETWORK_MEMORY_TRACKING_TRACK("Challenge", Challenge.CountBytes(Ar));
 		GRANULAR_NETWORK_MEMORY_TRACKING_TRACK("ClientResponse", ClientResponse.CountBytes(Ar));
diff --git a/Engine/Source/Runtime/Engine/Private/NetworkDriver.cpp b/Engine/Source/Runtime/Engine/Private/NetworkDriver.cpp
--- a/Engine/Source/Runtime/Engine/Private/NetworkDriver.cpp
+++ b/Engine/Source/Runtime/Engine/Private/NetworkDriver.cpp
@@ -2316,7 +2316,9 @@ void UNetDriver::FReplicationChangelistMgrWrapper::CountBytes(FArchive& Ar) cons
 
 void UNetDriver::Serialize( FArchive& Ar )
 {
-	Super::Serialize( Ar );
+	GRANULAR_NETWORK_MEMORY_TRACKING_INIT(Ar, "UNetDriver::Serialize");
+
+	GRANULAR_NETWORK_MEMORY_TRACKING_TRACK("UNetDriver::Super", Super::Serialize(Ar));
 
 	if (Ar.IsCountingMemory())
 	{
@@ -2330,8 +2332,6 @@ void UNetDriver::Serialize( FArchive& Ar )
 		//		DDoSDetection data
 		// These are probably insignificant, though.
 
-		GRANULAR_NETWORK_MEMORY_TRACKING_INIT(Ar, "UNetDriver::Serialize");
-
 		GRANULAR_NETWORK_MEMORY_TRACKING_TRACK("MappedClientConnection", MappedClientConnections.CountBytes(Ar));
 		GRANULAR_NETWORK_MEMORY_TRACKING_TRACK("RecentlyDisconnectedClients", RecentlyDisconnectedClients.CountBytes(Ar));
 		GRANULAR_NETWORK_MEMORY_TRACKING_TRACK("GuidCache",
diff --git a/Engine/Source/Runtime/Engine/Private/PackageMapClient.cpp b/Engine/Source/Runtime/Engine/Private/PackageMapClient.cpp
--- a/Engine/Source/Runtime/Engine/Private/PackageMapClient.cpp
+++ b/Engine/Source/Runtime/Engine/Private/PackageMapClient.cpp
@@ -2223,15 +2223,15 @@ void UPackageMapClient::SetHasQueuedBunches(const FNetworkGUID& NetGUID, bool bH
 
 void UPackageMapClient::Serialize(FArchive& Ar)
 {
-	Super::Serialize(Ar);
+	GRANULAR_NETWORK_MEMORY_TRACKING_INIT(Ar, "UPackageMapClient::Serialize");
+
+	GRANULAR_NETWORK_MEMORY_TRACKING_TRACK("Super", Super::Serialize(Ar));
 
 	if (Ar.IsCountingMemory())
 	{
 		// TODO: We don't currently track:
 		//		Working Bunches.
 
-		GRANULAR_NETWORK_MEMORY_TRACKING_INIT(Ar, "UPackageMapClient::Serialize");
-
 		GRANULAR_NETWORK_MEMORY_TRACKING_TRACK("NetGUIDExportCountMap", NetGUIDExportCountMap.CountBytes(Ar));
 
 		GRANULAR_NETWORK_MEMORY_TRACKING_TRACK("ExportGUIDArchives",
diff --git a/Engine/Source/Runtime/Engine/Private/RepLayout.cpp b/Engine/Source/Runtime/Engine/Private/RepLayout.cpp
--- a/Engine/Source/Runtime/Engine/Private/RepLayout.cpp
+++ b/Engine/Source/Runtime/Engine/Private/RepLayout.cpp
@@ -472,11 +472,15 @@ struct FCustomDeltaChangelistState
 
 	void CountBytes(FArchive& Ar) const
 	{
-		ArrayStates.CountBytes(Ar);
-		for (const FDeltaArrayHistoryState& ArrayState : ArrayStates)
-		{
-			ArrayState.CountBytes(Ar);
-		}
+		GRANULAR_NETWORK_MEMORY_TRACKING_INIT(Ar, "FCustomDeltaChangelistState::CountBytes");
+
+		GRANULAR_NETWORK_MEMORY_TRACKING_TRACK("ArrayStates",
+			ArrayStates.CountBytes(Ar);
+			for (const FDeltaArrayHistoryState& ArrayState : ArrayStates)
+			{
+				ArrayState.CountBytes(Ar);
+			}
+		);
 	}
 };
 
@@ -788,15 +792,17 @@ void FRepStateStaticBuffer::CountBytes(FArchive& Ar) const
 	struct FCountBytesHelper
 	{
 		FCountBytesHelper(
-			FArchive& InAr,
+			FArchiveCountMem& InAr,
 			const FConstRepShadowDataBuffer InShadowData,
+			const uint64 InTotalShadowMemory,
 			const TArray<FRepParentCmd>& InParents,
 			const TArray<FRepLayoutCmd>& InCmds)
 
-			: Ar((FArchiveCountMem&)InAr)
+			: Ar(InAr)
 			, MainShadowData(InShadowData)
 			, Parents(InParents)
 			, Cmds(InCmds)
+			, TotalShadowMemory(InTotalShadowMemory)
 		{
 		}
 
@@ -804,9 +810,12 @@ void FRepStateStaticBuffer::CountBytes(FArchive& Ar) const
 		{
 			uint64 NewMax = Ar.GetMax();
 			uint64 OldMax = 0;
+			uint64 OldShadowOffset = 0u;
 
-			for (const FRepParentCmd& Parent : Parents)
+			for (int32 ParentIndex = 0; ParentIndex < Parents.Num(); ++ParentIndex)
 			{
+				const FRepParentCmd& Parent = Parents[ParentIndex];
+
 				OldMax = NewMax;
 
 				CountBytes_Command(Parent, Parent.CmdStart, Parent.CmdEnd, MainShadowData);
@@ -820,26 +829,25 @@ void FRepStateStaticBuffer::CountBytes(FArchive& Ar) const
 				}
 				else
 				{
+					const uint64 NextShadowOffset = (ParentIndex < Parents.Num() - 1) ? Parents[ParentIndex + 1].ShadowOffset : TotalShadowMemory;
+
 					NonRepMemory += (NewMax - OldMax);
+					NonRepStaticMemory += NextShadowOffset - OldShadowOffset;
 				}
+
+				OldShadowOffset = Parent.ShadowOffset;
 			}
 		}
 
-		void CountBytes_Command(const FRepParentCmd& Parent, const int32 CmdStart, const int32 CmdEnd, const FConstRepShadowDataBuffer ShadowData) const
+		void CountBytes_Command(const FRepParentCmd& Parent, const int32 CmdStart, const int32 CmdEnd, const FConstRepShadowDataBuffer ShadowData)
 		{
 			for (int32 CmdIndex = CmdStart; CmdIndex < CmdEnd; ++CmdIndex)
 			{
-				const FRepLayoutCmd& Cmd = Cmds[CmdIndex];
-				CountBytes_r(Parent, Cmd, CmdIndex, ShadowData);
-
-				if (ERepLayoutCmdType::DynamicArray == Cmd.Type)
-				{
-					CmdIndex = Cmd.EndCmd - 1;
-				}
+				CountBytes_r(Parent, Cmds[CmdIndex], CmdIndex, ShadowData);
 			}
 		}
 
-		void CountBytes_r(const FRepParentCmd& Parent, const FRepLayoutCmd& Cmd, const int32 InCmdIndex, const FConstRepShadowDataBuffer ShadowData) const
+		void CountBytes_r(const FRepParentCmd& Parent, const FRepLayoutCmd& Cmd, int32& InCmdIndex, const FConstRepShadowDataBuffer ShadowData)
 		{
 			if (ERepLayoutCmdType::DynamicArray == Cmd.Type)
 			{
@@ -851,8 +859,10 @@ void FRepStateStaticBuffer::CountBytes(FArchive& Ar) const
 				for (int32 i = 0; i < Array->Num(); ++i)
 				{
 					const int32 ArrayElementOffset = Cmd.ElementSize * i;
-					CountBytes_Command(Parent, InCmdIndex + 1, Cmd.EndCmd, ShadowArrayData + ArrayElementOffset);
+					CountBytes_Command(Parent, InCmdIndex + 1, Cmd.EndCmd - 1, ShadowArrayData + ArrayElementOffset);
 				}
+
+				InCmdIndex = Cmd.EndCmd - 1;
 			}
 			else if (ERepLayoutCmdType::PropertyString == Cmd.Type)
 			{
@@ -865,17 +875,25 @@ void FRepStateStaticBuffer::CountBytes(FArchive& Ar) const
 		const TArray<FRepParentCmd>& Parents;
 		const TArray<FRepLayoutCmd>& Cmds;
 
+		const uint64 TotalShadowMemory;
 		uint64 OnRepMemory = 0;
 		uint64 NonRepMemory = 0;
+		uint64 NonRepStaticMemory = 0u;
 	};
 
-	GRANULAR_NETWORK_MEMORY_TRACKING_TRACK("Static Memory", Buffer.CountBytes(Ar));
-	GRANULAR_NETWORK_MEMORY_TRACKING_TRACK("Dynamic Memory (Undercounts!)",
-		// FCountBytesHelper CountBytesHelper(Ar, Buffer.GetData(), RepLayout->Parents, RepLayout->Cmds);
-		// CountBytesHelper.CountBytes();
-		// GRANULAR_NETWORK_MEMORY_TRACKING_CUSTOM_WORK("OnRepMemory", CountBytesHelper.OnRepMemory);
-		// GRANULAR_NETWORK_MEMORY_TRACKING_CUSTOM_WORK("NonRepMemory", CountBytesHelper.NonRepMemory);
-	);
+	FArchiveCountMem LocalAr(nullptr);
+	Buffer.CountBytes(LocalAr);
+	const uint64 StaticTotalMemory = LocalAr.GetMax();
+
+	FCountBytesHelper CountBytesHelper(LocalAr, Buffer.GetData(), StaticTotalMemory, RepLayout->Parents, RepLayout->Cmds);
+	CountBytesHelper.CountBytes();
+
+	const uint64 StaticOnRepMemory = StaticTotalMemory - CountBytesHelper.NonRepStaticMemory;
+
+	GRANULAR_NETWORK_MEMORY_TRACKING_TRACK("Static_OnRep", Ar.CountBytes(StaticTotalMemory, StaticTotalMemory));
+	GRANULAR_NETWORK_MEMORY_TRACKING_TRACK("Static_NotOnRep", Ar.CountBytes(CountBytesHelper.NonRepStaticMemory, CountBytesHelper.NonRepStaticMemory));
+	GRANULAR_NETWORK_MEMORY_TRACKING_TRACK("Dynamic_OnRep", Ar.CountBytes(CountBytesHelper.OnRepMemory, CountBytesHelper.OnRepMemory));
+	GRANULAR_NETWORK_MEMORY_TRACKING_TRACK("Dynamic_NotOnRep", Ar.CountBytes(CountBytesHelper.NonRepMemory, CountBytesHelper.NonRepMemory));
 }
 
 FRepChangelistState::FRepChangelistState(
@@ -892,13 +910,16 @@ FRepChangelistState::FRepChangelistState(
 
 void FRepChangelistState::CountBytes(FArchive& Ar) const
 {
-	StaticBuffer.CountBytes(Ar);
-	SharedSerialization.CountBytes(Ar);
+	GRANULAR_NETWORK_MEMORY_TRACKING_INIT(Ar, "FRepChangelistState::CountBytes");
+	GRANULAR_NETWORK_MEMORY_TRACKING_TRACK("StaticBuffer", StaticBuffer.CountBytes(Ar));
+	GRANULAR_NETWORK_MEMORY_TRACKING_TRACK("SharedSerialization", SharedSerialization.CountBytes(Ar));
 
 	if (CustomDeltaChangelistState)
 	{
-		Ar.CountBytes(sizeof(FCustomDeltaChangelistState), sizeof(FCustomDeltaChangelistState));
-		CustomDeltaChangelistState->CountBytes(Ar);
+		GRANULAR_NETWORK_MEMORY_TRACKING_TRACK("CustomDeltaChangelistState",
+			Ar.CountBytes(sizeof(FCustomDeltaChangelistState), sizeof(FCustomDeltaChangelistState));
+			CustomDeltaChangelistState->CountBytes(Ar);
+		);
 	}
 }
 
@@ -1969,6 +1990,21 @@ void FRepLayout::FilterChangeListToActive(
 	OutProperties.Add(0);
 }
 
+void FRepSerializationSharedInfo::CountBytes(FArchive& Ar) const
+{
+	GRANULAR_NETWORK_MEMORY_TRACKING_INIT(Ar, "FRepSerializationSharedInfo::CountBytes");
+
+	GRANULAR_NETWORK_MEMORY_TRACKING_TRACK("SharedPropertyInfo", SharedPropertyInfo.CountBytes(Ar));
+
+	GRANULAR_NETWORK_MEMORY_TRACKING_TRACK("SerializedProperties",
+		if (FNetBitWriter const* const LocalSerializedProperties = SerializedProperties.Get())
+		{
+			Ar.CountBytes(sizeof(FNetBitWriter), sizeof(FNetBitWriter));
+			LocalSerializedProperties->CountMemory(Ar);
+		}
+	);
+}
+
 const FRepSerializedPropertyInfo* FRepSerializationSharedInfo::WriteSharedProperty(
 	const FRepLayoutCmd& Cmd,
 	const FGuid& PropertyGuid,
diff --git a/Engine/Source/Runtime/Engine/Public/Net/NetworkGranularMemoryLogging.h b/Engine/Source/Runtime/Engine/Public/Net/NetworkGranularMemoryLogging.h
--- a/Engine/Source/Runtime/Engine/Public/Net/NetworkGranularMemoryLogging.h
+++ b/Engine/Source/Runtime/Engine/Public/Net/NetworkGranularMemoryLogging.h
@@ -8,9 +8,10 @@
 
 namespace GranularNetworkMemoryTrackingPrivate
 {
-	struct ENGINE_API FHelper
+	struct ENGINE_API FScopeMarker
 	{
-		FHelper(FArchive& InAr, FString&& InScopeName);
+		FScopeMarker(FArchive& InAr, FString&& InScopeName);
+		~FScopeMarker();
 
 		void BeginWork();
 
@@ -20,27 +21,34 @@ namespace GranularNetworkMemoryTrackingPrivate
 
 		const bool IsEnabled() const
 		{
-			return bShouldTrack;
+			return ScopeStack != nullptr;
+		}
+
+		const FString& GetScopeName()
+		{
+			return ScopeName;
 		}
 
 	private:
 
+		friend struct FNetworkMemoryTrackingScopeStack;
+
 		uint64 PreWorkPos = 0;
 
 		const FArchive& Ar;
 		const FString ScopeName;
-		const bool bShouldTrack;
+		struct FNetworkMemoryTrackingScopeStack* ScopeStack;
 	};
 }
 
-#define GRANULAR_NETWORK_MEMORY_TRACKING_INIT(Archive, ScopeName) GranularNetworkMemoryTrackingPrivate::FHelper GranularNetworkMemoryHelper(Archive, ScopeName);
+#define GRANULAR_NETWORK_MEMORY_TRACKING_INIT(Archive, ScopeName) GranularNetworkMemoryTrackingPrivate::FScopeMarker GranularNetworkMemoryScope(Archive, ScopeName);
 #define GRANULAR_NETWORK_MEMORY_TRACKING_TRACK(Id, Work) \
 	{ \
-		GranularNetworkMemoryHelper.BeginWork(); \
+		GranularNetworkMemoryScope.BeginWork(); \
 		Work; \
-		GranularNetworkMemoryHelper.EndWork(Id); \
+		GranularNetworkMemoryScope.EndWork(Id); \
 	}
-#define GRANULAR_NETWORK_MEMORY_TRACKING_CUSTOM_WORK(Id, Value) GranularNetworkMemoryHelper.LogCustomWork(Id, Value);
+#define GRANULAR_NETWORK_MEMORY_TRACKING_CUSTOM_WORK(Id, Value) GranularNetworkMemoryScope.LogCustomWork(Id, Value);
 
 #else
 
diff --git a/Engine/Source/Runtime/Engine/Public/Net/RepLayout.h b/Engine/Source/Runtime/Engine/Public/Net/RepLayout.h
--- a/Engine/Source/Runtime/Engine/Public/Net/RepLayout.h
+++ b/Engine/Source/Runtime/Engine/Public/Net/RepLayout.h
@@ -275,16 +275,7 @@ struct FRepSerializationSharedInfo
 	/** Binary blob of net serialized data to be shared */
 	TUniquePtr<FNetBitWriter> SerializedProperties;
 
-	void CountBytes(FArchive& Ar) const
-	{
-		SharedPropertyInfo.CountBytes(Ar);
-
-		if (FNetBitWriter const * const LocalSerializedProperties = SerializedProperties.Get())
-		{
-			Ar.CountBytes(sizeof(FNetBitWriter), sizeof(FNetBitWriter));
-			LocalSerializedProperties->CountMemory(Ar);
-		}
-	}
+	void CountBytes(FArchive& Ar) const;
 
 private:
 
