More contextual display names for native AttachToActor/AttachToComponent APIs to restore visual separation in the BP editor's context menu.

diff --git a/Engine/Source/Runtime/Engine/Classes/Components/SceneComponent.h b/Engine/Source/Runtime/Engine/Classes/Components/SceneComponent.h
--- a/Engine/Source/Runtime/Engine/Classes/Components/SceneComponent.h
+++ b/Engine/Source/Runtime/Engine/Classes/Components/SceneComponent.h
@@ -682,7 +682,7 @@ public:
 	 * @param  bWeldSimulatedBodies		Whether to weld together simulated physics bodies.
 	 * @return True if attachment is successful (or already attached to requested parent/socket), false if attachment is rejected and there is no change in AttachParent.
 	 */
-	UFUNCTION(BlueprintCallable, Category = "Utilities|Transformation", meta = (DisplayName = "AttachToComponent", ScriptName = "AttachToComponent", bWeldSimulatedBodies=true))
+	UFUNCTION(BlueprintCallable, Category = "Utilities|Transformation", meta = (DisplayName = "AttachComponentToComponent", ScriptName = "AttachToComponent", bWeldSimulatedBodies=true))
 	bool K2_AttachToComponent(USceneComponent* Parent, FName SocketName, EAttachmentRule LocationRule, EAttachmentRule RotationRule, EAttachmentRule ScaleRule, bool bWeldSimulatedBodies);
 
 	/** DEPRECATED - Use AttachToComponent() instead */
diff --git a/Engine/Source/Runtime/Engine/Classes/GameFramework/Actor.h b/Engine/Source/Runtime/Engine/Classes/GameFramework/Actor.h
--- a/Engine/Source/Runtime/Engine/Classes/GameFramework/Actor.h
+++ b/Engine/Source/Runtime/Engine/Classes/GameFramework/Actor.h
@@ -1251,7 +1251,7 @@ public:
 
 	/** DEPRECATED - Use AttachToComponent() instead */
 	UE_DEPRECATED(4.17, "Use AttachToComponent() instead.")
-	UFUNCTION(BlueprintCallable, meta = (DisplayName = "AttachActorToComponent (Deprecated)", ScriptNoExport, AttachLocationType = "KeepRelativeOffset"), Category = "Utilities|Transformation")
+	UFUNCTION(BlueprintCallable, meta = (DisplayName = "AttachRootComponentTo (Deprecated)", ScriptNoExport, AttachLocationType = "KeepRelativeOffset"), Category = "Utilities|Transformation")
 	void K2_AttachRootComponentTo(USceneComponent* InParent, FName InSocketName = NAME_None, EAttachLocation::Type AttachLocationType = EAttachLocation::KeepRelativeOffset, bool bWeldSimulatedBodies = true);
 
 	/**
@@ -1263,7 +1263,7 @@ public:
 	 * @param ScaleRule					How to handle scale when attaching.
 	 * @param bWeldSimulatedBodies		Whether to weld together simulated physics bodies.
 	 */
-	UFUNCTION(BlueprintCallable, meta = (DisplayName = "AttachToComponent", ScriptName = "AttachToComponent", bWeldSimulatedBodies = true), Category = "Utilities|Transformation")
+	UFUNCTION(BlueprintCallable, meta = (DisplayName = "AttachActorToComponent", ScriptName = "AttachToComponent", bWeldSimulatedBodies = true), Category = "Utilities|Transformation")
 	void K2_AttachToComponent(USceneComponent* Parent, FName SocketName, EAttachmentRule LocationRule, EAttachmentRule RotationRule, EAttachmentRule ScaleRule, bool bWeldSimulatedBodies);
 
 	/**
@@ -1276,7 +1276,7 @@ public:
 
 	/** DEPRECATED - Use AttachToActor() instead */
 	UE_DEPRECATED(4.17, "Use AttachToActor() instead.")
-	UFUNCTION(BlueprintCallable, meta = (DisplayName = "AttachActorToActor (Deprecated)", ScriptNoExport, AttachLocationType = "KeepRelativeOffset"), Category = "Utilities|Transformation")
+	UFUNCTION(BlueprintCallable, meta = (DisplayName = "AttachRootComponentToActor (Deprecated)", ScriptNoExport, AttachLocationType = "KeepRelativeOffset"), Category = "Utilities|Transformation")
 	void K2_AttachRootComponentToActor(AActor* InParentActor, FName InSocketName = NAME_None, EAttachLocation::Type AttachLocationType = EAttachLocation::KeepRelativeOffset, bool bWeldSimulatedBodies = true);
 
 	/**
@@ -1288,7 +1288,7 @@ public:
 	void AttachToActor(AActor* ParentActor, const FAttachmentTransformRules& AttachmentRules, FName SocketName = NAME_None);
 
 	/**
-	 * Attaches the RootComponent of this Actor to the supplied component, optionally at a named socket. It is not valid to call this on components that are not Registered.
+	 * Attaches the RootComponent of this Actor to the supplied actor, optionally at a named socket.
 	 * @param ParentActor				Actor to attach this actor's RootComponent to
 	 * @param SocketName				Socket name to attach to, if any
 	 * @param LocationRule				How to handle translation when attaching.
@@ -1296,7 +1296,7 @@ public:
 	 * @param ScaleRule					How to handle scale when attaching.
 	 * @param bWeldSimulatedBodies		Whether to weld together simulated physics bodies.
 	 */
-	UFUNCTION(BlueprintCallable, meta = (DisplayName = "AttachToActor", ScriptName = "AttachToActor", bWeldSimulatedBodies=true), Category = "Utilities|Transformation")
+	UFUNCTION(BlueprintCallable, meta = (DisplayName = "AttachActorToActor", ScriptName = "AttachToActor", bWeldSimulatedBodies=true), Category = "Utilities|Transformation")
 	void K2_AttachToActor(AActor* ParentActor, FName SocketName, EAttachmentRule LocationRule, EAttachmentRule RotationRule, EAttachmentRule ScaleRule, bool bWeldSimulatedBodies);
 
 	UE_DEPRECATED(4.17, "Use AttachToComponent() with EAttachLocation::SnapToTarget option instead")
