Don't create Custom Delta Tracking state when recording or playing back replays.

diff --git a/Engine/Source/Runtime/Engine/Classes/Engine/DemoNetDriver.h b/Engine/Source/Runtime/Engine/Classes/Engine/DemoNetDriver.h
--- a/Engine/Source/Runtime/Engine/Classes/Engine/DemoNetDriver.h
+++ b/Engine/Source/Runtime/Engine/Classes/Engine/DemoNetDriver.h
@@ -1215,4 +1215,9 @@ private:
 
 	bool ProcessFastForwardPackets(TArrayView<FPlaybackPacket> Packets, const TSet<int32>& LevelIndices);
 	void ProcessPlaybackPackets(TArrayView<FPlaybackPacket> Packets);
+
+	virtual ECreateReplicationChangelistMgrFlags GetCreateReplicationChangelistMgrFlags() const override
+	{
+		return ECreateReplicationChangelistMgrFlags::SkipDeltaCustomState;
+	}
 };
diff --git a/Engine/Source/Runtime/Engine/Classes/Engine/NetDriver.h b/Engine/Source/Runtime/Engine/Classes/Engine/NetDriver.h
--- a/Engine/Source/Runtime/Engine/Classes/Engine/NetDriver.h
+++ b/Engine/Source/Runtime/Engine/Classes/Engine/NetDriver.h
@@ -1576,6 +1576,8 @@ protected:
 
 private:
 
+	ENGINE_API virtual enum class ECreateReplicationChangelistMgrFlags GetCreateReplicationChangelistMgrFlags() const;
+
 	FDelegateHandle PostGarbageCollectHandle;
 	void PostGarbageCollect();
 
diff --git a/Engine/Source/Runtime/Engine/Private/NetworkDriver.cpp b/Engine/Source/Runtime/Engine/Private/NetworkDriver.cpp
--- a/Engine/Source/Runtime/Engine/Private/NetworkDriver.cpp
+++ b/Engine/Source/Runtime/Engine/Private/NetworkDriver.cpp
@@ -5299,7 +5299,7 @@ TSharedPtr< FReplicationChangelistMgr > UNetDriver::GetReplicationChangeListMgr(
 	if (!ReplicationChangeListMgrPtr)
 	{
 		const TSharedPtr<const FRepLayout> RepLayout = GetObjectClassRepLayout(Object->GetClass());
-		FReplicationChangelistMgrWrapper Wrapper(Object, RepLayout->CreateReplicationChangelistMgr(Object));
+		FReplicationChangelistMgrWrapper Wrapper(Object, RepLayout->CreateReplicationChangelistMgr(Object, GetCreateReplicationChangelistMgrFlags()));
 		ReplicationChangeListMgrPtr = &ReplicationChangeListMap.Add(Object, Wrapper);
 	}
 
@@ -5601,6 +5601,11 @@ void UNetDriver::ResetAsyncLoadDelinquencyAnalytics()
 	}
 }
 
+ECreateReplicationChangelistMgrFlags UNetDriver::GetCreateReplicationChangelistMgrFlags() const
+{
+	return ECreateReplicationChangelistMgrFlags::None;
+}
+
 FAutoConsoleCommandWithWorld	DumpRelevantActorsCommand(
 	TEXT("net.DumpRelevantActors"), 
 	TEXT( "Dumps information on relevant actors during next network update" ), 
diff --git a/Engine/Source/Runtime/Engine/Private/RepLayout.cpp b/Engine/Source/Runtime/Engine/Private/RepLayout.cpp
--- a/Engine/Source/Runtime/Engine/Private/RepLayout.cpp
+++ b/Engine/Source/Runtime/Engine/Private/RepLayout.cpp
@@ -6006,13 +6006,13 @@ FRepStateStaticBuffer FRepLayout::CreateShadowBuffer(const FConstRepObjectDataBu
 	return ShadowData;
 }
 
-TSharedPtr<FReplicationChangelistMgr> FRepLayout::CreateReplicationChangelistMgr(const UObject* InObject) const
+TSharedPtr<FReplicationChangelistMgr> FRepLayout::CreateReplicationChangelistMgr(const UObject* InObject, const ECreateReplicationChangelistMgrFlags CreateFlags) const
 {
 	// ChangelistManager / ChangelistState will hold onto a unique pointer for this
 	// so no need to worry about deleting it here.
 
 	FCustomDeltaChangelistState* DeltaChangelistState = nullptr;
-	if (LifetimeCustomPropertyState && !!LifetimeCustomPropertyState->GetNumFastArrayProperties())
+	if (!EnumHasAnyFlags(CreateFlags, ECreateReplicationChangelistMgrFlags::SkipDeltaCustomState) && LifetimeCustomPropertyState && !!LifetimeCustomPropertyState->GetNumFastArrayProperties())
 	{
 		DeltaChangelistState = new FCustomDeltaChangelistState(LifetimeCustomPropertyState->GetNumFastArrayProperties());
 	}
diff --git a/Engine/Source/Runtime/Engine/Public/Net/RepLayout.h b/Engine/Source/Runtime/Engine/Public/Net/RepLayout.h
--- a/Engine/Source/Runtime/Engine/Public/Net/RepLayout.h
+++ b/Engine/Source/Runtime/Engine/Public/Net/RepLayout.h
@@ -976,6 +976,14 @@ private:
 	int32 LastSuccessfulCmdIndex;
 };
 
+enum class ECreateReplicationChangelistMgrFlags
+{
+	None,
+	SkipDeltaCustomState,	//! Skip creating CustomDeltaState used for tracking.
+							//! Only do this if you know you'll never need it (like for replay recording)
+};
+ENUM_CLASS_FLAGS(ECreateReplicationChangelistMgrFlags);
+
 enum class ECreateRepLayoutFlags
 {
 	None,
@@ -1129,9 +1137,10 @@ public:
 	/**
 	 * Creates and initializes a new FReplicationChangelistMgr.
 	 *
-	 * @param InObject	The Object that is being managed.
+	 * @param InObject		The Object that is being managed.
+	 * @param CreateFlags	Flags modifying how the manager is created.
 	 */
-	TSharedPtr<FReplicationChangelistMgr> CreateReplicationChangelistMgr(const UObject* InObject) const;
+	TSharedPtr<FReplicationChangelistMgr> CreateReplicationChangelistMgr(const UObject* InObject, const ECreateReplicationChangelistMgrFlags CreateFlags) const;
 
 	/**
 	 * Creates and initializes a new FRepState.
