[NO NON-PII COMMIT MESSAGE]

diff --git a/Engine/Source/Runtime/Engine/Private/AudioDerivedData.cpp b/Engine/Source/Runtime/Engine/Private/AudioDerivedData.cpp
--- a/Engine/Source/Runtime/Engine/Private/AudioDerivedData.cpp
+++ b/Engine/Source/Runtime/Engine/Private/AudioDerivedData.cpp
@@ -1139,7 +1139,7 @@ static void CookSurroundWave( USoundWave* SoundWave, FName FormatName, const IAu
 	{
 		SoundWave->RawData.Unlock();
 
-		UE_LOG(LogAudioDerivedData, Warning, TEXT("No raw wave data for: %s"), *SoundWave->GetFullName());
+		UE_LOG(LogAudioDerivedData, Display, TEXT("No raw wave data for: %s"), *SoundWave->GetFullName());
 		return;
 	}
 
