Fix linux CIS issue by forward declaring enum in global scope, instead of function definition.

diff --git a/Engine/Source/Runtime/Engine/Classes/Engine/NetDriver.h b/Engine/Source/Runtime/Engine/Classes/Engine/NetDriver.h
--- a/Engine/Source/Runtime/Engine/Classes/Engine/NetDriver.h
+++ b/Engine/Source/Runtime/Engine/Classes/Engine/NetDriver.h
@@ -330,6 +330,8 @@ class IAnalyticsProvider;
 class FNetAnalyticsAggregator;
 class UNetDriver;
 
+enum class ECreateReplicationChangelistMgrFlags;
+
 using FConnectionMap = TMap<TSharedRef<const FInternetAddr>, UNetConnection*, FDefaultSetAllocator, FInternetAddrConstKeyMapFuncs<UNetConnection*>>;
 
 extern ENGINE_API TAutoConsoleVariable<int32> CVarNetAllowEncryption;
@@ -1576,7 +1578,7 @@ protected:
 
 private:
 
-	ENGINE_API virtual enum class ECreateReplicationChangelistMgrFlags GetCreateReplicationChangelistMgrFlags() const;
+	ENGINE_API virtual ECreateReplicationChangelistMgrFlags GetCreateReplicationChangelistMgrFlags() const;
 
 	FDelegateHandle PostGarbageCollectHandle;
 	void PostGarbageCollect();
