Fix light "Affect Global Illumination" flag being ignored by ray traced global illumination. Slightly improve flag description.

(ushell-p4-cherrypick of 9404729 by riley.alston)

diff --git a/Engine/Source/Runtime/Engine/Classes/Components/LightComponentBase.h b/Engine/Source/Runtime/Engine/Classes/Components/LightComponentBase.h
--- a/Engine/Source/Runtime/Engine/Classes/Components/LightComponentBase.h
+++ b/Engine/Source/Runtime/Engine/Classes/Components/LightComponentBase.h
@@ -91,7 +91,7 @@ class ENGINE_API ULightComponentBase : public USceneComponent
 	UPROPERTY(EditAnywhere, BlueprintReadOnly, Category = Light, AdvancedDisplay)
 	uint32 bAffectReflection : 1;
 
-	/** Whether the light affects objects in reflections, when ray-traced global illumination is enabled. */
+	/** Whether the light affects global illumination, when ray-traced global illumination is enabled. */
 	UPROPERTY(EditAnywhere, BlueprintReadOnly, Category = Light, AdvancedDisplay)
 	uint32 bAffectGlobalIllumination : 1;
 
diff --git a/Engine/Source/Runtime/Engine/Public/SceneManagement.h b/Engine/Source/Runtime/Engine/Public/SceneManagement.h
--- a/Engine/Source/Runtime/Engine/Public/SceneManagement.h
+++ b/Engine/Source/Runtime/Engine/Public/SceneManagement.h
@@ -1559,7 +1559,7 @@ protected:
 	/** Whether the light affects objects in reflections, when ray-traced reflection is enabled. */
 	const uint8 bAffectReflection : 1;
 
-	/** Whether the light affects objects in reflections, when ray-traced global illumination is enabled. */
+	/** Whether the light affects global illumination, when ray-traced global illumination is enabled. */
 	const uint8 bAffectGlobalIllumination : 1;
 
 	/** Whether the light affects translucency or not.  Disabling this can save GPU time when there are many small lights. */
diff --git a/Engine/Source/Runtime/Renderer/Private/RayTracing/RayTracingGlobalIllumination.cpp b/Engine/Source/Runtime/Renderer/Private/RayTracing/RayTracingGlobalIllumination.cpp
--- a/Engine/Source/Runtime/Renderer/Private/RayTracing/RayTracingGlobalIllumination.cpp
+++ b/Engine/Source/Runtime/Renderer/Private/RayTracing/RayTracingGlobalIllumination.cpp
@@ -166,6 +166,7 @@ void SetupLightParameters(
 		if (LightParameters->Count >= MaxLightCount) break;
 
 		if (Light.LightSceneInfo->Proxy->HasStaticLighting() && Light.LightSceneInfo->IsPrecomputedLightingValid()) continue;
+		if (!Light.LightSceneInfo->Proxy->AffectGlobalIllumination()) continue;
 
 		FLightShaderParameters LightShaderParameters;
 		Light.LightSceneInfo->Proxy->GetLightShaderParameters(LightShaderParameters);
