Prevent reentrant Actor construction due to post-edit RRCS at construction time.

Change summary:
- Added a bitflag that's set while the Actor is being constructed (editor only).
- Modified AActor::ReregisterComponentsWhenModified() to consider the bitflag.
- Added an ensure() to RRCS to avoid reconstruction if the bitflag is set (failsafe).

diff --git a/Engine/Source/Runtime/Engine/Classes/GameFramework/Actor.h b/Engine/Source/Runtime/Engine/Classes/GameFramework/Actor.h
--- a/Engine/Source/Runtime/Engine/Classes/GameFramework/Actor.h
+++ b/Engine/Source/Runtime/Engine/Classes/GameFramework/Actor.h
@@ -747,6 +747,9 @@ private:
 	/** Whether this actor is temporarily hidden within the editor; used for show/hide/etc functionality w/o dirtying the actor. */
 	UPROPERTY(Transient)
 	uint8 bHiddenEdTemporary:1;
+
+	/** Set while actor is being constructed. Used to ensure that construction is not re-entrant. */
+	uint8 bActorIsBeingConstructed:1;
 #endif // WITH_EDITORONLY_DATA
 
 public:
diff --git a/Engine/Source/Runtime/Engine/Private/ActorConstruction.cpp b/Engine/Source/Runtime/Engine/Private/ActorConstruction.cpp
--- a/Engine/Source/Runtime/Engine/Private/ActorConstruction.cpp
+++ b/Engine/Source/Runtime/Engine/Private/ActorConstruction.cpp
@@ -208,22 +208,27 @@ void AActor::RerunConstructionScripts()
 	FEditorScriptExecutionGuard ScriptGuard;
 	// don't allow (re)running construction scripts on dying actors and Actors that seamless traveled 
 	// were constructed in the previous level and should not have construction scripts rerun
-	bool bAllowReconstruction =  !bActorSeamlessTraveled && !IsPendingKill() && !HasAnyFlags(RF_BeginDestroyed|RF_FinishDestroyed);
+	bool bAllowReconstruction = !bActorSeamlessTraveled && !IsPendingKill() && !HasAnyFlags(RF_BeginDestroyed|RF_FinishDestroyed);
 #if WITH_EDITOR
 	if(bAllowReconstruction && GIsEditor)
 	{
-		// Generate the blueprint hierarchy for this actor
-		TArray<UBlueprint*> ParentBPStack;
-		bAllowReconstruction = UBlueprint::GetBlueprintHierarchyFromClass(GetClass(), ParentBPStack);
-		if(bAllowReconstruction)
+		// Don't allow reconstruction if we're still in the middle of construction.
+		bAllowReconstruction = !bActorIsBeingConstructed;
+		if (ensureMsgf(bAllowReconstruction, TEXT("Attempted to rerun construction scripts on an Actor that isn't fully constructed yet (%s)."), *GetFullName()))
 		{
-			for(int i = ParentBPStack.Num() - 1; i > 0 && bAllowReconstruction; --i)
+			// Generate the blueprint hierarchy for this actor
+			TArray<UBlueprint*> ParentBPStack;
+			bAllowReconstruction = UBlueprint::GetBlueprintHierarchyFromClass(GetClass(), ParentBPStack);
+			if (bAllowReconstruction)
 			{
-				const UBlueprint* ParentBP = ParentBPStack[i];
-				if(ParentBP && ParentBP->bBeingCompiled)
+				for (int i = ParentBPStack.Num() - 1; i > 0 && bAllowReconstruction; --i)
 				{
-					// don't allow (re)running construction scripts if a parent BP is being compiled
-					bAllowReconstruction = false;
+					const UBlueprint* ParentBP = ParentBPStack[i];
+					if (ParentBP && ParentBP->bBeingCompiled)
+					{
+						// don't allow (re)running construction scripts if a parent BP is being compiled
+						bAllowReconstruction = false;
+					}
 				}
 			}
 		}
@@ -684,6 +689,13 @@ bool AActor::ExecuteConstruction(const FTransform& Transform, const FRotationCon
 	check(!IsPendingKill());
 	check(!HasAnyFlags(RF_BeginDestroyed|RF_FinishDestroyed));
 
+#if WITH_EDITOR
+	// Guard against reentrancy due to attribute editing at construction time.
+	// @see RerunConstructionScripts()
+	checkf(!bActorIsBeingConstructed, TEXT("Actor construction is not reentrant"));
+	FGuardValue_Bitfield(bActorIsBeingConstructed, true);
+#endif
+
 	// ensure that any existing native root component gets this new transform
 	// we can skip this in the default case as the given transform will be the root component's transform
 	if (RootComponent && !bIsDefaultTransform)
diff --git a/Engine/Source/Runtime/Engine/Private/ActorEditor.cpp b/Engine/Source/Runtime/Engine/Private/ActorEditor.cpp
--- a/Engine/Source/Runtime/Engine/Private/ActorEditor.cpp
+++ b/Engine/Source/Runtime/Engine/Private/ActorEditor.cpp
@@ -185,7 +185,13 @@ void AActor::PostEditMove(bool bFinished)
 
 bool AActor::ReregisterComponentsWhenModified() const
 {
-	return !IsTemplate() && !GetOutermost()->HasAnyPackageFlags(PKG_PlayInEditor) && GetWorld() != nullptr;
+	// For child actors, redirect to the parent's owner (we do the same in RerunConstructionScripts).
+	if (const AActor* ParentActor = GetParentActor())
+	{
+		return ParentActor->ReregisterComponentsWhenModified();
+	}
+
+	return !bActorIsBeingConstructed && !IsTemplate() && !GetOutermost()->HasAnyPackageFlags(PKG_PlayInEditor) && GetWorld() != nullptr;
 }
 
 void AActor::DebugShowComponentHierarchy(  const TCHAR* Info, bool bShowPosition )
