Fixing issue where particle vector random instance parameters were incorrectly using VRand, causing unit vector sizing instead of a box for certain assets. A new unit vector sizing option was added to allow data created under the incorrect implementation to have the same behavior.

diff --git a/Engine/Source/Editor/DetailCustomizations/Private/ParticleSysParamStructCustomization.cpp b/Engine/Source/Editor/DetailCustomizations/Private/ParticleSysParamStructCustomization.cpp
--- a/Engine/Source/Editor/DetailCustomizations/Private/ParticleSysParamStructCustomization.cpp
+++ b/Engine/Source/Editor/DetailCustomizations/Private/ParticleSysParamStructCustomization.cpp
@@ -134,17 +134,17 @@ EVisibility FParticleSysParamStructCustomization::GetScalarLowVisibility() const
 
 EVisibility FParticleSysParamStructCustomization::GetVectorVisibility() const
 {
-	return (ParameterType == PSPT_Vector || ParameterType == PSPT_VectorRand) ? EVisibility::Visible : EVisibility::Collapsed;
+	return (ParameterType == PSPT_Vector || ParameterType == PSPT_VectorRand || ParameterType == PSPT_VectorUnitRand) ? EVisibility::Visible : EVisibility::Collapsed;
 }
 
 EVisibility FParticleSysParamStructCustomization::GetVectorLowVisibility() const
 {
-	return (ParameterType == PSPT_VectorRand) ? EVisibility::Visible : EVisibility::Collapsed;
+	return (ParameterType == PSPT_VectorRand || ParameterType == PSPT_VectorUnitRand) ? EVisibility::Visible : EVisibility::Collapsed;
 }
 
 EVisibility FParticleSysParamStructCustomization::GetColorVisibility() const
 {
-	return (ParameterType == PSPT_Vector || ParameterType == PSPT_VectorRand || ParameterType == PSPT_Color) ? EVisibility::Visible : EVisibility::Collapsed;
+	return (ParameterType == PSPT_Vector || ParameterType == PSPT_VectorRand || ParameterType == PSPT_Color || ParameterType == PSPT_VectorUnitRand) ? EVisibility::Visible : EVisibility::Collapsed;
 }
 
 EVisibility FParticleSysParamStructCustomization::GetActorVisibility() const
diff --git a/Engine/Source/Runtime/Engine/Classes/Particles/ParticleSystemComponent.h b/Engine/Source/Runtime/Engine/Classes/Particles/ParticleSystemComponent.h
--- a/Engine/Source/Runtime/Engine/Classes/Particles/ParticleSystemComponent.h
+++ b/Engine/Source/Runtime/Engine/Classes/Particles/ParticleSystemComponent.h
@@ -51,6 +51,7 @@ enum EParticleSysParamType
 	PSPT_Color UMETA(DisplayName="Color"),
 	PSPT_Actor UMETA(DisplayName="Actor"),
 	PSPT_Material UMETA(DisplayName="Material"),
+	PSPT_VectorUnitRand UMETA(DisplayName = "Vector Unit Random"),
 	PSPT_MAX,
 };
 
@@ -116,6 +117,7 @@ struct FParticleSysParam
 	 *	PSPT_Color      - Use the color value
 	 *	PSPT_Actor      - Use the actor value
 	 *	PSPT_Material   - Use the material value
+	 *	PSPT_VectorUnitRand - Select a random unit vector and scale along the range [Vector_Low..Vector)
 	 */
 	UPROPERTY(EditAnywhere, BlueprintReadWrite, Category=ParticleSysParam)
 	TEnumAsByte<enum EParticleSysParamType> ParamType;
@@ -1548,6 +1550,14 @@ public:
 	 */
 	void SetVectorRandParameter(FName ParameterName, const FVector& Param, const FVector& ParamLow);
 
+
+
+	/**
+	 *	Set a named random unit vector instance parameter on this ParticleSystemComponent.
+	 *	Updates the parameter if it already exists, or creates a new entry if not.
+	 */
+	void SetVectorUnitRandParameter(FName ParameterName, const FVector& Param, const FVector& ParamLow);
+
 	/** 
 	 *	Set a named random float instance parameter on this ParticleSystemComponent. 
 	 *	Updates the parameter if it already exists, or creates a new entry if not. 
diff --git a/Engine/Source/Runtime/Engine/Private/Particles/ParticleComponents.cpp b/Engine/Source/Runtime/Engine/Private/Particles/ParticleComponents.cpp
--- a/Engine/Source/Runtime/Engine/Private/Particles/ParticleComponents.cpp
+++ b/Engine/Source/Runtime/Engine/Private/Particles/ParticleComponents.cpp
@@ -4628,7 +4628,7 @@ void UParticleSystemComponent::PostEditChangeChainProperty(FPropertyChangedChain
 				for (int32 InstIdx = 0; InstIdx < InstanceParameters.Num(); InstIdx++)
 				{
 					FParticleSysParam& PSysParam = InstanceParameters[InstIdx];
-					if ((PSysParam.ParamType == PSPT_Vector) || (PSysParam.ParamType == PSPT_VectorRand))
+					if ((PSysParam.ParamType == PSPT_Vector) || (PSysParam.ParamType == PSPT_VectorRand) || (PSysParam.ParamType == PSPT_VectorUnitRand))
 					{
 						PSysParam.Vector.X = PSysParam.Color.R / 255.0f;
 						PSysParam.Vector.Y = PSysParam.Color.G / 255.0f;
@@ -7165,6 +7165,36 @@ void UParticleSystemComponent::SetVectorRandParameter(FName ParameterName,const
 	InstanceParameters[NewParamIndex].Vector_Low = ParamLow;
 }
 
+void UParticleSystemComponent::SetVectorUnitRandParameter(FName ParameterName, const FVector& Param, const FVector& ParamLow)
+{
+	LLM_SCOPE(ELLMTag::Particles);
+
+	if (ParameterName == NAME_None)
+	{
+		return;
+	}
+	check(IsInGameThread());
+
+	// First see if an entry for this name already exists
+	for (int32 i = 0; i < InstanceParameters.Num(); i++)
+	{
+		FParticleSysParam& P = InstanceParameters[i];
+		if (P.Name == ParameterName && P.ParamType == PSPT_VectorUnitRand)
+		{
+			P.Vector = Param;
+			P.Vector_Low = ParamLow;
+			return;
+		}
+	}
+
+	// We didn't find one, so create a new one.
+	int32 NewParamIndex = InstanceParameters.AddZeroed();
+	InstanceParameters[NewParamIndex].Name = ParameterName;
+	InstanceParameters[NewParamIndex].ParamType = PSPT_VectorUnitRand;
+	InstanceParameters[NewParamIndex].Vector = Param;
+	InstanceParameters[NewParamIndex].Vector_Low = ParamLow;
+}
+
 
 void UParticleSystemComponent::SetColorParameter(FName Name, FLinearColor Param)
 {
@@ -7308,6 +7338,12 @@ bool UParticleSystemComponent::GetVectorParameter(const FName InName,FVector& Ou
 				return true;
 			}
 			else if (Param.ParamType == PSPT_VectorRand)
+			{
+				FVector RandValue(RandomStream.FRand(), RandomStream.FRand(), RandomStream.FRand());
+				OutVector = Param.Vector + (Param.Vector_Low - Param.Vector) * RandValue;
+				return true;
+			}
+			else if (Param.ParamType == PSPT_VectorUnitRand)
 			{
 				OutVector = Param.Vector + (Param.Vector_Low - Param.Vector) * RandomStream.VRand();
 				return true;
@@ -7338,6 +7374,12 @@ bool UParticleSystemComponent::GetAnyVectorParameter(const FName InName,FVector&
 				return true;
 			}
 			if (Param.ParamType == PSPT_VectorRand)
+			{
+				FVector RandValue(RandomStream.FRand(), RandomStream.FRand(), RandomStream.FRand());
+				OutVector = Param.Vector + (Param.Vector_Low - Param.Vector) * RandValue;
+				return true;
+			}
+			else if (Param.ParamType == PSPT_VectorUnitRand)
 			{
 				OutVector = Param.Vector + (Param.Vector_Low - Param.Vector) * RandomStream.VRand();
 				return true;
