Fixed sky always creating built data registry when not needed.
9566337 also helped removing the creation of the registry but simply using current one if there as done by other components.
Tested again UI notifying "LIGHTING NEEDS REBUILD" and everything is working as expected.

(ushell-p4-cherrypick of 9614889 by [REDACTED_PERSON])

diff --git a/Engine/Source/Runtime/Engine/Private/Components/SkyAtmosphereComponent.cpp b/Engine/Source/Runtime/Engine/Private/Components/SkyAtmosphereComponent.cpp
--- a/Engine/Source/Runtime/Engine/Private/Components/SkyAtmosphereComponent.cpp
+++ b/Engine/Source/Runtime/Engine/Private/Components/SkyAtmosphereComponent.cpp
@@ -248,7 +248,6 @@ void USkyAtmosphereComponent::PostEditChangeProperty(FPropertyChangedEvent& Prop
 		CategoryName == FName(TEXT("Atmosphere - Absorption")) ||
 		CategoryName == FName(TEXT("Art direction")))
 	{
-		UMapBuildDataRegistry* Registry = GetOwner() && GetOwner()->GetLevel() ? GetOwner()->GetLevel()->GetOrCreateMapBuildData() : nullptr;
 		if (SkyAtmosphereComponentStaticLightingBuilt(this))
 		{
 			// If we have changed an atmosphere property and the lighyting has already been built, we need to ask for a rebuild by updating the static lighting GUIDs.
diff --git a/Engine/Source/Runtime/Renderer/Private/SkyAtmosphereRendering.cpp b/Engine/Source/Runtime/Renderer/Private/SkyAtmosphereRendering.cpp
--- a/Engine/Source/Runtime/Renderer/Private/SkyAtmosphereRendering.cpp
+++ b/Engine/Source/Runtime/Renderer/Private/SkyAtmosphereRendering.cpp
@@ -8,7 +8,6 @@
 #include "CanvasTypes.h"
 #include "Components/SkyAtmosphereComponent.h"
 #include "DeferredShadingRenderer.h"
-#include "Engine/MapBuildDataRegistry.h"
 #include "LightSceneInfo.h"
 #include "PixelShaderUtils.h"
 #include "RenderTargetTemp.h"
