AWorldSettings - Hide HierarchicalLODSetup & OverrideBaseMaterial if HLODSetupAsset is valid
Use editcondition instead of CanEditChange() to bypass a known issue with multiple details view not updating properly

diff --git a/Engine/Source/Runtime/Engine/Classes/GameFramework/WorldSettings.h b/Engine/Source/Runtime/Engine/Classes/GameFramework/WorldSettings.h
--- a/Engine/Source/Runtime/Engine/Classes/GameFramework/WorldSettings.h
+++ b/Engine/Source/Runtime/Engine/Classes/GameFramework/WorldSettings.h
@@ -629,12 +629,12 @@ public:
 	TSoftClassPtr<class UHierarchicalLODSetup> HLODSetupAsset;
 
 	/** If set overrides the project-wide base material used for Proxy Materials*/
-	UPROPERTY(EditAnywhere, config, Category = LODSystem)
+	UPROPERTY(EditAnywhere, config, Category = LODSystem, meta=(editcondition = "bEnableHierarchicalLODSystem && HLODSetupAsset == nullptr"))
 	TSoftObjectPtr<class UMaterialInterface> OverrideBaseMaterial;	
 
 protected:
 	/** Hierarchical LOD Setup */
-	UPROPERTY(EditAnywhere, Category=LODSystem, config)
+	UPROPERTY(EditAnywhere, Category=LODSystem, config, meta=(editcondition = "bEnableHierarchicalLODSystem && HLODSetupAsset == nullptr"))
 	TArray<struct FHierarchicalSimplification>	HierarchicalLODSetup;
 
 public:
diff --git a/Engine/Source/Runtime/Engine/Private/WorldSettings.cpp b/Engine/Source/Runtime/Engine/Private/WorldSettings.cpp
--- a/Engine/Source/Runtime/Engine/Private/WorldSettings.cpp
+++ b/Engine/Source/Runtime/Engine/Private/WorldSettings.cpp
@@ -556,12 +556,6 @@ bool AWorldSettings::CanEditChange(const UProperty* InProperty) const
 				return LightmassSettings.EnvironmentIntensity > 0;
 			}
 		}
-
-		if (PropertyName == GET_MEMBER_NAME_STRING_CHECKED(AWorldSettings, OverrideBaseMaterial) ||
-			PropertyName == GET_MEMBER_NAME_STRING_CHECKED(AWorldSettings, HierarchicalLODSetup))
-		{
-			return bEnableHierarchicalLODSystem && HLODSetupAsset.IsNull();
-		}
 	}
 
 	return Super::CanEditChange(InProperty);
