Editgrate cl 9586084

Expose SnapshotTransactionBuffer to Blueprint as SnapshotObject

diff --git a/Engine/Source/Runtime/Engine/Classes/Kismet/KismetSystemLibrary.h b/Engine/Source/Runtime/Engine/Classes/Kismet/KismetSystemLibrary.h
--- a/Engine/Source/Runtime/Engine/Classes/Kismet/KismetSystemLibrary.h
+++ b/Engine/Source/Runtime/Engine/Classes/Kismet/KismetSystemLibrary.h
@@ -1778,6 +1778,16 @@ class ENGINE_API UKismetSystemLibrary : public UBlueprintFunctionLibrary
 	UFUNCTION(BlueprintCallable, Category = "Transactions")
 	static void TransactObject(UObject* Object);
 
+	/**
+	 * Notify the current transaction (if any) that this object is about to be modified and should be snapshot for intermediate update.
+	 * @note Internally this calls SnapshotTransactionBuffer on the given object.
+	 * @note Only available in the editor.
+	 *
+	 * @param	Object		The object that is about to be modified.
+	 */
+	UFUNCTION(BlueprintCallable, Category = "Transactions")
+	static void SnapshotObject(UObject* Object);
+
 	// --- Asset Manager ------------------------------
 
 	/** Returns the Object associated with a Primary Asset Id, this will only return a valid object if it is in memory, it will not load it */
diff --git a/Engine/Source/Runtime/Engine/Private/KismetSystemLibrary.cpp b/Engine/Source/Runtime/Engine/Private/KismetSystemLibrary.cpp
--- a/Engine/Source/Runtime/Engine/Private/KismetSystemLibrary.cpp
+++ b/Engine/Source/Runtime/Engine/Private/KismetSystemLibrary.cpp
@@ -2680,6 +2680,16 @@ void UKismetSystemLibrary::TransactObject(UObject* Object)
 #endif
 }
 
+void UKismetSystemLibrary::SnapshotObject(UObject* Object)
+{
+#if WITH_EDITOR
+	if (Object)
+	{
+		SnapshotTransactionBuffer(Object);
+	}
+#endif
+}
+
 UObject* UKismetSystemLibrary::GetObjectFromPrimaryAssetId(FPrimaryAssetId PrimaryAssetId)
 {
 	if (UAssetManager* Manager = UAssetManager::GetIfValid())
