Fix HLOD MergeSetting.LODSelectionType overriden on load when saved to level instead of BP asset

diff --git a/Engine/Source/Runtime/Engine/Private/WorldSettings.cpp b/Engine/Source/Runtime/Engine/Private/WorldSettings.cpp
--- a/Engine/Source/Runtime/Engine/Private/WorldSettings.cpp
+++ b/Engine/Source/Runtime/Engine/Private/WorldSettings.cpp
@@ -439,7 +439,7 @@ void AWorldSettings::PostLoad()
 	for (FHierarchicalSimplification& Entry : HierarchicalLODSetup)
 	{
 		Entry.ProxySetting.PostLoadDeprecated();
-		Entry.MergeSetting.LODSelectionType = EMeshLODSelectionType::CalculateLOD;
+		Entry.MergeSetting.PostLoadDeprecated();
 	}
 
 #endif// WITH_EDITOR
