Make sure skeletal mesh duplication is handled correctly

diff --git a/Engine/Source/Runtime/Engine/Private/SkeletalMeshLODRenderData.cpp b/Engine/Source/Runtime/Engine/Private/SkeletalMeshLODRenderData.cpp
--- a/Engine/Source/Runtime/Engine/Private/SkeletalMeshLODRenderData.cpp
+++ b/Engine/Source/Runtime/Engine/Private/SkeletalMeshLODRenderData.cpp
@@ -10,6 +10,7 @@
 #include "Interfaces/ITargetPlatform.h"
 #include "Interfaces/ITargetPlatformManagerModule.h"
 #include "PlatformInfo.h"
+#include "UObject/PropertyPortFlags.h"
 
 #if WITH_EDITOR
 #include "Rendering/SkeletalMeshModel.h"
@@ -932,12 +933,25 @@ void FSkeletalMeshLODRenderData::Serialize(FArchive& Ar, UObject* Owner, int32 I
 	FStripDataFlags StripFlags(Ar, ClassDataStripFlags);
 
 	const bool bIsBelowMinLOD = StripFlags.IsClassDataStripped(CDSF_MinLodData);
-	bool bIsLODCookedOut = IsLODCookedOut(Ar.CookingTarget(), OwnerMesh, bIsBelowMinLOD);
-	Ar << bIsLODCookedOut;
+	bool bIsLODCookedOut = false;
+	bool bInlined = false;
 
-	bool bInlined = bIsLODCookedOut || IsLODInlined(Ar.CookingTarget(), OwnerMesh, Idx, bIsBelowMinLOD);
-	Ar << bInlined;
-	bStreamedDataInlined = bInlined;
+	if (Ar.IsSaving() && !Ar.IsCooking() && !!(Ar.GetPortFlags() & PPF_Duplicate))
+	{
+		bInlined = bStreamedDataInlined;
+		bIsLODCookedOut = bIsBelowMinLOD && bInlined;
+		Ar << bIsLODCookedOut;
+		Ar << bInlined;
+	}
+	else
+	{
+		bIsLODCookedOut = IsLODCookedOut(Ar.CookingTarget(), OwnerMesh, bIsBelowMinLOD);
+		Ar << bIsLODCookedOut;
+
+		bInlined = bIsLODCookedOut || IsLODInlined(Ar.CookingTarget(), OwnerMesh, Idx, bIsBelowMinLOD);
+		Ar << bInlined;
+		bStreamedDataInlined = bInlined;
+	}
 
 	// Skeletal mesh buffers are kept in CPU memory after initialization to support merging of skeletal meshes.
 	const bool bForceKeepCPUResources = ShouldForceKeepCPUResources();
diff --git a/Engine/Source/Runtime/Engine/Private/SkeletalMeshRenderData.cpp b/Engine/Source/Runtime/Engine/Private/SkeletalMeshRenderData.cpp
--- a/Engine/Source/Runtime/Engine/Private/SkeletalMeshRenderData.cpp
+++ b/Engine/Source/Runtime/Engine/Private/SkeletalMeshRenderData.cpp
@@ -59,7 +59,7 @@ static void SerializeLODInfoForDDC(USkeletalMesh* SkeletalMesh, FString& KeySuff
 // differences, etc.) replace the version GUID below with a new one.
 // In case of merge conflicts with DDC versions, you MUST generate a new GUID
 // and set this new GUID as the version.                                       
-#define SKELETALMESH_DERIVEDDATA_VER TEXT("51BDAF351B6941548818C2F8B33397D7")
+#define SKELETALMESH_DERIVEDDATA_VER TEXT("29B218D50FE444AF8EFC8E4234F5D4E1")
 
 static const FString& GetSkeletalMeshDerivedDataVersion()
 {
