Fixed engine crash when importing texture bigger than 8K

diff --git a/Engine/Source/Runtime/Engine/Private/Streaming/StreamingTexture.cpp b/Engine/Source/Runtime/Engine/Private/Streaming/StreamingTexture.cpp
--- a/Engine/Source/Runtime/Engine/Private/Streaming/StreamingTexture.cpp
+++ b/Engine/Source/Runtime/Engine/Private/Streaming/StreamingTexture.cpp
@@ -54,7 +54,7 @@ void FStreamingRenderAsset::UpdateStaticData(const FRenderAssetStreamingSettings
 
 		if (IsTexture())
 		{
-			check(MipCount <= MAX_TEXTURE_MIP_COUNT);
+			MipCount = FMath::Min<int32>(MipCount, MAX_TEXTURE_MIP_COUNT);
 			const TextureGroup TextureLODGroup = static_cast<TextureGroup>(LODGroup);
 			BoostFactor = GetExtraBoost(TextureLODGroup, Settings);
 			bIsCharacterTexture = (TextureLODGroup == TEXTUREGROUP_Character || TextureLODGroup == TEXTUREGROUP_CharacterSpecular || TextureLODGroup == TEXTUREGROUP_CharacterNormalMap);
@@ -161,6 +161,13 @@ void FStreamingRenderAsset::UpdateDynamicData(const int32* NumStreamedMips, int3
 		{
 			LODBias = FMath::Max<int32>(RenderAsset->GetCachedLODBias() - NumCinematicMipLevels, 0);
 
+
+#if WITH_EDITORONLY_DATA
+			// When data is not cooked, the asset can have more mips than the engine supports.
+			// The engine limit is applied in UpdateStaticData() when computing MipCount, but this will also be be accounted in GetCachedLODBias().
+			LODBias -= RenderAsset->GetNumMipsForStreaming() - MipCount;
+#endif
+
 			// Reduce the max allowed resolution according to LODBias if the texture group allows it.
 			if (IsMaxResolutionAffectedByGlobalBias() && !Settings.bUsePerTextureBias)
 			{
diff --git a/Engine/Source/Runtime/Engine/Private/Streaming/Texture2DStreamIn.cpp b/Engine/Source/Runtime/Engine/Private/Streaming/Texture2DStreamIn.cpp
--- a/Engine/Source/Runtime/Engine/Private/Streaming/Texture2DStreamIn.cpp
+++ b/Engine/Source/Runtime/Engine/Private/Streaming/Texture2DStreamIn.cpp
@@ -15,7 +15,7 @@ FTexture2DStreamIn::FTexture2DStreamIn(UTexture2D* InTexture, int32 InRequestedM
 {
 	ensure(InRequestedMips > InTexture->GetNumResidentMips());
 
-	FMemory::Memzero(MipData, sizeof(MipData));
+	MipData.AddZeroed(InTexture->GetNumMips());
 }
 
 FTexture2DStreamIn::~FTexture2DStreamIn()
@@ -142,7 +142,7 @@ void FTexture2DStreamIn::DoAsyncCreateWithNewMips(const FContext& Context)
 				Texture2DRHI->GetFormat(),
 				RequestedMips,
 				Flags,
-				MipData + PendingFirstMip,
+				&MipData[PendingFirstMip],
 				RequestedMips - ResidentMips);
 		}
 	}
diff --git a/Engine/Source/Runtime/Engine/Private/Streaming/Texture2DStreamIn.h b/Engine/Source/Runtime/Engine/Private/Streaming/Texture2DStreamIn.h
--- a/Engine/Source/Runtime/Engine/Private/Streaming/Texture2DStreamIn.h
+++ b/Engine/Source/Runtime/Engine/Private/Streaming/Texture2DStreamIn.h
@@ -21,7 +21,7 @@ public:
 protected:
 
 	// StreamIn_Default : Locked mips of the intermediate textures, used as disk load destination.
-	void* MipData[MAX_TEXTURE_MIP_COUNT];
+	TArray<void*, TInlineAllocator<MAX_TEXTURE_MIP_COUNT> > MipData;
 
 
 	// ****************************
diff --git a/Engine/Source/Runtime/Engine/Private/Streaming/Texture2DStreamIn_DDC.cpp b/Engine/Source/Runtime/Engine/Private/Streaming/Texture2DStreamIn_DDC.cpp
--- a/Engine/Source/Runtime/Engine/Private/Streaming/Texture2DStreamIn_DDC.cpp
+++ b/Engine/Source/Runtime/Engine/Private/Streaming/Texture2DStreamIn_DDC.cpp
@@ -104,7 +104,7 @@ FTexture2DStreamIn_DDC::FTexture2DStreamIn_DDC(UTexture2D* InTexture, int32 InRe
 	: FTexture2DStreamIn(InTexture, InRequestedMips)
 	, bDDCIsInvalid(false)
 {
-	FMemory::Memzero(DDCHandles, sizeof(DDCHandles));
+	DDCHandles.AddZeroed(InTexture->GetNumMips());
 }
 
 FTexture2DStreamIn_DDC::~FTexture2DStreamIn_DDC()
diff --git a/Engine/Source/Runtime/Engine/Private/Streaming/Texture2DStreamIn_DDC.h b/Engine/Source/Runtime/Engine/Private/Streaming/Texture2DStreamIn_DDC.h
--- a/Engine/Source/Runtime/Engine/Private/Streaming/Texture2DStreamIn_DDC.h
+++ b/Engine/Source/Runtime/Engine/Private/Streaming/Texture2DStreamIn_DDC.h
@@ -26,7 +26,7 @@ public:
 protected:
 
 	// StreamIn_Default : Locked mips of the intermediate textures, used as disk load destination.
-	uint32 DDCHandles[MAX_TEXTURE_MIP_COUNT];
+	TArray<uint32, TInlineAllocator<MAX_TEXTURE_MIP_COUNT> > DDCHandles;
 
 	// Whether the DDC data was compatible or not.
 	bool bDDCIsInvalid;
diff --git a/Engine/Source/Runtime/Engine/Private/Streaming/Texture2DStreamIn_IO.cpp b/Engine/Source/Runtime/Engine/Private/Streaming/Texture2DStreamIn_IO.cpp
--- a/Engine/Source/Runtime/Engine/Private/Streaming/Texture2DStreamIn_IO.cpp
+++ b/Engine/Source/Runtime/Engine/Private/Streaming/Texture2DStreamIn_IO.cpp
@@ -19,7 +19,7 @@ FTexture2DStreamIn_IO::FTexture2DStreamIn_IO(UTexture2D* InTexture, int32 InRequ
 	, bPrioritizedIORequest(InPrioritizedIORequest)
 	, IOFileOffset(0)
 {
-	FMemory::Memzero(IORequests, sizeof(IORequests));
+	IORequests.AddZeroed(InTexture->GetNumMips());
 }
 
 FTexture2DStreamIn_IO::~FTexture2DStreamIn_IO()
@@ -96,7 +96,7 @@ void FTexture2DStreamIn_IO::SetIORequests(const FContext& Context)
 
 void FTexture2DStreamIn_IO::CancelIORequests()
 {
-	for (int32 MipIndex = 0; MipIndex < MAX_TEXTURE_MIP_COUNT; ++MipIndex)
+	for (int32 MipIndex = 0; MipIndex < IORequests.Num(); ++MipIndex)
 	{
 		FBulkDataIORequest* IORequest = IORequests[MipIndex];
 		if (IORequest)
diff --git a/Engine/Source/Runtime/Engine/Private/Streaming/Texture2DStreamIn_IO.h b/Engine/Source/Runtime/Engine/Private/Streaming/Texture2DStreamIn_IO.h
--- a/Engine/Source/Runtime/Engine/Private/Streaming/Texture2DStreamIn_IO.h
+++ b/Engine/Source/Runtime/Engine/Private/Streaming/Texture2DStreamIn_IO.h
@@ -63,7 +63,7 @@ private:
 
 
 	// Request for loading into each mip.
-	FBulkDataIORequest* IORequests[MAX_TEXTURE_MIP_COUNT];
+	TArray<FBulkDataIORequest*, TInlineAllocator<MAX_TEXTURE_MIP_COUNT> > IORequests;
 
 	bool bPrioritizedIORequest;
 
diff --git a/Engine/Source/Runtime/Engine/Private/Texture2D.cpp b/Engine/Source/Runtime/Engine/Private/Texture2D.cpp
--- a/Engine/Source/Runtime/Engine/Private/Texture2D.cpp
+++ b/Engine/Source/Runtime/Engine/Private/Texture2D.cpp
@@ -1497,7 +1497,6 @@ FTexture2DResource::FTexture2DResource( UTexture2D* InOwner, int32 InitialMipCou
 	bSRGB = InOwner->SRGB;
 
 	check(InitialMipCount>0);
-	check(UE_ARRAY_COUNT(MipData)>=GMaxTextureMipCount);
 
 	// Keep track of first miplevel to use.
 	CurrentFirstMip = InOwner->GetNumMips() - FMath::Max(InitialMipCount, Owner->PlatformData ? (int32)Owner->PlatformData->NumMipsInTail : 0);
@@ -1508,7 +1507,7 @@ FTexture2DResource::FTexture2DResource( UTexture2D* InOwner, int32 InitialMipCou
 	check(CurrentFirstMip<=Owner->GetMipTailBaseIndex());
 
 	// Retrieve initial mip data.
-	FMemory::Memzero(MipData, sizeof(MipData));
+	MipData.AddZeroed(InOwner->GetNumMips());
 	InOwner->GetMipData(CurrentFirstMip, &MipData[CurrentFirstMip]);
 	STAT( TextureSize = Owner->CalcTextureMemorySize( InitialMipCount ) );
 	STAT( LODGroupStatName = TextureGroupStatFNames[Owner->LODGroup] );
@@ -1530,7 +1529,7 @@ FTexture2DResource::~FTexture2DResource()
 		});
 
 	// Make sure we're not leaking memory if InitRHI has never been called.
-	for( int32 MipIndex=0; MipIndex<UE_ARRAY_COUNT(MipData); MipIndex++ )
+	for( int32 MipIndex=0; MipIndex<MipData.Num(); MipIndex++ )
 	{
 		// free any mip data that was copied 
 		if( MipData[MipIndex] )
diff --git a/Engine/Source/Runtime/Engine/Private/TextureCube.cpp b/Engine/Source/Runtime/Engine/Private/TextureCube.cpp
--- a/Engine/Source/Runtime/Engine/Private/TextureCube.cpp
+++ b/Engine/Source/Runtime/Engine/Private/TextureCube.cpp
@@ -216,6 +216,7 @@ public:
 
 		// Read the mip-levels into the RHI texture.
 		int32 NumMips = Owner->GetNumMips();
+		check(NumMips < MAX_TEXTURE_MIP_COUNT);
 		for( int32 FaceIndex=0; FaceIndex<6; FaceIndex++ )
 		{
 			for(int32 MipIndex=0; MipIndex < NumMips; MipIndex++)
@@ -289,6 +290,7 @@ private:
 	 */
 	void GetData( int32 FaceIndex, int32 MipIndex, void* Dest, uint32 DestPitch )
 	{
+		check(MipIndex < MAX_TEXTURE_MIP_COUNT);
 		check( MipData[FaceIndex][MipIndex] );
 
 		// for platforms that returned 0 pitch from Lock, we need to just use the bulk data directly, never do 
diff --git a/Engine/Source/Runtime/Engine/Public/TextureResource.h b/Engine/Source/Runtime/Engine/Public/TextureResource.h
--- a/Engine/Source/Runtime/Engine/Public/TextureResource.h
+++ b/Engine/Source/Runtime/Engine/Public/TextureResource.h
@@ -238,7 +238,7 @@ private:
 	int32 CurrentFirstMip;
 	
 	/** Local copy/ cache of mip data between creation and first call to InitRHI.							*/
-	void*				MipData[MAX_TEXTURE_MIP_COUNT];
+	TArray<void*, TInlineAllocator<MAX_TEXTURE_MIP_COUNT> > MipData;
 
 	/** 2D texture version of TextureRHI which is used to lock the 2D texture during mip transitions.		*/
 	FTexture2DRHIRef	Texture2DRHI;
