[NO NON-PII COMMIT MESSAGE]

diff --git a/Engine/Source/Runtime/Engine/Private/AudioDerivedData.cpp b/Engine/Source/Runtime/Engine/Private/AudioDerivedData.cpp
--- a/Engine/Source/Runtime/Engine/Private/AudioDerivedData.cpp
+++ b/Engine/Source/Runtime/Engine/Private/AudioDerivedData.cpp
@@ -1241,7 +1241,23 @@ static void CookSurroundWave( USoundWave* SoundWave, FName FormatName, const IAu
 					ResampleWaveData(SourceBuffers[ChannelIndex], DataSize, 1, WaveSampleRate, SampleRateOverride);
 				}
 
+				// Since each channel is resampled independently, we may have slightly different sample counts in each channel.
+				// To counter this, we truncate or zero-pad every non-zero channel's buffer to the zeroth channel's length.
+				int32 SizeOfZerothChannel = SourceBuffers[0].Num();
+
+				for (int32 ChannelIndex = 1; ChannelIndex < ChannelCount; ChannelIndex++)
+				{
+					TArray<uint8>& SourceBuffer = SourceBuffers[ChannelIndex];
+
+					if (SourceBuffer.Num() != SizeOfZerothChannel)
+					{	
+						UE_LOG(LogAudioDerivedData, Display, TEXT("Fixing up channel %d from %d to %d samples."), ChannelIndex, SizeOfZerothChannel, SourceBuffer.Num());
+						SourceBuffer.SetNumZeroed(SizeOfZerothChannel);
+					}
+				}
+
 				WaveSampleRate = SampleRateOverride;
+				SampleDataSize = SizeOfZerothChannel * ChannelCount;
 			}
 
 			UE_LOG(LogAudioDerivedData, Display, TEXT("Cooking %d channels for: %s"), ChannelCount, *SoundWave->GetFullName());
