[NO NON-PII COMMIT MESSAGE]

diff --git a/Engine/Source/Runtime/Engine/Private/AudioDerivedData.cpp b/Engine/Source/Runtime/Engine/Private/AudioDerivedData.cpp
--- a/Engine/Source/Runtime/Engine/Private/AudioDerivedData.cpp
+++ b/Engine/Source/Runtime/Engine/Private/AudioDerivedData.cpp
@@ -538,10 +538,9 @@ public:
 		}
 	}
 
-	/** Finalize work. Must be called ONLY by the game thread! */
+	/** Finalize work. Must be called ONLY by the thread that started this task! */
 	bool Finalize()
 	{
-		check(IsInGameThread());
 		// if we couldn't get from the DDC or didn't build synchronously, then we have to build now.
 		// This is a super edge case that should rarely happen.
 		if (!bSucceeded)
