Only limit the net tick rate if it's lower than the engine tick rate.

diff --git a/Engine/Source/Runtime/Engine/Private/NetConnection.cpp b/Engine/Source/Runtime/Engine/Private/NetConnection.cpp
--- a/Engine/Source/Runtime/Engine/Private/NetConnection.cpp
+++ b/Engine/Source/Runtime/Engine/Private/NetConnection.cpp
@@ -35,6 +35,7 @@
 #include "UObject/UObjectIterator.h"
 #include "Net/NetworkGranularMemoryLogging.h"
 #include "SocketSubsystem.h"
+#include "Math/NumericLimits.h"
 
 static TAutoConsoleVariable<int32> CVarPingExcludeFrameTime( TEXT( "net.PingExcludeFrameTime" ), 0, TEXT( "Calculate RTT time between NIC's of server and client." ) );
 
@@ -2872,9 +2873,16 @@ void UNetConnection::Tick()
 
 	FrameTime = CurrentRealtimeSeconds - LastTime;
 	const int32 MaxNetTickRate = Driver->MaxNetTickRate;
-	const float MaxNetTickRateFloat = MaxNetTickRate > 0 ? float(MaxNetTickRate) : FLT_MAX;
-	const float DesiredTickRate = FMath::Clamp(GEngine->GetMaxTickRate(0.0f, false), 0.0f, MaxNetTickRateFloat);
-	if (!InternalAck && MaxNetTickRate > 0 && DesiredTickRate > 0.0f)
+	float EngineTickRate = GEngine->GetMaxTickRate(0.0f, false);
+	// We want to make sure the DesiredTickRate stays at <= 0 if there's no tick rate limiting of any kind, since it's used later in the function for bandwidth limiting.
+	if (MaxNetTickRate > 0 && EngineTickRate <= 0.0f)
+	{
+		EngineTickRate = MAX_flt;
+	}
+	const float MaxNetTickRateFloat = MaxNetTickRate > 0 ? float(MaxNetTickRate) : MAX_flt;
+	const float DesiredTickRate = FMath::Clamp(EngineTickRate, 0.0f, MaxNetTickRateFloat);
+	// Apply net tick rate limiting if the desired net tick rate is strictly less than the engine tick rate.
+	if (!InternalAck && MaxNetTickRateFloat < EngineTickRate && DesiredTickRate > 0.0f)
 	{
 		const float MinNetFrameTime = 1.0f/DesiredTickRate;
 		if (FrameTime < MinNetFrameTime)
