Fix for delay in battle bus music

Move the listener before checking if virtual loops should realize: UpdateVirtualLoops() checks virtual looping sounds against the closest listener to determine if the sound should realize (become audible).  Before this change, we would check if the listener is about to move a large distance, call UpdateVirtualLoops(), and THEN move the listener to it's new position.

What we want to do here is move the listener before we call UpdateVirtualLoops, so the sounds are checked against the listener's new position, instead of the old position.

diff --git a/Engine/Source/Runtime/Engine/Private/AudioDevice.cpp b/Engine/Source/Runtime/Engine/Private/AudioDevice.cpp
--- a/Engine/Source/Runtime/Engine/Private/AudioDevice.cpp
+++ b/Engine/Source/Runtime/Engine/Private/AudioDevice.cpp
@@ -2929,16 +2929,17 @@ void FAudioDevice::SetListener(UWorld* World, const int32 InViewportIndex, const
 			logOrEnsureNanError(TEXT("FAudioDevice::SetListener has detected a NaN in Listener Velocity"));
 		}
 #endif
+		const bool bShouldListenerForceUpdate = FAudioVirtualLoop::ShouldListenerMoveForceUpdate(Listener.Transform, ListenerTransformCopy);
 
-		if (FAudioVirtualLoop::ShouldListenerMoveForceUpdate(Listener.Transform, ListenerTransformCopy))
+		Listener.WorldID = WorldID;
+		Listener.Transform = ListenerTransformCopy;
+
+		if (bShouldListenerForceUpdate)
 		{
 			const bool bForceUpdate = true;
 			UpdateVirtualLoops(bForceUpdate);
 		}
 
-		Listener.WorldID = WorldID;
-		Listener.Transform = ListenerTransformCopy;
-
 		OnListenerUpdated(AudioThreadListeners);
 
 	}, GET_STATID(STAT_AudioSetListener));
