Change display name of some attributes in LightComponentBase to specify they affect ray tracing only.

diff --git a/Engine/Source/Runtime/Engine/Classes/Components/LightComponentBase.h b/Engine/Source/Runtime/Engine/Classes/Components/LightComponentBase.h
--- a/Engine/Source/Runtime/Engine/Classes/Components/LightComponentBase.h
+++ b/Engine/Source/Runtime/Engine/Classes/Components/LightComponentBase.h
@@ -84,15 +84,15 @@ class ENGINE_API ULightComponentBase : public USceneComponent
 	uint32 bCastDeepShadow : 1;
 
 	/** Whether the light shadows are computed with shadow-mapping or ray-tracing (when available). */
-	UPROPERTY(EditAnywhere, BlueprintReadOnly, Category = Light, meta = (DisplayName = "Cast Ray Tracing Shadow"), AdvancedDisplay)
+	UPROPERTY(EditAnywhere, BlueprintReadOnly, Category = Light, meta = (DisplayName = "Cast Ray Tracing Shadows"), AdvancedDisplay)
 	uint32 bCastRaytracedShadow : 1;
 
 	/** Whether the light affects objects in reflections, when ray-traced reflection is enabled. */
-	UPROPERTY(EditAnywhere, BlueprintReadOnly, Category = Light, AdvancedDisplay)
+	UPROPERTY(EditAnywhere, BlueprintReadOnly, Category = Light, AdvancedDisplay, meta = (DisplayName = "Affect Ray Tracing Reflections"))
 	uint32 bAffectReflection : 1;
 
 	/** Whether the light affects global illumination, when ray-traced global illumination is enabled. */
-	UPROPERTY(EditAnywhere, BlueprintReadOnly, Category = Light, AdvancedDisplay)
+	UPROPERTY(EditAnywhere, BlueprintReadOnly, Category = Light, AdvancedDisplay, meta = (DisplayName = "Affect Ray Tracing Global Illumination"))
 	uint32 bAffectGlobalIllumination : 1;
 
 	/** 
