Make changing section user data (cast shadow, recompute, enabled...) call posteditchange. This is like the old workflow.
#rnx

diff --git a/Engine/Plugins/Experimental/SkeletalReduction/Source/Private/SkeletalMeshReductionPlugin.cpp b/Engine/Plugins/Experimental/SkeletalReduction/Source/Private/SkeletalMeshReductionPlugin.cpp
--- a/Engine/Plugins/Experimental/SkeletalReduction/Source/Private/SkeletalMeshReductionPlugin.cpp
+++ b/Engine/Plugins/Experimental/SkeletalReduction/Source/Private/SkeletalMeshReductionPlugin.cpp
@@ -62,6 +62,11 @@ public:
 	}
 
 	virtual bool IsReductionActive(const FSkeletalMeshOptimizationSettings &ReductionSettings) const
+	{
+		return IsReductionActive(ReductionSettings, 0, 0);
+	}
+
+	virtual bool IsReductionActive(const struct FSkeletalMeshOptimizationSettings &ReductionSettings, uint32 NumVertices, uint32 NumTriangles) const
 	{
 		float Threshold_One = (1.0f - KINDA_SMALL_NUMBER);
 		float Threshold_Zero = (0.0f + KINDA_SMALL_NUMBER);
@@ -84,10 +89,18 @@ public:
 			break;
 			//Absolute count is consider has being always reduced
 			case SkeletalMeshTerminationCriterion::SMTC_AbsNumOfVerts:
+			{
+				return ReductionSettings.MaxNumOfVerts < NumVertices;
+			}
+			break;
 			case SkeletalMeshTerminationCriterion::SMTC_AbsNumOfTriangles:
+			{
+				return ReductionSettings.MaxNumOfTriangles < NumTriangles;
+			}
+			break;
 			case SkeletalMeshTerminationCriterion::SMTC_AbsTriangleOrVert:
 			{
-				return true;
+				return ReductionSettings.MaxNumOfVerts < NumVertices || ReductionSettings.MaxNumOfTriangles < NumTriangles;
 			}
 			break;
 		}
diff --git a/Engine/Source/Developer/MeshReductionInterface/Public/IMeshReductionInterfaces.h b/Engine/Source/Developer/MeshReductionInterface/Public/IMeshReductionInterfaces.h
--- a/Engine/Source/Developer/MeshReductionInterface/Public/IMeshReductionInterfaces.h
+++ b/Engine/Source/Developer/MeshReductionInterface/Public/IMeshReductionInterfaces.h
@@ -56,6 +56,7 @@ public:
 	*/
 	virtual bool IsReductionActive(const struct FMeshReductionSettings &ReductionSettings) const = 0;
 	virtual bool IsReductionActive(const struct FSkeletalMeshOptimizationSettings &ReductionSettings) const = 0;
+	virtual bool IsReductionActive(const struct FSkeletalMeshOptimizationSettings &ReductionSettings, uint32 NumVertices, uint32 NumTriangles) const = 0;
 };
 
 DECLARE_DELEGATE_ThreeParams(FProxyCompleteDelegate, struct FMeshDescription&, struct FFlattenMaterial&, const FGuid);
diff --git a/Engine/Source/Developer/MeshSimplifier/Private/QuadricMeshReduction.cpp b/Engine/Source/Developer/MeshSimplifier/Private/QuadricMeshReduction.cpp
--- a/Engine/Source/Developer/MeshSimplifier/Private/QuadricMeshReduction.cpp
+++ b/Engine/Source/Developer/MeshSimplifier/Private/QuadricMeshReduction.cpp
@@ -621,6 +621,11 @@ public:
 		return false;
 	}
 
+	virtual bool IsReductionActive(const struct FSkeletalMeshOptimizationSettings &ReductionSettings, uint32 NumVertices, uint32 NumTriangles) const
+	{
+		return false;
+	}
+
 	virtual ~FQuadricSimplifierMeshReduction() {}
 
 	static FQuadricSimplifierMeshReduction* Create()
diff --git a/Engine/Source/Developer/SkeletalMeshUtilitiesCommon/Private/LODUtilities.cpp b/Engine/Source/Developer/SkeletalMeshUtilitiesCommon/Private/LODUtilities.cpp
--- a/Engine/Source/Developer/SkeletalMeshUtilitiesCommon/Private/LODUtilities.cpp
+++ b/Engine/Source/Developer/SkeletalMeshUtilitiesCommon/Private/LODUtilities.cpp
@@ -1126,6 +1126,8 @@ void FLODUtilities::RestoreSkeletalMeshLODImportedData(USkeletalMesh* SkeletalMe
 		FSkeletalMeshUpdateContext UpdateContext;
 		UpdateContext.SkeletalMesh = SkeletalMesh;
 
+		ImportedBaseLODModel.UpdateChunkedSectionInfo();
+
 		//Copy the SkeletalMeshLODModel
 		SkeletalMesh->GetImportedModel()->LODModels[LodIndex] = ImportedBaseLODModel;
 		//Copy the morph target deltas
diff --git a/Engine/Source/Editor/Persona/Private/PersonaMeshDetails.cpp b/Engine/Source/Editor/Persona/Private/PersonaMeshDetails.cpp
--- a/Engine/Source/Editor/Persona/Private/PersonaMeshDetails.cpp
+++ b/Engine/Source/Editor/Persona/Private/PersonaMeshDetails.cpp
@@ -2921,6 +2921,7 @@ FReply FPersonaMeshDetails::ApplyLODChanges(int32 LODIndex)
 
 		if (!LODModel.RawSkeletalMeshBulkData.IsBuildDataAvailable())
 		{
+			SkelMesh->InvalidateDeriveDataCacheGUID();
 			RegenerateLOD(LODIndex);
 		}
 		else
@@ -4234,36 +4235,40 @@ void FPersonaMeshDetails::OnSectionEnabledChanged(int32 LodIndex, int32 SectionI
 
 				if(Section.bDisabled != !bEnable)
 				{
-					FScopedTransaction Transaction(LOCTEXT("ChangeSectionEnabled", "Set section disabled flag."));
+					FScopedSuspendAlternateSkinWeightPreview ScopedSuspendAlternateSkinnWeightPreview(SkeletalMeshPtr.Get());
+					{
+						FScopedSkeletalMeshPostEditChange ScopedPostEditChange(SkeletalMeshPtr.Get());
+						FScopedTransaction Transaction(LOCTEXT("ChangeSectionEnabled", "Set section disabled flag."));
 
-					SkeletalMeshPtr->Modify();
-					SkeletalMeshPtr->PreEditChange(nullptr);
+						SkeletalMeshPtr->Modify();
+						SkeletalMeshPtr->PreEditChange(nullptr);
 
-					Section.bDisabled = !bEnable;
-					for (int32 AfterSectionIndex = SectionIndex + 1; AfterSectionIndex < LodModel.Sections.Num(); ++AfterSectionIndex)
-					{
-						if (LodModel.Sections[AfterSectionIndex].ChunkedParentSectionIndex == SectionIndex)
+						Section.bDisabled = !bEnable;
+						for (int32 AfterSectionIndex = SectionIndex + 1; AfterSectionIndex < LodModel.Sections.Num(); ++AfterSectionIndex)
 						{
-							LodModel.Sections[AfterSectionIndex].bDisabled = Section.bDisabled;
-						}
-						else
-						{
-							break;
+							if (LodModel.Sections[AfterSectionIndex].ChunkedParentSectionIndex == SectionIndex)
+							{
+								LodModel.Sections[AfterSectionIndex].bDisabled = Section.bDisabled;
+							}
+							else
+							{
+								break;
+							}
 						}
-					}
-					//We display only the parent chunk
-					check(Section.ChunkedParentSectionIndex == INDEX_NONE);
+						//We display only the parent chunk
+						check(Section.ChunkedParentSectionIndex == INDEX_NONE);
 
-					SetSkelMeshSourceSectionUserData(LodModel, SectionIndex, Section.OriginalDataSectionIndex);
+						SetSkelMeshSourceSectionUserData(LodModel, SectionIndex, Section.OriginalDataSectionIndex);
 
-					// Disable highlight and isolate flags
-					UDebugSkelMeshComponent * MeshComponent = GetPersonaToolkit()->GetPreviewScene()->GetPreviewMeshComponent();
-					if(MeshComponent)
-					{
-						MeshComponent->SetSelectedEditorSection(INDEX_NONE);
-						MeshComponent->SetSelectedEditorMaterial(INDEX_NONE);
-						MeshComponent->SetMaterialPreview(INDEX_NONE);
-						MeshComponent->SetSectionPreview(INDEX_NONE);
+						// Disable highlight and isolate flags
+						UDebugSkelMeshComponent * MeshComponent = GetPersonaToolkit()->GetPreviewScene()->GetPreviewMeshComponent();
+						if (MeshComponent)
+						{
+							MeshComponent->SetSelectedEditorSection(INDEX_NONE);
+							MeshComponent->SetSelectedEditorMaterial(INDEX_NONE);
+							MeshComponent->SetMaterialPreview(INDEX_NONE);
+							MeshComponent->SetSectionPreview(INDEX_NONE);
+						}
 					}
 				}
 			}
@@ -4703,12 +4708,16 @@ void FPersonaMeshDetails::OnSectionShadowCastingChanged(ECheckBoxState NewState,
 	//Update Original PolygonGroup
 	auto UpdatePolygonGroupCastShadow = [&Mesh, &LODModel, &Section, &SectionIndex](bool bCastShadow)
 	{
-		Section.bCastShadow = bCastShadow;
-		//We change only the parent chunk data
-		check(Section.ChunkedParentSectionIndex == INDEX_NONE);
+		FScopedSuspendAlternateSkinWeightPreview ScopedSuspendAlternateSkinnWeightPreview(Mesh);
+		{
+			FScopedSkeletalMeshPostEditChange ScopedPostEditChange(Mesh);
+			Section.bCastShadow = bCastShadow;
+			//We change only the parent chunk data
+			check(Section.ChunkedParentSectionIndex == INDEX_NONE);
 
-		//The post edit change will kick a build
-		SetSkelMeshSourceSectionUserData(LODModel, SectionIndex, Section.OriginalDataSectionIndex);
+			//The post edit change will kick a build
+			SetSkelMeshSourceSectionUserData(LODModel, SectionIndex, Section.OriginalDataSectionIndex);
+		}
 	};
 
 	if (NewState == ECheckBoxState::Checked)
@@ -4770,21 +4779,25 @@ void FPersonaMeshDetails::OnSectionRecomputeTangentChanged(ECheckBoxState NewSta
 	//Update Original PolygonGroup
 	auto UpdatePolygonGroupRecomputeTangent = [&Mesh, &LODModel, &Section, &SectionIndex](bool bRecomputeTangent)
 	{
-		Section.bRecomputeTangent = bRecomputeTangent;
-		for (int32 AfterSectionIndex = SectionIndex + 1; AfterSectionIndex < LODModel.Sections.Num(); ++AfterSectionIndex)
+		FScopedSuspendAlternateSkinWeightPreview ScopedSuspendAlternateSkinnWeightPreview(Mesh);
 		{
-			if (LODModel.Sections[AfterSectionIndex].ChunkedParentSectionIndex == SectionIndex)
-			{
-				LODModel.Sections[AfterSectionIndex].bRecomputeTangent = bRecomputeTangent;
-			}
-			else
+			FScopedSkeletalMeshPostEditChange ScopedPostEditChange(Mesh);
+			Section.bRecomputeTangent = bRecomputeTangent;
+			for (int32 AfterSectionIndex = SectionIndex + 1; AfterSectionIndex < LODModel.Sections.Num(); ++AfterSectionIndex)
 			{
-				break;
+				if (LODModel.Sections[AfterSectionIndex].ChunkedParentSectionIndex == SectionIndex)
+				{
+					LODModel.Sections[AfterSectionIndex].bRecomputeTangent = bRecomputeTangent;
+				}
+				else
+				{
+					break;
+				}
 			}
+			//We display only the parent chunk
+			check(Section.ChunkedParentSectionIndex == INDEX_NONE);
+			SetSkelMeshSourceSectionUserData(LODModel, SectionIndex, Section.OriginalDataSectionIndex);
 		}
-		//We display only the parent chunk
-		check(Section.ChunkedParentSectionIndex == INDEX_NONE);
-		SetSkelMeshSourceSectionUserData(LODModel, SectionIndex, Section.OriginalDataSectionIndex);
 	};
 
 	if (NewState == ECheckBoxState::Checked)
diff --git a/Engine/Source/Runtime/Engine/Private/Rendering/SkeletalMeshLODImporterData.cpp b/Engine/Source/Runtime/Engine/Private/Rendering/SkeletalMeshLODImporterData.cpp
--- a/Engine/Source/Runtime/Engine/Private/Rendering/SkeletalMeshLODImporterData.cpp
+++ b/Engine/Source/Runtime/Engine/Private/Rendering/SkeletalMeshLODImporterData.cpp
@@ -477,6 +477,8 @@ void FReductionBaseSkeletalMeshBulkData::SaveReductionData(FSkeletalMeshLODModel
 	bUseSerializeLoadingCustomVersion = false;
 	FReductionSkeletalMeshData ReductionSkeletalMeshData(BaseLODModel, BaseLODMorphTargetData, Owner);
 
+	CacheGeometryInfo(BaseLODModel);
+
 	//Clear the bulk data before writing it
 	BulkData.RemoveBulkData();
 
@@ -509,11 +511,43 @@ void FReductionBaseSkeletalMeshBulkData::LoadReductionData(FSkeletalMeshLODModel
 			Ar.SetCustomVersions(SerializeLoadingCustomVersionContainer);
 
 			Ar << ReductionSkeletalMeshData;
+
+			CacheGeometryInfo(BaseLODModel);
 		}
 		// Unlock the bulk data
 	}
 }
 
+void FReductionBaseSkeletalMeshBulkData::CacheGeometryInfo(const FSkeletalMeshLODModel& SourceLODModel)
+{
+	CacheLODVertexNumber = 0;
+	CacheLODTriNumber = 0;
+	for (int32 SectionIndex = 0; SectionIndex < SourceLODModel.Sections.Num(); ++SectionIndex)
+	{
+		const FSkelMeshSection& Section = SourceLODModel.Sections[SectionIndex];
+
+		if (!Section.bDisabled)
+		{
+			//Make sure the count fit in a uint32
+			CacheLODVertexNumber += Section.NumVertices < 0 ? 0 : Section.NumVertices;
+			CacheLODTriNumber += Section.NumTriangles < 0 ? 0 : Section.NumTriangles;
+		}
+	}
+}
+
+void FReductionBaseSkeletalMeshBulkData::GetGeometryInfo(uint32& LODVertexNumber, uint32& LODTriNumber, UObject* Owner)
+{
+	if (!IsEmpty() && (CacheLODVertexNumber == MAX_uint32 || CacheLODTriNumber == MAX_uint32))
+	{
+		FSkeletalMeshLODModel ReductionSourceLODModel;
+		TMap<FString, TArray<FMorphTargetDelta>> TempLODMorphTargetData;
+		LoadReductionData(ReductionSourceLODModel, TempLODMorphTargetData, Owner);
+		CacheGeometryInfo(ReductionSourceLODModel);
+	}
+	LODVertexNumber = CacheLODVertexNumber;
+	LODTriNumber = CacheLODTriNumber;
+}
+
 /*------------------------------------------------------------------------------
 FRawSkeletalMeshBulkData
 ------------------------------------------------------------------------------*/
diff --git a/Engine/Source/Runtime/Engine/Private/SkeletalMesh.cpp b/Engine/Source/Runtime/Engine/Private/SkeletalMesh.cpp
--- a/Engine/Source/Runtime/Engine/Private/SkeletalMesh.cpp
+++ b/Engine/Source/Runtime/Engine/Private/SkeletalMesh.cpp
@@ -58,6 +58,7 @@
 
 #if WITH_EDITOR
 #include "Rendering/SkeletalMeshModel.h"
+#include "Rendering/SkeletalMeshLODImporterData.h"
 #include "MeshUtilities.h"
 
 #if WITH_APEX_CLOTHING
@@ -471,7 +472,39 @@ bool USkeletalMesh::IsReductionActive(int32 LODIndex) const
 {
 	FSkeletalMeshOptimizationSettings ReductionSettings = GetReductionSettings(LODIndex);
 	IMeshReduction* ReductionModule = FModuleManager::Get().LoadModuleChecked<IMeshReductionManagerModule>("MeshReductionInterface").GetSkeletalMeshReductionInterface();
-	return ReductionModule->IsReductionActive(ReductionSettings);
+	uint32 LODVertexNumber = MAX_uint32;
+	uint32 LODTriNumber = MAX_uint32;
+	const FSkeletalMeshLODInfo* LODInfoPtr = GetLODInfo(LODIndex);
+	bool bLODHasBeenSimplified = LODInfoPtr && LODInfoPtr->bHasBeenSimplified;
+	if (GetImportedModel() && GetImportedModel()->LODModels.IsValidIndex(LODIndex))
+	{
+		if (!bLODHasBeenSimplified)
+		{
+			LODVertexNumber = 0;
+			LODTriNumber = 0;
+			const FSkeletalMeshLODModel& LODModel = GetImportedModel()->LODModels[LODIndex];
+			//We can take the vertices and triangles count from the source model
+			for (int32 SectionIndex = 0; SectionIndex < LODModel.Sections.Num(); ++SectionIndex)
+			{
+				const FSkelMeshSection& Section = LODModel.Sections[SectionIndex];
+
+				if (!Section.bDisabled)
+				{
+					//Make sure the count fit in a uint32
+					LODVertexNumber += Section.NumVertices < 0 ? 0 : Section.NumVertices;
+					LODTriNumber += Section.NumTriangles < 0 ? 0 : Section.NumTriangles;
+				}
+			}
+		}
+		else if(GetImportedModel()->OriginalReductionSourceMeshData.IsValidIndex(LODIndex)
+			&& !GetImportedModel()->OriginalReductionSourceMeshData[LODIndex]->IsEmpty())
+		{
+			//In this case we have to use the stored reduction source data to know how many vertices/triangles we have before the reduction
+			USkeletalMesh* MutableSkeletalMesh = const_cast<USkeletalMesh*>(this);
+			GetImportedModel()->OriginalReductionSourceMeshData[LODIndex]->GetGeometryInfo(LODVertexNumber, LODTriNumber, MutableSkeletalMesh);
+		}
+	}
+	return ReductionModule->IsReductionActive(ReductionSettings, LODVertexNumber, LODTriNumber);
 }
 
 /* Get a copy of the reduction settings for a specified LOD index. */
@@ -2114,11 +2147,6 @@ void USkeletalMesh::PostLoad()
 				FSkeletalMeshLODModel& ThisLODModel = ImportedModel->LODModels[LodIndex];
 				FSkeletalMeshLODInfo* ThisLODInfo = GetLODInfo(LodIndex);
 				check(ThisLODInfo);
-				int32 LODModelSectionNum = ThisLODModel.Sections.Num();
-				//Fill the ChunkedParentSectionIndex data, we assume that every section using the same material are chunked
-				int32 LastMaterialIndex = INDEX_NONE;
-				int32 CurrentParentChunkIndex = INDEX_NONE;
-				int32 OriginalIndex = 0;
 				
 				//Reset the reduction setting to a non active state if the asset has active reduction but have no RawSkeletalMeshBulkData (we cannot reduce it)
 				if (IsReductionActive(LodIndex) && !ThisLODInfo->bHasBeenSimplified && ThisLODModel.RawSkeletalMeshBulkData.IsEmpty())
@@ -2150,42 +2178,7 @@ void USkeletalMesh::PostLoad()
 						}
 					}
 				}
-
-				for (int32 LODModelSectionIndex = 0; LODModelSectionIndex < LODModelSectionNum; ++LODModelSectionIndex)
-				{
-					//If we have cloth on a chunked section we treat the chunked section has a parent section (this is to get the same result has before the refactor)
-					if (ThisLODModel.Sections[LODModelSectionIndex].MaterialIndex == LastMaterialIndex && !ThisLODModel.Sections[LODModelSectionIndex].ClothingData.AssetGuid.IsValid())
-					{
-						ThisLODModel.Sections[LODModelSectionIndex].ChunkedParentSectionIndex = CurrentParentChunkIndex;
-						ThisLODModel.Sections[LODModelSectionIndex].OriginalDataSectionIndex = ThisLODModel.Sections[CurrentParentChunkIndex].OriginalDataSectionIndex;
-						//In case of a child section that was BONE chunked ensure it has the same setting has the original section
-						FSkelMeshSourceSectionUserData& SectionUserData = ThisLODModel.UserSectionsData.FindOrAdd(ThisLODModel.Sections[LODModelSectionIndex].OriginalDataSectionIndex);
-						ThisLODModel.Sections[LODModelSectionIndex].bDisabled = SectionUserData.bDisabled;
-						ThisLODModel.Sections[LODModelSectionIndex].bCastShadow = SectionUserData.bCastShadow;
-						ThisLODModel.Sections[LODModelSectionIndex].bRecomputeTangent = SectionUserData.bRecomputeTangent;
-						ThisLODModel.Sections[LODModelSectionIndex].GenerateUpToLodIndex = SectionUserData.GenerateUpToLodIndex;
-						//Chunked section cannot have cloth, a cloth section will be a parent section
-						ThisLODModel.Sections[LODModelSectionIndex].CorrespondClothAssetIndex = INDEX_NONE;
-						ThisLODModel.Sections[LODModelSectionIndex].ClothingData.AssetGuid = FGuid();
-						ThisLODModel.Sections[LODModelSectionIndex].ClothingData.AssetLodIndex = INDEX_NONE;
-					}
-					else
-					{
-						CurrentParentChunkIndex = LODModelSectionIndex;
-						FSkelMeshSourceSectionUserData& SectionUserData = ThisLODModel.UserSectionsData.FindOrAdd(OriginalIndex);
-						SectionUserData.bDisabled = ThisLODModel.Sections[LODModelSectionIndex].bDisabled;
-						SectionUserData.bCastShadow = ThisLODModel.Sections[LODModelSectionIndex].bCastShadow;
-						SectionUserData.bRecomputeTangent = ThisLODModel.Sections[LODModelSectionIndex].bRecomputeTangent;
-						SectionUserData.GenerateUpToLodIndex = ThisLODModel.Sections[LODModelSectionIndex].GenerateUpToLodIndex;
-						SectionUserData.CorrespondClothAssetIndex = ThisLODModel.Sections[LODModelSectionIndex].CorrespondClothAssetIndex;
-						SectionUserData.ClothingData.AssetGuid = ThisLODModel.Sections[LODModelSectionIndex].ClothingData.AssetGuid;
-						SectionUserData.ClothingData.AssetLodIndex = ThisLODModel.Sections[LODModelSectionIndex].ClothingData.AssetLodIndex;
-
-						ThisLODModel.Sections[LODModelSectionIndex].OriginalDataSectionIndex = OriginalIndex++;
-						ThisLODModel.Sections[LODModelSectionIndex].ChunkedParentSectionIndex = INDEX_NONE;
-					}
-					LastMaterialIndex = ThisLODModel.Sections[LODModelSectionIndex].MaterialIndex;
-				}
+				ThisLODModel.UpdateChunkedSectionInfo();
 			}
 		}
 
diff --git a/Engine/Source/Runtime/Engine/Private/SkeletalMeshLODModel.cpp b/Engine/Source/Runtime/Engine/Private/SkeletalMeshLODModel.cpp
--- a/Engine/Source/Runtime/Engine/Private/SkeletalMeshLODModel.cpp
+++ b/Engine/Source/Runtime/Engine/Private/SkeletalMeshLODModel.cpp
@@ -1023,4 +1023,55 @@ FString FSkeletalMeshLODModel::GetLODModelDeriveDataKey() const
 	return KeySuffix;
 }
 
+void FSkeletalMeshLODModel::UpdateChunkedSectionInfo()
+{
+	//Look if the data is pre skeletalmesh build refactor
+	if (RawSkeletalMeshBulkData.IsBuildDataAvailable())
+	{
+		//If the data is up to date just return and do nothing
+		return;
+	}
+	int32 LODModelSectionNum = Sections.Num();
+	//Fill the ChunkedParentSectionIndex data, we assume that every section using the same material are chunked
+	int32 LastMaterialIndex = INDEX_NONE;
+	int32 CurrentParentChunkIndex = INDEX_NONE;
+	int32 OriginalIndex = 0;
+
+	for (int32 LODModelSectionIndex = 0; LODModelSectionIndex < LODModelSectionNum; ++LODModelSectionIndex)
+	{
+		//If we have cloth on a chunked section we treat the chunked section has a parent section (this is to get the same result has before the refactor)
+		if (Sections[LODModelSectionIndex].MaterialIndex == LastMaterialIndex && !Sections[LODModelSectionIndex].ClothingData.AssetGuid.IsValid())
+		{
+			Sections[LODModelSectionIndex].ChunkedParentSectionIndex = CurrentParentChunkIndex;
+			Sections[LODModelSectionIndex].OriginalDataSectionIndex = Sections[CurrentParentChunkIndex].OriginalDataSectionIndex;
+			//In case of a child section that was BONE chunked ensure it has the same setting has the original section
+			FSkelMeshSourceSectionUserData& SectionUserData = UserSectionsData.FindOrAdd(Sections[LODModelSectionIndex].OriginalDataSectionIndex);
+			Sections[LODModelSectionIndex].bDisabled = SectionUserData.bDisabled;
+			Sections[LODModelSectionIndex].bCastShadow = SectionUserData.bCastShadow;
+			Sections[LODModelSectionIndex].bRecomputeTangent = SectionUserData.bRecomputeTangent;
+			Sections[LODModelSectionIndex].GenerateUpToLodIndex = SectionUserData.GenerateUpToLodIndex;
+			//Chunked section cannot have cloth, a cloth section will be a parent section
+			Sections[LODModelSectionIndex].CorrespondClothAssetIndex = INDEX_NONE;
+			Sections[LODModelSectionIndex].ClothingData.AssetGuid = FGuid();
+			Sections[LODModelSectionIndex].ClothingData.AssetLodIndex = INDEX_NONE;
+		}
+		else
+		{
+			CurrentParentChunkIndex = LODModelSectionIndex;
+			FSkelMeshSourceSectionUserData& SectionUserData = UserSectionsData.FindOrAdd(OriginalIndex);
+			SectionUserData.bDisabled = Sections[LODModelSectionIndex].bDisabled;
+			SectionUserData.bCastShadow = Sections[LODModelSectionIndex].bCastShadow;
+			SectionUserData.bRecomputeTangent = Sections[LODModelSectionIndex].bRecomputeTangent;
+			SectionUserData.GenerateUpToLodIndex = Sections[LODModelSectionIndex].GenerateUpToLodIndex;
+			SectionUserData.CorrespondClothAssetIndex = Sections[LODModelSectionIndex].CorrespondClothAssetIndex;
+			SectionUserData.ClothingData.AssetGuid = Sections[LODModelSectionIndex].ClothingData.AssetGuid;
+			SectionUserData.ClothingData.AssetLodIndex = Sections[LODModelSectionIndex].ClothingData.AssetLodIndex;
+
+			Sections[LODModelSectionIndex].OriginalDataSectionIndex = OriginalIndex++;
+			Sections[LODModelSectionIndex].ChunkedParentSectionIndex = INDEX_NONE;
+		}
+		LastMaterialIndex = Sections[LODModelSectionIndex].MaterialIndex;
+	}
+}
+
 #endif // WITH_EDITOR
diff --git a/Engine/Source/Runtime/Engine/Public/Rendering/SkeletalMeshLODImporterData.h b/Engine/Source/Runtime/Engine/Public/Rendering/SkeletalMeshLODImporterData.h
--- a/Engine/Source/Runtime/Engine/Public/Rendering/SkeletalMeshLODImporterData.h
+++ b/Engine/Source/Runtime/Engine/Public/Rendering/SkeletalMeshLODImporterData.h
@@ -414,6 +414,15 @@ struct FReductionBaseSkeletalMeshBulkData
 	//The custom version when this was load
 	FCustomVersionContainer SerializeLoadingCustomVersionContainer;
 	bool bUseSerializeLoadingCustomVersion = false;
+
+	uint32 CacheLODVertexNumber = MAX_uint32;
+	uint32 CacheLODTriNumber = MAX_uint32;
+
+	/*
+	 * Caching those value since this is a slow operation to load the bulk data to retrieve the original geometry information
+	 */
+	void CacheGeometryInfo(const FSkeletalMeshLODModel& SourceLODModel);
+
 public:
 	/** Default constructor. */
 	ENGINE_API FReductionBaseSkeletalMeshBulkData();
@@ -428,6 +437,9 @@ public:
 	/** Load the raw mesh from bulk data. */
 	ENGINE_API void LoadReductionData(FSkeletalMeshLODModel& BaseLODModel, TMap<FString, TArray<FMorphTargetDelta>>& BaseLODMorphTargetData, UObject* Owner);
 
+	/** Return the number of vertices and triangles store in this bulk data. */
+	ENGINE_API void GetGeometryInfo(uint32& LODVertexNumber, uint32& LODTriNumber, UObject* Owner);
+
 	ENGINE_API FByteBulkData& GetBulkData() { return BulkData; }
 
 	/** Empty the bulk data. */
diff --git a/Engine/Source/Runtime/Engine/Public/Rendering/SkeletalMeshLODModel.h b/Engine/Source/Runtime/Engine/Public/Rendering/SkeletalMeshLODModel.h
--- a/Engine/Source/Runtime/Engine/Public/Rendering/SkeletalMeshLODModel.h
+++ b/Engine/Source/Runtime/Engine/Public/Rendering/SkeletalMeshLODModel.h
@@ -388,6 +388,12 @@ public:
 	* Build a derive data key with the user section data (UserSectionsData) and the original bulk data
 	*/
 	ENGINE_API FString GetLODModelDeriveDataKey() const;
+
+	/**
+	* This function will update the chunked information for each section. Only old data before the 
+	* skeletal mesh build refactor should need to call this function.
+	*/
+	ENGINE_API void UpdateChunkedSectionInfo();
 };
 
 #endif // WITH_EDITOR
\ No newline at end of file
