Fix a crash on calling UMaterialInterface::SetForceMipLevelsToBeResident when texture streaming is disabled

diff --git a/Engine/Source/Runtime/Engine/Private/Materials/MaterialInterface.cpp b/Engine/Source/Runtime/Engine/Private/Materials/MaterialInterface.cpp
--- a/Engine/Source/Runtime/Engine/Private/Materials/MaterialInterface.cpp
+++ b/Engine/Source/Runtime/Engine/Private/Materials/MaterialInterface.cpp
@@ -228,7 +228,10 @@ void UMaterialInterface::SetForceMipLevelsToBeResident( bool OverrideForceMiplev
 				static IConsoleVariable* CVarAllowFastForceResident = IConsoleManager::Get().FindConsoleVariable(TEXT("r.Streaming.AllowFastForceResident"));
 
 				Texture->bIgnoreStreamingMipBias = CVarAllowFastForceResident && CVarAllowFastForceResident->GetInt();
-				IStreamingManager::Get().GetRenderAssetStreamingManager().FastForceFullyResident(Texture);
+				if (IStreamingManager::Get().IsRenderAssetStreamingEnabled())
+				{
+					IStreamingManager::Get().GetRenderAssetStreamingManager().FastForceFullyResident(Texture);
+				}
 			}
 		}
 	}
