Security fix for FRootMotionSourceGroup dynamic serialization of contained FRootMotionSources.
- We now ensure ScriptStructs are the objects being serialized through TCheckedObjPtr usage, and also restrict replication to derived classes of FRootMotionSource by checking the SuperStruct chain before allocation/NetSerialize of source
- Added error messages, and added removal of invalid sources from array
- If FRootMotionSourceGroup is replicated through a Server RPC, we need to prevent clients from sending us arbitrary ScriptStructs due to the allocation/reliance on GetCppStructOps  which could trigger a server crash for invalid structs. All provided sources are direct children of RMS and we never expect to have deep hierarchies so this should not be too costly

David.Ratti, Zak.Middleton, Michael.Noland, Fred.Kimberley
#rnx

diff --git a/Engine/Source/Runtime/Engine/Private/GameFramework/RootMotionSource.cpp b/Engine/Source/Runtime/Engine/Private/GameFramework/RootMotionSource.cpp
--- a/Engine/Source/Runtime/Engine/Private/GameFramework/RootMotionSource.cpp
+++ b/Engine/Source/Runtime/Engine/Private/GameFramework/RootMotionSource.cpp
@@ -1946,42 +1946,88 @@ bool FRootMotionSourceGroup::NetSerialize(FArchive& Ar, class UPackageMap* Map,
 
 	for (int32 i = 0; i < SourcesNum && !Ar.IsError(); ++i)
 	{
-		UScriptStruct* ScriptStruct = RootMotionSources[i].IsValid() ? RootMotionSources[i]->GetScriptStruct() : nullptr;
-		UScriptStruct* ScriptStructLocal = ScriptStruct;
+		TCheckedObjPtr<UScriptStruct> ScriptStruct = RootMotionSources[i].IsValid() ? RootMotionSources[i]->GetScriptStruct() : nullptr;
+		UScriptStruct* ScriptStructLocal = ScriptStruct.Get();
 		Ar << ScriptStruct;
 
-		if (ScriptStruct)
+		if (ScriptStruct.IsValid())
 		{
-			if (Ar.IsLoading())
+			// Restrict replication to derived classes of FRootMotionSource for security reasons:
+			// If FRootMotionSourceGroup is replicated through a Server RPC, we need to prevent clients from sending us
+			// arbitrary ScriptStructs due to the allocation/reliance on GetCppStructOps below which could trigger a server crash
+			// for invalid structs. All provided sources are direct children of RMS and we never expect to have deep hierarchies
+			// so this should not be too costly
+			bool bIsDerivedFromBaseRMS = false;
+			UStruct* CurrentSuperStruct = ScriptStruct->GetSuperStruct();
+			while (CurrentSuperStruct)
 			{
-				if (RootMotionSources[i].IsValid() && ScriptStructLocal == ScriptStruct)
+				if (CurrentSuperStruct == FRootMotionSource::StaticStruct())
 				{
-					// What we have locally is the same type as we're being serialized into, so we don't need to
-					// reallocate - just use existing structure
+					bIsDerivedFromBaseRMS = true;
+					break;
 				}
-				else
+				CurrentSuperStruct = CurrentSuperStruct->GetSuperStruct();
+			}
+
+			if (bIsDerivedFromBaseRMS)
+			{
+				if (Ar.IsLoading())
 				{
-					// For now, just reset/reallocate the data when loading.
-					// Longer term if we want to generalize this and use it for property replication, we should support
-					// only reallocating when necessary
-					FRootMotionSource* NewSource = (FRootMotionSource*)FMemory::Malloc(ScriptStruct->GetCppStructOps()->GetSize());
-					ScriptStruct->InitializeStruct(NewSource);
+					if (RootMotionSources[i].IsValid() && ScriptStructLocal == ScriptStruct.Get())
+					{
+						// What we have locally is the same type as we're being serialized into, so we don't need to
+						// reallocate - just use existing structure
+					}
+					else
+					{
+						// For now, just reset/reallocate the data when loading.
+						// Longer term if we want to generalize this and use it for property replication, we should support
+						// only reallocating when necessary
+						FRootMotionSource* NewSource = (FRootMotionSource*)FMemory::Malloc(ScriptStruct->GetCppStructOps()->GetSize());
+						ScriptStruct->InitializeStruct(NewSource);
 
-					RootMotionSources[i] = TSharedPtr<FRootMotionSource>(NewSource);
+						RootMotionSources[i] = TSharedPtr<FRootMotionSource>(NewSource);
+					}
 				}
-			}
 
-			void* ContainerPtr = RootMotionSources[i].Get();
+				void* ContainerPtr = RootMotionSources[i].Get();
 
-			if (ScriptStruct->StructFlags & STRUCT_NetSerializeNative)
-			{
-				ScriptStruct->GetCppStructOps()->NetSerialize(Ar, Map, bOutSuccess, RootMotionSources[i].Get());
+				if (ScriptStruct->StructFlags & STRUCT_NetSerializeNative)
+				{
+					ScriptStruct->GetCppStructOps()->NetSerialize(Ar, Map, bOutSuccess, RootMotionSources[i].Get());
+				}
+				else
+				{
+					checkf(false, TEXT("Serializing RootMotionSource without NetSerializeNative - not supported!"));
+				}
 			}
 			else
 			{
-				checkf(false, TEXT("Serializing RootMotionSource without NetSerializeNative - not supported!"));
+				UE_LOG(LogRootMotion, Error, TEXT("FRootMotionSourceGroup::NetSerialize: ScriptStruct not derived from FRootMotionSource attempted to serialize."));
+				Ar.SetError();
+				break;
 			}
 		}
+		else if (ScriptStruct.IsError())
+		{
+			UE_LOG(LogRootMotion, Error, TEXT("FRootMotionSourceGroup::NetSerialize: Invalid ScriptStruct serialized."));
+			Ar.SetError();
+			break;
+		}
+	}
+
+	if (Ar.IsError())
+	{
+		// Something bad happened, make sure to not return invalid shared ptrs
+		for (int32 i = RootMotionSources.Num() - 1; i >= 0; --i)
+		{
+			if (RootMotionSources[i].IsValid() == false)
+			{
+				RootMotionSources.RemoveAt(i);
+			}
+		}
+		bOutSuccess = false;
+		return false;
 	}
 
 	bOutSuccess = true;
