Add world context custom description to displayall text

diff --git a/Engine/Source/Runtime/Engine/Private/UnrealEngine.cpp b/Engine/Source/Runtime/Engine/Private/UnrealEngine.cpp
--- a/Engine/Source/Runtime/Engine/Private/UnrealEngine.cpp
+++ b/Engine/Source/Runtime/Engine/Private/UnrealEngine.cpp
@@ -10336,9 +10336,17 @@ void DrawStatsHUD( UWorld* World, FViewport* Viewport, FCanvas* Canvas, UCanvas*
 		{
 			UWorld* ServerWorld = *It;
 			if (ServerWorld->WorldType == EWorldType::PIE && ServerWorld->GetNetMode() == NM_DedicatedServer)
-			{				
+			{
+				FString WorldDesc = TEXT("[Server]");
+
+				FWorldContext& WorldContext = GEngine->GetWorldContextFromWorldChecked(ServerWorld);
+				if (!WorldContext.CustomDescription.IsEmpty())
+				{
+					WorldDesc += TEXT("[") + WorldContext.CustomDescription + TEXT("]");
+				}
+
 				Y += 20;
-				Y += Canvas->DrawShadowedString( X, Y, TEXT("[Server]"), GEngine->GetSmallFont(), FLinearColor::Gray);
+				Y += Canvas->DrawShadowedString( X, Y, *WorldDesc, GEngine->GetSmallFont(), FLinearColor::Gray);
 				DrawDebugPropertiesForWorld(ServerWorld, X, Y);
 			}
 		}
