Exporting FScopedActorRoleSwap for use by other modules

diff --git a/Engine/Source/Runtime/Engine/Public/Net/DataReplication.h b/Engine/Source/Runtime/Engine/Public/Net/DataReplication.h
--- a/Engine/Source/Runtime/Engine/Public/Net/DataReplication.h
+++ b/Engine/Source/Runtime/Engine/Public/Net/DataReplication.h
@@ -297,7 +297,7 @@ private:
 	TWeakObjectPtr<UObject> WeakObjectPtr;
 };
 
-class FScopedActorRoleSwap : public FNoncopyable
+class ENGINE_API FScopedActorRoleSwap : public FNoncopyable
 {
 public:
 	FScopedActorRoleSwap(AActor* InActor);
