Prevent divide by zero in beam update by taking the max of the size and a minimum value.

diff --git a/Engine/Source/Runtime/Engine/Private/Particles/ParticleBeamModules.cpp b/Engine/Source/Runtime/Engine/Private/Particles/ParticleBeamModules.cpp
--- a/Engine/Source/Runtime/Engine/Private/Particles/ParticleBeamModules.cpp
+++ b/Engine/Source/Runtime/Engine/Private/Particles/ParticleBeamModules.cpp
@@ -324,7 +324,7 @@ void UParticleModuleTypeDataBeam2::Update(FParticleEmitterInstance* Owner, int32
 
 		// Determine the step size, count, and travelled ratio
 		BeamData->Direction		= BeamData->TargetPoint - BeamData->SourcePoint;
-		float	FullMagnitude	= BeamData->Direction.Size();
+		float	FullMagnitude	= FMath::Max(BeamData->Direction.Size(), 0.001f);
 		BeamData->Direction.Normalize();
 
 		int32 InterpSteps = 0;
