Fix MaxDistance calculation being wrong for box attenuation
Add debug draw cvar (au.3dVisualize.Listeners 1) to help debug when listener is in and out of bounds for 3rd person (off by default when au.3dVisualize.Enabled 1)

diff --git a/Engine/Source/Runtime/Engine/Private/AudioDebug.cpp b/Engine/Source/Runtime/Engine/Private/AudioDebug.cpp
--- a/Engine/Source/Runtime/Engine/Private/AudioDebug.cpp
+++ b/Engine/Source/Runtime/Engine/Private/AudioDebug.cpp
@@ -40,6 +40,14 @@ FAutoConsoleVariableRef CVarAudioVisualizeActiveSoundsMode(
 	TEXT("0: Not Enabled, 1: Volume (Lin), 2: Volume (dB), 3: Distance, 4: Random color"),
 	ECVF_Default);
 
+static int32 ActiveSoundVisualizeListenersCVar = 0;
+FAutoConsoleVariableRef CVarAudioVisualizeListeners(
+	TEXT("au.3dVisualize.Listeners"),
+	ActiveSoundVisualizeListenersCVar,
+	TEXT("Whether or not listeners are visible when 3d visualize is enabled. \n")
+	TEXT("0: Not Enabled, 1: Enabled"), 
+	ECVF_Default);
+
 static int32 ActiveSoundVisualizeTypeCVar = 0;
 FAutoConsoleVariableRef CVarAudioVisualizeActiveSounds(
 	TEXT("au.3dVisualize.ActiveSounds.Type"),
@@ -502,6 +510,47 @@ void FAudioDebugger::DrawDebugInfo(const FActiveSound& ActiveSound, const TArray
 #endif // ENABLE_DRAW_DEBUG
 }
 
+void FAudioDebugger::DrawDebugInfo(const UWorld& World, const TArray<FListener>& Listeners, FVector& ListenerTransformOverride, bool bUseListenerTransformOverride)
+{
+#if ENABLE_DRAW_DEBUG
+	if (!ActiveSoundVisualizeListenersCVar)
+	{
+		return;
+	}
+
+	FAudioDeviceManager* DeviceManager = GEngine->GetAudioDeviceManager();
+	if (DeviceManager && DeviceManager->IsVisualizeDebug3dEnabled())
+	{
+		DECLARE_CYCLE_STAT(TEXT("FAudioThreadTask.DrawListenerDebugInfo"), STAT_AudioDrawListenerDebugInfo, STATGROUP_TaskGraphTasks);
+		const TWeakObjectPtr<UWorld> WorldPtr = &World;
+		for (const FListener& Listener : Listeners)
+		{
+			const FVector ListenerPosition = bUseListenerTransformOverride ? ListenerTransformOverride : Listener.Transform.GetTranslation();
+			const FVector ListenerFront = Listener.GetFront();
+			const FVector ListenerUp = Listener.GetUp();
+			const FVector ListenerRight = Listener.GetRight();
+			FAudioThread::RunCommandOnGameThread([WorldPtr, ListenerPosition, ListenerFront, ListenerUp, ListenerRight]()
+			{
+				if (WorldPtr.IsValid())
+				{
+					static float ArrowLength	 = 30.0f;
+					static float ArrowHeadSize	 = 8.0f;
+					static float Lifetime		 = -1.0f;
+					static float Thickness		 = 0.9f;
+					static uint8 DepthPriority	 = SDPG_Foreground;
+					static bool bPersistentLines = false;
+
+					DrawDebugDirectionalArrow(WorldPtr.Get(), ListenerPosition, ListenerPosition + (ArrowLength * ListenerFront), ArrowHeadSize, FColor::Red, bPersistentLines, Lifetime, DepthPriority, Thickness);
+					DrawDebugDirectionalArrow(WorldPtr.Get(), ListenerPosition, ListenerPosition + (ArrowLength * ListenerUp), ArrowHeadSize, FColor::Blue, bPersistentLines, Lifetime, DepthPriority, Thickness);
+					DrawDebugDirectionalArrow(WorldPtr.Get(), ListenerPosition, ListenerPosition + (ArrowLength * ListenerRight), ArrowHeadSize, FColor::Green, bPersistentLines, Lifetime, DepthPriority, Thickness);
+					DrawDebugSphere(WorldPtr.Get(), ListenerPosition, 5.0f, 8, FColor::Magenta, bPersistentLines, Lifetime, DepthPriority);
+				}
+			}, GET_STATID(STAT_AudioDrawListenerDebugInfo));
+		}
+	}
+#endif // ENABLE_DRAW_DEBUG
+}
+
 void FAudioDebugger::DrawDebugInfo(const FAudioVirtualLoop& VirtualLoop)
 {
 #if ENABLE_DRAW_DEBUG
diff --git a/Engine/Source/Runtime/Engine/Private/AudioDevice.cpp b/Engine/Source/Runtime/Engine/Private/AudioDevice.cpp
--- a/Engine/Source/Runtime/Engine/Private/AudioDevice.cpp
+++ b/Engine/Source/Runtime/Engine/Private/AudioDevice.cpp
@@ -4068,6 +4068,22 @@ void FAudioDevice::Update(bool bGameTicking)
 		}
 	}
 
+#if ENABLE_AUDIO_DEBUG
+	for (FActiveSound* ActiveSound : ActiveSounds)
+	{
+		if (!ActiveSound)
+		{
+			continue;
+		}
+		
+		if (UWorld* World = ActiveSound->GetWorld())
+		{
+			FAudioDebugger::DrawDebugInfo(*World, Listeners, ListenerAttenuationOverride, bUseListenerAttenuationOverride);
+			break;
+		}
+	}
+#endif // ENABLE_AUDIO_DEBUG
+
 	if (bHasActivatedReverb)
 	{
 		if (HighestPriorityActivatedReverb.Priority > AudioVolumePriority || bUsingDefaultReverb)
@@ -4847,6 +4863,11 @@ void FAudioDevice::GetMaxDistanceAndFocusFactor(USoundBase* Sound, const UWorld*
 		SoundTransform.SetTranslation(Location);
 
 		OutMaxDistance = AttenuationSettingsToApply->GetMaxDimension();
+		if (AttenuationSettingsToApply->AttenuationShape == EAttenuationShape::Box)
+		{
+			static const float Sqrt2 = 1.4142135f;
+			OutMaxDistance *= Sqrt2;
+		}
 
 		if (AttenuationSettingsToApply->bSpatialize && AttenuationSettingsToApply->bEnableListenerFocus)
 		{
diff --git a/Engine/Source/Runtime/Engine/Public/Audio/AudioDebug.h b/Engine/Source/Runtime/Engine/Public/Audio/AudioDebug.h
--- a/Engine/Source/Runtime/Engine/Public/Audio/AudioDebug.h
+++ b/Engine/Source/Runtime/Engine/Public/Audio/AudioDebug.h
@@ -12,9 +12,12 @@
  // Forward Declarations
 struct FActiveSound;
 struct FAudioVirtualLoop;
+struct FListener;
 struct FWaveInstance;
 
 class FSoundSource;
+class UWorld;
+
 
 class ENGINE_API FAudioDebugger
 {
@@ -42,6 +45,7 @@ public:
 
 	static void DrawDebugInfo(const FSoundSource& SoundSource);
 	static void DrawDebugInfo(const FActiveSound& ActiveSound, const TArray<FWaveInstance*>& ThisSoundsWaveInstances, const float DeltaTime);
+	static void DrawDebugInfo(const UWorld& World, const TArray<FListener>& Listeners, FVector& ListenerTransformOverride, bool bUseListenerTransformOverride);
 	static void DrawDebugInfo(const FAudioVirtualLoop& VirtualLoop);
 	static bool PostStatModulatorHelp(UWorld* World, FCommonViewportClient* ViewportClient, const TCHAR* Stream);
 	static int32 RenderStatCues(UWorld* World, FViewport* Viewport, FCanvas* Canvas, int32 X, int32 Y, const FVector* ViewLocation, const FRotator* ViewRotation);
