Fix const correctness compile error on Mac

diff --git a/Engine/Source/Runtime/Engine/Private/AudioDebug.cpp b/Engine/Source/Runtime/Engine/Private/AudioDebug.cpp
--- a/Engine/Source/Runtime/Engine/Private/AudioDebug.cpp
+++ b/Engine/Source/Runtime/Engine/Private/AudioDebug.cpp
@@ -510,7 +510,7 @@ void FAudioDebugger::DrawDebugInfo(const FActiveSound& ActiveSound, const TArray
 #endif // ENABLE_DRAW_DEBUG
 }
 
-void FAudioDebugger::DrawDebugInfo(const UWorld& World, const TArray<FListener>& Listeners, FVector& ListenerTransformOverride, bool bUseListenerTransformOverride)
+void FAudioDebugger::DrawDebugInfo(UWorld& World, const TArray<FListener>& Listeners, FVector& ListenerTransformOverride, bool bUseListenerTransformOverride)
 {
 #if ENABLE_DRAW_DEBUG
 	if (!ActiveSoundVisualizeListenersCVar)
diff --git a/Engine/Source/Runtime/Engine/Public/Audio/AudioDebug.h b/Engine/Source/Runtime/Engine/Public/Audio/AudioDebug.h
--- a/Engine/Source/Runtime/Engine/Public/Audio/AudioDebug.h
+++ b/Engine/Source/Runtime/Engine/Public/Audio/AudioDebug.h
@@ -45,7 +45,7 @@ public:
 
 	static void DrawDebugInfo(const FSoundSource& SoundSource);
 	static void DrawDebugInfo(const FActiveSound& ActiveSound, const TArray<FWaveInstance*>& ThisSoundsWaveInstances, const float DeltaTime);
-	static void DrawDebugInfo(const UWorld& World, const TArray<FListener>& Listeners, FVector& ListenerTransformOverride, bool bUseListenerTransformOverride);
+	static void DrawDebugInfo(UWorld& World, const TArray<FListener>& Listeners, FVector& ListenerTransformOverride, bool bUseListenerTransformOverride);
 	static void DrawDebugInfo(const FAudioVirtualLoop& VirtualLoop);
 	static bool PostStatModulatorHelp(UWorld* World, FCommonViewportClient* ViewportClient, const TCHAR* Stream);
 	static int32 RenderStatCues(UWorld* World, FViewport* Viewport, FCanvas* Canvas, int32 X, int32 Y, const FVector* ViewLocation, const FRotator* ViewRotation);
