Create streaming virtual texture staging textures with flags dependent on whether we are running in multithreaded rendering mode.
This fixes very slow performace for D3D11 non-multithreaded mode.
Reduce minimum size for staging texture.
This reduces unnecessary overhead from copying to it in multithreaded mode.

diff --git a/Engine/Source/Runtime/Engine/Private/VT/VirtualTextureUploadCache.cpp b/Engine/Source/Runtime/Engine/Private/VT/VirtualTextureUploadCache.cpp
--- a/Engine/Source/Runtime/Engine/Private/VT/VirtualTextureUploadCache.cpp
+++ b/Engine/Source/Runtime/Engine/Private/VT/VirtualTextureUploadCache.cpp
@@ -89,32 +89,33 @@ void FVirtualTextureUploadCache::Finalize(FRHICommandListImmediate& RHICmdList)
 		PoolEntry.BatchTextureIndex = (PoolEntry.BatchTextureIndex + 1u) % NUM_STAGING_TEXTURES;
 		FStagingTexture& StagingTexture = PoolEntry.StagingTexture[TextureIndex];
 
-		if (BatchCount > StagingTexture.BatchCapacity)
+		// On some platforms the staging texture create/lock behavior will depend on whether we are running with RHI threading
+		const bool bIsCpuWritable = !IsRunningRHIInSeparateThread();
+
+		if (BatchCount > StagingTexture.BatchCapacity || bIsCpuWritable != StagingTexture.bIsCPUWritable)
 		{
-			const uint32 MaxSizeInTiles = FMath::DivideAndRoundDown(GetMax2DTextureDimension(), TileSize);
+			// Staging texture is vertical stacked in widths of multiples of 4
+			// Smaller widths mean smaller stride which is more efficient for copying
+			// Round up to 4 to reduce likely wasted memory from width not aligning to whatever GPU prefers
+			const uint32 MaxTextureDimension = GetMax2DTextureDimension();
+			const uint32 MaxSizeInTiles = FMath::DivideAndRoundDown(MaxTextureDimension, TileSize);
 			const uint32 MaxCapacity = MaxSizeInTiles * MaxSizeInTiles;
 			check(BatchCount <= MaxCapacity);
-
-			// Try to create roughly square staging texture
-			// Stacking tiles on top of each other is potentially more cache efficient, since 'stride' will be smaller
-			// However, we're typically creating this texture with a tile size of 136, which on most GPUs will round up to next multiple of 32 (8x8 tiles of 4x4 BC compressed blocks) internally
-			// This means we'll waste less memory overall if width is larger
-			// Also, if we only stack vertically, we run into GPU limit of 16k texture dimension for large upload buffers
-			const uint32 NewCapacity = FMath::Clamp(BatchCount * 3u / 2u, 64u, MaxCapacity);
-			const uint32 WidthInTiles = FMath::FloorToInt(FMath::Sqrt((float)NewCapacity));
+			const uint32 WidthInTiles = FMath::DivideAndRoundUp(FMath::DivideAndRoundUp(BatchCount, MaxSizeInTiles), 4u) * 4;
 			check(WidthInTiles > 0u);
-			const uint32 HeightInTiles = (NewCapacity + WidthInTiles - 1u) / WidthInTiles;
+			const uint32 HeightInTiles = FMath::DivideAndRoundUp(BatchCount, WidthInTiles);
+			check(HeightInTiles > 0u);
 
 			if (StagingTexture.RHITexture)
 			{
 				DEC_MEMORY_STAT_BY(STAT_TotalGPUUploadSize, CalcTextureSize(StagingTexture.RHITexture->GetSizeX(), StagingTexture.RHITexture->GetSizeY(), PoolEntry.Format, 1u));
 			}
 
-			//todo[vt]: Intended to use TexCreate_CPUWritable on PC but it doesn't play well with D3D11 multi-threaded rendering Lock/Unlock
 			FRHIResourceCreateInfo CreateInfo;
-			StagingTexture.RHITexture = RHICmdList.CreateTexture2D(TileSize * WidthInTiles, TileSize * HeightInTiles, PoolEntry.Format, 1, 1, TexCreate_None, CreateInfo); 
+			StagingTexture.RHITexture = RHICmdList.CreateTexture2D(TileSize * WidthInTiles, TileSize * HeightInTiles, PoolEntry.Format, 1, 1, bIsCpuWritable ? TexCreate_CPUWritable : TexCreate_None, CreateInfo);
 			StagingTexture.WidthInTiles = WidthInTiles;
 			StagingTexture.BatchCapacity = WidthInTiles * HeightInTiles;
+			StagingTexture.bIsCPUWritable = bIsCpuWritable;
 			INC_MEMORY_STAT_BY(STAT_TotalGPUUploadSize, CalcTextureSize(TileSize * WidthInTiles, TileSize * HeightInTiles, PoolEntry.Format, 1u));
 		}
 
diff --git a/Engine/Source/Runtime/Engine/Private/VT/VirtualTextureUploadCache.h b/Engine/Source/Runtime/Engine/Private/VT/VirtualTextureUploadCache.h
--- a/Engine/Source/Runtime/Engine/Private/VT/VirtualTextureUploadCache.h
+++ b/Engine/Source/Runtime/Engine/Private/VT/VirtualTextureUploadCache.h
@@ -57,6 +57,7 @@ private:
 		TRefCountPtr<FRHITexture2D> RHITexture;
 		uint32_t WidthInTiles = 0u;
 		uint32_t BatchCapacity = 0u;
+		bool bIsCPUWritable;
 	};
 
 	struct FPoolEntry
