Fixed crash when building static mesh MeshDescriptions in editor.

diff --git a/Engine/Source/Runtime/Engine/Private/StaticMesh.cpp b/Engine/Source/Runtime/Engine/Private/StaticMesh.cpp
--- a/Engine/Source/Runtime/Engine/Private/StaticMesh.cpp
+++ b/Engine/Source/Runtime/Engine/Private/StaticMesh.cpp
@@ -5175,6 +5175,7 @@ bool UStaticMesh::BuildFromMeshDescriptions(const TArray<const FMeshDescription*
 	{
 #if WITH_EDITOR
 		// Editor builds cache the mesh description so that it can be preserved during map reloads etc
+		SetNumSourceModels(MeshDescriptions.Num());
 		CreateMeshDescription(LODIndex, *MeshDescriptionPtr);
 		CommitMeshDescription(LODIndex);
 #endif
