Fixed validation errors due to missing resource barriers in skin cache and ray tracing dynamic geometry systems

diff --git a/Engine/Source/Runtime/Engine/Private/GPUSkinCache.cpp b/Engine/Source/Runtime/Engine/Private/GPUSkinCache.cpp
--- a/Engine/Source/Runtime/Engine/Private/GPUSkinCache.cpp
+++ b/Engine/Source/Runtime/Engine/Private/GPUSkinCache.cpp
@@ -627,6 +627,9 @@ void FGPUSkinCache::CommitRayTracingGeometryUpdates(FRHICommandList& RHICmdList)
 
 	if (RayTracingGeometriesToUpdate.Num())
 	{
+		// Flush any remaining pending resource barriers before performing BVH updates
+		TransitionAllToReadable(RHICmdList);
+
 		TArray<FAccelerationStructureBuildParams> Updates;
 		for (const FRayTracingGeometry* RayTracingGeometry : RayTracingGeometriesToUpdate)
 		{
@@ -1286,10 +1289,6 @@ void FGPUSkinCache::DispatchUpdateSkinning(FRHICommandListImmediate& RHICmdList,
 	uint32 PreviousRevision = ShaderData.GetRevisionNumber(true);
 
 	DispatchData.TangentBuffer = DispatchData.PositionTracker.GetTangentBuffer();
-	if (DispatchData.TangentBuffer)
-	{
-		BuffersToTransition.Add(DispatchData.TangentBuffer->UAV);
-	}
 
 	DispatchData.PreviousPositionBuffer = DispatchData.PositionTracker.Find(PrevBoneBuffer, PreviousRevision);
 	if (!DispatchData.PreviousPositionBuffer)
@@ -1302,13 +1301,19 @@ void FGPUSkinCache::DispatchUpdateSkinning(FRHICommandListImmediate& RHICmdList,
 		Shader->SetParameters(RHICmdList, PrevBoneBuffer, Entry, DispatchData, DispatchData.GetPreviousPositionRWBuffer()->UAV, DispatchData.GetTangentRWBuffer() ? DispatchData.GetTangentRWBuffer()->UAV : nullptr);
 
 		RHICmdList.TransitionResource(EResourceTransitionAccess::EWritable, EResourceTransitionPipeline::EGfxToCompute, DispatchData.GetPreviousPositionRWBuffer()->UAV.GetReference());
+		BuffersToTransition.Add(DispatchData.GetPreviousPositionRWBuffer()->UAV);
+
+		if (DispatchData.TangentBuffer)
+		{
+			RHICmdList.TransitionResource(EResourceTransitionAccess::EWritable, EResourceTransitionPipeline::EGfxToCompute, DispatchData.TangentBuffer->UAV.GetReference());
+			BuffersToTransition.Add(DispatchData.TangentBuffer->UAV);
+		}
 
 		uint32 VertexCountAlign64 = FMath::DivideAndRoundUp(DispatchData.NumVertices, (uint32)64);
 		INC_DWORD_STAT_BY(STAT_GPUSkinCache_TotalNumVertices, VertexCountAlign64 * 64);
 		RHICmdList.DispatchComputeShader(VertexCountAlign64, 1, 1);
 		Shader->UnsetParameters(RHICmdList);
 
-		BuffersToTransition.Add(DispatchData.GetPreviousPositionRWBuffer()->UAV);
 	}
 
 	DispatchData.PositionBuffer = DispatchData.PositionTracker.Find(BoneBuffer, CurrentRevision);
@@ -1322,13 +1327,19 @@ void FGPUSkinCache::DispatchUpdateSkinning(FRHICommandListImmediate& RHICmdList,
 		Shader->SetParameters(RHICmdList, BoneBuffer, Entry, DispatchData, DispatchData.GetPositionRWBuffer()->UAV, DispatchData.GetTangentRWBuffer() ? DispatchData.GetTangentRWBuffer()->UAV : nullptr);
 
 		RHICmdList.TransitionResource(EResourceTransitionAccess::EWritable, EResourceTransitionPipeline::EGfxToCompute, DispatchData.GetPositionRWBuffer()->UAV.GetReference());
+		BuffersToTransition.Add(DispatchData.GetPositionRWBuffer()->UAV);
+
+		if (DispatchData.TangentBuffer)
+		{
+			RHICmdList.TransitionResource(EResourceTransitionAccess::EWritable, EResourceTransitionPipeline::EGfxToCompute, DispatchData.TangentBuffer->UAV.GetReference());
+			BuffersToTransition.Add(DispatchData.TangentBuffer->UAV);
+		}
 
 		uint32 VertexCountAlign64 = FMath::DivideAndRoundUp(DispatchData.NumVertices, (uint32)64);
 		INC_DWORD_STAT_BY(STAT_GPUSkinCache_TotalNumVertices, VertexCountAlign64 * 64);
 		RHICmdList.DispatchComputeShader(VertexCountAlign64, 1, 1);
 		Shader->UnsetParameters(RHICmdList);
 
-		BuffersToTransition.Add(DispatchData.GetPositionRWBuffer()->UAV);
 	}
 
 	check(DispatchData.PreviousPositionBuffer != DispatchData.PositionBuffer);
diff --git a/Engine/Source/Runtime/Engine/Private/SkeletalRenderGPUSkin.cpp b/Engine/Source/Runtime/Engine/Private/SkeletalRenderGPUSkin.cpp
--- a/Engine/Source/Runtime/Engine/Private/SkeletalRenderGPUSkin.cpp
+++ b/Engine/Source/Runtime/Engine/Private/SkeletalRenderGPUSkin.cpp
@@ -436,6 +436,9 @@ void FSkeletalMeshObjectGPUSkin::UpdateDynamicData_RenderThread(FGPUSkinCache* G
 
 				FGPUSkinCache::GetRayTracingSegmentVertexBuffers(*SkinCacheEntry, Initializer.Segments);
 
+				// Flush pending resource barriers before BVH is built for the first time
+				GPUSkinCache->TransitionAllToReadable(RHICmdList);
+
 				RayTracingGeometry.SetInitializer(Initializer);
 				RayTracingGeometry.UpdateRHI();
 			}
diff --git a/Engine/Source/Runtime/Engine/Public/GPUSkinCache.h b/Engine/Source/Runtime/Engine/Public/GPUSkinCache.h
--- a/Engine/Source/Runtime/Engine/Public/GPUSkinCache.h
+++ b/Engine/Source/Runtime/Engine/Public/GPUSkinCache.h
@@ -164,11 +164,11 @@ public:
 		{
 			for (int32 Index = 0; Index < NUM_BUFFERS; ++Index)
 			{
-				RWBuffers[Index].Initialize(4, NumVertices * 3, PF_R32_FLOAT, BUF_Static, TEXT("Vertices"));
+				RWBuffers[Index].Initialize(4, NumVertices * 3, PF_R32_FLOAT, BUF_Static, TEXT("SkinCacheVertices"));
 			}
 			if (WithTangents)
 			{
-				Tangents.Initialize(8, NumVertices * 2, PF_R16G16B16A16_SNORM, BUF_Static, TEXT("Tangents"));
+				Tangents.Initialize(8, NumVertices * 2, PF_R16G16B16A16_SNORM, BUF_Static, TEXT("SkinCacheTangents"));
 			}
 		}
 
diff --git a/Engine/Source/Runtime/Renderer/Private/RayTracing/RayTracingDynamicGeometry.cpp b/Engine/Source/Runtime/Renderer/Private/RayTracing/RayTracingDynamicGeometry.cpp
--- a/Engine/Source/Runtime/Renderer/Private/RayTracing/RayTracingDynamicGeometry.cpp
+++ b/Engine/Source/Runtime/Renderer/Private/RayTracing/RayTracingDynamicGeometry.cpp
@@ -250,7 +250,18 @@ void FRayTracingDynamicGeometryCollection::DispatchUpdates(CmdListType& RHICmdLi
 
 		{
 			SCOPED_DRAW_OR_COMPUTE_EVENT(RHICmdList, VSinCSComputeDispatch)
-			for (auto& Cmd : *DispatchCommands)
+
+			TArray<FRHIUnorderedAccessView*> BuffersToTransition;
+			BuffersToTransition.Reserve(DispatchCommands->Num());
+
+			for (FMeshComputeDispatchCommand& Cmd : *DispatchCommands)
+			{
+				BuffersToTransition.Add(Cmd.TargetBuffer->UAV.GetReference());
+			}
+
+			RHICmdList.TransitionResources(EResourceTransitionAccess::EWritable, EResourceTransitionPipeline::EGfxToCompute, BuffersToTransition.GetData(), BuffersToTransition.Num());
+
+			for (FMeshComputeDispatchCommand& Cmd : *DispatchCommands)
 			{
 				{
 					FRayTracingDynamicGeometryConverterCS* Shader = Cmd.MaterialShader;
@@ -268,6 +279,8 @@ void FRayTracingDynamicGeometryCollection::DispatchUpdates(CmdListType& RHICmdLi
 					Shader->RWVertexPositions.UnsetUAV(RHICmdList, Shader->GetComputeShader());
 				}
 			}
+
+			RHICmdList.TransitionResources(EResourceTransitionAccess::EReadable, EResourceTransitionPipeline::EComputeToGfx, BuffersToTransition.GetData(), BuffersToTransition.Num());
 		}
 
 		SCOPED_DRAW_OR_COMPUTE_EVENT(RHICmdList, Build);
