Fied crash when scrubbing qn qnimqtion in sequencer

Removed duplicate-update guard for linked anim graphs that meant slot nodes downstream could end up with an evaluate() without a graph update() call, and hence invalid weights.
This guard is safe to remove now as we always expect (actually force) linked instanes to be updated in parallel.

diff --git a/Engine/Source/Runtime/Engine/Private/Animation/AnimNode_LinkedAnimGraph.cpp b/Engine/Source/Runtime/Engine/Private/Animation/AnimNode_LinkedAnimGraph.cpp
--- a/Engine/Source/Runtime/Engine/Private/Animation/AnimNode_LinkedAnimGraph.cpp
+++ b/Engine/Source/Runtime/Engine/Private/Animation/AnimNode_LinkedAnimGraph.cpp
@@ -98,12 +98,11 @@ void FAnimNode_LinkedAnimGraph::Update_AnyThread(const FAnimationUpdateContext&
 
 		PropagateInputProperties(InContext.AnimInstanceProxy->GetAnimInstanceObject());
 
-		// Only update if we've not had a single-threaded update already
-		if(InstanceToRun->bNeedsUpdate)
-		{
-			FAnimationUpdateContext NewContext = InContext.WithOtherProxy(&Proxy);
-			Proxy.UpdateAnimation_WithRoot(NewContext, LinkedRoot, GetDynamicLinkFunctionName());
-		}
+		// We can call this unconditionally here now because linked anim instances are forced to have a parallel update
+		// in USkeletalMeshComponent::TickAnimation. It used to be the case that we could do non-parallel work in 
+		// USkeletalMeshComponent::TickAnimation, which would mean we would have to skip doing that work here.
+		FAnimationUpdateContext NewContext = InContext.WithOtherProxy(&Proxy);
+		Proxy.UpdateAnimation_WithRoot(NewContext, LinkedRoot, GetDynamicLinkFunctionName());
 	}
 	else if(InputPoses.Num() > 0)
 	{
