Fix static analysis warning. NumTriangles is uint32 so no need for < 0 check.
#rnx

diff --git a/Engine/Source/Runtime/Engine/Private/Rendering/SkeletalMeshLODImporterData.cpp b/Engine/Source/Runtime/Engine/Private/Rendering/SkeletalMeshLODImporterData.cpp
--- a/Engine/Source/Runtime/Engine/Private/Rendering/SkeletalMeshLODImporterData.cpp
+++ b/Engine/Source/Runtime/Engine/Private/Rendering/SkeletalMeshLODImporterData.cpp
@@ -530,7 +530,7 @@ void FReductionBaseSkeletalMeshBulkData::CacheGeometryInfo(const FSkeletalMeshLO
 		{
 			//Make sure the count fit in a uint32
 			CacheLODVertexNumber += Section.NumVertices < 0 ? 0 : Section.NumVertices;
-			CacheLODTriNumber += Section.NumTriangles < 0 ? 0 : Section.NumTriangles;
+			CacheLODTriNumber += Section.NumTriangles;
 		}
 	}
 }
diff --git a/Engine/Source/Runtime/Engine/Private/SkeletalMesh.cpp b/Engine/Source/Runtime/Engine/Private/SkeletalMesh.cpp
--- a/Engine/Source/Runtime/Engine/Private/SkeletalMesh.cpp
+++ b/Engine/Source/Runtime/Engine/Private/SkeletalMesh.cpp
@@ -492,7 +492,7 @@ bool USkeletalMesh::IsReductionActive(int32 LODIndex) const
 				{
 					//Make sure the count fit in a uint32
 					LODVertexNumber += Section.NumVertices < 0 ? 0 : Section.NumVertices;
-					LODTriNumber += Section.NumTriangles < 0 ? 0 : Section.NumTriangles;
+					LODTriNumber += Section.NumTriangles;
 				}
 			}
 		}
