Fix for stack-use-after-scope error, caught by address sanitizer.

The PostInitHandle variable in the InitGamePhys function was falling out of scope well ahead of the lambda execution.  This changelist moves the handle out to the global scope so that it may be safely referenced by the lambda as well as the TermGamePhys function.

#rnx

diff --git a/Engine/Source/Runtime/Engine/Private/PhysicsEngine/PhysLevel.cpp b/Engine/Source/Runtime/Engine/Private/PhysicsEngine/PhysLevel.cpp
--- a/Engine/Source/Runtime/Engine/Private/PhysicsEngine/PhysLevel.cpp
+++ b/Engine/Source/Runtime/Engine/Private/PhysicsEngine/PhysLevel.cpp
@@ -263,6 +263,8 @@ void PostEngineInitialize()
 	}
 }
 
+FDelegateHandle GPostInitHandle;
+
 bool InitGamePhys()
 {
 	if (!InitGamePhysCore())
@@ -272,11 +274,9 @@ bool InitGamePhys()
 
 	// We need to defer initializing the module as it will attempt to read from the settings provider. If the settings
 	// provider is backed by a UObject in any way access to it will fail because we're too early in the init process.
-	FDelegateHandle PostInitHandle;
-	PostInitHandle = FCoreDelegates::OnPostEngineInit.AddLambda([&PostInitHandle]
+	GPostInitHandle = FCoreDelegates::OnPostEngineInit.AddLambda([]()
 	{
 		PostEngineInitialize();
-		FCoreDelegates::OnPostEngineInit.Remove(PostInitHandle);
 	});
 
 #if WITH_PHYSX
@@ -297,6 +297,12 @@ bool InitGamePhys()
 
 void TermGamePhys()
 {
+	if (GPostInitHandle.IsValid())
+	{
+		FCoreDelegates::OnPostEngineInit.Remove(GPostInitHandle);
+		GPostInitHandle.Reset();
+	}
+
 #if WITH_PHYSX
 
 	// Do nothing if they were never initialized
