Revert CL-7441982. This caused issues where multiple bunches could be marked as closing bunches, and cause crashes.

diff --git a/Engine/Source/Runtime/Engine/Private/DataChannel.cpp b/Engine/Source/Runtime/Engine/Private/DataChannel.cpp
--- a/Engine/Source/Runtime/Engine/Private/DataChannel.cpp
+++ b/Engine/Source/Runtime/Engine/Private/DataChannel.cpp
@@ -553,12 +553,6 @@ bool UChannel::ReceivedNextBunch( FInBunch & Bunch, bool & bOutSkipAck )
 		Connection->InReliable[Bunch.ChIndex] = Bunch.ChSequence;
 	}
 
-	if (Bunch.bHasPackageMapExports && !Connection->InternalAck)
-	{
-		UE_LOG(LogNet, VeryVerbose, TEXT("Received new bunch. It only contained NetGUIDs and was processed in UChannel::ReceivedRawBunch(). PacketId %d. ChSequence %d. ChIndex %d"), Bunch.PacketId, Bunch.ChSequence, Bunch.ChIndex);
-		return false;
-	}
-
 	FInBunch* HandleBunch = &Bunch;
 	if (Bunch.bPartial)
 	{
@@ -593,7 +587,7 @@ bool UChannel::ReceivedNextBunch( FInBunch & Bunch, bool & bOutSkipAck )
 			}
 
 			InPartialBunch = new FInBunch(Bunch, false);
-			if (ensureMsgf((!Bunch.bHasPackageMapExports && Bunch.GetBitsLeft() > 0), TEXT("Received new partial bunch with bHasPackageMapExports flag set to 1. This should not have happened. %s"), *Describe()))
+			if ( !Bunch.bHasPackageMapExports && Bunch.GetBitsLeft() > 0 )
 			{
 				if ( Bunch.GetBitsLeft() % 8 != 0 )
 				{
@@ -606,6 +600,10 @@ bool UChannel::ReceivedNextBunch( FInBunch & Bunch, bool & bOutSkipAck )
 
 				LogPartialBunch(TEXT("Received new partial bunch."), Bunch, *InPartialBunch);
 			}
+			else
+			{
+				LogPartialBunch(TEXT("Received New partial bunch. It only contained NetGUIDs."), Bunch, *InPartialBunch);
+			}
 		}
 		else
 		{
@@ -972,13 +970,11 @@ FPacketIdRange UChannel::SendBunch( FOutBunch* Bunch, bool Merge )
 	TArray<FOutBunch*>& OutgoingBunches = Connection->GetOutgoingBunches();
 	OutgoingBunches.Reset();
 
-	int32 NumOfExportBunches = 0;
 	// Add any export bunches
 	// Replay connections will manage export bunches separately.
 	if (!Connection->InternalAck)
 	{
 		AppendExportBunches( OutgoingBunches );
-		NumOfExportBunches = OutgoingBunches.Num();
 	}
 
 	if ( OutgoingBunches.Num() )
@@ -1098,7 +1094,7 @@ FPacketIdRange UChannel::SendBunch( FOutBunch* Bunch, bool Merge )
 		return PacketIdRange;
 	}
 
-	UE_CLOG((OutgoingBunches.Num() > 1 + NumOfExportBunches), LogNetPartialBunch, Log, TEXT("Sending %d Bunches. Channel: %d %s"), OutgoingBunches.Num() - NumOfExportBunches, Bunch->ChIndex, *Describe());
+	UE_CLOG((OutgoingBunches.Num() > 1), LogNetPartialBunch, Log, TEXT("Sending %d Bunches. Channel: %d %s"), OutgoingBunches.Num(), Bunch->ChIndex, *Describe());
 	for( int32 PartialNum = 0; PartialNum < OutgoingBunches.Num(); ++PartialNum)
 	{
 		FOutBunch * NextBunch = OutgoingBunches[PartialNum];
@@ -1118,12 +1114,12 @@ FPacketIdRange UChannel::SendBunch( FOutBunch* Bunch, bool Merge )
 			NextBunch->bHasMustBeMappedGUIDs |= Bunch->bHasMustBeMappedGUIDs;
 		}
 
-		if (OutgoingBunches.Num() > 1 + NumOfExportBunches)
+		if (OutgoingBunches.Num() > 1)
 		{
-			NextBunch->bPartial = (PartialNum >= NumOfExportBunches ? 1 : 0);               // do not set bPartial bit of export bunches to 1
-			NextBunch->bPartialInitial = (PartialNum == NumOfExportBunches ? 1: 0);
+			NextBunch->bPartial = 1;
+			NextBunch->bPartialInitial = (PartialNum == 0 ? 1: 0);
 			NextBunch->bPartialFinal = (PartialNum == OutgoingBunches.Num() - 1 ? 1: 0);
-			NextBunch->bOpen &= (PartialNum == NumOfExportBunches);											// Only the first bunch should have the bOpen bit set
+			NextBunch->bOpen &= (PartialNum == 0);											// Only the first bunch should have the bOpen bit set
 			NextBunch->bClose = (Bunch->bClose && (OutgoingBunches.Num()-1 == PartialNum)); // Only last bunch should have bClose bit set
 		}
 
