Fix compiling editor with RTTI enabled for Linux/Mac

diff --git a/Engine/Source/Editor/BlueprintGraph/Classes/K2Node.h b/Engine/Source/Editor/BlueprintGraph/Classes/K2Node.h
--- a/Engine/Source/Editor/BlueprintGraph/Classes/K2Node.h
+++ b/Engine/Source/Editor/BlueprintGraph/Classes/K2Node.h
@@ -173,7 +173,7 @@ enum ERenamePinResult
  * Abstract base class of all blueprint graph nodes.
  */
 UCLASS(abstract, MinimalAPI)
-class UK2Node : public UEdGraphNode
+class BLUEPRINTGRAPH_VTABLE UK2Node : public UEdGraphNode
 {
 	GENERATED_UCLASS_BODY()
 
diff --git a/Engine/Source/Editor/BlueprintGraph/Classes/K2Node_Switch.h b/Engine/Source/Editor/BlueprintGraph/Classes/K2Node_Switch.h
--- a/Engine/Source/Editor/BlueprintGraph/Classes/K2Node_Switch.h
+++ b/Engine/Source/Editor/BlueprintGraph/Classes/K2Node_Switch.h
@@ -12,7 +12,7 @@
 #include "K2Node_Switch.generated.h"
 
 UCLASS(MinimalAPI, abstract)
-class UK2Node_Switch : public UK2Node
+class BLUEPRINTGRAPH_VTABLE UK2Node_Switch : public UK2Node
 {
 	GENERATED_UCLASS_BODY()
 
diff --git a/Engine/Source/Editor/UnrealEd/Classes/ActorFactories/ActorFactoryBoxVolume.h b/Engine/Source/Editor/UnrealEd/Classes/ActorFactories/ActorFactoryBoxVolume.h
--- a/Engine/Source/Editor/UnrealEd/Classes/ActorFactories/ActorFactoryBoxVolume.h
+++ b/Engine/Source/Editor/UnrealEd/Classes/ActorFactories/ActorFactoryBoxVolume.h
@@ -12,7 +12,7 @@ class AActor;
 struct FAssetData;
 
 UCLASS(MinimalAPI, config=Editor, collapsecategories, hidecategories=Object)
-class UActorFactoryBoxVolume : public UActorFactoryVolume
+class UNREALED_VTABLE UActorFactoryBoxVolume : public UActorFactoryVolume
 {
 	GENERATED_UCLASS_BODY()
 
diff --git a/Engine/Source/Editor/UnrealEd/Classes/ThumbnailRendering/DefaultSizedThumbnailRenderer.h b/Engine/Source/Editor/UnrealEd/Classes/ThumbnailRendering/DefaultSizedThumbnailRenderer.h
--- a/Engine/Source/Editor/UnrealEd/Classes/ThumbnailRendering/DefaultSizedThumbnailRenderer.h
+++ b/Engine/Source/Editor/UnrealEd/Classes/ThumbnailRendering/DefaultSizedThumbnailRenderer.h
@@ -13,7 +13,7 @@
 #include "DefaultSizedThumbnailRenderer.generated.h"
 
 UCLASS(abstract, config=Editor, MinimalAPI)
-class UDefaultSizedThumbnailRenderer : public UThumbnailRenderer
+class UNREALED_VTABLE UDefaultSizedThumbnailRenderer : public UThumbnailRenderer
 {
 	GENERATED_UCLASS_BODY()
 
diff --git a/Engine/Source/Editor/UnrealEd/Classes/ThumbnailRendering/TextureThumbnailRenderer.h b/Engine/Source/Editor/UnrealEd/Classes/ThumbnailRendering/TextureThumbnailRenderer.h
--- a/Engine/Source/Editor/UnrealEd/Classes/ThumbnailRendering/TextureThumbnailRenderer.h
+++ b/Engine/Source/Editor/UnrealEd/Classes/ThumbnailRendering/TextureThumbnailRenderer.h
@@ -16,7 +16,7 @@ class FCanvas;
 class FRenderTarget;
 
 UCLASS(MinimalAPI)
-class UTextureThumbnailRenderer : public UThumbnailRenderer
+class UNREALED_VTABLE UTextureThumbnailRenderer : public UThumbnailRenderer
 {
 	GENERATED_UCLASS_BODY()
 
diff --git a/Engine/Source/Editor/UnrealEd/Public/CameraController.h b/Engine/Source/Editor/UnrealEd/Public/CameraController.h
--- a/Engine/Source/Editor/UnrealEd/Public/CameraController.h
+++ b/Engine/Source/Editor/UnrealEd/Public/CameraController.h
@@ -213,7 +213,7 @@ public:
  *
  * An interactive camera movement system.  Supports simple physics-based animation.
  */
-class FEditorCameraController
+class UNREALED_VTABLE FEditorCameraController
 {
 
 public:
diff --git a/Engine/Source/Editor/UnrealEd/Public/ReferencedAssetsUtils.h b/Engine/Source/Editor/UnrealEd/Public/ReferencedAssetsUtils.h
--- a/Engine/Source/Editor/UnrealEd/Public/ReferencedAssetsUtils.h
+++ b/Engine/Source/Editor/UnrealEd/Public/ReferencedAssetsUtils.h
@@ -129,7 +129,7 @@ public:
 	FArchive& operator<<(class UObject*& Obj);
 };
 
-class FFindReferencedAssets : 
+class UNREALED_VTABLE FFindReferencedAssets : 
 	public FGCObject
 {
 public:
diff --git a/Engine/Source/Runtime/Core/Public/Containers/Ticker.h b/Engine/Source/Runtime/Core/Public/Containers/Ticker.h
--- a/Engine/Source/Runtime/Core/Public/Containers/Ticker.h
+++ b/Engine/Source/Runtime/Core/Public/Containers/Ticker.h
@@ -112,7 +112,7 @@ protected:
 /**
  * Base class for ticker objects
  */
-class FTickerObjectBase
+class CORE_VTABLE FTickerObjectBase
 {
 public:
 
@@ -158,4 +158,4 @@ private:
 	FTicker& Ticker;
 	/** Delegate for callbacks to Tick */
 	FDelegateHandle TickHandle;
-};
\ No newline at end of file
+};
diff --git a/Engine/Source/Runtime/Engine/Classes/Engine/DataAsset.h b/Engine/Source/Runtime/Engine/Classes/Engine/DataAsset.h
--- a/Engine/Source/Runtime/Engine/Classes/Engine/DataAsset.h
+++ b/Engine/Source/Runtime/Engine/Classes/Engine/DataAsset.h
@@ -13,7 +13,7 @@
  * Base class for a simple asset containing data. The editor will list this in the content browser if you inherit from this class
  */
 UCLASS(abstract, MinimalAPI)
-class UDataAsset : public UObject
+class ENGINE_VTABLE UDataAsset : public UObject
 {
 	GENERATED_UCLASS_BODY()
 public:
diff --git a/Engine/Source/Runtime/Engine/Classes/Kismet/BlueprintFunctionLibrary.h b/Engine/Source/Runtime/Engine/Classes/Kismet/BlueprintFunctionLibrary.h
--- a/Engine/Source/Runtime/Engine/Classes/Kismet/BlueprintFunctionLibrary.h
+++ b/Engine/Source/Runtime/Engine/Classes/Kismet/BlueprintFunctionLibrary.h
@@ -12,7 +12,7 @@
 // This class is a base class for any function libraries exposed to blueprints.
 // Methods in subclasses are expected to be static, and no methods should be added to this base class.
 UCLASS(Abstract, MinimalAPI)
-class UBlueprintFunctionLibrary : public UObject
+class ENGINE_VTABLE UBlueprintFunctionLibrary : public UObject
 {
 	GENERATED_UCLASS_BODY()
 
diff --git a/Engine/Source/Runtime/Engine/Public/HitProxies.h b/Engine/Source/Runtime/Engine/Public/HitProxies.h
--- a/Engine/Source/Runtime/Engine/Public/HitProxies.h
+++ b/Engine/Source/Runtime/Engine/Public/HitProxies.h
@@ -126,7 +126,7 @@ private:
 /**
  * Base class for detecting user-interface hits.
  */
-class HHitProxy : public FRefCountedObject
+class ENGINE_VTABLE HHitProxy : public FRefCountedObject
 {
 	//DECLARE_HIT_PROXY( ENGINE_API )
 	//We separate the GetType function implementation here because this is the base class/
diff --git a/Engine/Source/Runtime/Engine/Public/MaterialShared.h b/Engine/Source/Runtime/Engine/Public/MaterialShared.h
--- a/Engine/Source/Runtime/Engine/Public/MaterialShared.h
+++ b/Engine/Source/Runtime/Engine/Public/MaterialShared.h
@@ -2155,7 +2155,7 @@ inline bool IsTranslucentBlendMode(enum EBlendMode BlendMode)
 /**
  * Implementation of the FMaterial interface for a UMaterial or UMaterialInstance.
  */
-class FMaterialResource : public FMaterial
+class ENGINE_VTABLE FMaterialResource : public FMaterial
 {
 public:
 
diff --git a/Engine/Source/Runtime/Engine/Public/Net/OnlineBlueprintCallProxyBase.h b/Engine/Source/Runtime/Engine/Public/Net/OnlineBlueprintCallProxyBase.h
--- a/Engine/Source/Runtime/Engine/Public/Net/OnlineBlueprintCallProxyBase.h
+++ b/Engine/Source/Runtime/Engine/Public/Net/OnlineBlueprintCallProxyBase.h
@@ -10,7 +10,7 @@
 DECLARE_DYNAMIC_MULTICAST_DELEGATE(FEmptyOnlineDelegate);
 
 UCLASS(MinimalAPI)
-class UOnlineBlueprintCallProxyBase : public UObject
+class ENGINE_VTABLE UOnlineBlueprintCallProxyBase : public UObject
 {
 	GENERATED_UCLASS_BODY()
 
diff --git a/Engine/Source/Runtime/Engine/Public/PrimitiveSceneProxy.h b/Engine/Source/Runtime/Engine/Public/PrimitiveSceneProxy.h
--- a/Engine/Source/Runtime/Engine/Public/PrimitiveSceneProxy.h
+++ b/Engine/Source/Runtime/Engine/Public/PrimitiveSceneProxy.h
@@ -123,7 +123,7 @@ extern bool CacheShadowDepthsFromPrimitivesUsingWPO();
  * Encapsulates the data which is mirrored to render a UPrimitiveComponent parallel to the game thread.
  * This is intended to be subclassed to support different primitive types.  
  */
-class FPrimitiveSceneProxy
+class ENGINE_VTABLE FPrimitiveSceneProxy
 {
 public:
 
@@ -1116,4 +1116,4 @@ ENGINE_API extern bool SupportsCachingMeshDrawCommands(const FPrimitiveSceneProx
  * Returns if specified mesh command can be cached, or needs to be recreated every frame; this is a slightly slower version
  * used for materials with external textures that need invalidating their PSOs.
  */
-ENGINE_API extern bool SupportsCachingMeshDrawCommands(const FPrimitiveSceneProxy* RESTRICT PrimitiveSceneProxy, const FMeshBatch& MeshBatch, ERHIFeatureLevel::Type FeatureLevel);
\ No newline at end of file
+ENGINE_API extern bool SupportsCachingMeshDrawCommands(const FPrimitiveSceneProxy* RESTRICT PrimitiveSceneProxy, const FMeshBatch& MeshBatch, ERHIFeatureLevel::Type FeatureLevel);
diff --git a/Engine/Source/Runtime/MovieScene/Public/MovieSceneTrack.h b/Engine/Source/Runtime/MovieScene/Public/MovieSceneTrack.h
--- a/Engine/Source/Runtime/MovieScene/Public/MovieSceneTrack.h
+++ b/Engine/Source/Runtime/MovieScene/Public/MovieSceneTrack.h
@@ -111,7 +111,7 @@ struct FMovieSceneTrackDisplayOptions
  * Base class for a track in a Movie Scene
  */
 UCLASS(abstract, DefaultToInstanced, MinimalAPI, BlueprintType)
-class UMovieSceneTrack
+class MOVIESCENE_VTABLE UMovieSceneTrack
 	: public UMovieSceneSignedObject
 {
 	GENERATED_BODY()
diff --git a/Engine/Source/Runtime/SlateCore/Public/FastUpdate/SlateInvalidationRoot.h b/Engine/Source/Runtime/SlateCore/Public/FastUpdate/SlateInvalidationRoot.h
--- a/Engine/Source/Runtime/SlateCore/Public/FastUpdate/SlateInvalidationRoot.h
+++ b/Engine/Source/Runtime/SlateCore/Public/FastUpdate/SlateInvalidationRoot.h
@@ -47,7 +47,7 @@ struct FSlateInvalidationResult
 	bool bRepaintedWidgets;
 };
 
-class FSlateInvalidationRoot : public FGCObject, public FNoncopyable
+class SLATECORE_VTABLE FSlateInvalidationRoot : public FGCObject, public FNoncopyable
 {
 	friend class FSlateUpdateFastWidgetPathTask;
 	friend class FSlateUpdateFastPathAndHitTestGridTask;
@@ -130,4 +130,4 @@ private:
 	bool bChildOrderInvalidated;
 	bool bNeedsSlowPath;
 	bool bNeedScreenPositionShift;
-};
\ No newline at end of file
+};
