Adding option to swap roles when replicating actor properties

diff --git a/Engine/Plugins/Runtime/ReplicationGraph/Source/Private/ReplicationGraph.cpp b/Engine/Plugins/Runtime/ReplicationGraph/Source/Private/ReplicationGraph.cpp
--- a/Engine/Plugins/Runtime/ReplicationGraph/Source/Private/ReplicationGraph.cpp
+++ b/Engine/Plugins/Runtime/ReplicationGraph/Source/Private/ReplicationGraph.cpp
@@ -540,6 +540,13 @@ void UReplicationGraph::AddNetworkActor(AActor* Actor)
 	RouteAddNetworkActorToNodes(FNewReplicatedActorInfo(Actor), GlobalInfo);
 }
 
+void UReplicationGraph::SetRoleSwapOnReplicate(AActor* Actor, bool bSwapRoles)
+{
+	if (FGlobalActorReplicationInfo* GlobalInfo = GlobalActorReplicationInfoMap.Find(Actor))
+	{
+		GlobalInfo->bSwapRolesOnReplicate = bSwapRoles;
+	}
+}
 
 void UReplicationGraph::RouteAddNetworkActorToNodes(const FNewReplicatedActorInfo& ActorInfo, FGlobalActorReplicationInfo& GlobalInfo)
 {
@@ -1710,6 +1717,12 @@ int64 UReplicationGraph::ReplicateSingleActor(AActor* Actor, FConnectionReplicat
 		Actor->CallPreReplication(NetDriver);
 	}
 
+	TOptional<FScopedActorRoleSwap> SwapGuard;
+	if (GlobalActorInfo.bSwapRolesOnReplicate)
+	{
+		SwapGuard = FScopedActorRoleSwap(Actor);
+	}
+
 	const bool bWantsToGoDormant = GlobalActorInfo.bWantsToBeDormant;
 
 	const FActorRepListRefView& DependentActorList = GlobalActorInfo.GetDependentActorList();
diff --git a/Engine/Plugins/Runtime/ReplicationGraph/Source/Public/ReplicationGraph.h b/Engine/Plugins/Runtime/ReplicationGraph/Source/Public/ReplicationGraph.h
--- a/Engine/Plugins/Runtime/ReplicationGraph/Source/Public/ReplicationGraph.h
+++ b/Engine/Plugins/Runtime/ReplicationGraph/Source/Public/ReplicationGraph.h
@@ -828,6 +828,7 @@ public:
 	virtual void NotifyActorTearOff(AActor* Actor) override;
 	virtual void NotifyActorFullyDormantForConnection(AActor* Actor, UNetConnection* Connection) override;
 	virtual void NotifyActorDormancyChange(AActor* Actor, ENetDormancy OldDormancyState) override;
+	virtual void SetRoleSwapOnReplicate(AActor* Actor, bool bSwapRoles) override;
 	virtual bool ProcessRemoteFunction(class AActor* Actor, UFunction* Function, void* Parameters, FOutParmRec* OutParms, FFrame* Stack, class UObject* SubObject) override;
 	virtual int32 ServerReplicateActors(float DeltaSeconds) override;
 	virtual void PostTickDispatch() override;
diff --git a/Engine/Plugins/Runtime/ReplicationGraph/Source/Public/ReplicationGraphTypes.h b/Engine/Plugins/Runtime/ReplicationGraph/Source/Public/ReplicationGraphTypes.h
--- a/Engine/Plugins/Runtime/ReplicationGraph/Source/Public/ReplicationGraphTypes.h
+++ b/Engine/Plugins/Runtime/ReplicationGraph/Source/Public/ReplicationGraphTypes.h
@@ -471,6 +471,9 @@ struct FGlobalActorReplicationInfo
 	/** Mirrors AActor::NetDormancy > DORM_Awake */
 	bool bWantsToBeDormant = false;
 
+	/** True if we should swap the actor role and remote role before calling ReplicateActor() */
+	bool bSwapRolesOnReplicate = false;
+
 	/** Class default mirrors: state that is initialized directly from class defaults (and can be later changed on a per-actor basis) */
 	FClassReplicationInfo Settings;
 	
diff --git a/Engine/Source/Runtime/Engine/Classes/Engine/NetDriver.h b/Engine/Source/Runtime/Engine/Classes/Engine/NetDriver.h
--- a/Engine/Source/Runtime/Engine/Classes/Engine/NetDriver.h
+++ b/Engine/Source/Runtime/Engine/Classes/Engine/NetDriver.h
@@ -1332,6 +1332,9 @@ public:
 
 	ENGINE_API virtual void NotifyActorTearOff(AActor* Actor);
 
+	/** Set whether this actor should swap roles before replicating properties. */
+	ENGINE_API void SetRoleSwapOnReplicate(AActor* Actor, bool bSwapRoles);
+
 	// ---------------------------------------------------------------
 	//
 	// ---------------------------------------------------------------	
diff --git a/Engine/Source/Runtime/Engine/Classes/Engine/NetworkObjectList.h b/Engine/Source/Runtime/Engine/Classes/Engine/NetworkObjectList.h
--- a/Engine/Source/Runtime/Engine/Classes/Engine/NetworkObjectList.h
+++ b/Engine/Source/Runtime/Engine/Classes/Engine/NetworkObjectList.h
@@ -52,6 +52,9 @@ struct FNetworkObjectInfo
 	/** Should this object be considered for replay checkpoint writes */
 	uint8 bDirtyForReplay : 1;
 
+	/** Should channel swap roles while calling ReplicateActor */
+	uint8 bSwapRolesOnReplicate : 1;
+
 	/** Force this object to be considered relevant for at least one update */
 	uint32 ForceRelevantFrame = 0;
 
@@ -63,7 +66,8 @@ struct FNetworkObjectInfo
 		, LastNetUpdateTime(0.0f)
 		, bPendingNetUpdate(false)
 		, bForceRelevantNextUpdate(false)
-		, bDirtyForReplay(false) {}
+		, bDirtyForReplay(false)
+		, bSwapRolesOnReplicate(false) {}
 
 	FNetworkObjectInfo(AActor* InActor)
 		: Actor(InActor)
@@ -74,7 +78,8 @@ struct FNetworkObjectInfo
 		, LastNetUpdateTime(0.0f)
 		, bPendingNetUpdate(false)
 		, bForceRelevantNextUpdate(false)
-		, bDirtyForReplay(false) {}
+		, bDirtyForReplay(false)
+		, bSwapRolesOnReplicate(false) {}
 
 	void CountBytes(FArchive& Ar) const;
 };
diff --git a/Engine/Source/Runtime/Engine/Classes/Engine/ReplicationDriver.h b/Engine/Source/Runtime/Engine/Classes/Engine/ReplicationDriver.h
--- a/Engine/Source/Runtime/Engine/Classes/Engine/ReplicationDriver.h
+++ b/Engine/Source/Runtime/Engine/Classes/Engine/ReplicationDriver.h
@@ -92,6 +92,8 @@ public:
 
 	virtual void NotifyActorDormancyChange(AActor* Actor, ENetDormancy OldDormancyState) PURE_VIRTUAL(UReplicationDriver::NotifyActorDormancyChange, );
 
+	virtual void SetRoleSwapOnReplicate(AActor* Actor, bool bSwapRoles) PURE_VIRTUAL(UReplicationDriver::SetRoleSwapOnReplicate, );
+
 	/** Handles an RPC. Returns true if it actually handled it. Returning false will cause the rep driver function to handle it instead */
 	virtual bool ProcessRemoteFunction(class AActor* Actor, UFunction* Function, void* Parameters, FOutParmRec* OutParms, FFrame* Stack, UObject* SubObject ) { return false; }
 
diff --git a/Engine/Source/Runtime/Engine/Private/Actor.cpp b/Engine/Source/Runtime/Engine/Private/Actor.cpp
--- a/Engine/Source/Runtime/Engine/Private/Actor.cpp
+++ b/Engine/Source/Runtime/Engine/Private/Actor.cpp
@@ -3313,7 +3313,6 @@ void AActor::PostNetInit()
 	{
 		UE_LOG(LogActor, Warning, TEXT("AActor::PostNetInit %s Remoterole: %d"), *GetName(), (int)RemoteRole);
 	}
-	check(RemoteRole == ROLE_Authority);
 
 	if (!HasActorBegunPlay())
 	{
diff --git a/Engine/Source/Runtime/Engine/Private/DataChannel.cpp b/Engine/Source/Runtime/Engine/Private/DataChannel.cpp
--- a/Engine/Source/Runtime/Engine/Private/DataChannel.cpp
+++ b/Engine/Source/Runtime/Engine/Private/DataChannel.cpp
@@ -3885,6 +3885,7 @@ void UActorChannel::AddedToChannelPool()
 	bForceCompareProperties = false;
 	bIsReplicatingActor = false;
 	bActorIsPendingKill = false;
+	bSkipRoleSwap = false;
 	bClearRecentActorRefs = true;
 	QueuedBunchStartTime = 0;
 	CreateSubObjects.Empty();
diff --git a/Engine/Source/Runtime/Engine/Private/NetworkDriver.cpp b/Engine/Source/Runtime/Engine/Private/NetworkDriver.cpp
--- a/Engine/Source/Runtime/Engine/Private/NetworkDriver.cpp
+++ b/Engine/Source/Runtime/Engine/Private/NetworkDriver.cpp
@@ -448,6 +448,24 @@ void UNetDriver::AddNetworkActor(AActor* Actor)
 	}
 }
 
+void UNetDriver::SetRoleSwapOnReplicate(AActor* Actor, bool bSwapRoles)
+{
+	TSharedPtr<FNetworkObjectInfo> InfoPtr = GetNetworkObjectList().Find(Actor);
+	if (InfoPtr.IsValid())
+	{
+		InfoPtr->bSwapRolesOnReplicate = bSwapRoles;
+	}
+	else
+	{
+		UE_LOG(LogNet, Warning, TEXT("SetRoleSwapOnReplicate could not find network object info for: %s"), *GetFullNameSafe(Actor));
+	}
+
+	if (ReplicationDriver)
+	{
+		ReplicationDriver->SetRoleSwapOnReplicate(Actor, bSwapRoles);
+	}
+}
+
 FNetworkObjectInfo* UNetDriver::FindOrAddNetworkObjectInfo(const AActor* InActor)
 {
 	ensureMsgf(InActor == nullptr || !(InActor->HasAnyFlags(RF_ClassDefaultObject | RF_ArchetypeObject)), TEXT("%s is a CDO or Archetype and should not be replicated."), *GetFullNameSafe(InActor));
@@ -4163,6 +4181,12 @@ int32 UNetDriver::ServerReplicateActors_ProcessPrioritizedActors( UNetConnection
 			{
 				FinalRelevantCount++;
 
+				TOptional<FScopedActorRoleSwap> SwapGuard;
+				if (ActorInfo->bSwapRolesOnReplicate)
+				{
+					SwapGuard = FScopedActorRoleSwap(Actor);
+				}
+
 				// Find or create the channel for this actor.
 				// we can't create the channel if the client is in a different world than we are
 				// or the package map doesn't support the actor's class/archetype (or the actor itself in the case of serializable actors)
diff --git a/Engine/Source/Runtime/Engine/Public/Net/DataReplication.h b/Engine/Source/Runtime/Engine/Public/Net/DataReplication.h
--- a/Engine/Source/Runtime/Engine/Public/Net/DataReplication.h
+++ b/Engine/Source/Runtime/Engine/Public/Net/DataReplication.h
@@ -297,12 +297,27 @@ private:
 	TWeakObjectPtr<UObject> WeakObjectPtr;
 };
 
-class ENGINE_API FScopedActorRoleSwap : public FNoncopyable
+class ENGINE_API FScopedActorRoleSwap
 {
 public:
 	FScopedActorRoleSwap(AActor* InActor);
 	~FScopedActorRoleSwap();
 
+	FScopedActorRoleSwap(const FScopedActorRoleSwap&) = delete;
+	FScopedActorRoleSwap& operator=(const FScopedActorRoleSwap&) = delete;
+
+	FScopedActorRoleSwap(FScopedActorRoleSwap&& Other)
+	{
+		Actor = Other.Actor;
+		Other.Actor = nullptr;
+	}
+	FScopedActorRoleSwap& operator=(FScopedActorRoleSwap&& Other)
+	{
+		Actor = Other.Actor;
+		Other.Actor = nullptr;
+		return *this;
+	}
+
 private:
 	AActor* Actor;
 };
