Fixed custom blend graphs occasionally creating non-normalized poses

diff --git a/Engine/Source/Runtime/Engine/Private/Animation/AnimNode_TransitionPoseEvaluator.cpp b/Engine/Source/Runtime/Engine/Private/Animation/AnimNode_TransitionPoseEvaluator.cpp
--- a/Engine/Source/Runtime/Engine/Private/Animation/AnimNode_TransitionPoseEvaluator.cpp
+++ b/Engine/Source/Runtime/Engine/Private/Animation/AnimNode_TransitionPoseEvaluator.cpp
@@ -77,5 +77,6 @@ bool FAnimNode_TransitionPoseEvaluator::InputNodeNeedsEvaluate() const
 void FAnimNode_TransitionPoseEvaluator::CachePose(const FPoseContext& PoseToCache)
 {
 	CachedPose.CopyBonesFrom(PoseToCache.Pose);
+	CachedPose.NormalizeRotations();
 	CachedCurve.CopyFrom(PoseToCache.Curve);
 }
