Fixed GetInstanceTransitionTimeElapsed always returning 0 in non-editor builds

diff --git a/Engine/Source/Runtime/Engine/Private/Animation/AnimNode_StateMachine.cpp b/Engine/Source/Runtime/Engine/Private/Animation/AnimNode_StateMachine.cpp
--- a/Engine/Source/Runtime/Engine/Private/Animation/AnimNode_StateMachine.cpp
+++ b/Engine/Source/Runtime/Engine/Private/Animation/AnimNode_StateMachine.cpp
@@ -453,9 +453,7 @@ void FAnimNode_StateMachine::Update_AnyThread(const FAnimationUpdateContext& Con
 					}
 				}
 
-#if WITH_EDITORONLY_DATA
 				NewTransition->SourceTransitionIndices = PotentialTransition.SourceTransitionIndices;
-#endif
 
 				if (!bFirstUpdate)
 				{
