Warning fix

diff --git a/Engine/Source/Runtime/Engine/Private/PrimitiveUniformShaderParameters.cpp b/Engine/Source/Runtime/Engine/Private/PrimitiveUniformShaderParameters.cpp
--- a/Engine/Source/Runtime/Engine/Private/PrimitiveUniformShaderParameters.cpp
+++ b/Engine/Source/Runtime/Engine/Private/PrimitiveUniformShaderParameters.cpp
@@ -42,7 +42,7 @@ void FSinglePrimitiveStructured::UploadToGPU()
 		else
 		{
 			uint32 SrcStride;
-			void* LockedData = RHILockTexture2D(PrimitiveSceneDataTextureRHI, 0, RLM_WriteOnly, SrcStride, false);
+			LockedData = RHILockTexture2D(PrimitiveSceneDataTextureRHI, 0, RLM_WriteOnly, SrcStride, false);
 			FPlatformMemory::Memcpy(LockedData, PrimitiveSceneData.Data, FPrimitiveSceneShaderData::PrimitiveDataStrideInFloat4s * sizeof(FVector4));
 			RHIUnlockTexture2D(PrimitiveSceneDataTextureRHI, 0, false);
 		}
