Generate enums with RepIndex values for native Replicated Properties.
These can be used with simple macro expansions to get Rep Indices for properties without actually looking up UProperty pointers.

diff --git a/Engine/Source/Programs/UnrealHeaderTool/Private/CodeGenerator.cpp b/Engine/Source/Programs/UnrealHeaderTool/Private/CodeGenerator.cpp
--- a/Engine/Source/Programs/UnrealHeaderTool/Private/CodeGenerator.cpp
+++ b/Engine/Source/Programs/UnrealHeaderTool/Private/CodeGenerator.cpp
@@ -75,6 +75,8 @@ FCompilerMetadataManager GScriptHelper;
 /** C++ name lookup helper */
 FNameLookupCPP NameLookupCPP;
 
+bool HasIdentifierExactMatch(const TCHAR* StringBegin, const TCHAR* StringEnd, const FString& Find);
+
 namespace
 {
 	static FString AsTEXT(FString InStr)
@@ -138,6 +140,126 @@ namespace
 
 	/** Guard that should be put at the end of editor only generated code */
 	const TCHAR EndEditorOnlyGuard[] = TEXT("#endif //WITH_EDITOR") LINE_TERMINATOR;
+
+	/** Whether or not the given class has any replicated properties. */
+	static bool ClassHasReplicatedProperties(UClass* Class)
+	{
+		for (TFieldIterator<UProperty> It(Class, EFieldIteratorFlags::ExcludeSuper); It; ++It)
+		{
+			if ((It->PropertyFlags & CPF_Net) != 0)
+			{
+				return true;
+			}
+		}
+
+		return false;
+	}
+
+	static void WriteReplicatedMacroData(
+		const ClassDefinitionRange&	ClassRange,
+		const TCHAR*				ClassCPPName,
+		FClass*						Class,
+		FClass*						SuperClass,
+		FOutputDevice&				Writer,
+		const FUnrealSourceFile&	SourceFile)
+	{
+		const bool bHasGetLifetimeReplicatedProps = HasIdentifierExactMatch(ClassRange.Start, ClassRange.End, TEXT("GetLifetimeReplicatedProps"));
+	
+		if (!bHasGetLifetimeReplicatedProps)
+		{
+			// Default version autogenerates declarations.
+			if (SourceFile.GetGeneratedCodeVersionForStruct(Class) == EGeneratedCodeVersion::V1)
+			{
+				Writer.Logf(TEXT("\tvoid GetLifetimeReplicatedProps(TArray<FLifetimeProperty>& OutLifetimeProps) const override;\r\n"));
+			}
+			else
+			{
+				FError::Throwf(TEXT("Class %s has Net flagged properties and should declare member function: void GetLifetimeReplicatedProps(TArray<FLifetimeProperty>& OutLifetimeProps) const override"), ClassCPPName);
+			}
+		}	
+	}
+
+	static void ExportNetData(FOutputDevice& Out, UClass* Class)
+	{
+		if (ClassHasReplicatedProperties(Class))
+		{
+			Class->SetUpRuntimeReplicationData();
+
+			int32 StartingClassRep = 0;
+			if (UClass * SuperClass = Class->GetSuperClass())
+			{
+				StartingClassRep = SuperClass->ClassReps.Num();
+			}
+
+			// If we don't have any Replicated Properties, then don't do anything.
+			if (Class->ClassReps.Num() == StartingClassRep)
+			{
+				return;
+			}
+
+			FUHTStringBuilder NetFieldBuilder;
+			NetFieldBuilder.Logf(TEXT(""
+			"\t\tenum NetFields\r\n"
+			"\t\t{\r\n"
+			"\t\t\tNETFIELD_REP_START=%d,\r\n"), StartingClassRep);
+
+			FUHTStringBuilder ArrayDimBuilder;
+			ArrayDimBuilder.Log(TEXT(""
+			"\t\tenum ArrayDims\r\n"
+			"\t\t{\r\n"));
+
+			Out.Logf(TEXT("namespace NetworkingPrivate\r\n"
+				"{\r\n"
+				"\tnamespace Net_%s%s\r\n"
+				"\t{\r\n"), Class->GetPrefixCPP(), *Class->GetName());
+
+			TArray<const FRepRecord*> StaticArrayRecords;
+			const TArray<FRepRecord>& ClassReps = Class->ClassReps;
+			const int32 NumClassReps = ClassReps.Num();
+			for (int32 ClassRepIndex = StartingClassRep; ClassRepIndex < ClassReps.Num(); ++ClassRepIndex)
+			{
+				const FRepRecord& ClassRep = ClassReps[ClassRepIndex];
+				const FString& PropertyName = ClassRep.Property->GetName();
+
+				if (ClassRep.Property->ArrayDim == 1)
+				{
+					
+					NetFieldBuilder.Logf(TEXT("\t\t\tNETFIELD_%s=%d,\r\n"), *PropertyName, ClassRepIndex);
+					ArrayDimBuilder.Logf(TEXT("\t\t\tARRAYDIM_%s=1,\r\n"), *PropertyName);
+				}
+				else
+				{
+					const int32 NumClassRepArrayElements = ClassRep.Property->ArrayDim;
+					const int32 ClassRepArrayStart = ClassRepIndex;
+				
+					// -1 because we want this to be the last valid RepIndex.
+					const int32 ClassRepArrayEnd = ClassRepArrayStart + NumClassRepArrayElements - 1;
+
+					for (int32 i = 0; i < NumClassRepArrayElements; ++i)
+					{
+						NetFieldBuilder.Logf(TEXT("\t\t\tNETFIELD_%s_%d=%d,\r\n"), *PropertyName, i, ClassRepIndex + i);
+					}
+				
+					ClassRepIndex += NumClassRepArrayElements - 1;
+					
+					ArrayDimBuilder.Logf(TEXT("\t\t\tARRAYDIM_%s=%d,\r\n"), *PropertyName, NumClassRepArrayElements);
+				}
+			}
+
+			NetFieldBuilder.Logf(TEXT("\t\t\tNETFIELD_REP_END=%d,\r\n"), ClassReps.Num() - 1);
+
+			NetFieldBuilder.Log(TEXT("\t\t};"));
+			ArrayDimBuilder.Log(TEXT("\t\t};"));
+
+			Out.Logf(TEXT(""
+				"%s\r\n" // NetFields
+				"%s\r\n" // ArrayDims
+				"\t}\r\n"
+				"}\r\n"),
+				*NetFieldBuilder,
+				*ArrayDimBuilder);
+		}
+	}
 }
 
 #define BEGIN_WRAP_EDITOR_ONLY(DoWrap) DoWrap ? BeginEditorOnlyGuard : TEXT("")
@@ -2720,15 +2842,6 @@ void FNativeClassHeaderGenerator::ExportClassFromSourceFileInner(
 	FString PPOMacroName;
 
 	// Replication, add in the declaration for GetLifetimeReplicatedProps() automatically if there are any net flagged properties
-	bool bNeedsRep = false;
-	for (TFieldIterator<UProperty> It(Class, EFieldIteratorFlags::ExcludeSuper); It; ++It)
-	{
-		if ((It->PropertyFlags & CPF_Net) != 0)
-		{
-			bNeedsRep = true;
-			break;
-		}
-	}
 
 	ClassDefinitionRange ClassRange;
 	if (ClassDefinitionRange* FoundRange = ClassDefinitionRanges.Find(Class))
@@ -2781,8 +2894,6 @@ void FNativeClassHeaderGenerator::ExportClassFromSourceFileInner(
 		GeneratedSerializeFunctionCPP = BoilerPlateCPP;
 	}
 
-	bool bHasGetLifetimeReplicatedProps = HasIdentifierExactMatch(ClassRange.Start, ClassRange.End, TEXT("GetLifetimeReplicatedProps"));
-
 	{
 		FUHTStringBuilder Boilerplate;
 
@@ -2890,16 +3001,9 @@ void FNativeClassHeaderGenerator::ExportClassFromSourceFileInner(
 				InterfaceBoilerplate.Logf(TEXT("\tvirtual UObject* _getUObject() const { check(0 && \"Missing required implementation.\"); return nullptr; }\r\n"));
 			}
 
-			if (bNeedsRep && !bHasGetLifetimeReplicatedProps)
+			if (ClassHasReplicatedProperties(Class))
 			{
-				if (SourceFile.GetGeneratedCodeVersionForStruct(Class) == EGeneratedCodeVersion::V1)
-				{
-					InterfaceBoilerplate.Logf(TEXT("\tvoid GetLifetimeReplicatedProps(TArray<FLifetimeProperty>& OutLifetimeProps) const override;\r\n"));
-				}
-				else
-				{
-					FError::Throwf(TEXT("Class %s has Net flagged properties and should declare member function: void GetLifetimeReplicatedProps(TArray<FLifetimeProperty>& OutLifetimeProps) const override"), ClassCPPName);
-				}
+				WriteReplicatedMacroData(ClassRange, ClassCPPName, Class, SuperClass, InterfaceBoilerplate, SourceFile);
 			}
 
 			FString NoPureDeclsMacroName = SourceFile.GetGeneratedMacroName(ClassData, TEXT("_INCLASS_IINTERFACE_NO_PURE_DECLS"));
@@ -2924,18 +3028,11 @@ void FNativeClassHeaderGenerator::ExportClassFromSourceFileInner(
 				Boilerplate.Logf(TEXT("\tvirtual UObject* _getUObject() const override { return const_cast<%s*>(this); }\r\n"), ClassCPPName);
 			}
 
-			if (bNeedsRep && !bHasGetLifetimeReplicatedProps)
+			if (ClassHasReplicatedProperties(Class))
 			{
-				// Default version autogenerates declarations.
-				if (SourceFile.GetGeneratedCodeVersionForStruct(Class) == EGeneratedCodeVersion::V1)
-				{
-					Boilerplate.Logf(TEXT("\tvoid GetLifetimeReplicatedProps(TArray<FLifetimeProperty>& OutLifetimeProps) const override;\r\n"));
-				}
-				else
-				{
-					FError::Throwf(TEXT("Class %s has Net flagged properties and should declare member function: void GetLifetimeReplicatedProps(TArray<FLifetimeProperty>& OutLifetimeProps) const override"), ClassCPPName);
-				}
+				WriteReplicatedMacroData(ClassRange, ClassCPPName, Class, SuperClass, Boilerplate, SourceFile);
 			}
+
 			{
 				FString NoPureDeclsMacroName = SourceFile.GetGeneratedMacroName(ClassData, TEXT("_INCLASS_NO_PURE_DECLS"));
 				WriteMacro(OutGeneratedHeaderText, NoPureDeclsMacroName, Boilerplate);
@@ -5433,6 +5530,11 @@ FNativeClassHeaderGenerator::FNativeClassHeaderGenerator(
 			ExportEnum(GeneratedHeaderText, Enum);
 		}
 
+		for (UClass* Class : DefinedClasses)
+		{
+			ExportNetData(GeneratedHeaderText, Class);
+		}
+
 		FString HeaderPath = BaseSourceFilename + TEXT(".generated.h");
 		bool bHasChanged = WriteHeader(*HeaderPath, GeneratedHeaderText, ForwardDeclarations);
 
diff --git a/Engine/Source/Runtime/CoreUObject/Private/UObject/Class.cpp b/Engine/Source/Runtime/CoreUObject/Private/UObject/Class.cpp
--- a/Engine/Source/Runtime/CoreUObject/Private/UObject/Class.cpp
+++ b/Engine/Source/Runtime/CoreUObject/Private/UObject/Class.cpp
@@ -3452,7 +3452,7 @@ void UClass::Link(FArchive& Ar, bool bRelinkExistingProperties)
 	Super::Link(Ar, bRelinkExistingProperties);
 }
 
-#if (UE_BUILD_SHIPPING)
+#if (UE_BUILD_SHIPPING || HACK_HEADER_GENERATOR)
 static int32 GValidateReplicatedProperties = 0;
 #else 
 static int32 GValidateReplicatedProperties = 1;
@@ -3555,7 +3555,7 @@ void UClass::ValidateRuntimeReplicationData()
 {
 	DECLARE_SCOPE_CYCLE_COUNTER(TEXT("Class ValidateRuntimeReplicationData"), STAT_Class_ValidateRuntimeReplicationData, STATGROUP_Game);
 
-	if (HasAnyClassFlags(CLASS_CompiledFromBlueprint))
+	if (HasAnyClassFlags(CLASS_CompiledFromBlueprint|CLASS_LayoutChanging))
 	{
 		// Blueprint classes don't always generate a GetLifetimeReplicatedProps function. 
 		// Assume the Blueprint compiler was ok to do this.
@@ -3568,6 +3568,12 @@ void UClass::ValidateRuntimeReplicationData()
 		return;
 	}
 
+	// Our replication data was set up, but there are no class reps, so there's nothing to do.
+	if (ClassReps.Num() == 0)
+	{
+		return;
+	}
+
 	// Let's compare the CDO's registered lifetime properties with the Class's net properties
 	TArray<FLifetimeProperty> LifetimeProps;
 	LifetimeProps.Reserve(ClassReps.Num());
diff --git a/Engine/Source/Runtime/Engine/Private/RepLayout.cpp b/Engine/Source/Runtime/Engine/Private/RepLayout.cpp
--- a/Engine/Source/Runtime/Engine/Private/RepLayout.cpp
+++ b/Engine/Source/Runtime/Engine/Private/RepLayout.cpp
@@ -941,8 +941,6 @@ FReceivingRepState::FReceivingRepState(FRepStateStaticBuffer&& InStaticBuffer) :
 
 FRepLayout::FRepLayout() :
 	Flags(ERepLayoutFlags::None),
-	RoleIndex(INDEX_NONE),
-	RemoteRoleIndex(-1),
 	Owner(nullptr)
 {}
 
@@ -1012,8 +1010,6 @@ struct FComparePropertiesSharedParams
 {
 	const bool bIsInitial;
 	const bool bForceFail;
-	const int16 RoleIndex;
-	const int16 RemoteRoleIndex;
 	const ERepLayoutFlags Flags;
 	const TArray<FRepParentCmd>& Parents;
 	const TArray<FRepLayoutCmd>& Cmds;
@@ -1062,8 +1058,8 @@ static void CompareRoleProperties(
 {
 	if (RepState && EnumHasAnyFlags(SharedParams.Flags, ERepLayoutFlags::IsActor))
 	{
-		CompareRoleProperty(SharedParams, Data, SharedParams.RemoteRoleIndex, RepState->SavedRemoteRole, Changed);
-		CompareRoleProperty(SharedParams, Data, SharedParams.RoleIndex, RepState->SavedRole, Changed);
+		CompareRoleProperty(SharedParams, Data, NetworkingPrivate::Net_AActor::NETFIELD_RemoteRole, RepState->SavedRemoteRole, Changed);
+		CompareRoleProperty(SharedParams, Data, NetworkingPrivate::Net_AActor::NETFIELD_Role, RepState->SavedRole, Changed);
 	}
 }
 
@@ -1102,14 +1098,14 @@ static void CompareParentProperties(
 
 		if (bCheckForRole)
 		{
-			if (UNLIKELY(ParentIndex == SharedParams.RoleIndex))
+			if (UNLIKELY(ParentIndex == NetworkingPrivate::Net_AActor::NETFIELD_Role))
 			{
-				CompareRoleProperty(SharedParams, Data, SharedParams.RoleIndex, RepState->SavedRole, Changed);
+				CompareRoleProperty(SharedParams, Data, NetworkingPrivate::Net_AActor::NETFIELD_Role, RepState->SavedRole, Changed);
 				continue;
 			}
-			else if (UNLIKELY(ParentIndex == SharedParams.RemoteRoleIndex))
+			else if (UNLIKELY(ParentIndex == NetworkingPrivate::Net_AActor::NETFIELD_RemoteRole))
 			{
-				CompareRoleProperty(SharedParams, Data, SharedParams.RemoteRoleIndex, RepState->SavedRemoteRole, Changed);
+				CompareRoleProperty(SharedParams, Data, NetworkingPrivate::Net_AActor::NETFIELD_RemoteRole, RepState->SavedRemoteRole, Changed);
 				continue;
 			}
 		}
@@ -1240,8 +1236,6 @@ bool FRepLayout::CompareProperties(
 	FComparePropertiesSharedParams SharedParams{
 		/*bIsInitial=*/ !!RepFlags.bNetInitial,
 		/*bForceFail=*/ false,
-		RoleIndex,
-		RemoteRoleIndex,
 		Flags,
 		Parents,
 		Cmds
@@ -2577,8 +2571,26 @@ static bool ReceivePropertyHelper(
 	const FRepLayoutCmd& Cmd = Cmds[CmdIndex];
 	const FRepParentCmd& Parent = Parents[Cmd.ParentIndex];
 
+	auto GetSwappedCmd = [&Cmd, &Cmds, &Parents, bSkipSwapRoles]() -> const FRepLayoutCmd&
+	{
+		if (!bSkipSwapRoles)
+		{
+			// Swap Role to RemoteRole, and vice-versa. Leave everything else the same.
+			if (UNLIKELY(NetworkingPrivate::Net_AActor::NETFIELD_Role == Cmd.ParentIndex))
+			{
+				return Cmds[Parents[NetworkingPrivate::Net_AActor::NETFIELD_RemoteRole].CmdStart];
+			}
+			else if (UNLIKELY(NetworkingPrivate::Net_AActor::NETFIELD_RemoteRole == Cmd.ParentIndex))
+			{
+				return Cmds[Parents[NetworkingPrivate::Net_AActor::NETFIELD_Role].CmdStart];
+			}
+		}
+
+		return Cmd;
+	};
+
 	// This swaps Role/RemoteRole as we write it
-	const FRepLayoutCmd& SwappedCmd = (!bSkipSwapRoles && Parent.RoleSwapIndex != -1) ? Cmds[Parents[Parent.RoleSwapIndex].CmdStart] : Cmd;
+	const FRepLayoutCmd& SwappedCmd = GetSwappedCmd();
 
 	if (GuidReferencesMap)		// Don't reset unmapped guids here if we are told not to (assuming calling code is handling this)
 	{
@@ -5017,8 +5029,6 @@ void FRepLayout::InitFromClass(
 	SCOPE_CYCLE_UOBJECT(ObjectClass, InObjectClass);
 
 	const bool bIsObjectActor = InObjectClass->IsChildOf(AActor::StaticClass());
-	RoleIndex = INDEX_NONE;
-	RemoteRoleIndex = INDEX_NONE;
 
 	if (bIsObjectActor)
 	{
@@ -5066,48 +5076,21 @@ void FRepLayout::InitFromClass(
 			Parents[ParentHandle].Flags |= ERepParentFlags::IsConfig;
 		}
 
-		if (bIsObjectActor)
-		{
-		    // Find Role/RemoteRole property indexes so we can swap them on the client
-		    if (Property->GetFName() == NAME_Role)
-		    {
-			    check(RoleIndex == INDEX_NONE);
-			    check(Parents[ParentHandle].CmdEnd == Parents[ParentHandle].CmdStart + 1);
-			    RoleIndex = ParentHandle;
-		    }
-
-		    if (Property->GetFName() == NAME_RemoteRole)
-		    {
-			    check(RemoteRoleIndex == INDEX_NONE);
-			    check(Parents[ParentHandle].CmdEnd == Parents[ParentHandle].CmdStart + 1);
-			    RemoteRoleIndex = ParentHandle;
-		    }
-	    }
-
 		if (EnumHasAnyFlags(Parents[ParentHandle].Flags, ERepParentFlags::IsCustomDelta))
 		{
 			HighestCustomDeltaRepIndex = ParentHandle;
 		}
 	}
 
-	// Make sure it either found both, or didn't find either
-	check((RoleIndex == INDEX_NONE) == (RemoteRoleIndex == INDEX_NONE));
+	// Make sure RemoteRole has a lower RepIndex than Role, otherwise assumptions RepLayout may break.
+	static_assert(NetworkingPrivate::Net_AActor::NETFIELD_RemoteRole < NetworkingPrivate::Net_AActor::NETFIELD_Role, "Role and RemoteRole have been rearranged in AActor. This will break assumptions in RepLayout.");
 
-	// Make sure that we only find these if we're an Actor, and if we're
-	// an Actor we always find these.
-	check((RoleIndex == INDEX_NONE) == !bIsObjectActor);
+	// Make sure that our RemoteRole property actually points to RemoteRole.
+	check(!bIsObjectActor || Parents[NetworkingPrivate::Net_AActor::NETFIELD_RemoteRole].Property->GetFName() == NAME_RemoteRole);
+
+	// Make sure that our Role property actually points to Role.
+	check(!bIsObjectActor || Parents[NetworkingPrivate::Net_AActor::NETFIELD_Role].Property->GetFName() == NAME_Role);
 
-	// This is so the receiving side can swap these as it receives them
-	if (RoleIndex != -1)
-	{
-		// Make sure that if we have Role and RemoteRole, that Role comes before RemoteRole.
-		// If this fails, it means that the order of Role and RemoteRole has been changed in AActor, and that
-		// will break assumptions RepLayout makes.
-		check(RemoteRoleIndex < RoleIndex);
-		Parents[RoleIndex].RoleSwapIndex = RemoteRoleIndex;
-		Parents[RemoteRoleIndex].RoleSwapIndex = RoleIndex;
-	}
-	
 	AddReturnCmd(Cmds);
 
 	// Initialize lifetime props
@@ -5244,7 +5227,7 @@ void FRepLayout::InitFromClass(
 	{
 		// We handle remote role specially, since it can change between connections when downgraded
 		// So we force it on the conditional list
-		FRepParentCmd& RemoteRoleParent = Parents[RemoteRoleIndex];
+		FRepParentCmd& RemoteRoleParent = Parents[NetworkingPrivate::Net_AActor::NETFIELD_RemoteRole];
 		if (RemoteRoleParent.Condition != COND_Never)
 		{
 			if (COND_None != RemoteRoleParent.Condition)
@@ -5252,8 +5235,8 @@ void FRepLayout::InitFromClass(
 				UE_LOG(LogRep, Warning, TEXT("FRepLayout::InitFromClass: Forcing replication of RemoteRole. Owner=%s"), *InObjectClass->GetPathName());
 			}
 
-			Parents[RemoteRoleIndex].Flags |= ERepParentFlags::IsConditional;
-			Parents[RemoteRoleIndex].Condition = COND_None;
+			Parents[NetworkingPrivate::Net_AActor::NETFIELD_RemoteRole].Flags |= ERepParentFlags::IsConditional;
+			Parents[NetworkingPrivate::Net_AActor::NETFIELD_RemoteRole].Condition = COND_None;
 		}
 	}
 
@@ -6601,8 +6584,6 @@ bool FRepLayout::DeltaSerializeFastArrayProperty(FFastArrayDeltaSerializeParams&
 							FComparePropertiesSharedParams SharedParams{
 								/*bIsInitial=*/ bIsInitial,
 								/*bForceFail=*/ bIsInitial || ShadowArrayItemIsNew[IDIndexPair.Idx],
-								RoleIndex,
-								RemoteRoleIndex,
 								Flags,
 								Parents,
 								Cmds
diff --git a/Engine/Source/Runtime/Engine/Private/UnrealNetwork.cpp b/Engine/Source/Runtime/Engine/Private/UnrealNetwork.cpp
--- a/Engine/Source/Runtime/Engine/Private/UnrealNetwork.cpp
+++ b/Engine/Source/Runtime/Engine/Private/UnrealNetwork.cpp
@@ -23,19 +23,13 @@ void RegisterReplicatedLifetimeProperty(
 PRAGMA_ENABLE_DEPRECATION_WARNINGS
 
 void RegisterReplicatedLifetimeProperty(
-	const UProperty* ReplicatedProperty,
+	const NetworkingPrivate::FRepPropertyDescriptor& PropertyDescriptor,
 	TArray<FLifetimeProperty>& OutLifetimeProps,
 	const FDoRepLifetimeParams& Params)
 {
-	if (!ReplicatedProperty) 
-	{
-		check(false);
-		return;
-	}
-
-	for ( int32 i = 0; i < ReplicatedProperty->ArrayDim; i++ )
+	for (int32 i = 0; i < PropertyDescriptor.ArrayDim; i++)
 	{
-		const uint16 RepIndex = ReplicatedProperty->RepIndex + i;
+		const uint16 RepIndex = PropertyDescriptor.RepIndex + i;
 		FLifetimeProperty* RegisteredPropertyPtr = OutLifetimeProps.FindByPredicate([&RepIndex](const FLifetimeProperty& Var) { return Var.RepIndex == RepIndex; });
 
 		FLifetimeProperty LifetimeProp(RepIndex, Params.Condition, Params.RepNotifyCondition);
@@ -51,7 +45,7 @@ void RegisterReplicatedLifetimeProperty(
 			else
 			{
 				// Conditions should be identical when calling DOREPLIFETIME twice on the same variable.
-				checkf((*RegisteredPropertyPtr) == LifetimeProp, TEXT("Property %s was registered twice with different conditions (old:%d) (new:%d)"), *ReplicatedProperty->GetName(), RegisteredPropertyPtr->Condition, Params.Condition);
+				checkf((*RegisteredPropertyPtr) == LifetimeProp, TEXT("Property %s was registered twice with different conditions (old:%d) (new:%d)"), PropertyDescriptor.PropertyName, RegisteredPropertyPtr->Condition, Params.Condition);
 			}
 		}
 		else
@@ -61,13 +55,25 @@ void RegisterReplicatedLifetimeProperty(
 	}
 }
 
-void SetReplicatedPropertyToDisabled(const UProperty* ReplicatedProperty, TArray< FLifetimeProperty >& OutLifetimeProps)
+void RegisterReplicatedLifetimeProperty(
+	const UProperty* ReplicatedProperty,
+	TArray<FLifetimeProperty>& OutLifetimeProps,
+	const FDoRepLifetimeParams& Params)
 {
-	check(ReplicatedProperty);
+	if (ReplicatedProperty == nullptr)
+	{
+		check(false);
+		return;
+	}
+
+	RegisterReplicatedLifetimeProperty(NetworkingPrivate::FRepPropertyDescriptor(ReplicatedProperty), OutLifetimeProps, Params);
+}
 
-	for (int32 i = 0; i < ReplicatedProperty->ArrayDim; i++)
+void SetReplicatedPropertyToDisabled(const NetworkingPrivate::FRepPropertyDescriptor& PropertyDescriptor, TArray<FLifetimeProperty>& OutLifetimeProps)
+{
+	for (int32 i = 0; i < PropertyDescriptor.ArrayDim; i++)
 	{
-		const uint16 RepIndex = ReplicatedProperty->RepIndex + i;
+		const uint16 RepIndex = PropertyDescriptor.RepIndex + i;
 		FLifetimeProperty* RegisteredPropertyPtr = OutLifetimeProps.FindByPredicate([&RepIndex](const FLifetimeProperty& Var) { return Var.RepIndex == RepIndex; });
 
 		if (RegisteredPropertyPtr)
@@ -81,6 +87,16 @@ void SetReplicatedPropertyToDisabled(const UProperty* ReplicatedProperty, TArray
 	}
 }
 
+void SetReplicatedPropertyToDisabled(const UProperty* ReplicatedProperty, TArray<FLifetimeProperty>& OutLifetimeProps)
+{
+	SetReplicatedPropertyToDisabled(NetworkingPrivate::FRepPropertyDescriptor(ReplicatedProperty), OutLifetimeProps);
+}
+
+void DisableReplicatedLifetimeProperty(const NetworkingPrivate::FRepPropertyDescriptor& PropertyDescriptor, TArray<FLifetimeProperty>& OutLifetimeProps)
+{
+	SetReplicatedPropertyToDisabled(PropertyDescriptor, OutLifetimeProps);
+}
+
 void DisableReplicatedLifetimeProperty(const UClass* ThisClass, const UClass* PropertyClass, FName PropertyName, TArray< FLifetimeProperty >& OutLifetimeProps)
 {
 	const UProperty* ReplicatedProperty = GetReplicatedProperty(ThisClass, PropertyClass, PropertyName);
@@ -89,7 +105,29 @@ void DisableReplicatedLifetimeProperty(const UClass* ThisClass, const UClass* Pr
 		return;
 	}
 
-	SetReplicatedPropertyToDisabled(ReplicatedProperty, OutLifetimeProps);
+	SetReplicatedPropertyToDisabled(NetworkingPrivate::FRepPropertyDescriptor(ReplicatedProperty), OutLifetimeProps);
+}
+
+void ResetReplicatedLifetimeProperty(
+	const NetworkingPrivate::FRepPropertyDescriptor& PropertyDescriptor,
+	ELifetimeCondition LifetimeCondition,
+	TArray<FLifetimeProperty>& OutLifetimeProps)
+{
+	for (int32 i = 0; i < PropertyDescriptor.ArrayDim; i++)
+	{
+		uint16 RepIndex = PropertyDescriptor.RepIndex + i;
+		FLifetimeProperty* RegisteredPropertyPtr = OutLifetimeProps.FindByPredicate([&RepIndex](const FLifetimeProperty& Var) { return Var.RepIndex == RepIndex; });
+
+		// Set the new condition
+		if (RegisteredPropertyPtr)
+		{
+			RegisteredPropertyPtr->Condition = LifetimeCondition;
+		}
+		else
+		{
+			OutLifetimeProps.Add(FLifetimeProperty(RepIndex, LifetimeCondition));
+		}
+	}
 }
 
 void ResetReplicatedLifetimeProperty(const UClass* ThisClass, const UClass* PropertyClass, FName PropertyName, ELifetimeCondition LifetimeCondition, TArray< FLifetimeProperty >& OutLifetimeProps)
@@ -100,19 +138,23 @@ void ResetReplicatedLifetimeProperty(const UClass* ThisClass, const UClass* Prop
 		return;
 	}
 
-	for (int32 i = 0; i < ReplicatedProperty->ArrayDim; i++)
+	ResetReplicatedLifetimeProperty(NetworkingPrivate::FRepPropertyDescriptor(ReplicatedProperty), LifetimeCondition, OutLifetimeProps);
+}
+
+void DisableAllReplicatedPropertiesOfClass(const NetworkingPrivate::FRepClassDescriptor& ClassDescriptor, EFieldIteratorFlags::SuperClassFlags SuperClassBehavior, TArray<FLifetimeProperty>& OutLifetimeProps)
+{
+	const int32 StartIndex = (EFieldIteratorFlags::IncludeSuper == SuperClassBehavior) ? 0 : ClassDescriptor.StartRepIndex;
+	for (int32 RepIndex = StartIndex; RepIndex < ClassDescriptor.EndRepIndex; ++RepIndex)
 	{
-		uint16 RepIndex = ReplicatedProperty->RepIndex + i;
 		FLifetimeProperty* RegisteredPropertyPtr = OutLifetimeProps.FindByPredicate([&RepIndex](const FLifetimeProperty& Var) { return Var.RepIndex == RepIndex; });
 
-		// Set the new condition
 		if (RegisteredPropertyPtr)
 		{
-			RegisteredPropertyPtr->Condition = LifetimeCondition;
+			RegisteredPropertyPtr->Condition = COND_Never;
 		}
 		else
 		{
-			OutLifetimeProps.Add(FLifetimeProperty(RepIndex, LifetimeCondition));
+			OutLifetimeProps.Add(FLifetimeProperty(RepIndex, COND_Never));
 		}
 	}
 }
@@ -135,6 +177,7 @@ void DisableAllReplicatedPropertiesOfClass(const UClass* ThisClass, const UClass
 	}
 }
 
+
 void DeprecatedChangeCondition(const  UProperty* ReplicatedProperty, TArray<FLifetimeProperty>& OutLifetimeProps, ELifetimeCondition InCondition)
 {
 	bool bFound = false;
diff --git a/Engine/Source/Runtime/Engine/Public/Net/RepLayout.h b/Engine/Source/Runtime/Engine/Public/Net/RepLayout.h
--- a/Engine/Source/Runtime/Engine/Public/Net/RepLayout.h
+++ b/Engine/Source/Runtime/Engine/Public/Net/RepLayout.h
@@ -713,7 +713,6 @@ public:
 		ShadowOffset(0),
 		CmdStart(0), 
 		CmdEnd(0), 
-		RoleSwapIndex(-1), 
 		Condition(COND_None),
 		RepNotifyCondition(REPNOTIFY_OnChanged),
 		RepNotifyNumParams(INDEX_NONE),
@@ -749,14 +748,6 @@ public:
 	/** @see CmdStart */
 	uint16 CmdEnd;
 
-	/**
-	 * This value indicates whether or not this command needs to swapped, and what other
-	 * command it should be swapped with.
-	 *
-	 * This is used for Role and RemoteRole which have inverted values on Servers and Clients.
-	 */
-	int32 RoleSwapIndex;
-
 	ELifetimeCondition Condition;
 	ELifetimeRepNotifyCondition RepNotifyCondition;
 
@@ -1765,12 +1756,6 @@ private:
 
 	ERepLayoutFlags Flags;
 
-	/** Index of the Role property in the Parents list. May be INDEX_NONE if Owner doesn't have the property. */
-	int16 RoleIndex;
-
-	/** Index of the RemoteRole property in the Parents list. May be INDEX_NONE if Owner doesn't have the property. */
-	int16 RemoteRoleIndex;
-
 	/** Size (in bytes) needed to allocate a single instance of a Shadow buffer for this RepLayout. */
 	int32 ShadowDataBufferSize;
 
diff --git a/Engine/Source/Runtime/Engine/Public/Net/UnrealNetwork.h b/Engine/Source/Runtime/Engine/Public/Net/UnrealNetwork.h
--- a/Engine/Source/Runtime/Engine/Public/Net/UnrealNetwork.h
+++ b/Engine/Source/Runtime/Engine/Public/Net/UnrealNetwork.h
@@ -70,19 +70,88 @@ struct ENGINE_API FDoRepLifetimeParams
 	ELifetimeRepNotifyCondition RepNotifyCondition = REPNOTIFY_OnChanged;
 };
 
+namespace NetworkingPrivate
+{
+	struct ENGINE_API FRepPropertyDescriptor
+	{
+		FRepPropertyDescriptor(const UProperty* Property)
+			: PropertyName(VerifyPropertyAndGetName(Property))
+			, RepIndex(Property->RepIndex)
+			, ArrayDim(Property->ArrayDim)
+		{
+		}
+
+		FRepPropertyDescriptor(const TCHAR* InPropertyName, const int32 InRepIndex, const int32 InArrayDim)
+			: PropertyName(InPropertyName)
+			, RepIndex(InRepIndex)
+			, ArrayDim(InArrayDim)
+		{
+		}
+
+		const TCHAR* PropertyName;
+		const int32 RepIndex;
+		const int32 ArrayDim;
+
+	private:
+
+		static const TCHAR* VerifyPropertyAndGetName(const UProperty* Property)
+		{
+			check(Property);
+			return *(Property->GetName());
+		}
+
+		UE_NONCOPYABLE(FRepPropertyDescriptor);
+
+		void* operator new(size_t) = delete;
+		void* operator new[](size_t) = delete;
+		void operator delete(void*) = delete;
+		void operator delete[](void*) = delete;
+	};
+
+	struct ENGINE_API FRepClassDescriptor
+	{
+		FRepClassDescriptor(const TCHAR* InClassName, const int32 InStartRepIndex, const int32 InEndRepIndex)
+			: ClassName(InClassName)
+			, StartRepIndex(InStartRepIndex)
+			, EndRepIndex(InEndRepIndex)
+		{
+		}
+
+		const TCHAR* ClassName;
+		const int32 StartRepIndex;
+		const int32 EndRepIndex;
+
+	private:
+
+		UE_NONCOPYABLE(FRepClassDescriptor);
+
+		void* operator new(size_t) = delete;
+		void* operator new[](size_t) = delete;
+		void operator delete(void*) = delete;
+		void operator delete[](void*) = delete;
+	};
+}
 /*-----------------------------------------------------------------------------
 	Replication.
 -----------------------------------------------------------------------------*/
 
-/** wrapper to find replicated properties that also makes sure they're valid */
-static UProperty* GetReplicatedProperty(const UClass* CallingClass, const UClass* PropClass, const FName& PropName)
+static bool ValidateReplicatedClassInheritance(const UClass* CallingClass, const UClass* PropClass, const TCHAR* PropertyName)
 {
 #if !(UE_BUILD_SHIPPING || UE_BUILD_TEST)
 	if (!CallingClass->IsChildOf(PropClass))
 	{
-		UE_LOG(LogNet, Fatal,TEXT("Attempt to replicate property '%s.%s' in C++ but class '%s' is not a child of '%s'"), *PropClass->GetName(), *PropName.ToString(), *CallingClass->GetName(), *PropClass->GetName());
+		UE_LOG(LogNet, Fatal, TEXT("Attempt to replicate property '%s.%s' in C++ but class '%s' is not a child of '%s'"), *PropClass->GetName(), PropertyName, *CallingClass->GetName(), *PropClass->GetName());
 	}
 #endif
+
+	return true;
+}
+
+/** wrapper to find replicated properties that also makes sure they're valid */
+static UProperty* GetReplicatedProperty(const UClass* CallingClass, const UClass* PropClass, const FName& PropName)
+{
+	ValidateReplicatedClassInheritance(CallingClass, PropClass, *PropName.ToString());
+
 	UProperty* TheProperty = FindFieldChecked<UProperty>(PropClass, PropName);
 #if !(UE_BUILD_SHIPPING || UE_BUILD_TEST)
 	if (!(TheProperty->PropertyFlags & CPF_Net))
@@ -93,6 +162,14 @@ static UProperty* GetReplicatedProperty(const UClass* CallingClass, const UClass
 	return TheProperty;
 }
 
+#define DOREPLIFETIME_WITH_PARAMS_FAST(c,v,params) \
+{ \
+	static const bool bIsValid_##c_##v = ValidateReplicatedClassInheritance(StaticClass(), c::StaticClass(), TEXT(#v)); \
+	const TCHAR* DoRepPropertyName_##c_##v(TEXT(#v)); \
+	const NetworkingPrivate::FRepPropertyDescriptor PropertyDescriptor_##c_##v(DoRepPropertyName_##c_##v, NetworkingPrivate::Net_##c::NETFIELD_##v, NetworkingPrivate::Net_##c::ARRAYDIM_##v); \
+	RegisterReplicatedLifetimeProperty(PropertyDescriptor_##c_##v, OutLifetimeProps, params); \
+}
+
 #define DOREPLIFETIME_WITH_PARAMS(c,v,params) \
 { \
 	UProperty* ReplicatedProperty = GetReplicatedProperty(StaticClass(), c::StaticClass(),GET_MEMBER_NAME_CHECKED(c,v)); \
@@ -134,6 +211,15 @@ static UProperty* GetReplicatedProperty(const UClass* CallingClass, const UClass
 }
 
 
+#define DOREPLIFETIME_ACTIVE_OVERRIDE_FAST(c,v,active) \
+{ \
+	static const bool bIsValid_##c_##v = ValidateReplicatedClassInheritance(StaticClass(), c::StaticClass(), TEXT(#v)); \
+	for (int32 i = 0; i < NetworkingPrivate::Net_##c::ARRAYDIM_##v; i++) \
+	{ \
+		ChangedPropertyTracker.SetCustomIsActiveOverride(NetworkingPrivate::Net_##c::NETFIELD_##v + i, active); \
+	} \
+}
+
 #define DOREPLIFETIME_ACTIVE_OVERRIDE(c,v,active)	\
 {													\
 	static UProperty* sp##v = GetReplicatedProperty(StaticClass(), c::StaticClass(),GET_MEMBER_NAME_CHECKED(c,v)); \
@@ -146,6 +232,8 @@ static UProperty* GetReplicatedProperty(const UClass* CallingClass, const UClass
 UE_DEPRECATED(4.24, "Please use the RESET_REPLIFETIME_CONDITION macro")
 ENGINE_API void DeprecatedChangeCondition(UProperty* ReplicatedProperty, TArray< FLifetimeProperty >& OutLifetimeProps, ELifetimeCondition InCondition);
 
+
+//~ This is already using a deprecated method, don't bother updating it.
 #define DOREPLIFETIME_CHANGE_CONDITION(c,v,cond) \
 { \
 	UProperty* sp##v = GetReplicatedProperty(StaticClass(), c::StaticClass(),GET_MEMBER_NAME_CHECKED(c,v));			\
@@ -164,6 +252,11 @@ ENGINE_API void RegisterReplicatedLifetimeProperty(
 	TArray<FLifetimeProperty>& OutLifetimeProps,
 	const FDoRepLifetimeParams& Params);
 
+ENGINE_API void RegisterReplicatedLifetimeProperty(
+	const NetworkingPrivate::FRepPropertyDescriptor& PropertyDescriptor,
+	TArray<FLifetimeProperty>& OutLifetimeProps,
+	const FDoRepLifetimeParams& Params);
+
 /*-----------------------------------------------------------------------------
 	Disable macros.
 	Use these macros to state that properties tagged replicated 
@@ -175,25 +268,64 @@ ENGINE_API void RegisterReplicatedLifetimeProperty(
 #define DISABLE_REPLICATED_PROPERTY(c,v) \
 DisableReplicatedLifetimeProperty(StaticClass(), c::StaticClass(), GET_MEMBER_NAME_CHECKED(c,v), OutLifetimeProps);
 
+#define DISABLE_REPLICATED_PROPERTY_FAST(c,v) \
+{ \
+	static const bool bIsValid_##c_##v = ValidateReplicatedClassInheritance(StaticClass(), c::StaticClass(), TEXT(#v)); \
+	const TCHAR* DoRepPropertyName_##c_##v(TEXT(#v)); \
+	const NetworkingPrivate::FRepPropertyDescriptor PropertyDescriptor_##c_##v(DoRepPropertyName_##c_##v, NetworkingPrivate::Net_##c::NETFIELD_##v, NetworkingPrivate::Net_##c::ARRAYDIM_##v); \
+	DisableReplicatedLifetimeProperty(PropertyDescriptor_##c_##v, OutLifetimeProps); \
+}
+
 /** Use this macro in GetLifetimeReplicatedProps to flag all replicated properties of a class as not-replicated.
     Use the EFieldIteratorFlags enum to disable all inherited properties or only those of the class specified
 */
 #define DISABLE_ALL_CLASS_REPLICATED_PROPERTIES(c, SuperClassBehavior) \
 DisableAllReplicatedPropertiesOfClass(StaticClass(), c::StaticClass(), SuperClassBehavior, OutLifetimeProps);
 
+#define DISABLE_ALL_CLASS_REPLICATED_PROPERTIES_FAST(c, SuperClassBehavior) \
+{ \
+	static const bool bIsValid_##c_##v = ValidateReplicatedClassInheritance(StaticClass(), c::StaticClass(), TEXT("DISABLE_ALL_CLASS_REPLICATED_PROPERTIES")); \
+	const TCHAR* DoRepPropertyName_##c(TEXT(#c)); \
+	const NetworkingPrivate::FRepClassDescriptor ClassDescriptor_##c(DoRepPropertyName_##c, NetworkingPrivate::Net_##c::NETFIELD_REP_START, NetworkingPrivate::Net_##c::NETFIELD_REP_END); \
+	DisableAllReplicatedPropertiesOfClass(StaticClass(), c::StaticClass(), SuperClassBehavior, OutLifetimeProps); \
+}	
+
 ENGINE_API void DisableReplicatedLifetimeProperty(const UClass* ThisClass, const UClass* PropertyClass, FName PropertyName, TArray< FLifetimeProperty >& OutLifetimeProps);
 ENGINE_API void DisableAllReplicatedPropertiesOfClass(const UClass* ThisClass, const UClass* ClassToDisable, EFieldIteratorFlags::SuperClassFlags SuperClassBehavior, TArray< FLifetimeProperty >& OutLifetimeProps);
 
+ENGINE_API void DisableReplicatedLifetimeProperty(const NetworkingPrivate::FRepPropertyDescriptor& PropertyDescriptor, TArray<FLifetimeProperty>& OutLifetimeProps);
+ENGINE_API void DisableAllReplicatedPropertiesOfClass(const NetworkingPrivate::FRepClassDescriptor& ClassDescriptor, EFieldIteratorFlags::SuperClassFlags SuperClassBehavior, TArray<FLifetimeProperty>& OutLifetimeProps);
+
 /*-----------------------------------------------------------------------------
 	Reset macros.
 	Use these to change the replication settings of an inherited property
 -----------------------------------------------------------------------------*/
 
-#define RESET_REPLIFETIME(c,v)  ResetReplicatedLifetimeProperty(StaticClass(), c::StaticClass(), GET_MEMBER_NAME_CHECKED(c,v), COND_None, OutLifetimeProps);
+#define RESET_REPLIFETIME_CONDITION(c,v,cond)  ResetReplicatedLifetimeProperty(StaticClass(), c::StaticClass(), GET_MEMBER_NAME_CHECKED(c,v), cond, OutLifetimeProps);
+
+#define RESET_REPLIFETIME(c,v) RESET_REPLIFETIME_CONDITION(c, v, COND_None)
+
+#define RESET_REPLIFETIME_CONDITION_FAST(c,v,cond) \
+{ \
+	static const bool bIsValid_##c_##v = ValidateReplicatedClassInheritance(StaticClass(), c::StaticClass(), TEXT(#v)); \
+	const TCHAR* DoRepPropertyName_##c_##v(TEXT(#v)); \
+	const NetworkingPrivate::FRepPropertyDescriptor PropertyDescriptor_##c_##v(DoRepPropertyName_##c_##v, NetworkingPrivate::Net_##c::NETFIELD_##v, NetworkingPrivate::Net_##c::ARRAYDIM_##v); \
+	ResetReplicatedLifetimeProperty(StaticClass(), c::StaticClass(), GET_MEMBER_NAME_CHECKED(c,v), cond, OutLifetimeProps); \
+}
+
+#define RESET_REPLIFETIME_FAST(c,v) RESET_REPLIFETIME_CONDITION_FAST(c, v, COND_None)
 
-#define RESET_REPLIFETIME_CONDITION(c,v,cond) ResetReplicatedLifetimeProperty(StaticClass(), c::StaticClass(), GET_MEMBER_NAME_CHECKED(c,v), cond, OutLifetimeProps);
+ENGINE_API void ResetReplicatedLifetimeProperty(
+	const UClass* ThisClass,
+	const UClass* PropertyClass,
+	FName PropertyName,
+	ELifetimeCondition LifetimeCondition,
+	TArray< FLifetimeProperty >& OutLifetimeProps);
 
-ENGINE_API void ResetReplicatedLifetimeProperty(const UClass* ThisClass, const UClass* PropertyClass, FName PropertyName, ELifetimeCondition LifetimeCondition, TArray< FLifetimeProperty >& OutLifetimeProps);
+ENGINE_API void ResetReplicatedLifetimeProperty(
+	const NetworkingPrivate::FRepPropertyDescriptor& PropertyDescriptor,
+	ELifetimeCondition LifetimeCondition,
+	TArray<FLifetimeProperty>& OutLifetimeProps);
 
 /*-----------------------------------------------------------------------------
 	RPC Parameter Validation Helpers
