Fixed HealthSnapshots to properly capture TotalTracked LLM and output the results.

#rnx
#test DevelopmentEditor RunFort.bat -test=profilego

diff --git a/Engine/Source/Runtime/Engine/Private/ProfilingDebugging/HealthSnapshot.cpp b/Engine/Source/Runtime/Engine/Private/ProfilingDebugging/HealthSnapshot.cpp
--- a/Engine/Source/Runtime/Engine/Private/ProfilingDebugging/HealthSnapshot.cpp
+++ b/Engine/Source/Runtime/Engine/Private/ProfilingDebugging/HealthSnapshot.cpp
@@ -46,8 +46,7 @@ void FHealthSnapshot::CaptureMemoryStats()
 #if ENABLE_LOW_LEVEL_MEM_TRACKER
 	if (FLowLevelMemTracker::Get().IsEnabled())
 	{
-		LLMTotalMemoryMB += FLowLevelMemTracker::Get().GetTotalTrackedMemory(ELLMTracker::Default) * InvToMb;
-		LLMTotalMemoryMB += FLowLevelMemTracker::Get().GetTotalTrackedMemory(ELLMTracker::Platform) * InvToMb;
+		LLMTotalMemoryMB = FLowLevelMemTracker::Get().GetTagAmountForTracker(ELLMTracker::Default, ELLMTag::TrackedTotal) * InvToMb;
 	}
 #endif
 
@@ -186,6 +185,12 @@ void FHealthSnapshot::DumpStats(FOutputDevice& Ar, FName CategoryName)
 
 	Ar.CategorizedLogf(CategoryName, ELogVerbosity::Log, TEXT("CPU Memory: Used %.2fMB, Peak %.2fMB"), CPUMemoryMB.Used, CPUMemoryMB.Peak);
 	Ar.CategorizedLogf(CategoryName, ELogVerbosity::Log, TEXT("Physical Memory: Used %.2fMB, Peak %.2fMB"), PhysicalMemoryMB.Used, PhysicalMemoryMB.Peak);
+#if ENABLE_LOW_LEVEL_MEM_TRACKER
+	if (FLowLevelMemTracker::Get().IsEnabled())
+	{
+		Ar.CategorizedLogf(CategoryName, ELogVerbosity::Log, TEXT("Low Level Memory : TrackedTotal %.2fMB"), LLMTotalMemoryMB);
+}
+#endif
 
 #if PLATFORM_PS4
 	Ar.CategorizedLogf(CategoryName, ELogVerbosity::Log, TEXT("Garlic: Used %.2f MB"), GarlicMemoryMB.Used);
