Partialy reverting CL# 9696480 to fix error during post-load serialisation.

diff --git a/Engine/Source/Runtime/Engine/Private/SkeletalMesh.cpp b/Engine/Source/Runtime/Engine/Private/SkeletalMesh.cpp
--- a/Engine/Source/Runtime/Engine/Private/SkeletalMesh.cpp
+++ b/Engine/Source/Runtime/Engine/Private/SkeletalMesh.cpp
@@ -474,36 +474,6 @@ bool USkeletalMesh::IsReductionActive(int32 LODIndex) const
 	IMeshReduction* ReductionModule = FModuleManager::Get().LoadModuleChecked<IMeshReductionManagerModule>("MeshReductionInterface").GetSkeletalMeshReductionInterface();
 	uint32 LODVertexNumber = MAX_uint32;
 	uint32 LODTriNumber = MAX_uint32;
-	const FSkeletalMeshLODInfo* LODInfoPtr = GetLODInfo(LODIndex);
-	bool bLODHasBeenSimplified = LODInfoPtr && LODInfoPtr->bHasBeenSimplified;
-	if (GetImportedModel() && GetImportedModel()->LODModels.IsValidIndex(LODIndex))
-	{
-		if (!bLODHasBeenSimplified)
-		{
-			LODVertexNumber = 0;
-			LODTriNumber = 0;
-			const FSkeletalMeshLODModel& LODModel = GetImportedModel()->LODModels[LODIndex];
-			//We can take the vertices and triangles count from the source model
-			for (int32 SectionIndex = 0; SectionIndex < LODModel.Sections.Num(); ++SectionIndex)
-			{
-				const FSkelMeshSection& Section = LODModel.Sections[SectionIndex];
-
-				if (!Section.bDisabled)
-				{
-					//Make sure the count fit in a uint32
-					LODVertexNumber += Section.NumVertices < 0 ? 0 : Section.NumVertices;
-					LODTriNumber += Section.NumTriangles;
-				}
-			}
-		}
-		else if(GetImportedModel()->OriginalReductionSourceMeshData.IsValidIndex(LODIndex)
-			&& !GetImportedModel()->OriginalReductionSourceMeshData[LODIndex]->IsEmpty())
-		{
-			//In this case we have to use the stored reduction source data to know how many vertices/triangles we have before the reduction
-			USkeletalMesh* MutableSkeletalMesh = const_cast<USkeletalMesh*>(this);
-			GetImportedModel()->OriginalReductionSourceMeshData[LODIndex]->GetGeometryInfo(LODVertexNumber, LODTriNumber, MutableSkeletalMesh);
-		}
-	}
 	return ReductionModule->IsReductionActive(ReductionSettings, LODVertexNumber, LODTriNumber);
 }
 
