Ensure that ProcessInputStack has updated the player input key mappings before starting to process input components

diff --git a/Engine/Source/Runtime/Engine/Private/UserInterface/PlayerInput.cpp b/Engine/Source/Runtime/Engine/Private/UserInterface/PlayerInput.cpp
--- a/Engine/Source/Runtime/Engine/Private/UserInterface/PlayerInput.cpp
+++ b/Engine/Source/Runtime/Engine/Private/UserInterface/PlayerInput.cpp
@@ -985,6 +985,8 @@ void UPlayerInput::ProcessNonAxesKeys(FKey InKey, FKeyState* KeyState)
 
 void UPlayerInput::ProcessInputStack(const TArray<UInputComponent*>& InputComponentStack, const float DeltaTime, const bool bGamePaused)
 {
+	ConditionalBuildKeyMappings();
+
 	// We collect axis contributions by delegate, so we can sum up 
 	// contributions from multiple bindings.
 	struct FAxisDelegateDetails
