Fix replication fallout from CL-9622569

diff --git a/Engine/Source/Runtime/Engine/Private/DataReplication.cpp b/Engine/Source/Runtime/Engine/Private/DataReplication.cpp
--- a/Engine/Source/Runtime/Engine/Private/DataReplication.cpp
+++ b/Engine/Source/Runtime/Engine/Private/DataReplication.cpp
@@ -362,9 +362,7 @@ void FObjectReplicator::InitRecentProperties(uint8* Source)
 		return;
 	}
 
-	const bool bIsReplay = Connection->InternalAck;
-
-	bSupportsFastArrayDelta = !!GSupportsFastArrayDelta && !bIsReplay;
+	bSupportsFastArrayDelta = !!GSupportsFastArrayDelta;
 
 	if (FSendingRepState* SendingRepState = RepState->GetSendingRepState())
 	{
@@ -375,7 +373,9 @@ void FObjectReplicator::InitRecentProperties(uint8* Source)
 		const uint16 NumLifetimeCustomDeltaProperties = FNetSerializeCB::GetNumLifetimeCustomDeltaProperties(LocalRepLayout);
 		SendingRepState->RecentCustomDeltaState.SetNum(NumLifetimeCustomDeltaProperties);
 
-		if (bIsReplay)
+		const bool bIsRecordingReplay = Connection->InternalAck;
+
+		if (bIsRecordingReplay)
 		{
 			SendingRepState->CDOCustomDeltaState.SetNum(NumLifetimeCustomDeltaProperties);
 			SendingRepState->CheckpointCustomDeltaState.SetNum(NumLifetimeCustomDeltaProperties);
@@ -391,7 +391,7 @@ void FObjectReplicator::InitRecentProperties(uint8* Source)
 
 			SendCustomDeltaProperty(UseObject, CustomDeltaProperty, DeltaState, NewState, OldState);
 
-			if (bIsReplay)
+			if (bIsRecordingReplay)
 			{
 				// Store the initial delta state in case we need it for when we're asked to resend all data since channel was first opened (bResendAllDataSinceOpen)
 				SendingRepState->CDOCustomDeltaState[CustomDeltaProperty] = NewState;
