Fix for DDC key mismatch on Mac. FPlatformCompressionUtilities::GetCookOverridesForCurrentPlatform() would return nullptr on Mac, even though its content was cooked out with cook overrides, causing assets to hit a DDC key validation check and early exit.

diff --git a/Engine/Source/Runtime/Engine/Private/AudioCompressionSettingsUtils.cpp b/Engine/Source/Runtime/Engine/Private/AudioCompressionSettingsUtils.cpp
--- a/Engine/Source/Runtime/Engine/Private/AudioCompressionSettingsUtils.cpp
+++ b/Engine/Source/Runtime/Engine/Private/AudioCompressionSettingsUtils.cpp
@@ -4,7 +4,7 @@
 #include "AudioCompressionSettings.h"
 
 #define ENABLE_PLATFORM_COMPRESSION_OVERRIDES 1
-#define PLATFORM_SUPPORTS_COMPRESSION_OVERRIDES (PLATFORM_ANDROID && !PLATFORM_LUMIN) || PLATFORM_IOS || PLATFORM_SWITCH || PLATFORM_XBOXONE || PLATFORM_PS4 || PLATFORM_WINDOWS
+#define PLATFORM_SUPPORTS_COMPRESSION_OVERRIDES (PLATFORM_ANDROID && !PLATFORM_LUMIN) || PLATFORM_IOS || PLATFORM_SWITCH || PLATFORM_XBOXONE || PLATFORM_PS4 || PLATFORM_WINDOWS || PLATFORM_MAC
 
 #if PLATFORM_ANDROID && !PLATFORM_LUMIN && ENABLE_PLATFORM_COMPRESSION_OVERRIDES
 #include "AndroidRuntimeSettings.h"
