#UE4 Added an accessor to for the number of queued bunch net guids for diagnosis of networking issues

diff --git a/Engine/Source/Runtime/Engine/Classes/Engine/PackageMapClient.h b/Engine/Source/Runtime/Engine/Classes/Engine/PackageMapClient.h
--- a/Engine/Source/Runtime/Engine/Classes/Engine/PackageMapClient.h
+++ b/Engine/Source/Runtime/Engine/Classes/Engine/PackageMapClient.h
@@ -507,6 +507,7 @@ private:
 
 public:
 
+	int32 GetNumQueuedBunchNetGUIDs() const;
 	void ConsumeQueuedActorDelinquencyAnalytics(FNetQueuedActorDelinquencyAnalytics& Out);
 	const FNetQueuedActorDelinquencyAnalytics& GetQueuedActorDelinquencyAnalytics() const;
 	void ResetQueuedActorDelinquencyAnalytics();
diff --git a/Engine/Source/Runtime/Engine/Private/PackageMapClient.cpp b/Engine/Source/Runtime/Engine/Private/PackageMapClient.cpp
--- a/Engine/Source/Runtime/Engine/Private/PackageMapClient.cpp
+++ b/Engine/Source/Runtime/Engine/Private/PackageMapClient.cpp
@@ -3524,6 +3524,11 @@ void FPackageMapAckState::CountBytes(FArchive& Ar) const
 	NetFieldExportAcked.CountBytes(Ar);
 }
 
+int32 UPackageMapClient::GetNumQueuedBunchNetGUIDs() const
+{
+	return CurrentQueuedBunchNetGUIDs.Num();
+}
+
 void FNetGUIDCache::ConsumeAsyncLoadDelinquencyAnalytics(FNetAsyncLoadDelinquencyAnalytics& Out)
 {
 	Out = MoveTemp(DelinquentAsyncLoads);
