Fixed: Failed to find material attribute warnings

diff --git a/Engine/Source/Runtime/Engine/Private/Materials/Material.cpp b/Engine/Source/Runtime/Engine/Private/Materials/Material.cpp
--- a/Engine/Source/Runtime/Engine/Private/Materials/Material.cpp
+++ b/Engine/Source/Runtime/Engine/Private/Materials/Material.cpp
@@ -5244,8 +5244,6 @@ bool UMaterial::GetAllReferencedExpressions(TArray<UMaterialExpression*>& OutExp
 				MP_BaseColor,
 				MP_Normal,
 				MP_WorldPositionOffset,
-				MP_WorldDisplacement,
-				MP_TessellationMultiplier,
 			};
 
 
@@ -5307,7 +5305,14 @@ bool UMaterial::GetExpressionsInPropertyChain(EMaterialProperty InProperty,
 	if (StartingExpression->Expression)
 	{
 		ProcessedInputs.AddUnique(StartingExpression);
-		const EShaderFrequency ShaderFrequency = FMaterialAttributeDefinitionMap::GetShaderFrequency(InProperty);
+		
+		EShaderFrequency ShaderFrequency = SF_NumFrequencies;
+		// These properties are "special", attempting to pass them to FMaterialAttributeDefinitionMap::GetShaderFrequency() will generate log spam
+		if (!(InProperty == MP_MaterialAttributes || InProperty == MP_CustomOutput))
+		{
+			ShaderFrequency = FMaterialAttributeDefinitionMap::GetShaderFrequency(InProperty);
+		}
+
 		RecursiveGetExpressionChain(StartingExpression->Expression, ProcessedInputs, OutExpressions, InStaticParameterSet, InFeatureLevel, InQuality, InShadingPath, ShaderFrequency);
 	}
 	return true;
@@ -5452,6 +5457,7 @@ bool UMaterial::RecursiveGetExpressionChain(UMaterialExpression* InExpression, T
 {
 	OutExpressions.AddUnique(InExpression);
 	TArray<FExpressionInput*> Inputs;
+	TArray<EShaderFrequency> InputsFrequency;
 	
 	UMaterialExpressionFeatureLevelSwitch* FeatureLevelSwitchExp;
 	UMaterialExpressionQualitySwitch* QualitySwitchExp;
@@ -5464,10 +5470,12 @@ bool UMaterial::RecursiveGetExpressionChain(UMaterialExpression* InExpression, T
 		if (FeatureLevelSwitchExp->Inputs[InFeatureLevel].IsConnected())
 		{
 			Inputs.Add(&FeatureLevelSwitchExp->Inputs[InFeatureLevel]);
+			InputsFrequency.Add(InShaderFrequency);
 		}
 		else
 		{
 			Inputs.Add(&FeatureLevelSwitchExp->Default);
+			InputsFrequency.Add(InShaderFrequency);
 		}
 	}
 	else if (InQuality != EMaterialQualityLevel::Num && (QualitySwitchExp = Cast<UMaterialExpressionQualitySwitch>(InExpression)) != nullptr)
@@ -5475,10 +5483,12 @@ bool UMaterial::RecursiveGetExpressionChain(UMaterialExpression* InExpression, T
 		if (QualitySwitchExp->Inputs[InQuality].IsConnected())
 		{
 			Inputs.Add(&QualitySwitchExp->Inputs[InQuality]);
+			InputsFrequency.Add(InShaderFrequency);
 		}
 		else
 		{
 			Inputs.Add(&QualitySwitchExp->Default);
+			InputsFrequency.Add(InShaderFrequency);
 		}
 	}
 	else if (InShadingPath != ERHIShadingPath::Num && (ShadingPathSwitchExp = Cast<UMaterialExpressionShadingPathSwitch>(InExpression)) != nullptr)
@@ -5486,10 +5496,12 @@ bool UMaterial::RecursiveGetExpressionChain(UMaterialExpression* InExpression, T
 		if (ShadingPathSwitchExp->Inputs[InShadingPath].IsConnected())
 		{
 			Inputs.Add(&ShadingPathSwitchExp->Inputs[InShadingPath]);
+			InputsFrequency.Add(InShaderFrequency);
 		}
 		else
 		{
 			Inputs.Add(&ShadingPathSwitchExp->Default);
+			InputsFrequency.Add(InShaderFrequency);
 		}
 	}
 	else if (InShaderFrequency != SF_NumFrequencies && (ShaderStageSwitchExp = Cast<UMaterialExpressionShaderStageSwitch>(InExpression)) != nullptr)
@@ -5497,28 +5509,40 @@ bool UMaterial::RecursiveGetExpressionChain(UMaterialExpression* InExpression, T
 		if (UMaterialExpressionShaderStageSwitch::ShouldUsePixelShaderInput(InShaderFrequency))
 		{
 			Inputs.Add(&ShaderStageSwitchExp->PixelShader);
+			InputsFrequency.Add(InShaderFrequency);
 		}
 		else
 		{
 			Inputs.Add(&ShaderStageSwitchExp->VertexShader);
+			InputsFrequency.Add(InShaderFrequency);
 		}
 	}
 	else if (InFeatureLevel <= ERHIFeatureLevel::ES3_1 && (MakeMaterialAttributesExp = Cast<UMaterialExpressionMakeMaterialAttributes>(InExpression)) != nullptr)
 	{
 		// Follow only mobile-relevant inputs
 		Inputs.Add(&MakeMaterialAttributesExp->EmissiveColor);
+		InputsFrequency.Add(SF_Pixel);
 		Inputs.Add(&MakeMaterialAttributesExp->OpacityMask);
+		InputsFrequency.Add(SF_Pixel);
 		Inputs.Add(&MakeMaterialAttributesExp->BaseColor);
+		InputsFrequency.Add(SF_Pixel);
 		Inputs.Add(&MakeMaterialAttributesExp->Normal);
+		InputsFrequency.Add(SF_Pixel);
 		Inputs.Add(&MakeMaterialAttributesExp->WorldPositionOffset);
-		Inputs.Add(&MakeMaterialAttributesExp->WorldDisplacement);
-		Inputs.Add(&MakeMaterialAttributesExp->TessellationMultiplier);
+		InputsFrequency.Add(SF_Vertex);
 	}
 	else
 	{
 		Inputs = InExpression->GetInputs();
+		
+		for (FExpressionInput* Input : Inputs)
+		{
+			InputsFrequency.Add(InShaderFrequency);
+		}
 	}
 
+	check(Inputs.Num() == InputsFrequency.Num());
+
 	for (int32 InputIdx = 0; InputIdx < Inputs.Num(); InputIdx++)
 	{
 		FExpressionInput* InnerInput = Inputs[InputIdx];
@@ -5574,7 +5598,7 @@ bool UMaterial::RecursiveGetExpressionChain(UMaterialExpression* InExpression, T
 					if (bProcessInput == true)
 					{
 						InOutProcessedInputs.Add(InnerInput);
-						RecursiveGetExpressionChain(InnerInput->Expression, InOutProcessedInputs, OutExpressions, InStaticParameterSet, InFeatureLevel, InQuality, InShadingPath, InShaderFrequency);
+						RecursiveGetExpressionChain(InnerInput->Expression, InOutProcessedInputs, OutExpressions, InStaticParameterSet, InFeatureLevel, InQuality, InShadingPath, InputsFrequency[InputIdx]);
 					}
 				}
 			}
