Allow creation of RenderMeshDescription directly from BulkData/RawMesh when MeshDescription is not in memory
Clear MeshDescriptions sooner to reduce memory usage on datasmith import
   - 37.92 GB -> 26.70 GB Peak Working Set when importing DATASET-0008a
   - 66.04 GB -> 52.33 GB Peak Working Set when importing DATASET-0022a

diff --git a/Engine/Plugins/Enterprise/DatasmithImporter/Source/DatasmithImporter/Private/DatasmithStaticMeshImporter.cpp b/Engine/Plugins/Enterprise/DatasmithImporter/Source/DatasmithImporter/Private/DatasmithStaticMeshImporter.cpp
--- a/Engine/Plugins/Enterprise/DatasmithImporter/Source/DatasmithImporter/Private/DatasmithStaticMeshImporter.cpp
+++ b/Engine/Plugins/Enterprise/DatasmithImporter/Source/DatasmithImporter/Private/DatasmithStaticMeshImporter.cpp
@@ -237,6 +237,10 @@ bool FDatasmithStaticMeshImporter::PreBuildStaticMesh( UStaticMesh* StaticMesh )
 			Params.bUseHashAsGuid = true;
 
 			StaticMesh->CommitMeshDescription(LodIndex, Params);
+
+			// Get rid of the memory used now that its committed into its bulk form.
+			// Will be reloaded from bulk data when building the mesh if not present in memory.
+			StaticMesh->ClearMeshDescription(LodIndex);
 		}
 		else
 		{
diff --git a/Engine/Source/Developer/MeshBuilder/Private/MeshDescriptionHelper.cpp b/Engine/Source/Developer/MeshBuilder/Private/MeshDescriptionHelper.cpp
--- a/Engine/Source/Developer/MeshBuilder/Private/MeshDescriptionHelper.cpp
+++ b/Engine/Source/Developer/MeshBuilder/Private/MeshDescriptionHelper.cpp
@@ -30,32 +30,26 @@ FMeshDescriptionHelper::FMeshDescriptionHelper(FMeshBuildSettings* InBuildSettin
 {
 }
 
-void FMeshDescriptionHelper::GetRenderMeshDescription(UObject* Owner, const FMeshDescription& InOriginalMeshDescription, FMeshDescription& OutRenderMeshDescription)
+void FMeshDescriptionHelper::SetupRenderMeshDescription(UObject* Owner, FMeshDescription& RenderMeshDescription)
 {
 	TRACE_CPUPROFILER_EVENT_SCOPE(FMeshDescriptionHelper::GetRenderMeshDescription);
 
 	UStaticMesh* StaticMesh = Cast<UStaticMesh>(Owner);
 	check(StaticMesh);
 
-	//Copy the Original Mesh Description in the render mesh description
-	{
-		TRACE_CPUPROFILER_EVENT_SCOPE(CloneMeshDescription);
-		OutRenderMeshDescription = InOriginalMeshDescription;
-	}
-
 	float ComparisonThreshold = BuildSettings->bRemoveDegenerates ? THRESH_POINTS_ARE_SAME : 0.0f;
 	
 	//This function make sure the Polygon NTB are compute and also remove degenerated triangle from the render mesh description.
-	FMeshDescriptionOperations::CreatePolygonNTB(OutRenderMeshDescription, ComparisonThreshold);
+	FMeshDescriptionOperations::CreatePolygonNTB(RenderMeshDescription, ComparisonThreshold);
 	//OutRenderMeshDescription->ComputePolygonTangentsAndNormals(BuildSettings->bRemoveDegenerates ? SMALL_NUMBER : 0.0f);
 
-	FVertexInstanceArray& VertexInstanceArray = OutRenderMeshDescription.VertexInstances();
-	TVertexInstanceAttributesRef<FVector> Normals = OutRenderMeshDescription.VertexInstanceAttributes().GetAttributesRef<FVector>( MeshAttribute::VertexInstance::Normal );
-	TVertexInstanceAttributesRef<FVector> Tangents = OutRenderMeshDescription.VertexInstanceAttributes().GetAttributesRef<FVector>( MeshAttribute::VertexInstance::Tangent );
-	TVertexInstanceAttributesRef<float> BinormalSigns = OutRenderMeshDescription.VertexInstanceAttributes().GetAttributesRef<float>( MeshAttribute::VertexInstance::BinormalSign );
+	FVertexInstanceArray& VertexInstanceArray = RenderMeshDescription.VertexInstances();
+	TVertexInstanceAttributesRef<FVector> Normals = RenderMeshDescription.VertexInstanceAttributes().GetAttributesRef<FVector>( MeshAttribute::VertexInstance::Normal );
+	TVertexInstanceAttributesRef<FVector> Tangents = RenderMeshDescription.VertexInstanceAttributes().GetAttributesRef<FVector>( MeshAttribute::VertexInstance::Tangent );
+	TVertexInstanceAttributesRef<float> BinormalSigns = RenderMeshDescription.VertexInstanceAttributes().GetAttributesRef<float>( MeshAttribute::VertexInstance::BinormalSign );
 
 	// Find overlapping corners to accelerate adjacency.
-	FMeshDescriptionOperations::FindOverlappingCorners(OverlappingCorners, OutRenderMeshDescription, ComparisonThreshold);
+	FMeshDescriptionOperations::FindOverlappingCorners(OverlappingCorners, RenderMeshDescription, ComparisonThreshold);
 
 	// Compute any missing normals or tangents.
 	{
@@ -101,16 +95,16 @@ void FMeshDescriptionHelper::GetRenderMeshDescription(UObject* Owner, const FMes
 		{
 			if (!bHasAllNormals)
 			{
-				FMeshDescriptionOperations::CreateNormals(OutRenderMeshDescription, (FMeshDescriptionOperations::ETangentOptions)TangentOptions, false);
+				FMeshDescriptionOperations::CreateNormals(RenderMeshDescription, (FMeshDescriptionOperations::ETangentOptions)TangentOptions, false);
 
 				//EComputeNTBsOptions ComputeNTBsOptions = EComputeNTBsOptions::Normals;
 				//OutRenderMeshDescription.ComputeTangentsAndNormals(ComputeNTBsOptions);
 			}
-			FMeshDescriptionOperations::CreateMikktTangents(OutRenderMeshDescription, (FMeshDescriptionOperations::ETangentOptions)TangentOptions);
+			FMeshDescriptionOperations::CreateMikktTangents(RenderMeshDescription, (FMeshDescriptionOperations::ETangentOptions)TangentOptions);
 		}
 		else
 		{
-			FMeshDescriptionOperations::CreateNormals(OutRenderMeshDescription, (FMeshDescriptionOperations::ETangentOptions)TangentOptions, true);
+			FMeshDescriptionOperations::CreateNormals(RenderMeshDescription, (FMeshDescriptionOperations::ETangentOptions)TangentOptions, true);
 			//EComputeNTBsOptions ComputeNTBsOptions = (bHasAllNormals ? EComputeNTBsOptions::None : EComputeNTBsOptions::Normals) | (bHasAllTangents ? EComputeNTBsOptions::None : EComputeNTBsOptions::Tangents);
 			//OutRenderMeshDescription.ComputeTangentsAndNormals(ComputeNTBsOptions);
 		}
@@ -118,7 +112,7 @@ void FMeshDescriptionHelper::GetRenderMeshDescription(UObject* Owner, const FMes
 
 	if (BuildSettings->bGenerateLightmapUVs && VertexInstanceArray.Num() > 0)
 	{
-		TVertexInstanceAttributesRef<FVector2D> VertexInstanceUVs = OutRenderMeshDescription.VertexInstanceAttributes().GetAttributesRef<FVector2D>(MeshAttribute::VertexInstance::TextureCoordinate);
+		TVertexInstanceAttributesRef<FVector2D> VertexInstanceUVs = RenderMeshDescription.VertexInstanceAttributes().GetAttributesRef<FVector2D>(MeshAttribute::VertexInstance::TextureCoordinate);
 		int32 NumIndices = VertexInstanceUVs.GetNumIndices();
 		//Verify the src light map channel
 		if (BuildSettings->SrcLightmapIndex >= NumIndices)
@@ -138,7 +132,7 @@ void FMeshDescriptionHelper::GetRenderMeshDescription(UObject* Owner, const FMes
 			VertexInstanceUVs.SetNumIndices(BuildSettings->DstLightmapIndex + 1);
 			BuildSettings->DstLightmapIndex = NumIndices;
 		}
-		FMeshDescriptionOperations::CreateLightMapUVLayout(OutRenderMeshDescription,
+		FMeshDescriptionOperations::CreateLightMapUVLayout(RenderMeshDescription,
 			BuildSettings->SrcLightmapIndex,
 			BuildSettings->DstLightmapIndex,
 			BuildSettings->MinLightmapResolution,
diff --git a/Engine/Source/Developer/MeshBuilder/Private/MeshDescriptionHelper.h b/Engine/Source/Developer/MeshBuilder/Private/MeshDescriptionHelper.h
--- a/Engine/Source/Developer/MeshBuilder/Private/MeshDescriptionHelper.h
+++ b/Engine/Source/Developer/MeshBuilder/Private/MeshDescriptionHelper.h
@@ -20,8 +20,8 @@ public:
 
 	FMeshDescriptionHelper(FMeshBuildSettings* InBuildSettings);
 
-	//Build a render mesh description with the BuildSettings. This will update the InRenderMeshDescription ptr content
-	void GetRenderMeshDescription(UObject* Owner, const FMeshDescription& InOriginalMeshDescription, FMeshDescription& OutRenderMeshDescription);
+	//Build a render mesh description with the BuildSettings. This will update the RenderMeshDescription in place
+	void SetupRenderMeshDescription(UObject* Owner, FMeshDescription& RenderMeshDescription);
 
 	void ReduceLOD(const FMeshDescription& BaseMesh, FMeshDescription& DestMesh, const struct FMeshReductionSettings& ReductionSettings, const FOverlappingCorners& InOverlappingCorners, float &OutMaxDeviation);
 
diff --git a/Engine/Source/Developer/MeshBuilder/Private/StaticMeshBuilder.cpp b/Engine/Source/Developer/MeshBuilder/Private/StaticMeshBuilder.cpp
--- a/Engine/Source/Developer/MeshBuilder/Private/StaticMeshBuilder.cpp
+++ b/Engine/Source/Developer/MeshBuilder/Private/StaticMeshBuilder.cpp
@@ -92,19 +92,19 @@ bool FStaticMeshBuilder::Build(FStaticMeshRenderData& StaticMeshRenderData, USta
 
 		float MaxDeviation = 0.0f;
 		FMeshBuildSettings& LODBuildSettings = SrcModel.BuildSettings;
-		const FMeshDescription* MeshDescription = StaticMesh->GetMeshDescription(LodIndex);
+		bool bIsMeshDescriptionValid = StaticMesh->CloneMeshDescription(LodIndex, MeshDescriptions[LodIndex]);
 		FMeshDescriptionHelper MeshDescriptionHelper(&LODBuildSettings);
 
 		FMeshReductionSettings ReductionSettings = LODGroup.GetSettings(SrcModel.ReductionSettings, LodIndex);
 
 		//Make sure we do not reduce a non custom LOD by himself
-		const int32 BaseReduceLodIndex = FMath::Clamp<int32>(ReductionSettings.BaseLODModel, 0, MeshDescription == nullptr ? LodIndex - 1 : LodIndex);
+		const int32 BaseReduceLodIndex = FMath::Clamp<int32>(ReductionSettings.BaseLODModel, 0, bIsMeshDescriptionValid ? LodIndex : LodIndex - 1);
 		// Use simplifier if a reduction in triangles or verts has been requested.
 		bool bUseReduction = StaticMesh->IsReductionActive(LodIndex);
 
-		if (MeshDescription != nullptr)
+		if (bIsMeshDescriptionValid)
 		{
-			MeshDescriptionHelper.GetRenderMeshDescription(StaticMesh, *MeshDescription, MeshDescriptions[LodIndex]);
+			MeshDescriptionHelper.SetupRenderMeshDescription(StaticMesh, MeshDescriptions[LodIndex]);
 		}
 		else
 		{
diff --git a/Engine/Source/Runtime/Engine/Classes/Engine/StaticMesh.h b/Engine/Source/Runtime/Engine/Classes/Engine/StaticMesh.h
--- a/Engine/Source/Runtime/Engine/Classes/Engine/StaticMesh.h
+++ b/Engine/Source/Runtime/Engine/Classes/Engine/StaticMesh.h
@@ -790,6 +790,15 @@ public:
 	 * and there is a FRawMesh data for this LODIndex.
 	 */
 	ENGINE_API FMeshDescription* GetMeshDescription(int32 LodIndex) const;
+
+	/* 
+	 * Clone the MeshDescription associated to the LODIndex.
+	 *
+	 * This will make a copy of any pending mesh description that hasn't been committed or will deserialize
+	 * from the bulkdata or rawmesh directly if no current working copy exists.
+	 */
+	ENGINE_API bool CloneMeshDescription(int32 LodIndex, FMeshDescription& OutMeshDescription) const;
+
 	ENGINE_API bool IsMeshDescriptionValid(int32 LodIndex) const;
 	ENGINE_API FMeshDescription* CreateMeshDescription(int32 LodIndex);
 	ENGINE_API FMeshDescription* CreateMeshDescription(int32 LodIndex, FMeshDescription MeshDescription);
@@ -1300,6 +1309,13 @@ private:
 	 */
 	void PostBuildInternal(const TArray<UStaticMeshComponent*>& InAffectedComponents, bool bHasRenderDataChanged);
 
+#if WITH_EDITORONLY_DATA
+	/**
+	 * Deserialize MeshDescription for the specified LodIndex from BulkData, DDC or RawMesh.
+	 */
+	bool LoadMeshDescription(int32 LodIndex, FMeshDescription& OutMeshDescription) const;
+#endif
+
 public:
 	/**
 	 * Caches derived renderable data.
diff --git a/Engine/Source/Runtime/Engine/Private/StaticMesh.cpp b/Engine/Source/Runtime/Engine/Private/StaticMesh.cpp
--- a/Engine/Source/Runtime/Engine/Private/StaticMesh.cpp
+++ b/Engine/Source/Runtime/Engine/Private/StaticMesh.cpp
@@ -3760,68 +3760,102 @@ static FStaticMeshRenderData& GetPlatformStaticMeshRenderData(UStaticMesh* Mesh,
 
 #if WITH_EDITORONLY_DATA
 
-FMeshDescription* UStaticMesh::GetMeshDescription(int32 LodIndex) const
+bool UStaticMesh::LoadMeshDescription(int32 LodIndex, FMeshDescription& OutMeshDescription) const
 {
-	if (!IsSourceModelValid(LodIndex))
-	{
-		return nullptr;
-	}
-
-	// Require a const_cast here, because GetMeshDescription should ostensibly have const semantics,
-	// but the lazy initialization (from the BulkData or the DDC) is a one-off event which breaks constness.
-	UStaticMesh* MutableThis = const_cast<UStaticMesh*>(this);
+	TRACE_CPUPROFILER_EVENT_SCOPE(UStaticMesh::LoadMeshDescription);
 
-	FStaticMeshSourceModel& SourceModel = MutableThis->GetSourceModel(LodIndex);
+	const FStaticMeshSourceModel& SourceModel = GetSourceModel(LodIndex);
 
-	if (!SourceModel.MeshDescription.IsValid())
+	// If we don't have a valid MeshDescription, try and get one...
+	if (SourceModel.MeshDescriptionBulkData.IsValid())
 	{
-		// If we don't have a valid MeshDescription, try and get one...
-		if (SourceModel.MeshDescriptionBulkData.IsValid())
-		{
-			// Unpack MeshDescription from the bulk data which was deserialized
-			SourceModel.MeshDescription = MakeUnique<FMeshDescription>();
-			SourceModel.MeshDescriptionBulkData->LoadMeshDescription(*SourceModel.MeshDescription);
-		}
-		else
+		// Unpack MeshDescription from the bulk data which was deserialized
+		SourceModel.MeshDescriptionBulkData->LoadMeshDescription(OutMeshDescription);
+		return true;
+	}
+	
+	// If BulkData isn't valid, this means either:
+	// a) This LOD doesn't have a MeshDescription (because it's been generated), or;
+	// b) This is a legacy asset which still uses RawMesh, in which case we'll look in the DDC for it.
+	FString MeshDataKey;
+	if (GetMeshDataKey(LodIndex, MeshDataKey))
+	{
+		TArray<uint8> DerivedData;
+		if (GetDerivedDataCacheRef().GetSynchronous(*MeshDataKey, DerivedData))
 		{
-			// If BulkData isn't valid, this means either:
-			// a) This LOD doesn't have a MeshDescription (because it's been generated), or;
-			// b) This is a legacy asset which still uses RawMesh, in which case we'll look in the DDC for it.
-			FString MeshDataKey;
-			if (GetMeshDataKey(LodIndex, MeshDataKey))
-			{
-				TArray<uint8> DerivedData;
-				if (GetDerivedDataCacheRef().GetSynchronous(*MeshDataKey, DerivedData))
-				{
-					// If there was valid DDC data, we assume this is because the asset is an old one with valid RawMeshBulkData
-					check(!SourceModel.RawMeshBulkData->IsEmpty());
+			// If there was valid DDC data, we assume this is because the asset is an old one with valid RawMeshBulkData
+			check(!SourceModel.RawMeshBulkData->IsEmpty());
 
-					// Load from the DDC
-					const bool bIsPersistent = true;
-					FMemoryReader Ar(DerivedData, bIsPersistent);
+			// Load from the DDC
+			const bool bIsPersistent = true;
+			FMemoryReader Ar(DerivedData, bIsPersistent);
 
-					// Create a bulk data object which will be immediately thrown away (as it is not in an archive)
-					FMeshDescriptionBulkData MeshDescriptionBulkData;
-					MeshDescriptionBulkData.Serialize(Ar, const_cast<UStaticMesh*>(this));
+			// Create a bulk data object which will be immediately thrown away (as it is not in an archive)
+			FMeshDescriptionBulkData MeshDescriptionBulkData;
+			MeshDescriptionBulkData.Serialize(Ar, const_cast<UStaticMesh*>(this));
 
-					// Unpack MeshDescription from the bulk data
-					SourceModel.MeshDescription = MakeUnique<FMeshDescription>();
-					MeshDescriptionBulkData.LoadMeshDescription(*SourceModel.MeshDescription);
-				}
-			}
+			// Unpack MeshDescription from the bulk data
+			MeshDescriptionBulkData.LoadMeshDescription(OutMeshDescription);
+			return true;
 		}
 	}
 
 	// If after all this we *still* don't have a valid MeshDescription, but there's a valid RawMesh, convert that to a MeshDescription.
-	if (!SourceModel.MeshDescription.IsValid() && !SourceModel.RawMeshBulkData->IsEmpty())
+	if (!SourceModel.RawMeshBulkData->IsEmpty())
 	{
-		SourceModel.CreateMeshDescription();
-
 		FRawMesh LodRawMesh;
 		SourceModel.LoadRawMesh(LodRawMesh);
 		TMap<int32, FName> MaterialMap;
 		FillMaterialName(StaticMaterials, MaterialMap);
-		FMeshDescriptionOperations::ConvertFromRawMesh(LodRawMesh, *SourceModel.MeshDescription, MaterialMap);
+		FMeshDescriptionOperations::ConvertFromRawMesh(LodRawMesh, OutMeshDescription, MaterialMap);
+		return true;
+	}
+
+	return false;
+}
+
+bool UStaticMesh::CloneMeshDescription(int32 LodIndex, FMeshDescription& OutMeshDescription) const
+{
+	if (!IsSourceModelValid(LodIndex))
+	{
+		return false;
+	}
+
+	TRACE_CPUPROFILER_EVENT_SCOPE(UStaticMesh::CloneMeshDescription);
+	
+	const FStaticMeshSourceModel& SourceModel = GetSourceModel(LodIndex);
+
+	if (SourceModel.MeshDescription.IsValid())
+	{
+		OutMeshDescription = *SourceModel.MeshDescription.Get();
+		return true;
+	}
+
+	return LoadMeshDescription(LodIndex, OutMeshDescription);
+}
+
+FMeshDescription* UStaticMesh::GetMeshDescription(int32 LodIndex) const
+{
+	if (!IsSourceModelValid(LodIndex))
+	{
+		return nullptr;
+	}
+
+	TRACE_CPUPROFILER_EVENT_SCOPE(UStaticMesh::GetMeshDescription);
+
+	// Require a const_cast here, because GetMeshDescription should ostensibly have const semantics,
+	// but the lazy initialization (from the BulkData or the DDC) is a one-off event which breaks constness.
+	UStaticMesh* MutableThis = const_cast<UStaticMesh*>(this);
+
+	FStaticMeshSourceModel& SourceModel = MutableThis->GetSourceModel(LodIndex);
+
+	if (!SourceModel.MeshDescription.IsValid())
+	{
+		FMeshDescription MeshDescription;
+		if (LoadMeshDescription(LodIndex, MeshDescription))
+		{
+			SourceModel.MeshDescription = MakeUnique<FMeshDescription>(MoveTemp(MeshDescription));
+		}
 	}
 
 	return SourceModel.MeshDescription.Get();
