Fix from Mattias Hornlund to avoid crash on application teardown when physics singleton is not initialized anymore

diff --git a/Engine/Source/Runtime/Engine/Private/PhysicsEngine/BodySetup.cpp b/Engine/Source/Runtime/Engine/Private/PhysicsEngine/BodySetup.cpp
--- a/Engine/Source/Runtime/Engine/Private/PhysicsEngine/BodySetup.cpp
+++ b/Engine/Source/Runtime/Engine/Private/PhysicsEngine/BodySetup.cpp
@@ -530,7 +530,7 @@ void UBodySetup::FinishCreatingPhysicsMeshes_PhysX(const TArray<PxConvexMesh*>&
 {
 	ClearPhysicsMeshes();
 
-	FPhysxSharedData::Get().LockAccess();
+	FPhysxSharedData::LockAccess();
 
 	const FString FullName = GetFullName();
 	if (GetCollisionTraceFlag() != CTF_UseComplexAsSimple)
@@ -571,7 +571,7 @@ void UBodySetup::FinishCreatingPhysicsMeshes_PhysX(const TArray<PxConvexMesh*>&
 		}
 	}
 
-	FPhysxSharedData::Get().UnlockAccess();
+	FPhysxSharedData::UnlockAccess();
 
 	// Clear the cooked data
 	if (!GIsEditor && !bSharedCookedData)
@@ -733,7 +733,7 @@ void UBodySetup::ClearPhysicsMeshes()
 {
 #if WITH_PHYSX && PHYSICS_INTERFACE_PHYSX
 
-	FPhysxSharedData::Get().LockAccess();
+	FPhysxSharedData::LockAccess();
 
 	for(int32 i=0; i<AggGeom.ConvexElems.Num(); i++)
 	{
@@ -763,7 +763,7 @@ void UBodySetup::ClearPhysicsMeshes()
 		TriMeshes[ElementIndex] = NULL;
 	}
 
-	FPhysxSharedData::Get().UnlockAccess();
+	FPhysxSharedData::UnlockAccess();
 
 	TriMeshes.Empty();
 
diff --git a/Engine/Source/Runtime/PhysicsCore/Private/PhysXSupportCore.cpp b/Engine/Source/Runtime/PhysicsCore/Private/PhysXSupportCore.cpp
--- a/Engine/Source/Runtime/PhysicsCore/Private/PhysXSupportCore.cpp
+++ b/Engine/Source/Runtime/PhysicsCore/Private/PhysXSupportCore.cpp
@@ -162,12 +162,18 @@ void FPhysxSharedData::Terminate()
 
 void FPhysxSharedData::LockAccess()
 {
-	CriticalSection.Lock();
+	if (Singleton)
+	{
+		Singleton->CriticalSection.Lock();
+	}
 }
 
 void FPhysxSharedData::UnlockAccess()
 {
-	CriticalSection.Unlock();
+	if (Singleton)
+	{
+		Singleton->CriticalSection.Unlock();
+	}
 }
 
 void FPhysxSharedData::Add( PxBase* Obj, const FString& OwnerName )
diff --git a/Engine/Source/Runtime/PhysicsCore/Public/PhysXSupportCore.h b/Engine/Source/Runtime/PhysicsCore/Public/PhysXSupportCore.h
--- a/Engine/Source/Runtime/PhysicsCore/Public/PhysXSupportCore.h
+++ b/Engine/Source/Runtime/PhysicsCore/Public/PhysXSupportCore.h
@@ -248,8 +248,8 @@ public:
 	static void Initialize();
 	static void Terminate();
 
-	void LockAccess();
-	void UnlockAccess();
+	static void LockAccess();
+	static void UnlockAccess();
 
 	void Add(PxBase* Obj, const FString& OwnerName);
 	void Remove(PxBase* Obj);
