UE4R - Pack cubemap data into NumSlices to help reduce patch size
#rnx

diff --git a/Engine/Plugins/Runtime/Steam/SteamVR/Source/SteamVR/Private/SteamVRAssetManager.cpp b/Engine/Plugins/Runtime/Steam/SteamVR/Source/SteamVR/Private/SteamVRAssetManager.cpp
--- a/Engine/Plugins/Runtime/Steam/SteamVR/Source/SteamVR/Private/SteamVRAssetManager.cpp
+++ b/Engine/Plugins/Runtime/Steam/SteamVR/Source/SteamVR/Private/SteamVRAssetManager.cpp
@@ -378,7 +378,7 @@ public:
 				void* Data = Mip.BulkData.Lock(LOCK_READ_WRITE);
 				FMemory::Memcpy(Data, (void*)RawResource->rubTextureMapData, Mip.BulkData.GetBulkDataSize());
 				Mip.BulkData.Unlock();
-				NewTexture->PlatformData->NumSlices = 1;
+				NewTexture->PlatformData->SetNumSlices(1);
 #if WITH_EDITORONLY_DATA
 				NewTexture->CompressionNone = true;
 				NewTexture->DeferCompression = false;
diff --git a/Engine/Source/Runtime/Engine/Classes/Engine/Texture.h b/Engine/Source/Runtime/Engine/Classes/Engine/Texture.h
--- a/Engine/Source/Runtime/Engine/Classes/Engine/Texture.h
+++ b/Engine/Source/Runtime/Engine/Classes/Engine/Texture.h
@@ -427,16 +427,14 @@ struct FTexturePlatformData
 	int32 SizeX;
 	/** Height of the texture. */
 	int32 SizeY;
-	/** Number of texture slices. */
-	int32 NumSlices;
+	/** Number of texture slices. Bit CubeMapBitIndex contains if a texture is a Cubemap or not. */
+	uint32 NumSlicesCubemapMask;
 	/** Format in which mip data is stored. */
 	EPixelFormat PixelFormat;
 	/** Optional extra data that the runtime may need. */
 	uint32 ExtData;
 	/** Number of mips making up the mip tail, which must always be resident */
 	uint32 NumMipsInTail;
-	/** Is Cubemap texture */
-	bool bCubemap;
 	/** Mip data or VT data. one or the other. */
 	TIndirectArray<struct FTexture2DMipMap> Mips;
 	struct FVirtualTextureBuiltData* VTData;
@@ -459,6 +457,10 @@ struct FTexturePlatformData
 	/** Destructor. */
 	ENGINE_API ~FTexturePlatformData();
 
+	static constexpr uint32 CubeMapBitIndex = 31;
+	static constexpr uint32 CubeMapBitMask = (1u << CubeMapBitIndex);
+	static constexpr uint32 CubeMapRemainingBitMask = CubeMapBitMask - 1u;
+
 	/** Return whether TryLoadMips() would stall because async loaded mips are not yet available. */
 	bool IsReadyForAsyncPostLoad() const;
 
@@ -484,6 +486,26 @@ struct FTexturePlatformData
 	 */
 	void SerializeCooked(FArchive& Ar, class UTexture* Owner, bool bStreamable);
 
+	inline bool IsCubemap() const
+	{
+		return (NumSlicesCubemapMask & CubeMapBitMask) == CubeMapBitMask;
+	}
+
+	inline int32 GetNumSlices() const
+	{
+		return (int32)(NumSlicesCubemapMask & CubeMapRemainingBitMask);
+	}
+
+	inline void SetNumSlices(int32 NumSlices)
+	{
+		NumSlicesCubemapMask = ((uint32)NumSlices & CubeMapRemainingBitMask) | (NumSlicesCubemapMask & CubeMapBitMask);
+	}
+
+	inline void SetIsCubemap(bool bCubemap)
+	{
+		NumSlicesCubemapMask = (bCubemap ? CubeMapBitMask : 0) | (NumSlicesCubemapMask & CubeMapRemainingBitMask);
+	}
+
 #if WITH_EDITOR
 	void Cache(
 		class UTexture& InTexture,
diff --git a/Engine/Source/Runtime/Engine/Classes/Engine/Texture2DArray.h b/Engine/Source/Runtime/Engine/Classes/Engine/Texture2DArray.h
--- a/Engine/Source/Runtime/Engine/Classes/Engine/Texture2DArray.h
+++ b/Engine/Source/Runtime/Engine/Classes/Engine/Texture2DArray.h
@@ -36,7 +36,7 @@ public:
 	}
 	FORCEINLINE int32 GetNumSlices() const
 	{
-		return PlatformData ? PlatformData->NumSlices : 0;
+		return PlatformData ? PlatformData->GetNumSlices() : 0;
 	}
 	FORCEINLINE int32 GetNumMips() const
 	{
diff --git a/Engine/Source/Runtime/Engine/Classes/Engine/VolumeTexture.h b/Engine/Source/Runtime/Engine/Classes/Engine/VolumeTexture.h
--- a/Engine/Source/Runtime/Engine/Classes/Engine/VolumeTexture.h
+++ b/Engine/Source/Runtime/Engine/Classes/Engine/VolumeTexture.h
@@ -68,7 +68,7 @@ public:
 	}
 	FORCEINLINE int32 GetSizeZ() const
 	{
-		return PlatformData ? PlatformData->NumSlices : 0;
+		return PlatformData ? PlatformData->GetNumSlices() : 0;
 	}
 	FORCEINLINE int32 GetNumMips() const
 	{
diff --git a/Engine/Source/Runtime/Engine/Private/TextureDerivedData.cpp b/Engine/Source/Runtime/Engine/Private/TextureDerivedData.cpp
--- a/Engine/Source/Runtime/Engine/Private/TextureDerivedData.cpp
+++ b/Engine/Source/Runtime/Engine/Private/TextureDerivedData.cpp
@@ -47,7 +47,7 @@
 // In case of merge conflicts with DDC versions, you MUST generate a new GUID and set this new
 // guid as version
 
-#define TEXTURE_DERIVEDDATA_VER		TEXT("51C887FFAD4D4179AECB2FA5D9D6A0F8")
+#define TEXTURE_DERIVEDDATA_VER		TEXT("A52A75F077004A798078A42DCF5502E9")
 
 // This GUID is mixed into DDC version for virtual textures only, this allows updating DDC version for VT without invalidating DDC for all textures
 // This is useful during development, but once large numbers of VT are present in shipped content, it will have the same problem as TEXTURE_DERIVEDDATA_VER
@@ -571,7 +571,7 @@ uint32 PutDerivedDataInCache(FTexturePlatformData* DerivedData, const FString& D
 
 	// Write out individual mips to the derived data cache.
 	const int32 MipCount = DerivedData->Mips.Num();
-	const bool bIsCubemap = DerivedData->bCubemap;
+	const bool bIsCubemap = DerivedData->IsCubemap();
 	const int32 FirstInlineMip = bIsCubemap ? 0 : FMath::Max(0, MipCount - FMath::Max((int32)NUM_INLINE_DERIVED_MIPS, (int32)DerivedData->NumMipsInTail));
 	const int32 WritableMipCount = MipCount - ((DerivedData->NumMipsInTail > 0) ? (DerivedData->NumMipsInTail - 1) : 0);
 	for (int32 MipIndex = 0; MipIndex < WritableMipCount; ++MipIndex)
@@ -1010,11 +1010,10 @@ bool FTexturePlatformData::TryInlineMipData(int32 FirstMipToLoad)
 FTexturePlatformData::FTexturePlatformData()
 	: SizeX(0)
 	, SizeY(0)
-	, NumSlices(0)
+	, NumSlicesCubemapMask(0)
 	, PixelFormat(PF_Unknown)
 	, ExtData(0)
 	, NumMipsInTail(0)
-	, bCubemap(false)
 	, VTData(nullptr)
 #if WITH_EDITORONLY_DATA
 	, AsyncTask(NULL)
@@ -1251,7 +1250,7 @@ static void SerializePlatformData(
 
 	Ar << PlatformData->SizeX;
 	Ar << PlatformData->SizeY;
-	Ar << PlatformData->NumSlices;
+	Ar << PlatformData->NumSlicesCubemapMask;
 	if (Ar.IsLoading())
 	{
 		FString PixelFormatString;
@@ -1265,7 +1264,6 @@ static void SerializePlatformData(
 	}
  	Ar << PlatformData->ExtData;
  	Ar << PlatformData->NumMipsInTail;
-	Ar << PlatformData->bCubemap;
 
 	int32 NumMips = PlatformData->Mips.Num();
 	int32 FirstMipToSerialize = 0;
@@ -1455,7 +1453,7 @@ void FTexturePlatformData::SerializeCooked(FArchive& Ar, UTexture* Owner, bool b
 			// SizeZ is not the same as NumSlices for texture arrays and cubemaps.
 			if (Owner && Owner->IsA(UVolumeTexture::StaticClass()))
 			{
-				 NumSlices = Mips[0].SizeZ;
+				SetNumSlices(Mips[0].SizeZ);
 			}
 		}
 		else if ( VTData )
diff --git a/Engine/Source/Runtime/Engine/Private/TextureDerivedDataTask.cpp b/Engine/Source/Runtime/Engine/Private/TextureDerivedDataTask.cpp
--- a/Engine/Source/Runtime/Engine/Private/TextureDerivedDataTask.cpp
+++ b/Engine/Source/Runtime/Engine/Private/TextureDerivedDataTask.cpp
@@ -222,7 +222,7 @@ void FTextureCacheDerivedDataWorker::BuildTexture()
 		DerivedData->SizeX = DerivedData->VTData->Width;
 		DerivedData->SizeY = DerivedData->VTData->Height;
 		DerivedData->PixelFormat = DerivedData->VTData->LayerTypes[0];
-		DerivedData->NumSlices = 1;
+		DerivedData->SetNumSlices(1);
 
 		// Store it in the cache.
 		// @todo: This will remove the streaming bulk data, which we immediately reload below!
@@ -252,7 +252,7 @@ void FTextureCacheDerivedDataWorker::BuildTexture()
 		DerivedData->SizeX = 0;
 		DerivedData->SizeY = 0;
 		DerivedData->PixelFormat = PF_Unknown;
-		DerivedData->bCubemap = false;
+		DerivedData->SetIsCubemap(false);
 		DerivedData->VTData = nullptr;
 
 		// Compress the texture.
@@ -288,8 +288,8 @@ void FTextureCacheDerivedDataWorker::BuildTexture()
 					DerivedData->SizeX = CompressedImage.SizeX;
 					DerivedData->SizeY = CompressedImage.SizeY;
 					DerivedData->PixelFormat = (EPixelFormat)CompressedImage.PixelFormat;
-					DerivedData->NumSlices = BuildSettingsPerLayer[0].bCubemap ? 6 : (BuildSettingsPerLayer[0].bVolume || BuildSettingsPerLayer[0].bTextureArray) ? CompressedImage.SizeZ : 1;
-					DerivedData->bCubemap = BuildSettingsPerLayer[0].bCubemap;
+					DerivedData->SetNumSlices(BuildSettingsPerLayer[0].bCubemap ? 6 : (BuildSettingsPerLayer[0].bVolume || BuildSettingsPerLayer[0].bTextureArray) ? CompressedImage.SizeZ : 1);
+					DerivedData->SetIsCubemap(BuildSettingsPerLayer[0].bCubemap);
 				}
 				else
 				{
diff --git a/Engine/Source/Runtime/RuntimeAssetCache/Private/RuntimeAssetCacheBuilders.cpp b/Engine/Source/Runtime/RuntimeAssetCache/Private/RuntimeAssetCacheBuilders.cpp
--- a/Engine/Source/Runtime/RuntimeAssetCache/Private/RuntimeAssetCacheBuilders.cpp
+++ b/Engine/Source/Runtime/RuntimeAssetCache/Private/RuntimeAssetCacheBuilders.cpp
@@ -125,7 +125,7 @@ void UExampleTextureCacheBuilder::SerializeAsset(FArchive& Ar)
 
 		Ar << PlatformData->SizeX;
 		Ar << PlatformData->SizeY;
-		Ar << PlatformData->NumSlices;
+		Ar << PlatformData->NumSlicesCubemapMask;
 		if (Ar.IsLoading())
 		{
 			FString PixelFormatString;
