Add dynamic volume mix ability for SoundClasses

diff --git a/Engine/Source/Runtime/Engine/Private/Audio.cpp b/Engine/Source/Runtime/Engine/Private/Audio.cpp
--- a/Engine/Source/Runtime/Engine/Private/Audio.cpp
+++ b/Engine/Source/Runtime/Engine/Private/Audio.cpp
@@ -20,6 +20,7 @@
 #include "Misc/Paths.h"
 #include "Sound/SoundBase.h"
 #include "Sound/SoundCue.h"
+#include "Sound/SoundSubmix.h"
 #include "Sound/SoundNodeWavePlayer.h"
 #include "Sound/SoundWave.h"
 #include "UObject/UObjectHash.h"
@@ -1010,6 +1011,11 @@ float FWaveInstance::GetDynamicVolume() const
 					OutVolume *= DeviceManager->GetDynamicSoundVolume(ESoundType::Cue, Sound->GetFName());
 				}
 			}
+
+			if (SoundClass)
+			{
+				OutVolume *= DeviceManager->GetDynamicSoundVolume(ESoundType::Class, SoundClass->GetFName());
+			}
 		}
 	}
 
diff --git a/Engine/Source/Runtime/Engine/Private/AudioDevice.cpp b/Engine/Source/Runtime/Engine/Private/AudioDevice.cpp
--- a/Engine/Source/Runtime/Engine/Private/AudioDevice.cpp
+++ b/Engine/Source/Runtime/Engine/Private/AudioDevice.cpp
@@ -1758,9 +1758,16 @@ bool FAudioDevice::HandleSetDynamicSoundCommand(const TCHAR* Cmd, FOutputDevice&
 		// Optional: Defaults to Cue
 		FString SoundTypeStr;
 		ESoundType SoundType = ESoundType::Cue;
-		if (FParse::Value(Cmd, TEXT("Type="), SoundTypeStr) && SoundTypeStr == TEXT("Wave"))
+		if (FParse::Value(Cmd, TEXT("Type="), SoundTypeStr))
 		{
-			SoundType = ESoundType::Wave;
+			if (SoundTypeStr == TEXT("Wave"))
+			{
+				SoundType = ESoundType::Wave;
+			}
+			else if (SoundTypeStr == TEXT("Class"))
+			{
+				SoundType = ESoundType::Class;
+			}
 		}
 
 		float Volume;
diff --git a/Engine/Source/Runtime/Engine/Public/AudioDeviceManager.h b/Engine/Source/Runtime/Engine/Public/AudioDeviceManager.h
--- a/Engine/Source/Runtime/Engine/Public/AudioDeviceManager.h
+++ b/Engine/Source/Runtime/Engine/Public/AudioDeviceManager.h
@@ -25,6 +25,7 @@ struct FSourceEffectChainEntry;
 
 enum class ESoundType : uint8
 {
+	Class,
 	Cue,
 	Wave
 };
