Add dynamic volume get/reset ability for SoundClasses

diff --git a/Engine/Source/Runtime/Engine/Private/AudioDevice.cpp b/Engine/Source/Runtime/Engine/Private/AudioDevice.cpp
--- a/Engine/Source/Runtime/Engine/Private/AudioDevice.cpp
+++ b/Engine/Source/Runtime/Engine/Private/AudioDevice.cpp
@@ -1710,9 +1710,16 @@ bool FAudioDevice::HandleResetDynamicSoundVolumeCommand(const TCHAR* Cmd, FOutpu
 		// Optional: Defaults to Cue
 		FString SoundTypeStr;
 		ESoundType SoundType = ESoundType::Cue;
-		if (FParse::Value(Cmd, TEXT("Type="), SoundTypeStr) && SoundTypeStr == TEXT("Wave"))
+		if (FParse::Value(Cmd, TEXT("Type="), SoundTypeStr))
 		{
-			SoundType = ESoundType::Wave;
+			if (SoundTypeStr == TEXT("Wave"))
+			{
+				SoundType = ESoundType::Wave;
+			}
+			else if (SoundTypeStr == TEXT("Class"))
+			{
+				SoundType = ESoundType::Class;
+			}
 		}
 
 		DeviceManager->ResetDynamicSoundVolume(SoundType, SoundName);
@@ -1733,9 +1740,16 @@ bool FAudioDevice::HandleGetDynamicSoundVolumeCommand(const TCHAR* Cmd, FOutputD
 		// Optional: Defaults to Cue
 		FString SoundTypeStr;
 		ESoundType SoundType = ESoundType::Cue;
-		if (FParse::Value(Cmd, TEXT("Type="), SoundTypeStr) && SoundTypeStr == TEXT("Wave"))
+		if (FParse::Value(Cmd, TEXT("Type="), SoundTypeStr))
 		{
-			SoundType = ESoundType::Wave;
+			if (SoundTypeStr == TEXT("Wave"))
+			{
+				SoundType = ESoundType::Wave;
+			}
+			else if (SoundTypeStr == TEXT("Class"))
+			{
+				SoundType = ESoundType::Class;
+			}
 		}
 
 		const float Volume = DeviceManager->GetDynamicSoundVolume(SoundType, SoundName);
