Added texture gropu TEXTUREGROUP_Project11 to TEXTUREGROUP_Project15

diff --git a/Engine/Source/Runtime/Engine/Classes/Engine/TextureDefines.h b/Engine/Source/Runtime/Engine/Classes/Engine/TextureDefines.h
--- a/Engine/Source/Runtime/Engine/Classes/Engine/TextureDefines.h
+++ b/Engine/Source/Runtime/Engine/Classes/Engine/TextureDefines.h
@@ -74,6 +74,11 @@ enum TextureGroup
 	TEXTUREGROUP_Project08 UMETA(DisplayName="ini:Project Group 08"),
 	TEXTUREGROUP_Project09 UMETA(DisplayName="ini:Project Group 09"),
 	TEXTUREGROUP_Project10 UMETA(DisplayName="ini:Project Group 10"),
+	TEXTUREGROUP_Project11 UMETA(DisplayName="ini:Project Group 11"),
+	TEXTUREGROUP_Project12 UMETA(DisplayName="ini:Project Group 12"),
+	TEXTUREGROUP_Project13 UMETA(DisplayName="ini:Project Group 13"),
+	TEXTUREGROUP_Project14 UMETA(DisplayName="ini:Project Group 14"),
+	TEXTUREGROUP_Project15 UMETA(DisplayName="ini:Project Group 15"),
 	TEXTUREGROUP_MAX,
 };
 
