[NO NON-PII COMMIT MESSAGE]

diff --git a/Engine/Source/Runtime/Engine/Classes/Sound/SoundClass.h b/Engine/Source/Runtime/Engine/Classes/Sound/SoundClass.h
--- a/Engine/Source/Runtime/Engine/Classes/Sound/SoundClass.h
+++ b/Engine/Source/Runtime/Engine/Classes/Sound/SoundClass.h
@@ -22,6 +22,32 @@ namespace EAudioOutputTarget
 	};
 }
 
+/**
+ * Only used when stream caching is enabled. Determines how we are going to load or retain a given audio asset.
+ * A USoundWave's loading behavior can be overridden in the USoundWave itself, the sound wave's USoundClass, or by cvars.
+ * The order of priority is defined as:
+ * 1) The loading behavior set on the USoundWave
+ * 2) The loading behavior set on the USoundWave's USoundClass.
+ * 3) The loading behavior set on the nearest parent of a USoundWave's USoundClass.
+ * 4) The loading behavior set via the au.streamcache cvars.
+ */
+UENUM()
+enum class ESoundWaveLoadingBehavior : uint8
+{
+	// If set on a USoundWave, use the setting defined by the USoundClass. If set on the next parent USoundClass, or the default behavior defined via the au.streamcache cvars.
+	Inherited = 0,
+	// the first chunk of audio for this asset will be retained in the audio cache until a given USoundWave is either destroyed or USoundWave::ReleaseCompressedAudioData is called.
+	RetainOnLoad = 1,
+	// the first chunk of audio for this asset will be loaded into the cache from disk when this asset is loaded, but may be evicted to make room for other audio if it isn't played for a while.
+	PrimeOnLoad = 2,
+	// the first chunk of audio for this asset will not be loaded until this asset is played or primed.
+	LoadOnDemand = 3,
+	// Force all audio data for this audio asset to live outside of the cache and use the non-streaming decode pathways. Only usable if set on the USoundWave.
+	ForceInline = 4 UMETA(DisplayName = "Force Inline"),
+	// This value is used to delineate when the value of ESoundWaveLoadingBehavior hasn't been cached on a USoundWave yet.
+	Uninitialized = 0xff UMETA(Hidden)
+};
+
 USTRUCT()
 struct FSoundClassEditorData
 {
@@ -121,6 +147,10 @@ struct FSoundClassProperties
 	UPROPERTY(EditAnywhere, BlueprintReadOnly, Category=SoundClassProperties)
 	TEnumAsByte<EAudioOutputTarget::Type> OutputTarget;
 
+	/** The loading behavior used for USoundWaves of this USoundClass. */
+	UPROPERTY(EditAnywhere, BlueprintReadOnly, Category = SoundClassProperties, meta = (DisplayName = "Loading Behavior Override"))
+	ESoundWaveLoadingBehavior LoadingBehavior;
+
 	FSoundClassProperties()
 		: Volume(1.0f)
 		, Pitch(1.0f)
@@ -139,6 +169,7 @@ struct FSoundClassProperties
 		, bCenterChannelOnly(false)
 		, bApplyAmbientVolumes(false)
 		, OutputTarget(EAudioOutputTarget::Speaker)
+		, LoadingBehavior(ESoundWaveLoadingBehavior::Inherited)
 		{
 		}
 	
diff --git a/Engine/Source/Runtime/Engine/Classes/Sound/SoundWave.h b/Engine/Source/Runtime/Engine/Classes/Sound/SoundWave.h
--- a/Engine/Source/Runtime/Engine/Classes/Sound/SoundWave.h
+++ b/Engine/Source/Runtime/Engine/Classes/Sound/SoundWave.h
@@ -12,11 +12,13 @@
 #include "Misc/Guid.h"
 #include "Async/AsyncWork.h"
 #include "Sound/SoundBase.h"
+#include "Sound/SoundClass.h"
 #include "Serialization/BulkData.h"
 #include "Sound/SoundGroups.h"
 #include "AudioMixerTypes.h"
 #include "AudioCompressionSettings.h"
 #include "PerPlatformProperties.h"
+#include "ContentStreaming.h"
 #include "SoundWave.generated.h"
 
 class ITargetPlatform;
@@ -308,7 +310,7 @@ class ENGINE_API USoundWave : public USoundBase
 	UPROPERTY(EditAnywhere, Category=Sound, AssetRegistrySearchable)
 	uint8 bLooping:1;
 
-	/** Whether this sound can be streamed to avoid increased memory usage */
+	/** Whether this sound can be streamed to avoid increased memory usage. If using Stream Caching, use Loading Behavior instead to control memory usage. */
 	UPROPERTY(EditAnywhere, Category="Playback|Streaming", meta = (DisplayName = "Force Streaming"))
 	uint8 bStreaming:1;
 
@@ -316,15 +318,11 @@ class ENGINE_API USoundWave : public USoundBase
 	UPROPERTY(EditAnywhere, Category = "Playback|Streaming", meta = (DisplayName = "Seekable", EditCondition = "bStreaming"))
 	uint8 bSeekableStreaming:1;
 
-	/** If Load On Demand or Auto-Streaming is enabled, use this to force the entire compressed audio file to be loaded when this USoundWave is loaded. */
-	UPROPERTY(EditAnywhere, AdvancedDisplay, Category = "Loading", meta = (DisplayName = "Ensure Sound is Always Loaded (Inlined)"))
-	uint8 bForceInline : 1;
+	/** If stream caching is enabled, this can be used to specify how and when compressed audio data is loaded for this asset. */
+	UPROPERTY(EditAnywhere, Category = "Loading", meta = (DisplayName = "Loading Behavior Override"))
+	ESoundWaveLoadingBehavior LoadingBehavior;
 
-	/** If Load On Demand or Auto-Streaming is enabled, This will immediately kick off a load for this sound once postload is complete. */
-	UPROPERTY(EditAnywhere, AdvancedDisplay, Category = "Loading", meta = (DisplayName = "Load Compressed Audio When This Sound Wave Is Loaded"))
-	uint8 bLoadCompressedAudioWhenSoundWaveIsLoaded : 1;
-
-	/** Set to true for programmatically-generated, streamed audio. */
+	/** Set to true for programmatically generated audio. */
 	uint8 bProcedural:1;
 
 	/** Set to true of this is a bus sound source. This will result in the sound wave not generating audio for itself, but generate audio through instances. Used only in audio mixer. */
@@ -362,13 +360,13 @@ class ENGINE_API USoundWave : public USoundBase
 
 private:
 
-#if !WITH_EDITOR
 	// This is set to false on initialization, then set to true on non-editor platforms when we cache appropriate sample rate.
 	uint8 bCachedSampleRateFromPlatformSettings:1;
 
 	// This is set when SetSampleRate is called to invalidate our cached sample rate while not re-parsing project settings.
 	uint8 bSampleRateManuallyReset:1;
-#else
+
+#if WITH_EDITOR
 	// Whether this was previously cooked with stream caching enabled.
 	uint8 bWasStreamCachingEnabledOnLastCook:1;
 #endif // !WITH_EDITOR
@@ -449,6 +447,25 @@ public:
 	bool GetInterpolatedCookedFFTDataForTime(float InTime, uint32& InOutLastIndex, TArray<FSoundWaveSpectralData>& OutData, bool bLoop);
 	bool GetInterpolatedCookedEnvelopeDataForTime(float InTime, uint32& InOutLastIndex, float& OutAmplitude, bool bLoop);
 
+	/** If stream caching is enabled, allows the user to retain a strong handle to the first chunk of audio in the cache. 
+	 *  Please note that this USoundWave is NOT guaranteed to be still alive when OnLoadCompleted is called.
+	 */
+	void GetHandleForChunkOfAudio(TFunction<void(FAudioChunkHandle)> OnLoadCompleted, bool bForceSync = false, int32 ChunkIndex = 1, ENamedThreads::Type CallbackThread = ENamedThreads::GameThread);
+
+	/** If stream caching is enabled, set this sound wave to retain a strong handle to its first chunk. 
+	 *  If not called on the game thread, bForceSync must be true.
+	*/
+	void RetainCompressedAudio(bool bForceSync = false);
+
+	/** If stream caching is enabled and au.streamcache.KeepFirstChunkInMemory is 1, this will release this USoundWave's first chunk, allowing it to be deleted. */
+	void ReleaseCompressedAudio();
+
+	/** Returns the loading behavior we should use for this sound wave.
+	 *  If this is called within Serialize(), this should be called with bCheckSoundClasses = false,
+	 *  Since there is no guarantee that the deserialized USoundClasses have been resolved yet.
+	 */
+	ESoundWaveLoadingBehavior GetLoadingBehavior(bool bCheckSoundClasses = true) const;
+
 	/** Use this to override how much audio data is loaded when this USoundWave is loaded. */
 	UPROPERTY(EditAnywhere, AdvancedDisplay, Category = "Loading")
 	int32 InitialChunkSize;
@@ -465,10 +482,13 @@ private:
 	/** the number of sounds currently playing this sound wave. */
 	FThreadSafeCounter NumSourcesPlaying;
 
-#if !WITH_EDITOR
-	// This is the sample rate gotten from platform settings.
+	// This is the sample rate retrieved from platform settings.
 	float CachedSampleRateOverride;
-#endif // !WITH_EDITOR
+
+	// We cache a soundwave's loading behavior on the first call to USoundWave::GetLoadingBehaviorForWave(true);
+	// Caches resolved loading behavior from the SoundClass graph. Must be called on the game thread.
+	void CacheInheritedLoadingBehavior();
+	ESoundWaveLoadingBehavior CachedSoundWaveLoadingBehavior;
 
 public:
 
@@ -555,6 +575,9 @@ protected:
 	UPROPERTY()
 	class UCurveTable* InternalCurves;
 
+	/** Potential strong handle to the first chunk of audio data. Can be released via ReleaseCompressedAudioData. */
+	FAudioChunkHandle FirstChunk;
+
 private:
 
 	/**
diff --git a/Engine/Source/Runtime/Engine/Private/AudioStreaming.cpp b/Engine/Source/Runtime/Engine/Private/AudioStreaming.cpp
--- a/Engine/Source/Runtime/Engine/Private/AudioStreaming.cpp
+++ b/Engine/Source/Runtime/Engine/Private/AudioStreaming.cpp
@@ -890,7 +890,7 @@ bool FLegacyAudioStreamingManager::IsManagedStreamingSoundSource(const FSoundSou
 	return StreamingSoundSources.FindByKey(SoundSource) != NULL;
 }
 
-bool FLegacyAudioStreamingManager::RequestChunk(USoundWave* SoundWave, uint32 ChunkIndex, TFunction<void(EAudioChunkLoadResult)> OnLoadCompleted)
+bool FLegacyAudioStreamingManager::RequestChunk(USoundWave* SoundWave, uint32 ChunkIndex, TFunction<void(EAudioChunkLoadResult)> OnLoadCompleted, ENamedThreads::Type ThreadToCallOnLoadCompletedOn)
 {
 	UE_LOG(LogAudio, Warning, TEXT("RequestChunk is only supported in Stream Caching."));
 	return false;
diff --git a/Engine/Source/Runtime/Engine/Private/AudioStreamingCache.cpp b/Engine/Source/Runtime/Engine/Private/AudioStreamingCache.cpp
--- a/Engine/Source/Runtime/Engine/Private/AudioStreamingCache.cpp
+++ b/Engine/Source/Runtime/Engine/Private/AudioStreamingCache.cpp
@@ -13,6 +13,7 @@ AudioStreaming.cpp: Implementation of audio streaming classes.
 #include "HAL/PlatformFile.h"
 #include "HAL/PlatformFilemanager.h"
 #include "Async/AsyncFileHandle.h"
+#include "Async/Async.h"
 #include "Misc/ScopeLock.h"
 #include "HAL/IConsoleManager.h"
 #include "HAL/LowLevelMemTracker.h"
@@ -210,7 +211,7 @@ FAudioChunkHandle FCachedAudioStreamingManager::GetLoadedChunk(const USoundWave*
 		if (NextChunk != INDEX_NONE)
 		{
 			const FAudioChunkCache::FChunkKey NextChunkKey = { MutableWave, SoundWave->GetFName(), ((uint32)NextChunk) };
-			bool bIsValidChunk = Cache->AddOrTouchChunk(NextChunkKey, [](EAudioChunkLoadResult) {});
+			bool bIsValidChunk = Cache->AddOrTouchChunk(NextChunkKey, [](EAudioChunkLoadResult) {}, ENamedThreads::AnyThread);
 			if (!bIsValidChunk)
 			{
 				UE_LOG(LogAudio, Warning, TEXT("Cache overflow!!! couldn't load chunk %d for sound %s!"), ChunkIndex, *SoundWave->GetName());
@@ -254,7 +255,7 @@ FAudioChunkCache* FCachedAudioStreamingManager::GetCacheForChunkSize(uint32 InCh
 		}
 	}
 
-	// If we ever hit this, something went wrong during cook.
+	// If we ever hit this, something may have wrong during cook.
 	// Please check to make sure this platform's implementation of IAudioFormat honors the MaxChunkSize parameter passed into SplitDataForStreaming,
 	// or that FStreamedAudioCacheDerivedDataWorker::BuildStreamedAudio() is passing the correct MaxChunkSize to IAudioFormat::SplitDataForStreaming.
 	ensureMsgf(false, TEXT("Chunks in SoundWave are too large: %d bytes"), InChunkSize);
@@ -316,13 +317,13 @@ void FCachedAudioStreamingManager::RemoveReferenceToChunk(const FAudioChunkHandl
 	Cache->RemoveReferenceToChunk(ChunkKey);
 }
 
-bool FCachedAudioStreamingManager::RequestChunk(USoundWave* SoundWave, uint32 ChunkIndex, TFunction<void(EAudioChunkLoadResult)> OnLoadCompleted)
+bool FCachedAudioStreamingManager::RequestChunk(USoundWave* SoundWave, uint32 ChunkIndex, TFunction<void(EAudioChunkLoadResult)> OnLoadCompleted, ENamedThreads::Type ThreadToCallOnLoadCompletedOn)
 {
 	FAudioChunkCache* Cache = GetCacheForWave(SoundWave);
 	if (Cache)
 	{
 		FAudioChunkCache::FChunkKey ChunkKey = { SoundWave, SoundWave->GetFName(), ChunkIndex };
-		return Cache->AddOrTouchChunk(ChunkKey, OnLoadCompleted);
+		return Cache->AddOrTouchChunk(ChunkKey, OnLoadCompleted, ThreadToCallOnLoadCompletedOn);
 	}
 	else
 	{
@@ -355,12 +356,12 @@ FAudioChunkCache::~FAudioChunkCache()
 	check(NumberOfLoadsInFlight.GetValue() == 0);
 }
 
-bool FAudioChunkCache::AddOrTouchChunk(const FChunkKey& InKey, TFunction<void(EAudioChunkLoadResult)> OnLoadCompleted)
+bool FAudioChunkCache::AddOrTouchChunk(const FChunkKey& InKey, TFunction<void(EAudioChunkLoadResult)> OnLoadCompleted, ENamedThreads::Type CallbackThread)
 {
 	if (!IsKeyValid(InKey))
 	{
 		ensure(false);
-		OnLoadCompleted(EAudioChunkLoadResult::ChunkOutOfBounds);
+		ExecuteOnLoadCompleteCallback(EAudioChunkLoadResult::ChunkOutOfBounds, OnLoadCompleted, CallbackThread);
 		return false;
 	}
 
@@ -373,7 +374,7 @@ bool FAudioChunkCache::AddOrTouchChunk(const FChunkKey& InKey, TFunction<void(EA
 		TouchElement(FoundElement);
 		if (FoundElement->bIsLoaded)
 		{
-			OnLoadCompleted(EAudioChunkLoadResult::AlreadyLoaded);
+			ExecuteOnLoadCompleteCallback(EAudioChunkLoadResult::AlreadyLoaded, OnLoadCompleted, CallbackThread);
 		}
 
 #if DEBUG_STREAM_CACHE
@@ -388,11 +389,11 @@ bool FAudioChunkCache::AddOrTouchChunk(const FChunkKey& InKey, TFunction<void(EA
 
 		if (!CacheElement)
 		{
-			OnLoadCompleted(EAudioChunkLoadResult::CacheBlown);
+			ExecuteOnLoadCompleteCallback(EAudioChunkLoadResult::CacheBlown, OnLoadCompleted, CallbackThread);
 			return false;
 		}
 
-		KickOffAsyncLoad(CacheElement, InKey, OnLoadCompleted);
+		KickOffAsyncLoad(CacheElement, InKey, OnLoadCompleted, CallbackThread);
 
 		if (TrimCacheWhenOverBudgetCVar != 0 && MemoryCounterBytes > MemoryLimitBytes)
 		{
@@ -411,14 +412,18 @@ TArrayView<uint8> FAudioChunkCache::GetChunk(const FChunkKey& InKey, bool bBlock
 	if (FoundElement)
 	{
 		TouchElement(FoundElement);
-		if (bBlockForLoadCompletion)
-		{
-			FoundElement->WaitForAsyncLoadCompletion(false);
-		}
-		else if (FoundElement->IsLoadInProgress())
+		if (FoundElement->IsLoadInProgress())
 		{
-			return TArrayView<uint8>();
+			if (bBlockForLoadCompletion)
+			{
+				FoundElement->WaitForAsyncLoadCompletion(false);
+			}
+			else
+			{
+				return TArrayView<uint8>();
+			}
 		}
+		
 
 		// If this value is ever negative, it means that we're decrementing more than we're incrementing:
 		check(FoundElement->NumConsumers.GetValue() >= 0);
@@ -435,7 +440,7 @@ TArrayView<uint8> FAudioChunkCache::GetChunk(const FChunkKey& InKey, bool bBlock
 			return TArrayView<uint8>();
 		}
 
-		KickOffAsyncLoad(FoundElement, InKey, [](EAudioChunkLoadResult InResult) {});
+		KickOffAsyncLoad(FoundElement, InKey, [](EAudioChunkLoadResult InResult) {}, ENamedThreads::AnyThread);
 
 		if (bBlockForLoadCompletion)
 		{
@@ -822,7 +827,7 @@ FAudioChunkCache::FCacheElement* FAudioChunkCache::EvictLeastRecentChunk()
 	return CacheElement;
 }
 
-void FAudioChunkCache::KickOffAsyncLoad(FCacheElement* CacheElement, const FChunkKey& InKey, TFunction<void(EAudioChunkLoadResult)> OnLoadCompleted)
+void FAudioChunkCache::KickOffAsyncLoad(FCacheElement* CacheElement, const FChunkKey& InKey, TFunction<void(EAudioChunkLoadResult)> OnLoadCompleted, ENamedThreads::Type CallbackThread)
 {
 	LLM_SCOPE(ELLMTag::Audio);
 	check(CacheElement);
@@ -860,7 +865,7 @@ void FAudioChunkCache::KickOffAsyncLoad(FCacheElement* CacheElement, const FChun
 #endif
 
 
-		TFunction<void(bool)> OnLoadComplete = [OnLoadCompleted, CacheElement, InKey, ChunkDataSize](bool bRequestFailed)
+		TFunction<void(bool)> OnLoadComplete = [OnLoadCompleted, CallbackThread, CacheElement, InKey, ChunkDataSize](bool bRequestFailed)
 		{
 			// Populate key and DataSize. The async read request was set up to write directly into CacheElement->ChunkData.
 			CacheElement->Key = InKey;
@@ -870,8 +875,8 @@ void FAudioChunkCache::KickOffAsyncLoad(FCacheElement* CacheElement, const FChun
 #if DEBUG_STREAM_CACHE
 			CacheElement->DebugInfo.TimeToLoad = (FPlatformTime::Seconds() - CacheElement->DebugInfo.TimeLoadStarted) * 1000.0f;
 #endif
-
-			OnLoadCompleted(bRequestFailed ? EAudioChunkLoadResult::Interrupted : EAudioChunkLoadResult::Completed);
+			EAudioChunkLoadResult ChunkLoadResult = bRequestFailed ? EAudioChunkLoadResult::Interrupted : EAudioChunkLoadResult::Completed;
+			ExecuteOnLoadCompleteCallback(ChunkLoadResult, OnLoadCompleted, CallbackThread);
 		};
 
 		NumberOfLoadsInFlight.Increment();
@@ -965,6 +970,22 @@ EAsyncIOPriorityAndFlags FAudioChunkCache::GetAsyncPriorityForChunk(const FChunk
 	}
 }
 
+void FAudioChunkCache::ExecuteOnLoadCompleteCallback(EAudioChunkLoadResult Result, const TFunction<void(EAudioChunkLoadResult)>& OnLoadCompleted, const ENamedThreads::Type& CallbackThread)
+{
+	if (CallbackThread == ENamedThreads::AnyThread)
+	{
+		OnLoadCompleted(Result);
+	}
+	else
+	{
+		// Dispatch an async notify.
+		AsyncTask(CallbackThread, [OnLoadCompleted, Result]() 
+		{
+			OnLoadCompleted(Result);
+		});
+	}
+}
+
 bool FAudioChunkCache::IsKeyValid(const FChunkKey& InKey)
 {
 	return InKey.ChunkIndex < TNumericLimits<uint32>::Max() && ((int32)InKey.ChunkIndex) < InKey.SoundWave->RunningPlatformData->Chunks.Num();
diff --git a/Engine/Source/Runtime/Engine/Private/Components/AudioComponent.cpp b/Engine/Source/Runtime/Engine/Private/Components/AudioComponent.cpp
--- a/Engine/Source/Runtime/Engine/Private/Components/AudioComponent.cpp
+++ b/Engine/Source/Runtime/Engine/Private/Components/AudioComponent.cpp
@@ -12,6 +12,7 @@
 #include "Components/BillboardComponent.h"
 #include "UObject/FrameworkObjectVersion.h"
 #include "Misc/App.h"
+#include "Kismet/GameplayStatics.h"
 
 DECLARE_CYCLE_STAT(TEXT("AudioComponent Play"), STAT_AudioComp_Play, STATGROUP_Audio);
 
@@ -23,6 +24,14 @@ FAutoConsoleVariableRef CVarBackedAnalysisTimeShift(
 	TEXT("Value: The time in seconds to shift the timeline."),
 	ECVF_Default);
 
+static int32 PrimeSoundOnAudioComponentSpawnCVar = 0;
+FAutoConsoleVariableRef CVarPrimeSoundOnAudioComponentSpawn(
+	TEXT("au.streamcaching.PrimeSoundOnAudioComponents"),
+	PrimeSoundOnAudioComponentSpawnCVar,
+	TEXT("When set to 1, automatically primes a USoundBase when a UAudioComponent is spawned with that sound, or when UAudioComponent::SetSound is called.\n")
+	TEXT("Value: The time in seconds to shift the timeline."),
+	ECVF_Default);
+
 
 /*-----------------------------------------------------------------------------
 UAudioComponent implementation.
@@ -169,6 +178,11 @@ void UAudioComponent::PostLoad()
 	}
 #endif
 
+	if (PrimeSoundOnAudioComponentSpawnCVar && Sound)
+	{
+		UGameplayStatics::PrimeSound(Sound);
+	}
+
 	Super::PostLoad();
 }
 
@@ -254,6 +268,11 @@ void UAudioComponent::SetSound(USoundBase* NewSound)
 
 	Sound = NewSound;
 
+	if (PrimeSoundOnAudioComponentSpawnCVar && Sound)
+	{
+		UGameplayStatics::PrimeSound(Sound);
+	}
+
 	if (bPlay)
 	{
 		Play();
diff --git a/Engine/Source/Runtime/Engine/Private/ContentStreaming.cpp b/Engine/Source/Runtime/Engine/Private/ContentStreaming.cpp
--- a/Engine/Source/Runtime/Engine/Private/ContentStreaming.cpp
+++ b/Engine/Source/Runtime/Engine/Private/ContentStreaming.cpp
@@ -1384,6 +1384,7 @@ FAudioChunkHandle::FAudioChunkHandle(const uint8* InData, uint32 NumBytes, const
 }
 
 FAudioChunkHandle::FAudioChunkHandle(const FAudioChunkHandle& Other)
+	: FAudioChunkHandle()
 {
 	*this = Other;
 }
diff --git a/Engine/Source/Runtime/Engine/Private/SoundClass.cpp b/Engine/Source/Runtime/Engine/Private/SoundClass.cpp
--- a/Engine/Source/Runtime/Engine/Private/SoundClass.cpp
+++ b/Engine/Source/Runtime/Engine/Private/SoundClass.cpp
@@ -70,26 +70,34 @@ void USoundClass::PostLoad()
 #if WITH_EDITOR
 
 TArray<USoundClass*> BackupChildClasses;
+ESoundWaveLoadingBehavior BackupLoadingBehavior;
 
 void USoundClass::PreEditChange(UProperty* PropertyAboutToChange)
 {
-	static FName NAME_ChildClasses(TEXT("ChildClasses"));
+	static const FName NAME_ChildClasses = GET_MEMBER_NAME_CHECKED(USoundClass, ChildClasses);
+	static const FName NAME_Properties = GET_MEMBER_NAME_CHECKED(FSoundClassProperties, LoadingBehavior);
 
 	if (PropertyAboutToChange && PropertyAboutToChange->GetFName() == NAME_ChildClasses)
 	{
 		// Take a copy of the current state of child classes
 		BackupChildClasses = ChildClasses;
 	}
+	else if (PropertyAboutToChange && PropertyAboutToChange->GetFName() == NAME_Properties)
+	{
+		// Copy the current loading behavior in case it does not pass validation.
+		BackupLoadingBehavior = Properties.LoadingBehavior;
+	}
 }
 
 void USoundClass::PostEditChangeProperty(struct FPropertyChangedEvent& PropertyChangedEvent)
 {
 	if (PropertyChangedEvent.Property != NULL)
 	{
-		static const FName NAME_ChildClasses(TEXT("ChildClasses"));
-		static const FName NAME_ParentClass(TEXT("ParentClass"));
+		static const FName NAME_ChildClasses = GET_MEMBER_NAME_CHECKED(USoundClass, ChildClasses);
+		static const FName NAME_ParentClass = GET_MEMBER_NAME_CHECKED(USoundClass, ParentClass);
+		static const FName NAME_Properties = GET_MEMBER_NAME_CHECKED(FSoundClassProperties, LoadingBehavior);
 
-		if (PropertyChangedEvent.Property->GetFName() == NAME_ChildClasses)
+		if (PropertyChangedEvent.GetPropertyName() == NAME_ChildClasses)
 		{
 			// Find child that was changed/added
 			for (int32 ChildIndex = 0; ChildIndex < ChildClasses.Num(); ChildIndex++)
@@ -127,7 +135,7 @@ void USoundClass::PostEditChangeProperty(struct FPropertyChangedEvent& PropertyC
 
 			RefreshAllGraphs(false);
 		}
-		else if (PropertyChangedEvent.Property->GetFName() == NAME_ParentClass)
+		else if (PropertyChangedEvent.GetPropertyName() == NAME_ParentClass)
 		{
 			// Add this sound class to the parent class if it's not already added
 			if (ParentClass)
@@ -153,11 +161,24 @@ void USoundClass::PostEditChangeProperty(struct FPropertyChangedEvent& PropertyC
 			Modify();
 			RefreshAllGraphs(false);
 		}
+		else if (PropertyChangedEvent.GetPropertyName() == NAME_Properties)
+		{
+			// Until we can check FSoundClassProperties during USoundWave::Serialize, we can't support ForceInline here.
+			if (Properties.LoadingBehavior == ESoundWaveLoadingBehavior::ForceInline)
+			{
+				FNotificationInfo Info(NSLOCTEXT("Engine", "ForceInlineUnavailableOnSoundClass", "Using Force Inline on soundclasses is currently not supported. Set each Sound Wave to Force Inline individually instead."));
+				Info.ExpireDuration = 5.0f;
+				Info.Image = FCoreStyle::Get().GetBrush(TEXT("MessageLog.Error"));
+				FSlateNotificationManager::Get().AddNotification(Info);
+
+				Properties.LoadingBehavior = BackupLoadingBehavior;
+			}
+		}
 	}
 
 	if (PropertyChangedEvent.ChangeType != EPropertyChangeType::Interactive && PropertyChangedEvent.MemberProperty)
 	{
-		if (PropertyChangedEvent.MemberProperty->GetFName() == GET_MEMBER_NAME_CHECKED(USoundClass,PassiveSoundMixModifiers))
+		if (PropertyChangedEvent.GetPropertyName() == GET_MEMBER_NAME_CHECKED(USoundClass,PassiveSoundMixModifiers))
 		{
 			TArray<USoundClass*> ProblemClasses;
 			for (int32 Index = 0; Index < PassiveSoundMixModifiers.Num(); Index++)
diff --git a/Engine/Source/Runtime/Engine/Private/SoundWave.cpp b/Engine/Source/Runtime/Engine/Private/SoundWave.cpp
--- a/Engine/Source/Runtime/Engine/Private/SoundWave.cpp
+++ b/Engine/Source/Runtime/Engine/Private/SoundWave.cpp
@@ -31,12 +31,12 @@
 
 #include "Misc/CommandLine.h"
 
-static int32 LoadIntoCacheOnPostLoadCVar = 0;
-FAutoConsoleVariableRef CVarLoadIntoCacheOnPostLoad(
-	TEXT("au.streamcache.LoadIntoCacheOnPostLoad"),
-	LoadIntoCacheOnPostLoadCVar,
-	TEXT("When set to 1, immediately loads any sounds that are loaded into the cache on load.\n")
-	TEXT("0: Honor the Play When Silent flag, 1: stop all silent non-procedural sources."),
+static int32 SoundWaveDefaultLoadingBehaviorCVar = 0;
+FAutoConsoleVariableRef CVarSoundWaveDefaultLoadingBehavior(
+	TEXT("au.streamcache.SoundWaveDefaultLoadingBehavior"),
+	SoundWaveDefaultLoadingBehaviorCVar,
+	TEXT("This can be set to define the default behavior when a USoundWave is loaded.\n")
+	TEXT("0: Default (load on demand), 1: Retain audio data on load, 2: prime audio data on load, 3: load on demand (No audio data is loaded until a USoundWave is played or primed)."),
 	ECVF_Default);
 
 static int32 ForceNonStreamingInEditorCVar = 0;
@@ -194,14 +194,15 @@ USoundWave::USoundWave(const FObjectInitializer& ObjectInitializer)
 	EnvelopeFollowerReleaseTime = 100;
 #endif
 
-#if !WITH_EDITOR
 	bCachedSampleRateFromPlatformSettings = false;
 	bSampleRateManuallyReset = false;
 	CachedSampleRateOverride = 0.0f;
-#else
+	CachedSoundWaveLoadingBehavior = ESoundWaveLoadingBehavior::Uninitialized;
+
+#if WITH_EDITOR
 	bWasStreamCachingEnabledOnLastCook = FPlatformCompressionUtilities::IsCurrentPlatformUsingStreamCaching();
 	RunningPlatformData = nullptr;
-#endif //!WITH_EDITOR
+#endif
 }
 
 void USoundWave::GetResourceSizeEx(FResourceSizeEx& CumulativeResourceSize)
@@ -421,17 +422,34 @@ void USoundWave::Serialize( FArchive& Ar )
 #if WITH_EDITORONLY_DATA
 		if (Ar.IsLoading() && !Ar.IsTransacting() && !bCooked && !GetOutermost()->HasAnyPackageFlags(PKG_ReloadingForCooker))
 		{
-			BeginCachePlatformData();
+			CachePlatformData(false);
 		}
 #endif // #if WITH_EDITORONLY_DATA
+	}
+
+	if (!GIsEditor && Ar.IsLoading())
+	{
+		// For non-editor builds, we can immediately cache the sample rate.
+		SampleRate = GetSampleRateForCurrentPlatform();
 
-// For non-editor builds, we can immediately cache the sample rate.
-#if !WITH_EDITOR
-		if (Ar.IsLoading())
+		// If stream caching is enabled, here we determine if we should retain or prime this wave on load.
+		if (FPlatformCompressionUtilities::IsCurrentPlatformUsingStreamCaching())
 		{
-			SampleRate = GetSampleRateForCurrentPlatform();
+			ESoundWaveLoadingBehavior CurrentLoadingBehavior = GetLoadingBehavior(false);
+
+			//EnsureZerothChunkIsLoaded();
+			const bool bHasFirstChunk = GetNumChunks() > 1;
+
+			if (LoadingBehavior == ESoundWaveLoadingBehavior::RetainOnLoad && bHasFirstChunk)
+			{
+				RetainCompressedAudio(true);
+			}
+			else if (LoadingBehavior == ESoundWaveLoadingBehavior::PrimeOnLoad && bHasFirstChunk)
+			{
+				// Prime first chunk of audio.
+				IStreamingManager::Get().GetAudioStreamingManager().RequestChunk(this, 1, [](EAudioChunkLoadResult) {});
+			}
 		}
-#endif // !WITH_EDITOR
 	}
 }
 
@@ -497,12 +515,17 @@ bool USoundWave::HasCompressedData(FName Format, ITargetPlatform* TargetPlatform
 		return false;
 	}
 
-#if WITH_EDITOR
-	const FPlatformAudioCookOverrides* CompressionOverrides = (TargetPlatform) ? TargetPlatform->GetAudioCompressionSettings() : nullptr;
-#else
-	// TargetPlatform is not available on consoles/mobile, so we have to grab it ourselves:
-	const FPlatformAudioCookOverrides* CompressionOverrides = FPlatformCompressionUtilities::GetCookOverridesForCurrentPlatform();
-#endif // WITH_EDITOR
+	const FPlatformAudioCookOverrides* CompressionOverrides = nullptr;
+
+	if (GIsEditor)
+	{
+		CompressionOverrides = (TargetPlatform) ? TargetPlatform->GetAudioCompressionSettings() : nullptr;
+	}
+	else
+	{
+		// TargetPlatform is not available on consoles/mobile, so we have to grab it ourselves:
+		CompressionOverrides = FPlatformCompressionUtilities::GetCookOverridesForCurrentPlatform();
+	}
 
 	if (CompressionOverrides)
 	{
@@ -744,10 +767,34 @@ void USoundWave::PostLoad()
 		return;
 	}
 
-#if !WITH_EDITOR
-	// In case any code accesses bStreaming directly, we fix up bStreaming based on the current platform's cook overrides.
-	bStreaming = IsStreaming();
-#endif
+	CacheInheritedLoadingBehavior();
+
+	if (!GIsEditor)
+	{
+		// If our loading behavior is defined by a sound class, we need to update whether this sound wave actually needs to retain its audio data or not.
+		ESoundWaveLoadingBehavior ActualLoadingBehavior = GetLoadingBehavior();
+
+		if (ShouldUseStreamCaching() && ActualLoadingBehavior != GetLoadingBehavior(false))
+		{
+			if (ActualLoadingBehavior == ESoundWaveLoadingBehavior::RetainOnLoad)
+			{
+				RetainCompressedAudio(false);
+			}
+			else
+			{
+				// if a sound class defined our loading behavior as something other than Retain and our cvar default is to retain, we need to release our handle.
+				ReleaseCompressedAudio();
+
+				if (ActualLoadingBehavior == ESoundWaveLoadingBehavior::PrimeOnLoad && GetNumChunks() > 1)
+				{
+					IStreamingManager::Get().GetAudioStreamingManager().RequestChunk(this, 1, [](EAudioChunkLoadResult) {});
+				}
+			}
+		}
+
+		// In case any code accesses bStreaming directly, we fix up bStreaming based on the current platform's cook overrides.
+		bStreaming = IsStreaming();
+	}
 
 	// Compress to whatever formats the active target platforms want
 	// static here as an optimization
@@ -762,8 +809,6 @@ void USoundWave::PostLoad()
 		}
 	}
 
-	
-
 	// We don't precache default objects and we don't precache in the Editor as the latter will
 	// most likely cause us to run out of memory.
 	if (!GIsEditor && !IsTemplate( RF_ClassDefaultObject ) && GEngine)
@@ -792,19 +837,11 @@ void USoundWave::PostLoad()
 
 	const bool bDoesSoundWaveHaveStreamingAudioData = HasStreamingChunks();
 
-	if (ShouldUseStreamCaching() && IsStreaming() && bDoesSoundWaveHaveStreamingAudioData)
+	if (ShouldUseStreamCaching() && bDoesSoundWaveHaveStreamingAudioData)
 	{
 		EnsureZerothChunkIsLoaded();
 	}
 
-	const bool bShouldPrimeSound = LoadIntoCacheOnPostLoadCVar || bLoadCompressedAudioWhenSoundWaveIsLoaded;
-
-	if (bShouldPrimeSound && FPlatformCompressionUtilities::IsCurrentPlatformUsingStreamCaching() && bDoesSoundWaveHaveStreamingAudioData)
-	{
-		// Load rest of the audio into cache.
-		IStreamingManager::Get().GetAudioStreamingManager().RequestChunk(this, 1, [](EAudioChunkLoadResult) {});
-	}
-
 #if WITH_EDITORONLY_DATA
 	if (!SourceFilePath_DEPRECATED.IsEmpty() && AssetImportData)
 	{
@@ -872,6 +909,7 @@ uint32 USoundWave::GetNumChunks() const
 	}
 	else
 	{
+		ensureAlwaysMsgf(false, TEXT("Call CachePlatformData(false) before calling this function in editor."));
 		return 0;
 	}
 }
@@ -1880,7 +1918,7 @@ bool USoundWave::IsStreaming(const FPlatformAudioCookOverrides* Overrides /* = n
 	{
 		return true;
 	}
-	else if (bForceInline)
+	else if (LoadingBehavior == ESoundWaveLoadingBehavior::ForceInline) // TODO: Support setting ESoundWaveLoadingBehavior to ForceInline on a USoundClass.
 	{
 		return false;
 	}
@@ -1992,33 +2030,36 @@ float USoundWave::GetSampleRateForCurrentPlatform()
 		return SampleRate;
 	}
 
-#if WITH_EDITOR
-	float SampleRateOverride = FPlatformCompressionUtilities::GetTargetSampleRateForPlatform(SampleRateQuality);
-	return (SampleRateOverride > 0) ? FMath::Min(SampleRateOverride, (float) SampleRate) : SampleRate;
-#else
-	if (bCachedSampleRateFromPlatformSettings)
-	{
-		return CachedSampleRateOverride;
-	}
-	else if (bSampleRateManuallyReset)
+	if (GIsEditor)
 	{
-		CachedSampleRateOverride = SampleRate;
-		bCachedSampleRateFromPlatformSettings = true;
-
-		return CachedSampleRateOverride;
+		float SampleRateOverride = FPlatformCompressionUtilities::GetTargetSampleRateForPlatform(SampleRateQuality);
+		return (SampleRateOverride > 0) ? FMath::Min(SampleRateOverride, (float)SampleRate) : SampleRate;
 	}
 	else
 	{
-		CachedSampleRateOverride = FPlatformCompressionUtilities::GetTargetSampleRateForPlatform(SampleRateQuality);
-		if (CachedSampleRateOverride < 0 || SampleRate < CachedSampleRateOverride)
+		if (bCachedSampleRateFromPlatformSettings)
+		{
+			return CachedSampleRateOverride;
+		}
+		else if (bSampleRateManuallyReset)
 		{
 			CachedSampleRateOverride = SampleRate;
+			bCachedSampleRateFromPlatformSettings = true;
+
+			return CachedSampleRateOverride;
 		}
+		else
+		{
+			CachedSampleRateOverride = FPlatformCompressionUtilities::GetTargetSampleRateForPlatform(SampleRateQuality);
+			if (CachedSampleRateOverride < 0 || SampleRate < CachedSampleRateOverride)
+			{
+				CachedSampleRateOverride = SampleRate;
+			}
 
-		bCachedSampleRateFromPlatformSettings = true;
-		return CachedSampleRateOverride;
+			bCachedSampleRateFromPlatformSettings = true;
+			return CachedSampleRateOverride;
+		}
 	}
-#endif
 }
 
 float USoundWave::GetSampleRateForCompressionOverrides(const FPlatformAudioCookOverrides* CompressionOverrides)
@@ -2271,3 +2312,122 @@ bool USoundWave::GetInterpolatedCookedEnvelopeDataForTime(float InTime, uint32&
 	}
 	return InOutLastIndex != INDEX_NONE;
 }
+
+void USoundWave::GetHandleForChunkOfAudio(TFunction<void(FAudioChunkHandle)> OnLoadCompleted, bool bForceSync /*= false*/, int32 ChunkIndex /*= 1*/, ENamedThreads::Type CallbackThread /*= ENamedThreads::GameThread*/)
+{
+	if (bForceSync)
+	{
+		// For sync cases, we call GetLoadedChunk with bBlockForLoad = true, then execute the callback immediately.
+		FAudioChunkHandle ChunkHandle = IStreamingManager::Get().GetAudioStreamingManager().GetLoadedChunk(this, ChunkIndex, true);
+		OnLoadCompleted(ChunkHandle);
+	}
+	else
+	{
+		// For async cases, we call RequestChunk and request the loaded chunk in the completion callback.
+		IStreamingManager::Get().GetAudioStreamingManager().RequestChunk(this, ChunkIndex, [this, OnLoadCompleted, ChunkIndex](EAudioChunkLoadResult LoadResult)
+		{
+			if (LoadResult == EAudioChunkLoadResult::Completed || LoadResult == EAudioChunkLoadResult::AlreadyLoaded)
+			{
+				FAudioChunkHandle ChunkHandle = IStreamingManager::Get().GetAudioStreamingManager().GetLoadedChunk(this, ChunkIndex, true);
+				OnLoadCompleted(ChunkHandle);
+			}
+			else
+			{
+				// Load failed. Return an invalid chunk handle.
+				FAudioChunkHandle ChunkHandle;
+				OnLoadCompleted(ChunkHandle);
+			}
+		}, CallbackThread);
+	}
+}
+
+void USoundWave::RetainCompressedAudio(bool bForceSync /*= false*/)
+{
+	if (FirstChunk.IsValid())
+	{
+		return;
+	}
+	else if (bForceSync)
+	{
+		FirstChunk = IStreamingManager::Get().GetAudioStreamingManager().GetLoadedChunk(this, 1, true);
+	}
+	else
+	{
+		TWeakObjectPtr<USoundWave> WeakThis = TWeakObjectPtr<USoundWave>(this);
+		GetHandleForChunkOfAudio([WeakThis](FAudioChunkHandle OutHandle)
+		{
+			check(IsInGameThread());
+			
+			if (WeakThis.IsValid())
+			{
+				WeakThis->FirstChunk = OutHandle;
+			}
+		}, ENamedThreads::GameThread);
+	}
+}
+
+void USoundWave::ReleaseCompressedAudio()
+{
+	// Here we release the USoundWave's handle to the compressed asset by resetting it.
+	FirstChunk = FAudioChunkHandle();
+}
+
+void USoundWave::CacheInheritedLoadingBehavior()
+{
+	check(IsInGameThread());
+
+	// Determine this sound wave's loading behavior and cache it.
+	if (LoadingBehavior != ESoundWaveLoadingBehavior::Inherited)
+	{
+		// If this sound wave specifies it's own loading behavior, use that.
+		CachedSoundWaveLoadingBehavior = LoadingBehavior;
+	}
+	else
+	{
+		USoundClass* CurrentSoundClass = SoundClassObject;
+		ESoundWaveLoadingBehavior SoundClassLoadingBehavior = ESoundWaveLoadingBehavior::Inherited;
+
+		// Recurse through this sound class's parents until we find an override.
+		while (SoundClassLoadingBehavior == ESoundWaveLoadingBehavior::Inherited && CurrentSoundClass != nullptr)
+		{
+			SoundClassLoadingBehavior = CurrentSoundClass->Properties.LoadingBehavior;
+			CurrentSoundClass = CurrentSoundClass->ParentClass;
+		}
+
+		// If we could not find an override in the sound class hierarchy, use the loading behavior defined by our cvar.
+		if (SoundClassLoadingBehavior == ESoundWaveLoadingBehavior::Inherited)
+		{
+			ensureAlwaysMsgf(SoundWaveDefaultLoadingBehaviorCVar >= 0 && SoundWaveDefaultLoadingBehaviorCVar < 4, TEXT("Invalid default loading behavior CVar value. Use value 0, 1, 2 or 3."));
+			SoundClassLoadingBehavior = (ESoundWaveLoadingBehavior)FMath::Clamp<int32>(SoundWaveDefaultLoadingBehaviorCVar, 0, (int32)ESoundWaveLoadingBehavior::LoadOnDemand);
+		}
+
+		CachedSoundWaveLoadingBehavior = SoundClassLoadingBehavior;
+	}
+}
+
+ESoundWaveLoadingBehavior USoundWave::GetLoadingBehavior(bool bCheckSoundClasses /*= true*/) const
+{
+	checkf(!bCheckSoundClasses || CachedSoundWaveLoadingBehavior != ESoundWaveLoadingBehavior::Uninitialized, TEXT("this call is only valid if bCheckSoundClasses is false or CacheInheritedLoadingBehavior has already been called on the game thread."));
+
+	if (!bCheckSoundClasses)
+	{
+		if (LoadingBehavior != ESoundWaveLoadingBehavior::Inherited)
+		{
+			// If this sound wave specifies it's own loading behavior, use that.
+			return LoadingBehavior;
+		}
+		else
+		{
+			// Otherwise, use the loading behavior defined by our cvar.
+			ensureAlwaysMsgf(SoundWaveDefaultLoadingBehaviorCVar >= 0 && SoundWaveDefaultLoadingBehaviorCVar < 4, TEXT("Invalid default loading behavior CVar value. Use value 0, 1, 2 or 3."));
+			return (ESoundWaveLoadingBehavior)FMath::Clamp<int32>(SoundWaveDefaultLoadingBehaviorCVar, 0, (int32)ESoundWaveLoadingBehavior::LoadOnDemand);
+		}
+	}
+	else if (CachedSoundWaveLoadingBehavior != ESoundWaveLoadingBehavior::Uninitialized)
+	{
+		return CachedSoundWaveLoadingBehavior;
+	}
+
+	// We should never reach this. See check at the beginning of this function.
+	return ESoundWaveLoadingBehavior::Inherited;
+}
diff --git a/Engine/Source/Runtime/Engine/Public/AudioStreaming.h b/Engine/Source/Runtime/Engine/Public/AudioStreaming.h
--- a/Engine/Source/Runtime/Engine/Public/AudioStreaming.h
+++ b/Engine/Source/Runtime/Engine/Public/AudioStreaming.h
@@ -273,7 +273,7 @@ struct FLegacyAudioStreamingManager : public IAudioStreamingManager
 	virtual void AddStreamingSoundSource(FSoundSource* SoundSource) override;
 	virtual void RemoveStreamingSoundSource(FSoundSource* SoundSource) override;
 	virtual bool IsManagedStreamingSoundSource(const FSoundSource* SoundSource) const override;
-	virtual bool RequestChunk(USoundWave* SoundWave, uint32 ChunkIndex, TFunction<void(EAudioChunkLoadResult)> OnLoadCompleted) override;
+	virtual bool RequestChunk(USoundWave* SoundWave, uint32 ChunkIndex, TFunction<void(EAudioChunkLoadResult)> OnLoadCompleted, ENamedThreads::Type ThreadToCallOnLoadCompletedOn) override;
 	virtual FAudioChunkHandle GetLoadedChunk(const USoundWave* SoundWave, uint32 ChunkIndex, bool bBlockForLoad = false) const override;
 	virtual uint64 TrimMemory(uint64 NumBytesToFree) override;
 	virtual int32 RenderStatAudioStreaming(UWorld* World, FViewport* Viewport, FCanvas* Canvas, int32 X, int32 Y, const FVector* ViewLocation, const FRotator* ViewRotation) override;
diff --git a/Engine/Source/Runtime/Engine/Public/AudioStreamingCache.h b/Engine/Source/Runtime/Engine/Public/AudioStreamingCache.h
--- a/Engine/Source/Runtime/Engine/Public/AudioStreamingCache.h
+++ b/Engine/Source/Runtime/Engine/Public/AudioStreamingCache.h
@@ -41,7 +41,7 @@ public:
 	~FAudioChunkCache();
 
 	// Places chunk in cache, or puts this chunk back at the top of the cache if it's already loaded. Returns false on failure.
-	bool AddOrTouchChunk(const FChunkKey& InKey, TFunction<void(EAudioChunkLoadResult) > OnLoadCompleted);
+	bool AddOrTouchChunk(const FChunkKey& InKey, TFunction<void(EAudioChunkLoadResult) > OnLoadCompleted, ENamedThreads::Type CallbackThread);
 
 	// Returns the chunk asked for, or an empty TArrayView if that chunk is not loaded.
 	TArrayView<uint8> GetChunk(const FChunkKey& InKey, bool bBlockForLoadCompletion);
@@ -167,7 +167,11 @@ private:
 #if WITH_EDITORONLY_DATA
 			if (DDCTask.IsValid() && !DDCTask->IsDone())
 			{
-				DDCTask->Cancel();
+				if (bCancel)
+				{
+					DDCTask->Cancel();
+				}
+				
 				DDCTask->EnsureCompletion(false);
 			}
 #endif
@@ -259,8 +263,11 @@ private:
 	// Returns the least recent chunk and fixes up the linked list accordingly.
 	FCacheElement* EvictLeastRecentChunk();
 
-	void KickOffAsyncLoad(FCacheElement* CacheElement, const FChunkKey& InKey, TFunction<void(EAudioChunkLoadResult)> OnLoadCompleted);
+	void KickOffAsyncLoad(FCacheElement* CacheElement, const FChunkKey& InKey, TFunction<void(EAudioChunkLoadResult)> OnLoadCompleted, ENamedThreads::Type CallbackThread);
 	EAsyncIOPriorityAndFlags GetAsyncPriorityForChunk(const FChunkKey& InKey);
+
+	// Calls OnLoadCompleted on current thread if CallbackThread == ENamedThreads::AnyThread, and dispatchs an async task on a named thread otherwise.
+	static void ExecuteOnLoadCompleteCallback(EAudioChunkLoadResult Result, const TFunction<void(EAudioChunkLoadResult)>& OnLoadCompleted, const ENamedThreads::Type& CallbackThread);
 };
 
 // This is used to sort the cache array from smallest chunk size to biggest.
@@ -329,7 +336,7 @@ public:
 	virtual void AddStreamingSoundSource(FSoundSource* SoundSource) override;
 	virtual void RemoveStreamingSoundSource(FSoundSource* SoundSource) override;
 	virtual bool IsManagedStreamingSoundSource(const FSoundSource* SoundSource) const override;
-	virtual bool RequestChunk(USoundWave* SoundWave, uint32 ChunkIndex, TFunction<void(EAudioChunkLoadResult)> OnLoadCompleted) override;
+	virtual bool RequestChunk(USoundWave* SoundWave, uint32 ChunkIndex, TFunction<void(EAudioChunkLoadResult)> OnLoadCompleted, ENamedThreads::Type ThreadToCallOnLoadCompletedOn) override;
 	virtual FAudioChunkHandle GetLoadedChunk(const USoundWave* SoundWave, uint32 ChunkIndex, bool bBlockForLoad = false) const override;
 	virtual uint64 TrimMemory(uint64 NumBytesToFree) override;
 	virtual int32 RenderStatAudioStreaming(UWorld* World, FViewport* Viewport, FCanvas* Canvas, int32 X, int32 Y, const FVector* ViewLocation, const FRotator* ViewRotation) override;
diff --git a/Engine/Source/Runtime/Engine/Public/ContentStreaming.h b/Engine/Source/Runtime/Engine/Public/ContentStreaming.h
--- a/Engine/Source/Runtime/Engine/Public/ContentStreaming.h
+++ b/Engine/Source/Runtime/Engine/Public/ContentStreaming.h
@@ -497,8 +497,14 @@ struct IAudioStreamingManager : public IStreamingManager
 	/** Returns true if this is a streaming Sound Source that is managed by the streaming manager. */
 	virtual bool IsManagedStreamingSoundSource(const FSoundSource* SoundSource) const = 0;
 
-	/** Manually prepare a chunk to start playing back. This should only be used when the Load On Demand feature is enabled, and returns false on failure. */
-	virtual bool RequestChunk(USoundWave* SoundWave, uint32 ChunkIndex, TFunction<void(EAudioChunkLoadResult)> OnLoadCompleted) = 0;
+	/** 
+	 * Manually prepare a chunk to start playing back. This should only be used when the Load On Demand feature is enabled, and returns false on failure. 
+	 * @param SoundWave SoundWave we would like to request a chunk of.
+	 * @param ChunkIndex the index of that soundwave we'd like to request a chunk of.
+	 * @param OnLoadCompleted optional callback when the load completes.
+	 * @param ThreadToCallOnLoadCompleteOn. Optional specifier for which thread OnLoadCompleted should be called on.
+	 */
+	virtual bool RequestChunk(USoundWave* SoundWave, uint32 ChunkIndex, TFunction<void(EAudioChunkLoadResult)> OnLoadCompleted = [](EAudioChunkLoadResult) {}, ENamedThreads::Type ThreadToCallOnLoadCompletedOn = ENamedThreads::AnyThread) = 0;
 
 	/**
 	 * Gets a pointer to a chunk of audio data
