[NO NON-PII COMMIT MESSAGE]

diff --git a/Engine/Source/Runtime/Engine/Private/SoundWave.cpp b/Engine/Source/Runtime/Engine/Private/SoundWave.cpp
--- a/Engine/Source/Runtime/Engine/Private/SoundWave.cpp
+++ b/Engine/Source/Runtime/Engine/Private/SoundWave.cpp
@@ -42,7 +42,7 @@ FAutoConsoleVariableRef CVarSoundWaveDefaultLoadingBehavior(
 static int32 ForceNonStreamingInEditorCVar = 0;
 FAutoConsoleVariableRef CVarForceNonStreamingInEditor(
 	TEXT("au.editor.ForceAudioNonStreaming"),
-	LoadIntoCacheOnPostLoadCVar,
+	ForceNonStreamingInEditorCVar,
 	TEXT("When set to 1, forces any audio played to be non-streaming May force a DDC miss.\n")
 	TEXT("0: Honor the Play When Silent flag, 1: stop all silent non-procedural sources."),
 	ECVF_Default);
