[NO NON-PII COMMIT MESSAGE]

diff --git a/Engine/Source/Runtime/Engine/Private/SoundWave.cpp b/Engine/Source/Runtime/Engine/Private/SoundWave.cpp
--- a/Engine/Source/Runtime/Engine/Private/SoundWave.cpp
+++ b/Engine/Source/Runtime/Engine/Private/SoundWave.cpp
@@ -2362,7 +2362,7 @@ void USoundWave::RetainCompressedAudio(bool bForceSync /*= false*/)
 			{
 				WeakThis->FirstChunk = OutHandle;
 			}
-		}, ENamedThreads::GameThread);
+		}, false, 1, ENamedThreads::GameThread);
 	}
 }
 
