Trivial: Fix initialization not being macro'ed out as editor-only data to fix non-editor builds.

diff --git a/Engine/Source/Runtime/Engine/Private/SoundSubmix.cpp b/Engine/Source/Runtime/Engine/Private/SoundSubmix.cpp
--- a/Engine/Source/Runtime/Engine/Private/SoundSubmix.cpp
+++ b/Engine/Source/Runtime/Engine/Private/SoundSubmix.cpp
@@ -15,11 +15,16 @@
 
 USoundSubmix::USoundSubmix(const FObjectInitializer& ObjectInitializer)
 	: Super(ObjectInitializer)
+#if WITH_EDITORONLY_DATA
 	, SoundSubmixGraph(nullptr)
+#endif // WITH_EDITORONLY_DATA
+	, ChannelFormat(ESubmixChannelFormat::Device)
+	, bMuteWhenBackgrounded(0)
+	, AmbisonicsPluginSettings(nullptr)
+	, EnvelopeFollowerAttackTime(10)
+	, EnvelopeFollowerReleaseTime(500)
+	, OutputVolume(1.0f)
 {
-	EnvelopeFollowerAttackTime = 10;
-	EnvelopeFollowerReleaseTime = 500;
-	OutputVolume = 1.0f;
 }
 
 void USoundSubmix::StartRecordingOutput(const UObject* WorldContextObject, float ExpectedDuration)
