Fix missing include
#rnx

diff --git a/Engine/Source/Runtime/Engine/Classes/Sound/SoundSubmix.h b/Engine/Source/Runtime/Engine/Classes/Sound/SoundSubmix.h
--- a/Engine/Source/Runtime/Engine/Classes/Sound/SoundSubmix.h
+++ b/Engine/Source/Runtime/Engine/Classes/Sound/SoundSubmix.h
@@ -6,6 +6,7 @@
 #include "IAmbisonicsMixer.h"
 #include "SampleBufferIO.h"
 #include "SoundEffectSubmix.h"
+#include "SoundSubmixSend.h"
 #include "UObject/Object.h"
 #include "UObject/ObjectMacros.h"
 
@@ -13,7 +14,6 @@
 
 
 // Forward Declarations
-class FOnSubmixEnvelopeBP;
 class UEdGraph;
 class USoundEffectSubmixPreset;
 class USoundSubmix;
