Unregister all delegates when shutting down FEditorSessionSummaryWriter. This is necessary because Shutdown() is called even when the app isn't shutting down, but rather when we're only shutting down analytics.

diff --git a/Engine/Source/Runtime/Engine/Private/Analytics/EditorSessionSummaryWriter.cpp b/Engine/Source/Runtime/Engine/Private/Analytics/EditorSessionSummaryWriter.cpp
--- a/Engine/Source/Runtime/Engine/Private/Analytics/EditorSessionSummaryWriter.cpp
+++ b/Engine/Source/Runtime/Engine/Private/Analytics/EditorSessionSummaryWriter.cpp
@@ -52,8 +52,8 @@ void FEditorSessionSummaryWriter::Initialize()
 	// Register for crash and app state callbacks
 	FCoreDelegates::OnHandleSystemError.AddRaw(this, &FEditorSessionSummaryWriter::OnCrashing);
 	FCoreDelegates::ApplicationWillTerminateDelegate.AddRaw(this, &FEditorSessionSummaryWriter::OnTerminate);
-	FUserActivityTracking::OnActivityChanged.AddRaw(this, &FEditorSessionSummaryWriter::OnUserActivity);
 	FCoreDelegates::IsVanillaProductChanged.AddRaw(this, &FEditorSessionSummaryWriter::OnVanillaStateChanged);
+	FUserActivityTracking::OnActivityChanged.AddRaw(this, &FEditorSessionSummaryWriter::OnUserActivity);
 	FSlateApplication::Get().GetOnModalLoopTickEvent().AddRaw(this, &FEditorSessionSummaryWriter::Tick);
 
 	InitializeSessions();
@@ -179,6 +179,8 @@ void FEditorSessionSummaryWriter::Shutdown()
 	{
 		if (!CurrentSession->bIsTerminating && !CurrentSession->bCrashed)
 		{
+			FSlateApplication::Get().GetOnModalLoopTickEvent().RemoveAll(this);
+			
 			CurrentSession->bWasShutdown = true;
 		}
 
