Sequencer: Iterate through level sequence actors by index in case the array changes.

diff --git a/Engine/Source/Runtime/Engine/Private/LevelTick.cpp b/Engine/Source/Runtime/Engine/Private/LevelTick.cpp
--- a/Engine/Source/Runtime/Engine/Private/LevelTick.cpp
+++ b/Engine/Source/Runtime/Engine/Private/LevelTick.cpp
@@ -1511,11 +1511,11 @@ void UWorld::Tick( ELevelTick TickType, float DeltaSeconds )
 		}
 
 		// Tick level sequence actors first
-		for (AActor* LevelSequenceActor : LevelSequenceActors)
+		for (int32 i = LevelSequenceActors.Num() - 1; i >= 0; --i)
 		{
-			if (LevelSequenceActor != nullptr)
+			if (LevelSequenceActors[i] != nullptr)
 			{
-				LevelSequenceActor->Tick(DeltaSeconds);
+				LevelSequenceActors[i]->Tick(DeltaSeconds);
 			}
 		}
 	}
