Fix VT editor image, when using UDIMs with non-power-2 sizes

diff --git a/Engine/Source/Runtime/Engine/Private/Texture2D.cpp b/Engine/Source/Runtime/Engine/Private/Texture2D.cpp
--- a/Engine/Source/Runtime/Engine/Private/Texture2D.cpp
+++ b/Engine/Source/Runtime/Engine/Private/Texture2D.cpp
@@ -1963,13 +1963,14 @@ void FVirtualTexture2DResource::InitializeEditorResources(IVirtualTexture* InVir
 		while (((MipWidth > 128u && MipHeight > 128u) || MipWidth > MaxTextureSize || MipHeight > MaxTextureSize) && MipLevel < MaxMipLevel)
 		{
 			++MipLevel;
-			MipWidth = FMath::Max(MipWidth / 2u, 1u);
-			MipHeight = FMath::Max(MipHeight / 2u, 1u);
+			MipWidth = FMath::DivideAndRoundUp(MipWidth, 2u);
+			MipHeight = FMath::DivideAndRoundUp(MipHeight, 2u);
 		}
 
 		const EPixelFormat PixelFormat = VTData->LayerTypes[0];
-		const uint32 MipWidthInTiles = FMath::Max(GetNumTilesX() >> MipLevel, 1u);
-		const uint32 MipHeightInTiles = FMath::Max(GetNumTilesY() >> MipLevel, 1u);
+		const uint32 MipScaleFactor = (1u << MipLevel);
+		const uint32 MipWidthInTiles = FMath::DivideAndRoundUp(GetNumTilesX(), MipScaleFactor);
+		const uint32 MipHeightInTiles = FMath::DivideAndRoundUp(GetNumTilesY(), MipScaleFactor);
 		const uint32 TileSizeInPixels = GetTileSize();
 		const uint32 LayerMask = 1u; // FVirtualTexture2DResource should only have a single layer
 
@@ -2030,6 +2031,9 @@ void FVirtualTexture2DResource::InitializeEditorResources(IVirtualTexture* InVir
 		{
 			// Logical dimensions of mip image may be smaller than tile size (in this case tile will contain mirrored/wrapped padding)
 			// In this case, copy the proper sub-image from the tiled texture we produced into a new texture of the correct size
+			check(MipWidth <= MipWidthInTiles * TileSizeInPixels);
+			check(MipHeight <= MipHeightInTiles * TileSizeInPixels);
+
 			FTexture2DRHIRef ResizedTexture2DRHI = RHICreateTexture2D(MipWidth, MipHeight, PixelFormat, 1, 1, TexCreateFlags, CreateInfo);
 			FRHICopyTextureInfo CopyInfo;
 			CopyInfo.Size = FIntVector(MipWidth, MipHeight, 1);
