Wait for streaming virtual texture transcode to complete before deleting file handles

diff --git a/Engine/Source/Runtime/Engine/Private/VT/UploadingVirtualTexture.cpp b/Engine/Source/Runtime/Engine/Private/VT/UploadingVirtualTexture.cpp
--- a/Engine/Source/Runtime/Engine/Private/VT/UploadingVirtualTexture.cpp
+++ b/Engine/Source/Runtime/Engine/Private/VT/UploadingVirtualTexture.cpp
@@ -53,6 +53,9 @@ FUploadingVirtualTexture::~FUploadingVirtualTexture()
 	DEC_MEMORY_STAT_BY(STAT_TileHeaderSize, Data->GetTileMemoryFootprint());
 	DEC_DWORD_STAT_BY(STAT_NumTileHeaders, Data->GetNumTileHeaders());
 
+	// Complete all open transcode requests before deleting IFileCacheHandle objects
+	FVirtualTextureChunkStreamingManager::Get().WaitTasksFinished();
+
 	for (TUniquePtr<FVirtualTextureCodec>& Codec : CodecPerChunk)
 	{
 		if (Codec)
diff --git a/Engine/Source/Runtime/Engine/Private/VT/VirtualTextureChunkManager.cpp b/Engine/Source/Runtime/Engine/Private/VT/VirtualTextureChunkManager.cpp
--- a/Engine/Source/Runtime/Engine/Private/VT/VirtualTextureChunkManager.cpp
+++ b/Engine/Source/Runtime/Engine/Private/VT/VirtualTextureChunkManager.cpp
@@ -181,3 +181,8 @@ IVirtualTextureFinalizer* FVirtualTextureChunkStreamingManager::ProduceTile(FRHI
 
 	return &UploadCache;
 }
+
+void FVirtualTextureChunkStreamingManager::WaitTasksFinished() const
+{
+	TranscodeCache.WaitTasksFinished();
+}
diff --git a/Engine/Source/Runtime/Engine/Private/VT/VirtualTextureChunkManager.h b/Engine/Source/Runtime/Engine/Private/VT/VirtualTextureChunkManager.h
--- a/Engine/Source/Runtime/Engine/Private/VT/VirtualTextureChunkManager.h
+++ b/Engine/Source/Runtime/Engine/Private/VT/VirtualTextureChunkManager.h
@@ -52,6 +52,7 @@ public:
 
 	FVTRequestPageResult RequestTile(FUploadingVirtualTexture* VTexture, const FVirtualTextureProducerHandle& ProducerHandle, uint8 LayerMask, uint8 vLevel, uint32 vAddress, EVTRequestPagePriority Priority);
 	IVirtualTextureFinalizer* ProduceTile(FRHICommandListImmediate& RHICmdList, uint32 SkipBorderSize, uint8 NumLayers, uint8 LayerMask, uint64 RequestHandle, const FVTProduceTargetLayer* TargetLayers);
+	void WaitTasksFinished() const;
 
 private:
 	FVirtualTextureUploadCache UploadCache;
diff --git a/Engine/Source/Runtime/Engine/Private/VT/VirtualTextureTranscodeCache.cpp b/Engine/Source/Runtime/Engine/Private/VT/VirtualTextureTranscodeCache.cpp
--- a/Engine/Source/Runtime/Engine/Private/VT/VirtualTextureTranscodeCache.cpp
+++ b/Engine/Source/Runtime/Engine/Private/VT/VirtualTextureTranscodeCache.cpp
@@ -246,6 +246,21 @@ void FVirtualTextureTranscodeCache::WaitTaskFinished(FVTTranscodeTileHandle InHa
 	FTaskGraphInterface::Get().WaitUntilTaskCompletes(TaskEntry.GraphEvent, ENamedThreads::GetRenderThread_Local());
 }
 
+void FVirtualTextureTranscodeCache::WaitTasksFinished() const
+{
+	int32 TaskIndex = Tasks[LIST_PENDING].NextIndex;
+	while (TaskIndex != LIST_PENDING)
+	{
+		FTaskEntry const& TaskEntry = Tasks[TaskIndex];
+		const int32 NextIndex = TaskEntry.NextIndex;
+		if (!TaskEntry.GraphEvent->IsComplete())
+		{
+			FTaskGraphInterface::Get().WaitUntilTaskCompletes(TaskEntry.GraphEvent, ENamedThreads::GetRenderThread_Local());
+		}
+		TaskIndex = NextIndex;
+	}
+}
+
 const FVTUploadTileHandle* FVirtualTextureTranscodeCache::AcquireTaskResult(FVTTranscodeTileHandle InHandle)
 {
 	const uint32 TaskIndex = InHandle.Index;
diff --git a/Engine/Source/Runtime/Engine/Private/VT/VirtualTextureTranscodeCache.h b/Engine/Source/Runtime/Engine/Private/VT/VirtualTextureTranscodeCache.h
--- a/Engine/Source/Runtime/Engine/Private/VT/VirtualTextureTranscodeCache.h
+++ b/Engine/Source/Runtime/Engine/Private/VT/VirtualTextureTranscodeCache.h
@@ -72,6 +72,8 @@ public:
 	void WaitTaskFinished(FVTTranscodeTileHandle InHandle) const;
 	const FVTUploadTileHandle* AcquireTaskResult(FVTTranscodeTileHandle InHandle);
 
+	void WaitTasksFinished() const;
+
 	void RetireOldTasks(FVirtualTextureUploadCache& InUploadCache);
 
 private:
