Update comments on C, V, and Vector2D, and put a proper full comment on T.
#rnx

diff --git a/Engine/Source/Runtime/Engine/Classes/Kismet/KismetMathLibrary.h b/Engine/Source/Runtime/Engine/Classes/Kismet/KismetMathLibrary.h
--- a/Engine/Source/Runtime/Engine/Classes/Kismet/KismetMathLibrary.h
+++ b/Engine/Source/Runtime/Engine/Classes/Kismet/KismetMathLibrary.h
@@ -1007,7 +1007,7 @@ class ENGINE_API UKismetMathLibrary : public UBlueprintFunctionLibrary
 	 * @param		Current			Actual position
 	 * @param		Target			Target position
 	 * @param		DeltaTime		Time since last tick
-	 * @param		InterpSpeed		Interpolation speed
+	 * @param		InterpSpeed		Interpolation speed, if the speed given is 0, then jump to the target.
 	 * @return		New interpolated position
 	 */
 	UFUNCTION(BlueprintPure, Category="Math|Interpolation", meta=(ScriptMethod="InterpTo", Keywords="position"))
@@ -1568,7 +1568,7 @@ class ENGINE_API UKismetMathLibrary : public UBlueprintFunctionLibrary
 	 * @param		Current			Actual position
 	 * @param		Target			Target position
 	 * @param		DeltaTime		Time since last tick
-	 * @param		InterpSpeed		Interpolation speed
+	 * @param		InterpSpeed		Interpolation speed, if the speed given is 0, then jump to the target.
 	 * @return		New interpolated position
 	 */
 	UFUNCTION(BlueprintPure, Category="Math|Interpolation", meta=(ScriptMethod = "InterpTo", Keywords="position"))
@@ -2585,7 +2585,7 @@ class ENGINE_API UKismetMathLibrary : public UBlueprintFunctionLibrary
 	 * @param		Current			Current Color
 	 * @param		Target			Target Color
 	 * @param		DeltaTime		Time since last tick
-	 * @param		InterpSpeed		Interpolation speed
+	 * @param		InterpSpeed		Interpolation speed, if the speed given is 0, then jump to the target.
 	 * @return		New interpolated Color
 	 */
 	UFUNCTION(BlueprintPure, meta = (DisplayName = "Interpolate (LinearColor)", ScriptMethod = "InterpolateTo", Keywords = "color"), Category = "Math|Interpolation")
@@ -3254,7 +3254,15 @@ class ENGINE_API UKismetMathLibrary : public UBlueprintFunctionLibrary
 	UFUNCTION(BlueprintPure, meta = (DisplayName = "Ease (Transform)", BlueprintInternalUseOnly = "true", ScriptMethod = "Ease"), Category = "Math|Interpolation")
 	static FTransform TEase(const FTransform& A, const FTransform& B, float Alpha, TEnumAsByte<EEasingFunc::Type> EasingFunc, float BlendExp = 2, int32 Steps = 2);
 
-	/** Tries to reach a target transform. */
+	/**
+	 * Tries to reach Target transform based on distance from Current position, giving a nice smooth feeling when tracking a position.
+	 *
+	 * @param		Current			Actual transform
+	 * @param		Target			Target transform
+	 * @param		DeltaTime		Time since last tick
+	 * @param		InterpSpeed		Interpolation speed, if the speed given is 0, then jump to the target.
+	 * @return		New interpolated transform
+	 */
 	UFUNCTION(BlueprintPure, meta = (ScriptMethod = "InterpTo"), Category="Math|Interpolation")
 	static FTransform TInterpTo(const FTransform& Current, const FTransform& Target, float DeltaTime, float InterpSpeed);
 
