Reset the LastReatTime when the engine goes from "fixed time step" to un-fixed.

diff --git a/Engine/Source/Runtime/Engine/Private/UnrealEngine.cpp b/Engine/Source/Runtime/Engine/Private/UnrealEngine.cpp
--- a/Engine/Source/Runtime/Engine/Private/UnrealEngine.cpp
+++ b/Engine/Source/Runtime/Engine/Private/UnrealEngine.cpp
@@ -1969,9 +1969,16 @@ void UEngine::UpdateTimeAndHandleMaxTickRate()
 
 		// Did we just switch from a fixed time step to real-time?  If so, then we'll update our
 		// cached 'last time' so our current interval isn't huge (or negative!)
-		if( bTimeWasManipulated && !bUseFixedFrameRate )
+		if( bTimeWasManipulated )
 		{
-			LastRealTime = CurrentRealTime - FApp::GetDeltaTime();
+			if ( bUseFixedFrameRate )
+			{
+				LastRealTime = CurrentRealTime - (1.0/FixedFrameRate);
+			}
+			else
+			{
+				LastRealTime = CurrentRealTime - FApp::GetDeltaTime();
+			}
 			bTimeWasManipulated = false;
 		}
 
