Adding protections for audio plugin on world change listener triggering.

diff --git a/Engine/Source/Runtime/Engine/Private/AudioDevice.cpp b/Engine/Source/Runtime/Engine/Private/AudioDevice.cpp
--- a/Engine/Source/Runtime/Engine/Private/AudioDevice.cpp
+++ b/Engine/Source/Runtime/Engine/Private/AudioDevice.cpp
@@ -2890,9 +2890,16 @@ void FAudioDevice::SetListener(UWorld* World, const int32 InViewportIndex, const
 			InitializePluginListeners(World);
 			bPluginListenersInitialized = true;
 		}
-		else if (Listeners[InViewportIndex].WorldID != WorldID)
+		else 
 		{
-			NotifyPluginListenersWorldChanged(World);
+			// World change event triggered on change in world of existing listener.
+			if (InViewportIndex < Listeners.Num())
+			{
+				if (Listeners[InViewportIndex].WorldID != WorldID)
+				{
+					NotifyPluginListenersWorldChanged(World);
+				}
+			}
 		}
 	}
 
@@ -6214,4 +6221,4 @@ int32 FAudioDevice::GetNumPrecacheFrames() const
 	}
 	// Otherwise, use the default value or value set in ini file
 	return NumPrecacheFrames;
-}
\ No newline at end of file
+}
