Stream Cache Fixups: console command for clearing the audio cache (au.streamcaching.FlushAudioCache), new color scheme for stat audiostreaming, convert timing to cycles, add marker for cache miss.

diff --git a/Engine/Source/Runtime/Engine/Private/AudioStreamingCache.cpp b/Engine/Source/Runtime/Engine/Private/AudioStreamingCache.cpp
--- a/Engine/Source/Runtime/Engine/Private/AudioStreamingCache.cpp
+++ b/Engine/Source/Runtime/Engine/Private/AudioStreamingCache.cpp
@@ -52,6 +52,19 @@ FAutoConsoleVariableRef CVarReadRequestPriority(
 	TEXT("0: High, 1: Normal, 2: Below Normal, 3: Low, 4: Min"),
 	ECVF_Default);
 
+static FAutoConsoleCommand GFlushAudioCacheCommand(
+	TEXT("au.streamcaching.FlushAudioCache"),
+	TEXT("This will flush any non retained audio from the cache when Stream Caching is enabled."),
+	FConsoleCommandDelegate::CreateStatic(
+		[]()
+	{
+		static constexpr uint64 NumBytesToFree = TNumericLimits<uint64>::Max() / 2;
+		uint64 NumBytesFreed = IStreamingManager::Get().GetAudioStreamingManager().TrimMemory(NumBytesToFree);
+
+		UE_LOG(LogAudio, Display, TEXT("Audio Cache Flushed! %d megabytes free."), NumBytesFreed / (1024.0 * 1024.0));
+	})
+);
+
 FCachedAudioStreamingManager::FCachedAudioStreamingManager(const FCachedAudioStreamingManagerParams& InitParams)
 {
 	check(FPlatformCompressionUtilities::IsCurrentPlatformUsingStreamCaching());
@@ -393,6 +406,10 @@ bool FAudioChunkCache::AddOrTouchChunk(const FChunkKey& InKey, TFunction<void(EA
 			return false;
 		}
 
+#if DEBUG_STREAM_CACHE
+		CacheElement->DebugInfo.bWasCacheMiss = true;
+#endif
+
 		KickOffAsyncLoad(CacheElement, InKey, OnLoadCompleted, CallbackThread);
 
 		if (TrimCacheWhenOverBudgetCVar != 0 && MemoryCounterBytes > MemoryLimitBytes)
@@ -861,7 +878,7 @@ void FAudioChunkCache::KickOffAsyncLoad(FCacheElement* CacheElement, const FChun
 		}
 		
 #if DEBUG_STREAM_CACHE
-		CacheElement->DebugInfo.TimeLoadStarted = FPlatformTime::Seconds();
+		CacheElement->DebugInfo.TimeLoadStarted = FPlatformTime::Cycles64();
 #endif
 
 
@@ -873,7 +890,7 @@ void FAudioChunkCache::KickOffAsyncLoad(FCacheElement* CacheElement, const FChun
 			CacheElement->bIsLoaded = true;
 
 #if DEBUG_STREAM_CACHE
-			CacheElement->DebugInfo.TimeToLoad = (FPlatformTime::Seconds() - CacheElement->DebugInfo.TimeLoadStarted) * 1000.0f;
+			CacheElement->DebugInfo.TimeToLoad = FPlatformTime::ToMilliseconds64(FPlatformTime::Cycles64() - CacheElement->DebugInfo.TimeLoadStarted);
 #endif
 			EAudioChunkLoadResult ChunkLoadResult = bRequestFailed ? EAudioChunkLoadResult::Interrupted : EAudioChunkLoadResult::Completed;
 			ExecuteOnLoadCompleteCallback(ChunkLoadResult, OnLoadCompleted, CallbackThread);
@@ -917,7 +934,7 @@ void FAudioChunkCache::KickOffAsyncLoad(FCacheElement* CacheElement, const FChun
 			CacheElement->bIsLoaded = true;
 
 #if DEBUG_STREAM_CACHE
-			CacheElement->DebugInfo.TimeToLoad = FPlatformTime::Seconds() - CacheElement->DebugInfo.TimeLoadStarted;
+			CacheElement->DebugInfo.TimeToLoad = (FPlatformTime::Seconds() - CacheElement->DebugInfo.TimeLoadStarted) * 1000.0f;
 #endif
 
 			OnLoadCompleted(bWasCancelled ? EAudioChunkLoadResult::Interrupted : EAudioChunkLoadResult::Completed);
@@ -1048,6 +1065,26 @@ uint64 FCachedAudioStreamingManager::TrimMemory(uint64 NumBytesToFree)
 
 TPair<int, int> FAudioChunkCache::DebugDisplay(UWorld* World, FViewport* Viewport, FCanvas* Canvas, int32 X, int32 Y, const FVector* ViewLocation, const FRotator* ViewRotation) const
 {
+	// Color scheme:
+	static constexpr float ColorMax = 256.0f;
+
+
+	// Chunk color for a single retainer.
+	const FLinearColor RetainChunkColor(44.0f / ColorMax, 207.0f / ColorMax, 47 / ColorMax);
+
+	// Chunk color we lerp to as more retainers are added for a chunk.
+	const FLinearColor TotalMassRetainChunkColor(204 / ColorMax, 126 / ColorMax, 43 / ColorMax);
+
+	// A chunk that's loaded but not retained.
+	const FLinearColor LoadedChunkColor(47 / ColorMax, 44 / ColorMax, 207 / ColorMax);
+
+	// A chunk that's been trimmed by TrimMemory.
+	const FLinearColor TrimmedChunkColor(204 / ColorMax, 46 / ColorMax, 43 / ColorMax);
+
+	// A chunk that currently has an async load in flight.
+	const FLinearColor CurrentlyLoadingChunkColor = FLinearColor::Yellow;
+
+
 	const int32 InitialX = X;
 	const int32 InitialY = Y;
 
@@ -1102,17 +1139,19 @@ TPair<int, int> FAudioChunkCache::DebugDisplay(UWorld* World, FViewport* Viewpor
 		int32 NumTimesTouched = -1;
 		double TimeToLoad = -1.0;
 		float AveragePlaceInCache = -1.0f;
+		bool bWasCacheMiss = false;
 
 #if DEBUG_STREAM_CACHE
 		NumTotalChunks = CurrentElement->DebugInfo.NumTotalChunks;
 		NumTimesTouched = CurrentElement->DebugInfo.NumTimesTouched;
 		TimeToLoad = CurrentElement->DebugInfo.TimeToLoad;
 		AveragePlaceInCache = CurrentElement->DebugInfo.AverageLocationInCacheWhenNeeded;
+		bWasCacheMiss = CurrentElement->DebugInfo.bWasCacheMiss;
 #endif
 
 		const bool bWasTrimmed = CurrentElement->ChunkData.Num() == 0;
 
-		FString ElementInfo = *FString::Printf(TEXT("%4i. Size: %6.2f KB   Chunk: %d of %d   Request Count: %d    Average Index: %6.2f  Number of Decoders Using Chunk: %d     Chunk Load Time: %6.4fms      Name: %s "),
+		FString ElementInfo = *FString::Printf(TEXT("%4i. Size: %6.2f KB   Chunk: %d of %d   Request Count: %d    Average Index: %6.2f  Number of Handles Retaining Chunk: %d     Chunk Load Time: %6.4fms      Name: %s %s"),
 			Index,
 			CurrentElement->ChunkData.Num() / 1024.0f,
 			CurrentElement->Key.ChunkIndex,
@@ -1121,13 +1160,14 @@ TPair<int, int> FAudioChunkCache::DebugDisplay(UWorld* World, FViewport* Viewpor
 			AveragePlaceInCache,
 			CurrentElement->NumConsumers.GetValue(),
 			TimeToLoad,
-			bWasTrimmed ? TEXT("TRIMMED CHUNK") : *CurrentElement->Key.SoundWaveName.ToString()
+			bWasTrimmed ? TEXT("TRIMMED CHUNK") : *CurrentElement->Key.SoundWaveName.ToString(),
+			bWasCacheMiss ? TEXT("(Cache Miss!)") : TEXT("")
 			);
 
 		// Since there's a lot of info here,
 		// Subtly fading the chunk info to gray seems to help as a visual indicator of how far down on the list things are.
 		ColorLerpAmount = FMath::Min(ColorLerpAmount + ColorLerpStep, 1.0f);
-		FLinearColor TextColor = FLinearColor::LerpUsingHSV(FLinearColor::Blue, FLinearColor::Gray, ColorLerpAmount);
+		FLinearColor TextColor = FLinearColor::LerpUsingHSV(LoadedChunkColor, FLinearColor::Gray, ColorLerpAmount);
 
 		// If there's a load in flight, paint this element yellow.
 		if (CurrentElement->IsLoadInProgress())
@@ -1136,11 +1176,15 @@ TPair<int, int> FAudioChunkCache::DebugDisplay(UWorld* World, FViewport* Viewpor
 		}
 		else if (CurrentElement->IsInUse())
 		{
-			TextColor = FLinearColor::Green;
+			// We slowly fade our text color based on how many refererences there are to this chunk.
+			static const float MaxNumHandles = 12.0f;
+
+			ColorLerpAmount = FMath::Min(CurrentElement->NumConsumers.GetValue() / MaxNumHandles, 1.0f);
+			TextColor = FLinearColor::LerpUsingHSV(RetainChunkColor, TotalMassRetainChunkColor, ColorLerpAmount);
 		}
 		else if (bWasTrimmed)
 		{
-			TextColor = FLinearColor::Red;
+			TextColor = TrimmedChunkColor;
 		}
 
 		Canvas->DrawShadowedString(X, Y, *ElementInfo, UEngine::GetSmallFont(), TextColor);
@@ -1156,4 +1200,4 @@ TPair<int, int> FAudioChunkCache::DebugDisplay(UWorld* World, FViewport* Viewpor
 	UEngine::GetSmallFont()->GetStringHeightAndWidth(*CacheMemoryUsage, CacheMemoryTextOffsetX, CacheMemoryTextOffsetY);
 
 	return TPair<int, int>(X + CacheTitleOffsetX + CacheMemoryTextOffsetX - InitialX, Y - InitialY);
-}
+}
\ No newline at end of file
diff --git a/Engine/Source/Runtime/Engine/Public/AudioStreamingCache.h b/Engine/Source/Runtime/Engine/Public/AudioStreamingCache.h
--- a/Engine/Source/Runtime/Engine/Public/AudioStreamingCache.h
+++ b/Engine/Source/Runtime/Engine/Public/AudioStreamingCache.h
@@ -100,19 +100,23 @@ private:
 		// Number of times this chunk was requested during its time in the cache.
 		int32 NumTimesTouched;
 
-		double TimeLoadStarted;
+		uint64 TimeLoadStarted;
 		// Amount of time spent loading the audio file.
-		double TimeToLoad;
+		float TimeToLoad;
 
 		// This is a cumulative moving average of a chunks location before it was 
 		float AverageLocationInCacheWhenNeeded;
 
+		// if true, 
+		bool bWasCacheMiss;
+
 		FCacheElementDebugInfo()
 			: NumTotalChunks(0)
 			, NumTimesTouched(0)
 			, TimeLoadStarted(0.0)
 			, TimeToLoad(0.0)
 			, AverageLocationInCacheWhenNeeded(0.0f)
+			, bWasCacheMiss(false)
 		{
 		}
 
@@ -121,7 +125,8 @@ private:
 			NumTotalChunks = 0;
 			NumTimesTouched = 0;
 			TimeLoadStarted = 0;
-
+			TimeToLoad = 0.0f;
+			bWasCacheMiss = false;
 			AverageLocationInCacheWhenNeeded = 0.0f;
 		}
 	};
