[NO NON-PII COMMIT MESSAGE]

diff --git a/Engine/Source/Runtime/Engine/Private/SoundWave.cpp b/Engine/Source/Runtime/Engine/Private/SoundWave.cpp
--- a/Engine/Source/Runtime/Engine/Private/SoundWave.cpp
+++ b/Engine/Source/Runtime/Engine/Private/SoundWave.cpp
@@ -2353,8 +2353,7 @@ void USoundWave::RetainCompressedAudio(bool bForceSync /*= false*/)
 	}
 	else
 	{
-		TWeakObjectPtr<USoundWave> WeakThis = TWeakObjectPtr<USoundWave>(this);
-		GetHandleForChunkOfAudio([WeakThis](FAudioChunkHandle OutHandle)
+		GetHandleForChunkOfAudio([WeakThis = MakeWeakObjectPtr(this)](FAudioChunkHandle OutHandle)
 		{
 			check(IsInGameThread());
 			
