Add stream chunk override in per platform cook overrides.

diff --git a/Engine/Source/Developer/Android/AndroidTargetPlatform/Private/AndroidTargetPlatform.cpp b/Engine/Source/Developer/Android/AndroidTargetPlatform/Private/AndroidTargetPlatform.cpp
--- a/Engine/Source/Developer/Android/AndroidTargetPlatform/Private/AndroidTargetPlatform.cpp
+++ b/Engine/Source/Developer/Android/AndroidTargetPlatform/Private/AndroidTargetPlatform.cpp
@@ -668,6 +668,10 @@ namespace Android
 			OutOverrides.SoundCueCookQualityIndex = SoundCueQualityIndex;
 		}
 
+		int32 RetrievedChunkSizeKB = 256;
+		GConfig->GetInt(CategoryName, TEXT("ChunkSizeKB"), RetrievedChunkSizeKB, GEngineIni);
+		OutOverrides.StreamChunkSizeKB = RetrievedChunkSizeKB;
+
 		GConfig->GetBool(CategoryName, TEXT("bUseAudioStreamCaching"), OutOverrides.bUseStreamCaching, GEngineIni);
 
 		/** Memory Load On Demand Settings */
diff --git a/Engine/Source/Developer/IOS/IOSTargetPlatform/Private/IOSTargetPlatform.cpp b/Engine/Source/Developer/IOS/IOSTargetPlatform/Private/IOSTargetPlatform.cpp
--- a/Engine/Source/Developer/IOS/IOSTargetPlatform/Private/IOSTargetPlatform.cpp
+++ b/Engine/Source/Developer/IOS/IOSTargetPlatform/Private/IOSTargetPlatform.cpp
@@ -728,6 +728,10 @@ namespace IOS
 			OutOverrides.SoundCueCookQualityIndex = SoundCueQualityIndex;
 		}
 
+		int32 RetrievedChunkSizeKB = 256;
+		GConfig->GetInt(CategoryName, TEXT("ChunkSizeKB"), RetrievedChunkSizeKB, GEngineIni);
+		OutOverrides.StreamChunkSizeKB = RetrievedChunkSizeKB;
+
 		GConfig->GetBool(CategoryName, TEXT("bUseAudioStreamCaching"), OutOverrides.bUseStreamCaching, GEngineIni);
 
 		/** Memory Load On Demand Settings */
diff --git a/Engine/Source/Developer/Windows/WindowsTargetPlatform/Classes/WindowsTargetSettings.h b/Engine/Source/Developer/Windows/WindowsTargetPlatform/Classes/WindowsTargetSettings.h
--- a/Engine/Source/Developer/Windows/WindowsTargetPlatform/Classes/WindowsTargetSettings.h
+++ b/Engine/Source/Developer/Windows/WindowsTargetPlatform/Classes/WindowsTargetSettings.h
@@ -111,6 +111,10 @@ public:
 	UPROPERTY(GlobalConfig, EditAnywhere, Category = "Audio")
 	FPlatformRuntimeAudioCompressionOverrides CompressionOverrides;
 
+	/** This determines how we split compressed audio into chunks for this platform. The smaller this value is the more granular our chunking is. */
+	UPROPERTY(GlobalConfig, EditAnywhere, Category = "Audio|CookOverrides", meta = (DisplayName = "Max Size Per Streaming Chunk (KB)"))
+	int32 ChunkSizeKB;
+
 	/** When this is enabled, Actual compressed data will be separated from the USoundWave, and loaded into a cache. */
 	UPROPERTY(GlobalConfig, EditAnywhere, Category = "Audio|CookOverrides", meta = (DisplayName = "Use Stream Caching (Experimental)"))
 	bool bUseAudioStreamCaching;
diff --git a/Engine/Source/Developer/Windows/WindowsTargetPlatform/Private/GenericWindowsTargetPlatform.h b/Engine/Source/Developer/Windows/WindowsTargetPlatform/Private/GenericWindowsTargetPlatform.h
--- a/Engine/Source/Developer/Windows/WindowsTargetPlatform/Private/GenericWindowsTargetPlatform.h
+++ b/Engine/Source/Developer/Windows/WindowsTargetPlatform/Private/GenericWindowsTargetPlatform.h
@@ -34,6 +34,10 @@ namespace Windows
 
 		GConfig->GetBool(CategoryName, TEXT("bUseAudioStreamCaching"), OutOverrides.bUseStreamCaching, GEngineIni);
 
+		int32 RetrievedChunkSizeKB = 256;
+		GConfig->GetInt(CategoryName, TEXT("ChunkSizeKB"), RetrievedChunkSizeKB, GEngineIni);
+		OutOverrides.StreamChunkSizeKB = RetrievedChunkSizeKB;
+
 		/** Memory Load On Demand Settings */
 		if (OutOverrides.bUseStreamCaching)
 		{
diff --git a/Engine/Source/Runtime/Android/AndroidRuntimeSettings/Classes/AndroidRuntimeSettings.h b/Engine/Source/Runtime/Android/AndroidRuntimeSettings/Classes/AndroidRuntimeSettings.h
--- a/Engine/Source/Runtime/Android/AndroidRuntimeSettings/Classes/AndroidRuntimeSettings.h
+++ b/Engine/Source/Runtime/Android/AndroidRuntimeSettings/Classes/AndroidRuntimeSettings.h
@@ -490,6 +490,10 @@ public:
 	UPROPERTY(config, EditAnywhere, Category = "Audio")
 	FPlatformRuntimeAudioCompressionOverrides CompressionOverrides;
 
+	/** This determines how we split compressed audio into chunks for this platform. The smaller this value is the more granular our chunking is. */
+	UPROPERTY(GlobalConfig, EditAnywhere, Category = "Audio|CookOverrides", meta = (DisplayName = "Max Size Per Streaming Chunk (KB)"))
+	int32 ChunkSizeKB;
+
 	/** When this is enabled, Actual compressed data will be separated from the USoundWave, and loaded into a cache. */
 	UPROPERTY(GlobalConfig, EditAnywhere, Category = "Audio|CookOverrides", meta = (DisplayName = "Use Stream Caching (Experimental)"))
 	bool bUseAudioStreamCaching;
diff --git a/Engine/Source/Runtime/AudioPlatformConfiguration/Public/AudioCompressionSettings.h b/Engine/Source/Runtime/AudioPlatformConfiguration/Public/AudioCompressionSettings.h
--- a/Engine/Source/Runtime/AudioPlatformConfiguration/Public/AudioCompressionSettings.h
+++ b/Engine/Source/Runtime/AudioPlatformConfiguration/Public/AudioCompressionSettings.h
@@ -59,6 +59,8 @@ struct FPlatformAudioCookOverrides
 	// If set, the cooker will keep only this level of quality
 	int32 SoundCueCookQualityIndex = INDEX_NONE;
 
+	int32 StreamChunkSizeKB;
+	
 	// When set to any platform > 0.0, this will automatically set any USoundWave beyond this value to be streamed from disk.
 	// If StreamCaching is set to true, this will be used 
 	float AutoStreamingThreshold;
@@ -72,6 +74,7 @@ struct FPlatformAudioCookOverrides
 	FPlatformAudioCookOverrides()
 		: bResampleForDevice(false)
 		, CompressionQualityModifier(1.0f)
+		, StreamChunkSizeKB(256)
 		, AutoStreamingThreshold(0.0f)
 		, bUseStreamCaching(false)
 	{
@@ -93,6 +96,8 @@ struct FPlatformAudioCookOverrides
 		int32 CompressionQualityHash = FMath::FloorToInt(InOverrides->CompressionQualityModifier * 100.0f);
 		OutSuffix.AppendInt(CompressionQualityHash);
 
+		OutSuffix.AppendInt(InOverrides->StreamChunkSizeKB);
+
 		int32 AutoStreamingThresholdHash = FMath::FloorToInt(InOverrides->AutoStreamingThreshold * 100.0f);
 		OutSuffix.AppendInt(AutoStreamingThresholdHash);
 
diff --git a/Engine/Source/Runtime/Engine/Private/AudioCompressionSettingsUtils.cpp b/Engine/Source/Runtime/Engine/Private/AudioCompressionSettingsUtils.cpp
--- a/Engine/Source/Runtime/Engine/Private/AudioCompressionSettingsUtils.cpp
+++ b/Engine/Source/Runtime/Engine/Private/AudioCompressionSettingsUtils.cpp
@@ -103,6 +103,10 @@ void CacheCurrentPlatformAudioCookOverrides(FPlatformAudioCookOverrides& OutOver
 
 	GConfig->GetBool(CategoryName, TEXT("bUseAudioStreamCaching"), OutOverrides.bUseStreamCaching, GEngineIni);
 
+	int32 RetrievedChunkSizeKB = 256;
+	GConfig->GetInt(CategoryName, TEXT("ChunkSizeKB"), RetrievedChunkSizeKB, GEngineIni);
+	OutOverrides.StreamChunkSizeKB = RetrievedChunkSizeKB;
+
 	/** Memory Load On Demand Settings */
 	if (OutOverrides.bUseStreamCaching)
 	{
@@ -206,7 +210,16 @@ const FAudioStreamCachingSettings& FPlatformCompressionUtilities::GetStreamCachi
 FCachedAudioStreamingManagerParams FPlatformCompressionUtilities::BuildCachedStreamingManagerParams()
 {
 	const FAudioStreamCachingSettings& CacheSettings = GetStreamCachingSettingsForCurrentPlatform();
-	int32 MaxChunkSize = GetMaxChunkSizeForCookOverrides(GetCookOverridesForCurrentPlatform());
+
+	int32 MaxChunkSize = 256 * 1024;
+	const FPlatformAudioCookOverrides* CookOverrides = GetCookOverridesForCurrentPlatform();
+
+	if (CookOverrides && (CookOverrides->StreamChunkSizeKB != 0))
+	{
+		MaxChunkSize = CookOverrides->StreamChunkSizeKB * 1024;;
+	}
+
+	MaxChunkSize = GetMaxChunkSizeForCookOverrides(GetCookOverridesForCurrentPlatform(), MaxChunkSize / 1024);
 
 	// Our number of elements is tweakable based on the minimum cache usage we want to support.
 	const float MinimumCacheUsage = FMath::Clamp(MinimumCacheUsageCvar, 0.0f, 0.95f);
@@ -227,7 +240,7 @@ FCachedAudioStreamingManagerParams FPlatformCompressionUtilities::BuildCachedStr
 	return Params;
 }
 
-uint32 FPlatformCompressionUtilities::GetMaxChunkSizeForCookOverrides(const FPlatformAudioCookOverrides* InCompressionOverrides)
+uint32 FPlatformCompressionUtilities::GetMaxChunkSizeForCookOverrides(const FPlatformAudioCookOverrides* InCompressionOverrides, int32 DefaultMaxChunkSizeKB)
 {
 	check(InCompressionOverrides);
 
@@ -235,7 +248,6 @@ uint32 FPlatformCompressionUtilities::GetMaxChunkSizeForCookOverrides(const FPla
 	// If the game runs with higher than 32 voices, that means we will potentially have a larger cache than what was set in the target settings.
 	// In that case we log a warning on application launch.
 	const int32 MinimumNumChunks = 32;
-	const int32 DefaultMaxChunkSizeKB = 256;
 	
 	int32 CacheSizeKB = InCompressionOverrides->StreamCachingSettings.CacheSizeKB;
 	
diff --git a/Engine/Source/Runtime/Engine/Private/AudioDerivedData.cpp b/Engine/Source/Runtime/Engine/Private/AudioDerivedData.cpp
--- a/Engine/Source/Runtime/Engine/Private/AudioDerivedData.cpp
+++ b/Engine/Source/Runtime/Engine/Private/AudioDerivedData.cpp
@@ -330,6 +330,12 @@ class FStreamedAudioCacheDerivedDataWorker : public FNonAbandonableTask
 				// Set the ideal chunk size to be 256k to optimize for data reads on console.
 				int32 MaxChunkSize = 256 * 1024;
 				
+				// If there is a chunk size override, use that.
+				if (CompressionOverrides && (CompressionOverrides->StreamChunkSizeKB != 0))
+				{
+					MaxChunkSize = CompressionOverrides->StreamChunkSizeKB * 1024;
+				}
+
 				// By default, the first chunk's max size is the same as the other chunks.
 				int32 FirstChunkSize = MaxChunkSize;
 
@@ -352,8 +358,8 @@ class FStreamedAudioCacheDerivedDataWorker : public FNonAbandonableTask
 
 				if (bUseStreamCaching)
 				{
-					// Use the chunk size for this duration:
-					MaxChunkSize = FPlatformCompressionUtilities::GetMaxChunkSizeForCookOverrides(CompressionOverrides);
+					// Fix up chunk size if we have too small a cache size to safely play our sources.
+					MaxChunkSize = FPlatformCompressionUtilities::GetMaxChunkSizeForCookOverrides(CompressionOverrides, MaxChunkSize / 1024);
 					UE_LOG(LogAudio, Display, TEXT("Chunk size for %s: %d"), *SoundWave.GetFullName(), MaxChunkSize);
 				}
 				
diff --git a/Engine/Source/Runtime/Engine/Public/AudioCompressionSettingsUtils.h b/Engine/Source/Runtime/Engine/Public/AudioCompressionSettingsUtils.h
--- a/Engine/Source/Runtime/Engine/Public/AudioCompressionSettingsUtils.h
+++ b/Engine/Source/Runtime/Engine/Public/AudioCompressionSettingsUtils.h
@@ -32,7 +32,7 @@ public:
 	static FCachedAudioStreamingManagerParams BuildCachedStreamingManagerParams();
 
 	/** This is used at runtime in BuildCachedStreamingManagerParams, as well as cooktime in FStreamedAudioCacheDerivedDataWorker::BuildStreamedAudio to split compressed audio.  */
-	static uint32 GetMaxChunkSizeForCookOverrides(const FPlatformAudioCookOverrides* InCompressionOverrides);
+	static uint32 GetMaxChunkSizeForCookOverrides(const FPlatformAudioCookOverrides* InCompressionOverrides, int32 DefaultMaxChunkSizeKB);
 
 private:
 	static const FPlatformRuntimeAudioCompressionOverrides* GetRuntimeCompressionOverridesForCurrentPlatform();
diff --git a/Engine/Source/Runtime/IOS/IOSRuntimeSettings/Classes/IOSRuntimeSettings.h b/Engine/Source/Runtime/IOS/IOSRuntimeSettings/Classes/IOSRuntimeSettings.h
--- a/Engine/Source/Runtime/IOS/IOSRuntimeSettings/Classes/IOSRuntimeSettings.h
+++ b/Engine/Source/Runtime/IOS/IOSRuntimeSettings/Classes/IOSRuntimeSettings.h
@@ -515,6 +515,10 @@ public:
 	UPROPERTY(config, EditAnywhere, Category = "Audio")
 	FPlatformRuntimeAudioCompressionOverrides CompressionOverrides;
 
+	/** This determines how we split compressed audio into chunks for this platform. The smaller this value is the more granular our chunking is. */
+	UPROPERTY(GlobalConfig, EditAnywhere, Category = "Audio|CookOverrides", meta = (DisplayName = "Max Size Per Streaming Chunk (KB)"))
+	int32 ChunkSizeKB;
+
 	/** When this is enabled, Actual compressed data will be separated from the USoundWave, and loaded into a cache. */
 	UPROPERTY(GlobalConfig, EditAnywhere, Category = "Audio|CookOverrides", meta = (DisplayName = "Use Stream Caching (Experimental)"))
 	bool bUseAudioStreamCaching;
