#ue4 - Prevent NaNs (or excessively large input values causing later NaNs) from entering AController::SetControlLocation().

diff --git a/Engine/Source/Runtime/Engine/Classes/GameFramework/Controller.h b/Engine/Source/Runtime/Engine/Classes/GameFramework/Controller.h
--- a/Engine/Source/Runtime/Engine/Classes/GameFramework/Controller.h
+++ b/Engine/Source/Runtime/Engine/Classes/GameFramework/Controller.h
@@ -86,6 +86,9 @@ protected:
 	UPROPERTY()
 	FRotator ControlRotation;
 
+	/** Return false if rotation contains NaN or extremely large values (usually resulting from uninitialized values). */
+	bool IsValidControlRotation(FRotator CheckRotation) const;
+
 	/**
 	 * If true, the controller location will match the possessed Pawn's location. If false, it will not be updated. Rotation will match ControlRotation in either case.
 	 * Since a Controller's location is normally inaccessible, this is intended mainly for purposes of being able to attach
diff --git a/Engine/Source/Runtime/Engine/Private/Components/CharacterMovementComponent.cpp b/Engine/Source/Runtime/Engine/Private/Components/CharacterMovementComponent.cpp
--- a/Engine/Source/Runtime/Engine/Private/Components/CharacterMovementComponent.cpp
+++ b/Engine/Source/Runtime/Engine/Private/Components/CharacterMovementComponent.cpp
@@ -5636,7 +5636,8 @@ void UCharacterMovementComponent::OnTeleported()
 
 float GetAxisDeltaRotation(float InAxisRotationRate, float DeltaTime)
 {
-	return (InAxisRotationRate >= 0.f) ? (InAxisRotationRate * DeltaTime) : 360.f;
+	// Values over 360 don't do anything, see FMath::FixedTurn. However we are trying to avoid giant floats from overflowing other calculations.
+	return (InAxisRotationRate >= 0.f) ? FMath::Min(InAxisRotationRate * DeltaTime, 360.f) : 360.f;
 }
 
 FRotator UCharacterMovementComponent::GetDeltaRotation(float DeltaTime) const
diff --git a/Engine/Source/Runtime/Engine/Private/Controller.cpp b/Engine/Source/Runtime/Engine/Private/Controller.cpp
--- a/Engine/Source/Runtime/Engine/Private/Controller.cpp
+++ b/Engine/Source/Runtime/Engine/Private/Controller.cpp
@@ -29,6 +29,16 @@ DEFINE_LOG_CATEGORY(LogPath);
 
 #define LOCTEXT_NAMESPACE "Controller"
 
+namespace ControllerStatics
+{
+	static float InvalidControlRotationMagnitude = 8388608.f; // 2^23, largest float when fractions are lost, and where FMod loses meaningful precision.
+	static FAutoConsoleVariableRef CVarInvalidControlRotationMagnitude(
+		TEXT("Controller.InvalidControlRotationMagnitude"), InvalidControlRotationMagnitude,
+		TEXT("If any component of an FRotator passed to SetControlRotation is larger than this magnitude, ignore the value. Huge values are usually from uninitialized variables and can cause NaN/Inf to propagate later."),
+		ECVF_Default);
+}
+
+
 
 AController::AController(const FObjectInitializer& ObjectInitializer)
 	: Super(ObjectInitializer)
@@ -106,13 +116,12 @@ FRotator AController::GetControlRotation() const
 
 void AController::SetControlRotation(const FRotator& NewRotation)
 {
-#if ENABLE_NAN_DIAGNOSTIC
-	if (NewRotation.ContainsNaN())
+	if (!IsValidControlRotation(NewRotation))
 	{
-		logOrEnsureNanError(TEXT("AController::SetControlRotation about to apply NaN-containing rotation! (%s)"), *NewRotation.ToString());
+		logOrEnsureNanError(TEXT("AController::SetControlRotation atttempted to apply NaN-containing or NaN-causing rotation! (%s)"), *NewRotation.ToString());
 		return;
 	}
-#endif
+
 	if (!ControlRotation.Equals(NewRotation, 1e-3f))
 	{
 		ControlRotation = NewRotation;
@@ -124,8 +133,27 @@ void AController::SetControlRotation(const FRotator& NewRotation)
 	}
 	else
 	{
-		//UE_LOG(LogPlayerController, Log, TEXT("Skipping SetControlRotation for %s (Pawn %s)"), *GetNameSafe(this), *GetNameSafe(GetPawn()));
+		//UE_LOG(LogPlayerController, Log, TEXT("Skipping SetControlRotation %s for %s (Pawn %s)"), *NewRotation.ToString(), *GetNameSafe(this), *GetNameSafe(GetPawn()));
+	}
+}
+
+bool AController::IsValidControlRotation(FRotator CheckRotation) const
+{
+	if (CheckRotation.ContainsNaN())
+	{
+		return false;
+	}
+
+	// Really large values can be technically valid but are usually the result of uninitialized values, and those can cause
+	// conversion to FQuat or Vector to fail and generate NaN or Inf.
+	if (FMath::Abs(CheckRotation.Pitch) >= ControllerStatics::InvalidControlRotationMagnitude ||
+		FMath::Abs(CheckRotation.Yaw  ) >= ControllerStatics::InvalidControlRotationMagnitude ||
+		FMath::Abs(CheckRotation.Roll ) >= ControllerStatics::InvalidControlRotationMagnitude)
+	{
+		return false;
 	}
+
+	return true;
 }
 
 
