Inertialization error logging.
- Inertialization requests made without the context of an ancestor inertialization node now log an error identifying the source of the request.
- Added an editor only table to map runtime nodes back to their source. This enables runtime node errors to identify the source.
- The log messages are text only for now, as this is a limitation of the current message log handling in AnimInstanceProxy. However, having the reverse map gets us halfway to runtime node messages with clickable hyperlinks.
- Don't log errors if anim graph blend time is zero. Anim graph times are allowed to be zero to cancel inertialziation, but there's no need to complain about a missing inertialization node in that case.
- Improved inertialization documentation

diff --git a/Engine/Source/Editor/AnimGraph/Classes/AnimationGraph.h b/Engine/Source/Editor/AnimGraph/Classes/AnimationGraph.h
--- a/Engine/Source/Editor/AnimGraph/Classes/AnimationGraph.h
+++ b/Engine/Source/Editor/AnimGraph/Classes/AnimationGraph.h
@@ -21,7 +21,7 @@ class UAnimationGraph : public UEdGraph
 	/** Delegate fired when a pin's default value is changed */
 	FOnPinDefaultValueChanged OnPinDefaultValueChanged;
 
-	/** Blending options for this graph (only valid for linked anim graphs with inertialization) */
+	/** Blending options for animation graphs in Linked Animation Blueprints. */
 	UPROPERTY(EditAnywhere, Category = GraphBlending, meta = (ShowOnlyInnerProperties))
 	FAnimGraphBlendOptions BlendOptions;
 };
diff --git a/Engine/Source/Editor/KismetCompiler/Private/AnimBlueprintCompiler.cpp b/Engine/Source/Editor/KismetCompiler/Private/AnimBlueprintCompiler.cpp
--- a/Engine/Source/Editor/KismetCompiler/Private/AnimBlueprintCompiler.cpp
+++ b/Engine/Source/Editor/KismetCompiler/Private/AnimBlueprintCompiler.cpp
@@ -512,8 +512,10 @@ void FAnimBlueprintCompilerContext::ProcessAnimationNode(UAnimGraphNode_Base* Vi
 	SourceNodeToProcessedNodeMap.Add(TrueSourceObject, VisualAnimNode);
 
 	// Register the slightly more permanent debug information
-	NewAnimBlueprintClass->GetAnimBlueprintDebugData().NodePropertyToIndexMap.Add(TrueSourceObject, AllocatedIndex);
-	NewAnimBlueprintClass->GetAnimBlueprintDebugData().NodeGuidToIndexMap.Add(TrueSourceObject->NodeGuid, AllocatedIndex);
+	FAnimBlueprintDebugData& NewAnimBlueprintDebugData = NewAnimBlueprintClass->GetAnimBlueprintDebugData();
+	NewAnimBlueprintDebugData.NodePropertyToIndexMap.Add(TrueSourceObject, AllocatedIndex);
+	NewAnimBlueprintDebugData.NodeGuidToIndexMap.Add(TrueSourceObject->NodeGuid, AllocatedIndex);
+	NewAnimBlueprintDebugData.NodePropertyIndexToNodeMap.Add(AllocatedIndex, TrueSourceObject);
 	NewAnimBlueprintClass->GetDebugData().RegisterClassPropertyAssociation(TrueSourceObject, NewProperty);
 
 	// Node-specific compilation that requires compiler state info
@@ -1839,7 +1841,7 @@ void FAnimBlueprintCompilerContext::CopyTermDefaultsToDefaultObject(UObject* Def
 		
 			// @TODO this is quick solution for crash - if there were previous errors and some nodes were not added, they could still end here -
 			// this check avoids that and since there are already errors, compilation won't be successful.
-			// but I'd prefer stoping compilation earlier to avoid getting here in first place
+			// but I'd prefer stopping compilation earlier to avoid getting here in first place
 			if (LinkIndexMap.Contains(LinkingNode) && LinkIndexMap.Contains(LinkedNode))
 			{
 				const int32 SourceNodeIndex = LinkIndexMap.FindChecked(LinkingNode);
diff --git a/Engine/Source/Runtime/AnimGraphRuntime/Private/AnimNodes/AnimNode_BlendListBase.cpp b/Engine/Source/Runtime/AnimGraphRuntime/Private/AnimNodes/AnimNode_BlendListBase.cpp
--- a/Engine/Source/Runtime/AnimGraphRuntime/Private/AnimNodes/AnimNode_BlendListBase.cpp
+++ b/Engine/Source/Runtime/AnimGraphRuntime/Private/AnimNodes/AnimNode_BlendListBase.cpp
@@ -111,7 +111,11 @@ void FAnimNode_BlendListBase::Update_AnyThread(const FAnimationUpdateContext& Co
 				FAnimNode_Inertialization* InertializationNode = Context.GetAncestor<FAnimNode_Inertialization>();
 				if (InertializationNode)
 				{
-					InertializationNode->Request(BlendTime[ChildIndex]);
+					InertializationNode->RequestInertialization(BlendTime[ChildIndex]);
+				}
+				else
+				{
+					FAnimNode_Inertialization::LogRequestError(Context, BlendPose[ChildIndex]);
 				}
 				
 				RemainingBlendTime = 0.0f;
diff --git a/Engine/Source/Runtime/Engine/Classes/Animation/AnimBlueprintGeneratedClass.h b/Engine/Source/Runtime/Engine/Classes/Animation/AnimBlueprintGeneratedClass.h
--- a/Engine/Source/Runtime/Engine/Classes/Animation/AnimBlueprintGeneratedClass.h
+++ b/Engine/Source/Runtime/Engine/Classes/Animation/AnimBlueprintGeneratedClass.h
@@ -102,6 +102,9 @@ public:
 	// Map from animation node to their property index
 	TMap<TWeakObjectPtr<const UAnimGraphNode_Base>, int32> NodePropertyToIndexMap;
 
+	// Map from node property index to source editor node
+	TMap<int32, TWeakObjectPtr<const UEdGraphNode> > NodePropertyIndexToNodeMap;
+
 	// Map from animation node GUID to property index
 	TMap<FGuid, int32> NodeGuidToIndexMap;
 
@@ -348,6 +351,8 @@ public:
 	}
 
 	const int32* GetNodePropertyIndexFromGuid(FGuid Guid, EPropertySearchMode::Type SearchMode = EPropertySearchMode::OnlyThis);
+
+	const UEdGraphNode* GetVisualNodeFromNodePropertyIndex(int32 PropertyIndex) const;
 #endif
 
 	// Called after Link to patch up references to the nodes in the CDO
@@ -383,4 +388,4 @@ NodeType* GetNodeFromPropertyIndex(UObject* AnimInstanceObject, const IAnimClass
 	}
 
 	return NULL;
-}
+}
\ No newline at end of file
diff --git a/Engine/Source/Runtime/Engine/Classes/Animation/AnimClassInterface.h b/Engine/Source/Runtime/Engine/Classes/Animation/AnimClassInterface.h
--- a/Engine/Source/Runtime/Engine/Classes/Animation/AnimClassInterface.h
+++ b/Engine/Source/Runtime/Engine/Classes/Animation/AnimClassInterface.h
@@ -105,17 +105,23 @@ struct FGraphAssetPlayerInformation
 	TArray<int32> PlayerNodeIndices;
 };
 
-/** Blending options for anim graph layers */
+/** Blending options for animation graphs in Linked Animation Blueprints. */
 USTRUCT()
 struct FAnimGraphBlendOptions
 {
 	GENERATED_USTRUCT_BODY()
 
-	/** Time to blend this graph in. Only valid for linked anim graphs. Must use inertialization node to blend. */
+	/**
+	* Time to blend this graph in using Inertialization. Specify -1.0 to defer to the BlendOutTime of the previous graph.
+	* To blend this graph in you must place an Inertialization node after the Linked Anim Graph node or Linked Anim Layer node that uses this graph.
+	*/
 	UPROPERTY(EditAnywhere, Category = GraphBlending)
 	float BlendInTime;       
 
-	/** Time to blend this graph out. Only valid for linked anim graphs. Must use inertialization node to blend. */
+	/**
+	* Time to blend this graph out using Inertialization. Specify -1.0 to defer to the BlendInTime of the next graph.
+	* To blend this graph out you must place an Inertialization node after the Linked Anim Graph node or Linked Anim Layer node that uses this graph.
+	*/
 	UPROPERTY(EditAnywhere, Category = GraphBlending)
 	float BlendOutTime;
 
diff --git a/Engine/Source/Runtime/Engine/Classes/Animation/AnimNodeBase.h b/Engine/Source/Runtime/Engine/Classes/Animation/AnimNodeBase.h
--- a/Engine/Source/Runtime/Engine/Classes/Animation/AnimNodeBase.h
+++ b/Engine/Source/Runtime/Engine/Classes/Animation/AnimNodeBase.h
@@ -155,7 +155,7 @@ public:
 protected:
 
 	/** Interface for node contexts to register log messages with the proxy */
-	ENGINE_API void LogMessageInternal(FName InLogType, EMessageSeverity::Type InSeverity, FText InMessage);
+	ENGINE_API void LogMessageInternal(FName InLogType, EMessageSeverity::Type InSeverity, FText InMessage) const;
 };
 
 
@@ -313,7 +313,7 @@ public:
 	float GetDeltaTime() const { return DeltaTime; }
 
 	// Log update message
-	void LogMessage(EMessageSeverity::Type InSeverity, FText InMessage) { LogMessageInternal("Update", InSeverity, InMessage); }
+	void LogMessage(EMessageSeverity::Type InSeverity, FText InMessage) const { LogMessageInternal("Update", InSeverity, InMessage); }
 };
 
 
@@ -345,7 +345,7 @@ public:
 	ENGINE_API void Initialize(FAnimInstanceProxy* InAnimInstanceProxy);
 
 	// Log evaluation message
-	void LogMessage(EMessageSeverity::Type InSeverity, FText InMessage) { LogMessageInternal("Evaluate", InSeverity, InMessage); }
+	void LogMessage(EMessageSeverity::Type InSeverity, FText InMessage) const { LogMessageInternal("Evaluate", InSeverity, InMessage); }
 
 	void ResetToRefPose()
 	{
diff --git a/Engine/Source/Runtime/Engine/Classes/Animation/AnimNode_Inertialization.h b/Engine/Source/Runtime/Engine/Classes/Animation/AnimNode_Inertialization.h
--- a/Engine/Source/Runtime/Engine/Classes/Animation/AnimNode_Inertialization.h
+++ b/Engine/Source/Runtime/Engine/Classes/Animation/AnimNode_Inertialization.h
@@ -184,10 +184,13 @@ public: // FAnimNode_Inertialization
 	// Request to activate inertialization for a duration.
 	// If multiple requests are made on the same inertialization node, the minimum requested time will be used.
 	//
-	virtual void Request(float Duration);
+	virtual void RequestInertialization(float Duration);
 
 	virtual float GetRequestedDuration() const { return RequestedDuration; }
 
+	// Log an error when a node wants to inertialize but no inertialization ancestor node exists
+	//
+	static void LogRequestError(const FAnimationUpdateContext& Context, const FPoseLinkBase& RequesterPoseLink);
 
 public: // FAnimNode_Base
 
diff --git a/Engine/Source/Runtime/Engine/Classes/Animation/AnimNode_StateMachine.h b/Engine/Source/Runtime/Engine/Classes/Animation/AnimNode_StateMachine.h
--- a/Engine/Source/Runtime/Engine/Classes/Animation/AnimNode_StateMachine.h
+++ b/Engine/Source/Runtime/Engine/Classes/Animation/AnimNode_StateMachine.h
@@ -263,6 +263,8 @@ protected:
 
 	FAnimNode_AssetPlayerBase* GetRelevantAssetPlayerFromState(const FAnimationUpdateContext& Context, const FBakedAnimationState& StateInfo);
 
+	void LogInertializationRequestError(const FAnimationUpdateContext& Context, int32 PreviousState, int32 NextState);
+
 public:
 	friend struct FAnimInstanceProxy;
 };
diff --git a/Engine/Source/Runtime/Engine/Private/Animation/AnimBlueprintGeneratedClass.cpp b/Engine/Source/Runtime/Engine/Private/Animation/AnimBlueprintGeneratedClass.cpp
--- a/Engine/Source/Runtime/Engine/Private/Animation/AnimBlueprintGeneratedClass.cpp
+++ b/Engine/Source/Runtime/Engine/Private/Animation/AnimBlueprintGeneratedClass.cpp
@@ -18,6 +18,7 @@
 #include "Animation/AnimNode_LinkedAnimLayer.h"
 #include "Animation/AnimNode_AssetPlayerBase.h"
 #include "Animation/AnimNode_StateMachine.h"
+#include "EdGraph/EdGraphNode.h"
 
 /////////////////////////////////////////////////////
 // FStateMachineDebugData
@@ -571,4 +572,11 @@ const int32* UAnimBlueprintGeneratedClass::GetNodePropertyIndexFromGuid(FGuid Gu
 	return NULL;
 }
 
-#endif
+const UEdGraphNode* UAnimBlueprintGeneratedClass::GetVisualNodeFromNodePropertyIndex(int32 PropertyIndex) const
+{
+	const int32 ReversedPropertyIndex = AnimNodeProperties.Num() - PropertyIndex - 1;
+	const TWeakObjectPtr<const UEdGraphNode>* Node = AnimBlueprintDebugData.NodePropertyIndexToNodeMap.Find(ReversedPropertyIndex);
+	return Node ? Node->Get() : nullptr;
+}
+
+#endif // WITH_EDITORONLY_DATA
diff --git a/Engine/Source/Runtime/Engine/Private/Animation/AnimInstanceProxy.cpp b/Engine/Source/Runtime/Engine/Private/Animation/AnimInstanceProxy.cpp
--- a/Engine/Source/Runtime/Engine/Private/Animation/AnimInstanceProxy.cpp
+++ b/Engine/Source/Runtime/Engine/Private/Animation/AnimInstanceProxy.cpp
@@ -309,7 +309,7 @@ FGuid MakeGuidForMessage(const FText& Message)
 	return FGuid(Hash[0] ^ Hash[4], Hash[1], Hash[2], Hash[3]);
 }
 
-void FAnimInstanceProxy::LogMessage(FName InLogType, EMessageSeverity::Type InSeverity, const FText& InMessage)
+void FAnimInstanceProxy::LogMessage(FName InLogType, EMessageSeverity::Type InSeverity, const FText& InMessage) const
 {
 #if ENABLE_ANIM_LOGGING
 	FGuid CurrentMessageGuid = MakeGuidForMessage(InMessage);
diff --git a/Engine/Source/Runtime/Engine/Private/Animation/AnimNodeBase.cpp b/Engine/Source/Runtime/Engine/Private/Animation/AnimNodeBase.cpp
--- a/Engine/Source/Runtime/Engine/Private/Animation/AnimNodeBase.cpp
+++ b/Engine/Source/Runtime/Engine/Private/Animation/AnimNodeBase.cpp
@@ -82,7 +82,7 @@ UAnimBlueprint* FAnimationBaseContext::GetAnimBlueprint() const
 }
 #endif //WITH_EDITORONLY_DATA
 
-void FAnimationBaseContext::LogMessageInternal(FName InLogType, EMessageSeverity::Type InSeverity, FText InMessage)
+void FAnimationBaseContext::LogMessageInternal(FName InLogType, EMessageSeverity::Type InSeverity, FText InMessage) const
 {
 	AnimInstanceProxy->LogMessage(InLogType, InSeverity, InMessage);
 }
diff --git a/Engine/Source/Runtime/Engine/Private/Animation/AnimNode_Inertialization.cpp b/Engine/Source/Runtime/Engine/Private/Animation/AnimNode_Inertialization.cpp
--- a/Engine/Source/Runtime/Engine/Private/Animation/AnimNode_Inertialization.cpp
+++ b/Engine/Source/Runtime/Engine/Private/Animation/AnimNode_Inertialization.cpp
@@ -4,6 +4,8 @@
 #include "Animation/AnimInstanceProxy.h"
 #include "AnimationRuntime.h"
 
+#define LOCTEXT_NAMESPACE "AnimNode_Inertialization"
+
 
 TAutoConsoleVariable<int32> CVarAnimInertializationEnable(TEXT("a.AnimNode.Inertialization.Enable"), 1, TEXT("Enable / Disable Inertialization"));
 TAutoConsoleVariable<int32> CVarAnimInertializationIgnoreVelocity(TEXT("a.AnimNode.Inertialization.IgnoreVelocity"), 0, TEXT("Ignore velocity information during Inertialization (effectively reverting to a quintic diff blend)"));
@@ -24,7 +26,7 @@ FAnimNode_Inertialization::FAnimNode_Inertialization()
 }
 
 
-void FAnimNode_Inertialization::Request(float Duration)
+void FAnimNode_Inertialization::RequestInertialization(float Duration)
 {
 	if (RequestedDuration < 0.0f || Duration < RequestedDuration)
 	{
@@ -32,6 +34,18 @@ void FAnimNode_Inertialization::Request(float Duration)
 	}
 }
 
+/*static*/ void FAnimNode_Inertialization::LogRequestError(const FAnimationUpdateContext& Context, const FPoseLinkBase& RequesterPoseLink)
+{
+#if WITH_EDITORONLY_DATA	
+	UAnimBlueprint* AnimBlueprint = Context.AnimInstanceProxy->GetAnimBlueprint();
+	UAnimBlueprintGeneratedClass* AnimClass = AnimBlueprint ? AnimBlueprint->GetAnimBlueprintGeneratedClass() : nullptr;
+	const UObject* RequesterNode = AnimClass->GetVisualNodeFromNodePropertyIndex(RequesterPoseLink.SourceLinkID);
+
+	FText Message = FText::Format(LOCTEXT("InertializationRequestError", "No Inertialization node found for request from '{0}'. Add an Inertialization node after this request."),
+		FText::FromString(GetPathNameSafe(RequesterNode)));
+	Context.LogMessage(EMessageSeverity::Error, Message);
+#endif
+}
 
 void FAnimNode_Inertialization::Initialize_AnyThread(const FAnimationInitializeContext& Context)
 {
@@ -264,7 +278,7 @@ void FAnimNode_Inertialization::OnUpdatesSkipped(TArrayView<const FAnimationUpda
 	{
 		if (FAnimNode_Inertialization* OtherIneritalizationNode = Update->GetAncestor<FAnimNode_Inertialization>())
 		{
-			OtherIneritalizationNode->Request(RequestedDuration);
+			OtherIneritalizationNode->RequestInertialization(RequestedDuration);
 		}
 	}
 }
diff --git a/Engine/Source/Runtime/Engine/Private/Animation/AnimNode_LinkedAnimGraph.cpp b/Engine/Source/Runtime/Engine/Private/Animation/AnimNode_LinkedAnimGraph.cpp
--- a/Engine/Source/Runtime/Engine/Private/Animation/AnimNode_LinkedAnimGraph.cpp
+++ b/Engine/Source/Runtime/Engine/Private/Animation/AnimNode_LinkedAnimGraph.cpp
@@ -112,12 +112,19 @@ void FAnimNode_LinkedAnimGraph::Update_AnyThread(const FAnimationUpdateContext&
 	}
 
 	// Consume pending inertial blend request
-	if (PendingBlendDuration >= 0.0f)
+	if(PendingBlendDuration >= 0.0f)
 	{
-		FAnimNode_Inertialization* InertializationNode = InContext.GetAncestor<FAnimNode_Inertialization>();
-		if (InertializationNode)
+		if(InputPoses.Num() > 0)
 		{
-			InertializationNode->Request(PendingBlendDuration);
+			FAnimNode_Inertialization* InertializationNode = InContext.GetAncestor<FAnimNode_Inertialization>();
+			if(InertializationNode)
+			{
+				InertializationNode->RequestInertialization(PendingBlendDuration);
+			}
+			else if (PendingBlendDuration != 0.0f)
+			{
+				FAnimNode_Inertialization::LogRequestError(InContext, InputPoses[0]);
+			}
 		}
 
 		PendingBlendDuration = -1.0f;
diff --git a/Engine/Source/Runtime/Engine/Private/Animation/AnimNode_StateMachine.cpp b/Engine/Source/Runtime/Engine/Private/Animation/AnimNode_StateMachine.cpp
--- a/Engine/Source/Runtime/Engine/Private/Animation/AnimNode_StateMachine.cpp
+++ b/Engine/Source/Runtime/Engine/Private/Animation/AnimNode_StateMachine.cpp
@@ -10,6 +10,8 @@
 #include "Animation/AnimNode_LinkedAnimLayer.h"
 #include "Animation/AnimInstance.h"
 
+#define LOCTEXT_NAMESPACE "AnimNode_StateMachine"
+
 DECLARE_CYCLE_STAT(TEXT("StateMachine SetState"), Stat_StateMachineSetState, STATGROUP_Anim);
 
 //////////////////////////////////////////////////////////////////////////
@@ -334,6 +336,19 @@ const FAnimationTransitionBetweenStates& FAnimNode_StateMachine::GetTransitionIn
 	return PRIVATE_MachineDescription->Transitions[TransIndex];
 }
 
+void FAnimNode_StateMachine::LogInertializationRequestError(const FAnimationUpdateContext& Context, int32 PreviousState, int32 NextState)
+{
+#if WITH_EDITORONLY_DATA
+	const FBakedAnimationStateMachine* Machine = GetMachineDescription();
+	FText Message = FText::Format(LOCTEXT("InertialTransitionError", "No Inertialization node found for request from transition '{0}' to '{1}' in state machine '{2}' in anim blueprint '{3}'. Add an Inertialization node after this state machine."),
+		FText::FromName(GetStateInfo(PreviousState).StateName),
+		FText::FromName(GetStateInfo(NextState).StateName),
+		FText::FromName(Machine->MachineName),
+		FText::FromString(GetPathNameSafe(Context.AnimInstanceProxy->GetAnimBlueprint())));
+	Context.LogMessage(EMessageSeverity::Error, Message);
+#endif
+}
+
 // Temporarily turned off while we track down and fix https://jira.ol.epicgames.net/browse/OR-17066
 TAutoConsoleVariable<int32> CVarAnimStateMachineRelevancyReset(TEXT("a.AnimNode.StateMachine.EnableRelevancyReset"), 1, TEXT("Reset State Machine when it becomes relevant"));
 
@@ -449,7 +464,11 @@ void FAnimNode_StateMachine::Update_AnyThread(const FAnimationUpdateContext& Con
 					FAnimNode_Inertialization* InertializationNode = Context.GetAncestor<FAnimNode_Inertialization>();
 					if (InertializationNode)
 					{
-						InertializationNode->Request(ReferenceTransition.CrossfadeDuration);
+						InertializationNode->RequestInertialization(ReferenceTransition.CrossfadeDuration);
+					}
+					else
+					{
+						LogInertializationRequestError(Context, PreviousState, NextState);
 					}
 				}
 
diff --git a/Engine/Source/Runtime/Engine/Public/Animation/AnimInstanceProxy.h b/Engine/Source/Runtime/Engine/Public/Animation/AnimInstanceProxy.h
--- a/Engine/Source/Runtime/Engine/Public/Animation/AnimInstanceProxy.h
+++ b/Engine/Source/Runtime/Engine/Public/Animation/AnimInstanceProxy.h
@@ -765,7 +765,7 @@ protected:
 	void AddReferencedObjects(UAnimInstance* InAnimInstance, FReferenceCollector& Collector);
 
 	/** Allow nodes to register log messages to be processed on the game thread */
-	void LogMessage(FName InLogType, EMessageSeverity::Type InSeverity, const FText& InMessage);
+	void LogMessage(FName InLogType, EMessageSeverity::Type InSeverity, const FText& InMessage) const;
 
 	/** Get the current value of all animation curves **/
 	TMap<FName, float>& GetAnimationCurves(EAnimCurveType InCurveType) { return AnimationCurves[(uint8)InCurveType]; }
@@ -941,10 +941,10 @@ private:
 	 *  on a worked thread
 	 */
 	typedef TPair<EMessageSeverity::Type, FText> FLogMessageEntry;
-	TMap<FName, TArray<FLogMessageEntry>> LoggedMessagesMap;
+	mutable TMap<FName, TArray<FLogMessageEntry>> LoggedMessagesMap;
 
 	/** Cache of guids generated from previously sent messages so we can stop spam*/
-	TArray<FGuid> PreviouslyLoggedMessages;
+	mutable TArray<FGuid> PreviouslyLoggedMessages;
 #endif
 
 	/** Scope guard to prevent duplicate work on re-entracy */
