Fixed unity builds
#rnx

diff --git a/Engine/Source/Runtime/Engine/Private/Animation/AnimNode_Inertialization.cpp b/Engine/Source/Runtime/Engine/Private/Animation/AnimNode_Inertialization.cpp
--- a/Engine/Source/Runtime/Engine/Private/Animation/AnimNode_Inertialization.cpp
+++ b/Engine/Source/Runtime/Engine/Private/Animation/AnimNode_Inertialization.cpp
@@ -675,3 +675,5 @@ float FInertializationPoseDiff::CalcInertialFloat(float x0, float v0, float t, f
 
 	return x * sign;
 }
+
+#undef LOCTEXT_NAMESPACE
\ No newline at end of file
diff --git a/Engine/Source/Runtime/Engine/Private/Animation/AnimNode_StateMachine.cpp b/Engine/Source/Runtime/Engine/Private/Animation/AnimNode_StateMachine.cpp
--- a/Engine/Source/Runtime/Engine/Private/Animation/AnimNode_StateMachine.cpp
+++ b/Engine/Source/Runtime/Engine/Private/Animation/AnimNode_StateMachine.cpp
@@ -1108,3 +1108,5 @@ void FAnimNode_StateMachine::CacheMachineDescription(IAnimClassInterface* AnimBl
 {
 	PRIVATE_MachineDescription = AnimBlueprintClass->GetBakedStateMachines().IsValidIndex(StateMachineIndexInClass) ? &(AnimBlueprintClass->GetBakedStateMachines()[StateMachineIndexInClass]) : nullptr;
 }
+
+#undef LOCTEXT_NAMESPACE
\ No newline at end of file
