Fixed static analysis potential null
#rnx

diff --git a/Engine/Source/Runtime/Engine/Private/Animation/AnimNode_Inertialization.cpp b/Engine/Source/Runtime/Engine/Private/Animation/AnimNode_Inertialization.cpp
--- a/Engine/Source/Runtime/Engine/Private/Animation/AnimNode_Inertialization.cpp
+++ b/Engine/Source/Runtime/Engine/Private/Animation/AnimNode_Inertialization.cpp
@@ -39,7 +39,7 @@ void FAnimNode_Inertialization::RequestInertialization(float Duration)
 #if WITH_EDITORONLY_DATA	
 	UAnimBlueprint* AnimBlueprint = Context.AnimInstanceProxy->GetAnimBlueprint();
 	UAnimBlueprintGeneratedClass* AnimClass = AnimBlueprint ? AnimBlueprint->GetAnimBlueprintGeneratedClass() : nullptr;
-	const UObject* RequesterNode = AnimClass->GetVisualNodeFromNodePropertyIndex(RequesterPoseLink.SourceLinkID);
+	const UObject* RequesterNode = AnimClass ? AnimClass->GetVisualNodeFromNodePropertyIndex(RequesterPoseLink.SourceLinkID) : nullptr;
 
 	FText Message = FText::Format(LOCTEXT("InertializationRequestError", "No Inertialization node found for request from '{0}'. Add an Inertialization node after this request."),
 		FText::FromString(GetPathNameSafe(RequesterNode)));
