Fixed skyAtmosphere component group name.

diff --git a/Engine/Source/Runtime/Engine/Classes/Components/SkyAtmosphereComponent.h b/Engine/Source/Runtime/Engine/Classes/Components/SkyAtmosphereComponent.h
--- a/Engine/Source/Runtime/Engine/Classes/Components/SkyAtmosphereComponent.h
+++ b/Engine/Source/Runtime/Engine/Classes/Components/SkyAtmosphereComponent.h
@@ -64,15 +64,15 @@ class USkyAtmosphereComponent : public USceneComponent
 
 
 	/** Rayleigh scattering coefficient scale.*/
-	UPROPERTY(EditAnywhere, BlueprintReadOnly, interp, Category = "Atmosphere - Raleigh", meta = (UIMin = 0.0, UIMax = 2.0, ClampMin = 0.0, SliderExponent = 4.0))
+	UPROPERTY(EditAnywhere, BlueprintReadOnly, interp, Category = "Atmosphere - Rayleigh", meta = (UIMin = 0.0, UIMax = 2.0, ClampMin = 0.0, SliderExponent = 4.0))
 	float RayleighScatteringScale;
 
 	/** The Rayleigh scattering coefficients resulting from molecules in the air at an altitude of 0 kilometer. */
-	UPROPERTY(EditAnywhere, BlueprintReadOnly, interp, Category = "Atmosphere - Raleigh", meta=(HideAlphaChannel))
+	UPROPERTY(EditAnywhere, BlueprintReadOnly, interp, Category = "Atmosphere - Rayleigh", meta=(HideAlphaChannel))
 	FLinearColor RayleighScattering;
 
 	/** The altitude in kilometer at which Rayleigh scattering effect is reduced to 40%.*/
-	UPROPERTY(EditAnywhere, BlueprintReadOnly, interp, Category = "Atmosphere - Raleigh", meta = (UIMin = 0.01, UIMax = 20.0, ClampMin = 0.1, SliderExponent = 5.0))
+	UPROPERTY(EditAnywhere, BlueprintReadOnly, interp, Category = "Atmosphere - Rayleigh", meta = (UIMin = 0.01, UIMax = 20.0, ClampMin = 0.1, SliderExponent = 5.0))
 	float RayleighExponentialDistribution;
 
 
diff --git a/Engine/Source/Runtime/Engine/Private/Components/SkyAtmosphereComponent.cpp b/Engine/Source/Runtime/Engine/Private/Components/SkyAtmosphereComponent.cpp
--- a/Engine/Source/Runtime/Engine/Private/Components/SkyAtmosphereComponent.cpp
+++ b/Engine/Source/Runtime/Engine/Private/Components/SkyAtmosphereComponent.cpp
@@ -243,7 +243,7 @@ void USkyAtmosphereComponent::PostEditChangeProperty(FPropertyChangedEvent& Prop
 	const FName CategoryName = FObjectEditorUtils::GetCategoryFName(PropertyChangedEvent.Property);
 	if (CategoryName == FName(TEXT("Planet")) ||
 		CategoryName == FName(TEXT("Atmosphere")) ||
-		CategoryName == FName(TEXT("Atmosphere - Raleigh")) ||
+		CategoryName == FName(TEXT("Atmosphere - Rayleigh")) ||
 		CategoryName == FName(TEXT("Atmosphere - Mie")) ||
 		CategoryName == FName(TEXT("Atmosphere - Absorption")) ||
 		CategoryName == FName(TEXT("Art direction")))
