Make the serialization of the ExtData and NumMipsInTail members of FTexturePlatformData optional, based on whether the target platform actually requires them. This is to prevent changes to cooked asset data causing large patch sizes on platforms that otherwise would not use the new members.
Platforms that do need the additional members should override FPlatformProperties::RequiresOptTextureData() to return true.

diff --git a/Engine/Source/Developer/TargetPlatform/Public/Common/TargetPlatformBase.h b/Engine/Source/Developer/TargetPlatform/Public/Common/TargetPlatformBase.h
--- a/Engine/Source/Developer/TargetPlatform/Public/Common/TargetPlatformBase.h
+++ b/Engine/Source/Developer/TargetPlatform/Public/Common/TargetPlatformBase.h
@@ -322,6 +322,11 @@ public:
 		return TPlatformProperties::GetMemoryMappingAlignment();
 	}
 
+	virtual bool RequiresOptTextureData() const override
+	{
+		return TPlatformProperties::RequiresOptTextureData();
+	}
+
 
 #if WITH_ENGINE
 	virtual FName GetPhysicsFormat( class UBodySetup* Body ) const override
diff --git a/Engine/Source/Developer/TargetPlatform/Public/Interfaces/ITargetPlatform.h b/Engine/Source/Developer/TargetPlatform/Public/Interfaces/ITargetPlatform.h
--- a/Engine/Source/Developer/TargetPlatform/Public/Interfaces/ITargetPlatform.h
+++ b/Engine/Source/Developer/TargetPlatform/Public/Interfaces/ITargetPlatform.h
@@ -389,6 +389,13 @@ public:
 	*/
 	virtual float GetDownSampleMeshDistanceFieldDivider() const = 0;
 
+	/**
+	* Gets whether the platforms requires the OptData member of FTexturePlatformData.
+	* On platforms where this returns false, the OptData will be initialized to zero at
+	* runtime, and not serialized in the asset data stream.
+	*/
+	virtual bool RequiresOptTextureData() const = 0;
+
 #if WITH_ENGINE
 	/**
 	 * Gets the format to use for a particular body setup.
diff --git a/Engine/Source/Developer/TextureCompressor/Public/TextureCompressorModule.h b/Engine/Source/Developer/TextureCompressor/Public/TextureCompressorModule.h
--- a/Engine/Source/Developer/TextureCompressor/Public/TextureCompressorModule.h
+++ b/Engine/Source/Developer/TextureCompressor/Public/TextureCompressorModule.h
@@ -164,6 +164,8 @@ struct FTextureBuildSettings
 	uint32 bVirtualTextureEnableCompressZlib : 1;
 	/** Is crunch compression enabled */
 	uint32 bVirtualTextureEnableCompressCrunch : 1;
+	/** True if the target platform requires the FOptTexturePlatformData strucutre */
+	uint32 bOptData : 1;
 
 	/** Default settings. */
 	FTextureBuildSettings()
@@ -214,6 +216,7 @@ struct FTextureBuildSettings
 		, VirtualTextureBorderSize(0)
 		, bVirtualTextureEnableCompressZlib(false)
 		, bVirtualTextureEnableCompressCrunch(false)
+		, bOptData(false)
 	{
 	}
 
diff --git a/Engine/Source/Runtime/Core/Public/GenericPlatform/GenericPlatformProperties.h b/Engine/Source/Runtime/Core/Public/GenericPlatform/GenericPlatformProperties.h
--- a/Engine/Source/Runtime/Core/Public/GenericPlatform/GenericPlatformProperties.h
+++ b/Engine/Source/Runtime/Core/Public/GenericPlatform/GenericPlatformProperties.h
@@ -297,4 +297,9 @@ struct FGenericPlatformProperties
 	{
 		return nullptr;
 	}
+
+	static FORCEINLINE bool RequiresOptTextureData()
+	{
+		return false;
+	}
 };
diff --git a/Engine/Source/Runtime/Engine/Classes/Engine/Texture.h b/Engine/Source/Runtime/Engine/Classes/Engine/Texture.h
--- a/Engine/Source/Runtime/Engine/Classes/Engine/Texture.h
+++ b/Engine/Source/Runtime/Engine/Classes/Engine/Texture.h
@@ -415,6 +415,26 @@ private:
 #endif // WITH_EDITORONLY_DATA
 };
 
+struct FOptTexturePlatformData
+{
+	/** Arbitrary extra data that the runtime may need. */
+	uint32 ExtData;
+	/** Number of mips making up the mip tail, which must always be resident */
+	uint32 NumMipsInTail;
+
+	FOptTexturePlatformData()
+		: ExtData(0)
+		, NumMipsInTail(0)
+	{}
+
+	friend inline FArchive& operator << (FArchive& Ar, FOptTexturePlatformData& Data)
+	{
+		Ar << Data.ExtData;
+		Ar << Data.NumMipsInTail;
+		return Ar;
+	}
+};
+
 /**
  * Platform-specific data used by the texture resource at runtime.
  */
@@ -431,10 +451,8 @@ struct FTexturePlatformData
 	uint32 NumSlicesCubemapMask;
 	/** Format in which mip data is stored. */
 	EPixelFormat PixelFormat;
-	/** Optional extra data that the runtime may need. */
-	uint32 ExtData;
-	/** Number of mips making up the mip tail, which must always be resident */
-	uint32 NumMipsInTail;
+	/** Additional data required by some platforms.*/
+	FOptTexturePlatformData OptData;
 	/** Mip data or VT data. one or the other. */
 	TIndirectArray<struct FTexture2DMipMap> Mips;
 	struct FVirtualTextureBuiltData* VTData;
@@ -506,6 +524,16 @@ struct FTexturePlatformData
 		NumSlicesCubemapMask = (bCubemap ? CubeMapBitMask : 0) | (NumSlicesCubemapMask & CubeMapRemainingBitMask);
 	}
 
+	inline int32 GetNumMipsInTail() const
+	{
+		return OptData.NumMipsInTail;
+	}
+
+	inline int32 GetExtData() const
+	{
+		return OptData.ExtData;
+	}
+
 #if WITH_EDITOR
 	void Cache(
 		class UTexture& InTexture,
diff --git a/Engine/Source/Runtime/Engine/Classes/Engine/Texture2D.h b/Engine/Source/Runtime/Engine/Classes/Engine/Texture2D.h
--- a/Engine/Source/Runtime/Engine/Classes/Engine/Texture2D.h
+++ b/Engine/Source/Runtime/Engine/Classes/Engine/Texture2D.h
@@ -200,7 +200,7 @@ public:
 	{
 		if (PlatformData)
 		{
-			return FMath::Max(0, PlatformData->Mips.Num() - (int32)PlatformData->NumMipsInTail);
+			return FMath::Max(0, PlatformData->Mips.Num() - (int32)PlatformData->GetNumMipsInTail());
 		}
 		return 0;
 	}
@@ -213,7 +213,7 @@ public:
 	{
 		if (PlatformData)
 		{
-			return PlatformData->ExtData;
+			return PlatformData->GetExtData();
 		}
 		return 0;
 	}
@@ -239,7 +239,7 @@ public:
 	/** Returns the minimum number of mips that must be resident in memory (cannot be streamed). */
 	FORCEINLINE int32 GetMinTextureResidentMipCount() const
 	{
-		return FMath::Max(GMinTextureResidentMipCount, PlatformData ? (int32)PlatformData->NumMipsInTail : 0);
+		return FMath::Max(GMinTextureResidentMipCount, PlatformData ? (int32)PlatformData->GetNumMipsInTail() : 0);
 	}
 
 	/** Returns the minimum number of mips that must be resident in memory (cannot be streamed). */
diff --git a/Engine/Source/Runtime/Engine/Classes/Engine/Texture2DArray.cpp b/Engine/Source/Runtime/Engine/Classes/Engine/Texture2DArray.cpp
--- a/Engine/Source/Runtime/Engine/Classes/Engine/Texture2DArray.cpp
+++ b/Engine/Source/Runtime/Engine/Classes/Engine/Texture2DArray.cpp
@@ -55,7 +55,7 @@ uint32 UTexture2DArray::CalcTextureMemorySize(int32 MipCount) const
 	{
 		uint32 TextureAlign = 0;
 		// @todo: We should be using a "calc texture 2D" function here. Using 3D and passing numslices as the Z size is wrong, as 3D texture can have different tiling/padding vs 2D arrays.
-		uint64 TextureSize = RHICalcTexture3DPlatformSize(GetSizeX(), GetSizeY(), GetNumSlices(), GetPixelFormat(), MipCount, 0, FRHIResourceCreateInfo(PlatformData->ExtData), TextureAlign);
+		uint64 TextureSize = RHICalcTexture3DPlatformSize(GetSizeX(), GetSizeY(), GetNumSlices(), GetPixelFormat(), MipCount, 0, FRHIResourceCreateInfo(PlatformData->GetExtData()), TextureAlign);
 		Size = TextureSize;
 	}
 
diff --git a/Engine/Source/Runtime/Engine/Private/Texture2D.cpp b/Engine/Source/Runtime/Engine/Private/Texture2D.cpp
--- a/Engine/Source/Runtime/Engine/Private/Texture2D.cpp
+++ b/Engine/Source/Runtime/Engine/Private/Texture2D.cpp
@@ -946,7 +946,7 @@ int32 UTexture2D::CalcTextureMemorySize( int32 MipCount ) const
 		FIntPoint MipExtents = CalcMipMapExtent(SizeX, SizeY, Format, FirstMip);
 
 		uint32 TextureAlign = 0;
-		uint64 TextureSize = RHICalcTexture2DPlatformSize(MipExtents.X, MipExtents.Y, Format, MipCount, 1, 0, FRHIResourceCreateInfo(PlatformData->ExtData), TextureAlign);
+		uint64 TextureSize = RHICalcTexture2DPlatformSize(MipExtents.X, MipExtents.Y, Format, MipCount, 1, 0, FRHIResourceCreateInfo(PlatformData->GetExtData()), TextureAlign);
 		Size = (int32)TextureSize;
 	}
 	return Size;
@@ -1499,7 +1499,7 @@ FTexture2DResource::FTexture2DResource( UTexture2D* InOwner, int32 InitialMipCou
 	check(InitialMipCount>0);
 
 	// Keep track of first miplevel to use.
-	CurrentFirstMip = InOwner->GetNumMips() - FMath::Max(InitialMipCount, Owner->PlatformData ? (int32)Owner->PlatformData->NumMipsInTail : 0);
+	CurrentFirstMip = InOwner->GetNumMips() - FMath::Max(InitialMipCount, Owner->PlatformData ? (int32)Owner->PlatformData->GetNumMipsInTail() : 0);
 	InOwner->SetCachedNumResidentLODs(static_cast<uint8>(InitialMipCount));
 
 	check(CurrentFirstMip>=0);
@@ -1601,7 +1601,7 @@ void FTexture2DResource::InitRHI()
 				TexCreateFlags |= TexCreate_Virtual;
 
 				FRHIResourceCreateInfo CreateInfo(ResourceMem);
-				CreateInfo.ExtData = Owner->PlatformData->ExtData;
+				CreateInfo.ExtData = Owner->PlatformData->GetExtData();
 				Texture2DRHI = RHICreateTexture2D( OwnerMips[0].SizeX, OwnerMips[0].SizeY, EffectiveFormat, OwnerMips.Num(), 1, TexCreateFlags, CreateInfo);
 				RHIVirtualTextureSetFirstMipInMemory(Texture2DRHI, CurrentFirstMip);
 				RHIVirtualTextureSetFirstMipVisible(Texture2DRHI, CurrentFirstMip);
@@ -1637,7 +1637,7 @@ void FTexture2DResource::InitRHI()
 
 			// create texture with ResourceMem data when available
 			FRHIResourceCreateInfo CreateInfo(ResourceMem);
-			CreateInfo.ExtData = Owner->PlatformData->ExtData;
+			CreateInfo.ExtData = Owner->PlatformData->GetExtData();
 			Texture2DRHI	= RHICreateTexture2D( SizeX, SizeY, EffectiveFormat, RequestedMips, 1, TexCreateFlags, CreateInfo);
 			TextureRHI		= Texture2DRHI;
 			TextureRHI->SetName(Owner->GetFName());
@@ -1688,7 +1688,7 @@ void FTexture2DResource::InitRHI()
 		if (GIsEditor || (!bSkipRHITextureCreation)) //-V560
 		{
 			FRHIResourceCreateInfo CreateInfo;
-			CreateInfo.ExtData = Owner->PlatformData->ExtData;
+			CreateInfo.ExtData = Owner->PlatformData->GetExtData();
 			Texture2DRHI	= RHICreateTexture2D( SizeX, SizeY, EffectiveFormat, RequestedMips, 1, TexCreateFlags, CreateInfo );
 			TextureRHI		= Texture2DRHI;
 			TextureRHI->SetName(Owner->GetFName());
@@ -2259,7 +2259,7 @@ void FTexture2DResource::UpdateTexture(FTexture2DRHIRef& InTextureRHI, int32 InN
 
 		TextureRHI		= InTextureRHI;
 		Texture2DRHI	= InTextureRHI;
-		CurrentFirstMip = FMath::Min(InNewFirstMip, NumMips - (Owner->PlatformData ? (int32)Owner->PlatformData->NumMipsInTail : 0));
+		CurrentFirstMip = FMath::Min(InNewFirstMip, NumMips - (Owner->PlatformData ? (int32)Owner->PlatformData->GetNumMipsInTail() : 0));
 		Owner->SetCachedNumResidentLODs(static_cast<uint8>(NumMips - InNewFirstMip));
 		RHIUpdateTextureReference(Owner->TextureReference.TextureReferenceRHI, TextureRHI);
 	}
diff --git a/Engine/Source/Runtime/Engine/Private/TextureCube.cpp b/Engine/Source/Runtime/Engine/Private/TextureCube.cpp
--- a/Engine/Source/Runtime/Engine/Private/TextureCube.cpp
+++ b/Engine/Source/Runtime/Engine/Private/TextureCube.cpp
@@ -103,7 +103,7 @@ uint32 UTextureCube::CalcTextureMemorySize( int32 MipCount ) const
 		FIntPoint MipExtents = CalcMipMapExtent(SizeX, SizeY, Format, FirstMip);
 		
 		uint32 TextureAlign = 0;
-		uint64 TextureSize = RHICalcTextureCubePlatformSize(MipExtents.X, Format, MipCount, 0, FRHIResourceCreateInfo(PlatformData->ExtData), TextureAlign);
+		uint64 TextureSize = RHICalcTextureCubePlatformSize(MipExtents.X, Format, MipCount, 0, FRHIResourceCreateInfo(PlatformData->GetExtData()), TextureAlign);
 		Size = (uint32)TextureSize;
 	}
 	return Size;
@@ -148,7 +148,7 @@ public:
 		check(Owner->GetNumMips() > 0);
 
 		TIndirectArray<FTexture2DMipMap>& Mips = InOwner->PlatformData->Mips;
-		const int32 FirstMipTailIndex = Mips.Num() - InOwner->PlatformData->NumMipsInTail;
+		const int32 FirstMipTailIndex = Mips.Num() - InOwner->PlatformData->GetNumMipsInTail();
 		for (int32 MipIndex = 0; MipIndex < FirstMipTailIndex; MipIndex++)
 		{
 			FTexture2DMipMap& Mip = Mips[MipIndex];
@@ -207,7 +207,7 @@ public:
 		// Create the RHI texture.
 		uint32 TexCreateFlags = (Owner->SRGB ? TexCreate_SRGB : 0)  | TexCreate_OfflineProcessed;
 		FRHIResourceCreateInfo CreateInfo;
-		CreateInfo.ExtData = Owner->PlatformData ? Owner->PlatformData->ExtData : 0;
+		CreateInfo.ExtData = Owner->PlatformData ? Owner->PlatformData->GetExtData() : 0;
 		TextureCubeRHI = RHICreateTextureCube( Owner->GetSizeX(), Owner->GetPixelFormat(), Owner->GetNumMips(), TexCreateFlags, CreateInfo );
 		TextureRHI = TextureCubeRHI;
 		TextureRHI->SetName(Owner->GetFName());
diff --git a/Engine/Source/Runtime/Engine/Private/TextureDerivedData.cpp b/Engine/Source/Runtime/Engine/Private/TextureDerivedData.cpp
--- a/Engine/Source/Runtime/Engine/Private/TextureDerivedData.cpp
+++ b/Engine/Source/Runtime/Engine/Private/TextureDerivedData.cpp
@@ -47,7 +47,7 @@
 // In case of merge conflicts with DDC versions, you MUST generate a new GUID and set this new
 // guid as version
 
-#define TEXTURE_DERIVEDDATA_VER		TEXT("A52A75F077004A798078A42DCF5502E9")
+#define TEXTURE_DERIVEDDATA_VER		TEXT("0385EF5C8C62402ABCD94730886B0259")
 
 // This GUID is mixed into DDC version for virtual textures only, this allows updating DDC version for VT without invalidating DDC for all textures
 // This is useful during development, but once large numbers of VT are present in shipped content, it will have the same problem as TEXTURE_DERIVEDDATA_VER
@@ -222,6 +222,12 @@ static void SerializeForKey(FArchive& Ar, const FTextureBuildSettings& Settings)
 		OutKeySuffix.Append(FString::Printf(TEXT("VT%s_"), TEXTURE_VT_DERIVEDDATA_VER));
 	}
 
+	if (BuildSettings.bOptData)
+	{
+		// Add a flag to signal the platform needs the optional texture data structure
+		OutKeySuffix.Append(TEXT("O_"));
+	}
+
 	// Serialize the compressor settings into a temporary array. The archive
 	// is flagged as persistent so that machines of different endianness produce
 	// identical binary results.
@@ -333,6 +339,7 @@ static void GetTextureBuildSettings(
 	const UTextureLODSettings& TextureLODSettings,
 	bool bPlatformSupportsTextureStreaming,
 	bool bPlatformSupportsVirtualTextureStreaming,
+	bool bPlatformRequiresOptTextureData,
 	FTextureBuildSettings& OutBuildSettings
 	)
 {
@@ -354,6 +361,7 @@ static void GetTextureBuildSettings(
 	OutBuildSettings.bVolume = false;
 	OutBuildSettings.bCubemap = false;
 	OutBuildSettings.bTextureArray = false;
+	OutBuildSettings.bOptData = bPlatformRequiresOptTextureData;
 
 	if (Texture.MaxTextureSize > 0)
 	{
@@ -500,9 +508,10 @@ static void GetBuildSettingsForRunningPlatform(
 		const UTextureLODSettings* LODSettings = (UTextureLODSettings*)UDeviceProfileManager::Get().FindProfile(CurrentPlatform->PlatformName());
 		const bool bPlatformSupportsTextureStreaming = CurrentPlatform->SupportsFeature(ETargetPlatformFeatures::TextureStreaming);
 		const bool bPlatformSupportsVirtualTextureStreaming = CurrentPlatform->SupportsFeature(ETargetPlatformFeatures::VirtualTextureStreaming);
+		const bool bPlatformRequiresOptTextureData = CurrentPlatform->RequiresOptTextureData();
 
 		FTextureBuildSettings SourceBuildSettings;
-		GetTextureBuildSettings(Texture, *LODSettings, bPlatformSupportsTextureStreaming, bPlatformSupportsVirtualTextureStreaming, SourceBuildSettings);
+		GetTextureBuildSettings(Texture, *LODSettings, bPlatformSupportsTextureStreaming, bPlatformSupportsVirtualTextureStreaming, bPlatformRequiresOptTextureData, SourceBuildSettings);
 
 		TArray< TArray<FName> > PlatformFormats;
 		CurrentPlatform->GetTextureFormats(&Texture, PlatformFormats);
@@ -572,8 +581,8 @@ uint32 PutDerivedDataInCache(FTexturePlatformData* DerivedData, const FString& D
 	// Write out individual mips to the derived data cache.
 	const int32 MipCount = DerivedData->Mips.Num();
 	const bool bIsCubemap = DerivedData->IsCubemap();
-	const int32 FirstInlineMip = bIsCubemap ? 0 : FMath::Max(0, MipCount - FMath::Max((int32)NUM_INLINE_DERIVED_MIPS, (int32)DerivedData->NumMipsInTail));
-	const int32 WritableMipCount = MipCount - ((DerivedData->NumMipsInTail > 0) ? (DerivedData->NumMipsInTail - 1) : 0);
+	const int32 FirstInlineMip = bIsCubemap ? 0 : FMath::Max(0, MipCount - FMath::Max((int32)NUM_INLINE_DERIVED_MIPS, (int32)DerivedData->GetNumMipsInTail()));
+	const int32 WritableMipCount = MipCount - ((DerivedData->GetNumMipsInTail() > 0) ? (DerivedData->GetNumMipsInTail() - 1) : 0);
 	for (int32 MipIndex = 0; MipIndex < WritableMipCount; ++MipIndex)
 	{
 		FString MipDerivedDataKey;
@@ -1012,8 +1021,6 @@ FTexturePlatformData::FTexturePlatformData()
 	, SizeY(0)
 	, NumSlicesCubemapMask(0)
 	, PixelFormat(PF_Unknown)
-	, ExtData(0)
-	, NumMipsInTail(0)
 	, VTData(nullptr)
 #if WITH_EDITORONLY_DATA
 	, AsyncTask(NULL)
@@ -1050,7 +1057,7 @@ bool FTexturePlatformData::IsReadyForAsyncPostLoad() const
 bool FTexturePlatformData::TryLoadMips(int32 FirstMipToLoad, void** OutMipData)
 {
 	int32 NumMipsCached = 0;
-	const int32 LoadableMips = Mips.Num() - ((NumMipsInTail > 0) ? (NumMipsInTail - 1) : 0);
+	const int32 LoadableMips = Mips.Num() - ((GetNumMipsInTail() > 0) ? (GetNumMipsInTail() - 1) : 0);
 
 #if WITH_EDITOR
 	TArray<uint8> TempData;
@@ -1168,7 +1175,7 @@ int32 FTexturePlatformData::GetNumNonStreamingMips() const
 		int32 NumNonStreamingMips = 1;
 
 		// Take in to account the min resident limit.
-		NumNonStreamingMips = FMath::Max(NumNonStreamingMips, (int32)NumMipsInTail);
+		NumNonStreamingMips = FMath::Max(NumNonStreamingMips, (int32)GetNumMipsInTail());
 		NumNonStreamingMips = FMath::Max(NumNonStreamingMips, UTexture2D::GetStaticMinTextureResidentMipCount());
 		NumNonStreamingMips = FMath::Min(NumNonStreamingMips, MipCount);
 		int32 BlockSizeX = GPixelFormats[PixelFormat].BlockSizeX;
@@ -1242,9 +1249,9 @@ static void SerializePlatformData(
 	UTexture* Texture,
 	bool bCooked,
 	bool bStreamable
-	)
+)
 {
-	DECLARE_SCOPE_CYCLE_COUNTER( TEXT("SerializePlatformData"), STAT_Texture_SerializePlatformData, STATGROUP_LoadTime );
+	DECLARE_SCOPE_CYCLE_COUNTER(TEXT("SerializePlatformData"), STAT_Texture_SerializePlatformData, STATGROUP_LoadTime);
 
 	UEnum* PixelFormatEnum = UTexture::GetPixelFormatEnum();
 
@@ -1262,8 +1269,25 @@ static void SerializePlatformData(
 		FString PixelFormatString = PixelFormatEnum->GetNameByValue(PlatformData->PixelFormat).GetPlainNameString();
 		Ar << PixelFormatString;
 	}
- 	Ar << PlatformData->ExtData;
- 	Ar << PlatformData->NumMipsInTail;
+
+	// If we're cooking, check whether the target platform requires optional data, otherwise check the current running platform.
+	// In both cases, always serialize if we have editor only data.
+#if WITH_EDITORONLY_DATA
+	if (Ar.IsCooking())
+	{
+		if (Ar.CookingTarget()->RequiresOptTextureData() || Ar.CookingTarget()->HasEditorOnlyData())
+		{
+			Ar << PlatformData->OptData;
+		}
+	}
+	else 
+#endif
+	{
+		if (FPlatformProperties::RequiresOptTextureData() || FPlatformProperties::HasEditorOnlyData())
+		{
+			Ar << PlatformData->OptData;
+		}
+	}
 
 	int32 NumMips = PlatformData->Mips.Num();
 	int32 FirstMipToSerialize = 0;
@@ -1294,13 +1318,13 @@ static void SerializePlatformData(
 			const int32 NumCinematicMipLevels = Texture->NumCinematicMipLevels;
 			const TextureMipGenSettings MipGenSetting = Texture->MipGenSettings;
 			const int32 LastMip = FMath::Max(NumMips - 1, 0);
-			check(NumMips >= (int32)PlatformData->NumMipsInTail);
-			const int32 FirstMipTailMip = NumMips - (int32)PlatformData->NumMipsInTail;
+			check(NumMips >= (int32)PlatformData->GetNumMipsInTail());
+			const int32 FirstMipTailMip = NumMips - (int32)PlatformData->GetNumMipsInTail();
 
 			FirstMipToSerialize = Ar.CookingTarget()->GetTextureLODSettings().CalculateLODBias(Width, Height, Texture->MaxTextureSize, LODGroup, LODBias, 0, MipGenSetting, bIsVirtual);
 			if (!bIsVirtual)
 			{
-				FirstMipToSerialize = FMath::Clamp(FirstMipToSerialize, 0, PlatformData->NumMipsInTail > 0 ? FirstMipTailMip : LastMip);
+				FirstMipToSerialize = FMath::Clamp(FirstMipToSerialize, 0, PlatformData->GetNumMipsInTail() > 0 ? FirstMipTailMip : LastMip);
 				NumMips = FMath::Max(1, NumMips - FirstMipToSerialize);
 			}
 			else
@@ -1384,7 +1408,7 @@ static void SerializePlatformData(
 		}
 	}
 	Ar << NumMips;
-	check(NumMips >= (int32)PlatformData->NumMipsInTail);
+	check(NumMips >= (int32)PlatformData->GetNumMipsInTail());
 	if (Ar.IsLoading())
 	{
 		check(FirstMipToSerialize == 0);
@@ -1595,7 +1619,9 @@ void UTexture::BeginCacheForCookedPlatformData( const ITargetPlatform *TargetPla
 		FTextureBuildSettings BuildSettings;
 		const bool bPlatformSupportsTextureStreaming = TargetPlatform->SupportsFeature(ETargetPlatformFeatures::TextureStreaming);
 		const bool bPlatformSupportsVirtualTextureStreaming = TargetPlatform->SupportsFeature(ETargetPlatformFeatures::VirtualTextureStreaming);
-		GetTextureBuildSettings(*this, TargetPlatform->GetTextureLODSettings(), bPlatformSupportsTextureStreaming, bPlatformSupportsVirtualTextureStreaming, BuildSettings);
+		const bool bPlatformRequiresOptTextureData = TargetPlatform->RequiresOptTextureData();
+
+		GetTextureBuildSettings(*this, TargetPlatform->GetTextureLODSettings(), bPlatformSupportsTextureStreaming, bPlatformSupportsVirtualTextureStreaming, bPlatformRequiresOptTextureData, BuildSettings);
 		
 		TArray< TArray<FTextureBuildSettings> > BuildSettingsToCache;
 		GetBuildSettingsPerFormat(*this, BuildSettings, TargetPlatform, BuildSettingsToCache);
@@ -1678,7 +1704,9 @@ void UTexture::ClearCachedCookedPlatformData( const ITargetPlatform* TargetPlatf
 		FTextureBuildSettings BuildSettings;
 		const bool bPlatformSupportsTextureStreaming = TargetPlatform->SupportsFeature(ETargetPlatformFeatures::TextureStreaming);
 		const bool bPlatformSupportsVirtualTextureStreaming = TargetPlatform->SupportsFeature(ETargetPlatformFeatures::VirtualTextureStreaming);
-		GetTextureBuildSettings(*this, TargetPlatform->GetTextureLODSettings(), bPlatformSupportsTextureStreaming, bPlatformSupportsVirtualTextureStreaming, BuildSettings);
+		const bool bPlatformRequiresOptTextureData = TargetPlatform->RequiresOptTextureData();
+
+		GetTextureBuildSettings(*this, TargetPlatform->GetTextureLODSettings(), bPlatformSupportsTextureStreaming, bPlatformSupportsVirtualTextureStreaming, bPlatformRequiresOptTextureData, BuildSettings);
 
 		TArray< TArray<FTextureBuildSettings> > BuildSettingsToCache;
 		GetBuildSettingsPerFormat(*this, BuildSettings, TargetPlatform, BuildSettingsToCache);
@@ -1739,7 +1767,9 @@ bool UTexture::IsCachedCookedPlatformDataLoaded( const ITargetPlatform* TargetPl
 
 	const bool bPlatformSupportsTextureStreaming = TargetPlatform->SupportsFeature(ETargetPlatformFeatures::TextureStreaming);
 	const bool bPlatformSupportsVirtualTextureStreaming = TargetPlatform->SupportsFeature(ETargetPlatformFeatures::VirtualTextureStreaming);
-	GetTextureBuildSettings(*this, TargetPlatform->GetTextureLODSettings(), bPlatformSupportsTextureStreaming, bPlatformSupportsVirtualTextureStreaming, BuildSettings);
+	const bool bPlatformRequiresOptTextureData = TargetPlatform->RequiresOptTextureData();
+
+	GetTextureBuildSettings(*this, TargetPlatform->GetTextureLODSettings(), bPlatformSupportsTextureStreaming, bPlatformSupportsVirtualTextureStreaming, bPlatformRequiresOptTextureData, BuildSettings);
 
 	TArray< TArray<FTextureBuildSettings> > BuildSettingsToCache;
 	GetBuildSettingsPerFormat(*this, BuildSettings, TargetPlatform, BuildSettingsToCache);
@@ -1943,7 +1973,9 @@ void UTexture::SerializeCookedPlatformData(FArchive& Ar)
 			FTextureBuildSettings BuildSettings;
 			const bool bPlatformSupportsTextureStreaming = Ar.CookingTarget()->SupportsFeature(ETargetPlatformFeatures::TextureStreaming);
 			const bool bPlatformSupportsVirtualTextureStreaming = Ar.CookingTarget()->SupportsFeature(ETargetPlatformFeatures::VirtualTextureStreaming);
-			GetTextureBuildSettings(*this, Ar.CookingTarget()->GetTextureLODSettings(), bPlatformSupportsTextureStreaming, bPlatformSupportsVirtualTextureStreaming, BuildSettings);
+			const bool bPlatformRequiresOptTextureData = Ar.CookingTarget()->RequiresOptTextureData();
+
+			GetTextureBuildSettings(*this, Ar.CookingTarget()->GetTextureLODSettings(), bPlatformSupportsTextureStreaming, bPlatformSupportsVirtualTextureStreaming, bPlatformRequiresOptTextureData, BuildSettings);
 
 			TArray<FTexturePlatformData*> PlatformDataToSerialize;
 
diff --git a/Engine/Source/Runtime/Engine/Private/TextureDerivedDataTask.cpp b/Engine/Source/Runtime/Engine/Private/TextureDerivedDataTask.cpp
--- a/Engine/Source/Runtime/Engine/Private/TextureDerivedDataTask.cpp
+++ b/Engine/Source/Runtime/Engine/Private/TextureDerivedDataTask.cpp
@@ -261,8 +261,8 @@ void FTextureCacheDerivedDataWorker::BuildTexture()
 			((bool)Texture.CompositeTexture && CompositeTextureData.Blocks.Num() && CompositeTextureData.Blocks[0].MipsPerLayer.Num()) ? CompositeTextureData.Blocks[0].MipsPerLayer[0] : TArray<FImage>(),
 			BuildSettingsPerLayer[0],
 			CompressedMips,
-			DerivedData->NumMipsInTail,
-			DerivedData->ExtData))
+			DerivedData->OptData.NumMipsInTail,
+			DerivedData->OptData.ExtData))
 		{
 			check(CompressedMips.Num());
 
diff --git a/Engine/Source/Runtime/Engine/Private/VolumeTexture.cpp b/Engine/Source/Runtime/Engine/Private/VolumeTexture.cpp
--- a/Engine/Source/Runtime/Engine/Private/VolumeTexture.cpp
+++ b/Engine/Source/Runtime/Engine/Private/VolumeTexture.cpp
@@ -247,7 +247,7 @@ uint32 UVolumeTexture::CalcTextureMemorySize(int32 MipCount) const
 		CalcMipMapExtent3D(GetSizeX(), GetSizeY(), GetSizeZ(), Format, FMath::Max<int32>(0, GetNumMips() - MipCount), SizeX, SizeY, SizeZ);
 
 		uint32 TextureAlign = 0;
-		Size = (uint32)RHICalcTexture3DPlatformSize(SizeX, SizeY, SizeZ, Format, MipCount, Flags, FRHIResourceCreateInfo(PlatformData->ExtData), TextureAlign);
+		Size = (uint32)RHICalcTexture3DPlatformSize(SizeX, SizeY, SizeZ, Format, MipCount, Flags, FRHIResourceCreateInfo(PlatformData->GetExtData()), TextureAlign);
 	}
 	return Size;
 }
@@ -420,7 +420,7 @@ public:
 			const uint32 BaseMipSizeY = FMath::Max<uint32>(SizeY >> CurrentFirstMip, 1);
 			const uint32 BaseMipSizeZ = FMath::Max<uint32>(SizeZ >> CurrentFirstMip, 1);
 
-			CreateInfo.ExtData = Owner->PlatformData ? Owner->PlatformData->ExtData : 0;
+			CreateInfo.ExtData = Owner->PlatformData ? Owner->PlatformData->GetExtData() : 0;
 			Texture3DRHI = RHICreateTexture3D(BaseMipSizeX, BaseMipSizeY, BaseMipSizeZ, PixelFormat, NumMips - CurrentFirstMip, CreationFlags, CreateInfo);
 			TextureRHI = Texture3DRHI; 
 		}
