Warning: In asset 'None', there is an enum property of type 'EDepthOfFieldMethod' with an invalid value of '(' opening QAGame
#fix Add default value for EDepthOfFieldMethod

diff --git a/Engine/Source/Runtime/Engine/Private/Scene.cpp b/Engine/Source/Runtime/Engine/Private/Scene.cpp
--- a/Engine/Source/Runtime/Engine/Private/Scene.cpp
+++ b/Engine/Source/Runtime/Engine/Private/Scene.cpp
@@ -431,6 +431,7 @@ FPostProcessSettings::FPostProcessSettings()
 	BloomConvolutionCenterUV = FVector2D(0.5f, 0.5f);
 #if WITH_EDITORONLY_DATA
 	BloomConvolutionPreFilter_DEPRECATED = FVector(-1.f, -1.f, -1.f);
+	DepthOfFieldMethod_DEPRECATED = EDepthOfFieldMethod::DOFM_MAX;
 #endif
 	BloomConvolutionPreFilterMin = 7.f;
 	BloomConvolutionPreFilterMax = 15000.f;
@@ -757,7 +758,6 @@ FPostProcessSettings::FPostProcessSettings(const FPostProcessSettings& Settings)
 	, bMobileHQGaussian(Settings.bMobileHQGaussian)
 	, BloomMethod(Settings.BloomMethod)
 	, AutoExposureMethod(Settings.AutoExposureMethod)
-	//, DepthOfFieldMethod_DEPRECATED(Settings.DepthOfFieldMethod_DEPRECATED)
 	, WhiteTemp(Settings.WhiteTemp)
 	, WhiteTint(Settings.WhiteTint)
 	, ColorSaturation(Settings.ColorSaturation)
