Fixed crash on initialization of linked layer/graph nodes when the anim instance is re-used

diff --git a/Engine/Source/Runtime/Engine/Private/Animation/AnimNode_LinkedAnimGraph.cpp b/Engine/Source/Runtime/Engine/Private/Animation/AnimNode_LinkedAnimGraph.cpp
--- a/Engine/Source/Runtime/Engine/Private/Animation/AnimNode_LinkedAnimGraph.cpp
+++ b/Engine/Source/Runtime/Engine/Private/Animation/AnimNode_LinkedAnimGraph.cpp
@@ -42,6 +42,10 @@ void FAnimNode_LinkedAnimGraph::InitializeSubGraph_AnyThread(const FAnimationIni
 	if(InstanceToRun && LinkedRoot)
 	{
 		FAnimInstanceProxy& Proxy = InstanceToRun->GetProxyOnAnyThread<FAnimInstanceProxy>();
+
+		// Make sure we have valid objects in place for the sub-graph init
+		Proxy.InitializeObjects(InstanceToRun);
+
 		Proxy.InitializationCounter.SynchronizeWith(Context.AnimInstanceProxy->InitializationCounter);
 		Proxy.InitializeRootNode_WithRoot(LinkedRoot);
 	}
