Added WasRecentlyRendered to PrimitiveComponent.  This works just like the actor version, but sometimes its useful to test a single component.

diff --git a/Engine/Source/Runtime/Engine/Classes/Components/PrimitiveComponent.h b/Engine/Source/Runtime/Engine/Classes/Components/PrimitiveComponent.h
--- a/Engine/Source/Runtime/Engine/Classes/Components/PrimitiveComponent.h
+++ b/Engine/Source/Runtime/Engine/Classes/Components/PrimitiveComponent.h
@@ -680,6 +680,16 @@ private:
 
 public:
 
+	/**
+	 * Returns true if this component has been rendered "recently", with a tolerance in seconds to define what "recent" means.
+	 * e.g.: If a tolerance of 0.1 is used, this function will return true only if the actor was rendered in the last 0.1 seconds of game time.
+	 *
+	 * @param Tolerance  How many seconds ago the actor last render time can be and still count as having been "recently" rendered.
+	 * @return Whether this actor was recently rendered.
+	 */
+	UFUNCTION(Category = "Rendering", BlueprintCallable, meta=(DisplayName="WasComponentRecentlyRendered", Keywords="scene visible"))
+	bool WasRecentlyRendered(float Tolerance = 0.2) const;
+
 	void SetLastRenderTime(float InLastRenderTime);
 	float GetLastRenderTime() const { return LastRenderTime; }
 	float GetLastRenderTimeOnScreen() const { return LastRenderTimeOnScreen; }
diff --git a/Engine/Source/Runtime/Engine/Classes/GameFramework/Actor.h b/Engine/Source/Runtime/Engine/Classes/GameFramework/Actor.h
--- a/Engine/Source/Runtime/Engine/Classes/GameFramework/Actor.h
+++ b/Engine/Source/Runtime/Engine/Classes/GameFramework/Actor.h
@@ -2701,7 +2701,7 @@ public:
 	 * @param Tolerance  How many seconds ago the actor last render time can be and still count as having been "recently" rendered.
 	 * @return Whether this actor was recently rendered.
 	 */
-	UFUNCTION(Category="Rendering", BlueprintCallable, meta=(Keywords="scene visible"))
+	UFUNCTION(Category="Rendering", BlueprintCallable, meta=(DisplayName="WasActorRecentlyRendered", Keywords="scene visible"))
 	bool WasRecentlyRendered(float Tolerance = 0.2) const;
 
 	/** Returns the most recent time any of this actor's components were rendered */
diff --git a/Engine/Source/Runtime/Engine/Private/Components/PrimitiveComponent.cpp b/Engine/Source/Runtime/Engine/Private/Components/PrimitiveComponent.cpp
--- a/Engine/Source/Runtime/Engine/Private/Components/PrimitiveComponent.cpp
+++ b/Engine/Source/Runtime/Engine/Private/Components/PrimitiveComponent.cpp
@@ -3684,6 +3684,19 @@ void UPrimitiveComponent::SetCustomNavigableGeometry(const EHasCustomNavigableGe
 	bHasCustomNavigableGeometry = InType;
 }
  
+bool UPrimitiveComponent::WasRecentlyRendered(float Tolerance /*= 0.2*/) const
+{
+	if (const UWorld* const World = GetWorld())
+	{
+		// Adjust tolerance, so visibility is not affected by bad frame rate / hitches.
+		const float RenderTimeThreshold = FMath::Max(Tolerance, World->DeltaTimeSeconds + KINDA_SMALL_NUMBER);
+
+		// If the current cached value is less than the tolerance then we don't need to go look at the components
+		return World->TimeSince(GetLastRenderTime()) <= RenderTimeThreshold;
+	}
+	return false;
+}
+
 void UPrimitiveComponent::SetLastRenderTime(float InLastRenderTime)
 {
 	LastRenderTime = InLastRenderTime;
