#UE4 Fix for concurrent save, where materials are not saved on the game thread

diff --git a/Engine/Source/Runtime/Engine/Private/Materials/MaterialInstanceSupport.h b/Engine/Source/Runtime/Engine/Private/Materials/MaterialInstanceSupport.h
--- a/Engine/Source/Runtime/Engine/Private/Materials/MaterialInstanceSupport.h
+++ b/Engine/Source/Runtime/Engine/Private/Materials/MaterialInstanceSupport.h
@@ -36,7 +36,7 @@ public:
 	{
 		Material = const_cast<UMaterialInstance*>(InMaterial);
 
-		check(IsInRenderingThread() || IsInGameThread() || IsAsyncLoading());
+		check(IsInRenderingThread() || IsInGameThread() || IsAsyncLoading() || GIsSavingPackage);
 		if (Material->GetReentrantFlag() == true)
 		{
 			UE_LOG(LogMaterial, Warning, TEXT("InMaterial: %s GameThread: %d RenderThread: %d"), *InMaterial->GetFullName(), IsInGameThread(), IsInRenderingThread());
@@ -47,7 +47,7 @@ public:
 
 	~FMICReentranceGuard()
 	{
-		check(IsInRenderingThread() || IsInGameThread() || IsAsyncLoading());
+		check(IsInRenderingThread() || IsInGameThread() || IsAsyncLoading() || GIsSavingPackage);
 		Material->SetReentrantFlag(false);
 	}
 
