Fix thread contention and initialization issues for cloth rendering.

diff --git a/Engine/Shaders/Private/GpuSkinCacheComputeShader.usf b/Engine/Shaders/Private/GpuSkinCacheComputeShader.usf
--- a/Engine/Shaders/Private/GpuSkinCacheComputeShader.usf
+++ b/Engine/Shaders/Private/GpuSkinCacheComputeShader.usf
@@ -262,8 +262,8 @@ void SkinCacheUpdateBatchCS(uint3 GroupID : SV_GroupID,
 	}
 #else
 #if GPUSKIN_APEX_CLOTH
-    float3 ClothTangentX;
-    float3 ClothTangentZ;
+    float3 ClothTangentX = float3(0, 0, 0);
+    float3 ClothTangentZ = float3(0, 0, 0);
 	float TempClothBlendWeight = 0.0f;
 	if( Unpacked.SourceMeshVertIndices.w < 0xFFFF )
 	{ 
diff --git a/Engine/Source/Runtime/Engine/Private/SkeletalRenderGPUSkin.cpp b/Engine/Source/Runtime/Engine/Private/SkeletalRenderGPUSkin.cpp
--- a/Engine/Source/Runtime/Engine/Private/SkeletalRenderGPUSkin.cpp
+++ b/Engine/Source/Runtime/Engine/Private/SkeletalRenderGPUSkin.cpp
@@ -634,7 +634,7 @@ void FSkeletalMeshObjectGPUSkin::ProcessUpdatedDynamicData(FGPUSkinCache* GPUSki
 			if (bUseSkinCache)
 			{
 				// This takes the cloth positions from cloth space into world space
-				FMatrix ClothLocalToWorld = bClothFactory ? VertexFactoryData.ClothVertexFactories[SectionIdx]->GetClothShaderData().GetClothLocalToWorldForWriting(FrameNumberToPrepare) : FMatrix::Identity;
+				FMatrix ClothLocalToWorld = bClothFactory ? VertexFactoryData.ClothVertexFactories[SectionIdx]->GetClothShaderData().GetClothLocalToWorldForReading(FrameNumberToPrepare) : FMatrix::Identity;
 				// Matrices are transposed in ue4 meaning matrix multiples need to happen in reverse ((AB)x = b becomes xTBTAT = b).
 				FMatrix LocalToCloth = DynamicData->ClothObjectLocalToWorld * ClothLocalToWorld.Inverse();
 
