Notify all relevant net drivers when an actor tears off on a server
Notify net driver when an actor is torn off on the client

diff --git a/Engine/Source/Runtime/Engine/Classes/Engine/NetDriver.h b/Engine/Source/Runtime/Engine/Classes/Engine/NetDriver.h
--- a/Engine/Source/Runtime/Engine/Classes/Engine/NetDriver.h
+++ b/Engine/Source/Runtime/Engine/Classes/Engine/NetDriver.h
@@ -1498,10 +1498,12 @@ public:
 	ENGINE_API virtual bool ShouldClientDestroyActor(AActor* Actor) const;
 
 	/** Called when an actor channel is remotely opened for an actor. */
-	ENGINE_API virtual void NotifyActorChannelOpen(UActorChannel* Channel, AActor* Actor) {}
+	ENGINE_API virtual void NotifyActorChannelOpen(UActorChannel* Channel, AActor* Actor);
 	
 	/** Called when an actor channel is cleaned up foor an actor. */
-	ENGINE_API virtual void NotifyActorChannelCleanedUp(UActorChannel* Channel, EChannelCloseReason CloseReason) {}
+	ENGINE_API virtual void NotifyActorChannelCleanedUp(UActorChannel* Channel, EChannelCloseReason CloseReason);
+
+	ENGINE_API virtual void NotifyActorTornOff(AActor* Actor);
 
 	/**
 	 * Returns the current delinquency analytics and resets them.
diff --git a/Engine/Source/Runtime/Engine/Private/Actor.cpp b/Engine/Source/Runtime/Engine/Private/Actor.cpp
--- a/Engine/Source/Runtime/Engine/Private/Actor.cpp
+++ b/Engine/Source/Runtime/Engine/Private/Actor.cpp
@@ -2226,9 +2226,17 @@ void AActor::TearOff()
 	{
 		bTearOff = true;
 		// MARK_PROPERTY_DIRTY_FROM_NAME(AActor, bTearOff, this);
-		if (UNetDriver* NetDriver = GetNetDriver())
+		
+		FWorldContext* const Context = GEngine->GetWorldContextFromWorld(GetWorld());
+		if (Context != nullptr)
 		{
-			NetDriver->NotifyActorTearOff(this);
+			for (FNamedNetDriver& Driver : Context->ActiveNetDrivers)
+			{
+				if (Driver.NetDriver != nullptr && Driver.NetDriver->ShouldReplicateActor(this))
+				{
+					Driver.NetDriver->NotifyActorTearOff(this);
+				}
+			}
 		}
 	}
 }
diff --git a/Engine/Source/Runtime/Engine/Private/DataChannel.cpp b/Engine/Source/Runtime/Engine/Private/DataChannel.cpp
--- a/Engine/Source/Runtime/Engine/Private/DataChannel.cpp
+++ b/Engine/Source/Runtime/Engine/Private/DataChannel.cpp
@@ -2001,6 +2001,8 @@ bool UActorChannel::CleanUp(const bool bForDestroy, EChannelCloseReason CloseRea
 					{
 						Actor->TornOff();
 					}
+
+					Connection->Driver->NotifyActorTornOff(Actor);
 				}
 			}
 			else if (Dormant && (CloseReason == EChannelCloseReason::Dormancy) && !Actor->GetTearOff())	
diff --git a/Engine/Source/Runtime/Engine/Private/NetworkDriver.cpp b/Engine/Source/Runtime/Engine/Private/NetworkDriver.cpp
--- a/Engine/Source/Runtime/Engine/Private/NetworkDriver.cpp
+++ b/Engine/Source/Runtime/Engine/Private/NetworkDriver.cpp
@@ -5593,6 +5593,21 @@ bool UNetDriver::ShouldClientDestroyActor(AActor* Actor) const
 	return (Actor && !Actor->IsA(ALevelScriptActor::StaticClass()));
 }
 
+void UNetDriver::NotifyActorChannelOpen(UActorChannel* Channel, AActor* Actor)
+{
+
+}
+
+void UNetDriver::NotifyActorChannelCleanedUp(UActorChannel* Channel, EChannelCloseReason CloseReason)
+{
+
+}
+
+void UNetDriver::NotifyActorTornOff(AActor* Actor)
+{
+
+}
+
 void UNetDriver::ConsumeAsyncLoadDelinquencyAnalytics(FNetAsyncLoadDelinquencyAnalytics& Out)
 {
 	if (FNetGUIDCache* LocalGuidCache = GuidCache.Get())
