Changing behavior around XGE failures on shader compilation. Try for up to five seconds to read the output file just in case there is network latency causing the task to complete before the output file completes writing, then warn that the shader compile was rescheduled locally rather than being an error since we fall back on compiling locally anyway.

#rnx

diff --git a/Engine/Source/Runtime/Engine/Private/ShaderCompiler/ShaderCompilerXGE.cpp b/Engine/Source/Runtime/Engine/Private/ShaderCompiler/ShaderCompilerXGE.cpp
--- a/Engine/Source/Runtime/Engine/Private/ShaderCompiler/ShaderCompilerXGE.cpp
+++ b/Engine/Source/Runtime/Engine/Private/ShaderCompiler/ShaderCompilerXGE.cpp
@@ -876,8 +876,20 @@ int32 FShaderCompileXGEThreadRunnable_InterceptionInterface::CompilingLoop()
 
 		if (Result.bCompleted)
 		{
-			// Open the output file, and serialize in the completed jobs.
-			FArchive* OutputFileAr = IFileManager::Get().CreateFileReader(*Task->OutputFilePath, FILEREAD_Silent);
+			// Check the output file exists. If it does, attempt to open it and serialize in the completed jobs.
+			FArchive* OutputFileAr = nullptr;
+			if (IFileManager::Get().FileExists(*Task->OutputFilePath))
+			{
+				float TimeWaited = 0.0f;
+				FArchive* OutputFileAr = IFileManager::Get().CreateFileReader(*Task->OutputFilePath, FILEREAD_Silent);
+				while (OutputFileAr == nullptr && TimeWaited < 5.0f)
+				{
+					UE_LOG(LogShaderCompilers, Warning, TEXT("Expected XGE output file '%s' exists but can't be opened for read, waiting 1 second to try again.."), *Task->OutputFilePath);
+					FPlatformProcess::Sleep(1.0f);
+					OutputFileAr = IFileManager::Get().CreateFileReader(*Task->OutputFilePath, FILEREAD_Silent);
+				}
+			}
+
 			if (OutputFileAr)
 			{
 				FShaderCompileUtilities::DoReadTaskResults(Task->ShaderJobs, *OutputFileAr);
@@ -888,7 +900,7 @@ int32 FShaderCompileXGEThreadRunnable_InterceptionInterface::CompilingLoop()
 				// Reading result from XGE job failed, so recompile shaders in current job batch locally
 				bOutputFileReadFailed = true;
 
-				UE_LOG(LogShaderCompilers, Error, TEXT("Reschedule shader compilation after XGE job failed: %s"), *Task->OutputFilePath);
+				UE_LOG(LogShaderCompilers, Warning, TEXT("Rescheduling shader compilation to run locally after XGE job failed: %s"), *Task->OutputFilePath);
 
 				for (FShaderCommonCompileJob* Job : Task->ShaderJobs)
 				{
