Fixing double declaration of OutputFileAr in ShaderCompilerXGE.cpp.

#rnx

diff --git a/Engine/Source/Runtime/Engine/Private/ShaderCompiler/ShaderCompilerXGE.cpp b/Engine/Source/Runtime/Engine/Private/ShaderCompiler/ShaderCompilerXGE.cpp
--- a/Engine/Source/Runtime/Engine/Private/ShaderCompiler/ShaderCompilerXGE.cpp
+++ b/Engine/Source/Runtime/Engine/Private/ShaderCompiler/ShaderCompilerXGE.cpp
@@ -881,7 +881,7 @@ int32 FShaderCompileXGEThreadRunnable_InterceptionInterface::CompilingLoop()
 			if (IFileManager::Get().FileExists(*Task->OutputFilePath))
 			{
 				float TimeWaited = 0.0f;
-				FArchive* OutputFileAr = IFileManager::Get().CreateFileReader(*Task->OutputFilePath, FILEREAD_Silent);
+				OutputFileAr = IFileManager::Get().CreateFileReader(*Task->OutputFilePath, FILEREAD_Silent);
 				while (OutputFileAr == nullptr && TimeWaited < 5.0f)
 				{
 					UE_LOG(LogShaderCompilers, Warning, TEXT("Expected XGE output file '%s' exists but can't be opened for read, waiting 1 second to try again.."), *Task->OutputFilePath);
