#ue4 - Fix typo in comment

#rnx

diff --git a/Engine/Source/Runtime/Engine/Private/Controller.cpp b/Engine/Source/Runtime/Engine/Private/Controller.cpp
--- a/Engine/Source/Runtime/Engine/Private/Controller.cpp
+++ b/Engine/Source/Runtime/Engine/Private/Controller.cpp
@@ -118,7 +118,7 @@ void AController::SetControlRotation(const FRotator& NewRotation)
 {
 	if (!IsValidControlRotation(NewRotation))
 	{
-		logOrEnsureNanError(TEXT("AController::SetControlRotation atttempted to apply NaN-containing or NaN-causing rotation! (%s)"), *NewRotation.ToString());
+		logOrEnsureNanError(TEXT("AController::SetControlRotation attempted to apply NaN-containing or NaN-causing rotation! (%s)"), *NewRotation.ToString());
 		return;
 	}
 
