Don't set a level to unload unless that level is loaded

diff --git a/Engine/Source/Runtime/Engine/Private/LevelTick.cpp b/Engine/Source/Runtime/Engine/Private/LevelTick.cpp
--- a/Engine/Source/Runtime/Engine/Private/LevelTick.cpp
+++ b/Engine/Source/Runtime/Engine/Private/LevelTick.cpp
@@ -729,17 +729,17 @@ void UWorld::ProcessLevelStreamingVolumes(FVector* OverrideViewLocation)
 					LevelStreamingObject->SetShouldBeVisible(NewStreamingSettings->ShouldBeVisible(bOriginalShouldBeVisible));
 					LevelStreamingObject->bShouldBlockOnLoad	= NewStreamingSettings->ShouldBlockOnLoad();
 				}
-				// Prevent unload request flood.  The additional check ensures that unload requests can still be issued in the first UnloadCooldownTime seconds of play.
-				else 
-				if( TimeSeconds - LevelStreamingObject->LastVolumeUnloadRequestTime > LevelStreamingObject->MinTimeBetweenVolumeUnloadRequests 
-				||  LevelStreamingObject->LastVolumeUnloadRequestTime < 0.1f )
+				else if (LevelStreamingObject->ShouldBeLoaded())
 				{
-					//UE_LOG(LogLevel, Warning, TEXT("Unloading") );
-					if( GetPlayerControllerIterator() )
+					// Prevent unload request flood.  The additional check ensures that unload requests can still be issued in the first UnloadCooldownTime seconds of play.
+					if (TimeSeconds - LevelStreamingObject->LastVolumeUnloadRequestTime > LevelStreamingObject->MinTimeBetweenVolumeUnloadRequests ||  LevelStreamingObject->LastVolumeUnloadRequestTime < 0.1f)
 					{
-						LevelStreamingObject->LastVolumeUnloadRequestTime	= TimeSeconds;
-						LevelStreamingObject->SetShouldBeLoaded(false);
-						LevelStreamingObject->SetShouldBeVisible(false);
+						if (GetPlayerControllerIterator())
+						{
+							LevelStreamingObject->LastVolumeUnloadRequestTime = TimeSeconds;
+							LevelStreamingObject->SetShouldBeLoaded(false);
+							LevelStreamingObject->SetShouldBeVisible(false);
+						}
 					}
 				}
 
