Fix crash if null passed to SaveGameToMemory.
Properly handle SaveGameToMemory returning false in AsyncSaveGameToSlot

diff --git a/Engine/Source/Runtime/Engine/Private/GameplayStatics.cpp b/Engine/Source/Runtime/Engine/Private/GameplayStatics.cpp
--- a/Engine/Source/Runtime/Engine/Private/GameplayStatics.cpp
+++ b/Engine/Source/Runtime/Engine/Private/GameplayStatics.cpp
@@ -1938,16 +1938,21 @@ USaveGame* UGameplayStatics::CreateSaveGameObject(TSubclassOf<USaveGame> SaveGam
 
 bool UGameplayStatics::SaveGameToMemory(USaveGame* SaveGameObject, TArray<uint8>& OutSaveData )
 {
-	FMemoryWriter MemoryWriter(OutSaveData, true);
+	if (SaveGameObject)
+	{
+		FMemoryWriter MemoryWriter(OutSaveData, true);
+
+		FSaveGameHeader SaveHeader(SaveGameObject->GetClass());
+		SaveHeader.Write(MemoryWriter);
 
-	FSaveGameHeader SaveHeader(SaveGameObject->GetClass());
-	SaveHeader.Write(MemoryWriter);
+		// Then save the object state, replacing object refs and names with strings
+		FObjectAndNameAsStringProxyArchive Ar(MemoryWriter, false);
+		SaveGameObject->Serialize(Ar);
 
-	// Then save the object state, replacing object refs and names with strings
-	FObjectAndNameAsStringProxyArchive Ar(MemoryWriter, false);
-	SaveGameObject->Serialize(Ar);
+		return true; // Not sure if there's a failure case here.
+	}
 
-	return true; // Not sure if there's a failure case here.
+	return false;
 }
 
 bool UGameplayStatics::SaveDataToSlot(const TArray<uint8>& InSaveData, const FString& SlotName, const int32 UserIndex)
@@ -1964,23 +1969,28 @@ bool UGameplayStatics::SaveDataToSlot(const TArray<uint8>& InSaveData, const FSt
 }
 
 void UGameplayStatics::AsyncSaveGameToSlot(USaveGame* SaveGameObject, const FString& SlotName, const int32 UserIndex, FAsyncSaveGameToSlotDelegate SavedDelegate)
+{
+	TArray<uint8> ObjectBytes;
+	if (SaveGameToMemory(SaveGameObject, ObjectBytes))
 	{
-		TArray<uint8> ObjectBytes;
-	SaveGameToMemory(SaveGameObject, ObjectBytes);
-
-	AsyncTask(ENamedThreads::AnyHiPriThreadNormalTask, [SlotName, UserIndex, SavedDelegate, ObjectBytes]()
-	{
-		bool bSuccess = SaveDataToSlot(ObjectBytes, SlotName, UserIndex);
-
-		// Now schedule the callback on the game thread, but only if it was bound to anything
-		if (SavedDelegate.IsBound())
+		AsyncTask(ENamedThreads::AnyHiPriThreadNormalTask, [SlotName, UserIndex, SavedDelegate, ObjectBytes]()
 		{
-			AsyncTask(ENamedThreads::GameThread, [SlotName, UserIndex, SavedDelegate, bSuccess]()
+			bool bSuccess = SaveDataToSlot(ObjectBytes, SlotName, UserIndex);
+
+			// Now schedule the callback on the game thread, but only if it was bound to anything
+			if (SavedDelegate.IsBound())
 			{
-				SavedDelegate.ExecuteIfBound(SlotName, UserIndex, bSuccess);
-			});
-		}
-	});
+				AsyncTask(ENamedThreads::GameThread, [SlotName, UserIndex, SavedDelegate, bSuccess]()
+				{
+					SavedDelegate.ExecuteIfBound(SlotName, UserIndex, bSuccess);
+				});
+			}
+		});
+	}
+	else if (SavedDelegate.IsBound())
+	{
+		SavedDelegate.ExecuteIfBound(SlotName, UserIndex, false);
+	}
 }
 
 bool UGameplayStatics::SaveGameToSlot(USaveGame* SaveGameObject, const FString& SlotName, const int32 UserIndex)
