1) Forces a rebuild of streaming chunks when USoundWave::InvalidateCompressedData() is called.
2) add new editor only generation counter for chunks to ensure we don't return a stale chunk that's in the cache, but can keep stale chunks around if they are still being actively used.
3) Add editor only IsStale() function to FAudioChunkHandle so anyone using a chunk for their own purposes can tell if the usoundwave has been changed since they got the chunk.

#okforgitub public

diff --git a/Engine/Source/Runtime/Engine/Classes/Sound/SoundWave.h b/Engine/Source/Runtime/Engine/Classes/Sound/SoundWave.h
--- a/Engine/Source/Runtime/Engine/Classes/Sound/SoundWave.h
+++ b/Engine/Source/Runtime/Engine/Classes/Sound/SoundWave.h
@@ -358,6 +358,11 @@ class ENGINE_API USoundWave : public USoundBase
 	/** Whether this SoundWave was decompressed from OGG. */
 	uint8 bDecompressedFromOgg : 1;
 
+#if WITH_EDITOR
+	/** The current revision of our compressed audio data. Used to tell when a chunk in the cache is stale. */
+	FThreadSafeCounter CurrentChunkRevision;
+#endif
+
 private:
 
 	// This is set to false on initialization, then set to true on non-editor platforms when we cache appropriate sample rate.
diff --git a/Engine/Source/Runtime/Engine/Private/AudioStreamingCache.cpp b/Engine/Source/Runtime/Engine/Private/AudioStreamingCache.cpp
--- a/Engine/Source/Runtime/Engine/Private/AudioStreamingCache.cpp
+++ b/Engine/Source/Runtime/Engine/Private/AudioStreamingCache.cpp
@@ -207,7 +207,15 @@ FAudioChunkHandle FCachedAudioStreamingManager::GetLoadedChunk(const USoundWave*
 
 		// TODO:  See if we can avoid non-const calls on the USoundWave here.
 		USoundWave* MutableWave = const_cast<USoundWave*>(SoundWave);
-		const FAudioChunkCache::FChunkKey ChunkKey = { MutableWave, SoundWave->GetFName(), ChunkIndex };
+		const FAudioChunkCache::FChunkKey ChunkKey =
+		{
+			  MutableWave
+			, SoundWave->GetFName()
+			, ChunkIndex
+#if WITH_EDITOR
+			, SoundWave->CurrentChunkRevision.GetValue()
+#endif
+		};
 
 		if (!FAudioChunkCache::IsKeyValid(ChunkKey))
 		{
@@ -223,7 +231,16 @@ FAudioChunkHandle FCachedAudioStreamingManager::GetLoadedChunk(const USoundWave*
 
 		if (NextChunk != INDEX_NONE)
 		{
-			const FAudioChunkCache::FChunkKey NextChunkKey = { MutableWave, SoundWave->GetFName(), ((uint32)NextChunk) };
+			const FAudioChunkCache::FChunkKey NextChunkKey = 
+			{ 
+				  MutableWave 
+				, SoundWave->GetFName() 
+				, ((uint32)NextChunk) 
+#if WITH_EDITOR
+				, SoundWave->CurrentChunkRevision.GetValue()
+#endif
+			};
+
 			bool bIsValidChunk = Cache->AddOrTouchChunk(NextChunkKey, [](EAudioChunkLoadResult) {}, ENamedThreads::AnyThread);
 			if (!bIsValidChunk)
 			{
@@ -307,9 +324,12 @@ void FCachedAudioStreamingManager::AddReferenceToChunk(const FAudioChunkHandle&
 
 	FAudioChunkCache::FChunkKey ChunkKey =
 	{
-		const_cast<USoundWave*>(InHandle.CorrespondingWave),
-		InHandle.CorrespondingWaveName,
-		((uint32) InHandle.ChunkIndex)
+		  const_cast<USoundWave*>(InHandle.CorrespondingWave)
+		, InHandle.CorrespondingWaveName
+		, ((uint32) InHandle.ChunkIndex)
+#if WITH_EDITOR
+		, InHandle.ChunkGeneration
+#endif
 	};
 
 	Cache->AddNewReferenceToChunk(ChunkKey);
@@ -322,9 +342,12 @@ void FCachedAudioStreamingManager::RemoveReferenceToChunk(const FAudioChunkHandl
 
 	FAudioChunkCache::FChunkKey ChunkKey =
 	{
-		const_cast<USoundWave*>(InHandle.CorrespondingWave),
-		InHandle.CorrespondingWaveName,
-		((uint32) InHandle.ChunkIndex)
+		  const_cast<USoundWave*>(InHandle.CorrespondingWave)
+		, InHandle.CorrespondingWaveName
+		, ((uint32) InHandle.ChunkIndex)
+#if WITH_EDITOR
+		, InHandle.ChunkGeneration
+#endif
 	};
 
 	Cache->RemoveReferenceToChunk(ChunkKey);
@@ -1081,6 +1104,9 @@ TPair<int, int> FAudioChunkCache::DebugDisplay(UWorld* World, FViewport* Viewpor
 	// A chunk that's been trimmed by TrimMemory.
 	const FLinearColor TrimmedChunkColor(204 / ColorMax, 46 / ColorMax, 43 / ColorMax);
 
+	// In editor builds, this is a chunk that was built in a previous version of the cook quality settings.
+	const FLinearColor StaleChunkColor(143 / ColorMax, 73 / ColorMax, 70 / ColorMax);
+	
 	// A chunk that currently has an async load in flight.
 	const FLinearColor CurrentlyLoadingChunkColor = FLinearColor::Yellow;
 
@@ -1140,6 +1166,7 @@ TPair<int, int> FAudioChunkCache::DebugDisplay(UWorld* World, FViewport* Viewpor
 		double TimeToLoad = -1.0;
 		float AveragePlaceInCache = -1.0f;
 		bool bWasCacheMiss = false;
+		bool bIsStaleChunk = false;
 
 #if DEBUG_STREAM_CACHE
 		NumTotalChunks = CurrentElement->DebugInfo.NumTotalChunks;
@@ -1149,9 +1176,14 @@ TPair<int, int> FAudioChunkCache::DebugDisplay(UWorld* World, FViewport* Viewpor
 		bWasCacheMiss = CurrentElement->DebugInfo.bWasCacheMiss;
 #endif
 
+#if WITH_EDITOR
+		// TODO: Worry about whether the sound wave is alive here. In most editor cases this is ok because the soundwave will always be loaded, but this may not be the case in the future.
+		bIsStaleChunk = (CurrentElement->Key.SoundWave->CurrentChunkRevision.GetValue() != CurrentElement->Key.ChunkRevision);
+#endif
+
 		const bool bWasTrimmed = CurrentElement->ChunkData.Num() == 0;
 
-		FString ElementInfo = *FString::Printf(TEXT("%4i. Size: %6.2f KB   Chunk: %d of %d   Request Count: %d    Average Index: %6.2f  Number of Handles Retaining Chunk: %d     Chunk Load Time: %6.4fms      Name: %s %s"),
+		FString ElementInfo = *FString::Printf(TEXT("%4i. Size: %6.2f KB   Chunk: %d of %d   Request Count: %d    Average Index: %6.2f  Number of Handles Retaining Chunk: %d     Chunk Load Time: %6.4fms      Name: %s Notes: %s %s"),
 			Index,
 			CurrentElement->ChunkData.Num() / 1024.0f,
 			CurrentElement->Key.ChunkIndex,
@@ -1161,13 +1193,22 @@ TPair<int, int> FAudioChunkCache::DebugDisplay(UWorld* World, FViewport* Viewpor
 			CurrentElement->NumConsumers.GetValue(),
 			TimeToLoad,
 			bWasTrimmed ? TEXT("TRIMMED CHUNK") : *CurrentElement->Key.SoundWaveName.ToString(),
-			bWasCacheMiss ? TEXT("(Cache Miss!)") : TEXT("")
+			bWasCacheMiss ? TEXT("(Cache Miss!)") : TEXT(""),
+			bIsStaleChunk ? TEXT("(Stale Chunk)") : TEXT("")
 			);
 
 		// Since there's a lot of info here,
 		// Subtly fading the chunk info to gray seems to help as a visual indicator of how far down on the list things are.
 		ColorLerpAmount = FMath::Min(ColorLerpAmount + ColorLerpStep, 1.0f);
-		FLinearColor TextColor = FLinearColor::LerpUsingHSV(LoadedChunkColor, FLinearColor::Gray, ColorLerpAmount);
+		FLinearColor TextColor;
+		if (bIsStaleChunk)
+		{
+			TextColor = FLinearColor::LerpUsingHSV(StaleChunkColor, FLinearColor::Gray, ColorLerpAmount);
+		}
+		else
+		{
+			TextColor = FLinearColor::LerpUsingHSV(LoadedChunkColor, FLinearColor::Gray, ColorLerpAmount);
+		}
 
 		// If there's a load in flight, paint this element yellow.
 		if (CurrentElement->IsLoadInProgress())
diff --git a/Engine/Source/Runtime/Engine/Private/ContentStreaming.cpp b/Engine/Source/Runtime/Engine/Private/ContentStreaming.cpp
--- a/Engine/Source/Runtime/Engine/Private/ContentStreaming.cpp
+++ b/Engine/Source/Runtime/Engine/Private/ContentStreaming.cpp
@@ -1371,6 +1371,9 @@ FAudioChunkHandle::FAudioChunkHandle()
 	, CorrespondingWave(nullptr)
 	, CorrespondingWaveName()
 	, ChunkIndex(INDEX_NONE)
+#if WITH_EDITOR
+	, ChunkGeneration(INDEX_NONE)
+#endif
 {
 }
 
@@ -1380,6 +1383,9 @@ FAudioChunkHandle::FAudioChunkHandle(const uint8* InData, uint32 NumBytes, const
 	, CorrespondingWave(InSoundWave)
 	, CorrespondingWaveName(SoundWaveName)
 	, ChunkIndex(InChunkIndex)
+#if WITH_EDITOR
+	, ChunkGeneration(InSoundWave->CurrentChunkRevision.GetValue())
+#endif
 {
 }
 
@@ -1402,7 +1408,9 @@ FAudioChunkHandle& FAudioChunkHandle::operator=(const FAudioChunkHandle& Other)
 	CorrespondingWave = Other.CorrespondingWave;
 	CorrespondingWaveName = Other.CorrespondingWaveName;
 	ChunkIndex = Other.ChunkIndex;
-
+#if WITH_EDITOR
+	ChunkGeneration = Other.ChunkGeneration;
+#endif
 
 	if (IsValid())
 	{
@@ -1436,6 +1444,21 @@ bool FAudioChunkHandle::IsValid() const
 	return GetData() != nullptr;
 }
 
+#if WITH_EDITOR
+bool FAudioChunkHandle::IsStale() const
+{
+	if (CorrespondingWave != nullptr)
+	{
+		// NOTE: While this is currently safe in editor, there's no guarantee the USoundWave will be kept alive during the lifecycle of this chunk handle.
+		return ChunkGeneration != CorrespondingWave->CurrentChunkRevision.GetValue();
+	}
+	else
+	{
+		return false;
+	}
+}
+#endif
+
 FAudioChunkHandle IAudioStreamingManager::BuildChunkHandle(const uint8* InData, uint32 NumBytes, const USoundWave* InSoundWave, const FName& SoundWaveName, uint32 InChunkIndex)
 {
 	return FAudioChunkHandle(InData, NumBytes, InSoundWave, SoundWaveName, InChunkIndex);
diff --git a/Engine/Source/Runtime/Engine/Private/SoundWave.cpp b/Engine/Source/Runtime/Engine/Private/SoundWave.cpp
--- a/Engine/Source/Runtime/Engine/Private/SoundWave.cpp
+++ b/Engine/Source/Runtime/Engine/Private/SoundWave.cpp
@@ -735,6 +735,19 @@ void USoundWave::InvalidateCompressedData(bool bFreeResources)
 	{
 		FreeResources(false);
 	}
+
+#if WITH_EDITOR
+	// Rebuild our streaming chunks.
+	CachePlatformData();
+	CurrentChunkRevision.Increment();
+	
+	// If this sound wave is retained, release and re-retain the new chunk.
+	if (FirstChunk.IsValid())
+	{
+		ReleaseCompressedAudio();
+		RetainCompressedAudio(true);
+	}
+#endif
 }
 
 bool USoundWave::HasStreamingChunks()
diff --git a/Engine/Source/Runtime/Engine/Public/AudioStreamingCache.h b/Engine/Source/Runtime/Engine/Public/AudioStreamingCache.h
--- a/Engine/Source/Runtime/Engine/Public/AudioStreamingCache.h
+++ b/Engine/Source/Runtime/Engine/Public/AudioStreamingCache.h
@@ -30,9 +30,18 @@ public:
 		FName SoundWaveName = FName();
 		uint32 ChunkIndex = INDEX_NONE;
 
+#if WITH_EDITOR
+		// This is used in the editor to invalidate stale compressed chunks.
+		uint32 ChunkRevision = INDEX_NONE;
+#endif
+
 		inline bool operator==(const FChunkKey& Other) const
 		{
+#if WITH_EDITOR
+			return (SoundWaveName == Other.SoundWaveName) && (ChunkIndex == Other.ChunkIndex) && (ChunkRevision == Other.ChunkRevision);
+#else
 			return (SoundWaveName == Other.SoundWaveName) && (ChunkIndex == Other.ChunkIndex);
+#endif
 		}
 	};
 
diff --git a/Engine/Source/Runtime/Engine/Public/ContentStreaming.h b/Engine/Source/Runtime/Engine/Public/ContentStreaming.h
--- a/Engine/Source/Runtime/Engine/Public/ContentStreaming.h
+++ b/Engine/Source/Runtime/Engine/Public/ContentStreaming.h
@@ -107,6 +107,11 @@ public:
 	// Checks whether this points to a valid compressed chunk.
 	bool IsValid() const;
 
+#if WITH_EDITOR
+	// If the soundwave has been recompressed, the compressed audio retained by this handle will not be up to date, and this will return true. 
+	bool IsStale() const;
+#endif
+
 private:
 	// This constructor should only be called by an implementation of IAudioStreamingManager.
 	FAudioChunkHandle(const uint8* InData, uint32 NumBytes, const USoundWave* InSoundWave, const FName& SoundWaveName, uint32 InChunkIndex);
@@ -118,6 +123,10 @@ private:
 	FName CorrespondingWaveName;
 	int32 ChunkIndex;
 
+#if WITH_EDITOR
+	uint32 ChunkGeneration;
+#endif
+
 	friend struct IAudioStreamingManager;
 	friend struct FCachedAudioStreamingManager;
 };
