Fix for log warnings on PostInitProperties.

diff --git a/Engine/Source/Runtime/Engine/Classes/Sound/SoundWave.h b/Engine/Source/Runtime/Engine/Classes/Sound/SoundWave.h
--- a/Engine/Source/Runtime/Engine/Classes/Sound/SoundWave.h
+++ b/Engine/Source/Runtime/Engine/Classes/Sound/SoundWave.h
@@ -827,7 +827,7 @@ public:
 	 * Change the guid and flush all compressed data
 	 * @param bFreeResources if true, will delete any precached compressed data as well.
 	 */
-	void InvalidateCompressedData(bool bFreeResources = false);
+	void InvalidateCompressedData(bool bFreeResources = false, bool bRebuildStreamingChunks = true);
 
 	/** Returns curves associated with this sound wave */
 	virtual class UCurveTable* GetCurveData() const override { return Curves; }
diff --git a/Engine/Source/Runtime/Engine/Private/SoundWave.cpp b/Engine/Source/Runtime/Engine/Private/SoundWave.cpp
--- a/Engine/Source/Runtime/Engine/Private/SoundWave.cpp
+++ b/Engine/Source/Runtime/Engine/Private/SoundWave.cpp
@@ -497,7 +497,8 @@ void USoundWave::PostInitProperties()
 
 	if(!IsTemplate())
 	{
-		InvalidateCompressedData();
+		// Don't rebuild our streaming chunks yet because we may not have loaded the RawPCMData at this point.
+		InvalidateCompressedData(false, false);
 	}
 
 #if WITH_EDITORONLY_DATA
@@ -725,7 +726,7 @@ FByteBulkData* USoundWave::GetCompressedData(FName Format, const FPlatformAudioC
 	return Result->GetBulkDataSize() > 0 ? Result : NULL; // we don't return empty bulk data...but we save it to avoid thrashing the DDC
 }
 
-void USoundWave::InvalidateCompressedData(bool bFreeResources)
+void USoundWave::InvalidateCompressedData(bool bFreeResources, bool bRebuildStreamingChunks)
 {
 	CompressedDataGuid = FGuid::NewGuid();
 	ZerothChunkData.Reset();
@@ -737,9 +738,12 @@ void USoundWave::InvalidateCompressedData(bool bFreeResources)
 	}
 
 #if WITH_EDITOR
-	// Rebuild our streaming chunks.
-	CachePlatformData();
-	CurrentChunkRevision.Increment();
+	if (bRebuildStreamingChunks)
+	{
+		CachePlatformData();
+		CurrentChunkRevision.Increment();
+	}
+	
 	
 	// If this sound wave is retained, release and re-retain the new chunk.
 	if (FirstChunk.IsValid())
