Fix Mac and Linux editor compiles
#rnx

diff --git a/Engine/Source/Runtime/Engine/Private/AudioStreamingCache.cpp b/Engine/Source/Runtime/Engine/Private/AudioStreamingCache.cpp
--- a/Engine/Source/Runtime/Engine/Private/AudioStreamingCache.cpp
+++ b/Engine/Source/Runtime/Engine/Private/AudioStreamingCache.cpp
@@ -213,7 +213,7 @@ FAudioChunkHandle FCachedAudioStreamingManager::GetLoadedChunk(const USoundWave*
 			, SoundWave->GetFName()
 			, ChunkIndex
 #if WITH_EDITOR
-			, SoundWave->CurrentChunkRevision.GetValue()
+			, (uint32)SoundWave->CurrentChunkRevision.GetValue()
 #endif
 		};
 
@@ -237,7 +237,7 @@ FAudioChunkHandle FCachedAudioStreamingManager::GetLoadedChunk(const USoundWave*
 				, SoundWave->GetFName() 
 				, ((uint32)NextChunk) 
 #if WITH_EDITOR
-				, SoundWave->CurrentChunkRevision.GetValue()
+				, (uint32)SoundWave->CurrentChunkRevision.GetValue()
 #endif
 			};
 
