Fix for crash with HTML5 caused by autoinstancing

diff --git a/Engine/Source/Runtime/Engine/Private/SceneManagement.cpp b/Engine/Source/Runtime/Engine/Private/SceneManagement.cpp
--- a/Engine/Source/Runtime/Engine/Private/SceneManagement.cpp
+++ b/Engine/Source/Runtime/Engine/Private/SceneManagement.cpp
@@ -747,7 +747,7 @@ FViewUniformShaderParameters::FViewUniformShaderParameters()
 	CameraAerialPerspectiveVolume = GBlackAlpha1VolumeTexture->TextureRHI;
 	CameraAerialPerspectiveVolumeSampler = TStaticSamplerState<SF_Bilinear>::GetRHI();
 
-	PrimitiveSceneDataTexture = GIdentityPrimitiveBuffer.PrimitiveSceneDataTextureRHI;
+	PrimitiveSceneDataTexture = OrBlack2DIfNull(GIdentityPrimitiveBuffer.PrimitiveSceneDataTextureRHI);
 	PrimitiveSceneData = GIdentityPrimitiveBuffer.PrimitiveSceneDataBufferSRV;
 	LightmapSceneData = GIdentityPrimitiveBuffer.LightmapSceneDataBufferSRV;
 
