Remove default custom time step and timecode provider. Now, the engine does support that no timecode provider was set.

diff --git a/Engine/Plugins/Experimental/VirtualProductionUtilities/Source/VPUtilities/Private/VPTimecodeCustomTimeStep.cpp b/Engine/Plugins/Experimental/VirtualProductionUtilities/Source/VPUtilities/Private/VPTimecodeCustomTimeStep.cpp
--- a/Engine/Plugins/Experimental/VirtualProductionUtilities/Source/VPUtilities/Private/VPTimecodeCustomTimeStep.cpp
+++ b/Engine/Plugins/Experimental/VirtualProductionUtilities/Source/VPUtilities/Private/VPTimecodeCustomTimeStep.cpp
@@ -16,10 +16,8 @@ bool UVPTimecodeCustomTimeStep::Initialize(UEngine* InEngine)
 {
 	check(InEngine);
 
-	// InitializedTimecodeProvider, will be set in the first update.
 	//The user may initialize the CustomTimeStep and the TimecodeProvider in the same frame and order of operation may break the behaviour.
 	InitializedSeconds = FApp::GetCurrentTime();
-	InitializedTimecodeProvider = nullptr;
 	State = ECustomTimeStepSynchronizationState::Synchronizing;
 
 	return true;
@@ -30,7 +28,8 @@ void UVPTimecodeCustomTimeStep::Shutdown(UEngine* InEngine)
 {
 	check(InEngine);
 
-	InitializedTimecodeProvider = nullptr;
+	InEngine->OnTimecodeProviderChanged().RemoveAll(this);
+
 	State = ECustomTimeStepSynchronizationState::Closed;
 }
 
@@ -78,14 +77,6 @@ bool UVPTimecodeCustomTimeStep::UpdateTimeStep(UEngine* InEngine)
 				return true;
 			}
 
-			// Test if the Timecode provider changed
-			if (bErrorIfTimecodeProviderChanged && TimecodeProvider != InitializedTimecodeProvider)
-			{
-				UE_LOG(LogVPUtilities, Error, TEXT("Timecode Provider changed for '%s'."), *GetName());
-				State = ECustomTimeStepSynchronizationState::Error;
-				return true;
-			}
-
 			NewTimecode = TimecodeProvider->GetTimecode();
 			if (NewTimecode == PreviousTimecode)
 			{
@@ -140,29 +131,18 @@ bool UVPTimecodeCustomTimeStep::InitializeFirstStep(UEngine* InEngine)
 	bool bFirstFrame = true;
 
 	const UTimecodeProvider* TimecodeProvider = InEngine->GetTimecodeProvider();
-	if (InitializedTimecodeProvider == nullptr)
+	if (TimecodeProvider == nullptr)
 	{
-		InitializedTimecodeProvider = TimecodeProvider;
-		if (InitializedTimecodeProvider == nullptr)
-		{
-			UE_LOG(LogVPUtilities, Error, TEXT("There is no Timecode Provider for '%s'."), *GetName());
-			State = ECustomTimeStepSynchronizationState::Error;
-			return false;
-		}
-	}
-
-	// Test if the Timecode provider changed
-	if (bErrorIfTimecodeProviderChanged && TimecodeProvider != InitializedTimecodeProvider)
-	{
-		UE_LOG(LogVPUtilities, Error, TEXT("Timecode Provider changed for '%s'."), *GetName());
+		UE_LOG(LogVPUtilities, Error, TEXT("There is no Timecode Provider for '%s'."), *GetName());
 		State = ECustomTimeStepSynchronizationState::Error;
 		return false;
 	}
 
-	if (InitializedTimecodeProvider->GetSynchronizationState() == ETimecodeProviderSynchronizationState::Synchronized)
+	if (TimecodeProvider->GetSynchronizationState() == ETimecodeProviderSynchronizationState::Synchronized)
 	{
-		PreviousTimecode = InitializedTimecodeProvider->GetTimecode();
-		PreviousFrameRate = InitializedTimecodeProvider->GetFrameRate();
+		PreviousTimecode = TimecodeProvider->GetTimecode();
+		PreviousFrameRate = TimecodeProvider->GetFrameRate();
+		InEngine->OnTimecodeProviderChanged().AddUObject(this, &UVPTimecodeCustomTimeStep::OnTimecodeProviderChanged);
 		State = ECustomTimeStepSynchronizationState::Synchronized;
 		bFirstFrame = true;
 	}
@@ -172,10 +152,19 @@ bool UVPTimecodeCustomTimeStep::InitializeFirstStep(UEngine* InEngine)
 		if ((FApp::GetCurrentTime() - InitializedSeconds) > NumberOfSecondsBeforeWarning)
 		{
 			bWarnAboutSynchronizationState = true;
-			UE_LOG(LogVPUtilities, Warning, TEXT("The TimecodeProvider '%s' is not Synchronized for '%s'."), *InitializedTimecodeProvider->GetName(), *GetName());
+			UE_LOG(LogVPUtilities, Warning, TEXT("The TimecodeProvider '%s' is not Synchronized for '%s'."), *TimecodeProvider->GetName(), *GetName());
 		}
 	}
 
 	return bFirstFrame;
 }
 
+void UVPTimecodeCustomTimeStep::OnTimecodeProviderChanged()
+{
+	// Test if the Timecode provider changed
+	if (bErrorIfTimecodeProviderChanged)
+	{
+		UE_LOG(LogVPUtilities, Error, TEXT("The Timecode Provider changed for '%s'."), *GetName());
+		State = ECustomTimeStepSynchronizationState::Error;
+	}
+}
diff --git a/Engine/Plugins/Experimental/VirtualProductionUtilities/Source/VPUtilities/Public/VPTimecodeCustomTimeStep.h b/Engine/Plugins/Experimental/VirtualProductionUtilities/Source/VPUtilities/Public/VPTimecodeCustomTimeStep.h
--- a/Engine/Plugins/Experimental/VirtualProductionUtilities/Source/VPUtilities/Public/VPTimecodeCustomTimeStep.h
+++ b/Engine/Plugins/Experimental/VirtualProductionUtilities/Source/VPUtilities/Public/VPTimecodeCustomTimeStep.h
@@ -43,6 +43,8 @@ public:
 private:
 	bool InitializeFirstStep(UEngine* InEngine);
 
+	void OnTimecodeProviderChanged();
+
 private:
 	/** The current SynchronizationState of the CustomTimeStep. */
 	ECustomTimeStepSynchronizationState State = ECustomTimeStepSynchronizationState::Closed;
@@ -56,9 +58,6 @@ private:
 	/** The time at initialization. */
 	double InitializedSeconds = 0.0;
 
+	/** Only warn once about the synchronization state. */
 	bool bWarnAboutSynchronizationState = false;
-
-	/** TimecodeProvider used to initialize the CustomTimeStep. */
-	UPROPERTY(Transient)
-	const UTimecodeProvider* InitializedTimecodeProvider = nullptr;
 };
diff --git a/Engine/Plugins/Media/MediaFrameworkUtilities/Source/MediaFrameworkUtilities/Private/Profile/MediaProfile.cpp b/Engine/Plugins/Media/MediaFrameworkUtilities/Source/MediaFrameworkUtilities/Private/Profile/MediaProfile.cpp
--- a/Engine/Plugins/Media/MediaFrameworkUtilities/Source/MediaFrameworkUtilities/Private/Profile/MediaProfile.cpp
+++ b/Engine/Plugins/Media/MediaFrameworkUtilities/Source/MediaFrameworkUtilities/Private/Profile/MediaProfile.cpp
@@ -71,7 +71,7 @@ void UMediaProfile::Apply()
 		return;
 	}
 
-	// Make sure we have the same amount of souces and outputs as the number of proxies.
+	// Make sure we have the same amount of sources and outputs as the number of proxies.
 	FixNumSourcesAndOutputs();
 
 	{
@@ -100,35 +100,29 @@ void UMediaProfile::Apply()
 		}
 	}
 
+	ResetTimecodeProvider();
 	if (bOverrideTimecodeProvider)
 	{
-		if (TimecodeProvider)
+		bTimecodeProvideWasApplied = true;
+		AppliedTimecodeProvider = TimecodeProvider;
+		PreviousTimecodeProvider = GEngine->GetTimecodeProvider();
+		bool bResult = GEngine->SetTimecodeProvider(TimecodeProvider);
+		if (!bResult && TimecodeProvider)
 		{
-			bool bResult = GEngine->SetTimecodeProvider(TimecodeProvider);
-			if (!bResult)
-			{
-				UE_LOG(LogMediaFrameworkUtilities, Warning, TEXT("The TimecodeProvider '%s' could not be initialized."), *TimecodeProvider->GetName());
-			}
-		}
-		else
-		{
-			GEngine->SetTimecodeProvider(nullptr);
+			UE_LOG(LogMediaFrameworkUtilities, Warning, TEXT("The Timecode Provider '%s' could not be initialized."), *TimecodeProvider->GetName());
 		}
 	}
 
+	ResetCustomTimeStep();
 	if (bOverrideCustomTimeStep)
 	{
-		if (CustomTimeStep)
+		bCustomTimeStepWasApplied = true;
+		AppliedCustomTimeStep = CustomTimeStep;
+		PreviousCustomTimeStep = GEngine->GetCustomTimeStep();
+		bool bResult = GEngine->SetCustomTimeStep(CustomTimeStep);
+		if (!bResult && CustomTimeStep)
 		{
-			bool bResult = GEngine->SetCustomTimeStep(CustomTimeStep);
-			if (!bResult)
-			{
-				UE_LOG(LogMediaFrameworkUtilities, Warning, TEXT("The Custom Time Step '%s' could not be initialized."), *CustomTimeStep->GetName());
-			}
-		}
-		else
-		{
-			GEngine->SetCustomTimeStep(nullptr);
+			UE_LOG(LogMediaFrameworkUtilities, Warning, TEXT("The Custom Time Step '%s' could not be initialized."), *CustomTimeStep->GetName());
 		}
 	}
 }
@@ -166,28 +160,68 @@ void UMediaProfile::Reset()
 		}
 	}
 
+	// Reset the timecode provider
+	ResetTimecodeProvider();
+
+	// Reset the engine custom time step
+	ResetCustomTimeStep();
+}
+
+void UMediaProfile::ResetTimecodeProvider()
+{
+	if (bTimecodeProvideWasApplied)
 	{
-		// Reset the timecode provider
-		const UTimecodeProvider* CurrentTimecodeProvider = GEngine->GetTimecodeProvider();
-		if (CurrentTimecodeProvider)
+		if (AppliedTimecodeProvider == GEngine->GetTimecodeProvider())
+		{
+			bool bResult = GEngine->SetTimecodeProvider(PreviousTimecodeProvider);
+			if (!bResult && PreviousTimecodeProvider)
+			{
+				UE_LOG(LogMediaFrameworkUtilities, Warning, TEXT("The TimecodeProvider '%s' could not be initialized."), *PreviousTimecodeProvider->GetName());
+			}
+		}
+		else
 		{
-			if (CurrentTimecodeProvider->GetOuter() == this)
+			if (PreviousTimecodeProvider)
 			{
-				GEngine->SetTimecodeProvider(nullptr);
+				UE_LOG(LogMediaFrameworkUtilities, Warning, TEXT("Could not set the previous TimecodeProvider '%s'."), *PreviousTimecodeProvider->GetName());
+			}
+			else
+			{
+				UE_LOG(LogMediaFrameworkUtilities, Warning, TEXT("Could not set the previous TimecodeProvider."));
 			}
 		}
+		PreviousTimecodeProvider = nullptr;
+		AppliedTimecodeProvider = nullptr;
+		bTimecodeProvideWasApplied = false;
 	}
+}
 
+void UMediaProfile::ResetCustomTimeStep()
+{
+	if (bCustomTimeStepWasApplied)
 	{
-		// Reset the engine custom time step
-		const UEngineCustomTimeStep* CurrentCustomTimeStep = GEngine->GetCustomTimeStep();
-		if (CurrentCustomTimeStep)
+		if (AppliedCustomTimeStep == GEngine->GetCustomTimeStep())
 		{
-			if (CurrentCustomTimeStep->GetOuter() == this)
+			bool bResult = GEngine->SetCustomTimeStep(PreviousCustomTimeStep);
+			if (!bResult && PreviousCustomTimeStep)
+			{
+				UE_LOG(LogMediaFrameworkUtilities, Warning, TEXT("The Custom Time Step '%s' could not be initialized."), *PreviousCustomTimeStep->GetName());
+			}
+		}
+		else
+		{
+			if (PreviousCustomTimeStep)
+			{
+				UE_LOG(LogMediaFrameworkUtilities, Warning, TEXT("Could not set the previous Custom Time Step '%s'."), *PreviousCustomTimeStep->GetName());
+			}
+			else
 			{
-				GEngine->SetCustomTimeStep(GEngine->GetDefaultCustomTimeStep());
+				UE_LOG(LogMediaFrameworkUtilities, Warning, TEXT("Could not set the previous Custom Time Step."));
 			}
 		}
+		PreviousCustomTimeStep = nullptr;
+		AppliedCustomTimeStep = nullptr;
+		bCustomTimeStepWasApplied = false;
 	}
 }
 
diff --git a/Engine/Plugins/Media/MediaFrameworkUtilities/Source/MediaFrameworkUtilities/Public/Profile/MediaProfile.h b/Engine/Plugins/Media/MediaFrameworkUtilities/Source/MediaFrameworkUtilities/Public/Profile/MediaProfile.h
--- a/Engine/Plugins/Media/MediaFrameworkUtilities/Source/MediaFrameworkUtilities/Public/Profile/MediaProfile.h
+++ b/Engine/Plugins/Media/MediaFrameworkUtilities/Source/MediaFrameworkUtilities/Public/Profile/MediaProfile.h
@@ -121,4 +121,25 @@ public:
 	 * Update the number of sources and outputs to the number to proxies.
 	 */
 	void FixNumSourcesAndOutputs();
+
+private:
+
+	void ResetTimecodeProvider();
+	void ResetCustomTimeStep();
+
+private:
+
+	/** Applied Timecode provider, cached to reset the previous value. */
+	bool bTimecodeProvideWasApplied;
+	UPROPERTY(Transient)
+	UTimecodeProvider* AppliedTimecodeProvider;
+	UPROPERTY(Transient)
+	UTimecodeProvider* PreviousTimecodeProvider;
+
+	/** Applied Custom time step, cached to reset the previous value. */
+	bool bCustomTimeStepWasApplied;
+	UPROPERTY(Transient)
+	UEngineCustomTimeStep* AppliedCustomTimeStep;
+	UPROPERTY(Transient)
+	UEngineCustomTimeStep* PreviousCustomTimeStep;
 };
diff --git a/Engine/Plugins/Media/TimecodeSynchronizer/Source/TimecodeSynchronizer/Private/TimecodeSynchronizer.cpp b/Engine/Plugins/Media/TimecodeSynchronizer/Source/TimecodeSynchronizer/Private/TimecodeSynchronizer.cpp
--- a/Engine/Plugins/Media/TimecodeSynchronizer/Source/TimecodeSynchronizer/Private/TimecodeSynchronizer.cpp
+++ b/Engine/Plugins/Media/TimecodeSynchronizer/Source/TimecodeSynchronizer/Private/TimecodeSynchronizer.cpp
@@ -205,6 +205,7 @@ UTimecodeSynchronizer::UTimecodeSynchronizer()
 	, ActiveMasterSynchronizationTimecodedSourceIndex(INDEX_NONE)
 	, bFailGuard(false)
 	, bAddSourcesGuard(false)
+	, bShouldResetTimecodeProvider(false)
 {
 }
 
@@ -318,9 +319,9 @@ FFrameTime UTimecodeSynchronizer::GetProviderFrameTime() const
 			UE_LOG(LogTimecodeSynchronizer, Log, TEXT("Unable to get frame time - Invalid source specified."));
 		}
 	}
-	else if (CachedTimecodeProvider)
+	else if (CachedProxiedTimecodeProvider)
 	{
-		ProviderFrameTime = FFrameTime(CachedTimecodeProvider->GetTimecode().ToFrameNumber(GetFrameRate()));
+		ProviderFrameTime = FFrameTime(CachedProxiedTimecodeProvider->GetTimecode().ToFrameNumber(GetFrameRate()));
 	}
 
 	return ProviderFrameTime;
@@ -351,11 +352,11 @@ ETimecodeProviderSynchronizationState UTimecodeSynchronizer::GetSynchronizationS
 bool UTimecodeSynchronizer::Initialize(class UEngine* InEngine)
 {
 	// The engine only allows one provider to be active at a given time.
-	// However, we are a specical case just acting as a pass through.
+	// However, we are a special case just acting as a pass through.
 	// Therefore, we need to make sure we pass along the initialization / shutdown requests.
-	if (CachedTimecodeProvider)
+	if (CachedProxiedTimecodeProvider)
 	{
-		return const_cast<UTimecodeProvider*>(CachedTimecodeProvider)->Initialize(InEngine);
+		return CachedProxiedTimecodeProvider->Initialize(InEngine);
 	}
 
 	return true;
@@ -363,9 +364,9 @@ bool UTimecodeSynchronizer::Initialize(class UEngine* InEngine)
 
 void UTimecodeSynchronizer::Shutdown(class UEngine* InEngine)
 {
-	if (CachedTimecodeProvider)
+	if (CachedProxiedTimecodeProvider)
 	{
-		const_cast<UTimecodeProvider*>(CachedTimecodeProvider)->Shutdown(InEngine);
+		CachedProxiedTimecodeProvider->Shutdown(InEngine);
 	}
 }
 
@@ -392,6 +393,7 @@ void UTimecodeSynchronizer::Register()
 	if (!bRegistered)
 	{
 		bRegistered = true;
+		bShouldResetTimecodeProvider = false;
 
 		if (FrameRateSource == ETimecodeSynchronizationFrameRateSources::EngineCustomTimeStepFrameRate)
 		{
@@ -425,27 +427,21 @@ void UTimecodeSynchronizer::Register()
 			FApp::SetUseFixedTimeStep(true);
 		}
 
+		CachedPreviousTimecodeProvider = GEngine->GetTimecodeProvider();
 		if (TimecodeProviderType == ETimecodeSynchronizationTimecodeType::TimecodeProvider)
 		{
-			CachedTimecodeProvider = TimecodeProvider;
+			CachedProxiedTimecodeProvider = TimecodeProvider;
 		}
 		else if (TimecodeProviderType == ETimecodeSynchronizationTimecodeType::DefaultProvider)
 		{
-			CachedTimecodeProvider = GEngine->GetDefaultTimecodeProvider();
+			CachedProxiedTimecodeProvider = GEngine->GetTimecodeProvider();
 		}
 		else
 		{
-			CachedTimecodeProvider = nullptr;
-		}
-
-		// Set TimecodeProvider
-		if (GEngine->GetTimecodeProvider() != GEngine->GetDefaultTimecodeProvider())
-		{
-			UE_LOG(LogTimecodeSynchronizer, Error, TEXT("A Timecode Provider is already in place."));
-			SwitchState(ESynchronizationState::Error);
-			return;
+			CachedProxiedTimecodeProvider = nullptr;
 		}
 
+		bShouldResetTimecodeProvider = true;
 		if (!GEngine->SetTimecodeProvider(this))
 		{
 			UE_LOG(LogTimecodeSynchronizer, Error, TEXT("TimecodeProvider failed to be set on Engine."));
@@ -453,6 +449,8 @@ void UTimecodeSynchronizer::Register()
 			return;
 		}
 
+		GEngine->OnTimecodeProviderChanged().AddUObject(this, &UTimecodeSynchronizer::OnTimecodeProviderChanged);
+
 		SetTickEnabled(true);
 	}
 }
@@ -463,12 +461,14 @@ void UTimecodeSynchronizer::Unregister()
 	{
 		bRegistered = false;
 
-		const UTimecodeProvider* Provider = GEngine->GetTimecodeProvider();
-		if (Provider == this)
+		GEngine->OnTimecodeProviderChanged().RemoveAll(this);
+		if (GEngine->GetTimecodeProvider() == this)
 		{
-			GEngine->SetTimecodeProvider(nullptr);
+			GEngine->SetTimecodeProvider(bShouldResetTimecodeProvider ? CachedPreviousTimecodeProvider : nullptr);
+			bShouldResetTimecodeProvider = false;
 		}
-		CachedTimecodeProvider = nullptr;
+		CachedPreviousTimecodeProvider = nullptr;
+		CachedProxiedTimecodeProvider = nullptr;
 
 		if (RegisteredCustomTimeStep == nullptr)
 		{
@@ -481,6 +481,17 @@ void UTimecodeSynchronizer::Unregister()
 	}
 }
 
+void UTimecodeSynchronizer::OnTimecodeProviderChanged()
+{
+	if (bRegistered && !IsError())
+	{
+		bShouldResetTimecodeProvider = false;
+		UE_LOG(LogTimecodeSynchronizer, Error, TEXT("The Engine's TimecodeProvider changed."));
+		SwitchState(ESynchronizationState::Error);
+		return;
+	}
+}
+
 void UTimecodeSynchronizer::SetTickEnabled(bool bEnabled)
 {
 	bIsTickEnabled = bEnabled;
@@ -713,7 +724,8 @@ bool UTimecodeSynchronizer::Tick_TestTimecode()
 		SwitchState(ESynchronizationState::Error);
 		return false;
 	}
-	else if (TimecodeProviderType == ETimecodeSynchronizationTimecodeType::InputSource)
+
+	if (TimecodeProviderType == ETimecodeSynchronizationTimecodeType::InputSource)
 	{
 		if (!SynchronizedSources.IsValidIndex(ActiveMasterSynchronizationTimecodedSourceIndex))
 		{
@@ -726,14 +738,14 @@ bool UTimecodeSynchronizer::Tick_TestTimecode()
 	}
 	else 
 	{
-		if (CachedTimecodeProvider == nullptr)
+		if (CachedProxiedTimecodeProvider == nullptr)
 		{
 			UE_LOG(LogTimecodeSynchronizer, Error, TEXT("The TimecodeProvider became invalid while synchronizing."));
 			SwitchState(ESynchronizationState::Error);
 			return false;
 		}
 
-		const ETimecodeProviderSynchronizationState SynchronizationState = CachedTimecodeProvider->GetSynchronizationState();
+		const ETimecodeProviderSynchronizationState SynchronizationState = CachedProxiedTimecodeProvider->GetSynchronizationState();
 		if (SynchronizationState != ETimecodeProviderSynchronizationState::Synchronized && SynchronizationState != ETimecodeProviderSynchronizationState::Synchronizing)
 		{
 			UE_LOG(LogTimecodeSynchronizer, Error, TEXT("The TimecodeProvider stopped while synchronizing."));
@@ -741,7 +753,7 @@ bool UTimecodeSynchronizer::Tick_TestTimecode()
 			return false;
 		}
 
-		if (CachedTimecodeProvider->GetFrameRate() != GetFrameRate())
+		if (CachedProxiedTimecodeProvider->GetFrameRate() != GetFrameRate())
 		{
 			UE_LOG(LogTimecodeSynchronizer, Error, TEXT("The TimecodeProvider frame rate do not correspond to the specified frame rate."));
 			SwitchState(ESynchronizationState::Error);
diff --git a/Engine/Plugins/Media/TimecodeSynchronizer/Source/TimecodeSynchronizer/Public/TimecodeSynchronizer.h b/Engine/Plugins/Media/TimecodeSynchronizer/Source/TimecodeSynchronizer/Public/TimecodeSynchronizer.h
--- a/Engine/Plugins/Media/TimecodeSynchronizer/Source/TimecodeSynchronizer/Public/TimecodeSynchronizer.h
+++ b/Engine/Plugins/Media/TimecodeSynchronizer/Source/TimecodeSynchronizer/Public/TimecodeSynchronizer.h
@@ -393,6 +393,9 @@ private:
 		}
 	}
 
+	/** Callback when the engine's TimecodeProvider changed. */
+	void OnTimecodeProviderChanged();
+
 	/** Registers asset to MediaModule tick */
 	void SetTickEnabled(bool bEnabled);
 
@@ -530,7 +533,10 @@ private:
 	UFixedFrameRateCustomTimeStep* RegisteredCustomTimeStep;
 
 	UPROPERTY(Transient)
-	const UTimecodeProvider* CachedTimecodeProvider;
+	UTimecodeProvider* CachedPreviousTimecodeProvider;
+
+	UPROPERTY(Transient)
+	UTimecodeProvider* CachedProxiedTimecodeProvider;
 
 	UPROPERTY(Transient, DuplicateTransient, VisibleAnywhere, Category = "Synchronization")
 	int32 ActualFrameOffset;
@@ -568,4 +574,5 @@ private:
 
 	bool bFailGuard;
 	bool bAddSourcesGuard;
+	bool bShouldResetTimecodeProvider;
 };
diff --git a/Engine/Source/Editor/DetailCustomizations/Private/DetailCustomizations.cpp b/Engine/Source/Editor/DetailCustomizations/Private/DetailCustomizations.cpp
--- a/Engine/Source/Editor/DetailCustomizations/Private/DetailCustomizations.cpp
+++ b/Engine/Source/Editor/DetailCustomizations/Private/DetailCustomizations.cpp
@@ -121,7 +121,6 @@
 #include "VectorStructCustomization.h"
 #include "Vector4StructCustomization.h"
 #include "AssetViewerSettingsCustomization.h"
-#include "EngineCustomization.h"
 #include "MeshMergingSettingsCustomization.h"
 #include "MaterialAttributePropertyDetails.h"
 #include "CollectionReferenceStructCustomization.h"
@@ -369,7 +368,6 @@ void FDetailCustomizationsModule::RegisterObjectCustomizations()
 	RegisterCustomClassLayout("CrashReportsPrivacySettings", FOnGetDetailCustomizationInstance::CreateStatic(&FImportantToggleSettingCustomization::MakeInstance));
 
 	RegisterCustomClassLayout("AssetViewerSettings", FOnGetDetailCustomizationInstance::CreateStatic(&FAssetViewerSettingsCustomization::MakeInstance));
-	RegisterCustomClassLayout("Engine", FOnGetDetailCustomizationInstance::CreateStatic(&FEngineCustomization::MakeInstance));
 
 	RegisterCustomClassLayout("MeshMergingSettingsObject", FOnGetDetailCustomizationInstance::CreateStatic(&FMeshMergingSettingsObjectCustomization::MakeInstance));
 
diff --git a/Engine/Source/Editor/DetailCustomizations/Private/EngineCustomization.cpp b/Engine/Source/Editor/DetailCustomizations/Private/EngineCustomization.cpp
deleted file mode 100644
diff --git a/Engine/Source/Editor/DetailCustomizations/Private/EngineCustomization.h b/Engine/Source/Editor/DetailCustomizations/Private/EngineCustomization.h
deleted file mode 100644
diff --git a/Engine/Source/Editor/TimeManagementEditor/Private/STimecodeProvider.cpp b/Engine/Source/Editor/TimeManagementEditor/Private/STimecodeProvider.cpp
--- a/Engine/Source/Editor/TimeManagementEditor/Private/STimecodeProvider.cpp
+++ b/Engine/Source/Editor/TimeManagementEditor/Private/STimecodeProvider.cpp
@@ -29,7 +29,7 @@ void STimecodeProvider::Construct(const FArguments& InArgs)
 			{
 				return TimecodeProviderPtr->GetFrameRate().ToPrettyText();
 			}
-			return GEngine->DefaultTimecodeFrameRate.ToPrettyText();
+			return NSLOCTEXT("TimecodeProvider", "Undefined", "<Undefined>");
 		}))
 		.Font(InArgs._TimecodeProviderFont)
 		.ColorAndOpacity(InArgs._TimecodeProviderColor)
@@ -61,7 +61,7 @@ void STimecodeProvider::Construct(const FArguments& InArgs)
 					{
 						return FText::FromName(TimecodeProviderPtr->GetFName());
 					}
-					return FText::FromString(TEXT("[System Clock]"));
+					return NSLOCTEXT("TimecodeProvider", "Undefined", "<Undefined>");
 				}))
 				.Font(InArgs._TimecodeProviderFont)
 				.ColorAndOpacity(InArgs._TimecodeProviderColor)
diff --git a/Engine/Source/Runtime/Engine/Classes/Engine/Engine.h b/Engine/Source/Runtime/Engine/Classes/Engine/Engine.h
--- a/Engine/Source/Runtime/Engine/Classes/Engine/Engine.h
+++ b/Engine/Source/Runtime/Engine/Classes/Engine/Engine.h
@@ -1323,13 +1323,15 @@ public:
 	FFloatRange SmoothedFrameRateRange;
 
 private:
-	/** Control how the Engine process the Framerate/Timestep */
+	/** Controls how the Engine process the Framerate/Timestep */
 	UPROPERTY(transient)
-	UEngineCustomTimeStep* DefaultCustomTimeStep;
+	UEngineCustomTimeStep* CustomTimeStep;
 
-	/** Control how the Engine process the Framerate/Timestep */
-	UPROPERTY(transient)
-	UEngineCustomTimeStep* CurrentCustomTimeStep;
+	/** Broadcasts whenever the custom time step changed. */
+	FSimpleMulticastDelegate CustomTimeStepChangedEvent;
+
+	/** Is the current custom time step was initialized properly and if we should shut it down. */
+	bool bIsCurrentCustomTimeStepInitialized;
 
 public:
 	/**
@@ -1341,40 +1343,22 @@ public:
 	FSoftClassPath CustomTimeStepClassName;
 
 private:
-	/**
-	 * Default timecode provider that will be used when no custom provider is set.
-	 * This is expected to be valid throughout the entire life of the application.
-	 */
+	/** Controls the Engine's timecode. */
 	UPROPERTY(transient)
-	UTimecodeProvider* DefaultTimecodeProvider;
+	UTimecodeProvider* TimecodeProvider;
 
-	UPROPERTY(transient)
-	UTimecodeProvider* CustomTimecodeProvider;
+	/** Broadcasts whenever the timecode provider changed. */
+	FSimpleMulticastDelegate TimecodeProviderChangedEvent;
 
-public:
-
-	/**
-	 * Allows UEngine subclasses a chance to override the DefaultTimecodeProvider class.
-	 * This must be set before InitializeObjectReferences is called.
-	 * This is intentionally protected and not exposed to config.
-	 */
-	UPROPERTY(config, EditAnywhere, Category=Timecode, meta=(MetaClass="TimecodeProvider", DisplayName="DefaultTimecodeProvider", ConfigRestartRequired=true))
-	FSoftClassPath DefaultTimecodeProviderClassName;
+	/** Is the current timecode provider was initialized properly and if we should shut it down. */
+	bool bIsCurrentTimecodeProviderInitialized;
 
+public:
 	/**
-	 * Override the CustomTimecodeProvider when the engine is started.
-	 * When set, this does not change the DefaultTImecodeProvider class.
-	 * Instead, it will create an instance and set it as the CustomTimecodeProvider.
+	 * Set TimecodeProvider when the engine is started.
 	 */
 	UPROPERTY(config, EditAnywhere, Category=Timecode, meta=(MetaClass="TimecodeProvider", DisplayName="TimecodeProvider", ConfigRestartRequired=true))
 	FSoftClassPath TimecodeProviderClassName;
-	
-	/**
-	 * Frame rate used to generated the engine Timecode's frame number when no TimecodeProvider are specified.
-	 * It doesn't control the Engine frame rate. The Engine can run faster or slower that the specified TimecodeFrameRate.
-	 */
-	UPROPERTY(config, EditAnywhere, Category=Timecode, Meta=(ConfigRestartRequired=true))
-	FFrameRate DefaultTimecodeFrameRate;
 
 public:
 	/** 
@@ -2076,24 +2060,23 @@ public:
 	 */
 	virtual double CorrectNegativeTimeDelta(double DeltaRealTime);
 
-	/** Causes the current CustomTimeStep to be shut down and then reinitialized. */
+	/** Causes the current custom time step to be shut down and then reinitialized. */
 	void ReinitializeCustomTimeStep();
 
 	/**
-	 * Set the CustomTimeStep that will control the Engine Framerate/Timestep
+	 * Set the custom time step that will control the Engine Framerate/Timestep.
+	 * It will shutdown the previous custom time step.
+	 * The new custom time step will be initialized.
 	 *
-	 * @return	true if the CustomTimeStep was properly initialized
+	 * @return	the result of the custom time step initialization.
 	 */
 	bool SetCustomTimeStep(UEngineCustomTimeStep* InCustomTimeStep);
 
-	/** Get the CustomTimeStep that control the Engine Framerate/Timestep */
-	UEngineCustomTimeStep* GetCustomTimeStep() const { return CurrentCustomTimeStep; };
+	/** Get the custom time step that control the Engine Framerate/Timestep */
+	UEngineCustomTimeStep* GetCustomTimeStep() const { return CustomTimeStep; };
 
-	/**
-	 * Get the DefaultCustomTimeStep created at the engine initialization.
-	 * This may be null if no CustomTimeStep was defined in the project settings and may not be currently active.
-	 */
-	UEngineCustomTimeStep* GetDefaultCustomTimeStep() const { return DefaultCustomTimeStep; }
+	/** Return custom time step changed event. */
+	FSimpleMulticastDelegate& OnCustomTimeStepChanged() { return CustomTimeStepChangedEvent; }
 
 	/** Executes the deferred commands */
 	void TickDeferredCommands();
@@ -2116,34 +2099,23 @@ public:
 	/** Update FApp::Timecode. */
 	void UpdateTimecode();
 
-	/** Causes the current TimecodeProvider to be shut down and then reinitialized. */
+	/** Causes the current timecode provider to be shut down and then reinitialized. */
 	void ReinitializeTimecodeProvider();
 
 	/**
-	 * Sets the CustomTimecodeProvider for the engine, shutting down the the default provider if necessary.
-	 * Passing nullptr will clear the current CustomTimecodeProvider, and re-initialize the default.
+	 * Set the timecode provider that will control the Engine's timecode.
+	 * It will shutdown the previous timecode provider.
+	 * The new timecode provider will be initialized.
 	 *
-	 * @return True if the provider was set (and initialized successfully when non-null).
+	 * @return	the result value of the new timecode provider initialization.
 	 */
 	bool SetTimecodeProvider(UTimecodeProvider* InTimecodeProvider);
 
-	/**
-	 * Get the TimecodeProvider that control the Engine's Timecode
-	 * The return value should always be non-null.
-	 */
-	const UTimecodeProvider* GetTimecodeProvider() const { return CustomTimecodeProvider ? CustomTimecodeProvider : GetDefaultTimecodeProvider(); }
-
-	/**
-	 * Get the DefaultTimecodeProvider.
-	 * This should be valid throughout the lifetime of the Engine (although, it may not always be active).
-	 */
-	const UTimecodeProvider* GetDefaultTimecodeProvider() const { check(DefaultTimecodeProvider != nullptr); return DefaultTimecodeProvider; }
-
-protected:
+	/** Get the TimecodeProvider that control the Engine's Timecode. */
+	UTimecodeProvider* GetTimecodeProvider() const { return TimecodeProvider; };
 
-	/** Provide mutable access to the current TimecodeProvider to Engine. This is needed to Initialize / Shutdown providers. */
-	UTimecodeProvider* GetTimecodeProviderProtected() { return const_cast<UTimecodeProvider*>(const_cast<const UEngine*>(this)->GetTimecodeProvider()); }
-	UTimecodeProvider* GetDefaultTimecodeProviderProtected() { return const_cast<UTimecodeProvider*>(const_cast<const UEngine*>(this)->GetDefaultTimecodeProvider()); }
+	/** Return timecode provider changed event. */
+	FSimpleMulticastDelegate& OnTimecodeProviderChanged() { return TimecodeProviderChangedEvent; }
 
 public:
 
diff --git a/Engine/Source/Runtime/Engine/Classes/Engine/SystemTimeTimecodeProvider.h b/Engine/Source/Runtime/Engine/Classes/Engine/SystemTimeTimecodeProvider.h
--- a/Engine/Source/Runtime/Engine/Classes/Engine/SystemTimeTimecodeProvider.h
+++ b/Engine/Source/Runtime/Engine/Classes/Engine/SystemTimeTimecodeProvider.h
@@ -10,7 +10,7 @@
 /**
  * Converts the current system time to timecode, relative to a provided frame rate.
  */
-UCLASS(editinlinenew)
+UCLASS(config=Engine, Blueprintable, editinlinenew)
 class ENGINE_API USystemTimeTimecodeProvider : public UTimecodeProvider
 {
 	GENERATED_BODY()
diff --git a/Engine/Source/Runtime/Engine/Private/Engine/SystemTimeTimecodeProvider.cpp b/Engine/Source/Runtime/Engine/Private/Engine/SystemTimeTimecodeProvider.cpp
--- a/Engine/Source/Runtime/Engine/Private/Engine/SystemTimeTimecodeProvider.cpp
+++ b/Engine/Source/Runtime/Engine/Private/Engine/SystemTimeTimecodeProvider.cpp
@@ -5,23 +5,10 @@
 #include "Misc/CoreMisc.h"
 #include "Misc/DateTime.h"
 
-static double ComputeTimeCodeOffset()
+FTimecode USystemTimeTimecodeProvider::GetTimecode() const
 {
 	const FDateTime DateTime = FDateTime::Now();
-	double HighPerformanceClock = FPlatformTime::Seconds();
 	const FTimespan Timespan = DateTime.GetTimeOfDay();
-	double Delta = Timespan.GetTotalSeconds() - HighPerformanceClock;
-	return Delta;
-}
-
-
-FTimecode USystemTimeTimecodeProvider::GetTimecode() const
-{
-	static double HighPerformanceClockDelta = ComputeTimeCodeOffset();
-	static double SecondsPerDay = 24.0 * 60.0 * 60.0;
-	FTimespan Timespan;
-	Timespan = FTimespan::FromSeconds(fmod(HighPerformanceClockDelta + FPlatformTime::Seconds(), SecondsPerDay));
 
-	FTimecode Result =  FTimecode::FromTimespan(Timespan, FrameRate, FTimecode::IsDropFormatTimecodeSupported(FrameRate), false);
-	return Result;
+	return FTimecode::FromTimespan(Timespan, FrameRate, FTimecode::IsDropFormatTimecodeSupported(FrameRate), false);
 }
diff --git a/Engine/Source/Runtime/Engine/Private/GameEngine.cpp b/Engine/Source/Runtime/Engine/Private/GameEngine.cpp
--- a/Engine/Source/Runtime/Engine/Private/GameEngine.cpp
+++ b/Engine/Source/Runtime/Engine/Private/GameEngine.cpp
@@ -738,8 +738,6 @@ UEngine::UEngine(const FObjectInitializer& ObjectInitializer)
 		}
 	}
 	#endif
-
-	DefaultTimecodeFrameRate = FFrameRate(30, 1);
 }
 
 
diff --git a/Engine/Source/Runtime/Engine/Private/UnrealEngine.cpp b/Engine/Source/Runtime/Engine/Private/UnrealEngine.cpp
--- a/Engine/Source/Runtime/Engine/Private/UnrealEngine.cpp
+++ b/Engine/Source/Runtime/Engine/Private/UnrealEngine.cpp
@@ -1769,14 +1769,19 @@ void UEngine::PreExit()
 
 	delete ScreenSaverInhibitorRunnable;
 
-	GetTimecodeProviderProtected()->Shutdown(this);
-	CustomTimecodeProvider = nullptr;
-
-	// Don't clear the pointer to DefaultTimecodeProvider, as other systems shutting down may try to reference it
-	// for validation.
-	DefaultTimecodeProvider->RemoveFromRoot();
+	// Shutdown Timecode Provider
+	if (bIsCurrentTimecodeProviderInitialized && TimecodeProvider)
+	{
+		TimecodeProvider->Shutdown(this);
+	}
+	TimecodeProvider = nullptr;
 
-	SetCustomTimeStep(nullptr);
+	// Shutdown CustomTimeStep
+	if (bIsCurrentCustomTimeStepInitialized && CustomTimeStep)
+	{
+		CustomTimeStep->Shutdown(this);
+	}
+	CustomTimeStep = nullptr;
 
 	ShutdownHMD();
 
@@ -1930,9 +1935,9 @@ void UEngine::UpdateTimeAndHandleMaxTickRate()
 	FTimedMemReport::Get().PumpTimedMemoryReports();
 #endif
 
-	if (CurrentCustomTimeStep)
+	if (CustomTimeStep)
 	{
-		bool bRunEngineCode = CurrentCustomTimeStep->UpdateTimeStep(this);
+		bool bRunEngineCode = CustomTimeStep->UpdateTimeStep(this);
 		if (!bRunEngineCode)
 		{
 			UpdateTimecode();
@@ -2145,107 +2150,75 @@ void UEngine::UpdateTimeAndHandleMaxTickRate()
 
 void UEngine::ReinitializeCustomTimeStep()
 {
-	UEngineCustomTimeStep* CustomTimeStep = GetCustomTimeStep();
 	if (CustomTimeStep)
 	{
-		CustomTimeStep->Shutdown(this);
-		if (!CustomTimeStep->Initialize(this))
+		if (bIsCurrentCustomTimeStepInitialized)
 		{
-			UE_LOG(LogEngine, Warning, TEXT("Failed reinitializing CustomTimeStep %s"), *GetPathName(CustomTimeStep));
+			CustomTimeStep->Shutdown(this);
 		}
+		bIsCurrentCustomTimeStepInitialized = CustomTimeStep->Initialize(this);
 	}
 }
 
 bool UEngine::SetCustomTimeStep(UEngineCustomTimeStep* InCustomTimeStep)
 {
-	bool bResult = true;
-
-	UEngineCustomTimeStep* Previous = GetCustomTimeStep();
-	if (InCustomTimeStep != Previous)
+	if (InCustomTimeStep != CustomTimeStep)
 	{
-		if (Previous)
+		if (CustomTimeStep && bIsCurrentCustomTimeStepInitialized)
 		{
-			Previous->Shutdown(this);
+			CustomTimeStep->Shutdown(this);
 		}
 
-		CurrentCustomTimeStep = InCustomTimeStep;
+		bIsCurrentCustomTimeStepInitialized = false;
+		CustomTimeStep = InCustomTimeStep && !InCustomTimeStep->IsPendingKill() ? InCustomTimeStep : nullptr;
 
-		if (CurrentCustomTimeStep)
+		if (CustomTimeStep)
 		{
-			bResult = CurrentCustomTimeStep->Initialize(this);
-			if (!bResult)
-			{
-				UE_LOG(LogEngine, Warning, TEXT("SetCustomTimeStep - Failed to intialize CustomTimeStep %s"), *GetPathName(CurrentCustomTimeStep));
-				CurrentCustomTimeStep = nullptr;
-			}
+			bIsCurrentCustomTimeStepInitialized = CustomTimeStep->Initialize(this);
 		}
+		OnCustomTimeStepChanged().Broadcast();
 	}
-
-	return bResult;
+	return bIsCurrentCustomTimeStepInitialized;
 }
 
 void UEngine::ReinitializeTimecodeProvider()
 {
-	UTimecodeProvider* Provider = GetTimecodeProviderProtected();
-	Provider->Shutdown(this);
-	if (!Provider->Initialize(this))
+	UTimecodeProvider* Provider = GetTimecodeProvider();
+	if (Provider)
 	{
-		UE_LOG(LogEngine, Warning, TEXT("Failed reinitializing TimecodeProvider %s"), *GetPathName(Provider));
+		if (bIsCurrentTimecodeProviderInitialized)
+		{
+			Provider->Shutdown(this);
+		}
+		bIsCurrentTimecodeProviderInitialized = Provider->Initialize(this);
 	}
 }
 
 bool UEngine::SetTimecodeProvider(UTimecodeProvider* InTimecodeProvider)
 {
-	bool bResult = true;
-
-	if (InTimecodeProvider != CustomTimecodeProvider)
+	if (InTimecodeProvider != TimecodeProvider)
 	{
-		const bool bCurrentlyUsingDefault = !CustomTimecodeProvider;
-		if (bCurrentlyUsingDefault)
-		{
-			// If we're already using the default, and we're resetting to the default, we don't need to do anything.
-			if (InTimecodeProvider == DefaultTimecodeProvider || InTimecodeProvider == nullptr)
-			{
-				return bResult;
-			}
-			else
-			{
-				DefaultTimecodeProvider->Shutdown(this);
-			}
-		}
-		else
+		if (TimecodeProvider && bIsCurrentTimecodeProviderInitialized)
 		{
-			CustomTimecodeProvider->Shutdown(this);
-			CustomTimecodeProvider = nullptr;
+			TimecodeProvider->Shutdown(this);
 		}
 
-		if (InTimecodeProvider != nullptr)
-		{
-			bResult = InTimecodeProvider->Initialize(this);
-			if (bResult)
-			{
-				CustomTimecodeProvider = InTimecodeProvider;
-			}
-		}
+		bIsCurrentTimecodeProviderInitialized = false;
+		TimecodeProvider = InTimecodeProvider && !InTimecodeProvider->IsPendingKill() ? InTimecodeProvider : nullptr;
 
-		// If the new provider failed to initialized (or was null), then
-		// re-initialize the default provider.
-		if (!CustomTimecodeProvider)
+		if (TimecodeProvider)
 		{
-			if (!ensure(DefaultTimecodeProvider->Initialize(this)))
-			{
-				UE_LOG(LogEngine, Warning, TEXT("SetTimecodeProvider - Failed to intialize DefaultTimecodeProvider %s"), *GetPathName(DefaultTimecodeProvider));
-			}
+			bIsCurrentTimecodeProviderInitialized = TimecodeProvider->Initialize(this);
 		}
+		OnTimecodeProviderChanged().Broadcast();
 	}
-	
-	return bResult;
+	return bIsCurrentTimecodeProviderInitialized;
 }
 
 void UEngine::UpdateTimecode()
 {
 	const UTimecodeProvider* Provider = GetTimecodeProvider();
-	if (Provider->GetSynchronizationState() == ETimecodeProviderSynchronizationState::Synchronized)
+	if (Provider && Provider->GetSynchronizationState() == ETimecodeProviderSynchronizationState::Synchronized)
 	{
 		FApp::SetTimecodeAndFrameRate(Provider->GetDelayedTimecode(), Provider->GetFrameRate());
 	}
@@ -2375,23 +2348,52 @@ static void LoadEngineTexture(TextureType*& InOutTexture, const TCHAR* InName)
 	}
 }
 
-UEngineCustomTimeStep* InitializeCustomTimeStep(UEngine* InEngine, FSoftClassPath InCustomTimeStepClassName)
+static void LoadCustomTimeStep(UEngine* Engine)
 {
-	UEngineCustomTimeStep* NewCustomTimeStep = nullptr;
-	if (InEngine->GetCustomTimeStep() == nullptr && InCustomTimeStepClassName.IsValid())
+	if (Engine->CustomTimeStepClassName.IsValid())
 	{
-		UClass* CustomTimeStepClass = LoadClass<UObject>(nullptr, *InCustomTimeStepClassName.ToString());
-		if (CustomTimeStepClass)
+		if (UClass* CustomTimeStepClass = Engine->CustomTimeStepClassName.TryLoadClass<UEngineCustomTimeStep>())
 		{
-			NewCustomTimeStep = NewObject<UEngineCustomTimeStep>(InEngine, CustomTimeStepClass);
-			InEngine->SetCustomTimeStep(NewCustomTimeStep);
+			UEngineCustomTimeStep* NewCustomTimeStep = NewObject<UEngineCustomTimeStep>(Engine, CustomTimeStepClass);
+			if (!Engine->SetCustomTimeStep(NewCustomTimeStep))
+			{
+				UE_LOG(LogEngine, Warning, TEXT("InitializeCustomTimeStep - Failed to intialize CustomTimeStep '%s'."), *NewCustomTimeStep->GetPathName());
+			}
 		}
 		else
 		{
-			UE_LOG(LogEngine, Error, TEXT("Engine config value CustomTimeStepClassName '%s' is not a valid class name."), *InCustomTimeStepClassName.ToString());
+			Engine->SetCustomTimeStep(nullptr);
+			UE_LOG(LogEngine, Error, TEXT("Engine config value CustomTimeStepClassName '%s' is not a valid class name."), *Engine->CustomTimeStepClassName.ToString());
 		}
 	}
-	return NewCustomTimeStep;
+	else
+	{
+		Engine->SetCustomTimeStep(nullptr);
+	}
+}
+
+static void LoadTimecodeProvider(UEngine* Engine)
+{
+	if (Engine->TimecodeProviderClassName.IsValid())
+	{
+		if (UClass* TimecodeProviderClass = Engine->TimecodeProviderClassName.TryLoadClass<UTimecodeProvider>())
+		{
+			UTimecodeProvider* NewTimecodeProvider = NewObject<UTimecodeProvider>(Engine, TimecodeProviderClass);
+			if (!Engine->SetTimecodeProvider(NewTimecodeProvider))
+			{
+				UE_LOG(LogEngine, Warning, TEXT("InitializeTimecodeProvider - Failed to intialize TimecodeProvider '%s'."), *NewTimecodeProvider->GetPathName());
+			}
+		}
+		else
+		{
+			Engine->SetTimecodeProvider(nullptr);
+			UE_LOG(LogEngine, Error, TEXT("Engine config value TimecodeProviderClassName '%s' is not a valid class name."), *Engine->TimecodeProviderClassName.ToString());
+		}
+	}
+	else
+	{
+		Engine->SetTimecodeProvider(nullptr);
+	}
 }
 
 /**
@@ -2592,43 +2594,8 @@ void UEngine::InitializeObjectReferences()
 		}
 	}
 
-	DefaultCustomTimeStep = InitializeCustomTimeStep(this, CustomTimeStepClassName);
-
-	// Setup the timecode providers.
-	{
-		if (DefaultTimecodeProvider == nullptr)
-		{
-			UClass* DefaultTimecodeProviderClass = DefaultTimecodeProviderClassName.TryLoadClass<UTimecodeProvider>();
-			if (DefaultTimecodeProviderClass == nullptr)
-			{
-				DefaultTimecodeProviderClass = USystemTimeTimecodeProvider::StaticClass();
-			}
-
-			DefaultTimecodeProvider = NewObject<UTimecodeProvider>(this, DefaultTimecodeProviderClass);
-			if (!ensure(DefaultTimecodeProvider->Initialize(this)))
-			{
-				UE_LOG(LogEngine, Warning, TEXT("InitializeObjectReferences - Failed to intialize DefaultTimecodeProvider %s"), *GetPathName(DefaultTimecodeProvider));
-			}
-
-			DefaultTimecodeProvider->AddToRoot();
-			if (USystemTimeTimecodeProvider* LocalSystemTimeProvider = Cast<USystemTimeTimecodeProvider>(DefaultTimecodeProvider))
-			{
-				LocalSystemTimeProvider->SetFrameRate(DefaultTimecodeFrameRate);
-			}
-		}
-
-		if (CustomTimecodeProvider == nullptr && TimecodeProviderClassName.IsValid())
-		{
-			if (UClass* TimecodeProviderClass = TimecodeProviderClassName.TryLoadClass<UTimecodeProvider>())
-			{
-				SetTimecodeProvider(NewObject<UTimecodeProvider>(this, TimecodeProviderClass));
-			}
-			else
-			{
-				UE_LOG(LogEngine, Error, TEXT("Engine config value TimecodeProviderClassName '%s' is not a valid class name."), *TimecodeProviderClassName.ToString());
-			}
-		}
-	}
+	LoadCustomTimeStep(this);
+	LoadTimecodeProvider(this);
 
 	if (GameSingleton == nullptr && GameSingletonClassName.ToString().Len() > 0)
 	{
@@ -14560,7 +14527,6 @@ int32 UEngine::RenderStatFPS(UWorld* World, FViewport* Viewport, FCanvas* Canvas
 	// Start drawing the various counters.
 	const int32 RowHeight = FMath::TruncToInt(Font->GetMaxCharHeight() * 1.1f);
 
-	UEngineCustomTimeStep* CustomTimeStep = GetCustomTimeStep();
 	if (CustomTimeStep)
 	{
 		ECustomTimeStepSynchronizationState State = CustomTimeStep->GetSynchronizationState();
@@ -15816,28 +15782,31 @@ int32 UEngine::RenderStatTimecode(UWorld* World, FViewport* Viewport, FCanvas* C
 	const int32 RowHeight = FMath::TruncToInt(Font->GetMaxCharHeight() * 1.1f);
 
 	const UTimecodeProvider* Provider = GetTimecodeProvider();
-	ETimecodeProviderSynchronizationState State = Provider->GetSynchronizationState();
-	FString ProviderName = Provider->GetName();
-	float CharWidth, CharHeight;
-	Font->GetCharSize(TEXT(' '), CharWidth, CharHeight);
-	int32 NewX = X - Font->GetStringSize(*ProviderName) - (int32)CharWidth;
-	switch (State)
-	{
-	case ETimecodeProviderSynchronizationState::Closed:
-		Canvas->DrawShadowedString(NewX, Y, *FString::Printf(TEXT("%s TC: Closed"), *ProviderName), Font, FColor::Red);
-		break;
-	case ETimecodeProviderSynchronizationState::Error:
-		Canvas->DrawShadowedString(NewX, Y, *FString::Printf(TEXT("%s TC: Error"), *ProviderName), Font, FColor::Red);
-		break;
-	case ETimecodeProviderSynchronizationState::Synchronized:
-		Canvas->DrawShadowedString(NewX, Y, *FString::Printf(TEXT("%s TC: Synchronized"), *ProviderName), Font, FColor::Green);
-		break;
-	case ETimecodeProviderSynchronizationState::Synchronizing:
-		Canvas->DrawShadowedString(NewX, Y, *FString::Printf(TEXT("%s TC: Synchronizing"), *ProviderName), Font, FColor::Yellow);
-		break;
-	default:
-		check(false);
-		break;
+	if (Provider)
+	{
+		ETimecodeProviderSynchronizationState State = Provider->GetSynchronizationState();
+		FString ProviderName = Provider->GetName();
+		float CharWidth, CharHeight;
+		Font->GetCharSize(TEXT(' '), CharWidth, CharHeight);
+		int32 NewX = X - Font->GetStringSize(*ProviderName) - (int32)CharWidth;
+		switch (State)
+		{
+		case ETimecodeProviderSynchronizationState::Closed:
+			Canvas->DrawShadowedString(NewX, Y, *FString::Printf(TEXT("%s TC: Closed"), *ProviderName), Font, FColor::Red);
+			break;
+		case ETimecodeProviderSynchronizationState::Error:
+			Canvas->DrawShadowedString(NewX, Y, *FString::Printf(TEXT("%s TC: Error"), *ProviderName), Font, FColor::Red);
+			break;
+		case ETimecodeProviderSynchronizationState::Synchronized:
+			Canvas->DrawShadowedString(NewX, Y, *FString::Printf(TEXT("%s TC: Synchronized"), *ProviderName), Font, FColor::Green);
+			break;
+		case ETimecodeProviderSynchronizationState::Synchronizing:
+			Canvas->DrawShadowedString(NewX, Y, *FString::Printf(TEXT("%s TC: Synchronizing"), *ProviderName), Font, FColor::Yellow);
+			break;
+		default:
+			check(false);
+			break;
+		}
 	}
 	Y += RowHeight;
 	Canvas->DrawShadowedString(X, Y, *FString::Printf(TEXT("TC: %s"), *FApp::GetTimecode().ToString()), Font, FColor::Green);
