#ue4 - Fix glitches when walking on a component attached to a simulated component.

Github #5787

diff --git a/Engine/Source/Runtime/Engine/Classes/GameFramework/Character.h b/Engine/Source/Runtime/Engine/Classes/GameFramework/Character.h
--- a/Engine/Source/Runtime/Engine/Classes/GameFramework/Character.h
+++ b/Engine/Source/Runtime/Engine/Classes/GameFramework/Character.h
@@ -122,6 +122,9 @@ namespace MovementBaseUtility
 	/** Determine whether MovementBase can possibly move. */
 	ENGINE_API bool IsDynamicBase(const UPrimitiveComponent* MovementBase);
 
+	/** Determine whether MovementBase is simulating or attached to a simulating object. */
+	ENGINE_API bool IsSimulatedBase(const UPrimitiveComponent* MovementBase);
+
 	/** Determine if we should use relative positioning when based on a component (because it may move). */
 	FORCEINLINE bool UseRelativeLocation(const UPrimitiveComponent* MovementBase)
 	{
diff --git a/Engine/Source/Runtime/Engine/Private/Character.cpp b/Engine/Source/Runtime/Engine/Private/Character.cpp
--- a/Engine/Source/Runtime/Engine/Private/Character.cpp
+++ b/Engine/Source/Runtime/Engine/Private/Character.cpp
@@ -492,6 +492,18 @@ namespace MovementBaseUtility
 		return (MovementBase && MovementBase->Mobility == EComponentMobility::Movable);
 	}
 
+	bool IsSimulatedBase(const UPrimitiveComponent* MovementBase)
+	{
+		bool bBaseIsSimulatingPhysics = false;
+		const USceneComponent* AttachParent = MovementBase;
+		while (!bBaseIsSimulatingPhysics && AttachParent)
+		{
+			bBaseIsSimulatingPhysics = AttachParent->IsSimulatingPhysics();
+			AttachParent = AttachParent->GetAttachParent();
+		}
+		return bBaseIsSimulatingPhysics;
+	}
+
 	void AddTickDependency(FTickFunction& BasedObjectTick, UPrimitiveComponent* NewBase)
 	{
 		if (NewBase && MovementBaseUtility::UseRelativeLocation(NewBase))
@@ -689,7 +701,7 @@ void ACharacter::SetBase( UPrimitiveComponent* NewBaseComponent, const FName InB
 
 		if (CharacterMovement)
 		{
-			const bool bBaseIsSimulating = NewBaseComponent && NewBaseComponent->IsSimulatingPhysics();
+			const bool bBaseIsSimulating = MovementBaseUtility::IsSimulatedBase(NewBaseComponent);
 			if (bBaseChanged)
 			{
 				MovementBaseUtility::RemoveTickDependency(CharacterMovement->PrimaryComponentTick, OldBase);
diff --git a/Engine/Source/Runtime/Engine/Private/Components/CharacterMovementComponent.cpp b/Engine/Source/Runtime/Engine/Private/Components/CharacterMovementComponent.cpp
--- a/Engine/Source/Runtime/Engine/Private/Components/CharacterMovementComponent.cpp
+++ b/Engine/Source/Runtime/Engine/Private/Components/CharacterMovementComponent.cpp
@@ -1930,16 +1930,8 @@ void UCharacterMovementComponent::MaybeUpdateBasedMovement(float DeltaSeconds)
 	UPrimitiveComponent* MovementBase = CharacterOwner->GetMovementBase();
 	if (MovementBaseUtility::UseRelativeLocation(MovementBase))
 	{
-		// Need to see if anything we're on is simulating physics or has a parent that is.
-		bool bBaseIsSimulatingPhysics = false;
-		USceneComponent* AttachParent = MovementBase;
-		while (!bBaseIsSimulatingPhysics && AttachParent)
-		{
-			bBaseIsSimulatingPhysics = AttachParent->IsSimulatingPhysics();
-			AttachParent = AttachParent->GetAttachParent();
-		}
-		
-		if (!bBaseIsSimulatingPhysics)
+		// Need to see if anything we're on is simulating physics or has a parent that is.		
+		if (!MovementBaseUtility::IsSimulatedBase(MovementBase))
 		{
 			bDeferUpdateBasedMovement = false;
 			UpdateBasedMovement(DeltaSeconds);
