Further fix crash when LODData needs to be resized

diff --git a/Engine/Source/Runtime/Engine/Private/StaticMeshLight.cpp b/Engine/Source/Runtime/Engine/Private/StaticMeshLight.cpp
--- a/Engine/Source/Runtime/Engine/Private/StaticMeshLight.cpp
+++ b/Engine/Source/Runtime/Engine/Private/StaticMeshLight.cpp
@@ -215,7 +215,9 @@ void FStaticMeshStaticLightingTextureMapping::Apply(FQuantizedLightmapData* Quan
 
 	UStaticMeshComponent* StaticMeshComponent = Primitive.Get();
 
-	if (StaticMeshComponent && StaticMeshComponent->GetOwner() && StaticMeshComponent->GetOwner()->GetLevel() && StaticMeshComponent->LODData[LODIndex].MapBuildDataId.IsValid())
+	if (StaticMeshComponent && StaticMeshComponent->GetOwner() && StaticMeshComponent->GetOwner()->GetLevel() && 
+		(LODIndex >= StaticMeshComponent->LODData.Num() || StaticMeshComponent->LODData[LODIndex].MapBuildDataId.IsValid())
+		)
 	{
 		// Should have happened at a higher level
 		check(!StaticMeshComponent->IsRenderStateCreated());
