-Checking the world in SetPhysicsAsset to prevent accessing a null world.

diff --git a/Engine/Source/Runtime/Engine/Private/SkeletalMeshComponentPhysics.cpp b/Engine/Source/Runtime/Engine/Private/SkeletalMeshComponentPhysics.cpp
--- a/Engine/Source/Runtime/Engine/Private/SkeletalMeshComponentPhysics.cpp
+++ b/Engine/Source/Runtime/Engine/Private/SkeletalMeshComponentPhysics.cpp
@@ -1776,7 +1776,7 @@ void USkeletalMeshComponent::SetPhysicsAsset(UPhysicsAsset* InPhysicsAsset, bool
 
 			// Initialize new Physics Asset
 			UWorld* World = GetWorld();
-			if(World->GetPhysicsScene() != nullptr && ShouldCreatePhysicsState())
+			if(World && World->GetPhysicsScene() != nullptr && ShouldCreatePhysicsState())
 			{
 			//	UE_LOG(LogSkeletalMesh, Warning, TEXT("Creating Physics State (%s : %s)"), *GetNameSafe(GetOuter()),  *GetName());			
 				InitArticulated(World->GetPhysicsScene());
