#UE4 Minor perf improvement. Check for parameter name before casting in GetNamedParameterOfType.

diff --git a/Engine/Source/Runtime/Engine/Classes/Materials/MaterialFunctionInterface.h b/Engine/Source/Runtime/Engine/Classes/Materials/MaterialFunctionInterface.h
--- a/Engine/Source/Runtime/Engine/Classes/Materials/MaterialFunctionInterface.h
+++ b/Engine/Source/Runtime/Engine/Classes/Materials/MaterialFunctionInterface.h
@@ -168,9 +168,9 @@ public:
 			{
 				for (UMaterialExpression* FunctionExpression : *Function->GetFunctionExpressions())
 				{
-					if (ExpressionType* ExpressionParameter = Cast<ExpressionType>(FunctionExpression))
+					if (FunctionExpression->GetParameterName() == ParameterInfo.Name)
 					{
-						if (ExpressionParameter->ParameterName == ParameterInfo.Name)
+						if (ExpressionType* ExpressionParameter = Cast<ExpressionType>(FunctionExpression))
 						{
 							Parameter = ExpressionParameter;
 
