Back out changelist 9871790. Non-editor CIS error

diff --git a/Engine/Source/Runtime/Engine/Classes/Materials/MaterialFunctionInterface.h b/Engine/Source/Runtime/Engine/Classes/Materials/MaterialFunctionInterface.h
--- a/Engine/Source/Runtime/Engine/Classes/Materials/MaterialFunctionInterface.h
+++ b/Engine/Source/Runtime/Engine/Classes/Materials/MaterialFunctionInterface.h
@@ -168,9 +168,9 @@ public:
 			{
 				for (UMaterialExpression* FunctionExpression : *Function->GetFunctionExpressions())
 				{
-					if (FunctionExpression->GetParameterName() == ParameterInfo.Name)
+					if (ExpressionType* ExpressionParameter = Cast<ExpressionType>(FunctionExpression))
 					{
-						if (ExpressionType* ExpressionParameter = Cast<ExpressionType>(FunctionExpression))
+						if (ExpressionParameter->ParameterName == ParameterInfo.Name)
 						{
 							Parameter = ExpressionParameter;
 
