UE4R - Add more info to help track issue
#rnx

diff --git a/Engine/Source/Runtime/Engine/Private/Materials/MeshMaterialShader.cpp b/Engine/Source/Runtime/Engine/Private/Materials/MeshMaterialShader.cpp
--- a/Engine/Source/Runtime/Engine/Private/Materials/MeshMaterialShader.cpp
+++ b/Engine/Source/Runtime/Engine/Private/Materials/MeshMaterialShader.cpp
@@ -145,7 +145,7 @@ FShader* FMeshMaterialShaderType::FinishCompileShader(
 	const FString& InDebugDescription)
 {
 	checkf(CurrentJob.bSucceeded, TEXT("Failed MeshMaterialType compilation job: %s"), *GetJobName(&CurrentJob, ShaderPipelineType, InDebugDescription));
-	check(CurrentJob.VFType);
+	checkf(CurrentJob.VFType, TEXT("No VF on MeshMaterialType compilation job: %s"), *GetJobName(&CurrentJob, ShaderPipelineType, InDebugDescription));
 
 	// Reuse an existing resource with the same key or create a new one based on the compile output
 	// This allows FShaders to share compiled bytecode and RHI shader references
