Write output hash as part of the shader debug info.

diff --git a/Engine/Source/Runtime/Engine/Private/ShaderCompiler/ShaderCompiler.cpp b/Engine/Source/Runtime/Engine/Private/ShaderCompiler/ShaderCompiler.cpp
--- a/Engine/Source/Runtime/Engine/Private/ShaderCompiler/ShaderCompiler.cpp
+++ b/Engine/Source/Runtime/Engine/Private/ShaderCompiler/ShaderCompiler.cpp
@@ -620,6 +620,13 @@ static void ReadSingleJob(FShaderCompileJob* CurrentJob, FArchive& OutputFile)
 	// The shader processing this output will use it to search for existing FShaderResources
 	CurrentJob->Output.GenerateOutputHash();
 	CurrentJob->bSucceeded = CurrentJob->Output.bSucceeded;
+
+	if (CurrentJob->bSucceeded && CurrentJob->Input.DumpDebugInfoRootPath.Len() > 0)
+	{
+		// write down the output hash as a file
+		FString HashFileName = FPaths::Combine(CurrentJob->Input.DumpDebugInfoPath, TEXT("OutputHash.txt"));
+		FFileHelper::SaveStringToFile(CurrentJob->Output.OutputHash.ToString(), *HashFileName, FFileHelper::EEncodingOptions::ForceAnsi);
+	}
 };
 
 // Process results from Worker Process
