Fixed AttachmentGroup errors
- Stop being removed when the last primitive left the group, which meant it was not re-created properly when a primitive joined again.
- Re-add (MarkRenderStateDirty) all child primitive components when bLightAttachmentsAsGroup is set (through editor or setter function), otherwise children will not add themselves to the correct group (which may now be some other ancestor).

diff --git a/Engine/Source/Runtime/Engine/Classes/Components/PrimitiveComponent.h b/Engine/Source/Runtime/Engine/Classes/Components/PrimitiveComponent.h
--- a/Engine/Source/Runtime/Engine/Classes/Components/PrimitiveComponent.h
+++ b/Engine/Source/Runtime/Engine/Classes/Components/PrimitiveComponent.h
@@ -1907,6 +1907,9 @@ protected:
 
 	/** Ensure physics state created **/
 	void EnsurePhysicsStateCreated();
+
+	/**  Go through attached primitive components and call MarkRenderStateDirty */
+	void MarkChildPrimitiveComponentRenderStateDirty();
 public:
 
 	//~ Begin UObject Interface.
diff --git a/Engine/Source/Runtime/Engine/Private/Components/PrimitiveComponent.cpp b/Engine/Source/Runtime/Engine/Private/Components/PrimitiveComponent.cpp
--- a/Engine/Source/Runtime/Engine/Private/Components/PrimitiveComponent.cpp
+++ b/Engine/Source/Runtime/Engine/Private/Components/PrimitiveComponent.cpp
@@ -785,6 +785,30 @@ void UPrimitiveComponent::EnsurePhysicsStateCreated()
 	}
 }
 
+
+
+void UPrimitiveComponent::MarkChildPrimitiveComponentRenderStateDirty()
+{
+	// Go through all potential children and update them 
+	TArray<USceneComponent*, TInlineAllocator<8>> ProcessStack;
+	ProcessStack.Append(GetAttachChildren());
+
+	// Walk down the tree updating
+	while (ProcessStack.Num() > 0)
+	{
+		USceneComponent* Current = ProcessStack.Pop(/*bAllowShrinking=*/ false);
+		UPrimitiveComponent* CurrentPrimitive = Cast<UPrimitiveComponent>(Current);
+
+		if (CurrentPrimitive)
+		{
+			CurrentPrimitive->MarkRenderStateDirty();
+		}
+
+		ProcessStack.Append(Current->GetAttachChildren());
+	}
+}
+
+
 bool UPrimitiveComponent::IsWelded() const
 {
 	return BodyInstance.WeldParent != nullptr;
@@ -926,6 +950,12 @@ void UPrimitiveComponent::PostEditChangeProperty(FPropertyChangedEvent& Property
 			MarkRenderStateDirty();
 		}
 
+		// In the light attachment as group has changed, we need to notify attachment children that they are invalid (they may have a new root)
+		// Unless multiple roots are in the way, in either case, they need to work this out.
+		if (PropertyName == GET_MEMBER_NAME_CHECKED(UPrimitiveComponent, bLightAttachmentsAsGroup))
+		{
+			MarkChildPrimitiveComponentRenderStateDirty();
+		}
 	}
 
 	if (UProperty* MemberPropertyThatChanged = PropertyChangedEvent.MemberProperty)
@@ -1401,6 +1431,7 @@ void UPrimitiveComponent::SetLightAttachmentsAsGroup(bool bInLightAttachmentsAsG
 	{
 		bLightAttachmentsAsGroup = bInLightAttachmentsAsGroup;
 		MarkRenderStateDirty();
+		MarkChildPrimitiveComponentRenderStateDirty();
 	}
 }
 
diff --git a/Engine/Source/Runtime/Renderer/Private/PrimitiveSceneInfo.cpp b/Engine/Source/Runtime/Renderer/Private/PrimitiveSceneInfo.cpp
--- a/Engine/Source/Runtime/Renderer/Private/PrimitiveSceneInfo.cpp
+++ b/Engine/Source/Runtime/Renderer/Private/PrimitiveSceneInfo.cpp
@@ -895,9 +895,9 @@ void FPrimitiveSceneInfo::UnlinkAttachmentGroup()
 		FAttachmentGroupSceneInfo& AttachmentGroup = Scene->AttachmentGroups.FindChecked(LightingAttachmentRoot);
 		AttachmentGroup.Primitives.RemoveSwap(this);
 
-		if (AttachmentGroup.Primitives.Num() == 0)
+		if (AttachmentGroup.Primitives.Num() == 0 && AttachmentGroup.ParentSceneInfo == nullptr)
 		{
-			// If this was the last primitive attached that uses this attachment root, free the group.
+			// If this was the last primitive attached that uses this attachment group and the root has left the building, free the group.
 			Scene->AttachmentGroups.Remove(LightingAttachmentRoot);
 		}
 	}
@@ -908,6 +908,11 @@ void FPrimitiveSceneInfo::UnlinkAttachmentGroup()
 		if (AttachmentGroup)
 		{
 			AttachmentGroup->ParentSceneInfo = NULL;
+			if (AttachmentGroup->Primitives.Num() == 0)
+			{
+				// If this was the owner and the group is empty, remove it (otherwise the above will remove when the last attached goes).
+				Scene->AttachmentGroups.Remove(LightingAttachmentRoot);
+			}
 		}
 	}
 }
