Fix for blueprint property arrays of components not being reset when rerunning construction scripts

diff --git a/Engine/Source/Runtime/Engine/Private/ActorConstruction.cpp b/Engine/Source/Runtime/Engine/Private/ActorConstruction.cpp
--- a/Engine/Source/Runtime/Engine/Private/ActorConstruction.cpp
+++ b/Engine/Source/Runtime/Engine/Private/ActorConstruction.cpp
@@ -70,6 +70,16 @@ void AActor::SeedAllRandomStreams()
 }
 #endif //WITH_EDITOR
 
+bool IsBlueprintAddedContainer(UProperty* Prop)
+{
+	if (Prop->IsA<UArrayProperty>() || Prop->IsA<USetProperty>() || Prop->IsA<UMapProperty>())
+	{
+		return Prop->IsInBlueprint();
+	}
+
+	return false;
+}
+
 void AActor::ResetPropertiesForConstruction()
 {
 	// Get class CDO
@@ -96,7 +106,7 @@ void AActor::ResetPropertiesForConstruction()
 			StreamPtr->Reset();
 		}
 		// If it is a blueprint exposed variable that is not editable per-instance, reset to default before running construction script
-		else if (!bIsLevelScriptActor && !Prop->ContainsInstancedObjectProperty())
+		else if (!bIsLevelScriptActor && (!Prop->ContainsInstancedObjectProperty() || IsBlueprintAddedContainer(Prop)))
 		{
 			const bool bExposedOnSpawn = bIsPlayInEditor && Prop->HasAnyPropertyFlags(CPF_ExposeOnSpawn);
 			const bool bCanEditInstanceValue = !Prop->HasAnyPropertyFlags(CPF_DisableEditOnInstance) && Prop->HasAnyPropertyFlags(CPF_Edit);
