NetConnection will now sample packet loss on every StatPeriod and store the current value along with a rolling average in the InPacketsLossPercentage and OutPacketsLossPercentage variables

diff --git a/Engine/Source/Runtime/Engine/Classes/Engine/NetConnection.h b/Engine/Source/Runtime/Engine/Classes/Engine/NetConnection.h
--- a/Engine/Source/Runtime/Engine/Classes/Engine/NetConnection.h
+++ b/Engine/Source/Runtime/Engine/Classes/Engine/NetConnection.h
@@ -396,16 +396,32 @@ public:
 	double			StatUpdateTime;
 	/** Interval between gathering stats */
 	float			StatPeriod;
-	float			BestLag,   AvgLag;		// Lag.
 
-	// Stat accumulators.
-	double			LagAcc, BestLagAcc;		// Previous msec lag.
-	int32			LagCount;				// Counter for lag measurement.
-	double			LastTime, FrameTime;	// Monitors frame time.
-	/** @todo document */
-	double			CumulativeTime, AverageFrameTime; 
-	/** @todo document */
+	/** Average lag seen during the last StatPeriod */
+	float AvgLag;
+
+	//UE_DEPRECATED(4.24, "BestLag variable is deprecated. Use AvgLag instead");
+	float			BestLag;
+	//UE_DEPRECATED(4.24, "BestLagAcc variable is deprecated. Use LagAcc instead");
+	double			BestLagAcc;
+
+	/** Total accumulated lag values during the current StatPeriod */
+	double			LagAcc;
+	/** Nb of stats accumulated in LagAcc */
+	int32			LagCount;
+	
+	/** Monitors frame time */
+	double			LastTime, FrameTime;
+
+	/** Total frames times accumulator */
+	double			CumulativeTime;
+	
+	/** The average frame delta time over the last 1 second period.*/
+	double			AverageFrameTime;
+
+	/** Nb of stats accumulated in CumulativeTime */
 	int32			CountedFrames;
+
 	/** bytes sent/received on this connection (accumulated during a StatPeriod) */
 	int32 InBytes, OutBytes;
 	/** total bytes sent/received on this connection */
@@ -425,6 +441,22 @@ public:
 	/** total acks sent on this connection */
 	int32 OutTotalAcks;
 
+	/** Percentage of packets lost during the last StatPeriod */
+	using FNetConnectionPacketLoss = TPacketLossData<3>;
+	const FNetConnectionPacketLoss& GetInLossPercentage() const { return InPacketsLossPercentage; }
+	const FNetConnectionPacketLoss& GetOutLossPercentage() const { return OutPacketsLossPercentage;  }
+
+
+private:
+	FNetConnectionPacketLoss InPacketsLossPercentage;
+	FNetConnectionPacketLoss OutPacketsLossPercentage;
+
+public:
+
+	/** Counts the number of stat samples taken */
+	int32 StatPeriodCount;
+
+
 	/** Net Analytics */
 
 	/** The locally cached/updated analytics variables, for the NetConnection - aggregated upon connection Close */
diff --git a/Engine/Source/Runtime/Engine/Private/NetConnection.cpp b/Engine/Source/Runtime/Engine/Private/NetConnection.cpp
--- a/Engine/Source/Runtime/Engine/Private/NetConnection.cpp
+++ b/Engine/Source/Runtime/Engine/Private/NetConnection.cpp
@@ -175,11 +175,9 @@ UNetConnection::UNetConnection(const FObjectInitializer& ObjectInitializer)
 ,	SendBunchHeader		( MAX_BUNCH_HEADER_BITS )
 
 ,	StatPeriod			( 1.f  )
-,	BestLag				( 9999 )
 ,	AvgLag				( 9999 )
 
 ,	LagAcc				( 9999 )
-,	BestLagAcc			( 9999 )
 ,	LagCount			( 0 )
 ,	LastTime			( 0 )
 ,	FrameTime			( 0 )
@@ -201,6 +199,7 @@ UNetConnection::UNetConnection(const FObjectInitializer& ObjectInitializer)
 ,	InTotalPacketsLost	( 0 )
 ,	OutTotalPacketsLost	( 0 )
 ,	OutTotalAcks		( 0 )
+,	StatPeriodCount		( 0 )
 ,	AnalyticsVars		()
 ,	NetAnalyticsData	()
 ,	SendBuffer			( 0 )
@@ -2943,17 +2942,21 @@ void UNetConnection::Tick()
 		{
 			AvgLag = LagAcc/LagCount;
 		}
-		BestLag = AvgLag;
 
 		InBytesPerSecond = FMath::TruncToInt(static_cast<float>(InBytes) / RealTime);
 		OutBytesPerSecond = FMath::TruncToInt(static_cast<float>(OutBytes) / RealTime);
 		InPacketsPerSecond = FMath::TruncToInt(static_cast<float>(InPackets) / RealTime);
 		OutPacketsPerSecond = FMath::TruncToInt(static_cast<float>(OutPackets) / RealTime);
 
+		// Add TotalPacketsLost to total since InTotalPackets only counts ACK packets
+		InPacketsLossPercentage.UpdateLoss(InPacketsLost, InTotalPackets + InTotalPacketsLost, StatPeriodCount);
+
+		// Using OutTotalNotifiedPackets so we do not count packets that are still in transit.
+		OutPacketsLossPercentage.UpdateLoss(OutPacketsLost, OutTotalNotifiedPackets, StatPeriodCount);
+
 		// Init counters.
 		LagAcc = 0;
 		StatUpdateTime = CurrentRealtimeSeconds;
-		BestLagAcc = 9999;
 		LagCount = 0;
 		InPacketsLost = 0;
 		OutPacketsLost = 0;
@@ -2961,6 +2964,8 @@ void UNetConnection::Tick()
 		OutBytes = 0;
 		InPackets = 0;
 		OutPackets = 0;
+
+		++StatPeriodCount;
 	}
 
 	if (bConnectionPendingCloseDueToSocketSendFailure)
diff --git a/Engine/Source/Runtime/Engine/Public/Net/NetAnalyticsTypes.h b/Engine/Source/Runtime/Engine/Public/Net/NetAnalyticsTypes.h
--- a/Engine/Source/Runtime/Engine/Public/Net/NetAnalyticsTypes.h
+++ b/Engine/Source/Runtime/Engine/Public/Net/NetAnalyticsTypes.h
@@ -376,4 +376,55 @@ private:
 	uint32 NumberOfFramesWithBurstsOfPacketLoss;
 	uint32 NumberOfFramesWithNoPackets;
 	uint64 NumberOfTrackedFrames;
+};
+
+
+/**
+* Calculates the current packet loss along with a rolling average of the past X updates
+* Constructed with the number of samples you want to average over time.
+*/
+template<uint32 NbPeriodsForAverage>
+struct TPacketLossData
+{
+public:
+
+	/** The loss percentage over the last StatPeriod */
+	float GetLossPercentage() const { return LossPercentage; }
+
+	/** The average loss percentage over the previous X StatPeriods */
+	float GetAvgLossPercentage() const { return AvgLossPercentage; }
+
+	/** Update the packet loss based on total and lost packets */
+	void UpdateLoss(int32 NbPacketsLost, int32 TotalPackets, int32 SampleCount)
+	{
+		// Update the current statistic
+		const int32 PacketsDuringLastPeriod = TotalPackets - TotalPacketsAtPeriodStart;
+		LossPercentage = PacketsDuringLastPeriod > 0 ? ((float)NbPacketsLost / (float)PacketsDuringLastPeriod) : 0.0f;
+
+		// Update the rolling average
+		const int32 SampleIndex = SampleCount % NbPeriodsForAverage;
+		LossSamples[SampleIndex] = LossPercentage;
+
+		float RollingAverage = 0.0f;
+		for (float LossSample : LossSamples)
+		{
+			RollingAverage += LossSample;
+		}
+		AvgLossPercentage = RollingAverage / NbPeriodsForAverage;
+
+		// Store data for the next update
+		TotalPacketsAtPeriodStart = TotalPackets;
+	}
+
+private:
+	/**
+	* The amount of samples to average.
+	* The length of the period will be NbPeriodsForAverage * StatPeriod.
+	* With default values that is 3 samples * 1 second -> an average of the past 3 seconds
+	*/
+	float LossSamples[NbPeriodsForAverage] = { 0.0f };
+
+	int32 TotalPacketsAtPeriodStart = 0;
+	float LossPercentage = 0.0f;
+	float AvgLossPercentage = 0.0f;
 };
\ No newline at end of file
