Swap IsRigidBody and IsDynamic Chaos implementation to properly match PhysX behaviour.

diff --git a/Engine/Source/Runtime/Engine/Private/PhysicsEngine/Experimental/PhysInterface_Chaos.cpp b/Engine/Source/Runtime/Engine/Private/PhysicsEngine/Experimental/PhysInterface_Chaos.cpp
--- a/Engine/Source/Runtime/Engine/Private/PhysicsEngine/Experimental/PhysInterface_Chaos.cpp
+++ b/Engine/Source/Runtime/Engine/Private/PhysicsEngine/Experimental/PhysInterface_Chaos.cpp
@@ -203,6 +203,11 @@ void FPhysInterface_Chaos::SetActorUserData_AssumesLocked(FPhysicsActorHandle& I
 }
 
 bool FPhysInterface_Chaos::IsRigidBody(const FPhysicsActorHandle& InActorReference)
+{
+	return !IsStatic(InActorReference);
+}
+
+bool FPhysInterface_Chaos::IsDynamic(const FPhysicsActorHandle& InActorReference)
 {
 	return InActorReference->ObjectState() == Chaos::EObjectStateType::Dynamic;
 }
diff --git a/Engine/Source/Runtime/Engine/Public/Physics/Experimental/PhysInterface_Chaos.h b/Engine/Source/Runtime/Engine/Public/Physics/Experimental/PhysInterface_Chaos.h
--- a/Engine/Source/Runtime/Engine/Public/Physics/Experimental/PhysInterface_Chaos.h
+++ b/Engine/Source/Runtime/Engine/Public/Physics/Experimental/PhysInterface_Chaos.h
@@ -211,10 +211,7 @@ public:
 	static void SetActorUserData_AssumesLocked(FPhysicsActorHandle& InActorReference, FPhysicsUserData* InUserData);
 
 	static bool IsRigidBody(const FPhysicsActorHandle& InActorReference);
-	static bool IsDynamic(const FPhysicsActorHandle& InActorReference)
-    {
-        return !IsStatic(InActorReference);
-    }
+	static bool IsDynamic(const FPhysicsActorHandle& InActorReference);
     static bool IsStatic(const FPhysicsActorHandle& InActorReference);
     static bool IsKinematic(const FPhysicsActorHandle& InActorReference);
 	static bool IsKinematic_AssumesLocked(const FPhysicsActorHandle& InActorReference);
