Fixing crash when importing SkeletalMesh with more than 8 LOD.

Alexis.Matte

diff --git a/Engine/Source/Runtime/Engine/Classes/Engine/SkeletalMesh.h b/Engine/Source/Runtime/Engine/Classes/Engine/SkeletalMesh.h
--- a/Engine/Source/Runtime/Engine/Classes/Engine/SkeletalMesh.h
+++ b/Engine/Source/Runtime/Engine/Classes/Engine/SkeletalMesh.h
@@ -1335,7 +1335,16 @@ public:
 	/* 
 	 * Returns total number of LOD
 	 */
-	int32 GetLODNum() const { check(LODInfo.Num() <= MAX_MESH_LOD_COUNT); return LODInfo.Num();  }
+	int32 GetLODNum() const 
+	{ 
+#if WITH_EDITOR
+		if (bSupportLODStreaming.Default || bSupportLODStreaming.PerPlatform.FindKey(true))
+		{
+			check(LODInfo.Num() <= MAX_MESH_LOD_COUNT); 
+		}
+#endif
+		return LODInfo.Num();  
+	}
 
 public:
 	const TArray<FSkinWeightProfileInfo>& GetSkinWeightProfiles() const { return SkinWeightProfiles; }
