Fix CIS (backout shouldn't have been propogated)

diff --git a/Engine/Source/Runtime/Engine/Classes/Sound/SoundSubmixSend.h b/Engine/Source/Runtime/Engine/Classes/Sound/SoundSubmixSend.h
new file mode 100644
--- /dev/null
+++ b/Engine/Source/Runtime/Engine/Classes/Sound/SoundSubmixSend.h
@@ -0,0 +1,85 @@
+// Copyright 1998-2019 Epic Games, Inc. All Rights Reserved.
+#pragma once
+
+#include "CoreMinimal.h"
+
+#include "Curves/CurveFloat.h"
+#include "UObject/ObjectMacros.h"
+#include "UObject/Object.h"
+
+#include "SoundSubmixSend.generated.h"
+
+
+// Forward Declarations
+class USoundSubmix;
+
+
+DECLARE_DYNAMIC_DELEGATE_OneParam(FOnSubmixEnvelopeBP, const TArray<float>&, Envelope);
+
+UENUM(BlueprintType)
+enum class EAudioRecordingExportType : uint8
+{
+	// Exports a USoundWave.
+	SoundWave,
+
+	// Exports a WAV file.
+	WavFile
+};
+
+UENUM(BlueprintType)
+enum class ESendLevelControlMethod : uint8
+{
+	// A send based on linear interpolation between a distance range and send-level range
+	Linear,
+
+	// A send based on a supplied curve
+	CustomCurve,
+
+	// A manual send level (Uses the specified constant send level value. Useful for 2D sounds.)
+	Manual,
+};
+
+// Class used to send audio to submixes from USoundBase
+USTRUCT(BlueprintType)
+struct ENGINE_API FSoundSubmixSendInfo
+{
+	GENERATED_USTRUCT_BODY()
+
+	FSoundSubmixSendInfo();
+
+	/*
+		Manual: Use Send Level only
+		Linear: Interpolate between Min and Max Send Levels based on listener distance (between Distance Min and Distance Max)
+		Custom Curve: Use the float curve to map Send Level to distance (0.0-1.0 on curve maps to Distance Min - Distance Max)
+	*/
+	UPROPERTY(EditAnywhere, Category = SubmixSend)
+	ESendLevelControlMethod SendLevelControlMethod;
+
+	// The submix to send the audio to
+	UPROPERTY(EditAnywhere, Category = SubmixSend)
+	USoundSubmix* SoundSubmix;
+
+	// The amount of audio to send
+	UPROPERTY(EditAnywhere, Category = SubmixSend)
+	float SendLevel;
+
+	// The amount to send to master when sound is located at a distance equal to value specified in the min send distance.
+	UPROPERTY(EditAnywhere, BlueprintReadWrite, Category = SubmixSend)
+	float MinSendLevel;
+
+	// The amount to send to master when sound is located at a distance equal to value specified in the max send distance.
+	UPROPERTY(EditAnywhere, BlueprintReadWrite, Category = SubmixSend)
+	float MaxSendLevel;
+
+	// The min distance to send to the master
+	UPROPERTY(EditAnywhere, BlueprintReadWrite, Category = SubmixSend)
+	float MinSendDistance;
+
+	// The max distance to send to the master
+	UPROPERTY(EditAnywhere, BlueprintReadWrite, Category = SubmixSend)
+	float MaxSendDistance;
+
+	// The custom reverb send curve to use for distance-based send level.
+	UPROPERTY(EditAnywhere, BlueprintReadWrite, Category = SubmixSend)
+	FRuntimeFloatCurve CustomSendLevelCurve;
+};
diff --git a/Engine/Source/Runtime/Engine/Private/SoundSubmixSend.cpp b/Engine/Source/Runtime/Engine/Private/SoundSubmixSend.cpp
new file mode 100644
--- /dev/null
+++ b/Engine/Source/Runtime/Engine/Private/SoundSubmixSend.cpp
@@ -0,0 +1,15 @@
+// Copyright 1998-2019 Epic Games, Inc. All Rights Reserved.
+#include "Sound/SoundSubmixSend.h"
+
+#include "Sound/SoundSubmix.h"
+
+FSoundSubmixSendInfo::FSoundSubmixSendInfo()
+	: SendLevelControlMethod(ESendLevelControlMethod::Manual)
+	, SoundSubmix(nullptr)
+	, SendLevel(0.0f)
+	, MinSendLevel(0.0f)
+	, MaxSendLevel(1.0f)
+	, MinSendDistance(100.0f)
+	, MaxSendDistance(1000.0f)
+	{
+	}
