Ensure attributes are registered on MeshDescription when converting it from RawMesh.
Corrected deprecation version for RegisterMeshAttributes.

diff --git a/Engine/Source/Runtime/Engine/Classes/Engine/StaticMesh.h b/Engine/Source/Runtime/Engine/Classes/Engine/StaticMesh.h
--- a/Engine/Source/Runtime/Engine/Classes/Engine/StaticMesh.h
+++ b/Engine/Source/Runtime/Engine/Classes/Engine/StaticMesh.h
@@ -781,7 +781,7 @@ public:
 	/**
 	 * Registers the mesh attributes required by the mesh description for a static mesh.
 	 */
-	UE_DEPRECATED(4.25, "Please use FStaticMeshAttributes::Register to do this.")
+	UE_DEPRECATED(4.24, "Please use FStaticMeshAttributes::Register to do this.")
 	ENGINE_API static void RegisterMeshAttributes( FMeshDescription& MeshDescription );
 
 #if WITH_EDITORONLY_DATA
diff --git a/Engine/Source/Runtime/Engine/Private/StaticMesh.cpp b/Engine/Source/Runtime/Engine/Private/StaticMesh.cpp
--- a/Engine/Source/Runtime/Engine/Private/StaticMesh.cpp
+++ b/Engine/Source/Runtime/Engine/Private/StaticMesh.cpp
@@ -3764,6 +3764,9 @@ bool UStaticMesh::LoadMeshDescription(int32 LodIndex, FMeshDescription& OutMeshD
 {
 	TRACE_CPUPROFILER_EVENT_SCOPE(UStaticMesh::LoadMeshDescription);
 
+	// Ensure MeshDescription is empty, with no attributes registered
+	OutMeshDescription = FMeshDescription();
+
 	const FStaticMeshSourceModel& SourceModel = GetSourceModel(LodIndex);
 
 	// If we don't have a valid MeshDescription, try and get one...
@@ -3807,6 +3810,11 @@ bool UStaticMesh::LoadMeshDescription(int32 LodIndex, FMeshDescription& OutMeshD
 		SourceModel.LoadRawMesh(LodRawMesh);
 		TMap<int32, FName> MaterialMap;
 		FillMaterialName(StaticMaterials, MaterialMap);
+
+		// Register static mesh attributes on the mesh description
+		FStaticMeshAttributes StaticMeshAttributes(OutMeshDescription);
+		StaticMeshAttributes.Register();
+
 		FMeshDescriptionOperations::ConvertFromRawMesh(LodRawMesh, OutMeshDescription, MaterialMap);
 		return true;
 	}
