Required include fix for strict vs2019

diff --git a/Engine/Source/Runtime/Engine/Classes/Sound/SoundSubmix.h b/Engine/Source/Runtime/Engine/Classes/Sound/SoundSubmix.h
--- a/Engine/Source/Runtime/Engine/Classes/Sound/SoundSubmix.h
+++ b/Engine/Source/Runtime/Engine/Classes/Sound/SoundSubmix.h
@@ -14,7 +14,6 @@
 
 
 // Forward Declarations
-class FOnSubmixEnvelopeBP;
 class UEdGraph;
 class USoundEffectSubmixPreset;
 class USoundSubmix;
