Fix CIS

diff --git a/Engine/Source/Runtime/Engine/Private/PhysicsEngine/Experimental/ChaosDerivedDataReader.cpp b/Engine/Source/Runtime/Engine/Private/PhysicsEngine/Experimental/ChaosDerivedDataReader.cpp
--- a/Engine/Source/Runtime/Engine/Private/PhysicsEngine/Experimental/ChaosDerivedDataReader.cpp
+++ b/Engine/Source/Runtime/Engine/Private/PhysicsEngine/Experimental/ChaosDerivedDataReader.cpp
@@ -3,6 +3,7 @@
 #include "Physics/Experimental/ChaosDerivedDataReader.h"
 #include "Chaos/ChaosArchive.h"
 #include "Chaos/TriangleMeshImplicitObject.h"
+#include "Chaos/Convex.h"
 
 template<typename T, int d>
 FChaosDerivedDataReader<T, d>::FChaosDerivedDataReader(FUntypedBulkData* InBulkData)
