Fix a serialization crash with some assets save during the development
#rnx

diff --git a/Engine/Source/Runtime/Engine/Private/Rendering/SkeletalMeshLODImporterData.cpp b/Engine/Source/Runtime/Engine/Private/Rendering/SkeletalMeshLODImporterData.cpp
--- a/Engine/Source/Runtime/Engine/Private/Rendering/SkeletalMeshLODImporterData.cpp
+++ b/Engine/Source/Runtime/Engine/Private/Rendering/SkeletalMeshLODImporterData.cpp
@@ -3,12 +3,16 @@
 #if WITH_EDITOR
 
 #include "Rendering/SkeletalMeshLODImporterData.h"
+#include "Logging/LogVerbosity.h"
+#include "Logging/LogMacros.h"
 #include "Serialization/BulkDataWriter.h"
 #include "Serialization/BulkDataReader.h"
 #include "Rendering/SkeletalMeshModel.h"
 #include "Engine/SkeletalMesh.h"
 #include "Factories/FbxSkeletalMeshImportData.h"
 
+DEFINE_LOG_CATEGORY_STATIC(LogSkeletalMeshLODImporterData, Log, All);
+
 void FSkeletalMeshImportData::CopyDataNeedByMorphTargetImport(FSkeletalMeshImportData& Other) const
 {
 	Other.Points = Points;
@@ -403,7 +407,16 @@ FArchive& operator<<(FArchive& Ar, FReductionSkeletalMeshData& ReductionSkeletal
 	Ar << Version;
 	Ar << LicenseeVersion;
 	ReductionSkeletalMeshData.BaseLODModel.Serialize(Ar, ReductionSkeletalMeshData.Owner, INDEX_NONE);
-	Ar << ReductionSkeletalMeshData.BaseLODMorphTargetData;
+	
+	if(Ar.IsLoading() && Ar.AtEnd())
+	{
+		//Hack to fix a serialization error, serialize the MorphTargetData only if there is some left space in the archive
+		UE_ASSET_LOG(LogSkeletalMeshLODImporterData, Display, ReductionSkeletalMeshData.Owner, TEXT("This skeletalMesh should be re-import to save some missing reduction source data."));
+	}
+	else
+	{
+		Ar << ReductionSkeletalMeshData.BaseLODMorphTargetData;
+	}
 	return Ar;
 }
 
diff --git a/Engine/Source/Runtime/Engine/Private/SkeletalMesh.cpp b/Engine/Source/Runtime/Engine/Private/SkeletalMesh.cpp
--- a/Engine/Source/Runtime/Engine/Private/SkeletalMesh.cpp
+++ b/Engine/Source/Runtime/Engine/Private/SkeletalMesh.cpp
@@ -474,6 +474,36 @@ bool USkeletalMesh::IsReductionActive(int32 LODIndex) const
 	IMeshReduction* ReductionModule = FModuleManager::Get().LoadModuleChecked<IMeshReductionManagerModule>("MeshReductionInterface").GetSkeletalMeshReductionInterface();
 	uint32 LODVertexNumber = MAX_uint32;
 	uint32 LODTriNumber = MAX_uint32;
+	const FSkeletalMeshLODInfo* LODInfoPtr = GetLODInfo(LODIndex);
+	bool bLODHasBeenSimplified = LODInfoPtr && LODInfoPtr->bHasBeenSimplified;
+	if (GetImportedModel() && GetImportedModel()->LODModels.IsValidIndex(LODIndex))
+	{
+		if (!bLODHasBeenSimplified)
+		{
+			LODVertexNumber = 0;
+			LODTriNumber = 0;
+			const FSkeletalMeshLODModel& LODModel = GetImportedModel()->LODModels[LODIndex];
+			//We can take the vertices and triangles count from the source model
+			for (int32 SectionIndex = 0; SectionIndex < LODModel.Sections.Num(); ++SectionIndex)
+			{
+				const FSkelMeshSection& Section = LODModel.Sections[SectionIndex];
+
+				if (!Section.bDisabled)
+				{
+					//Make sure the count fit in a uint32
+					LODVertexNumber += Section.NumVertices < 0 ? 0 : Section.NumVertices;
+					LODTriNumber += Section.NumTriangles;
+				}
+			}
+		}
+		else if (GetImportedModel()->OriginalReductionSourceMeshData.IsValidIndex(LODIndex)
+			&& !GetImportedModel()->OriginalReductionSourceMeshData[LODIndex]->IsEmpty())
+		{
+			//In this case we have to use the stored reduction source data to know how many vertices/triangles we have before the reduction
+			USkeletalMesh* MutableSkeletalMesh = const_cast<USkeletalMesh*>(this);
+			GetImportedModel()->OriginalReductionSourceMeshData[LODIndex]->GetGeometryInfo(LODVertexNumber, LODTriNumber, MutableSkeletalMesh);
+		}
+	}
 	return ReductionModule->IsReductionActive(ReductionSettings, LODVertexNumber, LODTriNumber);
 }
 
