Fixed issue where render thread could read the world scene pointer after the game thread has nulled it.

diff --git a/Engine/Source/Runtime/Engine/Private/Components/ActorComponent.cpp b/Engine/Source/Runtime/Engine/Private/Components/ActorComponent.cpp
--- a/Engine/Source/Runtime/Engine/Private/Components/ActorComponent.cpp
+++ b/Engine/Source/Runtime/Engine/Private/Components/ActorComponent.cpp
@@ -77,20 +77,19 @@ int32 FGlobalComponentReregisterContext::ActiveGlobalReregisterContextCount = 0;
 
 void UpdateAllPrimitiveSceneInfosForSingleComponent(UActorComponent* InComponent, TSet<FSceneInterface*>* InScenesToUpdateAllPrimitiveSceneInfosForBatching /* = nullptr*/)
 {
-	if (InScenesToUpdateAllPrimitiveSceneInfosForBatching == nullptr)
+	if (FSceneInterface* Scene = InComponent->GetScene())
 	{
-		// If no batching is available (this ComponentReregisterContext is not created by a FGlobalComponentReregisterContext), issue one update per component
-		ENQUEUE_RENDER_COMMAND(UpdateAllPrimitiveSceneInfosCmd)([InComponent](FRHICommandListImmediate& RHICmdList) {
-			if (InComponent->GetScene())
-				InComponent->GetScene()->UpdateAllPrimitiveSceneInfos(RHICmdList);
+		if (InScenesToUpdateAllPrimitiveSceneInfosForBatching == nullptr)
+		{
+			// If no batching is available (this ComponentReregisterContext is not created by a FGlobalComponentReregisterContext), issue one update per component
+			ENQUEUE_RENDER_COMMAND(UpdateAllPrimitiveSceneInfosCmd)([Scene](FRHICommandListImmediate& RHICmdList) {
+				Scene->UpdateAllPrimitiveSceneInfos(RHICmdList);
 			});
-	}
-	else
-	{
-		if (InComponent->GetScene())
+		}
+		else
 		{
 			// Try to batch the updates inside FGlobalComponentReregisterContext
-			InScenesToUpdateAllPrimitiveSceneInfosForBatching->Add(InComponent->GetScene());
+			InScenesToUpdateAllPrimitiveSceneInfosForBatching->Add(Scene);
 		}
 	}
 }
