Variable hiding warning fix

diff --git a/Engine/Source/Runtime/Engine/Classes/Sound/SoundSubmix.h b/Engine/Source/Runtime/Engine/Classes/Sound/SoundSubmix.h
--- a/Engine/Source/Runtime/Engine/Classes/Sound/SoundSubmix.h
+++ b/Engine/Source/Runtime/Engine/Classes/Sound/SoundSubmix.h
@@ -85,11 +85,9 @@ class ENGINE_API USoundSubmix : public UObject
 	UPROPERTY(VisibleAnywhere, BlueprintReadOnly, Category = SoundSubmix)
 	TArray<USoundSubmix*> ChildSubmixes;
 
-public:
 	UPROPERTY()
 	USoundSubmix* ParentSubmix;
 
-private:
 #if WITH_EDITORONLY_DATA
 	/** EdGraph based representation of the SoundSubmix */
 	UEdGraph* SoundSubmixGraph;
