Fix comments

diff --git a/Engine/Source/Runtime/Engine/Classes/Engine/EngineTypes.h b/Engine/Source/Runtime/Engine/Classes/Engine/EngineTypes.h
--- a/Engine/Source/Runtime/Engine/Classes/Engine/EngineTypes.h
+++ b/Engine/Source/Runtime/Engine/Classes/Engine/EngineTypes.h
@@ -1981,7 +1981,7 @@ struct ENGINE_API FHitResult
 
 	/**
 	 * Normal of the hit in world space, for the object that was hit by the sweep, if any.
-	 * For example if a box hits a flat plane, this is a normalized vector pointing out from the plane.
+	 * For example if a sphere hits a flat plane, this is a normalized vector pointing out from the plane.
 	 * In the case of impact with a corner or edge of a surface, usually the "most opposing" normal (opposed to the query direction) is chosen.
 	 */
 	UPROPERTY()
