Remove the build settings from the SkeletalMeshLODSettings
#rnx

diff --git a/Engine/Source/Editor/Persona/Private/PersonaMeshDetails.cpp b/Engine/Source/Editor/Persona/Private/PersonaMeshDetails.cpp
--- a/Engine/Source/Editor/Persona/Private/PersonaMeshDetails.cpp
+++ b/Engine/Source/Editor/Persona/Private/PersonaMeshDetails.cpp
@@ -2330,7 +2330,7 @@ void FPersonaMeshDetails::AddLODLevelCategories(IDetailLayoutBuilder& DetailLayo
 					//Create the build setting UI Layout
 					BuildSettingsWidgetsPerLOD.Add(LODIndex, MakeShareable(new FSkeletalMeshBuildSettingsLayout(SkelMesh->GetLODInfo(LODIndex)->BuildSettings
 						, LODIndex
-						, FIsLODSettingsEnabledDelegate::CreateSP(this, &FPersonaMeshDetails::IsLODInfoEditingEnabled)
+						, FIsLODSettingsEnabledDelegate::CreateLambda([](int32 InLODIndex)->bool {return true;})
 						, FModifyMeshLODSettingsDelegate::CreateSP(this, &FPersonaMeshDetails::ModifyMeshLODSettings))));
 
 					LODCategory.AddCustomBuilder(BuildSettingsWidgetsPerLOD[LODIndex].ToSharedRef());
diff --git a/Engine/Source/Runtime/Engine/Classes/Engine/SkeletalMeshLODSettings.h b/Engine/Source/Runtime/Engine/Classes/Engine/SkeletalMeshLODSettings.h
--- a/Engine/Source/Runtime/Engine/Classes/Engine/SkeletalMeshLODSettings.h
+++ b/Engine/Source/Runtime/Engine/Classes/Engine/SkeletalMeshLODSettings.h
@@ -64,9 +64,6 @@ struct FSkeletalMeshLODGroupSettings
 	/** Get Skeletal mesh optimizations setting structure for the given LOD level */
 	ENGINE_API FSkeletalMeshOptimizationSettings GetReductionSettings() const;
 
-	/** Get Skeletal mesh build setting structure for the given LOD level */
-	ENGINE_API FSkeletalMeshBuildSettings GetBuildSettings() const;
-
 	/** Get Skeletal mesh optimizations setting structure for the given LOD level */
 	UE_DEPRECATED(4.20, "Please use GetReductionSettings instead")
 	FSkeletalMeshOptimizationSettings GetSettings() const
@@ -111,10 +108,6 @@ struct FSkeletalMeshLODGroupSettings
 	/** The optimization settings to use for the respective LOD level */
 	UPROPERTY(EditAnywhere, Category = Reduction)
 	FSkeletalMeshOptimizationSettings ReductionSettings;
-
-	/** The build settings to use for the respective LOD level */
-	UPROPERTY(EditAnywhere, Category = Build)
-	FSkeletalMeshBuildSettings BuildSettings;
 };
 
 UCLASS(config = Engine, defaultconfig, BlueprintType, MinimalAPI)
diff --git a/Engine/Source/Runtime/Engine/Private/SkeletalMesh.cpp b/Engine/Source/Runtime/Engine/Private/SkeletalMesh.cpp
--- a/Engine/Source/Runtime/Engine/Private/SkeletalMesh.cpp
+++ b/Engine/Source/Runtime/Engine/Private/SkeletalMesh.cpp
@@ -3847,19 +3847,17 @@ FGuid FSkeletalMeshLODInfo::ComputeDeriveDataCacheKey(const FSkeletalMeshLODGrou
 	FArchive_Serialize_BitfieldBool(Ar, bSupportUniformlyDistributedSampling);
 
 	//Use the LOD settings asset if there is one
-	FSkeletalMeshBuildSettings RealBuildSettings = BuildSettings;
 	FSkeletalMeshOptimizationSettings RealReductionSettings = ReductionSettings;
 	if (SkeletalMeshLODGroupSettings != nullptr)
 	{
-		RealBuildSettings = SkeletalMeshLODGroupSettings->GetBuildSettings();
 		RealReductionSettings = SkeletalMeshLODGroupSettings->GetReductionSettings();
 	}
 
-	if (!RealBuildSettings.bUseFullPrecisionUVs && !bIs16BitfloatBufferSupported)
+	if (!BuildSettings.bUseFullPrecisionUVs && !bIs16BitfloatBufferSupported)
 	{
-		RealBuildSettings.bUseFullPrecisionUVs = true;
+		BuildSettings.bUseFullPrecisionUVs = true;
 	}
-	SerializeBuildSettingsForDDC(Ar, RealBuildSettings);
+	SerializeBuildSettingsForDDC(Ar, BuildSettings);
 	SerializeReductionSettingsForDDC(Ar, RealReductionSettings);
 
 	FSHA1 Sha;
diff --git a/Engine/Source/Runtime/Engine/Private/SkeletalMeshLODSettings.cpp b/Engine/Source/Runtime/Engine/Private/SkeletalMeshLODSettings.cpp
--- a/Engine/Source/Runtime/Engine/Private/SkeletalMeshLODSettings.cpp
+++ b/Engine/Source/Runtime/Engine/Private/SkeletalMeshLODSettings.cpp
@@ -48,7 +48,6 @@ bool USkeletalMeshLODSettings::SetLODSettingsToMesh(USkeletalMesh* InMesh, int32
 		FSkeletalMeshLODInfo* LODInfo = InMesh->GetLODInfo(LODIndex);
 		const FSkeletalMeshLODGroupSettings& Setting = LODGroups[LODIndex];
 		LODInfo->ReductionSettings = Setting.ReductionSettings;
-		LODInfo->BuildSettings = Setting.BuildSettings;
 		LODInfo->ScreenSize = Setting.ScreenSize;
 		LODInfo->LODHysteresis = Setting.LODHysteresis;
 		LODInfo->WeightOfPrioritization = Setting.WeightOfPrioritization;
@@ -208,7 +207,6 @@ int32 USkeletalMeshLODSettings::SetLODSettingsFromMesh(USkeletalMesh* InMesh)
 			FSkeletalMeshLODInfo* LODInfo = InMesh->GetLODInfo(Index);
 			FSkeletalMeshLODGroupSettings& Setting = LODGroups[Index];
 			Setting.ReductionSettings = LODInfo->ReductionSettings;
-			Setting.BuildSettings = LODInfo->BuildSettings;
 			Setting.ScreenSize = LODInfo->ScreenSize;
 			Setting.LODHysteresis = LODInfo->LODHysteresis;
 			// copy mesh setting to shared setting
@@ -279,27 +277,6 @@ void USkeletalMeshLODSettings::Serialize(FArchive& Ar)
 			}
 		}
 	}
-
-	if (Ar.CustomVer(FEditorObjectVersion::GUID) < FEditorObjectVersion::SkeletalMeshBuildRefactor)
-	{
-		//Since we use the bake source model when loading old asset we do not want to change anything
-		for (int32 Index = 0; Index < LODGroups.Num(); ++Index)
-		{
-			FSkeletalMeshBuildSettings& BuildSettings = LODGroups[Index].BuildSettings;
-			BuildSettings.bRecomputeNormals = false;
-			BuildSettings.bRecomputeTangents = false;
-			BuildSettings.bUseMikkTSpace = false;
-			BuildSettings.bRemoveDegenerates = false;
-			BuildSettings.ThresholdPosition = 0.0f;
-			BuildSettings.ThresholdTangentNormal = 0.0f;
-			BuildSettings.ThresholdUV = 0.0f;
-
-			//The user will have to make sure those settings are ok since before the modification precisions
-			//were one value for the whole skeletal mesh LODs. We set them to false by default
-			BuildSettings.bUseFullPrecisionUVs = GVertexElementTypeSupport.IsSupported(VET_Half2) ? false : true;
-			BuildSettings.bUseHighPrecisionTangentBasis = false;
-		}
-	}
 }
 
 /////////////////////////////////////////////////////////////
@@ -310,14 +287,6 @@ FSkeletalMeshOptimizationSettings FSkeletalMeshLODGroupSettings::GetReductionSet
 	return ReductionSettings;
 }
 
-/////////////////////////////////////////////////////////////
-// FSkeletalMeshBuildSettings 
-/////////////////////////////////////////////////////////////
-FSkeletalMeshBuildSettings FSkeletalMeshLODGroupSettings::GetBuildSettings() const
-{
-	return BuildSettings;
-}
-
 const float FSkeletalMeshLODGroupSettings::GetScreenSize() const
 {
 	return ScreenSize.Default;
diff --git a/Engine/Source/Runtime/Engine/Private/SkeletalMeshRenderData.cpp b/Engine/Source/Runtime/Engine/Private/SkeletalMeshRenderData.cpp
--- a/Engine/Source/Runtime/Engine/Private/SkeletalMeshRenderData.cpp
+++ b/Engine/Source/Runtime/Engine/Private/SkeletalMeshRenderData.cpp
@@ -232,17 +232,6 @@ void FSkeletalMeshRenderData::Cache(USkeletalMesh* Owner)
 					bool bUseFullPrecisionUVs = LODInfo->BuildSettings.bUseFullPrecisionUVs;
 					bool bUseHighPrecisionTangentBasis = LODInfo->BuildSettings.bUseHighPrecisionTangentBasis;
 					bool bBuildAdjacencyBuffer = LODInfo->BuildSettings.bBuildAdjacencyBuffer;
-					if (Owner->LODSettings != nullptr)
-					{
-						const int32 NumSettings = Owner->LODSettings->GetNumberOfSettings();
-						if (LODIndex < NumSettings)
-						{
-							const FSkeletalMeshLODGroupSettings& SkeletalMeshLODGroupSettings = Owner->LODSettings->GetSettingsForLODLevel(LODIndex);
-							bUseFullPrecisionUVs = SkeletalMeshLODGroupSettings.GetBuildSettings().bUseFullPrecisionUVs;
-							bUseHighPrecisionTangentBasis = SkeletalMeshLODGroupSettings.GetBuildSettings().bUseHighPrecisionTangentBasis;
-							bBuildAdjacencyBuffer = SkeletalMeshLODGroupSettings.GetBuildSettings().bBuildAdjacencyBuffer;
-						}
-					}
 					if (bUseFullPrecisionUVs || !GVertexElementTypeSupport.IsSupported(VET_Half2))
 					{
 						VertexBufferBuildFlags |= ESkeletalMeshVertexFlags::UseFullPrecisionUVs;
