Merged 'Fixed mesh particle rendering artifact in stereo VR.'

diff --git a/Engine/Source/Runtime/Engine/Private/Particles/ParticleSystemRender.cpp b/Engine/Source/Runtime/Engine/Private/Particles/ParticleSystemRender.cpp
--- a/Engine/Source/Runtime/Engine/Private/Particles/ParticleSystemRender.cpp
+++ b/Engine/Source/Runtime/Engine/Private/Particles/ParticleSystemRender.cpp
@@ -2446,11 +2446,11 @@ void FDynamicMeshEmitterData::GetInstanceData(void* InstanceData, void* DynamicP
 		const FBaseParticle& Particle		= *((const FBaseParticle*) ParticleBase);
 
 
-		FMeshParticleInstanceVertex* CurrentInstanceVertex = (FMeshParticleInstanceVertex*)TempVert;
+		FMeshParticleInstanceVertex CurrentInstanceVertex;
 		
 		// Populate instance buffer;
 		// The particle color.
-		CurrentInstanceVertex->Color = Particle.Color;
+		CurrentInstanceVertex.Color = Particle.Color;
 		
 		// Instance to world transformation. Translation (Instance world position) is packed into W
 		FMatrix TransMat(FMatrix::Identity);
@@ -2458,9 +2458,9 @@ void FDynamicMeshEmitterData::GetInstanceData(void* InstanceData, void* DynamicP
 		
 		// Transpose on CPU to allow for simpler shader code to perform the transform. 
 		const FMatrix Transpose = TransMat.GetTransposed();
-		CurrentInstanceVertex->Transform[0] = FVector4(Transpose.M[0][0], Transpose.M[0][1], Transpose.M[0][2], Transpose.M[0][3]);
-		CurrentInstanceVertex->Transform[1] = FVector4(Transpose.M[1][0], Transpose.M[1][1], Transpose.M[1][2], Transpose.M[1][3]);
-		CurrentInstanceVertex->Transform[2] = FVector4(Transpose.M[2][0], Transpose.M[2][1], Transpose.M[2][2], Transpose.M[2][3]);
+		CurrentInstanceVertex.Transform[0] = FVector4(Transpose.M[0][0], Transpose.M[0][1], Transpose.M[0][2], Transpose.M[0][3]);
+		CurrentInstanceVertex.Transform[1] = FVector4(Transpose.M[1][0], Transpose.M[1][1], Transpose.M[1][2], Transpose.M[1][3]);
+		CurrentInstanceVertex.Transform[2] = FVector4(Transpose.M[2][0], Transpose.M[2][1], Transpose.M[2][2], Transpose.M[2][3]);
 
 		if (bUseStaticMeshLODs)
 		{
@@ -2495,9 +2495,9 @@ void FDynamicMeshEmitterData::GetInstanceData(void* InstanceData, void* DynamicP
 			}
 			else
 			{
-				PrevTransformVertex[0] = CurrentInstanceVertex->Transform[0];
-				PrevTransformVertex[1] = CurrentInstanceVertex->Transform[1];
-				PrevTransformVertex[2] = CurrentInstanceVertex->Transform[2];
+				PrevTransformVertex[0] = CurrentInstanceVertex.Transform[0];
+				PrevTransformVertex[1] = CurrentInstanceVertex.Transform[1];
+				PrevTransformVertex[2] = CurrentInstanceVertex.Transform[2];
 			}
 
 			TempPrevTranformVert += PrevTransformVertexStride;
@@ -2525,11 +2525,11 @@ void FDynamicMeshEmitterData::GetInstanceData(void* InstanceData, void* DynamicP
 			DeltaPosition.ToDirectionAndLength(Direction, Speed);
 
 			// Pack direction and speed.
-			CurrentInstanceVertex->Velocity = FVector4(Direction, Speed);
+			CurrentInstanceVertex.Velocity = FVector4(Direction, Speed);
 		}
 		else
 		{
-			CurrentInstanceVertex->Velocity = FVector4();
+			CurrentInstanceVertex.Velocity = FVector4();
 		}
 
 		// The particle dynamic value
@@ -2538,16 +2538,17 @@ void FDynamicMeshEmitterData::GetInstanceData(void* InstanceData, void* DynamicP
 			if (Source.DynamicParameterDataOffset > 0)
 			{
 				FVector4 DynamicParameterValue;
-				FMeshParticleInstanceVertexDynamicParameter* CurrentInstanceVertexDynParam = (FMeshParticleInstanceVertexDynamicParameter*)TempDynamicParameterVert;
-				GetDynamicValueFromPayload(Source.DynamicParameterDataOffset, Particle, DynamicParameterValue );
-				CurrentInstanceVertexDynParam->DynamicValue[0] = DynamicParameterValue.X;
-				CurrentInstanceVertexDynParam->DynamicValue[1] = DynamicParameterValue.Y;
-				CurrentInstanceVertexDynParam->DynamicValue[2] = DynamicParameterValue.Z;
-				CurrentInstanceVertexDynParam->DynamicValue[3] = DynamicParameterValue.W;
-
-				for (uint32 Factor = 1; Factor < InstanceFactor; Factor++)
+				FMeshParticleInstanceVertexDynamicParameter CurrentInstanceVertexDynParam;
+				GetDynamicValueFromPayload(Source.DynamicParameterDataOffset, Particle, DynamicParameterValue);
+				CurrentInstanceVertexDynParam.DynamicValue[0] = DynamicParameterValue.X;
+				CurrentInstanceVertexDynParam.DynamicValue[1] = DynamicParameterValue.Y;
+				CurrentInstanceVertexDynParam.DynamicValue[2] = DynamicParameterValue.Z;
+				CurrentInstanceVertexDynParam.DynamicValue[3] = DynamicParameterValue.W;
+
+				//@todo - refactor into instance step rate in the RHI
+				for (uint32 Factor = 0; Factor < InstanceFactor; Factor++)
 				{
-					FMemory::Memcpy(TempDynamicParameterVert + DynamicParameterVertexStride * Factor, TempDynamicParameterVert, DynamicParameterVertexStride);
+					FMemory::Memcpy(TempDynamicParameterVert + DynamicParameterVertexStride * Factor, &CurrentInstanceVertexDynParam, sizeof(FMeshParticleInstanceVertexDynamicParameter));
 				}
 
 				TempDynamicParameterVert += DynamicParameterVertexStride * InstanceFactor;
@@ -2569,20 +2570,20 @@ void FDynamicMeshEmitterData::GetInstanceData(void* InstanceData, void* DynamicP
 			int32 SubImageBV = SubImageB / SubImagesX;
 
 			// SubUV offsets and lerp value
-			CurrentInstanceVertex->SubUVParams[0] = SubImageAH;
-			CurrentInstanceVertex->SubUVParams[1] = SubImageAV;
-			CurrentInstanceVertex->SubUVParams[2] = SubImageBH;
-			CurrentInstanceVertex->SubUVParams[3] = SubImageBV;
-			CurrentInstanceVertex->SubUVLerp = SubImageLerp;
+			CurrentInstanceVertex.SubUVParams[0] = SubImageAH;
+			CurrentInstanceVertex.SubUVParams[1] = SubImageAV;
+			CurrentInstanceVertex.SubUVParams[2] = SubImageBH;
+			CurrentInstanceVertex.SubUVParams[3] = SubImageBV;
+			CurrentInstanceVertex.SubUVLerp = SubImageLerp;
 		}
 
 		// The particle's relative time
-		CurrentInstanceVertex->RelativeTime = Particle.RelativeTime;
+		CurrentInstanceVertex.RelativeTime = Particle.RelativeTime;
 
 		//@todo - refactor into instance step rate in the RHI
-		for (uint32 Factor = 1; Factor < InstanceFactor; Factor++)
+		for (uint32 Factor = 0; Factor < InstanceFactor; Factor++)
 		{
-			FMemory::Memcpy(TempVert + InstanceVertexStride * Factor, TempVert, InstanceVertexStride);
+			FMemory::Memcpy(TempVert + InstanceVertexStride * Factor, &CurrentInstanceVertex, InstanceVertexStride);
 		}
 
 		TempVert += InstanceVertexStride * InstanceFactor;
