Tessellated Materials Displayed as 0 in Shader Complexity ViewMode

diff --git a/Engine/Source/Runtime/Engine/Private/DebugViewModeMaterialProxy.cpp b/Engine/Source/Runtime/Engine/Private/DebugViewModeMaterialProxy.cpp
--- a/Engine/Source/Runtime/Engine/Private/DebugViewModeMaterialProxy.cpp
+++ b/Engine/Source/Runtime/Engine/Private/DebugViewModeMaterialProxy.cpp
@@ -15,13 +15,14 @@ DebugViewModeMaterialProxy.cpp : Contains definitions the debug view mode materi
 FDebugViewModeMaterialProxy::FDebugViewModeMaterialProxy(
 	UMaterialInterface* InMaterialInterface, 
 	EMaterialQualityLevel::Type QualityLevel, 
-	ERHIFeatureLevel::Type FeatureLevel, 
+	ERHIFeatureLevel::Type InFeatureLevel,
 	bool InSynchronousCompilation, 
 	EDebugViewShaderMode InDebugViewMode
 )
 	: FMaterial()
 	, MaterialInterface(InMaterialInterface)
 	, Material(nullptr)
+	, FeatureLevel(InFeatureLevel)
 	, Usage(EMaterialShaderMapUsage::DebugViewMode)
 	, DebugViewMode(InDebugViewMode)
 	, PixelShaderName(nullptr)
@@ -96,7 +97,7 @@ bool FDebugViewModeMaterialProxy::RequiresSynchronousCompilation() const
 	return bSynchronousCompilation;
 }
 
-const FMaterial& FDebugViewModeMaterialProxy::GetMaterialWithFallback(ERHIFeatureLevel::Type FeatureLevel, const FMaterialRenderProxy*& OutFallbackMaterialRenderProxy) const
+const FMaterial& FDebugViewModeMaterialProxy::GetMaterialWithFallback(ERHIFeatureLevel::Type InFeatureLevel, const FMaterialRenderProxy*& OutFallbackMaterialRenderProxy) const
 {
 	if (GetRenderingThreadShaderMap())
 	{
@@ -105,7 +106,7 @@ const FMaterial& FDebugViewModeMaterialProxy::GetMaterialWithFallback(ERHIFeatur
 	else
 	{
 		OutFallbackMaterialRenderProxy = UMaterial::GetDefaultMaterial(MD_Surface)->GetRenderProxy();
-		return OutFallbackMaterialRenderProxy->GetMaterialWithFallback(FeatureLevel, OutFallbackMaterialRenderProxy);
+		return OutFallbackMaterialRenderProxy->GetMaterialWithFallback(InFeatureLevel, OutFallbackMaterialRenderProxy);
 	}
 }
 
@@ -210,4 +211,28 @@ void FDebugViewModeMaterialProxy::GatherExpressionsForCustomInterpolators(TArray
 	}
 }
 
+enum EMaterialTessellationMode FDebugViewModeMaterialProxy::GetTessellationMode() const
+{
+	FMaterialResource* Resource = MaterialInterface->GetMaterialResource(FeatureLevel);
+	return Resource ? Resource->GetTessellationMode() : MTM_NoTessellation;
+}
+
+bool FDebugViewModeMaterialProxy::IsCrackFreeDisplacementEnabled() const
+{
+	FMaterialResource* Resource = MaterialInterface->GetMaterialResource(FeatureLevel);
+	return Resource ? Resource->IsCrackFreeDisplacementEnabled() : false;
+}
+
+bool FDebugViewModeMaterialProxy::IsAdaptiveTessellationEnabled() const
+{
+	FMaterialResource* Resource = MaterialInterface->GetMaterialResource(FeatureLevel);
+	return Resource ? Resource->IsAdaptiveTessellationEnabled() : false;
+}
+
+float FDebugViewModeMaterialProxy::GetMaxDisplacement() const
+{
+	FMaterialResource* Resource = MaterialInterface->GetMaterialResource(FeatureLevel);
+	return Resource ? Resource->GetMaxDisplacement() : 0.0f;
+}
+
 #endif // WITH_EDITORONLY_DATA
diff --git a/Engine/Source/Runtime/Engine/Private/DebugViewModeMaterialProxy.h b/Engine/Source/Runtime/Engine/Private/DebugViewModeMaterialProxy.h
--- a/Engine/Source/Runtime/Engine/Private/DebugViewModeMaterialProxy.h
+++ b/Engine/Source/Runtime/Engine/Private/DebugViewModeMaterialProxy.h
@@ -148,6 +148,10 @@ public:
 	virtual bool GetCastDynamicShadowAsMasked() const override;
 	virtual void GatherCustomOutputExpressions(TArray<class UMaterialExpressionCustomOutput*>& OutCustomOutputs) const override;
 	virtual void GatherExpressionsForCustomInterpolators(TArray<class UMaterialExpression*>& OutExpressions) const override;
+	virtual enum EMaterialTessellationMode GetTessellationMode() const override;
+	virtual bool IsCrackFreeDisplacementEnabled() const override;
+	virtual bool IsAdaptiveTessellationEnabled() const override;
+	virtual float GetMaxDisplacement() const override;
 
 	// Cached material usage.
 	virtual bool IsUsedWithSkeletalMesh() const override { return bIsUsedWithSkeletalMesh; }
@@ -171,7 +175,8 @@ private:
 
 	/** The material interface for this proxy */
 	UMaterialInterface* MaterialInterface;
-	UMaterial* Material;	
+	UMaterial* Material;
+	ERHIFeatureLevel::Type FeatureLevel;
 	TArray<UObject*> ReferencedTextures;
 	EMaterialShaderMapUsage::Type Usage;
 	EDebugViewShaderMode DebugViewMode;
