Fixing shader compiler XGE issue where time waited wasn't incremented causing it to wait infinitely.

#rnx

diff --git a/Engine/Source/Runtime/Engine/Private/ShaderCompiler/ShaderCompilerXGE.cpp b/Engine/Source/Runtime/Engine/Private/ShaderCompiler/ShaderCompilerXGE.cpp
--- a/Engine/Source/Runtime/Engine/Private/ShaderCompiler/ShaderCompilerXGE.cpp
+++ b/Engine/Source/Runtime/Engine/Private/ShaderCompiler/ShaderCompilerXGE.cpp
@@ -886,6 +886,7 @@ int32 FShaderCompileXGEThreadRunnable_InterceptionInterface::CompilingLoop()
 				{
 					UE_LOG(LogShaderCompilers, Warning, TEXT("Expected XGE output file '%s' exists but can't be opened for read, waiting 1 second to try again.."), *Task->OutputFilePath);
 					FPlatformProcess::Sleep(1.0f);
+					TimeWaited += 1.0f;
 					OutputFileAr = IFileManager::Get().CreateFileReader(*Task->OutputFilePath, FILEREAD_Silent);
 				}
 			}
