HLOD proxy meshes don't cast dynamic shadows
bCastDynamicShadow was forced off in ALODActor::PostLoad() since 4.19 as a workaround for another issue.

diff --git a/Engine/Source/Runtime/Engine/Private/LODActor.cpp b/Engine/Source/Runtime/Engine/Private/LODActor.cpp
--- a/Engine/Source/Runtime/Engine/Private/LODActor.cpp
+++ b/Engine/Source/Runtime/Engine/Private/LODActor.cpp
@@ -57,6 +57,12 @@ ENGINE_API TAutoConsoleVariable<FString> CVarHLODDistanceOverride(
 	TEXT("'r.HLOD.DistanceOverride 5000, 10000, 20000' would result in HLOD levels 0, 1 and 2 transitioning at 5000, 1000 and 20000 respectively."),
 	ECVF_Scalability);
 
+static TAutoConsoleVariable<int32> CVarHLODForceDisableCastDynamicShadow(
+	TEXT("r.HLOD.ForceDisableCastDynamicShadow"),
+	0,
+	TEXT("If non-zero, will set bCastDynamicShadow to false for all LODActors, regardless of the shadowing setting of their subactors."),
+	ECVF_ReadOnly);
+
 ENGINE_API TArray<float> ALODActor::HLODDistances;
 
 #if !(UE_BUILD_SHIPPING)
@@ -233,9 +239,14 @@ void ALODActor::PostLoad()
 {
 	Super::PostLoad();
 	SetComponentsMinDrawDistance(LODDrawDistance, false);
-	StaticMeshComponent->bCastDynamicShadow = false;	
 	UpdateRegistrationToMatchMaximumLODLevel();
 
+	// Force disabled dynamic shadow casting if requested from CVar
+	if (CVarHLODForceDisableCastDynamicShadow.GetValueOnAnyThread() != 0)
+	{
+		StaticMeshComponent->bCastDynamicShadow = false;
+	}
+
 #if WITH_EDITOR
 	if (bRequiresLODScreenSizeConversion)
 	{
