Fixing UAnimCurveCompressionCodec serialization crash when loading cooked content in the editor

diff --git a/Engine/Source/Runtime/Engine/Private/Animation/AnimCurveCompressionCodec.cpp b/Engine/Source/Runtime/Engine/Private/Animation/AnimCurveCompressionCodec.cpp
--- a/Engine/Source/Runtime/Engine/Private/Animation/AnimCurveCompressionCodec.cpp
+++ b/Engine/Source/Runtime/Engine/Private/Animation/AnimCurveCompressionCodec.cpp
@@ -29,7 +29,7 @@ void UAnimCurveCompressionCodec::Serialize(FArchive& Ar)
 {
 	Super::Serialize(Ar);
 
-	if (!Ar.IsCooking())
+	if (!Ar.IsCooking() && !(Ar.IsLoading() && GetOutermost()->bIsCookedForEditor))
 	{
 		Ar << InstanceGuid;
 	}
