Fix for deprecated encryption key path

diff --git a/Engine/Source/Runtime/Engine/Private/NetConnection.cpp b/Engine/Source/Runtime/Engine/Private/NetConnection.cpp
--- a/Engine/Source/Runtime/Engine/Private/NetConnection.cpp
+++ b/Engine/Source/Runtime/Engine/Private/NetConnection.cpp
@@ -3746,7 +3746,21 @@ void UNetConnection::SendChallengeControlMessage(const FEncryptionKeyResponse& R
 	{
 		if (Response.Response == EEncryptionResponse::Success)
 		{
-			EnableEncryptionServer(Response.EncryptionData);
+			// handle deprecated path where only the key is set
+			PRAGMA_DISABLE_DEPRECATION_WARNINGS
+			if ((Response.EncryptionKey.Num() > 0) && (Response.EncryptionData.Key.Num() == 0))
+			{
+				FEncryptionData ResponseData = Response.EncryptionData;
+				ResponseData.Key = Response.EncryptionKey;
+
+				EnableEncryptionServer(ResponseData);
+			}
+			PRAGMA_ENABLE_DEPRECATION_WARNINGS
+			else
+			{
+				EnableEncryptionServer(Response.EncryptionData);
+			}
+
 			SendChallengeControlMessage();
 		}
 		else
diff --git a/Engine/Source/Runtime/Engine/Private/PendingNetGame.cpp b/Engine/Source/Runtime/Engine/Private/PendingNetGame.cpp
--- a/Engine/Source/Runtime/Engine/Private/PendingNetGame.cpp
+++ b/Engine/Source/Runtime/Engine/Private/PendingNetGame.cpp
@@ -371,7 +371,20 @@ void UPendingNetGame::FinalizeEncryptedConnection(const FEncryptionKeyResponse&
 		{
 			if (Response.Response == EEncryptionResponse::Success)
 			{
-				Connection->EnableEncryption(Response.EncryptionData);
+				// handle deprecated path where only the key is set
+				PRAGMA_DISABLE_DEPRECATION_WARNINGS
+				if ((Response.EncryptionKey.Num() > 0) && (Response.EncryptionData.Key.Num() == 0))
+				{
+					FEncryptionData ResponseData = Response.EncryptionData;
+					ResponseData.Key = Response.EncryptionKey;
+
+					Connection->EnableEncryption(ResponseData);
+				}
+				PRAGMA_ENABLE_DEPRECATION_WARNINGS
+				else
+				{
+					Connection->EnableEncryption(Response.EncryptionData);
+				}
 			}
 			else
 			{
