Fixing non-unity compile error

diff --git a/Engine/Source/Runtime/Engine/Private/Animation/AnimCurveCompressionCodec.cpp b/Engine/Source/Runtime/Engine/Private/Animation/AnimCurveCompressionCodec.cpp
--- a/Engine/Source/Runtime/Engine/Private/Animation/AnimCurveCompressionCodec.cpp
+++ b/Engine/Source/Runtime/Engine/Private/Animation/AnimCurveCompressionCodec.cpp
@@ -1,6 +1,7 @@
 // Copyright 1998-2019 Epic Games, Inc. All Rights Reserved.
 
 #include "Animation/AnimCurveCompressionCodec.h"
+#include "UObject/Package.h"
 
 UAnimCurveCompressionCodec::UAnimCurveCompressionCodec(const FObjectInitializer& ObjectInitializer)
 	: Super(ObjectInitializer)
