Deprecated public replicated members of Actor, ActorComponent, and SceneComponent.

diff --git a/Engine/Plugins/FX/Niagara/Source/Niagara/Private/NiagaraModule.cpp b/Engine/Plugins/FX/Niagara/Source/Niagara/Private/NiagaraModule.cpp
--- a/Engine/Plugins/FX/Niagara/Source/Niagara/Private/NiagaraModule.cpp
+++ b/Engine/Plugins/FX/Niagara/Source/Niagara/Private/NiagaraModule.cpp
@@ -477,7 +477,7 @@ void INiagaraModule::OnChangeDetailLevel(class IConsoleVariable* CVar)
 			UNiagaraComponent* Comp = *It;
 			check(Comp);
 
-			const bool bWasActive = Comp->bIsActive;
+			const bool bWasActive = Comp->IsActive();
 			Comp->DestroyInstance();
 			if ( bWasActive )
 			{
diff --git a/Engine/Plugins/Runtime/NetworkPrediction/Source/NetworkPrediction/Private/NetworkedSimulationGlobalManager.cpp b/Engine/Plugins/Runtime/NetworkPrediction/Source/NetworkPrediction/Private/NetworkedSimulationGlobalManager.cpp
--- a/Engine/Plugins/Runtime/NetworkPrediction/Source/NetworkPrediction/Private/NetworkedSimulationGlobalManager.cpp
+++ b/Engine/Plugins/Runtime/NetworkPrediction/Source/NetworkPrediction/Private/NetworkedSimulationGlobalManager.cpp
@@ -35,7 +35,7 @@ void UNetworkSimulationGlobalManager::ReconcileSimulationsPostNetworkUpdate()
 	{
 		for (auto& SimItem : It.Value().Simulations)
 		{
-			SimItem.Simulation->Reconcile(SimItem.OwningActor->Role);
+			SimItem.Simulation->Reconcile(SimItem.OwningActor->GetLocalRole());
 		}
 	}
 }
diff --git a/Engine/Source/Runtime/Engine/Classes/Components/ActorComponent.h b/Engine/Source/Runtime/Engine/Classes/Components/ActorComponent.h
--- a/Engine/Source/Runtime/Engine/Classes/Components/ActorComponent.h
+++ b/Engine/Source/Runtime/Engine/Classes/Components/ActorComponent.h
@@ -133,6 +133,7 @@ protected:
 	uint8 bPhysicsStateCreated:1;
 
 	/** Is this component currently replicating? Should the network code consider it for replication? Owning Actor must be replicating first! */
+	UE_DEPRECATED(4.24, "This member will be made private. Please use GetIsReplicated, SetIsReplicated, or SetIsReplicatedByDefault for constructors.")
 	UPROPERTY(EditDefaultsOnly, BlueprintReadOnly, Replicated, Category=ComponentReplication,meta=(DisplayName = "Component Replicates"))
 	uint8 bReplicates:1;
 
@@ -189,6 +190,7 @@ public:
 	uint8 bAutoActivate:1;
 
 	/** Whether the component is currently active. */
+	UE_DEPRECATED(4.24, "This member will be made private. Please use IsActive or SetActive.")
 	UPROPERTY(transient, ReplicatedUsing=OnRep_IsActive)
 	uint8 bIsActive:1;
 
@@ -353,12 +355,14 @@ public:
 	UFUNCTION(BlueprintCallable, Category="Components|Activation", meta=(UnsafeDuringActorConstruction="true"))
 	virtual void ToggleActive();
 
+PRAGMA_DISABLE_DEPRECATION_WARNINGS
 	/**
 	 * Returns whether the component is active or not
 	 * @return - The active state of the component.
 	 */
 	UFUNCTION(BlueprintCallable, Category="Components|Activation", meta=(UnsafeDuringActorConstruction="true"))
 	bool IsActive() const { return bIsActive; }
+PRAGMA_ENABLE_DEPRECATION_WARNINGS
 
 	/**
 	 * Sets whether the component should be auto activate or not. Only safe during construction scripts.
@@ -387,11 +391,13 @@ public:
 	UFUNCTION(BlueprintCallable, Category="Components")
 	void SetIsReplicated(bool ShouldReplicate);
 
+PRAGMA_DISABLE_DEPRECATION_WARNINGS
 	/** Returns whether replication is enabled or not. */
 	FORCEINLINE bool GetIsReplicated() const
 	{
 		return bReplicates;
 	}
+PRAGMA_ENABLE_DEPRECATION_WARNINGS
 
 	/** Allows a component to replicate other subobject on the actor  */
 	virtual bool ReplicateSubobjects(class UActorChannel *Channel, class FOutBunch *Bunch, FReplicationFlags *RepFlags);
@@ -921,6 +927,7 @@ private:
 	friend class FComponentRecreateRenderStateContext;
 	friend struct FActorComponentTickFunction;
 
+PRAGMA_DISABLE_DEPRECATION_WARNINGS
 	//~ Begin Methods for Replicated Members.
 protected:
 
@@ -954,6 +961,7 @@ public:
 	void SetActiveFlag(const bool bNewIsActive);
 
 	//~ End Methods for Replicated Members.
+PRAGMA_ENABLE_DEPRECATION_WARNINGS
 
 protected:
 
diff --git a/Engine/Source/Runtime/Engine/Classes/Components/SceneComponent.h b/Engine/Source/Runtime/Engine/Classes/Components/SceneComponent.h
--- a/Engine/Source/Runtime/Engine/Classes/Components/SceneComponent.h
+++ b/Engine/Source/Runtime/Engine/Classes/Components/SceneComponent.h
@@ -149,10 +149,12 @@ public:
 	FBoxSphereBounds Bounds;
 
 	/** Location of the component relative to its parent */
+	UE_DEPRECATED(4.24, "This member will be made private. Please use GetRelativeLocation or SetRelativeLocation.")
 	UPROPERTY(EditAnywhere, BlueprintReadOnly, ReplicatedUsing=OnRep_Transform, Category = Transform)
 	FVector RelativeLocation;
 
 	/** Rotation of the component relative to its parent */
+	UE_DEPRECATED(4.24, "This member will be made private. Please use GetRelativeRotation or SetRelativeRotation.")
 	UPROPERTY(EditAnywhere, BlueprintReadOnly, ReplicatedUsing=OnRep_Transform, Category=Transform)
 	FRotator RelativeRotation;
 
@@ -160,6 +162,7 @@ public:
 	*	Non-uniform scaling of the component relative to its parent.
 	*	Note that scaling is always applied in local space (no shearing etc)
 	*/
+	UE_DEPRECATED(4.24, "This member will be made private. Please use GetRelativeScale3D or SetRelativeScale3D.")
 	UPROPERTY(BlueprintReadOnly, ReplicatedUsing=OnRep_Transform, interp, Category=Transform)
 	FVector RelativeScale3D;
 
@@ -181,18 +184,22 @@ private:
 
 public:
 	/** If RelativeLocation should be considered relative to the world, rather than the parent */
+	UE_DEPRECATED(4.24, "This member will be made private. Please use IsUsingAbsoluteLocation or SetUsingAbsoluteLocation.")
 	UPROPERTY(EditAnywhere, AdvancedDisplay, BlueprintReadWrite, ReplicatedUsing=OnRep_Transform, Category=Transform)
 	uint8 bAbsoluteLocation:1;
 
 	/** If RelativeRotation should be considered relative to the world, rather than the parent */
+	UE_DEPRECATED(4.24, "This member will be made private. Please use IsUsingAbsoluteRotation or SetUsingAbsoluteRotation.")
 	UPROPERTY(EditAnywhere, AdvancedDisplay, BlueprintReadWrite, ReplicatedUsing=OnRep_Transform, Category=Transform)
 	uint8 bAbsoluteRotation:1;
 
 	/** If RelativeScale3D should be considered relative to the world, rather than the parent */
+	UE_DEPRECATED(4.24, "This member will be made private. Please use IsUsingAbsoluteScale or SetUsingAbsoluteScale.")
 	UPROPERTY(EditAnywhere, AdvancedDisplay, BlueprintReadWrite, ReplicatedUsing=OnRep_Transform, Category=Transform)
 	uint8 bAbsoluteScale:1;
 
 	/** Whether to completely draw the primitive; if false, the primitive is not drawn, does not cast a shadow. */
+	UE_DEPRECATED(4.24, "This member will be made private. Please use IsVisible or SetVisibility.")
 	UPROPERTY(EditAnywhere, BlueprintReadOnly, ReplicatedUsing=OnRep_Visibility,  Category = Rendering)
 	uint8 bVisible:1;
 
@@ -1250,7 +1257,8 @@ private:
 	friend class FScopedPreventAttachedComponentMove;
 	friend struct FDirectAttachChildrenAccessor;
 
-	//~ Begin Methods for Replicated Members.	
+PRAGMA_DISABLE_DEPRECATION_WARNINGS
+	//~ Begin Methods for Replicated Members.
 private:
 
 	/**
@@ -1507,6 +1515,7 @@ public:
 	void SetVisibleFlag(const bool bInVisible);
 	
 	//~ End Methods for Replicated Members.
+PRAGMA_ENABLE_DEPRECATION_WARNINGS
 };
 
 /** 
diff --git a/Engine/Source/Runtime/Engine/Classes/GameFramework/Actor.h b/Engine/Source/Runtime/Engine/Classes/GameFramework/Actor.h
--- a/Engine/Source/Runtime/Engine/Classes/GameFramework/Actor.h
+++ b/Engine/Source/Runtime/Engine/Classes/GameFramework/Actor.h
@@ -163,6 +163,7 @@ public:
 	 * Allows us to only see this Actor in the Editor, and not in the actual game.
 	 * @see SetActorHiddenInGame()
 	 */
+	UE_DEPRECATED(4.24, "This member will be made private. Please use IsHidden or SetHidden.")
 	UPROPERTY(Interp, EditAnywhere, Category=Rendering, BlueprintReadOnly, Replicated, meta=(DisplayName="Actor Hidden In Game", SequencerTrackClass="MovieSceneVisibilityTrack"))
 	uint8 bHidden:1;
 
@@ -188,6 +189,7 @@ public:
 	 * @see SetReplicates()
 	 * @see https://docs.unrealengine.com/latest/INT/Gameplay/Networking/Replication/
 	 */
+	UE_DEPRECATED(4.24, "This member will be made private. Please use IsReplicatingMovement or SetReplicatingMovement.")
 	UPROPERTY(ReplicatedUsing=OnRep_ReplicateMovement, Category=Replication, EditDefaultsOnly)
 	uint8 bReplicateMovement:1;    
 
@@ -267,6 +269,7 @@ public:
 	 * @see https://www.unrealengine.com/blog/damage-in-ue4
 	 * @see TakeDamage(), ReceiveDamage()
 	 */
+	UE_DEPRECATED(4.24, "This member will be made private. Please use CanBeDamaged or SetCanBeDamaged.")
 	UPROPERTY(EditAnywhere, BlueprintReadWrite, SaveGame, Replicated, Category=Actor)
 	uint8 bCanBeDamaged:1;
 
@@ -479,6 +482,7 @@ public:
 	ENetRole GetRemoteRole() const;
 
 	/** Used for replication of our RootComponent's position and velocity */
+	UE_DEPRECATED(4.24, "This member will be made private. Please use GetReplicatedMovement or SetReplicatedMovement.")
 	UPROPERTY(EditDefaultsOnly, ReplicatedUsing=OnRep_ReplicatedMovement, Category=Replication, AdvancedDisplay)
 	struct FRepMovement ReplicatedMovement;
 
@@ -525,6 +529,7 @@ public:
 	virtual void OnRep_AttachmentReplication();
 
 	/** Describes how much control the local machine has over the actor. */
+	UE_DEPRECATED(4.24, "This member will be made private. Please use GetLocalRole or SetRole.")
 	UPROPERTY(Replicated)
 	TEnumAsByte<enum ENetRole> Role;
 
@@ -625,6 +630,7 @@ public:
 	void CallPreReplication(UNetDriver* NetDriver);	
 	
 	/** Pawn responsible for damage and other gameplay events caused by this actor. */
+	UE_DEPRECATED(4.24, "This member will be made private. Please use GetInstigator or SetInstigator.")
 	UPROPERTY(BlueprintReadWrite, ReplicatedUsing=OnRep_Instigator, meta=(ExposeOnSpawn=true), Category=Actor)
 	class APawn* Instigator;
 
@@ -3168,6 +3174,8 @@ private:
 		return (RootComponent != nullptr) ? RootComponent->GetRightVector() : FVector::RightVector;
 	}
 	
+PRAGMA_DISABLE_DEPRECATION_WARNINGS
+
 	//~ Begin Methods for Replicated Members.
 public:
 
@@ -3298,6 +3306,7 @@ public:
 	void SetInstigator(APawn* InInstigator);
 
 	//~ End Methods for Replicated Members.
+PRAGMA_ENABLE_DEPRECATION_WARNINGS
 };
 
 /** Internal struct used by level code to mark actors as destroyed */
@@ -3469,10 +3478,13 @@ FORCEINLINE_DEBUGGABLE const AActor* AActor::GetNetOwner() const
 	return Owner;
 }
 
+PRAGMA_DISABLE_DEPRECATION_WARNINGS
 FORCEINLINE_DEBUGGABLE ENetRole AActor::GetLocalRole() const
 {
 	return Role;
 }
+PRAGMA_ENABLE_DEPRECATION_WARNINGS
+
 
 FORCEINLINE_DEBUGGABLE ENetRole AActor::GetRemoteRole() const
 {
diff --git a/Engine/Source/Runtime/Engine/Private/Actor.cpp b/Engine/Source/Runtime/Engine/Private/Actor.cpp
--- a/Engine/Source/Runtime/Engine/Private/Actor.cpp
+++ b/Engine/Source/Runtime/Engine/Private/Actor.cpp
@@ -1056,7 +1056,9 @@ void AActor::PreReplication(IRepChangedPropertyTracker & ChangedPropertyTracker)
 
 	GatherCurrentMovement();
 
-	DOREPLIFETIME_ACTIVE_OVERRIDE(AActor, ReplicatedMovement, bReplicateMovement);
+PRAGMA_DISABLE_DEPRECATION_WARNINGS
+	DOREPLIFETIME_ACTIVE_OVERRIDE(AActor, ReplicatedMovement, IsReplicatingMovement());
+PRAGMA_ENABLE_DEPRECATION_WARNINGS
 
 	// Don't need to replicate AttachmentReplication if the root component replicates, because it already handles it.
 	DOREPLIFETIME_ACTIVE_OVERRIDE(AActor, AttachmentReplication, RootComponent && !RootComponent->GetIsReplicated());
@@ -1710,7 +1712,7 @@ void AActor::OnRep_AttachmentReplication()
 		// Handle the case where an object was both detached and moved on the server in the same frame.
 		// Calling this extraneously does not hurt but will properly fire events if the movement state changed while attached.
 		// This is needed because client side movement is ignored when attached
-		if (bReplicateMovement)
+		if (IsReplicatingMovement())
 		{
 			OnRep_ReplicatedMovement();
 		}
@@ -2554,6 +2556,7 @@ void AActor::EndViewTarget( APlayerController* PC )
 	K2_OnEndViewTarget(PC);
 }
 
+PRAGMA_DISABLE_DEPRECATION_WARNINGS
 APawn* AActor::GetInstigator() const
 {
 	return Instigator;
@@ -2563,6 +2566,7 @@ AController* AActor::GetInstigatorController() const
 {
 	return Instigator ? Instigator->Controller : nullptr;
 }
+PRAGMA_ENABLE_DEPRECATION_WARNINGS
 
 void AActor::CalcCamera(float DeltaTime, FMinimalViewInfo& OutResult)
 {
@@ -3341,8 +3345,9 @@ void AActor::ExchangeNetRoles(bool bRemoteOwned)
 	{
 		if (bRemoteOwned)
 		{
-			// Don't worry about calling SetRemoteRoleInternal here, as this should only be hit during initialization.
+PRAGMA_DISABLE_DEPRECATION_WARNINGS
 			Exchange( Role, RemoteRole );
+PRAGMA_ENABLE_DEPRECATION_WARNINGS
 		}
 		bExchangedRoles = true;
 	}
@@ -3350,7 +3355,9 @@ void AActor::ExchangeNetRoles(bool bRemoteOwned)
 
 void AActor::SwapRoles()
 {
+PRAGMA_DISABLE_DEPRECATION_WARNINGS
 	Swap(Role, RemoteRole);
+PRAGMA_ENABLE_DEPRECATION_WARNINGS
 
 	// MARK_PROPERTY_DIRTY_FROM_NAME(AActor, RemoteRole, this);
 	// MARK_PROPERTY_DIRTY_FROM_NAME(AActor, Role, this);
@@ -4977,34 +4984,45 @@ void AActor::SetLODParent(UPrimitiveComponent* InLODParent, float InParentDrawDi
 	}
 }
 
+
 void AActor::SetHidden(bool bInHidden)
 {
+PRAGMA_DISABLE_DEPRECATION_WARNINGS
 	bHidden = bInHidden;
 	// MARK_PROPERTY_DIRTY_FROM_NAME(AActor, bHidden, this);
+PRAGMA_ENABLE_DEPRECATION_WARNINGS
 }
 
 void AActor::SetReplicatingMovement(bool bInReplicateMovement)
 {
+PRAGMA_DISABLE_DEPRECATION_WARNINGS
 	bReplicateMovement = bInReplicateMovement;
 	// MARK_PROPERTY_DIRTY_FROM_NAME(AActor, bReplicateMovement, this);
+PRAGMA_ENABLE_DEPRECATION_WARNINGS
 }
 
 void AActor::SetCanBeDamaged(bool bInCanBeDamaged)
 {
+PRAGMA_DISABLE_DEPRECATION_WARNINGS
 	bCanBeDamaged = bInCanBeDamaged;
 	// MARK_PROPERTY_DIRTY_FROM_NAME(AActor, bCanBeDamaged, this);
+PRAGMA_ENABLE_DEPRECATION_WARNINGS
 }
 
 void AActor::SetRole(ENetRole InRole)
 {
+PRAGMA_DISABLE_DEPRECATION_WARNINGS
 	Role = InRole;
 	// MARK_PROPERTY_DIRTY_FROM_NAME(AActor, Role, this);
+PRAGMA_ENABLE_DEPRECATION_WARNINGS
 }
 
 FRepMovement& AActor::GetReplicatedMovement_Mutable()
 {
+PRAGMA_DISABLE_DEPRECATION_WARNINGS
 	// MARK_PROPERTY_DIRTY_FROM_NAME(AActor, ReplicatedMovement, this);
 	return ReplicatedMovement;
+PRAGMA_ENABLE_DEPRECATION_WARNINGS
 }
 
 void AActor::SetReplicatedMovement(const FRepMovement& InReplicatedMovement)
@@ -5014,8 +5032,10 @@ void AActor::SetReplicatedMovement(const FRepMovement& InReplicatedMovement)
 
 void AActor::SetInstigator(APawn* InInstigator)
 {
+PRAGMA_DISABLE_DEPRECATION_WARNINGS
 	Instigator = InInstigator;
 	// MARK_PROPERTY_DIRTY_FROM_NAME(AActor, Instigator, this);
+PRAGMA_ENABLE_DEPRECATION_WARNINGS
 }
 
 #undef LOCTEXT_NAMESPACE
diff --git a/Engine/Source/Runtime/Engine/Private/ActorReplication.cpp b/Engine/Source/Runtime/Engine/Private/ActorReplication.cpp
--- a/Engine/Source/Runtime/Engine/Private/ActorReplication.cpp
+++ b/Engine/Source/Runtime/Engine/Private/ActorReplication.cpp
@@ -144,10 +144,12 @@ void AActor::PostNetReceive()
 		bNetCheckedInitialPhysicsState = true;
 	}
 
+PRAGMA_DISABLE_DEPRECATION_WARNINGS
 	ExchangeB(bHidden, ActorReplication::SavedbHidden);
 	Exchange(Owner, ActorReplication::SavedOwner);
+PRAGMA_ENABLE_DEPRECATION_WARNINGS
 
-	if (bHidden != ActorReplication::SavedbHidden)
+	if (IsHidden() != ActorReplication::SavedbHidden)
 	{
 		SetActorHiddenInGame(ActorReplication::SavedbHidden);
 	}
@@ -156,7 +158,7 @@ void AActor::PostNetReceive()
 		SetOwner(ActorReplication::SavedOwner);
 	}
 
-	if (Role != ActorReplication::SavedRole)
+	if (GetLocalRole() != ActorReplication::SavedRole)
 	{
 		PostNetReceiveRole();
 	}
@@ -333,6 +335,7 @@ bool AActor::IsReplayRelevantFor(const AActor* RealViewer, const AActor* ViewTar
 
 void AActor::GatherCurrentMovement()
 {
+PRAGMA_DISABLE_DEPRECATION_WARNINGS
 	if (IsReplicatingMovement() || (RootComponent && RootComponent->GetAttachParent()))
 	{
 		bool bWasAttachmentModified = false;
@@ -404,6 +407,7 @@ void AActor::GatherCurrentMovement()
 			ReplicatedMovement.bRepPhysics = false;
 		}
 	}
+PRAGMA_ENABLE_DEPRECATION_WARNINGS
 }
 
 void AActor::GetLifetimeReplicatedProps( TArray< FLifetimeProperty > & OutLifetimeProps ) const
@@ -414,22 +418,29 @@ void AActor::GetLifetimeReplicatedProps( TArray< FLifetimeProperty > & OutLifeti
 		BPClass->GetLifetimeBlueprintReplicationList(OutLifetimeProps);
 	}
 
+PRAGMA_DISABLE_DEPRECATION_WARNINGS
 	DOREPLIFETIME( AActor, bReplicateMovement );
 	DOREPLIFETIME( AActor, Role );
+PRAGMA_ENABLE_DEPRECATION_WARNINGS
 	DOREPLIFETIME( AActor, RemoteRole );
 	DOREPLIFETIME( AActor, Owner );
+
+PRAGMA_DISABLE_DEPRECATION_WARNINGS
 	DOREPLIFETIME( AActor, bHidden );
+PRAGMA_ENABLE_DEPRECATION_WARNINGS
 
-	PRAGMA_DISABLE_DEPRECATION_WARNINGS
 	DOREPLIFETIME( AActor, bTearOff );
-	PRAGMA_ENABLE_DEPRECATION_WARNINGS
 
+PRAGMA_DISABLE_DEPRECATION_WARNINGS
 	DOREPLIFETIME( AActor, bCanBeDamaged );
+PRAGMA_ENABLE_DEPRECATION_WARNINGS
+
 	DOREPLIFETIME_CONDITION_NOTIFY( AActor, AttachmentReplication, COND_Custom, REPNOTIFY_Always );
 
+PRAGMA_DISABLE_DEPRECATION_WARNINGS
 	DOREPLIFETIME( AActor, Instigator );
-
 	DOREPLIFETIME_CONDITION_NOTIFY( AActor, ReplicatedMovement, COND_SimulatedOrPhysics, REPNOTIFY_Always );
+PRAGMA_ENABLE_DEPRECATION_WARNINGS
 }
 
 bool AActor::ReplicateSubobjects(UActorChannel *Channel, FOutBunch *Bunch, FReplicationFlags *RepFlags)
diff --git a/Engine/Source/Runtime/Engine/Private/Components/ActorComponent.cpp b/Engine/Source/Runtime/Engine/Private/Components/ActorComponent.cpp
--- a/Engine/Source/Runtime/Engine/Private/Components/ActorComponent.cpp
+++ b/Engine/Source/Runtime/Engine/Private/Components/ActorComponent.cpp
@@ -1742,8 +1742,10 @@ void UActorComponent::SetIsReplicated(bool bShouldReplicate)
 
 		if (GetComponentClassCanReplicate())
 		{
+PRAGMA_DISABLE_DEPRECATION_WARNINGS
 			bReplicates = bShouldReplicate;
 			// MARK_PROPERTY_DIRTY_FROM_NAME(UActorComponent, bReplicates, this);
+PRAGMA_ENABLE_DEPRECATION_WARNINGS
 
 			if (AActor* MyOwner = GetOwner())
 			{
@@ -1790,8 +1792,10 @@ void UActorComponent::GetLifetimeReplicatedProps( TArray< FLifetimeProperty > &
 		BPClass->GetLifetimeBlueprintReplicationList(OutLifetimeProps);
 	}
 
+PRAGMA_DISABLE_DEPRECATION_WARNINGS
 	DOREPLIFETIME( UActorComponent, bIsActive );
 	DOREPLIFETIME( UActorComponent, bReplicates );
+PRAGMA_ENABLE_DEPRECATION_WARNINGS
 }
 
 void UActorComponent::OnRep_IsActive()
@@ -1955,8 +1959,10 @@ void UActorComponent::SetIsReplicatedByDefault(const bool bNewReplicates)
 	// Don't bother checking parent here.
 	if (LIKELY(NeedsInitialization()))
 	{
+PRAGMA_DISABLE_DEPRECATION_WARNINGS
 		bReplicates = bNewReplicates;
 		// MARK_PROPERTY_DIRTY_FROM_NAME(UActorComponent, bReplicates, this);
+PRAGMA_ENABLE_DEPRECATION_WARNINGS
 	}
 	else
 	{
@@ -1967,8 +1973,10 @@ void UActorComponent::SetIsReplicatedByDefault(const bool bNewReplicates)
 
 void UActorComponent::SetActiveFlag(const bool bNewIsActive)
 {
+PRAGMA_DISABLE_DEPRECATION_WARNINGS
 	bIsActive = bNewIsActive;
 	// MARK_PROPERTY_DIRTY_FROM_NAME(UActorComponent, bIsActive, this);
+PRAGMA_ENABLE_DEPRECATION_WARNINGS
 }
 
 bool UActorComponent::OwnerNeedsInitialization() const
diff --git a/Engine/Source/Runtime/Engine/Private/Components/CharacterMovementComponent.cpp b/Engine/Source/Runtime/Engine/Private/Components/CharacterMovementComponent.cpp
--- a/Engine/Source/Runtime/Engine/Private/Components/CharacterMovementComponent.cpp
+++ b/Engine/Source/Runtime/Engine/Private/Components/CharacterMovementComponent.cpp
@@ -2821,7 +2821,7 @@ void UCharacterMovementComponent::UnCrouch(bool bClientSimulation)
 void UCharacterMovementComponent::UpdateCharacterStateBeforeMovement(float DeltaSeconds)
 {
 	// Proxies get replicated crouch state.
-	if (CharacterOwner->Role != ROLE_SimulatedProxy)
+	if (CharacterOwner->GetLocalRole() != ROLE_SimulatedProxy)
 	{
 		// Check for a change in crouch state. Players toggle crouch by changing bWantsToCrouch.
 		const bool bIsCrouching = IsCrouching();
@@ -2839,7 +2839,7 @@ void UCharacterMovementComponent::UpdateCharacterStateBeforeMovement(float Delta
 void UCharacterMovementComponent::UpdateCharacterStateAfterMovement(float DeltaSeconds)
 {
 	// Proxies get replicated crouch state.
-	if (CharacterOwner->Role != ROLE_SimulatedProxy)
+	if (CharacterOwner->GetLocalRole() != ROLE_SimulatedProxy)
 	{
 		// Uncrouch if no longer allowed to be crouched
 		if (IsCrouching() && !CanCrouchInCurrentState())
diff --git a/Engine/Source/Runtime/Engine/Private/Components/SceneComponent.cpp b/Engine/Source/Runtime/Engine/Private/Components/SceneComponent.cpp
--- a/Engine/Source/Runtime/Engine/Private/Components/SceneComponent.cpp
+++ b/Engine/Source/Runtime/Engine/Private/Components/SceneComponent.cpp
@@ -3204,8 +3204,8 @@ void USceneComponent::OnRep_AttachChildren()
 
 void USceneComponent::OnRep_Visibility(bool OldValue)
 {
-	bool ReppedValue = bVisible;
-	bVisible = OldValue;
+	bool ReppedValue = IsVisible();
+	SetVisibleFlag(OldValue);
 	SetVisibility(ReppedValue);
 }
 
@@ -3286,19 +3286,25 @@ void USceneComponent::GetLifetimeReplicatedProps(TArray< FLifetimeProperty > & O
 {
 	Super::GetLifetimeReplicatedProps(OutLifetimeProps);
 	
+PRAGMA_DISABLE_DEPRECATION_WARNINGS
 	DOREPLIFETIME(USceneComponent, bAbsoluteLocation);
 	DOREPLIFETIME(USceneComponent, bAbsoluteRotation);
 	DOREPLIFETIME(USceneComponent, bAbsoluteScale);
 	DOREPLIFETIME(USceneComponent, bVisible);
+PRAGMA_ENABLE_DEPRECATION_WARNINGS
+
 	DOREPLIFETIME(USceneComponent, bShouldBeAttached);
 	DOREPLIFETIME(USceneComponent, bShouldSnapLocationWhenAttached);
 	DOREPLIFETIME(USceneComponent, bShouldSnapRotationWhenAttached);
 	DOREPLIFETIME(USceneComponent, AttachParent);
 	DOREPLIFETIME(USceneComponent, AttachChildren);
 	DOREPLIFETIME(USceneComponent, AttachSocketName);
+
+PRAGMA_DISABLE_DEPRECATION_WARNINGS
 	DOREPLIFETIME(USceneComponent, RelativeLocation);
 	DOREPLIFETIME(USceneComponent, RelativeRotation);
 	DOREPLIFETIME(USceneComponent, RelativeScale3D);
+PRAGMA_ENABLE_DEPRECATION_WARNINGS
 }
 
 #if WITH_EDITOR
@@ -3363,18 +3369,23 @@ FScopedPreventAttachedComponentMove::FScopedPreventAttachedComponentMove(USceneC
 		bSavedAbsoluteScale = Owner->IsUsingAbsoluteScale();
 		bSavedNonAbsoluteComponent = !(bSavedAbsoluteLocation && bSavedAbsoluteRotation && bSavedAbsoluteScale);
 
+		// These are only going to be changed temporarily and reset, so we'll allow access here.
+PRAGMA_DISABLE_DEPRECATION_WARNINGS
 		// Use absolute (stay in world space no matter what parent does)
 		Owner->bAbsoluteLocation = true;
 		Owner->bAbsoluteRotation = true;
 		Owner->bAbsoluteScale = true;
+PRAGMA_ENABLE_DEPRECATION_WARNINGS
 
 		if (bSavedNonAbsoluteComponent && Owner->GetAttachParent())
 		{
+PRAGMA_DISABLE_DEPRECATION_WARNINGS
 			// Make RelativeLocation etc relative to the world.
 			Component->ConditionalUpdateComponentToWorld();
 			Owner->RelativeLocation = Owner->GetComponentLocation();
 			Owner->RelativeRotation = Owner->GetComponentRotation();
 			Owner->RelativeScale3D = Owner->GetComponentScale();
+PRAGMA_ENABLE_DEPRECATION_WARNINGS
 		}
 	}
 	else
@@ -3389,9 +3400,11 @@ FScopedPreventAttachedComponentMove::~FScopedPreventAttachedComponentMove()
 {
 	if (Owner)
 	{
+PRAGMA_DISABLE_DEPRECATION_WARNINGS
 		Owner->bAbsoluteLocation = bSavedAbsoluteLocation;
 		Owner->bAbsoluteRotation = bSavedAbsoluteRotation;
 		Owner->bAbsoluteScale = bSavedAbsoluteScale;
+PRAGMA_ENABLE_DEPRECATION_WARNINGS
 
 		if (bSavedNonAbsoluteComponent && Owner->GetAttachParent())
 		{
@@ -3499,11 +3512,13 @@ void FScopedMovementUpdate::RevertMove()
 		
 		if (IsTransformDirty())
 		{
+PRAGMA_DISABLE_DEPRECATION_WARNINGS
 			// Teleport to start
 			Component->ComponentToWorld = InitialTransform;
 			Component->RelativeLocation = InitialRelativeLocation;
 			Component->RelativeRotation = InitialRelativeRotation;
 			Component->RelativeScale3D = InitialRelativeScale;
+PRAGMA_ENABLE_DEPRECATION_WARNINGS
 
 			if (!IsDeferringUpdates())
 			{
@@ -3821,14 +3836,18 @@ void USceneComponent::K2_AddWorldTransform(const FTransform& DeltaTransform, boo
 
 void USceneComponent::SetVisibleFlag(const bool bNewVisible)
 {
+PRAGMA_DISABLE_DEPRECATION_WARNINGS
 	bVisible = bNewVisible;
 	// MARK_PROPERTY_DIRTY_FROM_NAME(USceneComponent, bVisible, this);
+PRAGMA_ENABLE_DEPRECATION_WARNINGS
 }
 
 FVector& USceneComponent::GetRelativeLocation_DirectMutable()
 {
+PRAGMA_DISABLE_DEPRECATION_WARNINGS
 	// MARK_PROPERTY_DIRTY_FROM_NAME(USceneComponent, RelativeLocation, this);
 	return RelativeLocation;
+PRAGMA_ENABLE_DEPRECATION_WARNINGS
 }
 
 void USceneComponent::SetRelativeLocation_Direct(const FVector NewRelativeLocation)
@@ -3838,8 +3857,10 @@ void USceneComponent::SetRelativeLocation_Direct(const FVector NewRelativeLocati
 
 FRotator& USceneComponent::GetRelativeRotation_DirectMutable()
 {
+PRAGMA_DISABLE_DEPRECATION_WARNINGS
 	// MARK_PROPERTY_DIRTY_FROM_NAME(USceneComponent, RelativeRotation, this);
 	return RelativeRotation;
+PRAGMA_ENABLE_DEPRECATION_WARNINGS
 }
 
 void USceneComponent::SetRelativeRotation_Direct(const FRotator NewRelativeRotation)
@@ -3849,8 +3870,10 @@ void USceneComponent::SetRelativeRotation_Direct(const FRotator NewRelativeRotat
 
 FVector& USceneComponent::GetRelativeScale3D_DirectMutable()
 {
+PRAGMA_DISABLE_DEPRECATION_WARNINGS
 	// MARK_PROPERTY_DIRTY_FROM_NAME(USceneComponent, RelativeScale3D, this);
 	return RelativeScale3D;
+PRAGMA_ENABLE_DEPRECATION_WARNINGS
 }
 
 void USceneComponent::SetRelativeScale3D_Direct(const FVector NewRelativeScale3D)
@@ -3860,32 +3883,42 @@ void USceneComponent::SetRelativeScale3D_Direct(const FVector NewRelativeScale3D
 
 void USceneComponent::SetUsingAbsoluteLocation(bool bInAbsoluteLocation)
 {
+PRAGMA_DISABLE_DEPRECATION_WARNINGS
 	bAbsoluteLocation = bInAbsoluteLocation;
 	// MARK_PROPERTY_DIRTY_FROM_NAME(USceneComponent, bAbsoluteLocation, this);
+PRAGMA_ENABLE_DEPRECATION_WARNINGS
 }
 
 void USceneComponent::SetUsingAbsoluteRotation(bool bInAbsoluteRotation)
 {
+PRAGMA_DISABLE_DEPRECATION_WARNINGS
 	bAbsoluteRotation = bInAbsoluteRotation;
 	// MARK_PROPERTY_DIRTY_FROM_NAME(USceneComponent, bAbsoluteRotation, this);
+PRAGMA_ENABLE_DEPRECATION_WARNINGS
 }
 
 void USceneComponent::SetUsingAbsoluteScale(bool bInAbsoluteScale)
 {
+PRAGMA_DISABLE_DEPRECATION_WARNINGS
 	bAbsoluteScale = bInAbsoluteScale;
 	// MARK_PROPERTY_DIRTY_FROM_NAME(USceneComponent, bAbsoluteScale, this);
+PRAGMA_ENABLE_DEPRECATION_WARNINGS
 }
 
 void USceneComponent::SetAttachParent(USceneComponent* NewAttachParent)
 {
+PRAGMA_DISABLE_DEPRECATION_WARNINGS
 	AttachParent = NewAttachParent;
 	// MARK_PROPERTY_DIRTY_FROM_NAME(USceneComponent, AttachParent, this);
+PRAGMA_ENABLE_DEPRECATION_WARNINGS
 }
 
 void USceneComponent::SetAttachSocketName(FName NewSocketName)
 {
+PRAGMA_DISABLE_DEPRECATION_WARNINGS
 	AttachSocketName = NewSocketName;
 	// MARK_PROPERTY_DIRTY_FROM_NAME(USceneComponent, AttachSocketName, this);
+PRAGMA_ENABLE_DEPRECATION_WARNINGS
 }
 
 void USceneComponent::ModifiedAttachChildren()
diff --git a/Engine/Source/Runtime/Engine/Private/Net/BandwidthTestActor.cpp b/Engine/Source/Runtime/Engine/Private/Net/BandwidthTestActor.cpp
--- a/Engine/Source/Runtime/Engine/Private/Net/BandwidthTestActor.cpp
+++ b/Engine/Source/Runtime/Engine/Private/Net/BandwidthTestActor.cpp
@@ -12,7 +12,7 @@ ABandwidthTestActor::ABandwidthTestActor()
 {
 	bReplicates = true;
 	bAlwaysRelevant = true;
-	bReplicateMovement = false;
+	SetReplicatingMovement(false);
 
 	NetDormancy = DORM_Never;
 }
diff --git a/Engine/Source/Runtime/Landscape/Private/ControlPointMeshActor.cpp b/Engine/Source/Runtime/Landscape/Private/ControlPointMeshActor.cpp
--- a/Engine/Source/Runtime/Landscape/Private/ControlPointMeshActor.cpp
+++ b/Engine/Source/Runtime/Landscape/Private/ControlPointMeshActor.cpp
@@ -18,7 +18,7 @@
 AControlPointMeshActor::AControlPointMeshActor(const FObjectInitializer& ObjectInitializer)
 	: Super(ObjectInitializer)
 {
-	bCanBeDamaged = false;
+	SetCanBeDamaged(false);
 
 	ControlPointMeshComponent = ObjectInitializer.CreateDefaultSubobject<UControlPointMeshComponent>(this, TEXT("ControlPointMeshComponent0"));
 	ControlPointMeshComponent->SetCollisionProfileName(UCollisionProfile::BlockAll_ProfileName);
