Add shading scalability engine settings.

diff --git a/Engine/Config/BaseScalability.ini b/Engine/Config/BaseScalability.ini
--- a/Engine/Config/BaseScalability.ini
+++ b/Engine/Config/BaseScalability.ini
@@ -29,6 +29,7 @@ PerfIndexThresholds_PostProcessQuality="GPU 18 42 115"
 PerfIndexThresholds_TextureQuality="GPU 18 42 115"
 PerfIndexThresholds_EffectsQuality="Min 18 42 105"
 PerfIndexThresholds_FoliageQuality="GPU 18 42 115"
+PerfIndexThresholds_ShadingQuality="GPU 18 42 115"
 
 ; This is the screen percentage for the resolution quality, corresponding to 25% pixels, 50% pixels, 75% pixels, and 100% pixels
 PerfIndexValues_ResolutionQuality="50 71 87 100 100"
@@ -556,3 +557,20 @@ foliage.DensityScale=1.0
 grass.DensityScale=1.0
 ;foliage.DiscardDataOnLoad=0
 ;grass.DiscardDataOnLoad=0
+
+;-----------------------------------------------------------------------------------------------------------------
+
+[ShadingQuality@0]
+;todo
+
+[ShadingQuality@1]
+;todo
+
+[ShadingQuality@2]
+;todo
+
+[ShadingQuality@3]
+;todo
+
+[ShadingQuality@Cine]
+;todo
diff --git a/Engine/Plugins/Experimental/AutomationUtils/Source/AutomationUtils/Private/AutomationUtilsBlueprintLibrary.cpp b/Engine/Plugins/Experimental/AutomationUtils/Source/AutomationUtils/Private/AutomationUtilsBlueprintLibrary.cpp
--- a/Engine/Plugins/Experimental/AutomationUtils/Source/AutomationUtils/Private/AutomationUtilsBlueprintLibrary.cpp
+++ b/Engine/Plugins/Experimental/AutomationUtils/Source/AutomationUtils/Private/AutomationUtilsBlueprintLibrary.cpp
@@ -153,6 +153,7 @@ FAutomationUtilsGameplayAutomationScreenshotInstance::FAutomationUtilsGameplayAu
 		JsonObject->SetNumberField(TEXT("textureQuality"), QualityLevels.TextureQuality);
 		JsonObject->SetNumberField(TEXT("effectsQuality"), QualityLevels.EffectsQuality);
 		JsonObject->SetNumberField(TEXT("foliageQuality"), QualityLevels.FoliageQuality);
+		JsonObject->SetNumberField(TEXT("shadingQuality"), QualityLevels.ShadingQuality);
 		//ComparisonOptions
 		JsonObject->SetBoolField(TEXT("bHasComparisonRules"), true);
 		JsonObject->SetNumberField(TEXT("toleranceRed"), 8);
diff --git a/Engine/Source/Editor/UnrealEd/Private/PerformanceMonitor.cpp b/Engine/Source/Editor/UnrealEd/Private/PerformanceMonitor.cpp
--- a/Engine/Source/Editor/UnrealEd/Private/PerformanceMonitor.cpp
+++ b/Engine/Source/Editor/UnrealEd/Private/PerformanceMonitor.cpp
@@ -167,6 +167,7 @@ bool FPerformanceMonitor::WillAutoScalabilityHelp() const
 	IsAutoScaleLower |= NewLevels.TextureQuality < CurrentLevels.TextureQuality;
 	IsAutoScaleLower |= NewLevels.EffectsQuality < CurrentLevels.EffectsQuality;
 	IsAutoScaleLower |= NewLevels.FoliageQuality < CurrentLevels.FoliageQuality;
+	IsAutoScaleLower |= NewLevels.ShadingQuality < CurrentLevels.ShadingQuality;
 
 	// We don't check things like real-time, because the user may have enabled it temporarily.
 
@@ -187,6 +188,7 @@ Scalability::FQualityLevels FPerformanceMonitor::GetAutoScalabilityQualityLevels
 	NewLevels.TextureQuality		= FMath::Min(NewLevels.TextureQuality, ExistingLevels.TextureQuality);
 	NewLevels.EffectsQuality		= FMath::Min(NewLevels.EffectsQuality, ExistingLevels.EffectsQuality);
 	NewLevels.FoliageQuality		= FMath::Min(NewLevels.FoliageQuality, ExistingLevels.FoliageQuality);
+	NewLevels.ShadingQuality		= FMath::Min(NewLevels.ShadingQuality, ExistingLevels.ShadingQuality);
 
 	return NewLevels;
 }
diff --git a/Engine/Source/Editor/UnrealEd/Private/SScalabilitySettings.cpp b/Engine/Source/Editor/UnrealEd/Private/SScalabilitySettings.cpp
--- a/Engine/Source/Editor/UnrealEd/Private/SScalabilitySettings.cpp
+++ b/Engine/Source/Editor/UnrealEd/Private/SScalabilitySettings.cpp
@@ -27,6 +27,7 @@ ECheckBoxState SScalabilitySettings::IsGroupQualityLevelSelected(const TCHAR* In
 	else if (FCString::Strcmp(InGroupName, TEXT("TextureQuality")) == 0) QualityLevel = CachedQualityLevels.TextureQuality;
 	else if (FCString::Strcmp(InGroupName, TEXT("EffectsQuality")) == 0) QualityLevel = CachedQualityLevels.EffectsQuality;
 	else if (FCString::Strcmp(InGroupName, TEXT("FoliageQuality")) == 0) QualityLevel = CachedQualityLevels.FoliageQuality;
+	else if (FCString::Strcmp(InGroupName, TEXT("ShadingQuality")) == 0) QualityLevel = CachedQualityLevels.ShadingQuality;
 
 	return (QualityLevel == InQualityLevel) ? ECheckBoxState::Checked : ECheckBoxState::Unchecked;
 }
@@ -41,6 +42,7 @@ void SScalabilitySettings::OnGroupQualityLevelChanged(ECheckBoxState NewState, c
 	else if (FCString::Strcmp(InGroupName, TEXT("TextureQuality")) == 0) CachedQualityLevels.TextureQuality = InQualityLevel;
 	else if (FCString::Strcmp(InGroupName, TEXT("EffectsQuality")) == 0) CachedQualityLevels.EffectsQuality = InQualityLevel;
 	else if (FCString::Strcmp(InGroupName, TEXT("FoliageQuality")) == 0) CachedQualityLevels.FoliageQuality = InQualityLevel;
+	else if (FCString::Strcmp(InGroupName, TEXT("ShadingQuality")) == 0) CachedQualityLevels.ShadingQuality = InQualityLevel;
 
 	Scalability::SetQualityLevels(CachedQualityLevels);
 	Scalability::SaveState(GEditorSettingsIni);
@@ -218,8 +220,9 @@ void SScalabilitySettings::Construct( const FArguments& InArgs )
 		FMath::Max(LevelCounts.ViewDistanceQuality,
 		FMath::Max(LevelCounts.EffectsQuality,
 		FMath::Max(LevelCounts.FoliageQuality,
+		FMath::Max(LevelCounts.ShadingQuality,
 		FMath::Max(LevelCounts.PostProcessQuality, LevelCounts.AntiAliasingQuality)
-		)))));
+		))))));
 	const int32 TotalWidth = MaxLevelCount + 1;
 
 	TSharedRef<SGridPanel> ButtonMatrix = 
@@ -234,6 +237,7 @@ void SScalabilitySettings::Construct( const FArguments& InArgs )
 		+MakeGridSlot(0,6,TotalWidth,1) [ SNew (SBorder).BorderImage(FEditorStyle::GetBrush("Scalability.RowBackground")) ]
 		+MakeGridSlot(0,7,TotalWidth,1) [ SNew (SBorder).BorderImage(FEditorStyle::GetBrush("Scalability.RowBackground")) ]
 		+MakeGridSlot(0,8,TotalWidth,1) [ SNew (SBorder).BorderImage(FEditorStyle::GetBrush("Scalability.RowBackground")) ]
+		+MakeGridSlot(0,9,TotalWidth,1) [ SNew (SBorder).BorderImage(FEditorStyle::GetBrush("Scalability.RowBackground")) ]
 
 		+MakeGridSlot(0,0).VAlign(VAlign_Center) [ SNew(STextBlock).Text(LOCTEXT("QualityLabel", "Quality")).Font(TitleFont) ]
 		+MakeGridSlot(1,0) [ MakeHeaderButtonWidget(NamesLow, 0, LOCTEXT("QualityLow", "Set all groups to low quality"), Scalability::EQualityLevelBehavior::EAbsolute) ]
@@ -259,6 +263,7 @@ void SScalabilitySettings::Construct( const FArguments& InArgs )
 
 		+MakeGridSlot(0,7) [ SNew(STextBlock).Text(LOCTEXT("EffectsQualityLabel1", "Effects")).Font(GroupFont) ]
 		+MakeGridSlot(0,8) [ SNew(STextBlock).Text(LOCTEXT("FoliageQualityLabel1", "Foliage")).Font(GroupFont) ]
+		+MakeGridSlot(0,9) [ SNew(STextBlock).Text(LOCTEXT("ShadingQualityLabel1", "Shading")).Font(GroupFont) ]
 		;
 
 	AddButtonsToGrid(1, 2, ButtonMatrix, FiveDistanceNames, LevelCounts.ViewDistanceQuality, TEXT("ViewDistanceQuality"), LOCTEXT("ViewDistanceQualityTooltip", "Set view distance to {0}"));
@@ -268,6 +273,7 @@ void SScalabilitySettings::Construct( const FArguments& InArgs )
 	AddButtonsToGrid(1, 6, ButtonMatrix, FiveNames, LevelCounts.TextureQuality, TEXT("TextureQuality"), LOCTEXT("TextureQualityTooltip", "Set texture quality to {0}"));
 	AddButtonsToGrid(1, 7, ButtonMatrix, FiveNames, LevelCounts.EffectsQuality, TEXT("EffectsQuality"), LOCTEXT("EffectsQualityTooltip", "Set effects quality to {0}"));
 	AddButtonsToGrid(1, 8, ButtonMatrix, FiveNames, LevelCounts.FoliageQuality, TEXT("FoliageQuality"), LOCTEXT("FoliageQualityTooltip", "Set foliage quality to {0}"));
+	AddButtonsToGrid(1, 9, ButtonMatrix, FiveNames, LevelCounts.ShadingQuality, TEXT("ShadingQuality"), LOCTEXT("ShadingQualityTooltip", "Set shading quality to {0}"));
 
 
 	this->ChildSlot
diff --git a/Engine/Source/Editor/UnrealEd/Private/Settings/EditorSettings.cpp b/Engine/Source/Editor/UnrealEd/Private/Settings/EditorSettings.cpp
--- a/Engine/Source/Editor/UnrealEd/Private/Settings/EditorSettings.cpp
+++ b/Engine/Source/Editor/UnrealEd/Private/Settings/EditorSettings.cpp
@@ -59,6 +59,7 @@ void UEditorSettings::LoadScalabilityBenchmark()
 		GConfig->GetInt(Section, TEXT("TextureQuality"), Temporary.TextureQuality, GEditorSettingsIni);
 		GConfig->GetInt(Section, TEXT("EffectsQuality"), Temporary.EffectsQuality, GEditorSettingsIni);
 		GConfig->GetInt(Section, TEXT("FoliageQuality"), Temporary.FoliageQuality, GEditorSettingsIni);
+		GConfig->GetInt(Section, TEXT("ShadingQuality"), Temporary.ShadingQuality, GEditorSettingsIni);
 		EngineBenchmarkResult = Temporary;
 	}
 }
@@ -82,6 +83,7 @@ void UEditorSettings::AutoApplyScalabilityBenchmark()
 	GConfig->SetInt(Section, TEXT("TextureQuality"), Temporary.TextureQuality, GEditorSettingsIni);
 	GConfig->SetInt(Section, TEXT("EffectsQuality"), Temporary.EffectsQuality, GEditorSettingsIni);
 	GConfig->SetInt(Section, TEXT("FoliageQuality"), Temporary.FoliageQuality, GEditorSettingsIni);
+	GConfig->SetInt(Section, TEXT("ShadingQuality"), Temporary.ShadingQuality, GEditorSettingsIni);
 
 	Scalability::SetQualityLevels(Temporary);
 	Scalability::SaveState(GEditorSettingsIni);
diff --git a/Engine/Source/Runtime/AutomationMessages/Public/AutomationWorkerMessages.h b/Engine/Source/Runtime/AutomationMessages/Public/AutomationWorkerMessages.h
--- a/Engine/Source/Runtime/AutomationMessages/Public/AutomationWorkerMessages.h
+++ b/Engine/Source/Runtime/AutomationMessages/Public/AutomationWorkerMessages.h
@@ -436,6 +436,8 @@ public:
 	int32 EffectsQuality;
 	UPROPERTY(EditAnywhere, Category="Message")
 	int32 FoliageQuality;
+	UPROPERTY(EditAnywhere, Category = "Message")
+	int32 ShadingQuality;
 
 	// Comparison Requests
 	UPROPERTY(EditAnywhere, Category="Message")
@@ -475,6 +477,7 @@ public:
 		, TextureQuality(0)
 		, EffectsQuality(0)
 		, FoliageQuality(0)
+		, ShadingQuality(0)
 		, bHasComparisonRules(0)
 		, ToleranceRed(0)
 		, ToleranceGreen(0)
@@ -528,6 +531,7 @@ public:
 		TextureQuality = Data.TextureQuality;
 		EffectsQuality = Data.EffectsQuality;
 		FoliageQuality = Data.FoliageQuality;
+		ShadingQuality = Data.ShadingQuality;
 
 		// Enabled Features
 		// TBD
diff --git a/Engine/Source/Runtime/Core/Public/Misc/AutomationTest.h b/Engine/Source/Runtime/Core/Public/Misc/AutomationTest.h
--- a/Engine/Source/Runtime/Core/Public/Misc/AutomationTest.h
+++ b/Engine/Source/Runtime/Core/Public/Misc/AutomationTest.h
@@ -579,6 +579,7 @@ struct FAutomationScreenshotData
 	int32 TextureQuality;
 	int32 EffectsQuality;
 	int32 FoliageQuality;
+	int32 ShadingQuality;
 
 	// Comparison Requests
 	bool bHasComparisonRules;
@@ -609,6 +610,7 @@ struct FAutomationScreenshotData
 		, TextureQuality(0)
 		, EffectsQuality(0)
 		, FoliageQuality(0)
+		, ShadingQuality(0)
 		, bHasComparisonRules(false)
 		, ToleranceRed(0)
 		, ToleranceGreen(0)
diff --git a/Engine/Source/Runtime/Engine/Classes/GameFramework/GameUserSettings.h b/Engine/Source/Runtime/Engine/Classes/GameFramework/GameUserSettings.h
--- a/Engine/Source/Runtime/Engine/Classes/GameFramework/GameUserSettings.h
+++ b/Engine/Source/Runtime/Engine/Classes/GameFramework/GameUserSettings.h
@@ -244,6 +244,15 @@ public:
 	UFUNCTION(BlueprintCallable, Category=Settings)
 	int32 GetFoliageQuality() const;
 
+	// Sets the shading quality (0..4, higher is better)
+	// @param Value 0:low, 1:medium, 2:high, 3:epic, 4:cinematic (gets clamped if needed)
+	UFUNCTION(BlueprintCallable, Category = Settings)
+	void SetShadingQuality(int32 Value);
+
+	// Returns the shading quality (0..4, higher is better)
+	UFUNCTION(BlueprintCallable, Category = Settings)
+	int32 GetShadingQuality() const;
+
 	/** Checks if any user settings is different from current */
 	UFUNCTION(BlueprintPure, Category=Settings)
 	virtual bool IsDirty() const;
diff --git a/Engine/Source/Runtime/Engine/Private/ChartCreation.cpp b/Engine/Source/Runtime/Engine/Private/ChartCreation.cpp
--- a/Engine/Source/Runtime/Engine/Private/ChartCreation.cpp
+++ b/Engine/Source/Runtime/Engine/Private/ChartCreation.cpp
@@ -176,6 +176,7 @@ void FDumpFPSChartToEndpoint::HandleBasicStats()
 	PrintToEndpoint(FString::Printf(TEXT("\tTexture Quality: %d"), ScalabilityQuality.TextureQuality));
 	PrintToEndpoint(FString::Printf(TEXT("\tEffects Quality: %d"), ScalabilityQuality.EffectsQuality));
 	PrintToEndpoint(FString::Printf(TEXT("\tFoliage Quality: %d"), ScalabilityQuality.FoliageQuality));
+	PrintToEndpoint(FString::Printf(TEXT("\tHair Quality: %d"), ScalabilityQuality.ShadingQuality));
 
 	PrintToEndpoint(FString::Printf(TEXT("\tWindow Mode: %s"), *WindowMode));
 	PrintToEndpoint(FString::Printf(TEXT("\tResolution: %dx%d"), GameResolution.X, GameResolution.Y));
@@ -392,6 +393,7 @@ protected:
 			ParamArray.Add(FAnalyticsEventAttribute(TEXT("TextureQuality"), ScalabilityQuality.TextureQuality));
 			ParamArray.Add(FAnalyticsEventAttribute(TEXT("FXQuality"), ScalabilityQuality.EffectsQuality));
 			ParamArray.Add(FAnalyticsEventAttribute(TEXT("FoliageQuality"), ScalabilityQuality.FoliageQuality));
+			ParamArray.Add(FAnalyticsEventAttribute(TEXT("ShadingQuality"), ScalabilityQuality.ShadingQuality));
 			ParamArray.Add(FAnalyticsEventAttribute(TEXT("PercentGPUBound"), FString::Printf(TEXT("%4.2f"), BoundGPUPct)));
 			ParamArray.Add(FAnalyticsEventAttribute(TEXT("AvgGPUTime"), FString::Printf(TEXT("%4.2f"), AvgGPUFrameTime)));
 		}
@@ -537,6 +539,7 @@ protected:
 		FPSChartRow = FPSChartRow.Replace(TEXT("TOKEN_SETTINGS_TEX"), *FString::Printf(TEXT("%d"), ScalabilityQuality.TextureQuality), ESearchCase::CaseSensitive);
 		FPSChartRow = FPSChartRow.Replace(TEXT("TOKEN_SETTINGS_FX"), *FString::Printf(TEXT("%d"), ScalabilityQuality.EffectsQuality), ESearchCase::CaseSensitive);
 		FPSChartRow = FPSChartRow.Replace(TEXT("TOKEN_SETTINGS_FLG"), *FString::Printf(TEXT("%d"), ScalabilityQuality.FoliageQuality), ESearchCase::CaseSensitive);
+		FPSChartRow = FPSChartRow.Replace(TEXT("TOKEN_SETTINGS_SHD"), *FString::Printf(TEXT("%d"), ScalabilityQuality.ShadingQuality), ESearchCase::CaseSensitive);
 
 		FPSChartRow = FPSChartRow.Replace(TEXT("TOKEN_AVG_FPS"), *FString::Printf(TEXT("%4.2f"), Chart.GetAverageFramerate()), ESearchCase::CaseSensitive);
 		FPSChartRow = FPSChartRow.Replace(TEXT("TOKEN_TIME_DISREGARDED"), *FString::Printf(TEXT("%4.2f"), Chart.TimeDisregarded), ESearchCase::CaseSensitive);
diff --git a/Engine/Source/Runtime/Engine/Private/GameUserSettings.cpp b/Engine/Source/Runtime/Engine/Private/GameUserSettings.cpp
--- a/Engine/Source/Runtime/Engine/Private/GameUserSettings.cpp
+++ b/Engine/Source/Runtime/Engine/Private/GameUserSettings.cpp
@@ -836,6 +836,16 @@ int32 UGameUserSettings::GetFoliageQuality() const
 	return ScalabilityQuality.FoliageQuality;
 }
 
+void UGameUserSettings::SetShadingQuality(int32 Value)
+{
+	ScalabilityQuality.SetShadingQuality(Value);
+}
+
+int32 UGameUserSettings::GetShadingQuality() const
+{
+	return ScalabilityQuality.ShadingQuality;
+}
+
 UGameUserSettings* UGameUserSettings::GetGameUserSettings()
 {
 	return GEngine->GetGameUserSettings();
diff --git a/Engine/Source/Runtime/Engine/Private/Scalability.cpp b/Engine/Source/Runtime/Engine/Private/Scalability.cpp
--- a/Engine/Source/Runtime/Engine/Private/Scalability.cpp
+++ b/Engine/Source/Runtime/Engine/Private/Scalability.cpp
@@ -66,6 +66,13 @@ static TAutoConsoleVariable<int32> CVarFoliageQuality(
 	TEXT(" 0:low, 1:med, 2:high, 3:epic, 4:cinematic, default: 3"),
 	ECVF_ScalabilityGroup);
 
+static TAutoConsoleVariable<int32> CVarShadingQuality(
+	TEXT("sg.ShadingQuality"),
+	3,
+	TEXT("Scalability quality state (internally used by scalability system, ini load/save or using SCALABILITY console command)\n")
+	TEXT(" 0:low, 1:med, 2:high, 3:epic, 4:cinematic, default: 3"),
+	ECVF_ScalabilityGroup);
+
 static TAutoConsoleVariable<int32> CVarViewDistanceQuality_NumLevels(
 	TEXT("sg.ViewDistanceQuality.NumLevels"),
 	5,
@@ -115,6 +122,12 @@ static TAutoConsoleVariable<int32> CVarFoliageQuality_NumLevels(
 	TEXT(" default: 5 (0..4)"),
 	ECVF_ReadOnly);
 
+static TAutoConsoleVariable<int32> CVarShadingQuality_NumLevels(
+	TEXT("sg.ShadingQuality.NumLevels"),
+	5,
+	TEXT("Number of settings quality levels in sg.ShadingQuality\n")
+	TEXT(" default: 5 (0..4)"),
+	ECVF_ReadOnly);
 
 namespace Scalability
 {
@@ -442,6 +455,11 @@ void OnChangeFoliageQuality(IConsoleVariable* Var)
 	SetGroupQualityLevel(TEXT("FoliageQuality"), Var->GetInt(), CVarFoliageQuality_NumLevels->GetInt());
 }
 
+void OnChangeShadingQuality(IConsoleVariable* Var)
+{
+	SetGroupQualityLevel(TEXT("ShadingQuality"), Var->GetInt(), CVarShadingQuality_NumLevels->GetInt());
+}
+
 void InitScalabilitySystem()
 {
 	// needed only once
@@ -464,6 +482,7 @@ void InitScalabilitySystem()
 	CVarTextureQuality.AsVariable()->SetOnChangedCallback(FConsoleVariableDelegate::CreateStatic(&OnChangeTextureQuality));
 	CVarEffectsQuality.AsVariable()->SetOnChangedCallback(FConsoleVariableDelegate::CreateStatic(&OnChangeEffectsQuality));
 	CVarFoliageQuality.AsVariable()->SetOnChangedCallback(FConsoleVariableDelegate::CreateStatic(&OnChangeFoliageQuality));
+	CVarShadingQuality.AsVariable()->SetOnChangedCallback(FConsoleVariableDelegate::CreateStatic(&OnChangeShadingQuality));
 
 	// Set defaults
 	SetQualityLevels(FQualityLevels());
@@ -520,6 +539,7 @@ FQualityLevels BenchmarkQualityLevels(uint32 WorkScale, float CPUMultiplier, flo
 	Results.TextureQuality = ComputeOptionFromPerfIndex(TEXT("TextureQuality"), CPUPerfIndex, GPUPerfIndex);
 	Results.EffectsQuality = ComputeOptionFromPerfIndex(TEXT("EffectsQuality"), CPUPerfIndex, GPUPerfIndex);
 	Results.FoliageQuality = ComputeOptionFromPerfIndex(TEXT("FoliageQuality"), CPUPerfIndex, GPUPerfIndex);
+	Results.ShadingQuality = ComputeOptionFromPerfIndex(TEXT("ShadingQuality"), CPUPerfIndex, GPUPerfIndex);
 	Results.CPUBenchmarkResults = CPUPerfIndex;
 	Results.GPUBenchmarkResults = GPUPerfIndex;
 
@@ -625,6 +645,7 @@ void ProcessCommand(const TCHAR* Cmd, FOutputDevice& Ar)
 		PrintGroupInfo(TEXT("TextureQuality"), bInfoMode);
 		PrintGroupInfo(TEXT("EffectsQuality"), bInfoMode);
 		PrintGroupInfo(TEXT("FoliageQuality"), bInfoMode);
+		PrintGroupInfo(TEXT("ShadingQuality"), bInfoMode);
 
 		if (CPUBenchmarkValue >= 0.0f)
 		{
@@ -648,6 +669,7 @@ void SetQualityLevels(const FQualityLevels& QualityLevels, bool bForce/* = false
 	ClampedLevels.TextureQuality = FMath::Clamp(QualityLevels.TextureQuality, 0, CVarTextureQuality_NumLevels->GetInt() - 1);
 	ClampedLevels.EffectsQuality = FMath::Clamp(QualityLevels.EffectsQuality, 0, CVarEffectsQuality_NumLevels->GetInt() - 1);
 	ClampedLevels.FoliageQuality = FMath::Clamp(QualityLevels.FoliageQuality, 0, CVarFoliageQuality_NumLevels->GetInt() - 1);
+	ClampedLevels.ShadingQuality = FMath::Clamp(QualityLevels.ShadingQuality, 0, CVarShadingQuality_NumLevels->GetInt() - 1);
 
 	if (GScalabilityUsingTemporaryQualityLevels && !bForce)
 	{
@@ -665,6 +687,7 @@ void SetQualityLevels(const FQualityLevels& QualityLevels, bool bForce/* = false
 		CVarTextureQuality.AsVariable()->SetWithCurrentPriority(ClampedLevels.TextureQuality);
 		CVarEffectsQuality.AsVariable()->SetWithCurrentPriority(ClampedLevels.EffectsQuality);
 		CVarFoliageQuality.AsVariable()->SetWithCurrentPriority(ClampedLevels.FoliageQuality);
+		CVarShadingQuality.AsVariable()->SetWithCurrentPriority(ClampedLevels.ShadingQuality);
 	}
 	else
 	{
@@ -676,6 +699,7 @@ void SetQualityLevels(const FQualityLevels& QualityLevels, bool bForce/* = false
 		CVarTextureQuality.AsVariable()->Set(ClampedLevels.TextureQuality, ECVF_SetByScalability);
 		CVarEffectsQuality.AsVariable()->Set(ClampedLevels.EffectsQuality, ECVF_SetByScalability);
 		CVarFoliageQuality.AsVariable()->Set(ClampedLevels.FoliageQuality, ECVF_SetByScalability);
+		CVarShadingQuality.AsVariable()->Set(ClampedLevels.ShadingQuality, ECVF_SetByScalability);
 	}
 
 	OnScalabilitySettingsChanged.Broadcast(ClampedLevels);
@@ -696,6 +720,7 @@ FQualityLevels GetQualityLevels()
 		Ret.TextureQuality = CVarTextureQuality.GetValueOnGameThread();
 		Ret.EffectsQuality = CVarEffectsQuality.GetValueOnGameThread();
 		Ret.FoliageQuality = CVarFoliageQuality.GetValueOnGameThread();
+		Ret.ShadingQuality = CVarShadingQuality.GetValueOnGameThread();
 	}
 	else
 	{
@@ -760,6 +785,7 @@ void FQualityLevels::SetFromSingleQualityLevel(int32 Value)
 	TextureQuality = FMath::Clamp(Value, 0, CVarTextureQuality_NumLevels->GetInt() - 1);
 	EffectsQuality = FMath::Clamp(Value, 0, CVarEffectsQuality_NumLevels->GetInt() - 1);
 	FoliageQuality = FMath::Clamp(Value, 0, CVarFoliageQuality_NumLevels->GetInt() - 1);
+	ShadingQuality = FMath::Clamp(Value, 0, CVarShadingQuality_NumLevels->GetInt() - 1);
 }
 
 void FQualityLevels::SetFromSingleQualityLevelRelativeToMax(int32 Value)
@@ -776,6 +802,7 @@ void FQualityLevels::SetFromSingleQualityLevelRelativeToMax(int32 Value)
 	TextureQuality = FMath::Max(CVarTextureQuality_NumLevels->GetInt() - Value, 0);
 	EffectsQuality = FMath::Max(CVarEffectsQuality_NumLevels->GetInt() - Value, 0);
 	FoliageQuality = FMath::Max(CVarFoliageQuality_NumLevels->GetInt() - Value, 0);
+	ShadingQuality = FMath::Max(CVarShadingQuality_NumLevels->GetInt() - Value, 0);
 }
 
 // Returns the overall value if all settings are set to the same thing
@@ -785,7 +812,7 @@ int32 FQualityLevels::GetSingleQualityLevel() const
 	int32 Result = ViewDistanceQuality;
 
 	const int32 Target = ViewDistanceQuality;
-	if ((Target == AntiAliasingQuality) && (Target == ShadowQuality) && (Target == PostProcessQuality) && (Target == TextureQuality) && (Target == EffectsQuality) && (Target == FoliageQuality))
+	if ((Target == AntiAliasingQuality) && (Target == ShadowQuality) && (Target == PostProcessQuality) && (Target == TextureQuality) && (Target == EffectsQuality) && (Target == FoliageQuality) && (Target == ShadingQuality))
 	{
 		if (GetRenderScaleLevelFromQualityLevel(Target) == ResolutionQuality)
 		{
@@ -831,6 +858,11 @@ void FQualityLevels::SetFoliageQuality(int32 Value)
 	FoliageQuality = FMath::Clamp(Value, 0, CVarFoliageQuality_NumLevels->GetInt() - 1);
 }
 
+void FQualityLevels::SetShadingQuality(int32 Value)
+{
+	ShadingQuality = FMath::Clamp(Value, 0, CVarShadingQuality_NumLevels->GetInt() - 1);
+}
+
 void LoadState(const FString& IniName)
 {
 	check(!IniName.IsEmpty());
@@ -852,6 +884,7 @@ void LoadState(const FString& IniName)
 	GConfig->GetInt(Section, TEXT("sg.TextureQuality"), State.TextureQuality, IniName);
 	GConfig->GetInt(Section, TEXT("sg.EffectsQuality"), State.EffectsQuality, IniName);
 	GConfig->GetInt(Section, TEXT("sg.FoliageQuality"), State.FoliageQuality, IniName);
+	GConfig->GetInt(Section, TEXT("sg.ShadingQuality"), State.ShadingQuality, IniName);
 
 	// If possible apply immediately, else store in backup so we can re-apply later
 	if (!GScalabilityUsingTemporaryQualityLevels)
@@ -882,6 +915,7 @@ void SaveState(const FString& IniName)
 	GConfig->SetInt(Section, TEXT("sg.TextureQuality"), State.TextureQuality, IniName);
 	GConfig->SetInt(Section, TEXT("sg.EffectsQuality"), State.EffectsQuality, IniName);
 	GConfig->SetInt(Section, TEXT("sg.FoliageQuality"), State.FoliageQuality, IniName);
+	GConfig->SetInt(Section, TEXT("sg.ShadingQuality"), State.ShadingQuality, IniName);
 }
 
 void RecordQualityLevelsAnalytics(bool bAutoApplied)
@@ -900,6 +934,7 @@ void RecordQualityLevelsAnalytics(bool bAutoApplied)
 		Attributes.Add(FAnalyticsEventAttribute(TEXT("TextureQuality"), State.TextureQuality));
 		Attributes.Add(FAnalyticsEventAttribute(TEXT("EffectsQuality"), State.EffectsQuality));
 		Attributes.Add(FAnalyticsEventAttribute(TEXT("FoliageQuality"), State.FoliageQuality));
+		Attributes.Add(FAnalyticsEventAttribute(TEXT("ShadingQuality"), State.ShadingQuality));
 		Attributes.Add(FAnalyticsEventAttribute(TEXT("AutoAppliedSettings"), bAutoApplied));
 		Attributes.Add(FAnalyticsEventAttribute(TEXT("Enterprise"), IProjectManager::Get().IsEnterpriseProject()));
 
@@ -918,6 +953,7 @@ FQualityLevels GetQualityLevelCounts()
 	Result.TextureQuality = CVarTextureQuality_NumLevels->GetInt();
 	Result.EffectsQuality = CVarEffectsQuality_NumLevels->GetInt();
 	Result.FoliageQuality = CVarFoliageQuality_NumLevels->GetInt();
+	Result.ShadingQuality = CVarShadingQuality_NumLevels->GetInt();
 	return Result;
 }
 
diff --git a/Engine/Source/Runtime/Engine/Private/Tests/AutomationCommon.cpp b/Engine/Source/Runtime/Engine/Private/Tests/AutomationCommon.cpp
--- a/Engine/Source/Runtime/Engine/Private/Tests/AutomationCommon.cpp
+++ b/Engine/Source/Runtime/Engine/Private/Tests/AutomationCommon.cpp
@@ -123,6 +123,7 @@ namespace AutomationCommon
 		Data.TextureQuality = QualityLevels.TextureQuality;
 		Data.EffectsQuality = QualityLevels.EffectsQuality;
 		Data.FoliageQuality = QualityLevels.FoliageQuality;
+		Data.ShadingQuality = QualityLevels.ShadingQuality;
 		
 		//GRHIDeviceId
 
diff --git a/Engine/Source/Runtime/Engine/Public/Scalability.h b/Engine/Source/Runtime/Engine/Public/Scalability.h
--- a/Engine/Source/Runtime/Engine/Public/Scalability.h
+++ b/Engine/Source/Runtime/Engine/Public/Scalability.h
@@ -30,6 +30,7 @@ namespace Scalability
 		int32 TextureQuality;
 		int32 EffectsQuality;
 		int32 FoliageQuality;
+		int32 ShadingQuality;
 
 		float CPUBenchmarkResults;
 		float GPUBenchmarkResults;
@@ -100,6 +101,10 @@ namespace Scalability
 		// @param Value 0:low, 1:medium, 2:high, 3:epic, 4:cinematic (gets clamped if needed)
 		void SetFoliageQuality(int32 Value);
 
+		// Sets the sharing quality
+		// @param Value 0:low, 1:medium, 2:high, 3:epic, 4:cinematic (gets clamped if needed)
+		void SetShadingQuality(int32 Value);
+
 		void SetBenchmarkFallback();
 
 		void SetDefaults();
