#UE4 Fix implicit conversion to bool

diff --git a/Engine/Source/Runtime/Engine/Private/KismetSystemLibrary.cpp b/Engine/Source/Runtime/Engine/Private/KismetSystemLibrary.cpp
--- a/Engine/Source/Runtime/Engine/Private/KismetSystemLibrary.cpp
+++ b/Engine/Source/Runtime/Engine/Private/KismetSystemLibrary.cpp
@@ -2689,7 +2689,7 @@ bool UKismetSystemLibrary::Generic_GetEditorProperty(const UObject* Object, cons
 	if (EnumHasAnyFlags(AccessResult, EPropertyAccessResultFlags::ConversionFailed))
 	{
 		FFrame::KismetExecutionMessage(*FString::Printf(TEXT("Property '%s' on '%s' (%s) tried to read into a property value of the incorrect type"), *Property->GetName(), *Object->GetPathName(), *Object->GetClass()->GetName()), ELogVerbosity::Warning, PropertyGetFailedWarning);
-		return nullptr;
+		return false;
 	}
 
 	return true;
