Fixup usage of Actor, ActorComponent, and SceneComponent deprecated properties.

diff --git a/Engine/Source/Runtime/Engine/Private/Components/SceneComponent.cpp b/Engine/Source/Runtime/Engine/Private/Components/SceneComponent.cpp
--- a/Engine/Source/Runtime/Engine/Private/Components/SceneComponent.cpp
+++ b/Engine/Source/Runtime/Engine/Private/Components/SceneComponent.cpp
@@ -2765,13 +2765,19 @@ bool USceneComponent::InternalSetWorldLocationAndRotation(FVector NewLocation, c
 			RelativeRotationCache.RotatorToQuat(NewRelativeRotation);
 		}
 
+
 #if ENABLE_NAN_DIAGNOSTIC
+PRAGMA_DISABLE_DEPRECATION_WARNINGS
 		if (RelativeRotation.ContainsNaN())
 		{
 			logOrEnsureNanError(TEXT("USceneComponent:InternalSetWorldLocationAndRotation found NaN in RelativeRotation: %s"), *RelativeRotation.ToString());
 			RelativeRotation = FRotator::ZeroRotator;
+			// MARK_PROPERTY_DIRTY_FROM_NAME(USceneComponent, RelativeRotation, this);
 		}
+PRAGMA_ENABLE_DEPRECATION_WARNINGS
 #endif
+
+
 		UpdateComponentToWorldWithParent(GetAttachParent(),GetAttachSocketName(), SkipPhysicsToEnum(bNoPhysics), RelativeRotationCache.GetCachedQuat(), Teleport);
 
 		// we need to call this even if this component itself is not navigation relevant
