[NO NON-PII COMMIT MESSAGE]

diff --git a/Engine/Source/Runtime/Engine/Private/AudioStreamingCache.cpp b/Engine/Source/Runtime/Engine/Private/AudioStreamingCache.cpp
--- a/Engine/Source/Runtime/Engine/Private/AudioStreamingCache.cpp
+++ b/Engine/Source/Runtime/Engine/Private/AudioStreamingCache.cpp
@@ -1088,6 +1088,8 @@ uint64 FCachedAudioStreamingManager::TrimMemory(uint64 NumBytesToFree)
 
 TPair<int, int> FAudioChunkCache::DebugDisplay(UWorld* World, FViewport* Viewport, FCanvas* Canvas, int32 X, int32 Y, const FVector* ViewLocation, const FRotator* ViewRotation) const
 {
+	FScopeLock ScopeLock(&CacheMutationCriticalSection);
+
 	// Color scheme:
 	static constexpr float ColorMax = 256.0f;
 
