Option to immediately add object to spatial acceleration, on by default.
Fixes getting stuck in builds when building through self.

diff --git a/Engine/Source/Runtime/Engine/Private/PhysicsEngine/Experimental/PhysScene_Chaos.cpp b/Engine/Source/Runtime/Engine/Private/PhysicsEngine/Experimental/PhysScene_Chaos.cpp
--- a/Engine/Source/Runtime/Engine/Private/PhysicsEngine/Experimental/PhysScene_Chaos.cpp
+++ b/Engine/Source/Runtime/Engine/Private/PhysicsEngine/Experimental/PhysScene_Chaos.cpp
@@ -34,6 +34,7 @@
 #include "Chaos/PerParticleGravity.h"
 #include "PBDRigidActiveParticlesBuffer.h"
 #include "Chaos/GeometryParticlesfwd.h"
+#include "Chaos/Box.h"
 
 
 #if !UE_BUILD_SHIPPING
@@ -942,13 +943,39 @@ void FPhysScene_ChaosInterface::OnWorldEndPlay()
 	}
 }
 
-void FPhysScene_ChaosInterface::AddActorsToScene_AssumesLocked(TArray<FPhysicsActorHandle>& InHandles)
+void FPhysScene_ChaosInterface::AddActorsToScene_AssumesLocked(TArray<FPhysicsActorHandle>& InHandles, const bool bImmediate)
 {
 	Chaos::FPhysicsSolver* Solver = Scene.GetSolver();
 	Chaos::IDispatcher* Dispatcher = Scene.GetDispatcher();
+	Chaos::ISpatialAcceleration<Chaos::TAccelerationStructureHandle<float, 3>, float, 3>* SpatialAcceleration = Scene.GetSpacialAcceleration();
 	for (FPhysicsActorHandle& Handle : InHandles)
 	{
 		FPhysicsInterface::AddActorToSolver(Handle, Solver, Dispatcher);
+
+		// Optionally add this to the game-thread acceleration structure immediately
+		if (bImmediate && SpatialAcceleration)
+		{
+			// Get the bounding box for the particle if it has one
+			bool bHasBounds = Handle->Geometry()->HasBoundingBox();
+			Chaos::TBox<float, 3> WorldBounds = Chaos::TBox<float, 3>();
+			if (bHasBounds)
+			{
+				const Chaos::TBox<float, 3> LocalBounds = Handle->Geometry()->BoundingBox();
+				Chaos::TVector<float, 3> Velocity(0.f);
+				if (Handle->ObjectType() == Chaos::EParticleType::Kinematic ||
+					Handle->ObjectType() == Chaos::EParticleType::Dynamic)
+				{
+					Velocity = static_cast<Chaos::TKinematicGeometryParticle<float, 3>*>(Handle)->V();
+				}
+				WorldBounds = LocalBounds.TransformedBox(Chaos::TRigidTransform<float, 3>(Handle->X(), Handle->R()));
+				WorldBounds.ThickenSymmetrically(Velocity);
+				bHasBounds = true;
+			}
+
+			// Insert the particle
+			Chaos::TAccelerationStructureHandle<float, 3> AccelerationHandle(Handle);
+			SpatialAcceleration->UpdateElementIn(AccelerationHandle, WorldBounds, bHasBounds, Handle->SpatialIdx());
+		}
 	}
 }
 
diff --git a/Engine/Source/Runtime/Engine/Public/Physics/Experimental/PhysScene_Chaos.h b/Engine/Source/Runtime/Engine/Public/Physics/Experimental/PhysScene_Chaos.h
--- a/Engine/Source/Runtime/Engine/Public/Physics/Experimental/PhysScene_Chaos.h
+++ b/Engine/Source/Runtime/Engine/Public/Physics/Experimental/PhysScene_Chaos.h
@@ -276,7 +276,7 @@ public:
 
 	// In Chaos, this function will update the pointers from actor handles to their proxies.
 	// So the array of handles must be non-const.
-	void ENGINE_API AddActorsToScene_AssumesLocked(TArray<FPhysicsActorHandle>& InActors);
+	void ENGINE_API AddActorsToScene_AssumesLocked(TArray<FPhysicsActorHandle>& InActors, const bool bImmediate=true);
 	void AddAggregateToScene(const FPhysicsAggregateHandle& InAggregate);
 
 	void SetOwningWorld(UWorld* InOwningWorld);
