Simplify bounds for immediate object insertion

diff --git a/Engine/Source/Runtime/Engine/Private/PhysicsEngine/Experimental/PhysScene_Chaos.cpp b/Engine/Source/Runtime/Engine/Private/PhysicsEngine/Experimental/PhysScene_Chaos.cpp
--- a/Engine/Source/Runtime/Engine/Private/PhysicsEngine/Experimental/PhysScene_Chaos.cpp
+++ b/Engine/Source/Runtime/Engine/Private/PhysicsEngine/Experimental/PhysScene_Chaos.cpp
@@ -957,19 +957,11 @@ void FPhysScene_ChaosInterface::AddActorsToScene_AssumesLocked(TArray<FPhysicsAc
 		{
 			// Get the bounding box for the particle if it has one
 			bool bHasBounds = Handle->Geometry()->HasBoundingBox();
-			Chaos::TBox<float, 3> WorldBounds = Chaos::TBox<float, 3>();
+			Chaos::TBox<float, 3> WorldBounds;
 			if (bHasBounds)
 			{
 				const Chaos::TBox<float, 3> LocalBounds = Handle->Geometry()->BoundingBox();
-				Chaos::TVector<float, 3> Velocity(0.f);
-				if (Handle->ObjectType() == Chaos::EParticleType::Kinematic ||
-					Handle->ObjectType() == Chaos::EParticleType::Dynamic)
-				{
-					Velocity = static_cast<Chaos::TKinematicGeometryParticle<float, 3>*>(Handle)->V();
-				}
 				WorldBounds = LocalBounds.TransformedBox(Chaos::TRigidTransform<float, 3>(Handle->X(), Handle->R()));
-				WorldBounds.ThickenSymmetrically(Velocity);
-				bHasBounds = true;
 			}
 
 			// Insert the particle
