[NO NON-PII COMMIT MESSAGE]

diff --git a/Engine/Source/Runtime/Engine/Private/AudioCompressionSettingsUtils.cpp b/Engine/Source/Runtime/Engine/Private/AudioCompressionSettingsUtils.cpp
--- a/Engine/Source/Runtime/Engine/Private/AudioCompressionSettingsUtils.cpp
+++ b/Engine/Source/Runtime/Engine/Private/AudioCompressionSettingsUtils.cpp
@@ -217,7 +217,7 @@ FCachedAudioStreamingManagerParams FPlatformCompressionUtilities::BuildCachedStr
 
 	if (CookOverrides && (CookOverrides->StreamChunkSizeKB != 0))
 	{
-		MaxChunkSize = CookOverrides->StreamChunkSizeKB * 1024;;
+		MaxChunkSize = FMath::Max(CookOverrides->StreamChunkSizeKB * 1024, 1024);
 	}
 
 	MaxChunkSize = GetMaxChunkSizeForCookOverrides(GetCookOverridesForCurrentPlatform(), MaxChunkSize / 1024);
diff --git a/Engine/Source/Runtime/Engine/Private/SoundWave.cpp b/Engine/Source/Runtime/Engine/Private/SoundWave.cpp
--- a/Engine/Source/Runtime/Engine/Private/SoundWave.cpp
+++ b/Engine/Source/Runtime/Engine/Private/SoundWave.cpp
@@ -1765,7 +1765,8 @@ void USoundWave::Parse(FAudioDevice* AudioDevice, const UPTRINT NodeWaveInstance
 	// If set to bAlwaysPlay, increase the current sound's priority scale by 10x. This will still result in a possible 0-priority output if the sound has 0 actual volume
 	if (bAlwaysPlay)
 	{
-		WaveInstance->Priority = MAX_FLT;
+		static constexpr float VolumeWeightedMaxPriority = TNumericLimits<float>::Max() / MAX_VOLUME;
+		WaveInstance->Priority = VolumeWeightedMaxPriority;
 	}
 	else
 	{
