Fix LODMaterialMap issue with skeletalmesh refactor

#rnx

diff --git a/Engine/Source/Developer/SkeletalMeshUtilitiesCommon/Private/LODUtilities.cpp b/Engine/Source/Developer/SkeletalMeshUtilitiesCommon/Private/LODUtilities.cpp
--- a/Engine/Source/Developer/SkeletalMeshUtilitiesCommon/Private/LODUtilities.cpp
+++ b/Engine/Source/Developer/SkeletalMeshUtilitiesCommon/Private/LODUtilities.cpp
@@ -1125,8 +1125,8 @@ void FLODUtilities::RestoreSkeletalMeshLODImportedData(USkeletalMesh* SkeletalMe
 	{
 		FSkeletalMeshUpdateContext UpdateContext;
 		UpdateContext.SkeletalMesh = SkeletalMesh;
-
-		ImportedBaseLODModel.UpdateChunkedSectionInfo();
+		TArray<int32> EmptyLodInfoMaterialMap;
+		ImportedBaseLODModel.UpdateChunkedSectionInfo(SkeletalMesh->GetName(), EmptyLodInfoMaterialMap);
 
 		//Copy the SkeletalMeshLODModel
 		SkeletalMesh->GetImportedModel()->LODModels[LodIndex] = ImportedBaseLODModel;
diff --git a/Engine/Source/Editor/Persona/Private/PersonaMeshDetails.cpp b/Engine/Source/Editor/Persona/Private/PersonaMeshDetails.cpp
--- a/Engine/Source/Editor/Persona/Private/PersonaMeshDetails.cpp
+++ b/Engine/Source/Editor/Persona/Private/PersonaMeshDetails.cpp
@@ -3915,21 +3915,39 @@ void FPersonaMeshDetails::OnDeleteMaterialSlot(int32 MaterialIndex)
 
 	FScopedTransaction Transaction(LOCTEXT("PersonaOnDeleteMaterialSlotTransaction", "Persona editor: Delete material slot"));
 	SkeletalMeshPtr->Modify();
-	SkeletalMeshPtr->Materials.RemoveAt(MaterialIndex);
-	FSkeletalMeshModel* Model = SkeletalMeshPtr->GetImportedModel();
-	
-	for (int32 LodIndex = 0; LodIndex < Model->LODModels.Num(); ++LodIndex)
 	{
-		for (int32 SectionIndex = 0; SectionIndex < Model->LODModels[LodIndex].Sections.Num(); ++SectionIndex)
+		FScopedSkeletalMeshPostEditChange ScopedPostEditChange(SkeletalMeshPtr.Get());
+		SkeletalMeshPtr->Materials.RemoveAt(MaterialIndex);
+		FSkeletalMeshModel* Model = SkeletalMeshPtr->GetImportedModel();
+
+		int32 NumLODInfos = SkeletalMeshPtr->GetLODNum();
+
+		//When we delete a material slot we need a full LODMaterialMap for all section of all LODs
+		for (int32 LODInfoIdx = 0; LODInfoIdx < NumLODInfos; LODInfoIdx++)
 		{
-			if (Model->LODModels[LodIndex].Sections[SectionIndex].MaterialIndex > MaterialIndex)
+			for (int32 SectionIndex = 0; SectionIndex < Model->LODModels[LODInfoIdx].Sections.Num(); ++SectionIndex)
 			{
-				Model->LODModels[LodIndex].Sections[SectionIndex].MaterialIndex -= 1;
+				int32 SectionOriginalMaterialIndex = Model->LODModels[LODInfoIdx].Sections[SectionIndex].MaterialIndex;
+				if (!SkeletalMeshPtr->GetLODInfo(LODInfoIdx)->LODMaterialMap.IsValidIndex(SectionOriginalMaterialIndex))
+				{
+					SkeletalMeshPtr->GetLODInfo(LODInfoIdx)->LODMaterialMap.Add(SectionOriginalMaterialIndex);
+				}
+
+				int32& CurrentMaterialIndexRemap = SkeletalMeshPtr->GetLODInfo(LODInfoIdx)->LODMaterialMap[SectionOriginalMaterialIndex];
+
+				if (CurrentMaterialIndexRemap == MaterialIndex)
+				{
+					// Set to first material
+					CurrentMaterialIndexRemap = 0;
+				}
+				else if (CurrentMaterialIndexRemap > MaterialIndex)
+				{
+					// Decrement to keep correct reference after removal
+					CurrentMaterialIndexRemap--;
+				}
 			}
 		}
 	}
-
-	SkeletalMeshPtr->PostEditChange();
 }
 
 bool FPersonaMeshDetails::OnMaterialListDirty()
@@ -4981,45 +4999,47 @@ void FPersonaMeshDetails::OnSectionChanged(int32 LODIndex, int32 SectionIndex, i
 		// There is only one transaction for all replacement
 		FScopedTransaction Transaction(LOCTEXT("PersonaOnSectionChangedTransaction", "Persona editor: Section material slot changed"));
 		Mesh->Modify();
-
-		int32 NumSections = ImportedResource->LODModels[LODIndex].Sections.Num();
-		FSkeletalMeshLODInfo& Info = *(Mesh->GetLODInfo(LODIndex));
-		int32 CurrentMaterialIndex = ImportedResource->LODModels[LODIndex].Sections[SectionIndex].MaterialIndex;
-		
-		auto SetLODMaterialMapValue = [&Info](int32 OriginalSectionMaterialIndex, int32 OverrideMaterialIndex)
 		{
-			if (OriginalSectionMaterialIndex == OverrideMaterialIndex)
+			FScopedSkeletalMeshPostEditChange ScopedPostEditChange(Mesh);
+			int32 NumSections = ImportedResource->LODModels[LODIndex].Sections.Num();
+			FSkeletalMeshLODInfo& Info = *(Mesh->GetLODInfo(LODIndex));
+
+			auto SetLODMaterialMapValue = [&LODIndex, &Info, &ImportedResource](int32 OriginalSectionMaterialIndex, int32 OverrideMaterialIndex)
 			{
-				if (Info.LODMaterialMap.IsValidIndex(OriginalSectionMaterialIndex))
+				if (ImportedResource->LODModels[LODIndex].Sections[OriginalSectionMaterialIndex].MaterialIndex == OverrideMaterialIndex)
 				{
-					Info.LODMaterialMap[OriginalSectionMaterialIndex] = INDEX_NONE;
+					if (Info.LODMaterialMap.IsValidIndex(OriginalSectionMaterialIndex))
+					{
+						Info.LODMaterialMap[OriginalSectionMaterialIndex] = INDEX_NONE;
+					}
 				}
-			}
-			else
-			{
-				while (Info.LODMaterialMap.Num() <= OriginalSectionMaterialIndex)
+				else
 				{
-					int32 NewSectionIndex = Info.LODMaterialMap.Num();
-					Info.LODMaterialMap.Add(INDEX_NONE);
+					while (Info.LODMaterialMap.Num() <= OriginalSectionMaterialIndex)
+					{
+						int32 NewSectionIndex = Info.LODMaterialMap.Num();
+						Info.LODMaterialMap.Add(INDEX_NONE);
+					}
+					check(OriginalSectionMaterialIndex < Info.LODMaterialMap.Num());
+					Info.LODMaterialMap[OriginalSectionMaterialIndex] = OverrideMaterialIndex;
 				}
-				check(OriginalSectionMaterialIndex < Info.LODMaterialMap.Num());
-				Info.LODMaterialMap[OriginalSectionMaterialIndex] = OverrideMaterialIndex;
-			}
-		};
+			};
 
-		SetLODMaterialMapValue(CurrentMaterialIndex, NewSkeletalMaterialIndex);
-		//Set the chunked section 
-		for (int32 SectionIdx = 0; SectionIdx < NumSections; SectionIdx++)
-		{
-			if (ImportedResource->LODModels[LODIndex].Sections[SectionIdx].ChunkedParentSectionIndex == SectionIndex)
+			SetLODMaterialMapValue(SectionIndex, NewSkeletalMaterialIndex);
+			//Set the chunked section 
+			for (int32 SectionIdx = SectionIndex+1; SectionIdx < NumSections; SectionIdx++)
 			{
-				int32 CurrentChunkMaterialIndex = ImportedResource->LODModels[LODIndex].Sections[SectionIdx].MaterialIndex;
-				SetLODMaterialMapValue(CurrentChunkMaterialIndex, NewSkeletalMaterialIndex);
+				if (ImportedResource->LODModels[LODIndex].Sections[SectionIdx].ChunkedParentSectionIndex == SectionIndex)
+				{
+					SetLODMaterialMapValue(SectionIdx, NewSkeletalMaterialIndex);
+				}
+				else
+				{
+					//Chunked section are contiguous
+					break;
+				}
 			}
 		}
-
-		Mesh->PostEditChange();
-
 		// Redraw viewports to reflect the material changes 
 		GUnrealEd->RedrawLevelEditingViewports();
 	}
diff --git a/Engine/Source/Editor/UnrealEd/Private/Factories/SkeletalMeshImport.cpp b/Engine/Source/Editor/UnrealEd/Private/Factories/SkeletalMeshImport.cpp
--- a/Engine/Source/Editor/UnrealEd/Private/Factories/SkeletalMeshImport.cpp
+++ b/Engine/Source/Editor/UnrealEd/Private/Factories/SkeletalMeshImport.cpp
@@ -978,10 +978,6 @@ void RestoreExistingSkelMeshData(ExistingSkelMeshData* MeshData, USkeletalMesh*
 						{
 							BaseLODInfo.LODMaterialMap.Add(INDEX_NONE);
 						}
-						if (RemapMaterial.IsValidIndex(ExistingLODMaterialIndex))
-						{
-							ExistingLODMaterialIndex = RemapMaterial[ExistingLODMaterialIndex];
-						}
 						BaseLODInfo.LODMaterialMap[MaterialIndex] = ExistingLODMaterialIndex;
 					}
 				}
@@ -1047,7 +1043,6 @@ void RestoreExistingSkelMeshData(ExistingSkelMeshData* MeshData, USkeletalMesh*
 
 				//Get the current skelmesh section slot import name
 				FName ExistMeshSectionSlotName = MeshData->ExistingImportMeshLodSectionMaterialData[SafeReimportLODIndex][OriginalSectionIndex].ImportedMaterialSlotName;
-				if (ExistMeshSectionSlotName != OriginalImportMeshSectionSlotName)
 				{
 					//The last import slot name match the New import slot name, but the Exist slot name is different then the last import slot name.
 					//This mean the user has change the section assign slot and the fbx file did not change it
diff --git a/Engine/Source/Runtime/Engine/Private/SkeletalMesh.cpp b/Engine/Source/Runtime/Engine/Private/SkeletalMesh.cpp
--- a/Engine/Source/Runtime/Engine/Private/SkeletalMesh.cpp
+++ b/Engine/Source/Runtime/Engine/Private/SkeletalMesh.cpp
@@ -2178,7 +2178,7 @@ void USkeletalMesh::PostLoad()
 						}
 					}
 				}
-				ThisLODModel.UpdateChunkedSectionInfo();
+				ThisLODModel.UpdateChunkedSectionInfo(GetName(), ThisLODInfo->LODMaterialMap);
 			}
 		}
 
@@ -4145,9 +4145,9 @@ FSkeletalMeshSceneProxy::FSkeletalMeshSceneProxy(const USkinnedMeshComponent* Co
 			// If we are at a dropped LOD, route material index through the LODMaterialMap in the LODInfo struct.
 			int32 UseMaterialIndex = Section.MaterialIndex;			
 			{
-				if(Section.MaterialIndex < Info.LODMaterialMap.Num() && Component->SkeletalMesh->Materials.IsValidIndex(Info.LODMaterialMap[Section.MaterialIndex]))
+				if(SectionIndex < Info.LODMaterialMap.Num() && Component->SkeletalMesh->Materials.IsValidIndex(Info.LODMaterialMap[SectionIndex]))
 				{
-					UseMaterialIndex = Info.LODMaterialMap[Section.MaterialIndex];
+					UseMaterialIndex = Info.LODMaterialMap[SectionIndex];
 					UseMaterialIndex = FMath::Clamp( UseMaterialIndex, 0, Component->SkeletalMesh->Materials.Num() );
 				}
 			}
diff --git a/Engine/Source/Runtime/Engine/Private/SkeletalMeshLODModel.cpp b/Engine/Source/Runtime/Engine/Private/SkeletalMeshLODModel.cpp
--- a/Engine/Source/Runtime/Engine/Private/SkeletalMeshLODModel.cpp
+++ b/Engine/Source/Runtime/Engine/Private/SkeletalMeshLODModel.cpp
@@ -18,6 +18,7 @@
 #include "UObject/RenderingObjectVersion.h"
 #include "Rendering/SkeletalMeshLODImporterData.h"
 #include "UObject/FortniteMainBranchObjectVersion.h"
+#include "GPUSkinVertexFactory.h"
 
 /*-----------------------------------------------------------------------------
 FSoftSkinVertex
@@ -1023,7 +1024,7 @@ FString FSkeletalMeshLODModel::GetLODModelDeriveDataKey() const
 	return KeySuffix;
 }
 
-void FSkeletalMeshLODModel::UpdateChunkedSectionInfo()
+void FSkeletalMeshLODModel::UpdateChunkedSectionInfo(const FString& SkeletalMeshName, TArray<int32>& LODMaterialMap)
 {
 	//Look if the data is pre skeletalmesh build refactor
 	if (RawSkeletalMeshBulkData.IsBuildDataAvailable())
@@ -1034,43 +1035,64 @@ void FSkeletalMeshLODModel::UpdateChunkedSectionInfo()
 	int32 LODModelSectionNum = Sections.Num();
 	//Fill the ChunkedParentSectionIndex data, we assume that every section using the same material are chunked
 	int32 LastMaterialIndex = INDEX_NONE;
+	uint32 LastBoneCount = 0;
 	int32 CurrentParentChunkIndex = INDEX_NONE;
 	int32 OriginalIndex = 0;
+	const uint32 MaxGPUSkinBones = FGPUBaseSkinVertexFactory::GetMaxGPUSkinBones();
+	check(MaxGPUSkinBones <= FGPUBaseSkinVertexFactory::GHardwareMaxGPUSkinBones);
 
 	for (int32 LODModelSectionIndex = 0; LODModelSectionIndex < LODModelSectionNum; ++LODModelSectionIndex)
 	{
+		FSkelMeshSection& Section = Sections[LODModelSectionIndex];
+		int32 SectionMaterialIndex = Section.MaterialIndex;
+		
 		//If we have cloth on a chunked section we treat the chunked section has a parent section (this is to get the same result has before the refactor)
-		if (Sections[LODModelSectionIndex].MaterialIndex == LastMaterialIndex && !Sections[LODModelSectionIndex].ClothingData.AssetGuid.IsValid())
+		if (LastBoneCount >= MaxGPUSkinBones && SectionMaterialIndex == LastMaterialIndex && !Section.ClothingData.AssetGuid.IsValid())
 		{
-			Sections[LODModelSectionIndex].ChunkedParentSectionIndex = CurrentParentChunkIndex;
-			Sections[LODModelSectionIndex].OriginalDataSectionIndex = Sections[CurrentParentChunkIndex].OriginalDataSectionIndex;
+			Section.ChunkedParentSectionIndex = CurrentParentChunkIndex;
+			Section.OriginalDataSectionIndex = Sections[CurrentParentChunkIndex].OriginalDataSectionIndex;
 			//In case of a child section that was BONE chunked ensure it has the same setting has the original section
-			FSkelMeshSourceSectionUserData& SectionUserData = UserSectionsData.FindOrAdd(Sections[LODModelSectionIndex].OriginalDataSectionIndex);
-			Sections[LODModelSectionIndex].bDisabled = SectionUserData.bDisabled;
-			Sections[LODModelSectionIndex].bCastShadow = SectionUserData.bCastShadow;
-			Sections[LODModelSectionIndex].bRecomputeTangent = SectionUserData.bRecomputeTangent;
-			Sections[LODModelSectionIndex].GenerateUpToLodIndex = SectionUserData.GenerateUpToLodIndex;
+			FSkelMeshSourceSectionUserData& SectionUserData = UserSectionsData.FindOrAdd(Section.OriginalDataSectionIndex);
+			Section.bDisabled = SectionUserData.bDisabled;
+			Section.bCastShadow = SectionUserData.bCastShadow;
+			Section.bRecomputeTangent = SectionUserData.bRecomputeTangent;
+			Section.GenerateUpToLodIndex = SectionUserData.GenerateUpToLodIndex;
 			//Chunked section cannot have cloth, a cloth section will be a parent section
-			Sections[LODModelSectionIndex].CorrespondClothAssetIndex = INDEX_NONE;
-			Sections[LODModelSectionIndex].ClothingData.AssetGuid = FGuid();
-			Sections[LODModelSectionIndex].ClothingData.AssetLodIndex = INDEX_NONE;
+			Section.CorrespondClothAssetIndex = INDEX_NONE;
+			Section.ClothingData.AssetGuid = FGuid();
+			Section.ClothingData.AssetLodIndex = INDEX_NONE;
 		}
 		else
 		{
 			CurrentParentChunkIndex = LODModelSectionIndex;
 			FSkelMeshSourceSectionUserData& SectionUserData = UserSectionsData.FindOrAdd(OriginalIndex);
-			SectionUserData.bDisabled = Sections[LODModelSectionIndex].bDisabled;
-			SectionUserData.bCastShadow = Sections[LODModelSectionIndex].bCastShadow;
-			SectionUserData.bRecomputeTangent = Sections[LODModelSectionIndex].bRecomputeTangent;
-			SectionUserData.GenerateUpToLodIndex = Sections[LODModelSectionIndex].GenerateUpToLodIndex;
-			SectionUserData.CorrespondClothAssetIndex = Sections[LODModelSectionIndex].CorrespondClothAssetIndex;
-			SectionUserData.ClothingData.AssetGuid = Sections[LODModelSectionIndex].ClothingData.AssetGuid;
-			SectionUserData.ClothingData.AssetLodIndex = Sections[LODModelSectionIndex].ClothingData.AssetLodIndex;
-
-			Sections[LODModelSectionIndex].OriginalDataSectionIndex = OriginalIndex++;
-			Sections[LODModelSectionIndex].ChunkedParentSectionIndex = INDEX_NONE;
+			SectionUserData.bDisabled = Section.bDisabled;
+			SectionUserData.bCastShadow = Section.bCastShadow;
+			SectionUserData.bRecomputeTangent = Section.bRecomputeTangent;
+			SectionUserData.GenerateUpToLodIndex = Section.GenerateUpToLodIndex;
+			SectionUserData.CorrespondClothAssetIndex = Section.CorrespondClothAssetIndex;
+			SectionUserData.ClothingData.AssetGuid = Section.ClothingData.AssetGuid;
+			SectionUserData.ClothingData.AssetLodIndex = Section.ClothingData.AssetLodIndex;
+
+			Section.OriginalDataSectionIndex = OriginalIndex++;
+			Section.ChunkedParentSectionIndex = INDEX_NONE;
+		}
+
+		//Use the LODMaterialMap if the material is different from the original section index and set back the section materialindex to the originalDataSectionIndex
+		if (SectionMaterialIndex != Section.OriginalDataSectionIndex)
+		{
+			while (LODMaterialMap.Num() <= LODModelSectionIndex)
+			{
+				LODMaterialMap.Add(INDEX_NONE);
+			}
+			LODMaterialMap[LODModelSectionIndex] = Section.ChunkedParentSectionIndex != INDEX_NONE ? Sections[Section.ChunkedParentSectionIndex].MaterialIndex : Section.MaterialIndex;
+			Section.MaterialIndex = Section.OriginalDataSectionIndex;
 		}
-		LastMaterialIndex = Sections[LODModelSectionIndex].MaterialIndex;
+		LastMaterialIndex = SectionMaterialIndex;
+		//Set the last bone count
+		LastBoneCount = (uint32)Sections[LODModelSectionIndex].BoneMap.Num();
+		//its impossible to have more bone then the maximum allowed
+		ensureMsgf(LastBoneCount <= MaxGPUSkinBones, TEXT("Skeletalmesh(%s) section %d have more bones then its alowed (MaxGPUSkinBones: %d)."), *SkeletalMeshName, LODModelSectionIndex, LastBoneCount);
 	}
 }
 
diff --git a/Engine/Source/Runtime/Engine/Public/Rendering/SkeletalMeshLODModel.h b/Engine/Source/Runtime/Engine/Public/Rendering/SkeletalMeshLODModel.h
--- a/Engine/Source/Runtime/Engine/Public/Rendering/SkeletalMeshLODModel.h
+++ b/Engine/Source/Runtime/Engine/Public/Rendering/SkeletalMeshLODModel.h
@@ -393,7 +393,7 @@ public:
 	* This function will update the chunked information for each section. Only old data before the 
 	* skeletal mesh build refactor should need to call this function.
 	*/
-	ENGINE_API void UpdateChunkedSectionInfo();
+	ENGINE_API void UpdateChunkedSectionInfo(const FString& SkeletalMeshName, TArray<int32>& LODMaterialMap);
 };
 
 #endif // WITH_EDITOR
\ No newline at end of file
