GameModeBase - Making it so if there's no gamesession on login we error gracefully rather than crashing.

Marc.Audy

diff --git a/Engine/Source/Runtime/Engine/Private/GameModeBase.cpp b/Engine/Source/Runtime/Engine/Private/GameModeBase.cpp
--- a/Engine/Source/Runtime/Engine/Private/GameModeBase.cpp
+++ b/Engine/Source/Runtime/Engine/Private/GameModeBase.cpp
@@ -640,6 +640,12 @@ void AGameModeBase::PreLogin(const FString& Options, const FString& Address, con
 
 APlayerController* AGameModeBase::Login(UPlayer* NewPlayer, ENetRole InRemoteRole, const FString& Portal, const FString& Options, const FUniqueNetIdRepl& UniqueId, FString& ErrorMessage)
 {
+	if (GameSession == nullptr)
+	{
+		ErrorMessage = TEXT("Failed to spawn player controller, GameSession is null");
+		return nullptr;
+	}
+
 	ErrorMessage = GameSession->ApproveLogin(Options);
 	if (!ErrorMessage.IsEmpty())
 	{
