Replay recording: a simplified fidelity of a frame to be used for detection of replay file intervals acceptable for highlights recording

diff --git a/Engine/Source/Runtime/Engine/Classes/Engine/DemoNetDriver.h b/Engine/Source/Runtime/Engine/Classes/Engine/DemoNetDriver.h
--- a/Engine/Source/Runtime/Engine/Classes/Engine/DemoNetDriver.h
+++ b/Engine/Source/Runtime/Engine/Classes/Engine/DemoNetDriver.h
@@ -16,6 +16,8 @@
 #include "NetworkReplayStreaming.h"
 #include "Engine/DemoNetConnection.h"
 #include "Net/RepLayout.h"
+#include "Templates/Atomic.h"
+
 #include "DemoNetDriver.generated.h"
 
 class FNetworkNotify;
@@ -1220,4 +1222,20 @@ private:
 	{
 		return ECreateReplicationChangelistMgrFlags::SkipDeltaCustomState;
 	}
+
+//////////////////////////////////////////////////////////////////////////
+// Replay frame fidelity
+public:
+	// Simplified rating of replay frame fidelity as percentage of actors that were replicated.
+	// [0..1] where 0 means nothing was recorded this frame and 1 means full fidelity.
+	// This treats all actors equally. Assuming more important actors are prioritized higher, in general actual "fidelity"
+	// is expected to be higher than reported, which should be fine for detecting low-fidelity frame/intervals in replay file.
+	float GetLastReplayFrameFidelity() const
+	{
+		return LastReplayFrameFidelity;
+	}
+
+private:
+	TAtomic<float> LastReplayFrameFidelity{ 0 };
+//////////////////////////////////////////////////////////////////////////
 };
diff --git a/Engine/Source/Runtime/Engine/Private/DemoNetDriver.cpp b/Engine/Source/Runtime/Engine/Private/DemoNetDriver.cpp
--- a/Engine/Source/Runtime/Engine/Private/DemoNetDriver.cpp
+++ b/Engine/Source/Runtime/Engine/Private/DemoNetDriver.cpp
@@ -2196,6 +2196,9 @@ void UDemoNetDriver::TickCheckpoint()
 			ReplayStreamer->FlushCheckpoint(GetLastCheckpointTimeInMS());
 		}
 
+		// writing checkpoint means we have 100% fidelity this frame
+		LastReplayFrameFidelity = 1;
+
 		const float TotalCheckpointTimeInMS = CheckpointSaveContext.TotalCheckpointReplicationTimeSeconds * 1000.0f;
 		const float TotalCheckpointTimeWithOverheadInMS = CheckpointSaveContext.TotalCheckpointSaveTimeSeconds * 1000.0f;
 
@@ -2204,6 +2207,11 @@ void UDemoNetDriver::TickCheckpoint()
 		// we are done, out
 		CheckpointSaveContext.CheckpointSaveState = ECheckpointSaveState_Idle;
 	}
+	else
+	{
+		// nothing was recorded this frame
+		LastReplayFrameFidelity = 0;
+	}
 }
 
 bool UDemoNetDriver::ReplicateCheckpointActor(AActor* ToReplicate, UDemoNetConnection* ClientConnection, FRepActorsCheckpointParams& Params)
@@ -2824,7 +2832,8 @@ void UDemoNetDriver::TickDemoRecordFrame(float DeltaSeconds)
 		ReplicatePrioritizedActors(PrioritizedActors.GetData(), PrioritizedActors.Num(), Params);
 	}
 
-	
+	LastReplayFrameFidelity = (float)Params.NumActorsReplicated / (float)NumPrioritizedActors;
+
 	CSV_CUSTOM_STAT(Basic, DemoNumReplicatedActors, Params.NumActorsReplicated, ECsvCustomStatOp::Set);
 
 	FlushNetChecked(*ClientConnection);
