Fix crash when using undo/redo quickly after a mesh edit change

diff --git a/Engine/Source/Runtime/Engine/Public/StaticMeshResources.h b/Engine/Source/Runtime/Engine/Public/StaticMeshResources.h
--- a/Engine/Source/Runtime/Engine/Public/StaticMeshResources.h
+++ b/Engine/Source/Runtime/Engine/Public/StaticMeshResources.h
@@ -685,6 +685,8 @@ public:
 			StaticMeshComponents.Add(StaticMesh);
 		}
 
+		TSet<FSceneInterface*> Scenes;
+
 		for (TObjectIterator<UStaticMeshComponent> It; It; ++It)
 		{
 			// First, flush all deferred render updates, as they may depend on InStaticMesh.
@@ -704,10 +706,13 @@ public:
 					check(It->IsRegistered());
 					It->DestroyRenderState_Concurrent();
 					StaticMeshComponents[StaticMesh].Add(*It);
+					Scenes.Add(It->GetScene());
 				}
 			}
 		}
 
+		UpdateAllPrimitiveSceneInfosForScenes(MoveTemp(Scenes));
+
 		// Flush the rendering commands generated by the detachments.
 		// The static mesh scene proxies reference the UStaticMesh, and this ensures that they are cleaned up before the UStaticMesh changes.
 		FlushRenderingCommands();
@@ -727,6 +732,8 @@ public:
 	/** Destructor: recreates render state for all components that had their render states destroyed in the constructor. */
 	~FStaticMeshComponentRecreateRenderStateContext()
 	{
+		TSet<FSceneInterface*> Scenes;
+
 		for (const auto& MeshComponents : StaticMeshComponents)
 		{
 			for (UStaticMeshComponent * Component : MeshComponents.Value)
@@ -746,9 +753,12 @@ public:
 				if (Component->IsRegistered() && !Component->bRenderStateCreated)
 				{
 					Component->CreateRenderState_Concurrent();
+					Scenes.Add(Component->GetScene());
 				}
 			}
 		}
+
+		UpdateAllPrimitiveSceneInfosForScenes(MoveTemp(Scenes));
 	}
 
 private:
