Added inertial blending for curves and fixed pose meltdown caused by continuous intertial blend interruptions

diff --git a/Engine/Source/Runtime/Engine/Classes/Animation/AnimNode_Inertialization.h b/Engine/Source/Runtime/Engine/Classes/Animation/AnimNode_Inertialization.h
--- a/Engine/Source/Runtime/Engine/Classes/Animation/AnimNode_Inertialization.h
+++ b/Engine/Source/Runtime/Engine/Classes/Animation/AnimNode_Inertialization.h
@@ -5,6 +5,7 @@
 #include "CoreMinimal.h"
 #include "UObject/ObjectMacros.h"
 #include "Animation/AnimNodeBase.h"
+#include "Animation/AnimCurveTypes.h"
 #include "AnimNode_Inertialization.generated.h"
 
 
@@ -60,6 +61,11 @@ struct FInertializationPose
 	//
 	TArray<EInertializationBoneState> BoneStates;
 
+	// Snapshot of active curves
+	// 
+	FBlendedHeapCurve Curves;
+	TArray<uint16> CurveUIDToArrayIndexLUT;
+
 	FName AttachParentName;
 	float DeltaTime;
 
@@ -70,7 +76,7 @@ struct FInertializationPose
 	{
 	}
 
-	void InitFrom(const FCompactPose& Pose, const FTransform& InComponentTransform, const FName& InAttachParentName, float InDeltaTime);
+	void InitFrom(const FCompactPose& Pose, const FBlendedCurve& InCurves, const FTransform& InComponentTransform, const FName& InAttachParentName, float InDeltaTime);
 };
 
 
@@ -119,6 +125,26 @@ struct FInertializationBoneDiff
 	}
 };
 
+USTRUCT()
+struct FInertializationCurveDiff
+{
+	GENERATED_BODY();
+
+	float Delta;
+	float Derivative;
+
+	FInertializationCurveDiff()
+		: Delta(0.0f)
+		, Derivative(0.0f)
+	{}
+
+	void Clear()
+	{
+		Delta = 0.0f;
+		Derivative = 0.0f;
+	}
+};
+
 
 USTRUCT()
 struct FInertializationPoseDiff
@@ -133,6 +159,7 @@ struct FInertializationPoseDiff
 	void Reset()
 	{
 		BoneDiffs.Empty();
+		CurveDiffs.Empty();
 		InertializationSpace = EInertializationSpace::Default;
 	}
 
@@ -144,11 +171,11 @@ struct FInertializationPoseDiff
 	// Prev1				the previous frame's pose
 	// Prev2				the pose from two frames before
 	//
-	void InitFrom(const FCompactPose& Pose, const FTransform& ComponentTransform, const FName& AttachParentName, const FInertializationPose& Prev1, const FInertializationPose& Prev2);
+	void InitFrom(const FCompactPose& Pose, const FBlendedCurve& Curves, const FTransform& ComponentTransform, const FName& AttachParentName, const FInertializationPose& Prev1, const FInertializationPose& Prev2);
 
 	// Apply this difference to a pose, decaying over time as InertializationElapsedTime approaches InertializationDuration
 	//
-	void ApplyTo(FCompactPose& Pose, float InertializationElapsedTime, float InertializationDuration) const;
+	void ApplyTo(FCompactPose& Pose, FBlendedCurve& Curves, float InertializationElapsedTime, float InertializationDuration) const;
 
 	// Get the inertialization space for this pose diff (for debug display)
 	//
@@ -164,6 +191,9 @@ private:
 	// Bone differences indexed by skeleton bone index
 	TArray<FInertializationBoneDiff> BoneDiffs;
 
+	// Curve differences indexed by CurveID
+	TArray<FInertializationCurveDiff> CurveDiffs;
+
 	// Inertialization space (local vs world for situations where we wish to correct a world-space discontinuity such as an abrupt orientation change)
 	EInertializationSpace InertializationSpace;
 };
@@ -251,7 +281,11 @@ private:
 	EInertializationState InertializationState;
 	float InertializationElapsedTime;
 	float InertializationDuration;
+	float InertializationDeficit;
 
 	// Inertialization pose differences
 	FInertializationPoseDiff InertializationPoseDiff;
+
+	// Reset inertialization timing and state
+	void Deactivate();
 };
diff --git a/Engine/Source/Runtime/Engine/Private/Animation/AnimNode_Inertialization.cpp b/Engine/Source/Runtime/Engine/Private/Animation/AnimNode_Inertialization.cpp
--- a/Engine/Source/Runtime/Engine/Private/Animation/AnimNode_Inertialization.cpp
+++ b/Engine/Source/Runtime/Engine/Private/Animation/AnimNode_Inertialization.cpp
@@ -9,6 +9,7 @@
 
 TAutoConsoleVariable<int32> CVarAnimInertializationEnable(TEXT("a.AnimNode.Inertialization.Enable"), 1, TEXT("Enable / Disable Inertialization"));
 TAutoConsoleVariable<int32> CVarAnimInertializationIgnoreVelocity(TEXT("a.AnimNode.Inertialization.IgnoreVelocity"), 0, TEXT("Ignore velocity information during Inertialization (effectively reverting to a quintic diff blend)"));
+TAutoConsoleVariable<int32> CVarAnimInertializationIgnoreDeficit(TEXT("a.AnimNode.Inertialization.IgnoreDeficit"), 0, TEXT("Ignore inertialization time deficit caused by interruptions"));
 
 
 static constexpr int32 INERTIALIZATION_MAX_POSE_SNAPSHOTS = 2;
@@ -22,6 +23,7 @@ FAnimNode_Inertialization::FAnimNode_Inertialization()
 	, InertializationState(EInertializationState::Inactive)
 	, InertializationElapsedTime(0.0f)
 	, InertializationDuration(0.0f)
+	, InertializationDeficit(0.0f)
 {
 }
 
@@ -58,9 +60,7 @@ void FAnimNode_Inertialization::Initialize_AnyThread(const FAnimationInitializeC
 
 	TeleportType = ETeleportType::None;
 
-	InertializationState = EInertializationState::Inactive;
-	InertializationElapsedTime = 0.0f;
-	InertializationDuration = 0.0f;
+	Deactivate();
 
 	InertializationPoseDiff.Reset();
 }
@@ -94,9 +94,7 @@ void FAnimNode_Inertialization::Evaluate_AnyThread(FPoseContext& Output)
 	// Disable inertialization if requested (for testing / debugging)
 	if (!CVarAnimInertializationEnable.GetValueOnAnyThread())
 	{
-		InertializationState = EInertializationState::Inactive;
-		InertializationElapsedTime = 0.0f;
-		InertializationDuration = 0.0f;
+		Deactivate();
 
 		// Clear the pose history
 		PoseSnapshots.Reset();
@@ -111,9 +109,19 @@ void FAnimNode_Inertialization::Evaluate_AnyThread(FPoseContext& Output)
 	// Update the inertialization state if a new inertialization request is pending
 	if (Duration > SMALL_NUMBER && PoseSnapshots.Num() > 0)
 	{
+		float AppliedDeficit = 0.0f;
+		if (InertializationState == EInertializationState::Active)
+		{
+			// An active inertialization is being interrupted. Keep track of the lost inertialization time
+			// and reduce future durations if interruptions continue. Without this mitigation, 
+			// repeated interruptions will lead to a degenerate pose because the pose target is unstable.
+			bool bApplyDeficit = InertializationDeficit > 0.0f && !CVarAnimInertializationIgnoreDeficit.GetValueOnAnyThread();
+			InertializationDeficit += InertializationDuration - InertializationElapsedTime;
+			AppliedDeficit = bApplyDeficit ? FMath::Min(Duration, InertializationDeficit) : 0.0f;
+		}
 		InertializationState = EInertializationState::Pending;
 		InertializationElapsedTime = 0.0f;
-		InertializationDuration = Duration;
+		InertializationDuration = Duration - AppliedDeficit;
 	}
 
 	// Update the inertialization timer
@@ -122,9 +130,12 @@ void FAnimNode_Inertialization::Evaluate_AnyThread(FPoseContext& Output)
 		InertializationElapsedTime += DeltaTime;
 		if (InertializationElapsedTime >= InertializationDuration)
 		{
-			InertializationState = EInertializationState::Inactive;
-			InertializationElapsedTime = 0.0f;
-			InertializationDuration = 0.0f;
+			Deactivate();
+		}
+		else
+		{
+			// Pay down the accumulated deficit caused by interruptions
+			InertializationDeficit -= FMath::Min(InertializationDeficit, DeltaTime);
 		}
 	}
 
@@ -154,9 +165,7 @@ void FAnimNode_Inertialization::Evaluate_AnyThread(FPoseContext& Output)
 		// Cancel inertialization requests during teleports
 		if (InertializationState == EInertializationState::Pending)
 		{
-			InertializationState = EInertializationState::Inactive;
-			InertializationElapsedTime = 0.0f;
-			InertializationDuration = 0.0f;
+			Deactivate();
 		}
 
 		// Clear the time accumulator during teleports (in order to invalidate any recorded velocities during the teleport)
@@ -200,7 +209,7 @@ void FAnimNode_Inertialization::Evaluate_AnyThread(FPoseContext& Output)
 	{
 		// Add the pose to the end of the buffer
 		FInertializationPose& Snapshot = PoseSnapshots.AddDefaulted_GetRef();
-		Snapshot.InitFrom(Output.Pose, ComponentTransform, AttachParentName, DeltaTime);
+		Snapshot.InitFrom(Output.Pose, Output.Curve, ComponentTransform, AttachParentName, DeltaTime);
 	}
 	else
 	{
@@ -213,7 +222,7 @@ void FAnimNode_Inertialization::Evaluate_AnyThread(FPoseContext& Output)
 		// Overwrite the (now irrelevant) pose in the last slot with the new post snapshot
 		// (thereby avoiding the reallocation costs we would have incurred had we simply added a new pose at the end)
 		FInertializationPose& Snapshot = PoseSnapshots[INERTIALIZATION_MAX_POSE_SNAPSHOTS - 1];
-		Snapshot.InitFrom(Output.Pose, ComponentTransform, AttachParentName, DeltaTime);
+		Snapshot.InitFrom(Output.Pose, Output.Curve, ComponentTransform, AttachParentName, DeltaTime);
 	}
 
 	// Reset the time accumulator and teleport state
@@ -228,18 +237,20 @@ void FAnimNode_Inertialization::GatherDebugData(FNodeDebugData& DebugData)
 
 	if (InertializationDuration > KINDA_SMALL_NUMBER)
 	{
-		DebugLine += FString::Printf(TEXT("('%s' Time: %.3f / %.3f (%.0f%%))"),
+		DebugLine += FString::Printf(TEXT("('%s' Time: %.3f / %.3f (%.0f%%) [%.3f])"),
 			*UEnum::GetValueAsString(InertializationState),
 			InertializationElapsedTime,
 			InertializationDuration,
-			100.0f * InertializationElapsedTime / InertializationDuration);
+			100.0f * InertializationElapsedTime / InertializationDuration,
+			InertializationDeficit);
 	}
 	else
 	{
-		DebugLine += FString::Printf(TEXT("('%s' Time: %.3f / %.3f)"),
+		DebugLine += FString::Printf(TEXT("('%s' Time: %.3f / %.3f [%.3f])"),
 			*UEnum::GetValueAsString(InertializationState),
 			InertializationElapsedTime,
-			InertializationDuration);
+			InertializationDuration,
+			InertializationDeficit);
 	}
 	DebugData.AddDebugItem(DebugLine);
 }
@@ -304,17 +315,25 @@ void FAnimNode_Inertialization::StartInertialization(FPoseContext& Context, FIne
 		}
 	}
 
-	OutPoseDiff.InitFrom(Context.Pose, Context.AnimInstanceProxy->GetComponentTransform(), AttachParentName, PreviousPose1, PreviousPose2);
+	OutPoseDiff.InitFrom(Context.Pose, Context.Curve, Context.AnimInstanceProxy->GetComponentTransform(), AttachParentName, PreviousPose1, PreviousPose2);
 }
 
 
 void FAnimNode_Inertialization::ApplyInertialization(FPoseContext& Context, const FInertializationPoseDiff& PoseDiff, float ElapsedTime, float Duration)
 {
-	PoseDiff.ApplyTo(Context.Pose, ElapsedTime, Duration);
+	PoseDiff.ApplyTo(Context.Pose, Context.Curve, ElapsedTime, Duration);
+}
+
+void FAnimNode_Inertialization::Deactivate()
+{
+	InertializationState = EInertializationState::Inactive;
+	InertializationElapsedTime = 0.0f;
+	InertializationDuration = 0.0f;
+	InertializationDeficit = 0.0f;
 }
 
 
-void FInertializationPose::InitFrom(const FCompactPose& Pose, const FTransform& InComponentTransform, const FName& InAttachParentName, float InDeltaTime)
+void FInertializationPose::InitFrom(const FCompactPose& Pose, const FBlendedCurve& InCurves, const FTransform& InComponentTransform, const FName& InAttachParentName, float InDeltaTime)
 {
 	const FBoneContainer& BoneContainer = Pose.GetBoneContainer();
 
@@ -333,6 +352,15 @@ void FInertializationPose::InitFrom(const FCompactPose& Pose, const FTransform&
 		}
 	}
 
+	// Copy curves and make a new copy of the curve lookup table as it may change from frame to frame
+	Curves.CopyFrom(InCurves);
+	CurveUIDToArrayIndexLUT.Reset();
+	if (InCurves.UIDToArrayIndexLUT)
+	{
+		CurveUIDToArrayIndexLUT = *InCurves.UIDToArrayIndexLUT;
+		Curves.UIDToArrayIndexLUT = &CurveUIDToArrayIndexLUT;
+	}
+
 	ComponentTransform = InComponentTransform;
 	AttachParentName = InAttachParentName;
 	DeltaTime = InDeltaTime;
@@ -342,11 +370,12 @@ void FInertializationPose::InitFrom(const FCompactPose& Pose, const FTransform&
 // Initialize the pose difference from the current pose and the two previous snapshots
 //
 // Pose			the current frame's pose
-// Prev1		the previous frame's pose
-// Prev2		the pose from two frames before
+// Curves       the current frame's curves
+// Prev1		the previous frame's pose and curves
+// Prev2		the pose and curves from two frames before
 // DeltaTime	the time elapsed between Prev1 and Pose
 //
-void FInertializationPoseDiff::InitFrom(const FCompactPose& Pose, const FTransform& ComponentTransform, const FName& AttachParentName, const FInertializationPose& Prev1, const FInertializationPose& Prev2)
+void FInertializationPoseDiff::InitFrom(const FCompactPose& Pose, const FBlendedCurve& Curves, const FTransform& ComponentTransform, const FName& AttachParentName, const FInertializationPose& Prev1, const FInertializationPose& Prev2)
 {
 	const FBoneContainer& BoneContainer = Pose.GetBoneContainer();
 
@@ -514,15 +543,46 @@ void FInertializationPoseDiff::InitFrom(const FCompactPose& Pose, const FTransfo
 			}
 		}
 	}
+
+	// Compute the curve differences
+	const int32 CurveNum = Curves.IsValid() ? Curves.UIDToArrayIndexLUT->Num() : 0;
+	CurveDiffs.Empty(CurveNum);
+	CurveDiffs.AddZeroed(CurveNum);
+	for (int32 CurveUID = 0; CurveUID != CurveNum; ++CurveUID)
+	{
+		const int32 CurrIdx = Curves.GetArrayIndexByUID(CurveUID);
+		const int32 Prev1Idx = Prev1.Curves.GetArrayIndexByUID(CurveUID);
+		if (CurrIdx == INDEX_NONE || Prev1Idx == INDEX_NONE)
+		{
+			// CurveDiff is zeroed
+			continue;
+		}
+
+		const FCurveElement& CurrElement = Curves.Elements[CurrIdx];
+		const FCurveElement& Prev1Element = Prev1.Curves.Elements[Prev1Idx];
+		FInertializationCurveDiff& CurveDiff = CurveDiffs[CurveUID];
+
+		// Note we intentionally ignore FCurveElement::bValid. We want to ease in/out when only one
+		// curve is valid, and we'll compute a zero delta and derivative when both are invalid.
+		CurveDiff.Delta = Prev1Element.Value - CurrElement.Value;
+
+		const int32 Prev2Idx = Prev2.Curves.GetArrayIndexByUID(CurveUID);
+		if (Prev2Idx != INDEX_NONE && Prev1.DeltaTime > KINDA_SMALL_NUMBER)
+		{
+			const FCurveElement& Prev2Element = Prev2.Curves.Elements[Prev2Idx];
+			CurveDiff.Derivative = (Prev1Element.Value - Prev2Element.Value) / Prev1.DeltaTime;
+		}
+	}
 }
 
 
-// Apply this difference to a pose, decaying over time as InertializationElapsedTime approaches InertializationDuration
+// Apply this difference to a pose and a set of curves, decaying over time as InertializationElapsedTime approaches InertializationDuration
 //
-void FInertializationPoseDiff::ApplyTo(FCompactPose& Pose, float InertializationElapsedTime, float InertializationDuration) const
+void FInertializationPoseDiff::ApplyTo(FCompactPose& Pose, FBlendedCurve& Curves, float InertializationElapsedTime, float InertializationDuration) const
 {
 	const FBoneContainer& BoneContainer = Pose.GetBoneContainer();
 
+	// Apply pose difference
 	for (FCompactPoseBoneIndex BoneIndex : Pose.ForEachBoneIndex())
 	{
 		const int32 SkeletonPoseBoneIndex = BoneContainer.GetSkeletonIndex(BoneIndex);
@@ -547,6 +607,32 @@ void FInertializationPoseDiff::ApplyTo(FCompactPose& Pose, float Inertialization
 			Pose[BoneIndex].SetScale3D(S + Pose[BoneIndex].GetScale3D());
 		}
 	}
+
+	Pose.NormalizeRotations();
+
+	// Apply curve differences
+	if (Curves.IsValid())
+	{
+		// Note Curves.Elements is indexed indirectly via lookup table while CurveDiffs is indexed directly by curve ID
+		const int32 CurveNum = FMath::Min(Curves.UIDToArrayIndexLUT->Num(), CurveDiffs.Num());
+		for (int32 CurveUID = 0; CurveUID != CurveNum; ++CurveUID)
+		{
+			const int32 CurrIdx = Curves.GetArrayIndexByUID(CurveUID);
+			if (CurrIdx == INDEX_NONE)
+			{
+				continue;
+			}
+
+			FCurveElement& CurrElement = Curves.Elements[CurrIdx];
+			const FInertializationCurveDiff& CurveDiff = CurveDiffs[CurveUID];
+			const float C = CalcInertialFloat(CurveDiff.Delta, CurveDiff.Derivative, InertializationElapsedTime, InertializationDuration);
+			if (C != 0.0f)
+			{
+				CurrElement.Value += C;
+				CurrElement.bValid = true;
+			}
+		}
+	}
 }
 
 
