Copying material function calls will update their input pin appearance correctly.

diff --git a/Engine/Source/Runtime/Engine/Private/Materials/MaterialExpressions.cpp b/Engine/Source/Runtime/Engine/Private/Materials/MaterialExpressions.cpp
--- a/Engine/Source/Runtime/Engine/Private/Materials/MaterialExpressions.cpp
+++ b/Engine/Source/Runtime/Engine/Private/Materials/MaterialExpressions.cpp
@@ -234,6 +234,7 @@
 #include "MaterialGraph/MaterialGraphNode.h"
 #include "Framework/Notifications/NotificationManager.h"
 #include "Widgets/Notifications/SNotificationList.h"
+#include "MaterialGraph/MaterialGraphSchema.h"
 #endif //WITH_EDITOR
 #include "Materials/MaterialInstanceConstant.h"
 #include "Curves/CurveLinearColorAtlas.h"
@@ -12931,6 +12932,14 @@ void UMaterialExpressionMaterialFunctionCall::UpdateFromFunctionResource(bool bR
 				CurrentInput.Input = OriginalInput->Input;
 				CurrentInput.Input.InputName = TempInputName;
 			}
+
+			if (GraphNode)
+			{
+				if (CurrentInput.ExpressionInput->bUsePreviewValueAsDefault)
+				{
+					bRecreateAndLinkNode = true;
+				}
+			}
 		}
 
 		for (int32 OutputIndex = 0; OutputIndex < FunctionOutputs.Num(); OutputIndex++)
