Improved comments on Blueprint-exposed Floor/Trunc/Ceil/Round functions and fixed a precision loss issue in the int64 versions

diff --git a/Engine/Source/Runtime/Engine/Classes/Kismet/KismetMathLibrary.h b/Engine/Source/Runtime/Engine/Classes/Kismet/KismetMathLibrary.h
--- a/Engine/Source/Runtime/Engine/Classes/Kismet/KismetMathLibrary.h
+++ b/Engine/Source/Runtime/Engine/Classes/Kismet/KismetMathLibrary.h
@@ -725,38 +725,38 @@ class ENGINE_API UKismetMathLibrary : public UBlueprintFunctionLibrary
 	UFUNCTION(BlueprintPure, meta = (DisplayName = "Ease", BlueprintInternalUseOnly = "true"), Category = "Math|Interpolation")
 	static float Ease(float A, float B, float Alpha, TEnumAsByte<EEasingFunc::Type> EasingFunc, float BlendExp = 2, int32 Steps = 2);
 
-	/** Rounds A to the nearest integer */
+	/** Rounds A to the nearest integer (e.g., -1.6 becomes -2 and 1.6 becomes 2) */
 	UFUNCTION(BlueprintPure, Category="Math|Float")
 	static int32 Round(float A);
 
-	/** Rounds A to the largest previous integer */
+	/** Rounds A down towards negative infinity / down to the previous integer (e.g., -1.6 becomes -2 and 1.6 becomes 1) */
 	UFUNCTION(BlueprintPure, meta=(DisplayName = "Floor"), Category="Math|Float")
 	static int32 FFloor(float A);
 	
-	/** Rounds A to an integer with truncation towards zero.  (e.g. -1.7 truncated to -1, 2.8 truncated to 2) */
+	/** Rounds A towards zero, truncating the fractional part (e.g., -1.6 becomes -1 and 1.6 becomes 1) */
 	UFUNCTION(BlueprintPure, meta=(DisplayName = "Truncate", BlueprintAutocast), Category="Math|Float")
 	static int32 FTrunc(float A);
 
-	/** Rounds A to the nearest 32 bit integer then upconverts to 64 bit integer */
+	/** Rounds A up towards positive infinity / up to the next integer (e.g., -1.6 becomes -1 and 1.6 becomes 2) */
+	UFUNCTION(BlueprintPure, Category = "Math|Float", meta=(DisplayName="Ceil"))
+	static int32 FCeil(float A);
+
+	/** Rounds A to the nearest integer (e.g., -1.6 becomes -2 and 1.6 becomes 2) */
 	UFUNCTION(BlueprintPure, meta = (DisplayName = "Round to Int64"), Category = "Math|Float")
 	static int64 Round64(float A);
 
-	/** Rounds A to the largest previous 32 bit integer then upconverts to 64 bit integer */
+	/** Rounds A down towards negative infinity / down to the previous integer (e.g., -1.6 becomes -2 and 1.6 becomes 1) */
 	UFUNCTION(BlueprintPure, meta = (DisplayName = "Floor to Int64"), Category = "Math|Float")
 	static int64 FFloor64(float A);
 
-	/** Rounds A to an 32 bit integer with truncation towards zero then upconverts to 64 bit integer.  (e.g. -1.7 truncated to -1, 2.8 truncated to 2) */
+	/** Rounds A towards zero, truncating the fractional part (e.g., -1.6 becomes -1 and 1.6 becomes 1) */
 	UFUNCTION(BlueprintPure, meta = (DisplayName = "Truncate to Int64", BlueprintAutocast), Category = "Math|Float")
 	static int64 FTrunc64(float A);
 
-	/** Rounds A to the smallest following 32 bit integer then upconverts to 64 bit integer */
-	UFUNCTION(BlueprintPure, meta = (DisplayName = "Floor to Int64"), Category = "Math|Float")
+	/** Rounds A up towards positive infinity / up to the next integer (e.g., -1.6 becomes -1 and 1.6 becomes 2) */
+	UFUNCTION(BlueprintPure, meta = (DisplayName = "Ceil to Int64"), Category = "Math|Float")
 	static int64 FCeil64(float A);
 
-	/** Rounds A to the smallest following integer */
-	UFUNCTION(BlueprintPure, Category="Math|Float")
-	static int32 FCeil(float A);
-
 	/** Returns the number of times Divisor will go into Dividend (i.e., Dividend divided by Divisor), as well as the remainder */
 	UFUNCTION(BlueprintPure, meta=(DisplayName = "Division (whole and remainder)"), Category="Math|Float")
 	static int32 FMod(float Dividend, float Divisor, float& Remainder);
diff --git a/Engine/Source/Runtime/Engine/Classes/Kismet/KismetMathLibrary.inl b/Engine/Source/Runtime/Engine/Classes/Kismet/KismetMathLibrary.inl
--- a/Engine/Source/Runtime/Engine/Classes/Kismet/KismetMathLibrary.inl
+++ b/Engine/Source/Runtime/Engine/Classes/Kismet/KismetMathLibrary.inl
@@ -696,25 +696,25 @@ int32 UKismetMathLibrary::FTrunc(float A)
 KISMET_MATH_FORCEINLINE
 int64 UKismetMathLibrary::Round64(float A)
 {
-	return (int64)FMath::RoundToInt(A);
+	return (int64)FMath::RoundToFloat(A);
 }
 
 KISMET_MATH_FORCEINLINE
 int64 UKismetMathLibrary::FFloor64(float A)
 {
-	return (int64)FMath::FloorToInt(A);
+	return (int64)FMath::FloorToFloat(A);
 }
 
 KISMET_MATH_FORCEINLINE
 int64 UKismetMathLibrary::FTrunc64(float A)
 {
-	return (int64)FMath::TruncToInt(A);
+	return (int64)FMath::TruncToFloat(A);
 }
 
 KISMET_MATH_FORCEINLINE
 int64 UKismetMathLibrary::FCeil64(float A)
 {
-	return (int64)FMath::CeilToInt(A);
+	return (int64)FMath::CeilToFloat(A);
 }
 
 KISMET_MATH_FORCEINLINE
