Harden RenameGeneratedClasses to avoid CDO name collisions

diff --git a/Engine/Source/Runtime/Engine/Private/Blueprint.cpp b/Engine/Source/Runtime/Engine/Private/Blueprint.cpp
--- a/Engine/Source/Runtime/Engine/Private/Blueprint.cpp
+++ b/Engine/Source/Runtime/Engine/Private/Blueprint.cpp
@@ -510,12 +510,28 @@ bool UBlueprint::RenameGeneratedClasses( const TCHAR* InName, UObject* NewOuter,
 
 	if(bRenameGeneratedClasses)
 	{
+		const auto TryFreeCDOName = [](UClass* ForClass, UObject* ToOuter, ERenameFlags InFlags)
+		{
+			if(ForClass->ClassDefaultObject)
+			{
+				FName CDOName = ForClass->GetDefaultObjectName();
+				
+				if(UObject* Obj = StaticFindObjectFast(UObject::StaticClass(), ToOuter, CDOName))
+				{
+					FName NewName = MakeUniqueObjectName(ToOuter, Obj->GetClass(), CDOName);
+					Obj->Rename(*(NewName.ToString()), ToOuter, InFlags|REN_ForceNoResetLoaders|REN_DontCreateRedirectors);
+				}
+			}
+		};
+
 		FName SkelClassName, GenClassName;
 		GetBlueprintClassNames(GenClassName, SkelClassName, FName(InName));
 
 		UPackage* NewTopLevelObjectOuter = NewOuter ? NewOuter->GetOutermost() : NULL;
 		if (GeneratedClass != NULL)
 		{
+			// check for collision of CDO name, move aside if necessary:
+			TryFreeCDOName(GeneratedClass, NewTopLevelObjectOuter, Flags);
 			bool bMovedOK = GeneratedClass->Rename(*GenClassName.ToString(), NewTopLevelObjectOuter, Flags);
 			if (!bMovedOK)
 			{
@@ -526,6 +542,7 @@ bool UBlueprint::RenameGeneratedClasses( const TCHAR* InName, UObject* NewOuter,
 		// Also move skeleton class, if different from generated class, to new package (again, to create redirector)
 		if (SkeletonGeneratedClass != NULL && SkeletonGeneratedClass != GeneratedClass)
 		{
+			TryFreeCDOName(SkeletonGeneratedClass, NewTopLevelObjectOuter, Flags);
 			bool bMovedOK = SkeletonGeneratedClass->Rename(*SkelClassName.ToString(), NewTopLevelObjectOuter, Flags);
 			if (!bMovedOK)
 			{
