Fixed crash in FAnimNode_SaveCachedPose
#rnx

diff --git a/Engine/Source/Runtime/Engine/Classes/Animation/AnimNodeBase.h b/Engine/Source/Runtime/Engine/Classes/Animation/AnimNodeBase.h
--- a/Engine/Source/Runtime/Engine/Classes/Animation/AnimNodeBase.h
+++ b/Engine/Source/Runtime/Engine/Classes/Animation/AnimNodeBase.h
@@ -77,7 +77,7 @@ struct FAnimNodeTracker
 		Map.Reserve(Source.Map.Num());
 		for (const auto& Iter : Source.Map)
 		{
-			if (Iter.Value.Top())
+			if (Iter.Value.Num() != 0)
 			{
 				FNodeStack& Stack = Map.Add(Iter.Key);
 				Stack.Push(Iter.Value.Top());
