Added scoped hierarchical counters for inertialization
#rnx

diff --git a/Engine/Source/Runtime/Engine/Private/Animation/AnimNode_Inertialization.cpp b/Engine/Source/Runtime/Engine/Private/Animation/AnimNode_Inertialization.cpp
--- a/Engine/Source/Runtime/Engine/Private/Animation/AnimNode_Inertialization.cpp
+++ b/Engine/Source/Runtime/Engine/Private/Animation/AnimNode_Inertialization.cpp
@@ -51,6 +51,8 @@ void FAnimNode_Inertialization::RequestInertialization(float Duration)
 
 void FAnimNode_Inertialization::Initialize_AnyThread(const FAnimationInitializeContext& Context)
 {
+	DECLARE_SCOPE_HIERARCHICAL_COUNTER_ANIMNODE(Initialize_AnyThread);
+
 	FAnimNode_Base::Initialize_AnyThread(Context);
 	Source.Initialize(Context);
 
@@ -68,6 +70,8 @@ void FAnimNode_Inertialization::Initialize_AnyThread(const FAnimationInitializeC
 
 void FAnimNode_Inertialization::CacheBones_AnyThread(const FAnimationCacheBonesContext& Context)
 {
+	DECLARE_SCOPE_HIERARCHICAL_COUNTER_ANIMNODE(CacheBones_AnyThread);
+
 	FAnimNode_Base::CacheBones_AnyThread(Context);
 	Source.CacheBones(Context);
 }
@@ -75,6 +79,8 @@ void FAnimNode_Inertialization::CacheBones_AnyThread(const FAnimationCacheBonesC
 
 void FAnimNode_Inertialization::Update_AnyThread(const FAnimationUpdateContext& Context)
 {
+	DECLARE_SCOPE_HIERARCHICAL_COUNTER_ANIMNODE(Update_AnyThread);
+
 	FScopedAnimNodeTracker Tracked = Context.TrackAncestor(this);
 
 	Source.Update(Context);
@@ -86,6 +92,8 @@ void FAnimNode_Inertialization::Update_AnyThread(const FAnimationUpdateContext&
 
 void FAnimNode_Inertialization::Evaluate_AnyThread(FPoseContext& Output)
 {
+	DECLARE_SCOPE_HIERARCHICAL_COUNTER_ANIMNODE(Evaluate_AnyThread);
+
 	Source.Evaluate(Output);
 
 	// Extract any pending inertialization request
@@ -233,6 +241,8 @@ void FAnimNode_Inertialization::Evaluate_AnyThread(FPoseContext& Output)
 
 void FAnimNode_Inertialization::GatherDebugData(FNodeDebugData& DebugData)
 {
+	DECLARE_SCOPE_HIERARCHICAL_COUNTER_ANIMNODE(GatherDebugData);
+
 	FString DebugLine = DebugData.GetNodeName(this);
 
 	if (InertializationDuration > KINDA_SMALL_NUMBER)
