Remove inertialization deficit accumulation, but keep deficit tracking and duration reduction. Just setting the deficit on interruption is sufficient to avoid pose meltdown. This allows for more frequent inertialization interruptions with fewer pops while still preventing meltdown.
#rnx

diff --git a/Engine/Source/Runtime/Engine/Private/Animation/AnimNode_Inertialization.cpp b/Engine/Source/Runtime/Engine/Private/Animation/AnimNode_Inertialization.cpp
--- a/Engine/Source/Runtime/Engine/Private/Animation/AnimNode_Inertialization.cpp
+++ b/Engine/Source/Runtime/Engine/Private/Animation/AnimNode_Inertialization.cpp
@@ -124,7 +124,7 @@ void FAnimNode_Inertialization::Evaluate_AnyThread(FPoseContext& Output)
 			// and reduce future durations if interruptions continue. Without this mitigation, 
 			// repeated interruptions will lead to a degenerate pose because the pose target is unstable.
 			bool bApplyDeficit = InertializationDeficit > 0.0f && !CVarAnimInertializationIgnoreDeficit.GetValueOnAnyThread();
-			InertializationDeficit += InertializationDuration - InertializationElapsedTime;
+			InertializationDeficit = InertializationDuration - InertializationElapsedTime;
 			AppliedDeficit = bApplyDeficit ? FMath::Min(Duration, InertializationDeficit) : 0.0f;
 		}
 		InertializationState = EInertializationState::Pending;
