Added improved LODBias description to the InvestigateTexture command :
- Bucket=XXX : where XXX is the memory bucked name { Largest ... Tiniest }
- [UI:XXX] : UI bias from "r.UITextureLODBias"
- [LODGroup.Bias:XXX] : LOD group bias from TextureLODGroups settings (culled out in cooked data)
- [LODGroup.MaxRes:XXX] : LOD group max resolution from TextureLODGroups settings
- [Asset.Bias:XXX] : Texture2D asset bias
- [Asset.Cine:XXX] : Bias from the asset cinematic mips
- [Asset.Opt:XXX] : Bias from the optional mips when they are not available
- [RHI.MaxRes:XXX] : Bias from the RHI limit (see GMaxTextureMipCount)
- [Global:XXX] : Bias from "r.Streaming.MipBias" when using a global mip bias "r.Streaming.UsePerTextureBias=0"
- [Budget:XXX] : Bias from memory constraint when in overbudget

diff --git a/Engine/Source/Runtime/Engine/Private/Streaming/StreamingManagerTexture.cpp b/Engine/Source/Runtime/Engine/Private/Streaming/StreamingManagerTexture.cpp
--- a/Engine/Source/Runtime/Engine/Private/Streaming/StreamingManagerTexture.cpp
+++ b/Engine/Source/Runtime/Engine/Private/Streaming/StreamingManagerTexture.cpp
@@ -2244,7 +2244,16 @@ bool FRenderAssetStreamingManager::HandleInvestigateRenderAssetCommand(const TCH
 					LODGroupName = StaticMesh->LODGroup.ToString();
 				}
 #endif
-				UE_LOG(LogContentStreaming, Log, TEXT("  LOD group:   %s"), *LODGroupName);
+				// put this up in some shared code somewhere in FGenericPlatformMemory
+				const TCHAR* BucketNames[] = { TEXT("Largest"), TEXT("Larger"), TEXT("Default"), TEXT("Smaller"), TEXT("Smallest"), TEXT("Tiniest") };
+				if ((int32)FPlatformMemory::GetMemorySizeBucket() < UE_ARRAY_COUNT(BucketNames))
+				{
+					UE_LOG(LogContentStreaming, Log, TEXT("  LOD group:   %s [Bucket=%s]"), *LODGroupName, BucketNames[(int32)FPlatformMemory::GetMemorySizeBucket()]);
+				}
+				else
+				{
+					UE_LOG(LogContentStreaming, Log, TEXT("  LOD group:   %s [Unkown Bucket]"), *LODGroupName);
+				}
 
 				if (RenderAsset->bGlobalForceMipLevelsToBeResident)
 				{
@@ -2278,11 +2287,105 @@ bool FRenderAssetStreamingManager::HandleInvestigateRenderAssetCommand(const TCH
 					UE_LOG(LogContentStreaming, Log, TEXT("  Current LOD index: %d"), CurrentMipIndex);
 					UE_LOG(LogContentStreaming, Log, TEXT("  Wanted LOD index: %d"), WantedMipIndex);
 				}
-				UE_LOG(LogContentStreaming, Log, TEXT("  Allowed mips:        %d-%d"), StreamingRenderAsset.MinAllowedMips, StreamingRenderAsset.MaxAllowedMips);
+				UE_LOG(LogContentStreaming, Log, TEXT("  Allowed mips:        %d-%d [%d]"), StreamingRenderAsset.MinAllowedMips, StreamingRenderAsset.MaxAllowedMips, StreamingRenderAsset.MipCount);
 				UE_LOG(LogContentStreaming, Log, TEXT("  LoadOrder Priority:  %d"), StreamingRenderAsset.LoadOrderPriority);
 				UE_LOG(LogContentStreaming, Log, TEXT("  Retention Priority:  %d"), StreamingRenderAsset.RetentionPriority);
 				UE_LOG(LogContentStreaming, Log, TEXT("  Boost factor:        %.1f"), StreamingRenderAsset.BoostFactor);
-				UE_LOG(LogContentStreaming, Log, TEXT("  Mip bias [Budget]:   %d [%d]"), StreamingRenderAsset.MipCount - StreamingRenderAsset.MaxAllowedMips, StreamingRenderAsset.BudgetMipBias + (Settings.bUsePerTextureBias ? 0 : Settings.GlobalMipBias));
+
+				FString BiasDesc;
+				{
+					int32 CumuBias = 0;
+					if (!Settings.bUseAllMips)
+					{
+						// UI specific Bias : see UTextureLODSettings::CalculateLODBias(), included in CachedCombinedLODBias.
+						extern int32 GUITextureLODBias;
+						if (StreamingRenderAsset.LODGroup == TEXTUREGROUP_UI && GUITextureLODBias)
+						{
+							BiasDesc += FString::Printf(TEXT(" [UI:%d]"), GUITextureLODBias);
+							CumuBias += GUITextureLODBias;
+						}
+
+						// LOD group Bias : see UTextureLODSettings::CalculateLODBias(), included in CachedCombinedLODBias
+						const FTextureLODGroup& LODGroupInfo = UDeviceProfileManager::Get().GetActiveProfile()->GetTextureLODSettings()->TextureLODGroups[StreamingRenderAsset.LODGroup];
+						if (LODGroupInfo.LODBias)
+						{
+							if (FPlatformProperties::RequiresCookedData())
+							{
+								BiasDesc += FString::Printf(TEXT(" [LODGroup.Bias:0(%d)]"), LODGroupInfo.LODBias);
+							}
+							else
+							{
+								BiasDesc += FString::Printf(TEXT(" [LODGroup.Bias:%d]"), LODGroupInfo.LODBias);
+								CumuBias += LODGroupInfo.LODBias;
+							}
+						}
+
+						// LOD group MaxResolution clamp : see UTextureLODSettings::CalculateLODBias(), included in CachedCombinedLODBias
+						const int32 MipCountBeforeMaxRes = StreamingRenderAsset.MipCount - Texture2D->NumCinematicMipLevels -
+							(StreamingRenderAsset.LODGroup == TEXTUREGROUP_UI ? GUITextureLODBias : 0) - 
+							(FPlatformProperties::RequiresCookedData() ? 0 : (LODGroupInfo.LODBias + Texture2D->LODBias));
+						const int32 MaxResBias = MipCountBeforeMaxRes - (LODGroupInfo.MaxLODMipCount + 1);
+						if (MaxResBias > 0)
+						{
+							BiasDesc += FString::Printf(TEXT(" [LODGroup.MaxRes:%d]"), MaxResBias);
+							CumuBias += MaxResBias;
+						}
+
+						// Asset LODBias : see UTextureLODSettings::CalculateLODBias(), included in CachedCombinedLODBias
+						if (Texture2D->LODBias)
+						{
+							if (FPlatformProperties::RequiresCookedData())
+							{
+								BiasDesc += FString::Printf(TEXT(" [Asset.Bias:0(%d)]"), Texture2D->LODBias);
+							}
+							else
+							{
+								BiasDesc += FString::Printf(TEXT(" [Asset.Bias:%d]"), Texture2D->LODBias);
+								CumuBias += Texture2D->LODBias;
+							}
+						}
+
+						// Asset cinematic mips : see UTextureLODSettings::CalculateLODBias() and FStreamingRenderAsset::UpdateDynamicData(), included in CachedCombinedLODBias
+						if (Texture2D->NumCinematicMipLevels && !(StreamingRenderAsset.bForceFullyLoad && Texture2D->bUseCinematicMipLevels))
+						{
+							BiasDesc += FString::Printf(TEXT(" [Asset.Cine:%d]"), Texture2D->NumCinematicMipLevels);
+							CumuBias += Texture2D->NumCinematicMipLevels;
+						}
+					}
+
+					// RHI max resolution and optional mips : see FStreamingRenderAsset::UpdateDynamicData().
+					{
+						const int32 OptMipBias = StreamingRenderAsset.MipCount - CumuBias - StreamingRenderAsset.NumNonOptionalMips;
+						if (StreamingRenderAsset.OptionalMipsState != FStreamingRenderAsset::EOptionalMipsState::OMS_HasOptionalMips && OptMipBias > 0)
+						{
+							BiasDesc += FString::Printf(TEXT(" [Asset.Opt:%d]"), OptMipBias);
+							CumuBias += OptMipBias;
+						}
+
+						const int32 RHIMaxResBias = StreamingRenderAsset.MipCount - CumuBias - GMaxTextureMipCount;
+						if (RHIMaxResBias > 0)
+						{
+							BiasDesc += FString::Printf(TEXT(" [RHI.MaxRes:%d]"), RHIMaxResBias);
+							CumuBias += RHIMaxResBias;
+						}
+					}
+
+					// Memory budget bias
+					// Global Bias : see FStreamingRenderAsset::UpdateDynamicData().
+					if (StreamingRenderAsset.IsMaxResolutionAffectedByGlobalBias() && !Settings.bUsePerTextureBias && Settings.GlobalMipBias)
+					{
+						BiasDesc += FString::Printf(TEXT(" [Global:%d]"), Settings.GlobalMipBias);
+						CumuBias += Settings.GlobalMipBias;
+					}
+					
+					if (StreamingRenderAsset.BudgetMipBias)
+					{
+						BiasDesc += FString::Printf(TEXT(" [Budget:%d]"), StreamingRenderAsset.BudgetMipBias);
+						CumuBias += StreamingRenderAsset.BudgetMipBias;
+					}
+				}
+
+				UE_LOG(LogContentStreaming, Log, TEXT("  Mip bias:			  %d %s"), StreamingRenderAsset.MipCount - StreamingRenderAsset.MaxAllowedMips, !BiasDesc.IsEmpty() ? *BiasDesc : TEXT("") );
 
 				if (InWorld && !GIsEditor)
 				{
