Prevent potential crash caused by a race condition when importing meshes in datasmith with distancefield enabled

diff --git a/Engine/Source/Runtime/Engine/Private/StaticMesh.cpp b/Engine/Source/Runtime/Engine/Private/StaticMesh.cpp
--- a/Engine/Source/Runtime/Engine/Private/StaticMesh.cpp
+++ b/Engine/Source/Runtime/Engine/Private/StaticMesh.cpp
@@ -4305,9 +4305,14 @@ void UStaticMesh::CacheDerivedData()
 
 	if (RenderData)
 	{
-		// Finish any previous async builds before modifying RenderData
-		// This can happen during import as the mesh is rebuilt redundantly
-		GDistanceFieldAsyncQueue->BlockUntilBuildComplete(this, true);
+		// This is the responsability of the caller to ensure this has been called
+		// on the main thread when calling CacheDerivedData() from another thread.
+		if (IsInGameThread())
+		{
+			// Finish any previous async builds before modifying RenderData
+			// This can happen during import as the mesh is rebuilt redundantly
+			GDistanceFieldAsyncQueue->BlockUntilBuildComplete(this, true);
+		}
 
 		for (int32 LODIndex = 0; LODIndex < RenderData->LODResources.Num(); ++LODIndex)
 		{
