PR #6327: Const correctness of TimerManager functions (Contributed by [REDACTED_PERSON])
#rnx

diff --git a/Engine/Source/Runtime/Engine/Public/TimerManager.h b/Engine/Source/Runtime/Engine/Public/TimerManager.h
--- a/Engine/Source/Runtime/Engine/Public/TimerManager.h
+++ b/Engine/Source/Runtime/Engine/Public/TimerManager.h
@@ -320,7 +320,7 @@ public:
 	* @param InHandle The handle of the timer to check for being paused.
 	* @return true if the timer exists and is paused, false otherwise.
 	*/
-	FORCEINLINE bool IsTimerPaused(FTimerHandle InHandle)
+	FORCEINLINE bool IsTimerPaused(FTimerHandle InHandle) const
 	{
 		FTimerData const* const TimerData = FindTimer(InHandle);
 		return TimerData && TimerData->Status == ETimerStatus::Paused;
@@ -332,7 +332,7 @@ public:
 	* @param InHandle The handle of the timer to check for being pending.
 	* @return true if the timer exists and is pending, false otherwise.
 	*/
-	FORCEINLINE bool IsTimerPending(FTimerHandle InHandle)
+	FORCEINLINE bool IsTimerPending(FTimerHandle InHandle) const
 	{
 		FTimerData const* const TimerData = FindTimer(InHandle);
 		return TimerData && TimerData->Status == ETimerStatus::Pending;
@@ -344,7 +344,7 @@ public:
 	* @param InHandle The handle of the timer to check for existence.
 	* @return true if the timer exists, false otherwise.
 	*/
-	FORCEINLINE bool TimerExists(FTimerHandle InHandle)
+	FORCEINLINE bool TimerExists(FTimerHandle InHandle) const
 	{
 		return FindTimer(InHandle) != nullptr;
 	}
