The GameInstance no longer assumes it knows which world it's creating the gamemode for when requested you need to provide the world you want the game mode created for, if none is provided it will use the default.

diff --git a/Engine/Source/Runtime/Engine/Classes/Engine/GameInstance.h b/Engine/Source/Runtime/Engine/Classes/Engine/GameInstance.h
--- a/Engine/Source/Runtime/Engine/Classes/Engine/GameInstance.h
+++ b/Engine/Source/Runtime/Engine/Classes/Engine/GameInstance.h
@@ -411,7 +411,7 @@ public:
 	virtual void PreloadContentForURL(FURL InURL);
 
 	/** Call to create the game mode for a given map URL */
-	virtual class AGameModeBase* CreateGameModeForURL(FURL InURL);
+	virtual class AGameModeBase* CreateGameModeForURL(FURL InURL, UWorld* InWorld);
 
 	/** Return the game mode subclass to use for a given map, options, and portal. By default return passed in one */
 	virtual TSubclassOf<AGameModeBase> OverrideGameModeClass(TSubclassOf<AGameModeBase> GameModeClass, const FString& MapName, const FString& Options, const FString& Portal) const;
diff --git a/Engine/Source/Runtime/Engine/Private/GameInstance.cpp b/Engine/Source/Runtime/Engine/Private/GameInstance.cpp
--- a/Engine/Source/Runtime/Engine/Private/GameInstance.cpp
+++ b/Engine/Source/Runtime/Engine/Private/GameInstance.cpp
@@ -1252,7 +1252,7 @@ void UGameInstance::PreloadContentForURL(FURL InURL)
 	// Preload game mode and other content if needed here
 }
 
-AGameModeBase* UGameInstance::CreateGameModeForURL(FURL InURL)
+AGameModeBase* UGameInstance::CreateGameModeForURL(FURL InURL, UWorld* InWorld)
 {
 	// Init the game info.
 	FString Options;
@@ -1264,7 +1264,7 @@ AGameModeBase* UGameInstance::CreateGameModeForURL(FURL InURL)
 		FParse::Value(*InURL.Op[i], TEXT("GAME="), GameParam);
 	}
 
-	UWorld* World = GetWorld();
+	UWorld* World = InWorld ? InWorld : GetWorld();
 	AWorldSettings* Settings = World->GetWorldSettings();
 	UGameEngine* const GameEngine = Cast<UGameEngine>(GEngine);
 
diff --git a/Engine/Source/Runtime/Engine/Private/World.cpp b/Engine/Source/Runtime/Engine/Private/World.cpp
--- a/Engine/Source/Runtime/Engine/Private/World.cpp
+++ b/Engine/Source/Runtime/Engine/Private/World.cpp
@@ -3907,7 +3907,7 @@ bool UWorld::SetGameMode(const FURL& InURL)
 {
 	if( IsServer() && !AuthorityGameMode )
 	{
-		AuthorityGameMode = GetGameInstance()->CreateGameModeForURL(InURL);
+		AuthorityGameMode = GetGameInstance()->CreateGameModeForURL(InURL, this);
 		if( AuthorityGameMode != NULL )
 		{
 			return true;
