Make NetConnection's BestLag variable private since it is just a mirror of AvgLag and should eventually be deleted

[RN] minor engine/network

diff --git a/Engine/Source/Runtime/Engine/Classes/Engine/NetConnection.h b/Engine/Source/Runtime/Engine/Classes/Engine/NetConnection.h
--- a/Engine/Source/Runtime/Engine/Classes/Engine/NetConnection.h
+++ b/Engine/Source/Runtime/Engine/Classes/Engine/NetConnection.h
@@ -400,11 +400,14 @@ public:
 	/** Average lag seen during the last StatPeriod */
 	float AvgLag;
 
-	//UE_DEPRECATED(4.24, "BestLag variable is deprecated. Use AvgLag instead");
+private:
+	// BestLag variable is deprecated. Use AvgLag instead
 	float			BestLag;
-	//UE_DEPRECATED(4.24, "BestLagAcc variable is deprecated. Use LagAcc instead");
+	// BestLagAcc variable is deprecated. Use LagAcc instead
 	double			BestLagAcc;
 
+public:
+
 	/** Total accumulated lag values during the current StatPeriod */
 	double			LagAcc;
 	/** Nb of stats accumulated in LagAcc */
@@ -443,19 +446,19 @@ public:
 
 	/** Percentage of packets lost during the last StatPeriod */
 	using FNetConnectionPacketLoss = TPacketLossData<3>;
-	const FNetConnectionPacketLoss& GetInLossPercentage() const { return InPacketsLossPercentage; }
+	const FNetConnectionPacketLoss& GetInLossPercentage() const  { return InPacketsLossPercentage; }
 	const FNetConnectionPacketLoss& GetOutLossPercentage() const { return OutPacketsLossPercentage;  }
 
 
 private:
+
 	FNetConnectionPacketLoss InPacketsLossPercentage;
 	FNetConnectionPacketLoss OutPacketsLossPercentage;
 
-public:
-
 	/** Counts the number of stat samples taken */
 	int32 StatPeriodCount;
 
+public:
 
 	/** Net Analytics */
 
diff --git a/Engine/Source/Runtime/Engine/Private/NetConnection.cpp b/Engine/Source/Runtime/Engine/Private/NetConnection.cpp
--- a/Engine/Source/Runtime/Engine/Private/NetConnection.cpp
+++ b/Engine/Source/Runtime/Engine/Private/NetConnection.cpp
@@ -176,7 +176,8 @@ UNetConnection::UNetConnection(const FObjectInitializer& ObjectInitializer)
 
 ,	StatPeriod			( 1.f  )
 ,	AvgLag				( 9999 )
-
+,   BestLag				( 9999 )
+,   BestLagAcc			( 9999 )
 ,	LagAcc				( 9999 )
 ,	LagCount			( 0 )
 ,	LastTime			( 0 )
@@ -2942,6 +2943,7 @@ void UNetConnection::Tick()
 		{
 			AvgLag = LagAcc/LagCount;
 		}
+		BestLag = AvgLag;
 
 		InBytesPerSecond = FMath::TruncToInt(static_cast<float>(InBytes) / RealTime);
 		OutBytesPerSecond = FMath::TruncToInt(static_cast<float>(OutBytes) / RealTime);
