Some safety and error checking code for anim trails.

diff --git a/Engine/Source/Runtime/Engine/Private/Particles/ParticleTrail2EmitterInstance.cpp b/Engine/Source/Runtime/Engine/Private/Particles/ParticleTrail2EmitterInstance.cpp
--- a/Engine/Source/Runtime/Engine/Private/Particles/ParticleTrail2EmitterInstance.cpp
+++ b/Engine/Source/Runtime/Engine/Private/Particles/ParticleTrail2EmitterInstance.cpp
@@ -492,6 +492,22 @@ void FParticleTrailsEmitterInstance_Base::KillParticles()
 					check(!TEXT("What the hell are you doing in here?"));
 				}
 
+				//Be certain we've broken links to this particle				
+				int32 Next = TRAIL_EMITTER_GET_NEXT(TrailData->Flags);
+				if (Next != TRAIL_EMITTER_NULL_NEXT)
+				{
+					DECLARE_PARTICLE_PTR(NextParticle, ParticleData + ParticleStride * Next);
+					FTrailsBaseTypeDataPayload* NextTrailData = ((FTrailsBaseTypeDataPayload*)((uint8*)NextParticle + TypeDataOffset));
+					NextTrailData->Flags = TRAIL_EMITTER_SET_PREV(NextTrailData->Flags, TRAIL_EMITTER_NULL_PREV);
+				}
+				int32 Prev = TRAIL_EMITTER_GET_PREV(TrailData->Flags);
+				if (Prev != TRAIL_EMITTER_NULL_PREV)
+				{
+					DECLARE_PARTICLE_PTR(PrevParticle, ParticleData + ParticleStride * Prev);
+					FTrailsBaseTypeDataPayload* PrevTrailData = ((FTrailsBaseTypeDataPayload*)((uint8*)PrevParticle + TypeDataOffset));
+					PrevTrailData->Flags = TRAIL_EMITTER_SET_NEXT(PrevTrailData->Flags, TRAIL_EMITTER_NULL_NEXT);
+				}
+
 				// Clear it out... just to be safe when it gets pulled back to active
 				TrailData->Flags = TRAIL_EMITTER_SET_NEXT(TrailData->Flags, TRAIL_EMITTER_NULL_NEXT);
 				TrailData->Flags = TRAIL_EMITTER_SET_PREV(TrailData->Flags, TRAIL_EMITTER_NULL_PREV);
@@ -3004,6 +3020,11 @@ void FParticleAnimTrailEmitterInstance::Tick_RecalculateTangents(float DeltaTime
 			CurrTrailData = (FAnimTrailTypeDataPayload*)(TempPayload);
 			while (CurrParticle != NULL)
 			{
+				if (CheckForCircularTrail(StartParticle, CurrParticle))
+				{
+					break;
+				}
+
 				// Grab the next particle in the trail...
 				GetParticleInTrail(true, CurrParticle, CurrTrailData, GET_Next, GET_Any, NextParticle, TempPayload);
 				NextTrailData = (FAnimTrailTypeDataPayload*)(TempPayload);
@@ -3196,6 +3217,8 @@ void FParticleAnimTrailEmitterInstance::SpawnParticle( int32& StartParticleIndex
 			// This it the first particle.
 			// Tag it as the 'only'
 			TrailData->Flags = TRAIL_EMITTER_SET_ONLY(TrailData->Flags);
+			TrailData->Flags = TRAIL_EMITTER_SET_PREV(TrailData->Flags, TRAIL_EMITTER_NULL_PREV);
+			TrailData->Flags = TRAIL_EMITTER_SET_NEXT(TrailData->Flags, TRAIL_EMITTER_NULL_NEXT);
 			TiledUDistanceTraveled[TrailIdx] = 0.0f;
 			TrailData->TiledU = 0.0f;
 			bAddedParticle		= true;
@@ -3790,6 +3813,11 @@ void FParticleAnimTrailEmitterInstance::DetermineVertexAndTriangleCount()
 				// The end of the trail, so there MUST be another particle
 				while (!bDone)
 				{
+					if (CheckForCircularTrail(Particle, PrevParticle))
+					{
+						break;
+					}
+
 					ParticleCount++;
 					int32 InterpCount = 1;
 					if( bApplyDistanceTessellation )
@@ -4293,3 +4321,23 @@ void FParticleAnimTrailEmitterInstance::PrintTrails()
 		}
 	}
 }
+
+static int32 GbEnableCircularAnimTrailDump = 1;
+static FAutoConsoleVariableRef CVarEnableCircularAnimTrailDump(
+	TEXT("fx.EnableCircularAnimTrailDump"),
+	GbEnableCircularAnimTrailDump,
+	TEXT("Controls logging for when circular links are discovered in anim trails.\n0 = No logging.\n1 = Minimal logging.\n2 = Verbose logging."),
+	ECVF_Default
+);
+
+void FParticleTrailsEmitterInstance_Base::DumpCircularTrailsSpam()
+{
+	if (GbEnableCircularAnimTrailDump > 0)
+	{
+		UE_LOG(LogParticles, Warning, TEXT("Circular links in Anim Trail discovered. \nPSys: %s\nEmitter: %s"), *Component->Template->GetFullName(), *SpriteTemplate->GetEmitterName().ToString());
+		if (GbEnableCircularAnimTrailDump > 1)
+		{
+			PrintAllActiveParticles();
+		}
+	}
+}
diff --git a/Engine/Source/Runtime/Engine/Public/ParticleEmitterInstances.h b/Engine/Source/Runtime/Engine/Public/ParticleEmitterInstances.h
--- a/Engine/Source/Runtime/Engine/Public/ParticleEmitterInstances.h
+++ b/Engine/Source/Runtime/Engine/Public/ParticleEmitterInstances.h
@@ -1358,6 +1358,21 @@ void ClearIndices(int32 TrailIndex, int32 ParticleIndex){}
 void CheckIndices(int32 TrailIdx){}
 void CheckAllIndices(){}
 
+#endif
+
+	void DumpCircularTrailsSpam();
+#if !UE_BUILD_SHIPPING
+	bool CheckForCircularTrail(FBaseParticle* StartParticle, FBaseParticle* CheckParticle)
+	{
+		if (StartParticle == CheckParticle)
+		{
+			DumpCircularTrailsSpam();
+			return true;
+		}
+		return false;
+	}
+#else
+	bool CheckForCircularTrail(FBaseParticle* StartParticle, FBaseParticle* CheckParticle) { return StartParticle == CheckParticle; }
 #endif
 
 	/** Constructor	*/
