Increasing some fatal error log spam to be more verbose by default.

diff --git a/Engine/Source/Runtime/Engine/Private/Particles/ParticleTrail2EmitterInstance.cpp b/Engine/Source/Runtime/Engine/Private/Particles/ParticleTrail2EmitterInstance.cpp
--- a/Engine/Source/Runtime/Engine/Private/Particles/ParticleTrail2EmitterInstance.cpp
+++ b/Engine/Source/Runtime/Engine/Private/Particles/ParticleTrail2EmitterInstance.cpp
@@ -4322,7 +4322,7 @@ void FParticleAnimTrailEmitterInstance::PrintTrails()
 	}
 }
 
-static int32 GbEnableCircularAnimTrailDump = 1;
+static int32 GbEnableCircularAnimTrailDump = 2;
 static FAutoConsoleVariableRef CVarEnableCircularAnimTrailDump(
 	TEXT("fx.EnableCircularAnimTrailDump"),
 	GbEnableCircularAnimTrailDump,
