Trying to get PreviewScenes to play bettter with EWorldType::GamePreview world types.  It was leaving their construction in a sort of incomplete state, this change rounds it out more so that things like a gamemode can be initialized.

diff --git a/Engine/Source/Runtime/Engine/Private/PreviewScene.cpp b/Engine/Source/Runtime/Engine/Private/PreviewScene.cpp
--- a/Engine/Source/Runtime/Engine/Private/PreviewScene.cpp
+++ b/Engine/Source/Runtime/Engine/Private/PreviewScene.cpp
@@ -15,9 +15,11 @@
 #include "Components/SkyLightComponent.h"
 #include "Components/LineBatchComponent.h"
 #include "Components/ReflectionCaptureComponent.h"
+#include "GameFramework/GameModeBase.h"
+#include "GameFramework/GameMode.h"
 
 FPreviewScene::FPreviewScene(FPreviewScene::ConstructionValues CVS)
-	: PreviewWorld(NULL)
+	: PreviewWorld(nullptr)
 	, bForceAllUsedMipsResident(CVS.bForceMipsResident)
 {
 	EObjectFlags NewObjectFlags = RF_NoFlags;
@@ -42,7 +44,31 @@ FPreviewScene::FPreviewScene(FPreviewScene::ConstructionValues CVS)
 										.SetTransactional(CVS.bTransactional)
 										.SetDefaultGameMode(CVS.DefaultGameMode));
 
-	PreviewWorld->InitializeActorsForPlay(FURL());
+	FURL URL = FURL();
+	//URL += TEXT("?SpectatorOnly=1");
+	//URL = FURL(NULL, *EditorEngine->BuildPlayWorldURL(*PIEMapName, Params.bStartInSpectatorMode, ExtraURLOptions), TRAVEL_Absolute);
+
+	if (CVS.OwningGameInstance && PreviewWorld->WorldType == EWorldType::GamePreview)
+	{
+		PreviewWorld->SetGameInstance(CVS.OwningGameInstance);
+
+		FWorldContext& PreviewWorldContext = GEngine->GetWorldContextFromWorldChecked(PreviewWorld);
+		PreviewWorldContext.OwningGameInstance = CVS.OwningGameInstance;
+		PreviewWorldContext.GameViewport = CVS.OwningGameInstance->GetGameViewportClient();
+		PreviewWorldContext.AddRef(PreviewWorld);
+		
+		//PreviewWorldContext.PIEInstance =
+
+		if (CVS.DefaultGameMode)
+		{
+			PreviewWorld->SetGameMode(URL);
+
+			AGameMode* Mode = PreviewWorld->GetAuthGameMode<AGameMode>();
+			ensure(Mode);
+		}
+	}
+
+	PreviewWorld->InitializeActorsForPlay(URL);
 
 	if (CVS.bDefaultLighting)
 	{
diff --git a/Engine/Source/Runtime/Engine/Public/PreviewScene.h b/Engine/Source/Runtime/Engine/Public/PreviewScene.h
--- a/Engine/Source/Runtime/Engine/Public/PreviewScene.h
+++ b/Engine/Source/Runtime/Engine/Public/PreviewScene.h
@@ -46,6 +46,7 @@ public:
 		uint32 bEditor:1;
 
 		TSubclassOf<class AGameModeBase> DefaultGameMode;
+		class UGameInstance* OwningGameInstance = nullptr;
 
 		ConstructionValues& SetCreateDefaultLighting(const bool bDefault) { bDefaultLighting = bDefault; return *this; }
 		ConstructionValues& SetLightRotation(const FRotator& Rotation) { LightRotation = Rotation; return *this; }
@@ -60,6 +61,7 @@ public:
 		ConstructionValues& SetEditor(const bool bInEditor) { bEditor = bInEditor; return *this; }
 
 		ConstructionValues& SetDefaultGameMode(TSubclassOf<class AGameModeBase> GameMode) { DefaultGameMode = GameMode; return *this; }
+		ConstructionValues& SetOwningGameInstance(class UGameInstance* InGameInstance) { OwningGameInstance = InGameInstance; return *this; }
 	};
 
 	// for physical correct light computations we multiply diffuse and specular lights by PI (see LABEL_RealEnergy)
