Fixed issue with multiple viewports

diff --git a/Engine/Plugins/Experimental/PanoramicCapture/Source/PanoramicCapture/Private/StereoStaticMeshComponent.cpp b/Engine/Plugins/Experimental/PanoramicCapture/Source/PanoramicCapture/Private/StereoStaticMeshComponent.cpp
--- a/Engine/Plugins/Experimental/PanoramicCapture/Source/PanoramicCapture/Private/StereoStaticMeshComponent.cpp
+++ b/Engine/Plugins/Experimental/PanoramicCapture/Source/PanoramicCapture/Private/StereoStaticMeshComponent.cpp
@@ -29,14 +29,14 @@ public:
         FPrimitiveViewRelevance viewRelevance = FStaticMeshSceneProxy::GetViewRelevance(View);
         bool bVisible = true;
 
-		if (IStereoRendering::IsASecondaryView(*View, GEngine->StereoRenderingDevice))
+		if (IStereoRendering::IsASecondaryView(*View))
 		{
 			if ((EyeToRender != ESPStereoCameraLayer::RightEye) && (EyeToRender != ESPStereoCameraLayer::BothEyes))
 			{
 				bVisible = false;
 			}
 		}
-		else if (IStereoRendering::IsAPrimaryView(*View, GEngine->StereoRenderingDevice))
+		else if (IStereoRendering::IsAPrimaryView(*View))
 		{
 			if ((EyeToRender != ESPStereoCameraLayer::LeftEye) && (EyeToRender != ESPStereoCameraLayer::BothEyes))
 			{
diff --git a/Engine/Plugins/Lumin/MagicLeap/Source/MagicLeap/Private/MagicLeapXRCamera.cpp b/Engine/Plugins/Lumin/MagicLeap/Source/MagicLeap/Private/MagicLeapXRCamera.cpp
--- a/Engine/Plugins/Lumin/MagicLeap/Source/MagicLeap/Private/MagicLeapXRCamera.cpp
+++ b/Engine/Plugins/Lumin/MagicLeap/Source/MagicLeap/Private/MagicLeapXRCamera.cpp
@@ -6,6 +6,7 @@
 #include "MagicLeapHMD.h"
 #include "IMagicLeapPlugin.h"
 
+
 FMagicLeapXRCamera::FMagicLeapXRCamera(const FAutoRegister& AutoRegister, FMagicLeapHMD& InMagicLeapSystem, int32 InDeviceID)
 	: FDefaultXRCamera(AutoRegister, &InMagicLeapSystem, InDeviceID)
 	, MagicLeapSystem(InMagicLeapSystem)
@@ -19,7 +20,7 @@ void FMagicLeapXRCamera::PreRenderView_RenderThread(FRHICommandListImmediate& RH
 	
 	if (MagicLeapSystem.IsStereoEyePass(View.StereoPass))
 	{
-		const int EyeIdx = IStereoRendering::IsAPrimaryView(View, MagicLeapSystem.AsShared()) ? 0 : 1;
+		const int EyeIdx = MagicLeapSystem.DeviceIsAPrimaryView(View) ? 0 : 1;
 
 		const FTrackingFrame& Frame = MagicLeapSystem.GetCurrentFrame();
 
diff --git a/Engine/Plugins/Runtime/Steam/SteamVR/Source/SteamVR/Private/SteamVRRender.cpp b/Engine/Plugins/Runtime/Steam/SteamVR/Source/SteamVR/Private/SteamVRRender.cpp
--- a/Engine/Plugins/Runtime/Steam/SteamVR/Source/SteamVR/Private/SteamVRRender.cpp
+++ b/Engine/Plugins/Runtime/Steam/SteamVR/Source/SteamVR/Private/SteamVRRender.cpp
@@ -72,7 +72,7 @@ bool FSteamVRHMD::IsActiveThisFrame(class FViewport* InViewport) const
 static void DrawOcclusionMesh(FRHICommandList& RHICmdList, EStereoscopicPass StereoPass, const FHMDViewMesh MeshAssets[])
 {
 	check(IsInRenderingThread());
-	check(GEngine->StereoRenderingDevice->IsStereoEyePass(StereoPass));
+	check(GEngine->StereoRenderingDevice->DeviceIsStereoEyePass(StereoPass));
 
 	const uint32 MeshIndex = GEngine->StereoRenderingDevice->GetViewIndexForPass(StereoPass);
 	const FHMDViewMesh& Mesh = MeshAssets[MeshIndex];
diff --git a/Engine/Plugins/Runtime/nDisplay/Source/DisplayCluster/Private/Game/Classes/Basics/DisplayClusterViewportClient.cpp b/Engine/Plugins/Runtime/nDisplay/Source/DisplayCluster/Private/Game/Classes/Basics/DisplayClusterViewportClient.cpp
--- a/Engine/Plugins/Runtime/nDisplay/Source/DisplayCluster/Private/Game/Classes/Basics/DisplayClusterViewportClient.cpp
+++ b/Engine/Plugins/Runtime/nDisplay/Source/DisplayCluster/Private/Game/Classes/Basics/DisplayClusterViewportClient.cpp
@@ -173,7 +173,8 @@ void UDisplayClusterViewportClient::Draw(FViewport* InViewport, FCanvas* SceneCa
 			InViewport,
 			MyWorld->Scene,
 			EngineShowFlags)
-			.SetRealtimeUpdate(true));
+			.SetRealtimeUpdate(true)
+			.SetAdditionalViewFamily(viewFamily > 0));
 
 #if WITH_EDITOR
 		if (GIsEditor)
diff --git a/Engine/Plugins/Runtime/nDisplay/Source/DisplayCluster/Private/Render/Device/DisplayClusterDeviceBase.h b/Engine/Plugins/Runtime/nDisplay/Source/DisplayCluster/Private/Render/Device/DisplayClusterDeviceBase.h
--- a/Engine/Plugins/Runtime/nDisplay/Source/DisplayCluster/Private/Render/Device/DisplayClusterDeviceBase.h
+++ b/Engine/Plugins/Runtime/nDisplay/Source/DisplayCluster/Private/Render/Device/DisplayClusterDeviceBase.h
@@ -89,12 +89,21 @@ protected:
 	virtual bool AllocateDepthTexture(uint32 Index, uint32 SizeX, uint32 SizeY, uint8 Format, uint32 NumMips, uint32 Flags, uint32 TargetableTextureFlags, FTexture2DRHIRef& OutTargetableTexture, FTexture2DRHIRef& OutShaderResourceTexture, uint32 NumSamples = 1)
 	{ return false; }
 
+	virtual bool DeviceIsAPrimaryPass(EStereoscopicPass Pass) override
+	{ return true; }
+
 	virtual bool DeviceIsAPrimaryView(const FSceneView& View) override
 	{ return true; }
+
+	virtual bool DeviceIsASecondaryPass(EStereoscopicPass Pass) override
+	{ return false; }
 	
 	virtual bool DeviceIsASecondaryView(const FSceneView& View) override
 	{ return false; }
 
+	virtual bool DeviceIsAnAdditionalPass(EStereoscopicPass Pass) override
+	{ return false; }
+
 	virtual bool DeviceIsAnAdditionalView(const FSceneView& View) override
 	{ return false; }
 
diff --git a/Engine/Source/Runtime/Engine/Private/GameViewportClient.cpp b/Engine/Source/Runtime/Engine/Private/GameViewportClient.cpp
--- a/Engine/Source/Runtime/Engine/Private/GameViewportClient.cpp
+++ b/Engine/Source/Runtime/Engine/Private/GameViewportClient.cpp
@@ -1400,7 +1400,7 @@ void UGameViewportClient::Draw(FViewport* InViewport, FCanvas* SceneCanvas)
 
 	// If the views don't cover the entire bounding rectangle, clear the entire buffer.
 	bool bBufferCleared = false;
-	bool bStereoscopicPass = (ViewFamily.Views.Num() != 0 && ViewFamily.Views[0]->StereoPass != eSSP_FULL);
+	bool bStereoscopicPass = (ViewFamily.Views.Num() != 0 && IStereoRendering::IsStereoEyeView(*ViewFamily.Views[0]));
 	if (ViewFamily.Views.Num() == 0 || TotalArea != (MaxX-MinX)*(MaxY-MinY) || bDisableWorldRendering || bStereoscopicPass)
 	{
 		if (bDisableWorldRendering || !bStereoscopicPass) // TotalArea computation does not work correctly for stereoscopic views
diff --git a/Engine/Source/Runtime/Engine/Private/LocalPlayer.cpp b/Engine/Source/Runtime/Engine/Private/LocalPlayer.cpp
--- a/Engine/Source/Runtime/Engine/Private/LocalPlayer.cpp
+++ b/Engine/Source/Runtime/Engine/Private/LocalPlayer.cpp
@@ -762,7 +762,7 @@ bool ULocalPlayer::CalcSceneViewInitOptions(
 	check(PlayerController && PlayerController->GetWorld());
 
 	int ViewIndex = 0;
-	if (StereoPass != eSSP_FULL)
+	if (IStereoRendering::IsStereoEyePass(StereoPass))
 	{
 		ViewIndex = StereoPass - 1;
 	}
@@ -1071,7 +1071,7 @@ bool ULocalPlayer::GetProjectionData(FViewport* Viewport, EStereoscopicPass Ster
 	GetViewPoint(/*out*/ ViewInfo, StereoPass);
 
 	// If stereo rendering is enabled, update the size and offset appropriately for this pass
-	const bool bNeedStereo = (StereoPass != eSSP_FULL) && GEngine->IsStereoscopic3D();
+	const bool bNeedStereo = IStereoRendering::IsStereoEyePass(StereoPass) && GEngine->IsStereoscopic3D();
 	const bool bIsHeadTrackingAllowed = GEngine->XRSystem.IsValid() && GEngine->XRSystem->IsHeadTrackingAllowed();
 	if (bNeedStereo)
 	{
diff --git a/Engine/Source/Runtime/Engine/Private/SceneManagement.cpp b/Engine/Source/Runtime/Engine/Private/SceneManagement.cpp
--- a/Engine/Source/Runtime/Engine/Private/SceneManagement.cpp
+++ b/Engine/Source/Runtime/Engine/Private/SceneManagement.cpp
@@ -527,7 +527,7 @@ int8 ComputeTemporalStaticMeshLOD( const FStaticMeshRenderData* RenderData, cons
 // Ensure we always use the left eye when selecting lods to avoid divergent selections in stereo
 const FSceneView& GetLODView(const FSceneView& InView)
 {
-	if (InView.StereoPass == EStereoscopicPass::eSSP_RIGHT_EYE && InView.Family)
+	if (IStereoRendering::IsASecondaryView(InView) && InView.Family)
 	{
 		return *InView.Family->Views[0];
 	}
diff --git a/Engine/Source/Runtime/Engine/Private/SceneView.cpp b/Engine/Source/Runtime/Engine/Private/SceneView.cpp
--- a/Engine/Source/Runtime/Engine/Private/SceneView.cpp
+++ b/Engine/Source/Runtime/Engine/Private/SceneView.cpp
@@ -565,7 +565,7 @@ FViewMatrices::FViewMatrices(const FSceneViewInitOptions& InitOptions) : FViewMa
 
 	// Compute screen scale factors.
 	// Stereo renders at half horizontal resolution, but compute shadow resolution based on full resolution.
-	const bool bStereo = InitOptions.StereoPass != eSSP_FULL;
+	const bool bStereo = IStereoRendering::IsStereoEyePass(InitOptions.StereoPass);
 	const float ScreenXScale = bStereo ? 2.0f : 1.0f;
 	ProjectionScale.X = ScreenXScale * FMath::Abs(ProjectionMatrix.M[0][0]);
 	ProjectionScale.Y = FMath::Abs(ProjectionMatrix.M[1][1]);
@@ -759,7 +759,7 @@ FSceneView::FSceneView(const FSceneViewInitOptions& InitOptions)
 	bShouldBindInstancedViewUB = bIsInstancedStereoEnabled || bIsMobileMultiViewEnabled;
 
 	// If the device doesn't support mobile multi-view, disable it.
-	bIsMobileMultiViewEnabled = bIsMobileMultiViewEnabled && GSupportsMobileMultiView && StereoPass != eSSP_FULL;
+	bIsMobileMultiViewEnabled = bIsMobileMultiViewEnabled && GSupportsMobileMultiView && IStereoRendering::IsStereoEyePass(StereoPass);
 
 	SetupAntiAliasingMethod();
 
@@ -968,7 +968,7 @@ void FSceneView::UpdateViewMatrix()
 {
 	FVector StereoViewLocation = ViewLocation;
 	FRotator StereoViewRotation = ViewRotation;
-	if (GEngine->StereoRenderingDevice.IsValid() && StereoPass != eSSP_FULL)
+	if (GEngine->StereoRenderingDevice.IsValid() && IStereoRendering::IsStereoEyePass(StereoPass))
 	{
 		GEngine->StereoRenderingDevice->CalculateStereoViewOffset(StereoPass, StereoViewRotation, WorldToMetersScale, StereoViewLocation);
 		ViewLocation = StereoViewLocation;
@@ -2101,6 +2101,11 @@ EShaderPlatform FSceneView::GetShaderPlatform() const
 	return GShaderPlatformForFeatureLevel[GetFeatureLevel()];
 }
 
+bool FSceneView::IsInstancedStereoPass() const
+{
+	return bIsInstancedStereoEnabled && IStereoRendering::IsAPrimaryView(*this);
+}
+
 void FSceneView::SetupViewRectUniformBufferParameters(FViewUniformShaderParameters& ViewUniformShaderParameters,
 	const FIntPoint& BufferSize,
 	const FIntRect& EffectiveViewRect,
@@ -2358,6 +2363,7 @@ FSceneViewFamily::FSceneViewFamily(const ConstructionValues& CVS)
 	DeltaWorldTime(CVS.DeltaWorldTime),
 	CurrentRealTime(CVS.CurrentRealTime),
 	FrameNumber(UINT_MAX),
+	bAdditionalViewFamily(CVS.bAdditionalViewFamily),
 	bRealtimeUpdate(CVS.bRealtimeUpdate),
 	bDeferClear(CVS.bDeferClear),
 	bResolveScene(CVS.bResolveScene),
diff --git a/Engine/Source/Runtime/Engine/Private/StereoRendering.cpp b/Engine/Source/Runtime/Engine/Private/StereoRendering.cpp
--- a/Engine/Source/Runtime/Engine/Private/StereoRendering.cpp
+++ b/Engine/Source/Runtime/Engine/Private/StereoRendering.cpp
@@ -3,32 +3,105 @@
 #include "StereoRendering.h"
 #include "SceneView.h"
 
+#include "Engine/Engine.h"
+
+
+bool IStereoRendering::IsStereoEyePass(EStereoscopicPass Pass)
+{
+	if (GEngine->StereoRenderingDevice.IsValid())
+	{
+		return GEngine->StereoRenderingDevice->DeviceIsStereoEyePass(Pass);
+	}
+
+	return Pass != EStereoscopicPass::eSSP_FULL;
+}
+
+bool IStereoRendering::IsStereoEyeView(const FSceneView& View)
+{
+	if (GEngine->StereoRenderingDevice.IsValid())
+	{
+		return GEngine->StereoRenderingDevice->DeviceIsStereoEyeView(View);
+	}
+
+	return View.StereoPass != EStereoscopicPass::eSSP_FULL;
+}
+
+bool IStereoRendering::DeviceIsStereoEyeView(const FSceneView& View)
+{
+	return !(View.StereoPass == EStereoscopicPass::eSSP_FULL);
+}
+
+bool IStereoRendering::IsAPrimaryPass(EStereoscopicPass Pass)
+{
+	if (GEngine->StereoRenderingDevice.IsValid())
+	{
+		return GEngine->StereoRenderingDevice->DeviceIsAPrimaryPass(Pass);
+	}
+
+	return Pass == EStereoscopicPass::eSSP_FULL || Pass == EStereoscopicPass::eSSP_LEFT_EYE;
+}
+
+bool IStereoRendering::IsAPrimaryView(const FSceneView& View)
+{
+	if (GEngine->StereoRenderingDevice.IsValid())
+	{
+		return GEngine->StereoRenderingDevice->DeviceIsAPrimaryView(View);
+	}
+
+	return View.StereoPass == EStereoscopicPass::eSSP_FULL || View.StereoPass == EStereoscopicPass::eSSP_LEFT_EYE;
+}
+
 bool IStereoRendering::DeviceIsAPrimaryView(const FSceneView& View)
 {
-	return View.StereoPass == eSSP_FULL || View.StereoPass == eSSP_LEFT_EYE;
+	return View.StereoPass == EStereoscopicPass::eSSP_FULL || View.StereoPass == EStereoscopicPass::eSSP_LEFT_EYE;
 }
 
-bool IStereoRendering::DeviceIsAnAdditionalView(const FSceneView& View)
+bool IStereoRendering::IsASecondaryPass(EStereoscopicPass Pass)
 {
-	return View.StereoPass > eSSP_RIGHT_EYE;
+	if (GEngine->StereoRenderingDevice.IsValid())
+	{
+		return GEngine->StereoRenderingDevice->DeviceIsASecondaryPass(Pass);
+	}
+
+	return !(Pass == EStereoscopicPass::eSSP_FULL || Pass == EStereoscopicPass::eSSP_LEFT_EYE);
 }
 
-bool IStereoRendering::IsAPrimaryView(const FSceneView& View, TSharedPtr<IStereoRendering, ESPMode::ThreadSafe> StereoRenderingDevice)
+bool IStereoRendering::IsASecondaryView(const FSceneView& View)
 {
-	if (StereoRenderingDevice.IsValid())
+	if (GEngine->StereoRenderingDevice.IsValid())
 	{
-		return StereoRenderingDevice->DeviceIsAPrimaryView(View);
+		return GEngine->StereoRenderingDevice->DeviceIsASecondaryView(View);
 	}
 
-	return View.StereoPass == eSSP_FULL || View.StereoPass == eSSP_LEFT_EYE;
+	return !(View.StereoPass == EStereoscopicPass::eSSP_FULL || View.StereoPass == EStereoscopicPass::eSSP_LEFT_EYE);
 }
 
-bool IStereoRendering::IsAnAdditionalView(const FSceneView& View, TSharedPtr<IStereoRendering, ESPMode::ThreadSafe> StereoRenderingDevice)
+bool IStereoRendering::DeviceIsASecondaryView(const FSceneView& View)
 {
-	if (StereoRenderingDevice.IsValid())
+	return !DeviceIsAPrimaryView(View);
+}
+
+bool IStereoRendering::IsAnAdditionalPass(EStereoscopicPass Pass)
+{
+	if (GEngine->StereoRenderingDevice.IsValid())
 	{
-		return StereoRenderingDevice->DeviceIsAnAdditionalView(View);
+		return GEngine->StereoRenderingDevice->DeviceIsAnAdditionalPass(Pass);
+	}
+
+	return Pass > eSSP_RIGHT_EYE;
+}
+
+bool IStereoRendering::IsAnAdditionalView(const FSceneView& View)
+{
+	if (GEngine->StereoRenderingDevice.IsValid())
+	{
+		return GEngine->StereoRenderingDevice->DeviceIsAnAdditionalView(View);
 	}
 
 	return View.StereoPass > eSSP_RIGHT_EYE;
 }
+
+bool IStereoRendering::DeviceIsAnAdditionalView(const FSceneView& View)
+{
+	return View.StereoPass > EStereoscopicPass::eSSP_RIGHT_EYE;
+}
diff --git a/Engine/Source/Runtime/Engine/Private/UnrealEngine.cpp b/Engine/Source/Runtime/Engine/Private/UnrealEngine.cpp
--- a/Engine/Source/Runtime/Engine/Private/UnrealEngine.cpp
+++ b/Engine/Source/Runtime/Engine/Private/UnrealEngine.cpp
@@ -2949,7 +2949,7 @@ public:
 		// We have to do all the work for secondary views that are on a different GPU than
 		// the primary view. NB: This assumes that the primary view is assigned to the
 		// first GPU of the AFR group. See FSceneRenderer::ComputeViewGPUMasks.
-		if (View.StereoPass != eSSP_FULL && AFRUtils::GetIndexWithinGroup(View.GPUMask.ToIndex()) != 0)
+		if (IStereoRendering::IsStereoEyeView(View) && AFRUtils::GetIndexWithinGroup(View.GPUMask.ToIndex()) != 0)
 		{
 			return true;
 		}
diff --git a/Engine/Source/Runtime/Engine/Private/UserInterface/Canvas.cpp b/Engine/Source/Runtime/Engine/Private/UserInterface/Canvas.cpp
--- a/Engine/Source/Runtime/Engine/Private/UserInterface/Canvas.cpp
+++ b/Engine/Source/Runtime/Engine/Private/UserInterface/Canvas.cpp
@@ -1920,7 +1920,7 @@ void UCanvas::SetView(FSceneView* InView)
 	SceneView = InView;
 	if (InView)
 	{
-		if (GEngine->StereoRenderingDevice.IsValid() && InView->StereoPass != eSSP_FULL)
+		if (GEngine->StereoRenderingDevice.IsValid() && IStereoRendering::IsStereoEyeView(*InView))
 		{
 			GEngine->StereoRenderingDevice->InitCanvasFromView(InView, this);
 		}
diff --git a/Engine/Source/Runtime/Engine/Public/SceneView.h b/Engine/Source/Runtime/Engine/Public/SceneView.h
--- a/Engine/Source/Runtime/Engine/Public/SceneView.h
+++ b/Engine/Source/Runtime/Engine/Public/SceneView.h
@@ -1230,7 +1230,7 @@ public:
 	EShaderPlatform GetShaderPlatform() const;
 
 	/** True if the view should render as an instanced stereo pass */
-	bool IsInstancedStereoPass() const { return bIsInstancedStereoEnabled && StereoPass == eSSP_LEFT_EYE; }
+	bool IsInstancedStereoPass() const;
 
 	/** Sets up the view rect parameters in the view's uniform shader parameters */
 	void SetupViewRectUniformBufferParameters(FViewUniformShaderParameters& ViewUniformShaderParameters, 
@@ -1392,6 +1392,7 @@ public:
 		,	DeltaWorldTime(0.0f)
 		,	CurrentRealTime(0.0f)
 		,	GammaCorrection(1.0f)
+		,	bAdditionalViewFamily(false)
 		,	bRealtimeUpdate(false)
 		,	bDeferClear(false)
 		,	bResolveScene(true)			
@@ -1436,6 +1437,9 @@ public:
 		/** Gamma correction used when rendering this family. Default is 1.0 */
 		float GammaCorrection;
 
+		/** Indicates whether the view family is additional. */
+		uint32 bAdditionalViewFamily : 1;
+
 		/** Indicates whether the view family is updated in real-time. */
 		uint32 bRealtimeUpdate:1;
 		
@@ -1450,7 +1454,10 @@ public:
 
 		/** Set the world time ,difference between the last world time and CurrentWorldTime and current real time. */
 		ConstructionValues& SetWorldTimes(const float InCurrentWorldTime,const float InDeltaWorldTime,const float InCurrentRealTime) { CurrentWorldTime = InCurrentWorldTime; DeltaWorldTime = InDeltaWorldTime; CurrentRealTime = InCurrentRealTime;bTimesSet = true;return *this; }
-		
+
+		/** Set  whether the view family is additional. */
+		ConstructionValues& SetAdditionalViewFamily(const bool Value) { bAdditionalViewFamily = Value; return *this; }
+
 		/** Set  whether the view family is updated in real-time. */
 		ConstructionValues& SetRealtimeUpdate(const bool Value) { bRealtimeUpdate = Value; return *this; }
 		
@@ -1494,6 +1501,9 @@ public:
 	/** Copy from main thread GFrameNumber to be accessible on render thread side. UINT_MAX before CreateSceneRenderer() or BeginRenderingViewFamily() was called */
 	uint32 FrameNumber;
 
+	/** Indicates this view family is an additional one. */
+	bool bAdditionalViewFamily;
+
 	/** Indicates whether the view family is updated in realtime. */
 	bool bRealtimeUpdate;
 
diff --git a/Engine/Source/Runtime/Engine/Public/StereoRendering.h b/Engine/Source/Runtime/Engine/Public/StereoRendering.h
--- a/Engine/Source/Runtime/Engine/Public/StereoRendering.h
+++ b/Engine/Source/Runtime/Engine/Public/StereoRendering.h
@@ -84,54 +84,98 @@ public:
 			return -1;
 		}
 	}
-	
+
+	/**
+	* Static helper. Passes request to the current stereo device
+	*/
+	static bool IsStereoEyePass(EStereoscopicPass Pass);
+
 	/**
 	* Return true if this pass is for a stereo eye view
 	*/
-	virtual bool IsStereoEyePass(EStereoscopicPass Pass)
+	virtual bool DeviceIsStereoEyePass(EStereoscopicPass Pass)
 	{
 		return !(Pass == EStereoscopicPass::eSSP_FULL);
 	}
-	
 
-	virtual bool DeviceIsAPrimaryView(const FSceneView& View);
+	/**
+	* Static helper. Passes request to the current stereo device
+	*/
+	static bool IsStereoEyeView(const FSceneView& View);
 
 	/**
-	* Return true if this pass is for a view for which we share some work done for eSSP_LEFT_EYE (ie borrow some intermediate state from that eye)
+	* Return true if this is a stereoscopic view
+	*/
+	virtual bool DeviceIsStereoEyeView(const FSceneView& View);
+
+	/**
+	* Static helper. Passes request to the current stereo device
+	*/
+	static bool IsAPrimaryPass(EStereoscopicPass Pass);
+
+	/**
+	* Return true if this pass is for a view we do all the work for (ie this view can't borrow from another)
 	*/
-	virtual bool DeviceIsASecondaryView(const FSceneView& View)
+	virtual bool DeviceIsAPrimaryPass(EStereoscopicPass Pass)
 	{
-		return !DeviceIsAPrimaryView(View);
+		return Pass == EStereoscopicPass::eSSP_FULL || Pass == EStereoscopicPass::eSSP_LEFT_EYE;
 	}
 
 	/**
-	* Return true for additional eyes past the first two (a plugin could implement additional 'eyes').
+	* Static helper. Passes request to the current stereo device
 	*/
-	virtual bool DeviceIsAnAdditionalView(const FSceneView& View);
+	static bool IsAPrimaryView(const FSceneView& View);
 
+	/**
+	* Return true if primary view
+	*/
+	virtual bool DeviceIsAPrimaryView(const FSceneView& View);
 
 	/**
-	* Return true if this pass is for a view we do all the work for (ie this view can't borrow from another)
+	* Static helper. Passes request to the current stereo device
 	*/
-	static bool IsAPrimaryView(const FSceneView& View, TSharedPtr< class IStereoRendering, ESPMode::ThreadSafe > StereoRenderingDevice);
-	
+	static bool IsASecondaryPass(EStereoscopicPass Pass);
+
 	/**
 	* Return true if this pass is for a view for which we share some work done for eSSP_LEFT_EYE (ie borrow some intermediate state from that eye)
 	*/
-	static bool IsASecondaryView(const FSceneView& View, TSharedPtr< class IStereoRendering, ESPMode::ThreadSafe > StereoRenderingDevice)
+	virtual bool DeviceIsASecondaryPass(EStereoscopicPass Pass)
 	{
-		if (StereoRenderingDevice.IsValid())
-		{
-			return StereoRenderingDevice->DeviceIsASecondaryView(View);
-		}
-
-		return !IsAPrimaryView(View, StereoRenderingDevice);
+		return !DeviceIsAPrimaryPass(Pass);
 	}
-	
+
+	/**
+	* Static helper. Passes request to the current stereo device
+	*/
+	static bool IsASecondaryView(const FSceneView& View);
+
+	/**
+	* Return true if secondary view
+	*/
+	virtual bool DeviceIsASecondaryView(const FSceneView& View);
+
+	/**
+	* Static helper. Passes request to the current stereo device
+	*/
+	static bool IsAnAdditionalPass(EStereoscopicPass Pass);
+
 	/**
 	* Return true for additional eyes past the first two (a plugin could implement additional 'eyes').
 	*/
-	static bool IsAnAdditionalView(const FSceneView& View, TSharedPtr< class IStereoRendering, ESPMode::ThreadSafe > StereoRenderingDevice);
+	virtual bool DeviceIsAnAdditionalPass(EStereoscopicPass Pass)
+	{
+		return Pass > EStereoscopicPass::eSSP_RIGHT_EYE;
+	}
+
+	/**
+	* Static helper. Passes request to the current stereo device
+	*/
+	static bool IsAnAdditionalView(const FSceneView& View);
+
+	/**
+	* Return true if additional view
+	*/
+	virtual bool DeviceIsAnAdditionalView(const FSceneView& View);
 
 	/**
 	 * Adjusts the viewport rectangle for stereo, based on which eye pass is being rendered.
diff --git a/Engine/Source/Runtime/HeadMountedDisplay/Private/DefaultStereoLayers.cpp b/Engine/Source/Runtime/HeadMountedDisplay/Private/DefaultStereoLayers.cpp
--- a/Engine/Source/Runtime/HeadMountedDisplay/Private/DefaultStereoLayers.cpp
+++ b/Engine/Source/Runtime/HeadMountedDisplay/Private/DefaultStereoLayers.cpp
@@ -209,7 +209,7 @@ void FDefaultStereoLayers::PreRenderViewFamily_RenderThread(FRHICommandListImmed
 
 void FDefaultStereoLayers::PostRenderView_RenderThread(FRHICommandListImmediate& RHICmdList, FSceneView& InView)
 {
-	if (!HMDDevice->IsStereoEyePass(InView.StereoPass))
+	if (!HMDDevice->DeviceIsStereoEyeView(InView))
 	{
 		return;
 	}
diff --git a/Engine/Source/Runtime/Renderer/Private/BasePassRendering.cpp b/Engine/Source/Runtime/Renderer/Private/BasePassRendering.cpp
--- a/Engine/Source/Runtime/Renderer/Private/BasePassRendering.cpp
+++ b/Engine/Source/Runtime/Renderer/Private/BasePassRendering.cpp
@@ -485,7 +485,7 @@ void SetupSharedBasePassParameters(
 
 	SetupFogUniformParameters(View, SharedParameters.Fog);
 
-	if (View.bIsInstancedStereoEnabled && View.StereoPass == EStereoscopicPass::eSSP_LEFT_EYE)
+	if (View.bIsInstancedStereoEnabled && IStereoRendering::IsAPrimaryView(View))
 	{
 		const FSceneView& RightEye = *View.Family->Views[1];
 		SharedParameters.ForwardISR = RightEye.ForwardLightingResources->ForwardLightData;
diff --git a/Engine/Source/Runtime/Renderer/Private/DepthRendering.cpp b/Engine/Source/Runtime/Renderer/Private/DepthRendering.cpp
--- a/Engine/Source/Runtime/Renderer/Private/DepthRendering.cpp
+++ b/Engine/Source/Runtime/Renderer/Private/DepthRendering.cpp
@@ -801,7 +801,7 @@ bool FDeferredShadingSceneRenderer::RenderPrePassHMD(FRHICommandListImmediate& R
 	for (int32 ViewIndex = 0; ViewIndex < Views.Num(); ++ViewIndex)
 	{
 		const FViewInfo& View = Views[ViewIndex];
-		if (View.StereoPass != eSSP_FULL)
+		if (IStereoRendering::IsStereoEyeView(View))
 		{
 			SCOPED_GPU_MASK(RHICmdList, View.GPUMask);
 			RHICmdList.SetViewport(View.ViewRect.Min.X, View.ViewRect.Min.Y, 0.0f, View.ViewRect.Max.X, View.ViewRect.Max.Y, 1.0f);
diff --git a/Engine/Source/Runtime/Renderer/Private/DistortionRendering.cpp b/Engine/Source/Runtime/Renderer/Private/DistortionRendering.cpp
--- a/Engine/Source/Runtime/Renderer/Private/DistortionRendering.cpp
+++ b/Engine/Source/Runtime/Renderer/Private/DistortionRendering.cpp
@@ -53,7 +53,7 @@ void SetupDistortionParams(FVector4& DistortionParams, const FViewInfo& View)
 	if ((View.IsInstancedStereoPass() || View.bIsMobileMultiViewEnabled) && View.Family->Views.Num() > 0)
 	{
 		// When drawing the left eye in a stereo scene, copy the right eye view values into the instanced view uniform buffer.
-		const EStereoscopicPass StereoPassIndex = (View.StereoPass != eSSP_FULL) ? eSSP_RIGHT_EYE : eSSP_FULL;
+		const EStereoscopicPass StereoPassIndex = IStereoRendering::IsStereoEyeView(View) ? eSSP_RIGHT_EYE : eSSP_FULL;
 
 		const FViewInfo& InstancedView = static_cast<const FViewInfo&>(View.Family->GetStereoEyeView(StereoPassIndex));
 		DistortionParams.Y = InstancedView.ViewMatrices.GetProjectionMatrix().M[0][0];
diff --git a/Engine/Source/Runtime/Renderer/Private/LightPropagationVolume.cpp b/Engine/Source/Runtime/Renderer/Private/LightPropagationVolume.cpp
--- a/Engine/Source/Runtime/Renderer/Private/LightPropagationVolume.cpp
+++ b/Engine/Source/Runtime/Renderer/Private/LightPropagationVolume.cpp
@@ -1460,12 +1460,12 @@ void FSceneViewState::SetupLightPropagationVolume(FSceneView& View, FSceneViewFa
 	const ERHIFeatureLevel::Type ViewFeatureLevel = View.GetFeatureLevel();
 
 	//#todo-rv-vr
-	if (IStereoRendering::IsASecondaryView(View, GEngine->StereoRenderingDevice))
+	if (IStereoRendering::IsASecondaryView(View))
 	{
 		// The right eye will reference the left eye's LPV with the assumption that the left eye uses the primary view (index 0)
 		const FSceneView* PrimaryView = ViewFamily.Views[0];
 		//#todo-rv-vr
-		if (PrimaryView->StereoPass == eSSP_LEFT_EYE && PrimaryView->State)
+		if (IStereoRendering::IsAPrimaryView(*PrimaryView) && PrimaryView->State)
 		{
 			FSceneViewState* PrimaryViewState = PrimaryView->State->GetConcreteViewState();
 			if (PrimaryViewState)
diff --git a/Engine/Source/Runtime/Renderer/Private/LightRendering.cpp b/Engine/Source/Runtime/Renderer/Private/LightRendering.cpp
--- a/Engine/Source/Runtime/Renderer/Private/LightRendering.cpp
+++ b/Engine/Source/Runtime/Renderer/Private/LightRendering.cpp
@@ -994,7 +994,7 @@ void FDeferredShadingSceneRenderer::RenderLights(FRHICommandListImmediate& RHICm
 				bool bAnyViewIsStereo = false;
 				for (int32 ViewIndex = 0; ViewIndex < Views.Num(); ++ViewIndex)
 				{
-					if (Views[ViewIndex].StereoPass != eSSP_FULL)
+					if (IStereoRendering::IsStereoEyeView(Views[ViewIndex]))
 					{
 						bAnyViewIsStereo = true;
 						break;
diff --git a/Engine/Source/Runtime/Renderer/Private/PlanarReflectionRendering.cpp b/Engine/Source/Runtime/Renderer/Private/PlanarReflectionRendering.cpp
--- a/Engine/Source/Runtime/Renderer/Private/PlanarReflectionRendering.cpp
+++ b/Engine/Source/Runtime/Renderer/Private/PlanarReflectionRendering.cpp
@@ -570,7 +570,7 @@ void FScene::UpdatePlanarReflectionContents(UPlanarReflectionComponent* CaptureC
 			SceneRenderer->Views[ViewIndex].FinalPostProcessSettings.ScreenPercentage =
 				MainSceneRenderer.Views[ViewIndex].FinalPostProcessSettings.ScreenPercentage;
 
-			const bool bIsStereo = MainSceneRenderer.Views[0].StereoPass != EStereoscopicPass::eSSP_FULL;
+			const bool bIsStereo = IStereoRendering::IsStereoEyeView(MainSceneRenderer.Views[0]);
 
 			const FMatrix ProjectionMatrix = SceneCaptureViewInfo[ViewIndex].ProjectionMatrix;
 			FPlanarReflectionSceneProxy* SceneProxy = CaptureComponent->SceneProxy;
diff --git a/Engine/Source/Runtime/Renderer/Private/PostProcess/PostProcessMobile.cpp b/Engine/Source/Runtime/Renderer/Private/PostProcess/PostProcessMobile.cpp
--- a/Engine/Source/Runtime/Renderer/Private/PostProcess/PostProcessMobile.cpp
+++ b/Engine/Source/Runtime/Renderer/Private/PostProcess/PostProcessMobile.cpp
@@ -2953,7 +2953,7 @@ void FRCPassPostProcessAaES2::Process(FRenderingCompositePassContext& Context)
 
 	ERenderTargetActions LoadStoreAction = ERenderTargetActions::Load_Store;
 	//#todo-rv-vr
-	if (Context.View.StereoPass != eSSP_RIGHT_EYE)
+	if (!IStereoRendering::IsASecondaryView(Context.View))
 	{
 		// Full clear to avoid restore
 		LoadStoreAction = ERenderTargetActions::Clear_Store;
diff --git a/Engine/Source/Runtime/Renderer/Private/PostProcess/PostProcessTonemap.cpp b/Engine/Source/Runtime/Renderer/Private/PostProcess/PostProcessTonemap.cpp
--- a/Engine/Source/Runtime/Renderer/Private/PostProcess/PostProcessTonemap.cpp
+++ b/Engine/Source/Runtime/Renderer/Private/PostProcess/PostProcessTonemap.cpp
@@ -1484,7 +1484,7 @@ void FRCPassPostProcessTonemapES2::Process(FRenderingCompositePassContext& Conte
 		const bool bFirstView = (&View == View.Family->Views[0]);
 		
 		// Full clear to avoid restore
-		if ((View.StereoPass == eSSP_FULL && bFirstView) || View.StereoPass == eSSP_LEFT_EYE)
+		if ((View.StereoPass == eSSP_FULL && bFirstView) || IStereoRendering::IsAPrimaryView(View))
 		{
 			LoadAction = ERenderTargetLoadAction::EClear;
 		}
diff --git a/Engine/Source/Runtime/Renderer/Private/PostProcess/PostProcessing.cpp b/Engine/Source/Runtime/Renderer/Private/PostProcess/PostProcessing.cpp
--- a/Engine/Source/Runtime/Renderer/Private/PostProcess/PostProcessing.cpp
+++ b/Engine/Source/Runtime/Renderer/Private/PostProcess/PostProcessing.cpp
@@ -326,7 +326,7 @@ void AddPostProcessingPasses(FRDGBuilder& GraphBuilder, const FViewInfo& View, c
 		};
 
 		const EStereoscopicPass StereoPass = View.StereoPass;
-		const bool bPrimaryView = IStereoRendering::IsAPrimaryView(View, GEngine->StereoRenderingDevice);
+		const bool bPrimaryView = IStereoRendering::IsAPrimaryView(View);
 		const bool bHasViewState = View.ViewState != nullptr;
 		const bool bDepthOfFieldEnabled = DiaphragmDOF::IsEnabled(View);
 		const bool bVisualizeDepthOfField = bDepthOfFieldEnabled && EngineShowFlags.VisualizeDOF;
@@ -1198,7 +1198,7 @@ static FRCPassPostProcessTonemap* AddTonemapper(
 	const EStereoscopicPass StereoPass = View.StereoPass;
 
 	FRenderingCompositeOutputRef TonemapperCombinedLUTOutputRef;
-	if (IStereoRendering::IsAPrimaryView(View, GEngine->StereoRenderingDevice))
+	if (IStereoRendering::IsAPrimaryView(View))
 	{
 		TonemapperCombinedLUTOutputRef = AddCombineLUTPass(Context.Graph);
 	}
diff --git a/Engine/Source/Runtime/Renderer/Private/PostProcess/SceneFilterRendering.cpp b/Engine/Source/Runtime/Renderer/Private/PostProcess/SceneFilterRendering.cpp
--- a/Engine/Source/Runtime/Renderer/Private/PostProcess/SceneFilterRendering.cpp
+++ b/Engine/Source/Runtime/Renderer/Private/PostProcess/SceneFilterRendering.cpp
@@ -281,7 +281,7 @@ void DrawPostProcessPass(
 	bool bHasCustomMesh,
 	EDrawRectangleFlags Flags)
 {
-	if (bHasCustomMesh && StereoView != eSSP_FULL)
+	if (bHasCustomMesh && IStereoRendering::IsStereoEyePass(StereoView))
 	{
 		DrawHmdMesh(RHICmdList, X, Y, SizeX, SizeY, U, V, SizeU, SizeV, TargetSize, TextureSize, StereoView, VertexShader);
 	}
diff --git a/Engine/Source/Runtime/Renderer/Private/PostProcess/SceneRenderTargets.cpp b/Engine/Source/Runtime/Renderer/Private/PostProcess/SceneRenderTargets.cpp
--- a/Engine/Source/Runtime/Renderer/Private/PostProcess/SceneRenderTargets.cpp
+++ b/Engine/Source/Runtime/Renderer/Private/PostProcess/SceneRenderTargets.cpp
@@ -389,7 +389,7 @@ FIntPoint FSceneRenderTargets::ComputeDesiredSize(const FSceneViewFamily& ViewFa
 
 		bIsSceneCapture |= View->bIsSceneCapture;
 		bIsReflectionCapture |= View->bIsReflectionCapture;
-		bIsVRScene |= (View->StereoPass != EStereoscopicPass::eSSP_FULL && GEngine->XRSystem.IsValid());
+		bIsVRScene |= (IStereoRendering::IsStereoEyeView(*View) && GEngine->XRSystem.IsValid());
 	}
 
 	{
diff --git a/Engine/Source/Runtime/Renderer/Private/ScenePrivate.h b/Engine/Source/Runtime/Renderer/Private/ScenePrivate.h
--- a/Engine/Source/Runtime/Renderer/Private/ScenePrivate.h
+++ b/Engine/Source/Runtime/Renderer/Private/ScenePrivate.h
@@ -2323,7 +2323,7 @@ public:
 	const FViewInfo& GetInstancedView(const FViewInfo& View)
 	{
 		// When drawing the left eye in a stereo scene, copy the right eye view values into the instanced view uniform buffer.
-		const EStereoscopicPass StereoPassIndex = (View.StereoPass != eSSP_FULL) ? eSSP_RIGHT_EYE : eSSP_FULL;
+		const EStereoscopicPass StereoPassIndex = IStereoRendering::IsStereoEyeView(View) ? eSSP_RIGHT_EYE : eSSP_FULL;
 
 		return static_cast<const FViewInfo&>(View.Family->GetStereoEyeView(StereoPassIndex));
 	}
diff --git a/Engine/Source/Runtime/Renderer/Private/SceneRendering.cpp b/Engine/Source/Runtime/Renderer/Private/SceneRendering.cpp
--- a/Engine/Source/Runtime/Renderer/Private/SceneRendering.cpp
+++ b/Engine/Source/Runtime/Renderer/Private/SceneRendering.cpp
@@ -1453,12 +1453,12 @@ void FViewInfo::SetupUniformBufferParameters(
 	{
 		// If rendering in stereo, the other stereo passes uses the left eye's translucency lighting volume.
 		const FViewInfo* PrimaryView = this;
-		if (IStereoRendering::IsASecondaryView(*this, GEngine->StereoRenderingDevice))
+		if (IStereoRendering::IsASecondaryView(*this))
 		{
 			if (Family->Views.IsValidIndex(0))
 			{
 				const FSceneView* LeftEyeView = Family->Views[0];
-				if (LeftEyeView->bIsViewInfo && LeftEyeView->StereoPass == eSSP_LEFT_EYE)
+				if (LeftEyeView->bIsViewInfo && IStereoRendering::IsAPrimaryView(*LeftEyeView))
 				{
 					PrimaryView = static_cast<const FViewInfo*>(LeftEyeView);
 				}
@@ -1580,7 +1580,7 @@ void FViewInfo::SetupUniformBufferParameters(
 		GMaxRHIFeatureLevel > ERHIFeatureLevel::ES3_1) ? 1.0f : 0.0f;
 
 	// Padding between the left and right eye may be introduced by an HMD, which instanced stereo needs to account for.
-	if ((StereoPass != eSSP_FULL) && (Family->Views.Num() > 1))
+	if ((IStereoRendering::IsStereoEyePass(StereoPass)) && (Family->Views.Num() > 1))
 	{
 		check(Family->Views.Num() >= 2);
 
@@ -1835,7 +1835,7 @@ FSceneViewState* FViewInfo::GetEffectiveViewState() const
 	FSceneViewState* EffectiveViewState = ViewState;
 
 	// When rendering in stereo we want to use the same exposure for both eyes.
-	if (IStereoRendering::IsASecondaryView(*this, GEngine->StereoRenderingDevice))
+	if (IStereoRendering::IsASecondaryView(*this))
 	{
 		int32 ViewIndex = Family->Views.Find(this);
 		if (Family->Views.IsValidIndex(ViewIndex))
@@ -1845,7 +1845,7 @@ FSceneViewState* FViewInfo::GetEffectiveViewState() const
 			if (Family->Views.IsValidIndex(ViewIndex))
 			{
 				const FSceneView* PrimaryView = Family->Views[ViewIndex];
-				if (PrimaryView->StereoPass == eSSP_LEFT_EYE)
+				if (IStereoRendering::IsAPrimaryView(*PrimaryView))
 				{
 					EffectiveViewState = (FSceneViewState*)PrimaryView->State;
 				}
@@ -2566,7 +2566,7 @@ void FSceneRenderer::ComputeFamilySize()
 		MaxFamilyX = FMath::Max(MaxFamilyX, FinalViewMaxX);
 		MaxFamilyY = FMath::Max(MaxFamilyY, FinalViewMaxY);
 
-		if (!IStereoRendering::IsAnAdditionalView(View, GEngine->StereoRenderingDevice))
+		if (!IStereoRendering::IsAnAdditionalView(View))
 		{
 			InstancedStereoWidth = FPlatformMath::Max(InstancedStereoWidth, static_cast<uint32>(View.ViewRect.Max.X));
 		}
@@ -3136,7 +3136,7 @@ void FSceneRenderer::RenderCustomDepthPass(FRHICommandListImmediate& RHICmdList)
 					{
 						// When drawing the left eye in a stereo scene, set up the instanced custom depth uniform buffer with the right-eye data,
 						// with the TAA jitter removed.
-						const EStereoscopicPass StereoPassIndex = (View.StereoPass != eSSP_FULL) ? eSSP_RIGHT_EYE : eSSP_FULL;
+						const EStereoscopicPass StereoPassIndex = IStereoRendering::IsStereoEyeView(View) ? eSSP_RIGHT_EYE : eSSP_FULL;
 
 						const FViewInfo& InstancedView = static_cast<const FViewInfo&>(View.Family->GetStereoEyeView(StereoPassIndex));
 
diff --git a/Engine/Source/Runtime/Renderer/Private/SceneRendering.h b/Engine/Source/Runtime/Renderer/Private/SceneRendering.h
--- a/Engine/Source/Runtime/Renderer/Private/SceneRendering.h
+++ b/Engine/Source/Runtime/Renderer/Private/SceneRendering.h
@@ -1288,11 +1288,11 @@ public:
 		{
 			return true;
 		}
-		else if (bIsInstancedStereoEnabled && StereoPass != eSSP_RIGHT_EYE)
+		else if (bIsInstancedStereoEnabled && !IStereoRendering::IsASecondaryPass(StereoPass))
 		{
 			return true;
 		}
-		else if (bIsMobileMultiViewEnabled && StereoPass != eSSP_RIGHT_EYE && Family && Family->Views.Num() > 1)
+		else if (bIsMobileMultiViewEnabled && !IStereoRendering::IsASecondaryPass(StereoPass) && Family && Family->Views.Num() > 1)
 		{
 			return true;
 		}
diff --git a/Engine/Source/Runtime/Renderer/Private/SceneVisibility.cpp b/Engine/Source/Runtime/Renderer/Private/SceneVisibility.cpp
--- a/Engine/Source/Runtime/Renderer/Private/SceneVisibility.cpp
+++ b/Engine/Source/Runtime/Renderer/Private/SceneVisibility.cpp
@@ -769,7 +769,7 @@ static void FetchVisibilityForPrimitives_Range(FVisForPrimParams& Params, FGloba
 
 	int32 ReadBackLagTolerance = NumBufferedFrames;
 
-	const bool bIsStereoView = View.StereoPass == eSSP_LEFT_EYE || View.StereoPass == eSSP_RIGHT_EYE;
+	const bool bIsStereoView = IStereoRendering::IsStereoEyeView(View);
 	const bool bUseRoundRobinOcclusion = bIsStereoView && !View.bIsSceneCapture && View.ViewState->IsRoundRobinEnabled();
 	if (bUseRoundRobinOcclusion)
 	{
@@ -1637,13 +1637,13 @@ static int32 OcclusionCull(FRHICommandListImmediate& RHICmdList, const FScene* S
 			if (View.ViewState->IsRoundRobinEnabled() &&
 				!View.bIsSceneCapture && // We only round-robin on the main renderer (not scene captures)
 				!View.bIgnoreExistingQueries && // We do not alternate occlusion queries when we want to refresh the occlusion history
-				(View.StereoPass == eSSP_LEFT_EYE || View.StereoPass == eSSP_RIGHT_EYE)) // Only relevant to stereo views
+				(IStereoRendering::IsStereoEyeView(View))) // Only relevant to stereo views
 			{
 				// For even frames, prevent left eye from occlusion querying
 				// For odd frames, prevent right eye from occlusion querying
 				const bool FrameParity = ((View.ViewState->PrevFrameNumber & 0x01) == 1);
-				bSubmitQueries &= (FrameParity && View.StereoPass == eSSP_LEFT_EYE) ||
-								  (!FrameParity && View.StereoPass == eSSP_RIGHT_EYE);
+				bSubmitQueries &= (FrameParity  && IStereoRendering::IsAPrimaryView(View)) ||
+								  (!FrameParity && IStereoRendering::IsASecondaryView(View));
 			}
 
 			NumOccludedPrimitives += FetchVisibilityForPrimitives(Scene, View, bSubmitQueries, bHZBOcclusion, DynamicVertexBuffer);
diff --git a/Engine/Source/Runtime/Renderer/Private/ScreenPass.inl b/Engine/Source/Runtime/Renderer/Private/ScreenPass.inl
--- a/Engine/Source/Runtime/Renderer/Private/ScreenPass.inl
+++ b/Engine/Source/Runtime/Renderer/Private/ScreenPass.inl
@@ -23,7 +23,7 @@ inline FScreenPassRenderTarget FScreenPassRenderTarget::CreateViewFamilyOutput(F
 		// Raw output mode uses the original view rect. Otherwise the final unscaled rect is used.
 		View.PrimaryScreenPercentageMethod == EPrimaryScreenPercentageMethod::RawOutput ? View.ViewRect : View.UnscaledViewRect,
 		// First view clears the view family texture; all remaining views load.
-		View.IsFirstInFamily() ? ERenderTargetLoadAction::EClear : ERenderTargetLoadAction::ELoad);
+		(!View.Family->bAdditionalViewFamily && View.IsFirstInFamily() )? ERenderTargetLoadAction::EClear : ERenderTargetLoadAction::ELoad);
 }
 
 inline FScreenPassTextureViewport FScreenPassTextureViewport::CreateDownscaled(const FScreenPassTextureViewport& Other, uint32 DownscaleFactor)
diff --git a/Engine/Source/Runtime/Renderer/Private/ShadowRendering.cpp b/Engine/Source/Runtime/Renderer/Private/ShadowRendering.cpp
--- a/Engine/Source/Runtime/Renderer/Private/ShadowRendering.cpp
+++ b/Engine/Source/Runtime/Renderer/Private/ShadowRendering.cpp
@@ -709,7 +709,7 @@ void FProjectedShadowInfo::SetupProjectionStencilMask(
 
 		// If instanced stereo is enabled, we need to render each view of the stereo pair using the instanced stereo transform to avoid bias issues.
 		// TODO: Support instanced stereo properly in the projection stenciling pass.
-		const bool bIsInstancedStereoEmulated = View->bIsInstancedStereoEnabled && !View->bIsMultiViewEnabled && View->StereoPass != eSSP_FULL;
+		const bool bIsInstancedStereoEmulated = View->bIsInstancedStereoEnabled && !View->bIsMultiViewEnabled && IStereoRendering::IsStereoEyeView(*View);
 		if (bIsInstancedStereoEmulated)
 		{
 			RHICmdList.SetViewport(0, 0, 0, SceneRender->InstancedStereoWidth, View->ViewRect.Max.Y, 1);
diff --git a/Engine/Source/Runtime/Renderer/Private/ShadowSetup.cpp b/Engine/Source/Runtime/Renderer/Private/ShadowSetup.cpp
--- a/Engine/Source/Runtime/Renderer/Private/ShadowSetup.cpp
+++ b/Engine/Source/Runtime/Renderer/Private/ShadowSetup.cpp
@@ -1356,7 +1356,7 @@ void FProjectedShadowInfo::SetupMeshDrawCommandsForProjectionStenciling(FSceneRe
 
 			// If instanced stereo is enabled, we need to render each view of the stereo pair using the instanced stereo transform to avoid bias issues.
 			// TODO: Support instanced stereo properly in the projection stenciling pass.
-			const uint32 InstanceFactor = View.bIsInstancedStereoEnabled && !View.bIsMultiViewEnabled && View.StereoPass != eSSP_FULL ? 2 : 1;
+			const uint32 InstanceFactor = View.bIsInstancedStereoEnabled && !View.bIsMultiViewEnabled && IStereoRendering::IsStereoEyeView(View) ? 2 : 1;
 			SortAndMergeDynamicPassMeshDrawCommands(Renderer.FeatureLevel, ProjectionStencilingPass.VisibleMeshDrawCommands, DynamicMeshDrawCommandStorage, ProjectionStencilingPass.PrimitiveIdVertexBuffer, InstanceFactor);
 		}
 	}
@@ -2762,9 +2762,9 @@ void FSceneRenderer::InitProjectedShadowVisibility(FRHICommandListImmediate& RHI
 				{
 					// The view dependent projected shadow is valid for this view if it's the
 					// right eye and the projected shadow is being rendered for the left eye.
-					const bool bIsValidForView = View.StereoPass == eSSP_RIGHT_EYE
+					const bool bIsValidForView = IStereoRendering::IsASecondaryView(View)
 						&& Views.IsValidIndex(ViewIndex - 1)
-						&& Views[ViewIndex - 1].StereoPass == eSSP_LEFT_EYE
+						&& IStereoRendering::IsAPrimaryView(Views[ViewIndex - 1])
 						&& ProjectedShadowInfo.FadeAlphas.IsValidIndex(ViewIndex)
 						&& ProjectedShadowInfo.FadeAlphas[ViewIndex] == 1.0f;
 
@@ -3342,15 +3342,15 @@ void FSceneRenderer::AddViewDependentWholeSceneShadowsForView(
 		FadeAlphas.Init(0.0f, Views.Num());
 		FadeAlphas[ViewIndex] = LightShadowAmount;
 
-		if (IStereoRendering::IsAPrimaryView(View, GEngine->StereoRenderingDevice)
+		if (IStereoRendering::IsAPrimaryView(View)
 			&& Views.IsValidIndex(ViewIndex + 1)
-			&& IStereoRendering::IsASecondaryView(Views[ViewIndex + 1], GEngine->StereoRenderingDevice))
+			&& IStereoRendering::IsASecondaryView(Views[ViewIndex + 1]))
 		{
 			FadeAlphas[ViewIndex + 1] = LightShadowAmount;
 		}		
 		
 		// If rendering in stereo mode we render shadow depths only for the left eye, but project for both eyes!
-		if (IStereoRendering::IsAPrimaryView(View, GEngine->StereoRenderingDevice))
+		if (IStereoRendering::IsAPrimaryView(View))
 		{
 			const bool bExtraDistanceFieldCascade = LightSceneInfo.Proxy->ShouldCreateRayTracedCascade(View.GetFeatureLevel(), LightSceneInfo.IsPrecomputedLightingValid(), View.MaxShadowCascades);
 
diff --git a/Engine/Source/Runtime/Renderer/Private/TranslucentRendering.cpp b/Engine/Source/Runtime/Renderer/Private/TranslucentRendering.cpp
--- a/Engine/Source/Runtime/Renderer/Private/TranslucentRendering.cpp
+++ b/Engine/Source/Runtime/Renderer/Private/TranslucentRendering.cpp
@@ -924,7 +924,7 @@ void FDeferredShadingSceneRenderer::RenderTranslucency(FRHICommandListImmediate&
 				if ((View.IsInstancedStereoPass() || View.bIsMobileMultiViewEnabled) && View.Family->Views.Num() > 0)
 				{
 					// When drawing the left eye in a stereo scene, copy the right eye view values into the instanced view uniform buffer.
-					const EStereoscopicPass StereoPassIndex = (View.StereoPass != eSSP_FULL) ? eSSP_RIGHT_EYE : eSSP_FULL;
+					const EStereoscopicPass StereoPassIndex = IStereoRendering::IsStereoEyeView(View) ? eSSP_RIGHT_EYE : eSSP_FULL;
 
 					const FViewInfo& InstancedView = static_cast<const FViewInfo&>(View.Family->GetStereoEyeView(StereoPassIndex));
 					SetupDownsampledTranslucencyViewParameters(RHICmdList, InstancedView, DownsampledTranslucencyViewParameters);
diff --git a/Engine/Source/Runtime/Renderer/Public/MaterialShader.h b/Engine/Source/Runtime/Renderer/Public/MaterialShader.h
--- a/Engine/Source/Runtime/Renderer/Public/MaterialShader.h
+++ b/Engine/Source/Runtime/Renderer/Public/MaterialShader.h
@@ -128,7 +128,7 @@ public:
 		if (View.bShouldBindInstancedViewUB && View.Family->Views.Num() > 0)
 		{
 			// When drawing the left eye in a stereo scene, copy the right eye view values into the instanced view uniform buffer.
-			const EStereoscopicPass StereoPassIndex = (View.StereoPass != eSSP_FULL) ? eSSP_RIGHT_EYE : eSSP_FULL;
+			const EStereoscopicPass StereoPassIndex = IStereoRendering::IsStereoEyeView(View) ? eSSP_RIGHT_EYE : eSSP_FULL;
 
 			const FSceneView& InstancedView = View.Family->GetStereoEyeView(StereoPassIndex);
 			const auto& InstancedViewUniformBufferParameter = GetUniformBufferParameter<FInstancedViewUniformShaderParameters>();
