Fix crash when max audio sources is reached and stopping sources are enforced but not observed by spatialization/modulation plugins

diff --git a/Engine/Source/Runtime/Engine/Private/AudioDevice.cpp b/Engine/Source/Runtime/Engine/Private/AudioDevice.cpp
--- a/Engine/Source/Runtime/Engine/Private/AudioDevice.cpp
+++ b/Engine/Source/Runtime/Engine/Private/AudioDevice.cpp
@@ -390,7 +390,7 @@ bool FAudioDevice::Init(int32 InMaxSources)
 			//Set up initialization parameters for system level effect plugins:
 			FAudioPluginInitializationParams PluginInitializationParams;
 			PluginInitializationParams.SampleRate = SampleRate;
-			PluginInitializationParams.NumSources = MaxSources;
+			PluginInitializationParams.NumSources = GetMaxSources();
 			PluginInitializationParams.BufferLength = PlatformSettings.CallbackBufferFrameSize;
 			PluginInitializationParams.AudioDevicePtr = this;
 
@@ -444,7 +444,7 @@ bool FAudioDevice::Init(int32 InMaxSources)
 		//Set up initialization parameters for system level effect plugins:
 		FAudioPluginInitializationParams PluginInitializationParams;
 		PluginInitializationParams.SampleRate = SampleRate;
-		PluginInitializationParams.NumSources = MaxChannels;
+		PluginInitializationParams.NumSources = GetMaxSources();
 		PluginInitializationParams.BufferLength = PlatformSettings.CallbackBufferFrameSize;
 		PluginInitializationParams.AudioDevicePtr = this;
 		ModulationInterface->Initialize(PluginInitializationParams);
