Ability to scale distance attenuation of attenuation settings via sound class.

diff --git a/Engine/Source/Runtime/Engine/Classes/Kismet/GameplayStatics.h b/Engine/Source/Runtime/Engine/Classes/Kismet/GameplayStatics.h
--- a/Engine/Source/Runtime/Engine/Classes/Kismet/GameplayStatics.h
+++ b/Engine/Source/Runtime/Engine/Classes/Kismet/GameplayStatics.h
@@ -450,6 +450,20 @@ public:
 	UFUNCTION(BlueprintCallable, BlueprintCosmetic, Category = "Audio", meta = (WorldContext = "WorldContextObject"))
 	static void SetGlobalPitchModulation(const UObject* WorldContextObject, float PitchModulation, float TimeSec);
 
+	/** 
+	* Sets attenuation distance scale value on the sound class over the given amount of time from it's current attenuation distance override value (1.0f it not overridden). 
+	* Attenuation scale allows scaling the attenuation distance used for computing distance attenuation. 
+	*
+	* * Fire and Forget.
+	* * Not Replicated.
+	* @param SoundClass - Sound class to to use to set the attenuation distance scale on.
+	* @param DistanceAttenuationScale - A distance attenuation scale value.
+	* @param TimeSec - A time value to linearly interpolate from the current distance attenuation scale value to the new value.
+	*/
+	UFUNCTION(BlueprintCallable, BlueprintCosmetic, Category = "Audio", meta = (WorldContext = "WorldContextObject"))
+	static void SetSoundClassDistanceScale(const UObject* WorldContextObject, USoundClass* SoundClass, float DistanceAttenuationScale, float TimeSec = 0.0f);
+
+
 	/**
 	* Sets the global listener focus parameters which will scale focus behavior of sounds based on their focus azimuth settings in their attenuation settings. 
 	*
diff --git a/Engine/Source/Runtime/Engine/Classes/Sound/SoundClass.h b/Engine/Source/Runtime/Engine/Classes/Sound/SoundClass.h
--- a/Engine/Source/Runtime/Engine/Classes/Sound/SoundClass.h
+++ b/Engine/Source/Runtime/Engine/Classes/Sound/SoundClass.h
@@ -6,6 +6,7 @@
 #include "UObject/ObjectMacros.h"
 #include "UObject/Object.h"
 #include "AudioDefines.h"
+#include "AudioDynamicParameter.h"
 #if WITH_EDITOR
 #include "EdGraph/EdGraph.h"
 #endif
@@ -87,6 +88,10 @@ struct FSoundClassProperties
 	UPROPERTY(EditAnywhere, BlueprintReadOnly, Category=SoundClassProperties)
 	float Pitch;
 
+	/** Distance scale to apply to sounds that play with this sound class. Sounds will have their attenuation distance scaled by this amount. Allows adjusting attenuation settings dynamically. **/
+	UPROPERTY(EditAnywhere, BlueprintReadOnly, Category = SoundClassProperties)
+	float AttenuationDistanceScale;
+	
 	/** Lowpass filter frequency */
 	UPROPERTY(EditAnywhere, BlueprintReadOnly, Category = SoundClassProperties)
 	float LowPassFilterFrequency;
@@ -151,9 +156,22 @@ struct FSoundClassProperties
 	UPROPERTY(EditAnywhere, BlueprintReadOnly, Category = SoundClassProperties, meta = (DisplayName = "Loading Behavior Override"))
 	ESoundWaveLoadingBehavior LoadingBehavior;
 
+	// Sets the attenuation scale of the sound class in the given amount of time
+	void SetAttenuationDistanceScale(float InAttenuationDistanceScale, float InTime);
+
+	// Sets the parent's attenuation scale
+	void SetParentAttenuationDistanceScale(float InAttenuationDistanceScale);
+
+	// Returns the current attenuation scale
+	float GetAttenuationDistanceScale() const;
+
+	// Updates any dynamic sound class properties
+	void UpdateSoundClassProperties(float DeltaTime);
+
 	FSoundClassProperties()
 		: Volume(1.0f)
 		, Pitch(1.0f)
+		, AttenuationDistanceScale(1.0f)
 		, LowPassFilterFrequency(MAX_FILTER_FREQUENCY)
 		, StereoBleed(0.25f)
 		, LFEBleed(0.5f)
@@ -170,9 +188,15 @@ struct FSoundClassProperties
 		, bApplyAmbientVolumes(false)
 		, OutputTarget(EAudioOutputTarget::Speaker)
 		, LoadingBehavior(ESoundWaveLoadingBehavior::Inherited)
+		, AttenuationScaleParam(AttenuationDistanceScale)
+		, ParentAttenuationScale(1.0f)
 		{
 		}
 	
+private:
+	FDynamicParameter AttenuationScaleParam;
+	float ParentAttenuationScale;
+
 };
 
 /**
diff --git a/Engine/Source/Runtime/Engine/Private/AudioDevice.cpp b/Engine/Source/Runtime/Engine/Private/AudioDevice.cpp
--- a/Engine/Source/Runtime/Engine/Private/AudioDevice.cpp
+++ b/Engine/Source/Runtime/Engine/Private/AudioDevice.cpp
@@ -2007,7 +2007,7 @@ void FAudioDevice::InitSoundClasses()
 	}
 
 	// Propagate the properties down the hierarchy
-	ParseSoundClasses();
+	ParseSoundClasses(0.0f);
 }
 
 void FAudioDevice::InitSoundSources()
@@ -2095,6 +2095,7 @@ void FAudioDevice::RecurseIntoSoundClasses(USoundClass* CurrentClass, FSoundClas
 				Properties->Pitch *= ParentProperties.Pitch;
 				Properties->bIsUISound |= ParentProperties.bIsUISound;
 				Properties->bIsMusic |= ParentProperties.bIsMusic;
+				Properties->SetParentAttenuationDistanceScale(ParentProperties.GetAttenuationDistanceScale());
 
 				// Not all values propagate equally...
 				// VoiceCenterChannelVolume, RadioFilterVolume, RadioFilterVolumeThreshold, bApplyEffects, BleedStereo, bReverb, and bCenterChannelOnly do not propagate (sub-classes can be non-zero even if parent class is zero)
@@ -2138,7 +2139,7 @@ void FAudioDevice::UpdateHighestPriorityReverb()
 	}
 }
 
-void FAudioDevice::ParseSoundClasses()
+void FAudioDevice::ParseSoundClasses(float InDeltaTime)
 {
 	TArray<USoundClass*> RootSoundClasses;
 
@@ -2148,6 +2149,9 @@ void FAudioDevice::ParseSoundClasses()
 		USoundClass* SoundClass = It.Key();
 		if (SoundClass)
 		{
+			// Update any dynamic sound class properties
+			SoundClass->Properties.UpdateSoundClassProperties(InDeltaTime);
+
 			It.Value() = SoundClass->Properties;
 			if (SoundClass->ParentClass == NULL)
 			{
@@ -2723,7 +2727,7 @@ void FAudioDevice::UpdateSoundClassProperties(float DeltaTime)
 	SCOPED_NAMED_EVENT(FAudioDevice_UpdateSoundClasses, FColor::Blue);
 
 	// Remove SoundMix modifications and propagate the properties down the hierarchy
-	ParseSoundClasses();
+	ParseSoundClasses(DeltaTime);
 
 	for (TMap< USoundMix*, FSoundMixState >::TIterator It(SoundMixModifiers); It; ++It)
 	{
@@ -6163,6 +6167,32 @@ void FAudioDevice::SetGlobalPitchModulation(float PitchModulation, float TimeSec
 	GlobalPitchScale.Set(PitchModulation, TimeSec);
 }
 
+void FAudioDevice::SetSoundClassDistanceScale(USoundClass* InSoundClass, float DistanceScale, float TimeSec)
+{
+	check(InSoundClass);
+
+	if (!IsInAudioThread())
+	{
+		DECLARE_CYCLE_STAT(TEXT("FAudioThreadTask.SetSoundClassDistanceScale"), STAT_AudioSetSoundClassDistanceScale, STATGROUP_TaskGraphTasks);
+
+		FAudioThread::RunCommandOnAudioThread([this, InSoundClass, DistanceScale, TimeSec]()
+		{
+			SetSoundClassDistanceScale(InSoundClass, DistanceScale, TimeSec);
+		}, GET_STATID(STAT_AudioSetSoundClassDistanceScale));
+
+		return;
+	}
+
+	if (FSoundClassProperties* Properties = SoundClasses.Find(InSoundClass))
+	{
+		Properties->SetAttenuationDistanceScale(DistanceScale, TimeSec);
+	}
+	else
+	{
+		UE_LOG(LogAudio, Warning, TEXT("Failed to find sound class %s to set attenuation scale."), *InSoundClass->GetName());
+	}
+}
+
 float FAudioDevice::ClampPitch(float InPitchScale) const
 {
 	return FMath::Clamp(InPitchScale, GlobalMinPitch, GlobalMaxPitch);
diff --git a/Engine/Source/Runtime/Engine/Private/GameplayStatics.cpp b/Engine/Source/Runtime/Engine/Private/GameplayStatics.cpp
--- a/Engine/Source/Runtime/Engine/Private/GameplayStatics.cpp
+++ b/Engine/Source/Runtime/Engine/Private/GameplayStatics.cpp
@@ -1249,6 +1249,25 @@ void UGameplayStatics::SetGlobalPitchModulation(const UObject* WorldContextObjec
 	}
 }
 
+void UGameplayStatics::SetSoundClassDistanceScale(const UObject* WorldContextObject, USoundClass* SoundClass, float DistanceAttenuationScale, float TimeSec)
+{
+	if (!GEngine || !GEngine->UseSound())
+	{
+		return;
+	}
+
+	UWorld* ThisWorld = GEngine->GetWorldFromContextObject(WorldContextObject, EGetWorldErrorMode::LogAndReturnNull);
+	if (!ThisWorld || !ThisWorld->bAllowAudioPlayback || ThisWorld->IsNetMode(NM_DedicatedServer))
+	{
+		return;
+	}
+
+	if (FAudioDevice* AudioDevice = ThisWorld->GetAudioDevice())
+	{
+		AudioDevice->SetSoundClassDistanceScale(SoundClass, DistanceAttenuationScale, TimeSec);
+	}
+}
+
 void UGameplayStatics::SetGlobalListenerFocusParameters(const UObject* WorldContextObject, float FocusAzimuthScale, float NonFocusAzimuthScale, float FocusDistanceScale, float NonFocusDistanceScale, float FocusVolumeScale, float NonFocusVolumeScale, float FocusPriorityScale, float NonFocusPriorityScale)
 {
 	if (!GEngine || !GEngine->UseSound())
diff --git a/Engine/Source/Runtime/Engine/Private/SoundClass.cpp b/Engine/Source/Runtime/Engine/Private/SoundClass.cpp
--- a/Engine/Source/Runtime/Engine/Private/SoundClass.cpp
+++ b/Engine/Source/Runtime/Engine/Private/SoundClass.cpp
@@ -17,6 +17,26 @@
 #include "Widgets/Notifications/SNotificationList.h"
 #endif
 
+void FSoundClassProperties::SetAttenuationDistanceScale(float InAttenuationScale, float InTime)
+{
+	AttenuationScaleParam.Set(AttenuationDistanceScale, InTime);
+}
+
+void FSoundClassProperties::SetParentAttenuationDistanceScale(float InAttenuationDistanceScale)
+{
+	ParentAttenuationScale = InAttenuationDistanceScale;
+}
+
+float FSoundClassProperties::GetAttenuationDistanceScale() const
+{
+	return AttenuationScaleParam.GetValue() * ParentAttenuationScale;
+}
+
+void FSoundClassProperties::UpdateSoundClassProperties(float DeltaTime)
+{
+	AttenuationScaleParam.Update(DeltaTime);
+}
+
 /*-----------------------------------------------------------------------------
 	USoundClass implementation.
 -----------------------------------------------------------------------------*/
diff --git a/Engine/Source/Runtime/Engine/Private/SoundWave.cpp b/Engine/Source/Runtime/Engine/Private/SoundWave.cpp
--- a/Engine/Source/Runtime/Engine/Private/SoundWave.cpp
+++ b/Engine/Source/Runtime/Engine/Private/SoundWave.cpp
@@ -1722,6 +1722,8 @@ void USoundWave::Parse(FAudioDevice* AudioDevice, const UPTRINT NodeWaveInstance
 		WaveInstance->SetVolumeMultiplier(VolumeMultiplier* SoundClassProperties->Volume);
 		WaveInstance->SetPitch(WaveInstance->Pitch * SoundClassProperties->Pitch);
 
+		WaveInstance->AttenuationDistance *= FMath::Max(SoundClassProperties->GetAttenuationDistanceScale(), 0.0f);
+
 		WaveInstance->SoundClassFilterFrequency = SoundClassProperties->LowPassFilterFrequency;
 		WaveInstance->VoiceCenterChannelVolume = SoundClassProperties->VoiceCenterChannelVolume;
 		WaveInstance->RadioFilterVolume = SoundClassProperties->RadioFilterVolume * ParseParams.VolumeMultiplier;
diff --git a/Engine/Source/Runtime/Engine/Public/AudioDevice.h b/Engine/Source/Runtime/Engine/Public/AudioDevice.h
--- a/Engine/Source/Runtime/Engine/Public/AudioDevice.h
+++ b/Engine/Source/Runtime/Engine/Public/AudioDevice.h
@@ -1265,7 +1265,7 @@ private:
 	/**
 	 * Parses the sound classes and propagates multiplicative properties down the tree.
 	 */
-	void ParseSoundClasses();
+	void ParseSoundClasses(float InDeltaTime);
 
 	/** Stops quiet sounds due to being evaluated as not fulfilling concurrency requirements
 	 */
@@ -1576,6 +1576,9 @@ public:
 	void SetGlobalPitchModulation(float PitchScale, float TimeSec);
 	float ClampPitch(float InPitchScale) const;
 
+	/** Overrides the attenuation scale used on a sound class. */
+	void SetSoundClassDistanceScale(USoundClass* InSoundClass, float DistanceScale, float TimeSec);
+
 	float GetPlatformAudioHeadroom() const { check(IsInAudioThread()); return PlatformAudioHeadroom; }
 	void SetPlatformAudioHeadroom(float PlatformHeadRoom);
 
