Fix a default scene root being added when swapping scene roots via drag and drap.

diff --git a/Engine/Source/Editor/Kismet/Private/SSCSEditor.cpp b/Engine/Source/Editor/Kismet/Private/SSCSEditor.cpp
--- a/Engine/Source/Editor/Kismet/Private/SSCSEditor.cpp
+++ b/Engine/Source/Editor/Kismet/Private/SSCSEditor.cpp
@@ -2939,7 +2939,7 @@ void SSCS_RowWidget::OnMakeNewRootDropAction(FSCSEditorTreeNodePtrType DroppedNo
 		check(bWasDefaultSceneRoot || SceneRootNodePtr->CanReparent());
 
 		// Remove the current scene root node from the SCS context
-		Blueprint->SimpleConstructionScript->RemoveNode(SceneRootNodePtr->GetSCSNode());
+		Blueprint->SimpleConstructionScript->RemoveNode(SceneRootNodePtr->GetSCSNode(), /*bValidateSceneRootNodes=*/false);
 
 		// Save old root node
 		OldSceneRootNodePtr = SceneRootNodePtr;
diff --git a/Engine/Source/Runtime/Engine/Classes/Engine/SimpleConstructionScript.h b/Engine/Source/Runtime/Engine/Classes/Engine/SimpleConstructionScript.h
--- a/Engine/Source/Runtime/Engine/Classes/Engine/SimpleConstructionScript.h
+++ b/Engine/Source/Runtime/Engine/Classes/Engine/SimpleConstructionScript.h
@@ -85,8 +85,12 @@ class USimpleConstructionScript : public UObject
 	/** Adds this node to the root set */
 	ENGINE_API void AddNode(USCS_Node* Node);
 
-	/** Remove this node from the script */
-	ENGINE_API void RemoveNode(USCS_Node* Node);
+	/** 
+	 * Remove this node from the script
+	 * @param Node							The node to remove.
+	 * @param bValidateSceneRootNodes		Whether to validate scene root nodes after removing the node. Set to false if there are other operations that need to occur before validating.
+	 */
+	ENGINE_API void RemoveNode(USCS_Node* Node, bool bValidateSceneRootNodes = true);
 
 	/** Remove this node from the script and if it's the root, promote its first child to replace it */
 	ENGINE_API void RemoveNodeAndPromoteChildren(USCS_Node* Node);
diff --git a/Engine/Source/Runtime/Engine/Private/SimpleConstructionScript.cpp b/Engine/Source/Runtime/Engine/Private/SimpleConstructionScript.cpp
--- a/Engine/Source/Runtime/Engine/Private/SimpleConstructionScript.cpp
+++ b/Engine/Source/Runtime/Engine/Private/SimpleConstructionScript.cpp
@@ -819,7 +819,7 @@ void USimpleConstructionScript::AddNode(USCS_Node* Node)
 	}
 }
 
-void USimpleConstructionScript::RemoveNode(USCS_Node* Node)
+void USimpleConstructionScript::RemoveNode(USCS_Node* Node, const bool bValidateSceneRootNodes)
 {
 	// If it's a root node we are removing, clear it from the list
 	if(RootNodes.Contains(Node))
@@ -835,7 +835,10 @@ void USimpleConstructionScript::RemoveNode(USCS_Node* Node)
 		Node->ParentComponentOrVariableName = NAME_None;
 		Node->ParentComponentOwnerClassName = NAME_None;
 
-		ValidateSceneRootNodes();
+		if (bValidateSceneRootNodes)
+		{
+			ValidateSceneRootNodes();
+		}
 	}
 	// Not the root, so iterate over all nodes looking for the one with us in its ChildNodes array
 	else
