Fix bone to remove feature
#rnx

diff --git a/Engine/Plugins/Experimental/SkeletalReduction/Source/Private/SkeletalMeshReductionPlugin.cpp b/Engine/Plugins/Experimental/SkeletalReduction/Source/Private/SkeletalMeshReductionPlugin.cpp
--- a/Engine/Plugins/Experimental/SkeletalReduction/Source/Private/SkeletalMeshReductionPlugin.cpp
+++ b/Engine/Plugins/Experimental/SkeletalReduction/Source/Private/SkeletalMeshReductionPlugin.cpp
@@ -2087,12 +2087,7 @@ void FQuadricSkeletalMeshReduction::ReduceSkeletalMesh(USkeletalMesh& SkeletalMe
 	}
 	else
 	{
-		//	Bulk data arrays need to be locked before a copy can be made.
-		SrcModel->RawPointIndices.Lock(LOCK_READ_ONLY);
-		SrcModel->LegacyRawPointIndices.Lock(LOCK_READ_ONLY);
-		*NewModel = *SrcModel;
-		SrcModel->RawPointIndices.Unlock();
-		SrcModel->LegacyRawPointIndices.Unlock();
+		FSkeletalMeshLODModel::CopyStructure(NewModel, SrcModel);
 
 		// Do any joint-welding / bone removal.
 
diff --git a/Engine/Source/Developer/MeshBoneReduction/Private/MeshBoneReduction.cpp b/Engine/Source/Developer/MeshBoneReduction/Private/MeshBoneReduction.cpp
--- a/Engine/Source/Developer/MeshBoneReduction/Private/MeshBoneReduction.cpp
+++ b/Engine/Source/Developer/MeshBoneReduction/Private/MeshBoneReduction.cpp
@@ -386,12 +386,7 @@ public:
 			NewModel = new FSkeletalMeshLODModel();
 			LODModels[DesiredLOD] = NewModel;
 
-			// Bulk data arrays need to be locked before a copy can be made.
-			SrcModel->RawPointIndices.Lock(LOCK_READ_ONLY);
-			SrcModel->LegacyRawPointIndices.Lock(LOCK_READ_ONLY);
-			*NewModel = *SrcModel;
-			SrcModel->RawPointIndices.Unlock();
-			SrcModel->LegacyRawPointIndices.Unlock();
+			FSkeletalMeshLODModel::CopyStructure(NewModel, SrcModel);
 
 			TArray<FBoneIndexType> BoneIndices;
 			TArray<FMatrix> RemovedBoneMatrices;
@@ -480,7 +475,6 @@ public:
 		if (bCallPostEditChange)
 		{
 			SkeletalMesh->PostEditChange();
-			SkeletalMesh->InitResources();
 		}
 		SkeletalMesh->MarkPackageDirty();
 		
diff --git a/Engine/Source/Developer/MeshBuilder/Private/SkeletalMeshBuilder.cpp b/Engine/Source/Developer/MeshBuilder/Private/SkeletalMeshBuilder.cpp
--- a/Engine/Source/Developer/MeshBuilder/Private/SkeletalMeshBuilder.cpp
+++ b/Engine/Source/Developer/MeshBuilder/Private/SkeletalMeshBuilder.cpp
@@ -172,6 +172,19 @@ bool FSkeletalMeshBuilder::Build(USkeletalMesh* SkeletalMesh, const int32 LODInd
 			}
 			FLODUtilities::SimplifySkeletalMeshLOD(UpdateContext, LODIndex, false);
 		}
+		else
+		{
+			if (LODInfo->BonesToRemove.Num() > 0)
+			{
+				TArray<FName> BonesToRemove;
+				BonesToRemove.Reserve(LODInfo->BonesToRemove.Num());
+				for (const FBoneReference& BoneReference : LODInfo->BonesToRemove)
+				{
+					BonesToRemove.Add(BoneReference.BoneName);
+				}
+				MeshUtilities.RemoveBonesFromMesh(SkeletalMesh, LODIndex, &BonesToRemove);
+			}
+		}
 	}
 
 	FSkeletalMeshLODModel& LODModelAfterReduction = SkeletalMesh->GetImportedModel()->LODModels[LODIndex];
diff --git a/Engine/Source/Runtime/Engine/Private/SkeletalMeshLODModel.cpp b/Engine/Source/Runtime/Engine/Private/SkeletalMeshLODModel.cpp
--- a/Engine/Source/Runtime/Engine/Private/SkeletalMeshLODModel.cpp
+++ b/Engine/Source/Runtime/Engine/Private/SkeletalMeshLODModel.cpp
@@ -1096,4 +1096,22 @@ void FSkeletalMeshLODModel::UpdateChunkedSectionInfo(const FString& SkeletalMesh
 	}
 }
 
+bool FSkeletalMeshLODModel::CopyStructure(FSkeletalMeshLODModel* Destination, FSkeletalMeshLODModel* Source)
+{
+	if (Source->RawPointIndices.IsLocked() || Source->LegacyRawPointIndices.IsLocked() || Source->RawSkeletalMeshBulkData.GetBulkData().IsLocked())
+	{
+		return false;
+	}
+	// Bulk data arrays need to be locked before a copy can be made.
+	Source->RawPointIndices.Lock(LOCK_READ_ONLY);
+	Source->LegacyRawPointIndices.Lock(LOCK_READ_ONLY);
+	Source->RawSkeletalMeshBulkData.GetBulkData().Lock(LOCK_READ_ONLY);
+	*Destination = *Source;
+	Source->RawSkeletalMeshBulkData.GetBulkData().Unlock();
+	Source->RawPointIndices.Unlock();
+	Source->LegacyRawPointIndices.Unlock();
+
+	return true;
+}
+
 #endif // WITH_EDITOR
diff --git a/Engine/Source/Runtime/Engine/Public/Rendering/SkeletalMeshLODModel.h b/Engine/Source/Runtime/Engine/Public/Rendering/SkeletalMeshLODModel.h
--- a/Engine/Source/Runtime/Engine/Public/Rendering/SkeletalMeshLODModel.h
+++ b/Engine/Source/Runtime/Engine/Public/Rendering/SkeletalMeshLODModel.h
@@ -394,6 +394,11 @@ public:
 	* skeletal mesh build refactor should need to call this function.
 	*/
 	ENGINE_API void UpdateChunkedSectionInfo(const FString& SkeletalMeshName, TArray<int32>& LODMaterialMap);
+
+	/**
+	* Copy one structure to the other, make sure all bulk data is unlock and the data can be read before copying.
+	*/
+	static ENGINE_API bool CopyStructure(FSkeletalMeshLODModel* Destination, FSkeletalMeshLODModel* Source);
 };
 
 #endif // WITH_EDITOR
\ No newline at end of file
