Fixed crash when building texture that have invalid data.

diff --git a/Engine/Source/Runtime/Engine/Private/TextureDerivedDataTask.cpp b/Engine/Source/Runtime/Engine/Private/TextureDerivedDataTask.cpp
--- a/Engine/Source/Runtime/Engine/Private/TextureDerivedDataTask.cpp
+++ b/Engine/Source/Runtime/Engine/Private/TextureDerivedDataTask.cpp
@@ -195,11 +195,14 @@ void FTextureSourceData::GetAsyncSourceMips(IImageWrapperModule* InImageWrapper)
 
 void FTextureCacheDerivedDataWorker::BuildTexture()
 {
+	const bool bHasValidMip0 = TextureData.Blocks.Num() && TextureData.Blocks[0].MipsPerLayer.Num() && TextureData.Blocks[0].MipsPerLayer[0].Num();
+
 	FFormatNamedArguments Args;
 	Args.Add(TEXT("TextureName"), FText::FromString(Texture.GetName()));
 	Args.Add(TEXT("TextureFormatName"), FText::FromString(BuildSettingsPerLayer[0].TextureFormatName.GetPlainNameString()));
-	Args.Add(TEXT("TextureResolutionX"), FText::FromString(FString::FromInt(TextureData.Blocks[0].MipsPerLayer[0][0].SizeX)));
-	Args.Add(TEXT("TextureResolutionY"), FText::FromString(FString::FromInt(TextureData.Blocks[0].MipsPerLayer[0][0].SizeY)));
+	Args.Add(TEXT("TextureResolutionX"), FText::FromString(FString::FromInt(bHasValidMip0 ? TextureData.Blocks[0].MipsPerLayer[0][0].SizeX : 0)));
+	Args.Add(TEXT("TextureResolutionY"), FText::FromString(FString::FromInt(bHasValidMip0 ? TextureData.Blocks[0].MipsPerLayer[0][0].SizeY : 0)));
+
 	FTextureStatusMessageContext StatusMessage(FText::Format(NSLOCTEXT("Engine", "BuildTextureStatus", "Building textures: {TextureName} ({TextureFormatName}, {TextureResolutionX}X{TextureResolutionY})"), Args));
 
 	if (!ensure(Compressor))
@@ -240,9 +243,7 @@ void FTextureCacheDerivedDataWorker::BuildTexture()
 			UE_LOG(LogTexture, Warning, TEXT("Failed to build %s derived data for %s"), *BuildSettingsPerLayer[0].TextureFormatName.GetPlainNameString(), *Texture.GetPathName());
 		}
 	}
-	else if (TextureData.Blocks.Num() &&
-		TextureData.Blocks[0].MipsPerLayer.Num() &&
-		TextureData.Blocks[0].MipsPerLayer[0].Num())
+	else if (bHasValidMip0)
 	{
 		// Only support single Block/Layer here (Blocks and Layers are intended for VT support)
 		ensure(TextureData.Blocks.Num() == 1);
