Cameras: Add proximity-based camera shakes

This change adds a new actor/component type which defines a "shake source", along with some attenuation settings.
Some new Blueprint functions allow users to start camera shakes that are anchored to those sources, and whose shake intensities will change dynamically based on the curent camera's distance to the respective sources.

diff --git a/Engine/Source/Runtime/Engine/Classes/Camera/CameraModifier_CameraShake.h b/Engine/Source/Runtime/Engine/Classes/Camera/CameraModifier_CameraShake.h
--- a/Engine/Source/Runtime/Engine/Classes/Camera/CameraModifier_CameraShake.h
+++ b/Engine/Source/Runtime/Engine/Classes/Camera/CameraModifier_CameraShake.h
@@ -14,14 +14,44 @@
 #include "CameraModifier_CameraShake.generated.h"
 
 class UCameraShake;
+class UCameraShakeSourceComponent;
 
 USTRUCT()
 struct FPooledCameraShakes
 {
-	GENERATED_USTRUCT_BODY()
+	GENERATED_BODY()
 		
 	UPROPERTY()
-	TArray<class UCameraShake*> PooledShakes;
+	TArray<UCameraShake*> PooledShakes;
+};
+
+USTRUCT()
+struct FActiveCameraShakeInfo
+{
+	GENERATED_BODY()
+
+	FActiveCameraShakeInfo() : ShakeInstance(nullptr), ShakeSource(nullptr) {}
+
+	UPROPERTY()
+	UCameraShake* ShakeInstance;
+
+	UPROPERTY()
+	UCameraShakeSourceComponent* ShakeSource;
+};
+
+struct FAddCameraShakeParams
+{
+	float Scale;
+	ECameraAnimPlaySpace::Type PlaySpace;
+	FRotator UserPlaySpaceRot;
+	UCameraShakeSourceComponent* SourceComponent;
+
+	FAddCameraShakeParams() 
+		: Scale(1.f), PlaySpace(ECameraAnimPlaySpace::CameraLocal), UserPlaySpaceRot(FRotator::ZeroRotator), SourceComponent(nullptr)
+	{}
+	FAddCameraShakeParams(float InScale, ECameraAnimPlaySpace::Type InPlaySpace = ECameraAnimPlaySpace::CameraLocal, FRotator InUserPlaySpaceRot = FRotator::ZeroRotator, UCameraShakeSourceComponent* InSourceComponent = nullptr)
+		: Scale(InScale), PlaySpace(InPlaySpace), UserPlaySpaceRot(InUserPlaySpaceRot), SourceComponent(InSourceComponent)
+	{}
 };
 
 //~=============================================================================
@@ -32,18 +62,17 @@ struct FPooledCameraShakes
 UCLASS(config=Camera)
 class ENGINE_API UCameraModifier_CameraShake : public UCameraModifier
 {
-	GENERATED_UCLASS_BODY()
+	GENERATED_BODY()
 
 public:
-	/** List of active CameraShake instances */
-	UPROPERTY()
-	TArray<class UCameraShake*> ActiveShakes;
+	UCameraModifier_CameraShake(const FObjectInitializer& ObjectInitializer);
 
-	UPROPERTY()
-	TMap<TSubclassOf<class UCameraShake>, FPooledCameraShakes> ExpiredPooledShakesMap;
-
-	void SaveShakeInExpiredPool(class UCameraShake* ShakeInst);
-	UCameraShake* ReclaimShakeFromExpiredPool(TSubclassOf<class UCameraShake> CameraShakeClass);
+	/** 
+	 * Adds a new active screen shake to be applied. 
+	 * @param NewShake - The class of camera shake to instantiate.
+	 * @param Params - The parameters for the new camera shake.
+	 */
+	virtual UCameraShake* AddCameraShake(TSubclassOf<UCameraShake> NewShake, const FAddCameraShakeParams& Params);
 
 	/** 
 	 * Adds a new active screen shake to be applied. 
@@ -52,7 +81,11 @@ public:
 	 * @param PlaySpace - Which coordinate system to play the shake in.
 	 * @param UserPlaySpaceRot - Coordinate system to play shake when PlaySpace == CAPS_UserDefined.
 	 */
-	virtual class UCameraShake* AddCameraShake(TSubclassOf<class UCameraShake> NewShake, float Scale, ECameraAnimPlaySpace::Type PlaySpace=ECameraAnimPlaySpace::CameraLocal, FRotator UserPlaySpaceRot = FRotator::ZeroRotator);
+	UE_DEPRECATED(4.25, "Please use the new AddCameraShake method that takes a parameter struct.")
+	virtual class UCameraShake* AddCameraShake(TSubclassOf<UCameraShake> NewShake, float Scale, ECameraAnimPlaySpace::Type PlaySpace=ECameraAnimPlaySpace::CameraLocal, FRotator UserPlaySpaceRot = FRotator::ZeroRotator)
+	{
+		return AddCameraShake(NewShake, FAddCameraShakeParams(Scale, PlaySpace, UserPlaySpaceRot));
+	}
 	
 	/**
 	 * Stops and removes the camera shake of the given class from the camera.
@@ -65,7 +98,9 @@ public:
 	 * Stops and removes all camera shakes of the given class from the camera. 
 	 * @param bImmediately		If true, shake stops right away regardless of blend out settings. If false, shake may blend out according to its settings.
 	 */
-	virtual void RemoveAllCameraShakesOfClass(TSubclassOf<class UCameraShake> ShakeClass, bool bImmediately = true);
+	virtual void RemoveAllCameraShakesOfClass(TSubclassOf<UCameraShake> ShakeClass, bool bImmediately = true);
+
+	virtual void RemoveAllCameraShakesFromSource(UCameraShakeSourceComponent* SourceComponent, bool bImmediately = true);
 
 	/** 
 	 * Stops and removes all camera shakes from the camera. 
@@ -78,6 +113,17 @@ public:
 	//~ End UCameraModifer Interface
 
 protected:
+
+	/** List of active CameraShake instances */
+	UPROPERTY()
+	TArray<FActiveCameraShakeInfo> ActiveShakes;
+
+	UPROPERTY()
+	TMap<TSubclassOf<UCameraShake>, FPooledCameraShakes> ExpiredPooledShakesMap;
+
+	void SaveShakeInExpiredPool(UCameraShake* ShakeInst);
+	UCameraShake* ReclaimShakeFromExpiredPool(TSubclassOf<UCameraShake> CameraShakeClass);
+
 	/** Scaling factor applied to all camera shakes in when in splitscreen mode. Normally used to reduce shaking, since shakes feel more intense in a smaller viewport. */
 	UPROPERTY(EditAnywhere, Category = CameraModifier_CameraShake)
 	float SplitScreenShakeScale;
diff --git a/Engine/Source/Runtime/Engine/Classes/Camera/CameraShakeSourceActor.cpp b/Engine/Source/Runtime/Engine/Classes/Camera/CameraShakeSourceActor.cpp
new file mode 100644
--- /dev/null
+++ b/Engine/Source/Runtime/Engine/Classes/Camera/CameraShakeSourceActor.cpp
@@ -0,0 +1,12 @@
+// Copyright 1998-2019 Epic Games, Inc. All Rights Reserved.
+
+#include "CameraShakeSourceActor.h"
+#include "CameraShakeSourceComponent.h"
+
+
+ACameraShakeSourceActor::ACameraShakeSourceActor(const FObjectInitializer& ObjectInitializer)
+	: Super(ObjectInitializer)
+{
+	CameraShakeSourceComponent = CreateDefaultSubobject<UCameraShakeSourceComponent>(TEXT("CameraShakeSourceComponent"));
+	RootComponent = CameraShakeSourceComponent;
+}
diff --git a/Engine/Source/Runtime/Engine/Classes/Camera/CameraShakeSourceActor.h b/Engine/Source/Runtime/Engine/Classes/Camera/CameraShakeSourceActor.h
new file mode 100644
--- /dev/null
+++ b/Engine/Source/Runtime/Engine/Classes/Camera/CameraShakeSourceActor.h
@@ -0,0 +1,30 @@
+// Copyright 1998-2019 Epic Games, Inc. All Rights Reserved.
+
+#pragma once
+
+#include "CoreMinimal.h"
+#include "UObject/ObjectMacros.h"
+#include "GameFramework/Actor.h"
+
+#include "CameraShakeSourceActor.generated.h"
+
+
+class UCameraShakeSourceComponent;
+class USceneComponent;
+class UWorld;
+
+
+UCLASS()
+class ENGINE_API ACameraShakeSourceActor : public AActor
+{
+	GENERATED_BODY()
+
+public:
+	ACameraShakeSourceActor(const FObjectInitializer& ObjectInitializer);
+
+	UCameraShakeSourceComponent* GetCameraShakeSourceComponent() const { return CameraShakeSourceComponent; }
+
+private:
+	UPROPERTY(VisibleAnywhere, BlueprintReadOnly, Category = CameraShakeSourceActor, meta = (AllowPrivateAccess = "true"))
+	UCameraShakeSourceComponent* CameraShakeSourceComponent;
+};
diff --git a/Engine/Source/Runtime/Engine/Classes/Camera/CameraShakeSourceComponent.cpp b/Engine/Source/Runtime/Engine/Classes/Camera/CameraShakeSourceComponent.cpp
new file mode 100644
--- /dev/null
+++ b/Engine/Source/Runtime/Engine/Classes/Camera/CameraShakeSourceComponent.cpp
@@ -0,0 +1,71 @@
+// Copyright 1998-2019 Epic Games, Inc. All Rights Reserved.
+
+#include "CameraShakeSourceComponent.h"
+#include "CameraShake.h"
+#include "CinematicCameraModule.h"
+#include "Engine/Engine.h"
+#include "Engine/World.h"
+#include "GameFramework/PlayerController.h"
+#include "Modules/ModuleManager.h"
+
+
+UCameraShakeSourceComponent::UCameraShakeSourceComponent(const FObjectInitializer& ObjectInitializer)
+	: Super(ObjectInitializer)
+	, Attenuation(ECameraShakeAttenuation::Quadratic)
+	, InnerAttenuationRadius(100.f)
+	, OuterAttenuationRadius(1000.f)
+{
+}
+
+void UCameraShakeSourceComponent::PlayCameraShake(TSubclassOf<UCameraShake> CameraShake)
+{
+	for (FConstPlayerControllerIterator Iterator = GetWorld()->GetPlayerControllerIterator(); Iterator; ++Iterator)
+	{
+		APlayerController* PlayerController = Iterator->Get();
+		if (PlayerController != nullptr && PlayerController->PlayerCameraManager != nullptr)
+		{
+			PlayerController->ClientPlayCameraShakeFromSource(CameraShake, this);
+		}
+	}
+}
+
+void UCameraShakeSourceComponent::StopAllCameraShakes(bool bImmediately)
+{
+	for (FConstPlayerControllerIterator Iterator = GetWorld()->GetPlayerControllerIterator(); Iterator; ++Iterator)
+	{
+		APlayerController* PlayerController = Iterator->Get();
+		if (PlayerController != nullptr && PlayerController->PlayerCameraManager != nullptr)
+		{
+			PlayerController->ClientStopCameraShakesFromSource(this, bImmediately);
+		}
+	}
+}
+
+float UCameraShakeSourceComponent::GetAttenuationFactor(const FVector& Location) const
+{
+	const FTransform& SourceTransform = GetComponentTransform();
+	const FVector LocationToSource = SourceTransform.GetTranslation() - Location;
+
+	float AttFactor = 1.0f;
+	switch (Attenuation)
+	{
+		case ECameraShakeAttenuation::Quadratic:
+			AttFactor = 2.0f;
+			break;
+		default:
+			break;
+	}
+
+	if (InnerAttenuationRadius < OuterAttenuationRadius)
+	{
+		float DistFactor = (LocationToSource.Size() - InnerAttenuationRadius) / (OuterAttenuationRadius - InnerAttenuationRadius);
+		DistFactor = 1.f - FMath::Clamp(DistFactor, 0.f, 1.f);
+		return FMath::Pow(DistFactor, AttFactor);
+	}
+	else if (OuterAttenuationRadius > 0)
+	{
+		// Just cut the intensity after the end distance.
+		return (LocationToSource.SizeSquared() < FMath::Square(OuterAttenuationRadius)) ? 1.f : 0.f;
+	}
+	return 1.f;
+}
diff --git a/Engine/Source/Runtime/Engine/Classes/Camera/CameraShakeSourceComponent.h b/Engine/Source/Runtime/Engine/Classes/Camera/CameraShakeSourceComponent.h
new file mode 100644
--- /dev/null
+++ b/Engine/Source/Runtime/Engine/Classes/Camera/CameraShakeSourceComponent.h
@@ -0,0 +1,50 @@
+// Copyright 1998-2019 Epic Games, Inc. All Rights Reserved.
+
+#pragma once
+
+#include "CoreMinimal.h"
+#include "Components/SceneComponent.h"
+#include "UObject/ObjectMacros.h"
+#include "CameraShakeSourceComponent.generated.h"
+
+class UCameraShake;
+
+UENUM(BlueprintType)
+enum class ECameraShakeAttenuation : uint8
+{
+	Linear,
+	Quadratic
+};
+
+UCLASS(Blueprintable, meta = (BlueprintSpawnableComponent))
+class ENGINE_API UCameraShakeSourceComponent : public USceneComponent
+{
+	GENERATED_BODY()
+
+public:
+	UCameraShakeSourceComponent(const FObjectInitializer& ObjectInitializer);
+
+	/** The attenuation profile for how camera shakes' intensity falls off with distance */
+	UPROPERTY(EditAnywhere, BlueprintReadWrite, Category = CameraShake)
+	ECameraShakeAttenuation Attenuation;
+
+	/** Under this distance from the source, the camera shakes are at full intensity */
+	UPROPERTY(EditAnywhere, BlueprintReadWrite, Category = CameraShake)
+	float InnerAttenuationRadius;
+
+	/** Outside of this distance from the source, the camera shakes don't apply at all */
+	UPROPERTY(EditAnywhere, BlueprintReadWrite, Category = CameraShake)
+	float OuterAttenuationRadius;
+
+	/** Starts a new camera shake originating from this source, and apply it on all player controllers */
+	UFUNCTION(BlueprintCallable, Category = CameraShake)
+	void PlayCameraShake(TSubclassOf<UCameraShake> CameraShake);
+
+	/** Stops all currently active camera shakes that are originating from this source from all player controllers */
+	UFUNCTION(BlueprintCallable, Category = CameraShake)
+	void StopAllCameraShakes(bool bImmediately = true);
+
+	/** Computes an attenuation factor from this source */
+	UFUNCTION(BlueprintCallable, Category = CameraShake)
+	float GetAttenuationFactor(const FVector& Location) const;
+};
diff --git a/Engine/Source/Runtime/Engine/Classes/Camera/PlayerCameraManager.h b/Engine/Source/Runtime/Engine/Classes/Camera/PlayerCameraManager.h
--- a/Engine/Source/Runtime/Engine/Classes/Camera/PlayerCameraManager.h
+++ b/Engine/Source/Runtime/Engine/Classes/Camera/PlayerCameraManager.h
@@ -769,15 +769,27 @@ public:
 	 */
 	UFUNCTION(BlueprintCallable, Category = "Camera Shakes")
 	virtual class UCameraShake* PlayCameraShake(TSubclassOf<class UCameraShake> ShakeClass, float Scale=1.f, enum ECameraAnimPlaySpace::Type PlaySpace = ECameraAnimPlaySpace::CameraLocal, FRotator UserPlaySpaceRot = FRotator::ZeroRotator);
+
+	/** 
+	 * Plays a camera shake on this camera.
+	 * @param Shake - The class of camera shake to play.
+	 * @param SourceComponent - The source from which the camera shake originates.
+	 */
+	UFUNCTION(BlueprintCallable, Category = "Camera Shakes")
+	virtual class UCameraShake* PlayCameraShakeFromSource(TSubclassOf<class UCameraShake> ShakeClass, class UCameraShakeSourceComponent* SourceComponent);
 	
 	/** Immediately stops the given shake instance and invalidates it. */
 	UFUNCTION(BlueprintCallable, Category = "Camera Shakes")
 	virtual void StopCameraShake(class UCameraShake* ShakeInstance, bool bImmediately = true);
 
-	/** Stops playing CameraShake of the given class. */
+	/** Stops playing all shakes of the given class. */
 	UFUNCTION(BlueprintCallable, Category = "Camera Shakes")
 	virtual void StopAllInstancesOfCameraShake(TSubclassOf<class UCameraShake> Shake, bool bImmediately = true);
 
+	/** Stops playing all shakes originating from the given source. */
+	UFUNCTION(BlueprintCallable, Category = "Camera Shakes")
+	virtual void StopAllInstancesOfCameraShakeFromSource(class UCameraShakeSourceComponent* SourceComponent, bool bImmediately = true);
+
 	/** Stops all active camera shakes on this camera. */
 	UFUNCTION(BlueprintCallable, Category = "Camera Shakes")
 	virtual void StopAllCameraShakes(bool bImmediately = true);
diff --git a/Engine/Source/Runtime/Engine/Classes/GameFramework/PlayerController.h b/Engine/Source/Runtime/Engine/Classes/GameFramework/PlayerController.h
--- a/Engine/Source/Runtime/Engine/Classes/GameFramework/PlayerController.h
+++ b/Engine/Source/Runtime/Engine/Classes/GameFramework/PlayerController.h
@@ -932,6 +932,14 @@ public:
 	UFUNCTION(unreliable, client, BlueprintCallable, Category="Game|Feedback")
 	void ClientPlayCameraShake(TSubclassOf<class UCameraShake> Shake, float Scale = 1.f, ECameraAnimPlaySpace::Type PlaySpace = ECameraAnimPlaySpace::CameraLocal, FRotator UserPlaySpaceRot = FRotator::ZeroRotator);
 
+	/** 
+	 * Play Camera Shake localized to a given source
+	 * @param Shake - Camera shake animation to play
+	 * @param SourceComponent - The source from which the camera shakes originates
+	 */
+	UFUNCTION(unreliable, client, BlueprintCallable, Category="Game|Feedback")
+	void ClientPlayCameraShakeFromSource(TSubclassOf<class UCameraShake> Shake, class UCameraShakeSourceComponent* SourceComponent);
+
 	/**
 	 * Play sound client-side (so only the client will hear it)
 	 * @param Sound	- Sound to play
@@ -1062,6 +1070,10 @@ public:
 	UFUNCTION(reliable, client, BlueprintCallable, Category="Game|Feedback")
 	void ClientStopCameraShake(TSubclassOf<class UCameraShake> Shake, bool bImmediately = true);
 
+	/** Stop camera shake on client.  */
+	UFUNCTION(reliable, client, BlueprintCallable, Category="Game|Feedback")
+	void ClientStopCameraShakesFromSource(class UCameraShakeSourceComponent* SourceComponent, bool bImmediately = true);
+
 	/** 
 	 * Play a force feedback pattern on the player's controller
 	 * @param	ForceFeedbackEffect		The force feedback pattern to play
diff --git a/Engine/Source/Runtime/Engine/Private/Camera/CameraModifier_CameraShake.cpp b/Engine/Source/Runtime/Engine/Private/Camera/CameraModifier_CameraShake.cpp
--- a/Engine/Source/Runtime/Engine/Private/Camera/CameraModifier_CameraShake.cpp
+++ b/Engine/Source/Runtime/Engine/Private/Camera/CameraModifier_CameraShake.cpp
@@ -1,10 +1,11 @@
 // Copyright 1998-2019 Epic Games, Inc. All Rights Reserved.
 
 #include "Camera/CameraModifier_CameraShake.h"
-#include "EngineGlobals.h"
 #include "Camera/CameraShake.h"
+#include "Camera/CameraShakeSourceComponent.h"
 #include "Camera/PlayerCameraManager.h"
 #include "Engine/Engine.h"
+#include "EngineGlobals.h"
 
 //////////////////////////////////////////////////////////////////////////
 // UCameraModifier_CameraShake
@@ -32,20 +33,35 @@ bool UCameraModifier_CameraShake::ModifyCamera(float DeltaTime, FMinimalViewInfo
 	// Update and apply active shakes
 	if( ActiveShakes.Num() > 0 )
 	{
-		for (UCameraShake* ShakeInst : ActiveShakes)
+		for (FActiveCameraShakeInfo& ShakeInfo : ActiveShakes)
 		{
-			ShakeInst->UpdateAndApplyCameraShake(DeltaTime, Alpha, InOutPOV);
+			if (ShakeInfo.ShakeInstance != nullptr)
+			{
+				// Compute the scale of this shake for this frame according to the location
+				// of its source.
+				float CurShakeAlpha = Alpha;
+				if (ShakeInfo.ShakeSource != nullptr)
+				{
+					const float AttenuationFactor = ShakeInfo.ShakeSource->GetAttenuationFactor(InOutPOV.Location);
+					CurShakeAlpha *= AttenuationFactor;
+				}
+
+				ShakeInfo.ShakeInstance->UpdateAndApplyCameraShake(DeltaTime, CurShakeAlpha, InOutPOV);
+			}
 		}
 
 		// Delete any obsolete shakes
-		for (int32 i=ActiveShakes.Num()-1; i>=0; i--)
+		for (int32 i = ActiveShakes.Num() - 1; i >= 0; i--)
 		{
-			UCameraShake* const ShakeInst = ActiveShakes[i];
-			if ((ShakeInst == nullptr) || ShakeInst->IsFinished())
+			FActiveCameraShakeInfo ShakeInfo = ActiveShakes[i]; // Copy struct, we're going to maybe delete it.
+			if (ShakeInfo.ShakeInstance == nullptr || ShakeInfo.ShakeInstance->IsFinished())
 			{
 				ActiveShakes.RemoveAt(i, 1);
 
-				SaveShakeInExpiredPool(ShakeInst);
+				if (ShakeInfo.ShakeInstance != nullptr)
+				{
+					SaveShakeInExpiredPool(ShakeInfo.ShakeInstance);
+				}
 			}
 		}
 	}
@@ -58,31 +74,43 @@ bool UCameraModifier_CameraShake::ModifyCamera(float DeltaTime, FMinimalViewInfo
 	return false;
 }
 
-UCameraShake* UCameraModifier_CameraShake::AddCameraShake(TSubclassOf<class UCameraShake> ShakeClass, float Scale, ECameraAnimPlaySpace::Type PlaySpace, FRotator UserPlaySpaceRot)
+UCameraShake* UCameraModifier_CameraShake::AddCameraShake(TSubclassOf<UCameraShake> ShakeClass, const FAddCameraShakeParams& Params)
 {
 	SCOPE_CYCLE_COUNTER(STAT_AddCameraShake);
 
 	if (ShakeClass != nullptr)
 	{
-		// adjust for splitscreen
+		float Scale = Params.Scale;
+		UCameraShakeSourceComponent* SourceComponent = Params.SourceComponent;
+
+		// Adjust for splitscreen
 		if (GEngine->IsSplitScreen(CameraOwner->GetWorld()))
 		{
 			Scale *= SplitScreenShakeScale;
 		}
 
 		UCameraShake const* const ShakeCDO = GetDefault<UCameraShake>(ShakeClass);
-		if (ShakeCDO && ShakeCDO->bSingleInstance)
+		const bool bIsSingleInstance = ShakeCDO && ShakeCDO->bSingleInstance;
+		if (bIsSingleInstance)
 		{
-			// look for existing instance of same class
-			for (UCameraShake* ShakeInst : ActiveShakes)
+			// Look for existing instance of same class
+			for (FActiveCameraShakeInfo& ShakeInfo : ActiveShakes)
 			{
+				UCameraShake* ShakeInst = ShakeInfo.ShakeInstance;
 				if (ShakeInst && (ShakeClass == ShakeInst->GetClass()))
 				{
-					// just restart the existing shake
-					ShakeInst->PlayShake(CameraOwner, Scale, PlaySpace, UserPlaySpaceRot);
+					// Just restart the existing shake
+					ShakeInst->PlayShake(CameraOwner, Scale, Params.PlaySpace, Params.UserPlaySpaceRot);
 					return ShakeInst;
 				}
 			}
+
+			// No existing instance... we'll build a new one, but let's null the
+			// source in case one was specified (we don't support a source because
+			// it would in some cases restart the shake by teleporting it to another
+			// location, which would look bad).
+			ensureMsgf(SourceComponent == nullptr, TEXT("CameraShake assets with SingleInstance enabled shouldn't be located to a source."));
+			SourceComponent = nullptr;
 		}
 
 		// Try to find a shake in the expired pool
@@ -97,15 +125,17 @@ UCameraShake* UCameraModifier_CameraShake::AddCameraShake(TSubclassOf<class UCam
 		if (NewInst)
 		{
 			// Initialize new shake and add it to the list of active shakes
-			NewInst->PlayShake(CameraOwner, Scale, PlaySpace, UserPlaySpaceRot);
+			NewInst->PlayShake(CameraOwner, Scale, Params.PlaySpace, Params.UserPlaySpaceRot);
 
-			// look for nulls in the array to replace first -- keeps the array compact
+			// Look for nulls in the array to replace first -- keeps the array compact
 			bool bReplacedNull = false;
 			for (int32 Idx = 0; Idx < ActiveShakes.Num(); ++Idx)
 			{
-				if (ActiveShakes[Idx] == nullptr)
+				FActiveCameraShakeInfo& ShakeInfo = ActiveShakes[Idx];
+				if (ShakeInfo.ShakeInstance == nullptr)
 				{
-					ActiveShakes[Idx] = NewInst;
+					ShakeInfo.ShakeInstance = NewInst;
+					ShakeInfo.ShakeSource = SourceComponent;
 					bReplacedNull = true;
 				}
 			}
@@ -113,7 +143,10 @@ UCameraShake* UCameraModifier_CameraShake::AddCameraShake(TSubclassOf<class UCam
 			// no holes, extend the array
 			if (bReplacedNull == false)
 			{
-				ActiveShakes.Emplace(NewInst);
+				FActiveCameraShakeInfo ShakeInfo;
+				ShakeInfo.ShakeInstance = NewInst;
+				ShakeInfo.ShakeSource = SourceComponent;
+				ActiveShakes.Emplace(ShakeInfo);
 			}
 		}
 
@@ -151,7 +184,8 @@ void UCameraModifier_CameraShake::RemoveCameraShake(UCameraShake* ShakeInst, boo
 {
 	for (int32 i = 0; i < ActiveShakes.Num(); ++i)
 	{
-		if (ActiveShakes[i] == ShakeInst)
+		FActiveCameraShakeInfo& ShakeInfo = ActiveShakes[i];
+		if (ShakeInfo.ShakeInstance == ShakeInst)
 		{
 			ShakeInst->StopShake(bImmediately);
 
@@ -165,17 +199,35 @@ void UCameraModifier_CameraShake::RemoveCameraShake(UCameraShake* ShakeInst, boo
 	}
 }
 
-void UCameraModifier_CameraShake::RemoveAllCameraShakesOfClass(TSubclassOf<class UCameraShake> ShakeClass, bool bImmediately)
+void UCameraModifier_CameraShake::RemoveAllCameraShakesOfClass(TSubclassOf<UCameraShake> ShakeClass, bool bImmediately)
+{
+	for (int32 i = ActiveShakes.Num()- 1; i >= 0; --i)
+	{
+		UCameraShake* ShakeInst = ActiveShakes[i].ShakeInstance;
+		if (ShakeInst != nullptr && (ShakeInst->GetClass()->IsChildOf(ShakeClass)))
+		{
+			ShakeInst->StopShake(bImmediately);
+			if (bImmediately)
+			{
+				ActiveShakes.RemoveAt(i, 1);
+				SaveShakeInExpiredPool(ShakeInst);
+			}
+		}
+	}
+}
+
+void UCameraModifier_CameraShake::RemoveAllCameraShakesFromSource(UCameraShakeSourceComponent* SourceComponent, bool bImmediately)
 {
-	for (int32 i = ActiveShakes.Num()-1; i >= 0; --i)
+	for (int32 i = ActiveShakes.Num() - 1; i >= 0; --i)
 	{
-		if ( ActiveShakes[i] && (ActiveShakes[i]->GetClass()->IsChildOf(ShakeClass)) )
+		FActiveCameraShakeInfo ShakeInfo = ActiveShakes[i];  // Copy struct because we might delete it.
+		if (ShakeInfo.ShakeSource == SourceComponent && ShakeInfo.ShakeInstance != nullptr)
 		{
-			ActiveShakes[i]->StopShake(bImmediately);
+			ShakeInfo.ShakeInstance->StopShake(bImmediately);
 			if (bImmediately)
 			{
-				SaveShakeInExpiredPool(ActiveShakes[i]);
 				ActiveShakes.RemoveAt(i, 1);
+				SaveShakeInExpiredPool(ShakeInfo.ShakeInstance);
 			}
 		}
 	}
@@ -183,20 +235,20 @@ void UCameraModifier_CameraShake::RemoveAllCameraShakesOfClass(TSubclassOf<class
 
 void UCameraModifier_CameraShake::RemoveAllCameraShakes(bool bImmediately)
 {
-	// clean up any active camera shake anims
-	for (UCameraShake* Inst : ActiveShakes)
+	// Clean up any active camera shake anims
+	for (FActiveCameraShakeInfo& ShakeInfo : ActiveShakes)
 	{
-		Inst->StopShake(bImmediately);
+		ShakeInfo.ShakeInstance->StopShake(bImmediately);
 	}
 
 	if (bImmediately)
 	{
-		for (UCameraShake* Inst : ActiveShakes)
+		for (FActiveCameraShakeInfo& ShakeInfo : ActiveShakes)
 		{
-			SaveShakeInExpiredPool(Inst);
+			SaveShakeInExpiredPool(ShakeInfo.ShakeInstance);
 		}
 
-		// clear ActiveShakes array
+		// Clear ActiveShakes array
 		ActiveShakes.Empty();
 	}
 }
diff --git a/Engine/Source/Runtime/Engine/Private/PlayerCameraManager.cpp b/Engine/Source/Runtime/Engine/Private/PlayerCameraManager.cpp
--- a/Engine/Source/Runtime/Engine/Private/PlayerCameraManager.cpp
+++ b/Engine/Source/Runtime/Engine/Private/PlayerCameraManager.cpp
@@ -1271,12 +1271,21 @@ UCameraShake* APlayerCameraManager::PlayCameraShake(TSubclassOf<UCameraShake> Sh
 {
 	if (ShakeClass && CachedCameraShakeMod && (Scale > 0.0f) )
 	{
-		return CachedCameraShakeMod->AddCameraShake(ShakeClass, Scale, PlaySpace, UserPlaySpaceRot);
+		return CachedCameraShakeMod->AddCameraShake(ShakeClass, FAddCameraShakeParams(Scale, PlaySpace, UserPlaySpaceRot));
 	}
 
 	return nullptr;
 }
 
+UCameraShake* APlayerCameraManager::PlayCameraShakeFromSource(TSubclassOf<class UCameraShake> ShakeClass, class UCameraShakeSourceComponent* SourceComponent)
+{
+	if (ShakeClass && CachedCameraShakeMod)
+	{
+		return CachedCameraShakeMod->AddCameraShake(ShakeClass, FAddCameraShakeParams(1.0f, ECameraAnimPlaySpace::CameraLocal, FRotator::ZeroRotator, SourceComponent));
+	}
+
+	return nullptr;
+}
 
 void APlayerCameraManager::StopCameraShake(UCameraShake* ShakeInst, bool bImmediately)
 {
@@ -1294,6 +1303,14 @@ void APlayerCameraManager::StopAllInstancesOfCameraShake(TSubclassOf<class UCame
 	}
 }
 
+void APlayerCameraManager::StopAllInstancesOfCameraShakeFromSource(class UCameraShakeSourceComponent* SourceComponent, bool bImmediately)
+{
+	if (SourceComponent && CachedCameraShakeMod)
+	{
+		CachedCameraShakeMod->RemoveAllCameraShakesFromSource(SourceComponent, bImmediately);
+	}
+}
+
 void APlayerCameraManager::StopAllCameraShakes(bool bImmediately)
 {
 	if (CachedCameraShakeMod)
diff --git a/Engine/Source/Runtime/Engine/Private/PlayerController.cpp b/Engine/Source/Runtime/Engine/Private/PlayerController.cpp
--- a/Engine/Source/Runtime/Engine/Private/PlayerController.cpp
+++ b/Engine/Source/Runtime/Engine/Private/PlayerController.cpp
@@ -4337,6 +4337,22 @@ void APlayerController::ClientStopCameraShake_Implementation( TSubclassOf<class
 	}
 }
 
+void APlayerController::ClientPlayCameraShakeFromSource_Implementation(TSubclassOf<class UCameraShake> Shake, class UCameraShakeSourceComponent* SourceComponent)
+{
+	if (PlayerCameraManager != NULL)
+	{
+		PlayerCameraManager->PlayCameraShakeFromSource(Shake, SourceComponent);
+	}
+}
+
+void APlayerController::ClientStopCameraShakesFromSource_Implementation(class UCameraShakeSourceComponent* SourceComponent, bool bImmediately)
+{
+	if (PlayerCameraManager != NULL)
+	{
+		PlayerCameraManager->StopAllInstancesOfCameraShakeFromSource(SourceComponent, bImmediately);
+	}
+}
+
 void APlayerController::ClientPlayCameraAnim_Implementation( UCameraAnim* AnimToPlay, float Scale, float Rate,
 						float BlendInTime, float BlendOutTime, bool bLoop,
 						bool bRandomStartTime, ECameraAnimPlaySpace::Type Space, FRotator CustomPlaySpace )
