Make the physics interface report to the log which underlying physics engine is in use when initialising the engine.

diff --git a/Engine/Source/Runtime/Engine/Private/PhysicsEngine/PhysLevel.cpp b/Engine/Source/Runtime/Engine/Private/PhysicsEngine/PhysLevel.cpp
--- a/Engine/Source/Runtime/Engine/Private/PhysicsEngine/PhysLevel.cpp
+++ b/Engine/Source/Runtime/Engine/Private/PhysicsEngine/PhysLevel.cpp
@@ -292,6 +292,9 @@ bool InitGamePhys()
 	
 #endif // WITH_PHYSX
 
+	// Message to the log that physics is initialised and which interface we are using.
+	UE_LOG(LogInit, Log, TEXT("Physics initialised using underlying interface: %s"), *FPhysicsInterface::GetInterfaceDescription());
+
 	return true;
 }
 
diff --git a/Engine/Source/Runtime/Engine/Public/Physics/Experimental/PhysInterface_Chaos.h b/Engine/Source/Runtime/Engine/Public/Physics/Experimental/PhysInterface_Chaos.h
--- a/Engine/Source/Runtime/Engine/Public/Physics/Experimental/PhysInterface_Chaos.h
+++ b/Engine/Source/Runtime/Engine/Public/Physics/Experimental/PhysInterface_Chaos.h
@@ -177,6 +177,9 @@ public:
     FPhysInterface_Chaos(const AWorldSettings* Settings=nullptr);
     ~FPhysInterface_Chaos();
 
+	// Describe the interface to identify it to the caller
+	static FString GetInterfaceDescription() { return TEXT("Chaos"); }
+
     // Interface needed for interface
 	static void CreateActor(const FActorCreationParams& InParams, FPhysicsActorHandle& Handle);
 	static void ReleaseActor(FPhysicsActorHandle& InActorReference, FPhysScene* InScene = nullptr, bool bNeverDeferRelease=false);
diff --git a/Engine/Source/Runtime/Engine/Public/Physics/PhysicsInterfacePhysX.h b/Engine/Source/Runtime/Engine/Public/Physics/PhysicsInterfacePhysX.h
--- a/Engine/Source/Runtime/Engine/Public/Physics/PhysicsInterfacePhysX.h
+++ b/Engine/Source/Runtime/Engine/Public/Physics/PhysicsInterfacePhysX.h
@@ -184,6 +184,9 @@ struct ENGINE_API FPhysicsCommand_PhysX
 
 struct ENGINE_API FPhysicsInterface_PhysX : public FGenericPhysicsInterface
 {
+	// Describe the interface to identify it to the caller
+	static FString GetInterfaceDescription() { return TEXT("PhysX"); }
+
 	// PhysX Only functions, not related to wider physics interface
 	// To be used only in code that handles PhysX
 	static physx::PxRigidActor* GetPxRigidActorFromScene_AssumesLocked(const FPhysicsActorHandle_PhysX& InActorHandle);
