Fixing various recursion crashes with composite data tables and composite curve tables:
- Curve tables now pop up the error message under the same circumstances that data tables do
- OnParentTablesUpdated is now guarded against reentrancy (on both types - primary fix)
- Curve tables now early out when there is a loop detected due to the row map being empty (parity with data tables, potentially unnecessary now)
- During undo in the Serialize implementation, the bIsLoading flag now gets cleared instead of remaining set and the OnParentTablesUpdated call is deferred until later on in PostEditUndo

diff --git a/Engine/Source/Runtime/Engine/Classes/Engine/CompositeCurveTable.h b/Engine/Source/Runtime/Engine/Classes/Engine/CompositeCurveTable.h
--- a/Engine/Source/Runtime/Engine/Classes/Engine/CompositeCurveTable.h
+++ b/Engine/Source/Runtime/Engine/Classes/Engine/CompositeCurveTable.h
@@ -47,10 +47,13 @@ protected:
 	UPROPERTY(EditAnywhere, BlueprintReadOnly, Category = Tables)
 	TArray<UCurveTable*> ParentTables;
 
-	// true if this asset is currently being loaded; false otherwise
-	uint8 bIsLoading : 1;
-
 	// temporary copy used to detect changes so we can update delegates correctly on removal
 	UPROPERTY(transient)
 	TArray<UCurveTable*> OldParentTables;
+
+	// true if this asset is currently being loaded; false otherwise
+	uint8 bIsLoading : 1;
+
+	// true if we're already in the middle of updating parent tables for this asset
+	uint8 bUpdatingParentTables : 1;
 };
diff --git a/Engine/Source/Runtime/Engine/Classes/Engine/CompositeDataTable.h b/Engine/Source/Runtime/Engine/Classes/Engine/CompositeDataTable.h
--- a/Engine/Source/Runtime/Engine/Classes/Engine/CompositeDataTable.h
+++ b/Engine/Source/Runtime/Engine/Classes/Engine/CompositeDataTable.h
@@ -50,6 +50,7 @@ class UCompositeDataTable
 	ENGINE_API virtual void RestoreAfterStructChange() override;
 
 	virtual void PostEditChangeProperty(FPropertyChangedEvent& PropertyChangedEvent) override;
+	virtual void PostEditUndo() override;
 #endif // WITH_EDITOR
 
 	// Support for runtime modification of parent tables
@@ -75,16 +76,20 @@ protected:
 	UPROPERTY(EditAnywhere, BlueprintReadOnly, Category = Tables)
 	TArray<UDataTable*> ParentTables;
 
-	// true if this asset is currently being loaded; false otherwise
-	uint8 bIsLoading : 1;
-
-	// if this is true then the parent table array will not be cleared when EmptyTable is called
-	uint8 bShouldNotClearParentTablesOnEmpty : 1;
-
 	// temporary copy used to detect changes so we can update delegates correctly on removal
 	UPROPERTY(transient)
 	TArray<UDataTable*> OldParentTables;
+
 #if WITH_EDITORONLY_DATA
 	TMap<FName, ERowState> RowSourceMap;
 #endif // WITH_EDITORONLY_DATA
+
+	// true if this asset is currently being loaded; false otherwise
+	uint8 bIsLoading : 1;
+
+	// true if we're already in the middle of updating parent tables for this asset
+	uint8 bUpdatingParentTables : 1;
+
+	// if this is true then the parent table array will not be cleared when EmptyTable is called
+	uint8 bShouldNotClearParentTablesOnEmpty : 1;
 };
diff --git a/Engine/Source/Runtime/Engine/Private/CompositeCurveTable.cpp b/Engine/Source/Runtime/Engine/Private/CompositeCurveTable.cpp
--- a/Engine/Source/Runtime/Engine/Private/CompositeCurveTable.cpp
+++ b/Engine/Source/Runtime/Engine/Private/CompositeCurveTable.cpp
@@ -48,29 +48,31 @@ void UCompositeCurveTable::Serialize(FArchive& Ar)
 
 	if (Ar.IsLoading())
 	{
-		for (UCurveTable* ParentTable : ParentTables)
+		if (GIsTransacting)
 		{
-			if (ParentTable && ParentTable->HasAnyFlags(RF_NeedLoad))
+			bIsLoading = false;
+		}
+		else
+		{
+			for (UCurveTable* ParentTable : ParentTables)
 			{
-				FLinkerLoad* ParentTableLinker = ParentTable->GetLinker();
-				if (ParentTableLinker)
+				if (ParentTable && ParentTable->HasAnyFlags(RF_NeedLoad))
 				{
-					ParentTableLinker->Preload(ParentTable);
+					FLinkerLoad* ParentTableLinker = ParentTable->GetLinker();
+					if (ParentTableLinker)
+					{
+						ParentTableLinker->Preload(ParentTable);
+					}
 				}
 			}
-		}
 
-		OnParentTablesUpdated();
+			OnParentTablesUpdated();
+		}
 	}
 }
 
 void UCompositeCurveTable::UpdateCachedRowMap(bool bWarnOnInvalidChildren)
 {
-#if WITH_EDITOR
-	FCurveTableEditorUtils::BroadcastPreChange(this, FCurveTableEditorUtils::ECurveTableChangeInfo::RowList);
-#endif
-	UCurveTable::EmptyTable();
-
 	bool bLeaveEmpty = false;
 	// Throw up an error message and stop if any loops are found
 	if (const UCompositeCurveTable* LoopTable = FindLoops(TArray<const UCompositeCurveTable*>()))
@@ -90,8 +92,19 @@ void UCompositeCurveTable::UpdateCachedRowMap(bool bWarnOnInvalidChildren)
 			}
 		}
 		bLeaveEmpty = true;
+
+		// if the rowmap is empty, stop. We don't need to do the pre and post change since no changes will actually be done
+		if (RowMap.Num() == 0)
+		{
+			return;
+		}
 	}
 
+#if WITH_EDITOR
+	FCurveTableEditorUtils::BroadcastPreChange(this, FCurveTableEditorUtils::ECurveTableChangeInfo::RowList);
+#endif
+	UCurveTable::EmptyTable();
+
 	if (!bLeaveEmpty)
 	{
 		CurveTableMode = ECurveTableMode::SimpleCurves;
@@ -185,7 +198,7 @@ void UCompositeCurveTable::PostEditChangeProperty(FPropertyChangedEvent& Propert
 
 	if (PropertyName == Name_ParentTables)
 	{
-		OnParentTablesUpdated();
+		OnParentTablesUpdated(PropertyChangedEvent.ChangeType);
 	}
 
 	Super::PostEditChangeProperty(PropertyChangedEvent);
@@ -193,14 +206,21 @@ void UCompositeCurveTable::PostEditChangeProperty(FPropertyChangedEvent& Propert
 
 void UCompositeCurveTable::PostEditUndo()
 {
-	Super::PostEditUndo();
-
 	OnParentTablesUpdated();
+
+	Super::PostEditUndo();
 }
 #endif // WITH_EDITOR
 
 void UCompositeCurveTable::OnParentTablesUpdated(EPropertyChangeType::Type ChangeType)
 {
+	// Prevent recursion when there was a cycle in the parent hierarchy (or during the undo of the action that created the cycle; in that case PostEditUndo will recall OnParentTablesUpdated when the dust has settled)
+	if (bUpdatingParentTables)
+	{
+		return;
+	}
+	bUpdatingParentTables = true;
+
 	for (UCurveTable* Table : OldParentTables)
 	{
 		if (Table && ParentTables.Find(Table) == INDEX_NONE)
@@ -213,13 +233,15 @@ void UCompositeCurveTable::OnParentTablesUpdated(EPropertyChangeType::Type Chang
 
 	for (UCurveTable* Table : ParentTables)
 	{
-		if (Table && OldParentTables.Find(Table) == INDEX_NONE)
+		if ((Table != nullptr) && (Table != this) && (OldParentTables.Find(Table) == INDEX_NONE))
 		{
 			Table->OnCurveTableChanged().AddUObject(this, &UCompositeCurveTable::OnParentTablesUpdated, EPropertyChangeType::Unspecified);
 		}
 	}
 
 	OldParentTables = ParentTables;
+
+	bUpdatingParentTables = false;
 }
 
 void UCompositeCurveTable::EmptyTable()
diff --git a/Engine/Source/Runtime/Engine/Private/CompositeDataTable.cpp b/Engine/Source/Runtime/Engine/Private/CompositeDataTable.cpp
--- a/Engine/Source/Runtime/Engine/Private/CompositeDataTable.cpp
+++ b/Engine/Source/Runtime/Engine/Private/CompositeDataTable.cpp
@@ -209,19 +209,26 @@ void UCompositeDataTable::Serialize(FArchive& Ar)
 
 	if (Ar.IsLoading())
 	{
-		for (UDataTable* ParentTable : ParentTables)
+		if (GIsTransacting)
 		{
-			if (ParentTable && ParentTable->HasAnyFlags(RF_NeedLoad))
+			bIsLoading = false;
+		}
+		else
+		{
+			for (UDataTable* ParentTable : ParentTables)
 			{
-				FLinkerLoad* ParentTableLinker = ParentTable->GetLinker();
-				if (ParentTableLinker)
+				if (ParentTable && ParentTable->HasAnyFlags(RF_NeedLoad))
 				{
-					ParentTableLinker->Preload(ParentTable);
+					FLinkerLoad* ParentTableLinker = ParentTable->GetLinker();
+					if (ParentTableLinker)
+					{
+						ParentTableLinker->Preload(ParentTable);
+					}
 				}
 			}
-		}
 
-		OnParentTablesUpdated();
+			OnParentTablesUpdated();
+		}
 	}
 }
 
@@ -256,6 +263,14 @@ void UCompositeDataTable::PostEditChangeProperty(FPropertyChangedEvent& Property
 
 	Super::PostEditChangeProperty(PropertyChangedEvent);
 }
+
+void UCompositeDataTable::PostEditUndo()
+{
+	OnParentTablesUpdated(EPropertyChangeType::ValueSet);
+
+	Super::PostEditUndo();
+}
+
 #endif // WITH_EDITOR
 
 void UCompositeDataTable::AppendParentTables(const TArray<UDataTable*>& NewTables)
@@ -266,6 +281,13 @@ void UCompositeDataTable::AppendParentTables(const TArray<UDataTable*>& NewTable
 
 void UCompositeDataTable::OnParentTablesUpdated(EPropertyChangeType::Type ChangeType)
 {
+	// Prevent recursion when there was a cycle in the parent hierarchy (or during the undo of the action that created the cycle; in that case PostEditUndo will recall OnParentTablesUpdated when the dust has settled)
+	if (bUpdatingParentTables)
+	{
+		return;
+	}
+	bUpdatingParentTables = true;
+
 	for (UDataTable* Table : OldParentTables)
 	{
 		if (Table && ParentTables.Find(Table) == INDEX_NONE)
@@ -278,13 +300,15 @@ void UCompositeDataTable::OnParentTablesUpdated(EPropertyChangeType::Type Change
 
 	for (UDataTable* Table : ParentTables)
 	{
-		if (Table && OldParentTables.Find(Table) == INDEX_NONE)
+		if ((Table != nullptr) && (Table != this) && (OldParentTables.Find(Table) == INDEX_NONE))
 		{
 			Table->OnDataTableChanged().AddUObject(this, &UCompositeDataTable::OnParentTablesUpdated, EPropertyChangeType::Unspecified);
 		}
 	}
 
 	OldParentTables = ParentTables;
+
+	bUpdatingParentTables = false;
 }
 
 #undef LOCTEXT_NAMESPACE
