Added additional logging to track down race condition in vertex factory.

diff --git a/Engine/Source/Runtime/Engine/Private/SceneManagement.cpp b/Engine/Source/Runtime/Engine/Private/SceneManagement.cpp
--- a/Engine/Source/Runtime/Engine/Private/SceneManagement.cpp
+++ b/Engine/Source/Runtime/Engine/Private/SceneManagement.cpp
@@ -1096,8 +1096,17 @@ void FMeshBatch::PreparePrimitiveUniformBuffer(const FPrimitiveSceneProxy* Primi
 			MeshElement.PrimitiveUniformBuffer = PrimitiveSceneProxy->GetUniformBuffer();
 		}
 
-		checkf(bPrimitiveShaderDataComesFromSceneBuffer || Elements[ElementIndex].PrimitiveUniformBuffer || Elements[ElementIndex].PrimitiveUniformBufferResource != NULL,
-			TEXT("FMeshBatch was not properly setup.  The primitive uniform buffer must be specified."));
+		const bool bValidPrimitiveData = bPrimitiveShaderDataComesFromSceneBuffer || Elements[ElementIndex].PrimitiveUniformBuffer || Elements[ElementIndex].PrimitiveUniformBufferResource;
+
+		UE_CLOG(!bValidPrimitiveData, LogEngine, Fatal,
+			TEXT("FMeshBatch was not properly setup. No primitive uniform buffer was specified and the vertex factory does not have a valid primitive id stream.\n")
+			TEXT("\tVertexFactory[Name: %s, Initialized: %u]\n")
+			TEXT("\tPrimitiveSceneProxy[Level: %s, Owner: %s, Resource: %s]"),
+			*VertexFactory->GetType()->GetFName().ToString(),
+			VertexFactory->IsInitialized() ? 1 : 0,
+			*PrimitiveSceneProxy->GetLevelName().ToString(),
+			*PrimitiveSceneProxy->GetOwnerName().ToString(),
+			*PrimitiveSceneProxy->GetResourceName().ToString());
 	}
 }
 
