Chaos - reenable collisions for RBAN following persistent manifold fix

diff --git a/Engine/Source/Runtime/Engine/Private/PhysicsEngine/ImmediatePhysics/ImmediatePhysicsChaos/ImmediatePhysicsSimulation_Chaos.cpp b/Engine/Source/Runtime/Engine/Private/PhysicsEngine/ImmediatePhysics/ImmediatePhysicsChaos/ImmediatePhysicsSimulation_Chaos.cpp
--- a/Engine/Source/Runtime/Engine/Private/PhysicsEngine/ImmediatePhysics/ImmediatePhysicsChaos/ImmediatePhysicsSimulation_Chaos.cpp
+++ b/Engine/Source/Runtime/Engine/Private/PhysicsEngine/ImmediatePhysics/ImmediatePhysicsChaos/ImmediatePhysicsSimulation_Chaos.cpp
@@ -27,7 +27,7 @@ FAutoConsoleVariableRef CVarChaosImmPhysDeltaTime(TEXT("p.Chaos.ImmPhys.DeltaTim
 FAutoConsoleVariableRef CVarChaosImmPhysIterations(TEXT("p.Chaos.ImmPhys.Iterations"), ChaosImmediate_Evolution_Iterations, TEXT("Number of constraint solver loops in immediate physics"));
 FAutoConsoleVariableRef CVarChaosImmPhysPushOutIterations(TEXT("p.Chaos.ImmPhys.PushOutIterations"), ChaosImmediate_Evolution_PushOutIterations, TEXT("Set the ApplyPushOut() (position correction) iteration count"));
 
-int32 ChaosImmediate_Collision_Enabled = 0;
+int32 ChaosImmediate_Collision_Enabled = 1;
 int32 ChaosImmediate_Collision_ApplyEnabled = 1;
 int32 ChaosImmediate_Collision_PushOutPairIterations = 2;
 int32 ChaosImmediate_Collision_Priority = 1;
