Guarded access of GIsTransacting to fix non-editor builds

diff --git a/Engine/Source/Runtime/Engine/Private/CompositeCurveTable.cpp b/Engine/Source/Runtime/Engine/Private/CompositeCurveTable.cpp
--- a/Engine/Source/Runtime/Engine/Private/CompositeCurveTable.cpp
+++ b/Engine/Source/Runtime/Engine/Private/CompositeCurveTable.cpp
@@ -46,28 +46,28 @@ void UCompositeCurveTable::Serialize(FArchive& Ar)
 
 	Super::Serialize(Ar);
 
-	if (Ar.IsLoading())
+#if WITH_EDITORONLY_DATA
+	if (Ar.IsLoading() && GIsTransacting)
 	{
-		if (GIsTransacting)
-		{
-			bIsLoading = false;
-		}
-		else
+		bIsLoading = false;
+	}
+#endif
+
+	if (bIsLoading)
+	{
+		for (UCurveTable* ParentTable : ParentTables)
 		{
-			for (UCurveTable* ParentTable : ParentTables)
+			if (ParentTable && ParentTable->HasAnyFlags(RF_NeedLoad))
 			{
-				if (ParentTable && ParentTable->HasAnyFlags(RF_NeedLoad))
+				FLinkerLoad* ParentTableLinker = ParentTable->GetLinker();
+				if (ParentTableLinker)
 				{
-					FLinkerLoad* ParentTableLinker = ParentTable->GetLinker();
-					if (ParentTableLinker)
-					{
-						ParentTableLinker->Preload(ParentTable);
-					}
+					ParentTableLinker->Preload(ParentTable);
 				}
 			}
-
-			OnParentTablesUpdated();
 		}
+
+		OnParentTablesUpdated();
 	}
 }
 
diff --git a/Engine/Source/Runtime/Engine/Private/CompositeDataTable.cpp b/Engine/Source/Runtime/Engine/Private/CompositeDataTable.cpp
--- a/Engine/Source/Runtime/Engine/Private/CompositeDataTable.cpp
+++ b/Engine/Source/Runtime/Engine/Private/CompositeDataTable.cpp
@@ -207,28 +207,28 @@ void UCompositeDataTable::Serialize(FArchive& Ar)
 
 	Super::Serialize(Ar); // When loading, this should load our RowStruct!	
 
-	if (Ar.IsLoading())
+#if WITH_EDITORONLY_DATA
+	if (Ar.IsLoading() && GIsTransacting)
 	{
-		if (GIsTransacting)
-		{
-			bIsLoading = false;
-		}
-		else
+		bIsLoading = false;
+	}
+#endif
+
+	if (bIsLoading)
+	{
+		for (UDataTable* ParentTable : ParentTables)
 		{
-			for (UDataTable* ParentTable : ParentTables)
+			if (ParentTable && ParentTable->HasAnyFlags(RF_NeedLoad))
 			{
-				if (ParentTable && ParentTable->HasAnyFlags(RF_NeedLoad))
+				FLinkerLoad* ParentTableLinker = ParentTable->GetLinker();
+				if (ParentTableLinker)
 				{
-					FLinkerLoad* ParentTableLinker = ParentTable->GetLinker();
-					if (ParentTableLinker)
-					{
-						ParentTableLinker->Preload(ParentTable);
-					}
+					ParentTableLinker->Preload(ParentTable);
 				}
 			}
-
-			OnParentTablesUpdated();
 		}
+
+		OnParentTablesUpdated();
 	}
 }
 
