Add CVar explicitly setting the stream cache size (which will update next time we try to add or get a chunk of audio), as well as a console command to explicitly shrink the cache immediately.

diff --git a/Engine/Source/Runtime/Engine/Private/AudioStreamingCache.cpp b/Engine/Source/Runtime/Engine/Private/AudioStreamingCache.cpp
--- a/Engine/Source/Runtime/Engine/Private/AudioStreamingCache.cpp
+++ b/Engine/Source/Runtime/Engine/Private/AudioStreamingCache.cpp
@@ -52,6 +52,14 @@ FAutoConsoleVariableRef CVarReadRequestPriority(
 	TEXT("0: High, 1: Normal, 2: Below Normal, 3: Low, 4: Min"),
 	ECVF_Default);
 
+static float StreamCacheSizeOverrideMBCVar = 0.0f;
+FAutoConsoleVariableRef CVarStreamCacheSizeOverrideMB(
+	TEXT("au.streamcaching.StreamCacheSizeOverrideMB"),
+	StreamCacheSizeOverrideMBCVar,
+	TEXT("This cvar can be set to override the size of the cache.\n")
+	TEXT("0: use cache size from project settings. n: the new cache size in megabytes."),
+	ECVF_Default);
+
 static FAutoConsoleCommand GFlushAudioCacheCommand(
 	TEXT("au.streamcaching.FlushAudioCache"),
 	TEXT("This will flush any non retained audio from the cache when Stream Caching is enabled."),
@@ -65,6 +73,45 @@ static FAutoConsoleCommand GFlushAudioCacheCommand(
 	})
 );
 
+static FAutoConsoleCommand GResizeAudioCacheCommand(
+	TEXT("au.streamcaching.ResizeAudioCacheTo"),
+	TEXT("This will try to cull enough audio chunks to shrink the audio stream cache to the new size if neccessary, and keep the cache at that size."),
+	FConsoleCommandWithArgsDelegate::CreateStatic(
+		[](const TArray< FString >& Args)
+{
+	if (Args.Num() < 1)
+	{
+		return;
+	}
+
+	const float InMB = FCString::Atof(*Args[0]);
+
+	if (InMB <= 0.0f)
+	{
+		return;
+	}
+
+	static IConsoleVariable* StreamCacheSizeCVar = IConsoleManager::Get().FindConsoleVariable(TEXT("au.streamcaching.StreamCacheSizeOverrideMB"));
+	check(StreamCacheSizeCVar);
+
+	uint64 NewCacheSizeInBytes = ((uint64)(InMB * 1024)) * 1024;
+	uint64 OldCacheSizeInBytes = ((uint64)(StreamCacheSizeCVar->GetFloat() * 1024)) * 1024;
+
+	// TODO: here we delete the difference between the old cache size and the new cache size,
+	// but we don't actually need to do this unless the cache is full.
+	// In the future we can use our current cache usage to figure out how much we need to trim.
+	if (NewCacheSizeInBytes < OldCacheSizeInBytes)
+	{
+		uint64 NumBytesToFree = OldCacheSizeInBytes - NewCacheSizeInBytes;
+		IStreamingManager::Get().GetAudioStreamingManager().TrimMemory(NumBytesToFree);
+	}
+
+	StreamCacheSizeCVar->Set(InMB);
+
+	UE_LOG(LogAudio, Display, TEXT("Audio Cache Shrunk! Now set to be %f MB."), InMB);
+})
+);
+
 FCachedAudioStreamingManager::FCachedAudioStreamingManager(const FCachedAudioStreamingManagerParams& InitParams)
 {
 	check(FPlatformCompressionUtilities::IsCurrentPlatformUsingStreamCaching());
@@ -394,6 +441,12 @@ FAudioChunkCache::~FAudioChunkCache()
 
 bool FAudioChunkCache::AddOrTouchChunk(const FChunkKey& InKey, TFunction<void(EAudioChunkLoadResult)> OnLoadCompleted, ENamedThreads::Type CallbackThread)
 {
+	// Update cache limit if needed.
+	if (!FMath::IsNearlyZero(StreamCacheSizeOverrideMBCVar) && StreamCacheSizeOverrideMBCVar > 0.0f)
+	{
+		MemoryLimitBytes = ((uint64)(StreamCacheSizeOverrideMBCVar * 1024)) * 1024;
+	}
+	
 	if (!IsKeyValid(InKey))
 	{
 		ensure(false);
@@ -1180,7 +1233,7 @@ TPair<int, int> FAudioChunkCache::DebugDisplay(UWorld* World, FViewport* Viewpor
 
 #if WITH_EDITOR
 		// TODO: Worry about whether the sound wave is alive here. In most editor cases this is ok because the soundwave will always be loaded, but this may not be the case in the future.
-		bIsStaleChunk = (CurrentElement->Key.SoundWave->CurrentChunkRevision.GetValue() != CurrentElement->Key.ChunkRevision);
+		bIsStaleChunk = (CurrentElement->Key.SoundWave == nullptr) || (CurrentElement->Key.SoundWave->CurrentChunkRevision.GetValue() != CurrentElement->Key.ChunkRevision);
 #endif
 
 		const bool bWasTrimmed = CurrentElement->ChunkData.Num() == 0;
