Fix crash when instanced stereo rendering is enabled
Fix splitscreen targets only showing the last view

diff --git a/Engine/Source/Runtime/Engine/Private/SceneView.cpp b/Engine/Source/Runtime/Engine/Private/SceneView.cpp
--- a/Engine/Source/Runtime/Engine/Private/SceneView.cpp
+++ b/Engine/Source/Runtime/Engine/Private/SceneView.cpp
@@ -2103,7 +2103,7 @@ EShaderPlatform FSceneView::GetShaderPlatform() const
 
 bool FSceneView::IsInstancedStereoPass() const
 {
-	return bIsInstancedStereoEnabled && IStereoRendering::IsAPrimaryView(*this);
+	return bIsInstancedStereoEnabled && IStereoRendering::IsStereoEyeView(*this) && IStereoRendering::IsAPrimaryView(*this);
 }
 
 void FSceneView::SetupViewRectUniformBufferParameters(FViewUniformShaderParameters& ViewUniformShaderParameters,
diff --git a/Engine/Source/Runtime/Engine/Private/StereoRendering.cpp b/Engine/Source/Runtime/Engine/Private/StereoRendering.cpp
--- a/Engine/Source/Runtime/Engine/Private/StereoRendering.cpp
+++ b/Engine/Source/Runtime/Engine/Private/StereoRendering.cpp
@@ -88,7 +88,7 @@ bool IStereoRendering::IsAnAdditionalPass(EStereoscopicPass Pass)
 		return GEngine->StereoRenderingDevice->DeviceIsAnAdditionalPass(Pass);
 	}
 
-	return Pass > eSSP_RIGHT_EYE;
+	return Pass > EStereoscopicPass::eSSP_RIGHT_EYE;
 }
 
 bool IStereoRendering::IsAnAdditionalView(const FSceneView& View)
@@ -98,7 +98,7 @@ bool IStereoRendering::IsAnAdditionalView(const FSceneView& View)
 		return GEngine->StereoRenderingDevice->DeviceIsAnAdditionalView(View);
 	}
 
-	return View.StereoPass > eSSP_RIGHT_EYE;
+	return View.StereoPass > EStereoscopicPass::eSSP_RIGHT_EYE;
 }
 
 bool IStereoRendering::DeviceIsAnAdditionalView(const FSceneView& View)
diff --git a/Engine/Source/Runtime/Renderer/Private/BasePassRendering.cpp b/Engine/Source/Runtime/Renderer/Private/BasePassRendering.cpp
--- a/Engine/Source/Runtime/Renderer/Private/BasePassRendering.cpp
+++ b/Engine/Source/Runtime/Renderer/Private/BasePassRendering.cpp
@@ -483,7 +483,7 @@ void SetupSharedBasePassParameters(
 
 	SetupFogUniformParameters(View, SharedParameters.Fog);
 
-	if (View.bIsInstancedStereoEnabled && IStereoRendering::IsAPrimaryView(View))
+	if (View.IsInstancedStereoPass())
 	{
 		const FSceneView& RightEye = *View.Family->Views[1];
 		SharedParameters.ForwardISR = RightEye.ForwardLightingResources->ForwardLightData;
diff --git a/Engine/Source/Runtime/Renderer/Private/PostProcess/PostProcessTonemap.cpp b/Engine/Source/Runtime/Renderer/Private/PostProcess/PostProcessTonemap.cpp
--- a/Engine/Source/Runtime/Renderer/Private/PostProcess/PostProcessTonemap.cpp
+++ b/Engine/Source/Runtime/Renderer/Private/PostProcess/PostProcessTonemap.cpp
@@ -1484,7 +1484,7 @@ void FRCPassPostProcessTonemapES2::Process(FRenderingCompositePassContext& Conte
 		const bool bFirstView = (&View == View.Family->Views[0]);
 		
 		// Full clear to avoid restore
-		if ((View.StereoPass == eSSP_FULL && bFirstView) || IStereoRendering::IsAPrimaryView(View))
+		if (IStereoRendering::IsAPrimaryView(View) && (bFirstView || IStereoRendering::IsStereoEyeView(View)))
 		{
 			LoadAction = ERenderTargetLoadAction::EClear;
 		}
