Fix crash on re-use of an anim instance using linked anim instances

When a re-init is not forced and the main instance is re-used (as in the case of a re-registration), the LinkedAnimInstances array is emptied but references are still held on the various linked nodes. This invalidates those links by marking the linked instances pending kill so linked nodes dont propogate to trash instances (or in this case, instances with different required bones).

diff --git a/Engine/Source/Runtime/Engine/Classes/Components/SkeletalMeshComponent.h b/Engine/Source/Runtime/Engine/Classes/Components/SkeletalMeshComponent.h
--- a/Engine/Source/Runtime/Engine/Classes/Components/SkeletalMeshComponent.h
+++ b/Engine/Source/Runtime/Engine/Classes/Components/SkeletalMeshComponent.h
@@ -795,6 +795,9 @@ PRAGMA_DISABLE_DEPRECATION_WARNINGS
 	TArray<UAnimInstance*>& GetLinkedAnimInstances() { return LinkedInstances; }
 PRAGMA_ENABLE_DEPRECATION_WARNINGS
 
+	/** Clear the linked anim instances and mark them pending kill */
+	void ResetLinkedAnimInstances();
+
 public:
 	UE_DEPRECATED(4.23, "This function is deprecated. Please use GetLinkedAnimGraphInstanceByTag")
 	UAnimInstance* GetSubInstanceByName(FName InTag) const { return GetLinkedAnimGraphInstanceByTag(InTag); }
diff --git a/Engine/Source/Runtime/Engine/Private/Components/SkeletalMeshComponent.cpp b/Engine/Source/Runtime/Engine/Private/Components/SkeletalMeshComponent.cpp
--- a/Engine/Source/Runtime/Engine/Private/Components/SkeletalMeshComponent.cpp
+++ b/Engine/Source/Runtime/Engine/Private/Components/SkeletalMeshComponent.cpp
@@ -515,7 +515,7 @@ void USkeletalMeshComponent::OnRegister()
 
 	// Ensure we have an empty list of linked instances on registration. Ready for the initialization below 
 	// to correctly populate that list.
-	LinkedInstances.Reset();
+	ResetLinkedAnimInstances();
 
 	// We force an initialization here because we're in one of two cases.
 	// 1) First register, no spawned instance, need to initialize
@@ -571,7 +571,7 @@ void USkeletalMeshComponent::OnUnregister()
 	{
 		LinkedInstance->UninitializeAnimation();
 	}
-	LinkedInstances.Reset();
+	ResetLinkedAnimInstances();
 
 	if(PostProcessAnimInstance)
 	{
@@ -703,7 +703,7 @@ bool USkeletalMeshComponent::InitializeAnimScriptInstance(bool bForceReinit, boo
 			if (AnimScriptInstance)
 			{
 				// If we have any linked instances left we need to clear them out now, we're about to have a new master instance
-				LinkedInstances.Empty();
+				ResetLinkedAnimInstances();
 
 				AnimScriptInstance->InitializeAnimation(bInDeferRootNodeInitialization);
 				bInitializedMainInstance = true;
@@ -820,7 +820,7 @@ void USkeletalMeshComponent::ClearAnimScriptInstance()
 		AnimScriptInstance->EndNotifyStates();
 	}
 	AnimScriptInstance = nullptr;
-	LinkedInstances.Empty();
+	ResetLinkedAnimInstances();
 }
 
 
@@ -2722,6 +2722,19 @@ UAnimInstance* USkeletalMeshComponent::GetPostProcessInstance() const
 	return PostProcessAnimInstance;
 }
 
+void USkeletalMeshComponent::ResetLinkedAnimInstances()
+{
+	for(UAnimInstance* LinkedInstance : LinkedInstances)
+	{
+		if(LinkedInstance)
+		{
+			LinkedInstance->MarkPendingKill();
+			LinkedInstance = nullptr;
+		}
+	}
+	LinkedInstances.Reset();
+}
+
 UAnimInstance* USkeletalMeshComponent::GetLinkedAnimGraphInstanceByTag(FName InName) const
 {
 	if(AnimScriptInstance)
