ImportBufferAsTexture2D crashes when image bit depth is not 8 or 16. Fixed a badly formatted log string and returning nullptr on failure.

diff --git a/Engine/Source/Runtime/Engine/Private/ImageUtils.cpp b/Engine/Source/Runtime/Engine/Private/ImageUtils.cpp
--- a/Engine/Source/Runtime/Engine/Private/ImageUtils.cpp
+++ b/Engine/Source/Runtime/Engine/Private/ImageUtils.cpp
@@ -850,7 +850,8 @@ UTexture2D* FImageUtils::ImportBufferAsTexture2D(const TArray<uint8>& Buffer)
 			}
 			else
 			{
-				UE_LOG(LogImageUtils, Warning, TEXT("Error creating texture. %s is not a valid bit depth (%d)"), BitDepth);
+				UE_LOG(LogImageUtils, Warning, TEXT("Error creating texture. Bit depth is unsupported. (%d)"), BitDepth);
+				return nullptr;
 			}
 			
 			const TArray<uint8>* UncompressedData = nullptr;
