Merge all fixes for Stream Caching to 4.24.

diff --git a/Engine/Source/Runtime/Engine/Classes/Sound/SoundClass.h b/Engine/Source/Runtime/Engine/Classes/Sound/SoundClass.h
--- a/Engine/Source/Runtime/Engine/Classes/Sound/SoundClass.h
+++ b/Engine/Source/Runtime/Engine/Classes/Sound/SoundClass.h
@@ -6,7 +6,6 @@
 #include "UObject/ObjectMacros.h"
 #include "UObject/Object.h"
 #include "AudioDefines.h"
-#include "AudioDynamicParameter.h"
 #if WITH_EDITOR
 #include "EdGraph/EdGraph.h"
 #endif
@@ -88,10 +87,6 @@ struct FSoundClassProperties
 	UPROPERTY(EditAnywhere, BlueprintReadOnly, Category=SoundClassProperties)
 	float Pitch;
 
-	/** Distance scale to apply to sounds that play with this sound class. Sounds will have their attenuation distance scaled by this amount. Allows adjusting attenuation settings dynamically. **/
-	UPROPERTY(EditAnywhere, BlueprintReadOnly, Category = SoundClassProperties)
-	float AttenuationDistanceScale;
-	
 	/** Lowpass filter frequency */
 	UPROPERTY(EditAnywhere, BlueprintReadOnly, Category = SoundClassProperties)
 	float LowPassFilterFrequency;
@@ -156,22 +151,9 @@ struct FSoundClassProperties
 	UPROPERTY(EditAnywhere, BlueprintReadOnly, Category = SoundClassProperties, meta = (DisplayName = "Loading Behavior Override"))
 	ESoundWaveLoadingBehavior LoadingBehavior;
 
-	// Sets the attenuation scale of the sound class in the given amount of time
-	void SetAttenuationDistanceScale(float InAttenuationDistanceScale, float InTime);
-
-	// Sets the parent's attenuation scale
-	void SetParentAttenuationDistanceScale(float InAttenuationDistanceScale);
-
-	// Returns the current attenuation scale
-	float GetAttenuationDistanceScale() const;
-
-	// Updates any dynamic sound class properties
-	void UpdateSoundClassProperties(float DeltaTime);
-
 	FSoundClassProperties()
 		: Volume(1.0f)
 		, Pitch(1.0f)
-		, AttenuationDistanceScale(1.0f)
 		, LowPassFilterFrequency(MAX_FILTER_FREQUENCY)
 		, StereoBleed(0.25f)
 		, LFEBleed(0.5f)
@@ -188,15 +170,9 @@ struct FSoundClassProperties
 		, bApplyAmbientVolumes(false)
 		, OutputTarget(EAudioOutputTarget::Speaker)
 		, LoadingBehavior(ESoundWaveLoadingBehavior::Inherited)
-		, AttenuationScaleParam(AttenuationDistanceScale)
-		, ParentAttenuationScale(1.0f)
 		{
 		}
 	
-private:
-	FDynamicParameter AttenuationScaleParam;
-	float ParentAttenuationScale;
-
 };
 
 /**
diff --git a/Engine/Source/Runtime/Engine/Private/SoundClass.cpp b/Engine/Source/Runtime/Engine/Private/SoundClass.cpp
--- a/Engine/Source/Runtime/Engine/Private/SoundClass.cpp
+++ b/Engine/Source/Runtime/Engine/Private/SoundClass.cpp
@@ -17,26 +17,6 @@
 #include "Widgets/Notifications/SNotificationList.h"
 #endif
 
-void FSoundClassProperties::SetAttenuationDistanceScale(float InAttenuationScale, float InTime)
-{
-	AttenuationScaleParam.Set(AttenuationDistanceScale, InTime);
-}
-
-void FSoundClassProperties::SetParentAttenuationDistanceScale(float InAttenuationDistanceScale)
-{
-	ParentAttenuationScale = InAttenuationDistanceScale;
-}
-
-float FSoundClassProperties::GetAttenuationDistanceScale() const
-{
-	return AttenuationScaleParam.GetValue() * ParentAttenuationScale;
-}
-
-void FSoundClassProperties::UpdateSoundClassProperties(float DeltaTime)
-{
-	AttenuationScaleParam.Update(DeltaTime);
-}
-
 /*-----------------------------------------------------------------------------
 	USoundClass implementation.
 -----------------------------------------------------------------------------*/
diff --git a/Engine/Source/Runtime/Engine/Private/SoundWave.cpp b/Engine/Source/Runtime/Engine/Private/SoundWave.cpp
--- a/Engine/Source/Runtime/Engine/Private/SoundWave.cpp
+++ b/Engine/Source/Runtime/Engine/Private/SoundWave.cpp
@@ -1722,8 +1722,6 @@ void USoundWave::Parse(FAudioDevice* AudioDevice, const UPTRINT NodeWaveInstance
 		WaveInstance->SetVolumeMultiplier(VolumeMultiplier* SoundClassProperties->Volume);
 		WaveInstance->SetPitch(WaveInstance->Pitch * SoundClassProperties->Pitch);
 
-		WaveInstance->AttenuationDistance *= FMath::Max(SoundClassProperties->GetAttenuationDistanceScale(), 0.0f);
-
 		WaveInstance->SoundClassFilterFrequency = SoundClassProperties->LowPassFilterFrequency;
 		WaveInstance->VoiceCenterChannelVolume = SoundClassProperties->VoiceCenterChannelVolume;
 		WaveInstance->RadioFilterVolume = SoundClassProperties->RadioFilterVolume * ParseParams.VolumeMultiplier;
