Fixed crash on uninitialized rendering resources when undo/redoing material changes on static meshes.

diff --git a/Engine/Source/Runtime/Engine/Public/StaticMeshResources.h b/Engine/Source/Runtime/Engine/Public/StaticMeshResources.h
--- a/Engine/Source/Runtime/Engine/Public/StaticMeshResources.h
+++ b/Engine/Source/Runtime/Engine/Public/StaticMeshResources.h
@@ -689,12 +689,6 @@ public:
 
 		for (TObjectIterator<UStaticMeshComponent> It; It; ++It)
 		{
-			// First, flush all deferred render updates, as they may depend on InStaticMesh.
-			if (It->IsRenderStateDirty() && It->IsRegistered() && !It->IsTemplate() && !It->IsPendingKill())
-			{
-				It->DoDeferredRenderUpdates_Concurrent();
-			}
-
 			UStaticMesh* StaticMesh = It->GetStaticMesh();
 
 			if (StaticMeshComponents.Contains(StaticMesh))
@@ -709,6 +703,11 @@ public:
 					Scenes.Add(It->GetScene());
 				}
 			}
+			// Recreate dirty render state, if needed, only for components not using the static mesh we currently have released resources for.
+			else if (It->IsRenderStateDirty() && It->IsRegistered() && !It->IsTemplate() && !It->IsPendingKill())
+			{
+				It->DoDeferredRenderUpdates_Concurrent();
+			}
 		}
 
 		UpdateAllPrimitiveSceneInfosForScenes(MoveTemp(Scenes));
