Exposed FMatrix functionality to Blueprints and provided an Autocast node to support easy printing.

diff --git a/Engine/Source/Runtime/Core/Public/Math/Matrix.h b/Engine/Source/Runtime/Core/Public/Math/Matrix.h
--- a/Engine/Source/Runtime/Core/Public/Math/Matrix.h
+++ b/Engine/Source/Runtime/Core/Public/Math/Matrix.h
@@ -274,6 +274,14 @@ public:
 	 */
 	inline FVector GetColumn(int32 i) const;
 
+	/**
+	 * Set a column of this matrix
+	 *
+	 * @param i index of the matrix column
+	 * @param Value new value of the column
+	 */
+	inline void SetColumn(int32 i, FVector Value);
+
 	/** @return rotator representation of this matrix */
 	CORE_API FRotator Rotator() const;
 
diff --git a/Engine/Source/Runtime/Core/Public/Math/Matrix.inl b/Engine/Source/Runtime/Core/Public/Math/Matrix.inl
--- a/Engine/Source/Runtime/Core/Public/Math/Matrix.inl
+++ b/Engine/Source/Runtime/Core/Public/Math/Matrix.inl
@@ -615,6 +615,14 @@ inline FVector FMatrix::GetColumn(int32 i) const
 	return FVector(M[0][i], M[1][i], M[2][i]);
 }
 
+inline void FMatrix::SetColumn(int32 i, FVector Value)
+{
+	checkSlow(i >= 0 && i <= 3);
+	M[0][i] = Value.X;
+	M[1][i] = Value.Y;
+	M[2][i] = Value.Z;
+}
+
 FORCEINLINE bool MakeFrustumPlane(float A,float B,float C,float D,FPlane& OutPlane)
 {
 	const float	LengthSquared = A * A + B * B + C * C;
diff --git a/Engine/Source/Runtime/Engine/Classes/Kismet/KismetMathLibrary.h b/Engine/Source/Runtime/Engine/Classes/Kismet/KismetMathLibrary.h
--- a/Engine/Source/Runtime/Engine/Classes/Kismet/KismetMathLibrary.h
+++ b/Engine/Source/Runtime/Engine/Classes/Kismet/KismetMathLibrary.h
@@ -86,6 +86,26 @@ namespace ELerpInterpolationMode
 	};
 }
 
+/** Possible columns for an FMatrix */
+UENUM(BlueprintType)
+namespace EMatrixColumns
+{
+	enum Type
+	{
+		/** First Column. */
+		First,
+
+		/** Second Column. */
+		Second,
+
+		/** Third Column. */
+		Third,
+
+		/** Fourth Column. */
+		Fourth
+	};
+}
+
 USTRUCT(BlueprintType)
 struct ENGINE_API FFloatSpringState
 {
@@ -2200,6 +2220,239 @@ class ENGINE_API UKismetMathLibrary : public UBlueprintFunctionLibrary
 	UFUNCTION(BlueprintPure, meta = (DisplayName = "Get Origin (Matrix)", ScriptMethod = "GetOrigin"), Category = "Math|Matrix")
 	static FVector Matrix_GetOrigin(const FMatrix& InMatrix);
 
+	// Identity matrix
+	UFUNCTION(BlueprintPure, meta = (DisplayName = "Identity (Matrix)", ScriptConstant = "Identity", ScriptConstantHost = "Matrix"), Category = "Math|Matrix")
+	static FMatrix Matrix_Identity();
+
+	/**
+	 * Gets the result of multiplying a Matrix to this.
+	 *
+	 * @param Other The matrix to multiply this by.
+	 * @return The result of multiplication.
+	 */
+	UFUNCTION(BlueprintPure, meta = (DisplayName = "Matrix * Matrix", CompactNodeTitle = "*", ScriptMethod = "Multiply", ScriptOperator = "*;*=", Keywords = "* multiply"), Category = "Math|Matrix")
+	static FMatrix Multiply_MatrixMatrix (const FMatrix& A, const FMatrix& B);
+
+	/**
+	 * Gets the result of adding a matrix to this.
+	 *
+	 * @param Other The Matrix to add.
+	 * @return The result of addition.
+	 */
+	UFUNCTION(BlueprintPure, meta = (DisplayName = "Matrix + Matrix", CompactNodeTitle = "+", ScriptMethod = "Add", ScriptOperator = "+;+=", Keywords = "+ add plus", CommutativeAssociativeBinaryOperator = "true"), Category = "Math|Matrix")
+	static FMatrix Add_MatrixMatrix (const FMatrix& A, const FMatrix& B);
+
+	/**
+	  * This isn't applying SCALE, just multiplying the value to all members - i.e. weighting
+	  */
+	UFUNCTION(BlueprintPure, meta = (DisplayName = "Matrix * Float", CompactNodeTitle = "*", ScriptMethod = "Multiply", ScriptOperator = "*;*=", Keywords = "* multiply"), Category = "Math|Matrix")
+	static FMatrix Multiply_MatrixFloat (const FMatrix& A, float B);
+
+	/**
+	 * Checks whether another Matrix is equal to this, within specified tolerance.
+	 *
+	 * @param Other The other Matrix.
+	 * @param Tolerance Error Tolerance.
+	 * @return true if two Matrix are equal, within specified tolerance, otherwise false.
+	 */
+	UFUNCTION(BlueprintPure, meta = (DisplayName = "Equal (Matrix)", CompactNodeTitle = "==", ScriptMethod = "Equals", ScriptOperator = "==", Keywords = "== equal"), Category = "Math|Matrix")
+	static bool EqualEqual_MatrixMatrix(const FMatrix& A, const FMatrix& B, float Tolerance = 1.e-4f);
+
+	/**
+	 * Checks whether another Matrix is not equal to this, within specified tolerance.
+	 *
+	 * @param Other The other Matrix.
+	 * @return true if two Matrix are not equal, within specified tolerance, otherwise false.
+	 */
+	UFUNCTION(BlueprintPure, meta = (DisplayName = "Not Equal (Matrix)", CompactNodeTitle = "!=", ScriptMethod = "NotEqual", ScriptOperator = "!=", Keywords = "!= not equal"), Category = "Math|Matrix")
+	static bool NotEqual_MatrixMatrix(const FMatrix& A, const FMatrix& B, float Tolerance = 1.e-4f);
+
+	// Homogeneous transform.
+	UFUNCTION(BlueprintPure, meta = (DisplayName = "Transform Vector4 (Matrix)", ScriptMethod = "TransformVector4"), Category = "Math|Matrix")
+	static FVector4 Matrix_TransformVector4(const FMatrix& M, FVector4 V);
+
+	/** Transform a location - will take into account translation part of the FMatrix. */
+	UFUNCTION(BlueprintPure, meta = (DisplayName = "Transform Position (Matrix)", ScriptMethod = "TransformPosition"), Category = "Math|Matrix")
+	static FVector4 Matrix_TransformPosition(const FMatrix& M, FVector V);
+
+	/** Inverts the matrix and then transforms V - correctly handles scaling in this matrix. */
+	UFUNCTION(BlueprintPure, meta = (DisplayName = "Inverse Transform Position (Matrix)", ScriptMethod = "InverseTransformPosition"), Category = "Math|Matrix")
+	static FVector Matrix_InverseTransformPosition(const FMatrix& M, FVector V);
+
+	/**
+	 *	Transform a direction vector - will not take into account translation part of the FMatrix.
+	 *	If you want to transform a surface normal (or plane) and correctly account for non-uniform scaling you should use TransformByUsingAdjointT.
+	 */
+	UFUNCTION(BlueprintPure, meta = (DisplayName = "Transform Vector (Matrix)", ScriptMethod = "TransformVector"), Category = "Math|Matrix")
+	static FVector4 Matrix_TransformVector(const FMatrix& M, FVector V);
+
+	/**
+	 *	Transform a direction vector by the inverse of this matrix - will not take into account translation part.
+	 *	If you want to transform a surface normal (or plane) and correctly account for non-uniform scaling you should use TransformByUsingAdjointT with adjoint of matrix inverse.
+	 */
+	UFUNCTION(BlueprintPure, meta = (DisplayName = "Inverse Transform Vector (Matrix)", ScriptMethod = "InverseTransformVector"), Category = "Math|Matrix")
+	static FVector Matrix_InverseTransformVector(const FMatrix& M, FVector V);
+
+	// Transpose.
+	UFUNCTION(BlueprintPure, meta = (DisplayName = "Get Transposed (Matrix)", ScriptMethod = "GetTransposed"), Category = "Math|Matrix")
+	static FMatrix Matrix_GetTransposed(const FMatrix& M);
+
+	// @return determinant of this matrix.
+	UFUNCTION(BlueprintPure, meta = (DisplayName = "Get Determinant (Matrix)", ScriptMethod = "GetDeterminant"), Category = "Math|Matrix")
+	static float Matrix_GetDeterminant(const FMatrix& M);
+
+	/** @return the determinant of rotation 3x3 matrix */
+	UFUNCTION(BlueprintPure, meta = (DisplayName = "Get Rotation Determinant (Matrix)", ScriptMethod = "GetRotDeterminant"), Category = "Math|Matrix")
+	static float Matrix_GetRotDeterminant(const FMatrix& M);
+
+	/** Fast path, and handles nil matrices. */
+	UFUNCTION(BlueprintPure, meta = (DisplayName = "GetInverse (Matrix)", ScriptMethod = "GetInverse"), Category = "Math|Matrix")
+	static FMatrix Matrix_GetInverse(const FMatrix& M);
+
+	UFUNCTION(BlueprintPure, meta = (DisplayName = "Get Transpose Adjoint (Matrix)", ScriptMethod = "GetTransposeAdjoint"), Category = "Math|Matrix")
+	static FMatrix Matrix_GetTransposeAdjoint(const FMatrix& M);
+
+	// Remove any scaling from this matrix (ie magnitude of each row is 1) with error Tolerance
+	UFUNCTION(BlueprintCallable, meta = (DisplayName = "Remove Scaling (Matrix)", ScriptMethod = "RemoveScaling"), Category = "Math|Matrix")
+	static void Matrix_RemoveScaling(UPARAM(Ref) FMatrix& M, float Tolerance = 1.e-8f);
+
+	// Returns matrix after RemoveScaling with error Tolerance
+	UFUNCTION(BlueprintPure, meta = (DisplayName = "Get Matrix Without Scale (Matrix)", ScriptMethod = "GetMatrixWithoutScale"), Category = "Math|Matrix")
+	static FMatrix Matrix_GetMatrixWithoutScale(const FMatrix& M, float Tolerance = 1.e-8f);
+
+	/** return a 3D scale vector calculated from this matrix (where each component is the magnitude of a row vector) with error Tolerance. */
+	UFUNCTION(BlueprintPure, meta = (DisplayName = "Get Scale Vector (Matrix)", ScriptMethod = "GetScaleVector"), Category = "Math|Matrix")
+	static FVector Matrix_GetScaleVector(const FMatrix& M, float Tolerance = 1.e-8f);
+
+	// Remove any translation from this matrix
+	UFUNCTION(BlueprintPure, meta = (DisplayName = "Remove Translation (Matrix)", ScriptMethod = "RemoveTranslation"), Category = "Math|Matrix")
+	static FMatrix Matrix_RemoveTranslation(const FMatrix& M);
+
+	/** Returns a matrix with an additional translation concatenated. */
+	UFUNCTION(BlueprintPure, meta = (DisplayName = "Concatenate Translation (Matrix)", ScriptMethod = "ConcatenateTranslation"), Category = "Math|Matrix")
+	static FMatrix Matrix_ConcatenateTranslation(const FMatrix& M, FVector Translation);
+
+	/** Returns true if any element of this matrix is NaN */
+	UFUNCTION(BlueprintPure, meta = (DisplayName = "Contains NaN (Matrix)", ScriptMethod = "ContainsNaN"), Category = "Math|Matrix")
+	static bool Matrix_ContainsNaN(const FMatrix& M);
+
+	/** Scale the translation part of the matrix by the supplied vector. */
+	UFUNCTION(BlueprintPure, meta = (DisplayName = "Scale Translation (Matrix)", ScriptMethod = "ScaleTranslation"), Category = "Math|Matrix")
+	static FMatrix Matrix_ScaleTranslation(const FMatrix& M, FVector Scale3D);
+
+	/** @return the maximum magnitude of any row of the matrix. */
+	UFUNCTION(BlueprintPure, meta = (DisplayName = "Get Maximum Axis Scale (Matrix)", ScriptMethod = "GetMaximumAxisScale"), Category = "Math|Matrix")
+	static float Matrix_GetMaximumAxisScale(const FMatrix& M);
+
+	/** Apply Scale to this matrix **/
+	UFUNCTION(BlueprintPure, meta = (DisplayName = "Apply Scale (Matrix)", ScriptMethod = "ApplyScale"), Category = "Math|Matrix")
+	static FMatrix Matrix_ApplyScale(const FMatrix& M, float Scale);
+
+	/**
+	 * get axis of this matrix scaled by the scale of the matrix
+	 *
+	 * @param i index into the axis of the matrix
+	 * @ return vector of the axis
+	 */
+	UFUNCTION(BlueprintPure, meta = (DisplayName = "Get Scaled Axis (Matrix)", ScriptMethod = "GetScaledAxis"), Category = "Math|Matrix")
+	static FVector Matrix_GetScaledAxis(const FMatrix& M, TEnumAsByte<EAxis::Type> Axis);
+
+	/**
+	 * get axes of this matrix scaled by the scale of the matrix
+	 *
+	 * @param X axes returned to this param
+	 * @param Y axes returned to this param
+	 * @param Z axes returned to this param
+	 */
+	UFUNCTION(BlueprintPure, meta = (DisplayName = "Get Scaled Axes (Matrix)", ScriptMethod = "GetScaledAxes"), Category = "Math|Matrix")
+	static void Matrix_GetScaledAxes(const FMatrix& M, FVector &X, FVector &Y, FVector &Z);
+
+	/**
+	 * get unit length axis of this matrix
+	 *
+	 * @param i index into the axis of the matrix
+	 * @return vector of the axis
+	 */
+	UFUNCTION(BlueprintPure, meta = (DisplayName = "Get Unit Axis (Matrix)", ScriptMethod = "GetUnitAxis"), Category = "Math|Matrix")
+	static FVector Matrix_GetUnitAxis(const FMatrix& M, TEnumAsByte<EAxis::Type> Axis);
+
+	/**
+	 * get unit length axes of this matrix
+	 *
+	 * @param X axes returned to this param
+	 * @param Y axes returned to this param
+	 * @param Z axes returned to this param
+	 */
+	UFUNCTION(BlueprintPure, meta = (DisplayName = "Get Unit Axes (Matrix)", ScriptMethod = "GetUnitAxes"), Category = "Math|Matrix")
+	static void Matrix_GetUnitAxes(const FMatrix& M, FVector &X, FVector &Y, FVector &Z);
+
+	/**
+	 * set an axis of this matrix
+	 *
+	 * @param i index into the axis of the matrix
+	 * @param Axis vector of the axis
+	 */
+	UFUNCTION(BlueprintCallable, meta = (DisplayName = "Set Axis (Matrix)", ScriptMethod = "SetAxis"), Category = "Math|Matrix")
+	static void Matrix_SetAxis(UPARAM(Ref) FMatrix& M, TEnumAsByte<EAxis::Type> Axis, FVector AxisVector);
+
+	// Set the origin of the coordinate system to the given vector
+	UFUNCTION(BlueprintCallable, meta = (DisplayName = "Set Origin (Matrix)", ScriptMethod = "SetOrigin"), Category = "Math|Matrix")
+	static void Matrix_SetOrigin(UPARAM(Ref) FMatrix& M, FVector NewOrigin);
+
+	/**
+	 * get a column of this matrix
+	 *
+	 * @param i index into the column of the matrix
+	 * @return vector of the column
+	 */
+	UFUNCTION(BlueprintPure, meta = (DisplayName = "Get Column (Matrix)", ScriptMethod = "GetColumn"), Category = "Math|Matrix")
+	static FVector Matrix_GetColumn(const FMatrix& M, TEnumAsByte<EMatrixColumns::Type> Column);
+
+	UFUNCTION(BlueprintCallable, meta = (DisplayName = "Set Column (Matrix)", ScriptMethod = "SetColumn"), Category = "Math|Matrix")
+	static void Matrix_SetColumn(UPARAM(Ref) FMatrix& M, TEnumAsByte<EMatrixColumns::Type> Column, FVector Value);
+
+	/** @return rotator representation of this matrix */
+	UFUNCTION(BlueprintPure, meta = (DisplayName = "Get Rotator (Matrix)", ScriptMethod = "GetRotator"), Category = "Math|Matrix")
+	static FRotator Matrix_GetRotator(const FMatrix& M);
+
+	/**
+	 * Transform a rotation matrix into a quaternion.
+	 *
+	 * @warning rotation part will need to be unit length for this to be right!
+	 */
+	UFUNCTION(BlueprintPure, meta = (DisplayName = "To Quat (Matrix)", ScriptMethod = "ToQuat"), Category = "Math|Matrix")
+	static FQuat Matrix_ToQuat(const FMatrix& M);
+
+	// Frustum plane extraction.
+	/** @param OutPlane the near plane of the Frustum of this matrix */
+	UFUNCTION(BlueprintPure, meta = (DisplayName = "Get Frustum Near Plane (Matrix)", ScriptMethod = "GetFrustumNearPlane"), Category = "Math|Matrix")
+	static bool Matrix_GetFrustumNearPlane(const FMatrix& M, FPlane& OutPlane);
+
+	/** @param OutPlane the far plane of the Frustum of this matrix */
+	UFUNCTION(BlueprintPure, meta = (DisplayName = "Get Frustum Far Plane (Matrix)", ScriptMethod = "GetFrustumFarPlane"), Category = "Math|Matrix")
+	static bool Matrix_GetFrustumFarPlane(const FMatrix& M, FPlane& OutPlane);
+
+	/** @param OutPlane the left plane of the Frustum of this matrix */
+	UFUNCTION(BlueprintPure, meta = (DisplayName = "Get Frustum Left Plane (Matrix)", ScriptMethod = "GetFrustumLeftPlane"), Category = "Math|Matrix")
+	static bool Matrix_GetFrustumLeftPlane(const FMatrix& M, FPlane& OutPlane);
+
+	/** @param OutPlane the right plane of the Frustum of this matrix */
+	UFUNCTION(BlueprintPure, meta = (DisplayName = "Get Frustum Right Plane (Matrix)", ScriptMethod = "GetFrustumRightPlane"), Category = "Math|Matrix")
+	static bool Matrix_GetFrustumRightPlane(const FMatrix& M, FPlane& OutPlane);
+
+	/** @param OutPlane the top plane of the Frustum of this matrix */
+	UFUNCTION(BlueprintPure, meta = (DisplayName = "Get Frustum Top Plane (Matrix)", ScriptMethod = "GetFrustumTopPlane"), Category = "Math|Matrix")
+	static bool Matrix_GetFrustumTopPlane(const FMatrix& M, FPlane& OutPlane);
+
+	/** @param OutPlane the bottom plane of the Frustum of this matrix */
+	UFUNCTION(BlueprintPure, meta = (DisplayName = "Get Frustum Bottom Plane (Matrix)", ScriptMethod = "GetFrustumBottomPlane"), Category = "Math|Matrix")
+	static bool Matrix_GetFrustumBottomPlane(const FMatrix& M, FPlane& OutPlane);
+
+	/**
+	 * Utility for mirroring this transform across a certain plane, and flipping one of the axis as well.
+	 */
+	UFUNCTION(BlueprintPure, meta = (DisplayName = "Mirror (Matrix)", ScriptMethod = "Mirror"), Category = "Math|Matrix")
+	static FMatrix Matrix_Mirror(const FMatrix& M, TEnumAsByte<EAxis::Type> MirrorAxis, TEnumAsByte<EAxis::Type> FlipAxis);
 
 	//
 	// Quat constants - exposed for scripting
@@ -3270,6 +3523,10 @@ class ENGINE_API UKismetMathLibrary : public UBlueprintFunctionLibrary
 	UFUNCTION(BlueprintPure, Category="Math|Transform", meta = (DisplayName = "Determinant", ScriptMethod = "Determinant"))
 	static float Transform_Determinant(const FTransform& Transform);
 
+	/** Convert a Transform to a Matrix with scale */
+	UFUNCTION(BlueprintPure, Category="Math|Transform", meta = (DisplayName = "To Matrix (Transform)", ScriptMethod = "ToMatrix", CompactNodeTitle = "->", Keywords = "cast convert", BlueprintAutocast))
+	static FMatrix Conv_TransformToMatrix(const FTransform& Transform);
+
 
 	//
 	// Interpolation functions
diff --git a/Engine/Source/Runtime/Engine/Classes/Kismet/KismetMathLibrary.inl b/Engine/Source/Runtime/Engine/Classes/Kismet/KismetMathLibrary.inl
--- a/Engine/Source/Runtime/Engine/Classes/Kismet/KismetMathLibrary.inl
+++ b/Engine/Source/Runtime/Engine/Classes/Kismet/KismetMathLibrary.inl
@@ -2635,6 +2635,261 @@ FVector UKismetMathLibrary::Matrix_GetOrigin(const FMatrix& InMatrix)
 	return InMatrix.GetOrigin();
 }
 
+KISMET_MATH_FORCEINLINE
+FMatrix UKismetMathLibrary::Matrix_Identity()
+{
+	return FMatrix::Identity;
+}
+
+KISMET_MATH_FORCEINLINE
+FMatrix UKismetMathLibrary::Multiply_MatrixMatrix(const FMatrix& A, const FMatrix& B)
+{
+	return A * B;
+}
+
+KISMET_MATH_FORCEINLINE
+FMatrix UKismetMathLibrary::Add_MatrixMatrix(const FMatrix& A, const FMatrix& B)
+{
+	return A + B;
+}
+
+KISMET_MATH_FORCEINLINE
+FMatrix UKismetMathLibrary::Multiply_MatrixFloat(const FMatrix& A, float B)
+{
+	return A * B;
+}
+
+KISMET_MATH_FORCEINLINE
+bool UKismetMathLibrary::EqualEqual_MatrixMatrix(const FMatrix& A, const FMatrix& B, float Tolerance /*= KINDA_SMALL_NUMBER*/)
+{
+	return A.Equals(B, Tolerance);
+}
+
+KISMET_MATH_FORCEINLINE
+bool UKismetMathLibrary::NotEqual_MatrixMatrix(const FMatrix& A, const FMatrix& B, float Tolerance /*= KINDA_SMALL_NUMBER*/)
+{
+	return !(A.Equals(B, Tolerance));
+}
+
+KISMET_MATH_FORCEINLINE
+FVector4 UKismetMathLibrary::Matrix_TransformVector4(const FMatrix& M, FVector4 V)
+{
+	return M.TransformFVector4(V);
+}
+
+KISMET_MATH_FORCEINLINE
+FVector4 UKismetMathLibrary::Matrix_TransformPosition(const FMatrix& M, FVector V)
+{
+	return M.TransformPosition(V);
+}
+
+KISMET_MATH_FORCEINLINE
+FVector UKismetMathLibrary::Matrix_InverseTransformPosition(const FMatrix& M, FVector V)
+{
+	return M.InverseTransformPosition(V);
+}
+
+KISMET_MATH_FORCEINLINE
+FVector4 UKismetMathLibrary::Matrix_TransformVector(const FMatrix& M, FVector V)
+{
+	return M.TransformVector(V);
+}
+
+KISMET_MATH_FORCEINLINE
+FVector UKismetMathLibrary::Matrix_InverseTransformVector(const FMatrix& M, FVector V)
+{
+	return M.InverseTransformVector(V);
+}
+
+KISMET_MATH_FORCEINLINE
+FMatrix UKismetMathLibrary::Matrix_GetTransposed(const FMatrix& M)
+{
+	return M.GetTransposed();
+}
+
+KISMET_MATH_FORCEINLINE
+float UKismetMathLibrary::Matrix_GetDeterminant(const FMatrix& M)
+{
+	return M.Determinant();
+}
+
+KISMET_MATH_FORCEINLINE
+float UKismetMathLibrary::Matrix_GetRotDeterminant(const FMatrix& M)
+{
+	return M.RotDeterminant();
+}
+
+KISMET_MATH_FORCEINLINE
+FMatrix UKismetMathLibrary::Matrix_GetInverse(const FMatrix& M)
+{
+	return M.Inverse();
+}
+
+KISMET_MATH_FORCEINLINE
+FMatrix UKismetMathLibrary::Matrix_GetTransposeAdjoint(const FMatrix& M)
+{
+	return M.TransposeAdjoint();
+}
+
+KISMET_MATH_FORCEINLINE
+void UKismetMathLibrary::Matrix_RemoveScaling(FMatrix& M, float Tolerance /*= SMALL_NUMBER*/)
+{
+	M.RemoveScaling(Tolerance);
+}
+
+KISMET_MATH_FORCEINLINE
+FMatrix UKismetMathLibrary::Matrix_GetMatrixWithoutScale(const FMatrix& M, float Tolerance /*= SMALL_NUMBER*/)
+{
+	return M.GetMatrixWithoutScale(Tolerance);
+}
+
+KISMET_MATH_FORCEINLINE
+FVector UKismetMathLibrary::Matrix_GetScaleVector(const FMatrix& M, float Tolerance /*= SMALL_NUMBER*/)
+{
+	return M.GetScaleVector(Tolerance);
+}
+
+KISMET_MATH_FORCEINLINE
+FMatrix UKismetMathLibrary::Matrix_RemoveTranslation(const FMatrix& M)
+{
+	return M.RemoveTranslation();
+}
+
+KISMET_MATH_FORCEINLINE
+FMatrix UKismetMathLibrary::Matrix_ConcatenateTranslation(const FMatrix& M, FVector Translation)
+{
+	return M.ConcatTranslation(Translation);
+}
+
+KISMET_MATH_FORCEINLINE
+bool UKismetMathLibrary::Matrix_ContainsNaN(const FMatrix& M)
+{
+	return M.ContainsNaN();
+}
+
+KISMET_MATH_FORCEINLINE
+FMatrix UKismetMathLibrary::Matrix_ScaleTranslation(const FMatrix& M, FVector Scale3D)
+{
+	FMatrix ScaledMatrix = M;
+	ScaledMatrix.ScaleTranslation(Scale3D);
+	return ScaledMatrix;
+}
+
+KISMET_MATH_FORCEINLINE
+float UKismetMathLibrary::Matrix_GetMaximumAxisScale(const FMatrix& M)
+{
+	return M.GetMaximumAxisScale();
+}
+
+KISMET_MATH_FORCEINLINE
+FMatrix UKismetMathLibrary::Matrix_ApplyScale(const FMatrix& M, float Scale)
+{
+	return M.ApplyScale(Scale);
+}
+
+KISMET_MATH_FORCEINLINE
+FVector UKismetMathLibrary::Matrix_GetScaledAxis(const FMatrix& M, TEnumAsByte<EAxis::Type> Axis)
+{
+	return M.GetScaledAxis(Axis);
+}
+
+KISMET_MATH_FORCEINLINE
+void UKismetMathLibrary::Matrix_GetScaledAxes(const FMatrix& M, FVector &X, FVector &Y, FVector &Z)
+{
+	M.GetScaledAxes(X, Y, Z);
+}
+
+KISMET_MATH_FORCEINLINE
+FVector UKismetMathLibrary::Matrix_GetUnitAxis(const FMatrix& M, TEnumAsByte<EAxis::Type> Axis)
+{
+	return M.GetUnitAxis(Axis);
+}
+
+KISMET_MATH_FORCEINLINE
+void UKismetMathLibrary::Matrix_GetUnitAxes(const FMatrix& M, FVector &X, FVector &Y, FVector &Z)
+{
+	M.GetUnitAxes(X, Y, Z);
+}
+
+KISMET_MATH_FORCEINLINE
+void UKismetMathLibrary::Matrix_SetAxis(FMatrix& M, TEnumAsByte<EAxis::Type> Axis, FVector AxisVector)
+{
+	M.SetAxis(StaticCast<int32>(Axis), AxisVector);
+}
+
+KISMET_MATH_FORCEINLINE
+void UKismetMathLibrary::Matrix_SetOrigin(FMatrix& M, FVector NewOrigin)
+{
+	M.SetOrigin(NewOrigin);
+}
+
+KISMET_MATH_FORCEINLINE
+FVector UKismetMathLibrary::Matrix_GetColumn(const FMatrix& M, TEnumAsByte<EMatrixColumns::Type> Column)
+{
+	return M.GetColumn(StaticCast<int32>(Column));
+}
+
+KISMET_MATH_FORCEINLINE
+void UKismetMathLibrary::Matrix_SetColumn(FMatrix& M, TEnumAsByte<EMatrixColumns::Type> Column, FVector Value)
+{
+	M.SetColumn(StaticCast<int32>(Column), Value);
+}
+
+KISMET_MATH_FORCEINLINE
+FRotator UKismetMathLibrary::Matrix_GetRotator(const FMatrix& M)
+{
+	return M.Rotator();
+}
+
+KISMET_MATH_FORCEINLINE
+FQuat UKismetMathLibrary::Matrix_ToQuat(const FMatrix& M)
+{
+	return M.ToQuat();
+}
+
+KISMET_MATH_FORCEINLINE
+bool UKismetMathLibrary::Matrix_GetFrustumNearPlane(const FMatrix& M, FPlane& OutPlane)
+{
+	return M.GetFrustumNearPlane(OutPlane);
+}
+
+KISMET_MATH_FORCEINLINE
+bool UKismetMathLibrary::Matrix_GetFrustumFarPlane(const FMatrix& M, FPlane& OutPlane)
+{
+	return M.GetFrustumFarPlane(OutPlane);
+}
+
+KISMET_MATH_FORCEINLINE
+bool UKismetMathLibrary::Matrix_GetFrustumLeftPlane(const FMatrix& M, FPlane& OutPlane)
+{
+	return M.GetFrustumLeftPlane(OutPlane);
+}
+
+KISMET_MATH_FORCEINLINE
+bool UKismetMathLibrary::Matrix_GetFrustumRightPlane(const FMatrix& M, FPlane& OutPlane)
+{
+	return M.GetFrustumRightPlane(OutPlane);
+}
+
+KISMET_MATH_FORCEINLINE
+bool UKismetMathLibrary::Matrix_GetFrustumTopPlane(const FMatrix& M, FPlane& OutPlane)
+{
+	return M.GetFrustumTopPlane(OutPlane);
+}
+
+KISMET_MATH_FORCEINLINE
+bool UKismetMathLibrary::Matrix_GetFrustumBottomPlane(const FMatrix& M, FPlane& OutPlane)
+{
+	return M.GetFrustumBottomPlane(OutPlane);
+}
+
+KISMET_MATH_FORCEINLINE
+FMatrix UKismetMathLibrary::Matrix_Mirror(const FMatrix& M, TEnumAsByte<EAxis::Type> MirrorAxis, TEnumAsByte<EAxis::Type> FlipAxis)
+{
+	FMatrix MirrorMatrix = M;
+	MirrorMatrix.Mirror(MirrorAxis, FlipAxis);
+	return MirrorMatrix;
+}
 
 /* FQuat
 *****************************************************************************/
@@ -2797,7 +3052,7 @@ FQuat UKismetMathLibrary::Quat_Log(const FQuat& Q)
 }
 
 KISMET_MATH_FORCEINLINE
-void UKismetMathLibrary::Quat_SetComponents(UPARAM(ref) FQuat& Q, float X, float Y, float Z, float W)
+void UKismetMathLibrary::Quat_SetComponents(FQuat& Q, float X, float Y, float Z, float W)
 {
 	Q.X = X;
 	Q.Y = Y;
@@ -2806,7 +3061,7 @@ void UKismetMathLibrary::Quat_SetComponents(UPARAM(ref) FQuat& Q, float X, float
 }
 
 KISMET_MATH_FORCEINLINE
-void UKismetMathLibrary::Quat_SetFromEuler(UPARAM(ref) FQuat& Q, const FVector& Euler)
+void UKismetMathLibrary::Quat_SetFromEuler(FQuat& Q, const FVector& Euler)
 {
 	Q = FQuat::MakeFromEuler(Euler);
 }
@@ -2919,7 +3174,7 @@ void UKismetMathLibrary::BreakColor(FLinearColor InColor, float& R, float& G, fl
 }
 
 KISMET_MATH_FORCEINLINE
-void UKismetMathLibrary::LinearColor_Set(UPARAM(ref) FLinearColor& InOutColor, FLinearColor InColor)
+void UKismetMathLibrary::LinearColor_Set(FLinearColor& InOutColor, FLinearColor InColor)
 {
 	InOutColor.R = InColor.R;
 	InOutColor.G = InColor.G;
@@ -2928,7 +3183,7 @@ void UKismetMathLibrary::LinearColor_Set(UPARAM(ref) FLinearColor& InOutColor, F
 }
 
 KISMET_MATH_FORCEINLINE
-void UKismetMathLibrary::LinearColor_SetRGBA(UPARAM(ref) FLinearColor& InOutColor, float R, float G, float B, float A)
+void UKismetMathLibrary::LinearColor_SetRGBA(FLinearColor& InOutColor, float R, float G, float B, float A)
 {
 	InOutColor.R = R;
 	InOutColor.G = G;
@@ -2937,32 +3192,32 @@ void UKismetMathLibrary::LinearColor_SetRGBA(UPARAM(ref) FLinearColor& InOutColo
 }
 
 KISMET_MATH_FORCEINLINE
-void UKismetMathLibrary::LinearColor_SetFromHSV(UPARAM(ref) FLinearColor& InOutColor, float H, float S, float V, float A)
+void UKismetMathLibrary::LinearColor_SetFromHSV(FLinearColor& InOutColor, float H, float S, float V, float A)
 {
 	const FLinearColor HSV(H, S, V, A);
 	InOutColor = HSV.HSVToLinearRGB();
 }
 
 KISMET_MATH_FORCEINLINE
-void UKismetMathLibrary::LinearColor_SetFromSRGB(UPARAM(ref) FLinearColor& InOutColor, const FColor& InSRGB)
+void UKismetMathLibrary::LinearColor_SetFromSRGB(FLinearColor& InOutColor, const FColor& InSRGB)
 {
 	InOutColor = FLinearColor::FromSRGBColor(InSRGB);
 }
 
 KISMET_MATH_FORCEINLINE
-void UKismetMathLibrary::LinearColor_SetFromPow22(UPARAM(ref) FLinearColor& InOutColor, const FColor& InColor)
+void UKismetMathLibrary::LinearColor_SetFromPow22(FLinearColor& InOutColor, const FColor& InColor)
 {
 	InOutColor = FLinearColor::FromPow22Color(InColor);
 }
 
 KISMET_MATH_FORCEINLINE
-void UKismetMathLibrary::LinearColor_SetTemperature(UPARAM(ref) FLinearColor& InOutColor, float InTemperature)
+void UKismetMathLibrary::LinearColor_SetTemperature(FLinearColor& InOutColor, float InTemperature)
 {
 	InOutColor = FLinearColor::MakeFromColorTemperature(InTemperature);
 }
 
 KISMET_MATH_FORCEINLINE
-void UKismetMathLibrary::LinearColor_SetRandomHue(UPARAM(ref) FLinearColor& InOutColor)
+void UKismetMathLibrary::LinearColor_SetRandomHue(FLinearColor& InOutColor)
 {
 	InOutColor = FLinearColor::MakeRandomColor();
 }
diff --git a/Engine/Source/Runtime/Engine/Classes/Kismet/KismetStringLibrary.h b/Engine/Source/Runtime/Engine/Classes/Kismet/KismetStringLibrary.h
--- a/Engine/Source/Runtime/Engine/Classes/Kismet/KismetStringLibrary.h
+++ b/Engine/Source/Runtime/Engine/Classes/Kismet/KismetStringLibrary.h
@@ -60,6 +60,10 @@ class ENGINE_API UKismetStringLibrary : public UBlueprintFunctionLibrary
 	UFUNCTION(BlueprintPure, meta=(DisplayName = "ToString (name)", CompactNodeTitle = "->", BlueprintAutocast), Category="Utilities|String")
 	static FString Conv_NameToString(FName InName);
 
+	/** Converts a name value to a string */
+	UFUNCTION(BlueprintPure, meta = (DisplayName = "ToString (Matrix)", CompactNodeTitle = "->", BlueprintAutocast), Category = "Utilities|String")
+	static FString Conv_MatrixToString(const FMatrix& InMatrix);
+
 	/** Converts a string to a name value */
 	UFUNCTION(BlueprintPure, meta=(DisplayName = "String To Name", CompactNodeTitle = "->", BlueprintAutocast), Category="Utilities|String")
 	static FName Conv_StringToName(const FString& InString);
diff --git a/Engine/Source/Runtime/Engine/Private/KismetMathLibrary.cpp b/Engine/Source/Runtime/Engine/Private/KismetMathLibrary.cpp
--- a/Engine/Source/Runtime/Engine/Private/KismetMathLibrary.cpp
+++ b/Engine/Source/Runtime/Engine/Private/KismetMathLibrary.cpp
@@ -575,6 +575,11 @@ float UKismetMathLibrary::Transform_Determinant(const FTransform& Transform)
 	return Transform.ToMatrixWithScale().Determinant();
 }
 
+FMatrix UKismetMathLibrary::Conv_TransformToMatrix(const FTransform& Tranform)
+{
+	return Tranform.ToMatrixWithScale();
+}
+
 bool UKismetMathLibrary::ClassIsChildOf(TSubclassOf<class UObject> TestClass, TSubclassOf<class UObject> ParentClass)
 {
 	return ((*ParentClass != NULL) && (*TestClass != NULL)) ? (*TestClass)->IsChildOf(*ParentClass) : false;
diff --git a/Engine/Source/Runtime/Engine/Private/KismetStringLibrary.cpp b/Engine/Source/Runtime/Engine/Private/KismetStringLibrary.cpp
--- a/Engine/Source/Runtime/Engine/Private/KismetStringLibrary.cpp
+++ b/Engine/Source/Runtime/Engine/Private/KismetStringLibrary.cpp
@@ -106,6 +106,11 @@ FString UKismetStringLibrary::Conv_NameToString(FName InName)
 	return InName.ToString();
 }
 
+FString UKismetStringLibrary::Conv_MatrixToString(const FMatrix& InMatrix)
+{
+	return InMatrix.ToString();
+}
+
 FName UKismetStringLibrary::Conv_StringToName(const FString& InString)
 {
 	return FName(*InString);
