Fix for duplicate ScriptMethod name in Multiple_MatrixFloat

diff --git a/Engine/Source/Runtime/Engine/Classes/Kismet/KismetMathLibrary.h b/Engine/Source/Runtime/Engine/Classes/Kismet/KismetMathLibrary.h
--- a/Engine/Source/Runtime/Engine/Classes/Kismet/KismetMathLibrary.h
+++ b/Engine/Source/Runtime/Engine/Classes/Kismet/KismetMathLibrary.h
@@ -2245,7 +2245,7 @@ class ENGINE_API UKismetMathLibrary : public UBlueprintFunctionLibrary
 	/**
 	  * This isn't applying SCALE, just multiplying the value to all members - i.e. weighting
 	  */
-	UFUNCTION(BlueprintPure, meta = (DisplayName = "Matrix * Float", CompactNodeTitle = "*", ScriptMethod = "Multiply", ScriptOperator = "*;*=", Keywords = "* multiply"), Category = "Math|Matrix")
+	UFUNCTION(BlueprintPure, meta = (DisplayName = "Matrix * Float", CompactNodeTitle = "*", ScriptMethod = "MultiplyFloat", ScriptOperator = "*;*=", Keywords = "* multiply"), Category = "Math|Matrix")
 	static FMatrix Multiply_MatrixFloat (const FMatrix& A, float B);
 
 	/**
