Deprecated boolean version of SetShaderValue and SetPixelShaderBool for 4.24.

diff --git a/Engine/Source/Runtime/Engine/Private/SimpleElementShaders.cpp b/Engine/Source/Runtime/Engine/Private/SimpleElementShaders.cpp
--- a/Engine/Source/Runtime/Engine/Private/SimpleElementShaders.cpp
+++ b/Engine/Source/Runtime/Engine/Private/SimpleElementShaders.cpp
@@ -165,7 +165,7 @@ void FSimpleElementDistanceFieldGammaPS::SetParameters(
 	float InGamma,
 	float InClipRef,
 	float SmoothWidthValue,
-	bool EnableShadowValue,
+	bool bEnableShadowValue,
 	const FVector2D& ShadowDirectionValue,
 	const FLinearColor& ShadowColorValue,
 	float ShadowSmoothWidthValue,
@@ -174,15 +174,16 @@ void FSimpleElementDistanceFieldGammaPS::SetParameters(
 	)
 {
 	FSimpleElementMaskedGammaBasePS::SetParameters(RHICmdList, Texture,InGamma,InClipRef,BlendMode);
-	SetShaderValue(RHICmdList, GetPixelShader(),SmoothWidth,SmoothWidthValue);		
-	SetPixelShaderBool(RHICmdList, GetPixelShader(),EnableShadow,EnableShadowValue);
-	if (EnableShadowValue)
+	SetShaderValue(RHICmdList, GetPixelShader(),SmoothWidth,SmoothWidthValue);
+	const uint32 bEnableShadowValueUInt = (bEnableShadowValue ? 1 : 0);
+	SetShaderValue(RHICmdList, GetPixelShader(),EnableShadow,bEnableShadowValueUInt);
+	if (bEnableShadowValue)
 	{
 		SetShaderValue(RHICmdList, GetPixelShader(),ShadowDirection,ShadowDirectionValue);
 		SetShaderValue(RHICmdList, GetPixelShader(),ShadowColor,ShadowColorValue);
 		SetShaderValue(RHICmdList, GetPixelShader(),ShadowSmoothWidth,ShadowSmoothWidthValue);
 	}
-	SetPixelShaderBool(RHICmdList, GetPixelShader(),EnableGlow,GlowInfo.bEnableGlow);
+	SetShaderValue(RHICmdList, GetPixelShader(),EnableGlow,GlowInfo.bEnableGlow);
 	if (GlowInfo.bEnableGlow)
 	{
 		SetShaderValue(RHICmdList, GetPixelShader(),GlowColor,GlowInfo.GlowColor);
diff --git a/Engine/Source/Runtime/RenderCore/Public/ShaderParameterUtils.h b/Engine/Source/Runtime/RenderCore/Public/ShaderParameterUtils.h
--- a/Engine/Source/Runtime/RenderCore/Public/ShaderParameterUtils.h
+++ b/Engine/Source/Runtime/RenderCore/Public/ShaderParameterUtils.h
@@ -83,6 +83,7 @@ void SetShaderValueOnContext(
 
 /** Specialization of the above for C++ bool type. */
 template<typename ShaderRHIParamRef>
+UE_DEPRECATED(4.24, "Please use integer values for boolean shader parameters instead.")
 void SetShaderValue(
 	FRHICommandList& RHICmdList, 
 	const ShaderRHIParamRef& Shader,
@@ -97,6 +98,7 @@ void SetShaderValue(
 
 /** Specialization of the above for C++ bool type. */
 template<typename ShaderRHIParamRef>
+UE_DEPRECATED(4.24, "Please use integer values for boolean shader parameters instead.")
 void SetShaderValue(
 	FRHIAsyncComputeCommandList& RHICmdList,
 	const ShaderRHIParamRef& Shader,
@@ -111,6 +113,7 @@ void SetShaderValue(
 
 /** Specialization of the above for C++ bool type. */
 template<typename ShaderRHIParamRef>
+UE_DEPRECATED(4.24, "Please use integer values for boolean shader parameters instead.")
 void SetShaderValue(
 	FRHICommandListImmediate& RHICmdList,
 	const ShaderRHIParamRef& Shader,
@@ -174,6 +177,7 @@ void SetShaderValueArray(
 /**
  * Sets the value of a pixel shader bool parameter.
  */
+UE_DEPRECATED(4.24, "Please use integer values for boolean shader parameters instead.")
 inline void SetPixelShaderBool(
 	FRHICommandList& RHICmdList, 
 	FRHIPixelShader* PixelShader,
