Support Crunch compression of runtime virtual texture streaming mips.

diff --git a/Engine/Source/Runtime/Engine/Classes/VT/RuntimeVirtualTexture.h b/Engine/Source/Runtime/Engine/Classes/VT/RuntimeVirtualTexture.h
--- a/Engine/Source/Runtime/Engine/Classes/VT/RuntimeVirtualTexture.h
+++ b/Engine/Source/Runtime/Engine/Classes/VT/RuntimeVirtualTexture.h
@@ -64,10 +64,6 @@ protected:
 	UPROPERTY(EditAnywhere, BlueprintReadOnly, Category = Size, meta = (UIMin = "0", UIMax = "4", DisplayName = "Border padding for each virtual texture tile"))
 	int32 TileBorderSize = 2; // 4
 
-	/** Number of low mips to cut from the virtual texture. This can reduce peak virtual texture update cost but will also increase the probability of mip shimmering. */
-	UPROPERTY(EditAnywhere, BlueprintReadOnly, AdvancedDisplay, Category = LowMips, meta = (UIMin = "0", UIMax = "6", DisplayName = "Number of low mips to remove from the virtual texture"))
-	int32 RemoveLowMips = 0;
-
 	/** Number of low mips to serialize and stream for the virtual texture. This can reduce rendering update cost. */
 	UPROPERTY(EditAnywhere, BlueprintReadOnly, Category = LowMips, meta = (UIMin = "0", UIMax = "6", DisplayName = "Number of low mips to stream to the virtual texture"))
 	int32 StreamLowMips = 0;
@@ -76,6 +72,14 @@ protected:
 	UPROPERTY(VisibleAnywhere, Category = LowMips, meta = (DisplayName = "Streaming low mip texture"))
 	class URuntimeVirtualTextureStreamingProxy* StreamingTexture;
 
+	/** Enable Crunch compression of streamed low mips. ZLib compression is used when Crunch is disabled. */
+	UPROPERTY(EditAnywhere, BlueprintReadOnly, AdvancedDisplay, Category = LowMips, meta = (DisplayName = "Enable Crunch compression"))
+	bool bEnableCompressCrunch = false;
+
+	/** Number of low mips to cut from the virtual texture. This can reduce peak virtual texture update cost but will also increase the probability of mip shimmering. */
+	UPROPERTY(EditAnywhere, BlueprintReadOnly, AdvancedDisplay, Category = LowMips, meta = (UIMin = "0", UIMax = "6", DisplayName = "Number of low mips to remove from the virtual texture"))
+	int32 RemoveLowMips = 0;
+
 public:
 	/** Public getter for enabled status */
 	bool GetEnabled() { return bEnable; }
diff --git a/Engine/Source/Runtime/Engine/Private/VT/RuntimeVirtualTexture.cpp b/Engine/Source/Runtime/Engine/Private/VT/RuntimeVirtualTexture.cpp
--- a/Engine/Source/Runtime/Engine/Private/VT/RuntimeVirtualTexture.cpp
+++ b/Engine/Source/Runtime/Engine/Private/VT/RuntimeVirtualTexture.cpp
@@ -485,6 +485,7 @@ uint32 URuntimeVirtualTexture::GetStreamingTextureBuildHash() const
 			uint32 TileSize : 4;
 			uint32 TileBorderSize : 4;
 			uint32 StreamLowMips : 4;
+			uint32 EnableCompressCrunch : 1;
 		};
 	};
 
@@ -496,6 +497,7 @@ uint32 URuntimeVirtualTexture::GetStreamingTextureBuildHash() const
 	Settings.TileSize = (uint32)TileSize;
 	Settings.TileBorderSize = (uint32)TileBorderSize;
 	Settings.StreamLowMips = (uint32)GetStreamLowMips();
+	Settings.EnableCompressCrunch = (uint32)bEnableCompressCrunch;
 
 	return Settings.PackedValue;
 }
@@ -513,6 +515,7 @@ void URuntimeVirtualTexture::InitializeStreamingTexture(uint32 InSizeX, uint32 I
 	StreamingTexture->Settings.Init();
 	StreamingTexture->Settings.TileSize = GetTileSize();
 	StreamingTexture->Settings.TileBorderSize = GetTileBorderSize();
+	StreamingTexture->Settings.bEnableCompressCrunch = bEnableCompressCrunch;
 
 	StreamingTexture->BuildHash = GetStreamingTextureBuildHash();
 
