Fix build

#rnx

diff --git a/Engine/Source/Runtime/Engine/Private/Animation/AnimCompressionTypes.cpp b/Engine/Source/Runtime/Engine/Private/Animation/AnimCompressionTypes.cpp
--- a/Engine/Source/Runtime/Engine/Private/Animation/AnimCompressionTypes.cpp
+++ b/Engine/Source/Runtime/Engine/Private/Animation/AnimCompressionTypes.cpp
@@ -10,7 +10,7 @@
 #include "AnimationRuntime.h"
 #include "UObject/FortniteMainBranchObjectVersion.h"
 
-CSV_DECLARE_CATEGORY_EXTERN(Animation);
+CSV_DECLARE_CATEGORY_MODULE_EXTERN(ENGINE_API, Animation);
 
 DECLARE_CYCLE_STAT(TEXT("Build Anim Track Pairs"), STAT_BuildAnimTrackPairs, STATGROUP_Anim);
 DECLARE_CYCLE_STAT(TEXT("Extract Pose From Anim Data"), STAT_ExtractPoseFromAnimData, STATGROUP_Anim);
diff --git a/Engine/Source/Runtime/Engine/Private/Animation/AnimSequence.cpp b/Engine/Source/Runtime/Engine/Private/Animation/AnimSequence.cpp
--- a/Engine/Source/Runtime/Engine/Private/Animation/AnimSequence.cpp
+++ b/Engine/Source/Runtime/Engine/Private/Animation/AnimSequence.cpp
@@ -46,7 +46,7 @@
 DECLARE_CYCLE_STAT(TEXT("AnimSeq GetBonePose"), STAT_AnimSeq_GetBonePose, STATGROUP_Anim);
 DECLARE_CYCLE_STAT(TEXT("AnimSeq EvalCurveData"), STAT_AnimSeq_EvalCurveData, STATGROUP_Anim);
 
-CSV_DECLARE_CATEGORY_EXTERN(Animation);
+CSV_DECLARE_CATEGORY_MODULE_EXTERN(ENGINE_API, Animation);
 
 int32 GPerformFrameStripping = 0;
 int32 GPerformFrameStrippingOddFramedAnimations = 0;
diff --git a/Engine/Source/Runtime/Engine/Private/Animation/AnimSequenceBase.cpp b/Engine/Source/Runtime/Engine/Private/Animation/AnimSequenceBase.cpp
--- a/Engine/Source/Runtime/Engine/Private/Animation/AnimSequenceBase.cpp
+++ b/Engine/Source/Runtime/Engine/Private/Animation/AnimSequenceBase.cpp
@@ -12,7 +12,7 @@
 #include "UObject/FortniteMainBranchObjectVersion.h"
 
 DEFINE_LOG_CATEGORY(LogAnimMarkerSync);
-CSV_DECLARE_CATEGORY_EXTERN(Animation);
+CSV_DECLARE_CATEGORY_MODULE_EXTERN(ENGINE_API, Animation);
 
 #define LOCTEXT_NAMESPACE "AnimSequenceBase"
 /////////////////////////////////////////////////////
diff --git a/Engine/Source/Runtime/Engine/Private/Animation/AnimStreamable.cpp b/Engine/Source/Runtime/Engine/Private/Animation/AnimStreamable.cpp
--- a/Engine/Source/Runtime/Engine/Private/Animation/AnimStreamable.cpp
+++ b/Engine/Source/Runtime/Engine/Private/Animation/AnimStreamable.cpp
@@ -18,6 +18,8 @@
 #include "BonePose.h"
 #include "ContentStreaming.h"
 
+CSV_DECLARE_CATEGORY_MODULE_EXTERN(ENGINE_API, Animation);
+
 DECLARE_CYCLE_STAT(TEXT("AnimStreamable GetAnimationPose"), STAT_AnimStreamable_GetAnimationPose, STATGROUP_Anim);
 
 // This is a version string for the streaming anim chunk logic
diff --git a/Engine/Source/Runtime/Engine/Private/Components/SkeletalMeshComponent.cpp b/Engine/Source/Runtime/Engine/Private/Components/SkeletalMeshComponent.cpp
--- a/Engine/Source/Runtime/Engine/Private/Components/SkeletalMeshComponent.cpp
+++ b/Engine/Source/Runtime/Engine/Private/Components/SkeletalMeshComponent.cpp
@@ -57,7 +57,7 @@ DECLARE_CYCLE_STAT_EXTERN(TEXT("Anim Instance Spawn Time"), STAT_AnimSpawnTime,
 DEFINE_STAT(STAT_AnimSpawnTime);
 DEFINE_STAT(STAT_PostAnimEvaluation);
 
-CSV_DECLARE_CATEGORY_EXTERN(Animation);
+CSV_DECLARE_CATEGORY_MODULE_EXTERN(ENGINE_API, Animation);
 CSV_DECLARE_CATEGORY_MODULE_EXTERN(CORE_API, Basic);
 
 FAutoConsoleTaskPriority CPrio_ParallelAnimationEvaluationTask(
diff --git a/Engine/Source/Runtime/Engine/Private/SkeletalMeshComponentPhysics.cpp b/Engine/Source/Runtime/Engine/Private/SkeletalMeshComponentPhysics.cpp
--- a/Engine/Source/Runtime/Engine/Private/SkeletalMeshComponentPhysics.cpp
+++ b/Engine/Source/Runtime/Engine/Private/SkeletalMeshComponentPhysics.cpp
@@ -58,8 +58,7 @@
 
 DECLARE_CYCLE_STAT(TEXT("CreateClothing"), STAT_CreateClothing, STATGROUP_Physics);
 
-CSV_DECLARE_CATEGORY_EXTERN(Animation);
-
+CSV_DECLARE_CATEGORY_MODULE_EXTERN(ENGINE_API, Animation);
 CSV_DECLARE_CATEGORY_MODULE_EXTERN(CORE_API, Basic);
 
 TAutoConsoleVariable<int32> CVarEnableClothPhysics(TEXT("p.ClothPhysics"), 1, TEXT("If 1, physics cloth will be used for simulation."));
