Re-order sound engine stats so mix/modulation comes before sounds/cues/waves

diff --git a/Engine/Source/Runtime/Engine/Private/UnrealEngine.cpp b/Engine/Source/Runtime/Engine/Private/UnrealEngine.cpp
--- a/Engine/Source/Runtime/Engine/Private/UnrealEngine.cpp
+++ b/Engine/Source/Runtime/Engine/Private/UnrealEngine.cpp
@@ -1672,13 +1672,13 @@ void UEngine::Init(IEngineLoop* InEngineLoop)
 	EngineStats.Add(FEngineStatFuncs(TEXT("STAT_ColorList"), TEXT("STATCAT_Engine"), FText::GetEmpty(), &UEngine::RenderStatColorList, NULL));
 	EngineStats.Add(FEngineStatFuncs(TEXT("STAT_Levels"), TEXT("STATCAT_Engine"), FText::GetEmpty(), &UEngine::RenderStatLevels, NULL));
 #if !UE_BUILD_SHIPPING
-	EngineStats.Add(FEngineStatFuncs(TEXT("STAT_Sounds"), TEXT("STATCAT_Engine"), FText::GetEmpty(), &UEngine::RenderStatSounds, &UEngine::ToggleStatSounds));
-	EngineStats.Add(FEngineStatFuncs(TEXT("STAT_SoundCues"), TEXT("STATCAT_Engine"), FText::GetEmpty(), &UEngine::RenderStatSoundCues, &UEngine::ToggleStatSoundCues));
 	EngineStats.Add(FEngineStatFuncs(TEXT("STAT_SoundMixes"), TEXT("STATCAT_Engine"), FText::GetEmpty(), &UEngine::RenderStatSoundMixes, &UEngine::ToggleStatSoundMixes));
-	EngineStats.Add(FEngineStatFuncs(TEXT("STAT_AudioStreaming"), TEXT("STATCAT_Engine"), FText::GetEmpty(), &UEngine::RenderStatAudioStreaming, &UEngine::ToggleStatAudioStreaming));
 	EngineStats.Add(FEngineStatFuncs(TEXT("STAT_SoundModulators"), TEXT("STATCAT_Engine"), FText::GetEmpty(), &UEngine::RenderStatSoundModulators, &UEngine::ToggleStatSoundModulators));
 	EngineStats.Add(FEngineStatFuncs(TEXT("STAT_SoundModulatorsHelp"), TEXT("STATCAT_Engine"), FText::GetEmpty(), nullptr, &UEngine::PostStatSoundModulatorHelp));
+	EngineStats.Add(FEngineStatFuncs(TEXT("STAT_AudioStreaming"), TEXT("STATCAT_Engine"), FText::GetEmpty(), &UEngine::RenderStatAudioStreaming, &UEngine::ToggleStatAudioStreaming));
 	EngineStats.Add(FEngineStatFuncs(TEXT("STAT_SoundReverb"), TEXT("STATCAT_Engine"), FText::GetEmpty(), &UEngine::RenderStatSoundReverb, nullptr));
+	EngineStats.Add(FEngineStatFuncs(TEXT("STAT_Sounds"), TEXT("STATCAT_Engine"), FText::GetEmpty(), &UEngine::RenderStatSounds, &UEngine::ToggleStatSounds));
+	EngineStats.Add(FEngineStatFuncs(TEXT("STAT_SoundCues"), TEXT("STATCAT_Engine"), FText::GetEmpty(), &UEngine::RenderStatSoundCues, &UEngine::ToggleStatSoundCues));
 	EngineStats.Add(FEngineStatFuncs(TEXT("STAT_SoundWaves"), TEXT("STATCAT_Engine"), FText::GetEmpty(), &UEngine::RenderStatSoundWaves, &UEngine::ToggleStatSoundWaves));
 #endif // !UE_BUILD_SHIPPING
 	/* @todo UE4 physx fix this once we have convexelem drawing again
