Fix RBAN body scale for PhysX skeletal mesh physics initialization

diff --git a/Engine/Source/Runtime/AnimGraphRuntime/Private/BoneControllers/AnimNode_RigidBody.cpp b/Engine/Source/Runtime/AnimGraphRuntime/Private/BoneControllers/AnimNode_RigidBody.cpp
--- a/Engine/Source/Runtime/AnimGraphRuntime/Private/BoneControllers/AnimNode_RigidBody.cpp
+++ b/Engine/Source/Runtime/AnimGraphRuntime/Private/BoneControllers/AnimNode_RigidBody.cpp
@@ -658,15 +658,22 @@ void FAnimNode_RigidBody::InitPhysics(const UAnimInstance* InAnimInstance)
 		TotalMass = 0.f;
 
 		// Instantiate a FBodyInstance/FConstraintInstance set that will be cloned into the Immediate Physics sim.
+		// NOTE: We do not have a skeleton at the moment, so we have to use the ref pose
 		TArray<FBodyInstance*> HighLevelBodyInstances;
 		TArray<FConstraintInstance*> HighLevelConstraintInstances;
-#if WITH_CHAOS
+
 		// Chaos relies on the initial pose to set up constraint positions
-		SkeletalMeshComp->InstantiatePhysicsAssetRefPose(*UsePhysicsAsset, SimulationSpace == ESimulationSpace::WorldSpace ? SkeletalMeshComp->GetComponentToWorld().GetScale3D() : FVector(1.f), HighLevelBodyInstances, HighLevelConstraintInstances);
-#else
-		// PhysX sets up constraints each tick
-		SkeletalMeshComp->InstantiatePhysicsAsset(*UsePhysicsAsset, SimulationSpace == ESimulationSpace::WorldSpace ? SkeletalMeshComp->GetComponentToWorld().GetScale3D() : FVector(1.f), HighLevelBodyInstances, HighLevelConstraintInstances);
-#endif
+		bool bCreateBodiesInRefPose = (WITH_CHAOS != 0);
+		SkeletalMeshComp->InstantiatePhysicsAssetRefPose(
+			*UsePhysicsAsset, 
+			SimulationSpace == ESimulationSpace::WorldSpace ? SkeletalMeshComp->GetComponentToWorld().GetScale3D() : FVector(1.f), 
+			HighLevelBodyInstances, 
+			HighLevelConstraintInstances, 
+			nullptr, 
+			nullptr, 
+			INDEX_NONE, 
+			FPhysicsAggregateHandle(),
+			bCreateBodiesInRefPose);
 
 		TMap<FName, ImmediatePhysics::FActorHandle*> NamesToHandles;
 		TArray<ImmediatePhysics::FActorHandle*> IgnoreCollisionActors;
diff --git a/Engine/Source/Runtime/Engine/Classes/Components/SkeletalMeshComponent.h b/Engine/Source/Runtime/Engine/Classes/Components/SkeletalMeshComponent.h
--- a/Engine/Source/Runtime/Engine/Classes/Components/SkeletalMeshComponent.h
+++ b/Engine/Source/Runtime/Engine/Classes/Components/SkeletalMeshComponent.h
@@ -1754,7 +1754,7 @@ public:
 	void InstantiatePhysicsAsset(const UPhysicsAsset& PhysAsset, const FVector& Scale3D, TArray<FBodyInstance*>& OutBodies, TArray<FConstraintInstance*>& OutConstraints, FPhysScene* PhysScene = nullptr, USkeletalMeshComponent* OwningComponent = nullptr, int32 UseRootBodyIndex = INDEX_NONE, const FPhysicsAggregateHandle& UseAggregate = FPhysicsAggregateHandle()) const;
 
 	/** Instantiates bodies given a physics asset like InstantiatePhysicsAsset but instead of reading the current component state, this reads the ref-pose from the reference skeleton of the mesh. Useful if trying to create bodies to be used during any evaluation work */
-	void InstantiatePhysicsAssetRefPose(const UPhysicsAsset& PhysAsset, const FVector& Scale3D, TArray<FBodyInstance*>& OutBodies, TArray<FConstraintInstance*>& OutConstraints, FPhysScene* PhysScene = nullptr, USkeletalMeshComponent* OwningComponent = nullptr, int32 UseRootBodyIndex = INDEX_NONE, const FPhysicsAggregateHandle& UseAggregate = FPhysicsAggregateHandle()) const;
+	void InstantiatePhysicsAssetRefPose(const UPhysicsAsset& PhysAsset, const FVector& Scale3D, TArray<FBodyInstance*>& OutBodies, TArray<FConstraintInstance*>& OutConstraints, FPhysScene* PhysScene = nullptr, USkeletalMeshComponent* OwningComponent = nullptr, int32 UseRootBodyIndex = INDEX_NONE, const FPhysicsAggregateHandle& UseAggregate = FPhysicsAggregateHandle(), bool bCreateBodiesInRefPose = false) const;
 
 	/** Turn off all physics and remove the instance. */
 	void TermArticulated();
diff --git a/Engine/Source/Runtime/Engine/Private/SkeletalMeshComponentPhysics.cpp b/Engine/Source/Runtime/Engine/Private/SkeletalMeshComponentPhysics.cpp
--- a/Engine/Source/Runtime/Engine/Private/SkeletalMeshComponentPhysics.cpp
+++ b/Engine/Source/Runtime/Engine/Private/SkeletalMeshComponentPhysics.cpp
@@ -634,23 +634,41 @@ void USkeletalMeshComponent::InstantiatePhysicsAsset(const UPhysicsAsset& PhysAs
 	InstantiatePhysicsAsset_Internal(PhysAsset, Scale3D, OutBodies, OutConstraints, BoneTMCallable, PhysScene, OwningComponent, UseRootBodyIndex, UseAggregate);
 }
 
-void USkeletalMeshComponent::InstantiatePhysicsAssetRefPose(const UPhysicsAsset& PhysAsset, const FVector& Scale3D, TArray<FBodyInstance*>& OutBodies, TArray<FConstraintInstance*>& OutConstraints, FPhysScene* PhysScene /*= nullptr*/, USkeletalMeshComponent* OwningComponent /*= nullptr*/, int32 UseRootBodyIndex /*= INDEX_NONE*/, const FPhysicsAggregateHandle& UseAggregate) const
+void USkeletalMeshComponent::InstantiatePhysicsAssetRefPose(const UPhysicsAsset& PhysAsset, const FVector& Scale3D, TArray<FBodyInstance*>& OutBodies, TArray<FConstraintInstance*>& OutConstraints, FPhysScene* PhysScene /*= nullptr*/, USkeletalMeshComponent* OwningComponent /*= nullptr*/, int32 UseRootBodyIndex /*= INDEX_NONE*/, const FPhysicsAggregateHandle& UseAggregate, bool bCreateBodiesInRefPose) const
 {
 	if(SkeletalMesh)
 	{
-		// @todo(ccaulfield): check perf here against uncached version: FAnimationRuntime::GetComponentSpaceTransform for RBAN scenarios
 		const FReferenceSkeleton& RefSkeleton = SkeletalMesh->RefSkeleton;
-		TArray<FTransform> CachedTransforms;
-		TArray<bool> CachedTransformReady;
-		CachedTransforms.SetNumUninitialized(RefSkeleton.GetRefBonePose().Num());
-		CachedTransformReady.SetNumZeroed(RefSkeleton.GetRefBonePose().Num());
 
-		auto BoneTMCallable = [&](int32 BoneIndex)
+		if (bCreateBodiesInRefPose)
 		{
-			return FAnimationRuntime::GetComponentSpaceTransformWithCache(RefSkeleton, RefSkeleton.GetRefBonePose(), BoneIndex, CachedTransforms, CachedTransformReady);
-		};
+			// Create the bodies posed in component-space with the ref pose
+			TArray<FTransform> CachedTransforms;
+			TArray<bool> CachedTransformReady;
+			CachedTransforms.SetNumUninitialized(RefSkeleton.GetRefBonePose().Num());
+			CachedTransformReady.SetNumZeroed(RefSkeleton.GetRefBonePose().Num());
 
-		InstantiatePhysicsAsset_Internal(PhysAsset, Scale3D, OutBodies, OutConstraints, BoneTMCallable, PhysScene, OwningComponent, UseRootBodyIndex, UseAggregate);
+			auto BoneTMCallable = [&](int32 BoneIndex)
+			{
+				return FAnimationRuntime::GetComponentSpaceTransformWithCache(RefSkeleton, RefSkeleton.GetRefBonePose(), BoneIndex, CachedTransforms, CachedTransformReady);
+			};
+
+			InstantiatePhysicsAsset_Internal(PhysAsset, Scale3D, OutBodies, OutConstraints, BoneTMCallable, PhysScene, OwningComponent, UseRootBodyIndex, UseAggregate);
+		}
+		else
+		{
+			// Create the bodies with each body in its local ref pose (all bodies will be on top of each other close to the origin, but should still have the correct scale)
+			auto BoneTMCallable = [&](int32 BoneIndex)
+			{
+				if (RefSkeleton.IsValidIndex(BoneIndex))
+				{
+					return RefSkeleton.GetRefBonePose()[BoneIndex];
+				}
+				return FTransform::Identity;
+			};;
+
+			InstantiatePhysicsAsset_Internal(PhysAsset, Scale3D, OutBodies, OutConstraints, BoneTMCallable, PhysScene, OwningComponent, UseRootBodyIndex, UseAggregate);
+		}
 	}
 }
 
