Chaos: Add support for SQCaptures to re-measure SQ and use average of multiple samples of duration before serializing. This prevents context switches and other anomalies from giving useless captures.

diff --git a/Engine/Source/Runtime/Engine/Private/Collision/SceneQueryLowLevel.cpp b/Engine/Source/Runtime/Engine/Private/Collision/SceneQueryLowLevel.cpp
--- a/Engine/Source/Runtime/Engine/Private/Collision/SceneQueryLowLevel.cpp
+++ b/Engine/Source/Runtime/Engine/Private/Collision/SceneQueryLowLevel.cpp
@@ -26,10 +26,12 @@ FAutoConsoleVariableRef CVarForceStandardSQ(TEXT("p.ForceStandardSQ"), ForceStan
 
 #if !UE_BUILD_SHIPPING
 int32 SerializeSQs = 0;
+int32 SerializeSQSamples = 100;
 int32 SerializeBadSQs = 0;
 int32 ReplaySQs = 0;
 
-FAutoConsoleVariableRef CVarSerializeSQs(TEXT("p.SerializeSQs"), SerializeSQs, TEXT("If enabled, we create a sq capture per sq. This can be very expensive as the entire scene is saved out"));
+FAutoConsoleVariableRef CVarSerializeSQs(TEXT("p.SerializeSQs"), SerializeSQs, TEXT("If enabled, we create a sq capture per sq that takes more than provided value in microseconds. This can be very expensive as the entire scene is saved out"));
+FAutoConsoleVariableRef CVarSerializeSQSamples(TEXT("p.SerializeSQSampleCount"), SerializeSQSamples, TEXT("If Query exceeds duration threshold, we will re-measure SQ this many times before serializing. Larger values cause hitching."));
 FAutoConsoleVariableRef CVarReplaySweeps(TEXT("p.ReplaySQs"), ReplaySQs, TEXT("If enabled, we rerun the sq against chaos"));
 FAutoConsoleVariableRef CVarSerializeBadSweeps(TEXT("p.SerializeBadSQs"), SerializeBadSQs, TEXT("If enabled, we create a sq capture whenever chaos and physx diverge"));
 
@@ -57,10 +59,116 @@ void FinalizeCapture(FPhysTestSerializer& Serializer)
 #else
 constexpr int32 SerializeSQs = 0;
 constexpr int32 ReplaySQs = 0;
+constexpr int32 SerializeSQSamples = 0;
 // No-op in shipping
 void FinalizeCapture(FPhysTestSerializer& Serializer) {}
 #endif
 
+namespace
+{
+	void SweepSQCaptureHelper(float QueryDurationSeconds, const FChaosSQAccelerator& SQAccelerator, FPhysScene& Scene, const FPhysicsGeometry& QueryGeom, const FTransform& StartTM, const FVector& Dir, float DeltaMag, const FPhysicsHitCallback<FHitSweep>& HitBuffer, EHitFlags OutputFlags, FQueryFlags QueryFlags, const FCollisionFilterData& Filter, const FQueryFilterData& QueryFilterData, ICollisionQueryFilterCallbackBase* QueryCallback, const FQueryDebugParams& DebugParams)
+	{
+		float QueryDurationMicro = QueryDurationSeconds * 1000.0 * 1000.0;
+		if (((SerializeSQs && QueryDurationMicro > SerializeSQs)) && IsInGameThread())
+		{
+			// Measure average time of query over multiple samples to reduce fluke from context switches or that kind of thing.
+			uint32 Cycles = 0.0;
+			const uint32 SampleCount = SerializeSQSamples;
+			for (uint32 Samples = 0; Samples < SampleCount; ++Samples)
+			{
+				// Reset output to not skew times with large buffer
+				FPhysicsHitCallback<FHitSweep> ScratchHitBuffer = FPhysicsHitCallback<FHitSweep>(HitBuffer.WantsSingleResult());
+
+				uint32 StartTime = FPlatformTime::Cycles();
+				SQAccelerator.Sweep(QueryGeom, StartTM, Dir, DeltaMag, ScratchHitBuffer, OutputFlags, QueryFilterData, *QueryCallback, DebugParams);
+				Cycles += FPlatformTime::Cycles() - StartTime;
+			}
+
+			float Milliseconds = FPlatformTime::ToMilliseconds(Cycles);
+			float AvgMicroseconds = (Milliseconds * 1000) / SampleCount;
+
+			if (AvgMicroseconds > SerializeSQs)
+			{
+				FPhysTestSerializer Serializer;
+				Serializer.SetPhysicsData(*Scene.GetSolver()->GetEvolution());
+				FSQCapture& SweepCapture = Serializer.CaptureSQ();
+				SweepCapture.StartCaptureChaosSweep(*Scene.GetSolver()->GetEvolution(), QueryGeom, StartTM, Dir, DeltaMag, OutputFlags, QueryFilterData, Filter, *QueryCallback);
+				SweepCapture.EndCaptureChaosSweep(HitBuffer);
+
+				FinalizeCapture(Serializer);
+			}
+		}
+	}
+
+	void RaycastSQCaptureHelper(float QueryDurationSeconds, const FChaosSQAccelerator& SQAccelerator, FPhysScene& Scene, const FVector& Start, const FVector& Dir, float DeltaMag, const FPhysicsHitCallback<FHitRaycast>& HitBuffer, EHitFlags OutputFlags, FQueryFlags QueryFlags, const FCollisionFilterData& Filter, const FQueryFilterData& QueryFilterData, ICollisionQueryFilterCallbackBase* QueryCallback, const FQueryDebugParams& DebugParams)
+	{
+		float QueryDurationMicro = QueryDurationSeconds * 1000.0 * 1000.0;
+		if (((SerializeSQs && QueryDurationMicro > SerializeSQs)) && IsInGameThread())
+		{
+			// Measure average time of query over multiple samples to reduce fluke from context switches or that kind of thing.
+			uint32 Cycles = 0.0;
+			const uint32 SampleCount = SerializeSQSamples;
+			for (uint32 Samples = 0; Samples < SampleCount; ++Samples)
+			{
+				// Reset output to not skew times with large buffer
+				FPhysicsHitCallback<FHitRaycast> ScratchHitBuffer = FPhysicsHitCallback<FHitRaycast>(HitBuffer.WantsSingleResult());
+
+				uint32 StartTime = FPlatformTime::Cycles();
+				SQAccelerator.Raycast(Start, Dir, DeltaMag, ScratchHitBuffer, OutputFlags, QueryFilterData, *QueryCallback, DebugParams);
+				Cycles += FPlatformTime::Cycles() - StartTime;
+			}
+
+			float Milliseconds = FPlatformTime::ToMilliseconds(Cycles);
+			float AvgMicroseconds = (Milliseconds * 1000) / SampleCount;
+
+			if (AvgMicroseconds > SerializeSQs)
+			{
+				FPhysTestSerializer Serializer;
+				Serializer.SetPhysicsData(*Scene.GetSolver()->GetEvolution());
+				FSQCapture& RaycastCapture = Serializer.CaptureSQ();
+				RaycastCapture.StartCaptureChaosRaycast(*Scene.GetSolver()->GetEvolution(), Start, Dir, DeltaMag, OutputFlags, QueryFilterData, Filter, *QueryCallback);
+				RaycastCapture.EndCaptureChaosRaycast(HitBuffer);
+
+				FinalizeCapture(Serializer);
+			}
+		}
+	}
+
+	void OverlapSQCaptureHelper(float QueryDurationSeconds, const FChaosSQAccelerator& SQAccelerator, FPhysScene& Scene, const FPhysicsGeometry& QueryGeom, const FTransform& GeomPose, const FPhysicsHitCallback<FHitOverlap>& HitBuffer, FQueryFlags QueryFlags, const FCollisionFilterData& Filter, const FQueryFilterData& QueryFilterData, ICollisionQueryFilterCallbackBase* QueryCallback, const FQueryDebugParams& DebugParams)
+	{
+		float QueryDurationMicro = QueryDurationSeconds * 1000.0 * 1000.0;
+		if (((SerializeSQs && QueryDurationMicro > SerializeSQs)) && IsInGameThread())
+		{
+			// Measure average time of query over multiple samples to reduce fluke from context switches or that kind of thing.
+			uint32 Cycles = 0.0;
+			const uint32 SampleCount = SerializeSQSamples;
+			for (uint32 Samples = 0; Samples < SampleCount; ++Samples)
+			{
+				// Reset output to not skew times with large buffer
+				FPhysicsHitCallback<FHitOverlap> ScratchHitBuffer = FPhysicsHitCallback<FHitOverlap>(HitBuffer.WantsSingleResult());
+
+				uint32 StartTime = FPlatformTime::Cycles();
+				SQAccelerator.Overlap(QueryGeom, GeomPose, ScratchHitBuffer, QueryFilterData, *QueryCallback);
+				Cycles += FPlatformTime::Cycles() - StartTime;
+			}
+
+			float Milliseconds = FPlatformTime::ToMilliseconds(Cycles);
+			float AvgMicroseconds = (Milliseconds * 1000) / SampleCount;
+
+			if (AvgMicroseconds > SerializeSQs)
+			{
+				FPhysTestSerializer Serializer;
+				Serializer.SetPhysicsData(*Scene.GetSolver()->GetEvolution());
+				FSQCapture& OverlapCapture = Serializer.CaptureSQ();
+				OverlapCapture.StartCaptureChaosOverlap(*Scene.GetSolver()->GetEvolution(), QueryGeom, GeomPose, QueryFilterData, Filter, *QueryCallback);
+				OverlapCapture.EndCaptureChaosOverlap(HitBuffer);
+
+				FinalizeCapture(Serializer);
+			}
+		}
+	}
+}
+
 void LowLevelRaycast(FPhysScene& Scene, const FVector& Start, const FVector& Dir, float DeltaMag, FPhysicsHitCallback<FHitRaycast>& HitBuffer, EHitFlags OutputFlags, FQueryFlags QueryFlags, const FCollisionFilterData& Filter, const FQueryFilterData& QueryFilterData, ICollisionQueryFilterCallbackBase* QueryCallback, const FQueryDebugParams& DebugParams)
 {
 #if !defined(PHYSICS_INTERFACE_PHYSX) || !PHYSICS_INTERFACE_PHYSX
@@ -73,16 +181,10 @@ void LowLevelRaycast(FPhysScene& Scene, const FVector& Start, const FVector& Dir
 			SQAccelerator.Raycast(Start, Dir, DeltaMag, HitBuffer, OutputFlags, QueryFilterData, *QueryCallback, DebugParams);
 		}
 
-		/*if (((SerializeSQs && Time * 1000.0 * 1000.0 > SerializeSQs) || (false)) && IsInGameThread())
+		if (SerializeSQs)
 		{
-			FPhysTestSerializer Serializer;
-			Serializer.SetPhysicsData(*Scene.GetSolver()->GetEvolution());
-			FSQCapture& RaycastCapture = Serializer.CaptureSQ();
-			RaycastCapture.StartCaptureChaosRaycast(*Scene.GetSolver()->GetEvolution(), Start, Dir, DeltaMag, OutputFlags, QueryFilterData, Filter, *QueryCallback);
-			RaycastCapture.EndCaptureChaosRaycast(HitBuffer);
-
-			FinalizeCapture(Serializer);
-		}*/
+			//RaycastSQCaptureHelper(QueryDurationSeconds, SQAccelerator, Scene, Start, Dir, DeltaMag, HitBuffer, OutputFlags, QueryFlags, Filter, QueryFilterData, *QueryCallback, DebugParams);
+		}
 	}
 #else
 	if (SerializeSQs | ReplaySQs)
@@ -116,40 +218,12 @@ void LowLevelSweep(FPhysScene& Scene, const FPhysicsGeometry& QueryGeom, const F
 				FScopedDurationTimer Timer(Time);
 				SQAccelerator.Sweep(QueryGeom, StartTM, Dir, DeltaMag, HitBuffer, OutputFlags, QueryFilterData, *QueryCallback, DebugParams);
 			}
-			
-			if (((SerializeSQs && Time * 1000.0 * 1000.0 > SerializeSQs) || (false)) && IsInGameThread())
-			{
-				FPhysTestSerializer Serializer;
-				Serializer.SetPhysicsData(*Scene.GetSolver()->GetEvolution());
-				FSQCapture& SweepCapture = Serializer.CaptureSQ();
-				SweepCapture.StartCaptureChaosSweep(*Scene.GetSolver()->GetEvolution(), QueryGeom, StartTM, Dir, DeltaMag, OutputFlags, QueryFilterData, Filter, *QueryCallback);
-				SweepCapture.EndCaptureChaosSweep(HitBuffer);
-
-				FinalizeCapture(Serializer);
-			}
-		}
 
-		/*if (SerializeSQs && IsInGameThread())
-		{
-			FPhysTestSerializer Serializer;
-			Serializer.SetPhysicsData(*Scene.GetSolver()->GetEvolution());
-			FSQCapture& SweepCapture = Serializer.CaptureSQ();
-			SweepCapture.StartCaptureChaosSweep(*Scene.GetSolver()->GetEvolution(), QueryGeom, StartTM, Dir, DeltaMag, OutputFlags, QueryFilterData, Filter, *QueryCallback);
-			SQAccelerator.Sweep(QueryGeom, StartTM, Dir, DeltaMag, HitBuffer, OutputFlags, QueryFilterData, *QueryCallback);
-			SweepCapture.EndCaptureChaosSweep(HitBuffer);
-
-			FinalizeCapture(Serializer);
-		}
-		else
-		{
-			//ISQAccelerator* SQAccelerator = Scene.GetSQAccelerator();
-			double Time = 0.0;
+			if (SerializeSQs)
 			{
-				FScopedDurationTimer Timer(Time);
-				SQAccelerator.Sweep(QueryGeom, StartTM, Dir, DeltaMag, HitBuffer, OutputFlags, QueryFilterData, *QueryCallback);
+				SweepSQCaptureHelper(Time, SQAccelerator, Scene, QueryGeom, StartTM, Dir, DeltaMag, HitBuffer, OutputFlags, QueryFlags, Filter, QueryFilterData, QueryCallback, DebugParams);
 			}
-			if(Time > )
-		}*/
+		}
 	}
 #else
 	if (SerializeSQs | ReplaySQs)
@@ -176,21 +250,15 @@ void LowLevelOverlap(FPhysScene& Scene, const FPhysicsGeometry& QueryGeom, const
 	if (const auto& SolverAccelerationStructure = Scene.GetScene().GetSpacialAcceleration())
 	{
 		FChaosSQAccelerator SQAccelerator(*SolverAccelerationStructure);
-		if (false && SerializeSQs && IsInGameThread())
+		double Time = 0.0;
 		{
-			FPhysTestSerializer Serializer;
-			Serializer.SetPhysicsData(*Scene.GetSolver()->GetEvolution());
-			FSQCapture& OverlapCapture = Serializer.CaptureSQ();
-			OverlapCapture.StartCaptureChaosOverlap(*Scene.GetSolver()->GetEvolution(), QueryGeom, GeomPose, QueryFilterData, Filter, *QueryCallback);
+			FScopedDurationTimer Timer(Time);
 			SQAccelerator.Overlap(QueryGeom, GeomPose, HitBuffer, QueryFilterData, *QueryCallback);
-			OverlapCapture.EndCaptureChaosOverlap(HitBuffer);
-
-			FinalizeCapture(Serializer);
 		}
-		else
+
+		if (SerializeSQs)
 		{
-			//ISQAccelerator* SQAccelerator = Scene.GetSQAccelerator();
-			SQAccelerator.Overlap(QueryGeom, GeomPose, HitBuffer, QueryFilterData, *QueryCallback);
+			//OverlapSQCaptureHelper(Time, SQAccelerator, Scene, QueryGeom, GeomPose, HitBuffer, QueryFlags, Filter, QueryFilterData, QueryCallback, DebugParams);
 		}
 	}
 #else
