Fixed crash when force-deleting an in-use notify state blueprint

diff --git a/Engine/Source/Runtime/Engine/Private/Animation/AnimInstance.cpp b/Engine/Source/Runtime/Engine/Private/Animation/AnimInstance.cpp
--- a/Engine/Source/Runtime/Engine/Private/Animation/AnimInstance.cpp
+++ b/Engine/Source/Runtime/Engine/Private/Animation/AnimInstance.cpp
@@ -1309,7 +1309,7 @@ void UAnimInstance::TriggerAnimNotifies(float DeltaSeconds)
 	// Send end notification to AnimNotifyState not active anymore.
 	for(const FAnimNotifyEvent& AnimNotifyEvent : ActiveAnimNotifyState)
 	{
-		if (ShouldTriggerAnimNotifyState(AnimNotifyEvent.NotifyStateClass))
+		if (AnimNotifyEvent.NotifyStateClass && ShouldTriggerAnimNotifyState(AnimNotifyEvent.NotifyStateClass))
 		{
 			AnimNotifyEvent.NotifyStateClass->NotifyEnd(SkelMeshComp, Cast<UAnimSequenceBase>(AnimNotifyEvent.NotifyStateClass->GetOuter()));
 		}
