If the world is gone when destructing a preview scene we now gracefully ignore it.

diff --git a/Engine/Source/Runtime/Engine/Private/PreviewScene.cpp b/Engine/Source/Runtime/Engine/Private/PreviewScene.cpp
--- a/Engine/Source/Runtime/Engine/Private/PreviewScene.cpp
+++ b/Engine/Source/Runtime/Engine/Private/PreviewScene.cpp
@@ -123,12 +123,14 @@ FPreviewScene::~FPreviewScene()
 		Component->UnregisterComponent();
 	}
 	
-	PreviewWorld->CleanupWorld();
-	GEngine->DestroyWorldContext(GetWorld());
-
-	// Release PhysicsScene for fixing big fbx importing bug
-	PreviewWorld->ReleasePhysicsScene();
-
+	// The world may be released by now.
+	if (PreviewWorld)
+	{
+		// Release PhysicsScene for fixing big fbx importing bug
+		PreviewWorld->ReleasePhysicsScene();
+		PreviewWorld->CleanupWorld();
+		GEngine->DestroyWorldContext(GetWorld());
+	}
 }
 
 void FPreviewScene::AddComponent(UActorComponent* Component,const FTransform& LocalToWorld, bool bAttachToRoot /*= false*/)
diff --git a/Engine/Source/Runtime/Engine/Public/PreviewScene.h b/Engine/Source/Runtime/Engine/Public/PreviewScene.h
--- a/Engine/Source/Runtime/Engine/Public/PreviewScene.h
+++ b/Engine/Source/Runtime/Engine/Public/PreviewScene.h
@@ -117,8 +117,8 @@ private:
 	TArray<class UActorComponent*> Components;
 
 protected:
-	class UWorld* PreviewWorld;
-	class ULineBatchComponent* LineBatcher;
+	class UWorld* PreviewWorld = nullptr;
+	class ULineBatchComponent* LineBatcher = nullptr;
 
 	/** This controls whether or not all mip levels of textures used by UMeshComponents added to this preview window should be loaded and remain loaded. */
 	bool bForceAllUsedMipsResident;
