Added a property to allow switching around the evaluation order of linked animations and the component's animation. Only controllable via property editing.

Aaron.Cox, Thomas.Sarkanen

diff --git a/Engine/Source/Runtime/Engine/Classes/Components/SkeletalMeshComponent.h b/Engine/Source/Runtime/Engine/Classes/Components/SkeletalMeshComponent.h
--- a/Engine/Source/Runtime/Engine/Classes/Components/SkeletalMeshComponent.h
+++ b/Engine/Source/Runtime/Engine/Classes/Components/SkeletalMeshComponent.h
@@ -170,6 +170,14 @@ namespace EAnimationMode
 	};
 }
 
+UENUM()
+enum class ELinkedAnimationUpdateOrder : uint8
+{
+	UpdateAnimationBeforeAnimScriptInstance,
+	UpdateAnimationAfterAnimScriptInstance,
+	DoNotUpdate
+};
+
 UENUM()
 namespace EPhysicsTransformUpdateMode
 {
@@ -419,7 +427,7 @@ protected:
 	/** Whether to use Animation Blueprint or play Single Animation Asset. */
 	UPROPERTY(EditAnywhere, BlueprintReadOnly, Category=Animation)
 	TEnumAsByte<EAnimationMode::Type>	AnimationMode;
-
+	
 private:
 	/** Teleport type to use on the next update */
 	ETeleportType PendingTeleportType;
@@ -431,6 +439,9 @@ private:
 	uint8 bDisablePostProcessBlueprint:1;
 
 public:
+	/** The order in which linked animation is evaluated with respect to the component's animation */
+	UPROPERTY(EditAnywhere, BlueprintReadOnly, Category = Animation)
+	ELinkedAnimationUpdateOrder LinkedAnimationEvaluationOrder;
 
 	/** Indicates that simulation (if it's enabled) is entirely responsible for children transforms. This is only ok if you are not animating attachment points relative to the simulation */
 	uint8 bSimulationUpdatesChildTransforms:1;
diff --git a/Engine/Source/Runtime/Engine/Private/Components/SkeletalMeshComponent.cpp b/Engine/Source/Runtime/Engine/Private/Components/SkeletalMeshComponent.cpp
--- a/Engine/Source/Runtime/Engine/Private/Components/SkeletalMeshComponent.cpp
+++ b/Engine/Source/Runtime/Engine/Private/Components/SkeletalMeshComponent.cpp
@@ -1014,10 +1014,13 @@ void USkeletalMeshComponent::TickAnimation(float DeltaTime, bool bNeedsValidRoot
 
 		// We update linked instances first incase we're using either root motion or non-threaded update.
 		// This ensures that we go through the pre update process and initialize the proxies correctly.
-		for(UAnimInstance* LinkedInstance : LinkedInstances)
+		if (LinkedAnimationEvaluationOrder == ELinkedAnimationUpdateOrder::UpdateAnimationBeforeAnimScriptInstance)
 		{
-			// Sub anim instances are always forced to do a parallel update 
-			LinkedInstance->UpdateAnimation(DeltaTime * GlobalAnimRateScale, false, UAnimInstance::EUpdateAnimationFlag::ForceParallelUpdate);
+			for (UAnimInstance* LinkedInstance : LinkedInstances)
+			{
+				// Sub anim instances are always forced to do a parallel update 
+				LinkedInstance->UpdateAnimation(DeltaTime * GlobalAnimRateScale, false, UAnimInstance::EUpdateAnimationFlag::ForceParallelUpdate);
+			}
 		}
 
 		if (AnimScriptInstance != nullptr)
@@ -1026,6 +1029,16 @@ void USkeletalMeshComponent::TickAnimation(float DeltaTime, bool bNeedsValidRoot
 			AnimScriptInstance->UpdateAnimation(DeltaTime * GlobalAnimRateScale, bNeedsValidRootMotion);
 		}
 
+		if (LinkedAnimationEvaluationOrder == ELinkedAnimationUpdateOrder::UpdateAnimationAfterAnimScriptInstance)
+		{
+			for (UAnimInstance* LinkedInstance : LinkedInstances)
+			{
+				// Sub anim instances are always forced to do a parallel update 
+				LinkedInstance->UpdateAnimation(DeltaTime * GlobalAnimRateScale, false, UAnimInstance::EUpdateAnimationFlag::ForceParallelUpdate);
+			}
+		}
+
+
 		if(ShouldUpdatePostProcessInstance())
 		{
 			PostProcessAnimInstance->UpdateAnimation(DeltaTime * GlobalAnimRateScale, false);
