# Fix for Server crash in UAnimInstance::TriggerAnimNotifies.
#rnx

diff --git a/Engine/Source/Runtime/Engine/Classes/Animation/AnimInstance.h b/Engine/Source/Runtime/Engine/Classes/Animation/AnimInstance.h
--- a/Engine/Source/Runtime/Engine/Classes/Animation/AnimInstance.h
+++ b/Engine/Source/Runtime/Engine/Classes/Animation/AnimInstance.h
@@ -1401,8 +1401,8 @@ protected:
 
 	friend struct FAnimNode_LinkedAnimGraph;
 	
-	/** Return whethere this AnimNotifyState should be triggered */
-	virtual bool ShouldTriggerAnimNotifyState(const UAnimNotifyState* AnimNotifyState) const { return true; }
+	/** Return whether this AnimNotifyState should be triggered */
+	virtual bool ShouldTriggerAnimNotifyState(const UAnimNotifyState* AnimNotifyState) const;
 
 protected:
 	/** Proxy object, nothing should access this from an externally-callable API as it is used as a scratch area on worker threads */
diff --git a/Engine/Source/Runtime/Engine/Private/Animation/AnimInstance.cpp b/Engine/Source/Runtime/Engine/Private/Animation/AnimInstance.cpp
--- a/Engine/Source/Runtime/Engine/Private/Animation/AnimInstance.cpp
+++ b/Engine/Source/Runtime/Engine/Private/Animation/AnimInstance.cpp
@@ -3147,6 +3147,15 @@ void UAnimInstance::DestroyAnimInstanceProxy(FAnimInstanceProxy* InProxy)
 	delete InProxy;
 }
 
+bool UAnimInstance::ShouldTriggerAnimNotifyState(const UAnimNotifyState* AnimNotifyState) const
+{
+	if (ensureAlwaysMsgf(AnimNotifyState != nullptr, TEXT("UAnimInstance::ShouldTriggerAnimNotifyState: AnimNotifyState is null on AnimInstance %s. ActiveAnimNotifyState array size is: %d"), *GetNameSafe(this), ActiveAnimNotifyState.Num()))
+	{
+		return true;
+	}
+	return false;
+}
+
 void UAnimInstance::RecordMachineWeight(const int32 InMachineClassIndex, const float InMachineWeight)
 {
 	GetProxyOnAnyThread<FAnimInstanceProxy>().RecordMachineWeight(InMachineClassIndex, InMachineWeight);
