#ue4 - Fix crouching clients observed from a listen server popping up briefly before interpolation corrects the location down. Apply the same fix as for simulated proxies on the listen server.

#udn https://udn.unrealengine.com/questions/536850/client-crouch-cause-mesh-offset-on-server.html

diff --git a/Engine/Source/Runtime/Engine/Private/Components/CharacterMovementComponent.cpp b/Engine/Source/Runtime/Engine/Private/Components/CharacterMovementComponent.cpp
--- a/Engine/Source/Runtime/Engine/Private/Components/CharacterMovementComponent.cpp
+++ b/Engine/Source/Runtime/Engine/Private/Components/CharacterMovementComponent.cpp
@@ -2662,10 +2662,10 @@ void UCharacterMovementComponent::Crouch(bool bClientSimulation)
 	CharacterOwner->OnStartCrouch( HalfHeightAdjust, ScaledHalfHeightAdjust );
 
 	// Don't smooth this change in mesh position
-	if (bClientSimulation && CharacterOwner->GetLocalRole() == ROLE_SimulatedProxy)
+	if ((bClientSimulation && CharacterOwner->GetLocalRole() == ROLE_SimulatedProxy) || (IsNetMode(NM_ListenServer) && CharacterOwner->GetRemoteRole() == ROLE_AutonomousProxy))
 	{
 		FNetworkPredictionData_Client_Character* ClientData = GetPredictionData_Client_Character();
-		if (ClientData && ClientData->MeshTranslationOffset.Z != 0.f)
+		if (ClientData)
 		{
 			ClientData->MeshTranslationOffset -= FVector(0.f, 0.f, MeshAdjust);
 			ClientData->OriginalMeshTranslationOffset = ClientData->MeshTranslationOffset;
@@ -2807,10 +2807,10 @@ void UCharacterMovementComponent::UnCrouch(bool bClientSimulation)
 	CharacterOwner->OnEndCrouch( HalfHeightAdjust, ScaledHalfHeightAdjust );
 
 	// Don't smooth this change in mesh position
-	if (bClientSimulation && CharacterOwner->GetLocalRole() == ROLE_SimulatedProxy)
+	if ((bClientSimulation && CharacterOwner->GetLocalRole() == ROLE_SimulatedProxy) || (IsNetMode(NM_ListenServer) && CharacterOwner->GetRemoteRole() == ROLE_AutonomousProxy))
 	{
 		FNetworkPredictionData_Client_Character* ClientData = GetPredictionData_Client_Character();
-		if (ClientData && ClientData->MeshTranslationOffset.Z != 0.f)
+		if (ClientData)
 		{
 			ClientData->MeshTranslationOffset += FVector(0.f, 0.f, MeshAdjust);
 			ClientData->OriginalMeshTranslationOffset = ClientData->MeshTranslationOffset;
