Adding support for an attenuation override position per listener.

diff --git a/Engine/Source/Runtime/Engine/Private/ActiveSound.cpp b/Engine/Source/Runtime/Engine/Private/ActiveSound.cpp
--- a/Engine/Source/Runtime/Engine/Private/ActiveSound.cpp
+++ b/Engine/Source/Runtime/Engine/Private/ActiveSound.cpp
@@ -102,7 +102,6 @@ FActiveSound::FActiveSound()
 #if ENABLE_AUDIO_DEBUG
 	, DebugColor(FColor::Black)
 #endif // !(UE_BUILD_SHIPPING || UE_BUILD_TEST)
-	, ClosestListenerPtr(nullptr)
 {
 	if (!ActiveSoundTraceDelegate.IsBound())
 	{
@@ -438,11 +437,6 @@ void FActiveSound::GetSoundSourceBusSends(EBusSendType BusSendType, TArray<FSoun
 	}
 }
 
-int32 FActiveSound::FindClosestListener( const TArray<FListener>& InListeners ) const
-{
-	return FAudioDevice::FindClosestListenerIndex(Transform, InListeners);
-}
-
 void FActiveSound::GetConcurrencyHandles(TArray<FConcurrencyHandle>& OutConcurrencyHandles) const
 {
 	OutConcurrencyHandles.Reset();
@@ -530,19 +524,11 @@ void FActiveSound::UpdateWaveInstances(TArray<FWaveInstance*> &InWaveInstances,
 	// we always pass the 'primary' listener (viewport 0) to the sound nodes and the underlying audio system
 	// then move the AudioComponent's CurrentLocation so that its position relative to that Listener is the same as its real position is relative to the closest Listener
 
-	int32 ClosestListenerIndex = 0;
-
-	const TArray<FListener>& Listeners = AudioDevice->GetListeners();
-
-	if (Listeners.Num() > 1)
 	{
 		SCOPE_CYCLE_COUNTER( STAT_AudioFindNearestLocation );
-		ClosestListenerIndex = FindClosestListener(Listeners);
+		ClosestListenerIndex = AudioDevice->FindClosestListenerIndex(Transform);
 	}
 
-	// Cache the closest listener ptr
-	ClosestListenerPtr = &Listeners[ClosestListenerIndex];
-
 	FSoundParseParameters ParseParams;
 	ParseParams.Transform = Transform;
 	ParseParams.StartTime = RequestedStartTime;
@@ -579,7 +565,7 @@ void FActiveSound::UpdateWaveInstances(TArray<FWaveInstance*> &InWaveInstances,
 	{
 		// Additional inside/outside processing for ambient sounds
 		// If we aren't in a world there is no interior volumes to be handled.
-		HandleInteriorVolumes(*ClosestListenerPtr, ParseParams);
+		HandleInteriorVolumes(ParseParams);
 	}
 
 	// for velocity-based effects like doppler
@@ -598,7 +584,7 @@ void FActiveSound::UpdateWaveInstances(TArray<FWaveInstance*> &InWaveInstances,
 		bool bReverbSendLevelWasSet = false;
 		if (bHasAttenuationSettings)
 		{
-			UpdateAttenuation(DeltaTime, ParseParams, *ClosestListenerPtr);
+			UpdateAttenuation(DeltaTime, ParseParams, ClosestListenerIndex);
 			bReverbSendLevelWasSet = true;
 		}
 		else
@@ -607,11 +593,11 @@ void FActiveSound::UpdateWaveInstances(TArray<FWaveInstance*> &InWaveInstances,
 			if (ParseParams.SoundClass)
 			{
 				ParseParams.ManualReverbSendLevel = ParseParams.SoundClass->Properties.Default2DReverbSendAmount;
-			}
-			else
-			{
-				ParseParams.ManualReverbSendLevel = AudioDevice->GetDefaultReverbSendLevel();
-			}
+		}
+		else
+		{
+			ParseParams.ManualReverbSendLevel = AudioDevice->GetDefaultReverbSendLevel();
+		}
 		}
 
 		ParseParams.ModulationPluginSettings = FindModulationSettings();
@@ -619,8 +605,12 @@ void FActiveSound::UpdateWaveInstances(TArray<FWaveInstance*> &InWaveInstances,
 		// if the closest listener is not the primary one, transform the sound transform so it's panned relative to primary listener position
 		if (ClosestListenerIndex != 0)
 		{
-			const FListener& Listener = Listeners[0];
-			ParseParams.Transform = ParseParams.Transform * ClosestListenerPtr->Transform.Inverse() * Listener.Transform;
+			FTransform FirstListenerTransform;
+			FTransform ClosestListenerTransform;
+			AudioDevice->GetListenerTransform(0, FirstListenerTransform);
+			AudioDevice->GetListenerTransform(ClosestListenerIndex, ClosestListenerTransform);
+
+			ParseParams.Transform = ParseParams.Transform * ClosestListenerTransform.Inverse() * FirstListenerTransform;
 		}
 
 		Sound->Parse(AudioDevice, 0, *this, ParseParams, ThisSoundsWaveInstances);
@@ -1102,7 +1092,7 @@ void FActiveSound::CheckOcclusion(const FVector ListenerLocation, const FVector
 	CurrentOcclusionVolumeAttenuation.Update(DeltaTime);
 }
 
-void FActiveSound::HandleInteriorVolumes( const FListener& Listener, FSoundParseParameters& ParseParams )
+void FActiveSound::HandleInteriorVolumes(FSoundParseParameters& ParseParams)
 {
 	// Get the settings of the ambient sound
 	if (!bGotInteriorSettings || (ParseParams.Transform.GetTranslation() - LastLocation).SizeSquared() > KINDA_SMALL_NUMBER)
@@ -1115,6 +1105,11 @@ void FActiveSound::HandleInteriorVolumes( const FListener& Listener, FSoundParse
 		bGotInteriorSettings = true;
 	}
 
+	check(IsInAudioThread());
+	const TArray<FListener>& Listeners = AudioDevice->GetListeners();
+	check(ClosestListenerIndex < Listeners.Num());
+	const FListener& Listener = Listeners[ClosestListenerIndex];
+
 	// Check to see if we've moved to a new audio volume
 	if (LastUpdateTime < Listener.InteriorStartTime)
 	{
@@ -1539,12 +1534,12 @@ void FActiveSound::UpdateFocusData(float DeltaTime, const FAttenuationListenerDa
 	FocusDataToUpdate->VolumeScale = ListenerData.AttenuationSettings->GetFocusAttenuation(FocusSettings, FocusDataToUpdate->FocusFactor);
 }
 
-void FActiveSound::ParseAttenuation(FSoundParseParameters& OutParseParams, const FListener& InListener, const FSoundAttenuationSettings& InAttenuationSettings)
+void FActiveSound::ParseAttenuation(FSoundParseParameters& OutParseParams, int32 ListenerIndex, const FSoundAttenuationSettings& InAttenuationSettings)
 {
-	UpdateAttenuation(0.0f, OutParseParams, InListener, &InAttenuationSettings);
+	UpdateAttenuation(0.0f, OutParseParams, ListenerIndex, &InAttenuationSettings);
 }
 
-void FActiveSound::UpdateAttenuation(float DeltaTime, FSoundParseParameters& ParseParams, const FListener& Listener, const FSoundAttenuationSettings* SettingsAttenuationNode)
+void FActiveSound::UpdateAttenuation(float DeltaTime, FSoundParseParameters& ParseParams, int32 ListenerIndex, const FSoundAttenuationSettings* SettingsAttenuationNode)
 {
 	// Get the attenuation settings to use for this application to the active sound
 	const FSoundAttenuationSettings* Settings = SettingsAttenuationNode ? SettingsAttenuationNode : &AttenuationSettings;
@@ -1562,7 +1557,7 @@ void FActiveSound::UpdateAttenuation(float DeltaTime, FSoundParseParameters& Par
 	}
 
 	check(AudioDevice);
-	FAttenuationListenerData ListenerData = FAttenuationListenerData::Create(*AudioDevice, Listener.Transform, ParseParams.Transform, *Settings);
+	FAttenuationListenerData ListenerData = FAttenuationListenerData::Create(*AudioDevice, ListenerIndex, ParseParams.Transform, *Settings);
 
 	// Apply priority attenuation if it's enabled
 	if (Settings->bEnablePriorityAttenuation)
@@ -1649,8 +1644,11 @@ void FActiveSound::UpdateAttenuation(float DeltaTime, FSoundParseParameters& Par
 		}
 		else if (ParseParams.Volume > 0.0f && !AudioDevice->IsAudioDeviceMuted())
 		{
-			check(ClosestListenerPtr);
-			CheckOcclusion(ClosestListenerPtr->Transform.GetTranslation(), ParseParams.Transform.GetTranslation(), Settings);
+			check(ClosestListenerIndex != INDEX_NONE);
+			FVector ListenerPosition;
+			const bool bAllowOverride = false;
+			AudioDevice->GetListenerPosition(ClosestListenerIndex, ListenerPosition, bAllowOverride);
+			CheckOcclusion(ListenerPosition, ParseParams.Transform.GetTranslation(), Settings);
 
 			// Apply the volume attenuation due to occlusion (using the interpolating dynamic parameter)
 			ParseParams.DistanceAttenuation *= CurrentOcclusionVolumeAttenuation.GetValue();
diff --git a/Engine/Source/Runtime/Engine/Private/AudioDebug.cpp b/Engine/Source/Runtime/Engine/Private/AudioDebug.cpp
--- a/Engine/Source/Runtime/Engine/Private/AudioDebug.cpp
+++ b/Engine/Source/Runtime/Engine/Private/AudioDebug.cpp
@@ -510,7 +510,7 @@ void FAudioDebugger::DrawDebugInfo(const FActiveSound& ActiveSound, const TArray
 #endif // ENABLE_DRAW_DEBUG
 }
 
-void FAudioDebugger::DrawDebugInfo(UWorld& World, const TArray<FListener>& Listeners, FVector& ListenerTransformOverride, bool bUseListenerTransformOverride)
+void FAudioDebugger::DrawDebugInfo(UWorld& World, const TArray<FListener>& Listeners)
 {
 #if ENABLE_DRAW_DEBUG
 	if (!ActiveSoundVisualizeListenersCVar)
@@ -518,6 +518,8 @@ void FAudioDebugger::DrawDebugInfo(UWorld& World, const TArray<FListener>& Liste
 		return;
 	}
 
+	check(IsInAudioThread());
+
 	FAudioDeviceManager* DeviceManager = GEngine->GetAudioDeviceManager();
 	if (DeviceManager && DeviceManager->IsVisualizeDebug3dEnabled())
 	{
@@ -525,10 +527,11 @@ void FAudioDebugger::DrawDebugInfo(UWorld& World, const TArray<FListener>& Liste
 		const TWeakObjectPtr<UWorld> WorldPtr = &World;
 		for (const FListener& Listener : Listeners)
 		{
-			const FVector ListenerPosition = bUseListenerTransformOverride ? ListenerTransformOverride : Listener.Transform.GetTranslation();
+			const FVector ListenerPosition = Listener.GetPosition(true);
 			const FVector ListenerFront = Listener.GetFront();
 			const FVector ListenerUp = Listener.GetUp();
 			const FVector ListenerRight = Listener.GetRight();
+
 			FAudioThread::RunCommandOnGameThread([WorldPtr, ListenerPosition, ListenerFront, ListenerUp, ListenerRight]()
 			{
 				if (WorldPtr.IsValid())
@@ -853,6 +856,8 @@ int32 FAudioDebugger::RenderStatReverb(UWorld* World, FViewport* Viewport, FCanv
 		return Y;
 	}
 
+	check(IsInGameThread());
+
 	const int32 Height = static_cast<int32>(GetStatsFont()->GetMaxCharHeight() + 2);
 
 	FString TheString;
@@ -863,8 +868,10 @@ int32 FAudioDebugger::RenderStatReverb(UWorld* World, FViewport* Viewport, FCanv
 		Y += Height;
 
 		AAudioVolume* CurrentAudioVolume = nullptr;
-		for (const FTransform& Transform : AudioDevice->ListenerTransforms)
+		const int32 ProxyCount = AudioDevice->ListenerProxies.Num();
+		for (int i = 0; i < ProxyCount; ++i)
 		{
+			const FTransform& Transform = AudioDevice->ListenerProxies[i].Transform;
 			AAudioVolume* PlayerAudioVolume = World->GetAudioSettings(Transform.GetLocation(), nullptr, nullptr);
 			if (PlayerAudioVolume && ((CurrentAudioVolume == nullptr) || (PlayerAudioVolume->GetPriority() > CurrentAudioVolume->GetPriority())))
 			{
@@ -1503,4 +1510,4 @@ void FAudioDebugger::ClearMutesAndSolos()
 	DebugNames.SoloSoundWave.Empty();
 }
 
-#endif // ENABLE_AUDIO_DEBUG
+#endif // ENABLE_AUDIO_DEBUG
\ No newline at end of file
diff --git a/Engine/Source/Runtime/Engine/Private/AudioDevice.cpp b/Engine/Source/Runtime/Engine/Private/AudioDevice.cpp
--- a/Engine/Source/Runtime/Engine/Private/AudioDevice.cpp
+++ b/Engine/Source/Runtime/Engine/Private/AudioDevice.cpp
@@ -168,12 +168,15 @@ namespace
 } // namespace <>
 
 
-FAttenuationListenerData FAttenuationListenerData::Create(const FAudioDevice& AudioDevice, const FTransform& InListenerTransform, const FTransform& InSoundTransform, const FSoundAttenuationSettings& InAttenuationSettings)
+FAttenuationListenerData FAttenuationListenerData::Create(const FAudioDevice& AudioDevice, int32 ListenerIndex, const FTransform& InSoundTransform, const FSoundAttenuationSettings& InAttenuationSettings)
 {
-	FAttenuationListenerData ListenerData(InListenerTransform, InSoundTransform, InAttenuationSettings);
+	FTransform ListenerTransform;
+	AudioDevice.GetListenerTransform(ListenerIndex, ListenerTransform);
+
+	FAttenuationListenerData ListenerData(ListenerTransform, InSoundTransform, InAttenuationSettings);
 
 	const FVector SoundTranslation = InSoundTransform.GetTranslation();
-	FVector ListenerToSound = SoundTranslation - InListenerTransform.GetTranslation();
+	FVector ListenerToSound = SoundTranslation - ListenerTransform.GetTranslation();
 	ListenerToSound.ToDirectionAndLength(ListenerData.ListenerToSoundDir, ListenerData.ListenerToSoundDistance);
 
 	// Store the actual distance for surround-panning sources with spread (AudioMixer)
@@ -181,12 +184,12 @@ FAttenuationListenerData FAttenuationListenerData::Create(const FAudioDevice& Au
 
 	// Calculating override listener-to-sound distance and transform must
 	// be applied after distance used for panning value is calculated.
-	if (AudioDevice.IsUsingListenerAttenuationOverride())
+	FVector ListenerPosition;
+	const bool bAllowAttenuationOverride = true;
+	if (AudioDevice.GetListenerPosition(ListenerIndex, ListenerPosition, bAllowAttenuationOverride))
 	{
-		const FVector& AttenuationOverride = AudioDevice.GetListenerAttenuationOverride();
-
-		ListenerData.ListenerToSoundDistance = (SoundTranslation - AttenuationOverride).Size();
-		ListenerData.ListenerTransform.SetTranslation(AttenuationOverride);
+		ListenerData.ListenerToSoundDistance = (SoundTranslation - ListenerPosition).Size();
+		ListenerData.ListenerTransform.SetTranslation(ListenerPosition);
 	}
 
 	const FSoundAttenuationSettings& AttenuationSettings = *ListenerData.AttenuationSettings;
@@ -228,7 +231,6 @@ FAudioDevice::FAudioDevice()
 	, GlobalPitchScale(1.0f)
 	, LastUpdateTime(FPlatformTime::Seconds())
 	, NextResourceID(1)
-	, bUseListenerAttenuationOverride(false)
 	, BaseSoundMix(nullptr)
 	, DefaultBaseSoundMix(nullptr)
 	, Effects(nullptr)
@@ -473,7 +475,7 @@ bool FAudioDevice::Init(int32 InMaxSources)
 
 	// Make sure the Listeners array has at least one entry, so we don't have to check for Listeners.Num() == 0 all the time
 	Listeners.Add(FListener(this));
-	ListenerTransforms.AddDefaulted();
+	ListenerProxies.AddDefaulted();
 	InverseListenerTransform.SetIdentity();
 
 	if (!bDeferStartupPrecache)
@@ -2823,6 +2825,16 @@ void FAudioDevice::VirtualizeInactiveLoops()
 	}
 }
 
+FVector FListener::GetPosition(bool bAllowOverride) const
+{
+	if (bAllowOverride && bUseAttenuationOverride)
+	{
+		return AttenuationOverride;
+	}
+
+	return Transform.GetTranslation();
+}
+
 float FListener::Interpolate(const double EndTime)
 {
 	if (FApp::GetCurrentTime() < InteriorStartTime)
@@ -2874,6 +2886,16 @@ void FListener::ApplyInteriorSettings(const uint32 InAudioVolumeID, const FInter
 	}
 }
 
+FVector FListenerProxy::GetPosition(bool bAllowOverride) const
+{
+	if (bAllowOverride && bUseAttenuationOverride)
+	{
+		return AttenuationOverride;
+	}
+
+	return Transform.GetTranslation();
+}
+
 void FAudioDevice::SetListener(UWorld* World, const int32 InViewportIndex, const FTransform& ListenerTransform, const float InDeltaSeconds)
 {
 	check(IsInGameThread());
@@ -2916,17 +2938,17 @@ void FAudioDevice::SetListener(UWorld* World, const int32 InViewportIndex, const
 		ListenerTransformCopy = FTransform::Identity;
 	}
 
-	if (InViewportIndex >= ListenerTransforms.Num())
+	if (InViewportIndex >= ListenerProxies.Num())
 	{
-		ListenerTransforms.AddDefaulted(InViewportIndex - ListenerTransforms.Num() + 1);
+		ListenerProxies.AddDefaulted(InViewportIndex - ListenerProxies.Num() + 1);
 	}
-	else if (ListenerTransforms[InViewportIndex].Equals(ListenerTransformCopy, KINDA_SMALL_NUMBER))
+	else if (ListenerProxies[InViewportIndex].Transform.Equals(ListenerTransformCopy, KINDA_SMALL_NUMBER))
 	{
-		// Don't need tp update hte listener if we're very close to our current transform
+		// Don't need tp update the listener if we're very close to our current transform
 		return;
 	}
 
-	ListenerTransforms[InViewportIndex] = ListenerTransformCopy;
+	ListenerProxies[InViewportIndex].Transform = ListenerTransformCopy;
 
 	DECLARE_CYCLE_STAT(TEXT("FAudioThreadTask.SetListener"), STAT_AudioSetListener, STATGROUP_AudioThreadCommands);
 
@@ -2983,43 +3005,86 @@ void FAudioDevice::SetListener(UWorld* World, const int32 InViewportIndex, const
 	}, GET_STATID(STAT_AudioSetListener));
 }
 
-void FAudioDevice::SetListenerAttenuationOverride(const FVector AttenuationPosition)
+void FAudioDevice::SetListenerAttenuationOverride(int32 ListenerIndex, const FVector AttenuationPosition)
 {
-	check(IsInGameThread());
-
-	DECLARE_CYCLE_STAT(TEXT("FAudioThreadTask.SetListenerAttenuationOverride"), STAT_AudioSetListenerAttenuationOverride, STATGROUP_AudioThreadCommands);
+	if (ListenerIndex == INDEX_NONE)
+	{
+		return;
+	}
 
-	FAudioDevice* AudioDevice = this;
-	FAudioThread::RunCommandOnAudioThread([AudioDevice, AttenuationPosition]()
+	if (!IsInAudioThread())
 	{
-		const bool bPrevAttenuationOverride = AudioDevice->bUseListenerAttenuationOverride;
-		AudioDevice->bUseListenerAttenuationOverride = true;
-		AudioDevice->ListenerAttenuationOverride = AttenuationPosition;
+		if (ListenerIndex >= ListenerProxies.Num())
+		{
+			return;
+		}
+
+		DECLARE_CYCLE_STAT(TEXT("FAudioThreadTask.SetListenerAttenuationOverride"), STAT_AudioSetListenerAttenuationOverride, STATGROUP_AudioThreadCommands);
 
-		if (!bPrevAttenuationOverride)
+		ListenerProxies[ListenerIndex].AttenuationOverride = AttenuationPosition;
+		ListenerProxies[ListenerIndex].bUseAttenuationOverride = true;
+
+		FAudioDevice* AudioDevice = this;
+		FAudioThread::RunCommandOnAudioThread([AudioDevice, ListenerIndex, AttenuationPosition]()
 		{
-			AudioDevice->UpdateVirtualLoops(true);
+			AudioDevice->SetListenerAttenuationOverride(ListenerIndex, AttenuationPosition);
+		}, GET_STATID(STAT_AudioSetListenerAttenuationOverride));
+	}
+	else
+	{
+		if (ensureMsgf(ListenerIndex < Listeners.Num(), TEXT("Listener Index %u out of range of available Listeners!"), ListenerIndex))
+		{
+			FListener& Listener = Listeners[ListenerIndex];
+			const bool bPrevAttenuationOverride = Listener.bUseAttenuationOverride;
+
+			Listener.bUseAttenuationOverride = true;
+			Listener.AttenuationOverride = AttenuationPosition;
+
+			if (!bPrevAttenuationOverride)
+			{
+				UpdateVirtualLoops(true);
+			}
 		}
-	}, GET_STATID(STAT_AudioSetListenerAttenuationOverride));
+	}
 }
 
-void FAudioDevice::ClearListenerAttenuationOverride()
+void FAudioDevice::ClearListenerAttenuationOverride(int32 ListenerIndex)
 {
-	check(IsInGameThread());
-
-	if (!bUseListenerAttenuationOverride)
+	if (ListenerIndex == INDEX_NONE)
 	{
 		return;
 	}
 
-	DECLARE_CYCLE_STAT(TEXT("FAudioThreadTask.ClearListenerAttenuationOverride"), STAT_AudioClearListenerAttenuationOverride, STATGROUP_AudioThreadCommands);
+	if (!IsInAudioThread())
+	{
+		if (ListenerIndex >= ListenerProxies.Num())
+		{
+			return;
+		}
 
-	FAudioDevice* AudioDevice = this;
-	FAudioThread::RunCommandOnAudioThread([AudioDevice]()
+		DECLARE_CYCLE_STAT(TEXT("FAudioThreadTask.ClearListenerAttenuationOverride"), STAT_AudioClearListenerAttenuationOverride, STATGROUP_AudioThreadCommands);
+
+		ListenerProxies[ListenerIndex].AttenuationOverride = FVector::ZeroVector;
+		ListenerProxies[ListenerIndex].bUseAttenuationOverride = true;
+
+		FAudioDevice* AudioDevice = this;
+		FAudioThread::RunCommandOnAudioThread([AudioDevice, ListenerIndex]()
+		{
+			AudioDevice->ClearListenerAttenuationOverride(ListenerIndex);
+		}, GET_STATID(STAT_AudioClearListenerAttenuationOverride));
+	}
+	else
 	{
-		AudioDevice->bUseListenerAttenuationOverride = false;
-		AudioDevice->UpdateVirtualLoops(true);
-	}, GET_STATID(STAT_AudioClearListenerAttenuationOverride));
+		if(ensureMsgf(ListenerIndex < Listeners.Num(), TEXT("Listener Index %u out of range of available Listeners!"), ListenerIndex))
+		{
+			FListener& Listener = Listeners[ListenerIndex];
+			if (Listener.bUseAttenuationOverride)
+			{
+				Listener.bUseAttenuationOverride = false;
+				UpdateVirtualLoops(true);
+			}
+		}
+	}
 }
 
 void FAudioDevice::SetDefaultAudioSettings(UWorld* World, const FReverbSettings& DefaultReverbSettings, const FInteriorSettings& DefaultInteriorSettings)
@@ -4116,7 +4181,7 @@ void FAudioDevice::Update(bool bGameTicking)
 		
 		if (UWorld* World = ActiveSound->GetWorld())
 		{
-			FAudioDebugger::DrawDebugInfo(*World, Listeners, ListenerAttenuationOverride, bUseListenerAttenuationOverride);
+			FAudioDebugger::DrawDebugInfo(*World, Listeners);
 			break;
 		}
 	}
@@ -4827,47 +4892,32 @@ bool FAudioDevice::LocationIsAudible(const FVector& Location, const float MaxDis
 
 	check(bInAudioThread || bInGameThread);
 
-	if (bInAudioThread)
+	const int32 ListenerCount = bInAudioThread ? Listeners.Num() : ListenerProxies.Num();
+	for (int32 i = 0; i < ListenerCount; ++i)
 	{
-		for (const FListener& Listener : Listeners)
+		if (LocationIsAudible(Location, i, MaxDistance))
 		{
-			if (LocationIsAudible(Location, Listener.Transform, MaxDistance))
-			{
-				return true;
-			}
-		}
-	}
-	else //if (bInGameThread)
-	{
-		for (const FTransform& ListenerTransform : ListenerTransforms)
-		{
-			if (LocationIsAudible(Location, ListenerTransform, MaxDistance))
-			{
-				return true;
-			}
+			return true;
 		}
 	}
 
 	return false;
 }
 
-bool FAudioDevice::LocationIsAudible(const FVector& Location, const FTransform& ListenerTransform, float MaxDistance) const
+bool FAudioDevice::LocationIsAudible(const FVector& Location, int32 ListenerIndex, float MaxDistance) const
 {
 	if (MaxDistance >= WORLD_MAX)
 	{
 		return true;
 	}
-
+	
 	FVector ListenerTranslation;
-	if (bUseListenerAttenuationOverride)
+	const bool bAllowOverride = true;
+	if (ListenerIndex == INDEX_NONE || !GetListenerPosition(ListenerIndex, ListenerTranslation, bAllowOverride))
 	{
-		ListenerTranslation = ListenerAttenuationOverride;
-	}
-	else
-	{
-		ListenerTranslation = ListenerTransform.GetTranslation();
+		return false;
 	}
-
+	
 	const float MaxDistanceSquared = MaxDistance * MaxDistance;
 	return (ListenerTranslation - Location).SizeSquared() < MaxDistanceSquared;
 }
@@ -4879,46 +4929,76 @@ float FAudioDevice::GetDistanceToNearestListener(const FVector& Location) const
 
 	check(bInAudioThread || bInGameThread);
 
-	float NearestDistSquared = FLT_MAX;
-	if (bInAudioThread)
+	float DistSquared;
+	if (FindClosestListenerIndex(Location, DistSquared, true) == INDEX_NONE)
 	{
-		for (const FListener& Listener : Listeners)
-		{
-			float DistSquared = GetSquaredDistanceToListener(Location, Listener.Transform);
-			if (DistSquared < NearestDistSquared)
-			{
-				NearestDistSquared = DistSquared;
-			}
-		}
-	}
-	else // bInGameThread
-	{
-		for (const FTransform& ListenerTransform : ListenerTransforms)
-		{
-			float DistSquared = GetSquaredDistanceToListener(Location, ListenerTransform);
-			if (DistSquared < NearestDistSquared)
-			{
-				NearestDistSquared = DistSquared;
-			}
-		}
+		return WORLD_MAX;
 	}
 
-	return FMath::Sqrt(NearestDistSquared);
+	return FMath::Sqrt(DistSquared);
 }
 
-float FAudioDevice::GetSquaredDistanceToListener(const FVector& Location, const FTransform& ListenerTransform) const
+float FAudioDevice::GetDistanceSquaredToListener(const FVector& Location, int32 ListenerIndex) const
 {
 	FVector ListenerTranslation;
-	if (bUseListenerAttenuationOverride)
+	const bool bAllowOverride = true;
+	if (!GetListenerPosition(ListenerIndex, ListenerTranslation, bAllowOverride))
 	{
-		ListenerTranslation = ListenerAttenuationOverride;
+		return WORLD_MAX;
 	}
-	else
+	return (ListenerTranslation - Location).SizeSquared();
+}
+
+bool FAudioDevice::GetListenerPosition(int32 ListenerIndex, FVector& OutPosition, bool bAllowOverride) const
+{
+	OutPosition = FVector::ZeroVector;
+	if (ListenerIndex == INDEX_NONE)
 	{
-		ListenerTranslation = ListenerTransform.GetTranslation();
+		return false;
 	}
 
-	return (ListenerTranslation - Location).SizeSquared();
+	if (IsInAudioThread())
+	{
+		checkf(ListenerIndex < Listeners.Num(), TEXT("Listener Index %u out of range of available Listeners!"), ListenerIndex);
+		const FListener& Listener = Listeners[ListenerIndex];
+		OutPosition = Listener.GetPosition(bAllowOverride);
+		return true;
+	}
+	else // IsInGameThread()
+	{
+		checkf(ListenerIndex < ListenerProxies.Num(), TEXT("Listener Index %u out of range of available Listeners!"), ListenerIndex);
+		const FListenerProxy& Proxy = ListenerProxies[ListenerIndex];
+		OutPosition = Proxy.GetPosition(bAllowOverride);
+		return true;
+	}
+	return false;
+}
+
+bool FAudioDevice::GetListenerTransform(int32 ListenerIndex, FTransform& OutTransform) const
+{
+	OutTransform.SetIdentity();
+	if (ListenerIndex == INDEX_NONE)
+	{
+		return false;
+	}
+
+	if (IsInAudioThread())
+	{
+		if (ListenerIndex < Listeners.Num())
+		{
+			OutTransform = Listeners[ListenerIndex].Transform;
+			return true;
+		}
+	}
+	else // IsInGameThread()
+	{
+		if (ListenerIndex < ListenerProxies.Num())
+		{
+			OutTransform = ListenerProxies[ListenerIndex].Transform;
+			return true;
+		}
+	}
+	return false;
 }
 
 void FAudioDevice::GetMaxDistanceAndFocusFactor(USoundBase* Sound, const UWorld* World, const FVector& Location, const FSoundAttenuationSettings* AttenuationSettingsToApply, float& OutMaxDistance, float& OutFocusFactor)
@@ -4945,14 +5025,14 @@ void FAudioDevice::GetMaxDistanceAndFocusFactor(USoundBase* Sound, const UWorld*
 		if (AttenuationSettingsToApply->bSpatialize && AttenuationSettingsToApply->bEnableListenerFocus)
 		{
 			const int32 ClosestListenerIndex = FindClosestListenerIndex(SoundTransform);
-			if (ClosestListenerIndex >= ListenerTransforms.Num())
+			if (ClosestListenerIndex == INDEX_NONE)
 			{
 				UE_LOG(LogAudio, Warning, TEXT("Invalid ClosestListenerIndex. Sound max distance and focus factor calculation failed."));
 				return;
 			}
 
 			// Now scale the max distance based on the focus settings in the attenuation settings
-			FAttenuationListenerData ListenerData = FAttenuationListenerData::Create(*this, ListenerTransforms[ClosestListenerIndex], SoundTransform, *AttenuationSettingsToApply);
+			FAttenuationListenerData ListenerData = FAttenuationListenerData::Create(*this, ClosestListenerIndex, SoundTransform, *AttenuationSettingsToApply);
 
 			float Azimuth = 0.0f;
 			float AbsoluteAzimuth = 0.0f;
@@ -5036,55 +5116,43 @@ bool FAudioDevice::SoundIsAudible(const FActiveSound& NewActiveSound)
 	return false;
 }
 
-int32 FAudioDevice::FindClosestListenerIndex(const FTransform& SoundTransform, const TArray<FListener>& InListeners)
+int32 FAudioDevice::FindClosestListenerIndex(const FTransform& SoundTransform) const
 {
-	int32 ClosestListenerIndex = 0;
-	if (InListeners.Num() > 0)
-	{
-		float ClosestDistSq = FVector::DistSquared(SoundTransform.GetTranslation(), InListeners[0].Transform.GetTranslation());
-
-		for (int32 i = 1; i < InListeners.Num(); i++)
-		{
-			const float DistSq = FVector::DistSquared(SoundTransform.GetTranslation(), InListeners[i].Transform.GetTranslation());
-			if (DistSq < ClosestDistSq)
-			{
-				ClosestListenerIndex = i;
-				ClosestDistSq = DistSq;
-			}
-		}
-	}
-
-	return ClosestListenerIndex;
+	float UnusedDistSq;
+	const bool bAllowOverrides = false;
+	return FindClosestListenerIndex(SoundTransform.GetTranslation(), UnusedDistSq, bAllowOverrides);
 }
 
-int32 FAudioDevice::FindClosestListenerIndex(const FTransform& SoundTransform) const
+int32 FAudioDevice::FindClosestListenerIndex(const FVector& Position, float& OutDistanceSq, bool bAllowAttenuationOverrides) const
 {
-	if (IsInAudioThread())
-	{
-		return FindClosestListenerIndex(SoundTransform, Listeners);
-	}
-	else if (IsInGameThread())
+	int32 ClosestListenerIndex = INDEX_NONE;
+	OutDistanceSq = WORLD_MAX;
+
+	FVector ListenerPosition;
+	const int32 ListenerCount = IsInAudioThread() ? Listeners.Num() : ListenerProxies.Num();
+
+	if (GetListenerPosition(0, ListenerPosition, bAllowAttenuationOverrides))
 	{
-		int32 ClosestListenerIndex = 0;
-		if (ListenerTransforms.Num() > 0)
+		OutDistanceSq = FVector::DistSquared(Position, ListenerPosition);
+		ClosestListenerIndex = 0;
+
+		for (int32 i = 1; i < ListenerCount; ++i)
 		{
-			float ClosestDistSq = FVector::DistSquared(SoundTransform.GetTranslation(), ListenerTransforms[0].GetTranslation());
+			if (!GetListenerPosition(i, ListenerPosition, bAllowAttenuationOverrides))
+			{
+				continue;
+			}
 
-			for (int32 i = 1; i < ListenerTransforms.Num(); i++)
+			const float DistSq = FVector::DistSquared(Position, ListenerPosition);
+			if (DistSq < OutDistanceSq)
 			{
-				const float DistSq = FVector::DistSquared(SoundTransform.GetTranslation(), ListenerTransforms[i].GetTranslation());
-				if (DistSq < ClosestDistSq)
-				{
-					ClosestListenerIndex = i;
-					ClosestDistSq = DistSq;
-				}
+				OutDistanceSq = DistSq;
+				ClosestListenerIndex = i;
 			}
 		}
-
-		return ClosestListenerIndex;
 	}
 
-	return INDEX_NONE;
+	return ClosestListenerIndex;
 }
 
 void FAudioDevice::UnlinkActiveSoundFromComponent(const FActiveSound& InActiveSound)
@@ -5422,9 +5490,9 @@ void FAudioDevice::Flush(UWorld* WorldToFlush, bool bClearActivatedReverb)
 		AudioFence.BeginFence();
 		AudioFence.Wait();
 
-		// Clear the GameThread cache of the listner
-		ListenerTransforms.Reset();
-		ListenerTransforms.AddDefaulted();
+		// Clear the GameThread cache of the listener
+		ListenerProxies.Reset();
+		ListenerProxies.AddDefaulted();
 
 		return;
 	}
@@ -6251,4 +6319,4 @@ int32 FAudioDevice::GetNumPrecacheFrames() const
 	}
 	// Otherwise, use the default value or value set in ini file
 	return NumPrecacheFrames;
-}
+}
\ No newline at end of file
diff --git a/Engine/Source/Runtime/Engine/Private/AudioVirtualLoop.cpp b/Engine/Source/Runtime/Engine/Private/AudioVirtualLoop.cpp
--- a/Engine/Source/Runtime/Engine/Private/AudioVirtualLoop.cpp
+++ b/Engine/Source/Runtime/Engine/Private/AudioVirtualLoop.cpp
@@ -182,17 +182,9 @@ void FAudioVirtualLoop::UpdateFocusData(float DeltaTime)
 
 	check(ActiveSound->AudioDevice);
 	const FAudioDevice& AudioDevice = *ActiveSound->AudioDevice;
+	const int32 ClosestListenerIndex = AudioDevice.FindClosestListenerIndex(ActiveSound->Transform);
 
-	FAttenuationFocusData FocusData;
-	FTransform ListenerTransform;
-	const TArray<FListener>& Listeners = AudioDevice.GetListeners();
-	if (Listeners.Num() > 0)
-	{
-		int32 ClosestListenerIndex = FAudioDevice::FindClosestListenerIndex(ActiveSound->Transform, Listeners);
-		ListenerTransform = Listeners[ClosestListenerIndex].Transform;
-	}
-
-	FAttenuationListenerData ListenerData = FAttenuationListenerData::Create(AudioDevice, ListenerTransform, ActiveSound->Transform, ActiveSound->AttenuationSettings);
+	FAttenuationListenerData ListenerData = FAttenuationListenerData::Create(AudioDevice, ClosestListenerIndex, ActiveSound->Transform, ActiveSound->AttenuationSettings);
 	ActiveSound->UpdateFocusData(DeltaTime, ListenerData);
 }
 
@@ -241,4 +233,4 @@ bool FAudioVirtualLoop::ShouldListenerMoveForceUpdate(const FTransform& LastTran
 	const float DistanceSq = FVector::DistSquared(LastTransform.GetTranslation(), CurrentTransform.GetTranslation());
 	const float ForceUpdateDistSq = VirtualLoopsForceUpdateListenerMoveDistanceCVar * VirtualLoopsForceUpdateListenerMoveDistanceCVar;
 	return DistanceSq > ForceUpdateDistSq;
-}
+}
\ No newline at end of file
diff --git a/Engine/Source/Runtime/Engine/Private/GameViewportClient.cpp b/Engine/Source/Runtime/Engine/Private/GameViewportClient.cpp
--- a/Engine/Source/Runtime/Engine/Private/GameViewportClient.cpp
+++ b/Engine/Source/Runtime/Engine/Private/GameViewportClient.cpp
@@ -1334,11 +1334,11 @@ void UGameViewportClient::Draw(FViewport* InViewport, FCanvas* SceneCanvas)
 								FVector OverrideAttenuation;
 								if (PlayerController->GetAudioListenerAttenuationOverridePosition(OverrideAttenuation))
 								{
-									AudioDevice->SetListenerAttenuationOverride(OverrideAttenuation);
+									AudioDevice->SetListenerAttenuationOverride(ViewportIndex, OverrideAttenuation);
 								}
 								else
 								{
-									AudioDevice->ClearListenerAttenuationOverride();
+									AudioDevice->ClearListenerAttenuationOverride(ViewportIndex);
 								}
 							}
 						}
@@ -3927,4 +3927,4 @@ bool UGameViewportClient::LoadAvailableCursorPngs(TArray< TSharedPtr<FPngFileDat
 	return bFoundCursor;
 }
 
-#undef LOCTEXT_NAMESPACE
+#undef LOCTEXT_NAMESPACE
\ No newline at end of file
diff --git a/Engine/Source/Runtime/Engine/Private/SoundConcurrency.cpp b/Engine/Source/Runtime/Engine/Private/SoundConcurrency.cpp
--- a/Engine/Source/Runtime/Engine/Private/SoundConcurrency.cpp
+++ b/Engine/Source/Runtime/Engine/Private/SoundConcurrency.cpp
@@ -480,10 +480,9 @@ FActiveSound* FSoundConcurrencyManager::GetEvictableSoundStopFarthest(const FAct
 	const EMaxConcurrentResolutionRule::Type Rule = ConcurrencyGroup.GetSettings().ResolutionRule;
 
 	check(AudioDevice);
-	TArray<FListener>& Listeners = AudioDevice->Listeners;
-
-	int32 ClosestListenerIndex = NewActiveSound.FindClosestListener(Listeners);
-	float DistanceToStopSoundSq = FVector::DistSquared(Listeners[ClosestListenerIndex].Transform.GetTranslation(), NewActiveSound.Transform.GetTranslation());
+	float DistanceToStopSoundSq;
+	const bool bAllowAttenuationOverrides = false;
+	int32 ClosestListenerIndex = AudioDevice->FindClosestListenerIndex(NewActiveSound.Transform.GetTranslation(), DistanceToStopSoundSq, bAllowAttenuationOverrides);
 
 	FActiveSound* EvictableSound = nullptr;
 	const TArray<FActiveSound*>& ActiveSounds = ConcurrencyGroup.GetActiveSounds();
@@ -491,8 +490,8 @@ FActiveSound* FSoundConcurrencyManager::GetEvictableSoundStopFarthest(const FAct
 	{
 		check(ActiveSound);
 
-		ClosestListenerIndex = ActiveSound->FindClosestListener(Listeners);
-		const float DistanceToActiveSoundSq = FVector::DistSquared(Listeners[ClosestListenerIndex].Transform.GetTranslation(), ActiveSound->Transform.GetTranslation());
+		float DistanceToActiveSoundSq;
+		ClosestListenerIndex = AudioDevice->FindClosestListenerIndex(ActiveSound->Transform.GetTranslation(), DistanceToActiveSoundSq, bAllowAttenuationOverrides);
 
 		// Stop Farthest
 		if (DistanceToActiveSoundSq > DistanceToStopSoundSq)
@@ -839,4 +838,4 @@ void FSoundConcurrencyManager::UpdateQuietSoundsToStop()
 	{
 		ConcurrenyGroupEntry.Value->StopQuietSoundsDueToMaxConcurrency();
 	}
-}
+}
\ No newline at end of file
diff --git a/Engine/Source/Runtime/Engine/Private/SoundNodeAttenuation.cpp b/Engine/Source/Runtime/Engine/Private/SoundNodeAttenuation.cpp
--- a/Engine/Source/Runtime/Engine/Private/SoundNodeAttenuation.cpp
+++ b/Engine/Source/Runtime/Engine/Private/SoundNodeAttenuation.cpp
@@ -58,10 +58,10 @@ void USoundNodeAttenuation::ParseNodes(FAudioDevice* AudioDevice, const UPTRINT
 	const FSoundAttenuationSettings* Settings = (ActiveSound.bAllowSpatialization ? GetAttenuationSettingsToApply() : nullptr);
 	if (Settings)
 	{
-		const FListener& Listener = AudioDevice->GetListeners()[0];
-
 		// Update this node's attenuation settings overrides
-		ActiveSound.ParseAttenuation(UpdatedParseParams, Listener, *Settings);
+		check(AudioDevice);
+		const int32 ClosestListenerIndex = AudioDevice->FindClosestListenerIndex(UpdatedParseParams.Transform);
+		ActiveSound.ParseAttenuation(UpdatedParseParams, ClosestListenerIndex, *Settings);
 	}
 	else
 	{
@@ -69,4 +69,4 @@ void USoundNodeAttenuation::ParseNodes(FAudioDevice* AudioDevice, const UPTRINT
 	}
 
 	Super::ParseNodes(AudioDevice, NodeWaveInstanceHash, ActiveSound, UpdatedParseParams, WaveInstances);
-}
+}
\ No newline at end of file
diff --git a/Engine/Source/Runtime/Engine/Public/ActiveSound.h b/Engine/Source/Runtime/Engine/Public/ActiveSound.h
--- a/Engine/Source/Runtime/Engine/Public/ActiveSound.h
+++ b/Engine/Source/Runtime/Engine/Public/ActiveSound.h
@@ -641,9 +641,6 @@ public:
 	/** Gets the sound source bus sends to use for this sound instance. */
 	void GetSoundSourceBusSends(EBusSendType BusSendType, TArray<FSoundSourceBusSendInfo>& OutSends) const;
 
-	/* Determines which listener is the closest to the sound */
-	int32 FindClosestListener( const TArray<struct FListener>& InListeners ) const;
-
 	/** Returns the unique ID of the active sound's owner if it exists. Returns 0 if the sound doesn't have an owner. */
 	FSoundOwnerObjectID GetOwnerID() const { return OwnerID; }
 
@@ -656,7 +653,7 @@ public:
 	static void OcclusionTraceDone(const FTraceHandle& TraceHandle, FTraceDatum& TraceDatum);
 
 	/** Applies the active sound's attenuation settings to the input parse params using the given listener */
-	void ParseAttenuation(FSoundParseParameters& OutParseParams, const FListener& InListener, const FSoundAttenuationSettings& InAttenuationSettings);
+	void ParseAttenuation(FSoundParseParameters& OutParseParams, int32 ListenerIndex, const FSoundAttenuationSettings& InAttenuationSettings);
 
 	/** Returns the effective priority of the active sound */
 	float GetPriority() const { return Priority * FocusData.PriorityScale; }
@@ -668,7 +665,7 @@ public:
 	void SetSourceBusSend(EBusSendType BusSendTyoe, const FSoundSourceBusSendInfo& SourceBusSendInfo);
 
 	/** Updates the active sound's attenuation settings to the input parse params using the given listener */
-	void UpdateAttenuation(float DeltaTime, FSoundParseParameters& ParseParams, const FListener& Listener, const FSoundAttenuationSettings* SettingsAttenuationNode = nullptr);
+	void UpdateAttenuation(float DeltaTime, FSoundParseParameters& ParseParams, int32 ListenerIndex, const FSoundAttenuationSettings* SettingsAttenuationNode = nullptr);
 
 	/** Updates the provided focus data using the local */
 	void UpdateFocusData(float DeltaTime, const FAttenuationListenerData& ListenerData, FAttenuationFocusData* OutFocusData = nullptr);
@@ -685,8 +682,8 @@ private:
 
 	static FTraceDelegate ActiveSoundTraceDelegate;
 
-	/** Cached ptr to the closest listener. So we don't have to do the work to find it twice. */
-	const FListener* ClosestListenerPtr;
+	/** Cached index to the closest listener. So we don't have to do the work to find it twice. */
+	int32 ClosestListenerIndex;
 
 	/** This is a friend so the audio device can call Stop() on the active sound. */
 	friend class FAudioDevice;
@@ -715,8 +712,8 @@ private:
 	void CheckOcclusion(const FVector ListenerLocation, const FVector SoundLocation, const FSoundAttenuationSettings* AttenuationSettingsPtr);
 
 	/** Apply the interior settings to the ambient sound as appropriate */
-	void HandleInteriorVolumes( const FListener& Listener, struct FSoundParseParameters& ParseParams );
+	void HandleInteriorVolumes(struct FSoundParseParameters& ParseParams);
 
 	/** Helper function which retrieves attenuation frequency value for HPF and LPF distance-based filtering. */
 	float GetAttenuationFrequency(const FSoundAttenuationSettings* InSettings, const FAttenuationListenerData& ListenerData, const FVector2D& FrequencyRange, const FRuntimeFloatCurve& CustomCurve);
-};
+};
\ No newline at end of file
diff --git a/Engine/Source/Runtime/Engine/Public/Audio/AudioDebug.h b/Engine/Source/Runtime/Engine/Public/Audio/AudioDebug.h
--- a/Engine/Source/Runtime/Engine/Public/Audio/AudioDebug.h
+++ b/Engine/Source/Runtime/Engine/Public/Audio/AudioDebug.h
@@ -45,7 +45,7 @@ public:
 
 	static void DrawDebugInfo(const FSoundSource& SoundSource);
 	static void DrawDebugInfo(const FActiveSound& ActiveSound, const TArray<FWaveInstance*>& ThisSoundsWaveInstances, const float DeltaTime);
-	static void DrawDebugInfo(UWorld& World, const TArray<FListener>& Listeners, FVector& ListenerTransformOverride, bool bUseListenerTransformOverride);
+	static void DrawDebugInfo(UWorld& World, const TArray<FListener>& Listeners);
 	static void DrawDebugInfo(const FAudioVirtualLoop& VirtualLoop);
 	static bool PostStatModulatorHelp(UWorld* World, FCommonViewportClient* ViewportClient, const TCHAR* Stream);
 	static int32 RenderStatCues(UWorld* World, FViewport* Viewport, FCanvas* Canvas, int32 X, int32 Y, const FVector* ViewLocation, const FRotator* ViewRotation);
@@ -120,4 +120,4 @@ private:
 	/** Whether or not 3d debug visualization is enabled. */
 	uint8 bVisualize3dDebug : 1;
 };
-#endif // ENABLE_AUDIO_DEBUG
+#endif // ENABLE_AUDIO_DEBUG
\ No newline at end of file
diff --git a/Engine/Source/Runtime/Engine/Public/AudioDevice.h b/Engine/Source/Runtime/Engine/Public/AudioDevice.h
--- a/Engine/Source/Runtime/Engine/Public/AudioDevice.h
+++ b/Engine/Source/Runtime/Engine/Public/AudioDevice.h
@@ -137,6 +137,12 @@ struct FListener
 	FTransform Transform;
 	FVector Velocity;
 
+	/** An attenuation override to use for distance and attenuation calculations */
+	FVector AttenuationOverride;
+
+	/** Is our attenuation override active */
+	uint32 bUseAttenuationOverride:1;
+
 	struct FInteriorSettings InteriorSettings;
 
 	/** The ID of the volume the listener resides in */
@@ -161,6 +167,11 @@ struct FListener
 	FVector GetFront() const	{ return Transform.GetUnitAxis(EAxis::Y); }
 	FVector GetRight() const	{ return Transform.GetUnitAxis(EAxis::X); }
 
+	/**
+	 * Gets the position of the listener
+	 */
+	FVector GetPosition(bool bAllowOverride) const;
+
 	/**
 	 * Works out the interp value between source and end
 	 */
@@ -179,6 +190,8 @@ struct FListener
 	FListener(FAudioDevice* InAudioDevice)
 		: Transform(FTransform::Identity)
 		, Velocity(ForceInit)
+		, AttenuationOverride(ForceInit)
+		, bUseAttenuationOverride(false)
 		, AudioVolumeID(0)
 		, InteriorStartTime(0.0)
 		, InteriorEndTime(0.0)
@@ -197,6 +210,37 @@ private:
 	FListener();
 };
 
+/**
+ * Game thread representation of a listener
+ */
+struct FListenerProxy
+{
+	FTransform Transform;
+
+	/** An attenuation override to use for distance and attenuation calculations */
+	FVector AttenuationOverride;
+
+	/** Is our attenuation override active */
+	uint32 bUseAttenuationOverride :1;
+
+	/**
+	 * Gets the position of the listener proxy
+	 */
+	FVector GetPosition(bool bAllowOverride) const;
+
+	FListenerProxy()
+		: bUseAttenuationOverride(false)
+	{
+	}
+
+	FListenerProxy(const FListener& Listener)
+		: Transform(Listener.Transform)
+		, AttenuationOverride(Listener.AttenuationOverride)
+		, bUseAttenuationOverride(Listener.bUseAttenuationOverride)
+	{
+	}
+};
+
 /**
  * Structure for collating info about sound classes
  */
@@ -283,7 +327,7 @@ struct FAttenuationListenerData
 	const FSoundAttenuationSettings* AttenuationSettings;
 
 	/** Computes and returns some geometry related to the listener and the given sound transform. */
-	static FAttenuationListenerData Create(const FAudioDevice& AudioDevice, const FTransform& InListenerTransform, const FTransform& InSoundTransform, const FSoundAttenuationSettings& InAttenuationSettings);
+	static FAttenuationListenerData Create(const FAudioDevice& AudioDevice, int32 ListenerIndex, const FTransform& InSoundTransform, const FSoundAttenuationSettings& InAttenuationSettings);
 
 private:
 	FAttenuationListenerData(const FTransform& InListenerTransform, const FTransform& InSoundTransform, const FSoundAttenuationSettings& InAttenuationSettings)
@@ -598,10 +642,10 @@ public:
 	void SetListener(UWorld* World, int32 InListenerIndex, const FTransform& ListenerTransform, float InDeltaSeconds);
 
 	/** Sets an override for the listener to do attenuation calculations. */
-	void SetListenerAttenuationOverride(const FVector AttenuationPosition);
+	void SetListenerAttenuationOverride(int32 ListenerIndex, const FVector AttenuationPosition);
 
 	/** Removes a listener attenuation override. */
-	void ClearListenerAttenuationOverride();
+	void ClearListenerAttenuationOverride(int32 ListenerIndex);
 
 	const TArray<FListener>& GetListeners() const { check(IsInAudioThread()); return Listeners; }
 
@@ -815,14 +859,32 @@ public:
 	/**
 	* Checks to see if a coordinate is within a distance of the given listener
 	*/
-	bool LocationIsAudible(const FVector& Location, const FTransform& ListenerTransform, const float MaxDistance) const;
+	bool LocationIsAudible(const FVector& Location, int32 ListenerIndex, const float MaxDistance) const;
 
 	/**
 	* Returns the distance to the nearest listener from the given location
 	*/
 	float GetDistanceToNearestListener(const FVector& Location) const;
 
-	float GetSquaredDistanceToListener(const FVector& Location, const FTransform& ListenerTransform) const;
+	/**
+	* Returns the distance from location to the appropriate listener representation, depending on calling thread
+	*/
+	float GetDistanceSquaredToListener(const FVector& Location, int32 ListenerIndex) const;
+
+	/**
+	* Returns a position from the appropriate listener representation, depending on calling thread.
+	*
+	* @param	ListenerIndex	index of the listener or proxy
+	* @param	OutPosition		filled in position of the listener or proxy
+	* @param	bAllowOverride	if true we will use the attenuation override for position, if set
+	* @return	true if successful
+	*/
+	bool GetListenerPosition(int32 ListenerIndex, FVector& OutPosition, bool bAllowOverride) const;
+
+	/**
+	* Returns the transform of the appropriate listener representation, depending on calling thread
+	*/
+	bool GetListenerTransform(int32 ListenerIndex, FTransform& OutTransform) const;
 
 	/**
 	 * Sets the Sound Mix that should be active by default
@@ -1028,8 +1090,8 @@ public:
 	bool SoundIsAudible(USoundBase* Sound, const UWorld* World, const FVector& Location, const FSoundAttenuationSettings* AttenuationSettingsToApply, float MaxDistance, float FocusFactor) const;
 
 	/** Returns the index of the listener closest to the given sound transform */
-	static int32 FindClosestListenerIndex(const FTransform& SoundTransform, const TArray<FListener>& InListeners);
 	int32 FindClosestListenerIndex(const FTransform& SoundTransform) const;
+	int32 FindClosestListenerIndex(const FVector& Position, float& OutSqDistance, bool AllowAttenuationOverrides) const;
 
 	/** Disables ActiveSound from responding to calls from its associated AudioComponent. */
 	void UnlinkActiveSoundFromComponent(const FActiveSound& InActiveSound);
@@ -1507,18 +1569,6 @@ public:
 	/** Returns the game's delta time */
 	float GetGameDeltaTime() const;
 
-	/** Whether device is using listener attenuation override or not. */
-	bool IsUsingListenerAttenuationOverride() const
-	{
-		return bUseListenerAttenuationOverride;
-	}
-
-	/** Returns the listener attenuation override */
-	const FVector& GetListenerAttenuationOverride() const
-	{
-		return ListenerAttenuationOverride;
-	}
-
 	void UpdateVirtualLoops(bool bForceUpdate);
 
 	/** Sets the update delta time for the audio frame */
@@ -1658,7 +1708,7 @@ public:
 	TArray<TAudioPluginListenerPtr> PluginListeners;
 
 	// Game thread cache of listener transforms
-	TArray<FTransform> ListenerTransforms;
+	TArray<FListenerProxy> ListenerProxies;
 
 private:
 	/** The maximum number of sources.  Value cannot change after initialization. */
@@ -1704,12 +1754,6 @@ protected:
 	// Audio thread representation of listeners
 	TArray<FListener> Listeners;
 
-	// A listener attenuation override to use for distance and attenuation calculations
-	FVector ListenerAttenuationOverride;
-
-	// Whether or not to use the listener attenuation override
-	bool bUseListenerAttenuationOverride;
-
 	TArray<FSoundSource*> Sources;
 	TArray<FSoundSource*> FreeSources;
 
@@ -1892,4 +1936,3 @@ private:
 	float GlobalMaxPitch;
 };
 
-
