Semi blind fix for crash during cook... Will need to be properly investigated

diff --git a/Engine/Source/Runtime/Engine/Private/InstancedStaticMesh.cpp b/Engine/Source/Runtime/Engine/Private/InstancedStaticMesh.cpp
--- a/Engine/Source/Runtime/Engine/Private/InstancedStaticMesh.cpp
+++ b/Engine/Source/Runtime/Engine/Private/InstancedStaticMesh.cpp
@@ -1444,7 +1444,7 @@ void UInstancedStaticMeshComponent::BuildRenderData(FStaticMeshInstanceData& Out
 	for (int32 Index = 0; Index < NumInstances; ++Index)
 	{
 		int32 RenderIndex = InstanceReorderTable.IsValidIndex(Index) ? InstanceReorderTable[Index] : Index;
-		if (RenderIndex == INDEX_NONE) 
+		if (RenderIndex == INDEX_NONE || !OutData.IsValidIndex(RenderIndex)) 
 		{
 			// could be skipped by density settings
 			continue;
