'stat sounds' fixes/updates
- Fix 'Stub' issue for source descriptions in stat sounds when running audio mixer
- Add colorization options for debug body text
- Add disabling of full path for debug
- Add sort by priority for stat sounds
- Only show pertinent information when sorting by given field
- Relabel 'Debug' header for stat sounds to 'Visualize Attenuation'... as that's what the debug flag does and its less confusing
Potential 'stat soundmixes' fix
ClearSolos/Mutes on Begin/Exit PIE

diff --git a/Engine/Config/BaseInput.ini b/Engine/Config/BaseInput.ini
--- a/Engine/Config/BaseInput.ini
+++ b/Engine/Config/BaseInput.ini
@@ -205,7 +205,7 @@ MaxScrollbackSize=1024
 +ManualAutoCompleteList=(Command="Stat SoundModulatorsHelp",Desc="Shows modulator debug help provided by active audio modulation plugin")
 +ManualAutoCompleteList=(Command="Stat SoundReverb",Desc="Shows active SoundReverb")
 +ManualAutoCompleteList=(Command="Stat SoundWaves",Desc="Shows active SoundWaves")
-+ManualAutoCompleteList=(Command="Stat Sounds",Desc="<?> <sort=distance|class|name|waves|default> <-debug> <off> Shows active SoundCues and SoundWaves")
++ManualAutoCompleteList=(Command="Stat Sounds",Desc="<?> <sort=class|distance|name|priority|time|volume|waves> <-debug> Shows all active sounds. Displays value sorted by when sort is set")
 +ManualAutoCompleteList=(Command="ScriptAudit LongestFunctions",Desc="List functions that contain the most bytecode - optionally include # of entries to list")
 +ManualAutoCompleteList=(Command="ScriptAudit FrequentFunctionsCalled",Desc="List functions that are most frequently called from bytecode - optionally include # of entries to list")
 +ManualAutoCompleteList=(Command="ScriptAudit FrequentInstructions",Desc="List most frequently used instructions - optionally include # of entries to list")
diff --git a/Engine/Source/Runtime/AudioMixer/Private/AudioMixerSource.cpp b/Engine/Source/Runtime/AudioMixer/Private/AudioMixerSource.cpp
--- a/Engine/Source/Runtime/AudioMixer/Private/AudioMixerSource.cpp
+++ b/Engine/Source/Runtime/AudioMixer/Private/AudioMixerSource.cpp
@@ -716,11 +716,6 @@ namespace Audio
 		return true;
 	}
 
-	FString FMixerSource::Describe(bool bUseLongName)
-	{
-		return FString(TEXT("Stub"));
-	}
-
 	float FMixerSource::GetPlaybackPercent() const
 	{
 		if (InitializationState != EMixerSourceInitializationState::Initialized)
diff --git a/Engine/Source/Runtime/AudioMixer/Private/AudioMixerSource.h b/Engine/Source/Runtime/AudioMixer/Private/AudioMixerSource.h
--- a/Engine/Source/Runtime/AudioMixer/Private/AudioMixerSource.h
+++ b/Engine/Source/Runtime/AudioMixer/Private/AudioMixerSource.h
@@ -48,7 +48,6 @@ namespace Audio
 		virtual bool IsStopping() override { return bIsStopping; }
 		virtual void Pause() override;
 		virtual bool IsFinished() override;
-		virtual FString Describe(bool bUseLongName) override;
 		virtual float GetPlaybackPercent() const override;
 		virtual float GetEnvelopeValue() const override;
 		//~ End FSoundSource Interface
diff --git a/Engine/Source/Runtime/Engine/Private/AudioDebug.cpp b/Engine/Source/Runtime/Engine/Private/AudioDebug.cpp
--- a/Engine/Source/Runtime/Engine/Private/AudioDebug.cpp
+++ b/Engine/Source/Runtime/Engine/Private/AudioDebug.cpp
@@ -9,6 +9,7 @@
 #include "AudioVirtualLoop.h"
 #include "CanvasTypes.h"
 #include "DrawDebugHelpers.h"
+#include "DSP/Dsp.h"
 #include "Engine/Font.h"
 #include "GameFramework/GameUserSettings.h"
 #include "HAL/IConsoleManager.h"
@@ -25,6 +26,7 @@
 #include "Sound/SoundWave.h"
 #include "UnrealEngine.h"
 
+
 #ifndef ENABLE_AUDIO_DEBUG
 #error "Please define ENABLE_AUDIO_DEBUG"
 #endif //ENABLE_AUDIO_DEBUG
@@ -72,6 +74,38 @@ FAutoConsoleVariableRef CVarAudioVisualizeVirtualLoopsEnabled(
 	TEXT("0: Not Enabled, 1: Enabled"),
 	ECVF_Default);
 
+static int32 AudioDebugSoundMaxNumDisplayedCVar = 32;
+FAutoConsoleVariableRef CVarAudioDebugSoundMaxNumDisplayed(
+	TEXT("au.Debug.Sounds.Max"),
+	AudioDebugSoundMaxNumDisplayedCVar,
+	TEXT("Max number of sounds to display in full sound debugger view. \n")
+	TEXT("Default: 32"),
+	ECVF_Default);
+
+static int32 AudioDebugSoundShowPathCVar = 1;
+FAutoConsoleVariableRef CVarAudioDebugSoundShowPath(
+	TEXT("au.Debug.Sounds.ShowPath"),
+	AudioDebugSoundShowPathCVar,
+	TEXT("Display full path of sound when enabled.\n")
+	TEXT("0: Not Enabled, 1: Enabled"),
+	ECVF_Default);
+
+static FString AudioDebugSoundSortCVarCVar = TEXT("Name");
+FAutoConsoleVariableRef CVarAudioDebugSoundSortCVar(
+	TEXT("au.Debug.Sounds.Sort"),
+	AudioDebugSoundSortCVarCVar,
+	TEXT("Value to sort by and display when sound stats are active. \n")
+	TEXT("Class, Distance, Name (Default), Priority, Time, Waves, Volume"),
+	ECVF_Default);
+
+static FString AudioDebugStatSoundTextColorCVar = TEXT("White");
+FAutoConsoleVariableRef CVarAudioDebugStatSoundColor(
+	TEXT("au.Debug.Sounds.TextColor"),
+	AudioDebugStatSoundTextColorCVar,
+	TEXT("Color of body text in audio debug views. \n")
+	TEXT("White, Red, Orange, Yellow, Blue, Magenta, Purple, Black"),
+	ECVF_Default);
+
 // Test Data
 namespace ERequestedAudioStats
 {
@@ -87,7 +121,50 @@ namespace ERequestedAudioStats
 namespace
 {
 	const FColor HeaderColor = FColor::Green;
-	const FColor BodyColor = FColor::White;
+
+	FColor GetBodyColor()
+	{
+		return FColor::White;
+	}
+
+	FColor GetStatSoundBodyColor()
+	{
+		if (AudioDebugStatSoundTextColorCVar == TEXT("Red"))
+		{
+			return FColor::Red;
+		}
+		if (AudioDebugStatSoundTextColorCVar == TEXT("Orange"))
+		{
+			return FColor::Orange;
+		}
+		if (AudioDebugStatSoundTextColorCVar == TEXT("Yellow"))
+		{
+			return FColor::Yellow;
+		}
+		if (AudioDebugStatSoundTextColorCVar == TEXT("Green"))
+		{
+			return FColor::Green;
+		}
+		if (AudioDebugStatSoundTextColorCVar == TEXT("Blue"))
+		{
+			return FColor::Blue;
+		}
+		if (AudioDebugStatSoundTextColorCVar == TEXT("Magenta"))
+		{
+			return FColor::Magenta;
+		}
+		if (AudioDebugStatSoundTextColorCVar == TEXT("Purple"))
+		{
+			return FColor::Purple;
+		}
+		if (AudioDebugStatSoundTextColorCVar == TEXT("Black"))
+		{
+			return FColor::Black;
+		}
+
+		return FColor::White;
+	}
+
 	const int32  TabWidth = 12;
 
 	const float MinDisplayVolume = KINDA_SMALL_NUMBER; // -80 dB
@@ -106,14 +183,20 @@ namespace
 
 	struct FAudioStats
 	{
+		enum class EDisplaySort : uint8
+		{
+			Class,
+			Distance,
+			Name,
+			PlaybackTime,
+			Priority,
+			Waves,
+			Volume
+		};
+
 		enum class EDisplayFlags : uint8
 		{
 			Debug = 0x01,
-			Sort_Distance = 0x02,
-			Sort_Class = 0x04,
-			Sort_Name = 0x08,
-			Sort_WavesNum = 0x10,
-			Sort_Disabled = 0x20,
 			Long_Names = 0x40,
 		};
 
@@ -130,11 +213,17 @@ namespace
 		struct FStatSoundInfo
 		{
 			FString SoundName;
+			FString SoundPath;
 			FName SoundClassName;
 			float Distance;
+			float PlaybackTime;
+			float PlaybackTimeNonVirtualized;
+			float Priority;
+			float Volume;
 			uint32 AudioComponentID;
 			FTransform Transform;
 			TArray<FStatWaveInstanceInfo> WaveInstanceInfos;
+
 			TMultiMap<EAttenuationShape::Type, FBaseAttenuationSettings::AttenuationShapeDetails> ShapeDetailsMap;
 		};
 
@@ -147,16 +236,50 @@ namespace
 		};
 
 		uint8 DisplayFlags;
+		EDisplaySort DisplaySort;
 		TArray<FTransform> ListenerTransforms;
 		TArray<FStatSoundInfo> StatSoundInfos;
 		TArray<FStatSoundMix> StatSoundMixes;
 
 		FAudioStats()
 			: DisplayFlags(0)
+			, DisplaySort(EDisplaySort::Name)
 		{
 		}
 	};
 
+	void UpdateDisplaySort(FAudioStats& InAudioStats)
+	{
+		if (AudioDebugSoundSortCVarCVar == TEXT("distance"))
+		{
+			InAudioStats.DisplaySort = FAudioStats::EDisplaySort::Distance;
+		}
+		else if (AudioDebugSoundSortCVarCVar == TEXT("class"))
+		{
+			InAudioStats.DisplaySort = FAudioStats::EDisplaySort::Class;
+		}
+		else if (AudioDebugSoundSortCVarCVar == TEXT("name"))
+		{
+			InAudioStats.DisplaySort = FAudioStats::EDisplaySort::Name;
+		}
+		else if (AudioDebugSoundSortCVarCVar == TEXT("time"))
+		{
+			InAudioStats.DisplaySort = FAudioStats::EDisplaySort::PlaybackTime;
+		}
+		else if (AudioDebugSoundSortCVarCVar == TEXT("priority"))
+		{
+			InAudioStats.DisplaySort = FAudioStats::EDisplaySort::Priority;
+		}
+		else if (AudioDebugSoundSortCVarCVar == TEXT("volume"))
+		{
+			InAudioStats.DisplaySort = FAudioStats::EDisplaySort::Volume;
+		}
+		else if (AudioDebugSoundSortCVarCVar == TEXT("waves"))
+		{
+			InAudioStats.DisplaySort = FAudioStats::EDisplaySort::Waves;
+		}
+	}
+
 	struct FAudioStats_AudioThread
 	{
 		uint8 RequestedStats;
@@ -210,6 +333,24 @@ void FAudioDebugger::ToggleVisualizeDebug3dEnabled()
 	bVisualize3dDebug = !bVisualize3dDebug;
 }
 
+void FAudioDebugger::OnBeginPIE()
+{
+	FAudioDeviceManager* DeviceManager = GEngine->GetAudioDeviceManager();
+	if (DeviceManager)
+	{
+		DeviceManager->GetDebugger().ClearMutesAndSolos();
+	}
+}
+
+void FAudioDebugger::OnEndPIE()
+{
+	FAudioDeviceManager* DeviceManager = GEngine->GetAudioDeviceManager();
+	if (DeviceManager)
+	{
+		DeviceManager->GetDebugger().ClearMutesAndSolos();
+	}
+}
+
 void FAudioDebugger::QuerySoloMuteSoundClass(const FString& Name, bool& bOutIsSoloed, bool& bOutIsMuted, FString& OutReason) const
 {
 	GetDebugSoloMuteStateX(Name, DebugNames.SoloSoundClass, DebugNames.MuteSoundClass, bOutIsSoloed, bOutIsMuted, OutReason);
@@ -773,7 +914,7 @@ int32 FAudioDebugger::RenderStatCues(UWorld* World, FViewport* Viewport, FCanvas
 		}
 	}
 
-	Canvas->DrawShadowedString(X, Y, *FString::Printf(TEXT("Total: %i"), ActiveSoundCount), GetStatsFont(), BodyColor);
+	Canvas->DrawShadowedString(X, Y, *FString::Printf(TEXT("Total: %i"), ActiveSoundCount), GetStatsFont(), GetBodyColor());
 	Y += FontHeight;
 
 	return Y;
@@ -802,16 +943,16 @@ int32 FAudioDebugger::RenderStatMixes(UWorld* World, FViewport* Viewport, FCanva
 		{
 			const FString TheString = FString::Printf(TEXT("%s - Fade Proportion: %1.2f - Total Ref Count: %i"), *StatSoundMix.MixName, StatSoundMix.InterpValue, StatSoundMix.RefCount);
 
-			const FColor& TextColour = (StatSoundMix.bIsCurrentEQ ? FColor::Yellow : FColor::White);
+			const FColor& TextColor = (StatSoundMix.bIsCurrentEQ ? FColor::Yellow : GetBodyColor());
 
-			Canvas->DrawShadowedString(X + TabWidth, Y, *TheString, GetStatsFont(), TextColour);
+			Canvas->DrawShadowedString(X + TabWidth, Y, *TheString, GetStatsFont(), TextColor);
 			Y += FontHeight;
 		}
 	}
 
 	if (!bDisplayedSoundMixes)
 	{
-		Canvas->DrawShadowedString(X + TabWidth, Y, TEXT("None"), GetStatsFont(), FColor::White);
+		Canvas->DrawShadowedString(X + TabWidth, Y, TEXT("None"), GetStatsFont(), GetBodyColor());
 		Y += FontHeight;
 	}
 
@@ -841,7 +982,7 @@ int32 FAudioDebugger::RenderStatModulators(UWorld* World, FViewport* Viewport, F
 
 	if (!bDisplayedSoundModulationInfo)
 	{
-		Canvas->DrawShadowedString(X + TabWidth, Y, TEXT("None"), GetStatsFont(), FColor::White);
+		Canvas->DrawShadowedString(X + TabWidth, Y, TEXT("None"), GetStatsFont(), GetBodyColor());
 		Y += FontHeight;
 	}
 
@@ -861,10 +1002,11 @@ int32 FAudioDebugger::RenderStatReverb(UWorld* World, FViewport* Viewport, FCanv
 	const int32 Height = static_cast<int32>(GetStatsFont()->GetMaxCharHeight() + 2);
 
 	FString TheString;
+	const FLinearColor LinearBodyColor = FLinearColor(GetBodyColor());
 	if (UReverbEffect* ReverbEffect = AudioDevice->GetCurrentReverbEffect())
 	{
 		TheString = FString::Printf(TEXT("Active Reverb Effect: %s"), *ReverbEffect->GetName());
-		Canvas->DrawShadowedString(X, Y, *TheString, GetStatsFont(), FLinearColor::White);
+		Canvas->DrawShadowedString(X, Y, *TheString, GetStatsFont(), LinearBodyColor);
 		Y += Height;
 
 		AAudioVolume* CurrentAudioVolume = nullptr;
@@ -886,26 +1028,26 @@ int32 FAudioDebugger::RenderStatReverb(UWorld* World, FViewport* Viewport, FCanv
 		{
 			TheString = TEXT("  Audio Volume Reverb Effect: None");
 		}
-		Canvas->DrawShadowedString(X, Y, *TheString, GetStatsFont(), FLinearColor::White);
+		Canvas->DrawShadowedString(X, Y, *TheString, GetStatsFont(), LinearBodyColor);
 		Y += Height;
 
 		const TMap<FName, FActivatedReverb>& ActivatedReverbs = AudioDevice->GetActiveReverb();
 		if (ActivatedReverbs.Num() == 0)
 		{
 			TheString = TEXT("  Activated Reverb: None");
-			Canvas->DrawShadowedString(X, Y, *TheString, GetStatsFont(), FLinearColor::White);
+			Canvas->DrawShadowedString(X, Y, *TheString, GetStatsFont(), LinearBodyColor);
 			Y += Height;
 		}
 		else if (ActivatedReverbs.Num() == 1)
 		{
 			auto It = ActivatedReverbs.CreateConstIterator();
 			TheString = FString::Printf(TEXT("  Activated Reverb Effect: %s (Priority: %g Tag: '%s')"), *It.Value().ReverbSettings.ReverbEffect->GetName(), It.Value().Priority, *It.Key().ToString());
-			Canvas->DrawShadowedString(X, Y, *TheString, GetStatsFont(), FLinearColor::White);
+			Canvas->DrawShadowedString(X, Y, *TheString, GetStatsFont(), LinearBodyColor);
 			Y += Height;
 		}
 		else
 		{
-			Canvas->DrawShadowedString(X, Y, TEXT("  Activated Reverb Effects:"), GetStatsFont(), FLinearColor::White);
+			Canvas->DrawShadowedString(X, Y, TEXT("  Activated Reverb Effects:"), GetStatsFont(), LinearBodyColor);
 			Y += Height;
 			TMap<int32, FString> PrioritySortedActivatedReverbs;
 			for (auto It = ActivatedReverbs.CreateConstIterator(); It; ++It)
@@ -915,7 +1057,7 @@ int32 FAudioDebugger::RenderStatReverb(UWorld* World, FViewport* Viewport, FCanv
 			}
 			for (auto It = PrioritySortedActivatedReverbs.CreateConstIterator(); It; ++It)
 			{
-				Canvas->DrawShadowedString(X, Y, *It.Value(), GetStatsFont(), FLinearColor::White);
+				Canvas->DrawShadowedString(X, Y, *It.Value(), GetStatsFont(), LinearBodyColor);
 				Y += Height;
 			}
 		}
@@ -923,7 +1065,7 @@ int32 FAudioDebugger::RenderStatReverb(UWorld* World, FViewport* Viewport, FCanv
 	else
 	{
 		TheString = TEXT("Active Reverb Effect: None");
-		Canvas->DrawShadowedString(X, Y, *TheString, GetStatsFont(), FLinearColor::White);
+		Canvas->DrawShadowedString(X, Y, *TheString, GetStatsFont(), LinearBodyColor);
 		Y += Height;
 	}
 
@@ -945,66 +1087,165 @@ int32 FAudioDebugger::RenderStatSounds(UWorld* World, FViewport* Viewport, FCanv
 
 	const uint8 bDebug = AudioStats.DisplayFlags & static_cast<uint8>(FAudioStats::EDisplayFlags::Debug);
 
-	FString SortingName = TEXT("disabled");
+	UpdateDisplaySort(AudioStats);
 
 	// Sort the list.
-	if (AudioStats.DisplayFlags & static_cast<uint8>(FAudioStats::EDisplayFlags::Sort_Name))
-	{
-		AudioStats.StatSoundInfos.Sort([](const FAudioStats::FStatSoundInfo& A, const FAudioStats::FStatSoundInfo& B) { return A.SoundName < B.SoundName; });
-		SortingName = TEXT("pathname");
-	}
-	else if (AudioStats.DisplayFlags & static_cast<uint8>(FAudioStats::EDisplayFlags::Sort_Distance))
+	FString SortingName;
+	switch (AudioStats.DisplaySort)
 	{
-		AudioStats.StatSoundInfos.Sort([](const FAudioStats::FStatSoundInfo& A, const FAudioStats::FStatSoundInfo& B) { return A.Distance < B.Distance; });
-		SortingName = TEXT("distance");
-	}
-	else if (AudioStats.DisplayFlags & static_cast<uint8>(FAudioStats::EDisplayFlags::Sort_Class))
-	{
-		AudioStats.StatSoundInfos.Sort([](const FAudioStats::FStatSoundInfo& A, const FAudioStats::FStatSoundInfo& B) { return A.SoundClassName.LexicalLess(B.SoundClassName); });
-		SortingName = TEXT("class");
-	}
-	else if (AudioStats.DisplayFlags & static_cast<uint8>(FAudioStats::EDisplayFlags::Sort_WavesNum))
-	{
-		AudioStats.StatSoundInfos.Sort([](const FAudioStats::FStatSoundInfo& A, const FAudioStats::FStatSoundInfo& B) { return A.WaveInstanceInfos.Num() > B.WaveInstanceInfos.Num(); });
-		SortingName = TEXT("waves' num");
+		case FAudioStats::EDisplaySort::Class:
+		{
+			AudioStats.StatSoundInfos.Sort([](const FAudioStats::FStatSoundInfo& A, const FAudioStats::FStatSoundInfo& B) { return A.SoundClassName.LexicalLess(B.SoundClassName); });
+			SortingName = TEXT("Class");
+			break;
+		}
+		case FAudioStats::EDisplaySort::Distance:
+		{
+			AudioStats.StatSoundInfos.Sort([](const FAudioStats::FStatSoundInfo& A, const FAudioStats::FStatSoundInfo& B) { return A.Distance < B.Distance; });
+			SortingName = TEXT("Distance");
+			break;
+		}
+		case FAudioStats::EDisplaySort::PlaybackTime:
+		{
+			AudioStats.StatSoundInfos.Sort([](const FAudioStats::FStatSoundInfo& A, const FAudioStats::FStatSoundInfo& B) { return A.PlaybackTime > B.PlaybackTime; });
+			SortingName = TEXT("Time");
+			break;
+		}
+		case FAudioStats::EDisplaySort::Priority:
+		{
+			AudioStats.StatSoundInfos.Sort([](const FAudioStats::FStatSoundInfo& A, const FAudioStats::FStatSoundInfo& B) { return A.Priority > B.Priority; });
+			SortingName = TEXT("Priority");
+			break;
+		}
+		case FAudioStats::EDisplaySort::Waves:
+		{
+			AudioStats.StatSoundInfos.Sort([](const FAudioStats::FStatSoundInfo& A, const FAudioStats::FStatSoundInfo& B) { return A.WaveInstanceInfos.Num() > B.WaveInstanceInfos.Num(); });
+			SortingName = TEXT("Waves");
+			break;
+		}
+		case FAudioStats::EDisplaySort::Volume:
+		{
+			AudioStats.StatSoundInfos.Sort([](const FAudioStats::FStatSoundInfo& A, const FAudioStats::FStatSoundInfo& B) { return A.Volume > B.Volume; });
+			SortingName = TEXT("Volume");
+			break;
+		}
+		case FAudioStats::EDisplaySort::Name:
+		default:
+		{
+			AudioStats.StatSoundInfos.Sort([](const FAudioStats::FStatSoundInfo& A, const FAudioStats::FStatSoundInfo& B)
+			{
+				if (AudioDebugSoundShowPathCVar)
+				{
+					return A.SoundPath < B.SoundPath;
+				}
+				return A.SoundName < B.SoundName;
+			});
+			SortingName = TEXT("Name");
+			break;
+		}
 	}
 
 	Canvas->DrawShadowedString(X, Y, TEXT("Active Sounds:"), GetStatsFont(), HeaderColor);
 	Y += FontHeight;
 
-	const FString InfoText = FString::Printf(TEXT(" Sorting: %s, Debug: %s"), *SortingName, bDebug ? TEXT("enabled") : TEXT("disabled"));
+	const FString InfoText = FString::Printf(TEXT(" Sorting By: %s, Visualize Attenuation: %s"), *SortingName, bDebug ? TEXT("Enabled") : TEXT("Disabled"));
 	Canvas->DrawShadowedString(X, Y, *InfoText, GetStatsFont(), FColor(128, 255, 128));
 	Y += FontHeight;
 
-	Canvas->DrawShadowedString(X, Y, TEXT("Index Path (Class) Distance"), GetStatsFont(), BodyColor);
+	static int32 FieldLength = 40;
+	static int32 NameHeaderPad = static_cast<int32>(FieldLength * 1.25f);
+	if (AudioStats.DisplaySort == FAudioStats::EDisplaySort::Name || AudioStats.DisplaySort == FAudioStats::EDisplaySort::Distance)
+	{
+		const FString FieldName = FString(TEXT("Distance")).RightPad(NameHeaderPad).Left(NameHeaderPad);
+		Canvas->DrawShadowedString(X, Y, *FString::Printf(TEXT("Index %s Name"), *FieldName), GetStatsFont(), GetBodyColor());
+	}
+	else
+	{
+		const FString FieldName = SortingName.RightPad(FieldLength).Left(FieldLength);
+		Canvas->DrawShadowedString(X, Y, *FString::Printf(TEXT("Index %s Name"), *FieldName), GetStatsFont(), GetBodyColor());
+	}
 	Y += FontHeight;
 
 	int32 TotalSoundWavesNum = 0;
+	const FColor BodyColor = GetStatSoundBodyColor();
 	for (int32 SoundIndex = 0; SoundIndex < AudioStats.StatSoundInfos.Num(); ++SoundIndex)
 	{
 		const FAudioStats::FStatSoundInfo& StatSoundInfo = AudioStats.StatSoundInfos[SoundIndex];
 		const int32 WaveInstancesNum = StatSoundInfo.WaveInstanceInfos.Num();
-		if (WaveInstancesNum > 0)
+		if (WaveInstancesNum == 0)
 		{
+			continue;
+		}
+		TotalSoundWavesNum += WaveInstancesNum;
+
+		if (SoundIndex >= AudioDebugSoundMaxNumDisplayedCVar)
+		{
+			continue;
+		}
+
+		bool bDisplayWaves = false;
+		const FString DisplayName = AudioDebugSoundShowPathCVar ? StatSoundInfo.SoundPath : StatSoundInfo.SoundName;
+		FString DebugValue;
+		switch (AudioStats.DisplaySort)
+		{
+			case FAudioStats::EDisplaySort::Class:
 			{
-				const FString TheString = FString::Printf(TEXT("%4i. %s (%s) %6.2f"), SoundIndex, *StatSoundInfo.SoundName, *StatSoundInfo.SoundClassName.ToString(), StatSoundInfo.Distance);
-				Canvas->DrawShadowedString(X, Y, *TheString, GetStatsFont(), FColor::White);
-				Y += FontHeight;
+				DebugValue = *StatSoundInfo.SoundClassName.ToString();
+				break;
+			}
+			case FAudioStats::EDisplaySort::PlaybackTime:
+			{
+				if (FMath::IsNearlyEqual(StatSoundInfo.PlaybackTime, StatSoundInfo.PlaybackTimeNonVirtualized))
+				{
+					DebugValue = FString::Printf(TEXT("%08.2f"), StatSoundInfo.PlaybackTime);
+				}
+				else
+				{
+					DebugValue = FString::Printf(TEXT("%08.2f (%08.2f Non-Virt)"), StatSoundInfo.PlaybackTime, StatSoundInfo.PlaybackTimeNonVirtualized);
+				}
+				break;
+			}
+			case FAudioStats::EDisplaySort::Priority:
+			{
+				DebugValue = FString::Printf(TEXT("%06.2f"), StatSoundInfo.Priority);
+				break;
+			}
+			case FAudioStats::EDisplaySort::Volume:
+			{
+				DebugValue = FString::Printf(TEXT("%01.2f (%04.2f dB)"), StatSoundInfo.Volume, Audio::ConvertToDecibels(StatSoundInfo.Volume));
+				break;
+			}
+			case FAudioStats::EDisplaySort::Waves:
+			{
+				DebugValue = FString::Printf(TEXT("%03u"), StatSoundInfo.WaveInstanceInfos.Num());
+				bDisplayWaves = true;
+				break;
+			}
+			case FAudioStats::EDisplaySort::Name:
+			case FAudioStats::EDisplaySort::Distance:
+			default:
+			{
+				DebugValue = FString::Printf(TEXT("%08.2f"), StatSoundInfo.Distance);
+				break;
 			}
+		}
 
-			TotalSoundWavesNum += WaveInstancesNum;
+		const FString DebugStr = FString::Printf(TEXT("  %03i    %s %s"), SoundIndex, *DebugValue.RightPad(FieldLength - 4).Left(FieldLength - 4), *DisplayName);
+		Canvas->DrawShadowedString(X, Y, *DebugStr, GetStatsFont(), BodyColor);
+		Y += FontHeight;
 
-			// Get the active sound waves.
+		if (bDisplayWaves)
+		{
 			for (int32 WaveIndex = 0; WaveIndex < WaveInstancesNum; WaveIndex++)
 			{
-				const FString TheString = *FString::Printf(TEXT("    %4i. %s"), WaveIndex, *StatSoundInfo.WaveInstanceInfos[WaveIndex].Description);
-				Canvas->DrawShadowedString(X, Y, *TheString, GetStatsFont(), FColor(205, 205, 205));
+				const FString WaveStr = *FString::Printf(TEXT("    %02i    %s"), WaveIndex, *StatSoundInfo.WaveInstanceInfos[WaveIndex].Description);
+				Canvas->DrawShadowedString(X, Y, *WaveStr, GetStatsFont(), FColor(205, 205, 205));
 				Y += FontHeight;
 			}
 		}
 	}
 
-	Canvas->DrawShadowedString(X, Y, *FString::Printf(TEXT("Audio Device ID: %u"), AudioDevice->DeviceHandle), GetStatsFont(), HeaderColor);
+	Canvas->DrawShadowedString(X, Y, *FString::Printf(TEXT("Audio Device ID: %u, Max Sounds Displayed: %i"), AudioDevice->DeviceHandle, AudioDebugSoundMaxNumDisplayedCVar), GetStatsFont(), HeaderColor);
 	Y += FontHeight;
 
 	Canvas->DrawShadowedString(X, Y, *FString::Printf(TEXT("Total Sounds: %i, Sound Waves: %i"), AudioStats.StatSoundInfos.Num(), TotalSoundWavesNum), GetStatsFont(), HeaderColor);
@@ -1035,7 +1276,8 @@ int32 FAudioDebugger::RenderStatSounds(UWorld* World, FViewport* Viewport, FCanv
 
 			if (StatSoundInfo.ShapeDetailsMap.Num() > 0)
 			{
-				DrawDebugString(World, SoundTransform.GetTranslation(), StatSoundInfo.SoundName, NULL, FColor::White, 0.01f);
+				const FString DebugName = AudioDebugSoundShowPathCVar ? StatSoundInfo.SoundPath : StatSoundInfo.SoundName;
+				DrawDebugString(World, SoundTransform.GetTranslation(), *DebugName, nullptr, BodyColor, 0.01f);
 
 				for (auto ShapeDetailsIt = StatSoundInfo.ShapeDetailsMap.CreateConstIterator(); ShapeDetailsIt; ++ShapeDetailsIt)
 				{
@@ -1120,7 +1362,7 @@ int32 FAudioDebugger::RenderStatWaves(UWorld* World, FViewport* Viewport, FCanva
 	const int32 FontHeight = GetStatsFont()->GetMaxCharHeight() + 2;
 
 	FAudioStats& AudioStats = AudioDeviceStats.FindOrAdd(AudioDevice->DeviceHandle);
-	Canvas->DrawShadowedString(X, Y, TEXT("Active Sound Waves:"), GetStatsFont(), FLinearColor::White);
+	Canvas->DrawShadowedString(X, Y, TEXT("Active Sound Waves:"), GetStatsFont(), FLinearColor::Green);
 	Y += TabWidth;
 
 	using FWaveInstancePair = TPair<const FAudioStats::FStatWaveInstanceInfo*, const FAudioStats::FStatSoundInfo*>;
@@ -1138,6 +1380,7 @@ int32 FAudioDebugger::RenderStatWaves(UWorld* World, FViewport* Viewport, FCanva
 
 	WaveInstances.Sort([](const FWaveInstancePair& A, const FWaveInstancePair& B) { return A.Key->InstanceIndex < B.Key->InstanceIndex; });
 
+	const FColor BodyColor = GetBodyColor();
 	for (const FWaveInstancePair& WaveInstanceInfo : WaveInstances)
 	{
 		UAudioComponent* AudioComponent = UAudioComponent::GetAudioComponentFromID(WaveInstanceInfo.Value->AudioComponentID);
@@ -1149,7 +1392,7 @@ int32 FAudioDebugger::RenderStatWaves(UWorld* World, FViewport* Viewport, FCanva
 			*WaveInstanceInfo.Key->WaveInstanceName.ToString(),
 			SoundOwner ? *SoundOwner->GetName() : TEXT("None"),
 			*WaveInstanceInfo.Value->SoundClassName.ToString());
-		Canvas->DrawShadowedString(X, Y, *TheString, GetStatsFont(), WaveInstanceInfo.Key->bPlayWhenSilent == 0 ? FColor::White : FColor::Yellow);
+		Canvas->DrawShadowedString(X, Y, *TheString, GetStatsFont(), WaveInstanceInfo.Key->bPlayWhenSilent == 0 ? BodyColor : FColor::Yellow);
 		Y += FontHeight;
 	}
 
@@ -1298,23 +1541,23 @@ bool FAudioDebugger::ToggleStatSounds(UWorld* World, FCommonViewportClient* View
 	const bool bHelp = Stream ? FCString::Stristr(Stream, TEXT("?")) != nullptr : false;
 	if (bHelp)
 	{
-		GLog->Logf(TEXT("stat sounds description"));
-		GLog->Logf(TEXT("  stat sounds off - Disables drawing stat sounds"));
-		GLog->Logf(TEXT("  stat sounds sort=distance|class|name|waves|default"));
+		GLog->Logf(TEXT("Shows all active sounds. Displays value sorted by when sort is set"));
+		GLog->Logf(TEXT("  stat sounds sort=class|distance|name|priority|time|volume|waves"));
 		GLog->Logf(TEXT("      distance - sort list by distance to player"));
 		GLog->Logf(TEXT("      class - sort by sound class name"));
-		GLog->Logf(TEXT("      name - sort by cue pathname"));
+		GLog->Logf(TEXT("      name - sort by cue pathname (default)"));
+		GLog->Logf(TEXT("      time - sort by time played back"));
 		GLog->Logf(TEXT("      waves - sort by waves' num"));
-		GLog->Logf(TEXT("      default - sorting is disabled"));
-		GLog->Logf(TEXT("  stat sounds -debug - enables debugging mode like showing sound radius sphere and names, but only for cues with enabled property bDebug"));
-		GLog->Logf(TEXT("  stat sounds -smalltext - use large text in debug output (default)"));
-		GLog->Logf(TEXT("  stat sounds -largetext - use large text in debug output"));
+		GLog->Logf(TEXT("      priority - sort by playback priority"));
+		GLog->Logf(TEXT("      volume - sort by volume"));
+		GLog->Logf(TEXT("  stat sounds -debug - enables debugging mode like showing sound radius sphere and names, but only for sounds with enabled 'Debug' property"));
 		GLog->Logf(TEXT(""));
 		GLog->Logf(TEXT("Ex. stat sounds sort=class -debug"));
-		GLog->Logf(TEXT(" This will show only debug sounds sorted by sound class"));
+		GLog->Logf(TEXT(" This will show only sounds that have 'Debug' property set to true and will sort them by and display their sound class"));
 	}
 
 	uint8 ShowSounds = 0;
+	FAudioStats::EDisplaySort DisplaySort = FAudioStats::EDisplaySort::Name;
 	if (Stream)
 	{
 		const bool bHide = FParse::Command(&Stream, TEXT("off"));
@@ -1333,26 +1576,9 @@ bool FAudioDebugger::ToggleStatSounds(UWorld* World, FCommonViewportClient* View
 			}
 
 			FString SortStr;
-			FParse::Value(Stream, TEXT("sort="), SortStr);
-			if (SortStr == TEXT("distance"))
-			{
-				ShowSounds |= static_cast<uint8>(FAudioStats::EDisplayFlags::Sort_Distance);
-			}
-			else if (SortStr == TEXT("class"))
-			{
-				ShowSounds |= static_cast<uint8>(FAudioStats::EDisplayFlags::Sort_Class);
-			}
-			else if (SortStr == TEXT("name"))
-			{
-				ShowSounds |= static_cast<uint8>(FAudioStats::EDisplayFlags::Sort_Name);
-			}
-			else if (SortStr == TEXT("waves"))
+			if (FParse::Value(Stream, TEXT("sort="), SortStr))
 			{
-				ShowSounds |= static_cast<uint8>(FAudioStats::EDisplayFlags::Sort_WavesNum);
-			}
-			else
-			{
-				ShowSounds |= static_cast<uint8>(FAudioStats::EDisplayFlags::Sort_Disabled);
+				AudioDebugSoundSortCVarCVar = SortStr;
 			}
 		}
 	}
@@ -1361,6 +1587,7 @@ bool FAudioDebugger::ToggleStatSounds(UWorld* World, FCommonViewportClient* View
 	check(AudioDevice);
 	FAudioStats& Stats = AudioDeviceStats.FindOrAdd(AudioDevice->DeviceHandle);
 	Stats.DisplayFlags = ShowSounds;
+	Stats.DisplaySort = DisplaySort;
 
 	ResolveDesiredStats(ViewportClient);
 
@@ -1390,7 +1617,11 @@ void FAudioDebugger::SendUpdateResultsToGameThread(const FAudioDevice& AudioDevi
 	
 	const bool bDebug = (RequestedStats & ERequestedAudioStats::DebugSounds) != 0;
 
-	static const uint8 SoundMask = ERequestedAudioStats::Sounds | ERequestedAudioStats::SoundCues | ERequestedAudioStats::SoundWaves;
+	static const uint8 SoundMask =
+		ERequestedAudioStats::Sounds |
+		ERequestedAudioStats::SoundCues |
+		ERequestedAudioStats::SoundMixes |
+		ERequestedAudioStats::SoundWaves;
 	if (RequestedStats & SoundMask)
 	{
 		for (FActiveSound* ActiveSound : AudioDevice.GetActiveSounds())
@@ -1401,8 +1632,18 @@ void FAudioDebugger::SendUpdateResultsToGameThread(const FAudioDevice& AudioDevi
 				{
 					ActiveSoundToInfoIndex.Add(ActiveSound, StatSoundInfos.AddDefaulted());
 					FAudioStats::FStatSoundInfo& StatSoundInfo = StatSoundInfos.Last();
-					StatSoundInfo.SoundName = ActiveSound->GetSound()->GetPathName();
+					StatSoundInfo.SoundName = SoundBase->GetName();
+					StatSoundInfo.SoundPath = SoundBase->GetPathName();
 					StatSoundInfo.Distance = AudioDevice.GetDistanceToNearestListener(ActiveSound->Transform.GetTranslation());
+					StatSoundInfo.PlaybackTime = ActiveSound->PlaybackTime;
+					StatSoundInfo.Priority = ActiveSound->GetPriority();
+					StatSoundInfo.PlaybackTimeNonVirtualized = ActiveSound->PlaybackTimeNonVirtualized;
+					
+					StatSoundInfo.Volume = 0.0f;
+					for (const TPair<UPTRINT, FWaveInstance*>& Pair : ActiveSound->GetWaveInstances())
+					{
+						StatSoundInfo.Volume = FMath::Max(StatSoundInfo.Volume, Pair.Value->GetVolumeWithDistanceAttenuation() * Pair.Value->GetDynamicVolume());
+					}
 
 					if (USoundClass* SoundClass = ActiveSound->GetSoundClass())
 					{
@@ -1413,13 +1654,14 @@ void FAudioDebugger::SendUpdateResultsToGameThread(const FAudioDevice& AudioDevi
 					{
 						StatSoundInfo.SoundClassName = NAME_None;
 					}
+
 					StatSoundInfo.Transform = ActiveSound->Transform;
 					StatSoundInfo.AudioComponentID = ActiveSound->GetAudioComponentID();
-										
+
 					if (bDebug && ActiveSound->GetSound()->bDebug)
 					{
 						ActiveSound->CollectAttenuationShapesForVisualization(StatSoundInfo.ShapeDetailsMap);
-					}										
+					}
 				}
 			}
 		}
diff --git a/Engine/Source/Runtime/Engine/Private/AudioDevice.cpp b/Engine/Source/Runtime/Engine/Private/AudioDevice.cpp
--- a/Engine/Source/Runtime/Engine/Private/AudioDevice.cpp
+++ b/Engine/Source/Runtime/Engine/Private/AudioDevice.cpp
@@ -6143,6 +6143,10 @@ void FAudioDevice::OnBeginPIE(const bool bIsSimulating)
 		USoundNode* SoundNode = *It;
 		SoundNode->OnBeginPIE(bIsSimulating);
 	}
+
+#if ENABLE_AUDIO_DEBUG
+	FAudioDebugger::OnBeginPIE();
+#endif // WITH_EDITOR
 }
 
 void FAudioDevice::OnEndPIE(const bool bIsSimulating)
@@ -6152,6 +6156,10 @@ void FAudioDevice::OnEndPIE(const bool bIsSimulating)
 		USoundNode* SoundNode = *It;
 		SoundNode->OnEndPIE(bIsSimulating);
 	}
+
+#if ENABLE_AUDIO_DEBUG
+	FAudioDebugger::OnEndPIE();
+#endif // WITH_EDITOR
 }
 #endif
 
diff --git a/Engine/Source/Runtime/Engine/Public/Audio/AudioDebug.h b/Engine/Source/Runtime/Engine/Public/Audio/AudioDebug.h
--- a/Engine/Source/Runtime/Engine/Public/Audio/AudioDebug.h
+++ b/Engine/Source/Runtime/Engine/Public/Audio/AudioDebug.h
@@ -69,7 +69,12 @@ public:
 
 	bool IsVisualizeDebug3dEnabled() const;
 	void ToggleVisualizeDebug3dEnabled();
-		
+
+#if WITH_EDITOR
+	static void OnBeginPIE();
+	static void OnEndPIE();
+#endif // WITH_EDITOR
+
 	// Evaluate Mute/Solos
 	void QuerySoloMuteSoundClass(const FString& Name, bool& bOutIsSoloed, bool& bOutIsMuted, FString& OutReason ) const;
 	void QuerySoloMuteSoundWave(const FString& Name, bool& bOutIsSoloed, bool& bOutIsMuted, FString& OutReason) const;
