Fix non-editor compile issue from last CL

diff --git a/Engine/Source/Runtime/Engine/Private/AudioDebug.cpp b/Engine/Source/Runtime/Engine/Private/AudioDebug.cpp
--- a/Engine/Source/Runtime/Engine/Private/AudioDebug.cpp
+++ b/Engine/Source/Runtime/Engine/Private/AudioDebug.cpp
@@ -333,6 +333,7 @@ void FAudioDebugger::ToggleVisualizeDebug3dEnabled()
 	bVisualize3dDebug = !bVisualize3dDebug;
 }
 
+#if WITH_EDITOR
 void FAudioDebugger::OnBeginPIE()
 {
 	FAudioDeviceManager* DeviceManager = GEngine->GetAudioDeviceManager();
@@ -350,6 +351,7 @@ void FAudioDebugger::OnEndPIE()
 		DeviceManager->GetDebugger().ClearMutesAndSolos();
 	}
 }
+#endif // WITH_EDITOR
 
 void FAudioDebugger::QuerySoloMuteSoundClass(const FString& Name, bool& bOutIsSoloed, bool& bOutIsMuted, FString& OutReason) const
 {
diff --git a/Engine/Source/Runtime/Engine/Private/AudioDevice.cpp b/Engine/Source/Runtime/Engine/Private/AudioDevice.cpp
--- a/Engine/Source/Runtime/Engine/Private/AudioDevice.cpp
+++ b/Engine/Source/Runtime/Engine/Private/AudioDevice.cpp
@@ -6146,7 +6146,7 @@ void FAudioDevice::OnBeginPIE(const bool bIsSimulating)
 
 #if ENABLE_AUDIO_DEBUG
 	FAudioDebugger::OnBeginPIE();
-#endif // WITH_EDITOR
+#endif // ENABLE_AUDIO_DEBUG
 }
 
 void FAudioDevice::OnEndPIE(const bool bIsSimulating)
@@ -6159,9 +6159,9 @@ void FAudioDevice::OnEndPIE(const bool bIsSimulating)
 
 #if ENABLE_AUDIO_DEBUG
 	FAudioDebugger::OnEndPIE();
-#endif // WITH_EDITOR
+#endif // ENABLE_AUDIO_DEBUG
 }
-#endif
+#endif // WITH_EDITOR
 
 bool FAudioDevice::CanUseVRAudioDevice()
 {
