Fix attenuation box shape reporting 0 incorrectly when distance is 0 (effectively point is inside box) and falloff is set to 0.

diff --git a/Engine/Source/Runtime/Engine/Private/Attenuation.cpp b/Engine/Source/Runtime/Engine/Private/Attenuation.cpp
--- a/Engine/Source/Runtime/Engine/Private/Attenuation.cpp
+++ b/Engine/Source/Runtime/Engine/Private/Attenuation.cpp
@@ -135,7 +135,7 @@ float FBaseAttenuationSettings::AttenuationEval(const float Distance, const floa
 float FBaseAttenuationSettings::AttenuationEvalBox(const FTransform& Origin, const FVector Location, const float DistanceScale) const
 {
 	const float DistanceSq = ComputeSquaredDistanceFromBoxToPoint(-AttenuationShapeExtents, AttenuationShapeExtents, Origin.InverseTransformPositionNoScale(Location));
-	if (DistanceSq < FalloffDistance * FalloffDistance)
+	if (FMath::IsNearlyZero(DistanceSq) || DistanceSq < FalloffDistance * FalloffDistance)
 	{ 
 		return AttenuationEval(FMath::Sqrt(DistanceSq), FalloffDistance, DistanceScale);
 	}
