Fix mirror window sizing logic bug introduced by a platform-specific fix, by only applying resolution check logic on that specific platform.

diff --git a/Engine/Source/Runtime/Engine/Private/Slate/SceneViewport.cpp b/Engine/Source/Runtime/Engine/Private/Slate/SceneViewport.cpp
--- a/Engine/Source/Runtime/Engine/Private/Slate/SceneViewport.cpp
+++ b/Engine/Source/Runtime/Engine/Private/Slate/SceneViewport.cpp
@@ -1328,7 +1328,12 @@ void FSceneViewport::ResizeFrame(uint32 NewWindowSizeX, uint32 NewWindowSizeY, E
 				IHeadMountedDisplay::MonitorInfo MonitorInfo;
 				if (GEngine->XRSystem.IsValid() && GEngine->XRSystem->GetHMDDevice() && GEngine->XRSystem->GetHMDDevice()->GetHMDMonitorInfo(MonitorInfo))
 				{
+#if PLATFORM_PS4
+					// Only do the resolution check on PS4/Morpheus. On desktop, this breaks the mirror window logic.
 					if (MonitorInfo.DesktopX > 0 || MonitorInfo.DesktopY > 0 || MonitorInfo.ResolutionX > 0 || MonitorInfo.ResolutionY > 0)
+#else
+					if (MonitorInfo.DesktopX > 0 || MonitorInfo.DesktopY > 0)
+#endif
 					{
 						NewWindowSize.X = MonitorInfo.ResolutionX;
 						NewWindowSize.Y = MonitorInfo.ResolutionY;
