More workarounds for corrupt instancedata.  These should be removed when the culprit is found

diff --git a/Engine/Source/Runtime/Engine/Private/InstancedStaticMesh.cpp b/Engine/Source/Runtime/Engine/Private/InstancedStaticMesh.cpp
--- a/Engine/Source/Runtime/Engine/Private/InstancedStaticMesh.cpp
+++ b/Engine/Source/Runtime/Engine/Private/InstancedStaticMesh.cpp
@@ -262,11 +262,11 @@ void FStaticMeshInstanceBuffer::UpdateFromCommandBuffer_RenderThread(FInstanceUp
 	int32 NumCommands = CmdBuffer.NumInlineCommands();
 	int32 NumAdds = CmdBuffer.NumAdds;
 	int32 AddIndex = INDEX_NONE;
-
+	int32 NewNumInstances = InstanceData->GetNumInstances();
 	if (NumAdds > 0)
 	{
 		AddIndex = InstanceData->GetNumInstances();
-		int32 NewNumInstances = NumAdds + InstanceData->GetNumInstances();
+		NewNumInstances = NumAdds + InstanceData->GetNumInstances();
 		InstanceData->AllocateInstances(NewNumInstances, GIsEditor ? EResizeBufferFlags::AllowSlackOnGrow | EResizeBufferFlags::AllowSlackOnReduce : EResizeBufferFlags::None, false); // In Editor always permit overallocation, to prevent too much realloc
 	}
 
@@ -279,16 +279,28 @@ void FStaticMeshInstanceBuffer::UpdateFromCommandBuffer_RenderThread(FInstanceUp
 			InstanceData->SetInstance(AddIndex++, Cmd.XForm, 0);
 			break;
 		case FInstanceUpdateCmdBuffer::Hide:
-			InstanceData->NullifyInstance(Cmd.InstanceIndex);
+			if (Cmd.InstanceIndex < NewNumInstances)
+			{
+				InstanceData->NullifyInstance(Cmd.InstanceIndex);
+			}
 			break;
 		case FInstanceUpdateCmdBuffer::Update:
-			InstanceData->SetInstance(Cmd.InstanceIndex, Cmd.XForm, 0);
+			if (Cmd.InstanceIndex < NewNumInstances)
+			{
+				InstanceData->SetInstance(Cmd.InstanceIndex, Cmd.XForm, 0);
+			}
 			break;
 		case FInstanceUpdateCmdBuffer::EditorData:
-			InstanceData->SetInstanceEditorData(Cmd.InstanceIndex, Cmd.HitProxyColor, Cmd.bSelected);
+			if (Cmd.InstanceIndex < NewNumInstances)
+			{
+				InstanceData->SetInstanceEditorData(Cmd.InstanceIndex, Cmd.HitProxyColor, Cmd.bSelected);
+			}
 			break;
 		case FInstanceUpdateCmdBuffer::LightmapData:
-			InstanceData->SetInstanceLightMapData(Cmd.InstanceIndex, Cmd.LightmapUVBias, Cmd.ShadowmapUVBias);
+			if (Cmd.InstanceIndex < NewNumInstances)
+			{
+				InstanceData->SetInstanceLightMapData(Cmd.InstanceIndex, Cmd.LightmapUVBias, Cmd.ShadowmapUVBias);
+			}
 			break;
 		default:
 			check(false);
@@ -1973,7 +1985,7 @@ void UInstancedStaticMeshComponent::SerializeRenderData(FArchive& Ar)
 					for (int32 Index = 0; Index < NumInstances; ++Index)
 					{
 						int32 RenderIndex = InstanceReorderTable.IsValidIndex(Index) ? InstanceReorderTable[Index] : Index;
-						if (RenderIndex == INDEX_NONE)
+						if (RenderIndex == INDEX_NONE || RenderIndex >= NumInstances)
 						{
 							// could be skipped by density settings
 							continue;
@@ -1992,7 +2004,7 @@ void UInstancedStaticMeshComponent::SerializeRenderData(FArchive& Ar)
 					for (int32 Index = 0; Index < NumInstances; ++Index)
 					{
 						int32 RenderIndex = InstanceReorderTable.IsValidIndex(Index) ? InstanceReorderTable[Index] : Index;
-						if (RenderIndex == INDEX_NONE)
+						if (RenderIndex == INDEX_NONE || RenderIndex >= NumInstances)
 						{
 							// could be skipped by density settings
 							continue;
diff --git a/Engine/Source/Runtime/Engine/Public/StaticMeshResources.h b/Engine/Source/Runtime/Engine/Public/StaticMeshResources.h
--- a/Engine/Source/Runtime/Engine/Public/StaticMeshResources.h
+++ b/Engine/Source/Runtime/Engine/Public/StaticMeshResources.h
@@ -1532,8 +1532,8 @@ private:
 	{
 		FVector4* ElementData = reinterpret_cast<FVector4*>(InstanceOriginDataPtr);
 		uint32 CurrentSize = InstanceOriginData->Num() * InstanceOriginData->GetStride();
-		check((void*)((&ElementData[InstanceIndex]) + 1) <= (void*)(InstanceOriginDataPtr + CurrentSize));
-		check((void*)((&ElementData[InstanceIndex]) + 0) >= (void*)(InstanceOriginDataPtr));
+		checkf((void*)((&ElementData[InstanceIndex]) + 1) <= (void*)(InstanceOriginDataPtr + CurrentSize), TEXT("OOB Instance Set Under: %i, %u, %p, %p"), InstanceIndex, CurrentSize, &ElementData, InstanceOriginDataPtr);
+		checkf((void*)((&ElementData[InstanceIndex]) + 0) >= (void*)(InstanceOriginDataPtr), TEXT("OOB Instance Set: %i, %u, %p, %p"), InstanceIndex, CurrentSize, &ElementData, InstanceOriginDataPtr);
 
 		ElementData[InstanceIndex] = Origin;
 	}
