Fix animation root motion issues causing network corrections
- 1) Flush pending server moves before LocalPredicted abilities activate. This flush is necessary BEFORE the ServerTryActivateAbility RPC is called so that the ServerMove RPC is received and processed by the server prior to the server activating ability logic. Prior to this fix, the client would tick their movement and save it as pending, then activate an ability which would apply an animation root motion, and afterwards send the ServerMove. The server would receive the ability activation first (which would apply an animation root motion) and afterwards process the ServerMove, but the ServerMove would not contain any animation root motion movement and trigger a correction.
- 2) Fix a bug in AbilitySystemComponent ServerCurrentMontageSetNextSectionName that was incorrectly setting montage position every time. This was causing a correction at the end of every decision window - the client and server would be progressing normally through an animation root motion, during the client movement tick they would tick their montage, triggering the decision window notify to send the ServerCurrentMontageSetNextSectionName RPC with the next section data, and then complete their movement and send the ServerMove RPC. The server would receive the ServerCurrentMontageSetNextSectionName RPC FIRST, incorrectly skip forward in the montage, and THEN tick the received ServerMove, which would not match what the client said because the server would be ticking its montage from the end of what the client sent. This meant the server was always skipping a section of animation root motion involving decision windows/changing sections locally on clients.
- 3) Fix the server from ticking character animation twice for a single ServerMove on moves where an animation root motion ends. In MoveAutonomous() after PerformMovement() we were only checking if we CURRENTLY had root motion to guard against calling TickCharacterPose again. We now save off whether we were playing animation root motion prior to the PerformMovement (meaning TickCharacterPose will get called during the PerformMovement) and skip the double TickCharacterPose call. Current CharacterMesh->TickPose code was fixed previously to allow multiple ticks per frame for when the server receives multiple ServerMoves in a single frame, but this edge case was caused which is now fixed.
- 4) Prevent the first Move after an animation root motion montage has started being combined with a previous PendingMove causing network corrections for that ServerMove. Currently there is logic in CanCombineWith that checks if the PendingMove or NewMove contain animation root motion and disallows combining, but unfortunately this CanCombineWith check is run before we run PerformMovement and NewMove->PostUpdate, which is where NewMove gets its RootMotionMontage set. This meant that the first frame of movement containing animation root motion was always being able to be combined even though we have a guard that is meant to prevent this. We now set bForceNoCombine in FSavedMove_Character::SetMoveFor when the character has animation root motion, which runs before the CanCombineWith code, and we also additionally set bForceNoCombine in PostUpdate for anim root motion-containing moves so that they will not get delayed and will be sent that frame.
- 5) Fix Character AnimRootMotionTranslationScale not being applied during clientside resimulation of SavedMove root motion after a correction, causing errors and leading to additional corrections. This isn't a perfect fix as it reads current scale instead of a scale saved in a SavedMove, but is better than the current situation and doesn't introduce the possibility of your manually-set scale getting overridden and staying incorrect indefinitely after a correction (this scale will need a push/pop or other more stateful tracking of modifications to it to work flawlessly with saved moves).

Zak.Middleton, Laurent.Delayen, Jaren.Peterson, Frank.Gigliotti
#rnx

diff --git a/Engine/Plugins/Runtime/GameplayAbilities/Source/GameplayAbilities/Private/AbilitySystemComponent_Abilities.cpp b/Engine/Plugins/Runtime/GameplayAbilities/Source/GameplayAbilities/Private/AbilitySystemComponent_Abilities.cpp
--- a/Engine/Plugins/Runtime/GameplayAbilities/Source/GameplayAbilities/Private/AbilitySystemComponent_Abilities.cpp
+++ b/Engine/Plugins/Runtime/GameplayAbilities/Source/GameplayAbilities/Private/AbilitySystemComponent_Abilities.cpp
@@ -31,6 +31,8 @@
 #include "TickableAttributeSetInterface.h"
 #include "GameplayTagResponseTable.h"
 #include "ProfilingDebugging/CsvProfiler.h"
+#include "GameFramework/Character.h"
+#include "GameFramework/CharacterMovementComponent.h"
 
 #define LOCTEXT_NAMESPACE "AbilitySystemComponent"
 
@@ -1351,6 +1353,22 @@ bool UAbilitySystemComponent::InternalTryActivateAbility(FGameplayAbilitySpecHan
 	}
 	else if (Ability->GetNetExecutionPolicy() == EGameplayAbilityNetExecutionPolicy::LocalPredicted)
 	{
+		// Flush server moves that occurred before this ability activation so that the server receives the RPCs in the correct order
+		// Necessary to prevent abilities that trigger animation root motion or impact movement from causing network corrections
+		AActor* const AvatarActor = ActorInfo->AvatarActor.Get();
+		if (AvatarActor && !ActorInfo->IsNetAuthority())
+		{
+			ACharacter* AvatarCharacter = Cast<ACharacter>(AvatarActor);
+			if (AvatarCharacter)
+			{
+				UCharacterMovementComponent* AvatarCharMoveComp = Cast<UCharacterMovementComponent>(AvatarCharacter->GetMovementComponent());
+				if (AvatarCharMoveComp)
+				{
+					AvatarCharMoveComp->FlushServerMoves();
+				}
+			}
+		}
+
 		// This execution is now officially EGameplayAbilityActivationMode:Predicting and has a PredictionKey
 		FScopedPredictionWindow ScopedPredictionWindow(this, true);
 
@@ -2729,7 +2747,7 @@ void UAbilitySystemComponent::ServerCurrentMontageSetNextSectionName_Implementat
 
 			int32 ClientSectionID = CurrentAnimMontage->GetSectionIndexFromPosition(ClientPosition);
 			FName ClientCurrentSectionName = CurrentAnimMontage->GetSectionName(ClientSectionID);
-			if ((CurrentSectionName != ClientCurrentSectionName) || (CurrentSectionName != SectionName) || (CurrentSectionName != NextSectionName))
+			if ((CurrentSectionName != ClientCurrentSectionName) || (CurrentSectionName != SectionName))
 			{
 				// We are in an invalid section, jump to client's position.
 				AnimInstance->Montage_SetPosition(CurrentAnimMontage, ClientPosition);
diff --git a/Engine/Source/Runtime/Engine/Private/Components/CharacterMovementComponent.cpp b/Engine/Source/Runtime/Engine/Private/Components/CharacterMovementComponent.cpp
--- a/Engine/Source/Runtime/Engine/Private/Components/CharacterMovementComponent.cpp
+++ b/Engine/Source/Runtime/Engine/Private/Components/CharacterMovementComponent.cpp
@@ -9017,6 +9017,8 @@ void UCharacterMovementComponent::MoveAutonomous
 	const FVector OldLocation = UpdatedComponent->GetComponentLocation();
 	const FQuat OldRotation = UpdatedComponent->GetComponentQuat();
 
+	const bool bWasPlayingRootMotion = CharacterOwner->IsPlayingRootMotion();
+
 	PerformMovement(DeltaTime);
 
 	// Check if data is valid as PerformMovement can mark character for pending kill
@@ -9028,7 +9030,10 @@ void UCharacterMovementComponent::MoveAutonomous
 	// If not playing root motion, tick animations after physics. We do this here to keep events, notifies, states and transitions in sync with client updates.
 	if( CharacterOwner && !CharacterOwner->bClientUpdating && !CharacterOwner->IsPlayingRootMotion() && CharacterOwner->GetMesh() )
 	{
-		TickCharacterPose(DeltaTime);
+		if (!bWasPlayingRootMotion) // If we were playing root motion before PerformMovement but aren't anymore, we're on the last frame of anim root motion and have already ticked character
+		{
+			TickCharacterPose(DeltaTime);
+		}
 		// TODO: SaveBaseLocation() in case tick moves us?
 
 		// Trigger Events right away, as we could be receiving multiple ServerMoves per frame.
@@ -10571,6 +10576,13 @@ void FSavedMove_Character::SetMoveFor(ACharacter* Character, float InDeltaTime,
 		bForceNoCombine = true;
 	}
 
+	// Moves with anim root motion should not be combined
+	const FAnimMontageInstance* RootMotionMontageInstance = Character->GetRootMotionAnimMontageInstance();
+	if (RootMotionMontageInstance)
+	{
+		bForceNoCombine = true;
+	}
+
 	// Launch velocity gives instant and potentially huge change of velocity
 	// Avoid combining move to prevent from reverting locations until server catches up
 	const bool bHasLaunchVelocity = !Character->GetCharacterMovement()->PendingLaunchVelocity.IsZero();
@@ -10675,6 +10687,9 @@ void FSavedMove_Character::PostUpdate(ACharacter* Character, FSavedMove_Characte
 			RootMotionMontage = RootMotionMontageInstance->Montage;
 			RootMotionTrackPosition = RootMotionMontageInstance->GetPosition();
 			RootMotionMovement = Character->ClientRootMotionParams;
+
+			// Moves with anim root motion should not be combined
+			bForceNoCombine = true;
 		}
 
 		// Save off Root Motion Sources
@@ -11060,6 +11075,7 @@ void FSavedMove_Character::PrepMoveFor(ACharacter* Character)
 				RootMotionTrackPosition = RootMotionMontageInstance->GetPosition();
 				RootMotionMontageInstance->SimulateAdvance(DeltaTime, RootMotionTrackPosition, RootMotionMovement);
 				RootMotionMontageInstance->SetPosition(RootMotionTrackPosition);
+				RootMotionMovement.ScaleRootMotionTranslation(Character->GetAnimRootMotionTranslationScale());
 			}
 		}
 
