Blueprints: Worked around precision loss in int64 floor/trunc/round/ceil operations (see UE-84223)
#rnx

diff --git a/Engine/Source/Runtime/Engine/Classes/Kismet/KismetMathLibrary.inl b/Engine/Source/Runtime/Engine/Classes/Kismet/KismetMathLibrary.inl
--- a/Engine/Source/Runtime/Engine/Classes/Kismet/KismetMathLibrary.inl
+++ b/Engine/Source/Runtime/Engine/Classes/Kismet/KismetMathLibrary.inl
@@ -696,25 +696,25 @@ int32 UKismetMathLibrary::FTrunc(float A)
 KISMET_MATH_FORCEINLINE
 int64 UKismetMathLibrary::Round64(float A)
 {
-	return (int64)FMath::RoundToFloat(A);
+	return (int64)FMath::RoundToDouble(A);
 }
 
 KISMET_MATH_FORCEINLINE
 int64 UKismetMathLibrary::FFloor64(float A)
 {
-	return (int64)FMath::FloorToFloat(A);
+	return (int64)FMath::FloorToDouble(A);
 }
 
 KISMET_MATH_FORCEINLINE
 int64 UKismetMathLibrary::FTrunc64(float A)
 {
-	return (int64)FMath::TruncToFloat(A);
+	return (int64)A;
 }
 
 KISMET_MATH_FORCEINLINE
 int64 UKismetMathLibrary::FCeil64(float A)
 {
-	return (int64)FMath::CeilToFloat(A);
+	return (int64)FMath::CeilToDouble(A);
 }
 
 KISMET_MATH_FORCEINLINE
