Fix SQ to run iniital overlap on zero length sweep.

diff --git a/Engine/Source/Runtime/Engine/Private/Collision/SceneQuery.cpp b/Engine/Source/Runtime/Engine/Private/Collision/SceneQuery.cpp
--- a/Engine/Source/Runtime/Engine/Private/Collision/SceneQuery.cpp
+++ b/Engine/Source/Runtime/Engine/Private/Collision/SceneQuery.cpp
@@ -351,9 +351,10 @@ bool TSceneCastCommon(const UWorld* World, typename Traits::TOutHits& OutHits, c
 	bool bHaveBlockingHit = false;
 
 	FVector Delta = End - Start;
-	float DeltaMag = Delta.Size();
+	float DeltaSize = Delta.Size();
+	float DeltaMag = FMath::IsNearlyZero(DeltaSize) ? 0.f : DeltaSize;
 	float MinBlockingDistance = DeltaMag;
-	if (DeltaMag > KINDA_SMALL_NUMBER)
+	if (Traits::IsSweep() || DeltaMag > 0.f)
 	{
 		// Create filter data used to filter collisions 
 		CA_SUPPRESS(6326);
@@ -371,7 +372,7 @@ bool TSceneCastCommon(const UWorld* World, typename Traits::TOutHits& OutHits, c
 		typename Traits::THitBuffer HitBufferSync;
 
 		bool bBlockingHit = false;
-		const FVector Dir = Delta / DeltaMag;
+		const FVector Dir = DeltaMag > 0.f ? (Delta / DeltaMag) : FVector(1, 0, 0);
 		const FTransform StartTM = Traits::IsRay() ? FTransform(Start) : FTransform(*GeomInputs.GetGeometryOrientation(), Start);
 
 		// Enable scene locks, in case they are required
