PR #6352: Templated GetHUD (Contributed by [REDACTED_PERSON])
PR #6353: Template version of GetPawn (Contributed by [REDACTED_PERSON])

diff --git a/Engine/Source/Runtime/Engine/Classes/GameFramework/Controller.h b/Engine/Source/Runtime/Engine/Classes/GameFramework/Controller.h
--- a/Engine/Source/Runtime/Engine/Classes/GameFramework/Controller.h
+++ b/Engine/Source/Runtime/Engine/Classes/GameFramework/Controller.h
@@ -227,6 +227,13 @@ public:
 	/** Getter for Pawn */
 	FORCEINLINE APawn* GetPawn() const { return Pawn; }
 
+	/** Templated version of GetPawn, will return nullptr if cast fails */
+	template<class T>
+	T* GetPawn() const
+	{
+		return Cast<T>(Pawn);
+	}
+
 	/** Getter for Character */
 	FORCEINLINE ACharacter* GetCharacter() const { return Character; }
 
diff --git a/Engine/Source/Runtime/Engine/Classes/GameFramework/PlayerController.h b/Engine/Source/Runtime/Engine/Classes/GameFramework/PlayerController.h
--- a/Engine/Source/Runtime/Engine/Classes/GameFramework/PlayerController.h
+++ b/Engine/Source/Runtime/Engine/Classes/GameFramework/PlayerController.h
@@ -1039,6 +1039,13 @@ public:
 	UFUNCTION(BlueprintCallable, Category="HUD")
 	AHUD* GetHUD() const;
 
+	/** Templated version of GetHUD, will return nullptr if cast fails */
+	template<class T>
+	T* GetHUD() const
+	{
+		return Cast<T>(MyHUD);
+	}
+
 	/**
 	 * Sets the Widget for the Mouse Cursor to display 
 	 * @param Cursor - the cursor to set the widget for
