Implement GetCOM and AddForceAtLocation. Needed for vehicles.

diff --git a/Engine/Source/Runtime/Engine/Private/PhysicsEngine/Experimental/PhysInterface_Chaos.cpp b/Engine/Source/Runtime/Engine/Private/PhysicsEngine/Experimental/PhysInterface_Chaos.cpp
--- a/Engine/Source/Runtime/Engine/Private/PhysicsEngine/Experimental/PhysInterface_Chaos.cpp
+++ b/Engine/Source/Runtime/Engine/Private/PhysicsEngine/Experimental/PhysInterface_Chaos.cpp
@@ -413,7 +413,11 @@ FVector FPhysInterface_Chaos::GetWorldVelocityAtPoint_AssumesLocked(const FPhysi
 
 FTransform FPhysInterface_Chaos::GetComTransform_AssumesLocked(const FPhysicsActorHandle& InActorReference)
 {
-	// #todo : Implement
+	// NOTE: This might need to change after we add separate COM transforms.
+	if (ensure(FPhysicsInterface::IsValid(InActorReference)))
+	{
+		return FTransform(InActorReference->R(), InActorReference->X());
+	}
 	return FTransform();
 }
 
diff --git a/Engine/Source/Runtime/Engine/Private/PhysicsEngine/Experimental/PhysScene_Chaos.cpp b/Engine/Source/Runtime/Engine/Private/PhysicsEngine/Experimental/PhysScene_Chaos.cpp
--- a/Engine/Source/Runtime/Engine/Private/PhysicsEngine/Experimental/PhysScene_Chaos.cpp
+++ b/Engine/Source/Runtime/Engine/Private/PhysicsEngine/Experimental/PhysScene_Chaos.cpp
@@ -1075,9 +1075,34 @@ void FPhysScene_ChaosInterface::AddForce_AssumesLocked(FBodyInstance* BodyInstan
 
 void FPhysScene_ChaosInterface::AddForceAtPosition_AssumesLocked(FBodyInstance* BodyInstance, const FVector& Force, const FVector& Position, bool bAllowSubstepping, bool bIsLocalForce /*= false*/)
 {
-	// #todo : Implement
-	//AddTorque(Chaos::TVector<float, 3>::CrossProduct(Position - Scene.GetSolver()->GetRigidParticles().X(GetIndexFromId(Id)), Force), Id);
-	//AddForce(Force, Id);
+	FPhysicsActorHandle& Handle = BodyInstance->GetPhysicsActorHandle();
+	if (ensure(FPhysicsInterface::IsValid(Handle)))
+	{
+		Chaos::TPBDRigidParticle<float, 3>* Rigid = Handle->AsDynamic();
+		if (ensure(Rigid))
+		{
+			const Chaos::FVec3& CurrentForce = Rigid->ExternalForce();
+			const Chaos::FVec3& CurrentTorque = Rigid->ExternalTorque();
+			const Chaos::FVec3& CurrentPosition = Rigid->X();
+			const Chaos::FVec3& CurrentCOM = CurrentPosition; // TODO: Remember to fix this once we add COM to particles!
+			const Chaos::FVec3 Torque = Chaos::FVec3::CrossProduct(Position - CurrentCOM, Force);
+			if (!bIsLocalForce)
+			{
+				Rigid->SetExternalForce(CurrentForce + Force);
+				Rigid->SetExternalTorque(CurrentTorque + Torque);
+			}
+			else
+			{
+				const Chaos::FRotation3& CurrentRotation = Rigid->R();
+				const Chaos::FRigidTransform3 CurrentTransform(CurrentPosition, CurrentRotation);
+				const Chaos::FVec3 WorldPosition = CurrentTransform.TransformPosition(Position);
+				const Chaos::FVec3 WorldForce = CurrentRotation.RotateVector(Force);
+				const Chaos::FVec3 WorldTorque = Chaos::FVec3::CrossProduct(WorldPosition - CurrentCOM, WorldForce);
+				Rigid->SetExternalForce(CurrentForce + WorldForce);
+				Rigid->SetExternalTorque(CurrentTorque + WorldTorque);
+			}
+		}
+	}
 }
 
 void FPhysScene_ChaosInterface::AddRadialForceToBody_AssumesLocked(FBodyInstance* BodyInstance, const FVector& Origin, const float Radius, const float Strength, const uint8 Falloff, bool bAccelChange, bool bAllowSubstepping)
