PostLoad Fix: InstanceReorderTable indices were offsetted by PerInstanceSMData.Num() causing crash on Load.

This doesn't explain why the data got generated that way

#rnx

diff --git a/Engine/Source/Runtime/Engine/Private/HierarchicalInstancedStaticMesh.cpp b/Engine/Source/Runtime/Engine/Private/HierarchicalInstancedStaticMesh.cpp
--- a/Engine/Source/Runtime/Engine/Private/HierarchicalInstancedStaticMesh.cpp
+++ b/Engine/Source/Runtime/Engine/Private/HierarchicalInstancedStaticMesh.cpp
@@ -31,6 +31,7 @@
 #include "HAL/LowLevelMemTracker.h"
 #include "UObject/ReleaseObjectVersion.h"
 #include "ComponentRecreateRenderStateContext.h"
+#include "Algo/AnyOf.h"
 
 static TAutoConsoleVariable<int32> CVarFoliageSplitFactor(
 	TEXT("foliage.SplitFactor"),
@@ -3047,6 +3048,13 @@ void UHierarchicalInstancedStaticMeshComponent::OnPostLoadPerInstanceData()
 				bForceTreeBuild = true;
 			}
 
+			
+			// If Reorder table contains invalid indices force a tree rebuild
+			if(!bForceTreeBuild)
+			{
+				bForceTreeBuild = Algo::AnyOf(InstanceReorderTable, [this](int32 ReorderIndex) { return ReorderIndex != INDEX_NONE && !PerInstanceSMData.IsValidIndex(ReorderIndex); });
+			}
+
 			if (!bForceTreeBuild)
 			{
 				// create PerInstanceRenderData either from current data or pre-built instance buffer
