Add perf marker for virtual texture transcode

diff --git a/Engine/Source/Runtime/Engine/Private/VT/VirtualTextureTranscodeCache.cpp b/Engine/Source/Runtime/Engine/Private/VT/VirtualTextureTranscodeCache.cpp
--- a/Engine/Source/Runtime/Engine/Private/VT/VirtualTextureTranscodeCache.cpp
+++ b/Engine/Source/Runtime/Engine/Private/VT/VirtualTextureTranscodeCache.cpp
@@ -39,6 +39,8 @@ struct FTranscodeTask
 
 	void DoTask(ENamedThreads::Type CurrentThread, const FGraphEventRef& MyCompletionGraphEvent)
 	{
+		SCOPED_NAMED_EVENT(VirtualTextureTranscodeTask_DoTask, FColor::Cyan);
+
 		static uint8 Black[4] = { 0,0,0,0 };
 		static uint8 OpaqueBlack[4] = { 0,0,0,255 };
 		static uint8 White[4] = { 0xFF, 0xFF, 0xFF, 0xFF };
