Wrap InstanceReorderTable Indexing in to GetRenderIndex method + Backout Temp Workarounds in CL 10088116 + 10090973

#rnx

diff --git a/Engine/Source/Runtime/Engine/Classes/Components/InstancedStaticMeshComponent.h b/Engine/Source/Runtime/Engine/Classes/Components/InstancedStaticMeshComponent.h
--- a/Engine/Source/Runtime/Engine/Classes/Components/InstancedStaticMeshComponent.h
+++ b/Engine/Source/Runtime/Engine/Classes/Components/InstancedStaticMeshComponent.h
@@ -150,6 +150,9 @@ class ENGINE_API UInstancedStaticMeshComponent : public UStaticMeshComponent
 	/** Tracks outstanding proxysize, as this is a bit hard to do with the fire-and-forget grass. */
 	SIZE_T ProxySize;
 
+	/** Returns the render instance buffer index. */
+	FORCEINLINE int32 GetRenderIndex(int32 InInstanceIndex) const { return InstanceReorderTable.IsValidIndex(InInstanceIndex) ? InstanceReorderTable[InInstanceIndex] : InInstanceIndex; }
+
 	/** Add an instance to this component. Transform is given in local space of this component. */
 	UFUNCTION(BlueprintCallable, Category="Components|InstancedStaticMesh")
 	virtual int32 AddInstance(const FTransform& InstanceTransform);
diff --git a/Engine/Source/Runtime/Engine/Private/HierarchicalInstancedStaticMesh.cpp b/Engine/Source/Runtime/Engine/Private/HierarchicalInstancedStaticMesh.cpp
--- a/Engine/Source/Runtime/Engine/Private/HierarchicalInstancedStaticMesh.cpp
+++ b/Engine/Source/Runtime/Engine/Private/HierarchicalInstancedStaticMesh.cpp
@@ -2234,7 +2234,7 @@ bool UHierarchicalInstancedStaticMeshComponent::UpdateInstanceTransform(int32 In
 		bConcurrentChanges = true;
 	}
 
-	int32 RenderIndex = InstanceReorderTable.IsValidIndex(InstanceIndex) ? InstanceReorderTable[InstanceIndex] : InstanceIndex;
+	const int32 RenderIndex = GetRenderIndex(InstanceIndex);
 	const FMatrix OldTransform = PerInstanceSMData[InstanceIndex].Transform;
 	const FTransform NewLocalTransform = bWorldSpace ? NewInstanceTransform.GetRelativeTransform(GetComponentTransform()) : NewInstanceTransform;
 	const FVector NewLocalLocation = NewLocalTransform.GetTranslation();
@@ -2391,8 +2391,7 @@ void UHierarchicalInstancedStaticMeshComponent::ClearInstances()
 
 		for (int32 Index = 0; Index < NumInstances; ++Index)
 		{
-			int32 RenderIndex = InstanceReorderTable.IsValidIndex(Index) ? InstanceReorderTable[Index] : Index;
-
+			const int32 RenderIndex = GetRenderIndex(Index);
 			if (RenderIndex == INDEX_NONE) 
 			{
 				// could be skipped by density settings
@@ -2896,7 +2895,7 @@ void UHierarchicalInstancedStaticMeshComponent::PropagateLightingScenarioChange(
 			{
 				for (int32 InstanceIndex = 0; InstanceIndex < PerInstanceSMData.Num(); ++InstanceIndex)
 				{
-					int32 RenderIndex = InstanceReorderTable.IsValidIndex(InstanceIndex) ? InstanceReorderTable[InstanceIndex] : InstanceIndex;
+					const int32 RenderIndex = GetRenderIndex(InstanceIndex);
 					if (RenderIndex != INDEX_NONE)
 					{
 						InstanceUpdateCmdBuffer.SetLightMapData(RenderIndex, MeshMapBuildData->PerInstanceLightmapData[InstanceIndex].LightmapUVBias);
diff --git a/Engine/Source/Runtime/Engine/Private/InstancedStaticMesh.cpp b/Engine/Source/Runtime/Engine/Private/InstancedStaticMesh.cpp
--- a/Engine/Source/Runtime/Engine/Private/InstancedStaticMesh.cpp
+++ b/Engine/Source/Runtime/Engine/Private/InstancedStaticMesh.cpp
@@ -262,11 +262,10 @@ void FStaticMeshInstanceBuffer::UpdateFromCommandBuffer_RenderThread(FInstanceUp
 	int32 NumCommands = CmdBuffer.NumInlineCommands();
 	int32 NumAdds = CmdBuffer.NumAdds;
 	int32 AddIndex = INDEX_NONE;
-	int32 NewNumInstances = InstanceData->GetNumInstances();
 	if (NumAdds > 0)
 	{
 		AddIndex = InstanceData->GetNumInstances();
-		NewNumInstances = NumAdds + InstanceData->GetNumInstances();
+		const int32 NewNumInstances = NumAdds + InstanceData->GetNumInstances();
 		InstanceData->AllocateInstances(NewNumInstances, GIsEditor ? EResizeBufferFlags::AllowSlackOnGrow | EResizeBufferFlags::AllowSlackOnReduce : EResizeBufferFlags::None, false); // In Editor always permit overallocation, to prevent too much realloc
 	}
 
@@ -279,28 +278,16 @@ void FStaticMeshInstanceBuffer::UpdateFromCommandBuffer_RenderThread(FInstanceUp
 			InstanceData->SetInstance(AddIndex++, Cmd.XForm, 0);
 			break;
 		case FInstanceUpdateCmdBuffer::Hide:
-			if (Cmd.InstanceIndex < NewNumInstances)
-			{
-				InstanceData->NullifyInstance(Cmd.InstanceIndex);
-			}
+			InstanceData->NullifyInstance(Cmd.InstanceIndex);
 			break;
 		case FInstanceUpdateCmdBuffer::Update:
-			if (Cmd.InstanceIndex < NewNumInstances)
-			{
-				InstanceData->SetInstance(Cmd.InstanceIndex, Cmd.XForm, 0);
-			}
+			InstanceData->SetInstance(Cmd.InstanceIndex, Cmd.XForm, 0);
 			break;
 		case FInstanceUpdateCmdBuffer::EditorData:
-			if (Cmd.InstanceIndex < NewNumInstances)
-			{
-				InstanceData->SetInstanceEditorData(Cmd.InstanceIndex, Cmd.HitProxyColor, Cmd.bSelected);
-			}
+			InstanceData->SetInstanceEditorData(Cmd.InstanceIndex, Cmd.HitProxyColor, Cmd.bSelected);
 			break;
 		case FInstanceUpdateCmdBuffer::LightmapData:
-			if (Cmd.InstanceIndex < NewNumInstances)
-			{
-				InstanceData->SetInstanceLightMapData(Cmd.InstanceIndex, Cmd.LightmapUVBias, Cmd.ShadowmapUVBias);
-			}
+			InstanceData->SetInstanceLightMapData(Cmd.InstanceIndex, Cmd.LightmapUVBias, Cmd.ShadowmapUVBias);
 			break;
 		default:
 			check(false);
@@ -1455,8 +1442,8 @@ void UInstancedStaticMeshComponent::BuildRenderData(FStaticMeshInstanceData& Out
 	
 	for (int32 Index = 0; Index < NumInstances; ++Index)
 	{
-		int32 RenderIndex = InstanceReorderTable.IsValidIndex(Index) ? InstanceReorderTable[Index] : Index;
-		if (RenderIndex == INDEX_NONE || !OutData.IsValidIndex(RenderIndex)) 
+		const int32 RenderIndex = GetRenderIndex(Index);
+		if (RenderIndex == INDEX_NONE) 
 		{
 			// could be skipped by density settings
 			continue;
@@ -1984,8 +1971,8 @@ void UInstancedStaticMeshComponent::SerializeRenderData(FArchive& Ar)
 					// Clear editor data for the cooked data
 					for (int32 Index = 0; Index < NumInstances; ++Index)
 					{
-						int32 RenderIndex = InstanceReorderTable.IsValidIndex(Index) ? InstanceReorderTable[Index] : Index;
-						if (RenderIndex == INDEX_NONE || RenderIndex >= NumInstances)
+						const int32 RenderIndex = GetRenderIndex(Index);
+						if (RenderIndex == INDEX_NONE)
 						{
 							// could be skipped by density settings
 							continue;
@@ -2003,8 +1990,8 @@ void UInstancedStaticMeshComponent::SerializeRenderData(FArchive& Ar)
 
 					for (int32 Index = 0; Index < NumInstances; ++Index)
 					{
-						int32 RenderIndex = InstanceReorderTable.IsValidIndex(Index) ? InstanceReorderTable[Index] : Index;
-						if (RenderIndex == INDEX_NONE || RenderIndex >= NumInstances)
+						const int32 RenderIndex = GetRenderIndex(Index);
+						if (RenderIndex == INDEX_NONE)
 						{
 							// could be skipped by density settings
 							continue;
@@ -2882,7 +2869,7 @@ void UInstancedStaticMeshComponent::SelectInstance(bool bInSelected, int32 InIns
 						HitProxyColor = PerInstanceRenderData->HitProxies[InstanceIndex]->Id.GetColor();
 					}
 
-					int32 RenderIndex = InstanceReorderTable.IsValidIndex(InstanceIndex) ? InstanceReorderTable[InstanceIndex] : InstanceIndex;
+					const int32 RenderIndex = GetRenderIndex(InstanceIndex);
 					if (RenderIndex != INDEX_NONE)
 					{
 						InstanceUpdateCmdBuffer.SetEditorData(RenderIndex, HitProxyColor, bSelected);
@@ -2914,7 +2901,7 @@ void UInstancedStaticMeshComponent::ClearInstanceSelection()
 					HitProxyColor = PerInstanceRenderData->HitProxies[InstanceIndex]->Id.GetColor();
 				}
 				
-				int32 RenderIndex = InstanceReorderTable.IsValidIndex(InstanceIndex) ? InstanceReorderTable[InstanceIndex] : InstanceIndex;
+				const int32 RenderIndex = GetRenderIndex(InstanceIndex);
 				if (RenderIndex != INDEX_NONE)
 				{
 					InstanceUpdateCmdBuffer.SetEditorData(RenderIndex, HitProxyColor, false);
diff --git a/Engine/Source/Runtime/Engine/Private/LightMap.cpp b/Engine/Source/Runtime/Engine/Private/LightMap.cpp
--- a/Engine/Source/Runtime/Engine/Private/LightMap.cpp
+++ b/Engine/Source/Runtime/Engine/Private/LightMap.cpp
@@ -307,7 +307,7 @@ struct FLightMapAllocation
 				// Need to create per-LOD instance data to fix that
 				MeshBuildData->PerInstanceLightmapData[InstanceIndex].LightmapUVBias = LightMap->GetCoordinateBias();
 
-				int32 RenderIndex = Component->InstanceReorderTable.IsValidIndex(InstanceIndex) ? Component->InstanceReorderTable[InstanceIndex] : InstanceIndex;
+				const int32 RenderIndex = Component->GetRenderIndex(InstanceIndex);
 				if (RenderIndex != INDEX_NONE)
 				{
 					Component->InstanceUpdateCmdBuffer.SetLightMapData(RenderIndex, MeshBuildData->PerInstanceLightmapData[InstanceIndex].LightmapUVBias);
diff --git a/Engine/Source/Runtime/Engine/Private/ShadowMap.cpp b/Engine/Source/Runtime/Engine/Private/ShadowMap.cpp
--- a/Engine/Source/Runtime/Engine/Private/ShadowMap.cpp
+++ b/Engine/Source/Runtime/Engine/Private/ShadowMap.cpp
@@ -109,7 +109,7 @@ struct FShadowMapAllocation
 				// TODO: We currently only support one LOD of static lighting in foliage
 				// Need to create per-LOD instance data to fix that
 				MeshBuildData->PerInstanceLightmapData[InstanceIndex].ShadowmapUVBias = ShadowMap->GetCoordinateBias();
-				int32 RenderIndex = Component->InstanceReorderTable.IsValidIndex(InstanceIndex) ? Component->InstanceReorderTable[InstanceIndex] : InstanceIndex;
+				const int32 RenderIndex = Component->GetRenderIndex(InstanceIndex);
 				if (RenderIndex != INDEX_NONE)
 				{
 					Component->InstanceUpdateCmdBuffer.SetShadowMapData(RenderIndex, MeshBuildData->PerInstanceLightmapData[InstanceIndex].ShadowmapUVBias);
