Fixed crash when when marking children render dirty in primitive scene component

diff --git a/Engine/Source/Runtime/Engine/Private/Components/PrimitiveComponent.cpp b/Engine/Source/Runtime/Engine/Private/Components/PrimitiveComponent.cpp
--- a/Engine/Source/Runtime/Engine/Private/Components/PrimitiveComponent.cpp
+++ b/Engine/Source/Runtime/Engine/Private/Components/PrimitiveComponent.cpp
@@ -796,15 +796,15 @@ void UPrimitiveComponent::MarkChildPrimitiveComponentRenderStateDirty()
 	// Walk down the tree updating
 	while (ProcessStack.Num() > 0)
 	{
-		USceneComponent* Current = ProcessStack.Pop(/*bAllowShrinking=*/ false);
-		UPrimitiveComponent* CurrentPrimitive = Cast<UPrimitiveComponent>(Current);
-
-		if (CurrentPrimitive)
+		if (USceneComponent* Current = ProcessStack.Pop(/*bAllowShrinking=*/ false))
 		{
-			CurrentPrimitive->MarkRenderStateDirty();
-		}
+			if (UPrimitiveComponent* CurrentPrimitive = Cast<UPrimitiveComponent>(Current))
+			{
+				CurrentPrimitive->MarkRenderStateDirty();
+			}
 
-		ProcessStack.Append(Current->GetAttachChildren());
+			ProcessStack.Append(Current->GetAttachChildren());
+		}
 	}
 }
 
