Extracted a game-specific condition to calculate ramp up/ramp down toward requested velocity (including friction) in a method so game project can override depending of their needs.

#rnx

diff --git a/Engine/Source/Runtime/Engine/Classes/GameFramework/CharacterMovementComponent.h b/Engine/Source/Runtime/Engine/Classes/GameFramework/CharacterMovementComponent.h
--- a/Engine/Source/Runtime/Engine/Classes/GameFramework/CharacterMovementComponent.h
+++ b/Engine/Source/Runtime/Engine/Classes/GameFramework/CharacterMovementComponent.h
@@ -1614,6 +1614,11 @@ protected:
 	/** Performs trace for ProjectLocationFromNavMesh */
 	virtual void FindBestNavMeshLocation(const FVector& TraceStart, const FVector& TraceEnd, const FVector& CurrentFeetLocation, const FVector& TargetNavLocation, FHitResult& OutHitResult) const;
 
+	/** 
+	 * When a character requests a velocity (like when following a path), this method returns true if when we should compute the 
+	 * acceleration toward requested velocity (including friction). If it returns false, it will snap instantly to requested velocity.
+	 */
+	virtual bool ShouldComputeAccelerationToReachRequestedVelocity(const float RequestedSpeed) const;
 public:
 
 	/** Called by owning Character upon successful teleport from AActor::TeleportTo(). */
diff --git a/Engine/Source/Runtime/Engine/Private/Components/CharacterMovementComponent.cpp b/Engine/Source/Runtime/Engine/Private/Components/CharacterMovementComponent.cpp
--- a/Engine/Source/Runtime/Engine/Private/Components/CharacterMovementComponent.cpp
+++ b/Engine/Source/Runtime/Engine/Private/Components/CharacterMovementComponent.cpp
@@ -3273,6 +3273,11 @@ void UCharacterMovementComponent::CalcVelocity(float DeltaTime, float Friction,
 	}
 }
 
+bool UCharacterMovementComponent::ShouldComputeAccelerationToReachRequestedVelocity(const float RequestedSpeed) const
+{
+	// Compute acceleration if accelerating toward requested speed, 1% buffer.
+	return bRequestedMoveUseAcceleration && Velocity.SizeSquared() < FMath::Square(RequestedSpeed * 1.01f);
+}
 
 bool UCharacterMovementComponent::ApplyRequestedMove(float DeltaTime, float MaxAccel, float MaxSpeed, float Friction, float BrakingDeceleration, FVector& OutAcceleration, float& OutRequestedSpeed)
 {
@@ -3295,7 +3300,7 @@ bool UCharacterMovementComponent::ApplyRequestedMove(float DeltaTime, float MaxA
 		// Compute acceleration. Use MaxAccel to limit speed increase, 1% buffer.
 		FVector NewAcceleration = FVector::ZeroVector;
 		const float CurrentSpeedSq = Velocity.SizeSquared();
-		if (bRequestedMoveUseAcceleration && CurrentSpeedSq < FMath::Square(RequestedSpeed * 1.01f))
+		if (ShouldComputeAccelerationToReachRequestedVelocity(RequestedSpeed))
 		{
 			// Turn in the same manner as with input acceleration.
 			const float VelSize = FMath::Sqrt(CurrentSpeedSq);
