Fix crash due to sound effects being referenced that have been destroyed.
- Safeguard against iterating over potentially destroyed effects
- Use weak pointer in FSoundEffectBase for generating USoundEffectPreset to avoid referencing if set to kill

diff --git a/Engine/Source/Runtime/Engine/Classes/Sound/SoundEffectBase.h b/Engine/Source/Runtime/Engine/Classes/Sound/SoundEffectBase.h
--- a/Engine/Source/Runtime/Engine/Classes/Sound/SoundEffectBase.h
+++ b/Engine/Source/Runtime/Engine/Classes/Sound/SoundEffectBase.h
@@ -86,7 +86,7 @@ public:
 	USoundEffectPreset* GetPreset();
 
 	/** Removes the instance from the preset. */
-	void ClearPreset();
+	void ClearPreset(bool bRemoveFromPreset = true);
 
 	/** Queries if the given preset object is the uobject preset for this preset instance, i.e. the preset which spawned this effect instance. */
 	bool IsPreset(USoundEffectPreset* InPreset) const;
@@ -102,7 +102,7 @@ protected:
 	TArray<uint8> CurrentAudioThreadSettingsData;
 
 	FThreadSafeBool bChanged;
-	USoundEffectPreset* Preset;
+	TWeakObjectPtr<USoundEffectPreset> Preset;
 
 	FThreadSafeBool bIsRunning;
 	FThreadSafeBool bIsActive;
diff --git a/Engine/Source/Runtime/Engine/Private/SoundEffectBase.cpp b/Engine/Source/Runtime/Engine/Private/SoundEffectBase.cpp
--- a/Engine/Source/Runtime/Engine/Private/SoundEffectBase.cpp
+++ b/Engine/Source/Runtime/Engine/Private/SoundEffectBase.cpp
@@ -6,7 +6,6 @@
 
 FSoundEffectBase::FSoundEffectBase()
 	: bChanged(false)
-	, Preset(nullptr)
 	, bIsRunning(false)
 	, bIsActive(false)
 {}
@@ -32,7 +31,7 @@ void FSoundEffectBase::SetPreset(USoundEffectPreset* Inpreset)
 		ClearPreset();
 
 		Preset = Inpreset;
-		if (Preset)
+		if (Preset.IsValid())
 		{
 			Preset->AddEffectInstance(this);
 		}
@@ -45,15 +44,18 @@ void FSoundEffectBase::SetPreset(USoundEffectPreset* Inpreset)
 
 USoundEffectPreset* FSoundEffectBase::GetPreset()
 {
-	return Preset;
+	return Preset.Get();
 }
 
-void FSoundEffectBase::ClearPreset()
+void FSoundEffectBase::ClearPreset(bool bRemoveFromPreset)
 {
-	if (Preset)
+	if (Preset.IsValid())
 	{
-		Preset->RemoveEffectInstance(this);
-		Preset = nullptr;
+		if (bRemoveFromPreset)
+		{
+			Preset->RemoveEffectInstance(this);
+		}
+		Preset.Reset();
 	}
 }
 
@@ -61,7 +63,7 @@ bool FSoundEffectBase::Update()
 {
 	PumpPendingMessages();
 
-	if (bChanged && Preset)
+	if (bChanged && Preset.IsValid())
 	{
 		OnPresetChanged();
 		bChanged = false;
diff --git a/Engine/Source/Runtime/Engine/Private/SoundEffectPreset.cpp b/Engine/Source/Runtime/Engine/Private/SoundEffectPreset.cpp
--- a/Engine/Source/Runtime/Engine/Private/SoundEffectPreset.cpp
+++ b/Engine/Source/Runtime/Engine/Private/SoundEffectPreset.cpp
@@ -17,8 +17,12 @@ USoundEffectPreset::~USoundEffectPreset()
 {
 	for (int32 i = 0; i < Instances.Num(); ++i)
 	{
-		Instances[i]->ClearPreset();
+		if (Instances[i])
+		{
+			Instances[i]->ClearPreset(false /* bRemoveFromPreset */);
+		}
 	}
+	Instances.Reset();
 }
 
 
@@ -26,15 +30,25 @@ void USoundEffectPreset::EffectCommand(TFunction<void()> Command)
 {
 	for (int32 i = 0; i < Instances.Num(); ++i)
 	{
-		Instances[i]->EffectCommand(Command);
+		if (Instances[i])
+		{
+			Instances[i]->EffectCommand(Command);
+		}
 	}
 }
 
 void USoundEffectPreset::Update()
 {
-	for (int32 i = 0; i < Instances.Num(); ++i)
+	for (int32 i = Instances.Num() - 1; i >= 0; --i)
 	{
-		Instances[i]->SetPreset(this);
+		if (!Instances[i] || Instances[i]->GetPreset() == nullptr)
+		{
+			Instances.RemoveAtSwap(i, 1, false /* bAllowShrinking */);
+		}
+		else
+		{
+			Instances[i]->SetPreset(this);
+		}
 	}
 }
 
@@ -66,10 +80,7 @@ void USoundEffectPreset::AddReferencedEffects(FReferenceCollector& Collector)
 
 void USoundEffectPreset::RemoveEffectInstance(FSoundEffectBase* InSource)
 {
-	if (Instances.Contains(InSource))
-	{
-		Instances.Remove(InSource);
-	}
+	Instances.RemoveSwap(InSource);
 }
 
 #if WITH_EDITORONLY_DATA
