Integrating CL# 10073505 from Dev-Enterprise
Fixing lightmap UV index not being properly updated after automatic lightmap UV setup.

diff --git a/Engine/Source/Runtime/Engine/Classes/Engine/StaticMesh.h b/Engine/Source/Runtime/Engine/Classes/Engine/StaticMesh.h
--- a/Engine/Source/Runtime/Engine/Classes/Engine/StaticMesh.h
+++ b/Engine/Source/Runtime/Engine/Classes/Engine/StaticMesh.h
@@ -1216,7 +1216,7 @@ public:
 	ENGINE_API void RemoveVertexColors();
 
 	/** Make sure the Lightmap UV point on a valid UVChannel */
-	ENGINE_API void EnforceLightmapRestrictions();
+	ENGINE_API void EnforceLightmapRestrictions(bool bUseRenderData = true);
 
 	/** Calculates the extended bounds */
 	ENGINE_API void CalculateExtendedBounds();
diff --git a/Engine/Source/Runtime/Engine/Private/StaticMesh.cpp b/Engine/Source/Runtime/Engine/Private/StaticMesh.cpp
--- a/Engine/Source/Runtime/Engine/Private/StaticMesh.cpp
+++ b/Engine/Source/Runtime/Engine/Private/StaticMesh.cpp
@@ -2975,7 +2975,8 @@ void UStaticMesh::PostEditChangeProperty(FPropertyChangedEvent& PropertyChangedE
 		}
 	}
 
-	EnforceLightmapRestrictions();
+	//Don't use the render data here because the property that just changed might be invalidating the current RenderData.
+	EnforceLightmapRestrictions(/*bUseRenderData=*/false);
 
 	// Following an undo or other operation which can change the SourceModels, ensure the StaticMeshOwner is up to date
 	for (int32 Index = 0; Index < GetNumSourceModels(); ++Index)
@@ -5905,7 +5906,7 @@ ENGINE_API void UStaticMesh::RemoveVertexColors()
 #endif
 }
 
-void UStaticMesh::EnforceLightmapRestrictions()
+void UStaticMesh::EnforceLightmapRestrictions(bool bUseRenderData)
 {
 	// Legacy content may contain a lightmap resolution of 0, which was valid when vertex lightmaps were supported, but not anymore with only texture lightmaps
 	LightMapResolution = FMath::Max(LightMapResolution, 4);
@@ -5913,21 +5914,65 @@ void UStaticMesh::EnforceLightmapRestrictions()
 	// Lightmass only supports 4 UVs
 	int32 NumUVs = 4;
 
-	if (RenderData)
+#if !WITH_EDITORONLY_DATA
+	//The source models are only available in the editor, fallback on the render data.
+	bUseRenderData = true;
+#endif //WITH_EDITORONLY_DATA
+
+	if (bUseRenderData)
 	{
-		for (int32 LODIndex = 0; LODIndex < RenderData->LODResources.Num(); ++LODIndex)
+		if (RenderData)
 		{
-			const FStaticMeshLODResources& LODResource = RenderData->LODResources[LODIndex];
-			if (LODResource.GetNumVertices() > 0) // skip LOD that was stripped (eg. MinLOD)
+			for (int32 LODIndex = 0; LODIndex < RenderData->LODResources.Num(); ++LODIndex)
 			{
-				NumUVs = FMath::Min(LODResource.GetNumTexCoords(), NumUVs);
+				const FStaticMeshLODResources& LODResource = RenderData->LODResources[LODIndex];
+				if (LODResource.GetNumVertices() > 0) // skip LOD that was stripped (eg. MinLOD)
+				{
+					NumUVs = FMath::Min(LODResource.GetNumTexCoords(), NumUVs);
+				}
 			}
 		}
+		else
+		{
+			NumUVs = 1;
+		}
 	}
+#if WITH_EDITORONLY_DATA
 	else
 	{
-		NumUVs = 1;
+		for (int32 LODIndex = 0; LODIndex < GetNumSourceModels(); ++LODIndex)
+		{
+			if (const FMeshDescription* MeshDescription = GetMeshDescription(LODIndex))
+			{
+				const TVertexInstanceAttributesConstRef<FVector2D> UVChannels = MeshDescription->VertexInstanceAttributes().GetAttributesRef<FVector2D>(MeshAttribute::VertexInstance::TextureCoordinate);
+
+				// skip empty LODs
+				if (UVChannels.GetNumElements() > 0)
+				{
+					int NumChannelsInLOD = UVChannels.GetNumIndices();
+					const FStaticMeshSourceModel& SourceModel = GetSourceModel(LODIndex);
+
+					if (SourceModel.BuildSettings.bGenerateLightmapUVs)
+					{
+						NumChannelsInLOD = FMath::Max(NumChannelsInLOD, SourceModel.BuildSettings.DstLightmapIndex);
+					}
+
+					NumUVs = FMath::Min(NumChannelsInLOD, NumUVs);
+				}
+			}
+			else
+			{
+				NumUVs = 1;
+				break;
+			}
+		}
+
+		if (GetNumSourceModels() == 0)
+		{
+			NumUVs = 1;
+		}
 	}
+#endif //WITH_EDITORONLY_DATA
 
 	// do not allow LightMapCoordinateIndex go negative
 	check(NumUVs > 0);
