Reduce size of FFileCache (the file cache used by [REDACTED_PERSON] system) from 32MB to 16MB.
Make the size configurable with a CVar: fc.NumFileCacheBlocks

diff --git a/Engine/Source/Runtime/Engine/Private/FileCache/FileCache.cpp b/Engine/Source/Runtime/Engine/Private/FileCache/FileCache.cpp
--- a/Engine/Source/Runtime/Engine/Private/FileCache/FileCache.cpp
+++ b/Engine/Source/Runtime/Engine/Private/FileCache/FileCache.cpp
@@ -25,14 +25,12 @@ DEFINE_LOG_CATEGORY_STATIC(LogStreamingFileCache, Log, All);
 
 static const int CacheLineSize = 64 * 1024;
 
-#define NUM_CACHE_BLOCKS 512
-
-static int32 MaxNumLiveRequests = 32;
-static FAutoConsoleVariableRef CVarNumLiveRequests(
-	TEXT("sfc.MaxNumLiveRequests"),
-	MaxNumLiveRequests,
-	TEXT("Number of read request that can be in flight. default 32\n")
-	, ECVF_Default
+static int32 GNumFileCacheBlocks = 256;
+static FAutoConsoleVariableRef CVarNumFileCacheBlocks(
+	TEXT("fc.NumFileCacheBlocks"),
+	GNumFileCacheBlocks,
+	TEXT("Number of blocks in the global file cache object\n"),
+	ECVF_RenderThreadSafe
 );
 
 // 
@@ -171,12 +169,13 @@ public:
 	// allocated with an extra dummy entry at index0 for linked list head
 	TArray<FSlotInfo> SlotInfo;
 	uint8* Memory;
+	int32 SizeInBytes;
 	int32 NumFreeSlots;
 };
 
 static FFileCache &GetCache()
 {
-	static FFileCache TheCache(NUM_CACHE_BLOCKS);
+	static FFileCache TheCache(GNumFileCacheBlocks);
 	return TheCache;
 }
 
@@ -247,9 +246,9 @@ private:
 FFileCache::FFileCache(int32 NumSlots)
 	: EvictFileCacheCommand(TEXT("r.VT.EvictFileCache"), TEXT("Evict all the file caches in the VT system."),
 		FConsoleCommandDelegate::CreateRaw(this, &FFileCache::EvictFileCacheFromConsole))
+	, SizeInBytes(NumSlots * CacheSlotID::BlockSize)
 	, NumFreeSlots(NumSlots)
 {
-	const int32 SizeInBytes = NumSlots * CacheSlotID::BlockSize;
 	Memory = (uint8*)FMemory::Malloc(SizeInBytes);
 
 	SlotInfo.AddUninitialized(NumSlots + 1);
@@ -641,3 +640,8 @@ IFileCacheHandle* IFileCacheHandle::CreateFileCacheHandle(const TCHAR* InFileNam
 
 	return new FFileCacheHandle(FileHandle);
 }
+
+uint32 IFileCacheHandle::GetFileCacheSize()
+{
+	return GetCache().SizeInBytes;
+}
diff --git a/Engine/Source/Runtime/Engine/Private/FileCache/FileCache.h b/Engine/Source/Runtime/Engine/Private/FileCache/FileCache.h
--- a/Engine/Source/Runtime/Engine/Private/FileCache/FileCache.h
+++ b/Engine/Source/Runtime/Engine/Private/FileCache/FileCache.h
@@ -24,6 +24,9 @@ public:
 	static IFileCacheHandle *CreateFileCacheHandle(const TCHAR* InFileName);
 	virtual ~IFileCacheHandle() {};
 
+	/** Return size of underlying file cache in bytes. */
+	static uint32 GetFileCacheSize();
+
 	/**
 	 * Read a byte range form the file. This can be a high-throughput operation and done lots of times for small reads.
 	 * The system will handle this efficiently.
@@ -37,4 +40,4 @@ public:
 	 * Wait until all outstanding read requests complete. 
 	 */
 	virtual void WaitAll() = 0;
-};
\ No newline at end of file
+};
diff --git a/Engine/Source/Runtime/Engine/Private/VT/UploadingVirtualTexture.cpp b/Engine/Source/Runtime/Engine/Private/VT/UploadingVirtualTexture.cpp
--- a/Engine/Source/Runtime/Engine/Private/VT/UploadingVirtualTexture.cpp
+++ b/Engine/Source/Runtime/Engine/Private/VT/UploadingVirtualTexture.cpp
@@ -8,9 +8,10 @@
 #include "CrunchCompression.h"
 #include "VirtualTextureChunkDDCCache.h"
 
-DECLARE_MEMORY_STAT(TEXT("Total Disk Size"), STAT_TotalDiskSize, STATGROUP_VirtualTextureMemory);
+DECLARE_MEMORY_STAT(TEXT("File Cache Size"), STAT_FileCacheSize, STATGROUP_VirtualTextureMemory);
 DECLARE_MEMORY_STAT(TEXT("Total Header Size"), STAT_TotalHeaderSize, STATGROUP_VirtualTextureMemory);
 DECLARE_MEMORY_STAT(TEXT("Tile Header Size"), STAT_TileHeaderSize, STATGROUP_VirtualTextureMemory);
+DECLARE_DWORD_ACCUMULATOR_STAT(TEXT("Total Disk Size (KB)"), STAT_TotalDiskSize, STATGROUP_VirtualTextureMemory);
 DECLARE_DWORD_ACCUMULATOR_STAT(TEXT("Num Tile Headers"), STAT_NumTileHeaders, STATGROUP_VirtualTextureMemory);
 DECLARE_DWORD_ACCUMULATOR_STAT(TEXT("Num Codecs"), STAT_NumCodecs, STATGROUP_VirtualTextureMemory);
 
@@ -38,9 +39,9 @@ FUploadingVirtualTexture::FUploadingVirtualTexture(FVirtualTextureBuiltData* InD
 	CodecPerChunk.AddDefaulted(InData->Chunks.Num());
 	InvalidChunks.Init(false, InData->Chunks.Num());
 
-	INC_MEMORY_STAT_BY(STAT_TotalDiskSize, InData->GetDiskMemoryFootprint());
 	INC_MEMORY_STAT_BY(STAT_TotalHeaderSize, InData->GetMemoryFootprint());
 	INC_MEMORY_STAT_BY(STAT_TileHeaderSize, InData->GetTileMemoryFootprint());
+	INC_DWORD_STAT_BY(STAT_TotalDiskSize, InData->GetDiskMemoryFootprint() / 1024);
 	INC_DWORD_STAT_BY(STAT_NumTileHeaders, InData->GetNumTileHeaders());
 }
 
@@ -48,9 +49,9 @@ FUploadingVirtualTexture::FUploadingVirtualTexture(FVirtualTextureBuiltData* InD
 FUploadingVirtualTexture::~FUploadingVirtualTexture()
 {
 	check(Data);
-	DEC_MEMORY_STAT_BY(STAT_TotalDiskSize, Data->GetDiskMemoryFootprint());
 	DEC_MEMORY_STAT_BY(STAT_TotalHeaderSize, Data->GetMemoryFootprint());
 	DEC_MEMORY_STAT_BY(STAT_TileHeaderSize, Data->GetTileMemoryFootprint());
+	DEC_DWORD_STAT_BY(STAT_TotalDiskSize, Data->GetDiskMemoryFootprint() / 1024);
 	DEC_DWORD_STAT_BY(STAT_NumTileHeaders, Data->GetNumTileHeaders());
 
 	// Complete all open transcode requests before deleting IFileCacheHandle objects
@@ -379,6 +380,7 @@ FVTDataAndStatus FUploadingVirtualTexture::ReadData(FGraphEventArray& OutComplet
 			}
 			return EVTRequestPageStatus::Invalid;
 		}
+		SET_MEMORY_STAT(STAT_FileCacheSize, IFileCacheHandle::GetFileCacheSize());
 	}
 
 	IMemoryReadStreamRef ReadData = Handle->ReadData(OutCompletionEvents, ChunkOffsetInFile + Offset, Size, Priority);
