Blueprints: Added array of UBlueprintExtensions to UBlueprint in order to better support preloading, generated functions, and any other future extension points

diff --git a/Engine/Source/Editor/Kismet/Kismet.Build.cs b/Engine/Source/Editor/Kismet/Kismet.Build.cs
--- a/Engine/Source/Editor/Kismet/Kismet.Build.cs
+++ b/Engine/Source/Editor/Kismet/Kismet.Build.cs
@@ -18,7 +18,6 @@ public class Kismet : ModuleRules
                 "DerivedDataCache",
                 "LevelEditor",
 				"GameProjectGeneration",
-				"Sequencer",
 				"SourceCodeAccess",
 			}
 			);
diff --git a/Engine/Source/Editor/KismetCompiler/KismetCompiler.Build.cs b/Engine/Source/Editor/KismetCompiler/KismetCompiler.Build.cs
--- a/Engine/Source/Editor/KismetCompiler/KismetCompiler.Build.cs
+++ b/Engine/Source/Editor/KismetCompiler/KismetCompiler.Build.cs
@@ -14,7 +14,7 @@ public class KismetCompiler : ModuleRules
                 "InputCore",
 				"UnrealEd",
 				"MovieScene",
-				"MovieSceneTools",
+				//"MovieSceneTools",
 				"BlueprintGraph",
 				"AnimGraph",
                 "MessageLog",
diff --git a/Engine/Source/Editor/KismetCompiler/Private/KismetCompiler.cpp b/Engine/Source/Editor/KismetCompiler/Private/KismetCompiler.cpp
--- a/Engine/Source/Editor/KismetCompiler/Private/KismetCompiler.cpp
+++ b/Engine/Source/Editor/KismetCompiler/Private/KismetCompiler.cpp
@@ -21,6 +21,7 @@
 #include "Components/TimelineComponent.h"
 #include "Engine/TimelineTemplate.h"
 #include "Engine/UserDefinedStruct.h"
+#include "Blueprint/BlueprintExtension.h"
 #include "EdGraphUtilities.h"
 #include "K2Node_CallFunction.h"
 #include "K2Node_Composite.h"
@@ -3841,10 +3842,11 @@ void FKismetCompilerContext::CreateFunctionList()
 	{
 		BP_SCOPED_COMPILER_EVENT_STAT(EKismetCompilerStats_GenerateFunctionGraphs);
 
-		// Broadcast the blueprint's function generation event to ensure that any function generators are run right before we create the function list for the class layout
-		FGenerateBlueprintFunctionParams Params;
-		Params.CompilerContext = this;
-		Blueprint->GenerateFunctionGraphsEvent.Broadcast(Params);
+		// Allow blueprint extensions for the blueprint to generate function graphs
+		for (UBlueprintExtension* Extension : Blueprint->Extensions)
+		{
+			Extension->GenerateFunctionGraphs(this);
+		}
 	}
 
 	BP_SCOPED_COMPILER_EVENT_STAT(EKismetCompilerStats_CreateFunctionList);
diff --git a/Engine/Source/Runtime/Engine/Classes/Engine/Blueprint.h b/Engine/Source/Runtime/Engine/Classes/Engine/Blueprint.h
--- a/Engine/Source/Runtime/Engine/Classes/Engine/Blueprint.h
+++ b/Engine/Source/Runtime/Engine/Classes/Engine/Blueprint.h
@@ -22,6 +22,7 @@ class UActorComponent;
 class UEdGraph;
 class FKismetCompilerContext;
 class UInheritableComponentHandler;
+class UBlueprintExtension;
 class FBlueprintActionDatabaseRegistrar;
 struct FDiffResults;
 
@@ -140,18 +141,6 @@ struct FBlueprintMacroCosmeticInfo
 	}
 };
 
-/** Structure that is passed to UBlueprint::GenerateFunctionGraphsEvent in order to generate procedural blueprint functions */
-USTRUCT()
-struct FGenerateBlueprintFunctionParams
-{
-	GENERATED_BODY()
-
-	/** The kismet compiler being used for compilation */
-	FKismetCompilerContext* CompilerContext;
-};
-
-DECLARE_DYNAMIC_MULTICAST_DELEGATE_OneParam(FGenerateBlueprintFunctionsEvent, const FGenerateBlueprintFunctionParams&, Params);
-
 struct FKismetCompilerOptions
 {
 public:
@@ -555,10 +544,6 @@ public:
 	UPROPERTY()
 	FGuid SearchGuid;
 
-	/** Persistent, serialized event that is triggered after skeleton class generation, but before functions are created for the class layout to give external subscribers an opportunity to generate function graphs on the compiler context */
-	UPROPERTY()
-	FGenerateBlueprintFunctionsEvent GenerateFunctionGraphsEvent;
-
 #endif //WITH_EDITORONLY_DATA
 
 	/** The version of the blueprint system that was used to  create this blueprint */
@@ -657,6 +642,11 @@ public:
 	/** Maps old to new component template names */
 	UPROPERTY(transient)
 	TMap<FName, FName> OldToNewComponentTemplateNames;
+
+	/** Array of extensions for this blueprint */
+	UPROPERTY()
+	TArray<UBlueprintExtension*> Extensions;
+
 #endif // WITH_EDITORONLY_DATA
 
 public:
diff --git a/Engine/Source/Runtime/Engine/Private/Blueprint.cpp b/Engine/Source/Runtime/Engine/Private/Blueprint.cpp
--- a/Engine/Source/Runtime/Engine/Private/Blueprint.cpp
+++ b/Engine/Source/Runtime/Engine/Private/Blueprint.cpp
@@ -40,6 +40,7 @@
 #include "Interfaces/ITargetPlatform.h"
 #include "UObject/MetaData.h"
 #include "BlueprintAssetHandler.h"
+#include "Blueprint/BlueprintExtension.h"
 #endif
 #include "Engine/InheritableComponentHandler.h"
 
@@ -623,7 +624,13 @@ UClass* UBlueprint::RegenerateClass(UClass* ClassToRegenerate, UObject* Previous
 		UBlueprint::ForceLoadMembers(GeneratedClassResolved->ClassDefaultObject);
 	}
 	UBlueprint::ForceLoadMembers(this);
-		
+
+	for (UBlueprintExtension* Extension : Extensions)
+	{
+		ForceLoad(Extension);
+		Extension->PreloadObjectsForCompilation(this);
+	}
+
 	FBlueprintEditorUtils::PreloadConstructionScript( this );
 
 	FBlueprintEditorUtils::LinkExternalDependencies( this );
@@ -2032,3 +2039,4 @@ void UBlueprint::LoadModulesRequiredForCompilation()
 }
 #endif //WITH_EDITORONLY_DATA
 
+
diff --git a/Engine/Source/Runtime/Engine/Public/Blueprint/BlueprintExtension.h b/Engine/Source/Runtime/Engine/Public/Blueprint/BlueprintExtension.h
new file mode 100644
--- /dev/null
+++ b/Engine/Source/Runtime/Engine/Public/Blueprint/BlueprintExtension.h
@@ -0,0 +1,55 @@
+// Copyright 1998-2019 Epic Games, Inc. All Rights Reserved.
+
+#pragma once
+
+#include "CoreTypes.h"
+#include "UObject/Object.h"
+
+#include "BlueprintExtension.generated.h"
+
+class UBlueprint;
+class FKismetCompilerContext;
+
+/**
+ * Per-instance extension object that can be added to UBlueprint::Extensions in order to augment built-in blueprint functionality
+ * Ideally this would be an editor-only class, but such classes are not permitted within Engine modules (even inside WITH_EDITORONLY_DATA blocks)
+ */
+UCLASS()
+class ENGINE_API UBlueprintExtension : public UObject
+{
+public:
+
+	GENERATED_BODY()
+
+#if WITH_EDITORONLY_DATA
+
+	/**
+	 * Called during compilation (after skeleton class generation) in order to generate dynamic function graphs for this blueprint
+	 */
+	void GenerateFunctionGraphs(FKismetCompilerContext* CompilerContext)
+	{
+		HandleGenerateFunctionGraphs(CompilerContext);
+	}
+
+	/**
+	 * Called before blueprint compilation to ensure that any objects necessary for the specified blueprint's compilation
+	 */
+	void PreloadObjectsForCompilation(UBlueprint* OwningBlueprint)
+	{
+		HandlePreloadObjectsForCompilation(OwningBlueprint);
+	}
+
+private:
+
+	/**
+	 * Override this function to define custom preload logic into the specified blueprint
+	 */
+	virtual void HandlePreloadObjectsForCompilation(UBlueprint* OwningBlueprint) {}
+
+	/**
+	 * Override this function to define custom function generation logic for the specified blueprint compiler context
+	 */
+	virtual void HandleGenerateFunctionGraphs(FKismetCompilerContext* CompilerContext) {}
+
+#endif // WITH_EDITORONLY_DATA
+};
