Debug solver name

diff --git a/Engine/Source/Runtime/Engine/Private/PhysicsEngine/Experimental/PhysScene_Chaos.cpp b/Engine/Source/Runtime/Engine/Private/PhysicsEngine/Experimental/PhysScene_Chaos.cpp
--- a/Engine/Source/Runtime/Engine/Private/PhysicsEngine/Experimental/PhysScene_Chaos.cpp
+++ b/Engine/Source/Runtime/Engine/Private/PhysicsEngine/Experimental/PhysScene_Chaos.cpp
@@ -337,7 +337,11 @@ struct FPhysScenePendingComponentTransform_Chaos
 	{}
 };
 
-FPhysScene_Chaos::FPhysScene_Chaos(AActor* InSolverActor)
+FPhysScene_Chaos::FPhysScene_Chaos(AActor* InSolverActor
+#if CHAOS_CHECKED
+	, const FName& DebugName
+#endif
+)
 	: ChaosModule(nullptr)
 	, SceneSolver(nullptr)
 	, SolverActor(InSolverActor)
@@ -353,7 +357,11 @@ FPhysScene_Chaos::FPhysScene_Chaos(AActor* InSolverActor)
 	ChaosModule = FModuleManager::Get().GetModulePtr<FChaosSolversModule>("ChaosSolvers");
 	check(ChaosModule);
 
-	SceneSolver = ChaosModule->CreateSolver();
+	SceneSolver = ChaosModule->CreateSolver(false
+#if CHAOS_CHECKED
+	, DebugName
+#endif
+);
 	check(SceneSolver);
 
 	// If we're running the physics thread, hand over the solver to it - we are no longer
@@ -914,7 +922,15 @@ void FPhysScene_Chaos::RemoveFromComponentMaps(IPhysicsProxyBase* InObject)
 
 #if WITH_CHAOS
 
-FPhysScene_ChaosInterface::FPhysScene_ChaosInterface(const AWorldSettings* InSettings /*= nullptr*/)
+FPhysScene_ChaosInterface::FPhysScene_ChaosInterface(const AWorldSettings* InSettings /*= nullptr*/
+#if CHAOS_CHECKED
+	, const FName& DebugName
+#endif
+)	: Scene(nullptr
+#if CHAOS_CHECKED
+	, DebugName
+#endif
+)
 {
 	//Initialize unique ptrs that are just here to allow forward declare. This should be reworked todo(ocohen)
 #if TODO_FIX_REFERENCES_TO_ADDARRAY
diff --git a/Engine/Source/Runtime/Engine/Private/World.cpp b/Engine/Source/Runtime/Engine/Private/World.cpp
--- a/Engine/Source/Runtime/Engine/Private/World.cpp
+++ b/Engine/Source/Runtime/Engine/Private/World.cpp
@@ -4497,7 +4497,12 @@ bool UWorld::AreActorsInitialized() const
 
 void UWorld::CreatePhysicsScene(const AWorldSettings* Settings)
 {
+#if CHAOS_CHECKED
+	const FName PhysicsName = IsNetMode(NM_DedicatedServer) ? TEXT("ServerPhysics") : TEXT("ClientPhysics");
+	FPhysScene* NewScene = new FPhysScene(Settings, PhysicsName);
+#else
 	FPhysScene* NewScene = new FPhysScene(Settings);
+#endif
 	SetPhysicsScene(NewScene);
 }
 
diff --git a/Engine/Source/Runtime/Engine/Public/Physics/Experimental/PhysScene_Chaos.h b/Engine/Source/Runtime/Engine/Public/Physics/Experimental/PhysScene_Chaos.h
--- a/Engine/Source/Runtime/Engine/Public/Physics/Experimental/PhysScene_Chaos.h
+++ b/Engine/Source/Runtime/Engine/Public/Physics/Experimental/PhysScene_Chaos.h
@@ -68,9 +68,17 @@ class ENGINE_API FPhysScene_Chaos : public FTickableGameObject, public FGCObject
 public:
 
 #if !WITH_CHAOS_NEEDS_TO_BE_FIXED
-	FPhysScene_Chaos(AActor* InSolverActor);
+	FPhysScene_Chaos(AActor* InSolverActor
+#if CHAOS_CHECKED
+	, const FName& DebugName=NAME_None
+#endif
+);
 #else
-	FPhysScene_Chaos(AActor* InSolverActor=nullptr);
+	FPhysScene_Chaos(AActor* InSolverActor=nullptr
+#if CHAOS_CHECKED
+	, const FName& DebugName=NAME_None
+#endif
+);
 #endif
 
 	virtual ~FPhysScene_Chaos();
@@ -268,7 +276,11 @@ public:
 
 	friend class FPhysInterface_Chaos;
 
-	ENGINE_API FPhysScene_ChaosInterface(const AWorldSettings* InSettings = nullptr);
+	ENGINE_API FPhysScene_ChaosInterface(const AWorldSettings* InSettings = nullptr
+#if CHAOS_CHECKED
+		, const FName& DebugName = NAME_None
+#endif
+	);
 
 	// Scene
 	void OnWorldBeginPlay();
diff --git a/Engine/Source/Runtime/Experimental/Chaos/Public/Chaos/Framework/PhysicsSolverBase.h b/Engine/Source/Runtime/Experimental/Chaos/Public/Chaos/Framework/PhysicsSolverBase.h
--- a/Engine/Source/Runtime/Experimental/Chaos/Public/Chaos/Framework/PhysicsSolverBase.h
+++ b/Engine/Source/Runtime/Experimental/Chaos/Public/Chaos/Framework/PhysicsSolverBase.h
@@ -31,11 +31,27 @@ namespace Chaos
 			DirtyProxiesSet.Remove(ProxyBaseIn);
 		}
 
+#if CHAOS_CHECKED
+		void SetDebugName(const FName& Name)
+		{
+			DebugName = Name;
+		}
+
+		const FName& GetDebugName() const
+		{
+			return DebugName;
+		}
+#endif
+
 	protected:
 			EMultiBufferMode BufferMode;
 
 			// Input Proxy Map
 			FCriticalSection DirtyProxiesCriticalSection;
 			TSet< IPhysicsProxyBase *> DirtyProxiesSet;
+
+#if CHAOS_CHECKED
+			FName DebugName;
+#endif
 	};
 }
diff --git a/Engine/Source/Runtime/Experimental/ChaosSolverEngine/Private/Chaos/ChaosSolverActor.cpp b/Engine/Source/Runtime/Experimental/ChaosSolverEngine/Private/Chaos/ChaosSolverActor.cpp
--- a/Engine/Source/Runtime/Experimental/ChaosSolverEngine/Private/Chaos/ChaosSolverActor.cpp
+++ b/Engine/Source/Runtime/Experimental/ChaosSolverEngine/Private/Chaos/ChaosSolverActor.cpp
@@ -204,7 +204,11 @@ AChaosSolverActor::AChaosSolverActor(const FObjectInitializer& ObjectInitializer
 	, ChaosDebugSubstepControl()
 {
 	// @question(Benn) : Does this need to be created on the Physics thread using a queued command?
-	PhysScene = MakeShareable(new FPhysScene_Chaos(this));
+	PhysScene = MakeShareable(new FPhysScene_Chaos(this
+#if CHAOS_CHECKED
+		, TEXT("Solver Actor Physics"))
+#endif
+	);
 	Solver = PhysScene->GetSolver();
 	// Ticking setup for collision/breaking notifies
 	PrimaryActorTick.TickGroup = TG_PostPhysics;
diff --git a/Engine/Source/Runtime/Experimental/ChaosSolvers/Private/ChaosSolversModule.cpp b/Engine/Source/Runtime/Experimental/ChaosSolvers/Private/ChaosSolversModule.cpp
--- a/Engine/Source/Runtime/Experimental/ChaosSolvers/Private/ChaosSolversModule.cpp
+++ b/Engine/Source/Runtime/Experimental/ChaosSolvers/Private/ChaosSolversModule.cpp
@@ -441,7 +441,11 @@ void FChaosSolversModule::SyncTask(bool bForceBlockingSync /*= false*/)
 	UpdateStats();
 }
 
-Chaos::FPhysicsSolver* FChaosSolversModule::CreateSolver(bool bStandalone /*= false*/)
+Chaos::FPhysicsSolver* FChaosSolversModule::CreateSolver(bool bStandalone /*= false*/
+#if CHAOS_CHECKED
+	, const FName& DebugName
+#endif
+)
 {
 	FChaosScopeSolverLock SolverScopeLock;
 	
@@ -454,6 +458,12 @@ Chaos::FPhysicsSolver* FChaosSolversModule::CreateSolver(bool bStandalone /*= fa
 	Solvers.Add(new Chaos::FPhysicsSolver(SolverBufferMode));
 	Chaos::FPhysicsSolver* NewSolver = Solvers.Last();
 
+#if CHAOS_CHECKED
+    // Add solver nubmer to solver name
+	const FName NewDebugName = *FString::Printf(TEXT("%s (%d)"), DebugName == NAME_None ? TEXT("Solver") : *DebugName.ToString(), Solvers.Num() - 1);
+	NewSolver->SetDebugName(NewDebugName);
+#endif
+
 	if(!bStandalone && IsPersistentTaskRunning() && Dispatcher)
 	{
 		// Need to let the thread know there's a new solver to care about
diff --git a/Engine/Source/Runtime/Experimental/ChaosSolvers/Public/ChaosSolversModule.h b/Engine/Source/Runtime/Experimental/ChaosSolvers/Public/ChaosSolversModule.h
--- a/Engine/Source/Runtime/Experimental/ChaosSolvers/Public/ChaosSolversModule.h
+++ b/Engine/Source/Runtime/Experimental/ChaosSolvers/Public/ChaosSolversModule.h
@@ -150,7 +150,11 @@ public:
 	 *
 	 * @param bStandalone Whether the solver is standalone (not sent to physics thread - updating left to caller)
 	 */
-	Chaos::FPhysicsSolver* CreateSolver(bool bStandalone = false);
+	Chaos::FPhysicsSolver* CreateSolver(bool bStandalone = false
+#if CHAOS_CHECKED
+		, const FName& DebugName = NAME_None
+#endif
+	);
 
 	Chaos::EMultiBufferMode GetBufferModeFromThreadingModel(Chaos::EThreadingMode ThreadingMode) const
 	{
