Fix for VOIPListenerSynthComponent entering a zombie state after going idle at startup. Also includes fix for issue where an audio component set to bAlwaysPlay would still get culled at near silent values.

diff --git a/Engine/Plugins/Online/OnlineSubsystemUtils/Source/OnlineSubsystemUtils/Private/VoiceEngineImpl.cpp b/Engine/Plugins/Online/OnlineSubsystemUtils/Source/OnlineSubsystemUtils/Private/VoiceEngineImpl.cpp
--- a/Engine/Plugins/Online/OnlineSubsystemUtils/Source/OnlineSubsystemUtils/Private/VoiceEngineImpl.cpp
+++ b/Engine/Plugins/Online/OnlineSubsystemUtils/Source/OnlineSubsystemUtils/Private/VoiceEngineImpl.cpp
@@ -146,6 +146,8 @@ void FRemoteTalkerDataImpl::Reset()
 		}
 
 		bIsActive = false;
+		
+		VoipSynthComponent = nullptr;
 	}
 
 	CurrentUncompressedDataQueueSize = 0;
diff --git a/Engine/Source/Runtime/Engine/Private/Audio.cpp b/Engine/Source/Runtime/Engine/Private/Audio.cpp
--- a/Engine/Source/Runtime/Engine/Private/Audio.cpp
+++ b/Engine/Source/Runtime/Engine/Private/Audio.cpp
@@ -1046,6 +1046,12 @@ bool FWaveInstance::ShouldStopDueToMaxConcurrency() const
 
 float FWaveInstance::GetVolumeWeightedPriority() const
 {
+	// If priority has been set via bAlwaysPlay, it will have a priority larger than MAX_SOUND_PRIORITY. If that's the case, we should ignore volume weighting.
+	if (Priority > MAX_SOUND_PRIORITY)
+	{
+		return Priority;
+	}
+
 	// This will result in zero-volume sounds still able to be sorted due to priority but give non-zero volumes higher priority than 0 volumes
 	float ActualVolume = GetVolumeWithDistanceAttenuation();
 	if (ActualVolume > 0.0f)
