Do not print an error when cooking physics data for BSP components which contain empty elements.

Empty elements can occur due to CSG operations and are handled correctly by the rendering code. The collision code just needs to ignore them.

#rnx

diff --git a/Engine/Source/Runtime/Engine/Private/Components/ModelComponent.cpp b/Engine/Source/Runtime/Engine/Private/Components/ModelComponent.cpp
--- a/Engine/Source/Runtime/Engine/Private/Components/ModelComponent.cpp
+++ b/Engine/Source/Runtime/Engine/Private/Components/ModelComponent.cpp
@@ -673,20 +673,19 @@ bool UModelComponent::GetPhysicsTriMeshData(struct FTriMeshCollisionData* Collis
 	for(int32 ElementIndex = 0;ElementIndex < Elements.Num();ElementIndex++)
 	{
 		FModelElement& Element = Elements[ElementIndex];
-		FRawIndexBuffer16or32* IndexBuffer = Element.IndexBuffer;
-		int32 NumIndices = IndexBuffer ? IndexBuffer->Indices.Num() : 0;
-
-		if (NumIndices == 0)
+		if (Element.NumTriangles == 0)
 		{
-			UE_LOG(LogPhysics, Error, TEXT("Found bad index buffer when cooking UModelComponent physics data! Component: %s, Buffer: %x, Index Count: %d, Element: %d"), *GetPathName(), IndexBuffer, NumIndices, ElementIndex);
+			// CSG can produce empty elements, just ignore them.
 			continue;
 		}
 
+		FRawIndexBuffer16or32* IndexBuffer = Element.IndexBuffer;
+		int32 NumIndices = IndexBuffer ? IndexBuffer->Indices.Num() : 0;
 		int32 NumVertices = Model->VertexBuffer.Vertices.Num();
 
 		if (NumVertices < (int32)Element.MaxVertexIndex)
 		{
-			UE_LOG(LogPhysics, Error, TEXT("Found bad vertex buffer when cooking UModelComponent physics data! Component: %s, Element: %d. Expected Vertex Count: %d, Actual Vertex Count: %d"),
+			UE_LOG(LogPhysics, Warning, TEXT("Found bad vertex buffer when cooking UModelComponent physics data! Component: %s, Element: %d. Expected Vertex Count: %d, Actual Vertex Count: %d"),
 				*GetPathName(), ElementIndex, Element.MaxVertexIndex, NumVertices);
 			continue;
 		}
