Added DFAtlasPercentageUsage, DFAtlasMaxX/Y/Z csv stats

#rnx

diff --git a/Engine/Source/Runtime/Engine/Private/DistanceFieldAtlas.cpp b/Engine/Source/Runtime/Engine/Private/DistanceFieldAtlas.cpp
--- a/Engine/Source/Runtime/Engine/Private/DistanceFieldAtlas.cpp
+++ b/Engine/Source/Runtime/Engine/Private/DistanceFieldAtlas.cpp
@@ -127,7 +127,12 @@ TGlobalResource<FDistanceFieldVolumeTextureAtlas> GDistanceFieldVolumeTextureAtl
 
 FDistanceFieldVolumeTextureAtlas::FDistanceFieldVolumeTextureAtlas() :
 	BlockAllocator(0, 0, 0, 0, 0, 0, false, false),
-	bInitialized(false)
+	bInitialized(false),
+	AllocatedPixels(0),
+	MaxUsedAtlasX(0),
+	MaxUsedAtlasY(0),
+	MaxUsedAtlasZ(0)
+
 {
 	// Warning: can't access cvars here, this is called during global init
 	Generation = 0;
@@ -151,6 +156,10 @@ void FDistanceFieldVolumeTextureAtlas::InitializeIfNeeded()
 		const int32 AtlasZ = CVarZ->GetValueOnAnyThread();
 
 		BlockAllocator = FTextureLayout3d(0, 0, 0, AtlasXY, AtlasXY, AtlasZ, false, false);
+
+		MaxUsedAtlasX = 0;
+		MaxUsedAtlasY = 0;
+		MaxUsedAtlasZ = 0;
 	}
 }
 
@@ -255,7 +264,7 @@ void FDistanceFieldVolumeTextureAtlas::AddAllocation(FDistanceFieldVolumeTexture
 	{
 		Texture->bThrottled = true;
 	}
-	
+	const FIntVector Size = Texture->VolumeData.Size;
 }
 
 void FDistanceFieldVolumeTextureAtlas::RemoveAllocation(FDistanceFieldVolumeTexture* Texture)
@@ -269,6 +278,7 @@ void FDistanceFieldVolumeTextureAtlas::RemoveAllocation(FDistanceFieldVolumeText
 		const FIntVector Size = Texture->VolumeData.Size;
 		verify(BlockAllocator.RemoveElement(Min.X, Min.Y, Min.Z, Size.X, Size.Y, Size.Z));
 		CurrentAllocations.Remove(Texture);
+		AllocatedPixels -= Size.X * Size.Y * Size.Z;
 	}
 }
 
@@ -282,6 +292,14 @@ struct FCompareVolumeAllocation
 
 void FDistanceFieldVolumeTextureAtlas::UpdateAllocations()
 {
+	{
+		uint32 TotalSurface = BlockAllocator.GetMaxSizeX() * BlockAllocator.GetMaxSizeY() * BlockAllocator.GetMaxSizeZ();
+		CSV_CUSTOM_STAT_GLOBAL(DFAtlasPercentageUsage, float((float(AllocatedPixels) / float(TotalSurface))*100.0f), ECsvCustomStatOp::Set);
+		CSV_CUSTOM_STAT_GLOBAL(DFAtlasMaxX, float(MaxUsedAtlasX), ECsvCustomStatOp::Set);
+		CSV_CUSTOM_STAT_GLOBAL(DFAtlasMaxY, float(MaxUsedAtlasY), ECsvCustomStatOp::Set);
+		CSV_CUSTOM_STAT_GLOBAL(DFAtlasMaxZ, float(MaxUsedAtlasZ), ECsvCustomStatOp::Set);
+	}
+
 	if (PendingAllocations.Num() > 0 || GDistanceFieldForceAtlasRealloc != 0)
 	{
 		const double StartTime = FPlatformTime::Seconds();
@@ -313,6 +331,13 @@ void FDistanceFieldVolumeTextureAtlas::UpdateAllocations()
 					LocalPendingAllocations->RemoveAt(AllocationIndex);
 					AllocationIndex--;
 				}
+				else
+				{
+					MaxUsedAtlasX = FMath::Max<uint32>(MaxUsedAtlasX, Texture->AtlasAllocationMin.X + Size.X);
+					MaxUsedAtlasY = FMath::Max<uint32>(MaxUsedAtlasY, Texture->AtlasAllocationMin.Y + Size.Y);
+					MaxUsedAtlasZ = FMath::Max<uint32>(MaxUsedAtlasZ, Texture->AtlasAllocationMin.Z + Size.Z);
+					AllocatedPixels += Size.X * Size.Y * Size.Z;
+				}
 			}
 		};
 
@@ -356,9 +381,13 @@ void FDistanceFieldVolumeTextureAtlas::UpdateAllocations()
 
 				// Remove all allocations from the layout so we have a clean slate
 				BlockAllocator = FTextureLayout3d(0, 0, 0, AtlasXY, AtlasXY, AtlasZ, false, false);
-				
+
 				Generation++;
 
+				MaxUsedAtlasX = 0;
+				MaxUsedAtlasY = 0;
+				MaxUsedAtlasZ = 0;
+
 				// Re-upload all textures since we had to reallocate
 				PendingAllocations.Append(CurrentAllocations);
 				if (bThrottleUpdateAllocation)
diff --git a/Engine/Source/Runtime/Engine/Public/DistanceFieldAtlas.h b/Engine/Source/Runtime/Engine/Public/DistanceFieldAtlas.h
--- a/Engine/Source/Runtime/Engine/Public/DistanceFieldAtlas.h
+++ b/Engine/Source/Runtime/Engine/Public/DistanceFieldAtlas.h
@@ -125,6 +125,14 @@ private:
 	int32 Generation;
 
 	bool bInitialized;
+
+	/** Number of pixel used in atlas distance field */
+	uint32 AllocatedPixels;
+
+	/** Max position used in distance field */
+	uint32 MaxUsedAtlasX;
+	uint32 MaxUsedAtlasY;
+	uint32 MaxUsedAtlasZ;
 };
 
 extern ENGINE_API TGlobalResource<FDistanceFieldVolumeTextureAtlas> GDistanceFieldVolumeTextureAtlas;
