HISM: Fix Possibility for IsTreeFullyBuilt to return true even if not true. (Removing N Instances, Adding N Instances without calling BuildTree). This is probably the root cause of then loading bad data and crashing in BuildRenderData with offset indices.

#rnx

diff --git a/Engine/Source/Runtime/Engine/Classes/Components/HierarchicalInstancedStaticMeshComponent.h b/Engine/Source/Runtime/Engine/Classes/Components/HierarchicalInstancedStaticMeshComponent.h
--- a/Engine/Source/Runtime/Engine/Classes/Components/HierarchicalInstancedStaticMeshComponent.h
+++ b/Engine/Source/Runtime/Engine/Classes/Components/HierarchicalInstancedStaticMeshComponent.h
@@ -163,11 +163,6 @@ class ENGINE_API UHierarchicalInstancedStaticMeshComponent : public UInstancedSt
 	// Current value of density scaling applied to this component
 	float CurrentDensityScaling;
 
-#if WITH_EDITOR
-	// in Editor mode we might disable the density scaling for edition
-	bool bCanEnableDensityScaling;
-#endif
-
 	// The number of nodes in the occlusion layer
 	UPROPERTY()
 	int32 OcclusionLayerNumNodes;
@@ -176,11 +171,6 @@ class ENGINE_API UHierarchicalInstancedStaticMeshComponent : public UInstancedSt
 	UPROPERTY()
 	FBoxSphereBounds CacheMeshExtendedBounds;
 
-	bool bIsAsyncBuilding;
-	bool bDiscardAsyncBuildResults;
-	bool bConcurrentChanges;
-	bool bAutoRebuildTreeOnInstanceChanges;
-
 	UPROPERTY()
 	bool bDisableCollision;
 
@@ -188,6 +178,16 @@ class ENGINE_API UHierarchicalInstancedStaticMeshComponent : public UInstancedSt
 	UPROPERTY()
 	int32 InstanceCountToRender;
 
+	bool bIsAsyncBuilding : 1;
+	bool bIsOutOfDate : 1;
+	bool bConcurrentChanges : 1;
+	bool bAutoRebuildTreeOnInstanceChanges : 1;
+
+#if WITH_EDITOR
+	// in Editor mode we might disable the density scaling for edition
+	bool bCanEnableDensityScaling : 1;
+#endif
+
 	// Apply the results of the async build
 	void ApplyBuildTreeAsync(ENamedThreads::Type CurrentThread, const FGraphEventRef& MyCompletionGraphEvent, TSharedRef<FClusterBuilder, ESPMode::ThreadSafe> Builder, double StartTime);
 
@@ -236,7 +236,7 @@ public:
 	static void BuildTreeAnyThread(TArray<FMatrix>& InstanceTransforms, const FBox& MeshBox, TArray<FClusterNode>& OutClusterTree, TArray<int32>& OutSortedInstances, TArray<int32>& OutInstanceReorderTable, int32& OutOcclusionLayerNum, int32 MaxInstancesPerLeaf, bool InGenerateInstanceScalingRange);
 	void AcceptPrebuiltTree(TArray<FClusterNode>& InClusterTree, int32 InOcclusionLayerNumNodes, int32 InNumBuiltRenderInstances);
 	bool IsAsyncBuilding() const { return bIsAsyncBuilding; }
-	bool IsTreeFullyBuilt() const { return NumBuiltInstances == PerInstanceSMData.Num(); }
+	bool IsTreeFullyBuilt() const { return !bIsOutOfDate; }
 
 	/** Heuristic for the number of leaves in the tree **/
 	int32 DesiredInstancesPerLeaf();
diff --git a/Engine/Source/Runtime/Engine/Private/HierarchicalInstancedStaticMesh.cpp b/Engine/Source/Runtime/Engine/Private/HierarchicalInstancedStaticMesh.cpp
--- a/Engine/Source/Runtime/Engine/Private/HierarchicalInstancedStaticMesh.cpp
+++ b/Engine/Source/Runtime/Engine/Private/HierarchicalInstancedStaticMesh.cpp
@@ -1947,15 +1947,15 @@ UHierarchicalInstancedStaticMeshComponent::UHierarchicalInstancedStaticMeshCompo
 	, UnbuiltInstanceBounds(ForceInit)
 	, bEnableDensityScaling(false)
 	, CurrentDensityScaling(1.0f)
-#if WITH_EDITOR
-	, bCanEnableDensityScaling(true)
-#endif
 	, OcclusionLayerNumNodes(0)
+	, InstanceCountToRender(0)
 	, bIsAsyncBuilding(false)
-	, bDiscardAsyncBuildResults(false)
+	, bIsOutOfDate(false)
 	, bConcurrentChanges(false)
 	, bAutoRebuildTreeOnInstanceChanges(true)
-	, InstanceCountToRender(0)
+#if WITH_EDITOR
+	, bCanEnableDensityScaling(true)
+#endif
 	, AccumulatedNavigationDirtyArea(ForceInit)
 {
 	bCanEverAffectNavigation = true;
@@ -2093,9 +2093,10 @@ void UHierarchicalInstancedStaticMeshComponent::PostLoad()
 
 void UHierarchicalInstancedStaticMeshComponent::RemoveInstancesInternal(const int32* InstanceIndices, int32 Num)
 {
-	if (IsAsyncBuilding() && Num > 0)
+	if ( Num > 0)
 	{
-		bConcurrentChanges = true;
+		bIsOutOfDate = true;
+		bConcurrentChanges |= IsAsyncBuilding();
 	}
 
 	for (int32 Index = 0; Index < Num; ++Index)
@@ -2228,12 +2229,10 @@ bool UHierarchicalInstancedStaticMeshComponent::UpdateInstanceTransform(int32 In
 		return false;
 	}
 
-	if (IsAsyncBuilding())
-	{
-		// invalidate the results of the current async build we need to modify the tree
-		bConcurrentChanges = true;
-	}
-
+	bIsOutOfDate = true;
+	// invalidate the results of the current async build we need to modify the tree
+	bConcurrentChanges |= IsAsyncBuilding();
+	
 	const int32 RenderIndex = GetRenderIndex(InstanceIndex);
 	const FMatrix OldTransform = PerInstanceSMData[InstanceIndex].Transform;
 	const FTransform NewLocalTransform = bWorldSpace ? NewInstanceTransform.GetRelativeTransform(GetComponentTransform()) : NewInstanceTransform;
@@ -2336,11 +2335,9 @@ int32 UHierarchicalInstancedStaticMeshComponent::AddInstance(const FTransform& I
 	{	
 		check(InstanceIndex == InstanceReorderTable.Num());
 
-		if (IsAsyncBuilding())
-		{
-			bConcurrentChanges = true;
-		}
-
+		bIsOutOfDate = true;
+		bConcurrentChanges |= IsAsyncBuilding();
+	
 		int32 InitialBufferOffset = InstanceCountToRender - InstanceReorderTable.Num(); // Until the build is done, we need to always add at the end of the buffer/reorder table
 		InstanceReorderTable.Add(InitialBufferOffset + InstanceIndex); // add to the end until the build is completed
 		++InstanceCountToRender;
@@ -2362,11 +2359,9 @@ int32 UHierarchicalInstancedStaticMeshComponent::AddInstance(const FTransform& I
 
 void UHierarchicalInstancedStaticMeshComponent::ClearInstances()
 {
-	if (IsAsyncBuilding())
-	{
-		bConcurrentChanges = true;
-	}
-
+	bIsOutOfDate = true;
+	bConcurrentChanges |= IsAsyncBuilding();
+	
 	ClusterTreePtr = MakeShareable(new TArray<FClusterNode>);
 	NumBuiltInstances = 0;
 	NumBuiltRenderInstances = 0;
@@ -2586,11 +2581,8 @@ void UHierarchicalInstancedStaticMeshComponent::BuildTree()
 		CacheMeshExtendedBounds = FBoxSphereBounds(ForceInitToZero);
 	}
 
-	if (bIsAsyncBuilding)
-	{
-		// We did a sync build while async building. The sync build is newer so we will use that.
-		bDiscardAsyncBuildResults = true;
-	}
+	// If an Async Build is running it will bail out because data is no longer bIsOutOfDate
+	bIsOutOfDate = false;
 }
 
 void UHierarchicalInstancedStaticMeshComponent::BuildTreeAnyThread(
@@ -2669,9 +2661,9 @@ void UHierarchicalInstancedStaticMeshComponent::ApplyBuildTreeAsync(ENamedThread
 	BuildTreeAsyncTasks.Empty();
 
 	// We did a sync build while async building. The sync build is newer so we will use that.
-	if (bDiscardAsyncBuildResults)
+	if (!bIsOutOfDate)
 	{
-		bDiscardAsyncBuildResults = false;
+		bConcurrentChanges = false;
 		return;
 	}
 
@@ -2687,6 +2679,9 @@ void UHierarchicalInstancedStaticMeshComponent::ApplyBuildTreeAsync(ENamedThread
 		return;
 	}
 
+	// Completed the build
+	bIsOutOfDate = false;
+
 	check(Builder->Result->InstanceReorderTable.Num() == PerInstanceSMData.Num());
 
 	NumBuiltInstances = Builder->Result->InstanceReorderTable.Num();
@@ -2755,6 +2750,7 @@ bool UHierarchicalInstancedStaticMeshComponent::BuildTreeIfOutdated(bool Async,
 	}
 
 	if (ForceUpdate 
+		|| bIsOutOfDate
 		|| InstanceUpdateCmdBuffer.NumTotalCommands() != 0
 		|| InstanceReorderTable.Num() != PerInstanceSMData.Num()
 		|| NumBuiltInstances != PerInstanceSMData.Num() 
@@ -2763,6 +2759,8 @@ bool UHierarchicalInstancedStaticMeshComponent::BuildTreeIfOutdated(bool Async,
 		|| GetLinkerUE4Version() < VER_UE4_REBUILD_HIERARCHICAL_INSTANCE_TREES
 		|| GetLinkerCustomVersion(FReleaseObjectVersion::GUID) < FReleaseObjectVersion::HISMCClusterTreeMigration)
 	{
+		// Make sure if any of those conditions is true, we mark ourselves out of date so the Async Build completes
+		bIsOutOfDate = true;
 		if (GetStaticMesh() != nullptr && !GetStaticMesh()->HasAnyFlags(RF_NeedLoad)) // we can build the tree if the static mesh is not even loaded, and we can't call PostLoad as the load is not even done
 		{
 			GetStaticMesh()->ConditionalPostLoad();
@@ -2773,7 +2771,6 @@ bool UHierarchicalInstancedStaticMeshComponent::BuildTreeIfOutdated(bool Async,
 				{
 					// invalidate the results of the current async build we need to modify the tree
 					bConcurrentChanges = true;
-					bDiscardAsyncBuildResults = false;
 				}
 				else
 				{
@@ -2874,6 +2871,7 @@ void UHierarchicalInstancedStaticMeshComponent::BuildTreeAsync()
 
 		UnbuiltInstanceBoundsList.Empty();
 		BuiltInstanceBounds.Init();		
+		bIsOutOfDate = false;
 	}
 }
 
