Added an extra check for NULL index buffers, mostly to get static code analysis to shut up.

#rnx

diff --git a/Engine/Source/Runtime/Engine/Private/Components/ModelComponent.cpp b/Engine/Source/Runtime/Engine/Private/Components/ModelComponent.cpp
--- a/Engine/Source/Runtime/Engine/Private/Components/ModelComponent.cpp
+++ b/Engine/Source/Runtime/Engine/Private/Components/ModelComponent.cpp
@@ -680,7 +680,13 @@ bool UModelComponent::GetPhysicsTriMeshData(struct FTriMeshCollisionData* Collis
 		}
 
 		FRawIndexBuffer16or32* IndexBuffer = Element.IndexBuffer;
-		int32 NumIndices = IndexBuffer ? IndexBuffer->Indices.Num() : 0;
+		if (IndexBuffer == nullptr)
+		{
+			UE_LOG(LogPhysics, Warning, TEXT("Found NULL index buffer when cooking UModelComponent physics data! Component: %s, Element: %d, NumTriangles: %d."), *GetPathName(), ElementIndex, Element.NumTriangles);
+			continue;
+		}
+
+		int32 NumIndices = IndexBuffer->Indices.Num();
 		int32 NumVertices = Model->VertexBuffer.Vertices.Num();
 
 		if (NumVertices < (int32)Element.MaxVertexIndex)
