Error in cook step for multiple clients

Readded range check in FStaticMeshInstanceBuffer::UpdateFromCommandBuffer_RenderThread()
#rnx

diff --git a/Engine/Source/Runtime/Engine/Private/InstancedStaticMesh.cpp b/Engine/Source/Runtime/Engine/Private/InstancedStaticMesh.cpp
--- a/Engine/Source/Runtime/Engine/Private/InstancedStaticMesh.cpp
+++ b/Engine/Source/Runtime/Engine/Private/InstancedStaticMesh.cpp
@@ -262,10 +262,11 @@ void FStaticMeshInstanceBuffer::UpdateFromCommandBuffer_RenderThread(FInstanceUp
 	int32 NumCommands = CmdBuffer.NumInlineCommands();
 	int32 NumAdds = CmdBuffer.NumAdds;
 	int32 AddIndex = INDEX_NONE;
+	int32 NewNumInstances = InstanceData->GetNumInstances();
 	if (NumAdds > 0)
 	{
 		AddIndex = InstanceData->GetNumInstances();
-		const int32 NewNumInstances = NumAdds + InstanceData->GetNumInstances();
+		NewNumInstances = NumAdds + InstanceData->GetNumInstances();
 		InstanceData->AllocateInstances(NewNumInstances, GIsEditor ? EResizeBufferFlags::AllowSlackOnGrow | EResizeBufferFlags::AllowSlackOnReduce : EResizeBufferFlags::None, false); // In Editor always permit overallocation, to prevent too much realloc
 	}
 
@@ -278,16 +279,28 @@ void FStaticMeshInstanceBuffer::UpdateFromCommandBuffer_RenderThread(FInstanceUp
 			InstanceData->SetInstance(AddIndex++, Cmd.XForm, 0);
 			break;
 		case FInstanceUpdateCmdBuffer::Hide:
-			InstanceData->NullifyInstance(Cmd.InstanceIndex);
+			if (Cmd.InstanceIndex < NewNumInstances)
+			{
+				InstanceData->NullifyInstance(Cmd.InstanceIndex);
+			}
 			break;
 		case FInstanceUpdateCmdBuffer::Update:
-			InstanceData->SetInstance(Cmd.InstanceIndex, Cmd.XForm, 0);
+			if (Cmd.InstanceIndex < NewNumInstances)
+			{
+				InstanceData->SetInstance(Cmd.InstanceIndex, Cmd.XForm, 0);
+			}
 			break;
 		case FInstanceUpdateCmdBuffer::EditorData:
-			InstanceData->SetInstanceEditorData(Cmd.InstanceIndex, Cmd.HitProxyColor, Cmd.bSelected);
+			if (Cmd.InstanceIndex < NewNumInstances)
+			{
+				InstanceData->SetInstanceEditorData(Cmd.InstanceIndex, Cmd.HitProxyColor, Cmd.bSelected);
+			}
 			break;
 		case FInstanceUpdateCmdBuffer::LightmapData:
-			InstanceData->SetInstanceLightMapData(Cmd.InstanceIndex, Cmd.LightmapUVBias, Cmd.ShadowmapUVBias);
+			if (Cmd.InstanceIndex < NewNumInstances)
+			{
+				InstanceData->SetInstanceLightMapData(Cmd.InstanceIndex, Cmd.LightmapUVBias, Cmd.ShadowmapUVBias);
+			}
 			break;
 		default:
 			check(false);
