Fix missing curve names when baking transform curves

diff --git a/Engine/Source/Runtime/Engine/Private/Animation/AnimSequence.cpp b/Engine/Source/Runtime/Engine/Private/Animation/AnimSequence.cpp
--- a/Engine/Source/Runtime/Engine/Private/Animation/AnimSequence.cpp
+++ b/Engine/Source/Runtime/Engine/Private/Animation/AnimSequence.cpp
@@ -4164,15 +4164,9 @@ void UAnimSequence::BakeTrackCurvesToRawAnimation()
 
 		USkeleton * CurSkeleton = GetSkeleton();
 		check(CurSkeleton);
-
-		const FSmartNameMapping* NameMapping = CurSkeleton->GetSmartNameContainer(USkeleton::AnimTrackCurveMappingName);
-		// if no mapping, that means there is no transform curves
-		if(!NameMapping)
-		{
-			// if no name mapping is found but curve exists, we should verify curve namex
-			VerifyCurveNames<FTransformCurve>(*CurSkeleton, USkeleton::AnimTrackCurveMappingName, RawCurveData.TransformCurves);
-			NameMapping = CurSkeleton->GetSmartNameContainer(USkeleton::AnimTrackCurveMappingName);
-		}
+		
+		VerifyCurveNames<FTransformCurve>(*CurSkeleton, USkeleton::AnimTrackCurveMappingName, RawCurveData.TransformCurves);
+		const FSmartNameMapping*  NameMapping = CurSkeleton->GetSmartNameContainer(USkeleton::AnimTrackCurveMappingName);
 		
 		// since now I'm about to modify Scale Keys. I should add all of them here at least one key. 
 		// if all turns out to be same, it will clear it up. 
