Fix ensure when compressing an animation with Per Track Compression and Max Zeroing Threshold is greater than Max Pos Diff Bitwise

diff --git a/Engine/Source/Runtime/Engine/Private/Animation/AnimCompress_PerTrackCompression.cpp b/Engine/Source/Runtime/Engine/Private/Animation/AnimCompress_PerTrackCompression.cpp
--- a/Engine/Source/Runtime/Engine/Private/Animation/AnimCompress_PerTrackCompression.cpp
+++ b/Engine/Source/Runtime/Engine/Private/Animation/AnimCompress_PerTrackCompression.cpp
@@ -2206,6 +2206,16 @@ void UAnimCompress_PerTrackCompression::PostEditChangeProperty(struct FPropertyC
 			bUseAdaptiveError = (!bUseAdaptiveError2) && bUseAdaptiveError;
 		}
 
+		if (PropertyName == GET_MEMBER_NAME_CHECKED(UAnimCompress_PerTrackCompression, MaxZeroingThreshold)
+			|| PropertyName == GET_MEMBER_NAME_CHECKED(UAnimCompress_PerTrackCompression, MaxPosDiffBitwise))
+		{
+			//MaxZeroingThreshold needs to stay below MaxPosDiffBitwise
+			if (MaxZeroingThreshold > MaxPosDiffBitwise)
+			{
+				MaxZeroingThreshold = MaxPosDiffBitwise;
+			}
+		}
+
 		if (PropertyName == GET_MEMBER_NAME_CHECKED(UAnimCompress_PerTrackCompression, AllowedScaleFormats))
 		{
 			for (TEnumAsByte<enum AnimationCompressionFormat>& ScaleFormat : AllowedScaleFormats)
