fix compile error on linux

diff --git a/Engine/Source/Runtime/Engine/Private/ShaderCompiler/ShaderCompiler.cpp b/Engine/Source/Runtime/Engine/Private/ShaderCompiler/ShaderCompiler.cpp
--- a/Engine/Source/Runtime/Engine/Private/ShaderCompiler/ShaderCompiler.cpp
+++ b/Engine/Source/Runtime/Engine/Private/ShaderCompiler/ShaderCompiler.cpp
@@ -1520,7 +1520,7 @@ void FShaderCompilerStats::WriteStats()
 		FlushRenderingCommands(true);
 
 		FString FileName = FString::Printf(TEXT("%s/MaterialStats/Stats-%s.csv"), *FPaths::ProjectSavedDir(), *FDateTime::Now().ToString());
-		auto DebugWriter = IFileManager::Get().CreateDebugFileWriter(*FileName);
+		auto DebugWriter = IFileManager::Get().CreateFileWriter(*FileName);
 		FDiagnosticTableWriterCSV StatWriter(DebugWriter);
 		const TSparseArray<ShaderCompilerStats>& PlatformStats = GetShaderCompilerStats();
 
@@ -1559,7 +1559,7 @@ void FShaderCompilerStats::WriteStats()
 	{
 
 		FString FileName = FString::Printf(TEXT("%s/MaterialStatsDebug/StatsDebug-%s.csv"), *FPaths::ProjectSavedDir(), *FDateTime::Now().ToString());
-		auto DebugWriter = IFileManager::Get().CreateDebugFileWriter(*FileName);
+		auto DebugWriter = IFileManager::Get().CreateFileWriter(*FileName);
 		FDiagnosticTableWriterCSV StatWriter(DebugWriter);
 		const TSparseArray<ShaderCompilerStats>& PlatformStats = GetShaderCompilerStats();
 		StatWriter.AddColumn(TEXT("Name"));
