Integrating CL# 10134945 from Dev-Enterprise
Adding warning on lightmap restrictions

diff --git a/Engine/Source/Runtime/Engine/Private/StaticMesh.cpp b/Engine/Source/Runtime/Engine/Private/StaticMesh.cpp
--- a/Engine/Source/Runtime/Engine/Private/StaticMesh.cpp
+++ b/Engine/Source/Runtime/Engine/Private/StaticMesh.cpp
@@ -5916,6 +5916,7 @@ void UStaticMesh::EnforceLightmapRestrictions(bool bUseRenderData)
 
 #if !WITH_EDITORONLY_DATA
 	//The source models are only available in the editor, fallback on the render data.
+	UE_ASSET_LOG(LogStaticMesh, Warning, this, TEXT("Trying to enforce lightmap restrictions using the static mesh SourceModels outside of the Editor."))
 	bUseRenderData = true;
 #endif //WITH_EDITORONLY_DATA
 
