Force a synchronous tree build if invalid indices are found in the InstanceReorderTable
#rnx

diff --git a/Engine/Source/Runtime/Engine/Private/HierarchicalInstancedStaticMesh.cpp b/Engine/Source/Runtime/Engine/Private/HierarchicalInstancedStaticMesh.cpp
--- a/Engine/Source/Runtime/Engine/Private/HierarchicalInstancedStaticMesh.cpp
+++ b/Engine/Source/Runtime/Engine/Private/HierarchicalInstancedStaticMesh.cpp
@@ -3023,7 +3023,6 @@ void UHierarchicalInstancedStaticMeshComponent::OnPostLoadPerInstanceData()
 		if (bEnableDensityScaling && GetWorld() && GetWorld()->IsGameWorld())
 		{
 			CurrentDensityScaling = FMath::Clamp(CVarFoliageDensityScale.GetValueOnGameThread(), 0.0f, 1.0f);
-			//CurrentDensityScaling = 0.2f;
 			bForceTreeBuild = CurrentDensityScaling < 1.0f;
 		}
 
@@ -3034,34 +3033,46 @@ void UHierarchicalInstancedStaticMeshComponent::OnPostLoadPerInstanceData()
 		}
 		else
 		{
-			AActor* Owner = GetOwner();
-			ULevel* OwnerLevel = Owner->GetLevel();
-			UWorld* OwnerWorld = OwnerLevel ? OwnerLevel->OwningWorld : nullptr;
+			bool bAsyncTreeBuild = true;
 
-			//update the instance data if the lighting scenario isn't the owner level or if the reorder table do not match the per instance sm data
-			if ((OwnerWorld && OwnerWorld->GetActiveLightingScenario() != nullptr && OwnerWorld->GetActiveLightingScenario() != OwnerLevel)
-				|| (PerInstanceSMData.Num() > 0 && PerInstanceSMData.Num() != InstanceReorderTable.Num()))
+#if WITH_EDITOR
+			// If InstanceReorderTable contains invalid indices force a synchronous tree rebuild
+			bool bHasInvalidIndices = Algo::AnyOf(InstanceReorderTable, [this](int32 ReorderIndex) { return !PerInstanceSMData.IsValidIndex(ReorderIndex); });
+			if (bHasInvalidIndices)
 			{
+				bAsyncTreeBuild = false;
 				bForceTreeBuild = true;
 			}
+#endif
 
-			
-			// If Reorder table contains invalid indices force a tree rebuild
-			if(!bForceTreeBuild)
+			// Update the instance data if the reorder table do not match the per instance sm data
+			if (!bForceTreeBuild && PerInstanceSMData.Num() > 0 && PerInstanceSMData.Num() != InstanceReorderTable.Num())
+			{
+				bForceTreeBuild = true;
+			}
+
+			// Update the instance data if the lighting scenario isn't the owner level
+			if (!bForceTreeBuild)
 			{
-				bForceTreeBuild = Algo::AnyOf(InstanceReorderTable, [this](int32 ReorderIndex) { return ReorderIndex != INDEX_NONE && !PerInstanceSMData.IsValidIndex(ReorderIndex); });
+				AActor* Owner = GetOwner();
+				ULevel* OwnerLevel = Owner->GetLevel();
+				UWorld* OwnerWorld = OwnerLevel ? OwnerLevel->OwningWorld : nullptr;
+				if (OwnerWorld && OwnerWorld->GetActiveLightingScenario() != nullptr && OwnerWorld->GetActiveLightingScenario() != OwnerLevel)
+				{
+					bForceTreeBuild = true;
+				}
 			}
 
 			if (!bForceTreeBuild)
 			{
-				// create PerInstanceRenderData either from current data or pre-built instance buffer
+				// Create PerInstanceRenderData either from current data or pre-built instance buffer
 				InitPerInstanceRenderData(true, InstanceDataBuffers.Release());
 				NumBuiltRenderInstances = PerInstanceRenderData->InstanceBuffer_GameThread->GetNumInstances();
 				InstanceCountToRender = NumBuiltInstances;
 			}
 
 			// If any of the data is out of sync, build the tree now!
-			BuildTreeIfOutdated(true, bForceTreeBuild);
+			BuildTreeIfOutdated(bAsyncTreeBuild, bForceTreeBuild);
 		}
 	}
 
diff --git a/Engine/Source/Runtime/Engine/Private/InstancedStaticMesh.cpp b/Engine/Source/Runtime/Engine/Private/InstancedStaticMesh.cpp
--- a/Engine/Source/Runtime/Engine/Private/InstancedStaticMesh.cpp
+++ b/Engine/Source/Runtime/Engine/Private/InstancedStaticMesh.cpp
@@ -262,45 +262,40 @@ void FStaticMeshInstanceBuffer::UpdateFromCommandBuffer_RenderThread(FInstanceUp
 	int32 NumCommands = CmdBuffer.NumInlineCommands();
 	int32 NumAdds = CmdBuffer.NumAdds;
 	int32 AddIndex = INDEX_NONE;
-	int32 NewNumInstances = InstanceData->GetNumInstances();
+
 	if (NumAdds > 0)
 	{
 		AddIndex = InstanceData->GetNumInstances();
-		NewNumInstances = NumAdds + InstanceData->GetNumInstances();
+		int32 NewNumInstances = NumAdds + InstanceData->GetNumInstances();
 		InstanceData->AllocateInstances(NewNumInstances, GIsEditor ? EResizeBufferFlags::AllowSlackOnGrow | EResizeBufferFlags::AllowSlackOnReduce : EResizeBufferFlags::None, false); // In Editor always permit overallocation, to prevent too much realloc
 	}
 
 	for (int32 i = 0; i < NumCommands; ++i)
 	{
 		const auto& Cmd = CmdBuffer.Cmds[i];
+
+		int32 InstanceIndex = Cmd.Type != FInstanceUpdateCmdBuffer::Add ? Cmd.InstanceIndex : AddIndex++;
+		if (!ensure(InstanceData->IsValidIndex(InstanceIndex)))
+		{
+			continue;
+		}
+
 		switch (Cmd.Type)
 		{
 		case FInstanceUpdateCmdBuffer::Add:
-			InstanceData->SetInstance(AddIndex++, Cmd.XForm, 0);
+			InstanceData->SetInstance(InstanceIndex, Cmd.XForm, 0);
 			break;
 		case FInstanceUpdateCmdBuffer::Hide:
-			if (Cmd.InstanceIndex < NewNumInstances)
-			{
-				InstanceData->NullifyInstance(Cmd.InstanceIndex);
-			}
+			InstanceData->NullifyInstance(InstanceIndex);
 			break;
 		case FInstanceUpdateCmdBuffer::Update:
-			if (Cmd.InstanceIndex < NewNumInstances)
-			{
-				InstanceData->SetInstance(Cmd.InstanceIndex, Cmd.XForm, 0);
-			}
+			InstanceData->SetInstance(InstanceIndex, Cmd.XForm, 0);
 			break;
 		case FInstanceUpdateCmdBuffer::EditorData:
-			if (Cmd.InstanceIndex < NewNumInstances)
-			{
-				InstanceData->SetInstanceEditorData(Cmd.InstanceIndex, Cmd.HitProxyColor, Cmd.bSelected);
-			}
+			InstanceData->SetInstanceEditorData(InstanceIndex, Cmd.HitProxyColor, Cmd.bSelected);
 			break;
 		case FInstanceUpdateCmdBuffer::LightmapData:
-			if (Cmd.InstanceIndex < NewNumInstances)
-			{
-				InstanceData->SetInstanceLightMapData(Cmd.InstanceIndex, Cmd.LightmapUVBias, Cmd.ShadowmapUVBias);
-			}
+			InstanceData->SetInstanceLightMapData(InstanceIndex, Cmd.LightmapUVBias, Cmd.ShadowmapUVBias);
 			break;
 		default:
 			check(false);
