Add CVAR toggles for considering each type of SQCapture (Raycast,Overlap,Sweep) when using p.SerializeSQs for more control.

diff --git a/Engine/Source/Runtime/Engine/Private/Collision/SceneQueryLowLevel.cpp b/Engine/Source/Runtime/Engine/Private/Collision/SceneQueryLowLevel.cpp
--- a/Engine/Source/Runtime/Engine/Private/Collision/SceneQueryLowLevel.cpp
+++ b/Engine/Source/Runtime/Engine/Private/Collision/SceneQueryLowLevel.cpp
@@ -29,11 +29,17 @@ int32 SerializeSQs = 0;
 int32 SerializeSQSamples = 100;
 int32 SerializeBadSQs = 0;
 int32 ReplaySQs = 0;
+int32 EnableRaycastSQCapture = 1;
+int32 EnableOverlapSQCapture = 1;
+int32 EnableSweepSQCapture = 1;
 
 FAutoConsoleVariableRef CVarSerializeSQs(TEXT("p.SerializeSQs"), SerializeSQs, TEXT("If enabled, we create a sq capture per sq that takes more than provided value in microseconds. This can be very expensive as the entire scene is saved out"));
 FAutoConsoleVariableRef CVarSerializeSQSamples(TEXT("p.SerializeSQSampleCount"), SerializeSQSamples, TEXT("If Query exceeds duration threshold, we will re-measure SQ this many times before serializing. Larger values cause hitching."));
 FAutoConsoleVariableRef CVarReplaySweeps(TEXT("p.ReplaySQs"), ReplaySQs, TEXT("If enabled, we rerun the sq against chaos"));
 FAutoConsoleVariableRef CVarSerializeBadSweeps(TEXT("p.SerializeBadSQs"), SerializeBadSQs, TEXT("If enabled, we create a sq capture whenever chaos and physx diverge"));
+FAutoConsoleVariableRef CVarSerializeSQsRaycastEnabled(TEXT("p.SerializeSQsRaycastEnabled"), EnableRaycastSQCapture, TEXT("If disabled, p.SerializeSQs will not consider raycasts"));
+FAutoConsoleVariableRef CVarSerializeSQsOverlapEnabled(TEXT("p.SerializeSQsOverlapEnabled"), EnableOverlapSQCapture, TEXT("If disabled, p.SerializeSQs will not consider overlaps"));
+FAutoConsoleVariableRef CVarSerializeSQsSweepEnabled(TEXT("p.SerializeSQsSweepEnabled"), EnableSweepSQCapture, TEXT("If disabled, p.SerializeSQs will not consider sweeps"));
 
 void FinalizeCapture(FPhysTestSerializer& Serializer)
 {
@@ -60,6 +66,10 @@ void FinalizeCapture(FPhysTestSerializer& Serializer)
 constexpr int32 SerializeSQs = 0;
 constexpr int32 ReplaySQs = 0;
 constexpr int32 SerializeSQSamples = 0;
+constexpr int32 EnableRaycastSQCapture = 0;
+constexpr int32 EnableOverlapSQCapture = 0;
+constexpr int32 EnableSweepSQCapture = 0;
+
 // No-op in shipping
 void FinalizeCapture(FPhysTestSerializer& Serializer) {}
 #endif
@@ -185,9 +195,9 @@ void LowLevelRaycast(FPhysScene& Scene, const FVector& Start, const FVector& Dir
 			SQAccelerator.Raycast(Start, Dir, DeltaMag, HitBuffer, OutputFlags, QueryFilterData, *QueryCallback, DebugParams);
 		}
 
-		if (SerializeSQs)
+		if (SerializeSQs && EnableRaycastSQCapture)
 		{
-			//RaycastSQCaptureHelper(QueryDurationSeconds, SQAccelerator, Scene, Start, Dir, DeltaMag, HitBuffer, OutputFlags, QueryFlags, Filter, QueryFilterData, *QueryCallback, DebugParams);
+			RaycastSQCaptureHelper(Time, SQAccelerator, Scene, Start, Dir, DeltaMag, HitBuffer, OutputFlags, QueryFlags, Filter, QueryFilterData, QueryCallback, DebugParams);
 		}
 	}
 #else
@@ -223,7 +233,7 @@ void LowLevelSweep(FPhysScene& Scene, const FPhysicsGeometry& QueryGeom, const F
 				SQAccelerator.Sweep(QueryGeom, StartTM, Dir, DeltaMag, HitBuffer, OutputFlags, QueryFilterData, *QueryCallback, DebugParams);
 			}
 
-			if (SerializeSQs)
+			if (SerializeSQs && EnableSweepSQCapture)
 			{
 				SweepSQCaptureHelper(Time, SQAccelerator, Scene, QueryGeom, StartTM, Dir, DeltaMag, HitBuffer, OutputFlags, QueryFlags, Filter, QueryFilterData, QueryCallback, DebugParams);
 			}
@@ -260,9 +270,9 @@ void LowLevelOverlap(FPhysScene& Scene, const FPhysicsGeometry& QueryGeom, const
 			SQAccelerator.Overlap(QueryGeom, GeomPose, HitBuffer, QueryFilterData, *QueryCallback);
 		}
 
-		if (SerializeSQs)
+		if (SerializeSQs && EnableOverlapSQCapture)
 		{
-			//OverlapSQCaptureHelper(Time, SQAccelerator, Scene, QueryGeom, GeomPose, HitBuffer, QueryFlags, Filter, QueryFilterData, QueryCallback, DebugParams);
+			OverlapSQCaptureHelper(Time, SQAccelerator, Scene, QueryGeom, GeomPose, HitBuffer, QueryFlags, Filter, QueryFilterData, QueryCallback, DebugParams);
 		}
 	}
 #else
