Renamed r.UpdateReflectionCaptureEveryFrame to r.ReflectionCaptureUpdateEveryFrame.
discussed with Daniel.

diff --git a/Engine/Source/Runtime/Engine/Private/Components/ReflectionCaptureComponent.cpp b/Engine/Source/Runtime/Engine/Private/Components/ReflectionCaptureComponent.cpp
--- a/Engine/Source/Runtime/Engine/Private/Components/ReflectionCaptureComponent.cpp
+++ b/Engine/Source/Runtime/Engine/Private/Components/ReflectionCaptureComponent.cpp
@@ -58,8 +58,8 @@ ENGINE_API TAutoConsoleVariable<int32> CVarReflectionCaptureSize(
 	TEXT("Set the resolution for all reflection capture cubemaps. Should be set via project's Render Settings. Must be power of 2. Defaults to 128.\n")
 	);
 
-TAutoConsoleVariable<int32> CVarUpdateReflectionCaptureEveryFrame(
-	TEXT("r.UpdateReflectionCaptureEveryFrame"),
+TAutoConsoleVariable<int32> CVarReflectionCaptureUpdateEveryFrame(
+	TEXT("r.ReflectionCaptureUpdateEveryFrame"),
 	0,
 	TEXT("When set, reflection captures will constantly be scheduled for update.\n")
 );
@@ -1056,7 +1056,7 @@ void UReflectionCaptureComponent::UpdateReflectionCaptureContents(UWorld* WorldT
 		//guarantee that all render proxies are up to date before kicking off this render
 		WorldToUpdate->SendAllEndOfFrameUpdates();
 
-		if (CVarUpdateReflectionCaptureEveryFrame.GetValueOnGameThread())
+		if (CVarReflectionCaptureUpdateEveryFrame.GetValueOnGameThread())
 		{
 			for (FActorIterator It(WorldToUpdate); It; ++It)
 			{
