Implement AddTorque

diff --git a/Engine/Source/Runtime/Engine/Private/PhysicsEngine/Experimental/PhysScene_Chaos.cpp b/Engine/Source/Runtime/Engine/Private/PhysicsEngine/Experimental/PhysScene_Chaos.cpp
--- a/Engine/Source/Runtime/Engine/Private/PhysicsEngine/Experimental/PhysScene_Chaos.cpp
+++ b/Engine/Source/Runtime/Engine/Private/PhysicsEngine/Experimental/PhysScene_Chaos.cpp
@@ -1158,8 +1158,23 @@ void FPhysScene_ChaosInterface::ClearForces_AssumesLocked(FBodyInstance* BodyIns
 
 void FPhysScene_ChaosInterface::AddTorque_AssumesLocked(FBodyInstance* BodyInstance, const FVector& Torque, bool bAllowSubstepping, bool bAccelChange)
 {
-	// #todo : Implement
-	// AddTorque(Torque, BodyInstance->ActorHandle.GetId());
+	FPhysicsActorHandle& Handle = BodyInstance->GetPhysicsActorHandle();
+	if (ensure(FPhysicsInterface::IsValid(Handle)))
+	{
+		Chaos::TPBDRigidParticle<float, 3>* Rigid = Handle->AsDynamic();
+		if (ensure(Rigid))
+		{
+			const Chaos::TVector<float, 3> CurrentTorque = Rigid->ExternalTorque();
+			if (bAccelChange)
+			{
+				ensureMsgf(false, TEXT("Needs implementation"));
+			}
+			else
+			{
+				Rigid->SetExternalTorque(CurrentTorque + Torque);
+			}
+		}
+	}
 }
 
 void FPhysScene_ChaosInterface::ClearTorques_AssumesLocked(FBodyInstance* BodyInstance, bool bAllowSubstepping)
