Change ensure in UAnimInstance::ShouldTriggerAnimNotifyState() to only fire once.
#rnx

diff --git a/Engine/Source/Runtime/Engine/Private/Animation/AnimInstance.cpp b/Engine/Source/Runtime/Engine/Private/Animation/AnimInstance.cpp
--- a/Engine/Source/Runtime/Engine/Private/Animation/AnimInstance.cpp
+++ b/Engine/Source/Runtime/Engine/Private/Animation/AnimInstance.cpp
@@ -3149,7 +3149,7 @@ void UAnimInstance::DestroyAnimInstanceProxy(FAnimInstanceProxy* InProxy)
 
 bool UAnimInstance::ShouldTriggerAnimNotifyState(const UAnimNotifyState* AnimNotifyState) const
 {
-	if (ensureAlwaysMsgf(AnimNotifyState != nullptr, TEXT("UAnimInstance::ShouldTriggerAnimNotifyState: AnimNotifyState is null on AnimInstance %s. ActiveAnimNotifyState array size is: %d"), *GetNameSafe(this), ActiveAnimNotifyState.Num()))
+	if (ensureMsgf(AnimNotifyState != nullptr, TEXT("UAnimInstance::ShouldTriggerAnimNotifyState: AnimNotifyState is null on AnimInstance %s. ActiveAnimNotifyState array size is: %d"), *GetNameSafe(this), ActiveAnimNotifyState.Num()))
 	{
 		return true;
 	}
