#UE4 CIS fix

diff --git a/Engine/Source/Runtime/Engine/Private/ShaderCompiler/ShaderCompiler.cpp b/Engine/Source/Runtime/Engine/Private/ShaderCompiler/ShaderCompiler.cpp
--- a/Engine/Source/Runtime/Engine/Private/ShaderCompiler/ShaderCompiler.cpp
+++ b/Engine/Source/Runtime/Engine/Private/ShaderCompiler/ShaderCompiler.cpp
@@ -1516,6 +1516,7 @@ FShaderCompilerStats* GShaderCompilerStats = NULL;
 
 void FShaderCompilerStats::WriteStats()
 {
+#if ALLOW_DEBUG_FILES
 	{
 		FlushRenderingCommands(true);
 
@@ -1597,6 +1598,7 @@ void FShaderCompilerStats::WriteStats()
 			}
 		}
 	}
+#endif // ALLOW_DEBUG_FILES
 }
 void FShaderCompilerStats::RegisterCookedShaders(uint32 NumCooked, EShaderPlatform Platform, const FString MaterialPath, FString PermutationString)
 {
