Chaos - Fix Linux CIS errors in SceneQueryLowLevel.

diff --git a/Engine/Source/Runtime/Engine/Private/Collision/SceneQueryLowLevel.cpp b/Engine/Source/Runtime/Engine/Private/Collision/SceneQueryLowLevel.cpp
--- a/Engine/Source/Runtime/Engine/Private/Collision/SceneQueryLowLevel.cpp
+++ b/Engine/Source/Runtime/Engine/Private/Collision/SceneQueryLowLevel.cpp
@@ -195,7 +195,7 @@ void LowLevelRaycast(FPhysScene& Scene, const FVector& Start, const FVector& Dir
 			SQAccelerator.Raycast(Start, Dir, DeltaMag, HitBuffer, OutputFlags, QueryFilterData, *QueryCallback, DebugParams);
 		}
 
-		if (SerializeSQs && EnableRaycastSQCapture)
+		if (SerializeSQs & EnableRaycastSQCapture)
 		{
 			RaycastSQCaptureHelper(Time, SQAccelerator, Scene, Start, Dir, DeltaMag, HitBuffer, OutputFlags, QueryFlags, Filter, QueryFilterData, QueryCallback, DebugParams);
 		}
@@ -233,7 +233,7 @@ void LowLevelSweep(FPhysScene& Scene, const FPhysicsGeometry& QueryGeom, const F
 				SQAccelerator.Sweep(QueryGeom, StartTM, Dir, DeltaMag, HitBuffer, OutputFlags, QueryFilterData, *QueryCallback, DebugParams);
 			}
 
-			if (SerializeSQs && EnableSweepSQCapture)
+			if (SerializeSQs & EnableSweepSQCapture)
 			{
 				SweepSQCaptureHelper(Time, SQAccelerator, Scene, QueryGeom, StartTM, Dir, DeltaMag, HitBuffer, OutputFlags, QueryFlags, Filter, QueryFilterData, QueryCallback, DebugParams);
 			}
@@ -270,7 +270,7 @@ void LowLevelOverlap(FPhysScene& Scene, const FPhysicsGeometry& QueryGeom, const
 			SQAccelerator.Overlap(QueryGeom, GeomPose, HitBuffer, QueryFilterData, *QueryCallback);
 		}
 
-		if (SerializeSQs && EnableOverlapSQCapture)
+		if (SerializeSQs & EnableOverlapSQCapture)
 		{
 			OverlapSQCaptureHelper(Time, SQAccelerator, Scene, QueryGeom, GeomPose, HitBuffer, QueryFlags, Filter, QueryFilterData, QueryCallback, DebugParams);
 		}
