Fixing lightmap restrictions warning

diff --git a/Engine/Source/Runtime/Engine/Private/StaticMesh.cpp b/Engine/Source/Runtime/Engine/Private/StaticMesh.cpp
--- a/Engine/Source/Runtime/Engine/Private/StaticMesh.cpp
+++ b/Engine/Source/Runtime/Engine/Private/StaticMesh.cpp
@@ -5915,9 +5915,12 @@ void UStaticMesh::EnforceLightmapRestrictions(bool bUseRenderData)
 	int32 NumUVs = 4;
 
 #if !WITH_EDITORONLY_DATA
-	//The source models are only available in the editor, fallback on the render data.
-	UE_ASSET_LOG(LogStaticMesh, Warning, this, TEXT("Trying to enforce lightmap restrictions using the static mesh SourceModels outside of the Editor."))
-	bUseRenderData = true;
+	if (!bUseRenderData)
+	{
+		//The source models are only available in the editor, fallback on the render data.
+		UE_ASSET_LOG(LogStaticMesh, Warning, this, TEXT("Trying to enforce lightmap restrictions using the static mesh SourceModels outside of the Editor."))
+		bUseRenderData = true;
+	}
 #endif //WITH_EDITORONLY_DATA
 
 	if (bUseRenderData)
