Remove test code from FTimerManager::GenerateHandle.

diff --git a/Engine/Source/Runtime/Engine/Private/TimerManager.cpp b/Engine/Source/Runtime/Engine/Private/TimerManager.cpp
--- a/Engine/Source/Runtime/Engine/Private/TimerManager.cpp
+++ b/Engine/Source/Runtime/Engine/Private/TimerManager.cpp
@@ -1088,13 +1088,8 @@ FTimerHandle FTimerManager::GenerateHandle(int32 Index)
 		NewSerialNumber = (uint64)1;
 	}
 
-	FTimerHandle Test;
-	Test.SetIndexAndSerialNumber(FTimerHandle::MaxIndex - 1, FTimerHandle::MaxSerialNumber - 1);
-	check(Test.GetIndex() == FTimerHandle::MaxIndex - 1 && Test.GetSerialNumber() == FTimerHandle::MaxSerialNumber - 1);
-
 	FTimerHandle Result;
 	Result.SetIndexAndSerialNumber(Index, NewSerialNumber);
-	check(Result.GetIndex() == Index && Result.GetSerialNumber() == NewSerialNumber);
 	return Result;
 }
 
