Fixing bug in ClearAttenuationOverride causing Listener and ListenerProxies to become unsync'ed.

diff --git a/Engine/Source/Runtime/Engine/Private/AudioDevice.cpp b/Engine/Source/Runtime/Engine/Private/AudioDevice.cpp
--- a/Engine/Source/Runtime/Engine/Private/AudioDevice.cpp
+++ b/Engine/Source/Runtime/Engine/Private/AudioDevice.cpp
@@ -3065,7 +3065,7 @@ void FAudioDevice::ClearListenerAttenuationOverride(int32 ListenerIndex)
 		DECLARE_CYCLE_STAT(TEXT("FAudioThreadTask.ClearListenerAttenuationOverride"), STAT_AudioClearListenerAttenuationOverride, STATGROUP_AudioThreadCommands);
 
 		ListenerProxies[ListenerIndex].AttenuationOverride = FVector::ZeroVector;
-		ListenerProxies[ListenerIndex].bUseAttenuationOverride = true;
+		ListenerProxies[ListenerIndex].bUseAttenuationOverride = false;
 
 		FAudioDevice* AudioDevice = this;
 		FAudioThread::RunCommandOnAudioThread([AudioDevice, ListenerIndex]()
