Implement tick clamping in chaos for behavior consistent with PhysX

diff --git a/Engine/Source/Runtime/Engine/Private/PhysicsEngine/Experimental/PhysScene_Chaos.cpp b/Engine/Source/Runtime/Engine/Private/PhysicsEngine/Experimental/PhysScene_Chaos.cpp
--- a/Engine/Source/Runtime/Engine/Private/PhysicsEngine/Experimental/PhysScene_Chaos.cpp
+++ b/Engine/Source/Runtime/Engine/Private/PhysicsEngine/Experimental/PhysScene_Chaos.cpp
@@ -1243,7 +1243,7 @@ void FPhysScene_ChaosInterface::ApplyWorldOffset(FVector InOffset)
 void FPhysScene_ChaosInterface::SetUpForFrame(const FVector* NewGrav, float InDeltaSeconds /*= 0.0f*/, float InMaxPhysicsDeltaTime /*= 0.0f*/)
 {
 	SetGravity(*NewGrav);
-	MDeltaTime = InDeltaSeconds;
+	MDeltaTime = InDeltaSeconds < InMaxPhysicsDeltaTime ? InDeltaSeconds : InMaxPhysicsDeltaTime;
 }
 
 void FPhysScene_ChaosInterface::StartFrame()
