Set instantaneous torque

diff --git a/Engine/Source/Runtime/Engine/Private/PhysicsEngine/Experimental/PhysScene_Chaos.cpp b/Engine/Source/Runtime/Engine/Private/PhysicsEngine/Experimental/PhysScene_Chaos.cpp
--- a/Engine/Source/Runtime/Engine/Private/PhysicsEngine/Experimental/PhysScene_Chaos.cpp
+++ b/Engine/Source/Runtime/Engine/Private/PhysicsEngine/Experimental/PhysScene_Chaos.cpp
@@ -1167,7 +1167,7 @@ void FPhysScene_ChaosInterface::AddTorque_AssumesLocked(FBodyInstance* BodyInsta
 			const Chaos::TVector<float, 3> CurrentTorque = Rigid->ExternalTorque();
 			if (bAccelChange)
 			{
-				ensureMsgf(false, TEXT("Needs implementation"));
+				Rigid->SetExternalTorque(CurrentTorque + (Rigid->I() * Torque));
 			}
 			else
 			{
