Implement get and set velocities

diff --git a/Engine/Source/Runtime/Engine/Private/PhysicsEngine/Experimental/PhysInterface_Chaos.cpp b/Engine/Source/Runtime/Engine/Private/PhysicsEngine/Experimental/PhysInterface_Chaos.cpp
--- a/Engine/Source/Runtime/Engine/Private/PhysicsEngine/Experimental/PhysInterface_Chaos.cpp
+++ b/Engine/Source/Runtime/Engine/Private/PhysicsEngine/Experimental/PhysInterface_Chaos.cpp
@@ -432,37 +432,80 @@ void FPhysInterface_Chaos::SetKinematicTarget_AssumesLocked(const FPhysicsActorH
 
 FVector FPhysInterface_Chaos::GetLinearVelocity_AssumesLocked(const FPhysicsActorHandle& InActorReference)
 {
-	// #todo : Implement
+	if (ensure(FPhysicsInterface::IsValid(InActorReference)))
+	{
+		Chaos::TKinematicGeometryParticle<float, 3>* Kinematic = InActorReference->AsKinematic();
+		if (ensure(Kinematic))
+		{
+			return Kinematic->V();
+		}
+	}
+
 	return FVector(0);
 }
 
 void FPhysInterface_Chaos::SetLinearVelocity_AssumesLocked(const FPhysicsActorHandle& InActorReference, const FVector& InNewVelocity, bool bAutoWake)
 {
-	// #todo : Implement
+	// TODO: Implement bAutoWake == false.
+	// For now we don't support auto-awake == false.
+	// This feature is meant to detect when velocity change small
+	// and the velocity is nearly zero, and to not wake up the
+	// body in that case.
+	ensure(bAutoWake);
+
+	if (ensure(FPhysicsInterface::IsValid(InActorReference)))
+	{
+		Chaos::TKinematicGeometryParticle<float, 3>* Kinematic = InActorReference->AsKinematic();
+		if (ensure(Kinematic))
+		{
+			return Kinematic->SetV(InNewVelocity);
+		}
+	}
 }
 
 FVector FPhysInterface_Chaos::GetAngularVelocity_AssumesLocked(const FPhysicsActorHandle& InActorReference)
 {
-	// #todo : Implement
+	if (ensure(FPhysicsInterface::IsValid(InActorReference)))
+	{
+		Chaos::TKinematicGeometryParticle<float, 3>* Kinematic = InActorReference->AsKinematic();
+		if (ensure(Kinematic))
+		{
+			return Kinematic->W();
+		}
+	}
+
 	return FVector(0);
 }
 
-void FPhysInterface_Chaos::SetAngularVelocity_AssumesLocked(const FPhysicsActorHandle& InActorReference, const FVector& InNewVelocity, bool bAutoWake)
+void FPhysInterface_Chaos::SetAngularVelocity_AssumesLocked(const FPhysicsActorHandle& InActorReference, const FVector& InNewAngularVelocity, bool bAutoWake)
 {
-	// #todo : Implement
+	// TODO: Implement bAutoWake == false.
+	ensure(bAutoWake);
+
+	if (ensure(FPhysicsInterface::IsValid(InActorReference)))
+	{
+		Chaos::TKinematicGeometryParticle<float, 3>* Kinematic = InActorReference->AsKinematic();
+		if (ensure(Kinematic))
+		{
+			return Kinematic->SetW(InNewAngularVelocity);
+		}
+	}
 }
 
 float FPhysInterface_Chaos::GetMaxAngularVelocity_AssumesLocked(const FPhysicsActorHandle& InActorReference)
 {
+	ensure(false);
     return FLT_MAX;
 }
 
 void FPhysInterface_Chaos::SetMaxAngularVelocity_AssumesLocked(const FPhysicsActorHandle& InActorReference, float InMaxAngularVelocity)
 {
+	ensure(false);
 }
 
 float FPhysInterface_Chaos::GetMaxDepenetrationVelocity_AssumesLocked(const FPhysicsActorHandle& InActorReference)
 {
+	ensure(false);
     return FLT_MAX;
 }
 
