Added Light Shaft Bloom post exposure brightness clamp.
Helps fixing artefacts because bloom cannot handle the bright physically based sun light values.

diff --git a/Engine/Shaders/Private/LightShaftShader.usf b/Engine/Shaders/Private/LightShaftShader.usf
--- a/Engine/Shaders/Private/LightShaftShader.usf
+++ b/Engine/Shaders/Private/LightShaftShader.usf
@@ -26,6 +26,8 @@ float4 WorldSpaceBlurOriginAndRadius;
 
 float LightSourceRadius;
 
+float BloomMaxBrightness;
+
 /** Spot light direction. */
 float3 WorldSpaceSpotDirection;
 
@@ -120,9 +122,10 @@ void DownsampleLightShaftsPixelMain(
 	float EdgeMask = 1.0f - NormalizedCoordinates.x * (1.0f - NormalizedCoordinates.x) * NormalizedCoordinates.y * (1.0f - NormalizedCoordinates.y) * 8.0f;
 	EdgeMask = EdgeMask * EdgeMask * EdgeMask * EdgeMask;
 
-	// Only bloom colors over BloomThreshold
+	// Only bloom colors according if post exposure brightness is over BloomThreshold,
+	// and if brightness is greater than BloomMaxBrightness, scale it down to BloomMaxBrightness to avoid saturated artefacts.
 	float Luminance = max(dot(SceneColor, half3(.3f, .59f, .11f)), 6.10352e-5);
-	float AdjustedLuminance = max(Luminance - BloomTintAndThreshold.a, 0.0f);
+	float AdjustedLuminance = clamp(Luminance - BloomTintAndThreshold.a, 0.0f, BloomMaxBrightness);
 	float3 BloomColor = LightShaftParameters.y * SceneColor / Luminance * AdjustedLuminance * 2.0f;
 
 	#if POINT_LIGHT_SHAFTS
diff --git a/Engine/Source/Runtime/Engine/Classes/Components/LightComponent.h b/Engine/Source/Runtime/Engine/Classes/Components/LightComponent.h
--- a/Engine/Source/Runtime/Engine/Classes/Components/LightComponent.h
+++ b/Engine/Source/Runtime/Engine/Classes/Components/LightComponent.h
@@ -203,6 +203,10 @@ class ENGINE_API ULightComponent : public ULightComponentBase
 	UPROPERTY(EditAnywhere, BlueprintReadOnly, Category=LightShafts, meta=(UIMin = "0", UIMax = "4"))
 	float BloomThreshold;
 
+	/** After exposure is applied, scene color brightness larger than BloomMaxBrightness will be rescaled down to BloomMaxBrightness. */
+	UPROPERTY(EditAnywhere, BlueprintReadOnly, Category=LightShafts, meta=(UIMin = "0", UIMax = "100", SliderExponent = 20.0))
+	float BloomMaxBrightness;
+
 	/** Multiplies against scene color to create the bloom color. */
 	UPROPERTY(EditAnywhere, BlueprintReadOnly, Category=LightShafts)
 	FColor BloomTint;
@@ -271,6 +275,9 @@ public:
 	UFUNCTION(BlueprintCallable, Category="Rendering|Components|Light")
 	void SetBloomThreshold(float NewValue);
 
+	UFUNCTION(BlueprintCallable, Category="Rendering|Components|Light")
+	void SetBloomMaxBrightness(float NewValue);
+
 	UFUNCTION(BlueprintCallable, Category="Rendering|Components|Light")
 	void SetBloomTint(FColor NewValue);
 
diff --git a/Engine/Source/Runtime/Engine/Private/Components/LightComponent.cpp b/Engine/Source/Runtime/Engine/Private/Components/LightComponent.cpp
--- a/Engine/Source/Runtime/Engine/Private/Components/LightComponent.cpp
+++ b/Engine/Source/Runtime/Engine/Private/Components/LightComponent.cpp
@@ -435,6 +435,7 @@ ULightComponent::ULightComponent(const FObjectInitializer& ObjectInitializer)
 	bEnableLightShaftBloom = false;
 	BloomScale = .2f;
 	BloomThreshold = 0;
+	BloomMaxBrightness = 100.0f;
 	BloomTint = FColor::White;
 
 	RayStartOffsetDepthScale = .003f;
@@ -627,6 +628,7 @@ bool ULightComponent::CanEditChange(const UProperty* InProperty) const
 
 		if (PropertyName == GET_MEMBER_NAME_STRING_CHECKED(ULightComponent, BloomScale)
 			|| PropertyName == GET_MEMBER_NAME_STRING_CHECKED(ULightComponent, BloomThreshold)
+			|| PropertyName == GET_MEMBER_NAME_STRING_CHECKED(ULightComponent, BloomMaxBrightness)
 			|| PropertyName == GET_MEMBER_NAME_STRING_CHECKED(ULightComponent, BloomTint))
 		{
 			return bEnableLightShaftBloom;
@@ -677,6 +679,7 @@ void ULightComponent::PostEditChangeProperty(FPropertyChangedEvent& PropertyChan
 		PropertyName != GET_MEMBER_NAME_STRING_CHECKED(ULightComponent, bEnableLightShaftBloom) &&
 		PropertyName != GET_MEMBER_NAME_STRING_CHECKED(ULightComponent, BloomScale) &&
 		PropertyName != GET_MEMBER_NAME_STRING_CHECKED(ULightComponent, BloomThreshold) &&
+		PropertyName != GET_MEMBER_NAME_STRING_CHECKED(ULightComponent, BloomMaxBrightness) &&
 		PropertyName != GET_MEMBER_NAME_STRING_CHECKED(ULightComponent, BloomTint) &&
 		PropertyName != GET_MEMBER_NAME_STRING_CHECKED(ULightComponent, bUseRayTracedDistanceFieldShadows) &&
 		PropertyName != GET_MEMBER_NAME_STRING_CHECKED(ULightComponent, RayStartOffsetDepthScale) &&
@@ -1023,6 +1026,16 @@ void ULightComponent::SetBloomThreshold(float NewValue)
 	}
 }
 
+void ULightComponent::SetBloomMaxBrightness(float NewValue)
+{
+	if (AreDynamicDataChangesAllowed()
+		&& BloomMaxBrightness != NewValue)
+	{
+		BloomMaxBrightness = NewValue;
+		MarkRenderStateDirty();
+	}
+}
+
 void ULightComponent::SetBloomTint(FColor NewValue)
 {
 	if (AreDynamicDataChangesAllowed()
diff --git a/Engine/Source/Runtime/Renderer/Private/LightSceneInfo.cpp b/Engine/Source/Runtime/Renderer/Private/LightSceneInfo.cpp
--- a/Engine/Source/Runtime/Renderer/Private/LightSceneInfo.cpp
+++ b/Engine/Source/Runtime/Renderer/Private/LightSceneInfo.cpp
@@ -47,6 +47,7 @@ FLightSceneInfo::FLightSceneInfo(FLightSceneProxy* InProxy, bool InbVisible)
 	, bEnableLightShaftBloom(InProxy->GetLightComponent()->bEnableLightShaftBloom)
 	, BloomScale(InProxy->GetLightComponent()->BloomScale)
 	, BloomThreshold(InProxy->GetLightComponent()->BloomThreshold)
+	, BloomMaxBrightness(InProxy->GetLightComponent()->BloomMaxBrightness)
 	, BloomTint(InProxy->GetLightComponent()->BloomTint)
 	, NumUnbuiltInteractions(0)
 	, bCreatePerObjectShadowsForDynamicObjects(Proxy->ShouldCreatePerObjectShadowsForDynamicObjects())
diff --git a/Engine/Source/Runtime/Renderer/Private/LightSceneInfo.h b/Engine/Source/Runtime/Renderer/Private/LightSceneInfo.h
--- a/Engine/Source/Runtime/Renderer/Private/LightSceneInfo.h
+++ b/Engine/Source/Runtime/Renderer/Private/LightSceneInfo.h
@@ -204,6 +204,9 @@ public:
 	/** Scene color must be larger than this to create bloom in the light shafts. */
 	float BloomThreshold;
 
+	/** After exposure is applied, scene color brightness larger than BloomMaxBrightness will be rescaled down to BloomMaxBrightness. */
+	float BloomMaxBrightness;
+
 	/** Multiplies against scene color to create the bloom color. */
 	FColor BloomTint;
 
diff --git a/Engine/Source/Runtime/Renderer/Private/LightShaftRendering.cpp b/Engine/Source/Runtime/Renderer/Private/LightShaftRendering.cpp
--- a/Engine/Source/Runtime/Renderer/Private/LightShaftRendering.cpp
+++ b/Engine/Source/Runtime/Renderer/Private/LightShaftRendering.cpp
@@ -121,6 +121,7 @@ public:
 		AspectRatioAndInvAspectRatioParameter.Bind(ParameterMap,TEXT("AspectRatioAndInvAspectRatio"));
 		LightShaftParameters.Bind(ParameterMap, TEXT("LightShaftParameters"));
 		BloomTintAndThresholdParameter.Bind(ParameterMap,TEXT("BloomTintAndThreshold"));
+		BloomMaxBrightnessParameter.Bind(ParameterMap,TEXT("BloomMaxBrightness"));
 		DistanceFadeParameter.Bind(ParameterMap, TEXT("DistanceFade"));
 		SourceTextureParameter.Bind(ParameterMap, TEXT("SourceTexture"));
 		SourceTextureSamplerParameter.Bind(ParameterMap, TEXT("SourceTextureSampler"));
@@ -138,6 +139,7 @@ public:
 		Ar << Parameters.AspectRatioAndInvAspectRatioParameter;
 		Ar << Parameters.LightShaftParameters;
 		Ar << Parameters.BloomTintAndThresholdParameter;
+		Ar << Parameters.BloomMaxBrightnessParameter;
 		Ar << Parameters.DistanceFadeParameter;
 		Ar << Parameters.SourceTextureParameter;
 		Ar << Parameters.SourceTextureSamplerParameter;
@@ -220,6 +222,7 @@ public:
 
 		const FLinearColor BloomTint = LightSceneInfo->BloomTint;
 		SetShaderValue(RHICmdList, Shader, BloomTintAndThresholdParameter, FVector4(BloomTint.R, BloomTint.G, BloomTint.B, LightSceneInfo->BloomThreshold));
+		SetShaderValue(RHICmdList, Shader, BloomMaxBrightnessParameter, LightSceneInfo->BloomMaxBrightness);
 
 		float OcclusionMaskDarkness;
 		float OcclusionDepthRange;
@@ -258,6 +261,7 @@ private:
 	FShaderParameter AspectRatioAndInvAspectRatioParameter;
 	FShaderParameter LightShaftParameters;
 	FShaderParameter BloomTintAndThresholdParameter;
+	FShaderParameter BloomMaxBrightnessParameter;
 	FShaderParameter DistanceFadeParameter;
 	FShaderResourceParameter SourceTextureParameter;
 	FShaderResourceParameter SourceTextureSamplerParameter;
