Fix constant curve keys in baked additive animations

diff --git a/Engine/Source/Runtime/Engine/Private/Animation/AnimSequence.cpp b/Engine/Source/Runtime/Engine/Private/Animation/AnimSequence.cpp
--- a/Engine/Source/Runtime/Engine/Private/Animation/AnimSequence.cpp
+++ b/Engine/Source/Runtime/Engine/Private/Animation/AnimSequence.cpp
@@ -2473,6 +2473,7 @@ void UAnimSequence::BakeOutAdditiveIntoRawData(TArray<FRawAnimSequenceTrack>& Ne
 		DummyBaseCurve.InitFrom(EvalContext.RequiredBones);
 
 		//Grab pose for this frame
+		const float PreviousFrameTime = (Frame - 1) * EvalContext.IntervalTime;
 		const float CurrentFrameTime = Frame * EvalContext.IntervalTime;
 		ExtractContext.CurrentTime = CurrentFrameTime;
 		GetAnimationPose(Pose, Curve, ExtractContext);
@@ -2521,6 +2522,15 @@ void UAnimSequence::BakeOutAdditiveIntoRawData(TArray<FRawAnimSequenceTrack>& Ne
 					if (!bHasKeys || IsNewKeyDifferent(RawCurve->FloatCurve.GetLastKey(), CurveEL.Value))
 					{
 						RawCurve->UpdateOrAddKey(CurveEL.Value, CurrentFrameTime);
+						TArray<FRichCurveKey>& CurveKeys = RawCurve->FloatCurve.Keys;
+						if (CurveKeys.Num() > 1)
+						{
+							FRichCurveKey& PrevKey = CurveKeys.Last(1);
+							if (PrevKey.Time < (PreviousFrameTime - SMALL_NUMBER)) // Did we skip a frame, if so need to make previous key const
+							{
+								PrevKey.InterpMode = RCIM_Constant;
+							}
+						}
 					}
 				}
 			}
