Implemented multiple materials for trimeshes in Chaos, fixed material destruction crash and made implicits use u16 for material index

diff --git a/Engine/Source/Runtime/Core/Public/UObject/ExternalPhysicsCustomObjectVersion.h b/Engine/Source/Runtime/Core/Public/UObject/ExternalPhysicsCustomObjectVersion.h
--- a/Engine/Source/Runtime/Core/Public/UObject/ExternalPhysicsCustomObjectVersion.h
+++ b/Engine/Source/Runtime/Core/Public/UObject/ExternalPhysicsCustomObjectVersion.h
@@ -65,6 +65,9 @@ struct CORE_API FExternalPhysicsCustomObjectVersion
 
 		// Added material manager to Chaos
 		AddedMaterialManager,
+
+		// Added material indices to trimesh collision data
+		AddTrimeshMaterialIndices,
 		
 		// -----<new versions can be added above this line>-------------------------------------------------
 		VersionPlusOne,
diff --git a/Engine/Source/Runtime/Engine/Private/PhysicsEngine/Experimental/ChaosDerivedData.cpp b/Engine/Source/Runtime/Engine/Private/PhysicsEngine/Experimental/ChaosDerivedData.cpp
--- a/Engine/Source/Runtime/Engine/Private/PhysicsEngine/Experimental/ChaosDerivedData.cpp
+++ b/Engine/Source/Runtime/Engine/Private/PhysicsEngine/Experimental/ChaosDerivedData.cpp
@@ -21,7 +21,7 @@ const TCHAR* FChaosDerivedDataCooker::GetPluginName() const
 
 const TCHAR* FChaosDerivedDataCooker::GetVersionString() const
 {
-	return TEXT("64EABBC88E0D47A98B01CD8BA0042E13");
+	return TEXT("9F79F10AF0354268BB194AA050225CF4");
 }
 
 FString FChaosDerivedDataCooker::GetPluginSpecificCacheKeySuffix() const
@@ -105,7 +105,8 @@ void FChaosDerivedDataCooker::BuildTriangleMeshes(TArray<TUniquePtr<Chaos::TTria
 		FinalIndices.Add(Tri.v2);
 	}
 
-	Chaos::CleanTrimesh(FinalVerts, FinalIndices, nullptr);
+	TArray<int32> Remap;
+	Chaos::CleanTrimesh(FinalVerts, FinalIndices, &Remap);
 
 	// Build particle list #BG Maybe allow TParticles to copy vectors?
 	Chaos::TParticles<Precision, 3> TriMeshParticles;
@@ -119,9 +120,16 @@ void FChaosDerivedDataCooker::BuildTriangleMeshes(TArray<TUniquePtr<Chaos::TTria
 
 	// Build chaos triangle list. #BGTODO Just make the clean function take these types instead of double copying
 	const int32 NumTriangles = FinalIndices.Num() / 3;
+	bool bHasMaterials = InParams.TriangleMeshDesc.MaterialIndices.Num() > 0;
 	TArray<Chaos::TVector<int32, 3>> Triangles;
+	TArray<uint16> MaterialIndices;
 	Triangles.Reserve(NumTriangles);
 
+	if(bHasMaterials)
+	{
+		MaterialIndices.Reserve(NumTriangles);
+	}
+
 	for(int32 TriangleIndex = 0; TriangleIndex < NumTriangles; ++TriangleIndex)
 	{
 		// Only add this triangle if it is valid
@@ -136,10 +144,33 @@ void FChaosDerivedDataCooker::BuildTriangleMeshes(TArray<TUniquePtr<Chaos::TTria
 		if (bIsValidTriangle)
 		{
 			Triangles.Add(Chaos::TVector<int32, 3>(FinalIndices[BaseIndex], FinalIndices[BaseIndex + 1], FinalIndices[BaseIndex + 2]));
+
+			if(bHasMaterials)
+			{
+				if(!ensure(Remap.IsValidIndex(TriangleIndex)))
+				{
+					MaterialIndices.Empty();
+					bHasMaterials = false;
+				}
+				else
+				{
+					const int32 OriginalIndex = Remap[TriangleIndex];
+
+					if(ensure(InParams.TriangleMeshDesc.MaterialIndices.IsValidIndex(OriginalIndex)))
+					{
+						MaterialIndices.Add(InParams.TriangleMeshDesc.MaterialIndices[OriginalIndex]);
+					}
+					else
+					{
+						MaterialIndices.Empty();
+						bHasMaterials = false;
+					}
+				}
+			}
 		}
 	}
 
-	OutTriangleMeshes.Emplace(new Chaos::TTriangleMeshImplicitObject<Precision>(MoveTemp(TriMeshParticles), MoveTemp(Triangles)));
+	OutTriangleMeshes.Emplace(new Chaos::TTriangleMeshImplicitObject<Precision>(MoveTemp(TriMeshParticles), MoveTemp(Triangles), MoveTemp(MaterialIndices)));
 }
 
 template void FChaosDerivedDataCooker::BuildTriangleMeshes(TArray<TUniquePtr<Chaos::TTriangleMeshImplicitObject<float>>>& OutTriangleMeshes, const FCookBodySetupInfo& InParams);
diff --git a/Engine/Source/Runtime/Engine/Private/PhysicsEngine/Experimental/ChaosDerivedDataUtil.cpp b/Engine/Source/Runtime/Engine/Private/PhysicsEngine/Experimental/ChaosDerivedDataUtil.cpp
--- a/Engine/Source/Runtime/Engine/Private/PhysicsEngine/Experimental/ChaosDerivedDataUtil.cpp
+++ b/Engine/Source/Runtime/Engine/Private/PhysicsEngine/Experimental/ChaosDerivedDataUtil.cpp
@@ -30,12 +30,9 @@ namespace Chaos
 		TArray<int32> LocalUniqueToSourceIndices;
 		// Remapping table from source index to unique index
 		TArray<int32> LocalVertexRemap;
-		// Remapping table from source triangle to unique triangle
-		TArray<int32> LocalTriangleRemap;
-
+		
 		LocalUniqueVerts.Reserve(NumSourceVerts);
 		LocalVertexRemap.AddUninitialized(NumSourceVerts);
-		LocalTriangleRemap.AddUninitialized(NumSourceTriangles);
 
 		auto ValidateTrianglesPre = [&InOutVertices](int32 A, int32 B, int32 C) -> bool
 		{
@@ -112,6 +109,10 @@ namespace Chaos
 		int32 NumDiscardedTriangles_Welded = 0;
 		int32 NumDiscardedTriangles_Area = 0;
 
+		// Remapping table from source triangle to unique triangle
+		TArray<int32> LocalTriangleRemap;
+		LocalTriangleRemap.Reserve(NumSourceTriangles);
+
 		for(int32 OriginalTriIndex = 0; OriginalTriIndex < NumSourceTriangles; ++OriginalTriIndex)
 		{
 			const int32 OrigAIndex = LocalSourceIndices[OriginalTriIndex * 3];
diff --git a/Engine/Source/Runtime/Experimental/Chaos/Private/Chaos/TriangleMeshImplicitObject.cpp b/Engine/Source/Runtime/Experimental/Chaos/Private/Chaos/TriangleMeshImplicitObject.cpp
--- a/Engine/Source/Runtime/Experimental/Chaos/Private/Chaos/TriangleMeshImplicitObject.cpp
+++ b/Engine/Source/Runtime/Experimental/Chaos/Private/Chaos/TriangleMeshImplicitObject.cpp
@@ -9,11 +9,12 @@
 namespace Chaos
 {
 template <typename T>
-TTriangleMeshImplicitObject<T>::TTriangleMeshImplicitObject(TParticles<T, 3>&& Particles, TArray<TVector<int32, 3>>&& Elements)
+TTriangleMeshImplicitObject<T>::TTriangleMeshImplicitObject(TParticles<T, 3>&& Particles, TArray<TVector<int32, 3>>&& Elements, TArray<uint16>&& InMaterialIndices)
 	: FImplicitObject(EImplicitObject::HasBoundingBox, ImplicitObjectType::TriangleMesh)
 	, MParticles(MoveTemp(Particles))
 	, MElements(MoveTemp(Elements))
 	, MLocalBoundingBox(MParticles.X(0), MParticles.X(0))
+    , MaterialIndices(MoveTemp(InMaterialIndices))
 {
 	for (uint32 Idx = 1; Idx < MParticles.Size(); ++Idx)
 	{
diff --git a/Engine/Source/Runtime/Experimental/Chaos/Public/Chaos/HeightField.h b/Engine/Source/Runtime/Experimental/Chaos/Public/Chaos/HeightField.h
--- a/Engine/Source/Runtime/Experimental/Chaos/Public/Chaos/HeightField.h
+++ b/Engine/Source/Runtime/Experimental/Chaos/Public/Chaos/HeightField.h
@@ -68,7 +68,7 @@ namespace Chaos
 		virtual int32 FindMostOpposingFace(const TVector<T, 3>& Position, const TVector<T, 3>& UnitDir, int32 HintFaceIndex, T SearchDist) const override;
 		virtual TVector<T, 3> FindGeometryOpposingNormal(const TVector<T, 3>& DenormDir, int32 FaceIndex, const TVector<T, 3>& OriginalNormal) const override;
 
-		virtual uint8 GetMaterialIndex(uint32 HintIndex) const override
+		virtual uint16 GetMaterialIndex(uint32 HintIndex) const override
 		{
 			ensure(GeomData.MaterialIndices.Num() > 0);
 
diff --git a/Engine/Source/Runtime/Experimental/Chaos/Public/Chaos/ImplicitObject.h b/Engine/Source/Runtime/Experimental/Chaos/Public/Chaos/ImplicitObject.h
--- a/Engine/Source/Runtime/Experimental/Chaos/Public/Chaos/ImplicitObject.h
+++ b/Engine/Source/Runtime/Experimental/Chaos/Public/Chaos/ImplicitObject.h
@@ -272,7 +272,7 @@ public:
 
 	static const FName GetTypeName(const EImplicitObjectType InType);
 
-	virtual uint8 GetMaterialIndex(uint32 HintIndex) const { return 0; }
+	virtual uint16 GetMaterialIndex(uint32 HintIndex) const { return 0; }
 
 protected:
 	EImplicitObjectType Type;
diff --git a/Engine/Source/Runtime/Experimental/Chaos/Public/Chaos/ImplicitObjectScaled.h b/Engine/Source/Runtime/Experimental/Chaos/Public/Chaos/ImplicitObjectScaled.h
--- a/Engine/Source/Runtime/Experimental/Chaos/Public/Chaos/ImplicitObjectScaled.h
+++ b/Engine/Source/Runtime/Experimental/Chaos/Public/Chaos/ImplicitObjectScaled.h
@@ -469,6 +469,11 @@ public:
 		return HashCombine(MObject->GetTypeHash(), ::GetTypeHash(MScale));
 	}
 
+	virtual uint16 GetMaterialIndex(uint32 HintIndex) const override
+	{
+		return MObject->GetMaterialIndex(HintIndex);
+	}
+
 #if 0
 	virtual TUniquePtr<FImplicitObject> Copy() const override
 	{
@@ -526,6 +531,7 @@ private:
 		TImplicitObjectScaled<QueryGeomType> ScaledB(B.Object(), InvScale * B.GetScale());
 		return ScaledB;
 	}
+
 };
 
 template <typename TConcrete>
diff --git a/Engine/Source/Runtime/Experimental/Chaos/Public/Chaos/ImplicitObjectTransformed.h b/Engine/Source/Runtime/Experimental/Chaos/Public/Chaos/ImplicitObjectTransformed.h
--- a/Engine/Source/Runtime/Experimental/Chaos/Public/Chaos/ImplicitObjectTransformed.h
+++ b/Engine/Source/Runtime/Experimental/Chaos/Public/Chaos/ImplicitObjectTransformed.h
@@ -204,7 +204,7 @@ public:
 		return HashCombine(MObject->GetTypeHash(), ::GetTypeHash(MTransform));
 	}
 
-	virtual uint8 GetMaterialIndex(uint32 HintIndex) const override
+	virtual uint16 GetMaterialIndex(uint32 HintIndex) const override
 	{
 		return MObject->GetMaterialIndex(HintIndex);
 	}
diff --git a/Engine/Source/Runtime/Experimental/Chaos/Public/Chaos/TriangleMeshImplicitObject.h b/Engine/Source/Runtime/Experimental/Chaos/Public/Chaos/TriangleMeshImplicitObject.h
--- a/Engine/Source/Runtime/Experimental/Chaos/Public/Chaos/TriangleMeshImplicitObject.h
+++ b/Engine/Source/Runtime/Experimental/Chaos/Public/Chaos/TriangleMeshImplicitObject.h
@@ -29,7 +29,7 @@ namespace Chaos
 	public:
 		using FImplicitObject::GetTypeName;
 
-		TTriangleMeshImplicitObject(TParticles<T,3>&& Particles, TArray<TVector<int32, 3>>&& Elements);
+		TTriangleMeshImplicitObject(TParticles<T,3>&& Particles, TArray<TVector<int32, 3>>&& Elements, TArray<uint16>&& InMaterialIndices);
 		TTriangleMeshImplicitObject(const TTriangleMeshImplicitObject& Other) = delete;
 		TTriangleMeshImplicitObject(TTriangleMeshImplicitObject&& Other) = default;
 		virtual ~TTriangleMeshImplicitObject() {}
@@ -106,6 +106,11 @@ namespace Chaos
 				// Serialize acceleration
 				Ar << BVH;
 			}
+
+			if(Ar.CustomVer(FExternalPhysicsCustomObjectVersion::GUID) >= FExternalPhysicsCustomObjectVersion::AddTrimeshMaterialIndices)
+			{
+				Ar << MaterialIndices;
+			}
 		}
 
 		virtual void Serialize(FChaosArchive& Ar) override
@@ -130,6 +135,17 @@ namespace Chaos
 
 		TVector<T, 3> GetFaceNormal(const int32 FaceIdx) const;
 
+		virtual uint16 GetMaterialIndex(uint32 HintIndex) const override
+		{
+			if (MaterialIndices.IsValidIndex(HintIndex))
+			{
+				return MaterialIndices[HintIndex];
+			}
+
+			// 0 should always be the default material for a shape
+			return 0;
+		}
+
 	private:
 
 		void RebuildBV();
@@ -137,6 +153,7 @@ namespace Chaos
 		TParticles<T, 3> MParticles;
 		TArray<TVector<int32, 3>> MElements;
 		TBox<T, 3> MLocalBoundingBox;
+		TArray<uint16> MaterialIndices;
 
 		//using BVHType = TBoundingVolume<int32, T, 3>;
 		using BVHType = TAABBTree<int32, TAABBTreeLeafArray<int32, T>, T>;
diff --git a/Engine/Source/Runtime/PhysicsCore/Public/PhysXToChaosUtil.h b/Engine/Source/Runtime/PhysicsCore/Public/PhysXToChaosUtil.h
--- a/Engine/Source/Runtime/PhysicsCore/Public/PhysXToChaosUtil.h
+++ b/Engine/Source/Runtime/PhysicsCore/Public/PhysXToChaosUtil.h
@@ -147,8 +147,13 @@ inline TUniquePtr<Chaos::TImplicitObjectTransformed<float, 3>> PxShapeToChaosGeo
 		const uint32 NumTriangles = TriangleMesh->getNbTriangles();
 		PxTriangleMeshFlags Flags = TriangleMesh->getTriangleMeshFlags();
 		TArray<TVector<int32, 3>> Triangles;
+		TArray<uint16> MaterialIndices;
 		Triangles.AddUninitialized(NumTriangles);
-		for (uint32 TriIdx = 0; TriIdx < NumTriangles; ++TriIdx)
+		MaterialIndices.Reserve(NumTriangles);
+
+		const bool bHasMaterials = NumTriangles > 0 && TriangleMesh->getTriangleMaterialIndex(0) != TNumericLimits<uint16>::Max();
+
+		for(uint32 TriIdx = 0; TriIdx < NumTriangles; ++TriIdx)
 		{
 			if (Flags & PxTriangleMeshFlag::e16_BIT_INDICES)
 			{
@@ -166,9 +171,14 @@ inline TUniquePtr<Chaos::TImplicitObjectTransformed<float, 3>> PxShapeToChaosGeo
 					static_cast<const int32*>(IndexBuffer)[TriIdx * 3+2] };
 				Triangles[TriIdx] = Triangle;
 			}
+
+			if(bHasMaterials)
+			{
+				MaterialIndices.Add(TriangleMesh->getTriangleMaterialIndex(TriIdx));
+			}
 		}
 
-		TUniquePtr<TTriangleMeshImplicitObject<float>> TriMeshObj = MakeUnique<TTriangleMeshImplicitObject<float>>(MoveTemp(Particles), MoveTemp(Triangles));
+		TUniquePtr<TTriangleMeshImplicitObject<float>> TriMeshObj = MakeUnique<TTriangleMeshImplicitObject<float>>(MoveTemp(Particles), MoveTemp(Triangles), MoveTemp(MaterialIndices));
 		if (TriMeshGeom.scale.isIdentity())
 		{
 			InnerObj = MoveTemp(TriMeshObj);
