Hook up overlap MTD to engine

diff --git a/Engine/Source/Runtime/Engine/Private/PhysicsEngine/Experimental/PhysInterface_Chaos.cpp b/Engine/Source/Runtime/Engine/Private/PhysicsEngine/Experimental/PhysInterface_Chaos.cpp
--- a/Engine/Source/Runtime/Engine/Private/PhysicsEngine/Experimental/PhysInterface_Chaos.cpp
+++ b/Engine/Source/Runtime/Engine/Private/PhysicsEngine/Experimental/PhysInterface_Chaos.cpp
@@ -1787,32 +1787,17 @@ bool Overlap_GeomInternal(const FBodyInstance* InInstance, const Chaos::FImplici
 
 		if (TargetInstance->IsShapeBoundToBody(ShapeRef))
 		{
-			//Chaos::TVector<float,3> POutDirection;
-			//float OutDistance;
-
 			if (OutOptResult)
 			{
-				ensure(false);	//todo: implement in Chaos
-				/*
-				PxTransform PTransform = U2PTransform(FPhysicsInterface::GetTransform(ShapeRef));
-				if (PxGeometryQuery::computePenetration(POutDirection, OutDistance, InPxGeom, ShapePose, PShape->getGeometry().any(), PTransform))
+				Chaos::FMTDInfo MTDInfo;
+				if (Chaos::CastHelper(InGeom, [&](const auto& Downcast) { return Chaos::OverlapQuery(*Shape->Geometry, ActorTM, Downcast, GeomTransform, /*Thickness=*/0, &MTDInfo); }))
 				{
-					//TODO: there are some edge cases that give us nan results. In these cases we skip
-					if (!POutDirection.isFinite())
-					{
-						POutDirection.x = 0.f;
-						POutDirection.y = 0.f;
-						POutDirection.z = 0.f;
-					}
-
-					OutOptResult->Direction = P2UVector(POutDirection);
-					OutOptResult->Distance = FMath::Abs(OutDistance);
-
-					return true;
-				}*/
-				OutOptResult->Distance = 0;
+					OutOptResult->Distance = MTDInfo.Penetration;
+					OutOptResult->Direction = MTDInfo.Normal;
+					return true;	//question: should we take most shallow penetration?
+				}
 			}
-			//else	//todo: don't bother with this once MTD is implemented
+			else	//question: why do we even allow user to not pass in MTD info?
 			{
 				if (Chaos::CastHelper(InGeom, [&](const auto& Downcast) { return Chaos::OverlapQuery(*Shape->Geometry, ActorTM, Downcast, GeomTransform); }))
 				{
diff --git a/Engine/Source/Runtime/Experimental/Chaos/Public/Chaos/GeometryQueries.h b/Engine/Source/Runtime/Experimental/Chaos/Public/Chaos/GeometryQueries.h
--- a/Engine/Source/Runtime/Experimental/Chaos/Public/Chaos/GeometryQueries.h
+++ b/Engine/Source/Runtime/Experimental/Chaos/Public/Chaos/GeometryQueries.h
@@ -38,8 +38,14 @@ namespace Chaos
 		return Func(Geom.template GetObjectChecked<TSphere<FReal, 3>>());	//needed for return type
 	}
 
+	struct FMTDInfo
+	{
+		FVec3 Normal;
+		float Penetration;
+	};
+
 	template <typename QueryGeometry>
-	bool OverlapQuery(const FImplicitObject& A, const FRigidTransform3& ATM, const QueryGeometry& B, const FRigidTransform3& BTM, const FReal Thickness = 0, FVec3* OutMTD=nullptr)
+	bool OverlapQuery(const FImplicitObject& A, const FRigidTransform3& ATM, const QueryGeometry& B, const FRigidTransform3& BTM, const FReal Thickness = 0, FMTDInfo* OutMTD=nullptr)
 	{
 		const EImplicitObjectType AType = A.GetType(true);
 		constexpr EImplicitObjectType BType = QueryGeometry::StaticType();
@@ -64,7 +70,16 @@ namespace Chaos
 		else if (A.IsConvex())
 		{
 			const FVec3 Offset = ATM.GetLocation() - BTM.GetLocation();
-			return CastHelper(A, [&](const auto& AConcrete) { return GJKIntersection<FReal>(AConcrete, B, BToATM, Thickness, Offset.SizeSquared() < 1e-4 ? FVec3(1, 0, 0) : Offset); });
+			if (OutMTD)
+			{
+				FVec3 ClosestA;
+				FVec3 ClosestB;
+				return CastHelper(A, [&](const auto& AConcrete) { return GJKPenetration<FReal>(AConcrete, B, BToATM, OutMTD->Penetration, ClosestA, ClosestB, OutMTD->Normal, Thickness, Offset.SizeSquared() < 1e-4 ? FVec3(1, 0, 0) : Offset); });
+			}
+			else
+			{
+				return CastHelper(A, [&](const auto& AConcrete) { return GJKIntersection<FReal>(AConcrete, B, BToATM, Thickness, Offset.SizeSquared() < 1e-4 ? FVec3(1, 0, 0) : Offset); });
+			}
 		}
 		else
 		{
