Adding in event hook for when the Player Camera Manager has completed a Blend.

diff --git a/Engine/Source/Runtime/Engine/Classes/Camera/PlayerCameraManager.h b/Engine/Source/Runtime/Engine/Classes/Camera/PlayerCameraManager.h
--- a/Engine/Source/Runtime/Engine/Classes/Camera/PlayerCameraManager.h
+++ b/Engine/Source/Runtime/Engine/Classes/Camera/PlayerCameraManager.h
@@ -257,6 +257,12 @@ public:
 	/** Current view target transition blend parameters. */
 	struct FViewTargetTransitionParams BlendParams;
 
+public:
+	/** Fires when ViewTarget is set to PendingViewTarget */
+	DECLARE_EVENT(APlayerCameraManager, FOnBlendComplete)
+	FOnBlendComplete& OnBlendComplete() const { return OnBlendCompleteEvent; }
+private:
+	mutable FOnBlendComplete OnBlendCompleteEvent;
 private:
 	/** Cached camera properties. */
 	UPROPERTY(transient)
diff --git a/Engine/Source/Runtime/Engine/Private/PlayerCameraManager.cpp b/Engine/Source/Runtime/Engine/Private/PlayerCameraManager.cpp
--- a/Engine/Source/Runtime/Engine/Private/PlayerCameraManager.cpp
+++ b/Engine/Source/Runtime/Engine/Private/PlayerCameraManager.cpp
@@ -1005,6 +1005,8 @@ void APlayerCameraManager::DoUpdateCamera(float DeltaTime)
 
 			// our camera is now viewing there
 			NewPOV = PendingViewTarget.POV;
+
+			OnBlendComplete().Broadcast();
 		}
 	}
 
