Integrating CL# 10254055 from Dev-Enterprise
Fixed Lightmap UV channel not set to correct index after import

diff --git a/Engine/Source/Runtime/Engine/Private/StaticMesh.cpp b/Engine/Source/Runtime/Engine/Private/StaticMesh.cpp
--- a/Engine/Source/Runtime/Engine/Private/StaticMesh.cpp
+++ b/Engine/Source/Runtime/Engine/Private/StaticMesh.cpp
@@ -5958,7 +5958,7 @@ void UStaticMesh::EnforceLightmapRestrictions(bool bUseRenderData)
 
 					if (SourceModel.BuildSettings.bGenerateLightmapUVs)
 					{
-						NumChannelsInLOD = FMath::Max(NumChannelsInLOD, SourceModel.BuildSettings.DstLightmapIndex);
+						NumChannelsInLOD = FMath::Max(NumChannelsInLOD, SourceModel.BuildSettings.DstLightmapIndex + 1);
 					}
 
 					NumUVs = FMath::Min(NumChannelsInLOD, NumUVs);
