Fix crash trying to use stale cached values when the anim instance has been cleared

diff --git a/Engine/Source/Runtime/Engine/Classes/Components/SkeletalMeshComponent.h b/Engine/Source/Runtime/Engine/Classes/Components/SkeletalMeshComponent.h
--- a/Engine/Source/Runtime/Engine/Classes/Components/SkeletalMeshComponent.h
+++ b/Engine/Source/Runtime/Engine/Classes/Components/SkeletalMeshComponent.h
@@ -2028,6 +2028,9 @@ public:
 	bool IsAnimBlueprintInstanced() const;
 	void ClearAnimScriptInstance();
 
+	/** Clear cached animation data generated for URO during evaluation */
+	void ClearCachedAnimProperties();
+
 protected:
 	bool NeedToSpawnAnimScriptInstance() const;
 	bool NeedToSpawnPostPhysicsInstance(bool bForceReinit) const;
diff --git a/Engine/Source/Runtime/Engine/Private/Components/SkeletalMeshComponent.cpp b/Engine/Source/Runtime/Engine/Private/Components/SkeletalMeshComponent.cpp
--- a/Engine/Source/Runtime/Engine/Private/Components/SkeletalMeshComponent.cpp
+++ b/Engine/Source/Runtime/Engine/Private/Components/SkeletalMeshComponent.cpp
@@ -806,8 +806,15 @@ void USkeletalMeshComponent::ClearAnimScriptInstance()
 	}
 	AnimScriptInstance = nullptr;
 	ResetLinkedAnimInstances();
+	ClearCachedAnimProperties();S
 }
 
+void USkeletalMeshComponent::ClearCachedAnimProperties()
+{
+	CachedBoneSpaceTransforms.Empty();
+	CachedComponentSpaceTransforms.Empty();
+	CachedCurve.Empty();
+}
 
 void USkeletalMeshComponent::InitializeComponent()
 {
@@ -1712,9 +1719,7 @@ void USkeletalMeshComponent::RecalcRequiredBones(int32 LODIndex)
 	bRequiredBonesUpToDate = true;
 
 	// Invalidate cached bones.
-	CachedBoneSpaceTransforms.Empty();
-	CachedComponentSpaceTransforms.Empty();
-	CachedCurve.Empty();
+	ClearCachedAnimProperties();
 }
 
 void USkeletalMeshComponent::MarkRequiredCurveUpToDate()
