Fix typo

diff --git a/Engine/Source/Runtime/Engine/Private/Components/SkeletalMeshComponent.cpp b/Engine/Source/Runtime/Engine/Private/Components/SkeletalMeshComponent.cpp
--- a/Engine/Source/Runtime/Engine/Private/Components/SkeletalMeshComponent.cpp
+++ b/Engine/Source/Runtime/Engine/Private/Components/SkeletalMeshComponent.cpp
@@ -806,7 +806,7 @@ void USkeletalMeshComponent::ClearAnimScriptInstance()
 	}
 	AnimScriptInstance = nullptr;
 	ResetLinkedAnimInstances();
-	ClearCachedAnimProperties();S
+	ClearCachedAnimProperties();
 }
 
 void USkeletalMeshComponent::ClearCachedAnimProperties()
