Chaos - add hit flags to queries to geom sweep

diff --git a/Engine/Source/Runtime/Engine/Private/PhysicsEngine/Experimental/PhysInterface_Chaos.cpp b/Engine/Source/Runtime/Engine/Private/PhysicsEngine/Experimental/PhysInterface_Chaos.cpp
--- a/Engine/Source/Runtime/Engine/Private/PhysicsEngine/Experimental/PhysInterface_Chaos.cpp
+++ b/Engine/Source/Runtime/Engine/Private/PhysicsEngine/Experimental/PhysInterface_Chaos.cpp
@@ -1648,6 +1648,7 @@ bool FPhysInterface_Chaos::LineTrace_Geom(FHitResult& OutHit, const FBodyInstanc
 					{
 						BestHit.WorldNormal = WorldTM.TransformVectorNoScale(BestHit.WorldNormal);
 						BestHit.WorldPosition = WorldTM.TransformPositionNoScale(BestHit.WorldPosition);
+						SetFlags(BestHit, EHitFlags::Distance | EHitFlags::Normal | EHitFlags::Position);
 
 						// we just like to make sure if the hit is made, set to test touch
 						FCollisionFilterData QueryFilter;
@@ -1747,6 +1748,7 @@ bool FPhysInterface_Chaos::Sweep_Geom(FHitResult& OutHit, const FBodyInstance* I
 								{
 									Hit.FaceIndex = FindFaceIndex(Hit, Dir);
 								}
+								SetFlags(Hit, EHitFlags::Distance | EHitFlags::Normal | EHitFlags::Position | EHitFlags::FaceIndex);
 
 								FTransform StartTransform(InStart);
 								ConvertQueryImpactHit(OwnerComponentInst->GetWorld(), Hit, OutHit, DeltaMag, QueryFilter, InStart, InEnd, nullptr, StartTransform, false, false);
