Fix build

diff --git a/Engine/Source/Runtime/Engine/Private/PhysicsEngine/ImmediatePhysics/ImmediatePhysicsChaos/ImmediatePhysicsSimulation_Chaos.cpp b/Engine/Source/Runtime/Engine/Private/PhysicsEngine/ImmediatePhysics/ImmediatePhysicsChaos/ImmediatePhysicsSimulation_Chaos.cpp
--- a/Engine/Source/Runtime/Engine/Private/PhysicsEngine/ImmediatePhysics/ImmediatePhysicsChaos/ImmediatePhysicsSimulation_Chaos.cpp
+++ b/Engine/Source/Runtime/Engine/Private/PhysicsEngine/ImmediatePhysics/ImmediatePhysicsChaos/ImmediatePhysicsSimulation_Chaos.cpp
@@ -89,7 +89,11 @@ int32 ChaosImmediate_DebugDrawCollisions = 0;
 int32 ChaosImmediate_DebugDrawJoints = 0;
 #endif
 
+#if CHAOS_DEBUG_DRAW
 int32 ChaosImmediate_DebugDrawJointFeatures = (int32)Chaos::DebugDraw::EDebugDrawJointFeature::Default;
+#else
+int32 ChaosImmediate_DebugDrawJointFeatures = 0;
+#endif
 
 FAutoConsoleVariableRef CVarChaosImmPhysDebugDrawParticles(TEXT("p.Chaos.ImmPhys.DebugDrawParticles"), ChaosImmediate_DebugDrawParticles, TEXT("Draw Particle Transforms (0 = never; 1 = end of frame; 2 = begin and end of frame; 3 = post-integate, post-apply and post-applypushout;)."));
 FAutoConsoleVariableRef CVarChaosImmPhysDebugDrawShapes(TEXT("p.Chaos.ImmPhys.DebugDrawShapes"), ChaosImmediate_DebugDrawShapes, TEXT("Draw Shapes (0 = never; 1 = end of frame; 2 = begin and end of frame; 3 = post-integate, post-apply and post-applypushout;"));
