Copied over change from Dev-Editor which prevents a crash when loading a map containing a runtime-generated static mesh into the editor.

diff --git a/Engine/Source/Runtime/Engine/Classes/Engine/StaticMesh.h b/Engine/Source/Runtime/Engine/Classes/Engine/StaticMesh.h
--- a/Engine/Source/Runtime/Engine/Classes/Engine/StaticMesh.h
+++ b/Engine/Source/Runtime/Engine/Classes/Engine/StaticMesh.h
@@ -661,6 +661,7 @@ class UStaticMesh : public UStreamableRenderAsset, public IInterface_CollisionDa
 	/**
 	 * If true, StaticMesh has been built at runtime
 	 */
+	UPROPERTY()
 	uint8 bIsBuiltAtRuntime : 1;
 
 protected:
