Add ability to mirror shader stats file to a remote location, and a setting for setting this by default in a project.

diff --git a/Engine/Source/Runtime/Engine/Private/ShaderCompiler/ShaderCompiler.cpp b/Engine/Source/Runtime/Engine/Private/ShaderCompiler/ShaderCompiler.cpp
--- a/Engine/Source/Runtime/Engine/Private/ShaderCompiler/ShaderCompiler.cpp
+++ b/Engine/Source/Runtime/Engine/Private/ShaderCompiler/ShaderCompiler.cpp
@@ -1520,7 +1520,7 @@ void FShaderCompilerStats::WriteStats()
 	{
 		FlushRenderingCommands(true);
 
-		FString FileName = FString::Printf(TEXT("%s/MaterialStats/Stats-%s.csv"), *FPaths::ProjectSavedDir(), *FDateTime::Now().ToString());
+		FString FileName = FPaths::Combine(*FPaths::ProjectSavedDir(), FString::Printf(TEXT("MaterialStats/Stats-%s.csv"),  *FDateTime::Now().ToString()));
 		auto DebugWriter = IFileManager::Get().CreateFileWriter(*FileName);
 		FDiagnosticTableWriterCSV StatWriter(DebugWriter);
 		const TSparseArray<ShaderCompilerStats>& PlatformStats = GetShaderCompilerStats();
@@ -1556,6 +1556,19 @@ void FShaderCompilerStats::WriteStats()
 				}
 			}
 		}
+		DebugWriter->Close();
+		FString MirrorLocation;
+		GConfig->GetString(TEXT("/Script/Engine.ShaderCompilerStats"), TEXT("MaterialStatsLocation"), MirrorLocation, GGameIni);
+		FParse::Value(FCommandLine::Get(), TEXT("MaterialStatsMirror="), MirrorLocation);
+
+		if (!MirrorLocation.IsEmpty())
+		{
+			FString TargetType = TEXT("Default");
+			FParse::Value(FCommandLine::Get(), TEXT("target="), TargetType);
+
+			FString CopyLocation = FPaths::Combine(*MirrorLocation, *FApp::GetBranchName(), FString::Printf(TEXT("Stats-Latest(%s).csv"), *TargetType));
+			IFileManager::Get().Copy(*CopyLocation, *FileName, true, true);
+		}
 	}
 	{
 
