Fixed highlighted line in DXC error message.

#rnx

diff --git a/Engine/Source/Runtime/Engine/Private/ShaderCompiler/ShaderCompiler.cpp b/Engine/Source/Runtime/Engine/Private/ShaderCompiler/ShaderCompiler.cpp
--- a/Engine/Source/Runtime/Engine/Private/ShaderCompiler/ShaderCompiler.cpp
+++ b/Engine/Source/Runtime/Engine/Private/ShaderCompiler/ShaderCompiler.cpp
@@ -849,20 +849,7 @@ static bool CheckSingleJob(FShaderCompileJob* SingleJob, const TArray<FMaterial*
 		for (int32 ErrorIndex = 0; ErrorIndex < SingleJob->Output.Errors.Num(); ErrorIndex++)
 		{
 			const FShaderCompilerError& InError = SingleJob->Output.Errors[ErrorIndex];
-			if (InError.HasLineMarker())
-			{
-				// Append highlighted line and its marker to the same error message with line terminators
-				// to get a similar multiline error output as with DXC
-				Errors.AddUnique(
-					InError.GetErrorString() + LINE_TERMINATOR +
-					TEXT("\t") + InError.HighlightedLine + LINE_TERMINATOR +
-					TEXT("\t") + InError.HighlightedLineMarker
-				);
-			}
-			else
-			{
-				Errors.AddUnique(InError.GetErrorString());
-			}
+			Errors.AddUnique(InError.GetErrorStringWithLineMarker());
 		}
 	}
 
@@ -4061,7 +4048,7 @@ FShader* FGlobalShaderTypeCompiler::FinishCompileShader(FGlobalShaderType* Shade
 			UE_LOG(LogShaderCompilers, Error, TEXT("Errors compiling global shader %s %s %s:\n"), CurrentJob.ShaderType->GetName(), ShaderPipelineType ? TEXT("ShaderPipeline") : TEXT(""), ShaderPipelineType ? ShaderPipelineType->GetName() : TEXT(""));
 			for (int32 ErrorIndex = 0; ErrorIndex < CurrentJob.Output.Errors.Num(); ErrorIndex++)
 			{
-				UE_LOG(LogShaderCompilers, Error, TEXT("	%s"), *CurrentJob.Output.Errors[ErrorIndex].GetErrorString());
+				UE_LOG(LogShaderCompilers, Error, TEXT("\t%s"), *CurrentJob.Output.Errors[ErrorIndex].GetErrorStringWithLineMarker());
 			}
 		}
 		else if (CVarShowShaderWarnings->GetInt())
@@ -4069,7 +4056,7 @@ FShader* FGlobalShaderTypeCompiler::FinishCompileShader(FGlobalShaderType* Shade
 			UE_LOG(LogShaderCompilers, Warning, TEXT("Warnings compiling global shader %s %s %s:\n"), CurrentJob.ShaderType->GetName(), ShaderPipelineType ? TEXT("ShaderPipeline") : TEXT(""), ShaderPipelineType ? ShaderPipelineType->GetName() : TEXT(""));
 			for (int32 ErrorIndex = 0; ErrorIndex < CurrentJob.Output.Errors.Num(); ErrorIndex++)
 			{
-				UE_LOG(LogShaderCompilers, Warning, TEXT("	%s"), *CurrentJob.Output.Errors[ErrorIndex].GetErrorString());
+				UE_LOG(LogShaderCompilers, Warning, TEXT("\t%s"), *CurrentJob.Output.Errors[ErrorIndex].GetErrorStringWithLineMarker());
 			}
 		}
 	}
diff --git a/Engine/Source/Runtime/RenderCore/Public/ShaderCore.h b/Engine/Source/Runtime/RenderCore/Public/ShaderCore.h
--- a/Engine/Source/Runtime/RenderCore/Public/ShaderCore.h
+++ b/Engine/Source/Runtime/RenderCore/Public/ShaderCore.h
@@ -774,6 +774,21 @@ struct FShaderCompilerError
 		}
 	}
 
+	/** Returns the error message with source file and source line (if present), as well as a line marker seperated with a LINE_TERMINATOR. */
+	FString GetErrorStringWithLineMarker() const
+	{
+		if (HasLineMarker())
+		{
+			// Append highlighted line and its marker to the same error message with line terminators
+			// to get a similar multiline error output as with DXC
+			return (GetErrorString() + LINE_TERMINATOR + TEXT("\t") + HighlightedLine + LINE_TERMINATOR + TEXT("\t") + HighlightedLineMarker);
+		}
+		else
+		{
+			return GetErrorString();
+		}
+	}
+
 	/**
 	Returns true if this error message has a marker string for the highlighted source line where the error occurred. Example:
 		/Engine/Private/MySourceFile.usf(120): error: undeclared identifier 'a'
@@ -790,7 +805,7 @@ struct FShaderCompilerError
 
 	friend FArchive& operator<<(FArchive& Ar,FShaderCompilerError& Error)
 	{
-		return Ar << Error.ErrorVirtualFilePath << Error.ErrorLineString << Error.StrippedErrorMessage;
+		return Ar << Error.ErrorVirtualFilePath << Error.ErrorLineString << Error.StrippedErrorMessage << Error.HighlightedLine << Error.HighlightedLineMarker;
 	}
 };
 
