Made StaticMesh's tri mesh data accessible in runtime (if the mesh data is flagged for CPU access, issue a warning otherwise).
This allows real-time collisions for spline meshes with 'Use Complex Collision as Simple'

#rnx

diff --git a/Engine/Source/Runtime/Engine/Private/StaticMesh.cpp b/Engine/Source/Runtime/Engine/Private/StaticMesh.cpp
--- a/Engine/Source/Runtime/Engine/Private/StaticMesh.cpp
+++ b/Engine/Source/Runtime/Engine/Private/StaticMesh.cpp
@@ -5603,6 +5603,18 @@ bool UStaticMesh::GetPhysicsTriMeshData(struct FTriMeshCollisionData* CollisionD
 	{
 		return ComplexCollisionMesh->GetPhysicsTriMeshData(CollisionData, bInUseAllTriData);
 	}
+#else // #if WITH_EDITORONLY_DATA
+	// the static mesh needs to be tagged for CPUAccess in order to access TriMeshData in runtime mode : 
+	if (!bAllowCPUAccess)
+	{
+		UE_LOG(LogStaticMesh, Warning, TEXT("UStaticMesh::GetPhysicsTriMeshData: Triangle data from '%s' cannot be accessed at runtime on a mesh that isn't flagged as Allow CPU Access. This asset needs to be flagged as such (in the Advanced section)."), *GetFullName());
+		return false;
+	}
+
+	// without editor data, we can't selectively generate a physics mesh for a given LOD index (we're missing access to GetSectionInfoMap()) so force bInUseAllTriData in order to use LOD index 0
+	bInUseAllTriData = true;
+#endif // #if !WITH_EDITORONLY_DATA
+
 	check(HasValidRenderData());
 
 	// Get the LOD level to use for collision
@@ -5610,6 +5622,7 @@ bool UStaticMesh::GetPhysicsTriMeshData(struct FTriMeshCollisionData* CollisionD
 	const int32 UseLODIndex = bInUseAllTriData ? 0 : FMath::Clamp(LODForCollision, 0, RenderData->LODResources.Num()-1);
 
 	FStaticMeshLODResources& LOD = RenderData->LODResources[UseLODIndex];
+
 	FIndexArrayView Indices = LOD.IndexBuffer.GetArrayView();
 
 	TMap<int32, int32> MeshToCollisionVertMap; // map of static mesh verts to collision verts
@@ -5626,7 +5639,12 @@ bool UStaticMesh::GetPhysicsTriMeshData(struct FTriMeshCollisionData* CollisionD
 	{
 		const FStaticMeshSection& Section = LOD.Sections[SectionIndex];
 
-		if (bInUseAllTriData || GetSectionInfoMap().Get(UseLODIndex,SectionIndex).bEnableCollision)
+#if WITH_EDITORONLY_DATA
+		// we can only use GetSectionInfoMap() in WITH_EDITORONLY_DATA mode, otherwise, assume bInUseAllTriData :
+		if (bInUseAllTriData || GetSectionInfoMap().Get(UseLODIndex, SectionIndex).bEnableCollision)
+#else // #if WITH_EDITORONLY_DATA
+		check(bInUseAllTriData && bAllowCPUAccess);
+#endif // #if !WITH_EDITORONLY_DATA
 		{
 			const uint32 OnePastLastIndex  = Section.FirstIndex + Section.NumTriangles*3;
 
@@ -5647,9 +5665,6 @@ bool UStaticMesh::GetPhysicsTriMeshData(struct FTriMeshCollisionData* CollisionD
 	// We only have a valid TriMesh if the CollisionData has vertices AND indices. For meshes with disabled section collision, it
 	// can happen that the indices will be empty, in which case we do not want to consider that as valid trimesh data
 	return CollisionData->Vertices.Num() > 0 && CollisionData->Indices.Num() > 0;
-#else // #if WITH_EDITORONLY_DATA
-	return false;
-#endif // #if WITH_EDITORONLY_DATA
 }
 
 bool UStaticMesh::ContainsPhysicsTriMeshData(bool bInUseAllTriData) const 
@@ -5659,6 +5674,11 @@ bool UStaticMesh::ContainsPhysicsTriMeshData(bool bInUseAllTriData) const
 	{
 		return ComplexCollisionMesh->ContainsPhysicsTriMeshData(bInUseAllTriData);
 	}
+#else // #if WITH_EDITORONLY_DATA
+	// without editor data, we can't selectively generate a physics mesh for a given LOD index (we're missing access to GetSectionInfoMap()) so force bInUseAllTriData in order to use LOD index 0
+	bInUseAllTriData = true;
+#endif // #if !WITH_EDITORONLY_DATA
+	
 	if(RenderData == nullptr || RenderData->LODResources.Num() == 0)
 	{
 		return false;
@@ -5675,16 +5695,18 @@ bool UStaticMesh::ContainsPhysicsTriMeshData(bool bInUseAllTriData) const
 		for (int32 SectionIndex = 0; SectionIndex < LOD.Sections.Num(); ++SectionIndex)
 		{
 			const FStaticMeshSection& Section = LOD.Sections[SectionIndex];
+#if WITH_EDITORONLY_DATA
+			// we can only use GetSectionInfoMap() in WITH_EDITORONLY_DATA mode, otherwise, assume bInUseAllTriData :
 			if ((bInUseAllTriData || GetSectionInfoMap().Get(UseLODIndex, SectionIndex).bEnableCollision) && Section.NumTriangles > 0)
 			{
 				return true;
 			}
+#else // #if WITH_EDITORONLY_DATA
+			return true;
+#endif // #if WITH_EDITORONLY_DATA
 		}
 	}
 	return false; 
-#else // #if WITH_EDITORONLY_DATA
-	return false;
-#endif // #if WITH_EDITORONLY_DATA
 }
 
 void UStaticMesh::GetMeshId(FString& OutMeshId)
