Use constexpr to calculate the TSequenceNumber's ShiftValue at compile time instead of runtime.

#rnx

diff --git a/Engine/Source/Runtime/Engine/Public/Net/Util/SequenceNumber.h b/Engine/Source/Runtime/Engine/Public/Net/Util/SequenceNumber.h
--- a/Engine/Source/Runtime/Engine/Public/Net/Util/SequenceNumber.h
+++ b/Engine/Source/Runtime/Engine/Public/Net/Util/SequenceNumber.h
@@ -61,7 +61,7 @@ private:
 template <SIZE_T NumBits, typename SequenceType>
 typename TSequenceNumber<NumBits, SequenceType>::DifferenceT TSequenceNumber<NumBits, SequenceType>::Diff(TSequenceNumber A, TSequenceNumber B) 
 { 
-	const SIZE_T ShiftValue = sizeof(DifferenceT)*8 - NumBits;
+	constexpr SIZE_T ShiftValue = sizeof(DifferenceT)*8 - NumBits;
 
 	const SequenceT ValueA = A.Value;
 	const SequenceT ValueB = B.Value;
