DF Atlas size cvars are no more read only and can be hotfixable.
When DF atlas runs out of space, all meshes that did not fit get added to FailedAllocation list. When a DF mesh is removed from the atlas, failed allocated meshes that would fit in the atlas will be moved to PendingAllocation list
Added csv stat to get number of pixels that did not fit in the atlas.

#rnx

diff --git a/Engine/Source/Runtime/Engine/Private/DistanceFieldAtlas.cpp b/Engine/Source/Runtime/Engine/Private/DistanceFieldAtlas.cpp
--- a/Engine/Source/Runtime/Engine/Private/DistanceFieldAtlas.cpp
+++ b/Engine/Source/Runtime/Engine/Private/DistanceFieldAtlas.cpp
@@ -78,14 +78,12 @@ static TAutoConsoleVariable<float> CVarDistFieldResScale(
 static TAutoConsoleVariable<int32> CVarDistFieldAtlasResXY(
 	TEXT("r.DistanceFields.AtlasSizeXY"),
 	512,	
-	TEXT("Max size of the global mesh distance field atlas volume texture in X and Y."),
-	ECVF_ReadOnly);
+	TEXT("Max size of the global mesh distance field atlas volume texture in X and Y."));
 
 static TAutoConsoleVariable<int32> CVarDistFieldAtlasResZ(
 	TEXT("r.DistanceFields.AtlasSizeZ"),
 	1024,	
-	TEXT("Max size of the global mesh distance field atlas volume texture in Z."),
-	ECVF_ReadOnly);
+	TEXT("Max size of the global mesh distance field atlas volume texture in Z."));
 
 int32 GDistanceFieldForceAtlasRealloc = 0;
  
@@ -129,6 +127,7 @@ FDistanceFieldVolumeTextureAtlas::FDistanceFieldVolumeTextureAtlas() :
 	BlockAllocator(0, 0, 0, 0, 0, 0, false, false),
 	bInitialized(false),
 	AllocatedPixels(0),
+	FailedAllocatedPixels(0),
 	MaxUsedAtlasX(0),
 	MaxUsedAtlasY(0),
 	MaxUsedAtlasZ(0)
@@ -272,13 +271,50 @@ void FDistanceFieldVolumeTextureAtlas::RemoveAllocation(FDistanceFieldVolumeText
 	InitializeIfNeeded();
 	PendingAllocations.Remove(Texture);
 
-	if (CurrentAllocations.Contains(Texture))
+	FIntVector Size = Texture->VolumeData.Size;
+	int PixelAreaSize = Size.X * Size.Y * Size.Z;
+
+	if (FailedAllocations.Remove(Texture) > 0)
 	{
-		const FIntVector Min = Texture->GetAllocationMin();
-		const FIntVector Size = Texture->VolumeData.Size;
-		verify(BlockAllocator.RemoveElement(Min.X, Min.Y, Min.Z, Size.X, Size.Y, Size.Z));
-		CurrentAllocations.Remove(Texture);
-		AllocatedPixels -= Size.X * Size.Y * Size.Z;
+		FailedAllocatedPixels -= PixelAreaSize;
+	}
+
+	if (!CurrentAllocations.Contains(Texture))
+	{
+		return;
+	}
+
+	const FIntVector Min = Texture->GetAllocationMin();
+	verify(BlockAllocator.RemoveElement(Min.X, Min.Y, Min.Z, Size.X, Size.Y, Size.Z));
+	CurrentAllocations.Remove(Texture);
+	AllocatedPixels -= PixelAreaSize;
+
+	FIntVector RemainingSize = Size;
+	
+	// Check if there is room for a previous failed allocation
+	for (int32 Index = 0; Index < FailedAllocations.Num(); Index++)
+	{
+		FDistanceFieldVolumeTexture* PreviouslyFailedAllocatedTexture = FailedAllocations[Index];
+		Size = PreviouslyFailedAllocatedTexture->VolumeData.Size;
+		if (Size.X > RemainingSize.X || Size.Y > RemainingSize.Y || Size.Z > RemainingSize.Z)
+		{
+			continue;
+		}
+		// Room available. Add texture to pending list
+		PendingAllocations.Add(PreviouslyFailedAllocatedTexture);
+		FailedAllocations.RemoveAt(Index);
+		Index--;
+		FailedAllocatedPixels -= PixelAreaSize;
+
+		RemainingSize.X -= Size.X;
+		RemainingSize.Y -= Size.Y;
+		RemainingSize.Z -= Size.Z;
+		
+		// Continue iterating if remaining size can support another mesh DF
+		if (RemainingSize.X < 4 || RemainingSize.Y < 4 || RemainingSize.Z < 4)
+		{
+			break;
+		}
 	}
 }
 
@@ -298,6 +334,19 @@ void FDistanceFieldVolumeTextureAtlas::UpdateAllocations()
 		CSV_CUSTOM_STAT_GLOBAL(DFAtlasMaxX, float(MaxUsedAtlasX), ECsvCustomStatOp::Set);
 		CSV_CUSTOM_STAT_GLOBAL(DFAtlasMaxY, float(MaxUsedAtlasY), ECsvCustomStatOp::Set);
 		CSV_CUSTOM_STAT_GLOBAL(DFAtlasMaxZ, float(MaxUsedAtlasZ), ECsvCustomStatOp::Set);
+		CSV_CUSTOM_STAT_GLOBAL(DFAtlasFailedAllocatedMagaPixels, (float(FailedAllocatedPixels)/1024)/1024, ECsvCustomStatOp::Set);
+	}
+
+	static const auto CVarXY = IConsoleManager::Get().FindTConsoleVariableDataInt(TEXT("r.DistanceFields.AtlasSizeXY"));
+	const int32 AtlasXY = CVarXY->GetValueOnAnyThread();
+
+	static const auto CVarZ = IConsoleManager::Get().FindTConsoleVariableDataInt(TEXT("r.DistanceFields.AtlasSizeZ"));
+	const int32 AtlasZ = CVarZ->GetValueOnAnyThread();
+
+	if (bInitialized && (BlockAllocator.GetMaxSizeX() != AtlasXY || BlockAllocator.GetMaxSizeZ() != AtlasZ))
+	{
+		// Atlas size has changed (most likely because of a hotfix). Reallocate everything
+		GDistanceFieldForceAtlasRealloc = 1;
 	}
 
 	if (PendingAllocations.Num() > 0 || GDistanceFieldForceAtlasRealloc != 0)
@@ -319,6 +368,7 @@ void FDistanceFieldVolumeTextureAtlas::UpdateAllocations()
 
 		auto AllocateBlocks = [&]()
 		{
+			int32 FailedAllocationCount = FailedAllocations.Num();
 			for (int32 AllocationIndex = 0; AllocationIndex < LocalPendingAllocations->Num(); AllocationIndex++)
 			{
 				FDistanceFieldVolumeTexture* Texture = (*LocalPendingAllocations)[AllocationIndex];
@@ -327,8 +377,10 @@ void FDistanceFieldVolumeTextureAtlas::UpdateAllocations()
 
 				if (!BlockAllocator.AddElement((uint32&)Texture->AtlasAllocationMin.X, (uint32&)Texture->AtlasAllocationMin.Y, (uint32&)Texture->AtlasAllocationMin.Z, Size.X, Size.Y, Size.Z))
 				{
-					UE_LOG(LogStaticMesh, Error, TEXT("Failed to allocate %ux%ux%u in distance field atlas"), Size.X, Size.Y, Size.Z);
+					UE_LOG(LogStaticMesh, Warning, TEXT("Failed to allocate %ux%ux%u in distance field atlas. Moved mesh distance field to FailedAllocations list"), Size.X, Size.Y, Size.Z);
 					LocalPendingAllocations->RemoveAt(AllocationIndex);
+					FailedAllocations.Add(Texture);
+					FailedAllocatedPixels += Size.X * Size.Y * Size.Z;
 					AllocationIndex--;
 				}
 				else
@@ -339,6 +391,12 @@ void FDistanceFieldVolumeTextureAtlas::UpdateAllocations()
 					AllocatedPixels += Size.X * Size.Y * Size.Z;
 				}
 			}
+
+			if (FailedAllocations.Num() > FailedAllocationCount)
+			{
+				// Sort largest to smallest
+				FailedAllocations.Sort(FCompareVolumeAllocation());
+			}
 		};
 
 		if (bThrottleUpdateAllocation)
@@ -373,12 +431,6 @@ void FDistanceFieldVolumeTextureAtlas::UpdateAllocations()
 		{
 			if (CurrentAllocations.Num() > 0)
 			{
-				static const auto CVarXY = IConsoleManager::Get().FindTConsoleVariableDataInt(TEXT("r.DistanceFields.AtlasSizeXY"));
-				const int32 AtlasXY = CVarXY->GetValueOnAnyThread();
-
-				static const auto CVarZ = IConsoleManager::Get().FindTConsoleVariableDataInt(TEXT("r.DistanceFields.AtlasSizeZ"));
-				const int32 AtlasZ = CVarZ->GetValueOnAnyThread();
-
 				// Remove all allocations from the layout so we have a clean slate
 				BlockAllocator = FTextureLayout3d(0, 0, 0, AtlasXY, AtlasXY, AtlasZ, false, false);
 
@@ -413,7 +465,7 @@ void FDistanceFieldVolumeTextureAtlas::UpdateAllocations()
 				AllocateBlocks();
 			}
 
-			// Fully free the previous atlas memory before allocating a new one6
+			// Fully free the previous atlas memory before allocating a new one
 			{
 				// Remove last ref, add to deferred delete list
 				VolumeTextureRHI = NULL;
diff --git a/Engine/Source/Runtime/Engine/Public/DistanceFieldAtlas.h b/Engine/Source/Runtime/Engine/Public/DistanceFieldAtlas.h
--- a/Engine/Source/Runtime/Engine/Public/DistanceFieldAtlas.h
+++ b/Engine/Source/Runtime/Engine/Public/DistanceFieldAtlas.h
@@ -121,6 +121,9 @@ private:
 	/** Allocations that have already been added, stored in case we need to realloc. */
 	TArray<FDistanceFieldVolumeTexture*> CurrentAllocations;
 
+	/** Allocations that have failed, stored in case they could fit next time a mesh is evicted from atlas. */
+	TArray<FDistanceFieldVolumeTexture*> FailedAllocations;
+		
 	/** Incremented when the atlas is reallocated, so dependencies know to update. */
 	int32 Generation;
 
@@ -128,6 +131,9 @@ private:
 
 	/** Number of pixel used in atlas distance field */
 	uint32 AllocatedPixels;
+	
+	/** Number of pixel that have failed to be allocated in atlas */
+	uint32 FailedAllocatedPixels;
 
 	/** Max position used in distance field */
 	uint32 MaxUsedAtlasX;
