Stop InvertedAxis array growing unbounded

diff --git a/Engine/Source/Runtime/Engine/Private/UserInterface/PlayerInput.cpp b/Engine/Source/Runtime/Engine/Private/UserInterface/PlayerInput.cpp
--- a/Engine/Source/Runtime/Engine/Private/UserInterface/PlayerInput.cpp
+++ b/Engine/Source/Runtime/Engine/Private/UserInterface/PlayerInput.cpp
@@ -504,7 +504,7 @@ void UPlayerInput::InvertAxis(const FName AxisName)
 		}
 		if (bInverted)
 		{
-			InvertedAxis.Add(AxisName);
+			InvertedAxis.AddUnique(AxisName);
 		}
 		else
 		{
