added accessor for ChildConnection's parent object

diff --git a/Engine/Source/Runtime/Engine/Classes/Engine/ChildConnection.h b/Engine/Source/Runtime/Engine/Classes/Engine/ChildConnection.h
--- a/Engine/Source/Runtime/Engine/Classes/Engine/ChildConnection.h
+++ b/Engine/Source/Runtime/Engine/Classes/Engine/ChildConnection.h
@@ -21,6 +21,8 @@ class UChildConnection
 
 public:
 
+	UNetConnection* GetParentConnection() { return Parent; }
+
 	// UNetConnection interface.
 
 	virtual UChildConnection* GetUChildConnection() override
