Chaos: Prevent crash and log more information, possibly corrupt asset

diff --git a/Engine/Source/Runtime/Engine/Private/PhysicsEngine/BodySetup.cpp b/Engine/Source/Runtime/Engine/Private/PhysicsEngine/BodySetup.cpp
--- a/Engine/Source/Runtime/Engine/Private/PhysicsEngine/BodySetup.cpp
+++ b/Engine/Source/Runtime/Engine/Private/PhysicsEngine/BodySetup.cpp
@@ -698,7 +698,8 @@ void UBodySetup::FinishCreatingPhysicsMeshes_Chaos(FChaosDerivedDataReader<float
 		{
 			FKConvexElem& ConvexElem = AggGeom.ConvexElems[ElementIndex];
 
-			if (CHAOS_ENSURE(InReader.ConvexImplicitObjects[ElementIndex]->IsValidGeometry()))
+			if (CHAOS_ENSURE( (ElementIndex < InReader.ConvexImplicitObjects.Num())
+				&& InReader.ConvexImplicitObjects[ElementIndex]->IsValidGeometry()))
 			{
 				ConvexElem.SetChaosConvexMesh(MoveTemp(InReader.ConvexImplicitObjects[ElementIndex]));
 
@@ -710,9 +711,13 @@ void UBodySetup::FinishCreatingPhysicsMeshes_Chaos(FChaosDerivedDataReader<float
 			}
 			else
 			{
+				if (ElementIndex >= InReader.ConvexImplicitObjects.Num())
+				{
+					UE_LOG(LogPhysics, Warning, TEXT("InReader.ConvexImplicitObjects.Num() [%d], AggGeom.ConvexElems.Num() [%d]"),
+						InReader.ConvexImplicitObjects.Num(), AggGeom.ConvexElems.Num());
+				}
 				UE_LOG(LogPhysics, Warning, TEXT("TConvex Name:%s, Element [%d] has no Geometry"), FullName.GetCharArray().GetData(), ElementIndex);
 			}
-
 		}
 		InReader.ConvexImplicitObjects.Reset();
 	}
