Hide mirroring of Material stats behind a commandline flag, add project name to file path, and set the project mirror directory as part of default sources in the stats window.

diff --git a/Engine/Source/Editor/StatsViewer/Private/StatsPages/ShaderCookerStatsPage.cpp b/Engine/Source/Editor/StatsViewer/Private/StatsPages/ShaderCookerStatsPage.cpp
--- a/Engine/Source/Editor/StatsViewer/Private/StatsPages/ShaderCookerStatsPage.cpp
+++ b/Engine/Source/Editor/StatsViewer/Private/StatsPages/ShaderCookerStatsPage.cpp
@@ -185,6 +185,18 @@ FShaderCookerStats::FShaderCookerStats()
 	TArray<FString> Files;
 	FString BasePath = FString::Printf(TEXT("%s/MaterialStats/"), *FPaths::ProjectSavedDir());
 	FPlatformFileManager::Get().GetPlatformFile().FindFiles(Files, *BasePath, TEXT("csv"));
+
+	FString MirrorLocation;
+	GConfig->GetString(TEXT("/Script/Engine.ShaderCompilerStats"), TEXT("MaterialStatsLocation"), MirrorLocation, GGameIni);
+	FParse::Value(FCommandLine::Get(), TEXT("MaterialStatsMirror="), MirrorLocation);
+	if (!MirrorLocation.IsEmpty())
+	{
+		TArray<FString> RemoteFiles;
+		FString RemotePath = FPaths::Combine(*MirrorLocation, FApp::GetProjectName(), *FApp::GetBranchName());
+		FPlatformFileManager::Get().GetPlatformFile().FindFiles(RemoteFiles, *RemotePath, TEXT("csv"));
+		Files.Append(RemoteFiles);
+	}
+
 	for (FString Filename : Files)
 	{
 		FShaderCookerStatsSet Set;
@@ -192,6 +204,7 @@ FShaderCookerStats::FShaderCookerStats()
 		Set.bInitialized = false;
 		StatSets.Emplace(Set);
 	}
+
 }
 FShaderCookerStats::~FShaderCookerStats()
 {
diff --git a/Engine/Source/Runtime/Engine/Private/ShaderCompiler/ShaderCompiler.cpp b/Engine/Source/Runtime/Engine/Private/ShaderCompiler/ShaderCompiler.cpp
--- a/Engine/Source/Runtime/Engine/Private/ShaderCompiler/ShaderCompiler.cpp
+++ b/Engine/Source/Runtime/Engine/Private/ShaderCompiler/ShaderCompiler.cpp
@@ -1544,17 +1544,20 @@ void FShaderCompilerStats::WriteStats()
 			}
 		}
 		DebugWriter->Close();
-		FString MirrorLocation;
-		GConfig->GetString(TEXT("/Script/Engine.ShaderCompilerStats"), TEXT("MaterialStatsLocation"), MirrorLocation, GGameIni);
-		FParse::Value(FCommandLine::Get(), TEXT("MaterialStatsMirror="), MirrorLocation);
-
-		if (!MirrorLocation.IsEmpty())
+		if (FParse::Param(FCommandLine::Get(), TEXT("mirrorshaderstats")))
 		{
-			FString TargetType = TEXT("Default");
-			FParse::Value(FCommandLine::Get(), TEXT("target="), TargetType);
+			FString MirrorLocation;
+			GConfig->GetString(TEXT("/Script/Engine.ShaderCompilerStats"), TEXT("MaterialStatsLocation"), MirrorLocation, GGameIni);
+			FParse::Value(FCommandLine::Get(), TEXT("MaterialStatsMirror="), MirrorLocation);
+
+			if (!MirrorLocation.IsEmpty())
+			{
+				FString TargetType = TEXT("Default");
+				FParse::Value(FCommandLine::Get(), TEXT("target="), TargetType);
 
-			FString CopyLocation = FPaths::Combine(*MirrorLocation, *FApp::GetBranchName(), FString::Printf(TEXT("Stats-Latest(%s).csv"), *TargetType));
-			IFileManager::Get().Copy(*CopyLocation, *FileName, true, true);
+				FString CopyLocation = FPaths::Combine(*MirrorLocation, FApp::GetProjectName(), *FApp::GetBranchName(), FString::Printf(TEXT("Stats-Latest(%s).csv"), *TargetType));
+				IFileManager::Get().Copy(*CopyLocation, *FileName, true, true);
+			}
 		}
 	}
 	{
