Fix crash when a texture fails to compress during cooking.

This is a regression introduced in CL 7839132 (it's not in 4.23). If a texture fails to compress, it will have 0 mips, so the clamp to 1 was incorrect.

#rnx

diff --git a/Engine/Source/Runtime/Engine/Private/TextureDerivedData.cpp b/Engine/Source/Runtime/Engine/Private/TextureDerivedData.cpp
--- a/Engine/Source/Runtime/Engine/Private/TextureDerivedData.cpp
+++ b/Engine/Source/Runtime/Engine/Private/TextureDerivedData.cpp
@@ -1325,7 +1325,7 @@ static void SerializePlatformData(
 			if (!bIsVirtual)
 			{
 				FirstMipToSerialize = FMath::Clamp(FirstMipToSerialize, 0, PlatformData->GetNumMipsInTail() > 0 ? FirstMipTailMip : LastMip);
-				NumMips = FMath::Max(1, NumMips - FirstMipToSerialize);
+				NumMips = FMath::Max(0, NumMips - FirstMipToSerialize);
 			}
 			else
 			{
