#framework Doc comment change: fixed grammar

diff --git a/Engine/Source/Runtime/Engine/Classes/GameFramework/Actor.h b/Engine/Source/Runtime/Engine/Classes/GameFramework/Actor.h
--- a/Engine/Source/Runtime/Engine/Classes/GameFramework/Actor.h
+++ b/Engine/Source/Runtime/Engine/Classes/GameFramework/Actor.h
@@ -2098,7 +2098,7 @@ protected:
 public:
 	/** 
 	 * Set the owner of this Actor, used primarily for network replication. 
-	 * @param NewOwner	The Actor whom takes over ownership of this Actor
+	 * @param NewOwner	The Actor who takes over ownership of this Actor
 	 */
 	UFUNCTION(BlueprintCallable, Category=Actor)
 	virtual void SetOwner( AActor* NewOwner );
