Fixing up cases where we weren't using the correction position for listener based distance calcuations

diff --git a/Engine/Source/Runtime/Engine/Private/AudioDevice.cpp b/Engine/Source/Runtime/Engine/Private/AudioDevice.cpp
--- a/Engine/Source/Runtime/Engine/Private/AudioDevice.cpp
+++ b/Engine/Source/Runtime/Engine/Private/AudioDevice.cpp
@@ -4930,7 +4930,8 @@ float FAudioDevice::GetDistanceToNearestListener(const FVector& Location) const
 	check(bInAudioThread || bInGameThread);
 
 	float DistSquared;
-	if (FindClosestListenerIndex(Location, DistSquared, true) == INDEX_NONE)
+	const bool bAllowAttenuationOverrides = true;
+	if (FindClosestListenerIndex(Location, DistSquared, bAllowAttenuationOverrides) == INDEX_NONE)
 	{
 		return WORLD_MAX;
 	}
@@ -5119,7 +5120,7 @@ bool FAudioDevice::SoundIsAudible(const FActiveSound& NewActiveSound)
 int32 FAudioDevice::FindClosestListenerIndex(const FTransform& SoundTransform) const
 {
 	float UnusedDistSq;
-	const bool bAllowOverrides = false;
+	const bool bAllowOverrides = true;
 	return FindClosestListenerIndex(SoundTransform.GetTranslation(), UnusedDistSq, bAllowOverrides);
 }
 
diff --git a/Engine/Source/Runtime/Engine/Private/SoundConcurrency.cpp b/Engine/Source/Runtime/Engine/Private/SoundConcurrency.cpp
--- a/Engine/Source/Runtime/Engine/Private/SoundConcurrency.cpp
+++ b/Engine/Source/Runtime/Engine/Private/SoundConcurrency.cpp
@@ -481,7 +481,7 @@ FActiveSound* FSoundConcurrencyManager::GetEvictableSoundStopFarthest(const FAct
 
 	check(AudioDevice);
 	float DistanceToStopSoundSq;
-	const bool bAllowAttenuationOverrides = false;
+	const bool bAllowAttenuationOverrides = true;
 	int32 ClosestListenerIndex = AudioDevice->FindClosestListenerIndex(NewActiveSound.Transform.GetTranslation(), DistanceToStopSoundSq, bAllowAttenuationOverrides);
 
 	FActiveSound* EvictableSound = nullptr;
