Fix ensure loading cooked animations in editor

diff --git a/Engine/Source/Runtime/Engine/Private/Animation/AnimSequenceBase.cpp b/Engine/Source/Runtime/Engine/Private/Animation/AnimSequenceBase.cpp
--- a/Engine/Source/Runtime/Engine/Private/Animation/AnimSequenceBase.cpp
+++ b/Engine/Source/Runtime/Engine/Private/Animation/AnimSequenceBase.cpp
@@ -437,20 +437,18 @@ void UAnimSequenceBase::RefreshCacheData()
 		// Handle busted track indices
 		if (!AnimNotifyTracks.IsValidIndex(Notify.TrackIndex))
 		{
-			// This really shouldn't happen, but try to handle it
-			ensureMsgf(0, TEXT("AnimNotifyTrack: Anim (%s) has notify (%s) with track index (%i) that does not exist"), *GetFullName(), *Notify.NotifyName.ToString(), Notify.TrackIndex);
+			// This really shouldn't happen (unless we are a cooked asset), but try to handle it
+			ensureMsgf(GetOutermost()->bIsCookedForEditor, TEXT("AnimNotifyTrack: Anim (%s) has notify (%s) with track index (%i) that does not exist"), *GetFullName(), *Notify.NotifyName.ToString(), Notify.TrackIndex);
 
 			// Don't create lots of extra tracks if we are way off supporting this track
 			if (Notify.TrackIndex < 0 || Notify.TrackIndex > 20)
 			{
 				Notify.TrackIndex = 0;
 			}
-			else
+
+			while (!AnimNotifyTracks.IsValidIndex(Notify.TrackIndex))
 			{
-				while (!AnimNotifyTracks.IsValidIndex(Notify.TrackIndex))
-				{
-					AddNewTrack(AnimNotifyTracks);
-				}
+				AddNewTrack(AnimNotifyTracks);
 			}
 		}
 
