Readded 'DumpTextureStreamingStats' command.  Logs current streaming poolsize and occupancy.  Can be used in test builds, where stat streaming is disabled.

diff --git a/Engine/Source/Runtime/Engine/Private/Streaming/StreamingManagerTexture.cpp b/Engine/Source/Runtime/Engine/Private/Streaming/StreamingManagerTexture.cpp
--- a/Engine/Source/Runtime/Engine/Private/Streaming/StreamingManagerTexture.cpp
+++ b/Engine/Source/Runtime/Engine/Private/Streaming/StreamingManagerTexture.cpp
@@ -111,14 +111,6 @@ FRenderAssetStreamingManager::FRenderAssetStreamingManager()
 	// Convert from MByte to byte.
 	MemoryMargin *= 1024 * 1024;
 
-#if STATS_FAST
-	MaxStreamingTexturesSize = 0;
-	MaxOptimalTextureSize = 0;
-	MaxStreamingOverBudget = MIN_int64;
-	MaxTexturePoolAllocatedSize = 0;
-	MaxNumWantingTextures = 0;
-#endif
-
 	for ( int32 LODGroup=0; LODGroup < TEXTUREGROUP_MAX; ++LODGroup )
 	{
 		const FTextureLODGroup& TexGroup = UDeviceProfileManager::Get().GetActiveProfile()->GetTextureLODSettings()->GetTextureLODGroup(TextureGroup(LODGroup));
@@ -1785,27 +1777,30 @@ void FRenderAssetStreamingManager::PropagateLightingScenarioChange()
 	}
 }
 
-#if STATS_FAST
+#if !UE_BUILD_SHIPPING
+
 bool FRenderAssetStreamingManager::HandleDumpTextureStreamingStatsCommand( const TCHAR* Cmd, FOutputDevice& Ar )
 {
 	FScopeLock ScopeLock(&CriticalSection);
 
-	Ar.Logf( TEXT("Current Texture Streaming Stats") );
-	Ar.Logf( TEXT("  Textures In Memory, Current (KB) = %f"), MaxStreamingTexturesSize / 1024.0f);
-	Ar.Logf( TEXT("  Textures In Memory, Target (KB) =  %f"), MaxOptimalTextureSize / 1024.0f );
-	Ar.Logf( TEXT("  Over Budget (KB) =                 %f"), MaxStreamingOverBudget / 1024.0f );
-	Ar.Logf( TEXT("  Pool Memory Used (KB) =            %f"), MaxTexturePoolAllocatedSize / 1024.0f );
-	Ar.Logf( TEXT("  Num Wanting Textures =             %d"), MaxNumWantingTextures );
-	MaxStreamingTexturesSize = 0;
-	MaxOptimalTextureSize = 0;
-	MaxStreamingOverBudget = MIN_int64;
-	MaxTexturePoolAllocatedSize = 0;
-	MaxNumWantingTextures = 0;
+	float CurrentOccupancyPct = 0.f;
+	float TargetOccupancyPct = 0.f;
+	if (DisplayedStats.StreamingPool > 0)
+	{
+		CurrentOccupancyPct = 100.f * (DisplayedStats.WantedMips / (float)DisplayedStats.StreamingPool);
+		TargetOccupancyPct = 100.f * (DisplayedStats.RequiredPool / (float)DisplayedStats.StreamingPool);
+	}
+
+	float StreamingPoolMB = DisplayedStats.StreamingPool / 1024.f / 1024.f;
+
+	// uses csv stats for access in test builds
+	Ar.Logf(TEXT("Current Texture Streaming Stats") );
+	Ar.Logf(TEXT("  Streaming Assets, Current/Pool = %.2f / %.2f MB (%d%%)"), DisplayedStats.WantedMips / 1024.f / 1024.f, StreamingPoolMB, FMath::RoundToInt(CurrentOccupancyPct));
+	Ar.Logf(TEXT("  Streaming Assets, Target/Pool =  %.2f / %.2f MB (%d%%)"), DisplayedStats.RequiredPool / 1024.f / 1024.f, StreamingPoolMB, FMath::RoundToInt(TargetOccupancyPct));
+	Ar.Logf(TEXT("  Non-Streaming Mips = %.2f MB"), DisplayedStats.NonStreamingMips / 1024.f / 1024.f);
+
 	return true;
 }
-#endif // STATS_FAST
-
-#if !UE_BUILD_SHIPPING
 
 bool FRenderAssetStreamingManager::HandleListStreamingRenderAssetsCommand( const TCHAR* Cmd, FOutputDevice& Ar )
 {
@@ -2419,14 +2414,12 @@ bool FRenderAssetStreamingManager::HandleInvestigateRenderAssetCommand(const TCH
  */
 bool FRenderAssetStreamingManager::Exec( UWorld* InWorld, const TCHAR* Cmd, FOutputDevice& Ar )
 {
-#if STATS_FAST
+#if !UE_BUILD_SHIPPING
 	if (FParse::Command(&Cmd,TEXT("DumpTextureStreamingStats"))
 		|| FParse::Command(&Cmd, TEXT("DumpRenderAssetStreamingStats")))
 	{
 		return HandleDumpTextureStreamingStatsCommand( Cmd, Ar );
 	}
-#endif
-#if !UE_BUILD_SHIPPING
 	if (FParse::Command(&Cmd,TEXT("ListStreamingTextures"))
 		|| FParse::Command(&Cmd, TEXT("ListStreamingRenderAssets")))
 	{
diff --git a/Engine/Source/Runtime/Engine/Private/Streaming/StreamingManagerTexture.h b/Engine/Source/Runtime/Engine/Private/Streaming/StreamingManagerTexture.h
--- a/Engine/Source/Runtime/Engine/Private/Streaming/StreamingManagerTexture.h
+++ b/Engine/Source/Runtime/Engine/Private/Streaming/StreamingManagerTexture.h
@@ -19,8 +19,6 @@ class UPrimitiveComponent;
 
 template<typename TTask> class FAsyncTask;
 
-#define STATS_FAST 0
-
 /*-----------------------------------------------------------------------------
 	Texture or mesh streaming.
 -----------------------------------------------------------------------------*/
@@ -113,10 +111,8 @@ struct FRenderAssetStreamingManager final : public IRenderAssetStreamingManager
 	/**
 	 * Exec command handlers
 	 */
-#if STATS_FAST
-	bool HandleDumpTextureStreamingStatsCommand( const TCHAR* Cmd, FOutputDevice& Ar );
-#endif // STATS_FAST
 #if !UE_BUILD_SHIPPING
+	bool HandleDumpTextureStreamingStatsCommand(const TCHAR* Cmd, FOutputDevice& Ar);
 	bool HandleListStreamingRenderAssetsCommand(const TCHAR* Cmd, FOutputDevice& Ar);
 	bool HandleResetMaxEverRequiredRenderAssetMemoryCommand(const TCHAR* Cmd, FOutputDevice& Ar);
 	bool HandleLightmapStreamingFactorCommand( const TCHAR* Cmd, FOutputDevice& Ar );
@@ -414,14 +410,6 @@ private:
 	TMap<FString, bool> CachedFileExistsChecks;
 	void OnPakFileMounted(const TCHAR* PakFilename);
 
-#if STATS_FAST
-	uint64 MaxStreamingTexturesSize;
-	uint64 MaxOptimalTextureSize;
-	int64 MaxStreamingOverBudget;
-	uint64 MaxTexturePoolAllocatedSize;
-	uint32 MaxNumWantingTextures;
-#endif
-	
 	// A critical section use around code that could be called in parallel with NotifyPrimitiveUpdated() or NotifyPrimitiveUpdated_Concurrent().
 	FCriticalSection CriticalSection;
 
