# Speculative fix for server crash in UAnimInstance::TriggerAnimNotifies
#rnx

diff --git a/Engine/Source/Runtime/Engine/Private/Animation/AnimInstance.cpp b/Engine/Source/Runtime/Engine/Private/Animation/AnimInstance.cpp
--- a/Engine/Source/Runtime/Engine/Private/Animation/AnimInstance.cpp
+++ b/Engine/Source/Runtime/Engine/Private/Animation/AnimInstance.cpp
@@ -1307,12 +1307,21 @@ void UAnimInstance::TriggerAnimNotifies(float DeltaSeconds)
 	}
 
 	// Send end notification to AnimNotifyState not active anymore.
-	for(const FAnimNotifyEvent& AnimNotifyEvent : ActiveAnimNotifyState)
+	for (int32 Index = 0; Index < ActiveAnimNotifyState.Num(); ++Index)
 	{
+		const FAnimNotifyEvent& AnimNotifyEvent = ActiveAnimNotifyState[Index];
 		if (AnimNotifyEvent.NotifyStateClass && ShouldTriggerAnimNotifyState(AnimNotifyEvent.NotifyStateClass))
 		{
 			AnimNotifyEvent.NotifyStateClass->NotifyEnd(SkelMeshComp, Cast<UAnimSequenceBase>(AnimNotifyEvent.NotifyStateClass->GetOuter()));
 		}
+		// The NotifyEnd callback above may have triggered actor destruction and the tear down
+		// of this instance via UninitializeAnimation which empties ActiveAnimNotifyState.
+		// If that happened, we should stop iterating the ActiveAnimNotifyState array
+		if (ActiveAnimNotifyState.IsValidIndex(Index) == false)
+		{
+			ensureMsgf(false, TEXT("UAnimInstance::ActiveAnimNotifyState has been invalidated by NotifyEnd. AnimInstance: %s, Owning Component: %s, Owning Actor: %s "), *GetNameSafe(this), *GetNameSafe(GetOwningComponent()), *GetNameSafe(GetOwningActor()));
+			return;
+		}
 	}
 
 	// Call 'NotifyBegin' event on freshly added AnimNotifyState.
