Adding customization to stat soundcue view
+ minimal mode: shows color coded items w/ volume, sound cue, sound class, and path.  Adds spacing to make things easier to see at a glance.
+ showpaths: on by default, allows the user to turn off sound paths to regain some debug real estate, only showing sound cue name

diff --git a/Engine/Source/Runtime/Engine/Private/AudioDebug.cpp b/Engine/Source/Runtime/Engine/Private/AudioDebug.cpp
--- a/Engine/Source/Runtime/Engine/Private/AudioDebug.cpp
+++ b/Engine/Source/Runtime/Engine/Private/AudioDebug.cpp
@@ -106,6 +106,38 @@ FAutoConsoleVariableRef CVarAudioDebugStatSoundColor(
 	TEXT("White, Red, Orange, Yellow, Blue, Magenta, Purple, Black"),
 	ECVF_Default);
 
+static int32 SoundCueDebugShowPathCVar = 1;
+FAutoConsoleVariableRef CVarAudioSoundCueDebugShowPath(
+	TEXT("au.Debug.Soundcues.ShowPath"),
+	SoundCueDebugShowPathCVar,
+	TEXT("Display full path of sound cue when enabled.\n")
+	TEXT("0: Not Enabled, 1: Enabled"),
+	ECVF_Default);
+
+static int32 SoundCueDebugMinimalCVar = 0;
+FAutoConsoleVariableRef CVarAudioSoundCueDebugMinimal(
+	TEXT("au.Debug.SoundCues.Minimal"),
+	SoundCueDebugMinimalCVar,
+	TEXT("Use the compact view of sound cue debug when enabled. \n")
+	TEXT("0: Not Enabled, 1: Enabled"),
+	ECVF_Default);
+
+static int32 SoundCueDebugTabSpacingCVar = 5;
+FAutoConsoleVariableRef CVarAudioSoundCueDebugTabSpacing(
+	TEXT("au.Debug.SoundCues.Spacing.Tab"),
+	SoundCueDebugTabSpacingCVar,
+	TEXT("Size of tab (in characters) with compact view. \n")
+	TEXT("Default: 5"),
+	ECVF_Default);
+
+static int32 SoundCueDebugCharSpacingCVar = 7;
+FAutoConsoleVariableRef CVarSoundCueDebugCharSpacing(
+	TEXT("au.Debug.SoundCues.Spacing.Char"),
+	SoundCueDebugCharSpacingCVar,
+	TEXT("Size of character (in pixels) with compact view. \n")
+	TEXT("Default: 7"),
+	ECVF_Default);
+
 // Test Data
 namespace ERequestedAudioStats
 {
@@ -885,12 +917,17 @@ int32 FAudioDebugger::RenderStatCues(UWorld* World, FViewport* Viewport, FCanvas
 		return Y;
 	}
 
-	const int32 FontHeight = GetStatsFont()->GetMaxCharHeight() + 2;
-	Canvas->DrawShadowedString(X, Y, TEXT("Active Sound Cues:"), GetStatsFont(), HeaderColor);
+	UFont* StatsFont = GetStatsFont();
+	const int32 FontSpacing = 2;
+	const int32 FontHeight = StatsFont->GetMaxCharHeight() + FontSpacing;
+	Canvas->DrawShadowedString(X, Y, TEXT("Active Sound Cues:"), StatsFont, HeaderColor);
 	Y += FontHeight;
 
 	int32 ActiveSoundCount = 0;
 	FAudioStats& AudioStats = AudioDeviceStats.FindOrAdd(AudioDevice->DeviceHandle);
+
+	if (!SoundCueDebugMinimalCVar)
+	{
 	for (const FAudioStats::FStatSoundInfo& StatSoundInfo : AudioStats.StatSoundInfos)
 	{
 		for (const FAudioStats::FStatWaveInstanceInfo& WaveInstanceInfo : StatSoundInfo.WaveInstanceInfos)
@@ -899,22 +936,113 @@ int32 FAudioDebugger::RenderStatCues(UWorld* World, FViewport* Viewport, FCanvas
 			{
 				FColor Color = FColor::White;
 				FString MuteSoloReason;
-				
+
+					if (FSoundSource::FDebugInfo* DebugInfo = WaveInstanceInfo.DebugInfo.Get())
+					{
+						// Color code same as icons. Red (mute), Yellow (solo), White (normal).
+						FScopeLock Lock(&DebugInfo->CS);
+						Color = DebugInfo->bIsMuted ? FColor::Red : DebugInfo->bIsSoloed ? FColor::Yellow : FColor::White;
+						MuteSoloReason = !DebugInfo->MuteSoloReason.IsEmpty() ? FString::Printf(TEXT(" - %s"), *DebugInfo->MuteSoloReason) : TEXT("");
+					}
+
+					const FString TheString = FString::Printf(TEXT("%4i. %6.2f %s %s %s"), ActiveSoundCount++, WaveInstanceInfo.Volume, *StatSoundInfo.SoundName, *StatSoundInfo.SoundClassName.ToString(), *MuteSoloReason);
+					Canvas->DrawShadowedString(X, Y, *TheString, StatsFont, Color);
+					Y += FontHeight;
+					break;
+				}
+			}
+		}
+	}
+	else
+	{
+		TArray<TPair<FString, FString> > SplitNames;
+		FString SoundPath, SoundName;
+		int32 MaxNameLength = 0;
+		uint32 MaxClassLength = 0;
+
+		// Get spacing data
+		for (const FAudioStats::FStatSoundInfo& StatSoundInfo : AudioStats.StatSoundInfos)
+		{
+			for (const FAudioStats::FStatWaveInstanceInfo& WaveInstanceInfo : StatSoundInfo.WaveInstanceInfos)
+			{
+				if (WaveInstanceInfo.Volume >= MinDisplayVolume)
+				{
+					if (!StatSoundInfo.SoundName.Split(TEXT("."), &SoundPath, &SoundName))
+					{
+						if (!StatSoundInfo.SoundName.Split(SUBOBJECT_DELIMITER, &SoundPath, &SoundName))
+						{
+							SoundName.Reset();
+							SoundPath = StatSoundInfo.SoundName;
+						}
+					}
+					SplitNames.Emplace(SoundName, SoundPath);
+
+					if (SoundName.Len() > MaxNameLength)
+					{
+						MaxNameLength = SoundName.Len();
+					}
+
+					if (StatSoundInfo.SoundClassName.GetStringLength() > MaxClassLength)
+					{
+						MaxClassLength = StatSoundInfo.SoundClassName.GetStringLength();
+					}
+				}
+			}
+		}
+
+		const int32 TabSpacing = FMath::Clamp(SoundCueDebugTabSpacingCVar, 1, SoundCueDebugTabSpacingCVar);
+		const int32 CharSpacing = FMath::Clamp(SoundCueDebugCharSpacingCVar, 1, SoundCueDebugCharSpacingCVar);
+		const int32 NumberSpacing = 6 * CharSpacing;	// 6 character len for 2 decimal float + 2 spaces 'X.XX  '
+
+		// Tab out name and class length
+		MaxNameLength = (MaxNameLength / TabSpacing + 1) * TabSpacing;
+		MaxClassLength = (MaxClassLength / TabSpacing + 1) * TabSpacing;
+
+		for (const FAudioStats::FStatSoundInfo& StatSoundInfo : AudioStats.StatSoundInfos)
+		{
+			for (const FAudioStats::FStatWaveInstanceInfo& WaveInstanceInfo : StatSoundInfo.WaveInstanceInfos)
+			{
+				if (WaveInstanceInfo.Volume >= MinDisplayVolume)
+				{
+					FColor Color = FColor::White;
+					FString MuteSoloReason;
+					bool bMutedOrSoloed = false;
+
 				if (FSoundSource::FDebugInfo* DebugInfo = WaveInstanceInfo.DebugInfo.Get())
 				{
-					// Color code same as icons. Orange (mute), Solo (red), Normal (white).					
+						// Color code same as icons. Red (mute), Yellow (solo), White (normal).
 					FScopeLock Lock(&DebugInfo->CS);
-					Color = DebugInfo->bIsMuted ? FColor::Orange : DebugInfo->bIsSoloed ? FColor::Red : FColor::White;
+						Color = DebugInfo->bIsMuted ? FColor::Red : DebugInfo->bIsSoloed ? FColor::Yellow : FColor::White;
+						bMutedOrSoloed = DebugInfo->bIsMuted || DebugInfo->bIsSoloed;
 					MuteSoloReason = !DebugInfo->MuteSoloReason.IsEmpty() ? FString::Printf(TEXT(" - %s"), *DebugInfo->MuteSoloReason) : TEXT("");
 				}
 
-				const FString TheString = FString::Printf(TEXT("%4i. %s %s %s"), ActiveSoundCount++, *StatSoundInfo.SoundName, *StatSoundInfo.SoundClassName.ToString(), *MuteSoloReason );
-				Canvas->DrawShadowedString(X, Y, *TheString, GetStatsFont(), Color);
+					const int32 SoundNameIndex = ActiveSoundCount++;
+					const FString LeadingNumber = FString::Printf(TEXT("%4i. "), SoundNameIndex);
+					const FString Volume = FString::Printf(TEXT("%6.2f "), WaveInstanceInfo.Volume);
+					const FString PathAndMuting = FString::Printf(TEXT("Path: %s %s"), *SplitNames[SoundNameIndex].Value, *MuteSoloReason);
+
+					int32 CurrentX = X;
+					Canvas->DrawShadowedString(CurrentX, Y, *LeadingNumber, StatsFont, Color);
+					CurrentX += NumberSpacing;
+					Canvas->DrawShadowedString(CurrentX, Y, *Volume, StatsFont, bMutedOrSoloed ? Color : FColor::Orange);
+					CurrentX += NumberSpacing;
+					Canvas->DrawShadowedString(CurrentX, Y, *SplitNames[SoundNameIndex].Key, StatsFont, bMutedOrSoloed ? Color : FColor(0, 255, 255));
+					CurrentX += (MaxNameLength * CharSpacing);
+					Canvas->DrawShadowedString(CurrentX, Y, *StatSoundInfo.SoundClassName.ToString(), StatsFont, bMutedOrSoloed ? Color : FColor::Yellow);
+
+					if (SoundCueDebugShowPathCVar)
+					{
+						CurrentX += (MaxClassLength * CharSpacing);
+						Canvas->DrawShadowedString(CurrentX, Y, *PathAndMuting, StatsFont, Color);
+					}
+
 				Y += FontHeight;
 				break;
 			}
 		}
 	}
+	}
 
 	Canvas->DrawShadowedString(X, Y, *FString::Printf(TEXT("Total: %i"), ActiveSoundCount), GetStatsFont(), GetBodyColor());
 	Y += FontHeight;
@@ -1096,19 +1224,19 @@ int32 FAudioDebugger::RenderStatSounds(UWorld* World, FViewport* Viewport, FCanv
 	switch (AudioStats.DisplaySort)
 	{
 		case FAudioStats::EDisplaySort::Class:
-		{
+	{
 			AudioStats.StatSoundInfos.Sort([](const FAudioStats::FStatSoundInfo& A, const FAudioStats::FStatSoundInfo& B) { return A.SoundClassName.LexicalLess(B.SoundClassName); });
 			SortingName = TEXT("Class");
 			break;
-		}
+	}
 		case FAudioStats::EDisplaySort::Distance:
-		{
-			AudioStats.StatSoundInfos.Sort([](const FAudioStats::FStatSoundInfo& A, const FAudioStats::FStatSoundInfo& B) { return A.Distance < B.Distance; });
+	{
+		AudioStats.StatSoundInfos.Sort([](const FAudioStats::FStatSoundInfo& A, const FAudioStats::FStatSoundInfo& B) { return A.Distance < B.Distance; });
 			SortingName = TEXT("Distance");
 			break;
-		}
+	}
 		case FAudioStats::EDisplaySort::PlaybackTime:
-		{
+	{
 			AudioStats.StatSoundInfos.Sort([](const FAudioStats::FStatSoundInfo& A, const FAudioStats::FStatSoundInfo& B) { return A.PlaybackTime > B.PlaybackTime; });
 			SortingName = TEXT("Time");
 			break;
@@ -1118,10 +1246,10 @@ int32 FAudioDebugger::RenderStatSounds(UWorld* World, FViewport* Viewport, FCanv
 			AudioStats.StatSoundInfos.Sort([](const FAudioStats::FStatSoundInfo& A, const FAudioStats::FStatSoundInfo& B) { return A.Priority > B.Priority; });
 			SortingName = TEXT("Priority");
 			break;
-		}
+	}
 		case FAudioStats::EDisplaySort::Waves:
-		{
-			AudioStats.StatSoundInfos.Sort([](const FAudioStats::FStatSoundInfo& A, const FAudioStats::FStatSoundInfo& B) { return A.WaveInstanceInfos.Num() > B.WaveInstanceInfos.Num(); });
+	{
+		AudioStats.StatSoundInfos.Sort([](const FAudioStats::FStatSoundInfo& A, const FAudioStats::FStatSoundInfo& B) { return A.WaveInstanceInfos.Num() > B.WaveInstanceInfos.Num(); });
 			SortingName = TEXT("Waves");
 			break;
 		}
@@ -1218,7 +1346,7 @@ int32 FAudioDebugger::RenderStatSounds(UWorld* World, FViewport* Viewport, FCanv
 				break;
 			}
 			case FAudioStats::EDisplaySort::Waves:
-			{
+		{
 				DebugValue = FString::Printf(TEXT("%03u"), StatSoundInfo.WaveInstanceInfos.Num());
 				bDisplayWaves = true;
 				break;
@@ -1230,7 +1358,7 @@ int32 FAudioDebugger::RenderStatSounds(UWorld* World, FViewport* Viewport, FCanv
 				DebugValue = FString::Printf(TEXT("%08.2f"), StatSoundInfo.Distance);
 				break;
 			}
-		}
+			}
 
 		const FString DebugStr = FString::Printf(TEXT("  %03i    %s %s"), SoundIndex, *DebugValue.RightPad(FieldLength - 4).Left(FieldLength - 4), *DisplayName);
 		Canvas->DrawShadowedString(X, Y, *DebugStr, GetStatsFont(), BodyColor);
