Inertia tensor computations
* This change corrects scaled inertia tensors (previous implementation approximated scaleds as boxes and then scaled them inversely, this one just uses the bounds of the scaled... still an approximation).
* Moves inertia tensor computations from the interface into the geometry.

diff --git a/Engine/Source/Runtime/Engine/Private/PhysicsEngine/Experimental/PhysInterface_Chaos.cpp b/Engine/Source/Runtime/Engine/Private/PhysicsEngine/Experimental/PhysInterface_Chaos.cpp
--- a/Engine/Source/Runtime/Engine/Private/PhysicsEngine/Experimental/PhysInterface_Chaos.cpp
+++ b/Engine/Source/Runtime/Engine/Private/PhysicsEngine/Experimental/PhysInterface_Chaos.cpp
@@ -511,11 +511,21 @@ float FPhysInterface_Chaos::GetMaxDepenetrationVelocity_AssumesLocked(const FPhy
 
 void FPhysInterface_Chaos::SetMaxDepenetrationVelocity_AssumesLocked(const FPhysicsActorHandle& InActorReference, float InMaxDepenetrationVelocity)
 {
+	CHAOS_ENSURE(false);
 }
 
 FVector FPhysInterface_Chaos::GetWorldVelocityAtPoint_AssumesLocked(const FPhysicsActorHandle& InActorReference, const FVector& InPoint)
 {
-	// #todo : Implement
+	if (ensure(FPhysicsInterface::IsValid(InActorReference)))
+	{
+		Chaos::TKinematicGeometryParticle<float, 3>* Kinematic = InActorReference->AsKinematic();
+		if (ensure(Kinematic))
+		{
+			const Chaos::FVec3 COM = Kinematic->X(); // TODO: Fix once we have separate COM
+			const Chaos::FVec3 Diff = InPoint - COM;
+			return Kinematic->V() + Chaos::FVec3::CrossProduct(Diff, Kinematic->W());
+		}
+	}
 	return FVector(0);
 }
 
@@ -566,34 +576,40 @@ void FPhysInterface_Chaos::AddImpulse_AssumesLocked(const FPhysicsActorHandle& I
 {
 	// #todo : Implement
     //InActorReference.GetScene()->AddForce(InForce, InActorReference.GetId());
+	CHAOS_ENSURE(false);
 }
 
 void FPhysInterface_Chaos::AddAngularImpulseInRadians_AssumesLocked(const FPhysicsActorHandle& InActorReference, const FVector& InTorque)
 {
 	// #todo : Implement
     //InActorReference.GetScene()->AddTorque(InTorque, InActorReference.GetId());
+	CHAOS_ENSURE(false);
 }
 
 void FPhysInterface_Chaos::AddVelocity_AssumesLocked(const FPhysicsActorHandle& InActorReference, const FVector& InForce)
 {	
 	// #todo : Implement
     //InActorReference.GetScene()->AddForce(InForce * InActorReference.GetScene()->Scene.GetSolver()->GetRigidParticles().M(InActorReference.GetScene()->GetIndexFromId(InActorReference.GetId())), InActorReference.GetId());
+	CHAOS_ENSURE(false);
 }
 
 void FPhysInterface_Chaos::AddAngularVelocityInRadians_AssumesLocked(const FPhysicsActorHandle& InActorReference, const FVector& InTorque)
 {
 	// #todo : Implement
     //InActorReference.GetScene()->AddTorque(InActorReference.GetScene()->Scene.GetSolver()->GetRigidParticles().I(InActorReference.GetScene()->GetIndexFromId(InActorReference.GetId())) * Chaos::TVector<float, 3>(InTorque), InActorReference.GetId());
+	CHAOS_ENSURE(false);
 }
 
 void FPhysInterface_Chaos::AddImpulseAtLocation_AssumesLocked(const FPhysicsActorHandle& InActorReference, const FVector& InImpulse, const FVector& InLocation)
 {
     // @todo(mlentine): We don't currently have a way to apply an instantaneous force. Do we need this?
+	CHAOS_ENSURE(false);
 }
 
 void FPhysInterface_Chaos::AddRadialImpulse_AssumesLocked(const FPhysicsActorHandle& InActorReference, const FVector& InOrigin, float InRadius, float InStrength, ERadialImpulseFalloff InFalloff, bool bInVelChange)
 {
     // @todo(mlentine): We don't currently have a way to apply an instantaneous force. Do we need this?
+	CHAOS_ENSURE(false);
 }
 
 bool FPhysInterface_Chaos::IsGravityEnabled_AssumesLocked(const FPhysicsActorHandle& InActorReference)
@@ -1400,138 +1416,18 @@ void CalculateMassPropertiesOfImplicitType(
 
 	if (ImplicitObject)
 	{
-		Chaos::TVector<float, 3> Scale = WorldTransform.GetScale3D();
-		Chaos::PMatrix<float, 3, 3> ScaleM(Scale.X, Scale.Y, Scale.Z);
-
-		if (ImplicitObject->GetType() == Chaos::ImplicitObjectType::Sphere)
-		{
-			const Chaos::TSphere<float, 3> * Sphere = ImplicitObject->template GetObject<Chaos::TSphere<float, 3>>();
-			float Mass = Sphere->GetVolume()*InDensityKGPerCM;
-
-			Chaos::TMassProperties<float, 3> MassProperty;
-			MassProperty.CenterOfMass = Chaos::TVector<float, 3>(0);
-			MassProperty.Volume = Sphere->GetVolume();
-			MassProperty.InertiaTensor = Sphere->GetInertiaTensor(Mass)*ScaleM;
-			MassProperty.RotationOfMass = Chaos::TRotation<float, 3>::FromIdentity();
-			//MassProperty.RotationOfMass = Chaos::TransformToLocalSpace<float,3>(MassProperty.InertiaTensor);
-
-			MassProperties.Add(MassProperty);
-			TotalMass += Mass;
-		}
-		else if (ImplicitObject->GetType() == Chaos::ImplicitObjectType::Box)
-		{
-			const Chaos::TBox<float, 3> * Box = ImplicitObject->template GetObject<Chaos::TBox<float, 3>>();
-			float Mass = Box->GetVolume()*InDensityKGPerCM;
-
-			Chaos::TMassProperties<float, 3> MassProperty;
-			MassProperty.CenterOfMass = Chaos::TVector<float,3>(0);
-			MassProperty.Volume = Box->GetVolume();
-			MassProperty.InertiaTensor = Box->GetInertiaTensor(Mass); // What's the box? Scale!
-			MassProperty.RotationOfMass = Chaos::TRotation<float, 3>::FromIdentity();
-			//MassProperty.RotationOfMass = Chaos::TransformToLocalSpace<float, 3>(MassProperty.InertiaTensor);
-
-			MassProperties.Add(MassProperty);
-			TotalMass += Mass;
-		}
-		else if (ImplicitObject->GetType() == Chaos::ImplicitObjectType::Capsule)
-		{
-			const Chaos::TCapsule<float> * Capsule = ImplicitObject->template GetObject<Chaos::TCapsule<float>>();
-			float Mass = Capsule->GetVolume()*InDensityKGPerCM;
-
-			Chaos::TMassProperties<float, 3> MassProperty;
-			MassProperty.CenterOfMass = Chaos::TVector<float, 3>(0);
-			MassProperty.Volume = Capsule->GetVolume();
-			MassProperty.InertiaTensor = Capsule->GetInertiaTensor(Mass)*ScaleM;
-			MassProperty.RotationOfMass = Chaos::TRotation<float, 3>::FromIdentity();
-			//MassProperty.RotationOfMass = Chaos::TransformToLocalSpace<float, 3>(MassProperty.InertiaTensor);
-
-			MassProperties.Add(MassProperty);
-			TotalMass += Mass;
-		}
-		else if (ImplicitObject->GetType() == Chaos::ImplicitObjectType::Cylinder)
-		{
-			const Chaos::TCylinder<float> * Cylinder = ImplicitObject->template GetObject<Chaos::TCylinder<float>>();
-			float Mass = Cylinder->GetVolume()*InDensityKGPerCM;
-
-			Chaos::TMassProperties<float, 3> MassProperty;
-			MassProperty.CenterOfMass = Chaos::TVector<float, 3>(0);
-			MassProperty.Volume = Cylinder->GetVolume();
-			MassProperty.InertiaTensor = Cylinder->GetInertiaTensor(Mass)*ScaleM;
-			MassProperty.RotationOfMass = Chaos::TRotation<float, 3>::FromIdentity();
-			//MassProperty.RotationOfMass = Chaos::TransformToLocalSpace<float, 3>(MassProperty.InertiaTensor);
-
-			MassProperties.Add(MassProperty);
-			TotalMass += Mass;
-		}
-		else if (ImplicitObject->GetType() == Chaos::ImplicitObjectType::TaperedCylinder)
-		{
-			const Chaos::TTaperedCylinder<float> * TaperedCylinder = ImplicitObject->template GetObject<Chaos::TTaperedCylinder<float>>();
-			float Mass = TaperedCylinder->GetVolume()*InDensityKGPerCM;
-
-			Chaos::TMassProperties<float, 3> MassProperty;
-			MassProperty.CenterOfMass = Chaos::TVector<float, 3>(0);
-			MassProperty.Volume = TaperedCylinder->GetVolume();
-			MassProperty.InertiaTensor = TaperedCylinder->GetInertiaTensor(Mass)*ScaleM;
-			MassProperty.RotationOfMass = Chaos::TRotation<float, 3>::FromIdentity();
-			//MassProperty.RotationOfMass = Chaos::TransformToLocalSpace<float, 3>(MassProperty.InertiaTensor);
-
-			MassProperties.Add(MassProperty);
-			TotalMass += Mass;
-		}
-		else if (ImplicitObject->GetType() == Chaos::ImplicitObjectType::Convex)
-		{
-			// @todo: until this is actually used by anything that matters, just keep it simple. 
-			const Chaos::TConvex<float,3> * Convex = ImplicitObject->template GetObject<Chaos::TConvex<float,3>>();
-			float Mass = Convex->BoundingBox().GetVolume()*InDensityKGPerCM;
-
-			Chaos::TMassProperties<float, 3> MassProperty;
-			MassProperty.CenterOfMass = Chaos::TVector<float, 3>(0);
-			MassProperty.Volume = Convex->BoundingBox().GetVolume();
-			MassProperty.InertiaTensor = Chaos::TBox<float, 3>::GetInertiaTensor(Mass, Convex->BoundingBox().Extents());
-			MassProperty.RotationOfMass = Chaos::TRotation<float, 3>::FromIdentity();
-			//MassProperty.RotationOfMass = Chaos::TransformToLocalSpace<float, 3>(MassProperty.InertiaTensor);
-
-			MassProperties.Add(MassProperty);
-			TotalMass += Mass;
-		}
-		else if (ImplicitObject->GetType() == Chaos::ImplicitObjectType::TriangleMesh)
+		const Chaos::TMassProperties<float, 3> MassProps = Chaos::CastHelper(*ImplicitObject, [&TotalMass, InDensityKGPerCM](const auto& Object)
 		{
-			// @todo: until this is actually used by anything that matters, just keep it simple. 
-			const Chaos::TTriangleMeshImplicitObject<float> * TriangleMeshImplicitObject = ImplicitObject->template GetObject<Chaos::TTriangleMeshImplicitObject<float>>();
-			float Mass = TriangleMeshImplicitObject->BoundingBox().GetVolume()*InDensityKGPerCM;
-
-			Chaos::TMassProperties<float, 3> MassProperty;
-			MassProperty.CenterOfMass = Chaos::TVector<float, 3>(0);
-			MassProperty.Volume = TriangleMeshImplicitObject->BoundingBox().GetVolume();
-			MassProperty.InertiaTensor = Chaos::TBox<float, 3>::GetInertiaTensor(Mass, TriangleMeshImplicitObject->BoundingBox().Extents());
-			MassProperty.RotationOfMass = Chaos::TRotation<float, 3>::FromIdentity();
-			//MassProperty.RotationOfMass = Chaos::TransformToLocalSpace<float, 3>(MassProperty.InertiaTensor);
-
-			MassProperties.Add(MassProperty);
+			Chaos::TMassProperties<float, 3> MassProperties;
+			MassProperties.Volume = Object.GetVolume();
+			const float Mass = MassProperties.Volume * InDensityKGPerCM;
 			TotalMass += Mass;
-		}
-		else if (const auto ScaledObject = Chaos::TImplicitObjectScaledGeneric<float,3>::AsScaled(*ImplicitObject))
-		{
-			if (ensureMsgf(ScaledObject->GetUnscaledObject(), TEXT("Error : Null internal MObject type within TImplicitObjectScaled.")))
-			{
-				Chaos::TRigidTransform<float, 3> ScaledWorldTransform(WorldTransform.GetTranslation(), WorldTransform.GetRotation(), ScaledObject->GetScale());
-				CalculateMassPropertiesOfImplicitType(MassProperties, TotalMass, ScaledWorldTransform, ScaledObject->GetUnscaledObject(), InDensityKGPerCM);
-			}
-		}
-		else if (ImplicitObject->GetType() == Chaos::ImplicitObjectType::Union)
-		{
-			const Chaos::TImplicitObjectUnion<float,3> * ImplicitUnion = ImplicitObject->template GetObject<Chaos::TImplicitObjectUnion<float,3>>();
-
-			for (const TUniquePtr < Chaos::FImplicitObject > & ImplicitSubObject : ImplicitUnion->GetObjects())
-			{
-				CalculateMassPropertiesOfImplicitType(MassProperties, TotalMass, WorldTransform, ImplicitSubObject.Get(), InDensityKGPerCM);
-			}
-		}
-		else
-		{
-			// @todo : Enable when ready.
-			//ensureMsgf(false, TEXT("Unsupported implicit object type (%s)."), *(ImplicitObject->GetTypeName().ToString()) );
-		}
+			MassProperties.InertiaTensor = Object.GetInertiaTensor(Mass);
+			MassProperties.CenterOfMass = Object.GetCenterOfMass();
+			MassProperties.RotationOfMass = Chaos::TRotation<float, 3>::FromIdentity();
+			return MassProperties;
+		});
+		MassProperties.Add(MassProps);
 	}
 }
 
diff --git a/Engine/Source/Runtime/Experimental/Chaos/Public/Chaos/Convex.h b/Engine/Source/Runtime/Experimental/Chaos/Public/Chaos/Convex.h
--- a/Engine/Source/Runtime/Experimental/Chaos/Public/Chaos/Convex.h
+++ b/Engine/Source/Runtime/Experimental/Chaos/Public/Chaos/Convex.h
@@ -200,6 +200,24 @@ namespace Chaos
 			return Planes;
 		}
 
+		const FReal GetVolume() const
+		{
+			// TODO: More precise volume!
+			return LocalBoundingBox.GetVolume();
+		}
+
+		const FMatrix33 GetInertiaTensor(const FReal Mass) const
+		{
+			// TODO: More precise inertia!
+			return LocalBoundingBox.GetInertiaTensor(Mass);
+		}
+
+		const FVec3 GetCenterOfMass() const
+		{
+			// TODO: Actually compute this!
+			return FVec3(0.f);
+		}
+
 		virtual uint32 GetTypeHash() const override
 		{
 			uint32 Result = LocalBoundingBox.GetTypeHash();
diff --git a/Engine/Source/Runtime/Experimental/Chaos/Public/Chaos/ImplicitObjectScaled.h b/Engine/Source/Runtime/Experimental/Chaos/Public/Chaos/ImplicitObjectScaled.h
--- a/Engine/Source/Runtime/Experimental/Chaos/Public/Chaos/ImplicitObjectScaled.h
+++ b/Engine/Source/Runtime/Experimental/Chaos/Public/Chaos/ImplicitObjectScaled.h
@@ -448,6 +448,24 @@ public:
 
 	virtual const TBox<T, d>& BoundingBox() const override { return MLocalBoundingBox; }
 
+	const FReal GetVolume() const
+	{
+		// TODO: More precise volume!
+		return BoundingBox().GetVolume();
+	}
+
+	const FMatrix33 GetInertiaTensor(const FReal Mass) const
+	{
+		// TODO: More precise inertia!
+		return BoundingBox().GetInertiaTensor(Mass);
+	}
+
+	const FVec3 GetCenterOfMass() const
+	{
+		// TODO: Actually compute this!
+		return FVec3(0.f);
+	}
+
 	const ObjectType Object() const { return MObject; }
 	
 	virtual void Serialize(FChaosArchive& Ar) override
