Fixed crash with guard when non valid FDistanceFieldVolumeTexture is dereferenced in FDistanceFieldVolumeTextureAtlas::RemoveAllocation.

#rnx

diff --git a/Engine/Source/Runtime/Engine/Private/DistanceFieldAtlas.cpp b/Engine/Source/Runtime/Engine/Private/DistanceFieldAtlas.cpp
--- a/Engine/Source/Runtime/Engine/Private/DistanceFieldAtlas.cpp
+++ b/Engine/Source/Runtime/Engine/Private/DistanceFieldAtlas.cpp
@@ -271,12 +271,9 @@ void FDistanceFieldVolumeTextureAtlas::RemoveAllocation(FDistanceFieldVolumeText
 	InitializeIfNeeded();
 	PendingAllocations.Remove(Texture);
 
-	FIntVector Size = Texture->VolumeData.Size;
-	int PixelAreaSize = Size.X * Size.Y * Size.Z;
-
 	if (FailedAllocations.Remove(Texture) > 0)
 	{
-		FailedAllocatedPixels -= PixelAreaSize;
+		FailedAllocatedPixels -= Texture->VolumeData.Size.X * Texture->VolumeData.Size.Y * Texture->VolumeData.Size.Z;;
 	}
 
 	if (!CurrentAllocations.Contains(Texture))
@@ -284,6 +281,9 @@ void FDistanceFieldVolumeTextureAtlas::RemoveAllocation(FDistanceFieldVolumeText
 		return;
 	}
 
+	FIntVector Size = Texture->VolumeData.Size;
+	int PixelAreaSize = Size.X * Size.Y * Size.Z;
+
 	const FIntVector Min = Texture->GetAllocationMin();
 	verify(BlockAllocator.RemoveElement(Min.X, Min.Y, Min.Z, Size.X, Size.Y, Size.Z));
 	CurrentAllocations.Remove(Texture);
