Fixed game mode HISM scale / rotation updates. GIsEditor check (introduced in CL 4037080) was preventing instance GPU buffer update in case when full tree update wasn't required (instance rotation or scale change).

diff --git a/Engine/Source/Runtime/Engine/Private/HierarchicalInstancedStaticMesh.cpp b/Engine/Source/Runtime/Engine/Private/HierarchicalInstancedStaticMesh.cpp
--- a/Engine/Source/Runtime/Engine/Private/HierarchicalInstancedStaticMesh.cpp
+++ b/Engine/Source/Runtime/Engine/Private/HierarchicalInstancedStaticMesh.cpp
@@ -2983,8 +2983,7 @@ FPrimitiveSceneProxy* UHierarchicalInstancedStaticMeshComponent::CreateSceneProx
 		check(InstancingRandomSeed != 0);
 
 		// if instance data was modified, update GPU copy
-		// generally happens only in editor 
-		if (GIsEditor && InstanceUpdateCmdBuffer.NumInlineCommands() > 0)
+		if (InstanceUpdateCmdBuffer.NumInlineCommands() > 0)
 		{
 			PerInstanceRenderData->UpdateFromCommandBuffer(InstanceUpdateCmdBuffer);
 		}
