Cameras: Add API to get the list of currently playing camera shakes. Fix const-correctness. Allow running without a camera controller.

diff --git a/Engine/Source/Runtime/Engine/Classes/Camera/CameraModifier_CameraShake.h b/Engine/Source/Runtime/Engine/Classes/Camera/CameraModifier_CameraShake.h
--- a/Engine/Source/Runtime/Engine/Classes/Camera/CameraModifier_CameraShake.h
+++ b/Engine/Source/Runtime/Engine/Classes/Camera/CameraModifier_CameraShake.h
@@ -36,7 +36,7 @@ struct FActiveCameraShakeInfo
 	UCameraShake* ShakeInstance;
 
 	UPROPERTY()
-	UCameraShakeSourceComponent* ShakeSource;
+	const UCameraShakeSourceComponent* ShakeSource;
 };
 
 struct FAddCameraShakeParams
@@ -44,12 +44,12 @@ struct FAddCameraShakeParams
 	float Scale;
 	ECameraAnimPlaySpace::Type PlaySpace;
 	FRotator UserPlaySpaceRot;
-	UCameraShakeSourceComponent* SourceComponent;
+	const UCameraShakeSourceComponent* SourceComponent;
 
 	FAddCameraShakeParams() 
 		: Scale(1.f), PlaySpace(ECameraAnimPlaySpace::CameraLocal), UserPlaySpaceRot(FRotator::ZeroRotator), SourceComponent(nullptr)
 	{}
-	FAddCameraShakeParams(float InScale, ECameraAnimPlaySpace::Type InPlaySpace = ECameraAnimPlaySpace::CameraLocal, FRotator InUserPlaySpaceRot = FRotator::ZeroRotator, UCameraShakeSourceComponent* InSourceComponent = nullptr)
+	FAddCameraShakeParams(float InScale, ECameraAnimPlaySpace::Type InPlaySpace = ECameraAnimPlaySpace::CameraLocal, FRotator InUserPlaySpaceRot = FRotator::ZeroRotator, const UCameraShakeSourceComponent* InSourceComponent = nullptr)
 		: Scale(InScale), PlaySpace(InPlaySpace), UserPlaySpaceRot(InUserPlaySpaceRot), SourceComponent(InSourceComponent)
 	{}
 };
@@ -86,6 +86,12 @@ public:
 	{
 		return AddCameraShake(NewShake, FAddCameraShakeParams(Scale, PlaySpace, UserPlaySpaceRot));
 	}
+
+	/**
+	 * Returns a list of currently active camera shakes.
+	 * @param ActiveCameraShakes - The array to fill up with shake information.
+	 */
+	virtual void GetActiveCameraShakes(TArray<FActiveCameraShakeInfo>& ActiveCameraShakes) const;
 	
 	/**
 	 * Stops and removes the camera shake of the given class from the camera.
@@ -100,7 +106,7 @@ public:
 	 */
 	virtual void RemoveAllCameraShakesOfClass(TSubclassOf<UCameraShake> ShakeClass, bool bImmediately = true);
 
-	virtual void RemoveAllCameraShakesFromSource(UCameraShakeSourceComponent* SourceComponent, bool bImmediately = true);
+	virtual void RemoveAllCameraShakesFromSource(const UCameraShakeSourceComponent* SourceComponent, bool bImmediately = true);
 
 	/** 
 	 * Stops and removes all camera shakes from the camera. 
diff --git a/Engine/Source/Runtime/Engine/Private/Camera/CameraModifier_CameraShake.cpp b/Engine/Source/Runtime/Engine/Private/Camera/CameraModifier_CameraShake.cpp
--- a/Engine/Source/Runtime/Engine/Private/Camera/CameraModifier_CameraShake.cpp
+++ b/Engine/Source/Runtime/Engine/Private/Camera/CameraModifier_CameraShake.cpp
@@ -81,10 +81,10 @@ UCameraShake* UCameraModifier_CameraShake::AddCameraShake(TSubclassOf<UCameraSha
 	if (ShakeClass != nullptr)
 	{
 		float Scale = Params.Scale;
-		UCameraShakeSourceComponent* SourceComponent = Params.SourceComponent;
+		const UCameraShakeSourceComponent* SourceComponent = Params.SourceComponent;
 
 		// Adjust for splitscreen
-		if (GEngine->IsSplitScreen(CameraOwner->GetWorld()))
+		if (CameraOwner != nullptr && GEngine->IsSplitScreen(CameraOwner->GetWorld()))
 		{
 			Scale *= SplitScreenShakeScale;
 		}
@@ -180,6 +180,11 @@ UCameraShake* UCameraModifier_CameraShake::ReclaimShakeFromExpiredPool(TSubclass
 	return nullptr;
 }
 
+void UCameraModifier_CameraShake::GetActiveCameraShakes(TArray<FActiveCameraShakeInfo>& ActiveCameraShakes) const
+{
+	ActiveCameraShakes.Append(ActiveShakes);
+}
+
 void UCameraModifier_CameraShake::RemoveCameraShake(UCameraShake* ShakeInst, bool bImmediately)
 {
 	for (int32 i = 0; i < ActiveShakes.Num(); ++i)
@@ -216,7 +221,7 @@ void UCameraModifier_CameraShake::RemoveAllCameraShakesOfClass(TSubclassOf<UCame
 	}
 }
 
-void UCameraModifier_CameraShake::RemoveAllCameraShakesFromSource(UCameraShakeSourceComponent* SourceComponent, bool bImmediately)
+void UCameraModifier_CameraShake::RemoveAllCameraShakesFromSource(const UCameraShakeSourceComponent* SourceComponent, bool bImmediately)
 {
 	for (int32 i = ActiveShakes.Num() - 1; i >= 0; --i)
 	{
