Adding option to record game specific per frame replay data (submitting on behalf of brian.bekich)

diff --git a/Engine/Source/Runtime/Engine/Classes/Engine/DemoNetDriver.h b/Engine/Source/Runtime/Engine/Classes/Engine/DemoNetDriver.h
--- a/Engine/Source/Runtime/Engine/Classes/Engine/DemoNetDriver.h
+++ b/Engine/Source/Runtime/Engine/Classes/Engine/DemoNetDriver.h
@@ -159,10 +159,11 @@ struct FLevelNameAndTime
 
 enum class EReplayHeaderFlags : uint32
 {
-	None				= 0,
-	ClientRecorded		= (1 << 0),
-	HasStreamingFixes	= (1 << 1),
-	DeltaCheckpoints	= (1 << 2),
+	None					= 0,
+	ClientRecorded			= (1 << 0),
+	HasStreamingFixes		= (1 << 1),
+	DeltaCheckpoints		= (1 << 2),
+	GameSpecificFrameData	= (1 << 3),
 };
 
 ENUM_CLASS_FLAGS(EReplayHeaderFlags);
@@ -372,6 +373,14 @@ struct FMulticastRecordOptions
 	bool bClientSkip;
 };
 
+enum class EWriteDemoFrameFlags : uint32
+{
+	None				= 0,
+	SkipGameSpecific	= (1 << 0),
+};
+
+ENUM_CLASS_FLAGS(EWriteDemoFrameFlags);
+
 /**
  * Simulated network driver for recording and playing back game sessions.
  */
@@ -603,6 +612,9 @@ private:
 	UPROPERTY(config)
 	TArray<FMulticastRecordOptions> MulticastRecordOptions;
 
+	/** PlaybackFrames are used to buffer per frame data up when we read a demo frame, which we can then process when the time is right */
+	TMap<float, TMap<FString, TArray<uint8>>> PlaybackFrames;
+
 public:
 
 	// UNetDriver interface.
@@ -744,8 +756,15 @@ public:
 		return ShouldSkipPlaybackPacket(PlaybackPacket) ||
 				ProcessPacket(PlaybackPacket.Data.GetData(), PlaybackPacket.Data.Num());
 	}
+	
+	void WriteDemoFrameFromQueuedDemoPackets(FArchive& Ar, TArray<FQueuedDemoPacket>& QueuedPackets, float FrameTime, EWriteDemoFrameFlags Flags);
+
+	UE_DEPRECATED(4.25, "WriteDemoFrameFromQueuedDemoPackets now takes an additional flag value")
+	void WriteDemoFrameFromQueuedDemoPackets(FArchive& Ar, TArray<FQueuedDemoPacket>& QueuedPackets, float FrameTime)
+	{
+		WriteDemoFrameFromQueuedDemoPackets(Ar, QueuedPackets, FrameTime, EWriteDemoFrameFlags::None);
+	}
 
-	void WriteDemoFrameFromQueuedDemoPackets(FArchive& Ar, TArray<FQueuedDemoPacket>& QueuedPackets, float FrameTime);
 	void WritePacket(FArchive& Ar, uint8* Data, int32 Count);
 
 	void TickDemoPlayback(float DeltaSeconds);
@@ -905,6 +924,12 @@ public:
 		return bHasDeltaCheckpoints;
 	}
 
+	/** Returns whether or not this replay was recorded / is playing with the game specific per frame data feature. */
+	bool HasGameSpecificFrameData() const
+	{
+		return bHasGameSpecificFrameData;
+	}
+
 	/**
 	 * Gets the actively recording or playback replay (stream) name.
 	 * Note, this will be empty when not recording or playing back.
@@ -1015,6 +1040,9 @@ private:
 	// Checkpoints are delta compressed
 	bool bHasDeltaCheckpoints;
 
+	// Allow appending per frame game specific data
+	bool bHasGameSpecificFrameData;
+
 	// Levels that are currently pending for fast forward.
 	// Using raw pointers, because we manually keep when levels are added and removed.
 	TSet<class ULevel*> LevelsPendingFastForward;
diff --git a/Engine/Source/Runtime/Engine/Private/DemoNetDriver.cpp b/Engine/Source/Runtime/Engine/Private/DemoNetDriver.cpp
--- a/Engine/Source/Runtime/Engine/Private/DemoNetDriver.cpp
+++ b/Engine/Source/Runtime/Engine/Private/DemoNetDriver.cpp
@@ -82,6 +82,7 @@ static FAutoConsoleVariableRef CVarDemoSaveRollbackActorState( TEXT( "demo.SaveR
 static TAutoConsoleVariable<int32> CVarWithLevelStreamingFixes(TEXT("demo.WithLevelStreamingFixes"), 0, TEXT("If 1, provides fixes for level streaming (but breaks backwards compatibility)."));
 static TAutoConsoleVariable<int32> CVarWithDemoTimeBurnIn(TEXT("demo.WithTimeBurnIn"), 0, TEXT("If true, adds an on screen message with the current DemoTime and Changelist."));
 static TAutoConsoleVariable<int32> CVarWithDeltaCheckpoints(TEXT("demo.WithDeltaCheckpoints"), 0, TEXT("If true, record checkpoints as a delta from the previous checkpoint."));
+static TAutoConsoleVariable<int32> CVarWithGameSpecificFrameData(TEXT("demo.WithGameSpecificFrameData"), 0, TEXT("If true, allow game specific data to be recorded with each demo frame."));
 
 #if !(UE_BUILD_SHIPPING || UE_BUILD_TEST)
 static TAutoConsoleVariable<int32> CVarDemoForceFailure( TEXT( "demo.ForceFailure" ), 0, TEXT( "" ) );
@@ -830,6 +831,8 @@ void UDemoNetDriver::ResetDemoState()
 
 	ExternalDataToObjectMap.Empty();
 	PlaybackPackets.Empty();
+	PlaybackFrames.Empty();
+
 	ClearLevelStreamingState();
 }
 
@@ -958,6 +961,8 @@ bool UDemoNetDriver::ReadPlaybackDemoHeader(FString& Error)
 		return false;
 	}
 
+	bHasGameSpecificFrameData = EnumHasAnyFlags(PlaybackDemoHeader.HeaderFlags, EReplayHeaderFlags::GameSpecificFrameData);
+
 	// Set network version on connection
 	ServerConnection->EngineNetworkProtocolVersion = PlaybackDemoHeader.EngineNetworkProtocolVersion;
 	ServerConnection->GameNetworkProtocolVersion = PlaybackDemoHeader.GameNetworkProtocolVersion;
@@ -1070,6 +1075,7 @@ bool UDemoNetDriver::InitListen(FNetworkNotify* InNotify, FURL& ListenURL, bool
 	// During playback the CVars won't be used. Instead, we'll rely on the DemoPacketHeader value.
 	bHasLevelStreamingFixes = !!CVarWithLevelStreamingFixes.GetValueOnAnyThread();
 	bHasDeltaCheckpoints = !!CVarWithDeltaCheckpoints.GetValueOnAnyThread() && ReplayStreamer->IsCheckpointTypeSupported(EReplayCheckpointType::Delta);
+	bHasGameSpecificFrameData = !!CVarWithGameSpecificFrameData.GetValueOnAnyThread();
 
 	// Recording, local machine is server, demo stream acts "as if" it's a client.
 	UDemoNetConnection* Connection = NewObject<UDemoNetConnection>();
@@ -1299,6 +1305,11 @@ bool UDemoNetDriver::WriteNetworkDemoHeader(FString& Error)
 		DemoHeader.HeaderFlags |= EReplayHeaderFlags::DeltaCheckpoints;
 	}
 
+	if (HasGameSpecificFrameData())
+	{
+		DemoHeader.HeaderFlags |= EReplayHeaderFlags::GameSpecificFrameData;
+	}
+
 	DemoHeader.Guid = FGuid::NewGuid();
 
 	// Write the header
@@ -2166,7 +2177,7 @@ void UDemoNetDriver::TickCheckpoint()
 						TGuardValue<uint32> NumLevelsAddedThisFrameGuard(NumLevelsAddedThisFrame, AllLevelStatuses.Num());
 
 						// Write out all of the queued up packets generated while saving the checkpoint
-						WriteDemoFrameFromQueuedDemoPackets( *CheckpointArchive, ClientConnection->QueuedCheckpointPackets, static_cast<float>(LastCheckpointTime) );
+						WriteDemoFrameFromQueuedDemoPackets( *CheckpointArchive, ClientConnection->QueuedCheckpointPackets, static_cast<float>(LastCheckpointTime), EWriteDemoFrameFlags::SkipGameSpecific );
 
 						CheckpointSaveContext.CheckpointSaveState = ECheckpointSaveState_Finalize;
 					}
@@ -2838,7 +2849,7 @@ void UDemoNetDriver::TickDemoRecordFrame(float DeltaSeconds)
 
 	FlushNetChecked(*ClientConnection);
 
-	WriteDemoFrameFromQueuedDemoPackets(*FileAr, ClientConnection->QueuedDemoPackets, DemoCurrentTime);
+	WriteDemoFrameFromQueuedDemoPackets(*FileAr, ClientConnection->QueuedDemoPackets, DemoCurrentTime, EWriteDemoFrameFlags::None);
 	
 	AdjustConsiderTime((float)Params.NumActorsReplicated / (float)NumPrioritizedActors);
 }
@@ -3216,6 +3227,27 @@ bool UDemoNetDriver::ReadDemoFrameIntoPlaybackPackets(FArchive& Ar, TArray<FPlay
 		Ar.Seek(Ar.Tell() + SkipExternalOffset);
 	}
 
+	FArchivePos SkipGameSpecificOffset = 0;
+	if (HasGameSpecificFrameData())
+	{
+		Ar << SkipGameSpecificOffset;
+
+		if (!bForLevelFastForward)
+		{
+			FDemoFrameDataMap Data;
+			Ar << Data;
+
+			if (Data.Num() > 0)
+			{
+				PlaybackFrames.Emplace(TimeSeconds, MoveTemp(Data));
+			}
+		}
+		else
+		{
+			Ar.Seek(Ar.Tell() + SkipGameSpecificOffset);
+		}
+	}
+
 	{
 		DECLARE_SCOPE_CYCLE_COUNTER(TEXT("Demo_ReadPackets"), Demo_ReadPackets, STATGROUP_Net);
 
@@ -3492,6 +3524,8 @@ void UDemoNetDriver::ProcessAllPlaybackPackets()
 {
 	ProcessPlaybackPackets(PlaybackPackets);
 	PlaybackPackets.Empty();
+	// this call is used for checkpoint loading, so not dealing with per frame data
+	PlaybackFrames.Empty();
 }
 
 void UDemoNetDriver::ProcessPlaybackPackets(TArrayView<FPlaybackPacket> Packets)
@@ -3529,7 +3563,7 @@ bool UDemoNetDriver::ProcessPacket(const uint8* Data, int32 Count)
 	return true;
 }
 
-void UDemoNetDriver::WriteDemoFrameFromQueuedDemoPackets(FArchive& Ar, TArray<FQueuedDemoPacket>& QueuedPackets, float FrameTime)
+void UDemoNetDriver::WriteDemoFrameFromQueuedDemoPackets(FArchive& Ar, TArray<FQueuedDemoPacket>& QueuedPackets, float FrameTime, EWriteDemoFrameFlags Flags)
 {
 	Ar << CurrentLevelIndex;
 	
@@ -3581,6 +3615,19 @@ void UDemoNetDriver::WriteDemoFrameFromQueuedDemoPackets(FArchive& Ar, TArray<FQ
 		SaveExternalData( Ar );
 	}
 
+	if (HasGameSpecificFrameData())
+	{
+		FScopedStoreArchiveOffset ScopedOffset(Ar);
+
+		if (!EnumHasAnyFlags(Flags, EWriteDemoFrameFlags::SkipGameSpecific))
+		{
+			FDemoFrameDataMap Data;
+			FNetworkReplayDelegates::OnWriteGameSpecificFrameData.Broadcast(GetWorld(), FrameTime, Data);
+
+			Ar << Data;
+		}
+	}
+
 	for (FQueuedDemoPacket& DemoPacket : QueuedPackets)
 	{
 		if (HasLevelStreamingFixes())
@@ -3937,6 +3984,20 @@ void UDemoNetDriver::TickDemoPlayback(float DeltaSeconds)
 			PlaybackPackets.RemoveAt(0, PlaybackPacketIndex);
 			PlaybackPacketIndex = 0;
 		}
+
+		// Process playback frames
+		for (auto It = PlaybackFrames.CreateIterator(); It; ++It)
+		{
+			if (It.Key() <= DemoCurrentTime)
+			{
+				if (!bIsFastForwarding)
+				{
+					FNetworkReplayDelegates::OnProcessGameSpecificFrameData.Broadcast(World, It.Key(), It.Value());
+				}
+	
+				It.RemoveCurrent();
+			}
+		}
 	}
 
 	// Finalize any fast forward stuff that needs to happen
@@ -5145,6 +5206,7 @@ bool UDemoNetDriver::LoadCheckpoint(const FGotoResult& GotoResult)
 
 	ExternalDataToObjectMap.Empty();
 	PlaybackPackets.Empty();
+	PlaybackFrames.Empty();
 
 	// Destroy startup actors that need to rollback via being destroyed and re-created
 	for (FActorIterator It(GetWorld()); It; ++It)
@@ -5447,6 +5509,7 @@ bool UDemoNetDriver::LoadCheckpoint(const FGotoResult& GotoResult)
 			}
 
 			PlaybackPackets.Empty();
+			PlaybackFrames.Empty();
 
 			if (DemoConnection)
 			{
@@ -6625,6 +6688,14 @@ void UDemoNetDriver::Serialize(FArchive& Ar)
 			}
 		);
 
+		GRANULAR_NETWORK_MEMORY_TRACKING_TRACK("PlaybackFrames",
+			PlaybackFrames.CountBytes(Ar);
+			for (auto& Frame : PlaybackFrames)
+			{
+				Frame.Value.CountBytes(Ar);
+			}
+		);
+
 		GRANULAR_NETWORK_MEMORY_TRACKING_TRACK("UniqueStreamingLevels", UniqueStreamingLevels.CountBytes(Ar));
 		GRANULAR_NETWORK_MEMORY_TRACKING_TRACK("NewStreamingLevelsThisFrame", NewStreamingLevelsThisFrame.CountBytes(Ar));
 		GRANULAR_NETWORK_MEMORY_TRACKING_TRACK("NonQueuedGUIDsForScrubbing", NonQueuedGUIDsForScrubbing.CountBytes(Ar));
diff --git a/Engine/Source/Runtime/Engine/Private/UnrealNetwork.cpp b/Engine/Source/Runtime/Engine/Private/UnrealNetwork.cpp
--- a/Engine/Source/Runtime/Engine/Private/UnrealNetwork.cpp
+++ b/Engine/Source/Runtime/Engine/Private/UnrealNetwork.cpp
@@ -5,6 +5,8 @@
 FPreReplayScrub FNetworkReplayDelegates::OnPreScrub;
 FOnWriteGameSpecificDemoHeader FNetworkReplayDelegates::OnWriteGameSpecificDemoHeader;
 FOnProcessGameSpecificDemoHeader FNetworkReplayDelegates::OnProcessGameSpecificDemoHeader;
+FOnWriteGameSpecificFrameData FNetworkReplayDelegates::OnWriteGameSpecificFrameData;
+FOnProcessGameSpecificFrameData FNetworkReplayDelegates::OnProcessGameSpecificFrameData;
 
 // ----------------------------------------------------------------
 
diff --git a/Engine/Source/Runtime/Engine/Public/Net/UnrealNetwork.h b/Engine/Source/Runtime/Engine/Public/Net/UnrealNetwork.h
--- a/Engine/Source/Runtime/Engine/Public/Net/UnrealNetwork.h
+++ b/Engine/Source/Runtime/Engine/Public/Net/UnrealNetwork.h
@@ -37,13 +37,14 @@ inline int32 MakeRelative( int32 Value, int32 Reference, int32 Max )
 }
 
 DECLARE_MULTICAST_DELEGATE_OneParam(FPreActorDestroyReplayScrub, AActor*);
-
 DECLARE_MULTICAST_DELEGATE_OneParam(FPreReplayScrub, UWorld*);
-
 DECLARE_MULTICAST_DELEGATE_OneParam(FOnWriteGameSpecificDemoHeader, TArray<FString>& /*GameSpecificData*/);
-
 DECLARE_MULTICAST_DELEGATE_TwoParams(FOnProcessGameSpecificDemoHeader, const TArray<FString>& /*GameSpecificData*/, FString& /*Error*/);
 
+typedef TMap<FString, TArray<uint8>> FDemoFrameDataMap;
+DECLARE_MULTICAST_DELEGATE_ThreeParams(FOnWriteGameSpecificFrameData, UWorld* /*World*/, float /*FrameTime*/, FDemoFrameDataMap& /*Data*/);
+DECLARE_MULTICAST_DELEGATE_ThreeParams(FOnProcessGameSpecificFrameData, UWorld* /*World*/, float /*FrameTime*/, const FDemoFrameDataMap& /*Data*/);
+
 struct ENGINE_API FNetworkReplayDelegates
 {
 	/** global delegate called one time prior to scrubbing */
@@ -52,6 +53,10 @@ struct ENGINE_API FNetworkReplayDelegates
 	/** Game specific demo headers */
 	static FOnWriteGameSpecificDemoHeader OnWriteGameSpecificDemoHeader;
 	static FOnProcessGameSpecificDemoHeader OnProcessGameSpecificDemoHeader;
+
+	/** Game specific per frame data */
+	static FOnWriteGameSpecificFrameData OnWriteGameSpecificFrameData;
+	static FOnProcessGameSpecificFrameData OnProcessGameSpecificFrameData;
 };
 
 /**
