Fix for UI materials not updating, and not having shaders available in -game

diff --git a/Engine/Source/Editor/MaterialEditor/Private/MaterialEditor.cpp b/Engine/Source/Editor/MaterialEditor/Private/MaterialEditor.cpp
--- a/Engine/Source/Editor/MaterialEditor/Private/MaterialEditor.cpp
+++ b/Engine/Source/Editor/MaterialEditor/Private/MaterialEditor.cpp
@@ -4522,6 +4522,7 @@ void FMaterialEditor::NotifyPostChange( const FPropertyChangedEvent& PropertyCha
 			Material->MaterialGraph->RebuildGraph();
 			TArray<TWeakObjectPtr<UObject>> SelectedObjects = MaterialDetailsView->GetSelectedObjects();
 			MaterialDetailsView->SetObjects( SelectedObjects, true );
+			SetPreviewMaterial(Material);
 
 			if (ExpressionPreviewMaterial)
 			{
diff --git a/Engine/Source/Runtime/Engine/Private/LocalVertexFactory.cpp b/Engine/Source/Runtime/Engine/Private/LocalVertexFactory.cpp
--- a/Engine/Source/Runtime/Engine/Private/LocalVertexFactory.cpp
+++ b/Engine/Source/Runtime/Engine/Private/LocalVertexFactory.cpp
@@ -184,12 +184,9 @@ void FLocalVertexFactoryShaderParameters::GetElementShaderBindings(
 bool FLocalVertexFactory::ShouldCompilePermutation(EShaderPlatform Platform, const class FMaterial* Material, const class FShaderType* ShaderType)
 {
 	// Only compile this permutation inside the editor - it's not applicable in games, but occasionally the editor needs it.
-	if (!GIsEditor)
+	if (Material->GetMaterialDomain() == MD_UI)
 	{
-		if (Material->GetMaterialDomain() == MD_UI)
-		{
-			return false;
-		}
+		return !!WITH_EDITOR;
 	}
 
 	return true; 
