Copy up from DS to Main

diff --git a/Engine/Source/Developer/ScriptDisassembler/Private/ScriptDisassemblerModule.cpp b/Engine/Source/Developer/ScriptDisassembler/Private/ScriptDisassemblerModule.cpp
deleted file mode 100644
diff --git a/Engine/Source/Developer/ScriptDisassembler/ScriptDisassembler.Build.cs b/Engine/Source/Developer/ScriptDisassembler/ScriptDisassembler.Build.cs
--- a/Engine/Source/Developer/ScriptDisassembler/ScriptDisassembler.Build.cs
+++ b/Engine/Source/Developer/ScriptDisassembler/ScriptDisassembler.Build.cs
@@ -6,6 +6,8 @@ namespace UnrealBuildTool.Rules
 	{
 		public ScriptDisassembler(ReadOnlyTargetRules Target) : base(Target)
         {
+			bRequiresImplementModule = false;
+
 			PublicDependencyModuleNames.AddRange(new string[]
 			{
 				"Core",
diff --git a/Engine/Source/Runtime/Engine/Public/Physics/ImmediatePhysics/ImmediatePhysicsPhysX/ImmediatePhysicsActorHandle_PhysX.h b/Engine/Source/Runtime/Engine/Public/Physics/ImmediatePhysics/ImmediatePhysicsPhysX/ImmediatePhysicsActorHandle_PhysX.h
--- a/Engine/Source/Runtime/Engine/Public/Physics/ImmediatePhysics/ImmediatePhysicsPhysX/ImmediatePhysicsActorHandle_PhysX.h
+++ b/Engine/Source/Runtime/Engine/Public/Physics/ImmediatePhysics/ImmediatePhysicsPhysX/ImmediatePhysicsActorHandle_PhysX.h
@@ -2,247 +2,7 @@
 
 #pragma once
 
-#if WITH_PHYSX && PHYSICS_INTERFACE_PHYSX
-#include "PhysXPublic.h"
-#include "Physics/ImmediatePhysics/ImmediatePhysicsPhysX/ImmediatePhysicsSimulation_PhysX.h"
-#include "Engine/EngineTypes.h"
+#include "Physics/ImmediatePhysics/ImmediatePhysicsPhysX/ImmediatePhysicsActorHandle_PhysX.h"
+#include "Physics/ImmediatePhysics/ImmediatePhysicsChaos/ImmediatePhysicsActorHandle_Chaos.h"
 
-namespace ImmediatePhysics_PhysX
-{
-
-	/** handle associated with a physics actor. This is the proper way to read/write to the physics simulation */
-	struct ENGINE_API FActorHandle
-	{
-
-	public:
-		/** Sets the world transform.*/
-		void SetWorldTransform(const FTransform& WorldTM)
-		{
-			OwningSimulation.GetLowLevelBody(ActorDataIndex).body2World = U2PTransform(ActorToBody * WorldTM);
-		}
-
-		/** Make a body kinematic, or non-kinematic */
-		void SetIsKinematic(bool bKinematic)
-		{
-			OwningSimulation.SetIsKinematic(ActorDataIndex, bKinematic);
-		}
-
-		/** Is the actor kinematic */
-		bool GetIsKinematic() const
-		{
-			return OwningSimulation.GetIsKinematic(ActorDataIndex);
-		}
-
-		/** Gets the kinematic target (next transform) for the actor if one is set (check HasKinematicTarget() to see if a target is available) */
-		FImmediateKinematicTarget& GetKinematicTarget()
-		{
-			return OwningSimulation.GetKinematicTarget(ActorDataIndex);
-		}
-
-		/** Sets the kinematic target. This will affect velocities as expected*/
-		void SetKinematicTarget(const FTransform& WorldTM)
-		{
-			FImmediateKinematicTarget& KinematicTarget = OwningSimulation.GetKinematicTarget(ActorDataIndex);
-			KinematicTarget.BodyToWorld = U2PTransform(ActorToBody * WorldTM);
-			KinematicTarget.bTargetSet = true;
-		}
-
-		/** Does this actor have a kinematic target (next kinematic transform to be applied) */
-		bool HasKinematicTarget() const
-		{
-			FImmediateKinematicTarget& KinematicTarget = OwningSimulation.GetKinematicTarget(ActorDataIndex);
-			return KinematicTarget.bTargetSet;
-		}
-
-		/** Whether the body is simulating */
-		bool IsSimulated() const
-		{
-			return OwningSimulation.IsSimulated(ActorDataIndex);
-		}
-
-		/** Get the world transform */
-		FTransform GetWorldTransform() const
-		{
-			return ActorToBody.GetRelativeTransformReverse(P2UTransform(OwningSimulation.GetLowLevelBody(ActorDataIndex).body2World));
-		}
-
-		/** Set the linear velocity */
-		void SetLinearVelocity(const FVector& NewLinearVelocity)
-		{
-			OwningSimulation.GetLowLevelBody(ActorDataIndex).linearVelocity = U2PVector(NewLinearVelocity);
-		}
-
-		/** Get the linear velocity */
-		FVector GetLinearVelocity() const
-		{
-			return P2UVector(OwningSimulation.GetLowLevelBody(ActorDataIndex).linearVelocity);
-		}
-
-		/** Set the angular velocity */
-		void SetAngularVelocity(const FVector& NewAngularVelocity)
-		{
-			OwningSimulation.GetLowLevelBody(ActorDataIndex).angularVelocity = U2PVector(NewAngularVelocity);
-		}
-
-		/** Get the angular velocity */
-		FVector GetAngularVelocity() const
-		{
-			return P2UVector(OwningSimulation.GetLowLevelBody(ActorDataIndex).angularVelocity);
-		}
-
-		void AddForce(const FVector& Force)
-		{
-			OwningSimulation.AddForce(ActorDataIndex, Force);
-		}
-
-		void AddRadialForce(const FVector& Origin, float Strength, float Radius, ERadialImpulseFalloff Falloff, EForceType ForceType)
-		{
-			OwningSimulation.AddRadialForce(ActorDataIndex, Origin, Strength, Radius, Falloff, ForceType);
-		}
-
-		/** Set the linear damping*/
-		void SetLinearDamping(float NewLinearDamping)
-		{
-			OwningSimulation.GetLowLevelBody(ActorDataIndex).linearDamping = NewLinearDamping;
-		}
-
-		/** Get the linear damping*/
-		float GetLinearDamping() const
-		{
-			return OwningSimulation.GetLowLevelBody(ActorDataIndex).linearDamping;
-		}
-
-		/** Set the angular damping*/
-		void SetAngularDamping(float NewAngularDamping)
-		{
-			OwningSimulation.GetLowLevelBody(ActorDataIndex).angularDamping = NewAngularDamping;
-		}
-
-		/** Get the angular damping*/
-		float GetAngularDamping() const
-		{
-			return OwningSimulation.GetLowLevelBody(ActorDataIndex).angularDamping;
-		}
-
-		/** Set the max linear velocity squared*/
-		void SetMaxLinearVelocitySquared(float NewMaxLinearVelocitySquared)
-		{
-			OwningSimulation.GetLowLevelBody(ActorDataIndex).maxLinearVelocitySq = NewMaxLinearVelocitySquared;
-		}
-
-		/** Get the max linear velocity squared*/
-		float GetMaxLinearVelocitySquared() const
-		{
-			return OwningSimulation.GetLowLevelBody(ActorDataIndex).maxLinearVelocitySq;
-		}
-
-		/** Set the max angular velocity squared*/
-		void SetMaxAngularVelocitySquared(float NewMaxAngularVelocitySquared)
-		{
-			OwningSimulation.GetLowLevelBody(ActorDataIndex).maxAngularVelocitySq = NewMaxAngularVelocitySquared;
-		}
-
-		/** Get the max angular velocity squared*/
-		float GetMaxAngularVelocitySquared() const
-		{
-			return OwningSimulation.GetLowLevelBody(ActorDataIndex).maxAngularVelocitySq;
-		}
-
-		/** Set the inverse mass. 0 indicates kinematic object */
-		void SetInverseMass(float NewInverseMass)
-		{
-			OwningSimulation.GetLowLevelBody(ActorDataIndex).invMass = NewInverseMass;
-		}
-
-		/** Get the inverse mass. */
-		float GetInverseMass() const
-		{
-			return OwningSimulation.GetLowLevelBody(ActorDataIndex).invMass;
-		}
-
-		/** Set the inverse inertia. Mass-space inverse inertia diagonal vector */
-		void SetInverseInertia(const FVector& NewInverseInertia)
-		{
-			OwningSimulation.GetLowLevelBody(ActorDataIndex).invInertia = U2PVector(NewInverseInertia);
-		}
-
-		/** Get the inverse inertia. Mass-space inverse inertia diagonal vector */
-		FVector GetInverseInertia() const
-		{
-			return P2UVector(OwningSimulation.GetLowLevelBody(ActorDataIndex).invInertia);
-		}
-
-		/** Set the max depenetration velocity*/
-		void SetMaxDepenetrationVelocity(float NewMaxDepenetrationVelocity)
-		{
-			OwningSimulation.GetLowLevelBody(ActorDataIndex).maxDepenetrationVelocity = NewMaxDepenetrationVelocity;
-		}
-
-		/** Get the max depenetration velocity*/
-		float GetMaxDepenetrationVelocity(float NewMaxDepenetrationVelocity) const
-		{
-			return OwningSimulation.GetLowLevelBody(ActorDataIndex).maxDepenetrationVelocity;
-		}
-
-		/** Set the max contact impulse*/
-		void SetMaxContactImpulse(float NewMaxContactImpulse)
-		{
-			OwningSimulation.GetLowLevelBody(ActorDataIndex).maxContactImpulse = NewMaxContactImpulse;
-		}
-
-		/** Get the max contact impulse*/
-		float GetMaxContactImpulse() const
-		{
-			return OwningSimulation.GetLowLevelBody(ActorDataIndex).maxContactImpulse;
-		}
-
-		const int32 GetActorIndex() const
-		{
-			return ActorDataIndex;
-		}
-
-		void AddShape(physx::PxShape* InShape)
-		{
-			OwningSimulation.Actors[ActorDataIndex].AddShape(InShape);
-		}
-
-		TArray<FShape>& GetShapes() const
-		{
-			return OwningSimulation.Actors[ActorDataIndex].Shapes;
-		}
-
-		FActor* GetSimulationActor() const
-		{
-			return &OwningSimulation.Actors[ActorDataIndex];
-		}
-
-		immediate::PxRigidBodyData* GetSimulationRigidBodyData() const
-		{
-			immediate::PxRigidBodyData& Data = OwningSimulation.GetLowLevelBody(ActorDataIndex);
-			return &Data;
-		}
-
-	private:
-		/** Converts from actor space (i.e. the transform in world space as the client gives us) to body space (body with its origin at the COM and oriented to inertia tensor) */
-		FTransform ActorToBody;
-		FSimulation& OwningSimulation;
-		int32 ActorDataIndex;
-
-		friend FSimulation;
-		FActorHandle(FSimulation& InOwningSimulation, int32 InActorDataIndex)
-			: ActorToBody(FTransform::Identity)
-			, OwningSimulation(InOwningSimulation)
-			, ActorDataIndex(InActorDataIndex)
-		{
-		}
-
-		~FActorHandle()
-		{
-		}
-
-		FActorHandle(const FActorHandle&);	//Ensure no copying of handles
-	};
-
-}
-
-#endif // WITH_PHYSX
+#include "Physics/ImmediatePhysics/ImmediatePhysicsDeclares.h"
