Fixing crash with removing static component from physics.

diff --git a/Engine/Source/Runtime/Engine/Private/PhysicsEngine/Experimental/PhysScene_Chaos.cpp b/Engine/Source/Runtime/Engine/Private/PhysicsEngine/Experimental/PhysScene_Chaos.cpp
--- a/Engine/Source/Runtime/Engine/Private/PhysicsEngine/Experimental/PhysScene_Chaos.cpp
+++ b/Engine/Source/Runtime/Engine/Private/PhysicsEngine/Experimental/PhysScene_Chaos.cpp
@@ -497,7 +497,7 @@ void FPhysScene_Chaos::AddObject(UPrimitiveComponent* Component, FSkeletalMeshPh
 {
 	AddToComponentMaps(Component, InObject);
 
-	check(false);
+	ensure(false);
 #if 0
 	Chaos::FPhysicsSolver* Solver = GetSolver();
 	Solver->RegisterObject(InObject);
@@ -523,7 +523,7 @@ void FPhysScene_Chaos::AddObject(UPrimitiveComponent* Component, FGeometryPartic
 {
 	AddToComponentMaps(Component, InObject);
 
-	check(false);
+	ensure(false);
 #if 0
 	Chaos::FPhysicsSolver* Solver = GetSolver();
 	Solver->RegisterObject(InObject);
@@ -621,7 +621,7 @@ void RemovePhysicsProxy(ObjectType* InObject, Chaos::FPhysicsSolver* InSolver, F
 
 void FPhysScene_Chaos::RemoveObject(FSkeletalMeshPhysicsProxy* InObject)
 {
-	check(false);
+	ensure(false);
 #if 0
 	Chaos::FPhysicsSolver* Solver = InObject->GetSolver();
 	const int32 NumRemoved = Solver->UnregisterObject(InObject);
@@ -639,7 +639,7 @@ void FPhysScene_Chaos::RemoveObject(FSkeletalMeshPhysicsProxy* InObject)
 
 void FPhysScene_Chaos::RemoveObject(FStaticMeshPhysicsProxy* InObject)
 {
-	check(false);
+	ensure(false);
 #if 0
 	Chaos::FPhysicsSolver* Solver = InObject->GetSolver();
 
@@ -658,7 +658,7 @@ void FPhysScene_Chaos::RemoveObject(FStaticMeshPhysicsProxy* InObject)
 
 void FPhysScene_Chaos::RemoveObject(FGeometryParticlePhysicsProxy* InObject)
 {
-	check(false);
+	ensure(false);
 #if 0
 	Chaos::FPhysicsSolver* Solver = InObject->GetSolver();
 
