Chaos - Work around giving Transformed Implicit to Cast Helper when computing implicit mass properties.
Make Transformed drill down to concrete in cast helper and ensure to notify transform data is lost.

diff --git a/Engine/Source/Runtime/Engine/Private/PhysicsEngine/Experimental/PhysInterface_Chaos.cpp b/Engine/Source/Runtime/Engine/Private/PhysicsEngine/Experimental/PhysInterface_Chaos.cpp
--- a/Engine/Source/Runtime/Engine/Private/PhysicsEngine/Experimental/PhysInterface_Chaos.cpp
+++ b/Engine/Source/Runtime/Engine/Private/PhysicsEngine/Experimental/PhysInterface_Chaos.cpp
@@ -1413,21 +1413,45 @@ void CalculateMassPropertiesOfImplicitType(
 	// WIP
 	// @todo : Support center of mass offsets.
 	// @todo : Support Mass space alignment. 
+	
+	using namespace Chaos;
 
 	if (ImplicitObject)
 	{
-		const Chaos::TMassProperties<float, 3> MassProps = Chaos::CastHelper(*ImplicitObject, [&TotalMass, InDensityKGPerCM](const auto& Object)
+		// Hack to handle Transformed until CastHelper can properly support transformed
+		Chaos::FRigidTransform3 Transform(FMatrix::Identity);
+		if (ImplicitObject->GetType(true) & Chaos::ImplicitObjectType::Transformed)
 		{
-			Chaos::TMassProperties<float, 3> MassProperties;
-			MassProperties.Volume = Object.GetVolume();
-			const float Mass = MassProperties.Volume * InDensityKGPerCM;
+			// TODO:  This all very wrong, but is wrong in the same way as scaled. Rotation/Translation are ignored though. The three methods on Transformed no implemented.
+			// Only adding this to hack around CastHelper, as TransformedImplicit is very not supported in that path.
+			const Chaos::TImplicitObjectTransformed<FReal, 3>& Object = ImplicitObject->template GetObjectChecked<Chaos::TImplicitObjectTransformed<FReal, 3>>();
+
+
+			Chaos::TMassProperties<float, 3> MassProps;
+			MassProps.Volume = Object.GetVolume();
+			const float Mass = MassProps.Volume * InDensityKGPerCM;
 			TotalMass += Mass;
-			MassProperties.InertiaTensor = Object.GetInertiaTensor(Mass);
-			MassProperties.CenterOfMass = Object.GetCenterOfMass();
-			MassProperties.RotationOfMass = Chaos::TRotation<float, 3>::FromIdentity();
-			return MassProperties;
-		});
-		MassProperties.Add(MassProps);
+			MassProps.InertiaTensor = Object.GetInertiaTensor(Mass);
+			MassProps.CenterOfMass = Object.GetCenterOfMass();
+			MassProps.RotationOfMass = Chaos::TRotation<float, 3>::FromIdentity();
+
+			MassProperties.Add(MassProps);
+		}
+		else
+		{
+			const Chaos::TMassProperties<float, 3> MassProps = Chaos::CastHelper(*ImplicitObject, [&TotalMass, InDensityKGPerCM](const auto& Object)
+				{
+					Chaos::TMassProperties<float, 3> MassProperties;
+					MassProperties.Volume = Object.GetVolume();
+					const float Mass = MassProperties.Volume * InDensityKGPerCM;
+					TotalMass += Mass;
+					MassProperties.InertiaTensor = Object.GetInertiaTensor(Mass);
+					MassProperties.CenterOfMass = Object.GetCenterOfMass();
+					MassProperties.RotationOfMass = Chaos::TRotation<float, 3>::FromIdentity();
+					return MassProperties;
+				});
+			MassProperties.Add(MassProps);
+		}
 	}
 }
 
diff --git a/Engine/Source/Runtime/Experimental/Chaos/Public/Chaos/GeometryQueries.h b/Engine/Source/Runtime/Experimental/Chaos/Public/Chaos/GeometryQueries.h
--- a/Engine/Source/Runtime/Experimental/Chaos/Public/Chaos/GeometryQueries.h
+++ b/Engine/Source/Runtime/Experimental/Chaos/Public/Chaos/GeometryQueries.h
@@ -33,6 +33,11 @@ namespace Chaos
 		case ImplicitObjectType::IsScaled | ImplicitObjectType::Box: return Func(Geom.template GetObjectChecked< TImplicitObjectScaled<TBox<FReal, 3>>>());
 		case ImplicitObjectType::IsScaled | ImplicitObjectType::Capsule: return Func(Geom.template GetObjectChecked< TImplicitObjectScaled<TCapsule<FReal>>>());
 		case ImplicitObjectType::IsScaled | ImplicitObjectType::Convex: return Func(Geom.template GetObjectChecked< TImplicitObjectScaled<TConvex<FReal, 3>>>());
+		case ImplicitObjectType::Transformed:
+			ensure(false); // We are drilling down to concrete implicit inside transformed, this is disregarding transform data. Caller must specially handle transform.
+			// TODO: Refactor Transformed implicit to use same structure as scaled.
+			return CastHelper(*(Geom.template GetObjectChecked<TImplicitObjectTransformed<FReal, 3>>().GetTransformedObject()), Func);
+
 		default: check(false);
 		}
 		return Func(Geom.template GetObjectChecked<TSphere<FReal, 3>>());	//needed for return type
diff --git a/Engine/Source/Runtime/Experimental/Chaos/Public/Chaos/ImplicitObjectTransformed.h b/Engine/Source/Runtime/Experimental/Chaos/Public/Chaos/ImplicitObjectTransformed.h
--- a/Engine/Source/Runtime/Experimental/Chaos/Public/Chaos/ImplicitObjectTransformed.h
+++ b/Engine/Source/Runtime/Experimental/Chaos/Public/Chaos/ImplicitObjectTransformed.h
@@ -186,6 +186,25 @@ public:
 
 	virtual const TBox<T, d>& BoundingBox() const override { return MLocalBoundingBox; }
 
+	const FReal GetVolume() const
+	{
+		// TODO: More precise volume!
+		return BoundingBox().GetVolume();
+	}
+
+	const FMatrix33 GetInertiaTensor(const FReal Mass) const
+	{
+		// TODO: More precise inertia!
+		return BoundingBox().GetInertiaTensor(Mass);
+	}
+
+	const FVec3 GetCenterOfMass() const
+	{
+		// TODO: Actually compute this!
+		return FVec3(0.f);
+	}
+
+
 	const ObjectType Object() const { return MObject; }
 	
 	virtual void Serialize(FChaosArchive& Ar) override
