Try to add in some extra context around weird crashes in FRepLayout::InitFromClass.

diff --git a/Engine/Source/Runtime/Engine/Private/RepLayout.cpp b/Engine/Source/Runtime/Engine/Private/RepLayout.cpp
--- a/Engine/Source/Runtime/Engine/Private/RepLayout.cpp
+++ b/Engine/Source/Runtime/Engine/Private/RepLayout.cpp
@@ -5047,6 +5047,24 @@ void FRepLayout::InitFromClass(
 		UProperty * Property = InObjectClass->ClassReps[i].Property;
 		const int32 ArrayIdx = InObjectClass->ClassReps[i].Index;
 
+#if DO_CHECK
+		if (UNLIKELY(Property == nullptr))
+		{
+			FString Message = FString::Printf(TEXT("Class: %s | Index: %d"), *InObjectClass->GetPathName(), i);
+			if (i > 0)
+			{
+				Message += FString::Printf(TEXT(" | PreviousClassRepProperty: %s"), *GetPathNameSafe(InObjectClass->ClassReps[i - 1].Property));
+			}
+			else if (i < InObjectClass->ClassReps.Num() - 1)
+			{
+				Message += FString::Printf(TEXT(" | NextClassRepProperty: %s"), *GetPathNameSafe(InObjectClass->ClassReps[i + 1].Property));
+			}
+
+			checkf(false, TEXT("Encountered an invalid property while creating RepLayout. This should never happen! %s"), *Message);
+			return;
+		}
+#endif
+
 		check(Property->PropertyFlags & CPF_Net);
 
 		const int32 ParentHandle = AddParentProperty(Parents, Property, ArrayIdx);
