Force Chaos tick mode to be taskgraph or singlethread as dedicated thread mode currently unsupported.

diff --git a/Engine/Source/Runtime/Engine/Private/PhysicsEngine/PhysicsSettings.cpp b/Engine/Source/Runtime/Engine/Private/PhysicsEngine/PhysicsSettings.cpp
--- a/Engine/Source/Runtime/Engine/Private/PhysicsEngine/PhysicsSettings.cpp
+++ b/Engine/Source/Runtime/Engine/Private/PhysicsEngine/PhysicsSettings.cpp
@@ -84,6 +84,13 @@ void UPhysicsSettings::PostInitProperties()
 	{
 		DefaultShapeComplexity = bDefaultHasComplexCollision_DEPRECATED ? CTF_UseSimpleAndComplex : CTF_UseSimpleAsComplex;
 	}
+
+	// Temporarily override dedicated thread to taskgraph. The enum selection for dedicated
+	// thread is hidden until that threading mode is made to work with the world physics system overall
+	if(ChaosSettings.DefaultThreadingModel == EChaosThreadingMode::DedicatedThread)
+	{
+		ChaosSettings.DefaultThreadingModel = EChaosThreadingMode::TaskGraph;
+	}
 }
 
 #if WITH_EDITOR
diff --git a/Engine/Source/Runtime/PhysicsCore/Public/PhysicsCoreTypes.h b/Engine/Source/Runtime/PhysicsCore/Public/PhysicsCoreTypes.h
--- a/Engine/Source/Runtime/PhysicsCore/Public/PhysicsCoreTypes.h
+++ b/Engine/Source/Runtime/PhysicsCore/Public/PhysicsCoreTypes.h
@@ -18,7 +18,7 @@ enum class EChaosSolverTickMode : uint8
 UENUM()
 enum class EChaosThreadingMode : uint8
 {
-	DedicatedThread,
+	DedicatedThread UMETA(Hidden),
 	TaskGraph,
 	SingleThread,
 	Num UMETA(Hidden),
