Cameras: Make camera shake sources able to start/stop a camera shake of their own. Add editor sprite.

Having a shake directly set on the source saves the user a whole bunch of clicks to set things up in the Blueprint editor for 90% of cases.

diff --git a/Engine/Source/Runtime/Engine/Classes/Camera/CameraShakeSourceActor.cpp b/Engine/Source/Runtime/Engine/Classes/Camera/CameraShakeSourceActor.cpp
--- a/Engine/Source/Runtime/Engine/Classes/Camera/CameraShakeSourceActor.cpp
+++ b/Engine/Source/Runtime/Engine/Classes/Camera/CameraShakeSourceActor.cpp
@@ -7,6 +7,6 @@
 ACameraShakeSourceActor::ACameraShakeSourceActor(const FObjectInitializer& ObjectInitializer)
 	: Super(ObjectInitializer)
 {
-	CameraShakeSourceComponent = CreateDefaultSubobject<UCameraShakeSourceComponent>(TEXT("CameraShakeSourceComponent"));
+	CameraShakeSourceComponent = ObjectInitializer.CreateDefaultSubobject<UCameraShakeSourceComponent>(this, TEXT("CameraShakeSourceComponent"));
 	RootComponent = CameraShakeSourceComponent;
 }
diff --git a/Engine/Source/Runtime/Engine/Classes/Camera/CameraShakeSourceComponent.cpp b/Engine/Source/Runtime/Engine/Classes/Camera/CameraShakeSourceComponent.cpp
--- a/Engine/Source/Runtime/Engine/Classes/Camera/CameraShakeSourceComponent.cpp
+++ b/Engine/Source/Runtime/Engine/Classes/Camera/CameraShakeSourceComponent.cpp
@@ -2,11 +2,17 @@
 
 #include "CameraShakeSourceComponent.h"
 #include "CameraShake.h"
+#include "Camera/CameraModifier_CameraShake.h"
 #include "CinematicCameraModule.h"
+#include "Components/BillboardComponent.h"
+#include "Editor.h"
 #include "Engine/Engine.h"
+#include "Engine/Texture2D.h"
 #include "Engine/World.h"
 #include "GameFramework/PlayerController.h"
+#include "LevelEditorViewport.h"
 #include "Modules/ModuleManager.h"
+#include "UObject/ConstructorHelpers.h"
 
 
 UCameraShakeSourceComponent::UCameraShakeSourceComponent(const FObjectInitializer& ObjectInitializer)
@@ -14,17 +20,68 @@ UCameraShakeSourceComponent::UCameraShakeSourceComponent(const FObjectInitialize
 	, Attenuation(ECameraShakeAttenuation::Quadratic)
 	, InnerAttenuationRadius(100.f)
 	, OuterAttenuationRadius(1000.f)
+	, bAutoPlay(false)
 {
+#if WITH_EDITORONLY_DATA
+	bVisualizeComponent = true;
+
+	if (!IsRunningCommandlet())
+	{
+		static ConstructorHelpers::FObjectFinder<UTexture2D> StaticTexture(TEXT("/Engine/EditorResources/S_CameraShakeSource"));
+
+		EditorSpriteTexture = StaticTexture.Object;
+		EditorSpriteTextureScale = .5f;
+	}
+#endif
+}
+
+#if WITH_EDITOR
+void UCameraShakeSourceComponent::OnRegister()
+{
+	Super::OnRegister();
+	UpdateEditorSpriteTexture();
+}
+
+void UCameraShakeSourceComponent::UpdateEditorSpriteTexture()
+{
+	if (SpriteComponent != nullptr)
+	{
+		SpriteComponent->SetSprite(EditorSpriteTexture);
+		SpriteComponent->SetRelativeScale3D(FVector(EditorSpriteTextureScale));
+	}
+}
+#endif
+
+void UCameraShakeSourceComponent::BeginPlay()
+{
+	Super::BeginPlay();
+
+	if (bAutoPlay)
+	{
+		Play();
+	}
+}
+
+void UCameraShakeSourceComponent::EndPlay(const EEndPlayReason::Type EndPlayReason)
+{
+	StopAllCameraShakes();
+
+	Super::EndPlay(EndPlayReason);
+}
+
+void UCameraShakeSourceComponent::Play()
+{
+	PlayCameraShake(CameraShake);
 }
 
-void UCameraShakeSourceComponent::PlayCameraShake(TSubclassOf<UCameraShake> CameraShake)
+void UCameraShakeSourceComponent::PlayCameraShake(TSubclassOf<UCameraShake> InCameraShake)
 {
 	for (FConstPlayerControllerIterator Iterator = GetWorld()->GetPlayerControllerIterator(); Iterator; ++Iterator)
 	{
 		APlayerController* PlayerController = Iterator->Get();
 		if (PlayerController != nullptr && PlayerController->PlayerCameraManager != nullptr)
 		{
-			PlayerController->ClientPlayCameraShakeFromSource(CameraShake, this);
+			PlayerController->ClientPlayCameraShakeFromSource(InCameraShake, this);
 		}
 	}
 }
diff --git a/Engine/Source/Runtime/Engine/Classes/Camera/CameraShakeSourceComponent.h b/Engine/Source/Runtime/Engine/Classes/Camera/CameraShakeSourceComponent.h
--- a/Engine/Source/Runtime/Engine/Classes/Camera/CameraShakeSourceComponent.h
+++ b/Engine/Source/Runtime/Engine/Classes/Camera/CameraShakeSourceComponent.h
@@ -24,21 +24,51 @@ class ENGINE_API UCameraShakeSourceComponent : public USceneComponent
 public:
 	UCameraShakeSourceComponent(const FObjectInitializer& ObjectInitializer);
 
+	virtual void BeginPlay() override;
+	virtual void EndPlay(const EEndPlayReason::Type EndPlayReason) override;
+
+#if WITH_EDITOR
+	virtual void OnRegister() override;
+
+	void UpdateEditorSpriteTexture();
+#endif
+
+public:
 	/** The attenuation profile for how camera shakes' intensity falls off with distance */
-	UPROPERTY(EditAnywhere, BlueprintReadWrite, Category = CameraShake)
+	UPROPERTY(EditAnywhere, BlueprintReadWrite, Category = Attenuation)
 	ECameraShakeAttenuation Attenuation;
 
 	/** Under this distance from the source, the camera shakes are at full intensity */
-	UPROPERTY(EditAnywhere, BlueprintReadWrite, Category = CameraShake)
+	UPROPERTY(EditAnywhere, BlueprintReadWrite, Category = Attenuation)
 	float InnerAttenuationRadius;
 
 	/** Outside of this distance from the source, the camera shakes don't apply at all */
-	UPROPERTY(EditAnywhere, BlueprintReadWrite, Category = CameraShake)
+	UPROPERTY(EditAnywhere, BlueprintReadWrite, Category = Attenuation)
 	float OuterAttenuationRadius;
 
+	UPROPERTY(EditAnywhere, BlueprintReadWrite, Category = CameraShake)
+	TSubclassOf<UCameraShake> CameraShake;
+
+	UPROPERTY(EditAnywhere, BlueprintReadWrite, Category = CameraShake)
+	bool bAutoPlay;
+
+#if WITH_EDITORONLY_DATA
+	/** Sprite to display in the editor. */
+	UPROPERTY(transient)
+	UTexture2D* EditorSpriteTexture;
+
+	/** Sprite scaling for display in the editor. */
+	UPROPERTY(transient)
+	float EditorSpriteTextureScale;
+#endif
+
+public:
+	UFUNCTION(BlueprintCallable, Category = CameraShake)
+	void Play();
+
 	/** Starts a new camera shake originating from this source, and apply it on all player controllers */
 	UFUNCTION(BlueprintCallable, Category = CameraShake)
-	void PlayCameraShake(TSubclassOf<UCameraShake> CameraShake);
+	void PlayCameraShake(TSubclassOf<UCameraShake> InCameraShake);
 
 	/** Stops all currently active camera shakes that are originating from this source from all player controllers */
 	UFUNCTION(BlueprintCallable, Category = CameraShake)
