Cameras: Use weak pointers and auto-cleanup shakes whose sources have been destroyed.

diff --git a/Engine/Source/Runtime/Engine/Classes/Camera/CameraModifier_CameraShake.h b/Engine/Source/Runtime/Engine/Classes/Camera/CameraModifier_CameraShake.h
--- a/Engine/Source/Runtime/Engine/Classes/Camera/CameraModifier_CameraShake.h
+++ b/Engine/Source/Runtime/Engine/Classes/Camera/CameraModifier_CameraShake.h
@@ -36,7 +36,7 @@ struct FActiveCameraShakeInfo
 	UCameraShake* ShakeInstance;
 
 	UPROPERTY()
-	const UCameraShakeSourceComponent* ShakeSource;
+	TWeakObjectPtr<const UCameraShakeSourceComponent> ShakeSource;
 };
 
 struct FAddCameraShakeParams
diff --git a/Engine/Source/Runtime/Engine/Private/Camera/CameraModifier_CameraShake.cpp b/Engine/Source/Runtime/Engine/Private/Camera/CameraModifier_CameraShake.cpp
--- a/Engine/Source/Runtime/Engine/Private/Camera/CameraModifier_CameraShake.cpp
+++ b/Engine/Source/Runtime/Engine/Private/Camera/CameraModifier_CameraShake.cpp
@@ -40,9 +40,10 @@ bool UCameraModifier_CameraShake::ModifyCamera(float DeltaTime, FMinimalViewInfo
 				// Compute the scale of this shake for this frame according to the location
 				// of its source.
 				float CurShakeAlpha = Alpha;
-				if (ShakeInfo.ShakeSource != nullptr)
+				if (ShakeInfo.ShakeSource.IsValid())
 				{
-					const float AttenuationFactor = ShakeInfo.ShakeSource->GetAttenuationFactor(InOutPOV.Location);
+					const UCameraShakeSourceComponent* SourceComponent = ShakeInfo.ShakeSource.Get();
+					const float AttenuationFactor = SourceComponent->GetAttenuationFactor(InOutPOV.Location);
 					CurShakeAlpha *= AttenuationFactor;
 				}
 
@@ -54,7 +55,7 @@ bool UCameraModifier_CameraShake::ModifyCamera(float DeltaTime, FMinimalViewInfo
 		for (int32 i = ActiveShakes.Num() - 1; i >= 0; i--)
 		{
 			FActiveCameraShakeInfo ShakeInfo = ActiveShakes[i]; // Copy struct, we're going to maybe delete it.
-			if (ShakeInfo.ShakeInstance == nullptr || ShakeInfo.ShakeInstance->IsFinished())
+			if (ShakeInfo.ShakeInstance == nullptr || ShakeInfo.ShakeInstance->IsFinished() || ShakeInfo.ShakeSource.IsStale())
 			{
 				ActiveShakes.RemoveAt(i, 1);
 
@@ -226,7 +227,7 @@ void UCameraModifier_CameraShake::RemoveAllCameraShakesFromSource(const UCameraS
 	for (int32 i = ActiveShakes.Num() - 1; i >= 0; --i)
 	{
 		FActiveCameraShakeInfo ShakeInfo = ActiveShakes[i];  // Copy struct because we might delete it.
-		if (ShakeInfo.ShakeSource == SourceComponent && ShakeInfo.ShakeInstance != nullptr)
+		if (ShakeInfo.ShakeSource.Get() == SourceComponent && ShakeInfo.ShakeInstance != nullptr)
 		{
 			ShakeInfo.ShakeInstance->StopShake(bImmediately);
 			if (bImmediately)
