Get rid of invalid cook override define for Mac, which was causing audio not to load correctly on cooked builds.

diff --git a/Engine/Source/Runtime/Engine/Private/AudioCompressionSettingsUtils.cpp b/Engine/Source/Runtime/Engine/Private/AudioCompressionSettingsUtils.cpp
--- a/Engine/Source/Runtime/Engine/Private/AudioCompressionSettingsUtils.cpp
+++ b/Engine/Source/Runtime/Engine/Private/AudioCompressionSettingsUtils.cpp
@@ -4,7 +4,7 @@
 #include "AudioCompressionSettings.h"
 
 #define ENABLE_PLATFORM_COMPRESSION_OVERRIDES 1
-#define PLATFORM_SUPPORTS_COMPRESSION_OVERRIDES (PLATFORM_ANDROID && !PLATFORM_LUMIN) || PLATFORM_IOS || PLATFORM_SWITCH || PLATFORM_XBOXONE || PLATFORM_PS4 || PLATFORM_WINDOWS || PLATFORM_MAC
+#define PLATFORM_SUPPORTS_COMPRESSION_OVERRIDES (PLATFORM_ANDROID && !PLATFORM_LUMIN) || PLATFORM_IOS || PLATFORM_SWITCH || PLATFORM_XBOXONE || PLATFORM_PS4 || PLATFORM_WINDOWS
 
 #if PLATFORM_ANDROID && !PLATFORM_LUMIN && ENABLE_PLATFORM_COMPRESSION_OVERRIDES
 #include "AndroidRuntimeSettings.h"
