Remove unused include.

diff --git a/Engine/Source/Runtime/Engine/Classes/Camera/CameraShakeSourceComponent.cpp b/Engine/Source/Runtime/Engine/Classes/Camera/CameraShakeSourceComponent.cpp
--- a/Engine/Source/Runtime/Engine/Classes/Camera/CameraShakeSourceComponent.cpp
+++ b/Engine/Source/Runtime/Engine/Classes/Camera/CameraShakeSourceComponent.cpp
@@ -9,7 +9,6 @@
 #include "Engine/Texture2D.h"
 #include "Engine/World.h"
 #include "GameFramework/PlayerController.h"
-#include "LevelEditorViewport.h"
 #include "Modules/ModuleManager.h"
 #include "UObject/ConstructorHelpers.h"
 
