Fix build break caused by mixed up WITH_EDITOR/WITH_EDITOR_DATA

diff --git a/Engine/Source/Runtime/Engine/Classes/Camera/CameraShakeSourceComponent.cpp b/Engine/Source/Runtime/Engine/Classes/Camera/CameraShakeSourceComponent.cpp
--- a/Engine/Source/Runtime/Engine/Classes/Camera/CameraShakeSourceComponent.cpp
+++ b/Engine/Source/Runtime/Engine/Classes/Camera/CameraShakeSourceComponent.cpp
@@ -33,7 +33,6 @@ UCameraShakeSourceComponent::UCameraShakeSourceComponent(const FObjectInitialize
 #endif
 }
 
-#if WITH_EDITOR
 void UCameraShakeSourceComponent::OnRegister()
 {
 	Super::OnRegister();
@@ -42,13 +41,14 @@ void UCameraShakeSourceComponent::OnRegister()
 
 void UCameraShakeSourceComponent::UpdateEditorSpriteTexture()
 {
+#if WITH_EDITORONLY_DATA
 	if (SpriteComponent != nullptr)
 	{
 		SpriteComponent->SetSprite(EditorSpriteTexture);
 		SpriteComponent->SetRelativeScale3D(FVector(EditorSpriteTextureScale));
 	}
-}
 #endif
+}
 
 void UCameraShakeSourceComponent::BeginPlay()
 {
diff --git a/Engine/Source/Runtime/Engine/Classes/Camera/CameraShakeSourceComponent.h b/Engine/Source/Runtime/Engine/Classes/Camera/CameraShakeSourceComponent.h
--- a/Engine/Source/Runtime/Engine/Classes/Camera/CameraShakeSourceComponent.h
+++ b/Engine/Source/Runtime/Engine/Classes/Camera/CameraShakeSourceComponent.h
@@ -27,11 +27,10 @@ public:
 	virtual void BeginPlay() override;
 	virtual void EndPlay(const EEndPlayReason::Type EndPlayReason) override;
 
-#if WITH_EDITOR
 	virtual void OnRegister() override;
 
+private:
 	void UpdateEditorSpriteTexture();
-#endif
 
 public:
 	/** The attenuation profile for how camera shakes' intensity falls off with distance */
