Adds SSGI project setting

diff --git a/Engine/Source/Runtime/Engine/Classes/Engine/RendererSettings.h b/Engine/Source/Runtime/Engine/Classes/Engine/RendererSettings.h
--- a/Engine/Source/Runtime/Engine/Classes/Engine/RendererSettings.h
+++ b/Engine/Source/Runtime/Engine/Classes/Engine/RendererSettings.h
@@ -467,6 +467,11 @@ class ENGINE_API URendererSettings : public UDeveloperSettings
 		ToolTip = "Whether to do primary screen percentage with temporal AA or not."))
 	uint32 bTemporalUpsampling : 1;
 
+	UPROPERTY(config, EditAnywhere, Category = Lighting, meta = (
+		ConsoleVariable = "r.SSGI.Enable", DisplayName = "Screen Space Global Illumination (Beta)",
+		ToolTip = "Whether enable screen space global illumination."))
+	uint32 bSSGI : 1;
+
 	UPROPERTY(config, EditAnywhere, Category = DefaultSettings, meta = (
 		ConsoleVariable = "r.DefaultFeature.AntiAliasing", DisplayName = "Anti-Aliasing Method",
 		ToolTip = "Which anti-aliasing mode is used by default"))
