Fix non-unity
#rnx

diff --git a/Engine/Source/Runtime/Engine/Classes/Camera/CameraShakeSourceComponent.h b/Engine/Source/Runtime/Engine/Classes/Camera/CameraShakeSourceComponent.h
--- a/Engine/Source/Runtime/Engine/Classes/Camera/CameraShakeSourceComponent.h
+++ b/Engine/Source/Runtime/Engine/Classes/Camera/CameraShakeSourceComponent.h
@@ -8,6 +8,7 @@
 #include "CameraShakeSourceComponent.generated.h"
 
 class UCameraShake;
+class UTexture2D;
 
 UENUM(BlueprintType)
 enum class ECameraShakeAttenuation : uint8
