Fix OutputHash.txt being created in Engine/Binaries directory.

#rnx

diff --git a/Engine/Source/Runtime/Engine/Private/ShaderCompiler/ShaderCompiler.cpp b/Engine/Source/Runtime/Engine/Private/ShaderCompiler/ShaderCompiler.cpp
--- a/Engine/Source/Runtime/Engine/Private/ShaderCompiler/ShaderCompiler.cpp
+++ b/Engine/Source/Runtime/Engine/Private/ShaderCompiler/ShaderCompiler.cpp
@@ -622,7 +622,7 @@ static void ReadSingleJob(FShaderCompileJob* CurrentJob, FArchive& OutputFile)
 	CurrentJob->Output.GenerateOutputHash();
 	CurrentJob->bSucceeded = CurrentJob->Output.bSucceeded;
 
-	if (CurrentJob->bSucceeded && CurrentJob->Input.DumpDebugInfoRootPath.Len() > 0)
+	if (CurrentJob->bSucceeded && CurrentJob->Input.DumpDebugInfoPath.Len() > 0)
 	{
 		// write down the output hash as a file
 		FString HashFileName = FPaths::Combine(CurrentJob->Input.DumpDebugInfoPath, TEXT("OutputHash.txt"));
