Remove PlatformRequiresOptTextureData and rework how those fields are serialized to prevent crashes when loading cooked data in the editor.
 - An additional bit in the FTexturePlatformData next to the Cubemap bit is used to indicate whether the structure contains non-default OptData. Contents of the OptData field is only serialized when that bit is set.
 - Platforms which do no make use of OptData will automatically not have that data in their cooked byte stream, as the values will remain default. Platform that do set those values will have the data included and the bit set.

diff --git a/Engine/Source/Developer/TargetPlatform/Public/Common/TargetPlatformBase.h b/Engine/Source/Developer/TargetPlatform/Public/Common/TargetPlatformBase.h
--- a/Engine/Source/Developer/TargetPlatform/Public/Common/TargetPlatformBase.h
+++ b/Engine/Source/Developer/TargetPlatform/Public/Common/TargetPlatformBase.h
@@ -322,11 +322,6 @@ public:
 		return TPlatformProperties::GetMemoryMappingAlignment();
 	}
 
-	virtual bool RequiresOptTextureData() const override
-	{
-		return TPlatformProperties::RequiresOptTextureData();
-	}
-
 
 #if WITH_ENGINE
 	virtual FName GetPhysicsFormat( class UBodySetup* Body ) const override
diff --git a/Engine/Source/Developer/TargetPlatform/Public/Interfaces/ITargetPlatform.h b/Engine/Source/Developer/TargetPlatform/Public/Interfaces/ITargetPlatform.h
--- a/Engine/Source/Developer/TargetPlatform/Public/Interfaces/ITargetPlatform.h
+++ b/Engine/Source/Developer/TargetPlatform/Public/Interfaces/ITargetPlatform.h
@@ -389,13 +389,6 @@ public:
 	*/
 	virtual float GetDownSampleMeshDistanceFieldDivider() const = 0;
 
-	/**
-	* Gets whether the platforms requires the OptData member of FTexturePlatformData.
-	* On platforms where this returns false, the OptData will be initialized to zero at
-	* runtime, and not serialized in the asset data stream.
-	*/
-	virtual bool RequiresOptTextureData() const = 0;
-
 #if WITH_ENGINE
 	/**
 	 * Gets the format to use for a particular body setup.
diff --git a/Engine/Source/Developer/TextureCompressor/Public/TextureCompressorModule.h b/Engine/Source/Developer/TextureCompressor/Public/TextureCompressorModule.h
--- a/Engine/Source/Developer/TextureCompressor/Public/TextureCompressorModule.h
+++ b/Engine/Source/Developer/TextureCompressor/Public/TextureCompressorModule.h
@@ -164,8 +164,6 @@ struct FTextureBuildSettings
 	uint32 bVirtualTextureEnableCompressZlib : 1;
 	/** Is crunch compression enabled */
 	uint32 bVirtualTextureEnableCompressCrunch : 1;
-	/** True if the target platform requires the FOptTexturePlatformData strucutre */
-	uint32 bOptData : 1;
 
 	/** Default settings. */
 	FTextureBuildSettings()
@@ -216,7 +214,6 @@ struct FTextureBuildSettings
 		, VirtualTextureBorderSize(0)
 		, bVirtualTextureEnableCompressZlib(false)
 		, bVirtualTextureEnableCompressCrunch(false)
-		, bOptData(false)
 	{
 	}
 
diff --git a/Engine/Source/Runtime/Core/Public/GenericPlatform/GenericPlatformProperties.h b/Engine/Source/Runtime/Core/Public/GenericPlatform/GenericPlatformProperties.h
--- a/Engine/Source/Runtime/Core/Public/GenericPlatform/GenericPlatformProperties.h
+++ b/Engine/Source/Runtime/Core/Public/GenericPlatform/GenericPlatformProperties.h
@@ -297,9 +297,4 @@ struct FGenericPlatformProperties
 	{
 		return nullptr;
 	}
-
-	static FORCEINLINE bool RequiresOptTextureData()
-	{
-		return false;
-	}
 };
diff --git a/Engine/Source/Runtime/Engine/Classes/Engine/Texture.h b/Engine/Source/Runtime/Engine/Classes/Engine/Texture.h
--- a/Engine/Source/Runtime/Engine/Classes/Engine/Texture.h
+++ b/Engine/Source/Runtime/Engine/Classes/Engine/Texture.h
@@ -415,6 +415,10 @@ private:
 #endif // WITH_EDITORONLY_DATA
 };
 
+/**
+ * Optional extra fields for texture platform data required by some platforms.
+ * Data in this struct is only serialized if the struct's value is non-default.
+ */
 struct FOptTexturePlatformData
 {
 	/** Arbitrary extra data that the runtime may need. */
@@ -427,6 +431,17 @@ struct FOptTexturePlatformData
 		, NumMipsInTail(0)
 	{}
 
+	inline bool operator == (FOptTexturePlatformData const& RHS) const 
+	{
+		return ExtData == RHS.ExtData
+			&& NumMipsInTail == RHS.NumMipsInTail;
+	}
+
+	inline bool operator != (FOptTexturePlatformData const& RHS) const
+	{
+		return !(*this == RHS);
+	}
+
 	friend inline FArchive& operator << (FArchive& Ar, FOptTexturePlatformData& Data)
 	{
 		Ar << Data.ExtData;
@@ -447,8 +462,8 @@ struct FTexturePlatformData
 	int32 SizeX;
 	/** Height of the texture. */
 	int32 SizeY;
-	/** Number of texture slices. Bit CubeMapBitIndex contains if a texture is a Cubemap or not. */
-	uint32 NumSlicesCubemapMask;
+	/** Packed bits [b31: CubeMap], [b30: HasOptData], [b29-0: NumSlices]. See bit masks below. */
+	uint32 PackedData;
 	/** Format in which mip data is stored. */
 	EPixelFormat PixelFormat;
 	/** Additional data required by some platforms.*/
@@ -475,10 +490,12 @@ struct FTexturePlatformData
 	/** Destructor. */
 	ENGINE_API ~FTexturePlatformData();
 
-	static constexpr uint32 CubeMapBitIndex = 31;
-	static constexpr uint32 CubeMapBitMask = (1u << CubeMapBitIndex);
-	static constexpr uint32 CubeMapRemainingBitMask = CubeMapBitMask - 1u;
+private:
+	static constexpr uint32 BitMask_CubeMap    = 1u << 31u;
+	static constexpr uint32 BitMask_HasOptData = 1u << 30u;
+	static constexpr uint32 BitMask_NumSlices  = BitMask_HasOptData - 1u;
 
+public:
 	/** Return whether TryLoadMips() would stall because async loaded mips are not yet available. */
 	bool IsReadyForAsyncPostLoad() const;
 
@@ -503,25 +520,39 @@ struct FTexturePlatformData
 	 * @param bStreamable Store some mips inline, only used during cooking
 	 */
 	void SerializeCooked(FArchive& Ar, class UTexture* Owner, bool bStreamable);
-
-	inline bool IsCubemap() const
+	
+	inline bool GetHasOptData() const
 	{
-		return (NumSlicesCubemapMask & CubeMapBitMask) == CubeMapBitMask;
+		return (PackedData & BitMask_HasOptData) == BitMask_HasOptData;
 	}
 
-	inline int32 GetNumSlices() const
+	inline void SetOptData(FOptTexturePlatformData Data)
 	{
-		return (int32)(NumSlicesCubemapMask & CubeMapRemainingBitMask);
+		// Set the opt data flag to true if the specified data is non-default.
+		bool bHasOptData = Data != FOptTexturePlatformData();
+		PackedData = (bHasOptData ? BitMask_HasOptData : 0) | (PackedData & (~BitMask_HasOptData));
+
+		OptData = Data;
 	}
 
-	inline void SetNumSlices(int32 NumSlices)
+	inline bool IsCubemap() const
 	{
-		NumSlicesCubemapMask = ((uint32)NumSlices & CubeMapRemainingBitMask) | (NumSlicesCubemapMask & CubeMapBitMask);
+		return (PackedData & BitMask_CubeMap) == BitMask_CubeMap; 
 	}
 
 	inline void SetIsCubemap(bool bCubemap)
 	{
-		NumSlicesCubemapMask = (bCubemap ? CubeMapBitMask : 0) | (NumSlicesCubemapMask & CubeMapRemainingBitMask);
+		PackedData = (bCubemap ? BitMask_CubeMap : 0) | (PackedData & (~BitMask_CubeMap));
+	}
+	
+	inline int32 GetNumSlices() const 
+	{
+		return (int32)(PackedData & BitMask_NumSlices);
+	}
+
+	inline void SetNumSlices(int32 NumSlices)
+	{
+		PackedData = (NumSlices & BitMask_NumSlices) | (PackedData & (~BitMask_NumSlices));
 	}
 
 	inline int32 GetNumMipsInTail() const
diff --git a/Engine/Source/Runtime/Engine/Private/TextureDerivedData.cpp b/Engine/Source/Runtime/Engine/Private/TextureDerivedData.cpp
--- a/Engine/Source/Runtime/Engine/Private/TextureDerivedData.cpp
+++ b/Engine/Source/Runtime/Engine/Private/TextureDerivedData.cpp
@@ -47,7 +47,7 @@
 // In case of merge conflicts with DDC versions, you MUST generate a new GUID and set this new
 // guid as version
 
-#define TEXTURE_DERIVEDDATA_VER		TEXT("0385EF5C8C62402ABCD94730886B0259")
+#define TEXTURE_DERIVEDDATA_VER		TEXT("564290F8998644E39A2118D5C683187B")
 
 // This GUID is mixed into DDC version for virtual textures only, this allows updating DDC version for VT without invalidating DDC for all textures
 // This is useful during development, but once large numbers of VT are present in shipped content, it will have the same problem as TEXTURE_DERIVEDDATA_VER
@@ -222,12 +222,6 @@ static void SerializeForKey(FArchive& Ar, const FTextureBuildSettings& Settings)
 		OutKeySuffix.Append(FString::Printf(TEXT("VT%s_"), TEXTURE_VT_DERIVEDDATA_VER));
 	}
 
-	if (BuildSettings.bOptData)
-	{
-		// Add a flag to signal the platform needs the optional texture data structure
-		OutKeySuffix.Append(TEXT("O_"));
-	}
-
 	// Serialize the compressor settings into a temporary array. The archive
 	// is flagged as persistent so that machines of different endianness produce
 	// identical binary results.
@@ -339,7 +333,6 @@ static void GetTextureBuildSettings(
 	const UTextureLODSettings& TextureLODSettings,
 	bool bPlatformSupportsTextureStreaming,
 	bool bPlatformSupportsVirtualTextureStreaming,
-	bool bPlatformRequiresOptTextureData,
 	FTextureBuildSettings& OutBuildSettings
 	)
 {
@@ -361,7 +354,6 @@ static void GetTextureBuildSettings(
 	OutBuildSettings.bVolume = false;
 	OutBuildSettings.bCubemap = false;
 	OutBuildSettings.bTextureArray = false;
-	OutBuildSettings.bOptData = bPlatformRequiresOptTextureData;
 
 	if (Texture.MaxTextureSize > 0)
 	{
@@ -508,10 +500,9 @@ static void GetBuildSettingsForRunningPlatform(
 		const UTextureLODSettings* LODSettings = (UTextureLODSettings*)UDeviceProfileManager::Get().FindProfile(CurrentPlatform->PlatformName());
 		const bool bPlatformSupportsTextureStreaming = CurrentPlatform->SupportsFeature(ETargetPlatformFeatures::TextureStreaming);
 		const bool bPlatformSupportsVirtualTextureStreaming = CurrentPlatform->SupportsFeature(ETargetPlatformFeatures::VirtualTextureStreaming);
-		const bool bPlatformRequiresOptTextureData = CurrentPlatform->RequiresOptTextureData();
 
 		FTextureBuildSettings SourceBuildSettings;
-		GetTextureBuildSettings(Texture, *LODSettings, bPlatformSupportsTextureStreaming, bPlatformSupportsVirtualTextureStreaming, bPlatformRequiresOptTextureData, SourceBuildSettings);
+		GetTextureBuildSettings(Texture, *LODSettings, bPlatformSupportsTextureStreaming, bPlatformSupportsVirtualTextureStreaming, SourceBuildSettings);
 
 		TArray< TArray<FName> > PlatformFormats;
 		CurrentPlatform->GetTextureFormats(&Texture, PlatformFormats);
@@ -1019,7 +1010,7 @@ bool FTexturePlatformData::TryInlineMipData(int32 FirstMipToLoad)
 FTexturePlatformData::FTexturePlatformData()
 	: SizeX(0)
 	, SizeY(0)
-	, NumSlicesCubemapMask(0)
+	, PackedData(0)
 	, PixelFormat(PF_Unknown)
 	, VTData(nullptr)
 #if WITH_EDITORONLY_DATA
@@ -1257,7 +1248,7 @@ static void SerializePlatformData(
 
 	Ar << PlatformData->SizeX;
 	Ar << PlatformData->SizeY;
-	Ar << PlatformData->NumSlicesCubemapMask;
+	Ar << PlatformData->PackedData;
 	if (Ar.IsLoading())
 	{
 		FString PixelFormatString;
@@ -1270,23 +1261,9 @@ static void SerializePlatformData(
 		Ar << PixelFormatString;
 	}
 
-	// If we're cooking, check whether the target platform requires optional data, otherwise check the current running platform.
-	// In both cases, always serialize if we have editor only data.
-#if WITH_EDITORONLY_DATA
-	if (Ar.IsCooking())
+	if (PlatformData->GetHasOptData())
 	{
-		if (Ar.CookingTarget()->RequiresOptTextureData() || Ar.CookingTarget()->HasEditorOnlyData())
-		{
-			Ar << PlatformData->OptData;
-		}
-	}
-	else 
-#endif
-	{
-		if (FPlatformProperties::RequiresOptTextureData() || FPlatformProperties::HasEditorOnlyData())
-		{
-			Ar << PlatformData->OptData;
-		}
+		Ar << PlatformData->OptData;
 	}
 
 	int32 NumMips = PlatformData->Mips.Num();
@@ -1619,9 +1596,8 @@ void UTexture::BeginCacheForCookedPlatformData( const ITargetPlatform *TargetPla
 		FTextureBuildSettings BuildSettings;
 		const bool bPlatformSupportsTextureStreaming = TargetPlatform->SupportsFeature(ETargetPlatformFeatures::TextureStreaming);
 		const bool bPlatformSupportsVirtualTextureStreaming = TargetPlatform->SupportsFeature(ETargetPlatformFeatures::VirtualTextureStreaming);
-		const bool bPlatformRequiresOptTextureData = TargetPlatform->RequiresOptTextureData();
 
-		GetTextureBuildSettings(*this, TargetPlatform->GetTextureLODSettings(), bPlatformSupportsTextureStreaming, bPlatformSupportsVirtualTextureStreaming, bPlatformRequiresOptTextureData, BuildSettings);
+		GetTextureBuildSettings(*this, TargetPlatform->GetTextureLODSettings(), bPlatformSupportsTextureStreaming, bPlatformSupportsVirtualTextureStreaming, BuildSettings);
 		
 		TArray< TArray<FTextureBuildSettings> > BuildSettingsToCache;
 		GetBuildSettingsPerFormat(*this, BuildSettings, TargetPlatform, BuildSettingsToCache);
@@ -1704,9 +1680,8 @@ void UTexture::ClearCachedCookedPlatformData( const ITargetPlatform* TargetPlatf
 		FTextureBuildSettings BuildSettings;
 		const bool bPlatformSupportsTextureStreaming = TargetPlatform->SupportsFeature(ETargetPlatformFeatures::TextureStreaming);
 		const bool bPlatformSupportsVirtualTextureStreaming = TargetPlatform->SupportsFeature(ETargetPlatformFeatures::VirtualTextureStreaming);
-		const bool bPlatformRequiresOptTextureData = TargetPlatform->RequiresOptTextureData();
 
-		GetTextureBuildSettings(*this, TargetPlatform->GetTextureLODSettings(), bPlatformSupportsTextureStreaming, bPlatformSupportsVirtualTextureStreaming, bPlatformRequiresOptTextureData, BuildSettings);
+		GetTextureBuildSettings(*this, TargetPlatform->GetTextureLODSettings(), bPlatformSupportsTextureStreaming, bPlatformSupportsVirtualTextureStreaming, BuildSettings);
 
 		TArray< TArray<FTextureBuildSettings> > BuildSettingsToCache;
 		GetBuildSettingsPerFormat(*this, BuildSettings, TargetPlatform, BuildSettingsToCache);
@@ -1767,9 +1742,8 @@ bool UTexture::IsCachedCookedPlatformDataLoaded( const ITargetPlatform* TargetPl
 
 	const bool bPlatformSupportsTextureStreaming = TargetPlatform->SupportsFeature(ETargetPlatformFeatures::TextureStreaming);
 	const bool bPlatformSupportsVirtualTextureStreaming = TargetPlatform->SupportsFeature(ETargetPlatformFeatures::VirtualTextureStreaming);
-	const bool bPlatformRequiresOptTextureData = TargetPlatform->RequiresOptTextureData();
 
-	GetTextureBuildSettings(*this, TargetPlatform->GetTextureLODSettings(), bPlatformSupportsTextureStreaming, bPlatformSupportsVirtualTextureStreaming, bPlatformRequiresOptTextureData, BuildSettings);
+	GetTextureBuildSettings(*this, TargetPlatform->GetTextureLODSettings(), bPlatformSupportsTextureStreaming, bPlatformSupportsVirtualTextureStreaming, BuildSettings);
 
 	TArray< TArray<FTextureBuildSettings> > BuildSettingsToCache;
 	GetBuildSettingsPerFormat(*this, BuildSettings, TargetPlatform, BuildSettingsToCache);
@@ -1973,9 +1947,8 @@ void UTexture::SerializeCookedPlatformData(FArchive& Ar)
 			FTextureBuildSettings BuildSettings;
 			const bool bPlatformSupportsTextureStreaming = Ar.CookingTarget()->SupportsFeature(ETargetPlatformFeatures::TextureStreaming);
 			const bool bPlatformSupportsVirtualTextureStreaming = Ar.CookingTarget()->SupportsFeature(ETargetPlatformFeatures::VirtualTextureStreaming);
-			const bool bPlatformRequiresOptTextureData = Ar.CookingTarget()->RequiresOptTextureData();
 
-			GetTextureBuildSettings(*this, Ar.CookingTarget()->GetTextureLODSettings(), bPlatformSupportsTextureStreaming, bPlatformSupportsVirtualTextureStreaming, bPlatformRequiresOptTextureData, BuildSettings);
+			GetTextureBuildSettings(*this, Ar.CookingTarget()->GetTextureLODSettings(), bPlatformSupportsTextureStreaming, bPlatformSupportsVirtualTextureStreaming, BuildSettings);
 
 			TArray<FTexturePlatformData*> PlatformDataToSerialize;
 
diff --git a/Engine/Source/Runtime/Engine/Private/TextureDerivedDataTask.cpp b/Engine/Source/Runtime/Engine/Private/TextureDerivedDataTask.cpp
--- a/Engine/Source/Runtime/Engine/Private/TextureDerivedDataTask.cpp
+++ b/Engine/Source/Runtime/Engine/Private/TextureDerivedDataTask.cpp
@@ -256,14 +256,16 @@ void FTextureCacheDerivedDataWorker::BuildTexture()
 		DerivedData->SetIsCubemap(false);
 		DerivedData->VTData = nullptr;
 
+		FOptTexturePlatformData OptData;
+
 		// Compress the texture.
 		TArray<FCompressedImage2D> CompressedMips;
 		if (Compressor->BuildTexture(TextureData.Blocks[0].MipsPerLayer[0],
 			((bool)Texture.CompositeTexture && CompositeTextureData.Blocks.Num() && CompositeTextureData.Blocks[0].MipsPerLayer.Num()) ? CompositeTextureData.Blocks[0].MipsPerLayer[0] : TArray<FImage>(),
 			BuildSettingsPerLayer[0],
 			CompressedMips,
-			DerivedData->OptData.NumMipsInTail,
-			DerivedData->OptData.ExtData))
+			OptData.NumMipsInTail,
+			OptData.ExtData))
 		{
 			check(CompressedMips.Num());
 
@@ -298,6 +300,8 @@ void FTextureCacheDerivedDataWorker::BuildTexture()
 				}
 			}
 
+			DerivedData->SetOptData(OptData);
+
 			// Store it in the cache.
 			// @todo: This will remove the streaming bulk data, which we immediately reload below!
 			// Should ideally avoid this redundant work, but it only happens when we actually have 
