- Guard UPrimitiveComponent::SetCustomPrimitiveDataInternal from being sent an invalid index below zero (it already guards the upper bounds)

diff --git a/Engine/Source/Runtime/Engine/Private/Components/PrimitiveComponent.cpp b/Engine/Source/Runtime/Engine/Private/Components/PrimitiveComponent.cpp
--- a/Engine/Source/Runtime/Engine/Private/Components/PrimitiveComponent.cpp
+++ b/Engine/Source/Runtime/Engine/Private/Components/PrimitiveComponent.cpp
@@ -1665,7 +1665,7 @@ UMaterialInstanceDynamic* UPrimitiveComponent::CreateDynamicMaterialInstance(int
 void UPrimitiveComponent::SetCustomPrimitiveDataInternal(int32 DataIndex, const TArray<float>& Values)
 {
 	// Can only set data on valid indices and only if there's actually any data to set
-	if (DataIndex < FCustomPrimitiveData::NumCustomPrimitiveDataFloats && Values.Num() > 0)
+	if (DataIndex >= 0 && DataIndex < FCustomPrimitiveData::NumCustomPrimitiveDataFloats && Values.Num() > 0)
 	{
 		// Number of floats needed in the custom primitive data array
 		const int32 NeededFloats = FMath::Min(DataIndex + Values.Num(), FCustomPrimitiveData::NumCustomPrimitiveDataFloats);
