Fix for replay playback crash if the replay was recorded with demo.WithGameSpecificFrameData enabled.

diff --git a/Engine/Source/Runtime/Engine/Private/DemoNetDriver.cpp b/Engine/Source/Runtime/Engine/Private/DemoNetDriver.cpp
--- a/Engine/Source/Runtime/Engine/Private/DemoNetDriver.cpp
+++ b/Engine/Source/Runtime/Engine/Private/DemoNetDriver.cpp
@@ -3232,7 +3232,7 @@ bool UDemoNetDriver::ReadDemoFrameIntoPlaybackPackets(FArchive& Ar, TArray<FPlay
 	{
 		Ar << SkipGameSpecificOffset;
 
-		if (!bForLevelFastForward)
+		if ((SkipGameSpecificOffset > 0) && !bForLevelFastForward)
 		{
 			FDemoFrameDataMap Data;
 			Ar << Data;
