Fixed: Enabling Fixed Frame Rate doesn't Lock the Frame Rate

This was a regression introduced in CL 9842913 while trying to fix negative time deltas when switching from a fixed time step to variable (UE-81145). However, using bTimeWasManipulated to detect the switch is incorrect, because that's always true when fixed frame rate is enabled. The code now stores the previous state of bUseFixedTimeStep explicitly and compares against that to determine when the interval needs to be reset.

#rnx

diff --git a/Engine/Source/Runtime/Engine/Private/UnrealEngine.cpp b/Engine/Source/Runtime/Engine/Private/UnrealEngine.cpp
--- a/Engine/Source/Runtime/Engine/Private/UnrealEngine.cpp
+++ b/Engine/Source/Runtime/Engine/Private/UnrealEngine.cpp
@@ -1948,6 +1948,7 @@ void UEngine::UpdateTimeAndHandleMaxTickRate()
 
 															// Figure out whether we want to use real or fixed time step.
 	const bool bUseFixedTimeStep = FApp::IsBenchmarking() || FApp::UseFixedTimeStep();
+	static bool bPreviousUseFixedTimeStep = bUseFixedTimeStep;
 
 	// Updates logical last time to match logical current time from last tick
 	FApp::UpdateLastTime();
@@ -1970,11 +1971,11 @@ void UEngine::UpdateTimeAndHandleMaxTickRate()
 
 		// Did we just switch from a fixed time step to real-time?  If so, then we'll update our
 		// cached 'last time' so our current interval isn't huge (or negative!)
-		if( bTimeWasManipulated )
+		if (bUseFixedTimeStep != bPreviousUseFixedTimeStep)
 		{
-			if ( bUseFixedFrameRate )
+			if (bUseFixedFrameRate)
 			{
-				LastRealTime = CurrentRealTime - (1.0/FixedFrameRate);
+				LastRealTime = CurrentRealTime - (1.0 / FixedFrameRate);
 			}
 			else
 			{
@@ -2127,6 +2128,8 @@ void UEngine::UpdateTimeAndHandleMaxTickRate()
 		}
 	}
 
+	bPreviousUseFixedTimeStep = bUseFixedTimeStep;
+
 #if !UE_BUILD_SHIPPING
 	{
 		float OverrideFPS = CVarSetOverrideFPS.GetValueOnGameThread();
