LOD0 mesh are not hidden when HLOD outliner window is off

Fixed issue where LOD parenting was reset when launching in PIE

diff --git a/Engine/Source/Runtime/Engine/Classes/Components/PrimitiveComponent.h b/Engine/Source/Runtime/Engine/Classes/Components/PrimitiveComponent.h
--- a/Engine/Source/Runtime/Engine/Classes/Components/PrimitiveComponent.h
+++ b/Engine/Source/Runtime/Engine/Classes/Components/PrimitiveComponent.h
@@ -1568,7 +1568,7 @@ public:
 
 private:
 	/** LOD parent primitive to draw instead of this one (multiple UPrim's will point to the same LODParent ) */
-	UPROPERTY(duplicatetransient)
+	UPROPERTY(NonPIEDuplicateTransient)
 	class UPrimitiveComponent* LODParentPrimitive;
 
 public:
