Edigrate CL 9892482 from //UE4/Dev-Console_to_//UE4/Release-4.24

Fix for incorrect vertex indexing in canvas tile renderer

- FCanvasTileRendererItem::FRenderData::InitTileMesh was incorrectly indexing into the vertex array when constructing the tile mesh resulting in malformed tile geometry
  - the indexing was only incorrect for tile index > 0 and hence the issue only manifested when multiple tiles were merged into a single mesh
    - any change in state between consecutive tiles would cause the creation of a new tile mesh and obscure the issue

diff --git a/Engine/Source/Runtime/Engine/Private/TileRendering.cpp b/Engine/Source/Runtime/Engine/Private/TileRendering.cpp
--- a/Engine/Source/Runtime/Engine/Private/TileRendering.cpp
+++ b/Engine/Source/Runtime/Engine/Private/TileRendering.cpp
@@ -92,13 +92,14 @@ void FCanvasTileRendererItem::FRenderData::InitTileMesh(const FSceneView& View,
 	{
 		const FTileInst& Tile = Tiles[i];
 		const uint32 FirstIndex = i * CanvasTileIndexCount;
+		const uint32 FirstVertex = i * CanvasTileVertexCount;
 
-		IndexBuffer.Indices[FirstIndex + 0] = FirstIndex + 0;
-		IndexBuffer.Indices[FirstIndex + 1] = FirstIndex + 1;
-		IndexBuffer.Indices[FirstIndex + 2] = FirstIndex + 2;
-		IndexBuffer.Indices[FirstIndex + 3] = FirstIndex + 2;
-		IndexBuffer.Indices[FirstIndex + 4] = FirstIndex + 1;
-		IndexBuffer.Indices[FirstIndex + 5] = FirstIndex + 3;
+		IndexBuffer.Indices[FirstIndex + 0] = FirstVertex + 0;
+		IndexBuffer.Indices[FirstIndex + 1] = FirstVertex + 1;
+		IndexBuffer.Indices[FirstIndex + 2] = FirstVertex + 2;
+		IndexBuffer.Indices[FirstIndex + 3] = FirstVertex + 2;
+		IndexBuffer.Indices[FirstIndex + 4] = FirstVertex + 1;
+		IndexBuffer.Indices[FirstIndex + 5] = FirstVertex + 3;
 
 		const float X = Tile.X;
 		const float Y = Tile.Y;
@@ -111,33 +112,33 @@ void FCanvasTileRendererItem::FRenderData::InitTileMesh(const FSceneView& View,
 
 		if (bNeedsToSwitchVerticalAxis)
 		{
-			StaticMeshVertexBuffers.PositionVertexBuffer.VertexPosition(i * CanvasTileVertexCount + 0) = FVector(X + SizeX, View.UnscaledViewRect.Height() - (Y + SizeY), 0.0f);
-			StaticMeshVertexBuffers.PositionVertexBuffer.VertexPosition(i * CanvasTileVertexCount + 1) = FVector(X, View.UnscaledViewRect.Height() - (Y + SizeY), 0.0f);
-			StaticMeshVertexBuffers.PositionVertexBuffer.VertexPosition(i * CanvasTileVertexCount + 2) = FVector(X + SizeX, View.UnscaledViewRect.Height() - Y, 0.0f);
-			StaticMeshVertexBuffers.PositionVertexBuffer.VertexPosition(i * CanvasTileVertexCount + 3) = FVector(X, View.UnscaledViewRect.Height() - Y, 0.0f);
-
-			StaticMeshVertexBuffers.StaticMeshVertexBuffer.SetVertexUV(i * CanvasTileVertexCount + 0, 0, FVector2D(U + SizeU, V + SizeV));
-			StaticMeshVertexBuffers.StaticMeshVertexBuffer.SetVertexUV(i * CanvasTileVertexCount + 1, 0, FVector2D(U, V + SizeV));
-			StaticMeshVertexBuffers.StaticMeshVertexBuffer.SetVertexUV(i * CanvasTileVertexCount + 2, 0, FVector2D(U + SizeU, V));
-			StaticMeshVertexBuffers.StaticMeshVertexBuffer.SetVertexUV(i * CanvasTileVertexCount + 3, 0, FVector2D(U, V));
+			StaticMeshVertexBuffers.PositionVertexBuffer.VertexPosition(FirstVertex + 0) = FVector(X + SizeX, View.UnscaledViewRect.Height() - (Y + SizeY), 0.0f);
+			StaticMeshVertexBuffers.PositionVertexBuffer.VertexPosition(FirstVertex + 1) = FVector(X, View.UnscaledViewRect.Height() - (Y + SizeY), 0.0f);
+			StaticMeshVertexBuffers.PositionVertexBuffer.VertexPosition(FirstVertex + 2) = FVector(X + SizeX, View.UnscaledViewRect.Height() - Y, 0.0f);
+			StaticMeshVertexBuffers.PositionVertexBuffer.VertexPosition(FirstVertex + 3) = FVector(X, View.UnscaledViewRect.Height() - Y, 0.0f);
+
+			StaticMeshVertexBuffers.StaticMeshVertexBuffer.SetVertexUV(FirstVertex + 0, 0, FVector2D(U + SizeU, V + SizeV));
+			StaticMeshVertexBuffers.StaticMeshVertexBuffer.SetVertexUV(FirstVertex + 1, 0, FVector2D(U, V + SizeV));
+			StaticMeshVertexBuffers.StaticMeshVertexBuffer.SetVertexUV(FirstVertex + 2, 0, FVector2D(U + SizeU, V));
+			StaticMeshVertexBuffers.StaticMeshVertexBuffer.SetVertexUV(FirstVertex + 3, 0, FVector2D(U, V));
 		}
 		else
 		{
-			StaticMeshVertexBuffers.PositionVertexBuffer.VertexPosition(i * CanvasTileVertexCount + 0) = FVector(X + SizeX, Y, 0.0f);
-			StaticMeshVertexBuffers.PositionVertexBuffer.VertexPosition(i * CanvasTileVertexCount + 1) = FVector(X, Y, 0.0f);
-			StaticMeshVertexBuffers.PositionVertexBuffer.VertexPosition(i * CanvasTileVertexCount + 2) = FVector(X + SizeX, Y + SizeY, 0.0f);
-			StaticMeshVertexBuffers.PositionVertexBuffer.VertexPosition(i * CanvasTileVertexCount + 3) = FVector(X, Y + SizeY, 0.0f);
-
-			StaticMeshVertexBuffers.StaticMeshVertexBuffer.SetVertexUV(i * CanvasTileVertexCount + 0, 0, FVector2D(U + SizeU, V));
-			StaticMeshVertexBuffers.StaticMeshVertexBuffer.SetVertexUV(i * CanvasTileVertexCount + 1, 0, FVector2D(U, V));
-			StaticMeshVertexBuffers.StaticMeshVertexBuffer.SetVertexUV(i * CanvasTileVertexCount + 2, 0, FVector2D(U + SizeU, V + SizeV));
-			StaticMeshVertexBuffers.StaticMeshVertexBuffer.SetVertexUV(i * CanvasTileVertexCount + 3, 0, FVector2D(U, V + SizeV));
+			StaticMeshVertexBuffers.PositionVertexBuffer.VertexPosition(FirstVertex + 0) = FVector(X + SizeX, Y, 0.0f);
+			StaticMeshVertexBuffers.PositionVertexBuffer.VertexPosition(FirstVertex + 1) = FVector(X, Y, 0.0f);
+			StaticMeshVertexBuffers.PositionVertexBuffer.VertexPosition(FirstVertex + 2) = FVector(X + SizeX, Y + SizeY, 0.0f);
+			StaticMeshVertexBuffers.PositionVertexBuffer.VertexPosition(FirstVertex + 3) = FVector(X, Y + SizeY, 0.0f);
+
+			StaticMeshVertexBuffers.StaticMeshVertexBuffer.SetVertexUV(FirstVertex + 0, 0, FVector2D(U + SizeU, V));
+			StaticMeshVertexBuffers.StaticMeshVertexBuffer.SetVertexUV(FirstVertex + 1, 0, FVector2D(U, V));
+			StaticMeshVertexBuffers.StaticMeshVertexBuffer.SetVertexUV(FirstVertex + 2, 0, FVector2D(U + SizeU, V + SizeV));
+			StaticMeshVertexBuffers.StaticMeshVertexBuffer.SetVertexUV(FirstVertex + 3, 0, FVector2D(U, V + SizeV));
 		}
 
 		for (int j = 0; j < CanvasTileVertexCount; j++)
 		{
-			StaticMeshVertexBuffers.StaticMeshVertexBuffer.SetVertexTangents(i * CanvasTileVertexCount + j, FVector(1.0f, 0.0f, 0.0f), FVector(0.0f, 1.0f, 0.0f), FVector(0.0f, 0.0f, 1.0f));
-			StaticMeshVertexBuffers.ColorVertexBuffer.VertexColor(i * CanvasTileVertexCount + j) = Tile.InColor;
+			StaticMeshVertexBuffers.StaticMeshVertexBuffer.SetVertexTangents(FirstVertex + j, FVector(1.0f, 0.0f, 0.0f), FVector(0.0f, 1.0f, 0.0f), FVector(0.0f, 0.0f, 1.0f));
+			StaticMeshVertexBuffers.ColorVertexBuffer.VertexColor(FirstVertex + j) = Tile.InColor;
 		}
 	}
 
