Chaos: Fix missing code for vehicle sleep state

diff --git a/Engine/Source/Runtime/Engine/Private/PhysicsEngine/Experimental/PhysInterface_Chaos.cpp b/Engine/Source/Runtime/Engine/Private/PhysicsEngine/Experimental/PhysInterface_Chaos.cpp
--- a/Engine/Source/Runtime/Engine/Private/PhysicsEngine/Experimental/PhysInterface_Chaos.cpp
+++ b/Engine/Source/Runtime/Engine/Private/PhysicsEngine/Experimental/PhysInterface_Chaos.cpp
@@ -302,8 +302,7 @@ bool FPhysInterface_Chaos::IsKinematic_AssumesLocked(const FPhysicsActorHandle&
 
 bool FPhysInterface_Chaos::IsSleeping(const FPhysicsActorHandle& InActorReference)
 {
-	// #todo : Implement
-	return false;
+	return InActorReference->ObjectState() == Chaos::EObjectStateType::Sleeping;
 }
 
 bool FPhysInterface_Chaos::IsCcdEnabled(const FPhysicsActorHandle& InActorReference)
@@ -313,8 +312,7 @@ bool FPhysInterface_Chaos::IsCcdEnabled(const FPhysicsActorHandle& InActorRefere
 
 bool FPhysInterface_Chaos::IsInScene(const FPhysicsActorHandle& InActorReference)
 {
-	// TODO: Implement
-	return false;
+	return (GetCurrentScene(InActorReference) != nullptr);
 }
 
 FPhysScene* FPhysInterface_Chaos::GetCurrentScene(const FPhysicsActorHandle& InHandle)
@@ -357,18 +355,26 @@ void FPhysInterface_Chaos::SetSendsSleepNotifies_AssumesLocked(const FPhysicsAct
 
 void FPhysInterface_Chaos::PutToSleep_AssumesLocked(const FPhysicsActorHandle& InActorReference)
 {
-	// #todo : Implement
+	if (Chaos::TPBDRigidParticle<float, 3>* Particle = InActorReference->AsDynamic())
+	{
+		Particle->SetObjectState(Chaos::EObjectStateType::Sleeping);
+	}
+	
 }
 
 void FPhysInterface_Chaos::WakeUp_AssumesLocked(const FPhysicsActorHandle& InActorReference)
 {
-	// #todo : Implement
+	if (Chaos::TPBDRigidParticle<float, 3>* Particle = InActorReference->AsDynamic())
+	{
+		Particle->SetObjectState(Chaos::EObjectStateType::Dynamic);
+	}
 }
 
 void FPhysInterface_Chaos::SetIsKinematic_AssumesLocked(const FPhysicsActorHandle& InActorReference, bool bIsKinematic)
 {
 	if (Chaos::TPBDRigidParticle<float, 3>* Particle = InActorReference->AsDynamic())
 	{
+		//#todo: what if object state has been previously set to sleeping?
 		const Chaos::EObjectStateType NewState
 			= bIsKinematic
 			? Chaos::EObjectStateType::Kinematic
