Added option on SetCollisionProfileName to skip updating overlaps for performance.  Also, reversed logic of bDeferUpdateOverlaps to be more consistant with other calls.

diff --git a/Engine/Source/Runtime/Engine/Classes/Components/PrimitiveComponent.h b/Engine/Source/Runtime/Engine/Classes/Components/PrimitiveComponent.h
--- a/Engine/Source/Runtime/Engine/Classes/Components/PrimitiveComponent.h
+++ b/Engine/Source/Runtime/Engine/Classes/Components/PrimitiveComponent.h
@@ -1483,7 +1483,7 @@ public:
 	 * @param InCollisionProfileName : New Profile Name
 	 */
 	UFUNCTION(BlueprintCallable, Category="Collision")	
-	virtual void SetCollisionProfileName(FName InCollisionProfileName);
+	virtual void SetCollisionProfileName(FName InCollisionProfileName, bool bUpdateOverlaps=true);
 
 	/** Get the collision profile name */
 	UFUNCTION(BlueprintPure, Category="Collision")
@@ -2191,7 +2191,7 @@ public:
 	
 protected:
 	/** Called when the BodyInstance ResponseToChannels, CollisionEnabled or bNotifyRigidBodyCollision changes, in case subclasses want to use that information. */
-	virtual void OnComponentCollisionSettingsChanged(bool bDeferUpdateOverlaps = false);
+	virtual void OnComponentCollisionSettingsChanged(bool bUpdateOverlaps=true);
 
 	/** Ends all current component overlaps. Generally used when destroying this component or when it can no longer generate overlaps. */
 	void ClearComponentOverlaps(bool bDoNotifies, bool bSkipNotifySelf);
diff --git a/Engine/Source/Runtime/Engine/Classes/Components/SkeletalMeshComponent.h b/Engine/Source/Runtime/Engine/Classes/Components/SkeletalMeshComponent.h
--- a/Engine/Source/Runtime/Engine/Classes/Components/SkeletalMeshComponent.h
+++ b/Engine/Source/Runtime/Engine/Classes/Components/SkeletalMeshComponent.h
@@ -1539,7 +1539,7 @@ public:
 	virtual bool IsAnyRigidBodyAwake() override;
 	virtual void SetEnableGravity(bool bGravityEnabled);
 	virtual bool IsGravityEnabled() const override;
-	virtual void OnComponentCollisionSettingsChanged(bool bDeferUpdateOverlaps = false) override;
+	virtual void OnComponentCollisionSettingsChanged(bool bUpdateOverlaps=true) override;
 	virtual void SetPhysMaterialOverride(UPhysicalMaterial* NewPhysMaterial) override;
 	virtual bool GetSquaredDistanceToCollision(const FVector& Point, float& OutSquaredDistance, FVector& OutClosestPointOnCollision) const override;
 
diff --git a/Engine/Source/Runtime/Engine/Private/PrimitiveComponentPhysics.cpp b/Engine/Source/Runtime/Engine/Private/PrimitiveComponentPhysics.cpp
--- a/Engine/Source/Runtime/Engine/Private/PrimitiveComponentPhysics.cpp
+++ b/Engine/Source/Runtime/Engine/Private/PrimitiveComponentPhysics.cpp
@@ -908,7 +908,7 @@ void UPrimitiveComponent::SetCollisionEnabled(ECollisionEnabled::Type NewType)
 }
 
 // @todo : implement skeletalmeshcomponent version
-void UPrimitiveComponent::SetCollisionProfileName(FName InCollisionProfileName)
+void UPrimitiveComponent::SetCollisionProfileName(FName InCollisionProfileName, bool bUpdateOverlaps)
 {
 	SCOPE_CYCLE_COUNTER(STAT_PrimComp_SetCollisionProfileName);
 
@@ -930,7 +930,7 @@ void UPrimitiveComponent::SetCollisionProfileName(FName InCollisionProfileName)
 		{
 			EnsurePhysicsStateCreated();
 		}
-		OnComponentCollisionSettingsChanged();
+		OnComponentCollisionSettingsChanged(bUpdateOverlaps);
 	}
 }
 
@@ -942,10 +942,10 @@ FName UPrimitiveComponent::GetCollisionProfileName() const
 void UPrimitiveComponent::OnActorEnableCollisionChanged()
 {
 	BodyInstance.UpdatePhysicsFilterData();
-	OnComponentCollisionSettingsChanged(true);
+	OnComponentCollisionSettingsChanged(false);
 }
 
-void UPrimitiveComponent::OnComponentCollisionSettingsChanged(bool bDeferUpdateOverlaps)
+void UPrimitiveComponent::OnComponentCollisionSettingsChanged(bool bUpdateOverlaps)
 {
 	if (IsRegistered() && !IsTemplate())			// not for CDOs
 	{
@@ -955,7 +955,7 @@ void UPrimitiveComponent::OnComponentCollisionSettingsChanged(bool bDeferUpdateO
 			ClearSkipUpdateOverlaps();
 		}
 
-		if (!bDeferUpdateOverlaps)
+		if (bUpdateOverlaps)
 		{
 			UpdateOverlaps();
 		}
diff --git a/Engine/Source/Runtime/Engine/Private/SkeletalMeshComponentPhysics.cpp b/Engine/Source/Runtime/Engine/Private/SkeletalMeshComponentPhysics.cpp
--- a/Engine/Source/Runtime/Engine/Private/SkeletalMeshComponentPhysics.cpp
+++ b/Engine/Source/Runtime/Engine/Private/SkeletalMeshComponentPhysics.cpp
@@ -206,7 +206,7 @@ void USkeletalMeshComponent::SetSimulatePhysics(bool bSimulate)
  	UpdateClothTickRegisteredState();
 }
 
-void USkeletalMeshComponent::OnComponentCollisionSettingsChanged(bool bDeferUpdateOverlaps)
+void USkeletalMeshComponent::OnComponentCollisionSettingsChanged(bool bUpdateOverlaps)
 {
 	for(int32 i=0; i<Bodies.Num(); i++)
 	{
@@ -221,7 +221,7 @@ void USkeletalMeshComponent::OnComponentCollisionSettingsChanged(bool bDeferUpda
 		((FSkeletalMeshSceneProxy*)SceneProxy)->SetCollisionEnabled_GameThread(IsCollisionEnabled());
 	}
 
-	Super::OnComponentCollisionSettingsChanged(bDeferUpdateOverlaps);
+	Super::OnComponentCollisionSettingsChanged(bUpdateOverlaps);
 }
 
 void USkeletalMeshComponent::AddRadialImpulse(FVector Origin, float Radius, float Strength, ERadialImpulseFalloff Falloff, bool bVelChange)
