Backing out of CL 10083951. Turns out this might not be such a good idea after all.

#rnx

diff --git a/Engine/Source/Runtime/Engine/Classes/Components/SkeletalMeshComponent.h b/Engine/Source/Runtime/Engine/Classes/Components/SkeletalMeshComponent.h
--- a/Engine/Source/Runtime/Engine/Classes/Components/SkeletalMeshComponent.h
+++ b/Engine/Source/Runtime/Engine/Classes/Components/SkeletalMeshComponent.h
@@ -170,14 +170,6 @@ namespace EAnimationMode
 	};
 }
 
-UENUM()
-enum class ELinkedAnimationUpdateOrder : uint8
-{
-	UpdateAnimationBeforeAnimScriptInstance,
-	UpdateAnimationAfterAnimScriptInstance,
-	DoNotUpdate
-};
-
 UENUM()
 namespace EPhysicsTransformUpdateMode
 {
@@ -439,10 +431,6 @@ private:
 	uint8 bDisablePostProcessBlueprint:1;
 
 public:
-	/** The order in which linked animation is evaluated with respect to the component's animation */
-	UPROPERTY(EditAnywhere, BlueprintReadOnly, Category = Animation)
-	ELinkedAnimationUpdateOrder LinkedAnimationEvaluationOrder;
-
 	/** Indicates that simulation (if it's enabled) is entirely responsible for children transforms. This is only ok if you are not animating attachment points relative to the simulation */
 	uint8 bSimulationUpdatesChildTransforms:1;
 
diff --git a/Engine/Source/Runtime/Engine/Private/Components/SkeletalMeshComponent.cpp b/Engine/Source/Runtime/Engine/Private/Components/SkeletalMeshComponent.cpp
--- a/Engine/Source/Runtime/Engine/Private/Components/SkeletalMeshComponent.cpp
+++ b/Engine/Source/Runtime/Engine/Private/Components/SkeletalMeshComponent.cpp
@@ -1004,13 +1004,10 @@ void USkeletalMeshComponent::TickAnimation(float DeltaTime, bool bNeedsValidRoot
 
 		// We update linked instances first incase we're using either root motion or non-threaded update.
 		// This ensures that we go through the pre update process and initialize the proxies correctly.
-		if (LinkedAnimationEvaluationOrder == ELinkedAnimationUpdateOrder::UpdateAnimationBeforeAnimScriptInstance)
+		for (UAnimInstance* LinkedInstance : LinkedInstances)
 		{
-			for (UAnimInstance* LinkedInstance : LinkedInstances)
-			{
-				// Sub anim instances are always forced to do a parallel update 
-				LinkedInstance->UpdateAnimation(DeltaTime * GlobalAnimRateScale, false, UAnimInstance::EUpdateAnimationFlag::ForceParallelUpdate);
-			}
+			// Sub anim instances are always forced to do a parallel update 
+			LinkedInstance->UpdateAnimation(DeltaTime * GlobalAnimRateScale, false, UAnimInstance::EUpdateAnimationFlag::ForceParallelUpdate);
 		}
 
 		if (AnimScriptInstance != nullptr)
@@ -1019,16 +1016,6 @@ void USkeletalMeshComponent::TickAnimation(float DeltaTime, bool bNeedsValidRoot
 			AnimScriptInstance->UpdateAnimation(DeltaTime * GlobalAnimRateScale, bNeedsValidRootMotion);
 		}
 
-		if (LinkedAnimationEvaluationOrder == ELinkedAnimationUpdateOrder::UpdateAnimationAfterAnimScriptInstance)
-		{
-			for (UAnimInstance* LinkedInstance : LinkedInstances)
-			{
-				// Sub anim instances are always forced to do a parallel update 
-				LinkedInstance->UpdateAnimation(DeltaTime * GlobalAnimRateScale, false, UAnimInstance::EUpdateAnimationFlag::ForceParallelUpdate);
-			}
-		}
-
-
 		if(ShouldUpdatePostProcessInstance())
 		{
 			PostProcessAnimInstance->UpdateAnimation(DeltaTime * GlobalAnimRateScale, false);
