Fix ensure on SetMass failing cook. (Change to CHAOS_ENSURE)

diff --git a/Engine/Source/Runtime/Engine/Private/PhysicsEngine/Experimental/PhysInterface_Chaos.cpp b/Engine/Source/Runtime/Engine/Private/PhysicsEngine/Experimental/PhysInterface_Chaos.cpp
--- a/Engine/Source/Runtime/Engine/Private/PhysicsEngine/Experimental/PhysInterface_Chaos.cpp
+++ b/Engine/Source/Runtime/Engine/Private/PhysicsEngine/Experimental/PhysInterface_Chaos.cpp
@@ -657,7 +657,7 @@ void FPhysInterface_Chaos::SetMass_AssumesLocked(FPhysicsActorHandle& InActorRef
 	if (Chaos::TPBDRigidParticle<float, 3 >* RigidParticle = InActorReference->AsDynamic())
 	{
 		RigidParticle->SetM(InMass);
-		if (ensure(!FMath::IsNearlyZero(InMass)))
+		if (CHAOS_ENSURE(!FMath::IsNearlyZero(InMass)))
 		{
 			RigidParticle->SetInvM(1./InMass);
 		}
@@ -672,7 +672,7 @@ void FPhysInterface_Chaos::SetMassSpaceInertiaTensor_AssumesLocked(FPhysicsActor
 {
 	if (Chaos::TPBDRigidParticle<float, 3 >* RigidParticle = InActorReference->AsDynamic())
 	{
-		if(ensure(!FMath::IsNearlyZero(InTensor.X)) && ensure(!FMath::IsNearlyZero(InTensor.Y)) && ensure(!FMath::IsNearlyZero(InTensor.Z)) )
+		if(CHAOS_ENSURE(!FMath::IsNearlyZero(InTensor.X)) && CHAOS_ENSURE(!FMath::IsNearlyZero(InTensor.Y)) && CHAOS_ENSURE(!FMath::IsNearlyZero(InTensor.Z)) )
 		{
 			RigidParticle->SetI(Chaos::PMatrix<float, 3, 3>(InTensor.X, InTensor.Y, InTensor.Z));
 			RigidParticle->SetInvI(Chaos::PMatrix<float, 3, 3>(1./InTensor.X, 1./InTensor.Y, 1./InTensor.Z));
