Cameras: Move implementation files to the proper place.

diff --git a/Engine/Source/Runtime/Engine/Classes/Camera/CameraShakeSourceActor.cpp b/Engine/Source/Runtime/Engine/Classes/Camera/CameraShakeSourceActor.cpp
deleted file mode 100644
diff --git a/Engine/Source/Runtime/Engine/Classes/Camera/CameraShakeSourceComponent.cpp b/Engine/Source/Runtime/Engine/Classes/Camera/CameraShakeSourceComponent.cpp
deleted file mode 100644
diff --git a/Engine/Source/Runtime/Engine/Private/Camera/CameraShakeSourceActor.cpp b/Engine/Source/Runtime/Engine/Private/Camera/CameraShakeSourceActor.cpp
new file mode 100644
--- /dev/null
+++ b/Engine/Source/Runtime/Engine/Private/Camera/CameraShakeSourceActor.cpp
@@ -0,0 +1,12 @@
+// Copyright 1998-2019 Epic Games, Inc. All Rights Reserved.
+
+#include "CameraShakeSourceActor.h"
+#include "CameraShakeSourceComponent.h"
+
+
+ACameraShakeSourceActor::ACameraShakeSourceActor(const FObjectInitializer& ObjectInitializer)
+	: Super(ObjectInitializer)
+{
+	CameraShakeSourceComponent = ObjectInitializer.CreateDefaultSubobject<UCameraShakeSourceComponent>(this, TEXT("CameraShakeSourceComponent"));
+	RootComponent = CameraShakeSourceComponent;
+}
diff --git a/Engine/Source/Runtime/Engine/Private/Camera/CameraShakeSourceComponent.cpp b/Engine/Source/Runtime/Engine/Private/Camera/CameraShakeSourceComponent.cpp
new file mode 100644
--- /dev/null
+++ b/Engine/Source/Runtime/Engine/Private/Camera/CameraShakeSourceComponent.cpp
@@ -0,0 +1,126 @@
+// Copyright 1998-2019 Epic Games, Inc. All Rights Reserved.
+
+#include "CameraShakeSourceComponent.h"
+#include "CameraShake.h"
+#include "Camera/CameraModifier_CameraShake.h"
+#include "CinematicCameraModule.h"
+#include "Components/BillboardComponent.h"
+#include "Engine/Engine.h"
+#include "Engine/Texture2D.h"
+#include "Engine/World.h"
+#include "GameFramework/PlayerController.h"
+#include "Modules/ModuleManager.h"
+#include "UObject/ConstructorHelpers.h"
+
+
+UCameraShakeSourceComponent::UCameraShakeSourceComponent(const FObjectInitializer& ObjectInitializer)
+	: Super(ObjectInitializer)
+	, Attenuation(ECameraShakeAttenuation::Quadratic)
+	, InnerAttenuationRadius(100.f)
+	, OuterAttenuationRadius(1000.f)
+	, bAutoPlay(false)
+{
+#if WITH_EDITORONLY_DATA
+	bVisualizeComponent = true;
+
+	if (!IsRunningCommandlet())
+	{
+		static ConstructorHelpers::FObjectFinder<UTexture2D> StaticTexture(TEXT("/Engine/EditorResources/S_CameraShakeSource"));
+
+		EditorSpriteTexture = StaticTexture.Object;
+		EditorSpriteTextureScale = .5f;
+	}
+#endif
+}
+
+void UCameraShakeSourceComponent::OnRegister()
+{
+	Super::OnRegister();
+	UpdateEditorSpriteTexture();
+}
+
+void UCameraShakeSourceComponent::UpdateEditorSpriteTexture()
+{
+#if WITH_EDITORONLY_DATA
+	if (SpriteComponent != nullptr)
+	{
+		SpriteComponent->SetSprite(EditorSpriteTexture);
+		SpriteComponent->SetRelativeScale3D(FVector(EditorSpriteTextureScale));
+	}
+#endif
+}
+
+void UCameraShakeSourceComponent::BeginPlay()
+{
+	Super::BeginPlay();
+
+	if (bAutoPlay)
+	{
+		Play();
+	}
+}
+
+void UCameraShakeSourceComponent::EndPlay(const EEndPlayReason::Type EndPlayReason)
+{
+	StopAllCameraShakes();
+
+	Super::EndPlay(EndPlayReason);
+}
+
+void UCameraShakeSourceComponent::Play()
+{
+	PlayCameraShake(CameraShake);
+}
+
+void UCameraShakeSourceComponent::PlayCameraShake(TSubclassOf<UCameraShake> InCameraShake)
+{
+	for (FConstPlayerControllerIterator Iterator = GetWorld()->GetPlayerControllerIterator(); Iterator; ++Iterator)
+	{
+		APlayerController* PlayerController = Iterator->Get();
+		if (PlayerController != nullptr && PlayerController->PlayerCameraManager != nullptr)
+		{
+			PlayerController->ClientPlayCameraShakeFromSource(InCameraShake, this);
+		}
+	}
+}
+
+void UCameraShakeSourceComponent::StopAllCameraShakes(bool bImmediately)
+{
+	for (FConstPlayerControllerIterator Iterator = GetWorld()->GetPlayerControllerIterator(); Iterator; ++Iterator)
+	{
+		APlayerController* PlayerController = Iterator->Get();
+		if (PlayerController != nullptr && PlayerController->PlayerCameraManager != nullptr)
+		{
+			PlayerController->ClientStopCameraShakesFromSource(this, bImmediately);
+		}
+	}
+}
+
+float UCameraShakeSourceComponent::GetAttenuationFactor(const FVector& Location) const
+{
+	const FTransform& SourceTransform = GetComponentTransform();
+	const FVector LocationToSource = SourceTransform.GetTranslation() - Location;
+
+	float AttFactor = 1.0f;
+	switch (Attenuation)
+	{
+		case ECameraShakeAttenuation::Quadratic:
+			AttFactor = 2.0f;
+			break;
+		default:
+			break;
+	}
+
+	if (InnerAttenuationRadius < OuterAttenuationRadius)
+	{
+		float DistFactor = (LocationToSource.Size() - InnerAttenuationRadius) / (OuterAttenuationRadius - InnerAttenuationRadius);
+		DistFactor = 1.f - FMath::Clamp(DistFactor, 0.f, 1.f);
+		return FMath::Pow(DistFactor, AttFactor);
+	}
+	else if (OuterAttenuationRadius > 0)
+	{
+		// Just cut the intensity after the end distance.
+		return (LocationToSource.SizeSquared() < FMath::Square(OuterAttenuationRadius)) ? 1.f : 0.f;
+	}
+	return 1.f;
+}
