Update LocalPlayer Projection Data to explicitly fail to generate if the local player's viewport Size is 0, within a validly sized FViewport.

diff --git a/Engine/Source/Runtime/Engine/Private/LocalPlayer.cpp b/Engine/Source/Runtime/Engine/Private/LocalPlayer.cpp
--- a/Engine/Source/Runtime/Engine/Private/LocalPlayer.cpp
+++ b/Engine/Source/Runtime/Engine/Private/LocalPlayer.cpp
@@ -1021,7 +1021,7 @@ bool ULocalPlayer::GetPixelPoint(const FSceneViewProjectionData& ProjectionData,
 bool ULocalPlayer::GetProjectionData(FViewport* Viewport, EStereoscopicPass StereoPass, FSceneViewProjectionData& ProjectionData) const
 {
 	// If the actor
-	if ((Viewport == NULL) || (PlayerController == NULL) || (Viewport->GetSizeXY().X == 0) || (Viewport->GetSizeXY().Y == 0))
+	if ((Viewport == NULL) || (PlayerController == NULL) || (Viewport->GetSizeXY().X == 0) || (Viewport->GetSizeXY().Y == 0) || (Size.X == 0) || (Size.Y == 0))
 	{
 		return false;
 	}
