Fix include paths

diff --git a/Engine/Source/Runtime/Engine/Private/Camera/CameraShakeSourceActor.cpp b/Engine/Source/Runtime/Engine/Private/Camera/CameraShakeSourceActor.cpp
--- a/Engine/Source/Runtime/Engine/Private/Camera/CameraShakeSourceActor.cpp
+++ b/Engine/Source/Runtime/Engine/Private/Camera/CameraShakeSourceActor.cpp
@@ -1,7 +1,7 @@
 // Copyright 1998-2019 Epic Games, Inc. All Rights Reserved.
 
-#include "CameraShakeSourceActor.h"
-#include "CameraShakeSourceComponent.h"
+#include "Camera/CameraShakeSourceActor.h"
+#include "Camera/CameraShakeSourceComponent.h"
 
 
 ACameraShakeSourceActor::ACameraShakeSourceActor(const FObjectInitializer& ObjectInitializer)
diff --git a/Engine/Source/Runtime/Engine/Private/Camera/CameraShakeSourceComponent.cpp b/Engine/Source/Runtime/Engine/Private/Camera/CameraShakeSourceComponent.cpp
--- a/Engine/Source/Runtime/Engine/Private/Camera/CameraShakeSourceComponent.cpp
+++ b/Engine/Source/Runtime/Engine/Private/Camera/CameraShakeSourceComponent.cpp
@@ -1,7 +1,7 @@
 // Copyright 1998-2019 Epic Games, Inc. All Rights Reserved.
 
-#include "CameraShakeSourceComponent.h"
-#include "CameraShake.h"
+#include "Camera/CameraShakeSourceComponent.h"
+#include "Camera/CameraShake.h"
 #include "Camera/CameraModifier_CameraShake.h"
 #include "CinematicCameraModule.h"
 #include "Components/BillboardComponent.h"
