Change the chaos physics interface IsDynamic function to match PhysX behavior
(fixes object mass properties not getting initialized)

diff --git a/Engine/Source/Runtime/Engine/Private/PhysicsEngine/Experimental/PhysInterface_Chaos.cpp b/Engine/Source/Runtime/Engine/Private/PhysicsEngine/Experimental/PhysInterface_Chaos.cpp
--- a/Engine/Source/Runtime/Engine/Private/PhysicsEngine/Experimental/PhysInterface_Chaos.cpp
+++ b/Engine/Source/Runtime/Engine/Private/PhysicsEngine/Experimental/PhysInterface_Chaos.cpp
@@ -282,7 +282,8 @@ bool FPhysInterface_Chaos::IsRigidBody(const FPhysicsActorHandle& InActorReferen
 
 bool FPhysInterface_Chaos::IsDynamic(const FPhysicsActorHandle& InActorReference)
 {
-	return InActorReference->ObjectState() == Chaos::EObjectStateType::Dynamic;
+	// Do this to match the PhysX interface behavior: :( :( :(
+	return !IsStatic(InActorReference);
 }
 
 bool FPhysInterface_Chaos::IsStatic(const FPhysicsActorHandle& InActorReference)
