#UE4 Added more info to a log

diff --git a/Engine/Source/Runtime/Engine/Private/Components/SceneComponent.cpp b/Engine/Source/Runtime/Engine/Private/Components/SceneComponent.cpp
--- a/Engine/Source/Runtime/Engine/Private/Components/SceneComponent.cpp
+++ b/Engine/Source/Runtime/Engine/Private/Components/SceneComponent.cpp
@@ -477,7 +477,7 @@ bool USceneComponent::NeedsLoadForTargetPlatform(const ITargetPlatform* TargetPl
 		// Child not culled, so warn
 		if(!bDescendantsCulled)
 		{
-			UE_LOG(LogSceneComponent, Warning, TEXT("Component %s not cooked out for client because descendants were not also cooked out."), *GetName());
+			UE_LOG(LogSceneComponent, Warning, TEXT("Component %s not cooked out for client because descendants were not also cooked out."), *GetPathName());
 			return true;
 		}
 
