Added a serial number to FMaterialRenderProxy tracking when the cached uniforms array is modified or invalidated

diff --git a/Engine/Source/Runtime/Engine/Private/Materials/MaterialShared.cpp b/Engine/Source/Runtime/Engine/Private/Materials/MaterialShared.cpp
--- a/Engine/Source/Runtime/Engine/Private/Materials/MaterialShared.cpp
+++ b/Engine/Source/Runtime/Engine/Private/Materials/MaterialShared.cpp
@@ -2447,6 +2447,7 @@ void FMaterialRenderProxy::EvaluateUniformExpressions(FUniformExpressionCache& O
 	OutUniformExpressionCache.ParameterCollections = UniformExpressionSet.ParameterCollections;
 
 	OutUniformExpressionCache.bUpToDate = true;
+	++UniformExpressionCacheSerialNumber;
 }
 
 void FMaterialRenderProxy::CacheUniformExpressions(bool bRecreateUniformBuffer)
@@ -2499,6 +2500,7 @@ void FMaterialRenderProxy::InvalidateUniformExpressionCache(bool bRecreateUnifor
 		HasVirtualTextureCallbacks = false;
 	}
 
+	++UniformExpressionCacheSerialNumber;
 	for (int32 i = 0; i < ERHIFeatureLevel::Num; ++i)
 	{
 		UniformExpressionCache[i].bUpToDate = false;
@@ -2524,7 +2526,7 @@ void FMaterialRenderProxy::UpdateUniformExpressionCacheIfNeeded(ERHIFeatureLevel
 
 		// Don't cache uniform expressions if an entirely different FMaterialRenderProxy is going to be used for rendering
 		if (!FallbackMaterialRenderProxy)
-{
+		{
 			FMaterialRenderContext MaterialRenderContext(this, Material, nullptr);
 			MaterialRenderContext.bShowSelection = GIsEditor;
 			EvaluateUniformExpressions(UniformExpressionCache[InFeatureLevel], MaterialRenderContext);
diff --git a/Engine/Source/Runtime/Engine/Public/MaterialShared.h b/Engine/Source/Runtime/Engine/Public/MaterialShared.h
--- a/Engine/Source/Runtime/Engine/Public/MaterialShared.h
+++ b/Engine/Source/Runtime/Engine/Public/MaterialShared.h
@@ -2037,6 +2037,7 @@ public:
 		return DeferredUniformExpressionCacheRequests.Num() > 0;
 	}
 
+	int32 GetExpressionCacheSerialNumber() const { return UniformExpressionCacheSerialNumber; }
 private:
 	IAllocatedVirtualTexture* GetPreallocatedVTStack(const FMaterialRenderContext& Context, const FUniformExpressionSet& UniformExpressionSet, const FMaterialVirtualTextureStack& VTStack) const;
 	IAllocatedVirtualTexture* AllocateVTStack(const FMaterialRenderContext& Context, const FUniformExpressionSet& UniformExpressionSet, const FMaterialVirtualTextureStack& VTStack) const;
@@ -2044,6 +2045,9 @@ private:
 	/** 0 if not set, game thread pointer, do not dereference, only for comparison */
 	const USubsurfaceProfile* SubsurfaceProfileRT;
 
+	/** Incremented each time UniformExpressionCache is modified */
+	mutable int32 UniformExpressionCacheSerialNumber = 0;
+
 	/** For tracking down a bug accessing a deleted proxy. */
 	mutable int8 MarkedForGarbageCollection : 1;
 	mutable int8 DeletedFlag : 1;
