# Added blueprint exposed Montage_StopByGroupName
#rnx

diff --git a/Engine/Source/Runtime/Engine/Classes/Animation/AnimInstance.h b/Engine/Source/Runtime/Engine/Classes/Animation/AnimInstance.h
--- a/Engine/Source/Runtime/Engine/Classes/Animation/AnimInstance.h
+++ b/Engine/Source/Runtime/Engine/Classes/Animation/AnimInstance.h
@@ -566,6 +566,10 @@ public:
 	UFUNCTION(BlueprintCallable, Category = "Montage")
 	void Montage_Stop(float InBlendOutTime, const UAnimMontage* Montage = NULL);
 
+	/** Stops all active montages belonging to a group. */
+	UFUNCTION(BlueprintCallable, Category = "Montage")
+	void Montage_StopGroupByName(float InBlendOutTime, FName GroupName);
+
 	/** Pauses the animation montage. If reference is NULL, it will pause ALL active montages. */
 	UFUNCTION(BlueprintCallable, Category = "Montage")
 	void Montage_Pause(const UAnimMontage* Montage = NULL);
diff --git a/Engine/Source/Runtime/Engine/Private/Animation/AnimInstance.cpp b/Engine/Source/Runtime/Engine/Private/Animation/AnimInstance.cpp
--- a/Engine/Source/Runtime/Engine/Private/Animation/AnimInstance.cpp
+++ b/Engine/Source/Runtime/Engine/Private/Animation/AnimInstance.cpp
@@ -1891,6 +1891,18 @@ void UAnimInstance::Montage_Stop(float InBlendOutTime, const UAnimMontage* Monta
 	}
 }
 
+void UAnimInstance::Montage_StopGroupByName(float InBlendOutTime, FName GroupName)
+{
+	for (int32 InstanceIndex = 0; InstanceIndex < MontageInstances.Num(); InstanceIndex++)
+	{
+		FAnimMontageInstance* MontageInstance = MontageInstances[InstanceIndex];
+		if (MontageInstance && MontageInstance->Montage && MontageInstance->IsActive() && (MontageInstance->Montage->GetGroupName() == GroupName))
+		{
+			MontageInstances[InstanceIndex]->Stop(FAlphaBlend(MontageInstance->Montage->BlendOut, InBlendOutTime));
+		}
+	}
+}
+
 void UAnimInstance::Montage_Pause(const UAnimMontage* Montage)
 {
 	if (Montage)
