Allow Shader stats  platform to be pulled off of -targetplatform in addition to -target

diff --git a/Engine/Source/Runtime/Engine/Private/ShaderCompiler/ShaderCompiler.cpp b/Engine/Source/Runtime/Engine/Private/ShaderCompiler/ShaderCompiler.cpp
--- a/Engine/Source/Runtime/Engine/Private/ShaderCompiler/ShaderCompiler.cpp
+++ b/Engine/Source/Runtime/Engine/Private/ShaderCompiler/ShaderCompiler.cpp
@@ -1554,7 +1554,10 @@ void FShaderCompilerStats::WriteStats()
 			{
 				FString TargetType = TEXT("Default");
 				FParse::Value(FCommandLine::Get(), TEXT("target="), TargetType);
-
+				if (TargetType == TEXT("Default"))
+				{
+					FParse::Value(FCommandLine::Get(), TEXT("targetplatform="), TargetType);
+				}
 				FString CopyLocation = FPaths::Combine(*MirrorLocation, FApp::GetProjectName(), *FApp::GetBranchName(), FString::Printf(TEXT("Stats-Latest(%s).csv"), *TargetType));
 				IFileManager::Get().Copy(*CopyLocation, *FileName, true, true);
 			}
