Fix bug with calling set parent submix explicitly

diff --git a/Engine/Source/Runtime/Engine/Private/SoundSubmix.cpp b/Engine/Source/Runtime/Engine/Private/SoundSubmix.cpp
--- a/Engine/Source/Runtime/Engine/Private/SoundSubmix.cpp
+++ b/Engine/Source/Runtime/Engine/Private/SoundSubmix.cpp
@@ -317,7 +317,7 @@ void USoundSubmix::PostEditChangeProperty(FPropertyChangedEvent& PropertyChanged
 				if (!bIsChildSubmix)
 				{
 					ParentSubmix->Modify();
-					ParentSubmix->ChildSubmixes.Add(this);
+					ParentSubmix->ChildSubmixes.AddUnique(this);
 				}
 			}
 
@@ -334,13 +334,10 @@ void USoundSubmix::PostEditChangeProperty(FPropertyChangedEvent& PropertyChanged
 		}
 	}
 
-	if (bReinitSubmix)
+	if (GEngine && bReinitSubmix)
 	{
-		// Use the main/default audio device for storing and retrieving sound class properties
-		FAudioDeviceManager* AudioDeviceManager = (GEngine ? GEngine->GetAudioDeviceManager() : nullptr);
-
 		// Force the properties to be initialized for this SoundSubmix on all active audio devices
-		if (AudioDeviceManager)
+		if (FAudioDeviceManager* AudioDeviceManager = GEngine->GetAudioDeviceManager())
 		{
 			AudioDeviceManager->RegisterSoundSubmix(this);
 		}
@@ -395,6 +392,7 @@ void USoundSubmix::SetParentSubmix(USoundSubmix* InParentSubmix)
 
 		Modify();
 		ParentSubmix = InParentSubmix;
+		ParentSubmix->ChildSubmixes.AddUnique(this);
 	}
 }
 
