Pass on COM computation to SetActorCOMPose

diff --git a/Engine/Source/Runtime/Engine/Private/PhysicsEngine/Experimental/PhysInterface_Chaos.cpp b/Engine/Source/Runtime/Engine/Private/PhysicsEngine/Experimental/PhysInterface_Chaos.cpp
--- a/Engine/Source/Runtime/Engine/Private/PhysicsEngine/Experimental/PhysInterface_Chaos.cpp
+++ b/Engine/Source/Runtime/Engine/Private/PhysicsEngine/Experimental/PhysInterface_Chaos.cpp
@@ -1415,9 +1415,8 @@ void FinishSceneStat()
 
 #if WITH_PHYSX
 
-void CalculateMassPropertiesOfImplicitType(
-	TArray< Chaos::TMassProperties<float, 3> > & MassProperties,
-	float & TotalMass,
+bool CalculateMassPropertiesOfImplicitType(
+	Chaos::TMassProperties<float, 3>& OutMassProperties,
 	const Chaos::TRigidTransform<float, 3> & WorldTransform,
 	const Chaos::FImplicitObject* ImplicitObject, 
 	float InDensityKGPerCM)
@@ -1439,39 +1438,34 @@ void CalculateMassPropertiesOfImplicitType(
 			const Chaos::TImplicitObjectTransformed<FReal, 3>& Object = ImplicitObject->template GetObjectChecked<Chaos::TImplicitObjectTransformed<FReal, 3>>();
 
 
-			Chaos::TMassProperties<float, 3> MassProps;
-			MassProps.Volume = Object.GetVolume();
-			const float Mass = MassProps.Volume * InDensityKGPerCM;
-			TotalMass += Mass;
-			MassProps.InertiaTensor = Object.GetInertiaTensor(Mass);
-			MassProps.CenterOfMass = Object.GetCenterOfMass();
-			MassProps.RotationOfMass = Chaos::TRotation<float, 3>::FromIdentity();
-
-			MassProperties.Add(MassProps);
+			OutMassProperties.Volume = Object.GetVolume();
+			OutMassProperties.Mass = OutMassProperties.Volume * InDensityKGPerCM;
+			OutMassProperties.InertiaTensor = Object.GetInertiaTensor(OutMassProperties.Mass);
+			OutMassProperties.CenterOfMass = Object.GetCenterOfMass();
+			OutMassProperties.RotationOfMass = Chaos::TRotation<float, 3>::FromIdentity();
+			return true;
 		}
 		else
 		{
-			const Chaos::TMassProperties<float, 3> MassProps = Chaos::CastHelper(*ImplicitObject, [&TotalMass, InDensityKGPerCM](const auto& Object)
-				{
-					Chaos::TMassProperties<float, 3> MassProperties;
-					MassProperties.Volume = Object.GetVolume();
-					const float Mass = MassProperties.Volume * InDensityKGPerCM;
-					TotalMass += Mass;
-					MassProperties.InertiaTensor = Object.GetInertiaTensor(Mass);
-					MassProperties.CenterOfMass = Object.GetCenterOfMass();
-					MassProperties.RotationOfMass = Chaos::TRotation<float, 3>::FromIdentity();
-					return MassProperties;
-				});
-			MassProperties.Add(MassProps);
+			Chaos::CastHelper(*ImplicitObject, [&OutMassProperties, InDensityKGPerCM](const auto& Object)
+			{
+				OutMassProperties.Volume = Object.GetVolume();
+				OutMassProperties.Mass = OutMassProperties.Volume * InDensityKGPerCM;
+				OutMassProperties.InertiaTensor = Object.GetInertiaTensor(OutMassProperties.Mass);
+				OutMassProperties.CenterOfMass = Object.GetCenterOfMass();
+				OutMassProperties.RotationOfMass = Chaos::TRotation<float, 3>::FromIdentity();
+			});
+			return true;
 		}
 	}
+	return false;
 }
 
 void FPhysInterface_Chaos::CalculateMassPropertiesFromShapeCollection(physx::PxMassProperties& OutProperties, const TArray<FPhysicsShapeHandle>& InShapes, float InDensityKGPerCM)
 {
-	float TotalMass = 0;
-	TArray< Chaos::TMassProperties<float,3> > MassProperties;
-
+	float TotalMass = 0.f;
+	Chaos::FVec3 TotalCenterOfMass(0.f);
+	TArray< Chaos::TMassProperties<float,3> > MassPropertiesList;
 	for (const FPhysicsShapeHandle& ShapeHandle : InShapes)
 	{
 		if (const Chaos::TPerShapeData<float, 3>* Shape = ShapeHandle.Shape)
@@ -1479,16 +1473,26 @@ void FPhysInterface_Chaos::CalculateMassPropertiesFromShapeCollection(physx::PxM
 			if (const Chaos::FImplicitObject * ImplicitObject = Shape->Geometry.Get())
 			{
 				FTransform WorldTransform(ShapeHandle.ActorRef->R(), ShapeHandle.ActorRef->X());
-				CalculateMassPropertiesOfImplicitType(MassProperties, TotalMass, WorldTransform, ImplicitObject, InDensityKGPerCM);
+				Chaos::TMassProperties<float,3> MassProperties;
+				if (CalculateMassPropertiesOfImplicitType(MassProperties, WorldTransform, ImplicitObject, InDensityKGPerCM))
+				{
+					MassPropertiesList.Add(MassProperties);
+					TotalMass += MassProperties.Mass;
+					TotalCenterOfMass += MassProperties.CenterOfMass * MassProperties.Mass;
+				}
 			}
 		}
 	}
 
+	if (TotalMass > 0.f)
+	{
+		TotalCenterOfMass /= TotalMass;
+	}
 
 	Chaos::PMatrix<float, 3, 3> Tensor;
-	if (MassProperties.Num())
+	if (MassPropertiesList.Num())
 	{
-		Tensor = Chaos::Combine<float, 3>(MassProperties).InertiaTensor;
+		Tensor = Chaos::Combine<float, 3>(MassPropertiesList).InertiaTensor;
 	}
 	else 
 	{
@@ -1500,6 +1504,7 @@ void FPhysInterface_Chaos::CalculateMassPropertiesFromShapeCollection(physx::PxM
 	float Mat[] = { Tensor.M[0][0],Tensor.M[0][1],Tensor.M[0][2],Tensor.M[1][0],Tensor.M[1][1],Tensor.M[1][2],Tensor.M[2][0],Tensor.M[2][1],Tensor.M[2][2] };
 	OutProperties.inertiaTensor = PxMat33(Mat);
 	OutProperties.mass = TotalMass;
+	OutProperties.centerOfMass = U2PVector(TotalCenterOfMass);
 }
 
 #endif
diff --git a/Engine/Source/Runtime/Experimental/Chaos/Public/Chaos/MassProperties.h b/Engine/Source/Runtime/Experimental/Chaos/Public/Chaos/MassProperties.h
--- a/Engine/Source/Runtime/Experimental/Chaos/Public/Chaos/MassProperties.h
+++ b/Engine/Source/Runtime/Experimental/Chaos/Public/Chaos/MassProperties.h
@@ -20,11 +20,13 @@ namespace Chaos
 	struct TMassProperties
 	{
 		TMassProperties()
-		    : Volume(0)
+		    : Mass(0)
+			, Volume(0)
 		    , CenterOfMass(0)
 		    , RotationOfMass(TRotation<T, d>::FromElements(TVector<T, d>(0), 1))
 		    , InertiaTensor(0)
 		{}
+		T Mass;
 		T Volume;
 		TVector<T, d> CenterOfMass;
 		TRotation<T, d> RotationOfMass;
