Minor cleanup of AddForceAtPosition

diff --git a/Engine/Source/Runtime/Engine/Private/PhysicsEngine/Experimental/PhysScene_Chaos.cpp b/Engine/Source/Runtime/Engine/Private/PhysicsEngine/Experimental/PhysScene_Chaos.cpp
--- a/Engine/Source/Runtime/Engine/Private/PhysicsEngine/Experimental/PhysScene_Chaos.cpp
+++ b/Engine/Source/Runtime/Engine/Private/PhysicsEngine/Experimental/PhysScene_Chaos.cpp
@@ -1102,21 +1102,20 @@ void FPhysScene_ChaosInterface::AddForceAtPosition_AssumesLocked(FBodyInstance*
 				const Chaos::FVec3& CurrentForce = Rigid->ExternalForce();
 				const Chaos::FVec3& CurrentTorque = Rigid->ExternalTorque();
 				const Chaos::FVec3& CurrentPosition = Rigid->X();
-				const Chaos::FVec3& CurrentCOM = CurrentPosition; // TODO: Remember to fix this once we add COM to particles!
-				const Chaos::FVec3 Torque = Chaos::FVec3::CrossProduct(Position - CurrentCOM, Force);
-				if (!bIsLocalForce)
+				const Chaos::FVec3 WorldCOM = Rigid->X(); // TODO: Remember to fix this once we add COM to particles!
+				if (bIsLocalForce)
 				{
-					Rigid->SetExternalForce(CurrentForce + Force);
-					Rigid->SetExternalTorque(CurrentTorque + Torque);
+					const Chaos::FRigidTransform3 CurrentTransform(Rigid->X(), Rigid->R());
+					const Chaos::FVec3 WorldPosition = CurrentTransform.TransformPosition(Position);
+					const Chaos::FVec3 WorldForce = CurrentTransform.TransformVector(Force);
+					const Chaos::FVec3 WorldTorque = Chaos::FVec3::CrossProduct(WorldPosition - WorldCOM, WorldForce);
+					Rigid->SetExternalForce(CurrentForce + WorldForce);
+					Rigid->SetExternalTorque(CurrentTorque + WorldTorque);
 				}
 				else
 				{
-					const Chaos::FRotation3& CurrentRotation = Rigid->R();
-					const Chaos::FRigidTransform3 CurrentTransform(CurrentPosition, CurrentRotation);
-					const Chaos::FVec3 WorldPosition = CurrentTransform.TransformPosition(Position);
-					const Chaos::FVec3 WorldForce = CurrentRotation.RotateVector(Force);
-					const Chaos::FVec3 WorldTorque = Chaos::FVec3::CrossProduct(WorldPosition - CurrentCOM, WorldForce);
-					Rigid->SetExternalForce(CurrentForce + WorldForce);
+					const Chaos::FVec3 WorldTorque = Chaos::FVec3::CrossProduct(Position - WorldCOM, Force);
+					Rigid->SetExternalForce(CurrentForce + Force);
 					Rigid->SetExternalTorque(CurrentTorque + WorldTorque);
 				}
 			}
