Fixed crash when attempting to build a runtime static mesh marked for streaming, with multiple LODs.
#rnx

diff --git a/Engine/Source/Runtime/Engine/Private/StaticMesh.cpp b/Engine/Source/Runtime/Engine/Private/StaticMesh.cpp
--- a/Engine/Source/Runtime/Engine/Private/StaticMesh.cpp
+++ b/Engine/Source/Runtime/Engine/Private/StaticMesh.cpp
@@ -5202,6 +5202,7 @@ bool UStaticMesh::BuildFromMeshDescriptions(const TArray<const FMeshDescription*
 	// Set up
 
 	bIsBuiltAtRuntime = true;
+	NeverStream = true;
 
 	TOptional<FStaticMeshComponentRecreateRenderStateContext> RecreateRenderStateContext;
 	
