Fixes for runtime static meshes when created in-editor; material sections are now preserved correctly. Added some extra checks to ensure that a valid polygon group is passed when creating polygons.
#rnx

diff --git a/Engine/Source/Runtime/Engine/Private/StaticMesh.cpp b/Engine/Source/Runtime/Engine/Private/StaticMesh.cpp
--- a/Engine/Source/Runtime/Engine/Private/StaticMesh.cpp
+++ b/Engine/Source/Runtime/Engine/Private/StaticMesh.cpp
@@ -5235,6 +5235,17 @@ bool UStaticMesh::BuildFromMeshDescriptions(const TArray<const FMeshDescription*
 
 		BuildFromMeshDescription(*MeshDescriptionPtr, LODResources);
 
+#if WITH_EDITOR
+		for (int32 SectionIndex = 0; SectionIndex < LODResources.Sections.Num(); SectionIndex++)
+		{
+			const FStaticMeshSection& StaticMeshSection = LODResources.Sections[SectionIndex];
+			FMeshSectionInfo SectionInfo;
+			SectionInfo.MaterialIndex = StaticMeshSection.MaterialIndex;
+			SectionInfo.bEnableCollision = StaticMeshSection.bEnableCollision;
+			SectionInfo.bCastShadow = StaticMeshSection.bCastShadow;
+			GetSectionInfoMap().Set(LODIndex, SectionIndex, SectionInfo);
+		}
+#endif
 		LODIndex++;
 	}
 
@@ -6277,6 +6288,11 @@ UMaterialInterface* UStaticMesh::GetMaterial(int32 MaterialIndex) const
 
 FName UStaticMesh::AddMaterial(UMaterialInterface* Material)
 {
+	if (Material == nullptr)
+	{
+		return NAME_None;
+	}
+
 	// Create a unique slot name for the material
 	FName MaterialName = Material->GetFName();
 	for (const FStaticMaterial& StaticMaterial : StaticMaterials)
diff --git a/Engine/Source/Runtime/MeshDescription/Private/MeshDescription.cpp b/Engine/Source/Runtime/MeshDescription/Private/MeshDescription.cpp
--- a/Engine/Source/Runtime/MeshDescription/Private/MeshDescription.cpp
+++ b/Engine/Source/Runtime/MeshDescription/Private/MeshDescription.cpp
@@ -573,6 +573,7 @@ void FMeshDescription::CreatePolygon_Internal(const FPolygonID PolygonID, const
 		// End backward compatibility
 	}
 
+	check(PolygonGroupID != FPolygonGroupID::Invalid);
 	Polygon.PolygonGroupID = PolygonGroupID;
 	PolygonGroupArray[PolygonGroupID].Polygons.Add(PolygonID);
 
