Proxies can now be deleted after a StaticMesh Releases its RenderData (UpdateAllPrimitiveSceneInfos) move the RemoveSpeedTreeWind() into the new DestroyRenderThreadResources() method

#rnx

diff --git a/Engine/Source/Runtime/Engine/Private/StaticMeshRender.cpp b/Engine/Source/Runtime/Engine/Private/StaticMeshRender.cpp
--- a/Engine/Source/Runtime/Engine/Private/StaticMeshRender.cpp
+++ b/Engine/Source/Runtime/Engine/Private/StaticMeshRender.cpp
@@ -317,8 +317,6 @@ FStaticMeshSceneProxy::~FStaticMeshSceneProxy()
 		Geometry.ReleaseResource();
 	}
 #endif
-
-	RemoveSpeedTreeWind();
 }
 
 void FStaticMeshSceneProxy::AddSpeedTreeWind()
@@ -574,6 +572,14 @@ void FStaticMeshSceneProxy::CreateRenderThreadResources()
 #endif
 }
 
+void FStaticMeshSceneProxy::DestroyRenderThreadResources()
+{
+	FPrimitiveSceneProxy::DestroyRenderThreadResources();
+
+	// Call here because it uses RenderData from the StaticMesh which is not guaranteed to still be valid after this DestroyRenderThreadResources call
+	RemoveSpeedTreeWind();
+}
+
 /** Sets up a wireframe FMeshBatch for a specific LOD. */
 bool FStaticMeshSceneProxy::GetWireframeMeshElement(int32 LODIndex, int32 BatchIndex, const FMaterialRenderProxy* WireframeRenderProxy, uint8 InDepthPriorityGroup, bool bAllowPreCulledIndices, FMeshBatch& OutMeshBatch) const
 {
diff --git a/Engine/Source/Runtime/Engine/Public/StaticMeshResources.h b/Engine/Source/Runtime/Engine/Public/StaticMeshResources.h
--- a/Engine/Source/Runtime/Engine/Public/StaticMeshResources.h
+++ b/Engine/Source/Runtime/Engine/Public/StaticMeshResources.h
@@ -802,6 +802,8 @@ public:
 	virtual int32 CollectOccluderElements(class FOccluderElementsCollector& Collector) const override;
 
 	virtual void CreateRenderThreadResources() override;
+		
+	virtual void DestroyRenderThreadResources() override;
 
 	/** Sets up a wireframe FMeshBatch for a specific LOD. */
 	virtual bool GetWireframeMeshElement(int32 LODIndex, int32 BatchIndex, const FMaterialRenderProxy* WireframeRenderProxy, uint8 InDepthPriorityGroup, bool bAllowPreCulledIndices, FMeshBatch& OutMeshBatch) const;
