#UE4 Fix crash in parallelsave

diff --git a/Engine/Source/Runtime/Engine/Private/Materials/MaterialInstance.cpp b/Engine/Source/Runtime/Engine/Private/Materials/MaterialInstance.cpp
--- a/Engine/Source/Runtime/Engine/Private/Materials/MaterialInstance.cpp
+++ b/Engine/Source/Runtime/Engine/Private/Materials/MaterialInstance.cpp
@@ -1319,6 +1319,12 @@ void UMaterialInstance::LogMaterialsAndTextures(FOutputDevice& Ar, int32 Indent)
 
 void UMaterialInstance::ValidateTextureOverrides(ERHIFeatureLevel::Type InFeatureLevel) const
 {
+	if (!(IsInGameThread() || IsAsyncLoading()))
+	{
+		// Fatal to call getmaterial in a non-game thread or async loading
+		return;
+	}
+	
 	const UMaterial* Material = GetMaterial();
 	const FMaterialResource* CurrentResource = Material->GetMaterialResource(InFeatureLevel);
 	
