Don't try to set Channel Names if the outputs aren't RGBA

diff --git a/Engine/Source/Runtime/Engine/Private/Materials/MaterialExpressions.cpp b/Engine/Source/Runtime/Engine/Private/Materials/MaterialExpressions.cpp
--- a/Engine/Source/Runtime/Engine/Private/Materials/MaterialExpressions.cpp
+++ b/Engine/Source/Runtime/Engine/Private/Materials/MaterialExpressions.cpp
@@ -2738,10 +2738,13 @@ void UMaterialExpressionTextureSampleParameter::ApplyChannelNames()
 	static const FName Green("G");
 	static const FName Blue("B");
 	static const FName Alpha("A");
-	Outputs[1].OutputName = !ChannelNames.R.IsEmpty() ? FName(*ChannelNames.R.ToString()) : Red;
-	Outputs[2].OutputName = !ChannelNames.G.IsEmpty() ? FName(*ChannelNames.G.ToString()) : Green;
-	Outputs[3].OutputName = !ChannelNames.B.IsEmpty() ? FName(*ChannelNames.B.ToString()) : Blue;
-	Outputs[4].OutputName = !ChannelNames.A.IsEmpty() ? FName(*ChannelNames.A.ToString()) : Alpha;
+	if (GetOutputType(0) != MCT_Texture)
+	{
+		Outputs[1].OutputName = !ChannelNames.R.IsEmpty() ? FName(*ChannelNames.R.ToString()) : Red;
+		Outputs[2].OutputName = !ChannelNames.G.IsEmpty() ? FName(*ChannelNames.G.ToString()) : Green;
+		Outputs[3].OutputName = !ChannelNames.B.IsEmpty() ? FName(*ChannelNames.B.ToString()) : Blue;
+		Outputs[4].OutputName = !ChannelNames.A.IsEmpty() ? FName(*ChannelNames.A.ToString()) : Alpha;
+	}
 }
 #endif
 
