Fix an edge case where the NetGuidCache could return a partially loaded Package.

diff --git a/Engine/Source/Runtime/Engine/Private/PackageMapClient.cpp b/Engine/Source/Runtime/Engine/Private/PackageMapClient.cpp
--- a/Engine/Source/Runtime/Engine/Private/PackageMapClient.cpp
+++ b/Engine/Source/Runtime/Engine/Private/PackageMapClient.cpp
@@ -3124,6 +3124,9 @@ UObject* FNetGUIDCache::GetObjectFromNetGUID( const FNetworkGUID& NetGUID, const
 				// Something else is already async loading this package, calling load again will add our callback to the existing load request
 				StartAsyncLoadingPackage(*CacheObjectPtr, NetGUID, true);
 				UE_LOG(LogNetPackageMap, Log, TEXT("GetObjectFromNetGUID: Listening to existing async load. Path: %s, NetGUID: %s"), *CacheObjectPtr->PathName.ToString(), *NetGUID.ToString());
+
+				// We don't want to hook up this package into the cache yet or return it, because it's only partially loaded.
+				return nullptr;
 			}
 			else
 			{
