Fix network corrections during anim root motion by disabling move combining even when AnimNotifyState_DisableRootMotion is active

#rnx

diff --git a/Engine/Source/Runtime/Engine/Private/Components/CharacterMovementComponent.cpp b/Engine/Source/Runtime/Engine/Private/Components/CharacterMovementComponent.cpp
--- a/Engine/Source/Runtime/Engine/Private/Components/CharacterMovementComponent.cpp
+++ b/Engine/Source/Runtime/Engine/Private/Components/CharacterMovementComponent.cpp
@@ -10687,13 +10687,16 @@ void FSavedMove_Character::PostUpdate(ACharacter* Character, FSavedMove_Characte
 	if (PostUpdateMode == PostUpdate_Record)
 	{
 		const FAnimMontageInstance* RootMotionMontageInstance = Character->GetRootMotionAnimMontageInstance();
-		if (RootMotionMontageInstance && !RootMotionMontageInstance->IsRootMotionDisabled())
+		if (RootMotionMontageInstance)
 		{
-			RootMotionMontage = RootMotionMontageInstance->Montage;
-			RootMotionTrackPosition = RootMotionMontageInstance->GetPosition();
-			RootMotionMovement = Character->ClientRootMotionParams;
+			if (!RootMotionMontageInstance->IsRootMotionDisabled())
+			{
+				RootMotionMontage = RootMotionMontageInstance->Montage;
+				RootMotionTrackPosition = RootMotionMontageInstance->GetPosition();
+				RootMotionMovement = Character->ClientRootMotionParams;
+			}
 
-			// Moves with anim root motion should not be combined
+			// Moves where anim root motion is being played should not be combined
 			bForceNoCombine = true;
 		}
 
