Adjusted material layer strings
Fixed crash when rapidly deleting layers.

diff --git a/Engine/Source/Runtime/Engine/Classes/Engine/RendererSettings.h b/Engine/Source/Runtime/Engine/Classes/Engine/RendererSettings.h
--- a/Engine/Source/Runtime/Engine/Classes/Engine/RendererSettings.h
+++ b/Engine/Source/Runtime/Engine/Classes/Engine/RendererSettings.h
@@ -761,8 +761,8 @@ class ENGINE_API URendererSettings : public UDeveloperSettings
 		ConfigRestartRequired = true))
 		uint32 bSupportReversedIndexBuffers : 1;
 
-	UPROPERTY(config, EditAnywhere, Category = Experimental, meta = (
-		ConsoleVariable = "r.SupportMaterialLayers", Tooltip = "Support new material layering system. Disabling it reduces some overhead in place to support the experimental feature",
+	UPROPERTY(config, EditAnywhere, Category = Materials, meta = (
+		ConsoleVariable = "r.SupportMaterialLayers", Tooltip = "Support new material layering system.",
 		ConfigRestartRequired = true))
 		uint32 bSupportMaterialLayers : 1;
 
diff --git a/Engine/Source/Runtime/Engine/Classes/Materials/MaterialLayersFunctions.h b/Engine/Source/Runtime/Engine/Classes/Materials/MaterialLayersFunctions.h
--- a/Engine/Source/Runtime/Engine/Classes/Materials/MaterialLayersFunctions.h
+++ b/Engine/Source/Runtime/Engine/Classes/Materials/MaterialLayersFunctions.h
@@ -148,15 +148,19 @@ struct ENGINE_API FMaterialLayersFunctions
 
 	void RemoveBlendedLayerAt(int32 Index)
 	{
-		check(Layers.IsValidIndex(Index) && Blends.IsValidIndex(Index-1) && LayerStates.IsValidIndex(Index));
-		Layers.RemoveAt(Index);
-		Blends.RemoveAt(Index-1);
-		LayerStates.RemoveAt(Index);
+		if (Layers.IsValidIndex(Index))
+		{
+			check(Layers.IsValidIndex(Index) && Blends.IsValidIndex(Index-1) && LayerStates.IsValidIndex(Index));
+			Layers.RemoveAt(Index);
+			Blends.RemoveAt(Index-1);
+			LayerStates.RemoveAt(Index);
 #if WITH_EDITOR
-		LayerNames.RemoveAt(Index);
-		RestrictToLayerRelatives.RemoveAt(Index);
-		RestrictToBlendRelatives.RemoveAt(Index - 1);
-#endif
+			check(LayerNames.IsValidIndex(Index) && RestrictToLayerRelatives.IsValidIndex(Index) && RestrictToBlendRelatives.IsValidIndex(Index-1));
+			LayerNames.RemoveAt(Index);
+			RestrictToLayerRelatives.RemoveAt(Index);
+			RestrictToBlendRelatives.RemoveAt(Index-1);
+#endif //WITH_EDITOR
+		}
 	}
 
 	void ToggleBlendedLayerVisibility(int32 Index)
diff --git a/Engine/Source/Runtime/Engine/Private/Materials/MaterialUniformExpressions.cpp b/Engine/Source/Runtime/Engine/Private/Materials/MaterialUniformExpressions.cpp
--- a/Engine/Source/Runtime/Engine/Private/Materials/MaterialUniformExpressions.cpp
+++ b/Engine/Source/Runtime/Engine/Private/Materials/MaterialUniformExpressions.cpp
@@ -20,7 +20,7 @@
 static TAutoConsoleVariable<int32> CVarSupportMaterialLayers(
 	TEXT("r.SupportMaterialLayers"),
 	0,
-	TEXT("Support new material layering in 4.19. Disabling it reduces some overhead in place to support the experimental feature."),
+	TEXT("Support new material layering"),
 	ECVF_ReadOnly | ECVF_RenderThreadSafe);
 
 // Temporary flag for toggling experimental material layers functionality
