Fix for race condition between GC of USoundWave and ~AudioMixerSourceBuffer()

diff --git a/Engine/Source/Runtime/AudioMixer/Private/AudioMixerSourceBuffer.cpp b/Engine/Source/Runtime/AudioMixer/Private/AudioMixerSourceBuffer.cpp
--- a/Engine/Source/Runtime/AudioMixer/Private/AudioMixerSourceBuffer.cpp
+++ b/Engine/Source/Runtime/AudioMixer/Private/AudioMixerSourceBuffer.cpp
@@ -90,7 +90,7 @@ namespace Audio
 		, bProcedural(InWave.bProcedural)
 		, bIsBus(InWave.bIsBus)
 	{
-		InWave.AddPlayingSource();
+		InWave.AddPlayingSource(this);
 
 		const uint32 TotalSamples = MONO_PCM_BUFFER_SAMPLES * NumChannels;
 		for (int32 BufferIndex = 0; BufferIndex < Audio::MAX_BUFFERS_QUEUED; ++BufferIndex)
@@ -130,7 +130,7 @@ namespace Audio
 
 		if (SoundWave)
 		{
-			SoundWave->RemovePlayingSource();
+			SoundWave->RemovePlayingSource(this);
 		}
 	}
 
@@ -447,6 +447,21 @@ namespace Audio
 		BufferQueue.Enqueue(InSourceVoiceBuffer);
 	}
 
+	void FMixerSourceBuffer::OnBeginDestroy(USoundWave * /*Wave*/)
+	{
+		SoundWave = nullptr;
+	}
+
+	bool FMixerSourceBuffer::OnIsReadyForFinishDestroy(USoundWave * /*Wave*/) const
+	{
+		return false;
+	}
+
+	void FMixerSourceBuffer::OnFinishDestroy(USoundWave * /*Wave*/)
+	{
+		SoundWave = nullptr;
+	}
+
 	bool FMixerSourceBuffer::IsAsyncTaskInProgress() const
 	{ 
 		return AsyncRealtimeAudioTask != nullptr; 
diff --git a/Engine/Source/Runtime/AudioMixer/Private/AudioMixerSourceBuffer.h b/Engine/Source/Runtime/AudioMixer/Private/AudioMixerSourceBuffer.h
--- a/Engine/Source/Runtime/AudioMixer/Private/AudioMixerSourceBuffer.h
+++ b/Engine/Source/Runtime/AudioMixer/Private/AudioMixerSourceBuffer.h
@@ -42,11 +42,13 @@ namespace Audio
 		AsynchronousSkipFirstFrame
 	};
 
+	using FMixerSourceBufferPtr = TSharedPtr<class FMixerSourceBuffer>;
+
 	/** Class which handles decoding audio for a particular source buffer. */
-	class FMixerSourceBuffer
+	class FMixerSourceBuffer : public ISoundWaveClient
 	{
-	public:
-		static TSharedPtr<FMixerSourceBuffer> Create(FMixerBuffer& InBuffer, USoundWave& InWave, ELoopingMode InLoopingMode, bool bInIsSeeking);
+	public:		
+		static FMixerSourceBufferPtr Create(FMixerBuffer& InBuffer, USoundWave& InWave, ELoopingMode InLoopingMode, bool bInIsSeeking);
 
 		~FMixerSourceBuffer();
 
@@ -125,5 +127,9 @@ namespace Audio
 		uint32 bLoopCallback : 1;
 		uint32 bProcedural : 1;
 		uint32 bIsBus : 1;
+	
+		virtual void OnBeginDestroy(class USoundWave* Wave) override;
+		virtual bool OnIsReadyForFinishDestroy(class USoundWave* Wave) const override;
+		virtual void OnFinishDestroy(class USoundWave* Wave) override;
 	};
 }
diff --git a/Engine/Source/Runtime/Engine/Classes/Sound/SoundWave.h b/Engine/Source/Runtime/Engine/Classes/Sound/SoundWave.h
--- a/Engine/Source/Runtime/Engine/Classes/Sound/SoundWave.h
+++ b/Engine/Source/Runtime/Engine/Classes/Sound/SoundWave.h
@@ -282,12 +282,21 @@ enum class ESoundWaveFFTSize : uint8
 	VeryLarge_2048,
 };
 
+struct ISoundWaveClient
+{
+	ISoundWaveClient() {}
+	virtual ~ISoundWaveClient() {}
+	
+	virtual void OnBeginDestroy(class USoundWave* Wave) = 0;
+	virtual bool OnIsReadyForFinishDestroy(class USoundWave* Wave) const = 0;
+	virtual void OnFinishDestroy(class USoundWave* Wave) = 0;
+};
 
 UCLASS(hidecategories=Object, editinlinenew, BlueprintType)
 class ENGINE_API USoundWave : public USoundBase
 {
 	GENERATED_UCLASS_BODY()
-
+public:
 	/** Platform agnostic compression quality. 1..100 with 1 being best compression and 100 being best quality. */
 	UPROPERTY(EditAnywhere, Category="Format|Quality", meta=(DisplayName = "Compression", ClampMin = "1", ClampMax = "100"), AssetRegistrySearchable)
 	int32 CompressionQuality;
@@ -485,7 +494,11 @@ private:
 	FThreadSafeCounter PrecacheState;
 
 	/** the number of sounds currently playing this sound wave. */
-	FThreadSafeCounter NumSourcesPlaying;
+	mutable FCriticalSection SourcesPlayingCs;
+
+	using FSoundWaveClientPtr = ISoundWaveClient*;
+
+	TArray<FSoundWaveClientPtr> SourcesPlaying;
 
 	// This is the sample rate retrieved from platform settings.
 	float CachedSampleRateOverride;
@@ -679,20 +692,16 @@ public:
 	// Called when the procedural sound wave is done generating on the render thread. Only used in the audio mixer and when bProcedural is true..
 	virtual void OnEndGenerate() {};
 
-	void AddPlayingSource()
-	{
-		NumSourcesPlaying.Increment();
-	}
-
-	void RemovePlayingSource()
-	{
-		check(NumSourcesPlaying.GetValue() > 0);
-		NumSourcesPlaying.Decrement();
-	}
+	void AddPlayingSource(const FSoundWaveClientPtr& Source);
+	void RemovePlayingSource(const FSoundWaveClientPtr& Source);
 
 	bool IsGeneratingAudio() const
-	{
-		return NumSourcesPlaying.GetValue() > 0;
+	{		
+		bool bIsGeneratingAudio = false;
+		FScopeLock Lock(&SourcesPlayingCs);
+		bIsGeneratingAudio = SourcesPlaying.Num() > 0;
+		
+		return bIsGeneratingAudio;
 	}
 
 	/**
diff --git a/Engine/Source/Runtime/Engine/Private/SoundWave.cpp b/Engine/Source/Runtime/Engine/Private/SoundWave.cpp
--- a/Engine/Source/Runtime/Engine/Private/SoundWave.cpp
+++ b/Engine/Source/Runtime/Engine/Private/SoundWave.cpp
@@ -948,6 +948,14 @@ uint32 USoundWave::GetSizeOfChunk(uint32 ChunkIndex)
 void USoundWave::BeginDestroy()
 {
 	Super::BeginDestroy();
+	
+	{
+		FScopeLock Lock(&SourcesPlayingCs);
+		for (ISoundWaveClient* i: SourcesPlaying)
+		{
+			i->OnBeginDestroy(this);
+		}
+	}
 
 #if WITH_EDITOR
 	// Flush any async results so we dont leak them in the DDC
@@ -1608,6 +1616,15 @@ FWaveInstance& USoundWave::HandleStart( FActiveSound& ActiveSound, const UPTRINT
 
 bool USoundWave::IsReadyForFinishDestroy()
 {
+	{
+		FScopeLock Lock(&SourcesPlayingCs);
+		
+		for (ISoundWaveClient* i : SourcesPlaying)
+		{
+			i->OnIsReadyForFinishDestroy(this);
+		}
+	}
+
 	const bool bIsStreamingInProgress = IStreamingManager::Get().GetAudioStreamingManager().IsStreamingInProgress(this);
 
 	check(GetPrecacheState() != ESoundWavePrecacheState::InProgress);
@@ -2017,6 +2034,28 @@ bool USoundWave::HasCookedAmplitudeEnvelopeData() const
 	return CookedEnvelopeTimeData.Num() > 0;
 }
 
+void USoundWave::AddPlayingSource(const FSoundWaveClientPtr& Source)
+{
+	check(Source);
+	check(IsInAudioThread() || IsInGameThread());   // Don't allow incrementing on other threads as it's not safe (for GCing of this wave).
+	if (Source)
+	{
+		FScopeLock Lock(&SourcesPlayingCs);
+		check(!SourcesPlaying.Contains(Source));
+		SourcesPlaying.Add(Source);
+	}
+}
+
+void USoundWave::RemovePlayingSource(const FSoundWaveClientPtr& Source)
+{
+	if (Source)
+	{
+		FScopeLock Lock(&SourcesPlayingCs);
+		check(SourcesPlaying.Contains(Source));
+		SourcesPlaying.RemoveSwap(Source);
+	}
+}
+
 void USoundWave::UpdatePlatformData()
 {
 	if (IsStreaming())
