Fix for comment box performance

diff --git a/Engine/Source/Runtime/Engine/Classes/Materials/MaterialExpressionComment.h b/Engine/Source/Runtime/Engine/Classes/Materials/MaterialExpressionComment.h
--- a/Engine/Source/Runtime/Engine/Classes/Materials/MaterialExpressionComment.h
+++ b/Engine/Source/Runtime/Engine/Classes/Materials/MaterialExpressionComment.h
@@ -25,7 +25,7 @@ class UMaterialExpressionComment : public UMaterialExpression
 	FString Text;
 
 	/** Color to style comment with */
-	UPROPERTY(EditAnywhere, Category=MaterialExpressionComment)
+	UPROPERTY(EditAnywhere, Category=MaterialExpressionComment, meta=(OnlyUpdateOnInteractionEnd))
 	FLinearColor CommentColor;
 
 	/** Size of the text in the comment box */
