Fix for PS4 compile error

diff --git a/Engine/Source/Runtime/Engine/Classes/Sound/SoundWave.h b/Engine/Source/Runtime/Engine/Classes/Sound/SoundWave.h
--- a/Engine/Source/Runtime/Engine/Classes/Sound/SoundWave.h
+++ b/Engine/Source/Runtime/Engine/Classes/Sound/SoundWave.h
@@ -695,6 +695,20 @@ public:
 	void AddPlayingSource(const FSoundWaveClientPtr& Source);
 	void RemovePlayingSource(const FSoundWaveClientPtr& Source);
 
+	/** the number of sounds currently playing this sound wave. */
+	FThreadSafeCounter NumSourcesPlaying;
+
+	void AddPlayingSource()
+	{
+		NumSourcesPlaying.Increment();
+	}
+
+	void RemovePlayingSource()
+	{
+		check(NumSourcesPlaying.GetValue() > 0);
+		NumSourcesPlaying.Decrement();
+	}
+
 	bool IsGeneratingAudio() const
 	{		
 		bool bIsGeneratingAudio = false;
