* Set COM on kinematics
* Set UserData on Shapes

diff --git a/Engine/Source/Runtime/Engine/Private/PhysicsEngine/Experimental/ChaosInterfaceUtils.cpp b/Engine/Source/Runtime/Engine/Private/PhysicsEngine/Experimental/ChaosInterfaceUtils.cpp
--- a/Engine/Source/Runtime/Engine/Private/PhysicsEngine/Experimental/ChaosInterfaceUtils.cpp
+++ b/Engine/Source/Runtime/Engine/Private/PhysicsEngine/Experimental/ChaosInterfaceUtils.cpp
@@ -140,31 +140,32 @@ namespace ChaosInterface
 		TArray<TUniquePtr<Chaos::FImplicitObject>>& Geoms = OutGeoms;
 		Chaos::TShapesArray<float, 3>& Shapes = OutShapes;
 
-		auto NewShapeHelper = [&InParams](Chaos::TSerializablePtr<Chaos::FImplicitObject> InGeom, bool bComplexShape = false)
+		auto NewShapeHelper = [&InParams](Chaos::TSerializablePtr<Chaos::FImplicitObject> InGeom, void* UserData, bool bComplexShape = false)
 		{
 			auto NewShape = Chaos::TPerShapeData<float, 3>::CreatePerShapeData();
 			NewShape->Geometry = InGeom;
 			NewShape->QueryData = bComplexShape ? InParams.CollisionData.CollisionFilterData.QueryComplexFilter : InParams.CollisionData.CollisionFilterData.QuerySimpleFilter;
 			NewShape->SimData = InParams.CollisionData.CollisionFilterData.SimFilter;
 			NewShape->UpdateShapeBounds(InParams.WorldTransform);
-
+			NewShape->UserData = UserData;
 			return NewShape;
 		};
 
 		for (uint32 i = 0; i < static_cast<uint32>(InParams.Geometry->SphereElems.Num()); ++i)
 		{
-			const FKSphereElem ScaledSphereElem = InParams.Geometry->SphereElems[i].GetFinalScaled(Scale, InParams.LocalTransform);
+			const FKSphereElem& SphereElem = InParams.Geometry->SphereElems[i];
+			const FKSphereElem ScaledSphereElem = SphereElem.GetFinalScaled(Scale, InParams.LocalTransform);
 			const float UseRadius = FMath::Max(ScaledSphereElem.Radius, KINDA_SMALL_NUMBER);
-
 			auto ImplicitSphere = MakeUnique<Chaos::TSphere<float, 3>>(ScaledSphereElem.Center, UseRadius);
-			auto NewShape = NewShapeHelper(MakeSerializable(ImplicitSphere));
+			auto NewShape = NewShapeHelper(MakeSerializable(ImplicitSphere), (void*)SphereElem.GetUserData());
 			Shapes.Emplace(MoveTemp(NewShape));
 			Geoms.Add(MoveTemp(ImplicitSphere));
-
 		}
+
 		for (uint32 i = 0; i < static_cast<uint32>(InParams.Geometry->BoxElems.Num()); ++i)
 		{
-			const FKBoxElem ScaledBoxElem = InParams.Geometry->BoxElems[i].GetFinalScaled(Scale, InParams.LocalTransform);
+			const FKBoxElem& BoxElem = InParams.Geometry->BoxElems[i];
+			const FKBoxElem ScaledBoxElem = BoxElem.GetFinalScaled(Scale, InParams.LocalTransform);
 			const FTransform& BoxTransform = ScaledBoxElem.GetTransform();
 			Chaos::TVector<float, 3> HalfExtents = Chaos::TVector<float, 3>(ScaledBoxElem.X * 0.5f, ScaledBoxElem.Y * 0.5f, ScaledBoxElem.Z * 0.5f);
 
@@ -184,14 +185,13 @@ namespace ChaosInterface
 				Implicit = MakeUnique<Chaos::TBox<float, 3>>(BoxTransform.GetTranslation() - HalfExtents, BoxTransform.GetTranslation() + HalfExtents);
 			}
 
-			auto NewShape = NewShapeHelper(MakeSerializable(Implicit));
+			auto NewShape = NewShapeHelper(MakeSerializable(Implicit), (void*)BoxElem.GetUserData());
 			Shapes.Emplace(MoveTemp(NewShape));
 			Geoms.Add(MoveTemp(Implicit));
 		}
 		for (uint32 i = 0; i < static_cast<uint32>(InParams.Geometry->SphylElems.Num()); ++i)
 		{
 			const FKSphylElem& UnscaledSphyl = InParams.Geometry->SphylElems[i];
-
 			const FKSphylElem ScaledSphylElem = UnscaledSphyl.GetFinalScaled(Scale, InParams.LocalTransform);
 			float HalfHeight = FMath::Max(ScaledSphylElem.Length * 0.5f, KINDA_SMALL_NUMBER);
 			const float Radius = FMath::Max(ScaledSphylElem.Radius, KINDA_SMALL_NUMBER);
@@ -200,7 +200,7 @@ namespace ChaosInterface
 			{
 				//not a capsule just use a sphere
 				auto ImplicitSphere = MakeUnique<Chaos::TSphere<float, 3>>(ScaledSphylElem.Center, Radius);
-				auto NewShape = NewShapeHelper(MakeSerializable(ImplicitSphere));
+				auto NewShape = NewShapeHelper(MakeSerializable(ImplicitSphere), (void*)UnscaledSphyl.GetUserData());
 				Shapes.Emplace(MoveTemp(NewShape));
 				Geoms.Add(MoveTemp(ImplicitSphere));
 
@@ -210,7 +210,7 @@ namespace ChaosInterface
 				Chaos::TVector<float, 3> HalfExtents = ScaledSphylElem.Rotation.RotateVector(Chaos::TVector<float, 3>(0, 0, HalfHeight));
 
 				auto ImplicitCapsule = MakeUnique<Chaos::TCapsule<float>>(ScaledSphylElem.Center - HalfExtents, ScaledSphylElem.Center + HalfExtents, Radius);
-				auto NewShape = NewShapeHelper(MakeSerializable(ImplicitCapsule));
+				auto NewShape = NewShapeHelper(MakeSerializable(ImplicitCapsule), (void*)UnscaledSphyl.GetUserData());
 				Shapes.Emplace(MoveTemp(NewShape));
 				Geoms.Add(MoveTemp(ImplicitCapsule));
 			}
@@ -261,7 +261,7 @@ namespace ChaosInterface
 				//else
 				{
 					auto Implicit = MakeUnique<Chaos::TImplicitObjectScaled<Chaos::TConvex<float, 3>>>(MakeSerializable(ConvexImplicit), Scale);
-					auto NewShape = NewShapeHelper(MakeSerializable(Implicit));
+					auto NewShape = NewShapeHelper(MakeSerializable(Implicit), (void*)CollisionBody.GetUserData());
 					Shapes.Emplace(MoveTemp(NewShape));
 					Geoms.Add(MoveTemp(Implicit));
 				}
@@ -271,7 +271,7 @@ namespace ChaosInterface
 		for (const auto& ChaosTriMesh : InParams.ChaosTriMeshes)
 		{
 			auto Implicit = MakeUnique<Chaos::TImplicitObjectScaled<Chaos::TTriangleMeshImplicitObject<float>>>(MakeSerializable(ChaosTriMesh), Scale);
-			auto NewShape = NewShapeHelper(MakeSerializable(Implicit), true);
+			auto NewShape = NewShapeHelper(MakeSerializable(Implicit), nullptr, true);
 			Shapes.Emplace(MoveTemp(NewShape));
 			Geoms.Add(MoveTemp(Implicit));
 		}
diff --git a/Engine/Source/Runtime/Engine/Private/PhysicsEngine/Experimental/PhysInterface_Chaos.cpp b/Engine/Source/Runtime/Engine/Private/PhysicsEngine/Experimental/PhysInterface_Chaos.cpp
--- a/Engine/Source/Runtime/Engine/Private/PhysicsEngine/Experimental/PhysInterface_Chaos.cpp
+++ b/Engine/Source/Runtime/Engine/Private/PhysicsEngine/Experimental/PhysInterface_Chaos.cpp
@@ -247,6 +247,23 @@ void FPhysInterface_Chaos::SetUserData(FPhysicsMaterialHandle& InHandle, void* I
 	Chaos::FPhysicalMaterialManager::Get().UpdateMaterial(InHandle);
 }
 
+void FPhysInterface_Chaos::SetUserData(const FPhysicsShapeHandle& InShape, void* InUserData)
+{
+	if (CHAOS_ENSURE(InShape.Shape))
+	{
+		InShape.Shape->UserData = InUserData;
+	}
+}
+
+void* FPhysInterface_Chaos::GetUserData(const FPhysicsShapeHandle& InShape)
+{
+	if (ensure(InShape.Shape))
+	{
+		return InShape.Shape->UserData;
+	}
+	return nullptr;
+}
+
 int32 FPhysInterface_Chaos::GetNumShapes(const FPhysicsActorHandle& InHandle)
 {
 	// #todo : Implement
@@ -263,11 +280,13 @@ void FPhysInterface_Chaos::ReleaseShape(const FPhysicsShapeHandle& InShape)
 void FPhysInterface_Chaos::AttachShape(const FPhysicsActorHandle& InActor, const FPhysicsShapeHandle& InNewShape)
 {
 	// #todo : Implement
+	CHAOS_ENSURE(false);
 }
 
 void FPhysInterface_Chaos::DetachShape(const FPhysicsActorHandle& InActor, FPhysicsShapeHandle& InShape, bool bWakeTouching)
 {
 	// #todo : Implement
+	CHAOS_ENSURE(false);
 }
 
 void FPhysInterface_Chaos::SetActorUserData_AssumesLocked(FPhysicsActorHandle& InActorReference, FPhysicsUserData* InUserData)
@@ -546,7 +565,12 @@ FTransform FPhysInterface_Chaos::GetComTransform_AssumesLocked(const FPhysicsAct
 	// NOTE: This might need to change after we add separate COM transforms.
 	if (ensure(FPhysicsInterface::IsValid(InActorReference)))
 	{
-		return FTransform(InActorReference->R(), InActorReference->X());
+		if (Chaos::TKinematicGeometryParticle<float, 3>* Kinematic = InActorReference->CastToKinematicParticle())
+		{
+			const FTransform WorldTransform(Kinematic->R(), Kinematic->X());
+			const FTransform COMTransform = FTransform(Kinematic->RotationOfMass(), Kinematic->CenterOfMass());
+			return WorldTransform * COMTransform;
+		}
 	}
 	return FTransform();
 }
@@ -554,6 +578,13 @@ FTransform FPhysInterface_Chaos::GetComTransform_AssumesLocked(const FPhysicsAct
 FTransform FPhysInterface_Chaos::GetComTransformLocal_AssumesLocked(const FPhysicsActorHandle& InActorReference)
 {
 	// #todo : Implement
+	if (ensure(FPhysicsInterface::IsValid(InActorReference)))
+	{
+		if (Chaos::TKinematicGeometryParticle<float, 3>* Kinematic = InActorReference->CastToKinematicParticle())
+		{
+			return FTransform(Kinematic->RotationOfMass(), Kinematic->CenterOfMass());
+		}
+	}
 	return FTransform();
 }
 
@@ -686,6 +717,12 @@ void FPhysInterface_Chaos::SetMassSpaceInertiaTensor_AssumesLocked(FPhysicsActor
 void FPhysInterface_Chaos::SetComLocalPose_AssumesLocked(const FPhysicsActorHandle& InHandle, const FTransform& InComLocalPose)
 {
     //@todo(mlentine): What is InComLocalPose? If the center of an object is not the local pose then many things break including the three vector represtnation of inertia.
+
+	if (Chaos::TKinematicGeometryParticle<float, 3>* KinematicParticle = InHandle->CastToKinematicParticle())
+	{
+		KinematicParticle->SetCenterOfMass(InComLocalPose.GetLocation());
+		KinematicParticle->SetRotationOfMass(InComLocalPose.GetRotation());
+	}
 }
 
 float FPhysInterface_Chaos::GetStabilizationEnergyThreshold_AssumesLocked(const FPhysicsActorHandle& InHandle)
diff --git a/Engine/Source/Runtime/Engine/Public/Physics/Experimental/PhysInterface_Chaos.h b/Engine/Source/Runtime/Engine/Public/Physics/Experimental/PhysInterface_Chaos.h
--- a/Engine/Source/Runtime/Engine/Public/Physics/Experimental/PhysInterface_Chaos.h
+++ b/Engine/Source/Runtime/Engine/Public/Physics/Experimental/PhysInterface_Chaos.h
@@ -378,7 +378,7 @@ public:
 	static bool IsQueryShape(const FPhysicsShapeHandle& InShape);
 	static ECollisionShapeType GetShapeType(const FPhysicsShapeHandle& InShape);
 	static FTransform GetLocalTransform(const FPhysicsShapeHandle& InShape);
-    static void* GetUserData(const FPhysicsShapeHandle& InShape) { return nullptr; }
+    static void* GetUserData(const FPhysicsShapeHandle& InShape);
 
 	// Trace functions for testing specific geometry (not against a world)
 	static bool LineTrace_Geom(FHitResult& OutHit, const FBodyInstance* InInstance, const FVector& InStart, const FVector& InEnd, bool bTraceComplex, bool bExtractPhysMaterial = false);
@@ -394,7 +394,7 @@ public:
 	static void SetQueryFilter(const FPhysicsShapeHandle& InShape, const FCollisionFilterData& InFilter);
     static void SetIsSimulationShape(const FPhysicsShapeHandle& InShape, bool bIsSimShape) { const_cast<FPhysicsShapeHandle&>(InShape).bSimulation = bIsSimShape; }
     static void SetIsQueryShape(const FPhysicsShapeHandle& InShape, bool bIsQueryShape) { const_cast<FPhysicsShapeHandle&>(InShape).bSimulation = bIsQueryShape; }
-    static void SetUserData(const FPhysicsShapeHandle& InShape, void* InUserData) {}
+    static void SetUserData(const FPhysicsShapeHandle& InShape, void* InUserData);
     static void SetGeometry(const FPhysicsShapeHandle& InShape, physx::PxGeometry& InGeom) {}
 	static void SetLocalTransform(const FPhysicsShapeHandle& InShape, const FTransform& NewLocalTransform);
     static void SetMaterials(const FPhysicsShapeHandle& InShape, const TArrayView<UPhysicalMaterial*>InMaterials);
