Fix PHYSICS_INTERFACE_PHYSX case in CreateGeometry

diff --git a/Engine/Source/Runtime/Engine/Private/PhysicsEngine/Experimental/ChaosInterfaceUtils.cpp b/Engine/Source/Runtime/Engine/Private/PhysicsEngine/Experimental/ChaosInterfaceUtils.cpp
--- a/Engine/Source/Runtime/Engine/Private/PhysicsEngine/Experimental/ChaosInterfaceUtils.cpp
+++ b/Engine/Source/Runtime/Engine/Private/PhysicsEngine/Experimental/ChaosInterfaceUtils.cpp
@@ -280,7 +280,7 @@ namespace ChaosInterface
 		for (const auto& PhysXMesh : InParams.TriMeshes)
 		{
 			auto Implicit = ConvertPhysXMeshToLevelset(PhysXMesh, Scale);
-			auto NewShape = NewShapeHelper(MakeSerializable(Implicit), true);
+			auto NewShape = NewShapeHelper(MakeSerializable(Implicit), nullptr, true);
 			Shapes.Emplace(MoveTemp(NewShape));
 			Geoms.Add(MoveTemp(Implicit));
 
