Exposing maximum number of concurrent players for hotfixing

diff --git a/Engine/Source/Runtime/Engine/Classes/Engine/GameViewportClient.h b/Engine/Source/Runtime/Engine/Classes/Engine/GameViewportClient.h
--- a/Engine/Source/Runtime/Engine/Classes/Engine/GameViewportClient.h
+++ b/Engine/Source/Runtime/Engine/Classes/Engine/GameViewportClient.h
@@ -71,7 +71,8 @@ public:
 	/** Array of the screen data needed for all the different splitscreen configurations */
 	TArray<struct FSplitscreenData> SplitscreenInfo;
 
-	int32 MaxSplitscreenPlayers;
+	UPROPERTY(Config)
+	int32 MaxSplitscreenPlayers = 4;
 
 	/** if true then the title safe border is drawn */
 	uint32 bShowTitleSafeZone:1;
diff --git a/Engine/Source/Runtime/Engine/Private/GameViewportClient.cpp b/Engine/Source/Runtime/Engine/Private/GameViewportClient.cpp
--- a/Engine/Source/Runtime/Engine/Private/GameViewportClient.cpp
+++ b/Engine/Source/Runtime/Engine/Private/GameViewportClient.cpp
@@ -223,7 +223,6 @@ UGameViewportClient::UGameViewportClient(const FObjectInitializer& ObjectInitial
 	SplitscreenInfo[ESplitScreenType::FourPlayer_Horizontal].PlayerData.Add(FPerPlayerSplitscreenData(1.f, 0.25f, 0.0f, 0.5f));
 	SplitscreenInfo[ESplitScreenType::FourPlayer_Horizontal].PlayerData.Add(FPerPlayerSplitscreenData(1.f, 0.25f, 0.0f, 0.75f));
 
-	MaxSplitscreenPlayers = 4;
 	bSuppressTransitionMessage = true;
 
 	if (HasAnyFlags(RF_ClassDefaultObject) == false)
