Network
* Increased EEngineNetworkVersionHistory to account for new jitter clock time in packet headers
* This will cause connections to fail between pre/post engine builds.

Added Jitter calculation to NetConnection
* Added AverageJitterInMS variable
* Jitter is the difference in latency between every packet. The closer to 0 it is the more stable the latency is.
* Packet headers now contain local clock data used to calculate jitter.
* Fixed ServerFrameTime so that we send the current frame's time instead of the previous frame.
* PacketInfo payload (ServerFrameTime, Jitter) will now only be sent on the first packet of the frame. Reduces overhead when sending large data chunks in a single frame.
* Deprecated the RemoteSaturation variable and stopped sending the InBytesPerSecond value in every packet header

[RN] minor engine/network

diff --git a/Engine/Source/Runtime/Core/Public/Misc/NetworkVersion.h b/Engine/Source/Runtime/Core/Public/Misc/NetworkVersion.h
--- a/Engine/Source/Runtime/Core/Public/Misc/NetworkVersion.h
+++ b/Engine/Source/Runtime/Core/Public/Misc/NetworkVersion.h
@@ -45,6 +45,7 @@ enum EEngineNetworkVersionHistory
 	HISTORY_FAST_ARRAY_DELTA_STRUCT = 11,			// Bump version to allow fast array serialization, delta struct serialization.
 	HISTORY_FIX_ENUM_SERIALIZATION = 12,			// Bump version to fix enum net serialization issues.
 	HISTORY_OPTIONALLY_QUANTIZE_SPAWN_INFO = 13,	// Bump version to conditionally disable quantization for Scale, Location, and Velocity when spawning network actors.
+	HISTORY_JITTER_IN_HEADER = 14,					// Bump version since we added jitter clock time to packet headers and removed remote saturation
 	// New history items go above here.
 
 	HISTORY_ENGINENETVERSION_PLUS_ONE,
diff --git a/Engine/Source/Runtime/Core/Public/Serialization/BitWriter.h b/Engine/Source/Runtime/Core/Public/Serialization/BitWriter.h
--- a/Engine/Source/Runtime/Core/Public/Serialization/BitWriter.h
+++ b/Engine/Source/Runtime/Core/Public/Serialization/BitWriter.h
@@ -213,6 +213,12 @@ public:
 		Overflowed = Writer.ArIsError;
 	}
 
+	void Reset()
+	{
+		Overflowed = false;
+		Num = 0;		
+	}
+
 	void Pop( FBitWriter& Writer );
 	void Copy( FBitWriter& Writer, TArray<uint8> &Buffer );
 
diff --git a/Engine/Source/Runtime/Engine/Classes/Engine/NetConnection.h b/Engine/Source/Runtime/Engine/Classes/Engine/NetConnection.h
--- a/Engine/Source/Runtime/Engine/Classes/Engine/NetConnection.h
+++ b/Engine/Source/Runtime/Engine/Classes/Engine/NetConnection.h
@@ -380,6 +380,15 @@ private:
 	FPacketTimestamp	LastOSReceiveTime;		// Last time a packet was received at the OS/NIC layer
 	bool				bIsOSReceiveTimeLocal;	// Whether LastOSReceiveTime uses the same clock as the game, or needs translating
 
+	/** Stores the bit number where we wrote the dummy packet info in the packet header */
+	FBitWriterMark HeaderMarkForPacketInfo;
+
+	/** The difference between the send and receive clock time of the last received packet */
+	int32 PreviousJitterTimeDelta;
+
+	/** Timestamp of the last packet sent*/
+	double PreviousPacketSendTimeInS;
+
 public:
 	// Merge info.
 	FBitWriterMark  LastStart;				// Most recently sent bunch start.
@@ -422,6 +431,9 @@ public:
 	/** The average frame delta time over the last 1 second period.*/
 	double			AverageFrameTime;
 
+	/** Current jitter for this connection in milliseconds */
+	float			GetAverageJitterInMS() const { return AverageJitterInMS; }
+
 	/** Nb of stats accumulated in CumulativeTime */
 	int32			CountedFrames;
 
@@ -458,6 +470,14 @@ private:
 	/** Counts the number of stat samples taken */
 	int32 StatPeriodCount;
 
+	/**
+	* Jitter represents the average time divergence of all sent packets.
+	* Ex:
+	* If the time between the sending and the reception of packets is always 100ms; the jitter will be 0.
+	* If the time difference is either 150ms or 100ms, the jitter will tend towards 50ms.
+	*/
+	float AverageJitterInMS;
+
 public:
 
 	/** Net Analytics */
@@ -473,6 +493,7 @@ public:
 	double			OutLagTime[256];				// For lag measuring.
 	int32			OutLagPacketId[256];			// For lag measuring.
 	uint8			OutBytesPerSecondHistory[256];	// For saturation measuring.
+	UE_DEPRECATED(4.25, TEXT("RemoteSaturation is not calculated anymore and is now deprecated"))
 	float			RemoteSaturation;
 	int32			InPacketId;						// Full incoming packet index.
 	int32			OutPacketId;					// Most recently sent packet.
@@ -670,8 +691,11 @@ private:
 	/** delayed incoming packet array */
 	TArray<FDelayedIncomingPacket> DelayedIncomingPackets;
 
+	/** Allow emulated packets to be sent out of order. This causes packet loss on the receiving end */
+	bool bSendDelayedPacketsOutofOrder = false;
+
 	/** set to true when already delayed packets are reinjected */
-	bool bIsReinjectingDelayedPackets;
+	bool bIsReinjectingDelayedPackets = false;
 
 	/** Process incoming packets that have been delayed for long enough */
 	void ReinjectDelayedPackets();
@@ -1013,8 +1037,11 @@ public:
 	/** Create a channel. */
 	ENGINE_API UChannel* CreateChannelByName( const FName& ChName, EChannelCreateFlags CreateFlags, int32 ChannelIndex=INDEX_NONE );
 
-	/** Handle a packet we just received. */
-	void ReceivedPacket( FBitReader& Reader );
+	/** 
+	* Handle a packet we just received. 
+	* bIsReinjectedPacket is true if a packet is reprocessed after getting cached 
+	*/
+	void ReceivedPacket( FBitReader& Reader, bool bIsReinjectedPacket=false );
 
 	/** Packet was negatively acknowledged. */
 	void ReceivedNak( int32 NakPacketId );
@@ -1202,6 +1229,11 @@ public:
 
 protected:
 
+	bool GetPendingCloseDueToSocketSendFailure() const
+	{
+		return bConnectionPendingCloseDueToSocketSendFailure;
+	}
+
 	ENGINE_API void SetPendingCloseDueToSocketSendFailure();
 
 	void CleanupDormantActorState();
@@ -1250,15 +1282,21 @@ private:
 	/** Write packetHeader */
 	void WritePacketHeader(FBitWriter& Writer);
 
-	/** Write extended packet header information (ServerFrameTime, SaturationData) */
-	void WritePacketInfo(FBitWriter& Writer) const;
+	/** Write placeholder bits for data filled before the real send (ServerFrameTime, JitterClockTime) */
+	void WriteDummyPacketInfo(FBitWriter& Writer);
+
+	/** Overwrite the dummy packet info values with real values before sending */
+	void WriteFinalPacketInfo(FBitWriter& Writer, double PacketSentTimeInS);
 	
-	/** Read extended packet header information (ServerFrameTime, SaturationData) */
-	bool ReadPacketInfo(FBitReader& Reader);
+	/** Read extended packet header information (ServerFrameTime) */
+	bool ReadPacketInfo(FBitReader& Reader, bool bHasPacketInfoPayload);
 
 	/** Packet was acknowledged as delivered */
 	void ReceivedAck(int32 AckPacketId);
 
+	/** Calculate the average jitter while adding the new packet's jitter value */
+	void ProcessJitter(int32 PacketJitterClockTimeMS);
+
 	/**
 	 * on the server, the world the client has told us it has loaded
 	 * used to make sure the client has traveled correctly, prevent replicating actors before level transitions are done, etc
@@ -1325,6 +1363,12 @@ private:
 
 	/** Whether or not PacketOrderCache is presently being flushed */
 	bool bFlushingPacketOrderCache;
+
+	/** 
+	* Set to true after the first packet is sent on a Tick. 
+	* Subsequent packets sent on the same frame will not contain the PacketInfo payload
+	*/
+	bool bSentPacketInfoThisFrame = false;
 };
 
 /** Help structs for temporarily setting network settings */
diff --git a/Engine/Source/Runtime/Engine/Classes/Engine/NetDriver.h b/Engine/Source/Runtime/Engine/Classes/Engine/NetDriver.h
--- a/Engine/Source/Runtime/Engine/Classes/Engine/NetDriver.h
+++ b/Engine/Source/Runtime/Engine/Classes/Engine/NetDriver.h
@@ -476,6 +476,13 @@ struct ENGINE_API FPacketSimulationSettings
 	UPROPERTY(EditAnywhere, Category = "Simulation Settings")
 	int32	PktIncomingLoss = 0;
 
+	/**
+	 * Causes sent packets to have a variable latency that fluctuates from [PktLagMin] to [PktLagMin+PktJitter]
+	 * Note that this will cause packet loss on the receiving end.
+	 */
+	UPROPERTY(EditAnywhere, Category = "Simulation Settings")
+	int32	PktJitter = 0;
+
 	/** reads in settings from the .ini file 
 	 * @note: overwrites all previous settings
 	 */
diff --git a/Engine/Source/Runtime/Engine/Classes/Engine/ServerStatReplicator.h b/Engine/Source/Runtime/Engine/Classes/Engine/ServerStatReplicator.h
--- a/Engine/Source/Runtime/Engine/Classes/Engine/ServerStatReplicator.h
+++ b/Engine/Source/Runtime/Engine/Classes/Engine/ServerStatReplicator.h
@@ -40,7 +40,7 @@ public:
 	uint32 OutRate;
 
 	/** @see Network stats counters in EngineStats.h */
-	UPROPERTY(Replicated)
+	UE_DEPRECATED(4.25, TEXT("Saturation is not calculated anymore and is now deprecated."))
 	uint32 OutSaturation;
 
 	/** @see Network stats counters in EngineStats.h */
diff --git a/Engine/Source/Runtime/Engine/Private/Engine/ServerStatReplicator.cpp b/Engine/Source/Runtime/Engine/Private/Engine/ServerStatReplicator.cpp
--- a/Engine/Source/Runtime/Engine/Private/Engine/ServerStatReplicator.cpp
+++ b/Engine/Source/Runtime/Engine/Private/Engine/ServerStatReplicator.cpp
@@ -33,7 +33,6 @@ void AServerStatReplicator::GetLifetimeReplicatedProps(TArray<FLifetimeProperty>
 	DOREPLIFETIME(AServerStatReplicator, Channels);
 	DOREPLIFETIME(AServerStatReplicator, InRate);
 	DOREPLIFETIME(AServerStatReplicator, OutRate);
-	DOREPLIFETIME(AServerStatReplicator, OutSaturation);
 	DOREPLIFETIME(AServerStatReplicator, MaxPacketOverhead);
 	DOREPLIFETIME(AServerStatReplicator, InRateClientMax);
 	DOREPLIFETIME(AServerStatReplicator, InRateClientMin);
@@ -99,7 +98,6 @@ void AServerStatReplicator::Tick(float DeltaSeconds)
 		SET_DWORD_STAT(STAT_InLoss, InLoss);
 		SET_DWORD_STAT(STAT_InRate, InRate);
 		SET_DWORD_STAT(STAT_OutRate, OutRate);
-		SET_DWORD_STAT(STAT_OutSaturation, OutSaturation);
 		SET_DWORD_STAT(STAT_InPackets, InPackets);
 		SET_DWORD_STAT(STAT_OutPackets, OutPackets);
 		SET_DWORD_STAT(STAT_InBunches, InBunches);
diff --git a/Engine/Source/Runtime/Engine/Private/NetConnection.cpp b/Engine/Source/Runtime/Engine/Private/NetConnection.cpp
--- a/Engine/Source/Runtime/Engine/Private/NetConnection.cpp
+++ b/Engine/Source/Runtime/Engine/Private/NetConnection.cpp
@@ -113,6 +113,29 @@ namespace UE4_NetConnectionPrivate
 
 		return Result;
 	}
+
+
+	/** Number of bits to use in the packet header for sending the milliseconds on the clock when the packet is sent */
+	constexpr int32 NumBitsForJitterClockTimeInHeader = 10;
+
+	/** Maximum possible clock time value with available bits. If that value is sent than the jitter clock time is ignored */
+	constexpr int32 MaxJitterClockTimeValue = (1 << NumBitsForJitterClockTimeInHeader) - 1;
+
+	/** Maximum possible precision of the jitter stat. Set to 1 second since we only send the milliseconds of the local clock time */
+	constexpr int32 MaxJitterPrecisionInMS = 1000;
+
+	/** Used to reduce the impact on the average jitter of a divergent value */
+	constexpr double JitterNoiseReduction = 16.0;
+
+	/**
+	 * Return the milliseconds of a ClockTime
+	 * Ex: 3.45 seconds => 450 milliseconds
+	 */
+	int32 GetClockTimeMilliseconds(double ClockTimeInSeconds)
+	{
+		return FMath::TruncToInt((ClockTimeInSeconds - FMath::FloorToDouble(ClockTimeInSeconds)) * 1000.0);
+	}
+	
 }
 
 // ChannelRecord Implementation
@@ -168,7 +191,8 @@ UNetConnection::UNetConnection(const FObjectInitializer& ObjectInitializer)
 ,	ConnectTime			( 0.0 )
 ,	LastOSReceiveTime	()
 ,	bIsOSReceiveTimeLocal(false)
-
+,   PreviousJitterTimeDelta(0)
+,	PreviousPacketSendTimeInS(0.0)
 ,	AllowMerge			( false )
 ,	TimeSensitive		( false )
 ,	LastOutBunch		( nullptr )
@@ -185,6 +209,7 @@ UNetConnection::UNetConnection(const FObjectInitializer& ObjectInitializer)
 ,	CumulativeTime		( 0 )
 ,	AverageFrameTime	( 0 )
 ,	CountedFrames		( 0 )
+,   AverageJitterInMS   (0.0f)
 ,	InBytes				( 0 )
 ,	OutBytes			( 0 )
 ,	InTotalBytes		( 0 )
@@ -1111,6 +1136,8 @@ void UNetConnection::InitSendBuffer()
 		SendBuffer = FBitWriter(FinalBufferSize);
 	}
 
+	HeaderMarkForPacketInfo.Reset();
+
 	ResetPacketBitCounts();
 
 	ValidateSendBuffer();
@@ -1364,6 +1391,7 @@ void UNetConnection::FlushNet(bool bIgnoreSimulation)
 			Handler->OutgoingHigh(SendBuffer);
 		}
 
+		const double PacketSentTimeInS = FPlatformTime::Seconds();
 
 		// Write the UNetConnection-level termination bit
 		SendBuffer.WriteBit(1);
@@ -1374,6 +1402,8 @@ void UNetConnection::FlushNet(bool bIgnoreSimulation)
 			// if we update ack, we also update received ack associated with outgoing seq
 			// so we know how many ack bits we need to write (which is updated in received packet)
 			WritePacketHeader(SendBuffer);
+
+			WriteFinalPacketInfo(SendBuffer, PacketSentTimeInS);
 		}
 
 		ValidateSendBuffer();
@@ -1425,7 +1455,7 @@ void UNetConnection::FlushNet(bool bIgnoreSimulation)
 
 		// Remember the actual time this packet was sent out, so we can compute ping when the ack comes back
 		OutLagPacketId[Index]			= OutPacketId;
-		OutLagTime[Index]				= FPlatformTime::Seconds();
+		OutLagTime[Index]				= PacketSentTimeInS;
 	
 		OutBytesPerSecondHistory[Index]	= FMath::Min(OutBytesPerSecond / 1024, 255);
 
@@ -1445,9 +1475,12 @@ void UNetConnection::FlushNet(bool bIgnoreSimulation)
 		Driver->OutPackets++;
 		Driver->OutTotalPackets++;
 
-		//Record the packet time to the histogram
+		//Record the first packet time in the histogram
+		if (bSentPacketInfoThisFrame == false)
+		{
 		double LastPacketTimeDiffInMs = (Driver->Time - LastSendTime) * 1000.0;
 		NetConnectionHistogram.AddMeasurement(LastPacketTimeDiffInMs);
+		}
 
 		LastSendTime = Driver->Time;
 
@@ -1463,6 +1496,9 @@ void UNetConnection::FlushNet(bool bIgnoreSimulation)
 
 		AnalyticsVars.OutAckOnlyCount += (NumAckBits > 0 && NumBunchBits == 0);
 
+		// Make sure subsequent packets won't have the packet info payload
+		bSentPacketInfoThisFrame = true;
+
 		InitSendBuffer();
 	}
 }
@@ -1495,6 +1531,30 @@ bool UNetConnection::CheckOutgoingPacketEmulation(FOutPacketTraits& Traits)
 
 		return true;
 	}
+	else if (PacketSimulationSettings.PktJitter != 0)
+	{
+		bSendDelayedPacketsOutofOrder = true;
+
+		FDelayedPacket& B = *(new(Delayed)FDelayedPacket(SendBuffer.GetData(), SendBuffer.GetNumBits(), Traits));
+
+		// In order to cause jitter, send one packet at min latency and the next packet at high latency
+		bool bIsLowLatency = (OutPacketId % 2) == 0;
+
+		if (bIsLowLatency)
+		{
+			// ExtraLag goes from [PktLagMin, PktLagMin+PktLagVariance]
+			const double LagVariance = 2.0f * (FMath::FRand() - 0.5f) * double(PacketSimulationSettings.PktLagVariance);
+			const double ExtraLag = (double(PacketSimulationSettings.PktLagMin) + LagVariance) / 1000.f;
+			B.SendTime = FPlatformTime::Seconds() + ExtraLag;
+		}
+		else
+		{
+			// ExtraLag goes from [PktLagMin+PktJitter, PktLagMin+PktJitter+PktLagVariance]
+			const double LagVariance = 2.0f * (FMath::FRand() - 0.5f) * double(PacketSimulationSettings.PktLagVariance);
+			const double ExtraLag = (double(PacketSimulationSettings.PktLagMin + PacketSimulationSettings.PktJitter) + LagVariance) / 1000.f;
+			B.SendTime = FPlatformTime::Seconds() + ExtraLag;
+		}
+	}
 	else if (PacketSimulationSettings.PktLag)
 	{
 		FDelayedPacket& B = *(new(Delayed)FDelayedPacket(SendBuffer.GetData(), SendBuffer.GetNumBits(), Traits));
@@ -1669,25 +1729,96 @@ void UNetConnection::WritePacketHeader(FBitWriter& Writer)
 	}
 }
 
-void UNetConnection::WritePacketInfo(FBitWriter& Writer) const
+void UNetConnection::WriteDummyPacketInfo(FBitWriter& Writer)
+{
+	// The first packet of a frame will contain the packet info payload
+	const bool bHasPacketInfoPayload = bSentPacketInfoThisFrame == false;
+	Writer.WriteBit(bHasPacketInfoPayload);
+
+	if (bHasPacketInfoPayload)
 {
+		// Pre-insert the bits since the final time values will be calculated and inserted right before LowLevelSend
+
+		HeaderMarkForPacketInfo.Init(Writer);
+		int32 DummyJitterClockTime(UE4_NetConnectionPrivate::MaxJitterClockTimeValue);
+		Writer.SerializeBits(&DummyJitterClockTime, UE4_NetConnectionPrivate::NumBitsForJitterClockTimeInHeader);
+
 	const uint8 bHasServerFrameTime = Driver->IsServer() ? bLastHasServerFrameTime : ( CVarPingExcludeFrameTime.GetValueOnGameThread() > 0 ? 1u : 0u );
+		Writer.WriteBit(bHasServerFrameTime);
+
+		if (bHasServerFrameTime && Driver->IsServer())
+		{
+			uint8 DummyFrameTimeByte(0);
+			Writer << DummyFrameTimeByte;
+		}
+	}
+}
+
+void UNetConnection::WriteFinalPacketInfo(FBitWriter& Writer, const double PacketSentTimeInS)
+{
+	const bool bHasPacketInfoPayload = bSentPacketInfoThisFrame == false;
+
+	if (bHasPacketInfoPayload == false)
+	{
+		// PacketInfo payload was already sent, nothing to do here
+		return;
+	}
+
+	FBitWriterMark CurrentMark(Writer);
 
-	// Write data used to calculate link latency
+	// Go back to write over the dummy bits
+	HeaderMarkForPacketInfo.PopWithoutClear(Writer);
+
+	// Write Jitter clock time
+	{
+		const double DeltaSendTimeInMS = (PacketSentTimeInS - PreviousPacketSendTimeInS) * 1000.0;
+
+		// If the delta is over our max precision, we send MAX value and jitter will be ignored by the receiver.
+		if (DeltaSendTimeInMS >= UE4_NetConnectionPrivate::MaxJitterPrecisionInMS)
+		{
+			int32 MaxClockTime = UE4_NetConnectionPrivate::MaxJitterClockTimeValue;
+			Writer.SerializeBits(&MaxClockTime, UE4_NetConnectionPrivate::NumBitsForJitterClockTimeInHeader);
+		}
+		else
+		{
+			// Get the fractional part (milliseconds) of the clock time
+			int32 ClockTimeMilliseconds = UE4_NetConnectionPrivate::GetClockTimeMilliseconds(PacketSentTimeInS);
+
+			// Ensure we don't overflow
+			ClockTimeMilliseconds &= UE4_NetConnectionPrivate::MaxJitterClockTimeValue;
+
+			Writer.SerializeBits(&ClockTimeMilliseconds, UE4_NetConnectionPrivate::NumBitsForJitterClockTimeInHeader);
+		}
+
+		PreviousPacketSendTimeInS = PacketSentTimeInS;
+	}
+
+	// Write server frame time
+	{
+		const uint8 bHasServerFrameTime = Driver->IsServer() ? bLastHasServerFrameTime : (CVarPingExcludeFrameTime.GetValueOnGameThread() > 0 ? 1u : 0u);
 	Writer.WriteBit(bHasServerFrameTime);
+
 	if (bHasServerFrameTime && Driver->IsServer())
 	{
+			// Write data used to calculate link latency
 		uint8 FrameTimeByte = FMath::Min( FMath::FloorToInt( FrameTime * 1000 ), 255 );
 		Writer << FrameTimeByte;
 	}
+	}
 
-	// Notify server of our current rate per second at this time
-	uint8 InKBytesPerSecondByte = FMath::Min(InBytesPerSecond / 1024, 255);
-	Writer << InKBytesPerSecondByte;
+	// Revert to the correct bit writing place
+	CurrentMark.PopWithoutClear(Writer);
 }
 
-bool UNetConnection::ReadPacketInfo(FBitReader& Reader)
+bool UNetConnection::ReadPacketInfo(FBitReader& Reader, bool bHasPacketInfoPayload)
 {
+	// If this packet did not contain any packet info, nothing else to read
+	if (bHasPacketInfoPayload == false)
+	{
+		const bool bCanContinueReading = Reader.IsError() == false;
+		return bCanContinueReading;
+	}
+
 	const bool bHasServerFrameTime = Reader.ReadBit() == 1u;
 	double ServerFrameTime = 0;
 
@@ -1706,10 +1837,13 @@ bool UNetConnection::ReadPacketInfo(FBitReader& Reader)
 		bLastHasServerFrameTime = bHasServerFrameTime;
 	}
 
-	// limit to known size to know the size of the packet header
+	if (Reader.EngineNetVer() < HISTORY_JITTER_IN_HEADER)
+	{
 	uint8 RemoteInKBytesPerSecondByte = 0;
 	Reader << RemoteInKBytesPerSecondByte;
+	}
 
+	// limit to known size to know the size of the packet header
 	if ( Reader.IsError() )
 	{
 		return false;
@@ -1752,15 +1886,6 @@ bool UNetConnection::ReadPacketInfo(FBitReader& Reader)
 		const double RTT		 = (CurrentTime - OutLagTime[Index] ) - ( CVarPingExcludeFrameTime.GetValueOnAnyThread() ? GameTime : 0.0 );
 		const double NewLag		 = FMath::Max( RTT, 0.0 );
 
-		if ( OutBytesPerSecondHistory[Index] > 0 )
-		{
-			RemoteSaturation = ( 1.0f - FMath::Min( ( float )RemoteInKBytesPerSecondByte / ( float )OutBytesPerSecondHistory[Index], 1.0f ) ) * 100.0f;
-		}
-		else
-		{
-			RemoteSaturation = 0.0f;
-		}
-
 		//UE_LOG( LogNet, Warning, TEXT( "Out: %i, InRemote: %i, Saturation: %f" ), OutBytesPerSecondHistory[Index], RemoteInKBytesPerSecond, RemoteSaturation );
 
 		LagAcc += NewLag;
@@ -1803,7 +1928,7 @@ FNetworkGUID UNetConnection::GetActorGUIDFromOpenBunch(FInBunch& Bunch)
 	return ActorGUID;
 }
 
-void UNetConnection::ReceivedPacket( FBitReader& Reader )
+void UNetConnection::ReceivedPacket( FBitReader& Reader, bool bIsReinjectedPacket)
 {
 	SCOPED_NAMED_EVENT(UNetConnection_ReceivedPacket, FColor::Green);
 	AssertValid();
@@ -1851,12 +1976,16 @@ void UNetConnection::ReceivedPacket( FBitReader& Reader )
 	ValidateSendBuffer();
 
 	//Record the packet time to the histogram
-	double LastPacketTimeDiffInMs = (FPlatformTime::Seconds() - LastReceiveRealtime) * 1000.0;
+	if (!bIsReinjectedPacket)
+	{
+		const double CurrentReceiveTimeInS = FPlatformTime::Seconds();
+		const double LastPacketTimeDiffInMs = (CurrentReceiveTimeInS - LastReceiveRealtime) * 1000.0;
 	NetConnectionHistogram.AddMeasurement(LastPacketTimeDiffInMs);
 
 	// Update receive time to avoid timeout.
 	LastReceiveTime		= Driver->Time;
-	LastReceiveRealtime = FPlatformTime::Seconds();
+		LastReceiveRealtime = CurrentReceiveTimeInS;
+	}
 
 	if (InternalAck)
 	{
@@ -1872,15 +2001,34 @@ void UNetConnection::ReceivedPacket( FBitReader& Reader )
 			return;
 		}
 
+		bool bHasPacketInfoPayload = true;
+
+		if (Reader.EngineNetVer() >= HISTORY_JITTER_IN_HEADER)
+		{
+			bHasPacketInfoPayload = Reader.ReadBit() == 1u;
+
+			if (bHasPacketInfoPayload)
+			{
+				// Read jitter clock time from the packet header
+				int32 PacketJitterClockTimeMS = 0;
+				Reader.SerializeBits(&PacketJitterClockTimeMS, UE4_NetConnectionPrivate::NumBitsForJitterClockTimeInHeader);
+
+				if (!bIsReinjectedPacket)
+				{
+					ProcessJitter(PacketJitterClockTimeMS);
+				}
+			}			
+		}
 
-		bool bPacketOrderCacheActive = !bFlushingPacketOrderCache && PacketOrderCache.IsSet();
-		bool bCheckForMissingSequence = bPacketOrderCacheActive && PacketOrderCacheCount == 0;
-		bool bFillingPacketOrderCache = bPacketOrderCacheActive && PacketOrderCacheCount > 0;
-		int32 MaxMissingPackets = (bCheckForMissingSequence ? CVarNetPacketOrderMaxMissingPackets.GetValueOnAnyThread() : 0);
-		int32 PacketSequenceDelta = PacketNotify.GetSequenceDelta(Header);
+		const int32 PacketSequenceDelta = PacketNotify.GetSequenceDelta(Header);
 
 		if (PacketSequenceDelta > 0)
 		{
+			const bool bPacketOrderCacheActive = !bFlushingPacketOrderCache && PacketOrderCache.IsSet();
+			const bool bCheckForMissingSequence = bPacketOrderCacheActive && PacketOrderCacheCount == 0;
+			const bool bFillingPacketOrderCache = bPacketOrderCacheActive && PacketOrderCacheCount > 0;
+			const int32 MaxMissingPackets = (bCheckForMissingSequence ? CVarNetPacketOrderMaxMissingPackets.GetValueOnAnyThread() : 0);
+
 			const int32 MissingPacketCount = PacketSequenceDelta - 1;
 
 			// Cache the packet if we are already caching, and begin caching if we just encountered a missing sequence, within range
@@ -1893,14 +2041,14 @@ void UNetConnection::ReceivedPacket( FBitReader& Reader )
 				// The last cache entry is only set, when we've reached capacity or when we receive a sequence which is out of bounds of the cache
 				LinearCacheIdx = bLastCacheEntry ? (CacheCapacity - 1) : LinearCacheIdx;
 
-				int32 CiruclarCacheIdx = PacketOrderCacheStartIdx;
+				int32 CircularCacheIdx = PacketOrderCacheStartIdx;
 
 				for (int32 LinearDec=LinearCacheIdx; LinearDec > 0; LinearDec--)
 				{
-					CiruclarCacheIdx = PacketOrderCache->GetNextIndex(CiruclarCacheIdx);
+					CircularCacheIdx = PacketOrderCache->GetNextIndex(CircularCacheIdx);
 				}
 
-				TUniquePtr<FBitReader>& CurCachePacket = PacketOrderCache.GetValue()[CiruclarCacheIdx];
+				TUniquePtr<FBitReader>& CurCachePacket = PacketOrderCache.GetValue()[CircularCacheIdx];
 
 				// Reset the reader to its initial position, and cache the packet
 				if (!CurCachePacket.IsValid())
@@ -1992,7 +2140,7 @@ void UNetConnection::ReceivedPacket( FBitReader& Reader )
 		PacketNotify.Update(Header, HandlePacketNotification);
 
 		// Extra information associated with the header (read only after acks have been processed)
-		if (PacketSequenceDelta > 0 && !ReadPacketInfo(Reader))
+		if (PacketSequenceDelta > 0 && !ReadPacketInfo(Reader, bHasPacketInfoPayload))
 		{
 			CLOSE_CONNECTION_DUE_TO_SECURITY_VIOLATION(this, ESecurityEvent::Malformed_Packet, TEXT("Failed to read PacketHeader"));
 			return;
@@ -2547,8 +2695,8 @@ int32 UNetConnection::WriteBitsToSendBuffer(
 		// Write Packet Header, before sending the packet we will go back and rewrite the data
 		WritePacketHeader(SendBuffer);
 
-		// Also write server RTT and Received rate
-		WritePacketInfo(SendBuffer);
+		// Pre-write the bits for the packet info
+		WriteDummyPacketInfo(SendBuffer);
 
 		// We do not allow the first bunch to merge with the ack data as this will "revert" the ack data.
 		AllowMerge = false;
@@ -2844,6 +2992,23 @@ void UNetConnection::Tick()
 #if DO_ENABLE_NET_TEST
 	if (Delayed.Num() > 0)
 	{
+		if (bSendDelayedPacketsOutofOrder)
+		{
+			for (int Idx=0; Idx < Delayed.Num(); ++Idx)
+			{
+				if (CurrentRealtimeSeconds > Delayed[Idx].SendTime)
+				{
+					LowLevelSend((char*)&(Delayed[Idx].Data[0]), Delayed[Idx].SizeBits, Delayed[Idx].Traits);
+				}
+			}
+
+			Delayed.RemoveAll([&CurrentRealtimeSeconds](const FDelayedPacket& rhs)
+			{
+				return CurrentRealtimeSeconds > rhs.SendTime;
+			});
+		}
+		else
+		{
 		uint32 NbPacketsSent(0);
 		for( FDelayedPacket& DelayedPacket : Delayed )
 		{
@@ -2861,6 +3026,7 @@ void UNetConnection::Tick()
 
 		Delayed.RemoveAt(0, NbPacketsSent, false);
 	}
+	}
 #endif
 
 	// if this is 0 it's our first tick since init, so start our real-time tracking from here
@@ -3180,6 +3346,8 @@ void UNetConnection::Tick()
 	{
 		QueuedBits = FMath::TruncToInt(-AllowedLag);
 	}
+
+	bSentPacketInfoThisFrame = false;
 }
 
 void UNetConnection::HandleClientPlayer( APlayerController *PC, UNetConnection* NetConnection )
@@ -3727,6 +3895,39 @@ void UNetConnection::TrackReplicationForAnalytics(const bool bWasSaturated)
 	SaturationAnalytics.TrackReplication(bWasSaturated);
 }
 
+void UNetConnection::ProcessJitter(int32 PacketJitterClockTimeMS)
+{
+	if (PacketJitterClockTimeMS >= UE4_NetConnectionPrivate::MaxJitterClockTimeValue)
+	{
+		// Ignore this packet for jitter calculations
+		return;
+	}
+
+	const int32 CurrentClockTimeMilliseconds = UE4_NetConnectionPrivate::GetClockTimeMilliseconds(LastReceiveRealtime);
+
+	// Get the delta between the sent and receive clock time.
+	int32 CurrentJitterTimeDelta = CurrentClockTimeMilliseconds - PacketJitterClockTimeMS;
+	
+	if (CurrentJitterTimeDelta < 0)
+	{
+		// Account for wrap-around
+		CurrentJitterTimeDelta += UE4_NetConnectionPrivate::MaxJitterPrecisionInMS;
+	}
+
+	// Get the difference between the last two deltas
+	const float CurrentJitter = FMath::Abs(CurrentJitterTimeDelta - PreviousJitterTimeDelta);
+	
+
+	// Add the value to the average and smooth it out to reduce divergences
+	const float PreviousJitterInMS = AverageJitterInMS;
+	AverageJitterInMS = PreviousJitterInMS + ((CurrentJitter - PreviousJitterInMS) / UE4_NetConnectionPrivate::JitterNoiseReduction);
+
+	//UE_LOG(LogNetTraffic, Verbose, TEXT("Avg Jitter: %f | CurrentJitter: %f | CurrentTimeDelta: %d | PreviousTimeDelta: %d | ReceivedClockMilli %d | LocalClockMilli %d"), AverageJitterInMS, CurrentJitter, CurrentJitterTimeDelta, PreviousJitterTimeDelta, PacketJitterClockTimeMS, CurrentClockTimeMilliseconds);
+
+	// Store for next process
+	PreviousJitterTimeDelta = CurrentJitterTimeDelta;
+}
+
 void UNetConnection::SendChallengeControlMessage()
 {
 	if (State != USOCK_Invalid && State != USOCK_Closed && Driver)
diff --git a/Engine/Source/Runtime/Engine/Private/NetworkDriver.cpp b/Engine/Source/Runtime/Engine/Private/NetworkDriver.cpp
--- a/Engine/Source/Runtime/Engine/Private/NetworkDriver.cpp
+++ b/Engine/Source/Runtime/Engine/Private/NetworkDriver.cpp
@@ -84,7 +84,6 @@ DEFINE_STAT(STAT_Channels);
 DEFINE_STAT(STAT_MaxPacketOverhead);
 DEFINE_STAT(STAT_InRate);
 DEFINE_STAT(STAT_OutRate);
-DEFINE_STAT(STAT_OutSaturation);
 DEFINE_STAT(STAT_InRateClientMax);
 DEFINE_STAT(STAT_InRateClientMin);
 DEFINE_STAT(STAT_InRateClientAvg);
@@ -616,7 +615,6 @@ void UNetDriver::TickFlush(float DeltaSeconds)
 			int32 ClientOutPacketsAvg = 0;
 			int NumClients = 0;
 			int32 MaxPacketOverhead = 0;
-			float RemoteSaturationMax = 0.0f;
 
 			// these need to be updated even if we are not collecting stats, since they get reported to analytics/QoS
 			for (UNetConnection * Client : ClientConnections)
@@ -698,13 +696,11 @@ void UNetDriver::TickFlush(float DeltaSeconds)
 				if (ServerConnection != NULL)
 				{
 					NumOpenChannels = ServerConnection->OpenChannels.Num();
-					RemoteSaturationMax = FMath::Max(RemoteSaturationMax, ServerConnection->RemoteSaturation);
 				}
 
 				for (int32 i = 0; i < ClientConnections.Num(); i++)
 				{
 					NumOpenChannels += ClientConnections[i]->OpenChannels.Num();
-					RemoteSaturationMax = FMath::Max(RemoteSaturationMax, ClientConnections[i]->RemoteSaturation);
 				}
 
 				// Use the elapsed time to keep things scaled to one measured unit
@@ -767,7 +763,6 @@ void UNetDriver::TickFlush(float DeltaSeconds)
 				SET_DWORD_STAT(STAT_InLoss, InPacketsLost);
 				SET_DWORD_STAT(STAT_InRate, InBytes);
 				SET_DWORD_STAT(STAT_OutRate, OutBytes);
-				SET_DWORD_STAT(STAT_OutSaturation, RemoteSaturationMax);
 				SET_DWORD_STAT(STAT_InRateClientMax, ClientInBytesMax);
 				SET_DWORD_STAT(STAT_InRateClientMin, ClientInBytesMin);
 				SET_DWORD_STAT(STAT_InRateClientAvg, ClientInBytesAvg);
@@ -819,7 +814,6 @@ void UNetDriver::TickFlush(float DeltaSeconds)
 					GMBase->ServerStatReplicator->InLoss = InPacketsLost;
 					GMBase->ServerStatReplicator->InRate = InBytes;
 					GMBase->ServerStatReplicator->OutRate = OutBytes;
-					GMBase->ServerStatReplicator->OutSaturation = RemoteSaturationMax;
 					GMBase->ServerStatReplicator->InRateClientMax = ClientInBytesMax;
 					GMBase->ServerStatReplicator->InRateClientMin = ClientInBytesMin;
 					GMBase->ServerStatReplicator->InRateClientAvg = ClientInBytesAvg;
@@ -959,8 +953,6 @@ void UNetDriver::TickFlush(float DeltaSeconds)
 				PerfCounters->Set(TEXT("OutPackets"), OutPackets);
 				PerfCounters->Set(TEXT("InBunches"), InBunches);
 				PerfCounters->Set(TEXT("OutBunches"), OutBunches);
-
-				PerfCounters->Set(TEXT("OutSaturationMax"), RemoteSaturationMax);
 			}
 #endif // USE_SERVER_PERF_COUNTERS
 
@@ -3361,6 +3353,8 @@ void FPacketSimulationSettings::LoadConfig(const TCHAR* OptionalQualifier)
 	ConfigHelperInt(TEXT("PktIncomingLagMax"), PktIncomingLagMax, OptionalQualifier);
 	ConfigHelperInt(TEXT("PktIncomingLoss"), PktIncomingLoss, OptionalQualifier);
 
+	ConfigHelperInt(TEXT("PktJitter"), PktJitter, OptionalQualifier);
+
 	ValidateSettings();
 }
 
@@ -3541,6 +3535,11 @@ bool FPacketSimulationSettings::ParseSettings(const TCHAR* Cmd, const TCHAR* Opt
 		bParsed = true;
 		UE_LOG(LogNet, Log, TEXT("PktIncomingLoss set to %d"), PktIncomingLoss);
 	}
+	if (ParseHelper(Cmd, TEXT("PktJitter="), PktJitter, OptionalQualifier))
+	{
+		bParsed = true;
+		UE_LOG(LogNet, Log, TEXT("PktJitter set to %d"), PktJitter);
+	}
 
 	ValidateSettings();
 	return bParsed;
@@ -5723,6 +5722,7 @@ BUILD_NETEMULATION_CONSOLE_COMMAND(PktLagMax, "Sets maximum outgoing packet late
 BUILD_NETEMULATION_CONSOLE_COMMAND(PktIncomingLagMin, "Sets minimum incoming packet latency");
 BUILD_NETEMULATION_CONSOLE_COMMAND(PktIncomingLagMax, "Sets maximum incoming packet latency");
 BUILD_NETEMULATION_CONSOLE_COMMAND(PktIncomingLoss, "Simulates incoming packet loss");
+BUILD_NETEMULATION_CONSOLE_COMMAND(PktJitter, "Simulates outgoing packet jitter");
 
 #endif //#if DO_ENABLE_NET_TEST
 
diff --git a/Engine/Source/Runtime/Engine/Public/EngineStats.h b/Engine/Source/Runtime/Engine/Public/EngineStats.h
--- a/Engine/Source/Runtime/Engine/Public/EngineStats.h
+++ b/Engine/Source/Runtime/Engine/Public/EngineStats.h
@@ -174,7 +174,6 @@ DECLARE_DWORD_ACCUMULATOR_STAT_EXTERN(TEXT("Channels"),STAT_Channels,STATGROUP_N
 DECLARE_DWORD_ACCUMULATOR_STAT_EXTERN(TEXT("Packet Overhead - max among connections"), STAT_MaxPacketOverhead, STATGROUP_Net, );
 DECLARE_DWORD_ACCUMULATOR_STAT_EXTERN(TEXT("In Rate (bytes)"),STAT_InRate,STATGROUP_Net, );
 DECLARE_DWORD_ACCUMULATOR_STAT_EXTERN(TEXT("Out Rate (bytes)"),STAT_OutRate,STATGROUP_Net, );
-DECLARE_DWORD_ACCUMULATOR_STAT_EXTERN( TEXT( "Out Saturation (%)" ),STAT_OutSaturation, STATGROUP_Net,);
 DECLARE_DWORD_ACCUMULATOR_STAT_EXTERN(TEXT("In Rate (bytes) - max among clients"),STAT_InRateClientMax,STATGROUP_Net, );
 DECLARE_DWORD_ACCUMULATOR_STAT_EXTERN(TEXT("In Rate (bytes) - min among clients"),STAT_InRateClientMin,STATGROUP_Net, );
 DECLARE_DWORD_ACCUMULATOR_STAT_EXTERN(TEXT("In Rate (bytes) - avg among clients"),STAT_InRateClientAvg,STATGROUP_Net, );
