Fix build due to TEXT("") reason in UE_DEPRECATED macro

#rnx

diff --git a/Engine/Source/Runtime/Engine/Classes/Engine/ServerStatReplicator.h b/Engine/Source/Runtime/Engine/Classes/Engine/ServerStatReplicator.h
--- a/Engine/Source/Runtime/Engine/Classes/Engine/ServerStatReplicator.h
+++ b/Engine/Source/Runtime/Engine/Classes/Engine/ServerStatReplicator.h
@@ -40,7 +40,7 @@ public:
 	uint32 OutRate;
 
 	/** @see Network stats counters in EngineStats.h */
-	UE_DEPRECATED(4.25, TEXT("Saturation is not calculated anymore and is now deprecated."))
+	UE_DEPRECATED(4.25, "Saturation is not calculated anymore and is now deprecated.")
 	uint32 OutSaturation;
 
 	/** @see Network stats counters in EngineStats.h */
