Fix another TEXT() in UE_DEPRECATED macro

#rnx

diff --git a/Engine/Source/Runtime/Engine/Classes/Engine/NetConnection.h b/Engine/Source/Runtime/Engine/Classes/Engine/NetConnection.h
--- a/Engine/Source/Runtime/Engine/Classes/Engine/NetConnection.h
+++ b/Engine/Source/Runtime/Engine/Classes/Engine/NetConnection.h
@@ -493,7 +493,7 @@ public:
 	double			OutLagTime[256];				// For lag measuring.
 	int32			OutLagPacketId[256];			// For lag measuring.
 	uint8			OutBytesPerSecondHistory[256];	// For saturation measuring.
-	UE_DEPRECATED(4.25, TEXT("RemoteSaturation is not calculated anymore and is now deprecated"))
+	UE_DEPRECATED(4.25, "RemoteSaturation is not calculated anymore and is now deprecated")
 	float			RemoteSaturation;
 	int32			InPacketId;						// Full incoming packet index.
 	int32			OutPacketId;					// Most recently sent packet.
