Updated comments for FMatrix blueprint functions to explicitly list what assumptions are being made in given functions

Removed autocast from MatrixToTransform and MatrixToRotator as these are only valid when the matrix represents a transform so users should be required to explicitly choose to convert.

diff --git a/Engine/Source/Runtime/Engine/Classes/Kismet/KismetMathLibrary.h b/Engine/Source/Runtime/Engine/Classes/Kismet/KismetMathLibrary.h
--- a/Engine/Source/Runtime/Engine/Classes/Kismet/KismetMathLibrary.h
+++ b/Engine/Source/Runtime/Engine/Classes/Kismet/KismetMathLibrary.h
@@ -2204,16 +2204,21 @@ class ENGINE_API UKismetMathLibrary : public UBlueprintFunctionLibrary
 	// Matrix functions
 	//
 
-	/** Convert a Matrix to a Transform */
-	UFUNCTION(BlueprintPure, meta = (DisplayName = "To Transform (Matrix)", CompactNodeTitle = "->", ScriptMethod = "Transform", Keywords = "cast convert", BlueprintAutocast), Category = "Math|Conversions")
+	/** Convert a Matrix to a Transform 
+	* (Assumes Matrix represents a transform) 
+	*/
+	UFUNCTION(BlueprintPure, meta = (DisplayName = "To Transform (Matrix)", CompactNodeTitle = "->", ScriptMethod = "Transform", Keywords = "cast convert"), Category = "Math|Conversions")
 	static FTransform Conv_MatrixToTransform(const FMatrix& InMatrix);
 
-	/** Convert a Matrix to a Rotator */
-	UFUNCTION(BlueprintPure, meta = (DisplayName = "To Rotator (Matrix)", CompactNodeTitle = "->", ScriptMethod = "Rotator", Keywords = "cast convert", BlueprintAutocast), Category = "Math|Conversions")
+	/** Convert a Matrix to a Rotator 
+	* (Assumes Matrix represents a transform) 
+	*/
+	UFUNCTION(BlueprintPure, meta = (DisplayName = "To Rotator (Matrix)", CompactNodeTitle = "->", ScriptMethod = "Rotator", Keywords = "cast convert"), Category = "Math|Conversions")
 	static FRotator Conv_MatrixToRotator(const FMatrix& InMatrix);
 
 	/**
 	 * Get the origin of the co-ordinate system
+	 * (Assumes Matrix represents a transform)
 	 *
 	 * @return co-ordinate system origin
 	 */
@@ -2243,7 +2248,8 @@ class ENGINE_API UKismetMathLibrary : public UBlueprintFunctionLibrary
 	static FMatrix Add_MatrixMatrix (const FMatrix& A, const FMatrix& B);
 
 	/**
-	  * This isn't applying SCALE, just multiplying the value to all members - i.e. weighting
+	  * Multiplies all values of the matrix by a float.
+	  * If your Matrix represents a Transform that you wish to scale you should use Apply Scale instead
 	  */
 	UFUNCTION(BlueprintPure, meta = (DisplayName = "Matrix * Float", CompactNodeTitle = "*", ScriptMethod = "MultiplyFloat", ScriptOperator = "*;*=", Keywords = "* multiply"), Category = "Math|Matrix")
 	static FMatrix Multiply_MatrixFloat (const FMatrix& A, float B);
@@ -2267,29 +2273,38 @@ class ENGINE_API UKismetMathLibrary : public UBlueprintFunctionLibrary
 	UFUNCTION(BlueprintPure, meta = (DisplayName = "Not Equal (Matrix)", CompactNodeTitle = "!=", ScriptMethod = "NotEqual", ScriptOperator = "!=", Keywords = "!= not equal"), Category = "Math|Matrix")
 	static bool NotEqual_MatrixMatrix(const FMatrix& A, const FMatrix& B, float Tolerance = 1.e-4f);
 
-	// Homogeneous transform.
+	/**
+	 * Transform a vector by the matrix.
+	 * (Assumes Matrix represents a transform) 
+	*/
 	UFUNCTION(BlueprintPure, meta = (DisplayName = "Transform Vector4 (Matrix)", ScriptMethod = "TransformVector4"), Category = "Math|Matrix")
 	static FVector4 Matrix_TransformVector4(const FMatrix& M, FVector4 V);
 
-	/** Transform a location - will take into account translation part of the FMatrix. */
+	/** Transform a location - will take into account translation part of the FMatrix.
+	 * (Assumes Matrix represents a transform)
+	 */
 	UFUNCTION(BlueprintPure, meta = (DisplayName = "Transform Position (Matrix)", ScriptMethod = "TransformPosition"), Category = "Math|Matrix")
 	static FVector4 Matrix_TransformPosition(const FMatrix& M, FVector V);
 
-	/** Inverts the matrix and then transforms V - correctly handles scaling in this matrix. */
+	/** Inverts the matrix and then transforms V - correctly handles scaling in this matrix.
+	* (Assumes Matrix represents a transform)
+	*/
 	UFUNCTION(BlueprintPure, meta = (DisplayName = "Inverse Transform Position (Matrix)", ScriptMethod = "InverseTransformPosition"), Category = "Math|Matrix")
 	static FVector Matrix_InverseTransformPosition(const FMatrix& M, FVector V);
 
 	/**
 	 *	Transform a direction vector - will not take into account translation part of the FMatrix.
 	 *	If you want to transform a surface normal (or plane) and correctly account for non-uniform scaling you should use TransformByUsingAdjointT.
-	 */
+	 * (Assumes Matrix represents a transform)
+	*/
 	UFUNCTION(BlueprintPure, meta = (DisplayName = "Transform Vector (Matrix)", ScriptMethod = "TransformVector"), Category = "Math|Matrix")
 	static FVector4 Matrix_TransformVector(const FMatrix& M, FVector V);
 
 	/**
 	 *	Transform a direction vector by the inverse of this matrix - will not take into account translation part.
 	 *	If you want to transform a surface normal (or plane) and correctly account for non-uniform scaling you should use TransformByUsingAdjointT with adjoint of matrix inverse.
-	 */
+	 * (Assumes Matrix represents a transform)
+	*/
 	UFUNCTION(BlueprintPure, meta = (DisplayName = "Inverse Transform Vector (Matrix)", ScriptMethod = "InverseTransformVector"), Category = "Math|Matrix")
 	static FVector Matrix_InverseTransformVector(const FMatrix& M, FVector V);
 
@@ -2301,34 +2316,47 @@ class ENGINE_API UKismetMathLibrary : public UBlueprintFunctionLibrary
 	UFUNCTION(BlueprintPure, meta = (DisplayName = "Get Determinant (Matrix)", ScriptMethod = "GetDeterminant"), Category = "Math|Matrix")
 	static float Matrix_GetDeterminant(const FMatrix& M);
 
-	/** @return the determinant of rotation 3x3 matrix */
+	/** @return the determinant of rotation 3x3 matrix 
+	* (Assumes Top Left 3x3 Submatrix represents a Rotation)
+	*/
 	UFUNCTION(BlueprintPure, meta = (DisplayName = "Get Rotation Determinant (Matrix)", ScriptMethod = "GetRotDeterminant"), Category = "Math|Matrix")
 	static float Matrix_GetRotDeterminant(const FMatrix& M);
 
-	/** Fast path, and handles nil matrices. */
+	/** Get the inverse of the Matrix. Handles nil matrices. */
 	UFUNCTION(BlueprintPure, meta = (DisplayName = "GetInverse (Matrix)", ScriptMethod = "GetInverse"), Category = "Math|Matrix")
 	static FMatrix Matrix_GetInverse(const FMatrix& M);
 
+	/** Get the Transose Adjoint of the Matrix. */
 	UFUNCTION(BlueprintPure, meta = (DisplayName = "Get Transpose Adjoint (Matrix)", ScriptMethod = "GetTransposeAdjoint"), Category = "Math|Matrix")
 	static FMatrix Matrix_GetTransposeAdjoint(const FMatrix& M);
 
-	// Remove any scaling from this matrix (ie magnitude of each row is 1) with error Tolerance
+	/** Remove any scaling from this matrix (ie magnitude of each row is 1) with error Tolerance
+	* (Assumes Matrix represents a transform) 
+	*/
 	UFUNCTION(BlueprintCallable, meta = (DisplayName = "Remove Scaling (Matrix)", ScriptMethod = "RemoveScaling"), Category = "Math|Matrix")
 	static void Matrix_RemoveScaling(UPARAM(Ref) FMatrix& M, float Tolerance = 1.e-8f);
 
-	// Returns matrix after RemoveScaling with error Tolerance
+	/** Returns matrix after RemoveScaling with error Tolerance
+	* (Assumes Matrix represents a transform)
+	*/
 	UFUNCTION(BlueprintPure, meta = (DisplayName = "Get Matrix Without Scale (Matrix)", ScriptMethod = "GetMatrixWithoutScale"), Category = "Math|Matrix")
 	static FMatrix Matrix_GetMatrixWithoutScale(const FMatrix& M, float Tolerance = 1.e-8f);
 
-	/** return a 3D scale vector calculated from this matrix (where each component is the magnitude of a row vector) with error Tolerance. */
+	/** return a 3D scale vector calculated from this matrix (where each component is the magnitude of a row vector) with error Tolerance.
+	* (Assumes Matrix represents a transform)
+	*/
 	UFUNCTION(BlueprintPure, meta = (DisplayName = "Get Scale Vector (Matrix)", ScriptMethod = "GetScaleVector"), Category = "Math|Matrix")
 	static FVector Matrix_GetScaleVector(const FMatrix& M, float Tolerance = 1.e-8f);
 
-	// Remove any translation from this matrix
+	/** Remove any translation from this matrix
+	* (Assumes Matrix represents a transform)
+	*/
 	UFUNCTION(BlueprintPure, meta = (DisplayName = "Remove Translation (Matrix)", ScriptMethod = "RemoveTranslation"), Category = "Math|Matrix")
 	static FMatrix Matrix_RemoveTranslation(const FMatrix& M);
 
-	/** Returns a matrix with an additional translation concatenated. */
+	/** Returns a matrix with an additional translation concatenated.
+	* (Assumes Matrix represents a transform)
+	*/
 	UFUNCTION(BlueprintPure, meta = (DisplayName = "Concatenate Translation (Matrix)", ScriptMethod = "ConcatenateTranslation"), Category = "Math|Matrix")
 	static FMatrix Matrix_ConcatenateTranslation(const FMatrix& M, FVector Translation);
 
@@ -2336,20 +2364,27 @@ class ENGINE_API UKismetMathLibrary : public UBlueprintFunctionLibrary
 	UFUNCTION(BlueprintPure, meta = (DisplayName = "Contains NaN (Matrix)", ScriptMethod = "ContainsNaN"), Category = "Math|Matrix")
 	static bool Matrix_ContainsNaN(const FMatrix& M);
 
-	/** Scale the translation part of the matrix by the supplied vector. */
+	/** Scale the translation part of the matrix by the supplied vector.
+	* (Assumes Matrix represents a transform)
+	*/
 	UFUNCTION(BlueprintPure, meta = (DisplayName = "Scale Translation (Matrix)", ScriptMethod = "ScaleTranslation"), Category = "Math|Matrix")
 	static FMatrix Matrix_ScaleTranslation(const FMatrix& M, FVector Scale3D);
 
-	/** @return the maximum magnitude of any row of the matrix. */
+	/** @return the maximum magnitude of any row of the matrix.
+	* (Assumes Matrix represents a transform)
+	*/
 	UFUNCTION(BlueprintPure, meta = (DisplayName = "Get Maximum Axis Scale (Matrix)", ScriptMethod = "GetMaximumAxisScale"), Category = "Math|Matrix")
 	static float Matrix_GetMaximumAxisScale(const FMatrix& M);
 
-	/** Apply Scale to this matrix **/
+	/** Apply Scale to this matrix
+	* (Assumes Matrix represents a transform)
+	*/
 	UFUNCTION(BlueprintPure, meta = (DisplayName = "Apply Scale (Matrix)", ScriptMethod = "ApplyScale"), Category = "Math|Matrix")
 	static FMatrix Matrix_ApplyScale(const FMatrix& M, float Scale);
 
 	/**
 	 * get axis of this matrix scaled by the scale of the matrix
+	 * (Assumes Matrix represents a transform)
 	 *
 	 * @param i index into the axis of the matrix
 	 * @ return vector of the axis
@@ -2359,6 +2394,7 @@ class ENGINE_API UKismetMathLibrary : public UBlueprintFunctionLibrary
 
 	/**
 	 * get axes of this matrix scaled by the scale of the matrix
+	 * (Assumes Matrix represents a transform)
 	 *
 	 * @param X axes returned to this param
 	 * @param Y axes returned to this param
@@ -2369,6 +2405,7 @@ class ENGINE_API UKismetMathLibrary : public UBlueprintFunctionLibrary
 
 	/**
 	 * get unit length axis of this matrix
+	 * (Assumes Matrix represents a transform)
 	 *
 	 * @param i index into the axis of the matrix
 	 * @return vector of the axis
@@ -2378,6 +2415,7 @@ class ENGINE_API UKismetMathLibrary : public UBlueprintFunctionLibrary
 
 	/**
 	 * get unit length axes of this matrix
+	 * (Assumes Matrix represents a transform)
 	 *
 	 * @param X axes returned to this param
 	 * @param Y axes returned to this param
@@ -2388,6 +2426,7 @@ class ENGINE_API UKismetMathLibrary : public UBlueprintFunctionLibrary
 
 	/**
 	 * set an axis of this matrix
+	 * (Assumes Matrix represents a transform)
 	 *
 	 * @param i index into the axis of the matrix
 	 * @param Axis vector of the axis
@@ -2395,7 +2434,9 @@ class ENGINE_API UKismetMathLibrary : public UBlueprintFunctionLibrary
 	UFUNCTION(BlueprintCallable, meta = (DisplayName = "Set Axis (Matrix)", ScriptMethod = "SetAxis"), Category = "Math|Matrix")
 	static void Matrix_SetAxis(UPARAM(Ref) FMatrix& M, TEnumAsByte<EAxis::Type> Axis, FVector AxisVector);
 
-	// Set the origin of the coordinate system to the given vector
+	/** Set the origin of the coordinate system to the given vector
+	* (Assumes Matrix represents a transform)
+	*/
 	UFUNCTION(BlueprintCallable, meta = (DisplayName = "Set Origin (Matrix)", ScriptMethod = "SetOrigin"), Category = "Math|Matrix")
 	static void Matrix_SetOrigin(UPARAM(Ref) FMatrix& M, FVector NewOrigin);
 
@@ -2411,12 +2452,17 @@ class ENGINE_API UKismetMathLibrary : public UBlueprintFunctionLibrary
 	UFUNCTION(BlueprintCallable, meta = (DisplayName = "Set Column (Matrix)", ScriptMethod = "SetColumn"), Category = "Math|Matrix")
 	static void Matrix_SetColumn(UPARAM(Ref) FMatrix& M, TEnumAsByte<EMatrixColumns::Type> Column, FVector Value);
 
-	/** @return rotator representation of this matrix */
+	/** 
+	* Get the rotator representation of this matrix
+	* (Assumes Matrix represents a transform)
+	*@return rotator representation of this matrix
+	*/
 	UFUNCTION(BlueprintPure, meta = (DisplayName = "Get Rotator (Matrix)", ScriptMethod = "GetRotator"), Category = "Math|Matrix")
 	static FRotator Matrix_GetRotator(const FMatrix& M);
 
 	/**
 	 * Transform a rotation matrix into a quaternion.
+	 * (Assumes Matrix represents a transform)
 	 *
 	 * @warning rotation part will need to be unit length for this to be right!
 	 */
@@ -2424,32 +2470,51 @@ class ENGINE_API UKismetMathLibrary : public UBlueprintFunctionLibrary
 	static FQuat Matrix_ToQuat(const FMatrix& M);
 
 	// Frustum plane extraction.
-	/** @param OutPlane the near plane of the Frustum of this matrix */
+
+	/** Get the near plane of the Frustum of this matrix 
+	/* (Assumes Matrix represents a View Projection Matrix)
+	/* @param OutPlane the near plane of the Frustum of this matrix 
+	*/
 	UFUNCTION(BlueprintPure, meta = (DisplayName = "Get Frustum Near Plane (Matrix)", ScriptMethod = "GetFrustumNearPlane"), Category = "Math|Matrix")
 	static bool Matrix_GetFrustumNearPlane(const FMatrix& M, FPlane& OutPlane);
 
-	/** @param OutPlane the far plane of the Frustum of this matrix */
+	/** Get the far plane of the Frustum of this matrix
+	/* (Assumes Matrix represents a View Projection Matrix)
+	/* @param OutPlane the far plane of the Frustum of this matrix 
+	*/
 	UFUNCTION(BlueprintPure, meta = (DisplayName = "Get Frustum Far Plane (Matrix)", ScriptMethod = "GetFrustumFarPlane"), Category = "Math|Matrix")
 	static bool Matrix_GetFrustumFarPlane(const FMatrix& M, FPlane& OutPlane);
 
-	/** @param OutPlane the left plane of the Frustum of this matrix */
+	/** Get the left plane of the Frustum of this matrix
+	/* (Assumes Matrix represents a View Projection Matrix)
+	/* @param OutPlane the left plane of the Frustum of this matrix 
+	*/
 	UFUNCTION(BlueprintPure, meta = (DisplayName = "Get Frustum Left Plane (Matrix)", ScriptMethod = "GetFrustumLeftPlane"), Category = "Math|Matrix")
 	static bool Matrix_GetFrustumLeftPlane(const FMatrix& M, FPlane& OutPlane);
 
-	/** @param OutPlane the right plane of the Frustum of this matrix */
+	/** Get the right plane of the Frustum of this matrix
+	/* (Assumes Matrix represents a View Projection Matrix)
+	/* @param OutPlane the right plane of the Frustum of this matrix 
+	*/
 	UFUNCTION(BlueprintPure, meta = (DisplayName = "Get Frustum Right Plane (Matrix)", ScriptMethod = "GetFrustumRightPlane"), Category = "Math|Matrix")
 	static bool Matrix_GetFrustumRightPlane(const FMatrix& M, FPlane& OutPlane);
 
-	/** @param OutPlane the top plane of the Frustum of this matrix */
+	/** Get the top plane of the Frustum of this matrix
+	/* (Assumes Matrix represents a View Projection Matrix)
+	/* @param OutPlane the top plane of the Frustum of this matrix 
+	*/
 	UFUNCTION(BlueprintPure, meta = (DisplayName = "Get Frustum Top Plane (Matrix)", ScriptMethod = "GetFrustumTopPlane"), Category = "Math|Matrix")
 	static bool Matrix_GetFrustumTopPlane(const FMatrix& M, FPlane& OutPlane);
 
-	/** @param OutPlane the bottom plane of the Frustum of this matrix */
+	/** Get the bottom plane of the Frustum of this matrix
+	/* (Assumes Matrix represents a View Projection Matrix)
+	/* @param OutPlane the bottom plane of the Frustum of this matrix */
 	UFUNCTION(BlueprintPure, meta = (DisplayName = "Get Frustum Bottom Plane (Matrix)", ScriptMethod = "GetFrustumBottomPlane"), Category = "Math|Matrix")
 	static bool Matrix_GetFrustumBottomPlane(const FMatrix& M, FPlane& OutPlane);
 
 	/**
 	 * Utility for mirroring this transform across a certain plane, and flipping one of the axis as well.
+	 * (Assumes Matrix represents a transform)
 	 */
 	UFUNCTION(BlueprintPure, meta = (DisplayName = "Mirror (Matrix)", ScriptMethod = "Mirror"), Category = "Math|Matrix")
 	static FMatrix Matrix_Mirror(const FMatrix& M, TEnumAsByte<EAxis::Type> MirrorAxis, TEnumAsByte<EAxis::Type> FlipAxis);
