Fix mismatched commenting
#rnx

diff --git a/Engine/Source/Runtime/Engine/Classes/Kismet/KismetMathLibrary.h b/Engine/Source/Runtime/Engine/Classes/Kismet/KismetMathLibrary.h
--- a/Engine/Source/Runtime/Engine/Classes/Kismet/KismetMathLibrary.h
+++ b/Engine/Source/Runtime/Engine/Classes/Kismet/KismetMathLibrary.h
@@ -2479,43 +2479,44 @@ class ENGINE_API UKismetMathLibrary : public UBlueprintFunctionLibrary
 	// Frustum plane extraction.
 
 	/** Get the near plane of the Frustum of this matrix 
-	/* (Assumes Matrix represents a View Projection Matrix)
-	/* @param OutPlane the near plane of the Frustum of this matrix 
-	*/
+	 * (Assumes Matrix represents a View Projection Matrix)
+	 * @param OutPlane the near plane of the Frustum of this matrix 
+	 */
 	UFUNCTION(BlueprintPure, meta = (DisplayName = "Get Frustum Near Plane (Matrix)", ScriptMethod = "GetFrustumNearPlane"), Category = "Math|Matrix")
 	static bool Matrix_GetFrustumNearPlane(const FMatrix& M, FPlane& OutPlane);
 
 	/** Get the far plane of the Frustum of this matrix
-	/* (Assumes Matrix represents a View Projection Matrix)
-	/* @param OutPlane the far plane of the Frustum of this matrix 
-	*/
+	 * (Assumes Matrix represents a View Projection Matrix)
+	 * @param OutPlane the far plane of the Frustum of this matrix 
+	 */
 	UFUNCTION(BlueprintPure, meta = (DisplayName = "Get Frustum Far Plane (Matrix)", ScriptMethod = "GetFrustumFarPlane"), Category = "Math|Matrix")
 	static bool Matrix_GetFrustumFarPlane(const FMatrix& M, FPlane& OutPlane);
 
 	/** Get the left plane of the Frustum of this matrix
-	/* (Assumes Matrix represents a View Projection Matrix)
-	/* @param OutPlane the left plane of the Frustum of this matrix 
-	*/
+	 * (Assumes Matrix represents a View Projection Matrix)
+	 * @param OutPlane the left plane of the Frustum of this matrix 
+	 */
 	UFUNCTION(BlueprintPure, meta = (DisplayName = "Get Frustum Left Plane (Matrix)", ScriptMethod = "GetFrustumLeftPlane"), Category = "Math|Matrix")
 	static bool Matrix_GetFrustumLeftPlane(const FMatrix& M, FPlane& OutPlane);
 
 	/** Get the right plane of the Frustum of this matrix
-	/* (Assumes Matrix represents a View Projection Matrix)
-	/* @param OutPlane the right plane of the Frustum of this matrix 
-	*/
+	 * (Assumes Matrix represents a View Projection Matrix)
+	 * @param OutPlane the right plane of the Frustum of this matrix 
+	 */
 	UFUNCTION(BlueprintPure, meta = (DisplayName = "Get Frustum Right Plane (Matrix)", ScriptMethod = "GetFrustumRightPlane"), Category = "Math|Matrix")
 	static bool Matrix_GetFrustumRightPlane(const FMatrix& M, FPlane& OutPlane);
 
 	/** Get the top plane of the Frustum of this matrix
-	/* (Assumes Matrix represents a View Projection Matrix)
-	/* @param OutPlane the top plane of the Frustum of this matrix 
-	*/
+	 * (Assumes Matrix represents a View Projection Matrix)
+	 * @param OutPlane the top plane of the Frustum of this matrix 
+	 */
 	UFUNCTION(BlueprintPure, meta = (DisplayName = "Get Frustum Top Plane (Matrix)", ScriptMethod = "GetFrustumTopPlane"), Category = "Math|Matrix")
 	static bool Matrix_GetFrustumTopPlane(const FMatrix& M, FPlane& OutPlane);
 
 	/** Get the bottom plane of the Frustum of this matrix
-	/* (Assumes Matrix represents a View Projection Matrix)
-	/* @param OutPlane the bottom plane of the Frustum of this matrix */
+	 * (Assumes Matrix represents a View Projection Matrix)
+	 * @param OutPlane the bottom plane of the Frustum of this matrix 
+	 */
 	UFUNCTION(BlueprintPure, meta = (DisplayName = "Get Frustum Bottom Plane (Matrix)", ScriptMethod = "GetFrustumBottomPlane"), Category = "Math|Matrix")
 	static bool Matrix_GetFrustumBottomPlane(const FMatrix& M, FPlane& OutPlane);
 
