#UE4 SceneComponentNeedsLoadForTarget now respects IsEditorOnly on platforms that do not have editor only data.

diff --git a/Engine/Source/Runtime/Engine/Private/Components/SceneComponent.cpp b/Engine/Source/Runtime/Engine/Private/Components/SceneComponent.cpp
--- a/Engine/Source/Runtime/Engine/Private/Components/SceneComponent.cpp
+++ b/Engine/Source/Runtime/Engine/Private/Components/SceneComponent.cpp
@@ -444,7 +444,7 @@ static bool SceneComponentNeedsLoadForTarget(USceneComponent const* SceneCompone
 		}
 	}
 
-	return true;
+	return TargetPlatform->HasEditorOnlyData() || !SceneComponentObject->IsEditorOnly();
 }
 
 static bool CheckDescendantsAreAlsoCulledForTarget(USceneComponent const* SceneComponentObject, const ITargetPlatform* TargetPlatform)
