Modify "Ray Tracing Global Illumination" post process setting from a boolean to enum. Enumerable types are: Disabled (default), Brute Force, and Final Gather.

diff --git a/Engine/Source/Runtime/Engine/Classes/Engine/Scene.h b/Engine/Source/Runtime/Engine/Classes/Engine/Scene.h
--- a/Engine/Source/Runtime/Engine/Classes/Engine/Scene.h
+++ b/Engine/Source/Runtime/Engine/Classes/Engine/Scene.h
@@ -84,6 +84,15 @@ enum class ETranslucencyType : uint8
 	RayTracing	UMETA(DisplayName = "Ray Tracing"),
 };
 
+
+UENUM()
+enum class ERayTracingGlobalIlluminationType : uint8
+{
+	Disabled    UMETA(DisplayName = "Disabled"),
+	BruteForce  UMETA(DisplayName = "Brute Force"),
+	FinalGather UMETA(DisplayName = "Final Gather")
+};
+
 UENUM()
 enum class EReflectedAndRefractedRayTracedShadows : uint8
 {
@@ -1084,7 +1093,7 @@ struct FPostProcessSettings
 	uint32 bOverride_RayTracingTranslucencyRefraction : 1;
 
 	UPROPERTY(EditAnywhere, BlueprintReadWrite, Category = Overrides, meta = (PinHiddenByDefault, InlineEditConditionToggle))
-	uint32 bOverride_RayTracingGI : 1;
+	uint32 bOverride_RayTracingGIType : 1;
 
 	UPROPERTY(EditAnywhere, BlueprintReadWrite, Category = Overrides, meta = (PinHiddenByDefault, InlineEditConditionToggle))
 	uint32 bOverride_RayTracingGIMaxBounces : 1;
@@ -1566,9 +1575,9 @@ struct FPostProcessSettings
 	UPROPERTY(interp, BlueprintReadWrite, Category="Rendering Features|Global Illumination", meta=(ClampMin = "0", UIMax = "4.0", editcondition = "bOverride_IndirectLightingIntensity", DisplayName = "Indirect Lighting Intensity"))
 	float IndirectLightingIntensity;
 
-	/** Enables ray tracing global illumination. */
-	UPROPERTY(EditAnywhere, BlueprintReadWrite, Category = "Rendering Features|Ray Tracing Global Illumination", meta = (editcondition = "bOverride_RayTracingGI", DisplayName = "Enabled"))
-	uint32 RayTracingGI : 1;
+	/** Sets the ray tracing global illumination type. */
+	UPROPERTY(EditAnywhere, BlueprintReadWrite, Category = "Rendering Features|Ray Tracing Global Illumination", meta = (editcondition = "bOverride_RayTracingGIType", DisplayName = "Type"))
+	ERayTracingGlobalIlluminationType RayTracingGIType;
 
 	/** Sets the ray tracing global illumination maximum bounces. */
 	UPROPERTY(EditAnywhere, BlueprintReadWrite, Category = "Rendering Features|Ray Tracing Global Illumination", meta = (ClampMin = "0", ClampMax = "50", editcondition = "bOverride_RayTracingGIMaxBounces", DisplayName = "Max. Bounces"))
diff --git a/Engine/Source/Runtime/Engine/Private/Scene.cpp b/Engine/Source/Runtime/Engine/Private/Scene.cpp
--- a/Engine/Source/Runtime/Engine/Private/Scene.cpp
+++ b/Engine/Source/Runtime/Engine/Private/Scene.cpp
@@ -514,7 +514,7 @@ FPostProcessSettings::FPostProcessSettings()
 	IndirectLightingColor = FLinearColor(1.0f, 1.0f, 1.0f);
 	IndirectLightingIntensity = 1.0f;
 	ColorGradingIntensity = 1.0f;
-	RayTracingGI = 0;
+	RayTracingGIType = ERayTracingGlobalIlluminationType::Disabled;
 	RayTracingGIMaxBounces = 1;
 	RayTracingGISamplesPerPixel = 4;
 
@@ -751,7 +751,7 @@ FPostProcessSettings::FPostProcessSettings(const FPostProcessSettings& Settings)
 	, bOverride_RayTracingTranslucencySamplesPerPixel(Settings.bOverride_RayTracingTranslucencySamplesPerPixel)
 	, bOverride_RayTracingTranslucencyShadows(Settings.bOverride_RayTracingTranslucencyShadows)
 	, bOverride_RayTracingTranslucencyRefraction(Settings.bOverride_RayTracingTranslucencyRefraction)
-	, bOverride_RayTracingGI(Settings.bOverride_RayTracingGI)
+	, bOverride_RayTracingGIType(Settings.bOverride_RayTracingGIType)
 	, bOverride_RayTracingGIMaxBounces(Settings.bOverride_RayTracingGIMaxBounces)
 	, bOverride_RayTracingGISamplesPerPixel(Settings.bOverride_RayTracingGISamplesPerPixel)
 	, bOverride_PathTracingMaxBounces(Settings.bOverride_PathTracingMaxBounces)
@@ -876,7 +876,7 @@ FPostProcessSettings::FPostProcessSettings(const FPostProcessSettings& Settings)
 	, RayTracingAOSamplesPerPixel(Settings.RayTracingAOSamplesPerPixel)
 	, IndirectLightingColor(Settings.IndirectLightingColor)
 	, IndirectLightingIntensity(Settings.IndirectLightingIntensity)
-	, RayTracingGI(Settings.RayTracingGI)
+	, RayTracingGIType(Settings.RayTracingGIType)
 	, RayTracingGIMaxBounces(Settings.RayTracingGIMaxBounces)
 	, RayTracingGISamplesPerPixel(Settings.RayTracingGISamplesPerPixel)
 	, ColorGradingIntensity(Settings.ColorGradingIntensity)
diff --git a/Engine/Source/Runtime/Engine/Private/SceneView.cpp b/Engine/Source/Runtime/Engine/Private/SceneView.cpp
--- a/Engine/Source/Runtime/Engine/Private/SceneView.cpp
+++ b/Engine/Source/Runtime/Engine/Private/SceneView.cpp
@@ -1454,9 +1454,9 @@ void FSceneView::OverridePostProcessSettings(const FPostProcessSettings& Src, fl
 			Dest.RayTracingTranslucencyRefraction = Src.RayTracingTranslucencyRefraction;
 		}
 
-		if (Src.bOverride_RayTracingGI)
+		if (Src.bOverride_RayTracingGIType)
 		{
-			Dest.RayTracingGI = Src.RayTracingGI;
+			Dest.RayTracingGIType = Src.RayTracingGIType;
 		}
 
 		if (Src.bOverride_RayTracingGIMaxBounces)
diff --git a/Engine/Source/Runtime/Renderer/Private/RayTracing/RayTracingGlobalIllumination.cpp b/Engine/Source/Runtime/Renderer/Private/RayTracing/RayTracingGlobalIllumination.cpp
--- a/Engine/Source/Runtime/Renderer/Private/RayTracing/RayTracingGlobalIllumination.cpp
+++ b/Engine/Source/Runtime/Renderer/Private/RayTracing/RayTracingGlobalIllumination.cpp
@@ -278,10 +278,22 @@ bool ShouldRenderRayTracingGlobalIllumination(const FViewInfo& View)
 	{
 		return (CVarRayTracingGlobalIlluminationValue > 0);
 	}
-	else 
+	else
+	{
+		return View.FinalPostProcessSettings.RayTracingGIType > ERayTracingGlobalIlluminationType::Disabled;
+	}
+}
+
+bool IsFinalGatherEnabled(const FViewInfo& View)
+{
+
+	int32 bEnableFinalGather = CVarRayTracingGlobalIlluminationEnableFinalGather.GetValueOnRenderThread();
+	if (bEnableFinalGather >= 0)
 	{
-		return View.FinalPostProcessSettings.RayTracingGI > 0;
+		return bEnableFinalGather > 0;
 	}
+
+	return View.FinalPostProcessSettings.RayTracingGIType == ERayTracingGlobalIlluminationType::FinalGather;
 }
 
 class FGlobalIlluminationRGS : public FGlobalShader
@@ -513,7 +525,7 @@ bool FDeferredShadingSceneRenderer::RenderRayTracingGlobalIllumination(
 	}
 
 	// Ray generation pass
-	if (CVarRayTracingGlobalIlluminationEnableFinalGather.GetValueOnRenderThread() != 0)
+	if (IsFinalGatherEnabled(View))
 	{
 		RenderRayTracingGlobalIlluminationFinalGather(GraphBuilder, SceneTextures, View, *OutRayTracingConfig, UpscaleFactor, OutDenoiserInputs);
 	}
