Merging CL#10434557

#UE4 Forward declare fix

diff --git a/Engine/Source/Runtime/Engine/Public/Physics/Experimental/ChaosDerivedDataReader.h b/Engine/Source/Runtime/Engine/Public/Physics/Experimental/ChaosDerivedDataReader.h
--- a/Engine/Source/Runtime/Engine/Public/Physics/Experimental/ChaosDerivedDataReader.h
+++ b/Engine/Source/Runtime/Engine/Public/Physics/Experimental/ChaosDerivedDataReader.h
@@ -11,8 +11,7 @@ namespace Chaos
 	template <typename T>
 	class TTriangleMeshImplicitObject;
 
-	template <typename T, int d>
-	class TConvex;
+	class FConvex;
 }
 
 struct FUntypedBulkData;
