Improved r.HLOD.DistanceOverride to support space delimiters between values
#rnx

diff --git a/Engine/Source/Runtime/Engine/Private/LODActor.cpp b/Engine/Source/Runtime/Engine/Private/LODActor.cpp
--- a/Engine/Source/Runtime/Engine/Private/LODActor.cpp
+++ b/Engine/Source/Runtime/Engine/Private/LODActor.cpp
@@ -365,15 +365,18 @@ void ALODActor::ParseOverrideDistancesCVar()
 {
 	// Parse HLOD override distance cvar into array
 	const FString DistanceOverrideValues = CVarHLODDistanceOverride.GetValueOnAnyThread();
+	const TCHAR* Delimiters[] = { TEXT(","), TEXT(" ") };
+
 	TArray<FString> Distances;
-	DistanceOverrideValues.ParseIntoArray(/*out*/ Distances, TEXT(","), /*bCullEmpty=*/ false);
-	HLODDistances.Empty(Distances.Num());
+	DistanceOverrideValues.ParseIntoArray(/*out*/ Distances, Delimiters, UE_ARRAY_COUNT(Delimiters), /*bCullEmpty=*/ false);
+	Distances.RemoveAll([](const FString& String) { return String.IsEmpty(); });
 
+	HLODDistances.Empty(Distances.Num());
 	for (const FString& DistanceString : Distances)
 	{
 		const float DistanceForThisLevel = FCString::Atof(*DistanceString);
 		HLODDistances.Add(DistanceForThisLevel);
-	}	
+	}
 }
 
 float ALODActor::GetLODDrawDistanceWithOverride() const
