Avoid invoking the CancelDelegate via destroying the FStreamable which could, theoretically, manipulate the StreamableItems map while iterating the map during OnPreGarbageCollect

diff --git a/Engine/Source/Runtime/Engine/Private/StreamableManager.cpp b/Engine/Source/Runtime/Engine/Private/StreamableManager.cpp
--- a/Engine/Source/Runtime/Engine/Private/StreamableManager.cpp
+++ b/Engine/Source/Runtime/Engine/Private/StreamableManager.cpp
@@ -744,6 +744,7 @@ FStreamableManager::~FStreamableManager()
 void FStreamableManager::OnPreGarbageCollect()
 {
 	TSet<FSoftObjectPath> RedirectsToRemove;
+	TArray<FStreamable*> StreamablesToDelete;
 
 	// Remove any streamables with no active handles, as GC may have freed them
 	for (TStreamableMap::TIterator It(StreamableItems); It; ++It)
@@ -764,11 +765,16 @@ void FStreamableManager::OnPreGarbageCollect()
 		if (Existing->ActiveHandles.Num() == 0)
 		{
 			RedirectsToRemove.Add(It.Key());
-			delete Existing;
+			StreamablesToDelete.Add(Existing);
 			It.RemoveCurrent();
 		}
 	}
 
+	for (FStreamable* StreamableToDelete : StreamablesToDelete)
+	{
+		delete StreamableToDelete;
+	}
+
 	if (RedirectsToRemove.Num() > 0)
 	{
 		for (TStreamableRedirects::TIterator It(StreamableRedirects); It; ++It)
