Fixing Splitscreen audio.
Main issue stemmed from now using the closest listener inside sound nodes, where the active sound had previously been transformed into listener 0's space.

diff --git a/Engine/Source/Runtime/Engine/Private/ActiveSound.cpp b/Engine/Source/Runtime/Engine/Private/ActiveSound.cpp
--- a/Engine/Source/Runtime/Engine/Private/ActiveSound.cpp
+++ b/Engine/Source/Runtime/Engine/Private/ActiveSound.cpp
@@ -520,10 +520,6 @@ void FActiveSound::UpdateWaveInstances(TArray<FWaveInstance*> &InWaveInstances,
 		return;
 	}
 
-	// splitscreen support:
-	// we always pass the 'primary' listener (viewport 0) to the sound nodes and the underlying audio system
-	// then move the AudioComponent's CurrentLocation so that its position relative to that Listener is the same as its real position is relative to the closest Listener
-
 	{
 		SCOPE_CYCLE_COUNTER( STAT_AudioFindNearestLocation );
 		ClosestListenerIndex = AudioDevice->FindClosestListenerIndex(Transform);
@@ -601,18 +597,6 @@ void FActiveSound::UpdateWaveInstances(TArray<FWaveInstance*> &InWaveInstances,
 		}
 
 		ParseParams.ModulationPluginSettings = FindModulationSettings();
-
-		// if the closest listener is not the primary one, transform the sound transform so it's panned relative to primary listener position
-		if (ClosestListenerIndex != 0)
-		{
-			FTransform FirstListenerTransform;
-			FTransform ClosestListenerTransform;
-			AudioDevice->GetListenerTransform(0, FirstListenerTransform);
-			AudioDevice->GetListenerTransform(ClosestListenerIndex, ClosestListenerTransform);
-
-			ParseParams.Transform = ParseParams.Transform * ClosestListenerTransform.Inverse() * FirstListenerTransform;
-		}
-
 		Sound->Parse(AudioDevice, 0, *this, ParseParams, ThisSoundsWaveInstances);
 
 		// Track this active sound's min pitch value. This is used to scale it's possible duration value.
diff --git a/Engine/Source/Runtime/Engine/Private/Audio.cpp b/Engine/Source/Runtime/Engine/Private/Audio.cpp
--- a/Engine/Source/Runtime/Engine/Private/Audio.cpp
+++ b/Engine/Source/Runtime/Engine/Private/Audio.cpp
@@ -447,7 +447,11 @@ void FSoundSource::UpdateStereoEmitterPositions()
 	if (!DisableStereoSpreadCvar && WaveInstance->StereoSpread > 0.0f)
 	{
 		// We need to compute the stereo left/right channel positions using the audio component position and the spread
-		FVector ListenerPosition = AudioDevice->Listeners[0].Transform.GetLocation();
+		FVector ListenerPosition;
+
+		const bool bAllowAttenuationOverride = false;
+		const int32 ListenerIndex = WaveInstance->ActiveSound ? WaveInstance->ActiveSound->GetClosestListenerIndex() : 0;
+		AudioDevice->GetListenerPosition(ListenerIndex, ListenerPosition, bAllowAttenuationOverride);
 		FVector ListenerToSourceDir = (WaveInstance->Location - ListenerPosition).GetSafeNormal();
 
 		float HalfSpread = 0.5f * WaveInstance->StereoSpread;
@@ -577,9 +581,11 @@ FSpatializationParams FSoundSource::GetSpatializationParams()
 	}
 	Params.EmitterWorldPosition = WaveInstance->Location;
 
+	int32 ListenerIndex = 0;
 	if (WaveInstance->ActiveSound != nullptr)
 	{
 		Params.EmitterWorldRotation = WaveInstance->ActiveSound->Transform.GetRotation();
+		ListenerIndex = WaveInstance->ActiveSound->GetClosestListenerIndex();
 	}
 	else
 	{
@@ -587,7 +593,8 @@ FSpatializationParams FSoundSource::GetSpatializationParams()
 	}
 
 	// Pass the actual listener orientation and position
-	const FTransform& ListenerTransform = AudioDevice->GetListeners()[0].Transform;
+	FTransform ListenerTransform;
+	AudioDevice->GetListenerTransform(ListenerIndex, ListenerTransform);
 	Params.ListenerOrientation = ListenerTransform.GetRotation();
 	Params.ListenerPosition = ListenerTransform.GetLocation();
 
diff --git a/Engine/Source/Runtime/Engine/Private/AudioDevice.cpp b/Engine/Source/Runtime/Engine/Private/AudioDevice.cpp
--- a/Engine/Source/Runtime/Engine/Private/AudioDevice.cpp
+++ b/Engine/Source/Runtime/Engine/Private/AudioDevice.cpp
@@ -4922,13 +4922,30 @@ float FAudioDevice::GetDistanceToNearestListener(const FVector& Location) const
 
 float FAudioDevice::GetDistanceSquaredToListener(const FVector& Location, int32 ListenerIndex) const
 {
+	float OutSqDistance;
+	GetSquaredDistanceToListener(Location, ListenerIndex, OutSqDistance);
+	return OutSqDistance;
+}
+
+bool FAudioDevice::GetSquaredDistanceToListener(const FVector& Location, int32 ListenerIndex, float& OutSqDistance) const
+{
+	OutSqDistance = TNumericLimits<float>::Max();
+	const int32 ListenerCount = IsInAudioThread() ? Listeners.Num() : ListenerProxies.Num();
+
+	if (ListenerIndex >= ListenerCount)
+	{
+		return false;
+	}
+
 	FVector ListenerTranslation;
 	const bool bAllowOverride = true;
 	if (!GetListenerPosition(ListenerIndex, ListenerTranslation, bAllowOverride))
 	{
-		return WORLD_MAX;
+		return false;
 	}
-	return (ListenerTranslation - Location).SizeSquared();
+
+	OutSqDistance = (ListenerTranslation - Location).SizeSquared();
+	return true;
 }
 
 bool FAudioDevice::GetListenerPosition(int32 ListenerIndex, FVector& OutPosition, bool bAllowOverride) const
@@ -5107,30 +5124,30 @@ int32 FAudioDevice::FindClosestListenerIndex(const FTransform& SoundTransform) c
 
 int32 FAudioDevice::FindClosestListenerIndex(const FVector& Position, float& OutDistanceSq, bool bAllowAttenuationOverrides) const
 {
-	int32 ClosestListenerIndex = INDEX_NONE;
-	OutDistanceSq = WORLD_MAX;
-
+	int32 ClosestListenerIndex = 0;
+	OutDistanceSq = 0.f;
 	FVector ListenerPosition;
-	const int32 ListenerCount = IsInAudioThread() ? Listeners.Num() : ListenerProxies.Num();
 
-	if (GetListenerPosition(0, ListenerPosition, bAllowAttenuationOverrides))
+	if (!GetListenerPosition(0, ListenerPosition, bAllowAttenuationOverrides))
 	{
-		OutDistanceSq = FVector::DistSquared(Position, ListenerPosition);
-		ClosestListenerIndex = 0;
+		return INDEX_NONE;
+	}
 
-		for (int32 i = 1; i < ListenerCount; ++i)
+	OutDistanceSq = FVector::DistSquared(Position, ListenerPosition);
+
+	const int32 ListenerCount = IsInAudioThread() ? Listeners.Num() : ListenerProxies.Num();
+	for (int32 i = 1; i < ListenerCount; ++i)
+	{
+		if (!GetListenerPosition(i, ListenerPosition, bAllowAttenuationOverrides))
 		{
-			if (!GetListenerPosition(i, ListenerPosition, bAllowAttenuationOverrides))
-			{
-				continue;
-			}
+			continue;
+		}
 
-			const float DistSq = FVector::DistSquared(Position, ListenerPosition);
-			if (DistSq < OutDistanceSq)
-			{
-				OutDistanceSq = DistSq;
-				ClosestListenerIndex = i;
-			}
+		const float DistSq = FVector::DistSquared(Position, ListenerPosition);
+		if (DistSq < OutDistanceSq)
+		{
+			OutDistanceSq = DistSq;
+			ClosestListenerIndex = i;
 		}
 	}
 
diff --git a/Engine/Source/Runtime/Engine/Private/SoundNodeDistanceCrossFade.cpp b/Engine/Source/Runtime/Engine/Private/SoundNodeDistanceCrossFade.cpp
--- a/Engine/Source/Runtime/Engine/Private/SoundNodeDistanceCrossFade.cpp
+++ b/Engine/Source/Runtime/Engine/Private/SoundNodeDistanceCrossFade.cpp
@@ -184,7 +184,16 @@ void USoundNodeDistanceCrossFade::SetChildNodes(TArray<USoundNode*>& InChildNode
 
 float USoundNodeDistanceCrossFade::GetCurrentDistance(FAudioDevice* AudioDevice, FActiveSound& ActiveSound, const FSoundParseParameters& ParseParams) const
 {
-	return ActiveSound.bLocationDefined ? FVector::Dist( ParseParams.Transform.GetTranslation(), AudioDevice->GetListeners()[0].Transform.GetTranslation() ) : 0.f;
+	if (!ActiveSound.bLocationDefined)
+	{
+		return 0.f;
+	}
+
+	check(AudioDevice);
+	FVector ListenerPosition;
+	const bool bAllowOverride = true;
+	AudioDevice->GetListenerPosition(ActiveSound.GetClosestListenerIndex(), ListenerPosition, bAllowOverride);
+	return FVector::Dist(ParseParams.Transform.GetTranslation(), ListenerPosition);
 }
 
 bool USoundNodeDistanceCrossFade::AllowCrossfading(FActiveSound& ActiveSound) const
diff --git a/Engine/Source/Runtime/Engine/Private/SoundNodeDoppler.cpp b/Engine/Source/Runtime/Engine/Private/SoundNodeDoppler.cpp
--- a/Engine/Source/Runtime/Engine/Private/SoundNodeDoppler.cpp
+++ b/Engine/Source/Runtime/Engine/Private/SoundNodeDoppler.cpp
@@ -23,19 +23,30 @@ void USoundNodeDoppler::ParseNodes( FAudioDevice* AudioDevice, const UPTRINT Nod
 
 	FSoundParseParameters UpdatedParams = ParseParams;
 
+	check(AudioDevice);
+	const TArray<FListener>& Listeners = AudioDevice->GetListeners();
+
+	int32 ListenerIndex = ActiveSound.GetClosestListenerIndex();
+	if (ListenerIndex >= Listeners.Num())
+	{
+		ListenerIndex = 0;
+	}
+
+	const FListener& Listener = Listeners[ListenerIndex];
+
 	// Default the parse to using the setting for smoothing
 	if (*RequiresInitialization)
 	{
 		*RequiresInitialization = 0;
 
 		// First time, do no smoothing, but initialize the current pitch scale value to the first value returned from this function
-		CurrentPitchScale = GetDopplerPitchMultiplier(CurrentPitchScale, false, AudioDevice->GetListeners()[0], ParseParams.Transform.GetTranslation(), ParseParams.Velocity, AudioDevice->GetDeviceDeltaTime());
+		CurrentPitchScale = GetDopplerPitchMultiplier(CurrentPitchScale, false, Listener, ParseParams.Transform.GetTranslation(), ParseParams.Velocity, AudioDevice->GetDeviceDeltaTime());
 		UpdatedParams.Pitch *= CurrentPitchScale;
 	}
 	else
 	{
 		// Subsequent calls to this will do smoothing from the first initial value
-		UpdatedParams.Pitch *= GetDopplerPitchMultiplier(CurrentPitchScale, bUseSmoothing, AudioDevice->GetListeners()[0], ParseParams.Transform.GetTranslation(), ParseParams.Velocity, AudioDevice->GetDeviceDeltaTime());
+		UpdatedParams.Pitch *= GetDopplerPitchMultiplier(CurrentPitchScale, bUseSmoothing, Listener, ParseParams.Transform.GetTranslation(), ParseParams.Velocity, AudioDevice->GetDeviceDeltaTime());
 	}
 
 	Super::ParseNodes(AudioDevice, NodeWaveInstanceHash, ActiveSound, UpdatedParams, WaveInstances);
diff --git a/Engine/Source/Runtime/Engine/Public/ActiveSound.h b/Engine/Source/Runtime/Engine/Public/ActiveSound.h
--- a/Engine/Source/Runtime/Engine/Public/ActiveSound.h
+++ b/Engine/Source/Runtime/Engine/Public/ActiveSound.h
@@ -313,6 +313,8 @@ public:
 		AudioDevice = InAudioDevice;
 	}
 
+	int32 GetClosestListenerIndex() const { return ClosestListenerIndex; }
+
 	/** Returns whether or not the active sound can be deleted. */
 	bool CanDelete() const { return !bAsyncOcclusionPending; }
 
diff --git a/Engine/Source/Runtime/Engine/Public/AudioDevice.h b/Engine/Source/Runtime/Engine/Public/AudioDevice.h
--- a/Engine/Source/Runtime/Engine/Public/AudioDevice.h
+++ b/Engine/Source/Runtime/Engine/Public/AudioDevice.h
@@ -869,8 +869,14 @@ public:
 	/**
 	* Returns the distance from location to the appropriate listener representation, depending on calling thread
 	*/
+	UE_DEPRECATED(4.25, "Use GetSquaredDistanceToListener instead")
 	float GetDistanceSquaredToListener(const FVector& Location, int32 ListenerIndex) const;
 
+	/**
+	* Returns the distance from location to the appropriate listener representation, depending on calling thread
+	*/
+	bool GetSquaredDistanceToListener(const FVector& Location, int32 ListenerIndex, float& OutSqDistance) const;
+
 	/**
 	* Returns a position from the appropriate listener representation, depending on calling thread.
 	*
