- Fix max dimension (and by extension max distance) sound calculation for custom curves that do not terminate at 0 or at a time/distance of 1.0f
- Place bool to enable/disable distance attenuation before attenuation base properties

diff --git a/Engine/Source/Runtime/Engine/Classes/Engine/Attenuation.h b/Engine/Source/Runtime/Engine/Classes/Engine/Attenuation.h
--- a/Engine/Source/Runtime/Engine/Classes/Engine/Attenuation.h
+++ b/Engine/Source/Runtime/Engine/Classes/Engine/Attenuation.h
@@ -98,16 +98,7 @@ struct ENGINE_API FBaseAttenuationSettings
 	UPROPERTY(EditAnywhere, BlueprintReadWrite, Category= AttenuationDistance)
 	FRuntimeFloatCurve CustomAttenuationCurve;
 
-	FBaseAttenuationSettings()
-		: DistanceAlgorithm(EAttenuationDistanceModel::Linear)
-		, AttenuationShape(EAttenuationShape::Sphere)
-		, dBAttenuationAtMax(-60.f)
-		, FalloffMode(ENaturalSoundFalloffMode::Continues)
-		, AttenuationShapeExtents(400.f, 0.f, 0.f)
-		, ConeOffset(0.f)
-		, FalloffDistance(3600.f)
-	{
-	}
+	FBaseAttenuationSettings();
 
 	struct AttenuationShapeDetails
 	{
@@ -118,6 +109,9 @@ struct ENGINE_API FBaseAttenuationSettings
 
 	virtual void CollectAttenuationShapesForVisualization(TMultiMap<EAttenuationShape::Type, AttenuationShapeDetails>& ShapeDetailsMap) const;
 	float GetMaxDimension() const;
+
+	float GetMaxFalloffDistance() const;
+
 	float Evaluate(const FTransform& Origin, FVector Location, float DistanceScale = 1.f) const;
 
 	float AttenuationEval(float Distance, float Falloff, float DistanceScale = 1.f) const;
diff --git a/Engine/Source/Runtime/Engine/Private/Attenuation.cpp b/Engine/Source/Runtime/Engine/Private/Attenuation.cpp
--- a/Engine/Source/Runtime/Engine/Private/Attenuation.cpp
+++ b/Engine/Source/Runtime/Engine/Private/Attenuation.cpp
@@ -6,31 +6,27 @@
 #include "EngineDefines.h"
 
 
-float FBaseAttenuationSettings::GetMaxDimension() const
+namespace
 {
-	static const float WorldMax = static_cast<float>(WORLD_MAX);
-	if (FalloffDistance > WorldMax)
-	{
-		return WorldMax;
-	}
-
-	float MaxDimension = FalloffDistance;
-	if (DistanceAlgorithm == EAttenuationDistanceModel::NaturalSound)
-	{
-		switch (FalloffMode)
-		{
-		case ENaturalSoundFalloffMode::Hold:
-			return WorldMax;
-
-		case ENaturalSoundFalloffMode::Continues:
-			MaxDimension = (FalloffDistance / dBAttenuationAtMax) * -60.f;
-			break;
+	static const float MinAttenuationValue   = 1.e-3f;
+	static const float MinAttenuationValueDb = -60.0f;
+} // namespace <>
+
+
+FBaseAttenuationSettings::FBaseAttenuationSettings()
+	: DistanceAlgorithm(EAttenuationDistanceModel::Linear)
+	, AttenuationShape(EAttenuationShape::Sphere)
+	, dBAttenuationAtMax(MinAttenuationValueDb)
+	, FalloffMode(ENaturalSoundFalloffMode::Continues)
+	, AttenuationShapeExtents(400.f, 0.f, 0.f)
+	, ConeOffset(0.f)
+	, FalloffDistance(3600.f)
+{
+}
 
-		case ENaturalSoundFalloffMode::Silent:
-		default:
-			break;
-		}
-	}
+float FBaseAttenuationSettings::GetMaxDimension() const
+{
+	float MaxDimension = GetMaxFalloffDistance();
 
 	switch (AttenuationShape)
 	{
@@ -54,7 +50,84 @@ float FBaseAttenuationSettings::GetMaxDimension() const
 		check(false);
 	}
 
-	return FMath::Clamp(MaxDimension, 0.0f, WorldMax);
+	return FMath::Clamp(MaxDimension, 0.0f, static_cast<float>(WORLD_MAX));
+}
+
+float FBaseAttenuationSettings::GetMaxFalloffDistance() const
+{
+	static const float WorldMax = static_cast<float>(WORLD_MAX);
+	if (FalloffDistance > WorldMax)
+	{
+		return WorldMax;
+	}
+
+	float MaxFalloffDistance = FalloffDistance;
+	switch (DistanceAlgorithm)
+	{
+		case EAttenuationDistanceModel::Custom:
+		{
+			const FRichCurve* Curve = CustomAttenuationCurve.GetRichCurveConst();
+			check(Curve);
+
+			float LastTime = 0.0f;
+			const FRichCurveKey* LastKey = nullptr;
+			for (const FRichCurveKey& Key : Curve->Keys)
+			{
+				if (Key.Time > LastTime)
+				{
+					LastTime = Key.Time;
+					LastKey = &Key;
+				}
+			}
+
+			const float MaxValue = LastKey ? FMath::Max(LastKey->Value, 0.0f) : 0.0f;
+
+			// If last key's distance is near zero, scale the falloff distance accordingly
+			if (FMath::IsNearlyZero(MaxValue, MinAttenuationValue))
+			{
+				MaxFalloffDistance *= LastKey->Time;
+			}
+			// Otherwise, curve never terminates to non-zero value, so return WorldMax
+			else
+			{
+				MaxFalloffDistance = WorldMax;
+			}
+		}
+		break;
+
+		case EAttenuationDistanceModel::NaturalSound:
+		{
+			switch (FalloffMode)
+			{
+				case ENaturalSoundFalloffMode::Hold:
+				{
+					MaxFalloffDistance = WorldMax;
+				}
+				break;
+
+				case ENaturalSoundFalloffMode::Continues:
+				{
+					MaxFalloffDistance = FalloffDistance * MinAttenuationValueDb / FMath::Max(dBAttenuationAtMax, KINDA_SMALL_NUMBER);
+				}
+				break;
+
+				case ENaturalSoundFalloffMode::Silent:
+				default:
+				break;
+			}
+		}
+		break;
+
+		// All of these cases scale over the provided FalloffDistance, so just return that as max
+		case EAttenuationDistanceModel::Inverse:
+		case EAttenuationDistanceModel::Linear:
+		case EAttenuationDistanceModel::Logarithmic:
+		case EAttenuationDistanceModel::LogReverse:
+		default:
+		break;
+	}
+
+	return MaxFalloffDistance;
 }
 
 float FBaseAttenuationSettings::Evaluate(const FTransform& Origin, const FVector Location, const float DistanceScale) const
