Trivial: Fix bug with NaturalSound 'continuous' mode maxing incorrectly to avoid div/0

diff --git a/Engine/Source/Runtime/Engine/Private/Attenuation.cpp b/Engine/Source/Runtime/Engine/Private/Attenuation.cpp
--- a/Engine/Source/Runtime/Engine/Private/Attenuation.cpp
+++ b/Engine/Source/Runtime/Engine/Private/Attenuation.cpp
@@ -107,7 +107,7 @@ float FBaseAttenuationSettings::GetMaxFalloffDistance() const
 
 				case ENaturalSoundFalloffMode::Continues:
 				{
-					MaxFalloffDistance = FalloffDistance * MinAttenuationValueDb / FMath::Max(dBAttenuationAtMax, KINDA_SMALL_NUMBER);
+					MaxFalloffDistance = FalloffDistance * MinAttenuationValueDb / FMath::Min(dBAttenuationAtMax, -KINDA_SMALL_NUMBER);
 				}
 				break;
 
