Respecting backwards compatibility of boolean RayTracingGI to enum RayTracingGIType.

diff --git a/Engine/Source/Runtime/Engine/Classes/Engine/Scene.h b/Engine/Source/Runtime/Engine/Classes/Engine/Scene.h
--- a/Engine/Source/Runtime/Engine/Classes/Engine/Scene.h
+++ b/Engine/Source/Runtime/Engine/Classes/Engine/Scene.h
@@ -1093,7 +1093,7 @@ struct FPostProcessSettings
 	uint32 bOverride_RayTracingTranslucencyRefraction : 1;
 
 	UPROPERTY(EditAnywhere, BlueprintReadWrite, Category = Overrides, meta = (PinHiddenByDefault, InlineEditConditionToggle))
-	uint32 bOverride_RayTracingGIType : 1;
+	uint32 bOverride_RayTracingGI : 1;
 
 	UPROPERTY(EditAnywhere, BlueprintReadWrite, Category = Overrides, meta = (PinHiddenByDefault, InlineEditConditionToggle))
 	uint32 bOverride_RayTracingGIMaxBounces : 1;
@@ -1575,8 +1575,13 @@ struct FPostProcessSettings
 	UPROPERTY(interp, BlueprintReadWrite, Category="Rendering Features|Global Illumination", meta=(ClampMin = "0", UIMax = "4.0", editcondition = "bOverride_IndirectLightingIntensity", DisplayName = "Indirect Lighting Intensity"))
 	float IndirectLightingIntensity;
 
+#if WITH_EDITORONLY_DATA
+	UPROPERTY()
+	uint32 RayTracingGI_DEPRECATED : 1;
+#endif
+
 	/** Sets the ray tracing global illumination type. */
-	UPROPERTY(EditAnywhere, BlueprintReadWrite, Category = "Rendering Features|Ray Tracing Global Illumination", meta = (editcondition = "bOverride_RayTracingGIType", DisplayName = "Type"))
+	UPROPERTY(EditAnywhere, BlueprintReadWrite, Category = "Rendering Features|Ray Tracing Global Illumination", meta = (editcondition = "bOverride_RayTracingGI", DisplayName = "Type"))
 	ERayTracingGlobalIlluminationType RayTracingGIType;
 
 	/** Sets the ray tracing global illumination maximum bounces. */
@@ -1853,6 +1858,10 @@ struct FPostProcessSettings
 			BloomConvolutionPreFilterMax = BloomConvolutionPreFilter_DEPRECATED.Y;
 			BloomConvolutionPreFilterMult = BloomConvolutionPreFilter_DEPRECATED.Z;
 		}
+		if (RayTracingGI_DEPRECATED)
+		{
+			RayTracingGIType = (ERayTracingGlobalIlluminationType)(RayTracingGI_DEPRECATED == 1);
+		}
 	}
 #endif
 
diff --git a/Engine/Source/Runtime/Engine/Private/Scene.cpp b/Engine/Source/Runtime/Engine/Private/Scene.cpp
--- a/Engine/Source/Runtime/Engine/Private/Scene.cpp
+++ b/Engine/Source/Runtime/Engine/Private/Scene.cpp
@@ -751,7 +751,7 @@ FPostProcessSettings::FPostProcessSettings(const FPostProcessSettings& Settings)
 	, bOverride_RayTracingTranslucencySamplesPerPixel(Settings.bOverride_RayTracingTranslucencySamplesPerPixel)
 	, bOverride_RayTracingTranslucencyShadows(Settings.bOverride_RayTracingTranslucencyShadows)
 	, bOverride_RayTracingTranslucencyRefraction(Settings.bOverride_RayTracingTranslucencyRefraction)
-	, bOverride_RayTracingGIType(Settings.bOverride_RayTracingGIType)
+	, bOverride_RayTracingGI(Settings.bOverride_RayTracingGI)
 	, bOverride_RayTracingGIMaxBounces(Settings.bOverride_RayTracingGIMaxBounces)
 	, bOverride_RayTracingGISamplesPerPixel(Settings.bOverride_RayTracingGISamplesPerPixel)
 	, bOverride_PathTracingMaxBounces(Settings.bOverride_PathTracingMaxBounces)
diff --git a/Engine/Source/Runtime/Engine/Private/SceneView.cpp b/Engine/Source/Runtime/Engine/Private/SceneView.cpp
--- a/Engine/Source/Runtime/Engine/Private/SceneView.cpp
+++ b/Engine/Source/Runtime/Engine/Private/SceneView.cpp
@@ -1454,7 +1454,7 @@ void FSceneView::OverridePostProcessSettings(const FPostProcessSettings& Src, fl
 			Dest.RayTracingTranslucencyRefraction = Src.RayTracingTranslucencyRefraction;
 		}
 
-		if (Src.bOverride_RayTracingGIType)
+		if (Src.bOverride_RayTracingGI)
 		{
 			Dest.RayTracingGIType = Src.RayTracingGIType;
 		}
