Fix random disconnection errors by always including PacketInfo even if multiple packets are sent in a single frame.

#rnx

diff --git a/Engine/Source/Runtime/Engine/Private/NetConnection.cpp b/Engine/Source/Runtime/Engine/Private/NetConnection.cpp
--- a/Engine/Source/Runtime/Engine/Private/NetConnection.cpp
+++ b/Engine/Source/Runtime/Engine/Private/NetConnection.cpp
@@ -1732,7 +1732,8 @@ void UNetConnection::WritePacketHeader(FBitWriter& Writer)
 void UNetConnection::WriteDummyPacketInfo(FBitWriter& Writer)
 {
 	// The first packet of a frame will contain the packet info payload
-	const bool bHasPacketInfoPayload = bSentPacketInfoThisFrame == false;
+	//const bool bHasPacketInfoPayload = bSentPacketInfoThisFrame == false;
+	const bool bHasPacketInfoPayload = true;
 	Writer.WriteBit(bHasPacketInfoPayload);
 
 	if (bHasPacketInfoPayload)
@@ -1756,13 +1757,13 @@ void UNetConnection::WriteDummyPacketInfo(FBitWriter& Writer)
 
 void UNetConnection::WriteFinalPacketInfo(FBitWriter& Writer, const double PacketSentTimeInS)
 {
-	const bool bHasPacketInfoPayload = bSentPacketInfoThisFrame == false;
-
-	if (bHasPacketInfoPayload == false)
-	{
-		// PacketInfo payload was already sent, nothing to do here
-		return;
-	}
+	//const bool bHasPacketInfoPayload = bSentPacketInfoThisFrame == false;
+	//
+	//if (bHasPacketInfoPayload == false)
+	//{
+	//	// PacketInfo payload was already sent, nothing to do here
+	//	return;
+	//}
 
 	FBitWriterMark CurrentMark(Writer);
 
