#ue4 - Fill in ServerLastTransformUpdateTimeStamp for simulated proxies with the value of ReplicatedServerLastTransformUpdateTimeStamp, for code paths that may try to read that value instead. Added client and server timestamps to "p.NetVisualizeSimulatedCorrections" output.

diff --git a/Engine/Source/Runtime/Engine/Private/Components/CharacterMovementComponent.cpp b/Engine/Source/Runtime/Engine/Private/Components/CharacterMovementComponent.cpp
--- a/Engine/Source/Runtime/Engine/Private/Components/CharacterMovementComponent.cpp
+++ b/Engine/Source/Runtime/Engine/Private/Components/CharacterMovementComponent.cpp
@@ -7174,7 +7174,12 @@ void UCharacterMovementComponent::SmoothCorrection(const FVector& OldLocation, c
 
 		// Using server timestamp lets us know how much time actually elapsed, regardless of packet lag variance.
 		double OldServerTimeStamp = ClientData->SmoothingServerTimeStamp;
-		ClientData->SmoothingServerTimeStamp = (bIsSimulatedProxy ? CharacterOwner->GetReplicatedServerLastTransformUpdateTimeStamp() : ServerLastTransformUpdateTimeStamp);
+		if (bIsSimulatedProxy)
+		{
+			// This value is normally only updated on the server, however some code paths might try to read it instead of the replicated value so copy it for proxies as well.
+			ServerLastTransformUpdateTimeStamp = CharacterOwner->GetReplicatedServerLastTransformUpdateTimeStamp();
+		}
+		ClientData->SmoothingServerTimeStamp = ServerLastTransformUpdateTimeStamp;
 
 		// Initial update has no delta.
 		if (ClientData->LastCorrectionTime == 0)
@@ -7355,9 +7360,11 @@ void UCharacterMovementComponent::SmoothClientPosition_Interpolate(float DeltaSe
 			if ( CharacterMovementCVars::NetVisualizeSimulatedCorrections >= 1 )
 			{
 				const FColor DebugColor = FColor::White;
-				const FVector DebugLocation = CharacterOwner->GetMesh()->GetComponentLocation() + FVector( 0.f, 0.f, 300.0f ) - CharacterOwner->GetBaseTranslationOffset();
+				const FVector DebugLocation = CharacterOwner->GetMesh()->GetComponentLocation() + FVector( 0.f, 0.f, 130.0f ) - CharacterOwner->GetBaseTranslationOffset();
 				FString DebugText = FString::Printf( TEXT( "Lerp: %2.2f" ), LerpPercent );
 				DrawDebugString( GetWorld(), DebugLocation, DebugText, nullptr, DebugColor, 0.f, true );
+				FString TimeText = FString::Printf( TEXT("ClientTime: %2.2f ServerTime: %2.2f" ), ClientData->SmoothingClientTimeStamp, ClientData->SmoothingServerTimeStamp);
+				DrawDebugString( GetWorld(), DebugLocation + 25.f, TimeText, nullptr, DebugColor, 0.f, true);
 			}
 #endif
 		}
