Flip Y axis for canvas only when render target is backbuffer. Probably broke some case in MobileHDR=OFF

diff --git a/Engine/Source/Runtime/Engine/Private/Slate/DebugCanvas.cpp b/Engine/Source/Runtime/Engine/Private/Slate/DebugCanvas.cpp
--- a/Engine/Source/Runtime/Engine/Private/Slate/DebugCanvas.cpp
+++ b/Engine/Source/Runtime/Engine/Private/Slate/DebugCanvas.cpp
@@ -256,8 +256,8 @@ void FDebugCanvasDrawer::DrawRenderThread(FRHICommandListImmediate& RHICmdList,
 		RenderTarget->SetRenderTargetTexture(*RT);
 
 		bool bNeedToFlipVertical = RenderThreadCanvas->GetAllowSwitchVerticalAxis();
-		// Do not flip when rendering to the back buffer
-		RenderThreadCanvas->SetAllowSwitchVerticalAxis(false);
+		// Flip when rendering to the back buffer
+		RenderThreadCanvas->SetAllowSwitchVerticalAxis(true);
 		if (RenderThreadCanvas->IsScaledToRenderTarget() && IsValidRef(*RT)) 
 		{
 			RenderThreadCanvas->SetRenderTargetRect( FIntRect(0, 0, (*RT)->GetSizeX(), (*RT)->GetSizeY()) );
diff --git a/Engine/Source/Runtime/Engine/Private/TileRendering.cpp b/Engine/Source/Runtime/Engine/Private/TileRendering.cpp
--- a/Engine/Source/Runtime/Engine/Private/TileRendering.cpp
+++ b/Engine/Source/Runtime/Engine/Private/TileRendering.cpp
@@ -227,7 +227,7 @@ bool FCanvasTileRendererItem::Render_RenderThread(FRHICommandListImmediate& RHIC
 
 	TUniquePtr<const FSceneView> View = MakeUnique<const FSceneView>(ViewInitOptions);
 
-	const bool bNeedsToSwitchVerticalAxis = RHINeedsToSwitchVerticalAxis(Canvas->GetShaderPlatform()) && IsMobileHDR();
+	const bool bNeedsToSwitchVerticalAxis = RHINeedsToSwitchVerticalAxis(Canvas->GetShaderPlatform()) && Canvas->GetAllowSwitchVerticalAxis(); 
 
 	Data->RenderTiles(RHICmdList, DrawRenderState, *View, Canvas->IsHitTesting(), bNeedsToSwitchVerticalAxis);
 
@@ -277,7 +277,7 @@ bool FCanvasTileRendererItem::Render_GameThread(const FCanvas* Canvas, FRenderTh
 
 	const FSceneView* View = new FSceneView(ViewInitOptions);
 
-	const bool bNeedsToSwitchVerticalAxis = RHINeedsToSwitchVerticalAxis(Canvas->GetShaderPlatform()) && IsMobileHDR();
+	const bool bNeedsToSwitchVerticalAxis = RHINeedsToSwitchVerticalAxis(Canvas->GetShaderPlatform()) && Canvas->GetAllowSwitchVerticalAxis();
 	const bool bIsHitTesting = Canvas->IsHitTesting();
 	const bool bDeleteOnRender = Canvas->GetAllowedModes() & FCanvas::Allow_DeleteOnRender;
 
diff --git a/Engine/Source/Runtime/Engine/Private/TriangleRendering.cpp b/Engine/Source/Runtime/Engine/Private/TriangleRendering.cpp
--- a/Engine/Source/Runtime/Engine/Private/TriangleRendering.cpp
+++ b/Engine/Source/Runtime/Engine/Private/TriangleRendering.cpp
@@ -197,7 +197,7 @@ bool FCanvasTriangleRendererItem::Render_RenderThread(FRHICommandListImmediate&
 
 	TUniquePtr<const FSceneView> View = MakeUnique<const FSceneView>(ViewInitOptions);
 
-	const bool bNeedsToSwitchVerticalAxis = RHINeedsToSwitchVerticalAxis(Canvas->GetShaderPlatform()) && XOR(IsMobileHDR(), Canvas->GetAllowSwitchVerticalAxis());
+	const bool bNeedsToSwitchVerticalAxis = RHINeedsToSwitchVerticalAxis(Canvas->GetShaderPlatform()) && Canvas->GetAllowSwitchVerticalAxis();
 
 	Data->RenderTriangles(RHICmdList, DrawRenderState, *View, Canvas->IsHitTesting(), bNeedsToSwitchVerticalAxis);
 
@@ -247,7 +247,7 @@ bool FCanvasTriangleRendererItem::Render_GameThread(const FCanvas* Canvas, FRend
 
 	const FSceneView* View = new FSceneView(ViewInitOptions);
 
-	const bool bNeedsToSwitchVerticalAxis = RHINeedsToSwitchVerticalAxis(Canvas->GetShaderPlatform()) && XOR(IsMobileHDR(),Canvas->GetAllowSwitchVerticalAxis());
+	const bool bNeedsToSwitchVerticalAxis = RHINeedsToSwitchVerticalAxis(Canvas->GetShaderPlatform()) && Canvas->GetAllowSwitchVerticalAxis();
 	const bool bIsHitTesting = Canvas->IsHitTesting();
 	const bool bDeleteOnRender = Canvas->GetAllowedModes() & FCanvas::Allow_DeleteOnRender;
 
diff --git a/Engine/Source/Runtime/Engine/Private/UnrealClient.cpp b/Engine/Source/Runtime/Engine/Private/UnrealClient.cpp
--- a/Engine/Source/Runtime/Engine/Private/UnrealClient.cpp
+++ b/Engine/Source/Runtime/Engine/Private/UnrealClient.cpp
@@ -1564,7 +1564,7 @@ void FViewport::Draw( bool bShouldPresent /*= true */)
 				Canvas.SetRenderTargetRect(FIntRect(0, 0, SizeX, SizeY));
 				{
 					// Make sure the Canvas is not rendered upside down
-					Canvas.SetAllowSwitchVerticalAxis(false);
+					Canvas.SetAllowSwitchVerticalAxis(true);
 					ViewportClient->Draw(this, &Canvas);
 				}
 				Canvas.Flush_GameThread();
diff --git a/Engine/Source/Runtime/Engine/Private/UserInterface/Canvas.cpp b/Engine/Source/Runtime/Engine/Private/UserInterface/Canvas.cpp
--- a/Engine/Source/Runtime/Engine/Private/UserInterface/Canvas.cpp
+++ b/Engine/Source/Runtime/Engine/Private/UserInterface/Canvas.cpp
@@ -256,6 +256,7 @@ FCanvas::FCanvas(FRenderTarget* InRenderTarget,FHitProxyConsumer* InHitProxyCons
 ,	CurrentRealTime(InRealTime)
 ,	CurrentWorldTime(InWorldTime)
 ,	CurrentDeltaWorldTime(InWorldDeltaTime)
+,	bAllowsToSwitchVerticalAxis(false)
 ,	FeatureLevel(InFeatureLevel)
 ,	bUseInternalTexture(false)
 ,	DrawMode(CDM_DeferDrawing)
@@ -268,7 +269,7 @@ void FCanvas::Construct()
 {
 	bStereoRendering = false;
 	bScaledToRenderTarget = false;
-	bAllowsToSwitchVerticalAxis = true;
+	bAllowsToSwitchVerticalAxis = false;
 
 	const FIntPoint RenderTargetSizeXY = RenderTarget ? RenderTarget->GetSizeXY() : FIntPoint(1, 1);
 
@@ -392,7 +393,7 @@ bool FCanvasBatchedElementRenderItem::Render_RenderThread(FRHICommandListImmedia
 			Gamma = 1.0f;
 		}
 
-		bool bNeedsToSwitchVerticalAxis = RHINeedsToSwitchVerticalAxis(Canvas->GetShaderPlatform()) && !Canvas->GetAllowSwitchVerticalAxis();
+		bool bNeedsToSwitchVerticalAxis = RHINeedsToSwitchVerticalAxis(Canvas->GetShaderPlatform()) && Canvas->GetAllowSwitchVerticalAxis();
 
 		// draw batched items
 		Data->BatchedElements.Draw(
@@ -430,7 +431,7 @@ bool FCanvasBatchedElementRenderItem::Render_GameThread(const FCanvas* Canvas, F
 			Gamma = 1.0f;
 		}
 
-		bool bNeedsToSwitchVerticalAxis = RHINeedsToSwitchVerticalAxis(Canvas->GetShaderPlatform()) && !Canvas->GetAllowSwitchVerticalAxis();
+		bool bNeedsToSwitchVerticalAxis = RHINeedsToSwitchVerticalAxis(Canvas->GetShaderPlatform()) && Canvas->GetAllowSwitchVerticalAxis();
 
 		// Render the batched elements.
 		struct FBatchedDrawParameters
diff --git a/Engine/Source/Runtime/Renderer/Private/SceneRendering.cpp b/Engine/Source/Runtime/Renderer/Private/SceneRendering.cpp
--- a/Engine/Source/Runtime/Renderer/Private/SceneRendering.cpp
+++ b/Engine/Source/Runtime/Renderer/Private/SceneRendering.cpp
@@ -2735,7 +2735,7 @@ void FSceneRenderer::RenderFinish(FRHICommandListImmediate& RHICmdList)
 					int32 Y = 130;
 					FCanvas Canvas(&TempRenderTarget, NULL, View.Family->CurrentRealTime, View.Family->CurrentWorldTime, View.Family->DeltaWorldTime, FeatureLevel);
 					// Make sure draws to the canvas are not rendered upside down.
-					Canvas.SetAllowSwitchVerticalAxis(false);
+					Canvas.SetAllowSwitchVerticalAxis(true);
 					if (bViewParentOrFrozen)
 					{
 						const FText StateText =
diff --git a/Engine/Source/Runtime/Renderer/Private/SceneSoftwareOcclusion.cpp b/Engine/Source/Runtime/Renderer/Private/SceneSoftwareOcclusion.cpp
--- a/Engine/Source/Runtime/Renderer/Private/SceneSoftwareOcclusion.cpp
+++ b/Engine/Source/Runtime/Renderer/Private/SceneSoftwareOcclusion.cpp
@@ -1141,7 +1141,7 @@ void FSceneSoftwareOcclusion::DebugDraw(FRHICommandListImmediate& RHICmdList, co
 	}
 
 	FCanvas Canvas(&TempRenderTarget, NULL, View.Family->CurrentRealTime, View.Family->CurrentWorldTime, View.Family->DeltaWorldTime, View.GetFeatureLevel());
-	Canvas.SetAllowSwitchVerticalAxis(false);
+	Canvas.SetAllowSwitchVerticalAxis(true);
 	FBatchedElements* BatchedElements = Canvas.GetBatchedElements(FCanvas::ET_Line);
 						
 	for (int32 i = 0; i < BIN_NUM; ++i)
