replay recording fidelity: avoiding NaN if number of prioritized actors is 0

diff --git a/Engine/Source/Runtime/Engine/Private/DemoNetDriver.cpp b/Engine/Source/Runtime/Engine/Private/DemoNetDriver.cpp
--- a/Engine/Source/Runtime/Engine/Private/DemoNetDriver.cpp
+++ b/Engine/Source/Runtime/Engine/Private/DemoNetDriver.cpp
@@ -2843,15 +2843,15 @@ void UDemoNetDriver::TickDemoRecordFrame(float DeltaSeconds)
 		ReplicatePrioritizedActors(PrioritizedActors.GetData(), PrioritizedActors.Num(), Params);
 	}
 
-	LastReplayFrameFidelity = (float)Params.NumActorsReplicated / (float)NumPrioritizedActors;
-
 	CSV_CUSTOM_STAT(Basic, DemoNumReplicatedActors, Params.NumActorsReplicated, ECsvCustomStatOp::Set);
 
 	FlushNetChecked(*ClientConnection);
 
 	WriteDemoFrameFromQueuedDemoPackets(*FileAr, ClientConnection->QueuedDemoPackets, DemoCurrentTime, EWriteDemoFrameFlags::None);
-	
-	AdjustConsiderTime((float)Params.NumActorsReplicated / (float)NumPrioritizedActors);
+
+	float ReplicatedPercent = NumPrioritizedActors != 0 ? (float)Params.NumActorsReplicated / (float)NumPrioritizedActors : 1.0f;
+	AdjustConsiderTime(ReplicatedPercent);
+	LastReplayFrameFidelity = ReplicatedPercent;
 }
 
 bool UDemoNetDriver::ReplicatePrioritizedActor(const FActorPriority& ActorPriority, const class FRepActorsParams& Params)
