Fix nullptr static analysis warning

diff --git a/Engine/Source/Runtime/Engine/Private/Attenuation.cpp b/Engine/Source/Runtime/Engine/Private/Attenuation.cpp
--- a/Engine/Source/Runtime/Engine/Private/Attenuation.cpp
+++ b/Engine/Source/Runtime/Engine/Private/Attenuation.cpp
@@ -85,7 +85,7 @@ float FBaseAttenuationSettings::GetMaxFalloffDistance() const
 			// If last key's distance is near zero, scale the falloff distance accordingly
 			if (FMath::IsNearlyZero(MaxValue, MinAttenuationValue))
 			{
-				MaxFalloffDistance *= LastKey->Time;
+				MaxFalloffDistance *= LastTime;
 			}
 			// Otherwise, curve never terminates to non-zero value, so return WorldMax
 			else
