Adding GameplayStatic BP function to get the Location of the closest Listener

diff --git a/Engine/Source/Runtime/Engine/Classes/Kismet/GameplayStatics.h b/Engine/Source/Runtime/Engine/Classes/Kismet/GameplayStatics.h
--- a/Engine/Source/Runtime/Engine/Classes/Kismet/GameplayStatics.h
+++ b/Engine/Source/Runtime/Engine/Classes/Kismet/GameplayStatics.h
@@ -436,9 +436,21 @@ public:
 	 * @note This will always return false if there is no audio device, or the audio device is disabled.
 	 */
 	UFUNCTION(BlueprintCallable, Category = "Audio", meta = (WorldContext = "WorldContextObject"))
-	static bool AreAnyListenersWithinRange(const UObject* WorldContextObject, FVector Location, float MaximumRange);
+	static bool AreAnyListenersWithinRange(const UObject* WorldContextObject, const FVector& Location, float MaximumRange);
 	
 
+	/**
+	 * Finds and returns the position of the closest listener to the specified location
+	 * @param Location						The location from which we'd like to find the closest listener, in world space.
+	 * @param MaximumRange					The maximum distance away from Location that a listener can be.
+	 * @param bAllowAttenuationOverride		True for the adjusted listener position (if attenuation override is set), false for the raw listener position (for panning)
+	 * @param ListenerPosition				[Out] The position of the closest listener in world space, if found.
+	 * @return true if we've successfully found a listener within MaximumRange of Location, otherwise false.
+ 	 * @note This will always return false if there is no audio device, or the audio device is disabled.
+	 */
+	UFUNCTION(BlueprintCallable, Category = "Audio", meta = (WorldContext = "WorldContextObject"))
+	static bool GetClosestListenerLocation(const UObject* WorldContextObject, const FVector& Location, float MaximumRange, const bool bAllowAttenuationOverride, FVector& ListenerPosition);
+
 	/**
 	* Sets a global pitch modulation scalar that will apply to all non-UI sounds
 	*
diff --git a/Engine/Source/Runtime/Engine/Private/GameplayStatics.cpp b/Engine/Source/Runtime/Engine/Private/GameplayStatics.cpp
--- a/Engine/Source/Runtime/Engine/Private/GameplayStatics.cpp
+++ b/Engine/Source/Runtime/Engine/Private/GameplayStatics.cpp
@@ -1208,7 +1208,7 @@ bool UGameplayStatics::FindCollisionUV(const struct FHitResult& Hit, int32 UVCha
 	return bSuccess;
 }
 
-bool UGameplayStatics::AreAnyListenersWithinRange(const UObject* WorldContextObject, FVector Location, float MaximumRange)
+bool UGameplayStatics::AreAnyListenersWithinRange(const UObject* WorldContextObject, const FVector& Location, float MaximumRange)
 {
 	if (!GEngine || !GEngine->UseSound())
 	{
@@ -1230,6 +1230,36 @@ bool UGameplayStatics::AreAnyListenersWithinRange(const UObject* WorldContextObj
 	return false;
 }
 
+bool UGameplayStatics::GetClosestListenerLocation(const UObject* WorldContextObject, const FVector& Location, float MaximumRange, const bool bAllowAttenuationOverride, FVector& ListenerPosition)
+{
+	if (!GEngine || !GEngine->UseSound())
+	{
+		return false;
+	}
+
+	UWorld* ThisWorld = GEngine->GetWorldFromContextObject(WorldContextObject, EGetWorldErrorMode::LogAndReturnNull);
+	if (!ThisWorld)
+	{
+		return false;
+	}
+
+	// If there is no valid world from the world context object then there certainly are no listeners
+	FAudioDevice* AudioDevice = ThisWorld->GetAudioDevice();
+	if (!AudioDevice)
+	{
+		return false;
+	}
+
+	float OutDistSq;
+	const int32 ClosestListenerIndex = AudioDevice->FindClosestListenerIndex(Location, OutDistSq, bAllowAttenuationOverride);
+	if (ClosestListenerIndex == INDEX_NONE || ((MaximumRange * MaximumRange) < OutDistSq))
+	{
+		return false;
+	}
+
+	return AudioDevice->GetListenerPosition(ClosestListenerIndex, ListenerPosition, bAllowAttenuationOverride);
+}
+
 void UGameplayStatics::SetGlobalPitchModulation(const UObject* WorldContextObject, float PitchModulation, float TimeSec)
 {
 	if (!GEngine || !GEngine->UseSound())
