Fix SoundCue Undo resulting in SoundWavePlayers being left in bad state where sound wave data is no longer loaded in editor.

diff --git a/Engine/Source/Runtime/Engine/Private/SoundNodeWavePlayer.cpp b/Engine/Source/Runtime/Engine/Private/SoundNodeWavePlayer.cpp
--- a/Engine/Source/Runtime/Engine/Private/SoundNodeWavePlayer.cpp
+++ b/Engine/Source/Runtime/Engine/Private/SoundNodeWavePlayer.cpp
@@ -103,12 +103,14 @@ void USoundNodeWavePlayer::SetSoundWave(USoundWave* InSoundWave)
 #if WITH_EDITOR
 void USoundNodeWavePlayer::PostEditChangeProperty(FPropertyChangedEvent& PropertyChangedEvent)
 {
-	if (PropertyChangedEvent.Property && PropertyChangedEvent.Property->GetFName() == GET_MEMBER_NAME_CHECKED(USoundNodeWavePlayer, SoundWaveAssetPtr))
-	{
-		LoadAsset();
-	}
+	// Undo calls this and reports null PropertyChangedEvent,
+	// (PostEditUndo is not guaranteed to be called post re-serialization
+	// of parent SoundCue, but PostEditChangeProperty is)
+	// so always ensure asset is loaded on any call.
+	// Secondary calls will be fast and result in no-op.
+	LoadAsset();
 }
-#endif
+#endif // WITH_EDITOR
 
 void USoundNodeWavePlayer::ParseNodes( FAudioDevice* AudioDevice, const UPTRINT NodeWaveInstanceHash, FActiveSound& ActiveSound, const FSoundParseParameters& ParseParams, TArray<FWaveInstance*>& WaveInstances )
 {
@@ -216,6 +218,10 @@ FText USoundNodeWavePlayer::GetTitle() const
 	{
 		SoundWaveName = FText::FromString(SoundWave->GetFName().ToString());
 	}
+	else if (SoundWaveAssetPtr.IsValid())
+	{
+		SoundWaveName = FText::FromString(SoundWaveAssetPtr.GetAssetName());
+	}
 	else
 	{
 		SoundWaveName = LOCTEXT("NoSoundWave", "NONE");
