bVisualizeComponent needs to be a UProperty because we're relying on it getting propagated from template instances to preview instances
#rnx

diff --git a/Engine/Source/Runtime/Engine/Classes/Components/SceneComponent.h b/Engine/Source/Runtime/Engine/Classes/Components/SceneComponent.h
--- a/Engine/Source/Runtime/Engine/Classes/Components/SceneComponent.h
+++ b/Engine/Source/Runtime/Engine/Classes/Components/SceneComponent.h
@@ -270,6 +270,7 @@ public:
 
 #if WITH_EDITORONLY_DATA
 	/** This component should create a sprite component for visualization in the editor */
+	UPROPERTY()
 	uint8 bVisualizeComponent : 1;
 #endif
 
