Added runtime distance field GPU downsampling

diff --git a/Engine/Shaders/Private/DistanceFieldDownsampling.usf b/Engine/Shaders/Private/DistanceFieldDownsampling.usf
new file mode 100644
--- /dev/null
+++ b/Engine/Shaders/Private/DistanceFieldDownsampling.usf
@@ -0,0 +1,27 @@
+// Copyright 1998-2019 Epic Games, Inc. All Rights Reserved.
+
+/*=============================================================================
+	DistanceFieldDownsampling.usf
+=============================================================================*/
+
+#include "Common.ush"
+
+
+float3 TexelSrcSize;
+uint4 DstSize;
+uint4 OffsetInAtlas;
+SamplerState MeshDFSampler;
+Texture3D<float> MeshDF;
+RWTexture3D<float> DFAtlas;
+
+[numthreads(8,8,1)]
+void DistanceFieldDownsamplingCS(uint3 ThreadId : SV_DispatchThreadID)
+{
+	if (all(ThreadId.xyz < DstSize.xyz))  
+	{
+		float3 VolumeUV = float3(ThreadId.xyz) / float3(DstSize.xyz);
+		VolumeUV += TexelSrcSize;
+		int3 TexturePixelOffset = ThreadId.xyz + (int3)OffsetInAtlas.xyz;
+		DFAtlas[TexturePixelOffset] = MeshDF.SampleLevel(MeshDFSampler, VolumeUV, 0).x;
+	}
+}
diff --git a/Engine/Source/Runtime/Engine/Private/DistanceFieldAtlas.cpp b/Engine/Source/Runtime/Engine/Private/DistanceFieldAtlas.cpp
--- a/Engine/Source/Runtime/Engine/Private/DistanceFieldAtlas.cpp
+++ b/Engine/Source/Runtime/Engine/Private/DistanceFieldAtlas.cpp
@@ -17,6 +17,7 @@
 #include "Engine/StaticMesh.h"
 #include "Misc/AutomationTest.h"
 #include "Async/ParallelFor.h"
+#include "DistanceFieldDownsampling.h"
 
 #if WITH_EDITOR
 #include "DerivedDataCacheInterface.h"
@@ -100,6 +101,11 @@ static TAutoConsoleVariable<int32> CVarDistFieldThrottleCopyToAtlasInBytes(
 	TEXT("When enabled (higher than 0), throttle mesh distance field copy to global mesh distance field atlas volume (in bytes uncompressed)."),
 	ECVF_Default);
 
+static TAutoConsoleVariable<float> CVarDistFieldRuntimeDownsampling(
+	TEXT("r.DistanceFields.RuntimeDownsamplingFactor"),
+	0,
+	TEXT("When enabled (higher than 0 and lower than 1), mesh distance field will be downsampled by factor value on GPU and uploaded to the atlas."),
+	ECVF_Default);
 
 static TAutoConsoleVariable<int32> CVarLandscapeGI(
 	TEXT("r.GenerateLandscapeGIData"),
@@ -326,7 +332,43 @@ struct FCompareVolumeAllocation
 	}
 };
 
-void FDistanceFieldVolumeTextureAtlas::UpdateAllocations()
+static void CopyToUpdateTextureData
+(
+	const FIntVector& SrcSize,
+	int32 FormatSize,
+	const TArray<uint8>& SrcData,
+	const FUpdateTexture3DData& UpdateTextureData,
+	const FIntVector& DstOffset
+)
+{
+	// Is there any padding? If not, straight memcpy
+	if ((UpdateTextureData.DepthPitch * SrcSize.Z) == SrcData.Num())
+	{
+		FPlatformMemory::Memcpy(UpdateTextureData.Data, SrcData.GetData(), SrcData.Num());
+	}
+	else
+	{
+
+		const int32 SourcePitch = SrcSize.X * FormatSize;
+		check(SourcePitch <= (int32)UpdateTextureData.RowPitch);
+
+		for (int32 ZIndex = 0; ZIndex < SrcSize.Z; ZIndex++)
+		{
+			const int32 DestZIndex = (DstOffset.Z + ZIndex) * UpdateTextureData.DepthPitch + DstOffset.X * FormatSize;
+			const int32 SourceZIndex = ZIndex * SrcSize.Y * SourcePitch;
+
+			for (int32 YIndex = 0; YIndex < SrcSize.Y; YIndex++)
+			{
+				const int32 DestIndex = DestZIndex + (DstOffset.Y + YIndex) * UpdateTextureData.RowPitch;
+				const int32 SourceIndex = SourceZIndex + YIndex * SourcePitch;
+				check(uint32(DestIndex) + SourcePitch <= UpdateTextureData.DataSizeBytes);
+				FMemory::Memcpy((uint8*)&UpdateTextureData.Data[DestIndex], (const uint8*)&(SrcData[SourceIndex]), SourcePitch);
+			}
+		}
+	}
+}
+
+void FDistanceFieldVolumeTextureAtlas::UpdateAllocations(FRHICommandListImmediate& RHICmdList, ERHIFeatureLevel::Type InFeatureLevel)
 {
 	{
 		uint32 TotalSurface = BlockAllocator.GetMaxSizeX() * BlockAllocator.GetMaxSizeY() * BlockAllocator.GetMaxSizeZ();
@@ -363,8 +405,8 @@ void FDistanceFieldVolumeTextureAtlas::UpdateAllocations()
 		int ThrottledCopyCount = 0;
 		int PendingCopyCount = PendingAllocations.Num();
 
-		const int32 ThrottleSize = CVarDistFieldThrottleCopyToAtlasInBytes.GetValueOnAnyThread();
-		const bool bThrottleUpdateAllocation = ThrottleSize >= 1024;
+		const float RuntimeDownsamplingFactor = CVarDistFieldRuntimeDownsampling->GetFloat();
+		const bool bRuntimeDownsampling = FDistanceFieldDownsampling::CanDownsample() && (RuntimeDownsamplingFactor > 0 && RuntimeDownsamplingFactor < 1);
 
 		auto AllocateBlocks = [&]()
 		{
@@ -372,7 +414,14 @@ void FDistanceFieldVolumeTextureAtlas::UpdateAllocations()
 			for (int32 AllocationIndex = 0; AllocationIndex < LocalPendingAllocations->Num(); AllocationIndex++)
 			{
 				FDistanceFieldVolumeTexture* Texture = (*LocalPendingAllocations)[AllocationIndex];
-				const FIntVector Size = Texture->VolumeData.Size;
+				FIntVector Size = Texture->VolumeData.Size;
+				
+				if (bRuntimeDownsampling)
+				{
+					FDistanceFieldDownsampling::GetDownsampledSize(Size, RuntimeDownsamplingFactor, Size);
+				}
+				
+				Texture->SizeInAtlas = Size;
 				Texture->bThrottled = false;
 
 				if (!BlockAllocator.AddElement((uint32&)Texture->AtlasAllocationMin.X, (uint32&)Texture->AtlasAllocationMin.Y, (uint32&)Texture->AtlasAllocationMin.Z, Size.X, Size.Y, Size.Z))
@@ -399,6 +448,9 @@ void FDistanceFieldVolumeTextureAtlas::UpdateAllocations()
 			}
 		};
 
+		const int32 ThrottleSize = CVarDistFieldThrottleCopyToAtlasInBytes.GetValueOnAnyThread();
+		const bool bThrottleUpdateAllocation = ThrottleSize >= 1024;
+
 		if (bThrottleUpdateAllocation)
 		{
 			int32 CurrentSize = 0;
@@ -423,6 +475,9 @@ void FDistanceFieldVolumeTextureAtlas::UpdateAllocations()
 		static const auto CVarCompress = IConsoleManager::Get().FindTConsoleVariableDataInt(TEXT("r.DistanceFieldBuild.Compress"));
 		const bool bDataIsCompressed = CVarCompress->GetValueOnAnyThread() != 0;
 
+		TArray<FDistanceFieldDownsamplingDataTask> DownsamplingTasks;
+		TArray<FUpdateTexture3DData> UpdateDataArray;
+
 		if (!VolumeTextureRHI
 			|| BlockAllocator.GetSizeX() > VolumeTextureRHI->GetSizeX()
 			|| BlockAllocator.GetSizeY() > VolumeTextureRHI->GetSizeY()
@@ -454,7 +509,7 @@ void FDistanceFieldVolumeTextureAtlas::UpdateAllocations()
 
 				if (bThrottleUpdateAllocation)
 				{
-					// Throttling during a full realloc when not using the max sime of volume texture will make the same blocks being reused over and over
+					// Throttling during a full realloc when not using the max size of volume texture will make the same blocks being reused over and over
 					// allocate everything pending to avoid this
 					LocalPendingAllocations = &PendingAllocations;
 					ThrottledCopyCount = ThrottledAllocations.Num();
@@ -468,7 +523,8 @@ void FDistanceFieldVolumeTextureAtlas::UpdateAllocations()
 			// Fully free the previous atlas memory before allocating a new one
 			{
 				// Remove last ref, add to deferred delete list
-				VolumeTextureRHI = NULL;
+				VolumeTextureRHI = nullptr;
+				VolumeTextureUAVRHI = nullptr;
 
 				// Flush commandlist, flush RHI thread, delete deferred resources (GNM Memblock defers further)
 				FRHICommandListExecutor::GetImmediateCommandList().ImmediateFlush(EImmediateFlushType::FlushRHIThreadFlushResources);
@@ -491,8 +547,9 @@ void FDistanceFieldVolumeTextureAtlas::UpdateAllocations()
 				VolumeTextureSize.Z, 
 				Format,
 				1,
-				TexCreate_ShaderResource,
+				TexCreate_ShaderResource | TexCreate_UAV,
 				CreateInfo);
+			VolumeTextureUAVRHI = RHICreateUnorderedAccessView(VolumeTextureRHI, 0);
 
 			UE_LOG(LogStaticMesh,Log,TEXT("%s"),*GetSizeString());
 
@@ -503,20 +560,34 @@ void FDistanceFieldVolumeTextureAtlas::UpdateAllocations()
 				const int32 DepthPitch = BlockAllocator.GetSizeX() * BlockAllocator.GetSizeY() * FormatSize;
 
 				const FUpdateTextureRegion3D UpdateRegion(FIntVector::ZeroValue, FIntVector::ZeroValue, BlockAllocator.GetSize());
-				FUpdateTexture3DData TextureUpdateData = RHIBeginUpdateTexture3D(VolumeTextureRHI, 0, UpdateRegion);
+				// FUpdateTexture3DData default constructor is private. it might not be used if copy is done on GPU
+				// Not sure we want to make it public. let's be conservative, and allocate on stack an array of its size
+				uint8 TextureUpdateDataStackMem[sizeof(FUpdateTexture3DData)];
+				FUpdateTexture3DData* TextureUpdateDataPtr = reinterpret_cast<FUpdateTexture3DData*>(&TextureUpdateDataStackMem);
+
+				if (!bRuntimeDownsampling)
+				{
+					*TextureUpdateDataPtr = RHIBeginUpdateTexture3D(VolumeTextureRHI, 0, UpdateRegion);
+				}
 
 				TArray<uint8> UncompressedData;
 
+				if (bRuntimeDownsampling)
+				{
+					UpdateDataArray.Empty(LocalPendingAllocations->Num());
+					UpdateDataArray.AddUninitialized(LocalPendingAllocations->Num());
+				}
+
 				for (int32 AllocationIndex = 0; AllocationIndex < LocalPendingAllocations->Num(); AllocationIndex++)
 				{
 					FDistanceFieldVolumeTexture* Texture = (*LocalPendingAllocations)[AllocationIndex];
-					const FIntVector Size = Texture->VolumeData.Size;
-					const int32 UncompressedSize = Size.X * Size.Y * Size.Z * FormatSize;
-
+					FIntVector Size = Texture->VolumeData.Size;
+					
 					const TArray<uint8>* SourceDataPtr = NULL;
 
 					if (bDataIsCompressed)
 					{
+						const int32 UncompressedSize = Size.X * Size.Y * Size.Z * FormatSize;
 						UncompressedData.Reset(UncompressedSize);
 						UncompressedData.AddUninitialized(UncompressedSize);
 
@@ -531,115 +602,105 @@ void FDistanceFieldVolumeTextureAtlas::UpdateAllocations()
 						SourceDataPtr = &Texture->VolumeData.CompressedDistanceFieldVolume;
 					}
 
-					const int32 SourcePitch = Size.X * FormatSize;
-					check(SourcePitch <= (int32)TextureUpdateData.RowPitch);
-
-					// Copy each row into the correct position in TextureUpdateData.Data
-					for (int32 ZIndex = 0; ZIndex < Size.Z; ZIndex++)
+					FIntVector DstOffset = FIntVector::ZeroValue;
+					
+					if (bRuntimeDownsampling)
 					{
-						const int32 DestZIndex = (Texture->AtlasAllocationMin.Z + ZIndex) * TextureUpdateData.DepthPitch + Texture->AtlasAllocationMin.X * FormatSize;
-						const int32 SourceZIndex = ZIndex * Size.Y * SourcePitch;
-
-						for (int32 YIndex = 0; YIndex < Size.Y; YIndex++)
-						{
-							const int32 DestIndex = DestZIndex + (Texture->AtlasAllocationMin.Y + YIndex) * TextureUpdateData.RowPitch;
-							const int32 SourceIndex = SourceZIndex + YIndex * SourcePitch;
-							check( uint32(DestIndex) + SourcePitch <= TextureUpdateData.DataSizeBytes );
-							FMemory::Memcpy((uint8*)&TextureUpdateData.Data[DestIndex], (const uint8*)&(*SourceDataPtr)[SourceIndex], SourcePitch );
-						}
+						FDistanceFieldDownsamplingDataTask& DownsamplingTask = DownsamplingTasks.Emplace_GetRef();
+						TextureUpdateDataPtr = &UpdateDataArray[AllocationIndex];
+						FDistanceFieldDownsampling::FillDownsamplingTask(Size, Texture->SizeInAtlas, Texture->GetAllocationMin(), Format, DownsamplingTask, *TextureUpdateDataPtr);
 					}
+					else
+					{
+						DstOffset = Texture->GetAllocationMin();
+					}
+					
+					CopyToUpdateTextureData(Size, FormatSize, *SourceDataPtr, *TextureUpdateDataPtr, DstOffset);
 				}
 
 				UncompressedData.Empty();
-				RHIEndUpdateTexture3D(TextureUpdateData);
+
+				if (!bRuntimeDownsampling)
+				{
+					RHIEndUpdateTexture3D(*TextureUpdateDataPtr);
+				}
 			}
 		}
 		else
 		{
 			const int32 NumUpdates = LocalPendingAllocations->Num();
-			TArray<FUpdateTexture3DData> UpdateDataArray;
 			UpdateDataArray.Empty(NumUpdates);
 			UpdateDataArray.AddUninitialized(NumUpdates);
 			
 			// Allocate upload buffers
-			for (int32 Idx = 0; Idx < NumUpdates; ++Idx)
+			if (!bRuntimeDownsampling)
 			{
-				FDistanceFieldVolumeTexture* Texture = (*LocalPendingAllocations)[Idx];
-				const FIntVector& Size = Texture->VolumeData.Size;
-				const FUpdateTextureRegion3D UpdateRegion(Texture->AtlasAllocationMin, FIntVector::ZeroValue, Size);
+				for (int32 Idx = 0; Idx < NumUpdates; ++Idx)
+				{
+					FDistanceFieldVolumeTexture* Texture = (*LocalPendingAllocations)[Idx];
 
-				UpdateDataArray[Idx] = RHIBeginUpdateTexture3D(VolumeTextureRHI, 0, UpdateRegion);
+					const FUpdateTextureRegion3D UpdateRegion(Texture->AtlasAllocationMin, FIntVector::ZeroValue, Texture->SizeInAtlas);
 
-				check(!!UpdateDataArray[Idx].Data);
-				check(static_cast<int32>(UpdateDataArray[Idx].RowPitch) >= Size.X * FormatSize);
-				check(static_cast<int32>(UpdateDataArray[Idx].DepthPitch) >= Size.X * Size.Y * FormatSize);
-			}
+					UpdateDataArray[Idx] = RHIBeginUpdateTexture3D(VolumeTextureRHI, 0, UpdateRegion);
 
-			// Copy data to upload buffers and decompress source data if necessary
-			ParallelFor(
-				NumUpdates,
-				[this, FormatSize, bDataIsCompressed, &UpdateDataArray, LocalPendingAllocations](int32 Idx)
+					check(!!UpdateDataArray[Idx].Data);
+					check(static_cast<int32>(UpdateDataArray[Idx].RowPitch) >= Texture->SizeInAtlas.X * FormatSize);
+					check(static_cast<int32>(UpdateDataArray[Idx].DepthPitch) >= Texture->SizeInAtlas.X * Texture->SizeInAtlas.Y * FormatSize);
+				}
+			}
+			else
 			{
-				FUpdateTexture3DData& UpdateData = UpdateDataArray[Idx];
-				FDistanceFieldVolumeTexture* Texture = (*LocalPendingAllocations)[Idx];
-				const FIntVector& Size = Texture->VolumeData.Size;
+				DownsamplingTasks.Empty(NumUpdates);
+				DownsamplingTasks.AddDefaulted(NumUpdates);
 
-				TArray<uint8> UncompressedData;
-				const uint32 SrcRowPitch = Size.X * FormatSize;
-				const uint32 SrcDepthPitch = Size.Y * SrcRowPitch;
-				const bool bRowByRowCopy = SrcRowPitch != UpdateData.RowPitch
-					|| SrcDepthPitch != UpdateData.DepthPitch;
-				const uint8* SrcData = Texture->VolumeData.CompressedDistanceFieldVolume.GetData();
-				uint32 SrcDataSize = Texture->VolumeData.CompressedDistanceFieldVolume.Num();
-
-				if (bDataIsCompressed)
+				for (int32 Idx = 0; Idx < NumUpdates; ++Idx)
 				{
-					const int32 UncompressedSize = Size.Z * SrcDepthPitch;
-					UncompressedData.Empty(UncompressedSize);
-					UncompressedData.AddUninitialized(UncompressedSize);
-					verify(FCompression::UncompressMemory(
-						NAME_Zlib,
-						UncompressedData.GetData(),
-						UncompressedSize,
-						SrcData,
-						SrcDataSize));
-					SrcData = UncompressedData.GetData();
-					SrcDataSize = UncompressedSize;
+					FDistanceFieldVolumeTexture* Texture = (*LocalPendingAllocations)[Idx];
+					FDistanceFieldDownsampling::FillDownsamplingTask(Texture->VolumeData.Size, Texture->SizeInAtlas, Texture->GetAllocationMin(), Format, DownsamplingTasks[Idx], UpdateDataArray[Idx]);
 				}
+			}
 
-				if (bRowByRowCopy)
+			// Copy data to upload buffers and decompress source data if necessary
+			ParallelFor(
+				NumUpdates,
+				[this, FormatSize, bDataIsCompressed, bRuntimeDownsampling, &UpdateDataArray, LocalPendingAllocations](int32 Idx)
 				{
-					const uint32 NumRows = UpdateData.DepthPitch / UpdateData.RowPitch;
-					uint8* DstSliceData = UpdateData.Data;
-					const uint8* SrcSliceData = SrcData;
-					for (uint32 SliceIdx = 0; SliceIdx < UpdateData.UpdateRegion.Depth; ++SliceIdx)
+					FUpdateTexture3DData& UpdateData = UpdateDataArray[Idx];
+					FDistanceFieldVolumeTexture* Texture = (*LocalPendingAllocations)[Idx];
+					const FIntVector& Size = Texture->VolumeData.Size;
+
+					if (!bDataIsCompressed)
 					{
-						uint8* DstRowData = DstSliceData;
-						const uint8* SrcRowData = SrcSliceData;
-						for (uint32 RowIdx = 0; RowIdx < NumRows; ++RowIdx)
-						{
-							FMemory::Memcpy(DstRowData, SrcRowData, SrcRowPitch);
-							DstRowData += UpdateData.RowPitch;
-							SrcRowData += SrcRowPitch;
-						}
-						DstSliceData += UpdateData.DepthPitch;
-						SrcSliceData += SrcDepthPitch;
+						CopyToUpdateTextureData(Size, FormatSize, Texture->VolumeData.CompressedDistanceFieldVolume, UpdateDataArray[Idx], FIntVector::ZeroValue);
 					}
-				}
-				else
-				{
-					FMemory::Memcpy(UpdateData.Data, SrcData, SrcDataSize);
-				}
-			},
+					else
+					{
+						const int32 UncompressedSize = Size.X * Size.Y * Size.Z * FormatSize;
+						TArray<uint8> UncompressedData;
+						UncompressedData.Empty(UncompressedSize);
+						UncompressedData.AddUninitialized(UncompressedSize);
+						verify(FCompression::UncompressMemory(NAME_Zlib, UncompressedData.GetData(), UncompressedSize, Texture->VolumeData.CompressedDistanceFieldVolume.GetData(), Texture->VolumeData.CompressedDistanceFieldVolume.Num()));
+						
+						CopyToUpdateTextureData(Size, FormatSize, UncompressedData, UpdateDataArray[Idx], FIntVector::ZeroValue);
+					}
+				},
 				!GDistanceFieldParallelAtlasUpdate);
 
-			// For some RHIs, this has the advantage of reducing transition barriers
-			RHIEndMultiUpdateTexture3D(UpdateDataArray);
+			if (!bRuntimeDownsampling)
+			{
+				// For some RHIs, this has the advantage of reducing transition barriers
+				RHIEndMultiUpdateTexture3D(UpdateDataArray);
+			}
 		}
 
 		CurrentAllocations.Append(*LocalPendingAllocations);
 		LocalPendingAllocations->Empty();
 
+		if (DownsamplingTasks.Num() > 0)
+		{
+			FDistanceFieldDownsampling::DispatchDownsampleTasks(RHICmdList, VolumeTextureUAVRHI, InFeatureLevel, DownsamplingTasks, UpdateDataArray);
+		}
+
 		const double EndTime = FPlatformTime::Seconds();
 		const float UpdateDurationMs = (float)(EndTime - StartTime) * 1000.0f;
 
@@ -648,7 +709,7 @@ void FDistanceFieldVolumeTextureAtlas::UpdateAllocations()
 			UE_LOG(LogStaticMesh,Verbose,TEXT("FDistanceFieldVolumeTextureAtlas::UpdateAllocations took %.1fms"), UpdateDurationMs);
 		}
 		GDistanceFieldForceAtlasRealloc = 0;
-	}	
+	}
 }
 
 FDistanceFieldVolumeTexture::~FDistanceFieldVolumeTexture()
diff --git a/Engine/Source/Runtime/Engine/Private/DistanceFieldDownsampling.cpp b/Engine/Source/Runtime/Engine/Private/DistanceFieldDownsampling.cpp
new file mode 100644
--- /dev/null
+++ b/Engine/Source/Runtime/Engine/Private/DistanceFieldDownsampling.cpp
@@ -0,0 +1,147 @@
+// Copyright 1998-2019 Epic Games, Inc. All Rights Reserved.
+
+/*=============================================================================
+	DistanceFieldSampling.cpp
+=============================================================================*/
+
+#include "DistanceFieldDownsampling.h"
+#include "CoreMinimal.h"
+#include "RHI.h"
+#include "RenderGraph.h"
+#include "GlobalShader.h"
+#include "DistanceFieldAtlas.h"
+
+// --------------------------------------------------------------------------------------------------------------------
+DECLARE_GPU_STAT(DFMeshDownsampling);
+
+// --------------------------------------------------------------------------------------------------------------------
+class FDistanceFieldDownsamplingCS : public FGlobalShader
+{
+public:
+	DECLARE_GLOBAL_SHADER(FDistanceFieldDownsamplingCS);
+	SHADER_USE_PARAMETER_STRUCT(FDistanceFieldDownsamplingCS, FGlobalShader);
+
+	BEGIN_SHADER_PARAMETER_STRUCT(FParameters, )
+		SHADER_PARAMETER(FVector, TexelSrcSize)
+		SHADER_PARAMETER(FIntVector4, DstSize)
+		SHADER_PARAMETER(FIntVector4, OffsetInAtlas)
+		SHADER_PARAMETER_TEXTURE(Texture3D<float>, MeshDF)
+		SHADER_PARAMETER_SAMPLER(SamplerState, MeshDFSampler)
+		SHADER_PARAMETER_UAV(RWTexture3D<float>, DFAtlas)
+	END_SHADER_PARAMETER_STRUCT()
+
+	static bool ShouldCompilePermutation(FGlobalShaderPermutationParameters const& Parameters)
+	{
+		return IsFeatureLevelSupported(Parameters.Platform, ERHIFeatureLevel::SM5);
+	}
+};
+
+// --------------------------------------------------------------------------------------------------------------------
+IMPLEMENT_GLOBAL_SHADER(FDistanceFieldDownsamplingCS, "/Engine/Private/DistanceFieldDownsampling.usf", "DistanceFieldDownsamplingCS", SF_Compute);
+
+// --------------------------------------------------------------------------------------------------------------------
+bool FDistanceFieldDownsampling::CanDownsample()
+{
+	static const auto CVarEightBit = IConsoleManager::Get().FindTConsoleVariableDataInt(TEXT("r.DistanceFieldBuild.EightBit"));
+	return (CVarEightBit->GetValueOnAnyThread() != 0);
+}
+
+// --------------------------------------------------------------------------------------------------------------------
+void FDistanceFieldDownsampling::GetDownsampledSize(const FIntVector& Size, float Factor, FIntVector& OutDownsampledSize)
+{
+	if (Size.X <= 6 || Size.Y <= 6 || Size.Z <= 6)
+	{
+		OutDownsampledSize = Size;
+		return;
+	}
+
+	FVector DstSizeFloat = FVector(Size);
+	DstSizeFloat *= Factor;
+
+	OutDownsampledSize.X = FMath::TruncToInt(DstSizeFloat.X);
+	OutDownsampledSize.Y = FMath::TruncToInt(DstSizeFloat.Y);
+	OutDownsampledSize.Z = FMath::TruncToInt(DstSizeFloat.Z);
+}
+
+// --------------------------------------------------------------------------------------------------------------------
+void FDistanceFieldDownsampling::FillDownsamplingTask
+(
+	const FIntVector& SrcSize, 
+	const FIntVector& DstSize, 
+	const FIntVector& OffsetInAtlas, 
+	EPixelFormat Format,
+	FDistanceFieldDownsamplingDataTask& OutDataTask,
+	FUpdateTexture3DData& OutUpdateTextureData
+)
+{
+	FRHIResourceCreateInfo CreateInfo;
+	OutDataTask.VolumeTextureRHI = RHICreateTexture3D
+	(
+		SrcSize.X,
+		SrcSize.Y,
+		SrcSize.Z,
+		Format,
+		1,
+		TexCreate_ShaderResource,
+		CreateInfo
+	);
+
+	OutDataTask.TexelSrcSize = FVector::OneVector / FVector(float(SrcSize.X), float(SrcSize.Y), float(SrcSize.Z));
+	OutDataTask.DstSize = DstSize;
+	OutDataTask.OffsetInAtlas = OffsetInAtlas;
+
+	const FUpdateTextureRegion3D UpdateRegion(FIntVector::ZeroValue, FIntVector::ZeroValue, SrcSize);
+	OutUpdateTextureData = RHIBeginUpdateTexture3D(OutDataTask.VolumeTextureRHI, 0, UpdateRegion);
+}
+
+// --------------------------------------------------------------------------------------------------------------------
+void FDistanceFieldDownsampling::DispatchDownsampleTasks(FRHICommandListImmediate& RHICmdList, FRHIUnorderedAccessView* DFAtlasUAV, ERHIFeatureLevel::Type FeatureLevel, TArray<FDistanceFieldDownsamplingDataTask>& DownsamplingTasks, TArray<FUpdateTexture3DData>& UpdateTextureDataArray)
+{
+	SCOPED_GPU_STAT(RHICmdList, DFMeshDownsampling);
+
+	check(DownsamplingTasks.Num()==UpdateTextureDataArray.Num());
+
+	// End Update texture 3D
+	for (int32 Index = 0; Index < DownsamplingTasks.Num(); Index++)
+	{
+		FUpdateTexture3DData& UpdateTextureData = UpdateTextureDataArray[Index];
+		RHIEndUpdateTexture3D(UpdateTextureData);
+	}
+		
+	// Dispatch CS downsample tasks
+	FRDGBuilder GraphBuilder(RHICmdList);
+
+	for (int32 Index=0; Index<DownsamplingTasks.Num(); Index++)
+	{
+		FDistanceFieldDownsamplingDataTask& Task = DownsamplingTasks[Index];
+
+		FDistanceFieldDownsamplingCS::FParameters* PassParameters = GraphBuilder.AllocParameters<FDistanceFieldDownsamplingCS::FParameters>();
+		PassParameters->TexelSrcSize = Task.TexelSrcSize;
+		PassParameters->DstSize = FIntVector4(Task.DstSize.X, Task.DstSize.Y, Task.DstSize.Z, 0);
+		PassParameters->OffsetInAtlas = FIntVector4(Task.OffsetInAtlas.X, Task.OffsetInAtlas.Y, Task.OffsetInAtlas.Z, 0);
+		PassParameters->MeshDF = Task.VolumeTextureRHI;
+		PassParameters->MeshDFSampler = TStaticSamplerState<SF_Bilinear, AM_Clamp, AM_Clamp, AM_Clamp>::GetRHI();
+		PassParameters->DFAtlas = DFAtlasUAV;	
+
+		TShaderMap<FGlobalShaderType>* GlobalShaderMap = GetGlobalShaderMap(FeatureLevel);
+		TShaderMapRef< FDistanceFieldDownsamplingCS > ComputeShader(GlobalShaderMap);
+		FDistanceFieldDownsamplingCS* ComputeShaderPtr = *ComputeShader;
+		const FIntVector GroupCount(FMath::DivideAndRoundUp(Task.DstSize.X, 8), FMath::DivideAndRoundUp(Task.DstSize.Y, 8), Task.DstSize.Z);
+
+		GraphBuilder.AddPass
+		(
+			RDG_EVENT_NAME("DownsampleMeshDF"),
+			PassParameters, 
+			ERDGPassFlags::Compute,
+			[PassParameters, ComputeShaderPtr, GroupCount, &Task, &DFAtlasUAV](FRHICommandList& RHICmdList)
+			{
+				FComputeShaderUtils::Dispatch(RHICmdList, ComputeShaderPtr, *PassParameters, GroupCount);
+				RHICmdList.TransitionResources(EResourceTransitionAccess::ERWNoBarrier, EResourceTransitionPipeline::EComputeToCompute, &DFAtlasUAV, 1); // No barrier needed
+				Task.VolumeTextureRHI = nullptr;
+			}
+		);
+	}
+
+	GraphBuilder.Execute();
+	RHICmdList.TransitionResources(EResourceTransitionAccess::ERWBarrier, EResourceTransitionPipeline::EComputeToCompute, &DFAtlasUAV, 1);
+}
diff --git a/Engine/Source/Runtime/Engine/Private/StaticMeshRender.cpp b/Engine/Source/Runtime/Engine/Private/StaticMeshRender.cpp
--- a/Engine/Source/Runtime/Engine/Private/StaticMeshRender.cpp
+++ b/Engine/Source/Runtime/Engine/Private/StaticMeshRender.cpp
@@ -1730,7 +1730,7 @@ void FStaticMeshSceneProxy::GetDistancefieldAtlasData(FBox& LocalVolumeBounds, F
 		LocalVolumeBounds = DistanceFieldData->LocalBoundingBox;
 		OutDistanceMinMax = DistanceFieldData->DistanceMinMax;
 		OutBlockMin = DistanceFieldData->VolumeTexture.GetAllocationMin();
-		OutBlockSize = DistanceFieldData->VolumeTexture.GetAllocationSize();
+		OutBlockSize = DistanceFieldData->VolumeTexture.GetAllocationSizeInAtlas();
 		bOutBuiltAsIfTwoSided = DistanceFieldData->bBuiltAsIfTwoSided;
 		bMeshWasPlane = DistanceFieldData->bMeshWasPlane;
 		ObjectLocalToWorldTransforms.Add(GetLocalToWorld());
diff --git a/Engine/Source/Runtime/Engine/Public/DistanceFieldAtlas.h b/Engine/Source/Runtime/Engine/Public/DistanceFieldAtlas.h
--- a/Engine/Source/Runtime/Engine/Public/DistanceFieldAtlas.h
+++ b/Engine/Source/Runtime/Engine/Public/DistanceFieldAtlas.h
@@ -28,6 +28,7 @@ public:
 	FDistanceFieldVolumeTexture(const class FDistanceFieldVolumeData& InVolumeData) :
 		VolumeData(InVolumeData),
 		AtlasAllocationMin(FIntVector(-1, -1, -1)),
+		SizeInAtlas(FIntVector::ZeroValue),
 		bReferencedByAtlas(false),
 		bThrottled(false),
 		StaticMesh(NULL)
@@ -46,6 +47,11 @@ public:
 		return AtlasAllocationMin;
 	}
 
+	FIntVector GetAllocationSizeInAtlas() const
+	{
+		return SizeInAtlas;
+	}
+
 	FIntVector GetAllocationSize() const;
 
 	int32 GetAllocationVolume() const
@@ -68,6 +74,8 @@ public:
 private:
 	const FDistanceFieldVolumeData& VolumeData;
 	FIntVector AtlasAllocationMin;
+	FIntVector SizeInAtlas;
+
 	bool bReferencedByAtlas : 1;
 	/** bThrottled prevents any objects using the texture from being uploaded to the scene buffer until upload of the texture to distance field atlas is complete */
 	bool bThrottled         : 1;
@@ -104,12 +112,13 @@ public:
 	void RemoveAllocation(FDistanceFieldVolumeTexture* Texture);
 
 	/** Reallocates the volume texture if necessary and uploads new allocations. */
-	void UpdateAllocations();
+	void UpdateAllocations(FRHICommandListImmediate& RHICmdList, ERHIFeatureLevel::Type InFeatureLevel);
 
 	int32 GetGeneration() const { return Generation; }
 
 	EPixelFormat Format;
 	FTexture3DRHIRef VolumeTextureRHI;
+	FUnorderedAccessViewRHIRef VolumeTextureUAVRHI;
 
 private:
 	/** Manages the atlas layout. */
diff --git a/Engine/Source/Runtime/Engine/Public/DistanceFieldDownsampling.h b/Engine/Source/Runtime/Engine/Public/DistanceFieldDownsampling.h
new file mode 100644
--- /dev/null
+++ b/Engine/Source/Runtime/Engine/Public/DistanceFieldDownsampling.h
@@ -0,0 +1,26 @@
+// Copyright 1998-2019 Epic Games, Inc. All Rights Reserved.
+
+/*=============================================================================
+	DistanceFieldDownsampling.h
+=============================================================================*/
+
+#pragma once
+
+class FRHICommandListImmediate;
+
+struct FDistanceFieldDownsamplingDataTask
+{
+	FTexture3DRHIRef VolumeTextureRHI;
+	FVector TexelSrcSize;
+	FIntVector DstSize;
+	FIntVector OffsetInAtlas;
+};
+
+class ENGINE_API FDistanceFieldDownsampling
+{
+public:
+	static bool CanDownsample();
+	static void GetDownsampledSize(const FIntVector& Size, float Factor, FIntVector& OutDownsampledSize);
+	static void FillDownsamplingTask(const FIntVector& SrcSize, const FIntVector& DstSize, const FIntVector& OffsetInAtlas, EPixelFormat Format, FDistanceFieldDownsamplingDataTask& OutDataTask, FUpdateTexture3DData& OutTextureUpdateData);
+	static void DispatchDownsampleTasks(FRHICommandListImmediate& RHICmdList, FRHIUnorderedAccessView* DFAtlasUAV, ERHIFeatureLevel::Type FeatureLevel, TArray<FDistanceFieldDownsamplingDataTask>& DownsamplingTasks, TArray<FUpdateTexture3DData>& UpdateTextureData);
+};
diff --git a/Engine/Source/Runtime/Renderer/Private/DeferredShadingRenderer.cpp b/Engine/Source/Runtime/Renderer/Private/DeferredShadingRenderer.cpp
--- a/Engine/Source/Runtime/Renderer/Private/DeferredShadingRenderer.cpp
+++ b/Engine/Source/Runtime/Renderer/Private/DeferredShadingRenderer.cpp
@@ -556,7 +556,7 @@ void FDeferredShadingSceneRenderer::PrepareDistanceFieldScene(FRHICommandListImm
 	{
 		CSV_SCOPED_TIMING_STAT_EXCLUSIVE(RenderDFAO);
 		SCOPE_CYCLE_COUNTER(STAT_FDeferredShadingSceneRenderer_DistanceFieldAO_Init);
-		GDistanceFieldVolumeTextureAtlas.UpdateAllocations();
+		GDistanceFieldVolumeTextureAtlas.UpdateAllocations(RHICmdList, FeatureLevel);
 		UpdateGlobalDistanceFieldObjectBuffers(RHICmdList);
 		if (bSplitDispatch)
 		{
