Fixed compilation warning

#rnx

diff --git a/Engine/Source/Runtime/Engine/Private/DistanceFieldDownsampling.cpp b/Engine/Source/Runtime/Engine/Private/DistanceFieldDownsampling.cpp
--- a/Engine/Source/Runtime/Engine/Private/DistanceFieldDownsampling.cpp
+++ b/Engine/Source/Runtime/Engine/Private/DistanceFieldDownsampling.cpp
@@ -133,10 +133,10 @@ void FDistanceFieldDownsampling::DispatchDownsampleTasks(FRHICommandListImmediat
 			RDG_EVENT_NAME("DownsampleMeshDF"),
 			PassParameters, 
 			ERDGPassFlags::Compute,
-			[PassParameters, ComputeShaderPtr, GroupCount, &Task, &DFAtlasUAV](FRHICommandList& RHICmdList)
+			[PassParameters, ComputeShaderPtr, GroupCount, &Task, &DFAtlasUAV](FRHICommandList& CmdList)
 			{
-				FComputeShaderUtils::Dispatch(RHICmdList, ComputeShaderPtr, *PassParameters, GroupCount);
-				RHICmdList.TransitionResources(EResourceTransitionAccess::ERWNoBarrier, EResourceTransitionPipeline::EComputeToCompute, &DFAtlasUAV, 1); // No barrier needed
+				FComputeShaderUtils::Dispatch(CmdList, ComputeShaderPtr, *PassParameters, GroupCount);
+				CmdList.TransitionResources(EResourceTransitionAccess::ERWNoBarrier, EResourceTransitionPipeline::EComputeToCompute, &DFAtlasUAV, 1); // No barrier needed
 				Task.VolumeTextureRHI = nullptr;
 			}
 		);
