UBlueprintGeneratedClass::GetAuthoritativeClass() returns the itself as the authorative class when it's cooked and therefore has no source UBlueprint

diff --git a/Engine/Source/Runtime/Engine/Private/BlueprintGeneratedClass.cpp b/Engine/Source/Runtime/Engine/Private/BlueprintGeneratedClass.cpp
--- a/Engine/Source/Runtime/Engine/Private/BlueprintGeneratedClass.cpp
+++ b/Engine/Source/Runtime/Engine/Private/BlueprintGeneratedClass.cpp
@@ -182,10 +182,18 @@ FPrimaryAssetId UBlueprintGeneratedClass::GetPrimaryAssetId() const
 
 UClass* UBlueprintGeneratedClass::GetAuthoritativeClass()
 {
-	if (nullptr == ClassGeneratedBy) // to track UE-11597 and UE-11595
-	{
-		UE_LOG(LogBlueprint, Fatal, TEXT("UBlueprintGeneratedClass::GetAuthoritativeClass: ClassGeneratedBy is null. class '%s'"), *GetPathName());
-	}
+ 	if (nullptr == ClassGeneratedBy) // to track UE-11597 and UE-11595
+ 	{
+		// If this is a cooked blueprint, the generatedby class will have been discarded so we'll just have to assume we're authoritative!
+		if (bCooked)
+		{ 
+			return this;
+		}
+		else
+		{
+			UE_LOG(LogBlueprint, Fatal, TEXT("UBlueprintGeneratedClass::GetAuthoritativeClass: ClassGeneratedBy is null. class '%s'"), *GetPathName());
+		}
+ 	}
 
 	UBlueprint* GeneratingBP = CastChecked<UBlueprint>(ClassGeneratedBy);
 
