Const correctness on PlayerController

diff --git a/Engine/Source/Runtime/Engine/Classes/GameFramework/PlayerController.h b/Engine/Source/Runtime/Engine/Classes/GameFramework/PlayerController.h
--- a/Engine/Source/Runtime/Engine/Classes/GameFramework/PlayerController.h
+++ b/Engine/Source/Runtime/Engine/Classes/GameFramework/PlayerController.h
@@ -1677,10 +1677,10 @@ public:
 	virtual void SpawnPlayerCameraManager();
 
 	/** get audio listener position and orientation */
-	virtual void GetAudioListenerPosition(FVector& OutLocation, FVector& OutFrontDir, FVector& OutRightDir);
+	virtual void GetAudioListenerPosition(FVector& OutLocation, FVector& OutFrontDir, FVector& OutRightDir) const;
 
 	/** Gets the attenuation position override. */
-	virtual bool GetAudioListenerAttenuationOverridePosition(FVector& OutLocation);
+	virtual bool GetAudioListenerAttenuationOverridePosition(FVector& OutLocation) const;
 
 	/**
 	 * Used to override the default positioning of the audio listener
diff --git a/Engine/Source/Runtime/Engine/Private/PlayerController.cpp b/Engine/Source/Runtime/Engine/Private/PlayerController.cpp
--- a/Engine/Source/Runtime/Engine/Private/PlayerController.cpp
+++ b/Engine/Source/Runtime/Engine/Private/PlayerController.cpp
@@ -2685,7 +2685,7 @@ void APlayerController::SpawnPlayerCameraManager()
 	}
 }
 
-void APlayerController::GetAudioListenerPosition(FVector& OutLocation, FVector& OutFrontDir, FVector& OutRightDir)
+void APlayerController::GetAudioListenerPosition(FVector& OutLocation, FVector& OutFrontDir, FVector& OutRightDir) const
 {
 	FVector ViewLocation;
 	FRotator ViewRotation;
@@ -2716,7 +2716,7 @@ void APlayerController::GetAudioListenerPosition(FVector& OutLocation, FVector&
 	OutRightDir = ViewRotationMatrix.GetUnitAxis( EAxis::Y );
 }
 
-bool APlayerController::GetAudioListenerAttenuationOverridePosition(FVector& OutLocation)
+bool APlayerController::GetAudioListenerAttenuationOverridePosition(FVector& OutLocation) const
 {
 	if (bOverrideAudioAttenuationListener)
 	{
