Fixed wrong size when downsampling is enabled when removing a mesh DF from the atlas

#rnx

diff --git a/Engine/Source/Runtime/Engine/Private/DistanceFieldAtlas.cpp b/Engine/Source/Runtime/Engine/Private/DistanceFieldAtlas.cpp
--- a/Engine/Source/Runtime/Engine/Private/DistanceFieldAtlas.cpp
+++ b/Engine/Source/Runtime/Engine/Private/DistanceFieldAtlas.cpp
@@ -286,8 +286,8 @@ void FDistanceFieldVolumeTextureAtlas::RemoveAllocation(FDistanceFieldVolumeText
 	{
 		return;
 	}
-
-	FIntVector Size = Texture->VolumeData.Size;
+	
+	FIntVector Size = Texture->SizeInAtlas;
 	int PixelAreaSize = Size.X * Size.Y * Size.Z;
 
 	const FIntVector Min = Texture->GetAllocationMin();
