Fix non-unity
#rnx

diff --git a/Engine/Source/Runtime/Engine/Private/DistanceFieldDownsampling.cpp b/Engine/Source/Runtime/Engine/Private/DistanceFieldDownsampling.cpp
--- a/Engine/Source/Runtime/Engine/Private/DistanceFieldDownsampling.cpp
+++ b/Engine/Source/Runtime/Engine/Private/DistanceFieldDownsampling.cpp
@@ -5,8 +5,7 @@
 =============================================================================*/
 
 #include "DistanceFieldDownsampling.h"
-#include "CoreMinimal.h"
-#include "RHI.h"
+#include "RHIStaticStates.h"
 #include "RenderGraph.h"
 #include "GlobalShader.h"
 #include "DistanceFieldAtlas.h"
diff --git a/Engine/Source/Runtime/Engine/Public/DistanceFieldDownsampling.h b/Engine/Source/Runtime/Engine/Public/DistanceFieldDownsampling.h
--- a/Engine/Source/Runtime/Engine/Public/DistanceFieldDownsampling.h
+++ b/Engine/Source/Runtime/Engine/Public/DistanceFieldDownsampling.h
@@ -6,6 +6,8 @@
 
 #pragma once
 
+#include "RHI.h"
+
 class FRHICommandListImmediate;
 
 struct FDistanceFieldDownsamplingDataTask
