Adding check to prevent crash on local player retrieval

diff --git a/Engine/Source/Runtime/Engine/Private/GameInstance.cpp b/Engine/Source/Runtime/Engine/Private/GameInstance.cpp
--- a/Engine/Source/Runtime/Engine/Private/GameInstance.cpp
+++ b/Engine/Source/Runtime/Engine/Private/GameInstance.cpp
@@ -821,7 +821,12 @@ int32 UGameInstance::GetNumLocalPlayers() const
 
 ULocalPlayer* UGameInstance::GetLocalPlayerByIndex(const int32 Index) const
 {
-	return LocalPlayers[Index];
+	if (LocalPlayers.IsValidIndex(Index))
+	{
+		return LocalPlayers[Index];
+	}
+
+	return nullptr;
 }
 
 APlayerController* UGameInstance::GetFirstLocalPlayerController(UWorld* World) const
