Edigrating CL 8356797:
Removed clearing of viewport focus for all users when the primary keyboard user focuses something other than the viewport. Updated hard-coded checks to User 0 to instead fetch the keyboard user ID.

diff --git a/Engine/Source/Runtime/Engine/Private/Slate/SceneViewport.cpp b/Engine/Source/Runtime/Engine/Private/Slate/SceneViewport.cpp
--- a/Engine/Source/Runtime/Engine/Private/Slate/SceneViewport.cpp
+++ b/Engine/Source/Runtime/Engine/Private/Slate/SceneViewport.cpp
@@ -14,7 +14,6 @@
 #include "Layout/WidgetPath.h"
 #include "UnrealEngine.h"
 #include "Framework/Application/SlateApplication.h"
-#include "Framework/Application/SlateUser.h"
 #include "Slate/SlateTextures.h"
 #include "Slate/DebugCanvas.h"
 
