Undo changelist 10478108

diff --git a/Engine/Source/Runtime/Engine/Private/Slate/SceneViewport.cpp b/Engine/Source/Runtime/Engine/Private/Slate/SceneViewport.cpp
--- a/Engine/Source/Runtime/Engine/Private/Slate/SceneViewport.cpp
+++ b/Engine/Source/Runtime/Engine/Private/Slate/SceneViewport.cpp
@@ -14,6 +14,7 @@
 #include "Layout/WidgetPath.h"
 #include "UnrealEngine.h"
 #include "Framework/Application/SlateApplication.h"
+#include "Framework/Application/SlateUser.h"
 #include "Slate/SlateTextures.h"
 #include "Slate/DebugCanvas.h"
 
@@ -284,8 +285,8 @@ void FSceneViewport::ProcessAccumulatedPointerInput()
 	if (NumMouseSamplesX > 0 || NumMouseSamplesY > 0)
 	{
 		const float DeltaTime = FApp::GetDeltaTime();
-		ViewportClient->InputAxis( this, MouseDeltaUserIndex, EKeys::MouseX, MouseDelta.X, DeltaTime, NumMouseSamplesX );
-		ViewportClient->InputAxis( this, MouseDeltaUserIndex, EKeys::MouseY, MouseDelta.Y, DeltaTime, NumMouseSamplesY );
+		ViewportClient->InputAxis( this, 0, EKeys::MouseX, MouseDelta.X, DeltaTime, NumMouseSamplesX );
+		ViewportClient->InputAxis( this, 0, EKeys::MouseY, MouseDelta.Y, DeltaTime, NumMouseSamplesY );
 	}
 
 	if ( bCursorHiddenDueToCapture )
@@ -320,7 +321,6 @@ void FSceneViewport::ProcessAccumulatedPointerInput()
 	MouseDelta = FIntPoint::ZeroValue;
 	NumMouseSamplesX = 0;
 	NumMouseSamplesY = 0;
-	MouseDeltaUserIndex = INDEX_NONE;
 }
 
 FVector2D FSceneViewport::VirtualDesktopPixelToViewport(FIntPoint VirtualDesktopPointPx) const
@@ -656,8 +656,6 @@ FReply FSceneViewport::OnMouseMove( const FGeometry& InGeometry, const FPointerE
 
 			MouseDelta.Y -= CursorDelta.Y;
 			++NumMouseSamplesY;
-
-			MouseDeltaUserIndex = InMouseEvent.GetUserIndex();
 		}
 
 		if ( bCursorHiddenDueToCapture )
@@ -1058,7 +1056,7 @@ FReply FSceneViewport::OnFocusReceived(const FFocusEvent& InFocusEvent)
 {
 	CurrentReplyState = FReply::Handled(); 
 
-	if ( InFocusEvent.GetUser() == FSlateApplication::Get().GetUserIndexForKeyboard() )
+	if ( InFocusEvent.GetUser() == 0 )
 	{
 		if ( ViewportClient != nullptr )
 		{
@@ -1107,7 +1105,7 @@ FReply FSceneViewport::OnFocusReceived(const FFocusEvent& InFocusEvent)
 void FSceneViewport::OnFocusLost( const FFocusEvent& InFocusEvent )
 {
 	// If the focus loss event isn't the for the primary 'keyboard' user, don't worry about it.
-	if ( InFocusEvent.GetUser() != FSlateApplication::Get().GetUserIndexForKeyboard() )
+	if ( InFocusEvent.GetUser() != 0 )
 	{
 		return;
 	}
@@ -1119,6 +1117,19 @@ void FSceneViewport::OnFocusLost( const FFocusEvent& InFocusEvent )
 	{
 		FScopedConditionalWorldSwitcher WorldSwitcher(ViewportClient);
 		ViewportClient->LostFocus(this);
+
+		TSharedPtr<SWidget> ViewportWidgetPin = ViewportWidget.Pin();
+		if ( ViewportWidgetPin.IsValid() )
+		{
+			FSlateApplication::Get().ForEachUser([&ViewportWidgetPin] (FSlateUser& User) {
+				if (User.GetFocusedWidget() == ViewportWidgetPin )
+				{
+					User.ClearFocus();
+				}
+			});
+		}
+
+		UE_LOG(LogSlate, Log, TEXT("FSceneViewport::OnFocusLost() reason %d"), (int32)InFocusEvent.GetCause());
 	}
 }
 
