Cherry picking change 8356797

diff --git a/Engine/Source/Runtime/Engine/Private/Slate/SceneViewport.cpp b/Engine/Source/Runtime/Engine/Private/Slate/SceneViewport.cpp
--- a/Engine/Source/Runtime/Engine/Private/Slate/SceneViewport.cpp
+++ b/Engine/Source/Runtime/Engine/Private/Slate/SceneViewport.cpp
@@ -1056,7 +1056,7 @@ FReply FSceneViewport::OnFocusReceived(const FFocusEvent& InFocusEvent)
 {
 	CurrentReplyState = FReply::Handled(); 
 
-	if ( InFocusEvent.GetUser() == 0 )
+	if ( InFocusEvent.GetUser() == FSlateApplication::Get().GetUserIndexForKeyboard() )
 	{
 		if ( ViewportClient != nullptr )
 		{
@@ -1105,7 +1105,7 @@ FReply FSceneViewport::OnFocusReceived(const FFocusEvent& InFocusEvent)
 void FSceneViewport::OnFocusLost( const FFocusEvent& InFocusEvent )
 {
 	// If the focus loss event isn't the for the primary 'keyboard' user, don't worry about it.
-	if ( InFocusEvent.GetUser() != 0 )
+	if ( InFocusEvent.GetUser() != FSlateApplication::Get().GetUserIndexForKeyboard() )
 	{
 		return;
 	}
@@ -1117,19 +1117,6 @@ void FSceneViewport::OnFocusLost( const FFocusEvent& InFocusEvent )
 	{
 		FScopedConditionalWorldSwitcher WorldSwitcher(ViewportClient);
 		ViewportClient->LostFocus(this);
-
-		TSharedPtr<SWidget> ViewportWidgetPin = ViewportWidget.Pin();
-		if ( ViewportWidgetPin.IsValid() )
-		{
-			FSlateApplication::Get().ForEachUser([&ViewportWidgetPin] (FSlateUser& User) {
-				if (User.GetFocusedWidget() == ViewportWidgetPin )
-				{
-					User.ClearFocus();
-				}
-			});
-		}
-
-		UE_LOG(LogSlate, Log, TEXT("FSceneViewport::OnFocusLost() reason %d"), (int32)InFocusEvent.GetCause());
 	}
 }
 
