fix shader compiler stats being empty for all platforms except windows

diff --git a/Engine/Source/Runtime/Engine/Private/ShaderCompiler/ShaderCompiler.cpp b/Engine/Source/Runtime/Engine/Private/ShaderCompiler/ShaderCompiler.cpp
--- a/Engine/Source/Runtime/Engine/Private/ShaderCompiler/ShaderCompiler.cpp
+++ b/Engine/Source/Runtime/Engine/Private/ShaderCompiler/ShaderCompiler.cpp
@@ -1522,7 +1522,7 @@ void FShaderCompilerStats::WriteStats()
 		StatWriter.CycleRow();
 
 		
-		for(int32 Platform = 0; Platform < PlatformStats.Num(); ++Platform)
+		for(int32 Platform = 0; Platform < PlatformStats.GetMaxIndex(); ++Platform)
 		{
 			if(PlatformStats.IsValidIndex(Platform))
 			{
@@ -1579,7 +1579,7 @@ void FShaderCompilerStats::WriteStats()
 		StatWriter.CycleRow();
 
 
-		for (int32 Platform = 0; Platform < PlatformStats.Num(); ++Platform)
+		for (int32 Platform = 0; Platform < PlatformStats.GetMaxIndex(); ++Platform)
 		{
 			if (PlatformStats.IsValidIndex(Platform))
 			{
