Added detection for clients sending duplicate packets from two IP's. Handles the case of this starting midgame - which fixes the restart handshake getting stuck in a loop, and should fix a (potentially) related ensure - but does not yet handle the case of this starting right from initial connect.

diff --git a/Engine/Source/Runtime/Engine/Private/PacketHandlers/StatelessConnectHandlerComponent.cpp b/Engine/Source/Runtime/Engine/Private/PacketHandlers/StatelessConnectHandlerComponent.cpp
--- a/Engine/Source/Runtime/Engine/Private/PacketHandlers/StatelessConnectHandlerComponent.cpp
+++ b/Engine/Source/Runtime/Engine/Private/PacketHandlers/StatelessConnectHandlerComponent.cpp
@@ -209,6 +209,7 @@ StatelessConnectHandlerComponent::StatelessConnectHandlerComponent()
 	, LastClientSequence(0)
 	, LastClientSendTimestamp(0.0)
 	, LastChallengeTimestamp(0.0)
+	, LastRestartPacketTimestamp(0.0)
 	, LastSecretId(0)
 	, LastTimestamp(0.f)
 	, LastCookie()
@@ -738,7 +739,7 @@ void StatelessConnectHandlerComponent::Incoming(FBitReader& Packet)
 						SetState(Handler::Component::State::Initialized);
 						Initialized();
 
-						bRestartedHandshake = false;;
+						bRestartedHandshake = false;
 					}
 				}
 				else if (bRestartHandshake)
@@ -750,10 +751,30 @@ void StatelessConnectHandlerComponent::Incoming(FBitReader& Packet)
 					// it has received traffic from us on a different address than before.
 					if (ensure(bValidAuthCookie))
 					{
-						// The server may send multiple restart handshake packets, so have a 10 second delay between accepting them
-						double LastSendTimeDiff = (FPlatformTime::Seconds() - LastClientSendTimestamp) - 10.0;
+						bool bPassedDelayCheck = false;
+						bool bPassedDualIPCheck = false;
+						double CurrentTime = FPlatformTime::Seconds();;
 
-						if (!bRestartedHandshake && LastSendTimeDiff > 0.0)
+						if (!bRestartedHandshake)
+						{
+							UNetConnection* ServerConn = (Driver != nullptr ? Driver->ServerConnection : nullptr);
+							double LastNetConnPacketTime = (ServerConn != nullptr ? ServerConn->LastReceiveRealtime : 0.0);
+
+							// The server may send multiple restart handshake packets, so have a 10 second delay between accepting them
+							bPassedDelayCheck = ((CurrentTime - LastClientSendTimestamp) - 10.0) > 0.0;
+
+							// Some clients end up sending packets duplicated over multiple IP's, triggering the restart handshake.
+							// Detect this by checking if any restart handshake requests have been received in roughly the last second
+							// (Dual IP situations will make the server send them constantly) - and override the checks as a failsafe,
+							// if no NetConnection packets have been received in the last second.
+							bPassedDualIPCheck = LastRestartPacketTimestamp == 0.0 ||
+													((CurrentTime - LastRestartPacketTimestamp) - 1.1) > 0.0 ||
+													((CurrentTime - LastNetConnPacketTime) - 1.0) > 0.0;
+						}
+
+						LastRestartPacketTimestamp = CurrentTime;
+
+						if (!bRestartedHandshake && bPassedDelayCheck && bPassedDualIPCheck)
 						{
 							UE_LOG(LogHandshake, Log, TEXT("Beginning restart handshake process."));
 
@@ -767,10 +788,14 @@ void StatelessConnectHandlerComponent::Incoming(FBitReader& Packet)
 							UE_LOG(LogHandshake, Log, TEXT("Ignoring restart handshake request, while already restarted (this is normal)."));
 						}
 #if !UE_BUILD_SHIPPING
-						else // if (LastSendTimeDiff > 0.0)
+						else if (!bPassedDelayCheck)
 						{
 							UE_LOG(LogHandshake, Log, TEXT("Ignoring restart handshake request, due to < 10 seconds since last handshake."));
 						}
+						else // if (!bPassedDualIPCheck)
+						{
+							UE_LOG(LogHandshake, Log, TEXT("Ignoring restart handshake request, due to recent NetConnection packets."));
+						}
 #endif
 					}
 					else
diff --git a/Engine/Source/Runtime/Engine/Public/PacketHandlers/StatelessConnectHandlerComponent.h b/Engine/Source/Runtime/Engine/Public/PacketHandlers/StatelessConnectHandlerComponent.h
--- a/Engine/Source/Runtime/Engine/Public/PacketHandlers/StatelessConnectHandlerComponent.h
+++ b/Engine/Source/Runtime/Engine/Public/PacketHandlers/StatelessConnectHandlerComponent.h
@@ -281,6 +281,9 @@ private:
 	/** The local (client) time at which the challenge was last updated */
 	double LastChallengeTimestamp;
 
+	/** The local (client) time at which the last restart handshake request was received */
+	double LastRestartPacketTimestamp;
+
 	/** The SecretId value of the last challenge response sent */
 	uint8 LastSecretId;
 
