Proper fix for packet overflow issue when multiple packets were buffered in a frame but the last one was not sent immediately.

#rnx

diff --git a/Engine/Source/Runtime/Engine/Classes/Engine/NetConnection.h b/Engine/Source/Runtime/Engine/Classes/Engine/NetConnection.h
--- a/Engine/Source/Runtime/Engine/Classes/Engine/NetConnection.h
+++ b/Engine/Source/Runtime/Engine/Classes/Engine/NetConnection.h
@@ -380,6 +380,9 @@ private:
 	FPacketTimestamp	LastOSReceiveTime;		// Last time a packet was received at the OS/NIC layer
 	bool				bIsOSReceiveTimeLocal;	// Whether LastOSReceiveTime uses the same clock as the game, or needs translating
 
+	/** Did we write the dummy PacketInfo in the current SendBuffer */
+	bool bSendBufferHasDummyPacketInfo = false;
+
 	/** Stores the bit number where we wrote the dummy packet info in the packet header */
 	FBitWriterMark HeaderMarkForPacketInfo;
 
@@ -1365,10 +1368,9 @@ private:
 	bool bFlushingPacketOrderCache;
 
 	/** 
-	* Set to true after the first packet is sent on a Tick. 
-	* Subsequent packets sent on the same frame will not contain the PacketInfo payload
+	* Set to true after a packet is flushed (sent) and reset at the end of the connection's Tick. 
 	*/
-	bool bSentPacketInfoThisFrame = false;
+	bool bFlushedNetThisFrame = false;
 };
 
 /** Help structs for temporarily setting network settings */
diff --git a/Engine/Source/Runtime/Engine/Private/NetConnection.cpp b/Engine/Source/Runtime/Engine/Private/NetConnection.cpp
--- a/Engine/Source/Runtime/Engine/Private/NetConnection.cpp
+++ b/Engine/Source/Runtime/Engine/Private/NetConnection.cpp
@@ -1476,10 +1476,10 @@ void UNetConnection::FlushNet(bool bIgnoreSimulation)
 		Driver->OutTotalPackets++;
 
 		//Record the first packet time in the histogram
-		if (bSentPacketInfoThisFrame == false)
+		if (bFlushedNetThisFrame == false)
 		{
-		double LastPacketTimeDiffInMs = (Driver->Time - LastSendTime) * 1000.0;
-		NetConnectionHistogram.AddMeasurement(LastPacketTimeDiffInMs);
+			double LastPacketTimeDiffInMs = (Driver->Time - LastSendTime) * 1000.0;
+			NetConnectionHistogram.AddMeasurement(LastPacketTimeDiffInMs);
 		}
 
 		LastSendTime = Driver->Time;
@@ -1496,8 +1496,7 @@ void UNetConnection::FlushNet(bool bIgnoreSimulation)
 
 		AnalyticsVars.OutAckOnlyCount += (NumAckBits > 0 && NumBunchBits == 0);
 
-		// Make sure subsequent packets won't have the packet info payload
-		bSentPacketInfoThisFrame = true;
+		bFlushedNetThisFrame = true;
 
 		InitSendBuffer();
 	}
@@ -1731,20 +1730,18 @@ void UNetConnection::WritePacketHeader(FBitWriter& Writer)
 
 void UNetConnection::WriteDummyPacketInfo(FBitWriter& Writer)
 {
-	// The first packet of a frame will contain the packet info payload
-	//const bool bHasPacketInfoPayload = bSentPacketInfoThisFrame == false;
-	const bool bHasPacketInfoPayload = true;
+	// The first packet of a frame will include the packet info payload
+	const bool bHasPacketInfoPayload = bFlushedNetThisFrame == false;
 	Writer.WriteBit(bHasPacketInfoPayload);
 
 	if (bHasPacketInfoPayload)
-{
+	{
 		// Pre-insert the bits since the final time values will be calculated and inserted right before LowLevelSend
-
 		HeaderMarkForPacketInfo.Init(Writer);
 		int32 DummyJitterClockTime(UE4_NetConnectionPrivate::MaxJitterClockTimeValue);
 		Writer.SerializeBits(&DummyJitterClockTime, UE4_NetConnectionPrivate::NumBitsForJitterClockTimeInHeader);
 
-	const uint8 bHasServerFrameTime = Driver->IsServer() ? bLastHasServerFrameTime : ( CVarPingExcludeFrameTime.GetValueOnGameThread() > 0 ? 1u : 0u );
+		const uint8 bHasServerFrameTime = Driver->IsServer() ? bLastHasServerFrameTime : (CVarPingExcludeFrameTime.GetValueOnGameThread() > 0 ? 1u : 0u);
 		Writer.WriteBit(bHasServerFrameTime);
 
 		if (bHasServerFrameTime && Driver->IsServer())
@@ -1753,17 +1750,19 @@ void UNetConnection::WriteDummyPacketInfo(FBitWriter& Writer)
 			Writer << DummyFrameTimeByte;
 		}
 	}
+
+	bSendBufferHasDummyPacketInfo = bHasPacketInfoPayload;
 }
 
 void UNetConnection::WriteFinalPacketInfo(FBitWriter& Writer, const double PacketSentTimeInS)
 {
-	//const bool bHasPacketInfoPayload = bSentPacketInfoThisFrame == false;
-	//
-	//if (bHasPacketInfoPayload == false)
-	//{
-	//	// PacketInfo payload was already sent, nothing to do here
-	//	return;
-	//}
+	if (bSendBufferHasDummyPacketInfo == false)
+	{
+		// PacketInfo payload is not included in this SendBuffer; nothing to rewrite
+		return;
+	}
+
+	check(HeaderMarkForPacketInfo.GetNumBits() != 0);
 
 	FBitWriterMark CurrentMark(Writer);
 
@@ -1774,22 +1773,23 @@ void UNetConnection::WriteFinalPacketInfo(FBitWriter& Writer, const double Packe
 	{
 		const double DeltaSendTimeInMS = (PacketSentTimeInS - PreviousPacketSendTimeInS) * 1000.0;
 
+		int32 ClockTimeMilliseconds = 0;
+
 		// If the delta is over our max precision, we send MAX value and jitter will be ignored by the receiver.
 		if (DeltaSendTimeInMS >= UE4_NetConnectionPrivate::MaxJitterPrecisionInMS)
 		{
-			int32 MaxClockTime = UE4_NetConnectionPrivate::MaxJitterClockTimeValue;
-			Writer.SerializeBits(&MaxClockTime, UE4_NetConnectionPrivate::NumBitsForJitterClockTimeInHeader);
+			ClockTimeMilliseconds = UE4_NetConnectionPrivate::MaxJitterClockTimeValue;
 		}
 		else
 		{
 			// Get the fractional part (milliseconds) of the clock time
-			int32 ClockTimeMilliseconds = UE4_NetConnectionPrivate::GetClockTimeMilliseconds(PacketSentTimeInS);
+			ClockTimeMilliseconds = UE4_NetConnectionPrivate::GetClockTimeMilliseconds(PacketSentTimeInS);
 
 			// Ensure we don't overflow
 			ClockTimeMilliseconds &= UE4_NetConnectionPrivate::MaxJitterClockTimeValue;
+		}
 
 			Writer.SerializeBits(&ClockTimeMilliseconds, UE4_NetConnectionPrivate::NumBitsForJitterClockTimeInHeader);
-		}
 
 		PreviousPacketSendTimeInS = PacketSentTimeInS;
 	}
@@ -1797,16 +1797,18 @@ void UNetConnection::WriteFinalPacketInfo(FBitWriter& Writer, const double Packe
 	// Write server frame time
 	{
 		const uint8 bHasServerFrameTime = Driver->IsServer() ? bLastHasServerFrameTime : (CVarPingExcludeFrameTime.GetValueOnGameThread() > 0 ? 1u : 0u);
-	Writer.WriteBit(bHasServerFrameTime);
+		Writer.WriteBit(bHasServerFrameTime);
 
-	if (bHasServerFrameTime && Driver->IsServer())
-	{
+		if (bHasServerFrameTime && Driver->IsServer())
+		{
 			// Write data used to calculate link latency
-		uint8 FrameTimeByte = FMath::Min( FMath::FloorToInt( FrameTime * 1000 ), 255 );
-		Writer << FrameTimeByte;
-	}
+			uint8 FrameTimeByte = FMath::Min(FMath::FloorToInt(FrameTime * 1000), 255);
+			Writer << FrameTimeByte;
+		}
 	}
-
+	
+	HeaderMarkForPacketInfo.Reset();
+	
 	// Revert to the correct bit writing place
 	CurrentMark.PopWithoutClear(Writer);
 }
@@ -1821,7 +1823,7 @@ bool UNetConnection::ReadPacketInfo(FBitReader& Reader, bool bHasPacketInfoPaylo
 	}
 
 	const bool bHasServerFrameTime = Reader.ReadBit() == 1u;
-	double ServerFrameTime = 0;
+	double ServerFrameTime = 0.0;
 
 	if ( !Driver->IsServer() )
 	{
@@ -1840,8 +1842,8 @@ bool UNetConnection::ReadPacketInfo(FBitReader& Reader, bool bHasPacketInfoPaylo
 
 	if (Reader.EngineNetVer() < HISTORY_JITTER_IN_HEADER)
 	{
-	uint8 RemoteInKBytesPerSecondByte = 0;
-	Reader << RemoteInKBytesPerSecondByte;
+		uint8 RemoteInKBytesPerSecondByte = 0;
+		Reader << RemoteInKBytesPerSecondByte;
 	}
 
 	// limit to known size to know the size of the packet header
@@ -1883,8 +1885,7 @@ bool UNetConnection::ReadPacketInfo(FBitReader& Reader, bool bHasPacketInfoPaylo
 
 		// use FApp's time because it is set closer to the beginning of the frame - we don't care about the time so far of the current frame to process the packet
 		const double CurrentTime = (PacketReceiveTime != 0.0 ? PacketReceiveTime : FApp::GetCurrentTime());
-		const double GameTime	 = ServerFrameTime;
-		const double RTT		 = (CurrentTime - OutLagTime[Index] ) - ( CVarPingExcludeFrameTime.GetValueOnAnyThread() ? GameTime : 0.0 );
+		const double RTT		 = (CurrentTime - OutLagTime[Index] ) - ( CVarPingExcludeFrameTime.GetValueOnAnyThread() ? ServerFrameTime : 0.0 );
 		const double NewLag		 = FMath::Max( RTT, 0.0 );
 
 		//UE_LOG( LogNet, Warning, TEXT( "Out: %i, InRemote: %i, Saturation: %f" ), OutBytesPerSecondHistory[Index], RemoteInKBytesPerSecond, RemoteSaturation );
@@ -1976,15 +1977,15 @@ void UNetConnection::ReceivedPacket( FBitReader& Reader, bool bIsReinjectedPacke
 
 	ValidateSendBuffer();
 
-	//Record the packet time to the histogram
+	// Record the packet time to the histogram
 	if (!bIsReinjectedPacket)
 	{
 		const double CurrentReceiveTimeInS = FPlatformTime::Seconds();
 		const double LastPacketTimeDiffInMs = (CurrentReceiveTimeInS - LastReceiveRealtime) * 1000.0;
-	NetConnectionHistogram.AddMeasurement(LastPacketTimeDiffInMs);
+		NetConnectionHistogram.AddMeasurement(LastPacketTimeDiffInMs);
 
-	// Update receive time to avoid timeout.
-	LastReceiveTime		= Driver->Time;
+		// Update receive time to avoid timeout.
+		LastReceiveTime = Driver->Time;
 		LastReceiveRealtime = CurrentReceiveTimeInS;
 	}
 
@@ -3010,23 +3011,23 @@ void UNetConnection::Tick()
 		}
 		else
 		{
-		uint32 NbPacketsSent(0);
-		for( FDelayedPacket& DelayedPacket : Delayed )
-		{
-			if( CurrentRealtimeSeconds > DelayedPacket.SendTime )
+			uint32 NbPacketsSent(0);
+			for (FDelayedPacket& DelayedPacket : Delayed)
 			{
-				LowLevelSend((char*)&DelayedPacket.Data[0], DelayedPacket.SizeBits, DelayedPacket.Traits);
-				++NbPacketsSent;
-			}
-			else
-			{
-				// Break now instead of continuing to iterate through the list. Otherwise may cause out of order sends
-				break;
+				if (CurrentRealtimeSeconds > DelayedPacket.SendTime)
+				{
+					LowLevelSend((char*)&DelayedPacket.Data[0], DelayedPacket.SizeBits, DelayedPacket.Traits);
+					++NbPacketsSent;
+				}
+				else
+				{
+					// Break now instead of continuing to iterate through the list. Otherwise may cause out of order sends
+					break;
+				}
 			}
-		}
 
-		Delayed.RemoveAt(0, NbPacketsSent, false);
-	}
+			Delayed.RemoveAt(0, NbPacketsSent, false);
+		}
 	}
 #endif
 
@@ -3348,7 +3349,7 @@ void UNetConnection::Tick()
 		QueuedBits = FMath::TruncToInt(-AllowedLag);
 	}
 
-	bSentPacketInfoThisFrame = false;
+	bFlushedNetThisFrame = false;
 }
 
 void UNetConnection::HandleClientPlayer( APlayerController *PC, UNetConnection* NetConnection )
