Build fix for: 'bModifierFillCollisionUnderneathForNavmesh': undeclared identifier
#rnx

diff --git a/Engine/Source/Runtime/Engine/Private/AI/Navigation/NavigationModifier.cpp b/Engine/Source/Runtime/Engine/Private/AI/Navigation/NavigationModifier.cpp
--- a/Engine/Source/Runtime/Engine/Private/AI/Navigation/NavigationModifier.cpp
+++ b/Engine/Source/Runtime/Engine/Private/AI/Navigation/NavigationModifier.cpp
@@ -731,7 +731,7 @@ void FCompositeNavModifier::Reset()
 	bHasPotentialLinks = false;
 	bAdjustHeight = false;
 	bIsPerInstanceModifier = false;
-	bModifierRejectNavmeshUnderneath = false;
+	bModifierFillCollisionUnderneathForNavmesh = false;
 }
 
 void FCompositeNavModifier::Empty()
diff --git a/Engine/Source/Runtime/Engine/Public/AI/NavigationModifier.h b/Engine/Source/Runtime/Engine/Public/AI/NavigationModifier.h
--- a/Engine/Source/Runtime/Engine/Public/AI/NavigationModifier.h
+++ b/Engine/Source/Runtime/Engine/Public/AI/NavigationModifier.h
@@ -249,7 +249,7 @@ protected:
 
 struct ENGINE_API FCompositeNavModifier : public FNavigationModifier
 {
-	FCompositeNavModifier() : bHasPotentialLinks(false), bAdjustHeight(false), bHasLowAreaModifiers(false), bIsPerInstanceModifier(false), bModifierRejectNavmeshUnderneath(false) {}
+	FCompositeNavModifier() : bHasPotentialLinks(false), bAdjustHeight(false), bHasLowAreaModifiers(false), bIsPerInstanceModifier(false), bModifierFillCollisionUnderneathForNavmesh(false) {}
 
 	void Shrink();
 	void Reset();
@@ -303,8 +303,8 @@ struct ENGINE_API FCompositeNavModifier : public FNavigationModifier
 	FORCEINLINE bool HasAreas() const { return Areas.Num() > 0; }
 	FORCEINLINE bool HasLowAreaModifiers() const { return bHasLowAreaModifiers; }
 	FORCEINLINE bool IsPerInstanceModifier() const { return bIsPerInstanceModifier; }
-	FORCEINLINE bool GetRejectNavmeshUnderneath() const { return bModifierRejectNavmeshUnderneath; }
-	FORCEINLINE void SetRejectNavmeshUnderneath(bool bInRejectNavmeshUnderneath) { bModifierRejectNavmeshUnderneath = bInRejectNavmeshUnderneath; }
+	FORCEINLINE bool GetFillCollisionUnderneathForNavmesh() const { return bModifierFillCollisionUnderneathForNavmesh; }
+	FORCEINLINE void SetFillCollisionUnderneathForNavmesh(bool bInRejectNavmeshUnderneath) { bModifierFillCollisionUnderneathForNavmesh = bInRejectNavmeshUnderneath; }
 
 	FORCEINLINE void ReserveForAdditionalAreas(int32 AdditionalElementsCount) { Areas.Reserve(Areas.Num() + AdditionalElementsCount); }
 
@@ -336,5 +336,5 @@ private:
 	uint32 bAdjustHeight : 1;
 	uint32 bHasLowAreaModifiers : 1;
     uint32 bIsPerInstanceModifier : 1;
-	uint32 bModifierRejectNavmeshUnderneath : 1;
+	uint32 bModifierFillCollisionUnderneathForNavmesh : 1;
 };
