#ue4 - Add location and current client base (if any) to logging in ClientAdjustPosition if the client does not resolve the movement base.

diff --git a/Engine/Source/Runtime/Engine/Private/Components/CharacterMovementComponent.cpp b/Engine/Source/Runtime/Engine/Private/Components/CharacterMovementComponent.cpp
--- a/Engine/Source/Runtime/Engine/Private/Components/CharacterMovementComponent.cpp
+++ b/Engine/Source/Runtime/Engine/Private/Components/CharacterMovementComponent.cpp
@@ -9253,7 +9253,8 @@ void UCharacterMovementComponent::ClientAdjustPosition_Implementation
 	{
 		if (bBaseRelativePosition)
 		{
-			UE_LOG(LogNetPlayerMovement, Warning, TEXT("ClientAdjustPosition_Implementation could not resolve the new relative movement base actor, ignoring server correction!"));
+			UE_LOG(LogNetPlayerMovement, Warning, TEXT("ClientAdjustPosition_Implementation could not resolve the new relative movement base actor, ignoring server correction! Client currently at world location %s on base %s"),
+				*UpdatedComponent->GetComponentLocation().ToString(), *GetNameSafe(GetMovementBase()));
 			return;
 		}
 		else
