Cherry-pick (edigrated) 10325377 :

Minor engine changes necessary for integration of procedural plugin

diff --git a/Engine/Source/Editor/UnrealEd/Classes/ActorFactories/ActorFactoryBoxVolume.h b/Engine/Source/Editor/UnrealEd/Classes/ActorFactories/ActorFactoryBoxVolume.h
--- a/Engine/Source/Editor/UnrealEd/Classes/ActorFactories/ActorFactoryBoxVolume.h
+++ b/Engine/Source/Editor/UnrealEd/Classes/ActorFactories/ActorFactoryBoxVolume.h
@@ -11,13 +11,13 @@
 class AActor;
 struct FAssetData;
 
-UCLASS(MinimalAPI, config=Editor, collapsecategories, hidecategories=Object)
+UCLASS(MinimalAPI, config = Editor, collapsecategories, hidecategories = Object)
 class UNREALED_VTABLE UActorFactoryBoxVolume : public UActorFactoryVolume
 {
 	GENERATED_UCLASS_BODY()
 
 	//~ Begin UActorFactory Interface
-	virtual bool CanCreateActorFrom( const FAssetData& AssetData, FText& OutErrorMsg ) override;
-	UNREALED_API virtual void PostSpawnActor( UObject* Asset, AActor* NewActor ) override;
+	UNREALED_API virtual bool CanCreateActorFrom(const FAssetData& AssetData, FText& OutErrorMsg) override;
+	UNREALED_API virtual void PostSpawnActor(UObject* Asset, AActor* NewActor) override;
 	//~ End UActorFactory Interface
 };
diff --git a/Engine/Source/Runtime/Engine/Public/GenericQuadTree.h b/Engine/Source/Runtime/Engine/Public/GenericQuadTree.h
--- a/Engine/Source/Runtime/Engine/Public/GenericQuadTree.h
+++ b/Engine/Source/Runtime/Engine/Public/GenericQuadTree.h
@@ -24,7 +24,8 @@ public:
 	void Insert(const ElementType& Element, const FBox2D& Box, const TCHAR* DebugContext = nullptr);
 
 	/** Given a 2D box, returns an array of elements within the box. There will not be any duplicates in the list. */
-	void GetElements(const FBox2D& Box, TArray<ElementType>& ElementsOut) const;
+	template<typename ElementAllocatorType>
+	void GetElements(const FBox2D& Box, TArray<ElementType, ElementAllocatorType>& ElementsOut) const;
 
 	/** Removes an object of type ElementType with an associated 2D box of size Box (log n). Does not cleanup tree*/
 	bool Remove(const ElementType& Instance, const FBox2D& Box);
@@ -76,20 +77,18 @@ private:
 	void Split();
 
 	/** Given a list of nodes, return which ones actually intersect the box */
-	void GetIntersectingElements(const FBox2D& Box, TArray/*TSet*/<ElementType>& ElementsOut) const;
+	template<typename ElementAllocatorType>
+	void GetIntersectingElements(const FBox2D& Box, TArray<ElementType, ElementAllocatorType>& ElementsOut) const;
 
 	/** Given a list of nodes, remove the node that contains the given element */
 	bool RemoveNodeForElement(const ElementType& Element);
 
-	/** Internal recursive implementation of @see GetElements */
-	void GetElementsRecursive(const FBox2D& Box, TArray<ElementType>& OutElementSet) const;
-
 	/** Internal recursive implementation of @see Insert */
 	void InsertElementRecursive(const ElementType& Element, const FBox2D& Box, const TCHAR* DebugContext);
 
 private:
-	
-	/** 
+
+	/**
 	 * Contains the actual elements this tree is responsible for. Nodes are used to keep track of each element's AABB as well.
 	 * For a non-internal leaf, this is the list of nodes that are fully contained within this tree.
 	 * For an internal tree, this contains the nodes that overlap multiple subtrees.
@@ -123,15 +122,15 @@ template <typename ElementType, int32 NodeCapacity /*= 4*/>
 void TQuadTree<ElementType, NodeCapacity>::Serialize(FArchive& Ar)
 {
 	Ar << Nodes;
-	
-	bool SubTreeFlags[4] = {SubTrees[0] != nullptr, SubTrees[1] != nullptr, SubTrees[2] != nullptr, SubTrees[3] != nullptr};
+
+	bool SubTreeFlags[4] = { SubTrees[0] != nullptr, SubTrees[1] != nullptr, SubTrees[2] != nullptr, SubTrees[3] != nullptr };
 	Ar << SubTreeFlags[0] << SubTreeFlags[1] << SubTreeFlags[2] << SubTreeFlags[3];
 
-	for(int32 Idx = 0 ; Idx < 4 ; ++Idx)
+	for (int32 Idx = 0; Idx < 4; ++Idx)
 	{
-		if(SubTreeFlags[Idx])
+		if (SubTreeFlags[Idx])
 		{
-			if(Ar.IsLoading())
+			if (Ar.IsLoading())
 			{
 				SubTrees[Idx] = new TreeType(FBox2D(), MinimumQuadSize);
 			}
@@ -147,10 +146,10 @@ void TQuadTree<ElementType, NodeCapacity>::Serialize(FArchive& Ar)
 
 template <typename ElementType, int32 NodeCapacity>
 TQuadTree<ElementType, NodeCapacity>::TQuadTree(const FBox2D& Box, float InMinimumQuadSize)
-: TreeBox(Box)
-, Position(Box.GetCenter())
-, MinimumQuadSize(InMinimumQuadSize)
-, bInternal(false)
+	: TreeBox(Box)
+	, Position(Box.GetCenter())
+	, MinimumQuadSize(InMinimumQuadSize)
+	, bInternal(false)
 {
 	SubTrees[0] = SubTrees[1] = SubTrees[2] = SubTrees[3] = nullptr;
 }
@@ -175,7 +174,7 @@ template <typename ElementType, int32 NodeCapacity>
 void TQuadTree<ElementType, NodeCapacity>::Split()
 {
 	check(bInternal == false);
-	
+
 	const FVector2D Extent = TreeBox.GetExtent();
 	const FVector2D XExtent = FVector2D(Extent.X, 0.f);
 	const FVector2D YExtent = FVector2D(0.f, Extent.Y);
@@ -203,7 +202,7 @@ void TQuadTree<ElementType, NodeCapacity>::Split()
 	SubTrees[TopRight] = new TreeType(FBox2D(C, TR), MinimumQuadSize);
 	SubTrees[BottomLeft] = new TreeType(FBox2D(BL, C), MinimumQuadSize);
 	SubTrees[BottomRight] = new TreeType(FBox2D(BM, MR), MinimumQuadSize);
-	
+
 	//mark as no longer a leaf
 	bInternal = true;
 
@@ -329,7 +328,8 @@ bool TQuadTree<ElementType, NodeCapacity>::Remove(const ElementType& Element, co
 }
 
 template <typename ElementType, int32 NodeCapacity>
-void TQuadTree<ElementType, NodeCapacity>::GetElements(const FBox2D& Box, TArray<ElementType>& ElementsOut) const
+template <typename ElementAllocatorType>
+void TQuadTree<ElementType, NodeCapacity>::GetElements(const FBox2D& Box, TArray<ElementType, ElementAllocatorType>& ElementsOut) const
 {
 	TreeType* Quads[4];
 	const int32 NumQuads = GetQuads(Box, Quads);
@@ -345,7 +345,8 @@ void TQuadTree<ElementType, NodeCapacity>::GetElements(const FBox2D& Box, TArray
 }
 
 template <typename ElementType, int32 NodeCapacity>
-void TQuadTree<ElementType, NodeCapacity>::GetIntersectingElements(const FBox2D& Box, TArray/*TSet*/<ElementType>& ElementsOut) const
+template <typename ElementAllocatorType>
+void TQuadTree<ElementType, NodeCapacity>::GetIntersectingElements(const FBox2D& Box, TArray<ElementType, ElementAllocatorType>& ElementsOut) const
 {
 	ElementsOut.Reserve(ElementsOut.Num() + Nodes.Num());
 	for (const FNode& Node : Nodes)
@@ -402,10 +403,10 @@ void TQuadTree<ElementType, NodeCapacity>::Duplicate(TreeType& OutDuplicate) con
 	{
 		if (TreeType* SubTree = SubTrees[TreeIdx])
 		{
-			OutDuplicate.SubTrees[TreeIdx] = new TreeType(FBox2D(0,0), MinimumQuadSize);
+			OutDuplicate.SubTrees[TreeIdx] = new TreeType(FBox2D(0, 0), MinimumQuadSize);
 			SubTree->Duplicate(*OutDuplicate.SubTrees[TreeIdx]);	//duplicate sub trees
 		}
-		
+
 	}
 
 	OutDuplicate.Nodes = Nodes;
@@ -427,7 +428,7 @@ void TQuadTree<ElementType, NodeCapacity>::Empty()
 		}
 
 	}
-	
+
 	Nodes.Empty();
 	bInternal = false;
-}
+}
\ No newline at end of file
diff --git a/Engine/Source/Runtime/Foliage/Public/FoliageInstancedStaticMeshComponent.h b/Engine/Source/Runtime/Foliage/Public/FoliageInstancedStaticMeshComponent.h
--- a/Engine/Source/Runtime/Foliage/Public/FoliageInstancedStaticMeshComponent.h
+++ b/Engine/Source/Runtime/Foliage/Public/FoliageInstancedStaticMeshComponent.h
@@ -5,6 +5,7 @@
 #include "CoreMinimal.h"
 #include "UObject/ObjectMacros.h"
 #include "Engine/EngineTypes.h"
+#include "Misc/Guid.h"
 #include "Components/HierarchicalInstancedStaticMeshComponent.h"
 #include "FoliageInstancedStaticMeshComponent.generated.h"
 
@@ -22,14 +23,14 @@ class FOLIAGE_API UFoliageInstancedStaticMeshComponent : public UHierarchicalIns
 public:
 
 	UPROPERTY(BlueprintAssignable, Category = "Game|Damage")
-	FInstancePointDamageSignature OnInstanceTakePointDamage;
+		FInstancePointDamageSignature OnInstanceTakePointDamage;
 
 	UPROPERTY(BlueprintAssignable, Category = "Game|Damage")
-	FInstanceRadialDamageSignature OnInstanceTakeRadialDamage;
+		FInstanceRadialDamageSignature OnInstanceTakeRadialDamage;
 
 #if WITH_EDITORONLY_DATA
 	UPROPERTY(transient)
-	uint64 FoliageHiddenEditorViews;
+		uint64 FoliageHiddenEditorViews;
 #endif// WITH_EDITOR
 
 	virtual void ReceiveComponentDamage(float DamageAmount, FDamageEvent const& DamageEvent, AController* EventInstigator, AActor* DamageCauser) override;
@@ -38,5 +39,14 @@ public:
 #if WITH_EDITOR
 	virtual uint64 GetHiddenEditorViews() const override;
 #endif// WITH_EDITOR
+
+	/** Used by procedural generation to link generated component with its creator */
+	void SetGenerationGuid(const FGuid& InGuid) { GenerationGuid = InGuid; }
+	const FGuid& GetGenerationGuid() const { return GenerationGuid; }
+
+private:
+
+	UPROPERTY()
+		FGuid GenerationGuid;
 };
 
