XB1 - Fix r.VirtualTextureReducedMemory code path. Fix 2d texture memory reporting.

diff --git a/Engine/Source/Runtime/Engine/Private/Streaming/AsyncTextureStreaming.cpp b/Engine/Source/Runtime/Engine/Private/Streaming/AsyncTextureStreaming.cpp
--- a/Engine/Source/Runtime/Engine/Private/Streaming/AsyncTextureStreaming.cpp
+++ b/Engine/Source/Runtime/Engine/Private/Streaming/AsyncTextureStreaming.cpp
@@ -869,7 +869,7 @@ void FRenderAssetStreamingMipCalcTask::UpdateStats_Async()
 		Stats.UsedStreamingPool += FMath::Min<int64>(RequiredSize, BudgetedSize);
 
 		// Remove from the non streaming budget what is actually taken by streaming.
-		Stats.NonStreamingMips -= ResidentSize;
+		Stats.NonStreamingMips -= ResidentSize * StreamingRenderAsset.IsTexture();
 
 		// All persistent mip bias bigger than the expected is considered overbudget.
 		const int32 OverBudgetBias = FMath::Max<int32>(0, StreamingRenderAsset.BudgetMipBias - Settings.GlobalMipBias);
@@ -981,7 +981,7 @@ void FRenderAssetStreamingMipCalcTask::UpdateCSVOnlyStats_Async()
 		Stats.RequiredPool += RequiredSize;
 
 		// Remove from the non streaming budget what is actually taken by streaming.
-		Stats.NonStreamingMips -= ResidentSize;
+		Stats.NonStreamingMips -= ResidentSize * StreamingRenderAsset.IsTexture();
 
 		const int64 UsedSize = FMath::Min3<int64>(RequiredSize, BudgetedSize, ResidentSize);
 
diff --git a/Engine/Source/Runtime/Engine/Private/Streaming/Texture2DStreamIn.cpp b/Engine/Source/Runtime/Engine/Private/Streaming/Texture2DStreamIn.cpp
--- a/Engine/Source/Runtime/Engine/Private/Streaming/Texture2DStreamIn.cpp
+++ b/Engine/Source/Runtime/Engine/Private/Streaming/Texture2DStreamIn.cpp
@@ -71,7 +71,7 @@ void FTexture2DStreamIn::DoLockNewMips(const FContext& Context)
 	{
 		// With virtual textures, all mips exist although they might not be allocated.
 		const FTexture2DRHIRef Texture2DRHI = Context.Resource->GetTexture2DRHI();
-		const bool bIsVirtualTexture = (Texture2DRHI->GetFlags() & TexCreate_Virtual) == TexCreate_Virtual;
+		const bool bIsVirtualTexture = (IntermediateTextureRHI->GetFlags() & TexCreate_Virtual) == TexCreate_Virtual;
 		const int32 MipOffset = bIsVirtualTexture ? 0 : PendingFirstMip;
 
 		const int32 CurrentFirstMip = Context.Resource->GetCurrentFirstMip();
@@ -93,7 +93,7 @@ void FTexture2DStreamIn::DoUnlockNewMips(const FContext& Context)
 	{
 		// With virtual textures, all mips exist although they might not be allocated.
 		const FTexture2DRHIRef Texture2DRHI = Context.Resource->GetTexture2DRHI();
-		const bool bIsVirtualTexture = (Texture2DRHI->GetFlags() & TexCreate_Virtual) == TexCreate_Virtual;
+		const bool bIsVirtualTexture = (IntermediateTextureRHI->GetFlags() & TexCreate_Virtual) == TexCreate_Virtual;
 		const int32 MipOffset = bIsVirtualTexture ? 0 : PendingFirstMip;
 
 		const int32 CurrentFirstMip = Context.Resource->GetCurrentFirstMip();
diff --git a/Engine/Source/Runtime/Engine/Private/Texture2D.cpp b/Engine/Source/Runtime/Engine/Private/Texture2D.cpp
--- a/Engine/Source/Runtime/Engine/Private/Texture2D.cpp
+++ b/Engine/Source/Runtime/Engine/Private/Texture2D.cpp
@@ -117,7 +117,7 @@ static bool CanCreateAsVirtualTexture(uint32 TexCreateFlags)
 	const uint32 iRequiredFlags =
 		TexCreate_OfflineProcessed;
 
-	return ((TexCreateFlags & (iDisableFlags | iRequiredFlags)) == iRequiredFlags) && CVarVirtualTextureEnabled.GetValueOnRenderThread();
+	return ((TexCreateFlags & (iDisableFlags | iRequiredFlags)) == iRequiredFlags) && CVarVirtualTextureEnabled.GetValueOnAnyThread();
 	
 #else
 	return false;
@@ -949,18 +949,30 @@ int32 UTexture2D::CalcTextureMemorySize( int32 MipCount ) const
 	int32 Size = 0;
 	if (PlatformData)
 	{
-		int32 SizeX = GetSizeX();
-		int32 SizeY = GetSizeY();
-		int32 NumMips = GetNumMips();
-		EPixelFormat Format = GetPixelFormat();
+		static TConsoleVariableData<int32>* CVarReducedMode = IConsoleManager::Get().FindTConsoleVariableDataInt(TEXT("r.VirtualTextureReducedMemory"));
+		check(CVarReducedMode);
 
-		// Figure out what the first mip to use is.
-		int32 FirstMip	= FMath::Max( 0, NumMips - MipCount );		
-		FIntPoint MipExtents = CalcMipMapExtent(SizeX, SizeY, Format, FirstMip);
+		uint32 TexCreateFlags = (SRGB ? TexCreate_SRGB : 0) | (bNoTiling ? TexCreate_NoTiling : 0) | TexCreate_OfflineProcessed | TexCreate_Streamable;
+		const bool bCanBeVirtual = CanCreateAsVirtualTexture(TexCreateFlags);
 
-		uint32 TextureAlign = 0;
-		uint64 TextureSize = RHICalcTexture2DPlatformSize(MipExtents.X, MipExtents.Y, Format, MipCount, 1, 0, FRHIResourceCreateInfo(PlatformData->GetExtData()), TextureAlign);
-		Size = (int32)TextureSize;
+		const int32 SizeX = GetSizeX();
+		const int32 SizeY = GetSizeY();
+		const int32 NumMips = GetNumMips();
+		const int32 FirstMip = FMath::Max(0, NumMips - MipCount);
+		const EPixelFormat Format = GetPixelFormat();
+		uint32 TextureAlign;
+
+		// Must be consistent with the logic in FTexture2DResource::InitRHI
+		if (bIsStreamable && bCanBeVirtual && (!CVarReducedMode->GetValueOnAnyThread() || MipCount > UTexture2D::GetMinTextureResidentMipCount()))
+		{
+			TexCreateFlags |= TexCreate_Virtual;
+			Size = (int32)RHICalcVMTexture2DPlatformSize(SizeX, SizeY, Format, NumMips, FirstMip, 1, TexCreateFlags, TextureAlign);
+		}
+		else
+		{
+			const FIntPoint MipExtents = CalcMipMapExtent(SizeX, SizeY, Format, FirstMip);
+			Size = (int32)RHICalcTexture2DPlatformSize(MipExtents.X, MipExtents.Y, Format, MipCount, 1, TexCreateFlags, FRHIResourceCreateInfo(PlatformData->GetExtData()), TextureAlign);
+		}
 	}
 	return Size;
 }
diff --git a/Engine/Source/Runtime/RHI/Private/DynamicRHI.cpp b/Engine/Source/Runtime/RHI/Private/DynamicRHI.cpp
--- a/Engine/Source/Runtime/RHI/Private/DynamicRHI.cpp
+++ b/Engine/Source/Runtime/RHI/Private/DynamicRHI.cpp
@@ -386,6 +386,12 @@ void FDynamicRHI::RHIUpdateShaderResourceView(FRHIShaderResourceView* SRV, FRHII
 	UE_LOG(LogRHI, Fatal, TEXT("RHIUpdateShaderResourceView isn't implemented for the current RHI"));
 }
 
+uint64 FDynamicRHI::RHICalcVMTexture2DPlatformSize(uint32 Mip0Width, uint32 Mip0Height, uint8 Format, uint32 NumMips, uint32 FirstMipIdx, uint32 NumSamples, uint32 Flags, uint32& OutAlign)
+{
+	UE_LOG(LogRHI, Fatal, TEXT("RHICalcVMTexture2DPlatformSize isn't implemented for the current RHI"));
+	return -1;
+}
+
 FDefaultRHIRenderQueryPool::FDefaultRHIRenderQueryPool(ERenderQueryType InQueryType, FDynamicRHI* InDynamicRHI, uint32 InNumQueries)
 	: DynamicRHI(InDynamicRHI)
 	, QueryType(InQueryType)
diff --git a/Engine/Source/Runtime/RHI/Public/DynamicRHI.h b/Engine/Source/Runtime/RHI/Public/DynamicRHI.h
--- a/Engine/Source/Runtime/RHI/Public/DynamicRHI.h
+++ b/Engine/Source/Runtime/RHI/Public/DynamicRHI.h
@@ -553,6 +553,21 @@ public:
 	// FlushType: Thread safe
 	virtual uint64 RHICalcTexture2DPlatformSize(uint32 SizeX, uint32 SizeY, uint8 Format, uint32 NumMips, uint32 NumSamples, uint32 Flags, const FRHIResourceCreateInfo& CreateInfo, uint32& OutAlign) = 0;
 
+	/**
+	* Computes the total size of a virtual memory (VM) based 2D texture with the specified parameters.
+	*
+	* @param Mip0Width - width of the top mip
+	* @param Mip0Height - height of the top mip
+	* @param Format - EPixelFormat texture format
+	* @param NumMips - number of mips the texture has
+	* @param FirstMipIdx - index of the first resident mip
+	* @param NumSamples - number of MSAA samples, usually 1
+	* @param Flags - ETextureCreateFlags creation flags
+	* @param OutAlign - Alignment required for this texture. Output parameter.
+	*/
+	// FlushType: Thread safe
+	virtual uint64 RHICalcVMTexture2DPlatformSize(uint32 Mip0Width, uint32 Mip0Height, uint8 Format, uint32 NumMips, uint32 FirstMipIdx, uint32 NumSamples, uint32 Flags, uint32& OutAlign);
+
 	/**
 	* Computes the total size of a 3D texture with the specified parameters.
 	* @param SizeX - width of the texture to create
@@ -1402,6 +1417,11 @@ FORCEINLINE uint64 RHICalcTexture2DPlatformSize(uint32 SizeX, uint32 SizeY, uint
 	return GDynamicRHI->RHICalcTexture2DPlatformSize(SizeX, SizeY, Format, NumMips, NumSamples, Flags, CreateInfo, OutAlign);
 }
 
+FORCEINLINE uint64 RHICalcVMTexture2DPlatformSize(uint32 Mip0Width, uint32 Mip0Height, uint8 Format, uint32 NumMips, uint32 FirstMipIdx, uint32 NumSamples, uint32 Flags, uint32& OutAlign)
+{
+	return GDynamicRHI->RHICalcVMTexture2DPlatformSize(Mip0Width, Mip0Height, Format, NumMips, FirstMipIdx, NumSamples, Flags, OutAlign);
+}
+
 FORCEINLINE uint64 RHICalcTexture3DPlatformSize(uint32 SizeX, uint32 SizeY, uint32 SizeZ, uint8 Format, uint32 NumMips, uint32 Flags, const FRHIResourceCreateInfo& CreateInfo, uint32& OutAlign)
 {
 	return GDynamicRHI->RHICalcTexture3DPlatformSize(SizeX, SizeY, SizeZ, Format, NumMips, Flags, CreateInfo, OutAlign);
