Fixed solver shutdown crash when attempting to flush physics scenes.

diff --git a/Engine/Source/Runtime/Engine/Private/PhysicsEngine/Experimental/PhysScene_Chaos.cpp b/Engine/Source/Runtime/Engine/Private/PhysicsEngine/Experimental/PhysScene_Chaos.cpp
--- a/Engine/Source/Runtime/Engine/Private/PhysicsEngine/Experimental/PhysScene_Chaos.cpp
+++ b/Engine/Source/Runtime/Engine/Private/PhysicsEngine/Experimental/PhysScene_Chaos.cpp
@@ -995,7 +995,7 @@ void FPhysScene_ChaosInterface::Flush_AssumesLocked()
 		Dispatcher->Execute();
 	}
 
-	Chaos::FPBDRigidsSolver* Solver = Scene.GetSolver();
+	Chaos::FPBDRigidsSolver* Solver = GetSolver();
 
 	if(Solver)
 	{
@@ -1005,14 +1005,14 @@ void FPhysScene_ChaosInterface::Flush_AssumesLocked()
 		{
 			Command(Solver);
 		}
-	}
 
-	// Populate the spacial acceleration
-	Chaos::FPBDRigidsSolver::FPBDRigidsEvolution* Evolution = GetSolver()->GetEvolution();
+		// Populate the spacial acceleration
+		Chaos::FPBDRigidsSolver::FPBDRigidsEvolution* Evolution = Solver->GetEvolution();
 
-	if(Evolution)
-	{
-		Evolution->FlushSpatialAcceleration();
+		if(Evolution)
+		{
+			Evolution->FlushSpatialAcceleration();
+		}
 	}
 
 	Scene.CopySolverAccelerationStructure();
