Fixed FGenericPhysicsInterface::GeomOverlapAnyTest that should return true if any hit is found, not only blocking hit.

#rnx

diff --git a/Engine/Source/Runtime/Engine/Private/Collision/SceneQuery.cpp b/Engine/Source/Runtime/Engine/Private/Collision/SceneQuery.cpp
--- a/Engine/Source/Runtime/Engine/Private/Collision/SceneQuery.cpp
+++ b/Engine/Source/Runtime/Engine/Private/Collision/SceneQuery.cpp
@@ -638,7 +638,7 @@ bool FGenericPhysicsInterface::GeomOverlapAnyTest(const UWorld* World, const str
 	TArray<FOverlapResult> Overlaps;	//needed only for template shared code
 	FTransform GeomTransform(Rot, Pos);
 	FPhysicsShapeAdapter Adaptor(GeomTransform.GetRotation(), CollisionShape);
-	return GeomOverlapMultiImp<EQueryInfo::GatherAll>(World, Adaptor.GetGeometry(), CollisionShape, Adaptor.GetGeomPose(GeomTransform.GetTranslation()), Overlaps, TraceChannel, Params, ResponseParams, ObjectParams);
+	return GeomOverlapMultiImp<EQueryInfo::IsAnything>(World, Adaptor.GetGeometry(), CollisionShape, Adaptor.GetGeomPose(GeomTransform.GetTranslation()), Overlaps, TraceChannel, Params, ResponseParams, ObjectParams);
 }
 
 template<>
