Hide VirtualTexture category on non-UMeshComponent

diff --git a/Engine/Source/Runtime/Engine/Classes/Components/MeshComponent.h b/Engine/Source/Runtime/Engine/Classes/Components/MeshComponent.h
--- a/Engine/Source/Runtime/Engine/Classes/Components/MeshComponent.h
+++ b/Engine/Source/Runtime/Engine/Classes/Components/MeshComponent.h
@@ -17,7 +17,7 @@
  * @see UStaticMeshComponent
  * @see USkeletalMeshComponent
  */
-UCLASS(abstract)
+UCLASS(abstract, ShowCategories = (VirtualTexture))
 class ENGINE_API UMeshComponent : public UPrimitiveComponent
 {
 	GENERATED_UCLASS_BODY()
diff --git a/Engine/Source/Runtime/Engine/Classes/Components/PrimitiveComponent.h b/Engine/Source/Runtime/Engine/Classes/Components/PrimitiveComponent.h
--- a/Engine/Source/Runtime/Engine/Classes/Components/PrimitiveComponent.h
+++ b/Engine/Source/Runtime/Engine/Classes/Components/PrimitiveComponent.h
@@ -169,7 +169,7 @@ DECLARE_DYNAMIC_MULTICAST_SPARSE_DELEGATE_TwoParams( FComponentEndTouchOverSigna
  * There are several subclasses for the various types of geometry, but the most common by far are the ShapeComponents (Capsule, Sphere, Box), StaticMeshComponent, and SkeletalMeshComponent.
  * ShapeComponents generate geometry that is used for collision detection but are not rendered, while StaticMeshComponents and SkeletalMeshComponents contain pre-built geometry that is rendered, but can also be used for collision detection.
  */
-UCLASS(abstract, HideCategories=(Mobility), ShowCategories=(PhysicsVolume))
+UCLASS(abstract, HideCategories=(Mobility, VirtualTexture), ShowCategories=(PhysicsVolume))
 class ENGINE_API UPrimitiveComponent : public USceneComponent, public INavRelevantInterface
 {
 	GENERATED_BODY()
