Chaos - Implement UpdateBodyScale.
Make UpdateShapesFromGeometry not destroy existing shape data.

diff --git a/Engine/Source/Runtime/Engine/Private/PhysicsEngine/BodyInstance.cpp b/Engine/Source/Runtime/Engine/Private/PhysicsEngine/BodyInstance.cpp
--- a/Engine/Source/Runtime/Engine/Private/PhysicsEngine/BodyInstance.cpp
+++ b/Engine/Source/Runtime/Engine/Private/PhysicsEngine/BodyInstance.cpp
@@ -28,6 +28,13 @@
 #include "Logging/TokenizedMessage.h"
 #include "Logging/MessageLog.h"
 #include "Misc/UObjectToken.h"
+#include "ChaosCheck.h"
+#include "Chaos/Capsule.h"
+#include "Chaos/Convex.h"
+#include "Chaos/ImplicitObject.h"
+#include "Chaos/ImplicitObjectScaled.h"
+#include "Chaos/ParticleHandle.h"
+#include "Chaos/TriangleMeshImplicitObject.h"
 
 #if WITH_PHYSX
 	#include "PhysXPublic.h"
@@ -37,11 +44,6 @@
 #include "PxShape.h"
 #endif // WITH_PHYSX
 
-#if WITH_CHAOS
-#include "Chaos/ImplicitObject.h"
-#endif
-
-
 
 #define LOCTEXT_NAMESPACE "BodyInstance"
 
@@ -1621,7 +1623,210 @@ bool FBodyInstance::UpdateBodyScale(const FVector& InScale3D, bool bForceUpdate)
 #endif
 
 	FVector UpdatedScale3D;
-#if WITH_PHYSX
+
+#if WITH_CHAOS 
+	using namespace Chaos;
+
+	//Get all shapes
+	EScaleMode::Type ScaleMode = EScaleMode::Free;
+
+	FPhysicsCommand::ExecuteWrite(ActorHandle, [&](const FPhysicsActorHandle& Actor)
+	{
+		TArray<FPhysicsShapeHandle> Shapes;
+		GetAllShapes_AssumesLocked(Shapes);
+		ScaleMode = ComputeScaleMode(Shapes);
+
+		FVector AdjustedScale3D;
+		FVector AdjustedScale3DAbs;
+
+		// Apply scaling
+		ComputeScalingVectors(ScaleMode, InScale3D, AdjustedScale3D, AdjustedScale3DAbs);
+
+		TArray<TUniquePtr<FImplicitObject>> NewGeometry;
+		NewGeometry.Reserve(Shapes.Num());
+
+
+		for (FPhysicsShapeHandle& ShapeHandle : Shapes)
+		{
+			const Chaos::FImplicitObject& ImplicitObject = ShapeHandle.GetGeometry();
+			EImplicitObjectType ImplicitType = ImplicitObject.GetType(true);
+			EImplicitObjectType GeomType = GetInnerType(ImplicitType);
+
+			const FTransform& RelativeTM = GetRelativeBodyTransform(ShapeHandle);
+
+			// TODO: support instanced and transformed.
+			CHAOS_ENSURE(!IsInstanced(ImplicitType));
+			CHAOS_ENSURE(ImplicitType != ImplicitObjectType::Transformed);
+
+			FKShapeElem* ShapeElem = FPhysxUserData::Get<FKShapeElem>(FPhysicsInterface::GetUserData(ShapeHandle));
+
+			switch (GeomType)
+			{
+				case ImplicitObjectType::Box:
+				{
+					if (!CHAOS_ENSURE(!IsScaled(ImplicitType)))
+					{
+						// No support for ScaledImplicit Box yet
+						break;
+					}
+
+					FKBoxElem* BoxElem = ShapeElem->GetShapeCheck<FKBoxElem>();
+
+					const TBox<FReal, 3> * BoxGeometry = static_cast<const TBox<FReal, 3>*>(&ImplicitObject);
+
+					FVec3 HalfExtents;
+					HalfExtents.X = FMath::Max((0.5f * BoxElem->X * AdjustedScale3DAbs.X), FCollisionShape::MinBoxExtent());
+					HalfExtents.Y = FMath::Max((0.5f * BoxElem->Y * AdjustedScale3DAbs.Y), FCollisionShape::MinBoxExtent());
+					HalfExtents.Z = FMath::Max((0.5f * BoxElem->Z * AdjustedScale3DAbs.Z), FCollisionShape::MinBoxExtent());
+
+
+					// TODO: For Transformed implicit, do not bake this in. Set Transform instead.
+					FRigidTransform3 LocalTransform = BoxElem->GetTransform() * RelativeTM;
+					LocalTransform.ScaleTranslation(AdjustedScale3D);
+
+					// Any Box with initial Rotation should've been initialized as Transformed implicit.
+					// If we hit this it means user set rotation on initially unrotated box.
+					// TODO: Handle this.
+					CHAOS_ENSURE(LocalTransform.GetRotation() == FQuat::Identity);
+
+					const FVec3 Min = LocalTransform.GetLocation() - HalfExtents;
+					const FVec3 Max = LocalTransform.GetLocation() + HalfExtents;
+
+					TUniquePtr<TBox<FReal, 3>> NewBox = MakeUnique<TBox<FReal, 3>>(Min, Max);
+					NewGeometry.Emplace(MoveTemp(NewBox));
+
+					bSuccess = true;
+
+					break;
+				}
+				case ImplicitObjectType::Capsule:
+				{
+					ensure(ScaleMode == EScaleMode::LockedXY || ScaleMode == EScaleMode::LockedXYZ);
+
+
+
+					FReal ScaleRadius = FMath::Max(AdjustedScale3DAbs.X, AdjustedScale3DAbs.Y);
+					FReal ScaleLength = AdjustedScale3DAbs.Z;
+
+					FKSphylElem* SphylElem = ShapeElem->GetShapeCheck<FKSphylElem>();
+
+					if (!CHAOS_ENSURE(!IsScaled(ImplicitType)))
+					{
+						// No support for ScaledImplicit Capsule yet
+						break;
+					}
+
+					const TCapsule<FReal> * CapsuleGeometry = static_cast<const TCapsule<FReal>*>(&ImplicitObject);
+
+
+					const FReal InitialHeight = SphylElem->Radius * 2.0f + SphylElem->Length;
+					FReal Radius = FMath::Max(SphylElem->Radius * ScaleRadius, 0.1f);
+					FReal HalfHeight = (SphylElem->Length * 0.5f + SphylElem->Radius) * ScaleLength;
+					Radius = FMath::Min(Radius, HalfHeight);	//radius is capped by half length
+					Radius = FMath::Max(Radius, FCollisionShape::MinCapsuleRadius());
+					FReal HalfLength = HalfHeight - Radius;
+					HalfLength = FMath::Max(FCollisionShape::MinCapsuleAxisHalfHeight(), HalfLength);
+
+
+					// TODO: For Transformed implicit, do not bake this in. Set Transform instead.
+					FVec3 Center = RelativeTM.TransformPosition(SphylElem->Center) * AdjustedScale3D;
+					const FVec3 Axis = SphylElem->Rotation.RotateVector(Chaos::TVector<float, 3>(0, 0, 1));
+
+					const FVec3 X1 = Center - HalfLength * Axis;
+					const FVec3 X2 = Center + HalfLength * Axis;
+
+
+					TUniquePtr<TCapsule<FReal>> NewCapsule =  MakeUnique<TCapsule<FReal>>(X1, X2, Radius);
+					NewGeometry.Emplace(MoveTemp(NewCapsule));
+
+					bSuccess = true;
+
+					break;
+				}
+				case ImplicitObjectType::Convex:
+				{
+
+					if (!CHAOS_ENSURE(IsScaled(ImplicitType)))
+					{
+						// Currently assuming all Convex are scaled implicits.
+						break;
+					}
+
+					// Get Convex
+					const TImplicitObjectScaled<FConvex>* ConvexGeometry = static_cast<const TImplicitObjectScaled<FConvex>*>(&ImplicitObject);
+
+					FKConvexElem* ConvexElem = ShapeElem->GetShapeCheck<FKConvexElem>();
+					const TUniquePtr<FConvex>& ConvexImplicit = ConvexElem->GetChaosConvexMesh();
+
+					// Ensure no rotation/translation in relative transform. PhysX supports this, we currently do not.
+					CHAOS_ENSURE(RelativeTM.GetRotation() == FQuat::Identity);
+					CHAOS_ENSURE(RelativeTM.GetTranslation() == FVector(0, 0, 0));
+
+
+					TUniquePtr<TImplicitObjectScaled<FConvex>> NewConvex = MakeUnique<TImplicitObjectScaled<FConvex>>(MakeSerializable(ConvexImplicit), AdjustedScale3D);
+					NewGeometry.Emplace(MoveTemp(NewConvex));
+
+					bSuccess = true;
+
+					break;
+				}
+				case ImplicitObjectType::TriangleMesh:
+				{
+					if (!CHAOS_ENSURE(IsScaled(ImplicitType)))
+					{
+						// Currently assuming all TriangleMesh are scaled implicits.
+						break;
+					}
+
+
+					// PhysX supports translation, we currently do not.
+					CHAOS_ENSURE(RelativeTM.GetTranslation() == FVector(0, 0, 0));
+
+					const TImplicitObjectScaled<TTriangleMeshImplicitObject<FReal>>* ScaledTriangleMesh = (static_cast<const TImplicitObjectScaled<TTriangleMeshImplicitObject<FReal>>*>(&ImplicitObject));
+					auto UnscaledTriangleMesh = ScaledTriangleMesh->Object();
+
+					TUniquePtr<TImplicitObjectScaled<TTriangleMeshImplicitObject<FReal>>> NewTriangleMesh = MakeUnique<TImplicitObjectScaled<TTriangleMeshImplicitObject<FReal>>>(UnscaledTriangleMesh, AdjustedScale3D);
+					NewGeometry.Emplace(MoveTemp(NewTriangleMesh));
+
+					bSuccess = true;
+
+					break;
+				}
+				default:
+				{
+					CHAOS_ENSURE(false);
+					UE_LOG(LogPhysics, Error, TEXT("Unimplemented ImplicitObject skipped in UpdateBodyScale."));
+				}
+			}// end switch
+		}
+
+		// Many types not yet implemented for UpdateBodyScale. If any geometry is missing from array, we cannot update geometry without losing data.
+		// Only follow through with update if all shapes succeeded.
+		if (CHAOS_ENSURE(NewGeometry.Num() == Shapes.Num()))
+		{
+			ActorHandle->SetGeometry(MakeUnique<Chaos::TImplicitObjectUnion<FReal, 3>>(MoveTemp(NewGeometry)));
+		}
+		else
+		{
+			bSuccess = false;
+		}
+	});
+
+	if (bSuccess)
+	{
+		Scale3D = UpdatedScale3D;
+
+		FPhysScene_Chaos& Scene = GetPhysicsScene()->GetScene();
+		Scene.UpdateActorInAccelerationStructure(ActorHandle);
+
+		// update mass if required
+		if (bUpdateMassWhenScaleChanges)
+		{
+			UpdateMassProperties();
+		}
+	}
+
+#elif WITH_PHYSX
 	//Get all shapes
 	EScaleMode::Type ScaleMode = EScaleMode::Free;
 
@@ -1648,9 +1853,6 @@ bool FBodyInstance::UpdateBodyScale(const FVector& InScale3D, bool bForceUpdate)
 			FKShapeElem* ShapeElem = FPhysxUserData::Get<FKShapeElem>(FPhysicsInterface::GetUserData(Shape));
 			const FTransform& RelativeTM = GetRelativeBodyTransform(Shape);
 
-#if WITH_CHAOS
-			//check(false);
-#else
 			FPhysicsGeometryCollection GeoCollection = FPhysicsInterface::GetGeometryCollection(Shape);
 
 			switch (GeomType)
@@ -1810,7 +2012,6 @@ bool FBodyInstance::UpdateBodyScale(const FVector& InScale3D, bool bForceUpdate)
 						   UE_LOG(LogPhysics, Error, TEXT("Unknown geom type."));
 				}
 			}// end switch
-#endif
 
 			if (UpdatedGeometry)
 			{
@@ -1834,8 +2035,6 @@ bool FBodyInstance::UpdateBodyScale(const FVector& InScale3D, bool bForceUpdate)
 			}
 		}
 	});
-	
-#endif
 
 	// if success, overwrite old Scale3D, otherwise, just don't do it. It will have invalid scale next time
 	if (bSuccess)
@@ -1848,6 +2047,7 @@ bool FBodyInstance::UpdateBodyScale(const FVector& InScale3D, bool bForceUpdate)
 			UpdateMassProperties();
 		}
 	}
+#endif
 
 	return bSuccess;
 }
@@ -2541,6 +2741,13 @@ void FBodyInstance::UpdateMassProperties()
 {
 	UPhysicalMaterial* PhysMat = GetSimplePhysicalMaterial();
 
+#if WITH_CHAOS
+	if (ActorHandle->CastToRigidParticle() != nullptr)
+	{
+		CHAOS_ENSURE(false);
+	}
+#endif
+
 #if WITH_PHYSX
 	if (FPhysicsInterface::IsValid(ActorHandle) && FPhysicsInterface::IsRigidBody(ActorHandle))
 	{
diff --git a/Engine/Source/Runtime/Engine/Private/PhysicsEngine/Experimental/PhysScene_Chaos.cpp b/Engine/Source/Runtime/Engine/Private/PhysicsEngine/Experimental/PhysScene_Chaos.cpp
--- a/Engine/Source/Runtime/Engine/Private/PhysicsEngine/Experimental/PhysScene_Chaos.cpp
+++ b/Engine/Source/Runtime/Engine/Private/PhysicsEngine/Experimental/PhysScene_Chaos.cpp
@@ -585,6 +585,38 @@ void FPhysScene_Chaos::RemoveActorFromAccelerationStructure(FPhysicsActorHandle&
 #endif
 }
 
+void FPhysScene_Chaos::UpdateActorInAccelerationStructure(FPhysicsActorHandle& Actor)
+{
+#if WITH_CHAOS
+	using namespace Chaos;
+
+	if (GetSpacialAcceleration())
+	{
+		ExternalDataLock.WriteLock();
+
+		auto SpatialAcceleration = GetSpacialAcceleration();
+
+		if (SpatialAcceleration)
+		{
+
+			TBox<FReal, 3> WorldBounds;
+			const bool bHasBounds = Actor->Geometry()->HasBoundingBox();
+			if (bHasBounds)
+			{
+				WorldBounds = Actor->Geometry()->BoundingBox().GetAABB().TransformedAABB(TRigidTransform<FReal, 3>(Actor->X(), Actor->R()));
+			}
+
+
+			Chaos::TAccelerationStructureHandle<float, 3> AccelerationHandle(Actor);
+			SpatialAcceleration->UpdateElementIn(AccelerationHandle, WorldBounds, bHasBounds, Actor->SpatialIdx());
+		}
+
+		ExternalDataLock.WriteUnlock();
+	}
+#endif
+}
+
+
 template<typename ObjectType>
 void RemovePhysicsProxy(ObjectType* InObject, Chaos::FPhysicsSolver* InSolver, FChaosSolversModule* InModule)
 {
diff --git a/Engine/Source/Runtime/Engine/Public/Physics/Experimental/PhysScene_Chaos.h b/Engine/Source/Runtime/Engine/Public/Physics/Experimental/PhysScene_Chaos.h
--- a/Engine/Source/Runtime/Engine/Public/Physics/Experimental/PhysScene_Chaos.h
+++ b/Engine/Source/Runtime/Engine/Public/Physics/Experimental/PhysScene_Chaos.h
@@ -126,6 +126,7 @@ public:
 	void RemoveObject(FFieldSystemPhysicsProxy* InObject);	
 
 	void RemoveActorFromAccelerationStructure(FPhysicsActorHandle& Actor);
+	void UpdateActorInAccelerationStructure(FPhysicsActorHandle& Actor);
 
 #if XGE_FIXED
 	template<typename PayloadType>
diff --git a/Engine/Source/Runtime/Experimental/Chaos/Private/Chaos/GeometryParticles.cpp b/Engine/Source/Runtime/Experimental/Chaos/Private/Chaos/GeometryParticles.cpp
--- a/Engine/Source/Runtime/Experimental/Chaos/Private/Chaos/GeometryParticles.cpp
+++ b/Engine/Source/Runtime/Experimental/Chaos/Private/Chaos/GeometryParticles.cpp
@@ -14,13 +14,19 @@ namespace Chaos
 		{
 			if(const auto* Union = Geometry->template GetObject<TImplicitObjectUnion<T, d>>())
 			{
+				const int32 OldShapeNum = ShapesArray.Num();
+
 				ShapesArray.SetNum(Union->GetObjects().Num());
-				int32 Inner = 0;
-				for(const auto& Geom : Union->GetObjects())
+
+				for (int32 ShapeIndex = 0; ShapeIndex < ShapesArray.Num(); ++ShapeIndex)
 				{
-					ShapesArray[Inner] = TPerShapeData<T, d>::CreatePerShapeData();
-					ShapesArray[Inner]->Geometry = MakeSerializable(Geom);
-					++Inner;
+					if (ShapeIndex >= OldShapeNum)
+					{
+						// If newly allocated shape, initialize it.
+						ShapesArray[ShapeIndex] = TPerShapeData<T, d>::CreatePerShapeData();
+					}
+
+					ShapesArray[ShapeIndex]->Geometry = MakeSerializable(Union->GetObjects()[ShapeIndex]);
 				}
 			}
 			else
@@ -46,7 +52,12 @@ namespace Chaos
 
 	template <typename T, int d>
 	TPerShapeData<T, d>::TPerShapeData()
-		: UserData(nullptr)
+		: QueryData()
+		, SimData()
+		, UserData(nullptr)
+		, Geometry()
+		, WorldSpaceInflatedShapeBounds(TAABB<FReal, 3>(FVec3(0), FVec3(0)))
+		, Materials()
 		, bDisable(false)
 	{
 	}
diff --git a/Engine/Source/Runtime/Experimental/ChaosSolvers/Private/PhysicsProxy/SingleParticlePhysicsProxy.cpp b/Engine/Source/Runtime/Experimental/ChaosSolvers/Private/PhysicsProxy/SingleParticlePhysicsProxy.cpp
--- a/Engine/Source/Runtime/Experimental/ChaosSolvers/Private/PhysicsProxy/SingleParticlePhysicsProxy.cpp
+++ b/Engine/Source/Runtime/Experimental/ChaosSolvers/Private/PhysicsProxy/SingleParticlePhysicsProxy.cpp
@@ -245,7 +245,7 @@ void FSingleParticlePhysicsProxy<Chaos::TPBDRigidParticle<float, 3>>::PushToPhys
 			GetSolver()->GetEvolution()->GetGravityForces().SetEnabled(*RigidHandle, Data->MGravityEnabled);
 		}
 
-		if (Data->DirtyFlags.IsDirty(Chaos::EParticleFlags::X | Chaos::EParticleFlags::R | Chaos::EParticleFlags::V))
+		if (Data->DirtyFlags.IsDirty(Chaos::EParticleFlags::X | Chaos::EParticleFlags::R | Chaos::EParticleFlags::V | Chaos::EParticleFlags::Geometry))
 		{
 			if (Data->Geometry && Data->Geometry->HasBoundingBox())
 			{
