Add additional timings to try and identify high sort time during cooking
#rnx

diff --git a/Engine/Source/Runtime/Engine/Private/Level.cpp b/Engine/Source/Runtime/Engine/Private/Level.cpp
--- a/Engine/Source/Runtime/Engine/Private/Level.cpp
+++ b/Engine/Source/Runtime/Engine/Private/Level.cpp
@@ -818,6 +818,8 @@ static void SortActorsHierarchy(TArray<AActor*>& Actors, UObject* Level)
 		}
 	}
 
+	const double CalcAttachDepthTime = FPlatformTime::Seconds() - StartTime;
+
 	auto DepthSorter = [&DepthMap](AActor* A, AActor* B)
 	{
 		const int32 DepthA = A ? DepthMap.FindRef(A) : MAX_int32;
@@ -825,12 +827,14 @@ static void SortActorsHierarchy(TArray<AActor*>& Actors, UObject* Level)
 		return DepthA < DepthB;
 	};
 
+	const double StableSortStartTime = FPlatformTime::Seconds();
 	StableSortInternal(Actors.GetData(), Actors.Num(), DepthSorter);
+	const double StableSortTime = FPlatformTime::Seconds() - StableSortStartTime;
 
-	const float ElapsedTime = (float)(FPlatformTime::Seconds() - StartTime);
-	if (ElapsedTime > 1.0f)
+	const double ElapsedTime = FPlatformTime::Seconds() - StartTime;
+	if (ElapsedTime > 1.0)
 	{
-		UE_LOG(LogLevel, Warning, TEXT("SortActorsHierarchy(%s) took %f seconds"), Level ? *GetNameSafe(Level->GetOutermost()) : TEXT("??"), ElapsedTime);
+		UE_LOG(LogLevel, Warning, TEXT("SortActorsHierarchy(%s) took %f seconds (CalcAttachDepth: %f StableSort: %f)"), Level ? *GetNameSafe(Level->GetOutermost()) : TEXT("??"), ElapsedTime, CalcAttachDepthTime, StableSortTime);
 	}
 
 	// Since all the null entries got sorted to the end, lop them off right now
