Making the CameraModifier capable of being updated from native code, rather than relying on a blueprint implementation to do anything.
Adding some additional constructors to FWeightedBlendables.

Bob.Tellez

diff --git a/Engine/Source/Runtime/Engine/Classes/Camera/CameraModifier.h b/Engine/Source/Runtime/Engine/Classes/Camera/CameraModifier.h
--- a/Engine/Source/Runtime/Engine/Classes/Camera/CameraModifier.h
+++ b/Engine/Source/Runtime/Engine/Classes/Camera/CameraModifier.h
@@ -112,7 +112,6 @@ public:
 	UFUNCTION(BlueprintCallable, Category = CameraModifier)
 	virtual AActor* GetViewTarget() const;
 
-
 	/** 
 	 *  Disables this modifier.
 	 *  @param  bImmediate  - true to disable with no blend out, false (default) to allow blend out
@@ -149,7 +148,11 @@ public:
 
 	/** @return Returns the appropriate world context for this object. */
 	UWorld* GetWorld() const;
-};
-
 
+protected:
+	/** Allows modifying the camera in native code. */
+	virtual void ModifyCamera(float DeltaTime, FVector ViewLocation, FRotator ViewRotation, float FOV, FVector& NewViewLocation, FRotator& NewViewRotation, float& NewFOV);
 
+	/** Allows modifying the post process in native code. */
+	virtual void ModifyPostProcess(float DeltaTime, float& PostProcessBlendWeight, FPostProcessSettings& PostProcessSettings);
+};
diff --git a/Engine/Source/Runtime/Engine/Classes/Engine/Scene.h b/Engine/Source/Runtime/Engine/Classes/Engine/Scene.h
--- a/Engine/Source/Runtime/Engine/Classes/Engine/Scene.h
+++ b/Engine/Source/Runtime/Engine/Classes/Engine/Scene.h
@@ -581,7 +581,12 @@ struct FWeightedBlendable
 USTRUCT(BlueprintType)
 struct FWeightedBlendables
 {
-	GENERATED_USTRUCT_BODY()
+	GENERATED_BODY()
+
+public:
+
+	FWeightedBlendables() { }
+	FWeightedBlendables(const TArray<FWeightedBlendable>& InArray) : Array(InArray) { }
 
 	UPROPERTY(EditAnywhere, BlueprintReadWrite, Category="PostProcessSettings", meta=( Keywords="PostProcess" ))
 	TArray<FWeightedBlendable> Array;
diff --git a/Engine/Source/Runtime/Engine/Private/Camera/CameraModifier.cpp b/Engine/Source/Runtime/Engine/Private/Camera/CameraModifier.cpp
--- a/Engine/Source/Runtime/Engine/Private/Camera/CameraModifier.cpp
+++ b/Engine/Source/Runtime/Engine/Private/Camera/CameraModifier.cpp
@@ -22,6 +22,8 @@ bool UCameraModifier::ModifyCamera(float DeltaTime, FMinimalViewInfo& InOutPOV)
 	// Update the alpha
 	UpdateAlpha(DeltaTime);
 
+	// let native do what it wants
+	ModifyCamera(DeltaTime, InOutPOV.Location, InOutPOV.Rotation, InOutPOV.FOV, InOutPOV.Location, InOutPOV.Rotation, InOutPOV.FOV);
 	// let BP do what it wants
 	BlueprintModifyCamera(DeltaTime, InOutPOV.Location, InOutPOV.Rotation, InOutPOV.FOV, InOutPOV.Location, InOutPOV.Rotation, InOutPOV.FOV);
 
@@ -30,12 +32,20 @@ bool UCameraModifier::ModifyCamera(float DeltaTime, FMinimalViewInfo& InOutPOV)
 		// note: pushing these through the cached PP blend system in the camera to get
 		// proper layered blending, rather than letting subsequent mods stomp over each other in the 
 		// InOutPOV struct.
-		float PPBlendWeight = 0.f;
-		FPostProcessSettings PPSettings;
-		BlueprintModifyPostProcess(DeltaTime, PPBlendWeight, PPSettings);
-		if (PPBlendWeight > 0.f)
 		{
-			CameraOwner->AddCachedPPBlend(PPSettings, PPBlendWeight);
+			float PPBlendWeight = 0.f;
+			FPostProcessSettings PPSettings;
+			
+			//  Let native code modify the post process settings.
+			ModifyPostProcess(DeltaTime, PPBlendWeight, PPSettings);
+			
+			// Allow blueprint to modify the post process settings.
+			BlueprintModifyPostProcess(DeltaTime, PPBlendWeight, PPSettings);
+			
+			if (PPBlendWeight > 0.f)
+			{
+				CameraOwner->AddCachedPPBlend(PPSettings, PPBlendWeight);
+			}
 		}
 	}
 
@@ -87,7 +97,6 @@ AActor* UCameraModifier::GetViewTarget() const
 	return CameraOwner ? CameraOwner->GetViewTarget() : nullptr;
 }
 
-
 void UCameraModifier::AddedToCamera( APlayerCameraManager* Camera ) 
 {
 	CameraOwner = Camera;
@@ -109,7 +118,6 @@ void UCameraModifier::DisableModifier(bool bImmediate)
 	{
 		bPendingDisable = true;
 	}
-
 }
 
 void UCameraModifier::EnableModifier()
@@ -135,3 +143,12 @@ bool UCameraModifier::ProcessViewRotation( AActor* ViewTarget, float DeltaTime,
 	return false;
 }
 
+void UCameraModifier::ModifyCamera(float DeltaTime, FVector ViewLocation, FRotator ViewRotation, float FOV, FVector& NewViewLocation, FRotator& NewViewRotation, float& NewFOV)
+{
+
+}
+
+void UCameraModifier::ModifyPostProcess(float DeltaTime, float& PostProcessBlendWeight, FPostProcessSettings& PostProcessSettings)
+{
+
+}
\ No newline at end of file
