Chaos - RBAN Minimal Evolution

diff --git a/Engine/Source/Runtime/Engine/Private/PhysicsEngine/ImmediatePhysics/ImmediatePhysicsChaos/ImmediatePhysicsActorHandle_Chaos.cpp b/Engine/Source/Runtime/Engine/Private/PhysicsEngine/ImmediatePhysics/ImmediatePhysicsChaos/ImmediatePhysicsActorHandle_Chaos.cpp
--- a/Engine/Source/Runtime/Engine/Private/PhysicsEngine/ImmediatePhysics/ImmediatePhysicsChaos/ImmediatePhysicsActorHandle_Chaos.cpp
+++ b/Engine/Source/Runtime/Engine/Private/PhysicsEngine/ImmediatePhysics/ImmediatePhysicsChaos/ImmediatePhysicsActorHandle_Chaos.cpp
@@ -4,6 +4,7 @@
 
 #include "Chaos/Box.h"
 #include "Chaos/Capsule.h"
+#include "Chaos/Evolution/PBDMinEvolution.h"
 #include "Chaos/MassProperties.h"
 #include "Chaos/Particle/ParticleUtilities.h"
 #include "Chaos/ParticleHandle.h"
@@ -211,28 +212,28 @@ namespace ImmediatePhysics_Chaos
 	// Actor Handle
 	//
 
-	FActorHandle::FActorHandle(Chaos::TPBDRigidsEvolutionGBF<FReal, Dimensions>* InEvolution, EActorType ActorType, FBodyInstance* BodyInstance, const FTransform& Transform)
-		: Evolution(InEvolution)
+	FActorHandle::FActorHandle(Chaos::TPBDRigidsSOAs<FReal, 3>& InParticles, EActorType ActorType, FBodyInstance* BodyInstance, const FTransform& Transform)
+		: Particles(InParticles)
 		, ParticleHandle(nullptr)
 	{
 		using namespace Chaos;
 
 		// @todo(ccaulfield): Scale
 		float Mass = 0;
-		TVector<float, 3> Inertia = TVector<float, 3>::OneVector;
-		TRigidTransform<float, 3> CoMTransform = TRigidTransform<float, 3>::Identity;
+		FVec3 Inertia = FVec3::OneVector;
+		FRigidTransform3 CoMTransform = FRigidTransform3::Identity;
 		if (CreateGeometry(BodyInstance, FVector::OneVector, Mass, Inertia, CoMTransform, Geometry, Shapes))
 		{
 			switch (ActorType)
 			{
 			case EActorType::StaticActor:
-				ParticleHandle = Evolution->CreateStaticParticles(1, TGeometryParticleParameters<FReal, Dimensions>())[0];
+				ParticleHandle = Particles.CreateStaticParticles(1, TGeometryParticleParameters<FReal, Dimensions>())[0];
 				break;
 			case EActorType::KinematicActor:
-				ParticleHandle = Evolution->CreateKinematicParticles(1, TKinematicGeometryParticleParameters<FReal, Dimensions>())[0];
+				ParticleHandle = Particles.CreateKinematicParticles(1, TKinematicGeometryParticleParameters<FReal, Dimensions>())[0];
 				break;
 			case EActorType::DynamicActor:
-				ParticleHandle = Evolution->CreateDynamicParticles(1, TPBDRigidParticleParameters<FReal, Dimensions>())[0];
+				ParticleHandle = Particles.CreateDynamicParticles(1, TPBDRigidParticleParameters<FReal, Dimensions>())[0];
 				break;
 			}
 
@@ -242,6 +243,13 @@ namespace ImmediatePhysics_Chaos
 
 				ParticleHandle->SetGeometry(MakeSerializable(Geometry));
 
+				if (Geometry && Geometry->HasBoundingBox())
+				{
+					ParticleHandle->SetHasBounds(true);
+					ParticleHandle->SetLocalBounds(Geometry->BoundingBox());
+					ParticleHandle->SetWorldSpaceInflatedBounds(Geometry->BoundingBox().TransformedBox(TRigidTransform<float, 3>(ParticleHandle->X(), ParticleHandle->R())));
+				}
+
 				if (auto* Kinematic = ParticleHandle->CastToKinematicParticle())
 				{
 					Kinematic->SetV(FVector::ZeroVector);
@@ -269,7 +277,7 @@ namespace ImmediatePhysics_Chaos
 	{
 		if (ParticleHandle != nullptr)
 		{
-			Evolution->DestroyParticle(ParticleHandle);
+			Particles.DestroyParticle(ParticleHandle);
 			ParticleHandle = nullptr;
 			Geometry = nullptr;
 		}
diff --git a/Engine/Source/Runtime/Engine/Private/PhysicsEngine/ImmediatePhysics/ImmediatePhysicsChaos/ImmediatePhysicsSimulation_Chaos.cpp b/Engine/Source/Runtime/Engine/Private/PhysicsEngine/ImmediatePhysics/ImmediatePhysicsChaos/ImmediatePhysicsSimulation_Chaos.cpp
--- a/Engine/Source/Runtime/Engine/Private/PhysicsEngine/ImmediatePhysics/ImmediatePhysicsChaos/ImmediatePhysicsSimulation_Chaos.cpp
+++ b/Engine/Source/Runtime/Engine/Private/PhysicsEngine/ImmediatePhysics/ImmediatePhysicsChaos/ImmediatePhysicsSimulation_Chaos.cpp
@@ -82,11 +82,13 @@ FAutoConsoleVariableRef CVarChaosImmPhysMaxInertiaRatio(TEXT("p.Chaos.ImmPhys.Jo
 #if UE_BUILD_DEBUG
 int32 ChaosImmediate_DebugDrawParticles = 0;
 int32 ChaosImmediate_DebugDrawShapes = 1;
+int32 ChaosImmediate_DebugDrawBounds = 0;
 int32 ChaosImmediate_DebugDrawCollisions = 0;
 int32 ChaosImmediate_DebugDrawJoints = 1;
 #else
 int32 ChaosImmediate_DebugDrawParticles = 0;
 int32 ChaosImmediate_DebugDrawShapes = 0;
+int32 ChaosImmediate_DebugDrawBounds = 0;
 int32 ChaosImmediate_DebugDrawCollisions = 0;
 int32 ChaosImmediate_DebugDrawJoints = 0;
 #endif
@@ -99,6 +101,7 @@ int32 ChaosImmediate_DebugDrawJointFeatures = 0;
 
 FAutoConsoleVariableRef CVarChaosImmPhysDebugDrawParticles(TEXT("p.Chaos.ImmPhys.DebugDrawParticles"), ChaosImmediate_DebugDrawParticles, TEXT("Draw Particle Transforms (0 = never; 1 = end of frame; 2 = begin and end of frame; 3 = post-integate, post-apply and post-applypushout;)."));
 FAutoConsoleVariableRef CVarChaosImmPhysDebugDrawShapes(TEXT("p.Chaos.ImmPhys.DebugDrawShapes"), ChaosImmediate_DebugDrawShapes, TEXT("Draw Shapes (0 = never; 1 = end of frame; 2 = begin and end of frame; 3 = post-integate, post-apply and post-applypushout;"));
+FAutoConsoleVariableRef CVarChaosImmPhysDebugDrawBounds(TEXT("p.Chaos.ImmPhys.DebugDrawBounds"), ChaosImmediate_DebugDrawBounds, TEXT("Draw Particle Bounds (0 = never; 1 = end of frame; 2 = begin and end of frame; 3 = post-integate, post-apply and post-applypushout;)."));
 FAutoConsoleVariableRef CVarChaosImmPhysDebugDrawCollisions(TEXT("p.Chaos.ImmPhys.DebugDrawCollisions"), ChaosImmediate_DebugDrawCollisions, TEXT("Draw Collisions (0 = never; 1 = end of frame; 2 = begin and end of frame; 3 = post-integate, post-apply and post-applypushout;)"));
 FAutoConsoleVariableRef CVarChaosImmPhysDebugDrawJoints(TEXT("p.Chaos.ImmPhys.DebugDrawJoints"), ChaosImmediate_DebugDrawJoints, TEXT("Draw Joints. (0 = never; 1 = end of frame; 2 = begin and end of frame; 3 = post-integate, post-apply and post-applypushout; 4 = each Apply step)."));
 FAutoConsoleVariableRef CVarChaosImmPhysDebugDrawJointFeatures(TEXT("p.Chaos.ImmPhys.DebugDrawJointFeatures"), ChaosImmediate_DebugDrawJointFeatures, TEXT("Joint features mask (see EDebugDrawJointFeature)."));
@@ -148,61 +151,46 @@ namespace ImmediatePhysics_Chaos
 	//
 
 	FSimulation::FSimulation()
-		: NumActiveActorHandles(0)
+		: Particles()
+		, Joints()
+		, Collisions(Particles, CollidedParticles, ParticleMaterials, 0, 0)
+		, BroadPhase(PotentiallyCollidingPairs)
+		, CollisionDetector(BroadPhase, Collisions)
+		, JointsRule(0, Joints)
+		, CollisionsRule(1, Collisions)
+		, Evolution(Particles, CollisionDetector)
+		, NumActiveActorHandles(0)
 		, SimulationSpaceTransform(FTransform::Identity)
 	{
 		using namespace Chaos;
 
-		Particles = MakeUnique<TPBDRigidsSOAs<FReal, Dimensions>>();
-		Joints = MakeUnique<FPBDJointConstraints>(FPBDJointSolverSettings());
-		JointsRule = MakeUnique<TPBDConstraintIslandRule<FPBDJointConstraints>>(*Joints);
-		Evolution = MakeUnique<TPBDRigidsEvolutionGBF<FReal, Dimensions>>(*Particles.Get(), 1);
-		TPBDCollisionConstraints<FReal, Dimensions>& Collisions = Evolution->GetCollisionConstraints();
+		Particles.GetParticleHandles().AddArray(&ParticleMaterials);
+		Particles.GetParticleHandles().AddArray(&PerParticleMaterials);
 
-		Evolution->AddConstraintRule(JointsRule.Get());
-
-		// We want collisions to override joints
-		JointsRule->SetPriority(0);
-		Evolution->GetCollisionConstraintsRule().SetPriority(1);
-
-		// Filter collisions after detection
-		// @todo(ccaulfield): Eventually we will build lists of potentially colliding pairs and won't need this
-		Evolution->SetPostDetectCollisionsCallback(
-			[this]()
-			{
-				Evolution->GetCollisionConstraints().ApplyCollisionModifier(
-					[this](const TPBDCollisionConstraints<FReal, Dimensions>::FConstraintContainerHandle* Handle)
-					{
-						TVector<const TGeometryParticleHandle<FReal, Dimensions>*, 2> ConstrainedParticles = Handle->GetConstrainedParticles();
-						if (ShouldIgnoreCollisionConstraint(ConstrainedParticles[0], ConstrainedParticles[1], IgnoreCollisionParticlePairTable))
-						{
-							return ECollisionModifierResult::Disabled;
-						}
-						return ECollisionModifierResult::Unchanged;
-					});
-			});
+		Evolution.AddConstraintRule(&CollisionsRule);
+		Evolution.AddConstraintRule(&JointsRule);
 
 #if CHAOS_DEBUG_DRAW
-		Evolution->SetPostIntegrateCallback(
+		Evolution.SetPostIntegrateCallback(
 			[this]()
 			{
 				// Dynamic only - Kinematics get drawn once at the end of frame
 				DebugDrawParticles(3, 3, 0.3f, false, true);
 				DebugDrawConstraints(3, 3, 0.3f);
 			});
-		Evolution->SetPostApplyCallback(
-			[this](int32 Island)
+		Evolution.SetPostApplyCallback(
+			[this]()
 			{
 				// Dynamic only - Kinematics get drawn once at the end of frame
-				DebugDrawIslandParticles(Island, 3, 3, 0.6f, false, true);
-				DebugDrawIslandConstraints(Island, 3, 3, 0.6f);
+				DebugDrawParticles(3, 3, 0.6f, false, true);
+				DebugDrawConstraints(3, 3, 0.6f);
 			});
-		Evolution->SetPostApplyPushOutCallback(
-			[this](int32 Island)
+		Evolution.SetPostApplyPushOutCallback(
+			[this]()
 			{
 				// Dynamic only - Kinematics get drawn once at the end of frame
-				DebugDrawIslandParticles(Island, 3, 3, 1.0f, false, true);
-				DebugDrawIslandConstraints(Island, 3, 3, 1.0f);
+				DebugDrawParticles(3, 3, 1.0f, false, true);
+				DebugDrawConstraints(3, 3, 1.0f);
 			});
 		Collisions.SetPostApplyCallback(
 			[this](const float Dt, const TArray<TPBDCollisionConstraintHandle<float, 3>*>& InConstraintHandles)
@@ -220,7 +208,7 @@ namespace ImmediatePhysics_Chaos
 					DebugDraw::DrawCollisions(SimulationSpaceTransform, InConstraintHandles, 0.6f);
 				}
 			});
-		Joints->SetPreApplyCallback(
+		Joints.SetPreApplyCallback(
 			[this](const float Dt, const TArray<FPBDJointConstraintHandle*>& InConstraintHandles)
 			{
 				if (ChaosImmediate_DebugDrawJoints == 4)
@@ -228,7 +216,7 @@ namespace ImmediatePhysics_Chaos
 					DebugDraw::DrawJointConstraints(SimulationSpaceTransform, InConstraintHandles, 0.3f);
 				}
 			});
-		Joints->SetPostApplyCallback(
+		Joints.SetPostApplyCallback(
 			[this](const float Dt, const TArray<FPBDJointConstraintHandle*>& InConstraintHandles)
 			{
 				if (ChaosImmediate_DebugDrawJoints == 4)
@@ -243,12 +231,6 @@ namespace ImmediatePhysics_Chaos
 	{
 		using namespace Chaos;
 
-		//Evolution->ClearDetectCollisionsCallback();
-		Evolution->GetCollisionConstraints().ClearPostApplyCallback();
-		Evolution->GetCollisionConstraints().ClearPostApplyPushOutCallback();
-
-		//Evolution->RemoveConstraintRule(JointsRule.Get());
-
 		for (FActorHandle* ActorHandle : ActorHandles)
 		{
 			delete ActorHandle;
@@ -260,11 +242,6 @@ namespace ImmediatePhysics_Chaos
 			delete JointHandle;
 		}
 		JointHandles.Empty();
-
-		Evolution.Reset();
-		JointsRule.Reset();
-		Joints.Reset();
-		Particles.Reset();
 	}
 
 	FActorHandle* FSimulation::CreateStaticActor(FBodyInstance* BodyInstance)
@@ -284,21 +261,32 @@ namespace ImmediatePhysics_Chaos
 
 	FActorHandle* FSimulation::CreateActor(EActorType ActorType, FBodyInstance* BodyInstance, const FTransform& Transform)
 	{
-		FActorHandle* ActorHandle = new FActorHandle(Evolution.Get(), ActorType, BodyInstance, Transform);
+		using namespace Chaos;
+
+		FActorHandle* ActorHandle = new FActorHandle(Particles, ActorType, BodyInstance, Transform);
 		int ActorIndex = ActorHandles.Add(ActorHandle);
+
+		// @todo(ccaulfield): add materials
+		TUniquePtr<FChaosPhysicsMaterial> Material = MakeUnique<FChaosPhysicsMaterial>();
+		ParticleMaterials.Add(MakeSerializable(Material));
+		PerParticleMaterials.Add(MoveTemp(Material));
+
 		return ActorHandle;
 	}
 
 	void FSimulation::DestroyActor(FActorHandle* ActorHandle)
 	{
 		// @todo(ccaulfield): FActorHandle could remember its index to optimize this
-		ActorHandles.Remove(ActorHandle);
+		int32 Index = ActorHandles.Remove(ActorHandle);
 		delete ActorHandle;
+
+		ParticleMaterials.RemoveAt(Index, 1);
+		PerParticleMaterials.RemoveAt(Index, 1);
 	}
 
 	FJointHandle* FSimulation::CreateJoint(FConstraintInstance* ConstraintInstance, FActorHandle* Body1, FActorHandle* Body2)
 	{
-		FJointHandle* JointHandle = new FJointHandle(Joints.Get(), ConstraintInstance, Body1, Body2);
+		FJointHandle* JointHandle = new FJointHandle(&Joints, ConstraintInstance, Body1, Body2);
 		JointHandles.Add(JointHandle);
 		return JointHandle;
 	}
@@ -350,23 +338,23 @@ namespace ImmediatePhysics_Chaos
 			IgnoreCollisionParticlePairTable.FindOrAdd(ParticleB).Add(ParticleA);
 		}
 
-		TArray<TVector<const TGeometryParticleHandle<FReal, Dimensions>*, 2>> BroadPhaseParticlePairs;
+		PotentiallyCollidingPairs.Empty();
 		int NumActorHandles = ActorHandles.Num();
 		for (int ActorHandleIndex0 = 0; ActorHandleIndex0 < NumActorHandles; ++ActorHandleIndex0)
 		{
 			FActorHandle* ActorHandle0 = ActorHandles[ActorHandleIndex0];
-			const TGeometryParticleHandle<FReal, Dimensions>* Particle0 = ActorHandle0->GetParticle();
+			TGeometryParticleHandle<FReal, Dimensions>* Particle0 = ActorHandle0->GetParticle();
 
 			for (int ActorHandleIndex1 = ActorHandleIndex0 + 1; ActorHandleIndex1 < NumActorHandles; ++ActorHandleIndex1)
 			{
 				FActorHandle* ActorHandle1 = ActorHandles[ActorHandleIndex1];
-				const TGeometryParticleHandle<FReal, Dimensions>* Particle1 = ActorHandle1->GetParticle();
+				TGeometryParticleHandle<FReal, Dimensions>* Particle1 = ActorHandle1->GetParticle();
 
 				const TSet<const FParticleHandle*>* Particle0IgnoreSet = IgnoreCollisionParticlePairTable.Find(Particle0);
 				bool bIgnoreActorHandlePair = (Particle0IgnoreSet != nullptr)  && Particle0IgnoreSet->Contains(Particle1);
 				if (!bIgnoreActorHandlePair)
 				{
-					BroadPhaseParticlePairs.Emplace(TVector<const TGeometryParticleHandle<FReal, Dimensions>*, 2>(Particle0, Particle1));
+					PotentiallyCollidingPairs.Emplace(FParticlePair(Particle0, Particle1));
 				}
 			}
 		}
@@ -464,7 +452,7 @@ namespace ImmediatePhysics_Chaos
 			}
 			UE_LOG(LogChaosJoint, Verbose, TEXT("Simulate dt = %f"), DeltaTime);
 
-			FPBDJointSolverSettings JointsSettings = Joints->GetSettings();
+			FPBDJointSolverSettings JointsSettings = Joints.GetSettings();
 			JointsSettings.ApplyPairIterations = ChaosImmediate_Joint_PairIterations;
 			JointsSettings.ApplyPushOutPairIterations = ChaosImmediate_Joint_PushOutPairIterations;
 			JointsSettings.SwingTwistAngleTolerance = ChaosImmediate_Joint_SwingTwistAngleTolerance;
@@ -482,18 +470,19 @@ namespace ImmediatePhysics_Chaos
 			JointsSettings.SoftLinearStiffness = ChaosImmediate_Joint_SoftLinearStiffness;
 			JointsSettings.SoftAngularStiffness = ChaosImmediate_Joint_SoftAngularStiffness;
 			JointsSettings.DriveStiffness = ChaosImmediate_Joint_DriveStiffness;
-			Joints->SetSettings(JointsSettings);
 
-			Evolution->SetNumIterations(ChaosImmediate_Evolution_Iterations);
-			Evolution->SetNumPushOutIterations(ChaosImmediate_Evolution_PushOutIterations);
+			Joints.SetSettings(JointsSettings);
 
-			Evolution->GetCollisionConstraints().SetThickness(ChaosImmediate_Collision_Thickness);
-			Evolution->GetCollisionConstraints().SetPairIterations(ChaosImmediate_Collision_PairIterations);
-			Evolution->GetCollisionConstraints().SetPushOutPairIterations(ChaosImmediate_Collision_PushOutPairIterations);
-			Evolution->GetCollisionConstraintsRule().SetPriority(ChaosImmediate_Collision_Priority);
+			Evolution.SetNumIterations(ChaosImmediate_Evolution_Iterations);
+			Evolution.SetNumPushOutIterations(ChaosImmediate_Evolution_PushOutIterations);
+
+			Collisions.SetThickness(ChaosImmediate_Collision_Thickness);
+			Collisions.SetPairIterations(ChaosImmediate_Collision_PairIterations);
+			Collisions.SetPushOutPairIterations(ChaosImmediate_Collision_PushOutPairIterations);
+			CollisionsRule.SetPriority(ChaosImmediate_Collision_Priority);
 		
 			// TEMP until we can remove constraints again, or I add broad-phase filtering. (FilterCollisionConstraints will crash since the persistent collision changes)
-			Evolution->GetCollisionConstraints().SetCollisionsEnabled(ChaosImmediate_Collision_Enabled != 0);
+			Collisions.SetCollisionsEnabled(ChaosImmediate_Collision_Enabled != 0);
 		}
 
 		DebugDrawParticles(2, 2, 0.7f, true, true);
@@ -501,11 +490,11 @@ namespace ImmediatePhysics_Chaos
 
 		ConditionConstraints();
 
-		Evolution->GetGravityForces().SetAcceleration(InGravity);
+		Evolution.SetGravity(InGravity);
 		
-		Evolution->Advance(DeltaTime, MaxDeltaTime, MaxSubSteps);
+		Evolution.Advance(DeltaTime, MaxDeltaTime, MaxSubSteps);
 		
-		Evolution->EndFrame(DeltaTime);
+		//Evolution.EndFrame(DeltaTime);
 
 		DebugDrawParticles(1, 2, 1.0f, true, true);
 		DebugDrawConstraints(1, 2, 1.0f);
@@ -520,53 +509,21 @@ namespace ImmediatePhysics_Chaos
 		{
 			if ((ChaosImmediate_DebugDrawParticles >= MinDebugLevel) && (ChaosImmediate_DebugDrawParticles <= MaxDebugLevel))
 			{
-				DebugDraw::DrawParticleTransforms(SimulationSpaceTransform, Evolution->GetParticles().GetAllParticlesView(), ColorScale, bDrawKinematic, bDrawDynamic);
+				DebugDraw::DrawParticleTransforms(SimulationSpaceTransform, Particles.GetAllParticlesView(), ColorScale, bDrawKinematic, bDrawDynamic);
 			}
 			if ((ChaosImmediate_DebugDrawShapes >= MinDebugLevel) && (ChaosImmediate_DebugDrawShapes <= MaxDebugLevel))
 			{
-				DebugDraw::DrawParticleShapes(SimulationSpaceTransform, Evolution->GetParticles().GetAllParticlesView(), ColorScale, bDrawKinematic, bDrawDynamic);
-			}
-		}
-#endif
-	}
-
-	void FSimulation::DebugDrawConstraints(const int32 MinDebugLevel, const int32 MaxDebugLevel, const float ColorScale)
-	{
-#if CHAOS_DEBUG_DRAW
-		using namespace Chaos;
-		if (FDebugDrawQueue::IsDebugDrawingEnabled())
-		{
-			if ((ChaosImmediate_DebugDrawCollisions >= MinDebugLevel) && (ChaosImmediate_DebugDrawCollisions <= MaxDebugLevel))
-			{
-				DebugDraw::DrawCollisions(SimulationSpaceTransform, Evolution->GetCollisionConstraints(), ColorScale);
+				DebugDraw::DrawParticleShapes(SimulationSpaceTransform, Particles.GetAllParticlesView(), ColorScale, bDrawKinematic, bDrawDynamic);
 			}
-			if ((ChaosImmediate_DebugDrawJoints >= MinDebugLevel) && (ChaosImmediate_DebugDrawJoints <= MaxDebugLevel))
+			if ((ChaosImmediate_DebugDrawBounds >= MinDebugLevel) && (ChaosImmediate_DebugDrawBounds <= MaxDebugLevel))
 			{
-				DebugDraw::DrawJointConstraints(SimulationSpaceTransform, *Joints, ColorScale, (uint32)ChaosImmediate_DebugDrawJointFeatures);
+				DebugDraw::DrawParticleBounds(SimulationSpaceTransform, Particles.GetAllParticlesView(), ColorScale, bDrawKinematic, bDrawDynamic);
 			}
 		}
 #endif
 	}
 
-	void FSimulation::DebugDrawIslandParticles(const int32 Island, const int32 MinDebugLevel, const int32 MaxDebugLevel, const float ColorScale, bool bDrawKinematic, bool bDrawDynamic)
-	{
-#if CHAOS_DEBUG_DRAW
-		using namespace Chaos;
-		if (FDebugDrawQueue::IsDebugDrawingEnabled())
-		{
-			if ((ChaosImmediate_DebugDrawParticles >= MinDebugLevel) && (ChaosImmediate_DebugDrawParticles <= MaxDebugLevel))
-			{
-				DebugDraw::DrawParticleTransforms(SimulationSpaceTransform, Evolution->GetIslandParticles(Island), ColorScale, bDrawKinematic, bDrawDynamic);
-			}
-			if ((ChaosImmediate_DebugDrawShapes >= MinDebugLevel) && (ChaosImmediate_DebugDrawShapes <= MaxDebugLevel))
-			{
-				DebugDraw::DrawParticleShapes(SimulationSpaceTransform, Evolution->GetIslandParticles(Island), ColorScale, bDrawKinematic, bDrawDynamic);
-			}
-		}
-#endif
-	}
-
-	void FSimulation::DebugDrawIslandConstraints(const int32 Island, const int32 MinDebugLevel, const int32 MaxDebugLevel, const float ColorScale)
+	void FSimulation::DebugDrawConstraints(const int32 MinDebugLevel, const int32 MaxDebugLevel, const float ColorScale)
 	{
 #if CHAOS_DEBUG_DRAW
 		using namespace Chaos;
@@ -574,19 +531,11 @@ namespace ImmediatePhysics_Chaos
 		{
 			if ((ChaosImmediate_DebugDrawCollisions >= MinDebugLevel) && (ChaosImmediate_DebugDrawCollisions <= MaxDebugLevel))
 			{
-				Evolution->GetCollisionConstraintsRule().VisitIslandConstraints(Island,
-					[this, ColorScale](const TArray<TPBDCollisionConstraintHandle<float, 3>*>& ConstraintHandles)
-					{
-						DebugDraw::DrawCollisions(SimulationSpaceTransform, ConstraintHandles, ColorScale);
-					});
+				DebugDraw::DrawCollisions(SimulationSpaceTransform, Collisions, ColorScale);
 			}
 			if ((ChaosImmediate_DebugDrawJoints >= MinDebugLevel) && (ChaosImmediate_DebugDrawJoints <= MaxDebugLevel))
 			{
-				JointsRule->VisitIslandConstraints(Island,
-					[this, ColorScale](const TArray<FPBDJointConstraintHandle*>& ConstraintHandles)
-					{
-						DebugDraw::DrawJointConstraints(SimulationSpaceTransform, *Joints, ColorScale, (uint32)ChaosImmediate_DebugDrawJointFeatures);
-					});
+				DebugDraw::DrawJointConstraints(SimulationSpaceTransform, Joints, ColorScale, (uint32)ChaosImmediate_DebugDrawJointFeatures);
 			}
 		}
 #endif
diff --git a/Engine/Source/Runtime/Engine/Public/Physics/ImmediatePhysics/ImmediatePhysicsChaos/ImmediatePhysicsActorHandle_Chaos.h b/Engine/Source/Runtime/Engine/Public/Physics/ImmediatePhysics/ImmediatePhysicsChaos/ImmediatePhysicsActorHandle_Chaos.h
--- a/Engine/Source/Runtime/Engine/Public/Physics/ImmediatePhysics/ImmediatePhysicsChaos/ImmediatePhysicsActorHandle_Chaos.h
+++ b/Engine/Source/Runtime/Engine/Public/Physics/ImmediatePhysics/ImmediatePhysicsChaos/ImmediatePhysicsActorHandle_Chaos.h
@@ -119,11 +119,11 @@ namespace ImmediatePhysics_Chaos
 		friend struct FSimulation;
 		friend struct FJointHandle;
 
-		FActorHandle(Chaos::TPBDRigidsEvolutionGBF<FReal, Dimensions>* InEvolution, EActorType ActorType, FBodyInstance* BodyInstance, const FTransform& Transform);
+		FActorHandle(Chaos::TPBDRigidsSOAs<FReal, 3>& InParticles, EActorType ActorType, FBodyInstance* BodyInstance, const FTransform& Transform);
 
 		Chaos::TGenericParticleHandle<FReal, Dimensions> Handle() const;
 
-		Chaos::TPBDRigidsEvolutionGBF<FReal, Dimensions>* Evolution;
+		Chaos::TPBDRigidsSOAs<FReal, 3>& Particles;
 		Chaos::TGeometryParticleHandle<FReal, Dimensions>* ParticleHandle;
 		TUniquePtr<Chaos::FImplicitObject> Geometry;
 		TArray<TUniquePtr<Chaos::TPerShapeData<float, 3>>> Shapes;
diff --git a/Engine/Source/Runtime/Engine/Public/Physics/ImmediatePhysics/ImmediatePhysicsChaos/ImmediatePhysicsSimulation_Chaos.h b/Engine/Source/Runtime/Engine/Public/Physics/ImmediatePhysics/ImmediatePhysicsChaos/ImmediatePhysicsSimulation_Chaos.h
--- a/Engine/Source/Runtime/Engine/Public/Physics/ImmediatePhysics/ImmediatePhysicsChaos/ImmediatePhysicsSimulation_Chaos.h
+++ b/Engine/Source/Runtime/Engine/Public/Physics/ImmediatePhysics/ImmediatePhysicsChaos/ImmediatePhysicsSimulation_Chaos.h
@@ -5,6 +5,15 @@
 #include "Physics/ImmediatePhysics/ImmediatePhysicsChaos/ImmediatePhysicsCore_Chaos.h"
 
 #include "Chaos/ChaosDebugDrawDeclares.h"
+#include "Chaos/Evolution/PBDMinEvolution.h"
+#include "Chaos/Collision/CollisionDetector.h"
+#include "Chaos/Collision/CollisionReceiver.h"
+#include "Chaos/Collision/NarrowPhase.h"
+#include "Chaos/Collision/ParticlePairBroadPhase.h"
+#include "Chaos/PBDCollisionConstraints.h"
+#include "Chaos/PBDConstraintRule.h"
+#include "Chaos/PBDJointConstraints.h"
+
 #include "Engine/EngineTypes.h"
 #include "Templates/UniquePtr.h"
 
@@ -86,11 +95,26 @@ namespace ImmediatePhysics_Chaos
 		void DebugDrawIslandParticles(const int32 Island, const int32 MinDebugLevel, const int32 MaxDebugLevel, const float ColorScale, bool bDrawKinematic, bool bDrawDynamic);
 		void DebugDrawIslandConstraints(const int32 Island, const int32 MinDebugLevel, const int32 MaxDebugLevel, const float ColorScale);
 
+		using FCollisionConstraints = Chaos::TPBDCollisionConstraints<FReal, 3>;
+		using FCollisionDetector = Chaos::TCollisionDetector<Chaos::FParticlePairBroadPhase, Chaos::FNarrowPhase, Chaos::FSyncCollisionReceiver, FCollisionConstraints>;
+		using FRigidParticleSOAs = Chaos::TPBDRigidsSOAs<FReal, 3>;
 		using FParticleHandle = Chaos::TGeometryParticleHandle<FReal, Dimensions>;
-		TUniquePtr<Chaos::TPBDRigidsEvolutionGBF<FReal, Dimensions>> Evolution;
-		TUniquePtr<Chaos::TPBDRigidsSOAs<FReal, Dimensions>> Particles;
-		TUniquePtr<Chaos::FPBDJointConstraints> Joints;
-		TUniquePtr<Chaos::TPBDConstraintIslandRule<Chaos::FPBDJointConstraints>> JointsRule;		// @todo(ccaulfield): Don't need islands for anim node physics...
+		using FParticlePair = Chaos::TVector<Chaos::TGeometryParticleHandle<Chaos::FReal, 3>*, 2>;
+
+		FRigidParticleSOAs Particles;
+		Chaos::FPBDJointConstraints Joints;
+		FCollisionConstraints Collisions;
+		Chaos::FParticlePairBroadPhase BroadPhase;
+		FCollisionDetector CollisionDetector;
+		Chaos::TSimpleConstraintRule<Chaos::FPBDJointConstraints> JointsRule;
+		Chaos::TSimpleConstraintRule<FCollisionConstraints> CollisionsRule;
+		Chaos::FPBDMinEvolution Evolution;
+
+		// @todo(ccaulfield): Look into these...
+		TArray<FParticlePair> PotentiallyCollidingPairs;
+		Chaos::TArrayCollectionArray<bool> CollidedParticles;
+		Chaos::TArrayCollectionArray<Chaos::TSerializablePtr<Chaos::FChaosPhysicsMaterial>> ParticleMaterials;
+		Chaos::TArrayCollectionArray<TUniquePtr<Chaos::FChaosPhysicsMaterial>> PerParticleMaterials;
 
 		/** Mapping from entity index to handle */
 		// @todo(ccaulfield): we now have handles pointing to handles which is inefficient - we can do better than this, but don't want to change API yet
diff --git a/Engine/Source/Runtime/Experimental/Chaos/Private/Chaos/ChaosDebugDraw.cpp b/Engine/Source/Runtime/Experimental/Chaos/Private/Chaos/ChaosDebugDraw.cpp
--- a/Engine/Source/Runtime/Experimental/Chaos/Private/Chaos/ChaosDebugDraw.cpp
+++ b/Engine/Source/Runtime/Experimental/Chaos/Private/Chaos/ChaosDebugDraw.cpp
@@ -121,6 +121,17 @@ namespace Chaos
 			DrawShapesImpl(FRigidTransform3(P, Q), Particle->Geometry().Get(), Color);
 		}
 
+		void DrawParticleBoundsImpl(const FRigidTransform3& SpaceTransform, const TGeometryParticleHandle<FReal, 3>* InParticle, FReal ColorScale)
+		{
+			FColor Color = (ColorScale * FColor::White).ToFColor(false);
+
+			TBox<FReal, 3> Box = InParticle->WorldSpaceInflatedBounds();
+			FVec3 P = SpaceTransform.TransformPosition(Box.GetCenter());
+			FRotation3 Q = SpaceTransform.GetRotation();
+			FMatrix33 Qm = Q.ToMatrix();
+			FDebugDrawQueue::GetInstance().DrawDebugBox(P, 0.5f * Box.Extents(), Q, Color, false, KINDA_SMALL_NUMBER, DrawPriority, LineThickness);
+		}
+
 		void DrawParticleTransformImpl(const FRigidTransform3& SpaceTransform, const TGeometryParticleHandle<FReal, 3>* InParticle, FReal ColorScale)
 		{
 			FColor Red = (ColorScale * FColor::Red).ToFColor(false);
@@ -264,6 +275,38 @@ namespace Chaos
 #endif
 		}
 
+		void DrawParticleBounds(const FRigidTransform3& SpaceTransform, const TParticleView<TGeometryParticles<float, 3>>& ParticlesView, float ColorScale, bool bDrawKinemtatic, bool bDrawDynamic)
+		{
+#if CHAOS_DEBUG_DRAW
+			if (FDebugDrawQueue::IsDebugDrawingEnabled())
+			{
+				for (auto& Particle : ParticlesView)
+				{
+					if ((bDrawKinemtatic && Particle.ObjectState() != EObjectStateType::Dynamic) || (bDrawDynamic && Particle.ObjectState() == EObjectStateType::Dynamic))
+					{
+						DrawParticleBoundsImpl(SpaceTransform, GetHandleHelper(&Particle), ColorScale);
+					}
+				}
+			}
+#endif
+		}
+
+		void DrawParticleBounds(const FRigidTransform3& SpaceTransform, const TArray<TGeometryParticleHandle<float, 3>*>& Particles, float ColorScale, bool bDrawKinemtatic, bool bDrawDynamic)
+		{
+#if CHAOS_DEBUG_DRAW
+			if (FDebugDrawQueue::IsDebugDrawingEnabled())
+			{
+				for (auto& Particle : Particles)
+				{
+					if ((bDrawKinemtatic && Particle->ObjectState() != EObjectStateType::Dynamic) || (bDrawDynamic && Particle->ObjectState() == EObjectStateType::Dynamic))
+					{
+						DrawParticleBoundsImpl(SpaceTransform, Particle, ColorScale);
+					}
+				}
+			}
+#endif
+		}
+
 		void DrawParticleTransforms(const FRigidTransform3& SpaceTransform, const TParticleView<TGeometryParticles<float, 3>>& ParticlesView, float ColorScale, bool bDrawKinemtatic, bool bDrawDynamic)
 		{
 #if CHAOS_DEBUG_DRAW
diff --git a/Engine/Source/Runtime/Experimental/Chaos/Private/Chaos/Evolution/PBDMinEvolution.cpp b/Engine/Source/Runtime/Experimental/Chaos/Private/Chaos/Evolution/PBDMinEvolution.cpp
new file mode 100644
--- /dev/null
+++ b/Engine/Source/Runtime/Experimental/Chaos/Private/Chaos/Evolution/PBDMinEvolution.cpp
@@ -0,0 +1,287 @@
+// Copyright 1998-2019 Epic Games, Inc. All Rights Reserved.
+
+#include "Chaos/Evolution/PBDMinEvolution.h"
+#include "Chaos/Collision/CollisionDetector.h"
+#include "Chaos/Collision/NarrowPhase.h"
+#include "Chaos/Collision/ParticlePairBroadPhase.h"
+#include "Chaos/PBDCollisionConstraints.h"
+#include "Chaos/PBDConstraintRule.h"
+#include "Chaos/PBDRigidsSOAs.h"
+#include "Chaos/PerParticleEtherDrag.h"
+#include "Chaos/PerParticleEulerStepVelocity.h"
+#include "Chaos/PerParticleGravity.h"
+#include "Chaos/PerParticleInitForce.h"
+#include "Chaos/PerParticlePBDEulerStep.h"
+#include "Chaos/PerParticlePBDGroundConstraint.h"
+#include "Chaos/PerParticlePBDUpdateFromDeltaPosition.h"
+#include "ChaosStats.h"
+
+namespace Chaos
+{
+	DECLARE_CYCLE_STAT(TEXT("FPBDMinEvolution::AdvanceOneTimeStep"), STAT_MinEvolution_AdvanceOneTimeStep, STATGROUP_Chaos);
+	DECLARE_CYCLE_STAT(TEXT("FPBDMinEvolution::Integrate"), STAT_MinEvolution_Integrate, STATGROUP_Chaos);
+	DECLARE_CYCLE_STAT(TEXT("FPBDMinEvolution::KinematicTargets"), STAT_MinEvolution_KinematicTargets, STATGROUP_Chaos);
+	DECLARE_CYCLE_STAT(TEXT("FPBDMinEvolution::ApplyConstraints"), STAT_MinEvolution_ApplyConstraints, STATGROUP_Chaos);
+	DECLARE_CYCLE_STAT(TEXT("FPBDMinEvolution::UpdateVelocities"), STAT_MinEvolution_UpdateVelocites, STATGROUP_Chaos);
+	DECLARE_CYCLE_STAT(TEXT("FPBDMinEvolution::ApplyPushOut"), STAT_MinEvolution_ApplyPushOut, STATGROUP_Chaos);
+	DECLARE_CYCLE_STAT(TEXT("FPBDMinEvolution::DetectCollisions"), STAT_MinEvolution_DetectCollisions, STATGROUP_Chaos);
+	DECLARE_CYCLE_STAT(TEXT("FPBDMinEvolution::UpdatePositions"), STAT_MinEvolution_UpdatePositions, STATGROUP_Chaos);
+
+	FPBDMinEvolution::FPBDMinEvolution(FRigidParticleSOAs& InParticles, FCollisionDetector& InCollisionDetector)
+		: Particles(InParticles)
+		, CollisionDetector(InCollisionDetector)
+		, NumApplyIterations(0)
+		, NumApplyPushOutIterations(0)
+		, Gravity(FVec3(0))
+	{
+	}
+
+	void FPBDMinEvolution::AddConstraintRule(FSimpleConstraintRule* Rule)
+	{
+		ConstraintRules.Add(Rule);
+	}
+
+	// @todo(ccaulfield): dedupe (PBDRigidsEvolutionGBF)
+	void FPBDMinEvolution::Advance(const FReal Dt, const FReal MaxStepDt, const int32 MaxSteps)
+	{
+		// Determine how many steps we would like to take
+		int32 NumSteps = FMath::CeilToInt(Dt / MaxStepDt);
+		if (NumSteps > 0)
+		{
+			// Determine the step time
+			const FReal StepDt = Dt / (FReal)NumSteps;
+
+			// Limit the number of steps
+			// NOTE: This is after step time calculation so sim will appear to slow down for large Dt
+			// but that is preferable to blowing up from a large timestep.
+			NumSteps = FMath::Clamp(NumSteps, 1, MaxSteps);
+
+			for (int32 Step = 0; Step < NumSteps; ++Step)
+			{
+				// StepFraction: how much of the remaining time this step represents, used to interpolate kinematic targets
+				// E.g., for 4 steps this will be: 1/4, 1/3, 1/2, 1
+				const float StepFraction = (FReal)1 / (FReal)(NumSteps - Step);
+
+				UE_LOG(LogChaos, Verbose, TEXT("Advance dt = %f [%d/%d]"), StepDt, Step + 1, NumSteps);
+
+				AdvanceOneTimeStep(StepDt, StepFraction);
+			}
+		}
+	}
+
+	void FPBDMinEvolution::AdvanceOneTimeStep(const FReal Dt, const FReal StepFraction)
+	{
+		SCOPE_CYCLE_COUNTER(STAT_MinEvolution_AdvanceOneTimeStep);
+
+		Integrate(Dt);
+
+		ApplyKinematicTargets(Dt, StepFraction);
+
+		if (PostIntegrateCallback != nullptr)
+		{
+			PostIntegrateCallback();
+		}
+
+		DetectCollisions(Dt);
+
+		if (PostDetectCollisionsCallback != nullptr)
+		{
+			PostDetectCollisionsCallback();
+		}
+
+		if (Dt > 0)
+		{
+			ApplyConstraints(Dt);
+
+			if (PostApplyCallback != nullptr)
+			{
+				PostApplyCallback();
+			}
+
+			UpdateVelocities(Dt);
+
+			ApplyPushOutConstraints(Dt);
+
+			if (PostApplyPushOutCallback != nullptr)
+			{
+				PostApplyPushOutCallback();
+			}
+
+			UpdatePositions(Dt);
+		}
+	}
+
+	// @todo(ccaulfield): dedupe (PBDRigidsEvolutionGBF)
+	void FPBDMinEvolution::Integrate(FReal Dt)
+	{
+		SCOPE_CYCLE_COUNTER(STAT_MinEvolution_Integrate);
+
+		TPerParticleEulerStepVelocity<FReal, 3> EulerStepVelocityRule;
+		//TPerParticleEtherDrag<FReal, 3> EtherDragRule(HackLinearDrag, HackAngularDrag);
+		TPerParticlePBDEulerStep<FReal, 3> EulerStepPositionRule;
+
+		for (TTransientPBDRigidParticleHandle<FReal, 3>& Particle : Particles.GetActiveParticlesView())
+		{
+			if (Particle.ObjectState() == EObjectStateType::Dynamic)
+			{
+				Particle.PreV() = Particle.V();
+				Particle.PreW() = Particle.W();
+
+				Particle.F() = Particle.ExternalForce() + Particle.M() * Gravity;
+				Particle.Torque() = Particle.ExternalTorque();
+
+				EulerStepVelocityRule.Apply(Particle, Dt);
+				//EtherDragRule.Apply(Particle, Dt);
+
+				EulerStepPositionRule.Apply(Particle, Dt);
+
+				if (Particle.HasBounds())
+				{
+					const FReal BoundsThickness = 1.0f;
+					const FReal BoundsThicknessVelocityInflation = 2.0f;
+					const TBox<FReal, 3>& LocalBounds = Particle.LocalBounds();
+					TBox<FReal, 3> WorldSpaceBounds = LocalBounds.TransformedBox(FRigidTransform3(Particle.P(), Particle.Q()));
+					WorldSpaceBounds.ThickenSymmetrically(FVec3(BoundsThickness) + BoundsThicknessVelocityInflation * Particle.V().GetAbs() * Dt);
+					Particle.SetWorldSpaceInflatedBounds(WorldSpaceBounds);
+				}
+			}
+		}
+	}
+
+	// @todo(ccaulfield): dedupe (PBDRigidsEvolutionGBF)
+	void FPBDMinEvolution::ApplyKinematicTargets(FReal Dt, FReal StepFraction)
+	{
+		SCOPE_CYCLE_COUNTER(STAT_MinEvolution_KinematicTargets);
+
+		check(StepFraction > (FReal)0);
+		check(StepFraction <= (FReal)1);
+
+		// @todo(ccaulfield): optimize. Depending on the number of kinematics relative to the number that have 
+		// targets set, it may be faster to process a command list rather than iterate over them all each frame. 
+		const FReal MinDt = 1e-6f;
+		for (auto& Particle : Particles.GetActiveKinematicParticlesView())
+		{
+			TKinematicTarget<FReal, 3>& KinematicTarget = Particle.KinematicTarget();
+			switch (KinematicTarget.GetMode())
+			{
+			case EKinematicTargetMode::None:
+				// Nothing to do
+				break;
+
+			case EKinematicTargetMode::Zero:
+			{
+				// Reset velocity and then switch to do-nothing mode
+				Particle.V() = FVec3(0);
+				Particle.W() = FVec3(0);
+				KinematicTarget.SetMode(EKinematicTargetMode::None);
+				break;
+			}
+
+			case EKinematicTargetMode::Position:
+			{
+				// Move to kinematic target and update velocities to match
+				// Target positions only need to be processed once, and we reset the velocity next frame (if no new target is set)
+				FVec3 TargetPos;
+				FRotation3 TargetRot;
+				if (FMath::IsNearlyEqual(StepFraction, (FReal)1, KINDA_SMALL_NUMBER))
+				{
+					TargetPos = KinematicTarget.GetTarget().GetLocation();
+					TargetRot = KinematicTarget.GetTarget().GetRotation();
+					KinematicTarget.SetMode(EKinematicTargetMode::Zero);
+				}
+				else
+				{
+					TargetPos = FVec3::Lerp(Particle.X(), KinematicTarget.GetTarget().GetLocation(), StepFraction);
+					TargetRot = FRotation3::Slerp(Particle.R(), KinematicTarget.GetTarget().GetRotation(), StepFraction);
+				}
+				if (Dt > MinDt)
+				{
+					Particle.V() = FVec3::CalculateVelocity(Particle.X(), TargetPos, Dt);
+					Particle.W() = FRotation3::CalculateAngularVelocity(Particle.R(), TargetRot, Dt);
+				}
+				Particle.X() = TargetPos;
+				Particle.R() = TargetRot;
+				break;
+			}
+
+			case EKinematicTargetMode::Velocity:
+			{
+				// Move based on velocity
+				Particle.X() = Particle.X() + Particle.V() * Dt;
+				FRotation3::IntegrateRotationWithAngularVelocity(Particle.R(), Particle.W(), Dt);
+				break;
+			}
+			}
+		}
+	}
+
+	void FPBDMinEvolution::DetectCollisions(FReal Dt)
+	{
+		SCOPE_CYCLE_COUNTER(STAT_MinEvolution_DetectCollisions);
+
+		for (FSimpleConstraintRule* ConstraintRule : PrioritizedConstraintRules)
+		{
+			ConstraintRule->UpdatePositionBasedState(Dt);
+		}
+
+		CollisionDetector.DetectCollisions(Dt);
+	}
+
+	void FPBDMinEvolution::ApplyConstraints(FReal Dt)
+	{
+		SCOPE_CYCLE_COUNTER(STAT_MinEvolution_ApplyConstraints);
+
+		// @todo(ccaulfield): doesn't need to be every frame
+		PrioritizedConstraintRules = ConstraintRules;
+		PrioritizedConstraintRules.StableSort();
+
+		// @todo(ccaulfield): track whether we are sufficiently solved and can early-out
+		for (int32 i = 0; i < NumApplyIterations; ++i)
+		{
+			for (FSimpleConstraintRule* ConstraintRule : PrioritizedConstraintRules)
+			{
+				ConstraintRule->ApplyConstraints(Dt, i, NumApplyIterations);
+			}
+		}
+	}
+
+	void FPBDMinEvolution::UpdateVelocities(FReal Dt)
+	{
+		SCOPE_CYCLE_COUNTER(STAT_MinEvolution_UpdateVelocites);
+
+		TPerParticlePBDUpdateFromDeltaPosition<FReal, 3> UpdateVelocityRule;
+		for (auto& Particle : Particles.GetActiveParticlesView())
+		{
+			UpdateVelocityRule.Apply(Particle, Dt);
+		}
+	}
+
+	void FPBDMinEvolution::ApplyPushOutConstraints(FReal Dt)
+	{
+		SCOPE_CYCLE_COUNTER(STAT_MinEvolution_ApplyPushOut);
+
+		bool bNeedsAnotherIteration = true;
+		for (int32 It = 0; bNeedsAnotherIteration && (It < NumApplyPushOutIterations); ++It)
+		{
+			bNeedsAnotherIteration = false;
+			for (FSimpleConstraintRule* ConstraintRule : PrioritizedConstraintRules)
+			{
+				if (ConstraintRule->ApplyPushOut(Dt, It, NumApplyPushOutIterations))
+				{
+					bNeedsAnotherIteration = true;
+				}
+			}
+		}
+	}
+
+	void FPBDMinEvolution::UpdatePositions(FReal Dt)
+	{
+		SCOPE_CYCLE_COUNTER(STAT_MinEvolution_UpdatePositions);
+		for (auto& Particle : Particles.GetActiveParticlesView())
+		{
+			Particle.X() = Particle.P();
+			Particle.R() = Particle.Q();
+		}
+	}
+
+}
\ No newline at end of file
diff --git a/Engine/Source/Runtime/Experimental/Chaos/Private/Chaos/PBDCollisionConstraints.cpp b/Engine/Source/Runtime/Experimental/Chaos/Private/Chaos/PBDCollisionConstraints.cpp
--- a/Engine/Source/Runtime/Experimental/Chaos/Private/Chaos/PBDCollisionConstraints.cpp
+++ b/Engine/Source/Runtime/Experimental/Chaos/Private/Chaos/PBDCollisionConstraints.cpp
@@ -123,7 +123,7 @@ namespace Chaos
 	}
 
 	template<typename T, int d>
-	void TPBDCollisionConstraints<T, d>::UpdatePositionBasedState(/*const TPBDRigidParticles<T, d>& InParticles, const TArray<int32>& InIndices,*/ const T Dt)
+	void TPBDCollisionConstraints<T, d>::UpdatePositionBasedState(const T Dt)
 	{
 		Reset();
 	
@@ -227,6 +227,21 @@ namespace Chaos
 	}
 
 
+	template<typename T, int d>
+	void TPBDCollisionConstraints<T, d>::Apply(const T Dt, const int32 It, const int32 NumIts)
+	{
+		Apply(Dt, Handles, It, NumIts);
+	}
+
+
+	template<typename T, int d>
+	bool TPBDCollisionConstraints<T, d>::ApplyPushOut(const T Dt, const int32 It, const int32 NumIts)
+	{
+		TSet<const TGeometryParticleHandle<T, d>*> TempStatic;
+		return ApplyPushOut(Dt, Handles, TempStatic, It, NumIts);
+	}
+
+
 	DECLARE_CYCLE_STAT(TEXT("TPBDCollisionConstraints::Apply"), STAT_Collisions_Apply, STATGROUP_Chaos);
 	template<typename T, int d>
 	void TPBDCollisionConstraints<T, d>::Apply(const T Dt, const TArray<FConstraintContainerHandle*>& InConstraintHandles, const int32 Iterations, const int32 NumIterations)
diff --git a/Engine/Source/Runtime/Experimental/Chaos/Private/Chaos/PBDConstraintRule.cpp b/Engine/Source/Runtime/Experimental/Chaos/Private/Chaos/PBDConstraintRule.cpp
--- a/Engine/Source/Runtime/Experimental/Chaos/Private/Chaos/PBDConstraintRule.cpp
+++ b/Engine/Source/Runtime/Experimental/Chaos/Private/Chaos/PBDConstraintRule.cpp
@@ -52,6 +52,9 @@ namespace Chaos
 		return Constraints.NumConstraints(); 
 	}
 
+	template class TSimpleConstraintRule<TPBDCollisionConstraints<float, 3>>;
+	template class TSimpleConstraintRule<FPBDJointConstraints>;
+
 	template class TPBDConstraintGraphRuleImpl<TPBDCollisionConstraints<float, 3>>;
 	template class TPBDConstraintGraphRuleImpl<FPBDJointConstraints>;
 	template class TPBDConstraintGraphRuleImpl<TPBDPositionConstraints<float, 3>>;
diff --git a/Engine/Source/Runtime/Experimental/Chaos/Private/Chaos/PBDJointConstraints.cpp b/Engine/Source/Runtime/Experimental/Chaos/Private/Chaos/PBDJointConstraints.cpp
--- a/Engine/Source/Runtime/Experimental/Chaos/Private/Chaos/PBDJointConstraints.cpp
+++ b/Engine/Source/Runtime/Experimental/Chaos/Private/Chaos/PBDJointConstraints.cpp
@@ -391,7 +391,23 @@ namespace Chaos
 		}
 	}
 
-	
+
+	void FPBDJointConstraints::Apply(const FReal Dt, const int32 It, const int32 NumIts)
+	{
+		Apply(Dt, Handles, It, NumIts);
+	}
+
+	bool FPBDJointConstraints::ApplyPushOut(const FReal Dt, const int32 It, const int32 NumIts)
+	{
+		return ApplyPushOut(Dt, Handles, It, NumIts);
+	}
+
+	void FPBDJointConstraints::ApplyProjection(const FReal Dt)
+	{
+		ApplyProjection(Dt, Handles);
+	}
+
+
 	void FPBDJointConstraints::Apply(const FReal Dt, const TArray<FConstraintContainerHandle*>& InConstraintHandles, const int32 It, const int32 NumIts)
 	{
 		SCOPE_CYCLE_COUNTER(STAT_Joints_Apply);
diff --git a/Engine/Source/Runtime/Experimental/Chaos/Public/Chaos/Box.h b/Engine/Source/Runtime/Experimental/Chaos/Public/Chaos/Box.h
--- a/Engine/Source/Runtime/Experimental/Chaos/Public/Chaos/Box.h
+++ b/Engine/Source/Runtime/Experimental/Chaos/Public/Chaos/Box.h
@@ -100,6 +100,16 @@ namespace Chaos
 			return TBox<T, d>(AABB.TransformedAABB(SpaceTransform));
 		}
 
+		static FORCEINLINE bool Intersects(const TBox<T, d>& A, const TBox<T, d>& B)
+		{
+			return A.AABB.Intersects(B.AABB);
+		}
+
+		static FORCEINLINE TBox<T, d> Intersection(const TBox<T, d>& A, const TBox<T, d>& B)
+		{
+			return TBox<T, d>(A.AABB.GetIntersection(B.AABB));
+		}
+
 		FORCEINLINE bool Intersects(const TBox<T, d>& Other) const
 		{
 			return AABB.Intersects(Other.AABB);
diff --git a/Engine/Source/Runtime/Experimental/Chaos/Public/Chaos/ChaosDebugDraw.h b/Engine/Source/Runtime/Experimental/Chaos/Public/Chaos/ChaosDebugDraw.h
--- a/Engine/Source/Runtime/Experimental/Chaos/Public/Chaos/ChaosDebugDraw.h
+++ b/Engine/Source/Runtime/Experimental/Chaos/Public/Chaos/ChaosDebugDraw.h
@@ -24,6 +24,8 @@ namespace Chaos
 
 		CHAOS_API void DrawParticleShapes(const FRigidTransform3& SpaceTransform, const TParticleView<TGeometryParticles<float, 3>>& ParticlesView, float ColorScale, bool bDrawKinematic = true, bool bDrawDynamic = true);
 		CHAOS_API void DrawParticleShapes(const FRigidTransform3& SpaceTransform, const TArray<TGeometryParticleHandle<float, 3>*>& Particles, float ColorScale, bool bDrawKinematic = true, bool bDrawDynamic = true);
+		CHAOS_API void DrawParticleBounds(const FRigidTransform3& SpaceTransform, const TParticleView<TGeometryParticles<float, 3>>& ParticlesView, float ColorScale, bool bDrawKinematic = true, bool bDrawDynamic = true);
+		CHAOS_API void DrawParticleBounds(const FRigidTransform3& SpaceTransform, const TArray<TGeometryParticleHandle<float, 3>*>& Particles, float ColorScale, bool bDrawKinematic = true, bool bDrawDynamic = true);
 		CHAOS_API void DrawParticleTransforms(const FRigidTransform3& SpaceTransform, const TParticleView<TGeometryParticles<float, 3>>& ParticlesView, float ColorScale, bool bDrawKinematic = true, bool bDrawDynamic = true);
 		CHAOS_API void DrawParticleTransforms(const FRigidTransform3& SpaceTransform, const TArray<TGeometryParticleHandle<float, 3>*>& Particles, float ColorScale, bool bDrawKinematic = true, bool bDrawDynamic = true);
 		CHAOS_API void DrawParticleCollisions(const FRigidTransform3& SpaceTransform, const TGeometryParticleHandle<float, 3>* Particle, const TPBDCollisionConstraints<float, 3>& Collisions);
diff --git a/Engine/Source/Runtime/Experimental/Chaos/Public/Chaos/Collision/CollisionDetector.h b/Engine/Source/Runtime/Experimental/Chaos/Public/Chaos/Collision/CollisionDetector.h
--- a/Engine/Source/Runtime/Experimental/Chaos/Public/Chaos/Collision/CollisionDetector.h
+++ b/Engine/Source/Runtime/Experimental/Chaos/Public/Chaos/Collision/CollisionDetector.h
@@ -32,6 +32,12 @@ namespace Chaos
 		FBroadPhase& GetBroadPhase() { return BroadPhase; }
 		FContainer& GetCollisionContainer() { return CollisionContainer; }
 
+		void DetectCollisions(const FReal Dt)
+		{
+			CollisionStats::FStatData StatData(false);
+			DetectCollisions(Dt, StatData);
+		}
+
 		void DetectCollisions(const FReal Dt, CollisionStats::FStatData& StatData)
 		{
 			SCOPE_CYCLE_COUNTER(STAT_DetectCollisions);
diff --git a/Engine/Source/Runtime/Experimental/Chaos/Public/Chaos/Collision/ParticlePairBroadPhase.h b/Engine/Source/Runtime/Experimental/Chaos/Public/Chaos/Collision/ParticlePairBroadPhase.h
--- a/Engine/Source/Runtime/Experimental/Chaos/Public/Chaos/Collision/ParticlePairBroadPhase.h
+++ b/Engine/Source/Runtime/Experimental/Chaos/Public/Chaos/Collision/ParticlePairBroadPhase.h
@@ -1,7 +1,9 @@
 // Copyright 1998-2019 Epic Games, Inc. All Rights Reserved.
 #pragma once
 
+#include "Chaos/BoundingVolumeUtilities.h"
 #include "Chaos/Collision/CollisionReceiver.h"
+#include "Chaos/Collision/NarrowPhase.h"
 #include "Chaos/ParticleHandle.h"
 #include "Chaos/PBDCollisionConstraints.h"
 #include "Chaos/PBDRigidsSOAs.h"
@@ -18,7 +20,7 @@ namespace Chaos
 		using FParticlePair = TVector<TGeometryParticleHandle<FReal, 3>*, 2>;
 		using FAABB = TAABB<FReal, 3>;
 
-		FParticlePairBroadPhase(const TArray<FParticlePair> InParticlePairs)
+		FParticlePairBroadPhase(const TArray<FParticlePair>& InParticlePairs)
 			: ParticlePairs(InParticlePairs)
 		{
 		}
@@ -31,10 +33,25 @@ namespace Chaos
 			FSyncCollisionReceiver& Receiver,
 			CollisionStats::FStatData& StatData)
 		{
-			// @todo(ccaulfield)
+			const FReal BoundsThickness = 1.0f;
+			const FReal BoundsThicknessVelocityInflation = 2.0f;
+
+			for (const FParticlePair& ParticlePair : ParticlePairs)	
+			{
+				// Array is const, things in it are not...
+				TGeometryParticleHandle<FReal, 3>* Particle0 = const_cast<TGeometryParticleHandle<FReal, 3>*>(ParticlePair[0]);
+				TGeometryParticleHandle<FReal, 3>* Particle1 = const_cast<TGeometryParticleHandle<FReal, 3>*>(ParticlePair[1]);
+
+				const TBox<FReal, 3>& Box0 = Particle0->WorldSpaceInflatedBounds();
+				const TBox<FReal, 3>& Box1 = Particle1->WorldSpaceInflatedBounds();
+				if (TBox<FReal, 3>::Intersects(Box0, Box1))
+				{
+					NarrowPhase.GenerateCollisions(Dt, Receiver, Particle0, Particle1, BoundsThickness, StatData);
+				}
+			}
 		}
 
 	private:
-		TArray<FParticlePair> ParticlePairs;
+		const TArray<FParticlePair>& ParticlePairs;
 	};
 }
diff --git a/Engine/Source/Runtime/Experimental/Chaos/Public/Chaos/Evolution/PBDMinEvolution.h b/Engine/Source/Runtime/Experimental/Chaos/Public/Chaos/Evolution/PBDMinEvolution.h
new file mode 100644
--- /dev/null
+++ b/Engine/Source/Runtime/Experimental/Chaos/Public/Chaos/Evolution/PBDMinEvolution.h
@@ -0,0 +1,106 @@
+// Copyright 1998-2019 Epic Games, Inc. All Rights Reserved.
+#pragma once
+
+#include "Chaos/Core.h"
+#include "Chaos/ParticleHandleFwd.h"
+
+
+namespace Chaos
+{
+	class FNarrowPhase;
+	class FParticlePairBroadPhase;
+	class FSimpleConstraintRule;
+	class FSyncCollisionReceiver;
+
+	template <typename T, int d>
+	class TPBDCollisionConstraints;
+
+	template<typename T_BROADPHASE, typename T_NARROWPHASE, typename T_RECEIVER, typename T_CONTAINER>
+	class TCollisionDetector;
+
+	template <typename T, int d>
+	class TPBDRigidsSOAs;
+
+	/**
+	 * A minimal optimized evolution with support for
+	 *	- PBD Rigids
+	 *	- Joints
+	 *	- Collisions
+	 *
+	 * It is single-threaded and does not use a constraint graph or partition the particles into islands.
+	 */
+	class CHAOS_API FPBDMinEvolution
+	{
+	public:
+		using FCollisionConstraints = TPBDCollisionConstraints<FReal, 3>;
+		using FCollisionDetector = TCollisionDetector<FParticlePairBroadPhase, FNarrowPhase, FSyncCollisionReceiver, FCollisionConstraints>;
+		using FEvolutionCallback = TFunction<void()>;
+		using FRigidParticleSOAs = TPBDRigidsSOAs<FReal, 3>;
+
+		FPBDMinEvolution(FRigidParticleSOAs& InParticles, FCollisionDetector& InCollisionDetector);
+
+		void AddConstraintRule(FSimpleConstraintRule* Rule);
+
+		void Advance(const FReal Dt, const FReal MaxStepDt, const int32 MaxSteps);
+		void AdvanceOneTimeStep(const FReal dt, const FReal StepFraction);
+
+		void SetNumIterations(const int32 NumIts)
+		{
+			NumApplyIterations = NumIts;
+		}
+
+		void SetNumPushOutIterations(const int32 NumIts)
+		{
+			NumApplyPushOutIterations = NumIts;
+		}
+
+		void SetGravity(const FVec3& G)
+		{
+			Gravity = G;
+		}
+
+		void SetPostIntegrateCallback(const FEvolutionCallback& Cb)
+		{
+			PostIntegrateCallback = Cb;
+		}
+
+		void SetPostDetectCollisionsCallback(const FEvolutionCallback& Cb)
+		{
+			PostDetectCollisionsCallback = Cb;
+		}
+
+		void SetPostApplyCallback(const FEvolutionCallback& Cb)
+		{
+			PostApplyCallback = Cb;
+		}
+
+		void SetPostApplyPushOutCallback(const FEvolutionCallback& Cb)
+		{
+			PostApplyPushOutCallback = Cb;
+		}
+
+	private:
+		void Integrate(FReal Dt);
+		void ApplyKinematicTargets(FReal Dt, FReal StepFraction);
+		void DetectCollisions(FReal Dt);
+		void ApplyConstraints(FReal Dt);
+		void UpdateVelocities(FReal Dt);
+		void ApplyPushOutConstraints(FReal Dt);
+		void UpdatePositions(FReal Dt);
+
+		FRigidParticleSOAs& Particles;
+		FCollisionDetector& CollisionDetector;
+
+		TArray<FSimpleConstraintRule*> ConstraintRules;
+		TArray<FSimpleConstraintRule*> PrioritizedConstraintRules;
+
+		int32 NumApplyIterations;
+		int32 NumApplyPushOutIterations;
+		FVec3 Gravity;
+
+		FEvolutionCallback PostIntegrateCallback;
+		FEvolutionCallback PostDetectCollisionsCallback;
+		FEvolutionCallback PostApplyCallback;
+		FEvolutionCallback PostApplyPushOutCallback;
+	};
+}
diff --git a/Engine/Source/Runtime/Experimental/Chaos/Public/Chaos/PBDCollisionConstraints.h b/Engine/Source/Runtime/Experimental/Chaos/Public/Chaos/PBDCollisionConstraints.h
--- a/Engine/Source/Runtime/Experimental/Chaos/Public/Chaos/PBDCollisionConstraints.h
+++ b/Engine/Source/Runtime/Experimental/Chaos/Public/Chaos/PBDCollisionConstraints.h
@@ -72,12 +72,6 @@ public:
 
 	void AddConstraint(FConstraintBase* InConstraint);
 
-	/**
-	 * Generate all contact constraints.
-	 */
-	void UpdatePositionBasedState(const T Dt);
-
-
 	/**
 	*  Reset the constraint frame. 
 	*/
@@ -109,22 +103,28 @@ public:
 	void UpdateConstraints(T Dt, const TSet<TGeometryParticleHandle<T, d>*>& AddedParticles);
 
 
+	//
+	// General Rule API
+	//
 
 	/**
-	 * Apply corrections for the specified list of constraints. (Runs Wide!)
-	 *
-	 * @todo(ccaulfield): this runs wide. The serial/parallel decision should be in the ConstraintRule
+	 * Generate all contact constraints.
 	 */
-	void Apply(const T Dt, const TArray<FConstraintContainerHandle*>& InConstraintHandles, const int32 It, const int32 NumIts);
+	void UpdatePositionBasedState(const T Dt);
 
+	//
+	// Simple Rule API
+	//
 
-	/**
-	 * Apply push out for the specified list of constraints.
-	 *
-	 * /return true if we need another iteration.
-	 *
-	 * @todo(ccaulfield): this runs wide. The serial/parallel decision should be in the ConstraintRule
-	 */
+	void Apply(const T Dt, const int32 It, const int32 NumIts);
+	bool ApplyPushOut(const T Dt, const int32 It, const int32 NumIts);
+
+	//
+	// Island Rule API
+	//
+	// @todo(ccaulfield): this runs wide. The serial/parallel decision should be in the ConstraintRule
+
+	void Apply(const T Dt, const TArray<FConstraintContainerHandle*>& InConstraintHandles, const int32 It, const int32 NumIts);
 	bool ApplyPushOut(const T Dt, const TArray<FConstraintContainerHandle*>& InConstraintHandles, 
 		const TSet<const TGeometryParticleHandle<T,d>*>& IsTemporarilyStatic, int32 Iteration, int32 NumIterations);
 
diff --git a/Engine/Source/Runtime/Experimental/Chaos/Public/Chaos/PBDConstraintRule.h b/Engine/Source/Runtime/Experimental/Chaos/Public/Chaos/PBDConstraintRule.h
--- a/Engine/Source/Runtime/Experimental/Chaos/Public/Chaos/PBDConstraintRule.h
+++ b/Engine/Source/Runtime/Experimental/Chaos/Public/Chaos/PBDConstraintRule.h
@@ -35,6 +35,51 @@ namespace Chaos
 		int32 Priority;
 	};
 
+	/**
+	 * Constraint rule for evolutions that do not use Constraint Graphs or other acceleration schemes.
+	 */
+	class CHAOS_API FSimpleConstraintRule : public FConstraintRule
+	{
+	public:
+		FSimpleConstraintRule(int32 InPriority) : FConstraintRule(InPriority) {}
+
+		virtual void UpdatePositionBasedState(const FReal Dt) {}
+		virtual void ApplyConstraints(const FReal Dt, const int32 It, const int32 NumIts) {}
+		virtual bool ApplyPushOut(const FReal Dt, const int32 It, const int32 NumIts) { return false; }
+	};
+
+	template<typename T_CONSTRAINTS>
+	class CHAOS_API TSimpleConstraintRule : public FSimpleConstraintRule
+	{
+	public:
+		using FConstraints = T_CONSTRAINTS;
+
+		TSimpleConstraintRule(int32 InPriority, FConstraints& InConstraints)
+			: FSimpleConstraintRule(InPriority)
+			, Constraints(InConstraints)
+		{
+		}
+
+		virtual void UpdatePositionBasedState(const FReal Dt) override
+		{
+			return Constraints.UpdatePositionBasedState(Dt);
+		}
+
+		virtual void ApplyConstraints(const FReal Dt, const int32 It, const int32 NumIts) override
+		{
+			Constraints.Apply(Dt, It, NumIts);
+		}
+
+		virtual bool ApplyPushOut(const FReal Dt, const int32 It, const int32 NumIts) override
+		{ 
+			return Constraints.ApplyPushOut(Dt, It, NumIts);
+		}
+
+	private:
+		FConstraints& Constraints;
+	};
+
+
 	/**
 	 * Base class for Constraint Rules that use the Contact Graph (which will be most optimized ones).
 	 * The graph is shared among many/all constraint rules and is held external to the Graph rule itself.
diff --git a/Engine/Source/Runtime/Experimental/Chaos/Public/Chaos/PBDJointConstraints.h b/Engine/Source/Runtime/Experimental/Chaos/Public/Chaos/PBDJointConstraints.h
--- a/Engine/Source/Runtime/Experimental/Chaos/Public/Chaos/PBDJointConstraints.h
+++ b/Engine/Source/Runtime/Experimental/Chaos/Public/Chaos/PBDJointConstraints.h
@@ -124,17 +124,28 @@ namespace Chaos
 		void SetParticleLevels(int32 ConstraintIndex, const TVector<int32, 2>& ParticleLevels);
 
 		//
-		// Island Rule API
+		// General Rule API
 		//
 
 		void UpdatePositionBasedState(const FReal Dt);
 
-		void Apply(const FReal Dt, const TArray<FConstraintContainerHandle*>& InConstraintHandles, const int32 It, const int32 NumIts);
+		//
+		// Simple Rule API
+		//
 
-		bool ApplyPushOut(const FReal Dt, const TArray<FConstraintContainerHandle*>& InConstraintHandles, const int32 It, const int32 NumIts);
+		void Apply(const FReal Dt, const int32 It, const int32 NumIts);
+		bool ApplyPushOut(const FReal Dt, const int32 It, const int32 NumIts);
+		void ApplyProjection(const FReal Dt);
+
+		//
+		// Island Rule API
+		//
 
+		void Apply(const FReal Dt, const TArray<FConstraintContainerHandle*>& InConstraintHandles, const int32 It, const int32 NumIts);
+		bool ApplyPushOut(const FReal Dt, const TArray<FConstraintContainerHandle*>& InConstraintHandles, const int32 It, const int32 NumIts);
 		void ApplyProjection(const FReal Dt, const TArray<FConstraintContainerHandle*>& InConstraintHandles);
 
+
 	protected:
 		using Base::GetConstraintIndex;
 		using Base::SetConstraintIndex;
diff --git a/Engine/Source/Runtime/Experimental/Chaos/Public/Chaos/PerParticleGravity.h b/Engine/Source/Runtime/Experimental/Chaos/Public/Chaos/PerParticleGravity.h
--- a/Engine/Source/Runtime/Experimental/Chaos/Public/Chaos/PerParticleGravity.h
+++ b/Engine/Source/Runtime/Experimental/Chaos/Public/Chaos/PerParticleGravity.h
@@ -11,7 +11,11 @@ namespace Chaos
 	class TPerParticleGravity : public TPerParticleRule<T, d>
 	{
 	public:
-		TPerParticleGravity(const TVector<T, d>& Direction = TVector<T, d>(0, 0, -1), const T Magnitude = (T)980.665)
+		TPerParticleGravity()
+			: MAcceleration(TVector<T, d>(0, 0, (T)-980.665)) {}
+		TPerParticleGravity(const TVector<T, d>& G)
+			: MAcceleration(G) {}
+		TPerParticleGravity(const TVector<T, d>& Direction, const T Magnitude)
 			: MAcceleration(Direction * Magnitude) {}
 		virtual ~TPerParticleGravity() {}
 
diff --git a/Engine/Source/Runtime/Experimental/ChaosCore/Public/Chaos/Vector.h b/Engine/Source/Runtime/Experimental/ChaosCore/Public/Chaos/Vector.h
--- a/Engine/Source/Runtime/Experimental/ChaosCore/Public/Chaos/Vector.h
+++ b/Engine/Source/Runtime/Experimental/ChaosCore/Public/Chaos/Vector.h
@@ -92,7 +92,7 @@ namespace Chaos
 		{
 			for (int32 i = 0; i < NumElements; ++i)
 			{
-				V[i] = static_cast<T>(Other.V[i]);
+				V[i] = static_cast<T>(Other[i]);
 			}
 		}
 
