Improve error messaging around default curve compression asset

diff --git a/Engine/Source/Runtime/Engine/Private/Animation/AnimationUtils.cpp b/Engine/Source/Runtime/Engine/Private/Animation/AnimationUtils.cpp
--- a/Engine/Source/Runtime/Engine/Private/Animation/AnimationUtils.cpp
+++ b/Engine/Source/Runtime/Engine/Private/Animation/AnimationUtils.cpp
@@ -2033,12 +2033,17 @@ UAnimCurveCompressionSettings* FAnimationUtils::GetDefaultAnimationCurveCompress
 		if (Value != nullptr)
 		{
 			const FString& CurveCompressionSettingsName = Value->GetValue();
+
 			DefaultCurveCompressionSettings = LoadObject<UAnimCurveCompressionSettings>(nullptr, *CurveCompressionSettingsName);
-		}
 
-		if (DefaultCurveCompressionSettings == nullptr)
+			if (DefaultCurveCompressionSettings == nullptr)
+			{
+				UE_LOG(LogAnimationCompression, Fatal, TEXT("Couldn't load default curve compression settings with path '%s'"), *CurveCompressionSettingsName);
+			}
+		}
+		else
 		{
-			UE_LOG(LogAnimationCompression, Fatal, TEXT("Couldn't find default curve compression settings under '[Animation.DefaultObjectSettings]'"));
+			UE_LOG(LogAnimationCompression, Fatal, TEXT("Couldn't find default curve compression setting under '[Animation.DefaultObjectSettings]'"));
 		}
 
 		// Force load the default settings and all its dependencies just in case it hasn't happened yet
