Added a way to construct CompositeNavModifier using a FCollisionShape

#rnx

diff --git a/Engine/Source/Runtime/Engine/Private/AI/Navigation/NavigationModifier.cpp b/Engine/Source/Runtime/Engine/Private/AI/Navigation/NavigationModifier.cpp
--- a/Engine/Source/Runtime/Engine/Private/AI/Navigation/NavigationModifier.cpp
+++ b/Engine/Source/Runtime/Engine/Private/AI/Navigation/NavigationModifier.cpp
@@ -893,6 +893,37 @@ void FCompositeNavModifier::CreateAreaModifiers(const UPrimitiveComponent* PrimC
 	}
 }
 
+void FCompositeNavModifier::CreateAreaModifiers(const FCollisionShape& CollisionShape, const FTransform& LocalToWorld, const TSubclassOf<UNavAreaBase> AreaClass, const bool bIncludeAgentHeight /*= false*/)
+{
+	if (CollisionShape.IsBox())
+	{
+		const FVector BoxExtent = CollisionShape.GetBox();
+		FAreaNavModifier AreaMod(FBox(-BoxExtent, BoxExtent), LocalToWorld, AreaClass);
+		AreaMod.SetIncludeAgentHeight(bIncludeAgentHeight);
+		Add(AreaMod);
+	}
+	else if (CollisionShape.IsCapsule())
+	{
+		const float CapsuleHalfHeight = CollisionShape.GetCapsuleHalfHeight();
+		const FTransform AreaOffset(FVector(0.0f, 0.0f, -CapsuleHalfHeight));
+		FAreaNavModifier AreaMod(CollisionShape.GetCapsuleRadius(), CapsuleHalfHeight * 2.0f, AreaOffset * LocalToWorld, AreaClass); // Note: FAreaNavModifier creates a cylinder shape under the hood
+		AreaMod.SetIncludeAgentHeight(bIncludeAgentHeight);
+		Add(AreaMod);
+	}
+	else if (CollisionShape.IsSphere())
+	{
+		const float SphereRadius = CollisionShape.GetSphereRadius();
+		const FTransform AreaOffset(FVector(0.0f, 0.0f, -SphereRadius));
+		FAreaNavModifier AreaMod(SphereRadius, SphereRadius * 2.0f, AreaOffset * LocalToWorld, AreaClass); // Note: FAreaNavModifier creates a cylinder shape under the hood
+		AreaMod.SetIncludeAgentHeight(bIncludeAgentHeight);
+		Add(AreaMod);
+	}
+	else
+	{
+		UE_LOG(LogNavigation, Error, TEXT("Asked to create a FAreaNavModifier with an unknown collision shape type! Collision Shape Type = %d"), CollisionShape.ShapeType);
+	}
+}
+
 uint32 FCompositeNavModifier::GetAllocatedSize() const
 {
 	uint32 MemUsed = Areas.GetAllocatedSize() + SimpleLinks.GetAllocatedSize() + CustomLinks.GetAllocatedSize();
diff --git a/Engine/Source/Runtime/Engine/Public/AI/NavigationModifier.h b/Engine/Source/Runtime/Engine/Public/AI/NavigationModifier.h
--- a/Engine/Source/Runtime/Engine/Public/AI/NavigationModifier.h
+++ b/Engine/Source/Runtime/Engine/Public/AI/NavigationModifier.h
@@ -292,6 +292,7 @@ struct ENGINE_API FCompositeNavModifier : public FNavigationModifier
 	}
 
 	void CreateAreaModifiers(const UPrimitiveComponent* PrimComp, const TSubclassOf<UNavAreaBase> AreaClass);
+	void CreateAreaModifiers(const FCollisionShape& CollisionShape, const FTransform& LocalToWorld, const TSubclassOf<UNavAreaBase> AreaClass, const bool bIncludeAgentHeight = false);
 
 	FORCEINLINE const TArray<FAreaNavModifier>& GetAreas() const { return Areas; }
 	FORCEINLINE const TArray<FSimpleLinkNavModifier>& GetSimpleLinks() const { return SimpleLinks; }
