LevelPartition: Fix Move Selected to Current Level + Duplicate (Alt+Drag).

#rnx

diff --git a/Engine/Source/Editor/UnrealEd/Private/EditorActor.cpp b/Engine/Source/Editor/UnrealEd/Private/EditorActor.cpp
--- a/Engine/Source/Editor/UnrealEd/Private/EditorActor.cpp
+++ b/Engine/Source/Editor/UnrealEd/Private/EditorActor.cpp
@@ -473,9 +473,13 @@ public:
 		// Cache the current dest level
 		OldLevel = DestLevel->OwningWorld->GetCurrentLevel();
 		// Paste to the dest level.
-		DestLevel->OwningWorld->SetCurrentLevel( DestLevel );
-		GEditor->edactPasteSelected( DestLevel->OwningWorld, true, bOffsetLocations, true, &ScratchData );
-
+		{
+			FLevelPartitionOperationScope LevelPartitionScope(DestLevel);
+			DestLevel->OwningWorld->SetCurrentLevel(LevelPartitionScope.GetLevel());
+			GEditor->edactPasteSelected(DestLevel->OwningWorld, true, bOffsetLocations, true, &ScratchData);
+			// Restore dest level
+			DestLevel->OwningWorld->SetCurrentLevel(OldLevel);
+		}
 		// The selection set will be the newly created actors; copy them over to the output array.
 		for ( FSelectionIterator It( GEditor->GetSelectedActorIterator() ) ; It ; ++It )
 		{
@@ -483,8 +487,6 @@ public:
 			checkSlow( Actor->IsA(AActor::StaticClass()) );
 			OutNewActors.Add( Actor );
 		}
-		// Restore dest level
-		DestLevel->OwningWorld->SetCurrentLevel( OldLevel );
 	}
 };
 }
diff --git a/Engine/Source/Editor/UnrealEd/Private/EditorLevelUtils.cpp b/Engine/Source/Editor/UnrealEd/Private/EditorLevelUtils.cpp
--- a/Engine/Source/Editor/UnrealEd/Private/EditorLevelUtils.cpp
+++ b/Engine/Source/Editor/UnrealEd/Private/EditorLevelUtils.cpp
@@ -139,18 +139,26 @@ int32 UEditorLevelUtils::MoveActorsToLevel(const TArray<AActor*>& ActorsToMove,
 				const bool bIsMove = true;
 				GEditor->CopySelectedActorsToClipboard(OwningWorld, bShoudCut, bIsMove, bWarnAboutReferences);
 
-				// Set the new level and force it visible while we do the paste
-				OwningWorld->SetCurrentLevel(DestLevel);
 				const bool bLevelVisible = DestLevel->bIsVisible;
 				if (!bLevelVisible)
 				{
 					UEditorLevelUtils::SetLevelVisibility(DestLevel, true, false);
 				}
 
-				const bool bDuplicate = false;
-				const bool bOffsetLocations = false;
-				const bool bWarnIfHidden = false;
-				GEditor->edactPasteSelected(OwningWorld, bDuplicate, bOffsetLocations, bWarnIfHidden);
+				// Scope this so that Actors that have been pasted will have their final levels set before doing the actor mapping
+				{
+					// Set the new level and force it visible while we do the paste
+					FLevelPartitionOperationScope LevelPartitionScope(DestLevel);
+					OwningWorld->SetCurrentLevel(LevelPartitionScope.GetLevel());
+										
+					const bool bDuplicate = false;
+					const bool bOffsetLocations = false;
+					const bool bWarnIfHidden = false;
+					GEditor->edactPasteSelected(OwningWorld, bDuplicate, bOffsetLocations, bWarnIfHidden);
+
+					// Restore the original current level
+					OwningWorld->SetCurrentLevel(OldCurrentLevel);
+				}
 
 				// Build a remapping of old to new names so we can do a fixup
 				for (FSelectionIterator It(GEditor->GetSelectedActorIterator()); It; ++It)
@@ -221,9 +229,6 @@ int32 UEditorLevelUtils::MoveActorsToLevel(const TArray<AActor*>& ActorsToMove,
 				{
 					UEditorLevelUtils::SetLevelVisibility(DestLevel, false, false);
 				}
-
-				// Restore the original current level
-				OwningWorld->SetCurrentLevel(OldCurrentLevel);
 			}
 
 			// The moved (pasted) actors will now be selected
diff --git a/Engine/Source/Runtime/Engine/Classes/Engine/Level.h b/Engine/Source/Runtime/Engine/Classes/Engine/Level.h
--- a/Engine/Source/Runtime/Engine/Classes/Engine/Level.h
+++ b/Engine/Source/Runtime/Engine/Classes/Engine/Level.h
@@ -41,10 +41,39 @@ class ILevelPartitionInterface
 {
 	GENERATED_BODY()
 
+#if WITH_EDITOR
+	friend class FLevelPartitionOperationScope;
+#endif
+
 public:
 	virtual ULevel* GetSubLevel(const FVector& Coords) const = 0;
+
+private:
+#if WITH_EDITOR
+	virtual void BeginOperation(FLevelPartitionOperationScope* Scope) = 0;
+	virtual void EndOperation() = 0;
+#endif
 };
 
+#if WITH_EDITOR
+class ENGINE_API FLevelPartitionOperationScope
+{
+public:
+	FLevelPartitionOperationScope(ULevel* Level);
+	~FLevelPartitionOperationScope();
+		
+	TArray<AActor*> GetActors() const;
+	ULevel* GetLevel() const;
+
+private:
+	static ULevel* CreateTransientLevel(UWorld* InWorld);
+	static void DestroyTransientLevel(ULevel* InLevel);
+
+	ILevelPartitionInterface* InterfacePtr = nullptr;
+	ULevel* Level = nullptr;
+};
+#endif
+
 /**
  * Structure containing all information needed for determining the screen space
  * size of an object/ texture instance.
diff --git a/Engine/Source/Runtime/Engine/Private/Level.cpp b/Engine/Source/Runtime/Engine/Private/Level.cpp
--- a/Engine/Source/Runtime/Engine/Private/Level.cpp
+++ b/Engine/Source/Runtime/Engine/Private/Level.cpp
@@ -49,6 +49,7 @@ Level.cpp: Level-related functions
 #if WITH_EDITOR
 #include "Kismet2/KismetEditorUtilities.h"
 #include "Kismet2/BlueprintEditorUtils.h"
+#include "Algo/AnyOf.h"
 #endif
 #include "Engine/LevelStreaming.h"
 #include "LevelUtils.h"
@@ -70,6 +71,75 @@ static FAutoConsoleVariableRef CVarActorClusteringEnabled(
 	ECVF_Default
 );
 
+#if WITH_EDITOR
+FLevelPartitionOperationScope::FLevelPartitionOperationScope(ULevel* InLevel)
+{
+	InterfacePtr = InLevel->GetLevelPartition();
+	Level = InLevel;
+	if (InterfacePtr)
+	{
+		InterfacePtr->BeginOperation(this);
+		Level = CreateTransientLevel(InLevel->GetWorld());
+	}
+}
+
+FLevelPartitionOperationScope::~FLevelPartitionOperationScope()
+{
+	if (InterfacePtr)
+	{
+		InterfacePtr->EndOperation();
+		DestroyTransientLevel(Level);
+	}
+	Level = nullptr;
+}
+
+TArray<AActor*> FLevelPartitionOperationScope::GetActors() const
+{
+	if (InterfacePtr)
+	{
+		return Level->Actors;
+	}
+
+	return {};
+}
+
+ULevel* FLevelPartitionOperationScope::GetLevel() const
+{
+	check(Level);
+	return Level;
+}
+
+ULevel* FLevelPartitionOperationScope::CreateTransientLevel(UWorld* InWorld)
+{
+	ULevel* Level = NewObject<ULevel>(GetTransientPackage(), TEXT("TempLevelPartitionOperationScopeLevel"));
+	check(Level);
+	Level->Initialize(FURL(nullptr));
+	Level->AddToRoot();
+	Level->OwningWorld = InWorld;
+	Level->Model = NewObject<UModel>(Level);
+	Level->Model->Initialize(nullptr, true);
+	Level->bIsVisible = true;
+
+	Level->SetFlags(RF_Transactional);
+	Level->Model->SetFlags(RF_Transactional);
+
+	return Level;
+}
+
+void FLevelPartitionOperationScope::DestroyTransientLevel(ULevel* Level)
+{
+	check(Level->GetOutermost() == GetTransientPackage());
+	// Make sure Level doesn't contain any Actors before destroying. That would mean the operation failed.
+	check(!Algo::AnyOf(Level->Actors, [](AActor* Actor) { return Actor != nullptr; }));
+	// Delete the temporary level
+	Level->ClearLevelComponents();
+	Level->GetWorld()->RemoveLevel(Level);
+	Level->OwningWorld = nullptr;
+	Level->RemoveFromRoot();
+	Level = nullptr;
+}
+#endif
+
 /*-----------------------------------------------------------------------------
 ULevel implementation.
 -----------------------------------------------------------------------------*/
