Ensure max packet header bits definition in NetConnection.h includes the new packet info data.
Added non-shipping checks to ensure WritePacketHeader and WriteDummyPacketInfo don't go over max bits

#rnx

diff --git a/Engine/Source/Runtime/Engine/Classes/Engine/NetConnection.h b/Engine/Source/Runtime/Engine/Classes/Engine/NetConnection.h
--- a/Engine/Source/Runtime/Engine/Classes/Engine/NetConnection.h
+++ b/Engine/Source/Runtime/Engine/Classes/Engine/NetConnection.h
@@ -38,6 +38,13 @@ class UChildConnection;
 
 typedef TMap<TWeakObjectPtr<AActor>, UActorChannel*, FDefaultSetAllocator, TWeakObjectPtrMapKeyFuncs<TWeakObjectPtr<AActor>, UActorChannel*>> FActorChannelMap;
 
+namespace NetConnectionHelper
+{
+	/** Number of bits to use in the packet header for sending the milliseconds on the clock when the packet is sent */
+	constexpr int32 NumBitsForJitterClockTimeInHeader = 10;
+}
+
+
 /*-----------------------------------------------------------------------------
 	Types.
 -----------------------------------------------------------------------------*/
@@ -45,8 +52,9 @@ enum { RELIABLE_BUFFER = 256 }; // Power of 2 >= 1.
 enum { MAX_PACKETID = FNetPacketNotify::SequenceNumberT::SeqNumberCount };  // Power of 2 >= 1, covering guaranteed loss/misorder time.
 enum { MAX_CHSEQUENCE = 1024 }; // Power of 2 >RELIABLE_BUFFER, covering loss/misorder time.
 enum { MAX_BUNCH_HEADER_BITS = 256 };
-enum { UE_PACKET_MAX_SEQUENCE_HISTORY_BITS = FNetPacketNotify::MaxSequenceHistoryLength };
-enum { MAX_PACKET_HEADER_BITS = 49 + UE_PACKET_MAX_SEQUENCE_HISTORY_BITS }; // = [(Header)32 + (SeqHistory)32-MaxSequenceHistoryLength|(bool)1|?(ServerTime)8|(ReceviedRateByte)8]
+enum { MAX_PACKET_RELIABLE_SEQUENCE_HEADER_BITS = 32 /*PackedHeader*/ + FNetPacketNotify::SequenceHistoryT::MaxSizeInBits };
+enum { MAX_PACKET_INFO_HEADER_BITS = 1 /*bHasPacketInfo*/ + NetConnectionHelper::NumBitsForJitterClockTimeInHeader + 1 /*bHasServerFrameTime*/ + 8 /*ServerFrameTime*/};
+enum { MAX_PACKET_HEADER_BITS = MAX_PACKET_RELIABLE_SEQUENCE_HEADER_BITS + MAX_PACKET_INFO_HEADER_BITS  };
 enum { MAX_PACKET_TRAILER_BITS = 1 };
 
 // 
diff --git a/Engine/Source/Runtime/Engine/Private/NetConnection.cpp b/Engine/Source/Runtime/Engine/Private/NetConnection.cpp
--- a/Engine/Source/Runtime/Engine/Private/NetConnection.cpp
+++ b/Engine/Source/Runtime/Engine/Private/NetConnection.cpp
@@ -115,11 +115,8 @@ namespace UE4_NetConnectionPrivate
 	}
 
 
-	/** Number of bits to use in the packet header for sending the milliseconds on the clock when the packet is sent */
-	constexpr int32 NumBitsForJitterClockTimeInHeader = 10;
-
 	/** Maximum possible clock time value with available bits. If that value is sent than the jitter clock time is ignored */
-	constexpr int32 MaxJitterClockTimeValue = (1 << NumBitsForJitterClockTimeInHeader) - 1;
+	constexpr int32 MaxJitterClockTimeValue = (1 << NetConnectionHelper::NumBitsForJitterClockTimeInHeader) - 1;
 
 	/** Maximum possible precision of the jitter stat. Set to 1 second since we only send the milliseconds of the local clock time */
 	constexpr int32 MaxJitterPrecisionInMS = 1000;
@@ -1715,6 +1712,10 @@ void UNetConnection::WritePacketHeader(FBitWriter& Writer)
 	// Write notification header or refresh the header if used space is the same.
 	bool bWroteHeader = PacketNotify.WriteHeader(Writer, bIsHeaderUpdate);
 
+#if !UE_BUILD_SHIPPING
+	checkf(Writer.GetNumBits() <= MAX_PACKET_RELIABLE_SEQUENCE_HEADER_BITS, TEXT("WritePacketHeader exceeded the max allowed bits. Wrote %d. Max %d"), Writer.GetNumBits(), MAX_PACKET_RELIABLE_SEQUENCE_HEADER_BITS);
+#endif
+
 	// Jump back to where we came from.
 	if (bIsHeaderUpdate)
 	{
@@ -1730,6 +1731,10 @@ void UNetConnection::WritePacketHeader(FBitWriter& Writer)
 
 void UNetConnection::WriteDummyPacketInfo(FBitWriter& Writer)
 {
+#if !UE_BUILD_SHIPPING
+	const int32 BitsWrittenPrePacketInfo = Writer.GetNumBits();
+#endif
+
 	// The first packet of a frame will include the packet info payload
 	const bool bHasPacketInfoPayload = bFlushedNetThisFrame == false;
 	Writer.WriteBit(bHasPacketInfoPayload);
@@ -1739,7 +1744,7 @@ void UNetConnection::WriteDummyPacketInfo(FBitWriter& Writer)
 		// Pre-insert the bits since the final time values will be calculated and inserted right before LowLevelSend
 		HeaderMarkForPacketInfo.Init(Writer);
 		int32 DummyJitterClockTime(UE4_NetConnectionPrivate::MaxJitterClockTimeValue);
-		Writer.SerializeBits(&DummyJitterClockTime, UE4_NetConnectionPrivate::NumBitsForJitterClockTimeInHeader);
+		Writer.SerializeBits(&DummyJitterClockTime, NetConnectionHelper::NumBitsForJitterClockTimeInHeader);
 
 		const uint8 bHasServerFrameTime = Driver->IsServer() ? bLastHasServerFrameTime : (CVarPingExcludeFrameTime.GetValueOnGameThread() > 0 ? 1u : 0u);
 		Writer.WriteBit(bHasServerFrameTime);
@@ -1752,6 +1757,10 @@ void UNetConnection::WriteDummyPacketInfo(FBitWriter& Writer)
 	}
 
 	bSendBufferHasDummyPacketInfo = bHasPacketInfoPayload;
+
+#if !UE_BUILD_SHIPPING
+	checkf(Writer.GetNumBits() - BitsWrittenPrePacketInfo <= MAX_PACKET_INFO_HEADER_BITS, TEXT("WriteDummyPacketInfo exceeded the max allowed bits. Wrote %d. Max %d"), Writer.GetNumBits() - BitsWrittenPrePacketInfo, MAX_PACKET_INFO_HEADER_BITS);
+#endif
 }
 
 void UNetConnection::WriteFinalPacketInfo(FBitWriter& Writer, const double PacketSentTimeInS)
@@ -1789,7 +1798,7 @@ void UNetConnection::WriteFinalPacketInfo(FBitWriter& Writer, const double Packe
 			ClockTimeMilliseconds &= UE4_NetConnectionPrivate::MaxJitterClockTimeValue;
 		}
 
-			Writer.SerializeBits(&ClockTimeMilliseconds, UE4_NetConnectionPrivate::NumBitsForJitterClockTimeInHeader);
+		Writer.SerializeBits(&ClockTimeMilliseconds, NetConnectionHelper::NumBitsForJitterClockTimeInHeader);
 
 		PreviousPacketSendTimeInS = PacketSentTimeInS;
 	}
@@ -2013,7 +2022,7 @@ void UNetConnection::ReceivedPacket( FBitReader& Reader, bool bIsReinjectedPacke
 			{
 				// Read jitter clock time from the packet header
 				int32 PacketJitterClockTimeMS = 0;
-				Reader.SerializeBits(&PacketJitterClockTimeMS, UE4_NetConnectionPrivate::NumBitsForJitterClockTimeInHeader);
+				Reader.SerializeBits(&PacketJitterClockTimeMS, NetConnectionHelper::NumBitsForJitterClockTimeInHeader);
 
 				if (!bIsReinjectedPacket)
 				{
diff --git a/Engine/Source/Runtime/Engine/Public/Net/Util/SequenceHistory.h b/Engine/Source/Runtime/Engine/Public/Net/Util/SequenceHistory.h
--- a/Engine/Source/Runtime/Engine/Public/Net/Util/SequenceHistory.h
+++ b/Engine/Source/Runtime/Engine/Public/Net/Util/SequenceHistory.h
@@ -16,9 +16,10 @@ class TSequenceHistory
 public:	
 	typedef uint32 WordT;
 	
-	static const SIZE_T BitsPerWord = sizeof(WordT) * 8;
-	static const SIZE_T WordCount = HistorySize / BitsPerWord;
-	static const SIZE_T Size = HistorySize;
+	static constexpr SIZE_T BitsPerWord = sizeof(WordT) * 8;
+	static constexpr SIZE_T WordCount = HistorySize / BitsPerWord;
+	static constexpr SIZE_T MaxSizeInBits = WordCount * BitsPerWord;
+	static constexpr SIZE_T Size = HistorySize;
 
 	static_assert(HistorySize > 0, "HistorySize must be > 0");
 	static_assert(HistorySize % BitsPerWord == 0, "InMaxHistorySize must be a modulo of the wordsize");
