Fixed non-unity build in NavigationModifier.h

#rnx

diff --git a/Engine/Source/Runtime/Engine/Public/AI/NavigationModifier.h b/Engine/Source/Runtime/Engine/Public/AI/NavigationModifier.h
--- a/Engine/Source/Runtime/Engine/Public/AI/NavigationModifier.h
+++ b/Engine/Source/Runtime/Engine/Public/AI/NavigationModifier.h
@@ -3,6 +3,7 @@
 #pragma once 
 
 #include "CoreMinimal.h"
+#include "CollisionShape.h"
 #include "Templates/SubclassOf.h"
 #include "UObject/WeakObjectPtrTemplates.h"
 #include "AI/Navigation/NavLinkDefinition.h"
