Fix missing forward declare
#rnx

diff --git a/Engine/Source/Runtime/Engine/Classes/Engine/Level.h b/Engine/Source/Runtime/Engine/Classes/Engine/Level.h
--- a/Engine/Source/Runtime/Engine/Classes/Engine/Level.h
+++ b/Engine/Source/Runtime/Engine/Classes/Engine/Level.h
@@ -30,6 +30,7 @@ class UNavigationDataChunk;
 class UTexture2D;
 struct FLevelCollection;
 class ULevelActorContainer;
+class FLevelPartitionOperationScope;
 
 UINTERFACE()
 class ULevelPartitionInterface : public UInterface
