Separated HLSLMaterialTranslator.h into header and source files for better maintainability.

#rnx

diff --git a/Engine/Source/Runtime/Engine/Private/Materials/HLSLMaterialTranslator.cpp b/Engine/Source/Runtime/Engine/Private/Materials/HLSLMaterialTranslator.cpp
new file mode 100644
--- /dev/null
+++ b/Engine/Source/Runtime/Engine/Private/Materials/HLSLMaterialTranslator.cpp
@@ -0,0 +1,6880 @@
+// Copyright 1998-2019 Epic Games, Inc. All Rights Reserved.
+/*=============================================================================
+	HLSLMaterialTranslator.cpp: Translates material expressions into HLSL code.
+=============================================================================*/
+
+#pragma once
+
+#include "HLSLMaterialTranslator.h"
+
+class Error;
+
+#if WITH_EDITORONLY_DATA
+
+/** @return the number of components in a vector type. */
+static inline uint32 GetNumComponents(EMaterialValueType Type)
+{
+	switch(Type)
+	{
+		case MCT_Float:
+		case MCT_Float1: return 1;
+		case MCT_Float2: return 2;
+		case MCT_Float3: return 3;
+		case MCT_Float4: return 4;
+		default: return 0;
+	}
+}
+
+/** @return the vector type containing a given number of components. */
+static inline EMaterialValueType GetVectorType(uint32 NumComponents)
+{
+	switch(NumComponents)
+	{
+		case 1: return MCT_Float;
+		case 2: return MCT_Float2;
+		case 3: return MCT_Float3;
+		case 4: return MCT_Float4;
+		default: return MCT_Unknown;
+	};
+}
+
+static inline int32 SwizzleComponentToIndex(TCHAR Component)
+{
+	switch (Component)
+	{
+	case TCHAR('x'): case TCHAR('X'): case TCHAR('r'): case TCHAR('R'): return 0;
+	case TCHAR('y'): case TCHAR('Y'): case TCHAR('g'): case TCHAR('G'): return 1;
+	case TCHAR('z'): case TCHAR('Z'): case TCHAR('b'): case TCHAR('B'): return 2;
+	case TCHAR('w'): case TCHAR('W'): case TCHAR('a'): case TCHAR('A'): return 3;
+	default:
+		return -1;
+	}
+}
+
+FHLSLMaterialTranslator::FHLSLMaterialTranslator(FMaterial* InMaterial,
+	FMaterialCompilationOutput& InMaterialCompilationOutput,
+	const FStaticParameterSet& InStaticParameters,
+	EShaderPlatform InPlatform,
+	EMaterialQualityLevel::Type InQualityLevel,
+	ERHIFeatureLevel::Type InFeatureLevel,
+	const ITargetPlatform* InTargetPlatform) //if InTargetPlatform is nullptr, we use the current active
+:	ShaderFrequency(SF_Pixel)
+,	MaterialProperty(MP_EmissiveColor)
+,	CurrentScopeChunks(nullptr)
+,	CurrentScopeID(0u)
+,	NextTempScopeID(SF_NumFrequencies)
+,	Material(InMaterial)
+,	MaterialCompilationOutput(InMaterialCompilationOutput)
+,	StaticParameters(InStaticParameters)
+,	Platform(InPlatform)
+,	QualityLevel(InQualityLevel)
+,	FeatureLevel(InFeatureLevel)
+,	MaterialTemplateLineNumber(INDEX_NONE)
+,	NextSymbolIndex(INDEX_NONE)
+,	NextVertexInterpolatorIndex(0)
+,	CurrentCustomVertexInterpolatorOffset(0)
+,	CompileErrorsSink(nullptr)
+,	CompileErrorExpressionsSink(nullptr)
+,	bSuccess(false)
+,	bCompileForComputeShader(false)
+,	bUsesSceneDepth(false)
+,	bNeedsParticlePosition(false)
+,	bNeedsParticleVelocity(false)
+,	bNeedsParticleTime(false)
+,	bUsesParticleMotionBlur(false)
+,	bNeedsParticleRandom(false)
+,	bUsesSphericalParticleOpacity(false)
+,   bUsesParticleSubUVs(false)
+,	bUsesLightmapUVs(false)
+,	bUsesAOMaterialMask(false)
+,	bUsesSpeedTree(false)
+,	bNeedsWorldPositionExcludingShaderOffsets(false)
+,	bNeedsParticleSize(false)
+,	bNeedsSceneTexturePostProcessInputs(false)
+,	bUsesAtmosphericFog(false)
+,	bUsesSkyAtmosphere(false)
+,	bUsesVertexColor(false)
+,	bUsesParticleColor(false)
+,	bUsesParticleTransform(false)
+,	bUsesVertexPosition(false)
+,	bUsesTransformVector(false)
+,	bCompilingPreviousFrame(false)
+,	bOutputsBasePassVelocities(true)
+,	bUsesPixelDepthOffset(false)
+,   bUsesWorldPositionOffset(false)
+,	bUsesEmissiveColor(false)
+,	bUsesDistanceCullFade(false)
+,	bIsFullyRough(0)
+,	bAllowCodeChunkGeneration(true)
+,	AllocatedUserTexCoords()
+,	AllocatedUserVertexTexCoords()
+,	DynamicParticleParameterMask(0)
+,	NumVtSamples(0)
+,	TargetPlatform(InTargetPlatform)
+{
+	FMemory::Memzero(SharedPixelProperties);
+
+	SharedPixelProperties[MP_Normal] = true;
+	SharedPixelProperties[MP_EmissiveColor] = true;
+	SharedPixelProperties[MP_Opacity] = true;
+	SharedPixelProperties[MP_OpacityMask] = true;
+	SharedPixelProperties[MP_BaseColor] = true;
+	SharedPixelProperties[MP_Metallic] = true;
+	SharedPixelProperties[MP_Specular] = true;
+	SharedPixelProperties[MP_Roughness] = true;
+	SharedPixelProperties[MP_AmbientOcclusion] = true;
+	SharedPixelProperties[MP_Refraction] = true;
+	SharedPixelProperties[MP_PixelDepthOffset] = true;
+	SharedPixelProperties[MP_SubsurfaceColor] = true;
+	SharedPixelProperties[MP_ShadingModel] = true;
+
+	for (int32 Frequency = 0; Frequency < SF_NumFrequencies; ++Frequency)
+	{
+		FunctionStacks[Frequency].Add(new FMaterialFunctionCompileState(nullptr));
+	}
+
+	// Default value for attribute stack added to simplify code when compiling new attributes, see SetMaterialProperty.
+	const FGuid& MissingAttribute = FMaterialAttributeDefinitionMap::GetID(MP_MAX);
+	MaterialAttributesStack.Add(MissingAttribute);
+
+	// Default owner for parameters
+	ParameterOwnerStack.Add(FMaterialParameterInfo());
+
+	if (TargetPlatform == nullptr)
+	{
+		ITargetPlatformManagerModule* TPM = GetTargetPlatformManager();
+		if (TPM)
+		{
+			TargetPlatform = TPM->GetRunningTargetPlatform();
+		}
+	}
+}
+
+FHLSLMaterialTranslator::~FHLSLMaterialTranslator()
+{
+	ClearAllFunctionStacks();
+}
+
+int32 FHLSLMaterialTranslator::GetNumUserTexCoords() const
+{
+	return AllocatedUserTexCoords.FindLast(true) + 1;
+}
+
+int32 FHLSLMaterialTranslator::GetNumUserVertexTexCoords() const
+{
+	return AllocatedUserVertexTexCoords.FindLast(true) + 1;
+}
+
+void FHLSLMaterialTranslator::ClearAllFunctionStacks()
+{
+	for (uint32 Frequency = 0; Frequency < SF_NumFrequencies; ++Frequency)
+	{
+		ClearFunctionStack(Frequency);
+	}
+}
+
+void FHLSLMaterialTranslator::ClearFunctionStack(uint32 Frequency)
+{
+	check(Frequency < SF_NumFrequencies);
+
+	if (FunctionStacks[Frequency].Num() == 0)
+	{
+		return;  // Already cleared (at the end of Translate(), for example)
+	}
+
+	check(FunctionStacks[Frequency].Num() == 1);  // All states should be popped off, leaving only the null state
+	delete FunctionStacks[Frequency][0];
+	FunctionStacks[Frequency].Empty();
+}
+
+void FHLSLMaterialTranslator::AssignTempScope(TArray<FShaderCodeChunk>& InScope)
+{
+	CurrentScopeChunks = &InScope;
+	CurrentScopeID = NextTempScopeID++;
+}
+
+void FHLSLMaterialTranslator::AssignShaderFrequencyScope(EShaderFrequency InShaderFrequency)
+{
+	check(InShaderFrequency < SF_NumFrequencies);
+	check(InShaderFrequency < NextTempScopeID);
+	CurrentScopeChunks = &SharedPropertyCodeChunks[InShaderFrequency];
+	CurrentScopeID = (uint64)InShaderFrequency;
+}
+
+void FHLSLMaterialTranslator::GatherCustomVertexInterpolators(TArray<UMaterialExpression*> Expressions)
+{
+	for (UMaterialExpression* Expression : Expressions)
+	{
+		if (UMaterialExpressionVertexInterpolator* Interpolator = Cast<UMaterialExpressionVertexInterpolator>(Expression))
+		{
+			TArray<FShaderCodeChunk> CustomExpressionChunks;
+			AssignTempScope(CustomExpressionChunks);
+
+			// Errors are appended to a temporary pool as it's not known at this stage which interpolators are required
+			CompileErrorsSink = &Interpolator->CompileErrors;
+			CompileErrorExpressionsSink = &Interpolator->CompileErrorExpressions;
+
+			// Compile node and store those successfully translated
+			int32 Ret = Interpolator->CompileInput(this, NextVertexInterpolatorIndex);
+			if (Ret != INDEX_NONE)
+			{
+				CustomVertexInterpolators.AddUnique(Interpolator);
+				NextVertexInterpolatorIndex++;
+			}
+
+			// Restore error handling
+			CompileErrorsSink = nullptr;
+			CompileErrorExpressionsSink = nullptr;
+
+			// Each interpolator chain must be handled as an independent compile
+			for (FMaterialFunctionCompileState* FunctionStack : FunctionStacks[SF_Vertex])
+			{
+				FunctionStack->Reset();
+			}
+		}
+		else if (UMaterialExpressionMaterialFunctionCall* FunctionCall = Cast<UMaterialExpressionMaterialFunctionCall>(Expression))
+		{
+			if (FunctionCall->MaterialFunction)
+			{
+				FMaterialFunctionCompileState LocalState(FunctionCall);
+				FunctionCall->LinkFunctionIntoCaller(this);
+				PushFunction(&LocalState);
+
+				if (const TArray<UMaterialExpression*>* FunctionExpressions = FunctionCall->MaterialFunction->GetFunctionExpressions())
+				{
+					GatherCustomVertexInterpolators(*FunctionExpressions);
+				}
+
+				FMaterialFunctionCompileState* CompileState = PopFunction();
+				check(CompileState->ExpressionStack.Num() == 0);
+				FunctionCall->UnlinkFunctionFromCaller(this);
+			}
+		}
+		else if (UMaterialExpressionMaterialAttributeLayers* LayersExpression = Cast<UMaterialExpressionMaterialAttributeLayers>(Expression))
+		{
+			const FMaterialLayersFunctions* OverrideLayers = StaticMaterialLayersParameter(LayersExpression->ParameterName);
+			if (OverrideLayers)
+			{
+				LayersExpression->OverrideLayerGraph(OverrideLayers);
+			}
+
+			if (LayersExpression->bIsLayerGraphBuilt)
+			{
+				for (auto* Layer : LayersExpression->LayerCallers)
+				{
+					if (Layer && Layer->MaterialFunction)
+					{
+						FMaterialFunctionCompileState LocalState(Layer);
+						Layer->LinkFunctionIntoCaller(this);
+						PushFunction(&LocalState);
+
+						if (const TArray<UMaterialExpression*>* FunctionExpressions = Layer->MaterialFunction->GetFunctionExpressions())
+						{
+							GatherCustomVertexInterpolators(*FunctionExpressions);
+						}
+
+						FMaterialFunctionCompileState* CompileState = PopFunction();
+						check(CompileState->ExpressionStack.Num() == 0);
+						Layer->UnlinkFunctionFromCaller(this);
+					}
+				}
+
+				for (auto* Blend : LayersExpression->BlendCallers)
+				{
+					if (Blend && Blend->MaterialFunction)
+					{
+						FMaterialFunctionCompileState LocalState(Blend);
+						Blend->LinkFunctionIntoCaller(this);
+						PushFunction(&LocalState);
+
+						if (const TArray<UMaterialExpression*>* FunctionExpressions = Blend->MaterialFunction->GetFunctionExpressions())
+						{
+							GatherCustomVertexInterpolators(*FunctionExpressions);
+						}
+
+						FMaterialFunctionCompileState* CompileState = PopFunction();
+						check(CompileState->ExpressionStack.Num() == 0);
+						Blend->UnlinkFunctionFromCaller(this);
+					}
+				}
+			}
+
+			if (OverrideLayers)
+			{
+				LayersExpression->OverrideLayerGraph(nullptr);
+			}
+		}
+	}
+}
+
+void FHLSLMaterialTranslator::CompileCustomOutputs(TArray<UMaterialExpressionCustomOutput*>& CustomOutputExpressions, TSet<UClass*>& SeenCustomOutputExpressionsClasses, bool bIsBeforeAttributes)
+{
+	for (UMaterialExpressionCustomOutput* CustomOutput : CustomOutputExpressions)
+	{
+		if (CustomOutput->HasCustomSourceOutput() || CustomOutput->ShouldCompileBeforeAttributes() != bIsBeforeAttributes)
+		{
+			continue;
+		}
+
+		if (!CustomOutput->AllowMultipleCustomOutputs() && SeenCustomOutputExpressionsClasses.Contains(CustomOutput->GetClass()))
+		{
+			Errorf(TEXT("The material can contain only one %s node"), *CustomOutput->GetDescription());
+		}
+		else
+		{
+			SeenCustomOutputExpressionsClasses.Add(CustomOutput->GetClass());
+			int32 NumOutputs = CustomOutput->GetNumOutputs();
+
+			if (CustomOutput->NeedsCustomOutputDefines())
+			{
+				ResourcesString += FString::Printf(TEXT("#define NUM_MATERIAL_OUTPUTS_%s %d\r\n"), *CustomOutput->GetFunctionName().ToUpper(), NumOutputs);
+			}
+
+			if (NumOutputs > 0)
+			{
+				for (int32 Index = 0; Index < NumOutputs; Index++)
+				{
+					{
+						ClearFunctionStack(SF_Pixel);
+						FunctionStacks[SF_Pixel].Add(new FMaterialFunctionCompileState(nullptr));
+					}
+					MaterialProperty = MP_MAX; // Indicates we're not compiling any material property.
+					ShaderFrequency = SF_Pixel;
+					TArray<FShaderCodeChunk> CustomExpressionChunks;
+					AssignTempScope(CustomExpressionChunks);
+					CustomOutput->Compile(this, Index);
+				}
+
+				ClearFunctionStack(SF_Pixel);
+				FunctionStacks[SF_Pixel].Add(new FMaterialFunctionCompileState(nullptr));
+			}
+		}
+	}
+}
+
+EMaterialExpressionVisitResult FHLSLMaterialTranslator::VisitExpressionsRecursive(TArray<UMaterialExpression*> Expressions, IMaterialExpressionVisitor& InVisitor)
+{
+	EMaterialExpressionVisitResult VisitResult = MVR_CONTINUE;
+	for (UMaterialExpression* Expression : Expressions)
+	{
+		VisitResult = InVisitor.Visit(Expression);
+		if (VisitResult == MVR_STOP)
+		{
+			break;
+		}
+
+		if (UMaterialExpressionMaterialFunctionCall* FunctionCall = Cast<UMaterialExpressionMaterialFunctionCall>(Expression))
+		{
+			if (FunctionCall->MaterialFunction)
+			{
+				FMaterialFunctionCompileState LocalState(FunctionCall);
+				FunctionCall->LinkFunctionIntoCaller(this);
+				PushFunction(&LocalState);
+
+				if (const TArray<UMaterialExpression*>* FunctionExpressions = FunctionCall->MaterialFunction->GetFunctionExpressions())
+				{
+					VisitResult = VisitExpressionsRecursive(*FunctionExpressions, InVisitor);
+				}
+
+				FMaterialFunctionCompileState* CompileState = PopFunction();
+				check(CompileState->ExpressionStack.Num() == 0);
+				FunctionCall->UnlinkFunctionFromCaller(this);
+
+				if (VisitResult == MVR_STOP)
+				{
+					break;
+				}
+			}
+		}
+		else if (UMaterialExpressionMaterialAttributeLayers* LayersExpression = Cast<UMaterialExpressionMaterialAttributeLayers>(Expression))
+		{
+			const FMaterialLayersFunctions* OverrideLayers = StaticMaterialLayersParameter(LayersExpression->ParameterName);
+			if (OverrideLayers)
+			{
+				LayersExpression->OverrideLayerGraph(OverrideLayers);
+			}
+
+			if (LayersExpression->bIsLayerGraphBuilt)
+			{
+				for (auto* Layer : LayersExpression->LayerCallers)
+				{
+					if (Layer && Layer->MaterialFunction)
+					{
+						FMaterialFunctionCompileState LocalState(Layer);
+						Layer->LinkFunctionIntoCaller(this);
+						PushFunction(&LocalState);
+
+						if (const TArray<UMaterialExpression*>* FunctionExpressions = Layer->MaterialFunction->GetFunctionExpressions())
+						{
+							VisitResult = VisitExpressionsRecursive(*FunctionExpressions, InVisitor);
+						}
+
+						FMaterialFunctionCompileState* CompileState = PopFunction();
+						check(CompileState->ExpressionStack.Num() == 0);
+						Layer->UnlinkFunctionFromCaller(this);
+
+						if (VisitResult == MVR_STOP)
+						{
+							break;
+						}
+					}
+				}
+
+				for (auto* Blend : LayersExpression->BlendCallers)
+				{
+					if (Blend && Blend->MaterialFunction)
+					{
+						FMaterialFunctionCompileState LocalState(Blend);
+						Blend->LinkFunctionIntoCaller(this);
+						PushFunction(&LocalState);
+
+						if (const TArray<UMaterialExpression*>* FunctionExpressions = Blend->MaterialFunction->GetFunctionExpressions())
+						{
+							VisitResult = VisitExpressionsRecursive(*FunctionExpressions, InVisitor);
+						}
+
+						FMaterialFunctionCompileState* CompileState = PopFunction();
+						check(CompileState->ExpressionStack.Num() == 0);
+						Blend->UnlinkFunctionFromCaller(this);
+
+						if (VisitResult == MVR_STOP)
+						{
+							break;
+						}
+					}
+				}
+			}
+
+			if (OverrideLayers)
+			{
+				LayersExpression->OverrideLayerGraph(nullptr);
+			}
+
+			if (VisitResult == MVR_STOP)
+			{
+				break;
+			}
+		}
+	}
+
+	return VisitResult;
+}
+
+EMaterialExpressionVisitResult FHLSLMaterialTranslator::VisitExpressionsForProperty(EMaterialProperty InProperty, IMaterialExpressionVisitor& InVisitor)
+{
+	UMaterialInterface *MatIf = Material->GetMaterialInterface();
+	// Some proxies return null for this. But the main one we are interested in doesn't
+	if (MatIf)
+	{
+		TArray<UMaterialExpression*> InputExpressions;
+		MatIf->GetMaterial()->GetExpressionsInPropertyChain(InProperty, InputExpressions, &StaticParameters);
+		return VisitExpressionsRecursive(InputExpressions, InVisitor);
+	}
+	return MVR_STOP;
+}
+
+void FHLSLMaterialTranslator::ValidateVtPropertyLimits()
+{
+	class FFindVirtualTextureVisitor : public IMaterialExpressionVisitor
+	{
+	public:
+		virtual EMaterialExpressionVisitResult Visit(UMaterialExpression* InExpression) override
+		{
+			if (UMaterialExpressionTextureBase *TextureExpr = Cast<UMaterialExpressionTextureBase>(InExpression))
+			{
+				if (IsVirtualSamplerType(TextureExpr->SamplerType))
+				{
+					FoundVirtualTexture = true;
+					return MVR_STOP;
+				}
+			}
+			return MVR_CONTINUE;
+		}
+
+		bool FoundVirtualTexture = false;
+	};
+	
+	for (uint32 PropertyIndex = 0u; PropertyIndex < MP_MAX; ++PropertyIndex)
+	{
+		const EMaterialProperty PropertyToValidate = (EMaterialProperty)PropertyIndex;
+		if (PropertyToValidate == MP_MaterialAttributes || PropertyToValidate == MP_CustomOutput)
+		{
+			// These properties are "special", attempting to pass them to FMaterialAttributeDefinitionMap::GetShaderFrequency() will generate log spam
+			continue;
+		}
+
+		const EShaderFrequency ShaderFrequencyToValidate = FMaterialAttributeDefinitionMap::GetShaderFrequency(PropertyToValidate);
+
+		// check to see if this is a property that doesn't support virtual texture connections
+		if (PropertyToValidate == MP_OpacityMask || ShaderFrequencyToValidate != SF_Pixel)
+		{
+			FFindVirtualTextureVisitor Visitor;
+			VisitExpressionsForProperty(PropertyToValidate, Visitor);
+			if (Visitor.FoundVirtualTexture)
+			{
+				// virtual texture connected to an invalid property, report the correct error
+#if MATERIAL_OPACITYMASK_DOESNT_SUPPORT_VIRTUALTEXTURE
+				if (PropertyToValidate == MP_OpacityMask)
+				{
+					Errorf(TEXT("Sampling a virtual texture is currently not supported when connected to the Opacity Mask material attribute."));
+				}
+#endif
+			}
+		}
+	}
+}
+ 
+bool FHLSLMaterialTranslator::Translate()
+{
+	STAT(double HLSLTranslateTime = 0);
+	{
+		SCOPE_SECONDS_COUNTER(HLSLTranslateTime);
+		bSuccess = true;
+
+		// WARNING: No compile outputs should be stored on the UMaterial / FMaterial / FMaterialResource, unless they are transient editor-only data (like error expressions)
+		// Compile outputs that need to be saved must be stored in MaterialCompilationOutput, which will be saved to the DDC.
+
+		Material->CompileErrors.Empty();
+		Material->ErrorExpressions.Empty();
+
+		bCompileForComputeShader = Material->IsLightFunction();
+
+		// Generate code:
+		// Normally one would expect the generator to emit something like
+		//		float Local0 = ...
+		//		...
+		//		float Local3= ...
+		//		...
+		//		float Localn= ...
+		//		PixelMaterialInputs.EmissiveColor = Local0 + ...
+		//		PixelMaterialInputs.Normal = Local3 * ...
+		// However because the Normal can be used in the middle of generating other Locals (which happens when using a node like PixelNormalWS)
+		// instead we generate this:
+		//		float Local0 = ...
+		//		...
+		//		float Local3= ...
+		//		PixelMaterialInputs.Normal = Local3 * ...
+		//		...
+		//		float Localn= ...
+		//		PixelMaterialInputs.EmissiveColor = Local0 + ...
+		// in other words, compile Normal first, then emit all the expressions up to the last one Normal requires;
+		// assign the normal into the shared struct, then emit the remaining expressions; finally assign the rest of the shared struct inputs.
+		// Inputs that are not shared, have false in the SharedPixelProperties array, and those ones will emit the full code.
+
+		int32 NormalCodeChunkEnd = -1;
+		int32 Chunk[CompiledMP_MAX];
+			
+		memset(Chunk, INDEX_NONE, sizeof(Chunk));
+
+		// Translate all custom vertex interpolators before main attributes so type information is available
+		{
+			CustomVertexInterpolators.Empty();
+			CurrentCustomVertexInterpolatorOffset = 0;
+			NextVertexInterpolatorIndex = 0;
+			MaterialProperty = MP_MAX;
+			ShaderFrequency = SF_Vertex;
+
+			TArray<UMaterialExpression*> Expressions;
+			Material->GatherExpressionsForCustomInterpolators(Expressions);
+			GatherCustomVertexInterpolators(Expressions);
+
+			// Reset shared stack data
+			while (FunctionStacks[SF_Vertex].Num() > 1)
+			{
+				FMaterialFunctionCompileState* Stack = FunctionStacks[SF_Vertex].Pop(false);
+				delete Stack;
+			}
+			FunctionStacks[SF_Vertex][0]->Reset();
+
+			// Whilst expression list is available, apply node count limits
+			int32 NumMaterialLayersAttributes = 0;
+			for (UMaterialExpression* Expression : Expressions)
+			{
+				if (UMaterialExpressionMaterialAttributeLayers* Layers = Cast<UMaterialExpressionMaterialAttributeLayers>(Expression))
+				{
+					++NumMaterialLayersAttributes;
+
+					if (NumMaterialLayersAttributes > 1)
+					{
+						Errorf(TEXT("Materials can contain only one Material Attribute Layers node."));
+						break;
+					}
+				}
+			}
+		}
+
+			
+		const EShaderFrequency NormalShaderFrequency = FMaterialAttributeDefinitionMap::GetShaderFrequency(MP_Normal);
+		const EMaterialDomain Domain = Material->GetMaterialDomain();
+		const EBlendMode BlendMode = Material->GetBlendMode();
+
+		// Gather the implementation for any custom output expressions
+		TArray<UMaterialExpressionCustomOutput*> CustomOutputExpressions;
+		Material->GatherCustomOutputExpressions(CustomOutputExpressions);
+		TSet<UClass*> SeenCustomOutputExpressionsClasses;
+
+		// Some custom outputs must be pre-compiled so they can be re-used as shared inputs
+		CompileCustomOutputs(CustomOutputExpressions, SeenCustomOutputExpressionsClasses, true);			
+
+		// Normal must always be compiled first; this will ensure its chunk calculations are the first to be added
+		{
+			// Verify that start chunk is 0
+			check(SharedPropertyCodeChunks[NormalShaderFrequency].Num() == 0);
+			Chunk[MP_Normal]					= Material->CompilePropertyAndSetMaterialProperty(MP_Normal					,this);
+			NormalCodeChunkEnd = SharedPropertyCodeChunks[NormalShaderFrequency].Num();
+		}
+
+		// Validate some things on the VT system. Since generated code for expressions shared between multiple properties
+		// (e.g. a texture sample connected to both diffuse and opacity mask) is reused we can't check based on the MaterialProperty
+		// variable inside the actual code generation pass. So we do a pre-pass over it here.
+		if (UseVirtualTexturing(FeatureLevel, TargetPlatform))
+		{
+			ValidateVtPropertyLimits();
+		}
+
+		// Rest of properties
+		Chunk[MP_EmissiveColor]					= Material->CompilePropertyAndSetMaterialProperty(MP_EmissiveColor			,this);
+		Chunk[MP_DiffuseColor]					= Material->CompilePropertyAndSetMaterialProperty(MP_DiffuseColor			,this);
+		Chunk[MP_SpecularColor]					= Material->CompilePropertyAndSetMaterialProperty(MP_SpecularColor			,this);
+		Chunk[MP_BaseColor]						= Material->CompilePropertyAndSetMaterialProperty(MP_BaseColor				,this);
+		Chunk[MP_Metallic]						= Material->CompilePropertyAndSetMaterialProperty(MP_Metallic				,this);
+		Chunk[MP_Specular]						= Material->CompilePropertyAndSetMaterialProperty(MP_Specular				,this);
+		Chunk[MP_Roughness]						= Material->CompilePropertyAndSetMaterialProperty(MP_Roughness				,this);
+		Chunk[MP_Opacity]						= Material->CompilePropertyAndSetMaterialProperty(MP_Opacity				,this);
+		Chunk[MP_OpacityMask]					= Material->CompilePropertyAndSetMaterialProperty(MP_OpacityMask			,this);
+		Chunk[MP_WorldPositionOffset]			= Material->CompilePropertyAndSetMaterialProperty(MP_WorldPositionOffset	,this);
+		Chunk[MP_WorldDisplacement]				= Material->CompilePropertyAndSetMaterialProperty(MP_WorldDisplacement		,this);
+		Chunk[MP_TessellationMultiplier]		= Material->CompilePropertyAndSetMaterialProperty(MP_TessellationMultiplier	,this);			
+
+		// Make sure to compile this property before using ShadingModelsFromCompilation
+		Chunk[MP_ShadingModel]					= Material->CompilePropertyAndSetMaterialProperty(MP_ShadingModel			,this);
+			
+		// Get shading models from material.
+		FMaterialShadingModelField MaterialShadingModels = Material->GetShadingModels(); 
+
+		// If the material gets its shading model from material expressions and we have compiled one or more shading model expressions, 
+		// then use that shading model field instead. It's the most optimal set of shading models
+		if (Material->IsShadingModelFromMaterialExpression() && ShadingModelsFromCompilation.IsValid())
+		{
+			MaterialShadingModels = ShadingModelsFromCompilation;
+		}
+
+		if (Domain == MD_Surface && IsSubsurfaceShadingModel(MaterialShadingModels))
+		{
+			// Note we don't test for the blend mode as you can have a translucent material using the subsurface shading model
+
+			// another ForceCast as CompilePropertyAndSetMaterialProperty() can return MCT_Float which we don't want here
+			int32 SubsurfaceColor = Material->CompilePropertyAndSetMaterialProperty(MP_SubsurfaceColor, this);
+			SubsurfaceColor = ForceCast(SubsurfaceColor, FMaterialAttributeDefinitionMap::GetValueType(MP_SubsurfaceColor), MFCF_ExactMatch | MFCF_ReplicateValue);
+
+			static FName NameSubsurfaceProfile(TEXT("__SubsurfaceProfile"));
+
+			// 1.0f is is a not used profile - later this gets replaced with the actual profile
+			int32 CodeSubsurfaceProfile = ForceCast(ScalarParameter(NameSubsurfaceProfile, 1.0f), MCT_Float1);
+
+			Chunk[MP_SubsurfaceColor] = AppendVector(SubsurfaceColor, CodeSubsurfaceProfile);		
+		}
+
+		Chunk[MP_CustomData0]					= Material->CompilePropertyAndSetMaterialProperty(MP_CustomData0		,this);
+		Chunk[MP_CustomData1]					= Material->CompilePropertyAndSetMaterialProperty(MP_CustomData1		,this);
+		Chunk[MP_AmbientOcclusion]				= Material->CompilePropertyAndSetMaterialProperty(MP_AmbientOcclusion	,this);
+
+		if (IsTranslucentBlendMode(BlendMode) || MaterialShadingModels.HasShadingModel(MSM_SingleLayerWater))
+		{
+			int32 UserRefraction = ForceCast(Material->CompilePropertyAndSetMaterialProperty(MP_Refraction, this), MCT_Float1);
+			int32 RefractionDepthBias = ForceCast(ScalarParameter(FName(TEXT("RefractionDepthBias")), Material->GetRefractionDepthBiasValue()), MCT_Float1);
+
+			Chunk[MP_Refraction] = AppendVector(UserRefraction, RefractionDepthBias);
+		}
+
+		if (bCompileForComputeShader)
+		{
+			Chunk[CompiledMP_EmissiveColorCS]	= Material->CompilePropertyAndSetMaterialProperty(MP_EmissiveColor, this, SF_Compute);
+		}
+
+		if (Chunk[MP_WorldPositionOffset] != INDEX_NONE)
+		{
+			// Only calculate previous WPO if there is a current WPO
+			Chunk[CompiledMP_PrevWorldPositionOffset] = Material->CompilePropertyAndSetMaterialProperty(MP_WorldPositionOffset, this, SF_Vertex, true);
+		}
+
+		Chunk[MP_PixelDepthOffset] = Material->CompilePropertyAndSetMaterialProperty(MP_PixelDepthOffset, this);
+
+
+		ResourcesString = TEXT("");
+
+#if HANDLE_CUSTOM_OUTPUTS_AS_MATERIAL_ATTRIBUTES
+		// Handle custom outputs when using material attribute output
+		if (Material->HasMaterialAttributesConnected())
+		{
+			TArray<FMaterialCustomOutputAttributeDefintion> CustomAttributeList;
+			FMaterialAttributeDefinitionMap::GetCustomAttributeList(CustomAttributeList);
+			TArray<FShaderCodeChunk> CustomExpressionChunks;
+
+			for (FMaterialCustomOutputAttributeDefintion& Attribute : CustomAttributeList)
+			{
+				// Compile all outputs for attribute
+				bool bValidResultCompiled = false;
+				int32 NumOutputs = 1;//CustomOutput->GetNumOutputs();
+
+				for (int32 OutputIndex = 0; OutputIndex < NumOutputs; ++OutputIndex)
+				{
+					MaterialProperty = Attribute.Property;
+					ShaderFrequency = Attribute.ShaderFrequency;
+					FunctionStacks[ShaderFrequency].Empty();
+					FunctionStacks[ShaderFrequency].Add(FMaterialFunctionCompileState(nullptr));
+
+					CustomExpressionChunks.Empty();
+					AssignTempScope(CustomExpressionChunks);
+					int32 Result = Material->CompileCustomAttribute(Attribute.AttributeID, this);
+
+					// Consider attribute used if varies from default value
+					if (Result != INDEX_NONE)
+					{
+						bool bValueNonDefault = true;
+
+						if (FMaterialUniformExpression* Expression = GetParameterUniformExpression(Result))
+						{
+							FLinearColor Value;
+							FMaterialRenderContext DummyContext(nullptr, *Material, nullptr);
+							Expression->GetNumberValue(DummyContext, Value);
+
+							bool bEqualValue = Value.R == Attribute.DefaultValue.X;
+							bEqualValue &= Value.G == Attribute.DefaultValue.Y || Attribute.ValueType < MCT_Float2;
+							bEqualValue &= Value.B == Attribute.DefaultValue.Z || Attribute.ValueType < MCT_Float3;
+							bEqualValue &= Value.A == Attribute.DefaultValue.W || Attribute.ValueType < MCT_Float4;
+
+							if (Expression->IsConstant() && bEqualValue)
+							{
+								bValueNonDefault = false;
+							}
+						}
+
+						// Valid, non-default value so generate shader code
+						if (bValueNonDefault)
+						{
+							GenerateCustomAttributeCode(OutputIndex, Result, Attribute.ValueType, Attribute.FunctionName);
+							bValidResultCompiled = true;
+						}
+					}
+				}
+
+				// If used, add compile data
+				if (bValidResultCompiled)
+				{
+					ResourcesString += FString::Printf(TEXT("#define NUM_MATERIAL_OUTPUTS_%s %d\r\n"), *Attribute.FunctionName.ToUpper(), NumOutputs);
+				}
+			}
+		}
+		else
+#endif // #if HANDLE_CUSTOM_OUTPUTS_AS_MATERIAL_ATTRIBUTES
+		{
+			CompileCustomOutputs(CustomOutputExpressions, SeenCustomOutputExpressionsClasses, false);
+		}
+
+		// Output the implementation for any custom expressions we will call below.
+		for (int32 ExpressionIndex = 0; ExpressionIndex < CustomExpressionImplementations.Num(); ExpressionIndex++)
+		{
+			ResourcesString += CustomExpressionImplementations[ExpressionIndex] + "\r\n\r\n";
+		}
+
+		// No more calls to non-vertex shader CompilePropertyAndSetMaterialProperty beyond this point
+		const uint32 SavedNumUserTexCoords = GetNumUserTexCoords();
+
+		for (uint32 CustomUVIndex = MP_CustomizedUVs0; CustomUVIndex <= MP_CustomizedUVs7; CustomUVIndex++)
+		{
+			// Only compile custom UV inputs for UV channels requested by the pixel shader inputs
+			// Any unconnected inputs will have a texcoord generated for them in Material->CompileProperty, which will pass through the vertex (uncustomized) texture coordinates
+			// Note: this is using NumUserTexCoords, which is set by translating all the pixel properties above
+			if (CustomUVIndex - MP_CustomizedUVs0 < SavedNumUserTexCoords)
+			{
+				Chunk[CustomUVIndex] = Material->CompilePropertyAndSetMaterialProperty((EMaterialProperty)CustomUVIndex, this);
+			}
+		}
+
+		// Translation is designed to have a code chunk generation phase followed by several passes that only has readonly access to the code chunks.
+		// At this point we mark the code chunk generation complete.
+		bAllowCodeChunkGeneration = false;
+
+		bUsesEmissiveColor = IsMaterialPropertyUsed(MP_EmissiveColor, Chunk[MP_EmissiveColor], FLinearColor(0, 0, 0, 0), 3);
+		bUsesPixelDepthOffset = (AllowPixelDepthOffset(Platform) && IsMaterialPropertyUsed(MP_PixelDepthOffset, Chunk[MP_PixelDepthOffset], FLinearColor(0, 0, 0, 0), 1))
+			|| (Domain == MD_DeferredDecal && Material->GetDecalBlendMode() == DBM_Volumetric_DistanceFunction);
+
+		bUsesWorldPositionOffset = IsMaterialPropertyUsed(MP_WorldPositionOffset, Chunk[MP_WorldPositionOffset], FLinearColor(0, 0, 0, 0), 3);
+		MaterialCompilationOutput.bModifiesMeshPosition = bUsesPixelDepthOffset || bUsesWorldPositionOffset;
+		MaterialCompilationOutput.bUsesWorldPositionOffset = bUsesWorldPositionOffset;
+		MaterialCompilationOutput.bUsesPixelDepthOffset = bUsesPixelDepthOffset;
+
+		// Fully rough if we have a roughness code chunk and it's constant and evaluates to 1.
+		bIsFullyRough = Chunk[MP_Roughness] != INDEX_NONE && IsMaterialPropertyUsed(MP_Roughness, Chunk[MP_Roughness], FLinearColor(1, 0, 0, 0), 1) == false;
+
+		if (BlendMode == BLEND_Modulate && MaterialShadingModels.IsLit() && !Material->IsDeferredDecal())
+		{
+			Errorf(TEXT("Dynamically lit translucency is not supported for BLEND_Modulate materials."));
+		}
+
+		if (Domain == MD_Surface)
+		{
+			if (BlendMode == BLEND_Modulate && Material->IsTranslucencyAfterDOFEnabled())
+			{
+				Errorf(TEXT("Translucency after DOF with BLEND_Modulate is not supported. Consider using BLEND_Translucent with black emissive"));
+			}
+		}
+
+		// Don't allow opaque and masked materials to scene depth as the results are undefined
+		if (bUsesSceneDepth && Domain != MD_PostProcess && !IsTranslucentBlendMode(BlendMode))
+		{
+			Errorf(TEXT("Only transparent or postprocess materials can read from scene depth."));
+		}
+
+		if (bUsesSceneDepth)
+		{
+			MaterialCompilationOutput.SetIsSceneTextureUsed(PPI_SceneDepth);
+		}
+
+		MaterialCompilationOutput.bUsesDistanceCullFade = bUsesDistanceCullFade;
+
+		if (MaterialCompilationOutput.RequiresSceneColorCopy())
+		{
+			if (Domain != MD_Surface)
+			{
+				Errorf(TEXT("Only 'surface' material domain can use the scene color node."));
+			}
+			else if (!IsTranslucentBlendMode(BlendMode))
+			{
+				Errorf(TEXT("Only translucent materials can use the scene color node."));
+			}
+		}
+
+		if (BlendMode == BLEND_AlphaHoldout && !MaterialShadingModels.IsUnlit())
+		{
+			Errorf(TEXT("Alpha Holdout blend mode must use unlit shading model."));
+		}
+
+		if (Domain == MD_Volume && BlendMode != BLEND_Additive)
+		{
+			Errorf(TEXT("Volume materials must use an Additive blend mode."));
+		}
+
+		if (Material->IsLightFunction() && BlendMode != BLEND_Opaque)
+		{
+			Errorf(TEXT("Light function materials must be opaque."));
+		}
+
+		if (Material->IsLightFunction() && MaterialShadingModels.IsLit())
+		{
+			Errorf(TEXT("Light function materials must use unlit."));
+		}
+
+		if (Domain == MD_PostProcess && MaterialShadingModels.IsLit())
+		{
+			Errorf(TEXT("Post process materials must use unlit."));
+		}
+
+		if (Material->AllowNegativeEmissiveColor() && MaterialShadingModels.IsLit())
+		{
+			Errorf(TEXT("Only unlit materials can output negative emissive color."));
+		}
+
+		if (Material->IsSky() && (!MaterialShadingModels.IsUnlit() || BlendMode!=BLEND_Opaque))
+		{
+			Errorf(TEXT("Sky materials must be opaque and unlit. They are expected to completely replace the background."));
+		}
+
+		if (MaterialShadingModels.HasShadingModel(MSM_SingleLayerWater))
+		{
+			if (BlendMode != BLEND_Opaque)
+			{
+				Errorf(TEXT("SingleLayerWater materials must be opaque."));
+			}
+			if (!MaterialShadingModels.HasOnlyShadingModel(MSM_SingleLayerWater))
+			{
+				Errorf(TEXT("SingleLayerWater materials cannot be combined with other shading models.")); // Simply untested for now
+			}
+			if (Material->GetMaterialInterface() && !Material->GetMaterialInterface()->GetMaterial()->HasAnyExpressionsInMaterialAndFunctionsOfType<UMaterialExpressionSingleLayerWaterMaterialOutput>())
+			{
+				Errorf(TEXT("SingleLayerWater materials requires the use of SingleLayerWaterMaterial output node."));
+			}
+		}
+
+		bool bDBufferAllowed = IsUsingDBuffers(Platform);
+		bool bDBufferBlendMode = IsDBufferDecalBlendMode((EDecalBlendMode)Material->GetDecalBlendMode());
+
+		if (bDBufferBlendMode && !bDBufferAllowed)
+		{
+			// Error feedback for when the decal would not be displayed due to project settings
+			Errorf(TEXT("DBuffer decal blend modes are only supported when the 'DBuffer Decals' Rendering Project setting is enabled."));
+		}
+
+		if (Domain == MD_DeferredDecal && BlendMode != BLEND_Translucent)
+		{
+			// We could make the change for the user but it would be confusing when going to DeferredDecal and back
+			// or we would have to pay a performance cost to make the change more transparently.
+			// The change saves performance as with translucency we don't need to test for MeshDecals in all opaque rendering passes
+			Errorf(TEXT("Material using the DeferredDecal domain need to use the BlendModel Translucent (this saves performance)"));
+		}
+
+		if (MaterialCompilationOutput.bNeedsSceneTextures)
+		{
+			if (Domain != MD_DeferredDecal && Domain != MD_PostProcess)
+			{
+				if (BlendMode == BLEND_Opaque || BlendMode == BLEND_Masked)
+				{
+					// In opaque pass, none of the textures are available
+					Errorf(TEXT("SceneTexture expressions cannot be used in opaque materials"));
+				}
+				else if (bNeedsSceneTexturePostProcessInputs)
+				{
+					Errorf(TEXT("SceneTexture expressions cannot use post process inputs or scene color in non post process domain materials"));
+				}
+			}
+		}
+
+		// Catch any modifications to NumUserTexCoords that will not seen by customized UVs
+		check(SavedNumUserTexCoords == GetNumUserTexCoords());
+
+		FString InterpolatorsOffsetsDefinitionCode;
+		TBitArray<> FinalAllocatedCoords = GetVertexInterpolatorsOffsets(InterpolatorsOffsetsDefinitionCode);
+
+		// Finished compilation, verify final interpolator count restrictions
+		if (CurrentCustomVertexInterpolatorOffset > 0)
+		{
+			const int32 MaxNumScalars = (FeatureLevel == ERHIFeatureLevel::ES2) ? 3 * 2 : 8 * 2;
+			const int32 TotalUsedScalars = FinalAllocatedCoords.FindLast(true) + 1;
+
+ 			if (TotalUsedScalars > MaxNumScalars)
+			{
+				Errorf(TEXT("Maximum number of custom vertex interpolators exceeded. (%i / %i scalar values) (TexCoord: %i scalars, Custom: %i scalars)"),
+					TotalUsedScalars, MaxNumScalars, GetNumUserTexCoords() * 2, CurrentCustomVertexInterpolatorOffset);
+			}
+		}
+
+		MaterialCompilationOutput.NumUsedUVScalars = GetNumUserTexCoords() * 2;
+		MaterialCompilationOutput.NumUsedCustomInterpolatorScalars = CurrentCustomVertexInterpolatorOffset;
+
+		// Do Normal Chunk first
+		{
+			GetFixedParameterCode(
+				0,
+				NormalCodeChunkEnd,
+				Chunk[MP_Normal],
+				SharedPropertyCodeChunks[NormalShaderFrequency],
+				TranslatedCodeChunkDefinitions[MP_Normal],
+				TranslatedCodeChunks[MP_Normal]);
+
+			// Always gather MP_Normal definitions as they can be shared by other properties
+			if (TranslatedCodeChunkDefinitions[MP_Normal].IsEmpty())
+			{
+				TranslatedCodeChunkDefinitions[MP_Normal] = GetDefinitions(SharedPropertyCodeChunks[NormalShaderFrequency], 0, NormalCodeChunkEnd);
+			}
+		}
+
+		// Now the rest, skipping Normal
+		for(uint32 PropertyId = 0; PropertyId < MP_MAX; ++PropertyId)
+		{
+			if (PropertyId == MP_MaterialAttributes || PropertyId == MP_Normal || PropertyId == MP_CustomOutput)
+			{
+				continue;
+			}
+
+			const EShaderFrequency PropertyShaderFrequency = FMaterialAttributeDefinitionMap::GetShaderFrequency((EMaterialProperty)PropertyId);
+
+			int32 StartChunk = 0;
+			if (PropertyShaderFrequency == NormalShaderFrequency && SharedPixelProperties[PropertyId])
+			{
+				// When processing shared properties, do not generate the code before the Normal was generated as those are already handled
+				StartChunk = NormalCodeChunkEnd;
+			}
+
+			GetFixedParameterCode(
+				StartChunk,
+				SharedPropertyCodeChunks[PropertyShaderFrequency].Num(),
+				Chunk[PropertyId],
+				SharedPropertyCodeChunks[PropertyShaderFrequency],
+				TranslatedCodeChunkDefinitions[PropertyId],
+				TranslatedCodeChunks[PropertyId]);
+		}
+
+		for(uint32 PropertyId = MP_MAX; PropertyId < CompiledMP_MAX; ++PropertyId)
+		{
+			switch(PropertyId)
+			{
+			case CompiledMP_EmissiveColorCS:
+			    if (bCompileForComputeShader)
+				{
+					GetFixedParameterCode(Chunk[PropertyId], SharedPropertyCodeChunks[SF_Compute], TranslatedCodeChunkDefinitions[PropertyId], TranslatedCodeChunks[PropertyId]);
+				}
+				break;
+			case CompiledMP_PrevWorldPositionOffset:
+				{
+					GetFixedParameterCode(Chunk[PropertyId], SharedPropertyCodeChunks[SF_Vertex], TranslatedCodeChunkDefinitions[PropertyId], TranslatedCodeChunks[PropertyId]);
+				}
+				break;
+			default: check(0);
+				break;
+			}
+		}
+
+		// Output the implementation for any custom output expressions
+		for (int32 ExpressionIndex = 0; ExpressionIndex < CustomOutputImplementations.Num(); ExpressionIndex++)
+		{
+			ResourcesString += CustomOutputImplementations[ExpressionIndex] + "\r\n\r\n";
+		}
+
+		LoadShaderSourceFileChecked(TEXT("/Engine/Private/MaterialTemplate.ush"), MaterialTemplate);
+
+		// Find the string index of the '#line' statement in MaterialTemplate.usf
+		const int32 LineIndex = MaterialTemplate.Find(TEXT("#line"), ESearchCase::CaseSensitive);
+		check(LineIndex != INDEX_NONE);
+
+		// Count line endings before the '#line' statement
+		MaterialTemplateLineNumber = INDEX_NONE;
+		int32 StartPosition = LineIndex + 1;
+		do 
+		{
+			MaterialTemplateLineNumber++;
+			// Using \n instead of LINE_TERMINATOR as not all of the lines are terminated consistently
+			// Subtract one from the last found line ending index to make sure we skip over it
+			StartPosition = MaterialTemplate.Find(TEXT("\n"), ESearchCase::CaseSensitive, ESearchDir::FromEnd, StartPosition - 1);
+		} 
+		while (StartPosition != INDEX_NONE);
+		check(MaterialTemplateLineNumber != INDEX_NONE);
+		// At this point MaterialTemplateLineNumber is one less than the line number of the '#line' statement
+		// For some reason we have to add 2 more to the #line value to get correct error line numbers from D3DXCompileShader
+		MaterialTemplateLineNumber += 3;
+
+		MaterialCompilationOutput.UniformExpressionSet.SetParameterCollections(ParameterCollections);
+
+		// Create the material uniform buffer struct.
+		MaterialCompilationOutput.UniformExpressionSet.CreateBufferStruct();
+
+		// Store the number of unique VT samples
+		MaterialCompilationOutput.EstimatedNumVirtualTextureLookups = NumVtSamples;
+	}
+	ClearAllFunctionStacks();
+		
+	INC_FLOAT_STAT_BY(STAT_ShaderCompiling_HLSLTranslation,(float)HLSLTranslateTime);
+	return bSuccess;
+}
+
+void FHLSLMaterialTranslator::GetMaterialEnvironment(EShaderPlatform InPlatform, FShaderCompilerEnvironment& OutEnvironment)
+{
+	if (bNeedsParticlePosition || Material->ShouldGenerateSphericalParticleNormals() || bUsesSphericalParticleOpacity)
+	{
+		OutEnvironment.SetDefine(TEXT("NEEDS_PARTICLE_POSITION"), 1);
+	}
+
+	if (bNeedsParticleVelocity || Material->IsUsedWithNiagaraMeshParticles())
+	{
+		OutEnvironment.SetDefine(TEXT("NEEDS_PARTICLE_VELOCITY"), 1);
+	}
+
+	if (DynamicParticleParameterMask)
+	{
+		OutEnvironment.SetDefine(TEXT("USE_DYNAMIC_PARAMETERS"), 1);
+		OutEnvironment.SetDefine(TEXT("DYNAMIC_PARAMETERS_MASK"), DynamicParticleParameterMask);
+	}
+
+	if (bNeedsParticleTime)
+	{
+		OutEnvironment.SetDefine(TEXT("NEEDS_PARTICLE_TIME"), 1);
+	}
+
+	if (bUsesParticleMotionBlur)
+	{
+		OutEnvironment.SetDefine(TEXT("USES_PARTICLE_MOTION_BLUR"), 1);
+	}
+
+	if (bNeedsParticleRandom)
+	{
+		OutEnvironment.SetDefine(TEXT("NEEDS_PARTICLE_RANDOM"), 1);
+	}
+
+	if (bUsesSphericalParticleOpacity)
+	{
+		OutEnvironment.SetDefine(TEXT("SPHERICAL_PARTICLE_OPACITY"), TEXT("1"));
+	}
+
+	if (bUsesParticleSubUVs)
+	{
+		OutEnvironment.SetDefine(TEXT("USE_PARTICLE_SUBUVS"), TEXT("1"));
+	}
+
+	if (bUsesLightmapUVs)
+	{
+		OutEnvironment.SetDefine(TEXT("LIGHTMAP_UV_ACCESS"),TEXT("1"));
+	}
+
+	if (bUsesAOMaterialMask)
+	{
+		OutEnvironment.SetDefine(TEXT("USES_AO_MATERIAL_MASK"),TEXT("1"));
+	}
+
+	if (bUsesSpeedTree)
+	{
+		OutEnvironment.SetDefine(TEXT("USES_SPEEDTREE"),TEXT("1"));
+	}
+
+	if (bNeedsWorldPositionExcludingShaderOffsets)
+	{
+		OutEnvironment.SetDefine(TEXT("NEEDS_WORLD_POSITION_EXCLUDING_SHADER_OFFSETS"), TEXT("1"));
+	}
+
+	if (bNeedsParticleSize)
+	{
+		OutEnvironment.SetDefine(TEXT("NEEDS_PARTICLE_SIZE"), TEXT("1"));
+	}
+
+	if (MaterialCompilationOutput.bNeedsSceneTextures)
+	{
+		OutEnvironment.SetDefine(TEXT("NEEDS_SCENE_TEXTURES"), TEXT("1"));
+	}
+	if (MaterialCompilationOutput.bUsesEyeAdaptation)
+	{
+		OutEnvironment.SetDefine(TEXT("USES_EYE_ADAPTATION"), TEXT("1"));
+	}
+
+	if (MaterialCompilationOutput.bHasRuntimeVirtualTextureOutput)
+	{
+		OutEnvironment.SetDefine(TEXT("VIRTUAL_TEXTURE_OUTPUT"), 1);
+	}
+		
+	// @todo MetalMRT: Remove this hack and implement proper atmospheric-fog solution for Metal MRT...
+	OutEnvironment.SetDefine(TEXT("MATERIAL_ATMOSPHERIC_FOG"), !IsMetalMRTPlatform(InPlatform) ? bUsesAtmosphericFog : 0);
+	OutEnvironment.SetDefine(TEXT("MATERIAL_SKY_ATMOSPHERE"), bUsesSkyAtmosphere);
+	OutEnvironment.SetDefine(TEXT("INTERPOLATE_VERTEX_COLOR"), bUsesVertexColor);
+	OutEnvironment.SetDefine(TEXT("NEEDS_PARTICLE_COLOR"), bUsesParticleColor); 
+	OutEnvironment.SetDefine(TEXT("NEEDS_PARTICLE_TRANSFORM"), bUsesParticleTransform);
+	OutEnvironment.SetDefine(TEXT("USES_TRANSFORM_VECTOR"), bUsesTransformVector);
+	OutEnvironment.SetDefine(TEXT("WANT_PIXEL_DEPTH_OFFSET"), bUsesPixelDepthOffset);
+	if (IsMetalPlatform(InPlatform))
+	{
+		OutEnvironment.SetDefine(TEXT("USES_WORLD_POSITION_OFFSET"), bUsesWorldPositionOffset);
+	}
+	OutEnvironment.SetDefine(TEXT("USES_EMISSIVE_COLOR"), bUsesEmissiveColor);
+	// Distortion uses tangent space transform 
+	OutEnvironment.SetDefine(TEXT("USES_DISTORTION"), Material->IsDistorted()); 
+
+	OutEnvironment.SetDefine(TEXT("MATERIAL_ENABLE_TRANSLUCENCY_FOGGING"), Material->ShouldApplyFogging());
+	OutEnvironment.SetDefine(TEXT("MATERIAL_FORCE_SKIP_AERIAL_PERSPECTIVE"), Material->IsSky());
+	OutEnvironment.SetDefine(TEXT("MATERIAL_COMPUTE_FOG_PER_PIXEL"), Material->ComputeFogPerPixel());
+	OutEnvironment.SetDefine(TEXT("MATERIAL_FULLY_ROUGH"), bIsFullyRough || Material->IsFullyRough());
+
+	// Count the number of VTStacks (each stack will allocate a feedback slot)
+	OutEnvironment.SetDefine(TEXT("NUM_VIRTUALTEXTURE_SAMPLES"), VTStacks.Num());
+
+	// Setup defines to map each VT stack to either 1 or 2 page table textures, depending on how many layers it uses
+	for (int i = 0; i < VTStacks.Num(); ++i)
+	{
+		const FMaterialVirtualTextureStack& Stack = MaterialCompilationOutput.UniformExpressionSet.VTStacks[i];
+		const FString PageTableValue = (Stack.GetNumLayers() > 4u)
+			? FString::Printf(TEXT("Material.VirtualTexturePageTable0_%d, Material.VirtualTexturePageTable1_%d"), i, i)
+			: FString::Printf(TEXT("Material.VirtualTexturePageTable0_%d"), i);
+
+		OutEnvironment.SetDefine(*FString::Printf(TEXT("VIRTUALTEXTURE_PAGETABLE_%d"), i), *PageTableValue);
+	}
+
+	for (int32 CollectionIndex = 0; CollectionIndex < ParameterCollections.Num(); CollectionIndex++)
+	{
+		// Add uniform buffer declarations for any parameter collections referenced
+		const FString CollectionName = FString::Printf(TEXT("MaterialCollection%u"), CollectionIndex);
+		// This can potentially become an issue for MaterialCollection Uniform Buffers if they ever get non-numeric resources (eg Textures), as
+		// OutEnvironment.ResourceTableMap has a map by name, and the N ParameterCollection Uniform Buffers ALL are names "MaterialCollection"
+		// (and the hlsl cbuffers are named MaterialCollection0, etc, so the names don't match the layout)
+		FShaderUniformBufferParameter::ModifyCompilationEnvironment(*CollectionName, ParameterCollections[CollectionIndex]->GetUniformBufferStruct(), InPlatform, OutEnvironment);
+	}
+	OutEnvironment.SetDefine(TEXT("IS_MATERIAL_SHADER"), TEXT("1"));
+
+	// Set all the shading models for this material here 
+	FMaterialShadingModelField ShadingModels = Material->GetShadingModels();
+
+	// If the material gets its shading model from the material expressions, then we use the result from the compilation (assuming it's valid).
+	// This result will potentially be tighter than what GetShadingModels() returns, because it only picks up the shading models from the expressions that get compiled for a specific feature level and quality level
+	// For example, the material might have shading models behind static switches. GetShadingModels() will return both the true and the false paths from that switch, whereas the shading model field from the compilation will only contain the actual shading model selected 
+	if (Material->IsShadingModelFromMaterialExpression() && ShadingModelsFromCompilation.IsValid())
+	{
+		// Shading models fetched from the compilation of the expression graph
+		ShadingModels = ShadingModelsFromCompilation;
+	}
+
+	ensure(ShadingModels.IsValid());
+
+	if (ShadingModels.IsLit())
+	{	
+		int NumSetMaterials = 0;
+		if (ShadingModels.HasShadingModel(MSM_DefaultLit))
+		{
+			OutEnvironment.SetDefine(TEXT("MATERIAL_SHADINGMODEL_DEFAULT_LIT"), TEXT("1"));
+			NumSetMaterials++;
+		}
+		if (ShadingModels.HasShadingModel(MSM_Subsurface))
+		{
+			OutEnvironment.SetDefine(TEXT("MATERIAL_SHADINGMODEL_SUBSURFACE"), TEXT("1"));
+			NumSetMaterials++;
+		}
+		if (ShadingModels.HasShadingModel(MSM_PreintegratedSkin))
+		{
+			OutEnvironment.SetDefine(TEXT("MATERIAL_SHADINGMODEL_PREINTEGRATED_SKIN"), TEXT("1"));
+			NumSetMaterials++;
+		}
+		if (ShadingModels.HasShadingModel(MSM_SubsurfaceProfile))
+		{
+			OutEnvironment.SetDefine(TEXT("MATERIAL_SHADINGMODEL_SUBSURFACE_PROFILE"), TEXT("1"));
+			NumSetMaterials++;
+		}
+		if (ShadingModels.HasShadingModel(MSM_ClearCoat))
+		{
+			OutEnvironment.SetDefine(TEXT("MATERIAL_SHADINGMODEL_CLEAR_COAT"), TEXT("1"));
+			NumSetMaterials++;
+		}
+		if (ShadingModels.HasShadingModel(MSM_TwoSidedFoliage))
+		{
+			OutEnvironment.SetDefine(TEXT("MATERIAL_SHADINGMODEL_TWOSIDED_FOLIAGE"), TEXT("1"));
+			NumSetMaterials++;
+		}
+		if (ShadingModels.HasShadingModel(MSM_Hair))
+		{
+			OutEnvironment.SetDefine(TEXT("MATERIAL_SHADINGMODEL_HAIR"), TEXT("1"));
+			NumSetMaterials++;
+		}
+		if (ShadingModels.HasShadingModel(MSM_Cloth))
+		{
+			OutEnvironment.SetDefine(TEXT("MATERIAL_SHADINGMODEL_CLOTH"), TEXT("1"));
+			NumSetMaterials++;
+		}
+		if (ShadingModels.HasShadingModel(MSM_Eye))
+		{
+			OutEnvironment.SetDefine(TEXT("MATERIAL_SHADINGMODEL_EYE"), TEXT("1"));
+			NumSetMaterials++;
+		}
+		if (ShadingModels.HasShadingModel(MSM_SingleLayerWater))
+		{
+			OutEnvironment.SetDefine(TEXT("MATERIAL_SHADINGMODEL_SINGLELAYERWATER"), TEXT("1"));
+			NumSetMaterials++;
+		}
+
+		if(ShadingModels.HasShadingModel(MSM_SingleLayerWater) && (IsSwitchPlatform(Platform) || IsPS4Platform(Platform) || Platform == SP_XBOXONE_D3D12))
+		{
+			OutEnvironment.SetDefine(TEXT("DISABLE_FORWARD_LOCAL_LIGHTS"), TEXT("1"));
+		}
+
+		// This is to have switch use the simple single layer water shading similar to mobile: no dynamic lights, only sun and sky, no distortion, no colored transmittance on background, no custom depth read.
+		const bool bSingleLayerWaterUsesSimpleShading = IsSwitchPlatform(InPlatform) && IsForwardShadingEnabled(InPlatform);
+		if (ShadingModels.HasShadingModel(MSM_SingleLayerWater) && bSingleLayerWaterUsesSimpleShading)
+		{
+			OutEnvironment.SetDefine(TEXT("SINGLE_LAYER_WATER_SIMPLE_FORWARD"), TEXT("1"));
+		}
+
+		if (NumSetMaterials == 1)
+		{
+			OutEnvironment.SetDefine(TEXT("MATERIAL_SINGLE_SHADINGMODEL"), TEXT("1"));
+		}
+
+		ensure(NumSetMaterials != 0);
+		if (NumSetMaterials == 0)
+		{
+			// Should not really end up here
+			UE_LOG(LogMaterial, Warning, TEXT("Unknown material shading model(s). Setting to MSM_DefaultLit"));
+			OutEnvironment.SetDefine(TEXT("MATERIAL_SHADINGMODEL_DEFAULT_LIT"),TEXT("1"));
+		}
+	}
+	else
+	{
+		// Unlit shading model can only exist by itself
+		OutEnvironment.SetDefine(TEXT("MATERIAL_SINGLE_SHADINGMODEL"), TEXT("1"));
+		OutEnvironment.SetDefine(TEXT("MATERIAL_SHADINGMODEL_UNLIT"), TEXT("1"));
+	}
+}
+
+// Assign custom interpolators to slots, packing them as much as possible in unused slots.
+TBitArray<> FHLSLMaterialTranslator::GetVertexInterpolatorsOffsets(FString& VertexInterpolatorsOffsetsDefinitionCode) const
+{
+	TBitArray<> AllocatedCoords = AllocatedUserTexCoords; // Don't mess with the already assigned sets of UV coords
+
+	int32 CurrentSlot = INDEX_NONE;
+	int32 EndAllocatedSlot = INDEX_NONE;
+
+	auto GetNextUVSlot = [&CurrentSlot, &EndAllocatedSlot, &AllocatedCoords]() -> int32
+	{
+		if (CurrentSlot == EndAllocatedSlot)
+		{
+			CurrentSlot = AllocatedCoords.FindAndSetFirstZeroBit();
+			if (CurrentSlot == INDEX_NONE)
+			{
+				CurrentSlot = AllocatedCoords.Add(true);
+			}
+
+			// Track one slot per component (u,v)
+			const int32 NUM_COMPONENTS = 2;
+			CurrentSlot *= NUM_COMPONENTS;
+			EndAllocatedSlot = CurrentSlot + NUM_COMPONENTS;
+		}
+
+		int32 ResultUVSlot = CurrentSlot / 2;
+		CurrentSlot++;
+
+		return ResultUVSlot;
+	};
+
+	TArray<UMaterialExpressionVertexInterpolator*> SortedInterpolators;
+	Algo::TransformIf(CustomVertexInterpolators, 
+						SortedInterpolators, 
+						[](const UMaterialExpressionVertexInterpolator* Interpolator) { return Interpolator && Interpolator->InterpolatorIndex != INDEX_NONE && Interpolator->InterpolatorOffset != INDEX_NONE; },
+						[](UMaterialExpressionVertexInterpolator* Interpolator) { return Interpolator; });
+						
+	SortedInterpolators.Sort([](const UMaterialExpressionVertexInterpolator& LHS, const UMaterialExpressionVertexInterpolator& RHS)  { return LHS.InterpolatorOffset < RHS.InterpolatorOffset; });
+		
+	for (UMaterialExpressionVertexInterpolator* Interpolator : SortedInterpolators)
+	{
+		int32 Index = Interpolator->InterpolatorIndex;
+
+		const EMaterialValueType Type = Interpolator->InterpolatedType == MCT_Float ? MCT_Float1 : Interpolator->InterpolatedType;
+
+		VertexInterpolatorsOffsetsDefinitionCode += LINE_TERMINATOR;
+		VertexInterpolatorsOffsetsDefinitionCode += FString::Printf(TEXT("#define VERTEX_INTERPOLATOR_%i_TEXCOORDS_X\t%i") LINE_TERMINATOR, Index, GetNextUVSlot());
+
+		if (Type >= MCT_Float2)
+		{
+			VertexInterpolatorsOffsetsDefinitionCode += FString::Printf(TEXT("#define VERTEX_INTERPOLATOR_%i_TEXCOORDS_Y\t%i") LINE_TERMINATOR, Index, GetNextUVSlot());
+
+			if (Type >= MCT_Float3)
+			{
+				VertexInterpolatorsOffsetsDefinitionCode += FString::Printf(TEXT("#define VERTEX_INTERPOLATOR_%i_TEXCOORDS_Z\t%i") LINE_TERMINATOR, Index, GetNextUVSlot());
+
+				if (Type == MCT_Float4)
+				{
+					VertexInterpolatorsOffsetsDefinitionCode += FString::Printf(TEXT("#define VERTEX_INTERPOLATOR_%i_TEXCOORDS_W\t%i") LINE_TERMINATOR, Index, GetNextUVSlot());
+				}
+			}
+		}
+			
+		VertexInterpolatorsOffsetsDefinitionCode += LINE_TERMINATOR;
+	}
+
+	return AllocatedCoords;
+}
+
+void FHLSLMaterialTranslator::GetSharedInputsMaterialCode(FString& PixelMembersDeclaration, FString& NormalAssignment, FString& PixelMembersInitializationEpilog)
+{
+	{
+		int32 LastProperty = -1;
+
+		FString PixelInputInitializerValues;
+		FString NormalInitializerValue;
+
+		for (int32 PropertyIndex = 0; PropertyIndex < MP_MAX; ++PropertyIndex)
+		{
+			// Skip non-shared properties
+			if (!SharedPixelProperties[PropertyIndex])
+			{
+				continue;
+			}
+
+			const EMaterialProperty Property = (EMaterialProperty)PropertyIndex;
+			check(FMaterialAttributeDefinitionMap::GetShaderFrequency(Property) == SF_Pixel);
+			// Special case MP_SubsurfaceColor as the actual property is a combination of the color and the profile but we don't want to expose the profile
+			const FString PropertyName = Property == MP_SubsurfaceColor ? "Subsurface" : FMaterialAttributeDefinitionMap::GetDisplayName(Property);
+			check(PropertyName.Len() > 0);				
+			const EMaterialValueType Type = Property == MP_SubsurfaceColor ? MCT_Float4 : FMaterialAttributeDefinitionMap::GetValueType(Property);
+
+			// Normal requires its own separate initializer
+			if (Property == MP_Normal)
+			{
+				NormalInitializerValue = FString::Printf(TEXT("\tPixelMaterialInputs.%s = %s;\n"), *PropertyName, *TranslatedCodeChunks[Property]);
+			}
+			else
+			{
+				if (TranslatedCodeChunkDefinitions[Property].Len() > 0)
+				{
+					if (LastProperty >= 0)
+					{
+						// Verify that all code chunks have the same contents
+						check(TranslatedCodeChunkDefinitions[Property].Len() == TranslatedCodeChunkDefinitions[LastProperty].Len());
+					}
+
+					LastProperty = Property;
+				}
+
+				PixelInputInitializerValues += FString::Printf(TEXT("\tPixelMaterialInputs.%s = %s;\n"), *PropertyName, *TranslatedCodeChunks[Property]);
+			}
+
+			PixelMembersDeclaration += FString::Printf(TEXT("\t%s %s;\n"), HLSLTypeString(Type), *PropertyName);
+		}
+
+		NormalAssignment = NormalInitializerValue;
+		if (LastProperty != -1)
+		{
+			PixelMembersInitializationEpilog += TranslatedCodeChunkDefinitions[LastProperty] + TEXT("\n");
+		}
+
+		PixelMembersInitializationEpilog += PixelInputInitializerValues;
+	}
+}
+
+FString FHLSLMaterialTranslator::GetMaterialShaderCode()
+{	
+	// use "/Engine/Private/MaterialTemplate.ush" to create the functions to get data (e.g. material attributes) and code (e.g. material expressions to create specular color) from C++
+	FLazyPrintf LazyPrintf(*MaterialTemplate);
+
+	// Assign slots to vertex interpolators
+	FString VertexInterpolatorsOffsetsDefinition;
+	TBitArray<> FinalAllocatedCoords = GetVertexInterpolatorsOffsets(VertexInterpolatorsOffsetsDefinition);
+
+	const uint32 NumUserVertexTexCoords = GetNumUserVertexTexCoords();
+	const uint32 NumUserTexCoords = GetNumUserTexCoords();
+	const uint32 NumCustomVectors = FMath::DivideAndRoundUp((uint32)CurrentCustomVertexInterpolatorOffset, 2u);
+	const uint32 NumTexCoordVectors = FinalAllocatedCoords.FindLast(true) + 1;
+
+	LazyPrintf.PushParam(*FString::Printf(TEXT("%u"),NumUserVertexTexCoords));
+	LazyPrintf.PushParam(*FString::Printf(TEXT("%u"),NumUserTexCoords));
+	LazyPrintf.PushParam(*FString::Printf(TEXT("%u"),NumCustomVectors));
+	LazyPrintf.PushParam(*FString::Printf(TEXT("%u"),NumTexCoordVectors));
+
+	LazyPrintf.PushParam(*VertexInterpolatorsOffsetsDefinition);
+
+	// Stores the shared shader results member declarations
+	FString PixelMembersDeclaration;
+
+	FString NormalAssignment;
+
+	// Stores the code to initialize all inputs after MP_Normal
+	FString PixelMembersSetupAndAssignments;
+
+	GetSharedInputsMaterialCode(PixelMembersDeclaration, NormalAssignment, PixelMembersSetupAndAssignments);
+
+	LazyPrintf.PushParam(*PixelMembersDeclaration);
+
+	LazyPrintf.PushParam(*ResourcesString);
+
+	if (bCompileForComputeShader)
+	{
+		LazyPrintf.PushParam(*GenerateFunctionCode(CompiledMP_EmissiveColorCS));
+	}
+	else
+	{
+		LazyPrintf.PushParam(TEXT("return 0"));
+	}
+
+	LazyPrintf.PushParam(*FString::Printf(TEXT("return %.5f"), Material->GetTranslucencyDirectionalLightingIntensity()));
+
+	LazyPrintf.PushParam(*FString::Printf(TEXT("return %.5f"), Material->GetTranslucentShadowDensityScale()));
+	LazyPrintf.PushParam(*FString::Printf(TEXT("return %.5f"), Material->GetTranslucentSelfShadowDensityScale()));
+	LazyPrintf.PushParam(*FString::Printf(TEXT("return %.5f"), Material->GetTranslucentSelfShadowSecondDensityScale()));
+	LazyPrintf.PushParam(*FString::Printf(TEXT("return %.5f"), Material->GetTranslucentSelfShadowSecondOpacity()));
+	LazyPrintf.PushParam(*FString::Printf(TEXT("return %.5f"), Material->GetTranslucentBackscatteringExponent()));
+
+	{
+		FLinearColor Extinction = Material->GetTranslucentMultipleScatteringExtinction();
+
+		LazyPrintf.PushParam(*FString::Printf(TEXT("return MaterialFloat3(%.5f, %.5f, %.5f)"), Extinction.R, Extinction.G, Extinction.B));
+	}
+
+	LazyPrintf.PushParam(*FString::Printf(TEXT("return %.5f"), Material->GetOpacityMaskClipValue()));
+
+	LazyPrintf.PushParam(*GenerateFunctionCode(MP_WorldPositionOffset));
+	LazyPrintf.PushParam(*GenerateFunctionCode(CompiledMP_PrevWorldPositionOffset));
+	LazyPrintf.PushParam(*GenerateFunctionCode(MP_WorldDisplacement));
+	LazyPrintf.PushParam(*FString::Printf(TEXT("return %.5f"), Material->GetMaxDisplacement()));
+	LazyPrintf.PushParam(*GenerateFunctionCode(MP_TessellationMultiplier));
+	LazyPrintf.PushParam(*GenerateFunctionCode(MP_CustomData0));
+	LazyPrintf.PushParam(*GenerateFunctionCode(MP_CustomData1));
+
+	// Print custom texture coordinate assignments
+	FString CustomUVAssignments;
+
+	int32 LastProperty = -1;
+	for (uint32 CustomUVIndex = 0; CustomUVIndex < NumUserTexCoords; CustomUVIndex++)
+	{
+		if (CustomUVIndex == 0)
+		{
+			CustomUVAssignments += TranslatedCodeChunkDefinitions[MP_CustomizedUVs0 + CustomUVIndex];
+		}
+
+		if (TranslatedCodeChunkDefinitions[MP_CustomizedUVs0 + CustomUVIndex].Len() > 0)
+		{
+			if (LastProperty >= 0)
+			{
+				check(TranslatedCodeChunkDefinitions[LastProperty].Len() == TranslatedCodeChunkDefinitions[MP_CustomizedUVs0 + CustomUVIndex].Len());
+			}
+			LastProperty = MP_CustomizedUVs0 + CustomUVIndex;
+		}
+		CustomUVAssignments += FString::Printf(TEXT("\tOutTexCoords[%u] = %s;") LINE_TERMINATOR, CustomUVIndex, *TranslatedCodeChunks[MP_CustomizedUVs0 + CustomUVIndex]);
+	}
+
+	LazyPrintf.PushParam(*CustomUVAssignments);
+
+	// Print custom vertex shader interpolator assignments
+	FString CustomInterpolatorAssignments;
+
+	for (UMaterialExpressionVertexInterpolator* Interpolator : CustomVertexInterpolators)
+	{
+		if (Interpolator->InterpolatorOffset != INDEX_NONE)
+		{
+			check(Interpolator->InterpolatorIndex != INDEX_NONE);
+			check(Interpolator->InterpolatedType & MCT_Float);
+
+			const EMaterialValueType Type = Interpolator->InterpolatedType == MCT_Float ? MCT_Float1 : Interpolator->InterpolatedType;
+			const TCHAR* Swizzle[2] = { TEXT("x"), TEXT("y") };
+			const int32 Offset = Interpolator->InterpolatorOffset;
+			const int32 Index = Interpolator->InterpolatorIndex;
+
+			// Note: We reference the UV define directly to avoid having to pre-accumulate UV counts before property translation
+			CustomInterpolatorAssignments += FString::Printf(TEXT("\tOutTexCoords[VERTEX_INTERPOLATOR_%i_TEXCOORDS_X].%s = VertexInterpolator%i(Parameters).x;") LINE_TERMINATOR, Index, Swizzle[Offset%2], Index);
+				
+			if (Type >= MCT_Float2)
+			{
+				CustomInterpolatorAssignments += FString::Printf(TEXT("\tOutTexCoords[VERTEX_INTERPOLATOR_%i_TEXCOORDS_Y].%s = VertexInterpolator%i(Parameters).y;") LINE_TERMINATOR, Index, Swizzle[(Offset+1)%2], Index);
+
+				if (Type >= MCT_Float3)
+				{
+					CustomInterpolatorAssignments += FString::Printf(TEXT("\tOutTexCoords[VERTEX_INTERPOLATOR_%i_TEXCOORDS_Z].%s = VertexInterpolator%i(Parameters).z;") LINE_TERMINATOR, Index, Swizzle[(Offset+2)%2], Index);
+
+					if (Type == MCT_Float4)
+					{
+						CustomInterpolatorAssignments += FString::Printf(TEXT("\tOutTexCoords[VERTEX_INTERPOLATOR_%i_TEXCOORDS_W].%s = VertexInterpolator%i(Parameters).w;") LINE_TERMINATOR, Index, Swizzle[(Offset+3)%2], Index);
+					}
+				}
+			}
+		}
+	}
+
+	LazyPrintf.PushParam(*CustomInterpolatorAssignments);
+
+	// Initializers required for Normal
+	LazyPrintf.PushParam(*TranslatedCodeChunkDefinitions[MP_Normal]);
+	LazyPrintf.PushParam(*NormalAssignment);
+	// Finally the rest of common code followed by assignment into each input
+	LazyPrintf.PushParam(*PixelMembersSetupAndAssignments);
+
+	LazyPrintf.PushParam(*FString::Printf(TEXT("%u"),MaterialTemplateLineNumber));
+
+	return LazyPrintf.GetResultString();
+}
+
+
+// ========== PROTECTED: ========== //
+
+bool FHLSLMaterialTranslator::IsMaterialPropertyUsed(EMaterialProperty Property, int32 PropertyChunkIndex, const FLinearColor& ReferenceValue, int32 NumComponents)
+{
+	bool bPropertyUsed = false;
+
+	if (PropertyChunkIndex == -1)
+	{
+		bPropertyUsed = false;
+	}
+	else
+	{
+		int32 Frequency = (int32)FMaterialAttributeDefinitionMap::GetShaderFrequency(Property);
+		FShaderCodeChunk& PropertyChunk = SharedPropertyCodeChunks[Frequency][PropertyChunkIndex];
+
+		// Determine whether the property is used. 
+		// If the output chunk has a uniform expression, it is constant, and GetNumberValue returns the default property value then property isn't used.
+		bPropertyUsed = true;
+
+		if( PropertyChunk.UniformExpression && PropertyChunk.UniformExpression->IsConstant() )
+		{
+			FLinearColor Value;
+			FMaterialRenderContext DummyContext(nullptr, *Material, nullptr);
+			PropertyChunk.UniformExpression->GetNumberValue(DummyContext, Value);
+
+			if ((NumComponents < 1 || Value.R == ReferenceValue.R)
+				&& (NumComponents < 2 || Value.G == ReferenceValue.G)
+				&& (NumComponents < 3 || Value.B == ReferenceValue.B)
+				&& (NumComponents < 4 || Value.A == ReferenceValue.A))
+			{
+				bPropertyUsed = false;
+			}
+		}
+	}
+
+	return bPropertyUsed;
+}
+
+// only used by GetMaterialShaderCode()
+// @param Index ECompiledMaterialProperty or EMaterialProperty
+FString FHLSLMaterialTranslator::GenerateFunctionCode(uint32 Index) const
+{
+	check(Index < CompiledMP_MAX);
+	return TranslatedCodeChunkDefinitions[Index] + TEXT("	return ") + TranslatedCodeChunks[Index] + TEXT(";");
+}
+
+// GetParameterCode
+FString FHLSLMaterialTranslator::GetParameterCode(int32 Index, const TCHAR* Default)
+{
+	if(Index == INDEX_NONE && Default)
+	{
+		return Default;
+	}
+
+	checkf(Index >= 0 && Index < CurrentScopeChunks->Num(), TEXT("Index %d/%d, Platform=%d"), Index, CurrentScopeChunks->Num(), (int)Platform);
+	const FShaderCodeChunk& CodeChunk = (*CurrentScopeChunks)[Index];
+	if((CodeChunk.UniformExpression && CodeChunk.UniformExpression->IsConstant()) || CodeChunk.bInline)
+	{
+		// Constant uniform expressions and code chunks which are marked to be inlined are accessed via Definition
+		return CodeChunk.Definition;
+	}
+	else
+	{
+		if (CodeChunk.UniformExpression)
+		{
+			// If the code chunk has a uniform expression, create a new code chunk to access it
+			const int32 AccessedIndex = AccessUniformExpression(Index);
+			const FShaderCodeChunk& AccessedCodeChunk = (*CurrentScopeChunks)[AccessedIndex];
+			if(AccessedCodeChunk.bInline)
+			{
+				// Handle the accessed code chunk being inlined
+				return AccessedCodeChunk.Definition;
+			}
+			// Return the symbol used to reference this code chunk
+			check(AccessedCodeChunk.SymbolName.Len() > 0);
+			return AccessedCodeChunk.SymbolName;
+		}
+			
+		// Return the symbol used to reference this code chunk
+		check(CodeChunk.SymbolName.Len() > 0);
+		return CodeChunk.SymbolName;
+	}
+}
+
+uint64 FHLSLMaterialTranslator::GetParameterHash(int32 Index)
+{
+	if (Index == INDEX_NONE)
+	{
+		return 0u;
+	}
+
+	checkf(Index >= 0 && Index < CurrentScopeChunks->Num(), TEXT("Index %d/%d, Platform=%d"), Index, CurrentScopeChunks->Num(), (int)Platform);
+	const FShaderCodeChunk& CodeChunk = (*CurrentScopeChunks)[Index];
+
+	if (CodeChunk.UniformExpression && !CodeChunk.UniformExpression->IsConstant())
+	{
+		// Non-constant uniform expressions are accessed through a separate code chunk...need to give the hash of that
+		const int32 AccessedIndex = AccessUniformExpression(Index);
+		const FShaderCodeChunk& AccessedCodeChunk = (*CurrentScopeChunks)[AccessedIndex];
+		return AccessedCodeChunk.Hash;
+	}
+
+	return CodeChunk.Hash;
+}
+
+/** Creates a string of all definitions needed for the given material input. */
+FString FHLSLMaterialTranslator::GetDefinitions(TArray<FShaderCodeChunk>& CodeChunks, int32 StartChunk, int32 EndChunk) const
+{
+	FString Definitions;
+	for (int32 ChunkIndex = StartChunk; ChunkIndex < EndChunk; ChunkIndex++)
+	{
+		const FShaderCodeChunk& CodeChunk = CodeChunks[ChunkIndex];
+		// Uniform expressions (both constant and variable) and inline expressions don't have definitions.
+		if (!CodeChunk.UniformExpression && !CodeChunk.bInline)
+		{
+			Definitions += CodeChunk.Definition;
+		}
+	}
+	return Definitions;
+}
+
+// GetFixedParameterCode
+void FHLSLMaterialTranslator::GetFixedParameterCode(int32 StartChunk, int32 EndChunk, int32 ResultIndex, TArray<FShaderCodeChunk>& CodeChunks, FString& OutDefinitions, FString& OutValue)
+{
+	if (ResultIndex != INDEX_NONE)
+	{
+		checkf(ResultIndex >= 0 && ResultIndex < CodeChunks.Num(), TEXT("Index out of range %d/%d [%s]"), ResultIndex, CodeChunks.Num(), *Material->GetFriendlyName());
+		check(!CodeChunks[ResultIndex].UniformExpression || CodeChunks[ResultIndex].UniformExpression->IsConstant());
+		if (CodeChunks[ResultIndex].UniformExpression && CodeChunks[ResultIndex].UniformExpression->IsConstant())
+		{
+			// Handle a constant uniform expression being the only code chunk hooked up to a material input
+			const FShaderCodeChunk& ResultChunk = CodeChunks[ResultIndex];
+			OutValue = ResultChunk.Definition;
+		}
+		else
+		{
+			const FShaderCodeChunk& ResultChunk = CodeChunks[ResultIndex];
+			// Combine the definition lines and the return statement
+			check(ResultChunk.bInline || ResultChunk.SymbolName.Len() > 0);
+			OutDefinitions = GetDefinitions(CodeChunks, StartChunk, EndChunk);
+			OutValue = ResultChunk.bInline ? ResultChunk.Definition : ResultChunk.SymbolName;
+		}
+	}
+	else
+	{
+		OutValue = TEXT("0");
+	}
+}
+
+void FHLSLMaterialTranslator::GetFixedParameterCode(int32 ResultIndex, TArray<FShaderCodeChunk>& CodeChunks, FString& OutDefinitions, FString& OutValue)
+{
+	GetFixedParameterCode(0, CodeChunks.Num(), ResultIndex, CodeChunks, OutDefinitions, OutValue);
+}
+
+/** Used to get a user friendly type from EMaterialValueType */
+const TCHAR* FHLSLMaterialTranslator::DescribeType(EMaterialValueType Type) const
+{
+	switch(Type)
+	{
+	case MCT_Float1:				return TEXT("float");
+	case MCT_Float2:				return TEXT("float2");
+	case MCT_Float3:				return TEXT("float3");
+	case MCT_Float4:				return TEXT("float4");
+	case MCT_Float:					return TEXT("float");
+	case MCT_Texture2D:				return TEXT("texture2D");
+	case MCT_TextureCube:			return TEXT("textureCube");
+	case MCT_VolumeTexture:			return TEXT("volumeTexture");
+	case MCT_StaticBool:			return TEXT("static bool");
+	case MCT_MaterialAttributes:	return TEXT("MaterialAttributes");
+	case MCT_TextureExternal:		return TEXT("TextureExternal");
+	case MCT_TextureVirtual:		return TEXT("TextureVirtual");
+	case MCT_VTPageTableResult:		return TEXT("VTPageTableResult");
+	case MCT_ShadingModel:			return TEXT("ShadingModel");
+	default:						return TEXT("unknown");
+	};
+}
+
+/** Used to get an HLSL type from EMaterialValueType */
+const TCHAR* FHLSLMaterialTranslator::HLSLTypeString(EMaterialValueType Type) const
+{
+	switch(Type)
+	{
+	case MCT_Float1:				return TEXT("MaterialFloat");
+	case MCT_Float2:				return TEXT("MaterialFloat2");
+	case MCT_Float3:				return TEXT("MaterialFloat3");
+	case MCT_Float4:				return TEXT("MaterialFloat4");
+	case MCT_Float:					return TEXT("MaterialFloat");
+	case MCT_Texture2D:				return TEXT("texture2D");
+	case MCT_TextureCube:			return TEXT("textureCube");
+	case MCT_VolumeTexture:			return TEXT("volumeTexture");
+	case MCT_StaticBool:			return TEXT("static bool");
+	case MCT_MaterialAttributes:	return TEXT("MaterialAttributes");
+	case MCT_TextureExternal:		return TEXT("TextureExternal");
+	case MCT_TextureVirtual:		return TEXT("TextureVirtual");
+	case MCT_VTPageTableResult:		return TEXT("VTPageTableResult");
+	case MCT_ShadingModel:			return TEXT("uint");
+	default:						return TEXT("unknown");
+	};
+}
+
+int32 FHLSLMaterialTranslator::NonPixelShaderExpressionError()
+{
+	return Errorf(TEXT("Invalid node used in vertex/hull/domain shader input!"));
+}
+
+int32 FHLSLMaterialTranslator::ErrorUnlessFeatureLevelSupported(ERHIFeatureLevel::Type RequiredFeatureLevel)
+{
+	if (FeatureLevel < RequiredFeatureLevel)
+	{
+		FString FeatureLevelName, RequiredLevelName;
+		GetFeatureLevelName(FeatureLevel, FeatureLevelName);
+		GetFeatureLevelName(RequiredFeatureLevel, RequiredLevelName);
+		return Errorf(TEXT("Node not supported in feature level %s. %s required."), *FeatureLevelName, *RequiredLevelName);
+	}
+
+	return 0;
+}
+
+int32 FHLSLMaterialTranslator::NonVertexShaderExpressionError()
+{
+	return Errorf(TEXT("Invalid node used in pixel/hull/domain shader input!"));
+}
+
+int32 FHLSLMaterialTranslator::NonVertexOrPixelShaderExpressionError()
+{
+	return Errorf(TEXT("Invalid node used in hull/domain shader input!"));
+}
+
+void FHLSLMaterialTranslator::AddEstimatedTextureSample(const uint32 Count)
+{
+	if (IsCurrentlyCompilingForPreviousFrame())
+	{
+		// Ignore non-actionable cases
+		return;
+	}
+
+	if (ShaderFrequency == SF_Pixel || ShaderFrequency == SF_Compute)
+	{
+		MaterialCompilationOutput.EstimatedNumTextureSamplesPS += Count;
+	}
+	else
+	{
+		MaterialCompilationOutput.EstimatedNumTextureSamplesVS += Count;
+	}
+}
+
+/** Creates a unique symbol name and adds it to the symbol list. */
+FString FHLSLMaterialTranslator::CreateSymbolName(const TCHAR* SymbolNameHint)
+{
+	NextSymbolIndex++;
+	return FString(SymbolNameHint) + FString::FromInt(NextSymbolIndex);
+}
+
+/** Adds an already formatted inline or referenced code chunk */
+int32 FHLSLMaterialTranslator::AddCodeChunkInner(uint64 Hash, const TCHAR* FormattedCode, EMaterialValueType Type, bool bInlined)
+{
+	check(bAllowCodeChunkGeneration);
+
+	if (Type == MCT_Unknown)
+	{
+		return INDEX_NONE;
+	}
+
+	if (bInlined)
+	{
+		const int32 CodeIndex = CurrentScopeChunks->Num();
+		// Adding an inline code chunk, the definition will be the code to inline
+		new(*CurrentScopeChunks) FShaderCodeChunk(Hash, FormattedCode,TEXT(""),Type,true);
+		return CodeIndex;
+	}
+	// Can only create temporaries for certain types
+	else if ((Type & (MCT_Float | MCT_VTPageTableResult)) || Type == MCT_ShadingModel)
+	{
+		// Check for existing
+		for (int32 i = 0; i < CurrentScopeChunks->Num(); ++i)
+		{
+			if ((*CurrentScopeChunks)[i].Hash == Hash)
+			{
+				return i;
+			}
+		}
+
+		const int32 CodeIndex = CurrentScopeChunks->Num();
+		// Allocate a local variable name
+		const FString SymbolName = CreateSymbolName(TEXT("Local"));
+		// Construct the definition string which stores the result in a temporary and adds a newline for readability
+		const FString LocalVariableDefinition = FString("	") + HLSLTypeString(Type) + TEXT(" ") + SymbolName + TEXT(" = ") + FormattedCode + TEXT(";") + LINE_TERMINATOR;
+		// Adding a code chunk that creates a local variable
+		new(*CurrentScopeChunks) FShaderCodeChunk(Hash, *LocalVariableDefinition,SymbolName,Type,false);
+		return CodeIndex;
+	}
+	else
+	{
+		if (Type == MCT_MaterialAttributes)
+		{
+			return Errorf(TEXT("Operation not supported on Material Attributes"));
+		}
+
+		if (Type & MCT_Texture)
+		{
+			return Errorf(TEXT("Operation not supported on a Texture"));
+		}
+
+		if (Type == MCT_StaticBool)
+		{
+			return Errorf(TEXT("Operation not supported on a Static Bool"));
+		}
+			
+		return INDEX_NONE;
+	}
+}
+
+/** 
+	* Constructs the formatted code chunk and creates a new local variable definition from it. 
+	* This should be used over AddInlinedCodeChunk when the code chunk adds actual instructions, and especially when calling a function.
+	* Creating local variables instead of inlining simplifies the generated code and reduces redundant expression chains,
+	* Making compiles faster and enabling the shader optimizer to do a better job.
+	*/
+int32 FHLSLMaterialTranslator::AddCodeChunk(EMaterialValueType Type,const TCHAR* Format,...)
+{
+	int32	BufferSize		= 256;
+	TCHAR*	FormattedCode	= NULL;
+	int32	Result			= -1;
+
+	while(Result == -1)
+	{
+		FormattedCode = (TCHAR*) FMemory::Realloc( FormattedCode, BufferSize * sizeof(TCHAR) );
+		GET_VARARGS_RESULT( FormattedCode, BufferSize, BufferSize-1, Format, Format, Result );
+		BufferSize *= 2;
+	};
+	FormattedCode[Result] = 0;
+
+	const uint64 Hash = CityHash64((char*)FormattedCode, Result * sizeof(TCHAR));
+	const int32 CodeIndex = AddCodeChunkInner(Hash, FormattedCode,Type,false);
+	FMemory::Free(FormattedCode);
+
+	return CodeIndex;
+}
+
+static inline uint32 GetTCharStringLength(const TCHAR* String)
+{
+	uint32 Length = 0u;
+	while (String[Length])
+	{
+		++Length;
+	}
+	return Length * sizeof(TCHAR);
+}
+
+int32 FHLSLMaterialTranslator::AddCodeChunkWithHash(uint64 BaseHash, EMaterialValueType Type, const TCHAR* Format, ...)
+{
+	int32	BufferSize = 256;
+	TCHAR*	FormattedCode = NULL;
+	int32	Result = -1;
+
+	while (Result == -1)
+	{
+		FormattedCode = (TCHAR*)FMemory::Realloc(FormattedCode, BufferSize * sizeof(TCHAR));
+		GET_VARARGS_RESULT(FormattedCode, BufferSize, BufferSize - 1, Format, Format, Result);
+		BufferSize *= 2;
+	};
+	FormattedCode[Result] = 0;
+
+	const uint64 Hash = CityHash64WithSeed((char*)Format, GetTCharStringLength(Format), BaseHash);
+	const int32 CodeIndex = AddCodeChunkInner(Hash, FormattedCode, Type, false);
+	FMemory::Free(FormattedCode);
+
+	return CodeIndex;
+}
+
+/** 
+	* Constructs the formatted code chunk and creates an inlined code chunk from it. 
+	* This should be used instead of AddCodeChunk when the code chunk does not add any actual shader instructions, for example a component mask.
+	*/
+int32 FHLSLMaterialTranslator::AddInlinedCodeChunk(EMaterialValueType Type, const TCHAR* Format,...)
+{
+	int32	BufferSize		= 256;
+	TCHAR*	FormattedCode	= NULL;
+	int32	Result			= -1;
+
+	while(Result == -1)
+	{
+		FormattedCode = (TCHAR*) FMemory::Realloc( FormattedCode, BufferSize * sizeof(TCHAR) );
+		GET_VARARGS_RESULT( FormattedCode, BufferSize, BufferSize-1, Format, Format, Result );
+		BufferSize *= 2;
+	};
+	FormattedCode[Result] = 0;
+
+	const uint64 Hash = CityHash64((char*)FormattedCode, Result * sizeof(TCHAR));
+	const int32 CodeIndex = AddCodeChunkInner(Hash, FormattedCode,Type,true);
+	FMemory::Free(FormattedCode);
+
+	return CodeIndex;
+}
+
+int32 FHLSLMaterialTranslator::AddInlinedCodeChunkWithHash(uint64 BaseHash, EMaterialValueType Type, const TCHAR* Format, ...)
+{
+	int32	BufferSize = 256;
+	TCHAR*	FormattedCode = NULL;
+	int32	Result = -1;
+
+	while (Result == -1)
+	{
+		FormattedCode = (TCHAR*)FMemory::Realloc(FormattedCode, BufferSize * sizeof(TCHAR));
+		GET_VARARGS_RESULT(FormattedCode, BufferSize, BufferSize - 1, Format, Format, Result);
+		BufferSize *= 2;
+	};
+	FormattedCode[Result] = 0;
+
+	const uint64 Hash = CityHash64WithSeed((char*)Format, GetTCharStringLength(Format), BaseHash);
+	const int32 CodeIndex = AddCodeChunkInner(Hash, FormattedCode, Type, true);
+	FMemory::Free(FormattedCode);
+
+	return CodeIndex;
+}
+
+int32 FHLSLMaterialTranslator::AddUniformExpressionInner(uint64 Hash, FMaterialUniformExpression* UniformExpression, EMaterialValueType Type, const TCHAR* FormattedCode)
+{
+	check(bAllowCodeChunkGeneration);
+
+	if (Type == MCT_Unknown)
+	{
+		return INDEX_NONE;
+	}
+
+	check(UniformExpression);
+
+	// Only a texture uniform expression can have MCT_Texture type
+	if ((Type & MCT_Texture) && !UniformExpression->GetTextureUniformExpression()
+		&& !UniformExpression->GetExternalTextureUniformExpression())
+	{
+		return Errorf(TEXT("Operation not supported on a Texture"));
+	}
+
+	// External textures must have an external texture uniform expression
+	if ((Type & MCT_TextureExternal) && !UniformExpression->GetExternalTextureUniformExpression())
+	{
+		return Errorf(TEXT("Operation not supported on an external texture"));
+	}
+
+	if (Type == MCT_StaticBool)
+	{
+		return Errorf(TEXT("Operation not supported on a Static Bool"));
+	}
+
+	if (Type == MCT_MaterialAttributes)
+	{
+		return Errorf(TEXT("Operation not supported on a MaterialAttributes"));
+	}
+
+	bool bFoundExistingExpression = false;
+	// Search for an existing code chunk with the same uniform expression in the array of all uniform expressions used by this material.
+	for (int32 ExpressionIndex = 0; ExpressionIndex < UniformExpressions.Num() && !bFoundExistingExpression; ExpressionIndex++)
+	{
+		FMaterialUniformExpression* TestExpression = UniformExpressions[ExpressionIndex].UniformExpression;
+		check(TestExpression);
+		if (TestExpression->IsIdentical(UniformExpression))
+		{
+			bFoundExistingExpression = true;
+			// This code chunk has an identical uniform expression to the new expression, reuse it.
+			// This allows multiple material properties to share uniform expressions because AccessUniformExpression uses AddUniqueItem when adding uniform expressions.
+			check(Type == UniformExpressions[ExpressionIndex].Type);
+			// Search for an existing code chunk with the same uniform expression in the array of code chunks for this material property.
+			for (int32 ChunkIndex = 0; ChunkIndex < CurrentScopeChunks->Num(); ChunkIndex++)
+			{
+				FMaterialUniformExpression* OtherExpression = (*CurrentScopeChunks)[ChunkIndex].UniformExpression;
+				if (OtherExpression && OtherExpression->IsIdentical(UniformExpression))
+				{
+					delete UniformExpression;
+					// Reuse the entry in CurrentScopeChunks
+					return ChunkIndex;
+				}
+			}
+			delete UniformExpression;
+			// Use the existing uniform expression from a different material property,
+			// And continue so that a code chunk using the uniform expression will be generated for this material property.
+			UniformExpression = TestExpression;
+			break;
+		}
+
+#if 0
+		// Test for the case where we have non-identical expressions of the same type and name.
+		// This means they exist with separate values and the one retrieved for shading will
+		// effectively be random, as we evaluate the first found during expression traversal
+		if (TestExpression->GetType() == UniformExpression->GetType())
+		{
+			if (TestExpression->GetType() == &FMaterialUniformExpressionScalarParameter::StaticType)
+			{
+				FMaterialUniformExpressionScalarParameter* ScalarParameterA = (FMaterialUniformExpressionScalarParameter*)TestExpression;
+				FMaterialUniformExpressionScalarParameter* ScalarParameterB = (FMaterialUniformExpressionScalarParameter*)UniformExpression;
+
+				if (!ScalarParameterA->GetParameterInfo().Name.IsNone() && ScalarParameterA->GetParameterInfo() == ScalarParameterB->GetParameterInfo())
+				{
+					delete UniformExpression;
+					return Errorf(TEXT("Invalid scalar parameter '%s' found. Identical parameters must have the same value."), *(ScalarParameterA->GetParameterInfo().Name.ToString()));
+				}
+			}
+			else if (TestExpression->GetType() == &FMaterialUniformExpressionVectorParameter::StaticType)
+			{
+				FMaterialUniformExpressionVectorParameter* VectorParameterA = (FMaterialUniformExpressionVectorParameter*)TestExpression;
+				FMaterialUniformExpressionVectorParameter* VectorParameterB = (FMaterialUniformExpressionVectorParameter*)UniformExpression;
+
+				// Note: Skipping NAME_SelectionColor here as this behavior is relied on for editor materials
+				if (!VectorParameterA->GetParameterInfo().Name.IsNone() && VectorParameterA->GetParameterInfo() == VectorParameterB->GetParameterInfo()
+					&& VectorParameterA->GetParameterInfo().Name != NAME_SelectionColor)
+				{
+					delete UniformExpression;
+					return Errorf(TEXT("Invalid vector parameter '%s' found. Identical parameters must have the same value."), *(VectorParameterA->GetParameterInfo().Name.ToString()));
+				}
+			}
+		}
+#endif
+	}
+
+	const int32 ReturnIndex = CurrentScopeChunks->Num();
+	// Create a new code chunk for the uniform expression
+	new(*CurrentScopeChunks) FShaderCodeChunk(Hash, UniformExpression, FormattedCode, Type);
+
+	if (!bFoundExistingExpression)
+	{
+		// Add an entry to the material-wide list of uniform expressions
+		new(UniformExpressions) FShaderCodeChunk(Hash, UniformExpression, FormattedCode, Type);
+	}
+
+	return ReturnIndex;
+}
+
+// AddUniformExpression - Adds an input to the Code array and returns its index.
+int32 FHLSLMaterialTranslator::AddUniformExpression(FMaterialUniformExpression* UniformExpression,EMaterialValueType Type, const TCHAR* Format,...)
+{
+	int32	BufferSize = 256;
+	TCHAR*	FormattedCode = NULL;
+	int32	Result = -1;
+
+	while (Result == -1)
+	{
+		FormattedCode = (TCHAR*)FMemory::Realloc(FormattedCode, BufferSize * sizeof(TCHAR));
+		GET_VARARGS_RESULT(FormattedCode, BufferSize, BufferSize - 1, Format, Format, Result);
+		BufferSize *= 2;
+	};
+	FormattedCode[Result] = 0;
+
+	const uint64 Hash = CityHash64((char*)FormattedCode, Result * sizeof(TCHAR));
+	const int32 CodeIndex = AddUniformExpressionInner(Hash, UniformExpression, Type, FormattedCode);
+	FMemory::Free(FormattedCode);
+	return CodeIndex;
+}
+
+int32 FHLSLMaterialTranslator::AddUniformExpressionWithHash(uint64 BaseHash, FMaterialUniformExpression* UniformExpression, EMaterialValueType Type, const TCHAR* Format, ...)
+{
+	int32	BufferSize = 256;
+	TCHAR*	FormattedCode = NULL;
+	int32	Result = -1;
+
+	while (Result == -1)
+	{
+		FormattedCode = (TCHAR*)FMemory::Realloc(FormattedCode, BufferSize * sizeof(TCHAR));
+		GET_VARARGS_RESULT(FormattedCode, BufferSize, BufferSize - 1, Format, Format, Result);
+		BufferSize *= 2;
+	};
+	FormattedCode[Result] = 0;
+
+	const uint64 Hash = CityHash64WithSeed((char*)Format, GetTCharStringLength(Format), BaseHash);
+	const int32 CodeIndex = AddUniformExpressionInner(Hash, UniformExpression, Type, FormattedCode);
+	FMemory::Free(FormattedCode);
+	return CodeIndex;
+}
+
+// AccessUniformExpression - Adds code to access the value of a uniform expression to the Code array and returns its index.
+int32 FHLSLMaterialTranslator::AccessUniformExpression(int32 Index)
+{
+	check(Index >= 0 && Index < CurrentScopeChunks->Num());
+	const FShaderCodeChunk&	CodeChunk = (*CurrentScopeChunks)[Index];
+	check(CodeChunk.UniformExpression && !CodeChunk.UniformExpression->IsConstant());
+
+	FMaterialUniformExpressionTexture* TextureUniformExpression = CodeChunk.UniformExpression->GetTextureUniformExpression();
+	FMaterialUniformExpressionExternalTexture* ExternalTextureUniformExpression = CodeChunk.UniformExpression->GetExternalTextureUniformExpression();
+
+	// Any code chunk can have a texture uniform expression (eg FMaterialUniformExpressionFlipBookTextureParameter),
+	// But a texture code chunk must have a texture uniform expression
+	check(!(CodeChunk.Type & MCT_Texture) || TextureUniformExpression || ExternalTextureUniformExpression);
+	// External texture samples must have a corresponding uniform expression
+	check(!(CodeChunk.Type & MCT_TextureExternal) || ExternalTextureUniformExpression);
+	// Virtual texture samples must have a corresponding uniform expression
+	check(!(CodeChunk.Type & MCT_TextureVirtual) || TextureUniformExpression);
+
+	TCHAR FormattedCode[MAX_SPRINTF]=TEXT("");
+	if(CodeChunk.Type == MCT_Float)
+	{
+		const static TCHAR IndexToMask[] = {'x', 'y', 'z', 'w'};
+		const int32 ScalarInputIndex = MaterialCompilationOutput.UniformExpressionSet.UniformScalarExpressions.AddUnique(CodeChunk.UniformExpression);
+		// Update the above FMemory::Malloc if this FCString::Sprintf grows in size, e.g. %s, ...
+		FCString::Sprintf(FormattedCode, TEXT("Material.ScalarExpressions[%u].%c"), ScalarInputIndex / 4, IndexToMask[ScalarInputIndex % 4]);
+	}
+	else if(CodeChunk.Type & MCT_Float)
+	{
+		const TCHAR* Mask;
+		switch(CodeChunk.Type)
+		{
+		case MCT_Float:
+		case MCT_Float1: Mask = TEXT(".r"); break;
+		case MCT_Float2: Mask = TEXT(".rg"); break;
+		case MCT_Float3: Mask = TEXT(".rgb"); break;
+		default: Mask = TEXT(""); break;
+		};
+
+		const int32 VectorInputIndex = MaterialCompilationOutput.UniformExpressionSet.UniformVectorExpressions.AddUnique(CodeChunk.UniformExpression);
+		FCString::Sprintf(FormattedCode, TEXT("Material.VectorExpressions[%u]%s"), VectorInputIndex, Mask);
+	}
+	else if(CodeChunk.Type & MCT_Texture)
+	{
+		int32 TextureInputIndex = INDEX_NONE;
+		const TCHAR* BaseName = TEXT("");
+		bool GenerateCode = true;
+		switch(CodeChunk.Type)
+		{
+		case MCT_Texture2D:
+			TextureInputIndex = MaterialCompilationOutput.UniformExpressionSet.Uniform2DTextureExpressions.AddUnique(TextureUniformExpression);
+			BaseName = TEXT("Texture2D");
+			break;
+		case MCT_TextureCube:
+			TextureInputIndex = MaterialCompilationOutput.UniformExpressionSet.UniformCubeTextureExpressions.AddUnique(TextureUniformExpression);
+			BaseName = TEXT("TextureCube");
+			break;
+		case MCT_VolumeTexture:
+			TextureInputIndex = MaterialCompilationOutput.UniformExpressionSet.UniformVolumeTextureExpressions.AddUnique(TextureUniformExpression);
+			BaseName = TEXT("VolumeTexture");
+			break;
+		case MCT_TextureExternal:
+			TextureInputIndex = MaterialCompilationOutput.UniformExpressionSet.UniformExternalTextureExpressions.AddUnique(ExternalTextureUniformExpression);
+			BaseName = TEXT("ExternalTexture");
+			break;
+		case MCT_TextureVirtual:
+			TextureInputIndex = MaterialCompilationOutput.UniformExpressionSet.UniformVirtualTextureExpressions.AddUnique(TextureUniformExpression);
+			GenerateCode = false;
+			break;
+		default: UE_LOG(LogMaterial, Fatal,TEXT("Unrecognized texture material value type: %u"),(int32)CodeChunk.Type);
+		};
+		if(GenerateCode)
+		{
+			FCString::Sprintf(FormattedCode, TEXT("Material.%s_%u"), BaseName, TextureInputIndex);
+		}
+	}
+	else
+	{
+		UE_LOG(LogMaterial, Fatal,TEXT("User input of unknown type: %s"),DescribeType(CodeChunk.Type));
+	}
+
+	return AddInlinedCodeChunk((*CurrentScopeChunks)[Index].Type,FormattedCode);
+}
+
+// CoerceParameter
+FString FHLSLMaterialTranslator::CoerceParameter(int32 Index,EMaterialValueType DestType)
+{
+	check(Index >= 0 && Index < CurrentScopeChunks->Num());
+	const FShaderCodeChunk&	CodeChunk = (*CurrentScopeChunks)[Index];
+	if( CodeChunk.Type == DestType )
+	{
+		return GetParameterCode(Index);
+	}
+	else
+	{
+		if( (CodeChunk.Type & DestType) && (CodeChunk.Type & MCT_Float) )
+		{
+			switch( DestType )
+			{
+			case MCT_Float1:
+				return FString::Printf( TEXT("MaterialFloat(%s)"), *GetParameterCode(Index) );
+			case MCT_Float2:
+				return FString::Printf( TEXT("MaterialFloat2(%s,%s)"), *GetParameterCode(Index), *GetParameterCode(Index) );
+			case MCT_Float3:
+				return FString::Printf( TEXT("MaterialFloat3(%s,%s,%s)"), *GetParameterCode(Index), *GetParameterCode(Index), *GetParameterCode(Index) );
+			case MCT_Float4:
+				return FString::Printf( TEXT("MaterialFloat4(%s,%s,%s,%s)"), *GetParameterCode(Index), *GetParameterCode(Index), *GetParameterCode(Index), *GetParameterCode(Index) );
+			default: 
+				return FString::Printf( TEXT("%s"), *GetParameterCode(Index) );
+			}
+		}
+		else
+		{
+			Errorf(TEXT("Coercion failed: %s: %s -> %s"),*CodeChunk.Definition,DescribeType(CodeChunk.Type),DescribeType(DestType));
+			return TEXT("");
+		}
+	}
+}
+
+// GetParameterType
+EMaterialValueType FHLSLMaterialTranslator::GetParameterType(int32 Index) const
+{
+	check(Index >= 0 && Index < CurrentScopeChunks->Num());
+	return (*CurrentScopeChunks)[Index].Type;
+}
+
+// GetParameterUniformExpression
+FMaterialUniformExpression* FHLSLMaterialTranslator::GetParameterUniformExpression(int32 Index) const
+{
+	check(Index >= 0 && Index < CurrentScopeChunks->Num());
+
+	const FShaderCodeChunk& Chunk = (*CurrentScopeChunks)[Index];
+
+	return Chunk.UniformExpression;
+}
+
+bool FHLSLMaterialTranslator::GetTextureForExpression(int32 Index, int32& OutTextureIndex, EMaterialSamplerType& OutSamplerType, TOptional<FName>& OutParameterName) const
+{
+	check(Index >= 0 && Index < CurrentScopeChunks->Num());
+	const FShaderCodeChunk& Chunk = (*CurrentScopeChunks)[Index];
+	const EMaterialValueType TexInputType = Chunk.Type;
+	if (!(TexInputType & MCT_Texture))
+	{
+		return false;
+	}
+
+	// If 'InputExpression' is connected, we use need to find the texture object that was passed in
+	// In this case, the texture/sampler assigned on this expression node are not used
+	FMaterialUniformExpression* TextureUniformBase = Chunk.UniformExpression;
+	checkf(TextureUniformBase, TEXT("TexInputType is %d, but missing FMaterialUniformExpression"), TexInputType);
+
+	if (FMaterialUniformExpressionTexture* TextureUniform = TextureUniformBase->GetTextureUniformExpression())
+	{
+		OutSamplerType = TextureUniform->GetSamplerType();
+		OutTextureIndex = TextureUniform->GetTextureIndex();
+		if (FMaterialUniformExpressionTextureParameter* TextureParameterUniform = TextureUniform->GetTextureParameterUniformExpression())
+		{
+			OutParameterName = TextureParameterUniform->GetParameterName();
+		}
+	}
+	else if (FMaterialUniformExpressionExternalTexture* ExternalTextureUniform = TextureUniformBase->GetExternalTextureUniformExpression())
+	{
+		OutTextureIndex = ExternalTextureUniform->GetSourceTextureIndex();
+		OutSamplerType = SAMPLERTYPE_External;
+		if (FMaterialUniformExpressionExternalTextureParameter* ExternalTextureParameterUniform = ExternalTextureUniform->GetExternalTextureParameterUniformExpression())
+		{
+			OutParameterName = ExternalTextureParameterUniform->GetParameterName();
+		}
+	}
+
+	return true;
+}
+
+// GetArithmeticResultType
+EMaterialValueType FHLSLMaterialTranslator::GetArithmeticResultType(EMaterialValueType TypeA, EMaterialValueType TypeB)
+{
+	if (!((TypeA & MCT_Float) || TypeA == MCT_ShadingModel) || !((TypeB & MCT_Float) || TypeB == MCT_ShadingModel))
+	{
+		Errorf(TEXT("Attempting to perform arithmetic on non-numeric types: %s %s"),DescribeType(TypeA),DescribeType(TypeB));
+		return MCT_Unknown;
+	}
+
+	if(TypeA == TypeB)
+	{
+		return TypeA;
+	}
+	else if(TypeA & TypeB)
+	{
+		if(TypeA == MCT_Float)
+		{
+			return TypeB;
+		}
+		else
+		{
+			check(TypeB == MCT_Float);
+			return TypeA;
+		}
+	}
+	else
+	{
+		Errorf(TEXT("Arithmetic between types %s and %s are undefined"),DescribeType(TypeA),DescribeType(TypeB));
+		return MCT_Unknown;
+	}
+}
+
+EMaterialValueType FHLSLMaterialTranslator::GetArithmeticResultType(int32 A,int32 B)
+{
+	check(A >= 0 && A < CurrentScopeChunks->Num());
+	check(B >= 0 && B < CurrentScopeChunks->Num());
+
+	EMaterialValueType	TypeA = (*CurrentScopeChunks)[A].Type,
+		TypeB = (*CurrentScopeChunks)[B].Type;
+
+	return GetArithmeticResultType(TypeA,TypeB);
+}
+
+// FMaterialCompiler interface.
+
+/** 
+	* Sets the current material property being compiled.  
+	* This affects the internal state of the compiler and the results of all functions except GetFixedParameterCode.
+	* @param OverrideShaderFrequency SF_NumFrequencies to not override
+	*/
+void FHLSLMaterialTranslator::SetMaterialProperty(EMaterialProperty InProperty, EShaderFrequency OverrideShaderFrequency, bool bUsePreviousFrameTime)
+{
+	MaterialProperty = InProperty;
+	SetBaseMaterialAttribute(FMaterialAttributeDefinitionMap::GetID(InProperty));
+
+	if(OverrideShaderFrequency != SF_NumFrequencies)
+	{
+		ShaderFrequency = OverrideShaderFrequency;
+	}
+	else
+	{
+		ShaderFrequency = FMaterialAttributeDefinitionMap::GetShaderFrequency(InProperty);
+	}
+
+	bCompilingPreviousFrame = bUsePreviousFrameTime;
+	AssignShaderFrequencyScope(ShaderFrequency);
+}
+
+void FHLSLMaterialTranslator::PushMaterialAttribute(const FGuid& InAttributeID)
+{
+	MaterialAttributesStack.Push(InAttributeID);
+}
+
+FGuid FHLSLMaterialTranslator::PopMaterialAttribute()
+{
+	return MaterialAttributesStack.Pop(false);
+}
+
+const FGuid FHLSLMaterialTranslator::GetMaterialAttribute()
+{
+	checkf(MaterialAttributesStack.Num() > 0, TEXT("Tried to query empty material attributes stack."));
+	return MaterialAttributesStack.Top();
+}
+
+void FHLSLMaterialTranslator::SetBaseMaterialAttribute(const FGuid& InAttributeID)
+{
+	// This is atypical behavior but is done to allow cleaner code and preserve existing paths.
+	// A base property is kept on the stack and updated by SetMaterialProperty(), the stack is only utilized during translation
+	checkf(MaterialAttributesStack.Num() == 1, TEXT("Tried to set non-base attribute on stack."));
+	MaterialAttributesStack.Top() = InAttributeID;
+}
+
+void FHLSLMaterialTranslator::PushParameterOwner(const FMaterialParameterInfo& InOwnerInfo)
+{
+	ParameterOwnerStack.Push(InOwnerInfo);
+}
+
+FMaterialParameterInfo FHLSLMaterialTranslator::PopParameterOwner()
+{
+	return ParameterOwnerStack.Pop(false);
+}
+
+EShaderFrequency FHLSLMaterialTranslator::GetCurrentShaderFrequency() const
+{
+	return ShaderFrequency;
+}
+
+FMaterialShadingModelField FHLSLMaterialTranslator::GetMaterialShadingModels() const
+{
+	check(Material);
+	return Material->GetShadingModels();
+}
+
+int32 FHLSLMaterialTranslator::Error(const TCHAR* Text)
+{
+	// Optionally append errors into proxy arrays which allow pre-translation stages to selectively include errors later
+	bool bUsingErrorProxy = (CompileErrorsSink && CompileErrorExpressionsSink);	
+	TArray<FString>& CompileErrors = bUsingErrorProxy ? *CompileErrorsSink : Material->CompileErrors;
+	TArray<UMaterialExpression*>& ErrorExpressions = bUsingErrorProxy ? *CompileErrorExpressionsSink : Material->ErrorExpressions;
+
+	FString ErrorString;
+	UMaterialExpression* ExpressionToError = nullptr;
+
+	check(ShaderFrequency < SF_NumFrequencies);
+	auto& CurrentFunctionStack = FunctionStacks[ShaderFrequency];
+	if (CurrentFunctionStack.Num() > 1)
+	{
+		// If we are inside a function, add that to the error message.  
+		// Only add the function call node to ErrorExpressions, since we can't add a reference to the expressions inside the function as they are private objects.
+		// Add the first function node on the stack because that's the one visible in the material being compiled, the rest are all nested functions.
+		UMaterialExpressionMaterialFunctionCall* ErrorFunction = CurrentFunctionStack[1]->FunctionCall;
+		check(ErrorFunction);
+		ExpressionToError = ErrorFunction;			
+		ErrorString = FString(TEXT("Function ")) + ErrorFunction->MaterialFunction->GetName() + TEXT(": ");
+	}
+
+	if (CurrentFunctionStack.Last()->ExpressionStack.Num() > 0)
+	{
+		UMaterialExpression* ErrorExpression = CurrentFunctionStack.Last()->ExpressionStack.Last().Expression;
+		check(ErrorExpression);
+
+		if (ErrorExpression->GetClass() != UMaterialExpressionMaterialFunctionCall::StaticClass()
+			&& ErrorExpression->GetClass() != UMaterialExpressionFunctionInput::StaticClass()
+			&& ErrorExpression->GetClass() != UMaterialExpressionFunctionOutput::StaticClass())
+		{
+			// Add the expression currently being compiled to ErrorExpressions so we can draw it differently
+			ExpressionToError = ErrorExpression;
+
+			const int32 ChopCount = FCString::Strlen(TEXT("MaterialExpression"));
+			const FString ErrorClassName = ErrorExpression->GetClass()->GetName();
+
+			// Add the node type to the error message
+			ErrorString += FString(TEXT("(Node ")) + ErrorClassName.Right(ErrorClassName.Len() - ChopCount) + TEXT(") ");
+		}
+	}
+			
+	ErrorString += Text;
+
+	if (!bUsingErrorProxy)
+	{
+		// Standard error handling, immediately append one-off errors and signal failure
+		CompileErrors.AddUnique(ErrorString);
+	
+		if (ExpressionToError)
+		{
+			ErrorExpressions.Add(ExpressionToError);
+			ExpressionToError->LastErrorText = Text;
+		}
+
+		bSuccess = false;
+	}
+	else
+	{
+		// When a proxy is intercepting errors, ignore the failure and match arrays to allow later error type selection
+		CompileErrors.Add(ErrorString);
+		ErrorExpressions.Add(ExpressionToError);		
+	}
+		
+	return INDEX_NONE;
+}
+
+void FHLSLMaterialTranslator::AppendExpressionError(UMaterialExpression* Expression, const TCHAR* Text)
+{
+	if (Expression && Text)
+	{
+		FString ErrorText(Text);
+
+		Material->ErrorExpressions.Add(Expression);
+		Expression->LastErrorText = ErrorText;
+		Material->CompileErrors.Add(ErrorText);
+	}
+}
+
+int32 FHLSLMaterialTranslator::CallExpression(FMaterialExpressionKey ExpressionKey,FMaterialCompiler* Compiler)
+{
+	// For any translated result not relying on material attributes, we can discard the attribute ID from the key
+	// to allow result sharing. In cases where we detect an expression loop we must err on the side of caution
+	if (ExpressionKey.Expression && !ExpressionKey.Expression->ContainsInputLoop() && !ExpressionKey.Expression->IsResultMaterialAttributes(ExpressionKey.OutputIndex))
+	{
+		ExpressionKey.MaterialAttributeID = FGuid(0, 0, 0, 0);
+	}
+
+	// Some expressions can discard output indices and share compiles with a swizzle/mask
+	if (ExpressionKey.Expression && ExpressionKey.Expression->CanIgnoreOutputIndex())
+	{
+		ExpressionKey.OutputIndex = INDEX_NONE;
+	}
+
+	// Check if this expression has already been translated.
+	check(ShaderFrequency < SF_NumFrequencies);
+	auto& CurrentFunctionStack = FunctionStacks[ShaderFrequency];
+	FMaterialFunctionCompileState* CurrentFunctionState = CurrentFunctionStack.Last();
+
+	int32* ExistingCodeIndex = CurrentFunctionState->ExpressionCodeMap.Find(ExpressionKey);
+	if (ExistingCodeIndex)
+	{
+		return *ExistingCodeIndex;
+	}
+	else
+	{
+		// Disallow reentrance.
+		if (CurrentFunctionState->ExpressionStack.Find(ExpressionKey) != INDEX_NONE)
+		{
+			return Error(TEXT("Reentrant expression"));
+		}
+
+		// The first time this expression is called, translate it.
+		CurrentFunctionState->ExpressionStack.Add(ExpressionKey);
+		const int32 FunctionDepth = CurrentFunctionStack.Num();
+			
+		// Attempt to share function states between function calls
+		UMaterialExpressionMaterialFunctionCall* FunctionCall = Cast<UMaterialExpressionMaterialFunctionCall>(ExpressionKey.Expression);
+		if (FunctionCall)
+		{
+			FMaterialExpressionKey ReuseCompileStateExpressionKey = ExpressionKey;
+ 			ReuseCompileStateExpressionKey.OutputIndex = INDEX_NONE; // Discard the output so we can share the stack internals
+
+			FMaterialFunctionCompileState* SharedFunctionState = CurrentFunctionState->FindOrAddSharedFunctionState(ReuseCompileStateExpressionKey, FunctionCall);
+			FunctionCall->SetSharedCompileState(SharedFunctionState);
+		}
+
+		int32 Result = ExpressionKey.Expression->Compile(Compiler, ExpressionKey.OutputIndex);
+
+		// Restore state
+		if (FunctionCall)
+		{
+			FunctionCall->SetSharedCompileState(nullptr);
+		}
+
+		FMaterialExpressionKey PoppedExpressionKey = CurrentFunctionState->ExpressionStack.Pop();
+
+		// Verify state integrity
+		check(PoppedExpressionKey == ExpressionKey);
+		check(FunctionDepth == CurrentFunctionStack.Num());
+
+		// Cache the translation.
+		CurrentFunctionStack.Last()->ExpressionCodeMap.Add(ExpressionKey,Result);
+
+		return Result;
+	}
+}
+
+EMaterialValueType FHLSLMaterialTranslator::GetType(int32 Code)
+{
+	if(Code != INDEX_NONE)
+	{
+		return GetParameterType(Code);
+	}
+	else
+	{
+		return MCT_Unknown;
+	}
+}
+
+EMaterialQualityLevel::Type FHLSLMaterialTranslator::GetQualityLevel()
+{
+	return QualityLevel;
+}
+
+ERHIFeatureLevel::Type FHLSLMaterialTranslator::GetFeatureLevel()
+{
+	return FeatureLevel;
+}
+
+EShaderPlatform FHLSLMaterialTranslator::GetShaderPlatform()
+{
+	return Platform;
+}
+
+const ITargetPlatform* FHLSLMaterialTranslator::GetTargetPlatform() const
+{
+	return TargetPlatform;
+}
+
+/** 
+	* Casts the passed in code to DestType, or generates a compile error if the cast is not valid. 
+	* This will truncate a type (float4 -> float3) but not add components (float2 -> float3), however a float1 can be cast to any float type by replication. 
+	*/
+int32 FHLSLMaterialTranslator::ValidCast(int32 Code, EMaterialValueType DestType)
+{
+	if(Code == INDEX_NONE)
+	{
+		return INDEX_NONE;
+	}
+
+	EMaterialValueType SourceType = GetParameterType(Code);
+	int32 CompiledResult = INDEX_NONE;
+
+	if (SourceType & DestType)
+	{
+		CompiledResult = Code;
+	}
+	else if(GetParameterUniformExpression(Code) && !GetParameterUniformExpression(Code)->IsConstant())
+	{
+		if ((SourceType & MCT_TextureVirtual) && (DestType & MCT_Texture2D))
+		{
+			return Code;
+		}
+		else
+		{
+			return ValidCast(AccessUniformExpression(Code), DestType);
+		}
+	}
+	else if((SourceType & MCT_Float) && (DestType & MCT_Float))
+	{
+		const uint32 NumSourceComponents = GetNumComponents(SourceType);
+		const uint32 NumDestComponents = GetNumComponents(DestType);
+
+		if(NumSourceComponents > NumDestComponents) // Use a mask to select the first NumDestComponents components from the source.
+		{
+			const TCHAR*	Mask;
+			switch(NumDestComponents)
+			{
+			case 1: Mask = TEXT(".r"); break;
+			case 2: Mask = TEXT(".rg"); break;
+			case 3: Mask = TEXT(".rgb"); break;
+			default: UE_LOG(LogMaterial, Fatal,TEXT("Should never get here!")); return INDEX_NONE;
+			};
+
+			return AddInlinedCodeChunk(DestType,TEXT("%s%s"),*GetParameterCode(Code),Mask);
+		}
+		else if(NumSourceComponents < NumDestComponents) // Pad the source vector up to NumDestComponents.
+		{
+			// Only allow replication when the Source is a Float1
+			if (NumSourceComponents == 1)
+			{
+				const uint32 NumPadComponents = NumDestComponents - NumSourceComponents;
+				FString CommaParameterCodeString = FString::Printf(TEXT(",%s"), *GetParameterCode(Code));
+
+				CompiledResult = AddInlinedCodeChunk(
+					DestType,
+					TEXT("%s(%s%s%s%s)"),
+					HLSLTypeString(DestType),
+					*GetParameterCode(Code),
+					NumPadComponents >= 1 ? *CommaParameterCodeString : TEXT(""),
+					NumPadComponents >= 2 ? *CommaParameterCodeString : TEXT(""),
+					NumPadComponents >= 3 ? *CommaParameterCodeString : TEXT("")
+					);
+			}
+			else
+			{
+				CompiledResult = Errorf(TEXT("Cannot cast from %s to %s."), DescribeType(SourceType), DescribeType(DestType));
+			}
+		}
+		else
+		{
+			CompiledResult = Code;
+		}
+	}
+	else
+	{
+		//We can feed any type into a material attributes socket as we're really just passing them through.
+		if( DestType == MCT_MaterialAttributes )
+		{
+			CompiledResult = Code;
+		}
+		else
+		{
+			CompiledResult = Errorf(TEXT("Cannot cast from %s to %s."), DescribeType(SourceType), DescribeType(DestType));
+		}
+	}
+
+	return CompiledResult;
+}
+
+int32 FHLSLMaterialTranslator::ForceCast(int32 Code, EMaterialValueType DestType, uint32 ForceCastFlags)
+{
+	if(Code == INDEX_NONE)
+	{
+		return INDEX_NONE;
+	}
+
+	if(GetParameterUniformExpression(Code) && !GetParameterUniformExpression(Code)->IsConstant())
+	{
+		return ForceCast(AccessUniformExpression(Code),DestType,ForceCastFlags);
+	}
+
+	EMaterialValueType	SourceType = GetParameterType(Code);
+
+	bool bExactMatch = (ForceCastFlags & MFCF_ExactMatch) ? true : false;
+	bool bReplicateValue = (ForceCastFlags & MFCF_ReplicateValue) ? true : false;
+
+	if (bExactMatch ? (SourceType == DestType) : (SourceType & DestType))
+	{
+		return Code;
+	}
+	else if((SourceType & MCT_Float) && (DestType & MCT_Float))
+	{
+		const uint32 NumSourceComponents = GetNumComponents(SourceType);
+		const uint32 NumDestComponents = GetNumComponents(DestType);
+
+		if(NumSourceComponents > NumDestComponents) // Use a mask to select the first NumDestComponents components from the source.
+		{
+			const TCHAR*	Mask;
+			switch(NumDestComponents)
+			{
+			case 1: Mask = TEXT(".r"); break;
+			case 2: Mask = TEXT(".rg"); break;
+			case 3: Mask = TEXT(".rgb"); break;
+			default: UE_LOG(LogMaterial, Fatal,TEXT("Should never get here!")); return INDEX_NONE;
+			};
+
+			return AddInlinedCodeChunk(DestType,TEXT("%s%s"),*GetParameterCode(Code),Mask);
+		}
+		else if(NumSourceComponents < NumDestComponents) // Pad the source vector up to NumDestComponents.
+		{
+			// Only allow replication when the Source is a Float1
+			if (NumSourceComponents != 1)
+			{
+				bReplicateValue = false;
+			}
+
+			const uint32 NumPadComponents = NumDestComponents - NumSourceComponents;
+			FString CommaParameterCodeString = FString::Printf(TEXT(",%s"), *GetParameterCode(Code));
+
+			return AddInlinedCodeChunk(
+				DestType,
+				TEXT("%s(%s%s%s%s)"),
+				HLSLTypeString(DestType),
+				*GetParameterCode(Code),
+				NumPadComponents >= 1 ? (bReplicateValue ? *CommaParameterCodeString : TEXT(",0")) : TEXT(""),
+				NumPadComponents >= 2 ? (bReplicateValue ? *CommaParameterCodeString : TEXT(",0")) : TEXT(""),
+				NumPadComponents >= 3 ? (bReplicateValue ? *CommaParameterCodeString : TEXT(",0")) : TEXT("")
+				);
+		}
+		else
+		{
+			return Code;
+		}
+	}
+	else if ((SourceType & MCT_TextureVirtual) && (DestType & MCT_Texture2D))
+	{
+		return Code;
+	}
+	else
+	{
+		return Errorf(TEXT("Cannot force a cast between non-numeric types."));
+	}
+}
+
+/** Pushes a function onto the compiler's function stack, which indicates that compilation is entering a function. */
+void FHLSLMaterialTranslator::PushFunction(FMaterialFunctionCompileState* FunctionState)
+{
+	check(ShaderFrequency < SF_NumFrequencies);
+	auto& CurrentFunctionStack = FunctionStacks[ShaderFrequency];
+	CurrentFunctionStack.Push(FunctionState);
+}	
+
+/** Pops a function from the compiler's function stack, which indicates that compilation is leaving a function. */
+FMaterialFunctionCompileState* FHLSLMaterialTranslator::PopFunction()
+{
+	check(ShaderFrequency < SF_NumFrequencies);
+	auto& CurrentFunctionStack = FunctionStacks[ShaderFrequency];
+	return CurrentFunctionStack.Pop();
+}
+
+int32 FHLSLMaterialTranslator::GetCurrentFunctionStackDepth()
+{
+	check(ShaderFrequency < SF_NumFrequencies);
+	auto& CurrentFunctionStack = FunctionStacks[ShaderFrequency];
+	return CurrentFunctionStack.Num();
+}
+
+int32 FHLSLMaterialTranslator::AccessCollectionParameter(UMaterialParameterCollection* ParameterCollection, int32 ParameterIndex, int32 ComponentIndex)
+{
+	if (!ParameterCollection || ParameterIndex == -1)
+	{
+		return INDEX_NONE;
+	}
+
+	int32 CollectionIndex = ParameterCollections.Find(ParameterCollection);
+
+	if (CollectionIndex == INDEX_NONE)
+	{
+		if (ParameterCollections.Num() >= MaxNumParameterCollectionsPerMaterial)
+		{
+			return Error(TEXT("Material references too many MaterialParameterCollections!  A material may only reference 2 different collections."));
+		}
+
+		ParameterCollections.Add(ParameterCollection);
+		CollectionIndex = ParameterCollections.Num() - 1;
+	}
+
+	int32 VectorChunk = AddCodeChunk(MCT_Float4,TEXT("MaterialCollection%u.Vectors[%u]"),CollectionIndex,ParameterIndex);
+
+	return ComponentMask(VectorChunk, 
+		ComponentIndex == -1 ? true : ComponentIndex % 4 == 0,
+		ComponentIndex == -1 ? true : ComponentIndex % 4 == 1,
+		ComponentIndex == -1 ? true : ComponentIndex % 4 == 2,
+		ComponentIndex == -1 ? true : ComponentIndex % 4 == 3);
+}
+
+int32 FHLSLMaterialTranslator::ScalarParameter(FName ParameterName, float DefaultValue)
+{
+	FMaterialParameterInfo ParameterInfo = GetParameterAssociationInfo();
+	ParameterInfo.Name = ParameterName;
+	return AddUniformExpression(new FMaterialUniformExpressionScalarParameter(ParameterInfo,DefaultValue),MCT_Float,TEXT(""));
+}
+
+int32 FHLSLMaterialTranslator::VectorParameter(FName ParameterName, const FLinearColor& DefaultValue)
+{
+	FMaterialParameterInfo ParameterInfo = GetParameterAssociationInfo();
+	ParameterInfo.Name = ParameterName;
+	return AddUniformExpression(new FMaterialUniformExpressionVectorParameter(ParameterInfo,DefaultValue),MCT_Float4,TEXT(""));
+}
+
+int32 FHLSLMaterialTranslator::Constant(float X)
+{
+	return AddUniformExpression(new FMaterialUniformExpressionConstant(FLinearColor(X,X,X,X),MCT_Float),MCT_Float,TEXT("%0.8f"),X);
+}
+
+int32 FHLSLMaterialTranslator::Constant2(float X,float Y)
+{
+	return AddUniformExpression(new FMaterialUniformExpressionConstant(FLinearColor(X,Y,0,0),MCT_Float2),MCT_Float2,TEXT("MaterialFloat2(%0.8f,%0.8f)"),X,Y);
+}
+
+int32 FHLSLMaterialTranslator::Constant3(float X,float Y,float Z)
+{
+	return AddUniformExpression(new FMaterialUniformExpressionConstant(FLinearColor(X,Y,Z,0),MCT_Float3),MCT_Float3,TEXT("MaterialFloat3(%0.8f,%0.8f,%0.8f)"),X,Y,Z);
+}
+
+int32 FHLSLMaterialTranslator::Constant4(float X,float Y,float Z,float W)
+{
+	return AddUniformExpression(new FMaterialUniformExpressionConstant(FLinearColor(X,Y,Z,W),MCT_Float4),MCT_Float4,TEXT("MaterialFloat4(%0.8f,%0.8f,%0.8f,%0.8f)"),X,Y,Z,W);
+}
+	
+int32 FHLSLMaterialTranslator::ViewProperty(EMaterialExposedViewProperty Property, bool InvProperty)
+{
+	check(Property < MEVP_MAX);
+
+	// Compile time struct storing all EMaterialExposedViewProperty's enumerations' HLSL compilation specific meta information
+	struct EMaterialExposedViewPropertyMeta
+	{
+		EMaterialExposedViewProperty EnumValue;
+		EMaterialValueType Type;
+		const TCHAR * PropertyCode;
+		const TCHAR * InvPropertyCode;
+	};
+
+	static const EMaterialExposedViewPropertyMeta ViewPropertyMetaArray[] = {
+		{MEVP_BufferSize, MCT_Float2, TEXT("View.BufferSizeAndInvSize.xy"), TEXT("View.BufferSizeAndInvSize.zw")},
+		{MEVP_FieldOfView, MCT_Float2, TEXT("View.<PREV>FieldOfViewWideAngles"), nullptr},
+		{MEVP_TanHalfFieldOfView, MCT_Float2, TEXT("Get<PREV>TanHalfFieldOfView()"), TEXT("Get<PREV>CotanHalfFieldOfView()")},
+		{MEVP_ViewSize, MCT_Float2, TEXT("View.ViewSizeAndInvSize.xy"), TEXT("View.ViewSizeAndInvSize.zw")},
+		{MEVP_WorldSpaceViewPosition, MCT_Float3, TEXT("ResolvedView.<PREV>WorldViewOrigin"), nullptr},
+		{MEVP_WorldSpaceCameraPosition, MCT_Float3, TEXT("ResolvedView.<PREV>WorldCameraOrigin"), nullptr},
+		{MEVP_ViewportOffset, MCT_Float2, TEXT("View.ViewRectMin.xy"), nullptr},
+		{MEVP_TemporalSampleCount, MCT_Float1, TEXT("View.TemporalAAParams.y"), nullptr},
+		{MEVP_TemporalSampleIndex, MCT_Float1, TEXT("View.TemporalAAParams.x"), nullptr},
+		{MEVP_TemporalSampleOffset, MCT_Float2, TEXT("View.TemporalAAParams.zw"), nullptr},
+		{MEVP_RuntimeVirtualTextureOutputLevel, MCT_Float1, TEXT("View.VirtualTextureParams.x"), nullptr},
+		{MEVP_RuntimeVirtualTextureOutputDerivative, MCT_Float2, TEXT("View.VirtualTextureParams.zw"), nullptr},
+	};
+	static_assert((sizeof(ViewPropertyMetaArray) / sizeof(ViewPropertyMetaArray[0])) == MEVP_MAX, "incoherency between EMaterialExposedViewProperty and ViewPropertyMetaArray");
+
+	auto& PropertyMeta = ViewPropertyMetaArray[Property];
+	check(Property == PropertyMeta.EnumValue);
+
+	FString Code = PropertyMeta.PropertyCode;
+
+	if (InvProperty && PropertyMeta.InvPropertyCode)
+	{
+		Code = PropertyMeta.InvPropertyCode;
+	}
+
+	// Resolved templated code
+	Code.ReplaceInline(TEXT("<PREV>"), bCompilingPreviousFrame ? TEXT("Prev") : TEXT(""));
+		
+	if (InvProperty && !PropertyMeta.InvPropertyCode)
+	{
+		// fall back to compute the property's inverse from PropertyCode
+		return Div(Constant(1.f), AddInlinedCodeChunk(PropertyMeta.Type, *Code));
+	}
+
+	return AddCodeChunk(PropertyMeta.Type, *Code);
+}
+
+int32 FHLSLMaterialTranslator::GameTime(bool bPeriodic, float Period)
+{
+	if (!bPeriodic)
+	{
+		if (bCompilingPreviousFrame)
+		{
+			return AddInlinedCodeChunk(MCT_Float, TEXT("View.PrevFrameGameTime"));
+		}
+
+		return AddInlinedCodeChunk(MCT_Float, TEXT("View.GameTime"));
+	}
+	else if (Period == 0.0f)
+	{
+		return Constant(0.0f);
+	}
+
+	int32 PeriodChunk = Constant(Period);
+
+	if (bCompilingPreviousFrame)
+	{
+		return AddInlinedCodeChunk(MCT_Float, TEXT("fmod(View.PrevFrameGameTime,%s)"), *GetParameterCode(PeriodChunk));
+	}
+
+	// Note: not using FHLSLMaterialTranslator::Fmod(), which will emit MaterialFloat types which will be converted to fp16 on mobile.
+	// We want full 32 bit float precision until the fmod when using a period.
+	return AddInlinedCodeChunk(MCT_Float, TEXT("fmod(View.GameTime,%s)"), *GetParameterCode(PeriodChunk));
+}
+
+int32 FHLSLMaterialTranslator::RealTime(bool bPeriodic, float Period)
+{
+	if (!bPeriodic)
+	{
+		if (bCompilingPreviousFrame)
+		{
+			return AddInlinedCodeChunk(MCT_Float, TEXT("View.PrevFrameRealTime"));
+		}
+
+		return AddInlinedCodeChunk(MCT_Float, TEXT("View.RealTime"));
+	}
+	else if (Period == 0.0f)
+	{
+		return Constant(0.0f);
+	}
+
+	int32 PeriodChunk = Constant(Period);
+
+	if (bCompilingPreviousFrame)
+	{
+		return AddInlinedCodeChunk(MCT_Float, TEXT("fmod(View.PrevFrameRealTime,%s)"), *GetParameterCode(PeriodChunk));
+	}
+
+	return AddInlinedCodeChunk(MCT_Float, TEXT("fmod(View.RealTime,%s)"), *GetParameterCode(PeriodChunk));
+}
+
+int32 FHLSLMaterialTranslator::DeltaTime()
+{
+	// explicitly avoid trying to return previous frame's delta time for bCompilingPreviousFrame here
+	// DeltaTime expression is designed to be used when generating custom motion vectors, by using world position offset along with previous frame switch
+	// in this context, we will technically be evaluating the previous frame, but we want to use the current frame's delta tick in order to offset the vector used to create previous position
+	return AddInlinedCodeChunk(MCT_Float, TEXT("View.DeltaTime"));
+}
+
+int32 FHLSLMaterialTranslator::PeriodicHint(int32 PeriodicCode)
+{
+	if(PeriodicCode == INDEX_NONE)
+	{
+		return INDEX_NONE;
+	}
+
+	if(GetParameterUniformExpression(PeriodicCode))
+	{
+		return AddUniformExpression(new FMaterialUniformExpressionPeriodic(GetParameterUniformExpression(PeriodicCode)),GetParameterType(PeriodicCode),TEXT("%s"),*GetParameterCode(PeriodicCode));
+	}
+	else
+	{
+		return PeriodicCode;
+	}
+}
+
+int32 FHLSLMaterialTranslator::Sine(int32 X)
+{
+	if(X == INDEX_NONE)
+	{
+		return INDEX_NONE;
+	}
+
+	if(GetParameterUniformExpression(X))
+	{
+		return AddUniformExpression(new FMaterialUniformExpressionTrigMath(GetParameterUniformExpression(X),TMO_Sin),MCT_Float,TEXT("sin(%s)"),*CoerceParameter(X,MCT_Float));
+	}
+	else
+	{
+		return AddCodeChunk(GetParameterType(X),TEXT("sin(%s)"),*GetParameterCode(X));
+	}
+}
+
+int32 FHLSLMaterialTranslator::Cosine(int32 X)
+{
+	if(X == INDEX_NONE)
+	{
+		return INDEX_NONE;
+	}
+
+	if(GetParameterUniformExpression(X))
+	{
+		return AddUniformExpression(new FMaterialUniformExpressionTrigMath(GetParameterUniformExpression(X),TMO_Cos),MCT_Float,TEXT("cos(%s)"),*CoerceParameter(X,MCT_Float));
+	}
+	else
+	{
+		return AddCodeChunk(GetParameterType(X),TEXT("cos(%s)"),*GetParameterCode(X));
+	}
+}
+
+int32 FHLSLMaterialTranslator::Tangent(int32 X)
+{
+	if(X == INDEX_NONE)
+	{
+		return INDEX_NONE;
+	}
+
+	if(GetParameterUniformExpression(X))
+	{
+		return AddUniformExpression(new FMaterialUniformExpressionTrigMath(GetParameterUniformExpression(X),TMO_Tan),MCT_Float,TEXT("tan(%s)"),*CoerceParameter(X,MCT_Float));
+	}
+	else
+	{
+		return AddCodeChunk(GetParameterType(X),TEXT("tan(%s)"),*GetParameterCode(X));
+	}
+}
+
+int32 FHLSLMaterialTranslator::Arcsine(int32 X)
+{
+	if(X == INDEX_NONE)
+	{
+		return INDEX_NONE;
+	}
+
+	if(GetParameterUniformExpression(X))
+	{
+		return AddUniformExpression(new FMaterialUniformExpressionTrigMath(GetParameterUniformExpression(X),TMO_Asin),MCT_Float,TEXT("asin(%s)"),*CoerceParameter(X,MCT_Float));
+	}
+	else
+	{
+		return AddCodeChunk(GetParameterType(X),TEXT("asin(%s)"),*GetParameterCode(X));
+	}
+}
+
+int32 FHLSLMaterialTranslator::ArcsineFast(int32 X)
+{
+	if(X == INDEX_NONE)
+	{
+		return INDEX_NONE;
+	}
+
+	if(GetParameterUniformExpression(X))
+	{
+		return AddUniformExpression(new FMaterialUniformExpressionTrigMath(GetParameterUniformExpression(X),TMO_Asin),MCT_Float,TEXT("asinFast(%s)"),*CoerceParameter(X,MCT_Float));
+	}
+	else
+	{
+		return AddCodeChunk(GetParameterType(X),TEXT("asinFast(%s)"),*GetParameterCode(X));
+	}
+}
+
+int32 FHLSLMaterialTranslator::Arccosine(int32 X)
+{
+	if(X == INDEX_NONE)
+	{
+		return INDEX_NONE;
+	}
+
+	if(GetParameterUniformExpression(X))
+	{
+		return AddUniformExpression(new FMaterialUniformExpressionTrigMath(GetParameterUniformExpression(X),TMO_Acos),MCT_Float,TEXT("acos(%s)"),*CoerceParameter(X,MCT_Float));
+	}
+	else
+	{
+		return AddCodeChunk(GetParameterType(X),TEXT("acos(%s)"),*GetParameterCode(X));
+	}
+}
+
+int32 FHLSLMaterialTranslator::ArccosineFast(int32 X)
+{
+	if(X == INDEX_NONE)
+	{
+		return INDEX_NONE;
+	}
+
+	if(GetParameterUniformExpression(X))
+	{
+		return AddUniformExpression(new FMaterialUniformExpressionTrigMath(GetParameterUniformExpression(X),TMO_Acos),MCT_Float,TEXT("acosFast(%s)"),*CoerceParameter(X,MCT_Float));
+	}
+	else
+	{
+		return AddCodeChunk(GetParameterType(X),TEXT("acosFast(%s)"),*GetParameterCode(X));
+	}
+}
+
+int32 FHLSLMaterialTranslator::Arctangent(int32 X)
+{
+	if(X == INDEX_NONE)
+	{
+		return INDEX_NONE;
+	}
+
+	if(GetParameterUniformExpression(X))
+	{
+		return AddUniformExpression(new FMaterialUniformExpressionTrigMath(GetParameterUniformExpression(X),TMO_Atan),MCT_Float,TEXT("atan(%s)"),*CoerceParameter(X,MCT_Float));
+	}
+	else
+	{
+		return AddCodeChunk(GetParameterType(X),TEXT("atan(%s)"),*GetParameterCode(X));
+	}
+}
+
+int32 FHLSLMaterialTranslator::ArctangentFast(int32 X)
+{
+	if(X == INDEX_NONE)
+	{
+		return INDEX_NONE;
+	}
+
+	if(GetParameterUniformExpression(X))
+	{
+		return AddUniformExpression(new FMaterialUniformExpressionTrigMath(GetParameterUniformExpression(X),TMO_Atan),MCT_Float,TEXT("atanFast(%s)"),*CoerceParameter(X,MCT_Float));
+	}
+	else
+	{
+		return AddCodeChunk(GetParameterType(X),TEXT("atanFast(%s)"),*GetParameterCode(X));
+	}
+}
+
+int32 FHLSLMaterialTranslator::Arctangent2(int32 Y, int32 X)
+{
+	if(Y == INDEX_NONE || X == INDEX_NONE)
+	{
+		return INDEX_NONE;
+	}
+
+	if(GetParameterUniformExpression(Y) && GetParameterUniformExpression(X))
+	{
+		return AddUniformExpression(new FMaterialUniformExpressionTrigMath(GetParameterUniformExpression(Y),GetParameterUniformExpression(X),TMO_Atan2),MCT_Float,TEXT("atan2(%s, %s)"),*CoerceParameter(Y,MCT_Float),*CoerceParameter(X,MCT_Float));
+	}
+	else
+	{
+		return AddCodeChunk(GetParameterType(Y),TEXT("atan2(%s, %s)"),*GetParameterCode(Y),*GetParameterCode(X));
+	}
+}
+
+int32 FHLSLMaterialTranslator::Arctangent2Fast(int32 Y, int32 X)
+{
+	if(Y == INDEX_NONE || X == INDEX_NONE)
+	{
+		return INDEX_NONE;
+	}
+
+	if(GetParameterUniformExpression(Y) && GetParameterUniformExpression(X))
+	{
+		return AddUniformExpression(new FMaterialUniformExpressionTrigMath(GetParameterUniformExpression(Y),GetParameterUniformExpression(X),TMO_Atan2),MCT_Float,TEXT("atan2Fast(%s, %s)"),*CoerceParameter(Y,MCT_Float),*CoerceParameter(X,MCT_Float));
+	}
+	else
+	{
+		return AddCodeChunk(GetParameterType(Y),TEXT("atan2Fast(%s, %s)"),*GetParameterCode(Y),*GetParameterCode(X));
+	}
+}
+
+int32 FHLSLMaterialTranslator::Floor(int32 X)
+{
+	if(X == INDEX_NONE)
+	{
+		return INDEX_NONE;
+	}
+
+	if(GetParameterUniformExpression(X))
+	{
+		return AddUniformExpression(new FMaterialUniformExpressionFloor(GetParameterUniformExpression(X)),GetParameterType(X),TEXT("floor(%s)"),*GetParameterCode(X));
+	}
+	else
+	{
+		return AddCodeChunk(GetParameterType(X),TEXT("floor(%s)"),*GetParameterCode(X));
+	}
+}
+
+int32 FHLSLMaterialTranslator::Ceil(int32 X)
+{
+	if(X == INDEX_NONE)
+	{
+		return INDEX_NONE;
+	}
+
+	if(GetParameterUniformExpression(X))
+	{
+		return AddUniformExpression(new FMaterialUniformExpressionCeil(GetParameterUniformExpression(X)),GetParameterType(X),TEXT("ceil(%s)"),*GetParameterCode(X));
+	}
+	else
+	{
+		return AddCodeChunk(GetParameterType(X),TEXT("ceil(%s)"),*GetParameterCode(X));
+	}
+}
+
+int32 FHLSLMaterialTranslator::Round(int32 X)
+{
+	if(X == INDEX_NONE)
+	{
+		return INDEX_NONE;
+	}
+
+	if(GetParameterUniformExpression(X))
+	{
+		return AddUniformExpression(new FMaterialUniformExpressionRound(GetParameterUniformExpression(X)),GetParameterType(X),TEXT("round(%s)"),*GetParameterCode(X));
+	}
+	else
+	{
+		return AddCodeChunk(GetParameterType(X),TEXT("round(%s)"),*GetParameterCode(X));
+	}
+}
+
+int32 FHLSLMaterialTranslator::Truncate(int32 X)
+{
+	if(X == INDEX_NONE)
+	{
+		return INDEX_NONE;
+	}
+
+	if(GetParameterUniformExpression(X))
+	{
+		return AddUniformExpression(new FMaterialUniformExpressionTruncate(GetParameterUniformExpression(X)),GetParameterType(X),TEXT("trunc(%s)"),*GetParameterCode(X));
+	}
+	else
+	{
+		return AddCodeChunk(GetParameterType(X),TEXT("trunc(%s)"),*GetParameterCode(X));
+	}
+}
+
+int32 FHLSLMaterialTranslator::Sign(int32 X)
+{
+	if(X == INDEX_NONE)
+	{
+		return INDEX_NONE;
+	}
+
+	if(GetParameterUniformExpression(X))
+	{
+		return AddUniformExpression(new FMaterialUniformExpressionSign(GetParameterUniformExpression(X)),GetParameterType(X),TEXT("sign(%s)"),*GetParameterCode(X));
+	}
+	else
+	{
+		return AddCodeChunk(GetParameterType(X),TEXT("sign(%s)"),*GetParameterCode(X));
+	}
+}	
+
+int32 FHLSLMaterialTranslator::Frac(int32 X)
+{
+	if(X == INDEX_NONE)
+	{
+		return INDEX_NONE;
+	}
+
+	if(GetParameterUniformExpression(X))
+	{
+		return AddUniformExpression(new FMaterialUniformExpressionFrac(GetParameterUniformExpression(X)),GetParameterType(X),TEXT("frac(%s)"),*GetParameterCode(X));
+	}
+	else
+	{
+		return AddCodeChunk(GetParameterType(X),TEXT("frac(%s)"),*GetParameterCode(X));
+	}
+}
+
+int32 FHLSLMaterialTranslator::Fmod(int32 A, int32 B)
+{
+	if ((A == INDEX_NONE) || (B == INDEX_NONE))
+	{
+		return INDEX_NONE;
+	}
+
+	if (GetParameterUniformExpression(A) && GetParameterUniformExpression(B))
+	{
+		return AddUniformExpression(new FMaterialUniformExpressionFmod(GetParameterUniformExpression(A),GetParameterUniformExpression(B)),
+			GetParameterType(A),TEXT("fmod(%s,%s)"),*GetParameterCode(A),*CoerceParameter(B,GetParameterType(A)));
+	}
+	else
+	{
+		return AddCodeChunk(GetParameterType(A),
+			TEXT("fmod(%s,%s)"),*GetParameterCode(A),*CoerceParameter(B,GetParameterType(A)));
+	}
+}
+
+/**
+* Creates the new shader code chunk needed for the Abs expression
+*
+* @param	X - Index to the FMaterialCompiler::CodeChunk entry for the input expression
+* @return	Index to the new FMaterialCompiler::CodeChunk entry for this expression
+*/	
+int32 FHLSLMaterialTranslator::Abs( int32 X )
+{
+	if(X == INDEX_NONE)
+	{
+		return INDEX_NONE;
+	}
+
+	// get the user input struct for the input expression
+	FMaterialUniformExpression* pInputParam = GetParameterUniformExpression(X);
+	if( pInputParam )
+	{
+		FMaterialUniformExpressionAbs* pUniformExpression = new FMaterialUniformExpressionAbs( pInputParam );
+		return AddUniformExpression( pUniformExpression, GetParameterType(X), TEXT("abs(%s)"), *GetParameterCode(X) );
+	}
+	else
+	{
+		return AddCodeChunk( GetParameterType(X), TEXT("abs(%s)"), *GetParameterCode(X) );
+	}
+}
+
+int32 FHLSLMaterialTranslator::ReflectionVector()
+{
+	if (ShaderFrequency != SF_Pixel && ShaderFrequency != SF_Compute)
+	{
+		return NonPixelShaderExpressionError();
+	}
+	if (ShaderFrequency != SF_Vertex)
+	{
+		bUsesTransformVector = true;
+	}
+	return AddInlinedCodeChunk(MCT_Float3,TEXT("Parameters.ReflectionVector"));
+}
+
+int32 FHLSLMaterialTranslator::ReflectionAboutCustomWorldNormal(int32 CustomWorldNormal, int32 bNormalizeCustomWorldNormal)
+{
+	if (ShaderFrequency != SF_Pixel && ShaderFrequency != SF_Compute)
+	{
+		return NonPixelShaderExpressionError();
+	}
+
+	if (CustomWorldNormal == INDEX_NONE)
+	{
+		return INDEX_NONE;
+	}
+
+	if (ShaderFrequency != SF_Vertex)
+	{
+		bUsesTransformVector = true;
+	}
+
+	const TCHAR* ShouldNormalize = (!!bNormalizeCustomWorldNormal) ? TEXT("true") : TEXT("false");
+
+	return AddCodeChunk(MCT_Float3,TEXT("ReflectionAboutCustomWorldNormal(Parameters, %s, %s)"), *GetParameterCode(CustomWorldNormal), ShouldNormalize);
+}
+
+int32 FHLSLMaterialTranslator::CameraVector()
+{
+	if (ShaderFrequency != SF_Pixel && ShaderFrequency != SF_Compute)
+	{
+		return NonPixelShaderExpressionError();
+	}
+	if (ShaderFrequency != SF_Vertex)
+	{
+		bUsesTransformVector = true;
+	}
+	return AddInlinedCodeChunk(MCT_Float3,TEXT("Parameters.CameraVector"));
+}
+
+int32 FHLSLMaterialTranslator::LightVector()
+{
+	if (ShaderFrequency != SF_Pixel && ShaderFrequency != SF_Compute)
+	{
+		return NonPixelShaderExpressionError();
+	}
+
+	if (!Material->IsLightFunction() && !Material->IsDeferredDecal())
+	{
+		return Errorf(TEXT("LightVector can only be used in LightFunction or DeferredDecal materials"));
+	}
+
+	return AddInlinedCodeChunk(MCT_Float3,TEXT("Parameters.LightVector"));
+}
+
+int32 FHLSLMaterialTranslator::GetViewportUV()
+{
+	if (ShaderFrequency != SF_Pixel && ShaderFrequency != SF_Compute && ShaderFrequency != SF_Vertex)
+	{
+		return Errorf(TEXT("GetViewportUV() node is only available in vertex or pixel shader input."));
+	}
+	return AddCodeChunk(MCT_Float2, TEXT("GetViewportUV(Parameters)"));	
+}
+
+int32 FHLSLMaterialTranslator::GetPixelPosition()
+{
+	if (ShaderFrequency != SF_Pixel && ShaderFrequency != SF_Compute && ShaderFrequency != SF_Vertex)
+	{
+		return Errorf(TEXT("GetPixelPosition() node is only available in vertex or pixel shader input."));
+	}
+	return AddCodeChunk(MCT_Float2, TEXT("GetPixelPosition(Parameters)"));
+}
+
+int32 FHLSLMaterialTranslator::ParticleMacroUV()
+{
+	if (ShaderFrequency != SF_Pixel && ShaderFrequency != SF_Compute)
+	{
+		return NonPixelShaderExpressionError();
+	}
+
+	return AddCodeChunk(MCT_Float2,TEXT("GetParticleMacroUV(Parameters)"));
+}
+
+int32 FHLSLMaterialTranslator::ParticleSubUV(int32 TextureIndex, EMaterialSamplerType SamplerType, bool bBlend)
+{
+	if (ShaderFrequency != SF_Pixel && ShaderFrequency != SF_Compute)
+	{
+		return NonPixelShaderExpressionError();
+	}
+
+	if (TextureIndex == INDEX_NONE)
+	{
+		return INDEX_NONE;
+	}
+
+	int32 ParticleSubUV;
+	const FString TexCoordCode = TEXT("Parameters.Particle.SubUVCoords[%u].xy");
+	const int32 TexCoord1 = AddCodeChunk(MCT_Float2,*TexCoordCode,0);
+
+	if(bBlend)
+	{
+		// Out	 = linear interpolate... using 2 sub-images of the texture
+		// A	 = RGB sample texture with Parameters.Particle.SubUVCoords[0]
+		// B	 = RGB sample texture with Parameters.Particle.SubUVCoords[1]
+		// Alpha = Parameters.Particle.SubUVLerp
+
+		const int32 TexCoord2 = AddCodeChunk( MCT_Float2,*TexCoordCode,1);
+		const int32 SubImageLerp = AddCodeChunk(MCT_Float, TEXT("Parameters.Particle.SubUVLerp"));
+
+		const int32 TexSampleA = TextureSample(TextureIndex, TexCoord1, SamplerType );
+		const int32 TexSampleB = TextureSample(TextureIndex, TexCoord2, SamplerType );
+		ParticleSubUV = Lerp( TexSampleA,TexSampleB, SubImageLerp);
+	} 
+	else
+	{
+		ParticleSubUV = TextureSample(TextureIndex, TexCoord1, SamplerType );
+	}
+	
+	bUsesParticleSubUVs = true;
+	return ParticleSubUV;
+}
+
+int32 FHLSLMaterialTranslator::ParticleColor()
+{
+	if (ShaderFrequency != SF_Vertex && ShaderFrequency != SF_Pixel && ShaderFrequency != SF_Compute)
+	{
+		return NonVertexOrPixelShaderExpressionError();
+	}
+	bUsesParticleColor |= (ShaderFrequency != SF_Vertex);
+	return AddInlinedCodeChunk(MCT_Float4,TEXT("Parameters.Particle.Color"));	
+}
+
+int32 FHLSLMaterialTranslator::ParticlePosition()
+{
+	if (ShaderFrequency != SF_Vertex && ShaderFrequency != SF_Pixel && ShaderFrequency != SF_Compute)
+	{
+		return NonVertexOrPixelShaderExpressionError();
+	}
+	bNeedsParticlePosition = true;
+	return AddInlinedCodeChunk(MCT_Float3,TEXT("(Parameters.Particle.TranslatedWorldPositionAndSize.xyz - ResolvedView.PreViewTranslation.xyz)"));	
+}
+
+int32 FHLSLMaterialTranslator::ParticleRadius()
+{
+	if (ShaderFrequency != SF_Vertex && ShaderFrequency != SF_Pixel && ShaderFrequency != SF_Compute)
+	{
+		return NonVertexOrPixelShaderExpressionError();
+	}
+	bNeedsParticlePosition = true;
+	return AddInlinedCodeChunk(MCT_Float,TEXT("max(Parameters.Particle.TranslatedWorldPositionAndSize.w, .001f)"));	
+}
+
+int32 FHLSLMaterialTranslator::SphericalParticleOpacity(int32 Density)
+{
+	if (ShaderFrequency != SF_Pixel && ShaderFrequency != SF_Compute)
+	{
+		return NonPixelShaderExpressionError();
+	}
+
+	if (Density == INDEX_NONE)
+	{
+		return INDEX_NONE;
+	}
+
+	bNeedsParticlePosition = true;
+	bUsesSphericalParticleOpacity = true;
+	bNeedsWorldPositionExcludingShaderOffsets = true;
+	bUsesSceneDepth = true;
+	return AddCodeChunk(MCT_Float, TEXT("GetSphericalParticleOpacity(Parameters,%s)"), *GetParameterCode(Density));
+}
+
+int32 FHLSLMaterialTranslator::ParticleRelativeTime()
+{
+	if (ShaderFrequency != SF_Vertex && ShaderFrequency != SF_Pixel && ShaderFrequency != SF_Compute)
+	{
+		return NonVertexOrPixelShaderExpressionError();
+	}
+	bNeedsParticleTime = true;
+	return AddInlinedCodeChunk(MCT_Float,TEXT("Parameters.Particle.RelativeTime"));
+}
+
+int32 FHLSLMaterialTranslator::ParticleMotionBlurFade()
+{
+	if (ShaderFrequency != SF_Vertex && ShaderFrequency != SF_Pixel && ShaderFrequency != SF_Compute)
+	{
+		return NonVertexOrPixelShaderExpressionError();
+	}
+	bUsesParticleMotionBlur = true;
+	return AddInlinedCodeChunk(MCT_Float,TEXT("Parameters.Particle.MotionBlurFade"));
+}
+
+int32 FHLSLMaterialTranslator::ParticleRandom()
+{
+	if (ShaderFrequency != SF_Vertex && ShaderFrequency != SF_Pixel && ShaderFrequency != SF_Compute)
+	{
+		return NonVertexOrPixelShaderExpressionError();
+	}
+	bNeedsParticleRandom = true;
+	return AddInlinedCodeChunk(MCT_Float,TEXT("Parameters.Particle.Random"));
+}
+
+
+int32 FHLSLMaterialTranslator::ParticleDirection()
+{
+	if (ShaderFrequency != SF_Vertex && ShaderFrequency != SF_Pixel && ShaderFrequency != SF_Compute)
+	{
+		return NonVertexOrPixelShaderExpressionError();
+	}
+	bNeedsParticleVelocity = true;
+	return AddInlinedCodeChunk(MCT_Float3,TEXT("Parameters.Particle.Velocity.xyz"));
+}
+
+int32 FHLSLMaterialTranslator::ParticleSpeed()
+{
+	if (ShaderFrequency != SF_Vertex && ShaderFrequency != SF_Pixel && ShaderFrequency != SF_Compute)
+	{
+		return NonVertexOrPixelShaderExpressionError();
+	}
+	bNeedsParticleVelocity = true;
+	return AddInlinedCodeChunk(MCT_Float,TEXT("Parameters.Particle.Velocity.w"));
+}
+
+int32 FHLSLMaterialTranslator::ParticleSize()
+{
+	if (ShaderFrequency != SF_Vertex && ShaderFrequency != SF_Pixel && ShaderFrequency != SF_Compute)
+	{
+		return NonVertexOrPixelShaderExpressionError();
+	}
+	bNeedsParticleSize = true;
+	return AddInlinedCodeChunk(MCT_Float2,TEXT("Parameters.Particle.Size"));
+}
+
+int32 FHLSLMaterialTranslator::WorldPosition(EWorldPositionIncludedOffsets WorldPositionIncludedOffsets)
+{
+	FString FunctionNamePattern;
+
+	// If this material has no expressions for world position offset or world displacement, the non-offset world position will
+	// be exactly the same as the offset one, so there is no point bringing in the extra code.
+	// Also, we can't access the full offset world position in anything other than the pixel shader, because it won't have
+	// been calculated yet
+	switch (WorldPositionIncludedOffsets)
+	{
+	case WPT_Default:
+		{
+			FunctionNamePattern = TEXT("Get<PREV>WorldPosition");
+			break;
+		}
+
+	case WPT_ExcludeAllShaderOffsets:
+		{
+			if (FeatureLevel < ERHIFeatureLevel::ES3_1)
+			{
+				// World position excluding shader offsets is not available on ES2
+				FunctionNamePattern = TEXT("Get<PREV>WorldPosition");
+			}
+			else
+			{
+				bNeedsWorldPositionExcludingShaderOffsets = true;
+				FunctionNamePattern = TEXT("Get<PREV>WorldPosition<NO_MATERIAL_OFFSETS>");
+			}
+			break;
+		}
+
+	case WPT_CameraRelative:
+		{
+			FunctionNamePattern = TEXT("Get<PREV>TranslatedWorldPosition");
+			break;
+		}
+
+	case WPT_CameraRelativeNoOffsets:
+		{
+			if (FeatureLevel < ERHIFeatureLevel::ES3_1)
+			{
+				// World position excluding shader offsets is not available on ES2
+				FunctionNamePattern = TEXT("Get<PREV>TranslatedWorldPosition");
+			}
+			else
+			{
+				bNeedsWorldPositionExcludingShaderOffsets = true;
+				FunctionNamePattern = TEXT("Get<PREV>TranslatedWorldPosition<NO_MATERIAL_OFFSETS>");
+			}
+			break;
+		}
+
+	default:
+		{
+			Errorf(TEXT("Encountered unknown world position type '%d'"), WorldPositionIncludedOffsets);
+			return INDEX_NONE;
+		}
+	}
+
+	// If compiling for the previous frame in the vertex shader
+	FunctionNamePattern.ReplaceInline(TEXT("<PREV>"), bCompilingPreviousFrame && ShaderFrequency == SF_Vertex ? TEXT("Prev") : TEXT(""));
+		
+	if (ShaderFrequency == SF_Pixel)
+	{
+		// No material offset only available in the vertex shader.
+		// TODO: should also be available in the tesselation shader
+		FunctionNamePattern.ReplaceInline(TEXT("<NO_MATERIAL_OFFSETS>"), TEXT("_NoMaterialOffsets"));
+	}
+	else
+	{
+		FunctionNamePattern.ReplaceInline(TEXT("<NO_MATERIAL_OFFSETS>"), TEXT(""));
+	}
+
+	bUsesVertexPosition = true;
+
+	return AddInlinedCodeChunk(MCT_Float3, TEXT("%s(Parameters)"), *FunctionNamePattern);
+}
+
+int32 FHLSLMaterialTranslator::ObjectWorldPosition()
+{
+	return AddInlinedCodeChunk(MCT_Float3,TEXT("GetObjectWorldPosition(Parameters)"));		
+}
+
+int32 FHLSLMaterialTranslator::ObjectRadius()
+{
+	return GetPrimitiveProperty(MCT_Float, TEXT("ObjectRadius"), TEXT("ObjectWorldPositionAndRadius.w"));		
+}
+
+int32 FHLSLMaterialTranslator::ObjectBounds()
+{
+	return GetPrimitiveProperty(MCT_Float3, TEXT("ObjectBounds"), TEXT("ObjectBounds.xyz"));
+}
+
+int32 FHLSLMaterialTranslator::PreSkinnedLocalBounds(int32 OutputIndex)
+{
+	switch (OutputIndex)
+	{
+	case 0: // Half extents
+		return AddInlinedCodeChunk(MCT_Float3, TEXT("((GetPrimitiveData(Parameters.PrimitiveId).PreSkinnedLocalBoundsMax - GetPrimitiveData(Parameters.PrimitiveId).PreSkinnedLocalBoundsMin) / 2.0f)"));
+	case 1: // Full extents
+		return AddInlinedCodeChunk(MCT_Float3, TEXT("(GetPrimitiveData(Parameters.PrimitiveId).PreSkinnedLocalBoundsMax - GetPrimitiveData(Parameters.PrimitiveId).PreSkinnedLocalBoundsMin)"));
+	case 2: // Min point
+		return GetPrimitiveProperty(MCT_Float3, TEXT("PreSkinnedLocalBounds"), TEXT("PreSkinnedLocalBoundsMin"));
+	case 3: // Max point
+		return GetPrimitiveProperty(MCT_Float3, TEXT("PreSkinnedLocalBounds"), TEXT("PreSkinnedLocalBoundsMax"));
+	default:
+		check(false);
+	}
+
+	return INDEX_NONE; 
+}
+
+int32 FHLSLMaterialTranslator::DistanceCullFade()
+{
+	bUsesDistanceCullFade = true;
+
+	return AddInlinedCodeChunk(MCT_Float,TEXT("GetDistanceCullFade()"));		
+}
+
+int32 FHLSLMaterialTranslator::ActorWorldPosition()
+{
+	if (bCompilingPreviousFrame && ShaderFrequency == SF_Vertex)
+	{
+		// Decal VS doesn't have material code so FMaterialVertexParameters
+		// and primitve uniform buffer are guaranteed to exist if ActorPosition
+		// material node is used in VS
+		return AddInlinedCodeChunk(
+			MCT_Float3,
+			TEXT("mul(mul(float4(GetActorWorldPosition(Parameters.PrimitiveId), 1), GetPrimitiveData(Parameters.PrimitiveId).WorldToLocal), Parameters.PrevFrameLocalToWorld)"));
+	}
+	else
+	{
+		return AddInlinedCodeChunk(MCT_Float3, TEXT("GetActorWorldPosition(Parameters.PrimitiveId)"));
+	}
+}
+
+int32 FHLSLMaterialTranslator::If(int32 A,int32 B,int32 AGreaterThanB,int32 AEqualsB,int32 ALessThanB,int32 ThresholdArg)
+{
+	if(A == INDEX_NONE || B == INDEX_NONE || AGreaterThanB == INDEX_NONE || ALessThanB == INDEX_NONE)
+	{
+		return INDEX_NONE;
+	}
+
+	if (AEqualsB != INDEX_NONE)
+	{
+		if (ThresholdArg == INDEX_NONE)
+		{
+			return INDEX_NONE;
+		}
+
+		EMaterialValueType ResultType = GetArithmeticResultType(GetParameterType(AGreaterThanB),GetArithmeticResultType(AEqualsB,ALessThanB));
+
+		int32 CoercedAGreaterThanB = ForceCast(AGreaterThanB,ResultType);
+		int32 CoercedAEqualsB = ForceCast(AEqualsB,ResultType);
+		int32 CoercedALessThanB = ForceCast(ALessThanB,ResultType);
+
+		if(CoercedAGreaterThanB == INDEX_NONE || CoercedAEqualsB == INDEX_NONE || CoercedALessThanB == INDEX_NONE)
+		{
+			return INDEX_NONE;
+		}
+
+		return AddCodeChunk(
+			ResultType,
+			TEXT("((abs(%s - %s) > %s) ? (%s >= %s ? %s : %s) : %s)"),
+			*GetParameterCode(A),
+			*GetParameterCode(B),
+			*GetParameterCode(ThresholdArg),
+			*GetParameterCode(A),
+			*GetParameterCode(B),
+			*GetParameterCode(CoercedAGreaterThanB),
+			*GetParameterCode(CoercedALessThanB),
+			*GetParameterCode(CoercedAEqualsB)
+			);
+	}
+	else
+	{
+		EMaterialValueType ResultType = GetArithmeticResultType(AGreaterThanB,ALessThanB);
+
+		int32 CoercedAGreaterThanB = ForceCast(AGreaterThanB,ResultType);
+		int32 CoercedALessThanB = ForceCast(ALessThanB,ResultType);
+
+		if(CoercedAGreaterThanB == INDEX_NONE || CoercedALessThanB == INDEX_NONE)
+		{
+			return INDEX_NONE;
+		}
+
+		return AddCodeChunk(
+			ResultType,
+			TEXT("((%s >= %s) ? %s : %s)"),
+			*GetParameterCode(A),
+			*GetParameterCode(B),
+			*GetParameterCode(CoercedAGreaterThanB),
+			*GetParameterCode(CoercedALessThanB)
+			);
+	}
+}
+
+void FHLSLMaterialTranslator::AllocateSlot(TBitArray<>& InBitArray, int32 InSlotIndex, int32 InSlotCount) const
+{
+	// Grow as needed
+	int32 NumSlotsNeeded = InSlotIndex + InSlotCount;
+	int32 CurrentNumSlots = InBitArray.Num();
+	if(NumSlotsNeeded > CurrentNumSlots)
+	{
+		InBitArray.Add(false, NumSlotsNeeded - CurrentNumSlots);
+	}
+
+	// Allocate the requested slot(s)
+	for (int32 i = InSlotIndex; i < NumSlotsNeeded; ++i)
+	{
+		InBitArray[i] = true;
+	}
+}
+
+#if WITH_EDITOR
+int32 FHLSLMaterialTranslator::MaterialBakingWorldPosition()
+{
+	if (ShaderFrequency == SF_Vertex)
+	{
+		AllocateSlot(AllocatedUserVertexTexCoords, 6, 2);
+	}
+	else
+	{
+		AllocateSlot(AllocatedUserTexCoords, 6, 2);
+	}
+
+	// Note: inlining is important so that on ES2 devices, where half precision is used in the pixel shader, 
+	// The UV does not get assigned to a half temporary in cases where the texture sample is done directly from interpolated UVs
+	return AddInlinedCodeChunk(MCT_Float3, TEXT("float3(Parameters.TexCoords[6].x, Parameters.TexCoords[6].y, Parameters.TexCoords[7].x)"));
+}
+#endif
+		
+
+int32 FHLSLMaterialTranslator::TextureCoordinate(uint32 CoordinateIndex, bool UnMirrorU, bool UnMirrorV)
+{
+	// For WebGL 1 which is essentially GLES2.0, we can safely assume a higher number of supported vertex attributes
+	// even when we are compiling ES 2 feature level shaders.
+	// For UI materials can safely use more texture coordinates due to how they are packed in the slate material shader
+	// Landscape materials also calculate their texture coordinates in the vertex factory and do not need to be sent using an interpolator
+	const uint32 MaxNumCoordinates = ((Platform == SP_OPENGL_ES2_WEBGL) || (FeatureLevel != ERHIFeatureLevel::ES2) || Material->IsUIMaterial() || Material->IsUsedWithLandscape()) ? 8 : 3;
+
+	if (CoordinateIndex >= MaxNumCoordinates)
+	{
+		return Errorf(TEXT("Only %u texture coordinate sets can be used by this feature level, currently using %u"), MaxNumCoordinates, CoordinateIndex + 1);
+	}
+
+	if (ShaderFrequency == SF_Vertex)
+	{
+		AllocateSlot(AllocatedUserVertexTexCoords, CoordinateIndex);
+	}
+	else
+	{
+		AllocateSlot(AllocatedUserTexCoords, CoordinateIndex);
+	}
+
+	FString	SampleCode;
+	if ( UnMirrorU && UnMirrorV )
+	{
+		SampleCode = TEXT("UnMirrorUV(Parameters.TexCoords[%u].xy, Parameters)");
+	}
+	else if ( UnMirrorU )
+	{
+		SampleCode = TEXT("UnMirrorU(Parameters.TexCoords[%u].xy, Parameters)");
+	}
+	else if ( UnMirrorV )
+	{
+		SampleCode = TEXT("UnMirrorV(Parameters.TexCoords[%u].xy, Parameters)");
+	}
+	else
+	{
+		SampleCode = TEXT("Parameters.TexCoords[%u].xy");
+	}
+
+	// Note: inlining is important so that on ES2 devices, where half precision is used in the pixel shader, 
+	// The UV does not get assigned to a half temporary in cases where the texture sample is done directly from interpolated UVs
+	return AddInlinedCodeChunk(
+			MCT_Float2,
+			*SampleCode,
+			CoordinateIndex
+			);
+}
+
+static const TCHAR* GetVTAddressMode(TextureAddress Address)
+{
+	switch (Address)
+	{
+	case TA_Wrap: return TEXT("VTADDRESSMODE_WRAP");
+	case TA_Clamp: return TEXT("VTADDRESSMODE_CLAMP");
+	case TA_Mirror: return TEXT("VTADDRESSMODE_MIRROR");
+	default: checkNoEntry(); return nullptr;
+	}
+}
+
+uint32 FHLSLMaterialTranslator::AcquireVTStackIndex(ETextureMipValueMode MipValueMode, TextureAddress AddressU, TextureAddress AddressV, float AspectRatio, int32 CoordinateIndex, int32 MipValue0Index, int32 MipValue1Index, int32 PreallocatedStackTextureIndex, bool bGenerateFeedback)
+{
+	const uint64 CoordinatHash = GetParameterHash(CoordinateIndex);
+	const uint64 MipValue0Hash = GetParameterHash(MipValue0Index);
+	const uint64 MipValue1Hash = GetParameterHash(MipValue1Index);
+
+	uint64 Hash = CityHash128to64({ CurrentScopeID, CoordinatHash });
+	Hash = CityHash128to64({ Hash, MipValue0Hash });
+	Hash = CityHash128to64({ Hash, MipValue1Hash });
+	Hash = CityHash128to64({ Hash, (uint64)MipValueMode });
+	Hash = CityHash128to64({ Hash, (uint64)AddressU });
+	Hash = CityHash128to64({ Hash, (uint64)AddressV });
+	Hash = CityHash128to64({ Hash, (uint64)(AspectRatio * 1000.0f) });
+	Hash = CityHash128to64({ Hash, (uint64)PreallocatedStackTextureIndex });
+	Hash = CityHash128to64({ Hash, (uint64)(bGenerateFeedback ? 1 : 0) });
+
+	// First check to see if we have an existing VTStack that matches this key, that can still fit another layer
+	for (int32 Index = VTStackHash.First(Hash); VTStackHash.IsValid(Index); Index = VTStackHash.Next(Index))
+	{
+		const FMaterialVirtualTextureStack& Stack = MaterialCompilationOutput.UniformExpressionSet.VTStacks[Index];
+		const FMaterialVTStackEntry& Entry = VTStacks[Index];
+		if (!Stack.AreLayersFull() &&
+			Entry.ScopeID == CurrentScopeID &&
+			Entry.CoordinateHash == CoordinatHash &&
+			Entry.MipValue0Hash == MipValue0Hash &&
+			Entry.MipValue1Hash == MipValue1Hash &&
+			Entry.MipValueMode == MipValueMode &&
+			Entry.AddressU == AddressU &&
+			Entry.AddressV == AddressV &&
+			Entry.AspectRatio == AspectRatio &&
+			Entry.PreallocatedStackTextureIndex == PreallocatedStackTextureIndex &&
+			Entry.bGenerateFeedback == bGenerateFeedback)
+		{
+			return Index;
+		}
+	}
+
+	// Need to allocate a new VTStack
+	const int32 StackIndex = VTStacks.AddDefaulted();
+	VTStackHash.Add(Hash, StackIndex);
+	FMaterialVTStackEntry& Entry = VTStacks[StackIndex];
+	Entry.ScopeID = CurrentScopeID;
+	Entry.CoordinateHash = CoordinatHash;
+	Entry.MipValue0Hash = MipValue0Hash;
+	Entry.MipValue1Hash = MipValue1Hash;
+	Entry.MipValueMode = MipValueMode;
+	Entry.AddressU = AddressU;
+	Entry.AddressV = AddressV;
+	Entry.AspectRatio = AspectRatio;
+	Entry.DebugCoordinateIndex = CoordinateIndex;
+	Entry.DebugMipValue0Index = MipValue0Index;
+	Entry.DebugMipValue1Index = MipValue1Index;
+	Entry.PreallocatedStackTextureIndex = PreallocatedStackTextureIndex;
+	Entry.bGenerateFeedback = bGenerateFeedback;
+
+	MaterialCompilationOutput.UniformExpressionSet.VTStacks.Add(FMaterialVirtualTextureStack(PreallocatedStackTextureIndex));
+
+	// these two arrays need to stay in sync
+	check(VTStacks.Num() == MaterialCompilationOutput.UniformExpressionSet.VTStacks.Num());
+
+	// Optionally sample without virtual texture feedback but only for miplevel mode
+	check(bGenerateFeedback || MipValueMode == TMVM_MipLevel)
+	FString FeedbackParameter = bGenerateFeedback ? TEXT("Parameters.VirtualTextureFeedback,") : TEXT("");
+
+	// Code to load the VT page table...this will execute the first time a given VT stack is accessed
+	// Additional stack layers will simply reuse these results
+	switch (MipValueMode)
+	{
+	case TMVM_None:
+		Entry.CodeIndex = AddCodeChunk(MCT_VTPageTableResult, TEXT("TextureLoadVirtualPageTable(VIRTUALTEXTURE_PAGETABLE_%d, VTPageTableUniform_Unpack(Material.VTPackedPageTableUniform[%d*2], Material.VTPackedPageTableUniform[%d*2+1]), Parameters.SvPosition.xy, Parameters.VirtualTextureFeedback, %d + LIGHTMAP_VT_ENABLED, %s, %s, %s)"),
+			StackIndex, StackIndex, StackIndex, StackIndex, *CoerceParameter(CoordinateIndex, MCT_Float2), GetVTAddressMode(AddressU), GetVTAddressMode(AddressV));
+		break;
+	case TMVM_MipBias:
+		Entry.CodeIndex = AddCodeChunk(MCT_VTPageTableResult, TEXT("TextureLoadVirtualPageTableBias(VIRTUALTEXTURE_PAGETABLE_%d, VTPageTableUniform_Unpack(Material.VTPackedPageTableUniform[%d*2], Material.VTPackedPageTableUniform[%d*2+1]), Parameters.SvPosition.xy, Parameters.VirtualTextureFeedback, %d + LIGHTMAP_VT_ENABLED, %s, %s, %s, %s)"),
+			StackIndex, StackIndex, StackIndex, StackIndex, *CoerceParameter(CoordinateIndex, MCT_Float2), GetVTAddressMode(AddressU), GetVTAddressMode(AddressV), *CoerceParameter(MipValue0Index, MCT_Float1));
+		break;
+	case TMVM_MipLevel:
+		Entry.CodeIndex = AddCodeChunk(MCT_VTPageTableResult, TEXT("TextureLoadVirtualPageTableLevel(VIRTUALTEXTURE_PAGETABLE_%d, VTPageTableUniform_Unpack(Material.VTPackedPageTableUniform[%d*2], Material.VTPackedPageTableUniform[%d*2+1]), %s %d + LIGHTMAP_VT_ENABLED, %s, %s, %s, %s)"),
+			StackIndex, StackIndex, StackIndex, *FeedbackParameter, StackIndex, *CoerceParameter(CoordinateIndex, MCT_Float2), GetVTAddressMode(AddressU), GetVTAddressMode(AddressV), *CoerceParameter(MipValue0Index, MCT_Float1));
+		break;
+	case TMVM_Derivative:
+		Entry.CodeIndex = AddCodeChunk(MCT_VTPageTableResult, TEXT("TextureLoadVirtualPageTableGrad(VIRTUALTEXTURE_PAGETABLE_%d, VTPageTableUniform_Unpack(Material.VTPackedPageTableUniform[%d*2], Material.VTPackedPageTableUniform[%d*2+1]), Parameters.SvPosition.xy, Parameters.VirtualTextureFeedback, %d + LIGHTMAP_VT_ENABLED, %s, %s, %s, %s, %s)"),
+			StackIndex, StackIndex, StackIndex, StackIndex, *CoerceParameter(CoordinateIndex, MCT_Float2), GetVTAddressMode(AddressU), GetVTAddressMode(AddressV), *CoerceParameter(MipValue0Index, MCT_Float2), *CoerceParameter(MipValue1Index, MCT_Float2));
+		break;
+	default:
+		checkNoEntry();
+		break;
+	}
+
+	return StackIndex;
+}
+
+int32 FHLSLMaterialTranslator::TextureSample(
+	int32 TextureIndex,
+	int32 CoordinateIndex,
+	EMaterialSamplerType SamplerType,
+	int32 MipValue0Index,
+	int32 MipValue1Index,
+	ETextureMipValueMode MipValueMode,
+	ESamplerSourceMode SamplerSource,
+	int32 TextureReferenceIndex,
+	bool AutomaticViewMipBias
+	)
+{
+	if(TextureIndex == INDEX_NONE || CoordinateIndex == INDEX_NONE)
+	{
+		return INDEX_NONE;
+	}
+
+	if (FeatureLevel == ERHIFeatureLevel::ES2 && ShaderFrequency == SF_Vertex)
+	{
+		if (MipValueMode != TMVM_MipLevel)
+		{
+			Errorf(TEXT("Sampling from vertex textures requires an absolute mip level on feature level ES2!"));
+			return INDEX_NONE;
+		}
+	}
+	else if (ShaderFrequency != SF_Pixel
+		&& ErrorUnlessFeatureLevelSupported(ERHIFeatureLevel::ES3_1) == INDEX_NONE)
+	{
+		return INDEX_NONE;
+	}
+
+	EMaterialValueType TextureType = GetParameterType(TextureIndex);
+
+	if(!(TextureType & MCT_Texture))
+	{
+		Errorf(TEXT("Sampling unknown texture type: %s"),DescribeType(TextureType));
+		return INDEX_NONE;
+	}
+
+	if(ShaderFrequency != SF_Pixel && MipValueMode == TMVM_MipBias)
+	{
+		Errorf(TEXT("MipBias is only supported in the pixel shader"));
+		return INDEX_NONE;
+	}
+
+	const bool bVirtualTexture = TextureType == MCT_TextureVirtual;
+	if (bVirtualTexture)
+	{
+		if (Material->GetMaterialDomain() == MD_DeferredDecal)
+		{
+			if (Material->GetDecalBlendMode() == DBM_Volumetric_DistanceFunction)
+			{
+				return Errorf(TEXT("Sampling a virtual texture is currently only supported inside a volumetric decal."));
+			}
+		}
+		else if (Material->GetMaterialDomain() != MD_Surface)
+		{
+			return Errorf(TEXT("Sampling a virtual texture is currently only supported inside surface and decal shaders."));
+		}
+	}
+
+	if (MipValueMode == TMVM_Derivative)
+	{
+		if (MipValue0Index == INDEX_NONE)
+		{
+			return Errorf(TEXT("Missing DDX(UVs) parameter"));
+		}
+		else if (MipValue1Index == INDEX_NONE)
+		{
+			return Errorf(TEXT("Missing DDY(UVs) parameter"));
+		}
+		else if (!(GetParameterType(MipValue0Index) & MCT_Float))
+		{
+			return Errorf(TEXT("Invalid DDX(UVs) parameter"));
+		}
+		else if (!(GetParameterType(MipValue1Index) & MCT_Float))
+		{
+			return Errorf(TEXT("Invalid DDY(UVs) parameter"));
+		}
+	}
+	else if (MipValueMode != TMVM_None && MipValue0Index != INDEX_NONE && !(GetParameterType(MipValue0Index) & MCT_Float))
+	{
+		return Errorf(TEXT("Invalid mip map parameter"));
+	}
+
+	// if we are not in the PS we need a mip level
+	if(ShaderFrequency != SF_Pixel)
+	{
+		MipValueMode = TMVM_MipLevel;
+		AutomaticViewMipBias = false;
+
+		if (MipValue0Index == INDEX_NONE)
+		{
+			MipValue0Index = Constant(0.f);
+		}
+	}
+
+	// Automatic view mip bias is only for surface and decal domains.
+	if (Material->GetMaterialDomain() != MD_Surface && Material->GetMaterialDomain() != MD_DeferredDecal)
+	{
+		AutomaticViewMipBias = false;
+	}
+
+	// If mobile, then disabling AutomaticViewMipBias.
+	if (FeatureLevel < ERHIFeatureLevel::SM4)
+	{
+		AutomaticViewMipBias = false;
+	}
+
+	// If not 2D texture, disable AutomaticViewMipBias.
+	if (!(TextureType & (MCT_Texture2D|MCT_TextureVirtual)))
+	{
+		AutomaticViewMipBias = false;
+	}
+
+	FString SamplerStateCode;
+	bool RequiresManualViewMipBias = AutomaticViewMipBias;
+
+	if (!bVirtualTexture) //VT does not have explict samplers (and always requires manual view mip bias)
+	{
+		if (SamplerSource == SSM_FromTextureAsset)
+		{
+			SamplerStateCode = TEXT("%sSampler");
+		}
+		else if (SamplerSource == SSM_Wrap_WorldGroupSettings)
+		{
+			// Use the shared sampler to save sampler slots
+			SamplerStateCode = AutomaticViewMipBias
+				? TEXT("GetMaterialSharedSampler(%sSampler,View.MaterialTextureBilinearWrapedSampler)")
+				: TEXT("GetMaterialSharedSampler(%sSampler,Material.Wrap_WorldGroupSettings)");
+			RequiresManualViewMipBias = false;
+		}
+		else if (SamplerSource == SSM_Clamp_WorldGroupSettings)
+		{
+			// Use the shared sampler to save sampler slots
+			SamplerStateCode = AutomaticViewMipBias
+				? TEXT("GetMaterialSharedSampler(%sSampler,View.MaterialTextureBilinearClampedSampler)")
+				: TEXT("GetMaterialSharedSampler(%sSampler,Material.Clamp_WorldGroupSettings)");
+			RequiresManualViewMipBias = false;
+		}
+	}
+
+	FString SampleCode;
+	if (TextureType == MCT_TextureCube)
+	{
+		SampleCode += TEXT("TextureCubeSample");
+	}
+	else if (TextureType == MCT_VolumeTexture)
+	{
+		SampleCode += TEXT("Texture3DSample");
+	}
+	else if (TextureType == MCT_TextureExternal)
+	{
+		SampleCode += TEXT("TextureExternalSample");
+	}
+	else if (bVirtualTexture)
+	{
+		SampleCode += TEXT("TextureVirtualSample");
+	}
+	else // MCT_Texture2D
+	{
+		SampleCode += TEXT("Texture2DSample");
+	}
+		
+	EMaterialValueType UVsType = (TextureType == MCT_TextureCube || TextureType == MCT_VolumeTexture) ? MCT_Float3 : MCT_Float2;
+	
+	if (RequiresManualViewMipBias)
+	{
+		if (MipValueMode == TMVM_Derivative)
+		{
+			// When doing derivative based sampling, multiply.
+			int32 Multiplier = AddInlinedCodeChunk(MCT_Float, TEXT("View.MaterialTextureDerivativeMultiply"));
+			MipValue0Index = Mul(MipValue0Index, Multiplier);
+			MipValue1Index = Mul(MipValue1Index, Multiplier);
+		}
+		else if (MipValue0Index != INDEX_NONE && MipValueMode != TMVM_None)
+		{
+			// Adds bias to existing input level bias.
+			MipValue0Index = Add(MipValue0Index, AddInlinedCodeChunk(MCT_Float, TEXT("View.MaterialTextureMipBias")));
+		}
+		else
+		{
+			// Sets bias.
+			MipValue0Index = AddInlinedCodeChunk(MCT_Float1, TEXT("View.MaterialTextureMipBias"));
+		}
+
+		// If no Mip mode, then use MipBias.
+		MipValueMode = MipValueMode == TMVM_None ? TMVM_MipBias : MipValueMode;
+	}
+
+	FString MipValue0Code = TEXT("0.0f");
+	FString MipValue1Code = TEXT("0.0f");
+	if (MipValue0Index != INDEX_NONE && (MipValueMode == TMVM_MipBias || MipValueMode == TMVM_MipLevel))
+	{
+		MipValue0Code = CoerceParameter(MipValue0Index, MCT_Float1);
+	}
+	else if (MipValueMode == TMVM_Derivative)
+	{
+		MipValue0Code = CoerceParameter(MipValue0Index, UVsType);
+		MipValue1Code = CoerceParameter(MipValue1Index, UVsType);
+	}
+
+	if (bVirtualTexture)
+	{
+		// VT MipValueMode logic (most of work for VT case is in page table lookup)
+		if (MipValueMode == TMVM_MipLevel)
+		{
+			SampleCode += TEXT("Level");
+		}
+
+		// 'Texture name/sampler', 'PageTableResult', 'LayerIndex', 'PackedUniform'
+		SampleCode += TEXT("(%s, %s, %d, VTUniform_Unpack(Material.VTPackedUniform[%d]))");
+	}
+	else
+	{
+		// Non-VT MipValueMode logic			
+
+		// Re-route decal texture sampling so platforms may add specific workarounds there
+		if (ShaderFrequency == SF_Pixel && Material->GetMaterialDomain() == MD_DeferredDecal && MipValueMode == TMVM_None)
+		{
+			SampleCode += TEXT("_Decal");
+		}
+
+		SamplerStateCode = ", " + SamplerStateCode;
+
+		if (MipValueMode == TMVM_None)
+		{
+			SampleCode += TEXT("(%s") + SamplerStateCode + TEXT(",%s)");
+		}
+		else if (MipValueMode == TMVM_MipLevel)
+		{
+			// WebGL 2/GLES3.0 (or browsers with the texture lod extension) it is possible to sample from specific mip levels
+			// GLSL >= 100 should support this in the vertex shader
+			bool bES2MipSupport = (Platform == SP_OPENGL_ES2_WEBGL) || (Platform == SP_OPENGL_ES2_ANDROID && ShaderFrequency == SF_Vertex);
+			// Mobile: Sampling of a particular level depends on an extension; iOS does have it by default but
+			// there's a driver as of 7.0.2 that will cause a GPU hang if used with an Aniso > 1 sampler, so show an error for now
+			if (!bES2MipSupport && ErrorUnlessFeatureLevelSupported(ERHIFeatureLevel::ES3_1) == INDEX_NONE)
+			{
+				Errorf(TEXT("Sampling for a specific mip-level is not supported for ES2"));
+				return INDEX_NONE;
+			}
+
+			SampleCode += TEXT("Level(%s") + SamplerStateCode + TEXT(",%s,%s)");
+		}
+		else if (MipValueMode == TMVM_MipBias)
+		{
+			SampleCode += TEXT("Bias(%s") + SamplerStateCode + TEXT(",%s,%s)");
+		}
+		else if (MipValueMode == TMVM_Derivative)
+		{
+			SampleCode += TEXT("Grad(%s") + SamplerStateCode + TEXT(",%s,%s,%s)");
+		}
+		else
+		{
+			check(0);
+		}
+	}
+
+	switch( SamplerType )
+	{
+		case SAMPLERTYPE_External:
+			SampleCode = FString::Printf(TEXT("ProcessMaterialExternalTextureLookup(%s)"), *SampleCode);
+			break;
+
+		case SAMPLERTYPE_Color:
+		case SAMPLERTYPE_VirtualColor:
+			SampleCode = FString::Printf( TEXT("ProcessMaterialColorTextureLookup(%s)"), *SampleCode );
+			break;
+
+		case SAMPLERTYPE_LinearColor:
+		case SAMPLERTYPE_VirtualLinearColor:
+			SampleCode = FString::Printf(TEXT("ProcessMaterialLinearColorTextureLookup(%s)"), *SampleCode);
+		break;
+
+		case SAMPLERTYPE_Alpha:
+		case SAMPLERTYPE_VirtualAlpha:
+		case SAMPLERTYPE_DistanceFieldFont:
+			// Sampling a single channel texture in D3D9 gives: (G,G,G)
+			// Sampling a single channel texture in D3D11 gives: (G,0,0)
+			// This replication reproduces the D3D9 behavior in all cases.
+			SampleCode = FString::Printf( TEXT("(%s).rrrr"), *SampleCode );
+			break;
+			
+		case SAMPLERTYPE_Grayscale:
+		case SAMPLERTYPE_VirtualGrayscale:
+			// Sampling a greyscale texture in D3D9 gives: (G,G,G)
+			// Sampling a greyscale texture in D3D11 gives: (G,0,0)
+			// This replication reproduces the D3D9 behavior in all cases.
+			SampleCode = FString::Printf( TEXT("ProcessMaterialGreyscaleTextureLookup((%s).r).rrrr"), *SampleCode );
+			break;
+
+		case SAMPLERTYPE_LinearGrayscale:
+		case SAMPLERTYPE_VirtualLinearGrayscale:
+			// Sampling a greyscale texture in D3D9 gives: (G,G,G)
+			// Sampling a greyscale texture in D3D11 gives: (G,0,0)
+			// This replication reproduces the D3D9 behavior in all cases.
+			SampleCode = FString::Printf(TEXT("ProcessMaterialLinearGreyscaleTextureLookup((%s).r).rrrr"), *SampleCode);
+			break;
+
+		case SAMPLERTYPE_Normal:
+		case SAMPLERTYPE_VirtualNormal:
+			// Normal maps need to be unpacked in the pixel shader.
+			SampleCode = FString::Printf( TEXT("UnpackNormalMap(%s)"), *SampleCode );
+			break;
+
+		case SAMPLERTYPE_Masks:
+		case SAMPLERTYPE_VirtualMasks:
+			break;
+
+		case SAMPLERTYPE_Data:
+			break;
+	}
+
+	FString TextureName;
+	int32 VirtualTextureIndex = INDEX_NONE;
+
+	if (TextureType == MCT_TextureCube)
+	{
+		TextureName = CoerceParameter(TextureIndex, MCT_TextureCube);
+	}
+	else if (TextureType == MCT_VolumeTexture)
+	{
+		TextureName = CoerceParameter(TextureIndex, MCT_VolumeTexture);
+	}
+	else if (TextureType == MCT_TextureExternal)
+	{
+		TextureName = CoerceParameter(TextureIndex, MCT_TextureExternal);
+	}
+	else if (bVirtualTexture)
+	{
+		// Note, this does not really do anything (by design) other than adding it to the UniformExpressionSet
+		/*TextureName =*/ CoerceParameter(TextureIndex, TextureType);
+
+		FMaterialUniformExpression* UniformExpression = GetParameterUniformExpression(TextureIndex);
+		if (UniformExpression == nullptr)
+		{
+			return Errorf(TEXT("Unable to find VT uniform expression."));
+		}
+		FMaterialUniformExpressionTexture* TextureUniformExpression = UniformExpression->GetTextureUniformExpression();
+		if (TextureUniformExpression == nullptr)
+		{
+			return Errorf(TEXT("The provided uniform expression is not a texture"));
+		}
+
+		VirtualTextureIndex = MaterialCompilationOutput.UniformExpressionSet.UniformVirtualTextureExpressions.Find(TextureUniformExpression);
+		check(MaterialCompilationOutput.UniformExpressionSet.UniformVirtualTextureExpressions.IsValidIndex(VirtualTextureIndex));
+
+		if (SamplerSource != SSM_FromTextureAsset)
+		{
+			// VT doesn't care if the shared sampler is wrap or clamp this is handled in the shader explicitly by our code so we still inherit this from the texture
+			TextureName += FString::Printf(TEXT("Material.VirtualTexturePhysicalTable_%d, GetMaterialSharedSampler(Material.VirtualTexturePhysicalTable_%dSampler, View.SharedBilinearClampedSampler)")
+				, VirtualTextureIndex, VirtualTextureIndex);
+		}
+		else
+		{
+			TextureName += FString::Printf(TEXT("Material.VirtualTexturePhysicalTable_%d, Material.VirtualTexturePhysicalTable_%dSampler")
+				, VirtualTextureIndex, VirtualTextureIndex);
+		}
+
+		NumVtSamples++;
+ 	}
+	else // MCT_Texture2D
+	{
+		TextureName = CoerceParameter(TextureIndex, MCT_Texture2D);
+	}
+
+	const FString UVs = CoerceParameter(CoordinateIndex, UVsType);
+	const bool bStoreTexCoordScales = ShaderFrequency == SF_Pixel && TextureReferenceIndex != INDEX_NONE && Material && Material->GetShaderMapUsage() == EMaterialShaderMapUsage::DebugViewMode;
+	const bool bStoreAvailableVTLevel = ShaderFrequency == SF_Pixel && TextureReferenceIndex != INDEX_NONE && Material && Material->GetShaderMapUsage() == EMaterialShaderMapUsage::DebugViewMode;
+
+	if (bStoreTexCoordScales)
+	{
+		AddCodeChunk(MCT_Float, TEXT("StoreTexCoordScale(Parameters.TexCoordScalesParams, %s, %d)"), *UVs, (int)TextureReferenceIndex);
+	}
+
+	int32 VTStackIndex = INDEX_NONE;
+	int32 VTLayerIndex = INDEX_NONE;
+	int32 VTPageTableIndex = INDEX_NONE;
+	if (bVirtualTexture)
+	{
+		check(VirtualTextureIndex >= 0);
+
+		const FShaderCodeChunk&	TextureChunk = (*CurrentScopeChunks)[TextureIndex];
+		check(TextureChunk.UniformExpression);
+		const FMaterialUniformExpressionTexture* Expr = TextureChunk.UniformExpression->GetTextureUniformExpression();
+		check(Expr);
+		const UTexture2D* Tex2D = Cast<UTexture2D>(Material->GetReferencedTextures()[Expr->GetTextureIndex()]);
+
+		TextureAddress AddressU = TA_Wrap;
+		TextureAddress AddressV = TA_Wrap;
+		if (Tex2D && Tex2D->Source.GetNumBlocks() > 1)
+		{
+			// UDIM (multi-block) texture are forced to use wrap address mode
+			// This is important for supporting VT stacks made from UDIMs with differing number of blocks, as this requires wrapping vAddress for certain layers
+			AddressU = TA_Wrap;
+			AddressV = TA_Wrap;
+		}
+		else
+		{
+			switch (SamplerSource)
+			{
+			case SSM_FromTextureAsset:
+				check(Tex2D);
+				AddressU = Tex2D->AddressX;
+				AddressV = Tex2D->AddressY;
+				break;
+			case SSM_Wrap_WorldGroupSettings:
+				AddressU = TA_Wrap;
+				AddressV = TA_Wrap;
+				break;
+			case SSM_Clamp_WorldGroupSettings:
+				AddressU = TA_Clamp;
+				AddressV = TA_Clamp;
+				break;
+			default:
+				checkNoEntry();
+				break;
+			}
+		}
+
+		// Only support GPU feedback from pixel shader
+		//todo[vt]: Support feedback from other shader types
+		const bool bGenerateFeedback = ShaderFrequency == SF_Pixel;
+
+		VTLayerIndex = MaterialCompilationOutput.UniformExpressionSet.UniformVirtualTextureExpressions[VirtualTextureIndex]->GetTextureLayerIndex();
+		if (VTLayerIndex != INDEX_NONE)
+		{
+			// The layer index in the virtual texture stack is already known
+			// Create a page table sample for each new combination of virtual texture and sample parameters
+			VTStackIndex = AcquireVTStackIndex(MipValueMode, AddressU, AddressV, 1.0f, CoordinateIndex, MipValue0Index, MipValue1Index, TextureReferenceIndex, bGenerateFeedback);
+			VTPageTableIndex = MaterialCompilationOutput.UniformExpressionSet.UniformVirtualTextureExpressions[VirtualTextureIndex]->GetPageTableLayerIndex();
+		}
+		else
+		{
+			// Textures can only be combined in a VT stack if they have the same aspect ratio
+			// This also means that any texture parameters set in material instances for VTs must match the aspect ratio of the texture in the parent material
+			// (Otherwise could potentially break stacks)
+			check(Tex2D);
+
+			// Using Source size because we care about the aspect ratio of each block (each block of multi-block texture must have same aspect ratio)
+			// We can still combine multi-block textures of different block aspect ratios, as long as each block has the same ratio
+			// This is because we only need to overlay VT pages from within a given block
+			const float TextureAspectRatio = (float)Tex2D->Source.GetSizeX() / (float)Tex2D->Source.GetSizeY();
+
+			// Create a page table sample for each new set of sample parameters
+			VTStackIndex = AcquireVTStackIndex(MipValueMode, AddressU, AddressV, TextureAspectRatio, CoordinateIndex, MipValue0Index, MipValue1Index, INDEX_NONE, bGenerateFeedback);
+			// Allocate a layer in the virtual texture stack for this physical sample
+			VTLayerIndex = MaterialCompilationOutput.UniformExpressionSet.VTStacks[VTStackIndex].AddLayer();
+			VTPageTableIndex = VTLayerIndex;
+		}
+
+		MaterialCompilationOutput.UniformExpressionSet.VTStacks[VTStackIndex].SetLayer(VTLayerIndex, VirtualTextureIndex);
+	}
+
+	int32 SamplingCodeIndex = INDEX_NONE;
+	if (bVirtualTexture)
+	{
+		const FMaterialVTStackEntry& VTStackEntry = VTStacks[VTStackIndex];
+		const FString VTPageTableResult = GetParameterCode(VTStackEntry.CodeIndex);
+
+		SamplingCodeIndex = AddCodeChunk(
+			MCT_Float4,
+			*SampleCode,
+			*TextureName,
+			*VTPageTableResult,
+			VTPageTableIndex,
+			VirtualTextureIndex);
+
+		// TODO
+		/*if (bStoreAvailableVTLevel)
+		{
+			check(VirtualTextureUniformExpressionIndex >= 0);
+			check(VirtualTextureIndex >= 0);
+
+			AddCodeChunk(MCT_Float, TEXT("StoreAvailableVTLevel(Parameters.TexCoordScalesParams, TextureVirtualGetSampledLevelSize(Material.VirtualTexturePageTable_%d, Material.VirtualTextureUniformData[%d], Parameters.SvPosition.xy, %s), %d)"),
+				VirtualTextureUniformExpressionIndex,
+				VirtualTextureIndex,
+				*UVs,
+				(int)TextureReferenceIndex);
+		}*/
+	}
+	else
+	{
+		SamplingCodeIndex = AddCodeChunk(
+			MCT_Float4,
+			*SampleCode,
+			*TextureName,
+			*TextureName,
+			*UVs,
+			*MipValue0Code,
+			*MipValue1Code
+		);
+	}
+
+	AddEstimatedTextureSample();
+	if (bStoreTexCoordScales)
+	{
+		FString SamplingCode = CoerceParameter(SamplingCodeIndex, MCT_Float4);
+		AddCodeChunk(MCT_Float, TEXT("StoreTexSample(Parameters.TexCoordScalesParams, %s, %d)"), *SamplingCode, (int)TextureReferenceIndex);
+	}
+
+	return SamplingCodeIndex;
+}
+
+int32 FHLSLMaterialTranslator::TextureProperty(int32 TextureIndex, EMaterialExposedTextureProperty Property)
+{
+	EMaterialValueType TextureType = GetParameterType(TextureIndex);
+
+	if(TextureType != MCT_Texture2D && TextureType != MCT_TextureVirtual)
+	{
+		return Errorf(TEXT("Texture size only available for Texture2D, not %s"),DescribeType(TextureType));
+	}
+		
+	auto TextureExpression = (FMaterialUniformExpressionTexture*) (*CurrentScopeChunks)[TextureIndex].UniformExpression.GetReference();
+
+	return AddUniformExpression(new FMaterialUniformExpressionTextureProperty(TextureExpression, Property), MCT_Float2, TEXT(""));
+}
+
+int32 FHLSLMaterialTranslator::TextureDecalMipmapLevel(int32 TextureSizeInput)
+{
+	if (Material->GetMaterialDomain() != MD_DeferredDecal)
+	{
+		return Errorf(TEXT("Decal mipmap level only available in the decal material domain."));
+	}
+
+	EMaterialValueType TextureSizeType = GetParameterType(TextureSizeInput);
+
+	if (TextureSizeType != MCT_Float2)
+	{
+		Errorf(TEXT("Unmatching conversion %s -> float2"), DescribeType(TextureSizeType));
+		return INDEX_NONE;
+	}
+
+	FString TextureSize = CoerceParameter(TextureSizeInput, MCT_Float2);
+
+	return AddCodeChunk(
+		MCT_Float1,
+		TEXT("ComputeDecalMipmapLevel(Parameters,%s)"),
+		*TextureSize
+		);
+}
+
+int32 FHLSLMaterialTranslator::TextureDecalDerivative(bool bDDY)
+{
+	if (Material->GetMaterialDomain() != MD_DeferredDecal)
+	{
+		return Errorf(TEXT("Decal derivatives only available in the decal material domain."));
+	}
+
+	return AddCodeChunk(
+		MCT_Float2,
+		bDDY ? TEXT("ComputeDecalDDY(Parameters)") : TEXT("ComputeDecalDDX(Parameters)")
+		);
+}
+
+int32 FHLSLMaterialTranslator::DecalLifetimeOpacity()
+{
+	if (Material->GetMaterialDomain() != MD_DeferredDecal)
+	{
+		return Errorf(TEXT("Decal lifetime fade is only available in the decal material domain."));
+	}
+
+	if (ShaderFrequency != SF_Pixel)
+	{
+		return Errorf(TEXT("Decal lifetime fade is only available in the pixel shader."));
+	}
+
+	return AddCodeChunk(
+		MCT_Float,
+		TEXT("DecalLifetimeOpacity()")
+		);
+}
+
+int32 FHLSLMaterialTranslator::PixelDepth()
+{
+	if (ShaderFrequency != SF_Pixel && ShaderFrequency != SF_Compute && ShaderFrequency != SF_Vertex)
+	{
+		return Errorf(TEXT("Invalid node used in hull/domain shader input!"));
+	}
+	return AddInlinedCodeChunk(MCT_Float, TEXT("GetScreenPosition(Parameters).w"));		
+}
+
+/** Calculate screen aligned UV coordinates from an offset fraction or texture coordinate */
+int32 FHLSLMaterialTranslator::GetScreenAlignedUV(int32 Offset, int32 ViewportUV, bool bUseOffset)
+{
+	if(bUseOffset)
+	{
+		return AddCodeChunk(MCT_Float2, TEXT("CalcScreenUVFromOffsetFraction(GetScreenPosition(Parameters), %s)"), *GetParameterCode(Offset));
+	}
+	else if (ViewportUV != INDEX_NONE)
+	{
+		int32 BufferUV = AddCodeChunk(MCT_Float2, TEXT("MaterialFloat2(ViewportUVToBufferUV(%s))"), *CoerceParameter(ViewportUV, MCT_Float2));
+
+		EMaterialDomain MaterialDomain = Material->GetMaterialDomain();
+		int32 Min = AddInlinedCodeChunk(MCT_Float2, MaterialDomain == MD_Surface ? TEXT("ResolvedView.BufferBilinearUVMinMax.xy") : TEXT("View.BufferBilinearUVMinMax.xy"));
+		int32 Max = AddInlinedCodeChunk(MCT_Float2, MaterialDomain == MD_Surface ? TEXT("ResolvedView.BufferBilinearUVMinMax.zw") : TEXT("View.BufferBilinearUVMinMax.zw"));
+		return Clamp(BufferUV, Min, Max);
+	}
+	else
+	{
+		return AddInlinedCodeChunk(MCT_Float2, TEXT("ScreenAlignedPosition(GetScreenPosition(Parameters))"));
+	}
+}
+
+int32 FHLSLMaterialTranslator::SceneDepth(int32 Offset, int32 ViewportUV, bool bUseOffset)
+{
+	if (ShaderFrequency == SF_Vertex && FeatureLevel <= ERHIFeatureLevel::ES3_1)
+	{
+		// mobile currently does not support this, we need to read a separate copy of the depth, we must disable framebuffer fetch and force scene texture reads.
+		return Errorf(TEXT("Cannot read scene depth from the vertex shader with feature level ES3.1 or below."));
+	}
+
+	if (Offset == INDEX_NONE && bUseOffset)
+	{
+		return INDEX_NONE;
+	}
+
+	bUsesSceneDepth = true;
+	AddEstimatedTextureSample();
+
+	FString	UserDepthCode(TEXT("CalcSceneDepth(%s)"));
+	int32 TexCoordCode = GetScreenAlignedUV(Offset, ViewportUV, bUseOffset);
+	// add the code string
+	return AddCodeChunk(
+		MCT_Float,
+		*UserDepthCode,
+		*GetParameterCode(TexCoordCode)
+		);
+}
+	
+// @param SceneTextureId of type ESceneTextureId e.g. PPI_SubsurfaceColor
+int32 FHLSLMaterialTranslator::SceneTextureLookup(int32 ViewportUV, uint32 InSceneTextureId, bool bFiltered)
+{
+	ESceneTextureId SceneTextureId = (ESceneTextureId)InSceneTextureId;
+
+	const bool bSupportedOnMobile = SceneTextureId == PPI_PostProcessInput0 ||
+									SceneTextureId == PPI_CustomDepth ||
+									SceneTextureId == PPI_SceneDepth ||
+									SceneTextureId == PPI_CustomStencil;
+
+	if (!bSupportedOnMobile	&& ErrorUnlessFeatureLevelSupported(ERHIFeatureLevel::SM4) == INDEX_NONE)
+	{
+		return INDEX_NONE;
+	}
+
+	if (ShaderFrequency != SF_Pixel && ShaderFrequency != SF_Vertex)
+	{
+		// we can relax this later if needed
+		return NonPixelShaderExpressionError();
+	}
+		
+	if (SceneTextureId == PPI_DecalMask)
+	{
+		return Error(TEXT("Decal Mask bit was move out of GBuffer to the stencil buffer for performance optimisation and is therefor no longer available"));
+	}
+
+	UseSceneTextureId(SceneTextureId, true);
+
+	int32 BufferUV;
+	if (ViewportUV != INDEX_NONE)
+	{
+		BufferUV = AddCodeChunk(MCT_Float2,
+			TEXT("ClampSceneTextureUV(ViewportUVToSceneTextureUV(%s, %d), %d)"),
+			*CoerceParameter(ViewportUV, MCT_Float2), (int)SceneTextureId, (int)SceneTextureId);
+	}
+	else
+	{
+		BufferUV = AddInlinedCodeChunk(MCT_Float2, TEXT("GetDefaultSceneTextureUV(Parameters, %d)"), (int)SceneTextureId);
+	}
+
+	AddEstimatedTextureSample();
+
+	if (FeatureLevel >= ERHIFeatureLevel::SM4)
+	{
+		return AddCodeChunk(
+			MCT_Float4,
+			TEXT("SceneTextureLookup(%s, %d, %s)"),
+			*CoerceParameter(BufferUV, MCT_Float2), (int)SceneTextureId, bFiltered ? TEXT("true") : TEXT("false")
+			);
+	}
+	else // mobile
+	{
+		int32 UV = BufferUV;
+		if (Material->GetMaterialDomain() == MD_PostProcess)
+		{
+			int32 BlendableLocation = Material->GetBlendableLocation();
+			if (SceneTextureId == PPI_SceneDepth && !(BlendableLocation == BL_BeforeTranslucency || BlendableLocation == BL_BeforeTonemapping))
+			{
+				// SceneDepth lookups are not available when using MSAA, but we can access depth stored in SceneColor.A channel
+				// SceneColor.A channel holds depth till BeforeTonemapping location, then it's gets overwritten
+				return Errorf(TEXT("SceneDepth lookups are only available when BlendableLocation is BeforeTranslucency or BeforeTonemapping"));
+			}
+		}
+			
+		return AddCodeChunk(MCT_Float4,	TEXT("MobileSceneTextureLookup(Parameters, %d, %s)"), (int32)SceneTextureId, *CoerceParameter(UV, MCT_Float2));
+	}
+}
+
+int32 FHLSLMaterialTranslator::GetSceneTextureViewSize(int32 SceneTextureId, bool InvProperty)
+{
+	if (InvProperty)
+	{
+		return AddCodeChunk(MCT_Float2, TEXT("GetSceneTextureViewSize(%d).zw"), SceneTextureId);
+	}
+	return AddCodeChunk(MCT_Float2, TEXT("GetSceneTextureViewSize(%d).xy"), SceneTextureId);
+}
+
+// @param bTextureLookup true: texture, false:no texture lookup, usually to get the size
+void FHLSLMaterialTranslator::UseSceneTextureId(ESceneTextureId SceneTextureId, bool bTextureLookup)
+{
+	MaterialCompilationOutput.bNeedsSceneTextures = true;
+	MaterialCompilationOutput.SetIsSceneTextureUsed(SceneTextureId);
+
+	if(Material->GetMaterialDomain() == MD_DeferredDecal)
+	{
+		EDecalBlendMode DecalBlendMode = (EDecalBlendMode)Material->GetDecalBlendMode();
+		bool bDBuffer = IsDBufferDecalBlendMode(DecalBlendMode);
+
+		bool bRequiresSM5 = (SceneTextureId == PPI_WorldNormal || SceneTextureId == PPI_CustomDepth || SceneTextureId == PPI_CustomStencil || SceneTextureId == PPI_AmbientOcclusion);
+
+		if(bDBuffer)
+		{
+			if(!(SceneTextureId == PPI_SceneDepth || SceneTextureId == PPI_CustomDepth || SceneTextureId == PPI_CustomStencil))
+			{
+				// Note: For DBuffer decals: CustomDepth and CustomStencil are only available if r.CustomDepth.Order = 0
+				Errorf(TEXT("DBuffer decals (MaterialDomain=DeferredDecal and DecalBlendMode is using DBuffer) can only access SceneDepth, CustomDepth, CustomStencil"));
+			}
+		}
+		else
+		{
+			if(!(SceneTextureId == PPI_SceneDepth || SceneTextureId == PPI_CustomDepth || SceneTextureId == PPI_CustomStencil || SceneTextureId == PPI_WorldNormal || SceneTextureId == PPI_AmbientOcclusion))
+			{
+				Errorf(TEXT("Decals (MaterialDomain=DeferredDecal) can only access WorldNormal, AmbientOcclusion, SceneDepth, CustomDepth, CustomStencil"));
+			}
+
+			if (SceneTextureId == PPI_WorldNormal && Material->HasNormalConnected())
+			{
+					// GBuffer can only relate to WorldNormal here.
+				Errorf(TEXT("Decals that read WorldNormal cannot output to normal at the same time"));
+			}
+		}
+
+		if (bRequiresSM5)
+		{
+			ErrorUnlessFeatureLevelSupported(ERHIFeatureLevel::SM4);
+		}
+	}
+
+	if(SceneTextureId == PPI_SceneColor && Material->GetMaterialDomain() != MD_Surface)
+	{
+		if(Material->GetMaterialDomain() == MD_PostProcess)
+		{
+			Errorf(TEXT("SceneColor lookups are only available when MaterialDomain = Surface. PostProcessMaterials should use the SceneTexture PostProcessInput0."));
+		}
+		else
+		{
+			Errorf(TEXT("SceneColor lookups are only available when MaterialDomain = Surface."));
+		}
+	}
+
+	if(bTextureLookup)
+	{
+		bNeedsSceneTexturePostProcessInputs = bNeedsSceneTexturePostProcessInputs
+			|| ((SceneTextureId >= PPI_PostProcessInput0 && SceneTextureId <= PPI_PostProcessInput6)
+			|| SceneTextureId == PPI_Velocity
+			|| SceneTextureId == PPI_SceneColor);
+
+	}
+
+	if (SceneTextureId == PPI_SceneDepth && bTextureLookup)
+	{
+		bUsesSceneDepth = true;
+	}
+
+	const bool bNeedsGBuffer = MaterialCompilationOutput.NeedsGBuffer();
+
+	if (bNeedsGBuffer && IsForwardShadingEnabled(Platform))
+	{
+		Errorf(TEXT("GBuffer scene textures not available with forward shading."));
+	}
+
+	if (SceneTextureId == PPI_Velocity)
+	{
+		if (Material->GetMaterialDomain() != MD_PostProcess)
+		{
+			Errorf(TEXT("Velocity scene textures are only available in post process materials."));
+		}
+	}
+
+	// not yet tracked:
+	//   PPI_SeparateTranslucency, PPI_CustomDepth, PPI_AmbientOcclusion
+}
+
+int32 FHLSLMaterialTranslator::SceneColor(int32 Offset, int32 ViewportUV, bool bUseOffset)
+{
+	if (Offset == INDEX_NONE && bUseOffset)
+	{
+		return INDEX_NONE;
+	}
+
+	if (ShaderFrequency != SF_Pixel)
+	{
+		return NonPixelShaderExpressionError();
+	}
+
+	if(Material->GetMaterialDomain() != MD_Surface)
+	{
+		Errorf(TEXT("SceneColor lookups are only available when MaterialDomain = Surface."));
+	}
+
+	if (ErrorUnlessFeatureLevelSupported(ERHIFeatureLevel::SM4) == INDEX_NONE)
+	{
+		return INDEX_NONE;
+	}
+
+	MaterialCompilationOutput.SetIsSceneTextureUsed(PPI_SceneColor);
+	AddEstimatedTextureSample();
+
+	int32 ScreenUVCode = GetScreenAlignedUV(Offset, ViewportUV, bUseOffset);
+	return AddCodeChunk(
+		MCT_Float3,
+		TEXT("DecodeSceneColorForMaterialNode(%s)"),
+		*GetParameterCode(ScreenUVCode)
+		);
+}
+
+int32 FHLSLMaterialTranslator::Texture(UTexture* InTexture, int32& TextureReferenceIndex, EMaterialSamplerType SamplerType, ESamplerSourceMode SamplerSource, ETextureMipValueMode MipValueMode)
+{
+	if (FeatureLevel == ERHIFeatureLevel::ES2 && ShaderFrequency == SF_Vertex)
+	{
+		if (MipValueMode != TMVM_MipLevel)
+		{
+			Errorf(TEXT("Sampling from vertex textures requires an absolute mip level on feature level ES2"));
+			return INDEX_NONE;
+		}
+	}
+	else if (ShaderFrequency != SF_Pixel
+		&& ErrorUnlessFeatureLevelSupported(ERHIFeatureLevel::ES3_1) == INDEX_NONE)
+	{
+		return INDEX_NONE;
+	}
+
+	EMaterialValueType ShaderType = InTexture->GetMaterialType();
+	TextureReferenceIndex = Material->GetReferencedTextures().Find(InTexture);
+
+#if DO_CHECK
+	// UE-3518: Additional pre-assert logging to help determine the cause of this failure.
+	if (TextureReferenceIndex == INDEX_NONE)
+	{
+		const TArray<UObject*>& ReferencedTextures = Material->GetReferencedTextures();
+		UE_LOG(LogMaterial, Error, TEXT("Compiler->Texture() failed to find texture '%s' in referenced list of size '%i':"), *InTexture->GetName(), ReferencedTextures.Num());
+		for (int32 i = 0; i < ReferencedTextures.Num(); ++i)
+		{
+			UE_LOG(LogMaterial, Error, TEXT("%i: '%s'"), i, ReferencedTextures[i] ? *ReferencedTextures[i]->GetName() : TEXT("nullptr"));
+		}
+	}
+#endif
+	checkf(TextureReferenceIndex != INDEX_NONE, TEXT("Material expression called Compiler->Texture() without implementing UMaterialExpression::GetReferencedTexture properly"));
+
+	const bool bVirtualTexturesEnabeled = UseVirtualTexturing(FeatureLevel, TargetPlatform);
+	bool bVirtual = ShaderType == MCT_TextureVirtual;
+	if (bVirtualTexturesEnabeled == false && ShaderType == MCT_TextureVirtual)
+	{
+		bVirtual = false;
+		ShaderType = MCT_Texture2D;
+	}
+	return AddUniformExpression(new FMaterialUniformExpressionTexture(TextureReferenceIndex, SamplerType, SamplerSource, bVirtual),ShaderType,TEXT(""));
+}
+
+int32 FHLSLMaterialTranslator::TextureParameter(FName ParameterName, UTexture* DefaultValue, int32& TextureReferenceIndex, EMaterialSamplerType SamplerType, ESamplerSourceMode SamplerSource)
+{
+	if (ShaderFrequency != SF_Pixel
+		&& ErrorUnlessFeatureLevelSupported(ERHIFeatureLevel::ES3_1) == INDEX_NONE)
+	{
+		return INDEX_NONE;
+	}
+
+	EMaterialValueType ShaderType = DefaultValue->GetMaterialType();
+	TextureReferenceIndex = Material->GetReferencedTextures().Find(DefaultValue);
+	checkf(TextureReferenceIndex != INDEX_NONE, TEXT("Material expression called Compiler->TextureParameter() without implementing UMaterialExpression::GetReferencedTexture properly"));
+
+	FMaterialParameterInfo ParameterInfo = GetParameterAssociationInfo();
+	ParameterInfo.Name = ParameterName;
+
+	const bool bVirtualTexturesEnabled = UseVirtualTexturing(FeatureLevel, TargetPlatform);
+	bool bVirtual = ShaderType == MCT_TextureVirtual;
+	if (bVirtualTexturesEnabled == false && ShaderType == MCT_TextureVirtual)
+	{
+		bVirtual = false;
+		ShaderType = MCT_Texture2D;
+	}
+	return AddUniformExpression(new FMaterialUniformExpressionTextureParameter(ParameterInfo, TextureReferenceIndex, SamplerType, SamplerSource, bVirtual),ShaderType,TEXT(""));
+}
+
+int32 FHLSLMaterialTranslator::VirtualTexture(URuntimeVirtualTexture* InTexture, int32 TextureLayerIndex, int32 PageTableLayerIndex, int32& TextureReferenceIndex, EMaterialSamplerType SamplerType)
+{
+	if (!UseVirtualTexturing(FeatureLevel, TargetPlatform))
+	{
+		return INDEX_NONE;
+	}
+
+	TextureReferenceIndex = Material->GetReferencedTextures().Find(InTexture);
+	checkf(TextureReferenceIndex != INDEX_NONE, TEXT("Material expression called Compiler->VirtualTexture() without implementing UMaterialExpression::GetReferencedTexture properly"));
+
+	return AddUniformExpression(new FMaterialUniformExpressionTexture(TextureReferenceIndex, TextureLayerIndex, PageTableLayerIndex, SamplerType), MCT_TextureVirtual, TEXT(""));
+}
+
+int32 FHLSLMaterialTranslator::VirtualTextureParameter(FName ParameterName, URuntimeVirtualTexture* DefaultValue, int32 TextureLayerIndex, int32 PageTableLayerIndex, int32& TextureReferenceIndex, EMaterialSamplerType SamplerType)
+{
+	if (!UseVirtualTexturing(FeatureLevel, TargetPlatform))
+	{
+		return INDEX_NONE;
+	}
+
+	TextureReferenceIndex = Material->GetReferencedTextures().Find(DefaultValue);
+	checkf(TextureReferenceIndex != INDEX_NONE, TEXT("Material expression called Compiler->VirtualTexture() without implementing UMaterialExpression::GetReferencedTexture properly"));
+
+	FMaterialParameterInfo ParameterInfo = GetParameterAssociationInfo();
+	ParameterInfo.Name = ParameterName;
+
+	return AddUniformExpression(new FMaterialUniformExpressionTextureParameter(ParameterInfo, TextureReferenceIndex, TextureLayerIndex, PageTableLayerIndex, SamplerType), MCT_TextureVirtual, TEXT(""));
+}
+
+int32 FHLSLMaterialTranslator::VirtualTextureUniform(int32 TextureIndex, int32 VectorIndex)
+{
+	return AddUniformExpression(new FMaterialUniformExpressionRuntimeVirtualTextureUniform(TextureIndex, VectorIndex), MCT_Float3, TEXT(""));
+}
+
+int32 FHLSLMaterialTranslator::VirtualTextureUniform(FName ParameterName, int32 TextureIndex, int32 VectorIndex)
+{
+	FMaterialParameterInfo ParameterInfo = GetParameterAssociationInfo();
+	ParameterInfo.Name = ParameterName;
+
+	return AddUniformExpression(new FMaterialUniformExpressionRuntimeVirtualTextureUniform(ParameterInfo, TextureIndex, VectorIndex), MCT_Float3, TEXT(""));
+}
+
+int32 FHLSLMaterialTranslator::VirtualTextureWorldToUV(int32 WorldPositionIndex, int32 P0, int32 P1, int32 P2)
+{
+	FString	SampleCode(TEXT("VirtualTextureWorldToUV(%s, %s, %s, %s)"));
+	return AddInlinedCodeChunk(MCT_Float2, *SampleCode, *GetParameterCode(WorldPositionIndex), *GetParameterCode(P0), *GetParameterCode(P1), *GetParameterCode(P2));
+}
+
+int32 FHLSLMaterialTranslator::VirtualTextureUnpack(int32 CodeIndex0, int32 CodeIndex1, int32 CodeIndex2, EVirtualTextureUnpackType UnpackType)
+{
+	if (UnpackType == EVirtualTextureUnpackType::BaseColorYCoCg)
+	{
+		FString	SampleCode(TEXT("VirtualTextureUnpackBaseColorYCoCg(%s)"));
+		return CodeIndex0 == INDEX_NONE ? INDEX_NONE : AddCodeChunk(MCT_Float3, *SampleCode, *GetParameterCode(CodeIndex0));
+	}
+	else if (UnpackType == EVirtualTextureUnpackType::NormalBC3)
+	{
+		FString	SampleCode(TEXT("VirtualTextureUnpackNormalBC3(%s)"));
+		return CodeIndex1 == INDEX_NONE ? INDEX_NONE : AddCodeChunk(MCT_Float3, *SampleCode, *GetParameterCode(CodeIndex1));
+	}
+	else if (UnpackType == EVirtualTextureUnpackType::NormalBC5)
+	{
+		FString	SampleCode(TEXT("VirtualTextureUnpackNormalBC5(%s)"));
+		return CodeIndex1 == INDEX_NONE ? INDEX_NONE : AddCodeChunk(MCT_Float3, *SampleCode, *GetParameterCode(CodeIndex1));
+	}
+	else if (UnpackType == EVirtualTextureUnpackType::NormalBC3BC3)
+	{
+		FString	SampleCode(TEXT("VirtualTextureUnpackNormalBC3BC3(%s, %s)"));
+		return CodeIndex0 == INDEX_NONE || CodeIndex1 == INDEX_NONE ? INDEX_NONE : AddCodeChunk(MCT_Float3, *SampleCode, *GetParameterCode(CodeIndex0), *GetParameterCode(CodeIndex1));
+	}
+	else if (UnpackType == EVirtualTextureUnpackType::NormalBC5BC1)
+	{
+		FString	SampleCode(TEXT("VirtualTextureUnpackNormalBC5BC1(%s, %s)"));
+		return CodeIndex0 == INDEX_NONE || CodeIndex1 == INDEX_NONE ? INDEX_NONE : AddCodeChunk(MCT_Float3, *SampleCode, *GetParameterCode(CodeIndex1), *GetParameterCode(CodeIndex2));
+	}
+	else if (UnpackType == EVirtualTextureUnpackType::HeightR16)
+	{
+		FString	SampleCode(TEXT("VirtualTextureUnpackHeightR16(%s)"));
+		return CodeIndex0 == INDEX_NONE ? INDEX_NONE : AddCodeChunk(MCT_Float, *SampleCode, *GetParameterCode(CodeIndex0));
+	}
+
+	return CodeIndex0;
+}
+
+int32 FHLSLMaterialTranslator::ExternalTexture(const FGuid& ExternalTextureGuid)
+{
+	bool bOnlyInPixelShader = GetFeatureLevel() < ERHIFeatureLevel::SM4;
+
+	if (bOnlyInPixelShader && ShaderFrequency != SF_Pixel)
+	{
+		return NonPixelShaderExpressionError();
+	}
+
+	return AddUniformExpression(new FMaterialUniformExpressionExternalTexture(ExternalTextureGuid), MCT_TextureExternal, TEXT(""));
+}
+
+int32 FHLSLMaterialTranslator::ExternalTexture(UTexture* InTexture, int32& TextureReferenceIndex)
+{
+	bool bOnlyInPixelShader = GetFeatureLevel() < ERHIFeatureLevel::SM4;
+
+	if (bOnlyInPixelShader && ShaderFrequency != SF_Pixel)
+	{
+		return NonPixelShaderExpressionError();
+	}
+
+	TextureReferenceIndex = Material->GetReferencedTextures().Find(InTexture);
+	checkf(TextureReferenceIndex != INDEX_NONE, TEXT("Material expression called Compiler->ExternalTexture() without implementing UMaterialExpression::GetReferencedTexture properly"));
+
+	return AddUniformExpression(new FMaterialUniformExpressionExternalTexture(TextureReferenceIndex), MCT_TextureExternal, TEXT(""));
+}
+
+int32 FHLSLMaterialTranslator::ExternalTextureParameter(FName ParameterName, UTexture* DefaultValue, int32& TextureReferenceIndex)
+{
+	bool bOnlyInPixelShader = GetFeatureLevel() < ERHIFeatureLevel::SM4;
+
+	if (bOnlyInPixelShader && ShaderFrequency != SF_Pixel)
+	{
+		return NonPixelShaderExpressionError();
+	}
+
+	TextureReferenceIndex = Material->GetReferencedTextures().Find(DefaultValue);
+	checkf(TextureReferenceIndex != INDEX_NONE, TEXT("Material expression called Compiler->ExternalTextureParameter() without implementing UMaterialExpression::GetReferencedTexture properly"));
+	return AddUniformExpression(new FMaterialUniformExpressionExternalTextureParameter(ParameterName, TextureReferenceIndex), MCT_TextureExternal, TEXT(""));
+}
+
+int32 FHLSLMaterialTranslator::ExternalTextureCoordinateScaleRotation(int32 TextureReferenceIndex, TOptional<FName> ParameterName)
+{
+	return AddUniformExpression(new FMaterialUniformExpressionExternalTextureCoordinateScaleRotation(TextureReferenceIndex, ParameterName), MCT_Float4, TEXT(""));
+}
+int32 FHLSLMaterialTranslator::ExternalTextureCoordinateScaleRotation(const FGuid& ExternalTextureGuid)
+{
+	return AddUniformExpression(new FMaterialUniformExpressionExternalTextureCoordinateScaleRotation(ExternalTextureGuid), MCT_Float4, TEXT(""));
+}
+int32 FHLSLMaterialTranslator::ExternalTextureCoordinateOffset(int32 TextureReferenceIndex, TOptional<FName> ParameterName)
+{
+	return AddUniformExpression(new FMaterialUniformExpressionExternalTextureCoordinateOffset(TextureReferenceIndex, ParameterName), MCT_Float4, TEXT(""));
+}
+int32 FHLSLMaterialTranslator::ExternalTextureCoordinateOffset(const FGuid& ExternalTextureGuid)
+{
+	return AddUniformExpression(new FMaterialUniformExpressionExternalTextureCoordinateOffset(ExternalTextureGuid), MCT_Float4, TEXT(""));
+}
+
+UObject* FHLSLMaterialTranslator::GetReferencedTexture(int32 Index)
+{
+	return Material->GetReferencedTextures()[Index];
+}
+
+int32 FHLSLMaterialTranslator::StaticBool(bool bValue)
+{
+	return AddInlinedCodeChunk(MCT_StaticBool,(bValue ? TEXT("true") : TEXT("false")));
+}
+
+int32 FHLSLMaterialTranslator::StaticBoolParameter(FName ParameterName,bool bDefaultValue)
+{
+	// Look up the value we are compiling with for this static parameter.
+	bool bValue = bDefaultValue;
+
+	FMaterialParameterInfo ParameterInfo = GetParameterAssociationInfo();
+	ParameterInfo.Name = ParameterName;
+
+	for (const FStaticSwitchParameter& Parameter : StaticParameters.StaticSwitchParameters)
+	{
+		if (Parameter.ParameterInfo == ParameterInfo)
+		{
+			bValue = Parameter.Value;
+			break;
+		}
+	}
+
+	return StaticBool(bValue);
+}
+	
+int32 FHLSLMaterialTranslator::StaticComponentMask(int32 Vector,FName ParameterName,bool bDefaultR,bool bDefaultG,bool bDefaultB,bool bDefaultA)
+{
+	// Look up the value we are compiling with for this static parameter.
+	bool bValueR = bDefaultR;
+	bool bValueG = bDefaultG;
+	bool bValueB = bDefaultB;
+	bool bValueA = bDefaultA;
+
+	FMaterialParameterInfo ParameterInfo = GetParameterAssociationInfo();
+	ParameterInfo.Name = ParameterName;
+
+	for (const FStaticComponentMaskParameter& Parameter : StaticParameters.StaticComponentMaskParameters)
+	{
+		if (Parameter.ParameterInfo == ParameterInfo)
+		{
+			bValueR = Parameter.R;
+			bValueG = Parameter.G;
+			bValueB = Parameter.B;
+			bValueA = Parameter.A;
+			break;
+		}
+	}
+
+	return ComponentMask(Vector,bValueR,bValueG,bValueB,bValueA);
+}
+
+const FMaterialLayersFunctions* FHLSLMaterialTranslator::StaticMaterialLayersParameter(FName ParameterName)
+{
+	FMaterialParameterInfo ParameterInfo = GetParameterAssociationInfo();
+	ParameterInfo.Name = ParameterName;
+
+	for (const FStaticMaterialLayersParameter& Parameter : StaticParameters.MaterialLayersParameters)
+	{
+		if(Parameter.ParameterInfo == ParameterInfo)
+		{
+			return &Parameter.Value;
+		}
+	}
+
+	return nullptr;
+}
+
+bool FHLSLMaterialTranslator::GetStaticBoolValue(int32 BoolIndex, bool& bSucceeded)
+{
+	bSucceeded = true;
+	if (BoolIndex == INDEX_NONE)
+	{
+		bSucceeded = false;
+		return false;
+	}
+
+	if (GetParameterType(BoolIndex) != MCT_StaticBool)
+	{
+		Errorf(TEXT("Failed to cast %s input to static bool type"), DescribeType(GetParameterType(BoolIndex)));
+		bSucceeded = false;
+		return false;
+	}
+
+	if (GetParameterCode(BoolIndex).Contains(TEXT("true")))
+	{
+		return true;
+	}
+	return false;
+}
+
+int32 FHLSLMaterialTranslator::StaticTerrainLayerWeight(FName ParameterName,int32 Default)
+{
+	if (GetFeatureLevel() <= ERHIFeatureLevel::ES3_1 && ShaderFrequency != SF_Pixel)
+	{
+		return Errorf(TEXT("Landscape layer weights are only available in the pixel shader."));
+	}
+		
+	// Look up the weight-map index for this static parameter.
+	int32 WeightmapIndex = INDEX_NONE;
+	bool bFoundParameter = false;
+	bool bAtLeastOneWeightBasedBlend = false;
+
+	FMaterialParameterInfo ParameterInfo = GetParameterAssociationInfo();
+	ParameterInfo.Name = ParameterName;
+
+	int32 NumActiveTerrainLayerWeightParameters = 0;
+	for(int32 ParameterIndex = 0;ParameterIndex < StaticParameters.TerrainLayerWeightParameters.Num(); ++ParameterIndex)
+	{
+		const FStaticTerrainLayerWeightParameter& Parameter = StaticParameters.TerrainLayerWeightParameters[ParameterIndex];
+		if (Parameter.WeightmapIndex != INDEX_NONE)
+		{
+			NumActiveTerrainLayerWeightParameters++;
+		}
+		if(Parameter.ParameterInfo == ParameterInfo)
+		{
+			WeightmapIndex = Parameter.WeightmapIndex;
+			bFoundParameter = true;
+		}
+		if (Parameter.bWeightBasedBlend)
+		{
+			bAtLeastOneWeightBasedBlend = true;
+		}
+	}
+
+	if(!bFoundParameter)
+	{
+		return Default;
+	}
+	else if(WeightmapIndex == INDEX_NONE)
+	{
+		return INDEX_NONE;
+	}
+	else
+	{			
+		int32 WeightmapCode;
+		if (GetFeatureLevel() <= ERHIFeatureLevel::ES3_1 && NumActiveTerrainLayerWeightParameters <= 3 && bAtLeastOneWeightBasedBlend)
+		{
+			// Mobile can pack 3 layers into the normal map texture B and A channels, implying the 3rd using weight based blending
+			// Layer texture is sampled into Parameters.LayerWeights in LandscapeVertexFactory.ush
+			WeightmapCode = AddInlinedCodeChunk(MCT_Float4, TEXT("Parameters.LayerWeights"));
+		}
+		else
+		{
+			// Otherwise we sample normally
+			const EMaterialSamplerType SamplerType = SAMPLERTYPE_Masks;
+			FString WeightmapName = FString::Printf(TEXT("Weightmap%d"),WeightmapIndex);
+			int32 TextureReferenceIndex = INDEX_NONE;
+			int32 TextureCodeIndex = TextureParameter(FName(*WeightmapName), GEngine->WeightMapPlaceholderTexture, TextureReferenceIndex, SamplerType);
+			WeightmapCode = TextureSample(TextureCodeIndex, TextureCoordinate(3, false, false), SamplerType);
+		}
+
+		FString LayerMaskName = FString::Printf(TEXT("LayerMask_%s"),*ParameterName.ToString());
+		return Dot(WeightmapCode,VectorParameter(FName(*LayerMaskName), FLinearColor(1.f,0.f,0.f,0.f)));
+	}
+}
+
+int32 FHLSLMaterialTranslator::VertexColor()
+{
+	bUsesVertexColor |= (ShaderFrequency != SF_Vertex);
+	return AddInlinedCodeChunk(MCT_Float4,TEXT("Parameters.VertexColor"));
+}
+
+int32 FHLSLMaterialTranslator::PreSkinnedPosition()
+{
+	if (ShaderFrequency != SF_Vertex)
+	{
+		return Errorf(TEXT("Pre-skinned position is only available in the vertex shader, pass through custom interpolators if needed."));
+	}
+
+	return AddInlinedCodeChunk(MCT_Float3,TEXT("Parameters.PreSkinnedPosition"));
+}
+
+int32 FHLSLMaterialTranslator::PreSkinnedNormal()
+{
+	if (ShaderFrequency != SF_Vertex)
+	{
+		return Errorf(TEXT("Pre-skinned normal is only available in the vertex shader, pass through custom interpolators if needed."));
+	}
+
+	return AddInlinedCodeChunk(MCT_Float3,TEXT("Parameters.PreSkinnedNormal"));
+}
+
+int32 FHLSLMaterialTranslator::VertexInterpolator(uint32 InterpolatorIndex)
+{
+	if (ShaderFrequency != SF_Pixel)
+	{
+		return Errorf(TEXT("Custom interpolator outputs only available in pixel shaders."));
+	}
+
+	UMaterialExpressionVertexInterpolator** InterpolatorPtr = CustomVertexInterpolators.FindByPredicate([InterpolatorIndex](const UMaterialExpressionVertexInterpolator* Item) { return Item && Item->InterpolatorIndex == InterpolatorIndex; });
+	if (InterpolatorPtr == nullptr)
+	{
+		return Errorf(TEXT("Invalid custom interpolator index."));
+	}
+
+	UMaterialExpressionVertexInterpolator* Interpolator = *InterpolatorPtr;
+	check(Interpolator->InterpolatorIndex == InterpolatorIndex);
+	check(Interpolator->InterpolatedType & MCT_Float);
+
+	// Assign interpolator offset and accumulate size
+	int32 InterpolatorSize = 0;
+	switch (Interpolator->InterpolatedType)
+	{
+	case MCT_Float4:	InterpolatorSize = 4; break;
+	case MCT_Float3:	InterpolatorSize = 3; break;
+	case MCT_Float2:	InterpolatorSize = 2; break;
+	default:			InterpolatorSize = 1;
+	};
+
+	if (Interpolator->InterpolatorOffset == INDEX_NONE)
+	{
+		Interpolator->InterpolatorOffset = CurrentCustomVertexInterpolatorOffset;
+		CurrentCustomVertexInterpolatorOffset += InterpolatorSize;
+	}
+	check(CurrentCustomVertexInterpolatorOffset != INDEX_NONE && Interpolator->InterpolatorOffset < CurrentCustomVertexInterpolatorOffset);
+
+	// Copy interpolated data from pixel parameters to local
+	const EMaterialValueType Type = Interpolator->InterpolatedType == MCT_Float ? MCT_Float1 : Interpolator->InterpolatedType;
+	const TCHAR* TypeName = HLSLTypeString(Type);
+	const TCHAR* Swizzle[2] = { TEXT("x"), TEXT("y") };
+	const int32 Offset = Interpolator->InterpolatorOffset;
+	
+	// Note: We reference the UV define directly to avoid having to pre-accumulate UV counts before property translation
+	FString GetValueCode = FString::Printf(TEXT("%s(Parameters.TexCoords[VERTEX_INTERPOLATOR_%i_TEXCOORDS_X].%s"), TypeName, InterpolatorIndex, Swizzle[Offset%2]);
+
+	if (Type >= MCT_Float2)
+	{
+		GetValueCode += FString::Printf(TEXT(", Parameters.TexCoords[VERTEX_INTERPOLATOR_%i_TEXCOORDS_Y].%s"), InterpolatorIndex, Swizzle[(Offset+1)%2]);
+
+		if (Type >= MCT_Float3)
+		{
+			GetValueCode += FString::Printf(TEXT(", Parameters.TexCoords[VERTEX_INTERPOLATOR_%i_TEXCOORDS_Z].%s"), InterpolatorIndex, Swizzle[(Offset+2)%2]);
+
+			if (Type == MCT_Float4)
+			{
+				GetValueCode += FString::Printf(TEXT(", Parameters.TexCoords[VERTEX_INTERPOLATOR_%i_TEXCOORDS_W].%s"), InterpolatorIndex, Swizzle[(Offset+3)%2]);
+			}
+		}
+	}
+
+	GetValueCode.Append(TEXT(")"));
+
+	int32 RetCode = AddCodeChunk(Type, *GetValueCode);
+	return RetCode;
+}
+
+int32 FHLSLMaterialTranslator::Add(int32 A,int32 B)
+{
+	if(A == INDEX_NONE || B == INDEX_NONE)
+	{
+		return INDEX_NONE;
+	}
+
+	const uint64 Hash = CityHash128to64({ GetParameterHash(A), GetParameterHash(B) });
+	if(GetParameterUniformExpression(A) && GetParameterUniformExpression(B))
+	{
+		return AddUniformExpressionWithHash(Hash, new FMaterialUniformExpressionFoldedMath(GetParameterUniformExpression(A),GetParameterUniformExpression(B),FMO_Add),GetArithmeticResultType(A,B),TEXT("(%s + %s)"),*GetParameterCode(A),*GetParameterCode(B));
+	}
+	else
+	{
+		return AddCodeChunkWithHash(Hash, GetArithmeticResultType(A,B),TEXT("(%s + %s)"),*GetParameterCode(A),*GetParameterCode(B));
+	}
+}
+
+int32 FHLSLMaterialTranslator::Sub(int32 A,int32 B)
+{
+	if(A == INDEX_NONE || B == INDEX_NONE)
+	{
+		return INDEX_NONE;
+	}
+
+	const uint64 Hash = CityHash128to64({ GetParameterHash(A), GetParameterHash(B) });
+	if(GetParameterUniformExpression(A) && GetParameterUniformExpression(B))
+	{
+		return AddUniformExpressionWithHash(Hash, new FMaterialUniformExpressionFoldedMath(GetParameterUniformExpression(A),GetParameterUniformExpression(B),FMO_Sub),GetArithmeticResultType(A,B),TEXT("(%s - %s)"),*GetParameterCode(A),*GetParameterCode(B));
+	}
+	else
+	{
+		return AddCodeChunkWithHash(Hash, GetArithmeticResultType(A,B),TEXT("(%s - %s)"),*GetParameterCode(A),*GetParameterCode(B));
+	}
+}
+
+int32 FHLSLMaterialTranslator::Mul(int32 A,int32 B)
+{
+	if(A == INDEX_NONE || B == INDEX_NONE)
+	{
+		return INDEX_NONE;
+	}
+
+	const uint64 Hash = CityHash128to64({ GetParameterHash(A), GetParameterHash(B) });
+	if(GetParameterUniformExpression(A) && GetParameterUniformExpression(B))
+	{
+		return AddUniformExpressionWithHash(Hash, new FMaterialUniformExpressionFoldedMath(GetParameterUniformExpression(A),GetParameterUniformExpression(B),FMO_Mul),GetArithmeticResultType(A,B),TEXT("(%s * %s)"),*GetParameterCode(A),*GetParameterCode(B));
+	}
+	else
+	{
+		return AddCodeChunkWithHash(Hash, GetArithmeticResultType(A,B),TEXT("(%s * %s)"),*GetParameterCode(A),*GetParameterCode(B));
+	}
+}
+
+int32 FHLSLMaterialTranslator::Div(int32 A,int32 B)
+{
+	if(A == INDEX_NONE || B == INDEX_NONE)
+	{
+		return INDEX_NONE;
+	}
+
+	const uint64 Hash = CityHash128to64({ GetParameterHash(A), GetParameterHash(B) });
+	if(GetParameterUniformExpression(A) && GetParameterUniformExpression(B))
+	{
+		return AddUniformExpressionWithHash(Hash, new FMaterialUniformExpressionFoldedMath(GetParameterUniformExpression(A),GetParameterUniformExpression(B),FMO_Div),GetArithmeticResultType(A,B),TEXT("(%s / %s)"),*GetParameterCode(A),*GetParameterCode(B));
+	}
+	else
+	{
+		return AddCodeChunkWithHash(Hash, GetArithmeticResultType(A,B),TEXT("(%s / %s)"),*GetParameterCode(A),*GetParameterCode(B));
+	}
+}
+
+int32 FHLSLMaterialTranslator::Dot(int32 A,int32 B)
+{
+	if(A == INDEX_NONE || B == INDEX_NONE)
+	{
+		return INDEX_NONE;
+	}
+
+	FMaterialUniformExpression* ExpressionA = GetParameterUniformExpression(A);
+	FMaterialUniformExpression* ExpressionB = GetParameterUniformExpression(B);
+
+	EMaterialValueType TypeA = GetParameterType(A);
+	EMaterialValueType TypeB = GetParameterType(B);
+	if(ExpressionA && ExpressionB)
+	{
+		if (TypeA == MCT_Float && TypeB == MCT_Float)
+		{
+			return AddUniformExpression(new FMaterialUniformExpressionFoldedMath(ExpressionA,ExpressionB,FMO_Mul),MCT_Float,TEXT("(%s * %s)"),*GetParameterCode(A),*GetParameterCode(B));
+		}
+		else
+		{
+			if (TypeA == TypeB)
+			{
+				return AddUniformExpression(new FMaterialUniformExpressionFoldedMath(ExpressionA,ExpressionB,FMO_Dot,TypeA),MCT_Float,TEXT("dot(%s,%s)"),*GetParameterCode(A),*GetParameterCode(B));
+			}
+			else
+			{
+				// Promote scalar (or truncate the bigger type)
+				if (TypeA == MCT_Float || (TypeB != MCT_Float && GetNumComponents(TypeA) > GetNumComponents(TypeB)))
+				{
+					return AddUniformExpression(new FMaterialUniformExpressionFoldedMath(ExpressionA,ExpressionB,FMO_Dot,TypeB),MCT_Float,TEXT("dot(%s,%s)"),*CoerceParameter(A, TypeB),*GetParameterCode(B));
+				}
+				else
+				{
+					return AddUniformExpression(new FMaterialUniformExpressionFoldedMath(ExpressionA,ExpressionB,FMO_Dot,TypeA),MCT_Float,TEXT("dot(%s,%s)"),*GetParameterCode(A),*CoerceParameter(B, TypeA));
+				}
+			}
+		}
+	}
+	else
+	{
+		// Promote scalar (or truncate the bigger type)
+		if (TypeA == MCT_Float || (TypeB != MCT_Float && GetNumComponents(TypeA) > GetNumComponents(TypeB)))
+		{
+			return AddCodeChunk(MCT_Float,TEXT("dot(%s, %s)"), *CoerceParameter(A, TypeB), *GetParameterCode(B));
+		}
+		else
+		{
+			return AddCodeChunk(MCT_Float,TEXT("dot(%s, %s)"), *GetParameterCode(A), *CoerceParameter(B, TypeA));
+		}
+	}
+}
+
+int32 FHLSLMaterialTranslator::Cross(int32 A,int32 B)
+{
+	if(A == INDEX_NONE || B == INDEX_NONE)
+	{
+		return INDEX_NONE;
+	}
+
+	if(GetParameterUniformExpression(A) && GetParameterUniformExpression(B))
+	{		
+		EMaterialValueType ResultType = GetArithmeticResultType(A,B);
+		if (ResultType == MCT_Float2 || (ResultType & MCT_Float) == 0)
+		{
+			return Errorf(TEXT("Cross product requires 3-component vector input."));
+		}
+		return AddUniformExpression(new FMaterialUniformExpressionFoldedMath(GetParameterUniformExpression(A),GetParameterUniformExpression(B),FMO_Cross,ResultType),MCT_Float3,TEXT("cross(%s,%s)"),*GetParameterCode(A),*GetParameterCode(B));
+	}
+	else
+	{
+		return AddCodeChunk(MCT_Float3,TEXT("cross(%s,%s)"),*CoerceParameter(A,MCT_Float3),*CoerceParameter(B,MCT_Float3));
+	}
+}
+
+int32 FHLSLMaterialTranslator::Power(int32 Base,int32 Exponent)
+{
+	if(Base == INDEX_NONE || Exponent == INDEX_NONE)
+	{
+		return INDEX_NONE;
+	}
+
+	// Clamp Pow input to >= 0 to help avoid common NaN cases
+	return AddCodeChunk(GetParameterType(Base),TEXT("PositiveClampedPow(%s,%s)"),*GetParameterCode(Base),*CoerceParameter(Exponent,MCT_Float));
+}
+	
+int32 FHLSLMaterialTranslator::Logarithm2(int32 X)
+{
+	if(X == INDEX_NONE)
+	{
+		return INDEX_NONE;
+	}
+		
+	if(GetParameterUniformExpression(X))
+	{
+		return AddUniformExpression(new FMaterialUniformExpressionLogarithm2(GetParameterUniformExpression(X)),GetParameterType(X),TEXT("log2(%s)"),*GetParameterCode(X));
+	}
+
+	return AddCodeChunk(GetParameterType(X),TEXT("log2(%s)"),*GetParameterCode(X));
+}
+
+int32 FHLSLMaterialTranslator::Logarithm10(int32 X)
+{
+	if(X == INDEX_NONE)
+	{
+		return INDEX_NONE;
+	}
+		
+	if(GetParameterUniformExpression(X))
+	{
+		return AddUniformExpression(new FMaterialUniformExpressionLogarithm10(GetParameterUniformExpression(X)),GetParameterType(X),TEXT("log10(%s)"),*GetParameterCode(X));
+	}
+
+	return AddCodeChunk(GetParameterType(X),TEXT("log10(%s)"),*GetParameterCode(X));
+}
+
+int32 FHLSLMaterialTranslator::SquareRoot(int32 X)
+{
+	if(X == INDEX_NONE)
+	{
+		return INDEX_NONE;
+	}
+
+	if(GetParameterUniformExpression(X))
+	{
+		return AddUniformExpression(new FMaterialUniformExpressionSquareRoot(GetParameterUniformExpression(X)),GetParameterType(X),TEXT("sqrt(%s)"),*GetParameterCode(X));
+	}
+	else
+	{
+		return AddCodeChunk(GetParameterType(X),TEXT("sqrt(%s)"),*GetParameterCode(X));
+	}
+}
+
+int32 FHLSLMaterialTranslator::Length(int32 X)
+{
+	if(X == INDEX_NONE)
+	{
+		return INDEX_NONE;
+	}
+
+	if(GetParameterUniformExpression(X))
+	{
+		return AddUniformExpression(new FMaterialUniformExpressionLength(GetParameterUniformExpression(X),GetParameterType(X)),MCT_Float,TEXT("length(%s)"),*GetParameterCode(X));
+	}
+	else
+	{
+		return AddCodeChunk(MCT_Float,TEXT("length(%s)"),*GetParameterCode(X));
+	}
+}
+
+int32 FHLSLMaterialTranslator::Lerp(int32 X,int32 Y,int32 A)
+{
+	if(X == INDEX_NONE || Y == INDEX_NONE || A == INDEX_NONE)
+	{
+		return INDEX_NONE;
+	}
+
+	FMaterialUniformExpression* ExpressionX = GetParameterUniformExpression(X);
+	FMaterialUniformExpression* ExpressionY = GetParameterUniformExpression(Y);
+	FMaterialUniformExpression* ExpressionA = GetParameterUniformExpression(A);
+	bool bExpressionsAreEqual = false;
+
+	// Skip over interpolations where inputs are equal
+	if (X == Y)
+	{
+		bExpressionsAreEqual = true;
+	}
+	else if (ExpressionX && ExpressionY)
+	{
+		if (ExpressionX->IsConstant() && ExpressionY->IsConstant() && (*CurrentScopeChunks)[X].Type == (*CurrentScopeChunks)[Y].Type)
+		{
+			FLinearColor ValueX, ValueY;
+			FMaterialRenderContext DummyContext(nullptr, *Material, nullptr);
+			ExpressionX->GetNumberValue(DummyContext, ValueX);
+			ExpressionY->GetNumberValue(DummyContext, ValueY);
+
+			if (ValueX == ValueY)
+			{
+				bExpressionsAreEqual = true;
+			}
+		}
+	}
+
+	if (bExpressionsAreEqual)
+	{
+		return X;
+	}
+
+	EMaterialValueType ResultType = GetArithmeticResultType(X,Y);
+	EMaterialValueType AlphaType = ResultType == (*CurrentScopeChunks)[A].Type ? ResultType : MCT_Float1;
+
+	if (AlphaType == MCT_Float1 && ExpressionA && ExpressionA->IsConstant())
+	{
+		// Skip over interpolations that explicitly select an input
+		FLinearColor Value;
+		FMaterialRenderContext DummyContext(nullptr, *Material, nullptr);
+		ExpressionA->GetNumberValue(DummyContext, Value);
+
+		if (Value.R == 0.0f)
+		{
+			return X;
+		}
+		else if (Value.R == 1.f)
+		{
+			return Y;
+		}
+	}
+
+	return AddCodeChunk(ResultType,TEXT("lerp(%s,%s,%s)"),*CoerceParameter(X,ResultType),*CoerceParameter(Y,ResultType),*CoerceParameter(A,AlphaType));
+}
+
+int32 FHLSLMaterialTranslator::Min(int32 A,int32 B)
+{
+	if(A == INDEX_NONE || B == INDEX_NONE)
+	{
+		return INDEX_NONE;
+	}
+
+	if(GetParameterUniformExpression(A) && GetParameterUniformExpression(B))
+	{
+		return AddUniformExpression(new FMaterialUniformExpressionMin(GetParameterUniformExpression(A),GetParameterUniformExpression(B)),GetParameterType(A),TEXT("min(%s,%s)"),*GetParameterCode(A),*CoerceParameter(B,GetParameterType(A)));
+	}
+	else
+	{
+		return AddCodeChunk(GetParameterType(A),TEXT("min(%s,%s)"),*GetParameterCode(A),*CoerceParameter(B,GetParameterType(A)));
+	}
+}
+
+int32 FHLSLMaterialTranslator::Max(int32 A,int32 B)
+{
+	if(A == INDEX_NONE || B == INDEX_NONE)
+	{
+		return INDEX_NONE;
+	}
+
+	if(GetParameterUniformExpression(A) && GetParameterUniformExpression(B))
+	{
+		return AddUniformExpression(new FMaterialUniformExpressionMax(GetParameterUniformExpression(A),GetParameterUniformExpression(B)),GetParameterType(A),TEXT("max(%s,%s)"),*GetParameterCode(A),*CoerceParameter(B,GetParameterType(A)));
+	}
+	else
+	{
+		return AddCodeChunk(GetParameterType(A),TEXT("max(%s,%s)"),*GetParameterCode(A),*CoerceParameter(B,GetParameterType(A)));
+	}
+}
+
+int32 FHLSLMaterialTranslator::Clamp(int32 X,int32 A,int32 B)
+{
+	if(X == INDEX_NONE || A == INDEX_NONE || B == INDEX_NONE)
+	{
+		return INDEX_NONE;
+	}
+
+	if(GetParameterUniformExpression(X) && GetParameterUniformExpression(A) && GetParameterUniformExpression(B))
+	{
+		return AddUniformExpression(new FMaterialUniformExpressionClamp(GetParameterUniformExpression(X),GetParameterUniformExpression(A),GetParameterUniformExpression(B)),GetParameterType(X),TEXT("min(max(%s,%s),%s)"),*GetParameterCode(X),*CoerceParameter(A,GetParameterType(X)),*CoerceParameter(B,GetParameterType(X)));
+	}
+	else
+	{
+		return AddCodeChunk(GetParameterType(X),TEXT("min(max(%s,%s),%s)"),*GetParameterCode(X),*CoerceParameter(A,GetParameterType(X)),*CoerceParameter(B,GetParameterType(X)));
+	}
+}
+
+int32 FHLSLMaterialTranslator::Saturate(int32 X)
+{
+	if(X == INDEX_NONE)
+	{
+		return INDEX_NONE;
+	}
+
+	if(GetParameterUniformExpression(X))
+	{
+		return AddUniformExpression(new FMaterialUniformExpressionSaturate(GetParameterUniformExpression(X)),GetParameterType(X),TEXT("saturate(%s)"),*GetParameterCode(X));
+	}
+	else
+	{
+		return AddCodeChunk(GetParameterType(X),TEXT("saturate(%s)"),*GetParameterCode(X));
+	}
+}
+
+int32 FHLSLMaterialTranslator::ComponentMask(int32 Vector,bool R,bool G,bool B,bool A)
+{
+	if(Vector == INDEX_NONE)
+	{
+		return INDEX_NONE;
+	}
+
+	EMaterialValueType	VectorType = GetParameterType(Vector);
+
+	if(	(A && (VectorType & MCT_Float) < MCT_Float4) ||
+		(B && (VectorType & MCT_Float) < MCT_Float3) ||
+		(G && (VectorType & MCT_Float) < MCT_Float2) ||
+		(R && (VectorType & MCT_Float) < MCT_Float1))
+	{
+		return Errorf(TEXT("Not enough components in (%s: %s) for component mask %u%u%u%u"),*GetParameterCode(Vector),DescribeType(GetParameterType(Vector)),R,G,B,A);
+	}
+
+	EMaterialValueType	ResultType;
+	switch((R ? 1 : 0) + (G ? 1 : 0) + (B ? 1 : 0) + (A ? 1 : 0))
+	{
+	case 1: ResultType = MCT_Float; break;
+	case 2: ResultType = MCT_Float2; break;
+	case 3: ResultType = MCT_Float3; break;
+	case 4: ResultType = MCT_Float4; break;
+	default: 
+		return Errorf(TEXT("Couldn't determine result type of component mask %u%u%u%u"),R,G,B,A);
+	};
+
+	FString MaskString = FString::Printf(TEXT("%s%s%s%s"),
+		R ? TEXT("r") : TEXT(""),
+		// If VectorType is set to MCT_Float which means it could be any of the float types, assume it is a float1
+		G ? (VectorType == MCT_Float ? TEXT("r") : TEXT("g")) : TEXT(""),
+		B ? (VectorType == MCT_Float ? TEXT("r") : TEXT("b")) : TEXT(""),
+		A ? (VectorType == MCT_Float ? TEXT("r") : TEXT("a")) : TEXT("")
+		);
+
+	auto* Expression = GetParameterUniformExpression(Vector);
+	if (Expression)
+	{
+		int8 Mask[4] = {-1, -1, -1, -1};
+		for (int32 Index = 0; Index < MaskString.Len(); ++Index)
+		{
+			Mask[Index] = SwizzleComponentToIndex(MaskString[Index]);
+		}
+		return AddUniformExpression(
+			new FMaterialUniformExpressionComponentSwizzle(Expression, Mask[0], Mask[1], Mask[2], Mask[3]),
+			ResultType,
+			TEXT("%s.%s"),
+			*GetParameterCode(Vector),
+			*MaskString
+			);
+	}
+
+	return AddInlinedCodeChunk(
+		ResultType,
+		TEXT("%s.%s"),
+		*GetParameterCode(Vector),
+		*MaskString
+		);
+}
+
+int32 FHLSLMaterialTranslator::AppendVector(int32 A,int32 B)
+{
+	if(A == INDEX_NONE || B == INDEX_NONE)
+	{
+		return INDEX_NONE;
+	}
+
+	int32 NumResultComponents = GetNumComponents(GetParameterType(A)) + GetNumComponents(GetParameterType(B));
+	EMaterialValueType	ResultType = GetVectorType(NumResultComponents);
+
+	if(GetParameterUniformExpression(A) && GetParameterUniformExpression(B))
+	{
+		return AddUniformExpression(new FMaterialUniformExpressionAppendVector(GetParameterUniformExpression(A),GetParameterUniformExpression(B),GetNumComponents(GetParameterType(A))),ResultType,TEXT("MaterialFloat%u(%s,%s)"),NumResultComponents,*GetParameterCode(A),*GetParameterCode(B));
+	}
+	else
+	{
+		return AddInlinedCodeChunk(ResultType,TEXT("MaterialFloat%u(%s,%s)"),NumResultComponents,*GetParameterCode(A),*GetParameterCode(B));
+	}
+}
+
+int32 FHLSLMaterialTranslator::TransformBase(EMaterialCommonBasis SourceCoordBasis, EMaterialCommonBasis DestCoordBasis, int32 A, int AWComponent)
+{
+	if (A == INDEX_NONE)
+	{
+		// unable to compile
+		return INDEX_NONE;
+	}
+		
+	{ // validation
+		if (ShaderFrequency != SF_Pixel && ShaderFrequency != SF_Compute && ShaderFrequency != SF_Domain && ShaderFrequency != SF_Vertex)
+		{
+			return NonPixelShaderExpressionError();
+		}
+
+		if (ShaderFrequency != SF_Pixel && ShaderFrequency != SF_Compute && ShaderFrequency != SF_Vertex)
+		{
+			if ((SourceCoordBasis == MCB_Local || DestCoordBasis == MCB_Local))
+			{
+				return Errorf(TEXT("Local space is only supported for vertex, compute or pixel shader"));
+			}
+		}
+
+		if (AWComponent != 0 && (SourceCoordBasis == MCB_Tangent || DestCoordBasis == MCB_Tangent))
+		{
+			return Errorf(TEXT("Tangent basis not available for position transformations"));
+		}
+		
+		// Construct float3(0,0,x) out of the input if it is a scalar
+		// This way artists can plug in a scalar and it will be treated as height, or a vector displacement
+		if (GetType(A) == MCT_Float1 && SourceCoordBasis == MCB_Tangent)
+		{
+			A = AppendVector(Constant2(0, 0), A);
+		}
+		else if (GetNumComponents(GetParameterType(A)) < 3)
+		{
+			return Errorf(TEXT("input must be a vector (%s: %s) or a scalar (if source is Tangent)"), *GetParameterCode(A), DescribeType(GetParameterType(A)));
+		}
+	}
+		
+	if (SourceCoordBasis == DestCoordBasis)
+	{
+		// no transformation needed
+		return A;
+	}
+		
+	FString CodeStr;
+	EMaterialCommonBasis IntermediaryBasis = MCB_World;
+
+	switch (SourceCoordBasis)
+	{
+		case MCB_Tangent:
+		{
+			check(AWComponent == 0);
+			if (DestCoordBasis == MCB_World)
+			{
+				if (ShaderFrequency == SF_Domain)
+				{
+					// domain shader uses a prescale value to preserve scaling factor on WorldTransform	when sampling a displacement map
+					CodeStr = FString(TEXT("TransformTangent<TO>World_PreScaled(Parameters, <A>.xyz)"));
+				}
+				else
+				{
+					CodeStr = TEXT("mul(<A>, <MATRIX>(Parameters.TangentToWorld))");
+				}
+			}
+			// else use MCB_World as intermediary basis
+			break;
+		}
+		case MCB_Local:
+		{
+			if (DestCoordBasis == MCB_World)
+			{
+				CodeStr = TEXT("TransformLocal<TO><PREV>World(Parameters, <A>.xyz)");
+			}
+			// else use MCB_World as intermediary basis
+			break;
+		}
+		case MCB_TranslatedWorld:
+		{
+			if (DestCoordBasis == MCB_World)
+			{
+				if (AWComponent)
+				{
+					CodeStr = TEXT("(<A>.xyz - ResolvedView.<PREV>PreViewTranslation.xyz)");
+				}
+				else
+				{
+					CodeStr = TEXT("<A>");
+				}
+			}
+			else if (DestCoordBasis == MCB_Camera)
+			{
+				CodeStr = TEXT("mul(<A>, <MATRIX>(ResolvedView.<PREV>TranslatedWorldToCameraView))");
+			}
+			else if (DestCoordBasis == MCB_View)
+			{
+				CodeStr = TEXT("mul(<A>, <MATRIX>(ResolvedView.<PREV>TranslatedWorldToView))");
+			}
+			// else use MCB_World as intermediary basis
+			break;
+		}
+		case MCB_World:
+		{
+			if (DestCoordBasis == MCB_Tangent)
+			{
+				CodeStr = TEXT("mul(<MATRIX>(Parameters.TangentToWorld), <A>)");
+			}
+			else if (DestCoordBasis == MCB_Local)
+			{
+				const EMaterialDomain Domain = (const EMaterialDomain)Material->GetMaterialDomain();
+
+				if(Domain != MD_Surface && Domain != MD_Volume)
+				{
+					// TODO: for decals we could support it
+					Errorf(TEXT("This transformation is only supported in the 'Surface' material domain."));
+					return INDEX_NONE;
+				}
+
+				// TODO: inconsistent with TransformLocal<TO>World with instancing
+				if (bCompilingPreviousFrame)
+				{
+					// uses different prefix than other Prev* names, so can't use <PREV> tag here
+					CodeStr = TEXT("mul(<A>, <MATRIX>(GetPrimitiveData(Parameters.PrimitiveId).PreviousWorldToLocal))");
+				}
+				else
+				{
+					CodeStr = TEXT("mul(<A>, <MATRIX>(GetPrimitiveData(Parameters.PrimitiveId).WorldToLocal))");
+				}
+			}
+			else if (DestCoordBasis == MCB_TranslatedWorld)
+			{
+				if (AWComponent)
+				{
+					CodeStr = TEXT("(<A>.xyz + ResolvedView.<PREV>PreViewTranslation.xyz)");
+				}
+				else
+				{
+					CodeStr = TEXT("<A>");
+				}
+			}
+			else if (DestCoordBasis == MCB_MeshParticle)
+			{
+				CodeStr = TEXT("mul(<A>, <MATRIX>(Parameters.Particle.LocalToWorld))");
+				bUsesParticleTransform = true;
+			}
+
+			// else use MCB_TranslatedWorld as intermediary basis
+			IntermediaryBasis = MCB_TranslatedWorld;
+			break;
+		}
+		case MCB_Camera:
+		{
+			if (DestCoordBasis == MCB_TranslatedWorld)
+			{
+				CodeStr = TEXT("mul(<A>, <MATRIX>(ResolvedView.<PREV>CameraViewToTranslatedWorld))");
+			}
+			// else use MCB_TranslatedWorld as intermediary basis
+			IntermediaryBasis = MCB_TranslatedWorld;
+			break;
+		}
+		case MCB_View:
+		{
+			if (DestCoordBasis == MCB_TranslatedWorld)
+			{
+				CodeStr = TEXT("mul(<A>, <MATRIX>(ResolvedView.<PREV>ViewToTranslatedWorld))");
+			}
+			// else use MCB_TranslatedWorld as intermediary basis
+			IntermediaryBasis = MCB_TranslatedWorld;
+			break;
+		}
+		case MCB_MeshParticle:
+		{
+			if (DestCoordBasis == MCB_World)
+			{
+				CodeStr = TEXT("mul(<MATRIX>(Parameters.Particle.LocalToWorld), <A>)");
+				bUsesParticleTransform = true;
+			}
+			else
+			{
+				return Errorf(TEXT("Can transform only to world space from particle space"));
+			}
+			break;
+		}
+		default:
+			check(0);
+			break;
+	}
+
+	if (CodeStr.IsEmpty())
+	{
+		// check intermediary basis so we don't have infinite recursion
+		check(IntermediaryBasis != SourceCoordBasis);
+		check(IntermediaryBasis != DestCoordBasis);
+
+		// use intermediary basis
+		const int32 IntermediaryA = TransformBase(SourceCoordBasis, IntermediaryBasis, A, AWComponent);
+
+		return TransformBase(IntermediaryBasis, DestCoordBasis, IntermediaryA, AWComponent);
+	}
+		
+	if (AWComponent != 0)
+	{
+		if (GetType(A) == MCT_Float3)
+		{
+			A = AppendVector(A, Constant(1));
+		}
+		CodeStr.ReplaceInline(TEXT("<TO>"),TEXT("PositionTo"));
+		CodeStr.ReplaceInline(TEXT("<MATRIX>"),TEXT(""));
+		CodeStr += ".xyz";
+	}
+	else
+	{
+		CodeStr.ReplaceInline(TEXT("<TO>"),TEXT("VectorTo"));
+		CodeStr.ReplaceInline(TEXT("<MATRIX>"),TEXT("(MaterialFloat3x3)"));
+	}
+		
+	if (bCompilingPreviousFrame)
+	{
+		CodeStr.ReplaceInline(TEXT("<PREV>"),TEXT("Prev"));
+	}
+	else
+	{
+		CodeStr.ReplaceInline(TEXT("<PREV>"),TEXT(""));
+	}
+		
+	CodeStr.ReplaceInline(TEXT("<A>"), *GetParameterCode(A));
+
+	if (ShaderFrequency != SF_Vertex && (DestCoordBasis == MCB_Tangent || SourceCoordBasis == MCB_Tangent))
+	{
+		bUsesTransformVector = true;
+	}
+
+	return AddCodeChunk(
+		MCT_Float3,
+		*CodeStr
+		);
+}
+
+int32 FHLSLMaterialTranslator::TransformVector(EMaterialCommonBasis SourceCoordBasis, EMaterialCommonBasis DestCoordBasis, int32 A)
+{
+	return TransformBase(SourceCoordBasis, DestCoordBasis, A, 0);
+}
+
+int32 FHLSLMaterialTranslator::TransformPosition(EMaterialCommonBasis SourceCoordBasis, EMaterialCommonBasis DestCoordBasis, int32 A)
+{
+	return TransformBase(SourceCoordBasis, DestCoordBasis, A, 1);
+}
+
+int32 FHLSLMaterialTranslator::DynamicParameter(FLinearColor& DefaultValue, uint32 ParameterIndex)
+{
+	if (ShaderFrequency != SF_Vertex && ShaderFrequency != SF_Pixel && ShaderFrequency != SF_Compute)
+	{
+		return NonVertexOrPixelShaderExpressionError();
+	}
+
+	DynamicParticleParameterMask |= (1 << ParameterIndex);
+
+	int32 Default = Constant4(DefaultValue.R, DefaultValue.G, DefaultValue.B, DefaultValue.A);
+	return AddInlinedCodeChunk(
+		MCT_Float4,
+		TEXT("GetDynamicParameter(Parameters.Particle, %s, %u)"),
+		*GetParameterCode(Default),
+		ParameterIndex
+		);
+}
+
+int32 FHLSLMaterialTranslator::LightmapUVs()
+{
+	if (ShaderFrequency != SF_Pixel && ShaderFrequency != SF_Compute)
+	{
+		return NonPixelShaderExpressionError();
+	}
+
+	if (ErrorUnlessFeatureLevelSupported(ERHIFeatureLevel::SM4) == INDEX_NONE)
+	{
+		return INDEX_NONE;
+	}
+
+	bUsesLightmapUVs = true;
+
+	int32 ResultIdx = INDEX_NONE;
+	FString CodeChunk = FString::Printf(TEXT("GetLightmapUVs(Parameters)"));
+	ResultIdx = AddCodeChunk(
+		MCT_Float2,
+		*CodeChunk
+		);
+	return ResultIdx;
+}
+
+int32 FHLSLMaterialTranslator::PrecomputedAOMask()
+{
+	if (ShaderFrequency != SF_Pixel && ShaderFrequency != SF_Compute)
+	{
+		return NonPixelShaderExpressionError();
+	}
+
+	if (ErrorUnlessFeatureLevelSupported(ERHIFeatureLevel::SM4) == INDEX_NONE)
+	{
+		return INDEX_NONE;
+	}
+
+	bUsesAOMaterialMask = true;
+
+	int32 ResultIdx = INDEX_NONE;
+	FString CodeChunk = FString::Printf(TEXT("Parameters.AOMaterialMask"));
+	ResultIdx = AddCodeChunk(
+		MCT_Float,
+		*CodeChunk
+		);
+	return ResultIdx;
+}
+
+int32 FHLSLMaterialTranslator::LightmassReplace(int32 Realtime, int32 Lightmass)
+{
+	return Realtime;
+}
+
+int32 FHLSLMaterialTranslator::GIReplace(int32 Direct, int32 StaticIndirect, int32 DynamicIndirect)
+{ 
+	if(Direct == INDEX_NONE || DynamicIndirect == INDEX_NONE)
+	{
+		return INDEX_NONE;
+	}
+
+	EMaterialValueType ResultType = GetArithmeticResultType(Direct, DynamicIndirect);
+
+	return AddCodeChunk(ResultType,TEXT("(GetGIReplaceState() ? (%s) : (%s))"), *GetParameterCode(DynamicIndirect), *GetParameterCode(Direct));
+}
+
+int32 FHLSLMaterialTranslator::ShadowReplace(int32 Default, int32 Shadow)
+{
+	if (Default == INDEX_NONE || Shadow == INDEX_NONE)
+	{
+		return INDEX_NONE;
+	}
+
+	EMaterialValueType ResultType = GetArithmeticResultType(Default, Shadow);
+	return AddCodeChunk(ResultType, TEXT("(GetShadowReplaceState() ? (%s) : (%s))"), *GetParameterCode(Shadow), *GetParameterCode(Default));
+}
+
+int32 FHLSLMaterialTranslator::RayTracingQualitySwitchReplace(int32 Normal, int32 RayTraced)
+{
+	if (Normal == INDEX_NONE || RayTraced == INDEX_NONE)
+	{
+		return INDEX_NONE;
+	}
+
+	EMaterialValueType ResultType = GetArithmeticResultType(Normal, RayTraced);
+	return AddCodeChunk(ResultType, TEXT("(GetRayTracingQualitySwitch() ? (%s) : (%s))"), *GetParameterCode(RayTraced), *GetParameterCode(Normal));
+}
+
+int32 FHLSLMaterialTranslator::MaterialProxyReplace(int32 Realtime, int32 MaterialProxy)
+{
+	return Realtime;
+}
+
+int32 FHLSLMaterialTranslator::VirtualTextureOutputReplace(int32 Default, int32 VirtualTexture)
+{
+	if (Default == INDEX_NONE || VirtualTexture == INDEX_NONE)
+	{
+		return INDEX_NONE;
+	}
+
+	EMaterialValueType ResultType = GetArithmeticResultType(Default, VirtualTexture);
+	return AddCodeChunk(ResultType, TEXT("(GetRuntimeVirtualTextureOutputSwitch() ? (%s) : (%s))"), *GetParameterCode(VirtualTexture), *GetParameterCode(Default));
+}
+
+int32 FHLSLMaterialTranslator::ObjectOrientation()
+{ 
+	return AddInlinedCodeChunk(MCT_Float3,TEXT("GetObjectOrientation(Parameters.PrimitiveId)"));
+}
+
+int32 FHLSLMaterialTranslator::RotateAboutAxis(int32 NormalizedRotationAxisAndAngleIndex, int32 PositionOnAxisIndex, int32 PositionIndex)
+{
+	if (NormalizedRotationAxisAndAngleIndex == INDEX_NONE
+		|| PositionOnAxisIndex == INDEX_NONE
+		|| PositionIndex == INDEX_NONE)
+	{
+		return INDEX_NONE;
+	}
+	else
+	{
+		return AddCodeChunk(
+			MCT_Float3,
+			TEXT("RotateAboutAxis(%s,%s,%s)"),
+			*CoerceParameter(NormalizedRotationAxisAndAngleIndex,MCT_Float4),
+			*CoerceParameter(PositionOnAxisIndex,MCT_Float3),
+			*CoerceParameter(PositionIndex,MCT_Float3)
+			);	
+	}
+}
+
+int32 FHLSLMaterialTranslator::TwoSidedSign()
+{
+	if (ShaderFrequency != SF_Pixel && ShaderFrequency != SF_Compute)
+	{
+		return NonPixelShaderExpressionError();
+	}
+	return AddInlinedCodeChunk(MCT_Float,TEXT("Parameters.TwoSidedSign"));	
+}
+
+int32 FHLSLMaterialTranslator::VertexNormal()
+{
+	if (ShaderFrequency != SF_Vertex)
+	{
+		bUsesTransformVector = true;
+	}
+	return AddInlinedCodeChunk(MCT_Float3,TEXT("Parameters.TangentToWorld[2]"));	
+}
+
+int32 FHLSLMaterialTranslator::PixelNormalWS()
+{
+	if (ShaderFrequency != SF_Pixel && ShaderFrequency != SF_Compute)
+	{
+		return NonPixelShaderExpressionError();
+	}
+	if(MaterialProperty == MP_Normal)
+	{
+		return Errorf(TEXT("Invalid node PixelNormalWS used for Normal input."));
+	}
+	if (ShaderFrequency != SF_Vertex)
+	{
+		bUsesTransformVector = true;
+	}
+	return AddInlinedCodeChunk(MCT_Float3,TEXT("Parameters.WorldNormal"));	
+}
+
+int32 FHLSLMaterialTranslator::DDX( int32 X )
+{
+	if ((Platform != SP_OPENGL_ES2_WEBGL) && // WebGL 2/GLES3.0 - DDX() function is available
+		ErrorUnlessFeatureLevelSupported(ERHIFeatureLevel::ES3_1) == INDEX_NONE)
+	{
+		return INDEX_NONE;
+	}
+
+	if (X == INDEX_NONE)
+	{
+		return INDEX_NONE;
+	}
+
+	if (ShaderFrequency == SF_Compute)
+	{
+		// running a material in a compute shader pass (e.g. when using SVOGI)
+		return AddInlinedCodeChunk(MCT_Float, TEXT("0"));	
+	}
+
+	if (ShaderFrequency != SF_Pixel)
+	{
+		return NonPixelShaderExpressionError();
+	}
+
+	return AddCodeChunk(GetParameterType(X),TEXT("DDX(%s)"),*GetParameterCode(X));
+}
+
+int32 FHLSLMaterialTranslator::DDY( int32 X )
+{
+	if ((Platform != SP_OPENGL_ES2_WEBGL) && // WebGL 2/GLES3.0 - DDY() function is available
+		ErrorUnlessFeatureLevelSupported(ERHIFeatureLevel::ES3_1) == INDEX_NONE)
+	{
+		return INDEX_NONE;
+	}
+
+	if(X == INDEX_NONE)
+	{
+		return INDEX_NONE;
+	}
+
+	if (ShaderFrequency == SF_Compute)
+	{
+		// running a material in a compute shader pass
+		return AddInlinedCodeChunk(MCT_Float, TEXT("0"));	
+	}
+	if (ShaderFrequency != SF_Pixel)
+	{
+		return NonPixelShaderExpressionError();
+	}
+
+	return AddCodeChunk(GetParameterType(X),TEXT("DDY(%s)"),*GetParameterCode(X));
+}
+
+int32 FHLSLMaterialTranslator::AntialiasedTextureMask(int32 Tex, int32 UV, float Threshold, uint8 Channel)
+{
+	if (ErrorUnlessFeatureLevelSupported(ERHIFeatureLevel::SM4) == INDEX_NONE)
+	{
+		return INDEX_NONE;
+	}
+
+	if (Tex == INDEX_NONE || UV == INDEX_NONE)
+	{
+		return INDEX_NONE;
+	}
+
+	int32 ThresholdConst = Constant(Threshold);
+	int32 ChannelConst = Constant(Channel);
+	FString TextureName = CoerceParameter(Tex, GetParameterType(Tex));
+
+	return AddCodeChunk(MCT_Float, 
+		TEXT("AntialiasedTextureMask(%s,%sSampler,%s,%s,%s)"), 
+		*GetParameterCode(Tex),
+		*TextureName,
+		*GetParameterCode(UV),
+		*GetParameterCode(ThresholdConst),
+		*GetParameterCode(ChannelConst));
+}
+
+int32 FHLSLMaterialTranslator::DepthOfFieldFunction(int32 Depth, int32 FunctionValueIndex)
+{
+	if (ShaderFrequency == SF_Hull)
+	{
+		return Errorf(TEXT("Invalid node DepthOfFieldFunction used in hull shader input!"));
+	}
+
+	if (Depth == INDEX_NONE)
+	{
+		return INDEX_NONE;
+	}
+
+	return AddCodeChunk(MCT_Float, 
+		TEXT("MaterialExpressionDepthOfFieldFunction(%s, %d)"), 
+		*GetParameterCode(Depth), FunctionValueIndex);
+}
+
+int32 FHLSLMaterialTranslator::Sobol(int32 Cell, int32 Index, int32 Seed)
+{
+	if (ErrorUnlessFeatureLevelSupported(ERHIFeatureLevel::ES3_1) == INDEX_NONE)
+	{
+		return INDEX_NONE;
+	}
+
+	AddEstimatedTextureSample(2);
+
+	return AddCodeChunk(MCT_Float2,
+		TEXT("floor(%s) + float2(SobolIndex(SobolPixel(uint2(%s)), uint(%s)) ^ uint2(%s * 0x10000) & 0xffff) / 0x10000"),
+		*GetParameterCode(Cell),
+		*GetParameterCode(Cell),
+		*GetParameterCode(Index),
+		*GetParameterCode(Seed));
+}
+
+int32 FHLSLMaterialTranslator::TemporalSobol(int32 Index, int32 Seed)
+{
+	if (ErrorUnlessFeatureLevelSupported(ERHIFeatureLevel::ES3_1) == INDEX_NONE)
+	{
+		return INDEX_NONE;
+	}
+
+	AddEstimatedTextureSample(2);
+
+	return AddCodeChunk(MCT_Float2,
+		TEXT("float2(SobolIndex(SobolPixel(uint2(Parameters.SvPosition.xy)), uint(View.StateFrameIndexMod8 + 8 * %s)) ^ uint2(%s * 0x10000) & 0xffff) / 0x10000"),
+		*GetParameterCode(Index),
+		*GetParameterCode(Seed));
+}
+
+int32 FHLSLMaterialTranslator::Noise(int32 Position, float Scale, int32 Quality, uint8 NoiseFunction, bool bTurbulence, int32 Levels, float OutputMin, float OutputMax, float LevelScale, int32 FilterWidth, bool bTiling, uint32 RepeatSize)
+{
+	// GradientTex3D uses 3D texturing, which is not available on ES2
+	if (NoiseFunction == NOISEFUNCTION_GradientTex3D)
+	{
+		if (ErrorUnlessFeatureLevelSupported(ERHIFeatureLevel::SM4) == INDEX_NONE)
+		{
+			Errorf(TEXT("3D textures are not supported for ES2"));
+			return INDEX_NONE;
+		}
+	}
+	// all others are fine for ES2 feature level
+	else if (ErrorUnlessFeatureLevelSupported(ERHIFeatureLevel::ES2) == INDEX_NONE)
+	{
+		return INDEX_NONE;
+	}
+
+	if(Position == INDEX_NONE || FilterWidth == INDEX_NONE)
+	{
+		return INDEX_NONE;
+	}
+
+	if (NoiseFunction == NOISEFUNCTION_SimplexTex ||
+		NoiseFunction == NOISEFUNCTION_GradientTex ||
+		NoiseFunction == NOISEFUNCTION_GradientTex3D)
+	{
+		AddEstimatedTextureSample();
+	}
+
+	// to limit performance problems due to values outside reasonable range
+	Levels = FMath::Clamp(Levels, 1, 10);
+
+	int32 ScaleConst = Constant(Scale);
+	int32 QualityConst = Constant(Quality);
+	int32 NoiseFunctionConst = Constant(NoiseFunction);
+	int32 TurbulenceConst = Constant(bTurbulence);
+	int32 LevelsConst = Constant(Levels);
+	int32 OutputMinConst = Constant(OutputMin);
+	int32 OutputMaxConst = Constant(OutputMax);
+	int32 LevelScaleConst = Constant(LevelScale);
+	int32 TilingConst = Constant(bTiling);
+	int32 RepeatSizeConst = Constant(RepeatSize);
+
+	return AddCodeChunk(MCT_Float, 
+		TEXT("MaterialExpressionNoise(%s,%s,%s,%s,%s,%s,%s,%s,%s,%s,%s,%s)"), 
+		*GetParameterCode(Position),
+		*GetParameterCode(ScaleConst),
+		*GetParameterCode(QualityConst),
+		*GetParameterCode(NoiseFunctionConst),
+		*GetParameterCode(TurbulenceConst),
+		*GetParameterCode(LevelsConst),
+		*GetParameterCode(OutputMinConst),
+		*GetParameterCode(OutputMaxConst),
+		*GetParameterCode(LevelScaleConst),
+		*GetParameterCode(FilterWidth),
+		*GetParameterCode(TilingConst),
+		*GetParameterCode(RepeatSizeConst));
+}
+
+int32 FHLSLMaterialTranslator::VectorNoise(int32 Position, int32 Quality, uint8 NoiseFunction, bool bTiling, uint32 TileSize)
+{
+	if (ErrorUnlessFeatureLevelSupported(ERHIFeatureLevel::ES2) == INDEX_NONE)
+	{
+		return INDEX_NONE;
+	}
+
+	if (Position == INDEX_NONE)
+	{
+		return INDEX_NONE;
+	}
+
+	int32 QualityConst = Constant(Quality);
+	int32 NoiseFunctionConst = Constant(NoiseFunction);
+	int32 TilingConst = Constant(bTiling);
+	int32 TileSizeConst = Constant(TileSize);
+
+	if (NoiseFunction == VNF_GradientALU || NoiseFunction == VNF_VoronoiALU)
+	{
+		return AddCodeChunk(MCT_Float4,
+			TEXT("MaterialExpressionVectorNoise(%s,%s,%s,%s,%s)"),
+			*GetParameterCode(Position),
+			*GetParameterCode(QualityConst),
+			*GetParameterCode(NoiseFunctionConst),
+			*GetParameterCode(TilingConst),
+			*GetParameterCode(TileSizeConst));
+	}
+	else
+	{
+		return AddCodeChunk(MCT_Float3,
+			TEXT("MaterialExpressionVectorNoise(%s,%s,%s,%s,%s).xyz"),
+			*GetParameterCode(Position),
+			*GetParameterCode(QualityConst),
+			*GetParameterCode(NoiseFunctionConst),
+			*GetParameterCode(TilingConst),
+			*GetParameterCode(TileSizeConst));
+	}
+}
+
+int32 FHLSLMaterialTranslator::BlackBody( int32 Temp )
+{
+	if( Temp == INDEX_NONE )
+	{
+		return INDEX_NONE;
+	}
+
+	return AddCodeChunk( MCT_Float3, TEXT("MaterialExpressionBlackBody(%s)"), *GetParameterCode(Temp) );
+}
+
+int32 FHLSLMaterialTranslator::DistanceToNearestSurface(int32 PositionArg)
+{
+	if (ErrorUnlessFeatureLevelSupported(ERHIFeatureLevel::SM5) == INDEX_NONE)
+	{
+		return INDEX_NONE;
+	}
+
+	if (PositionArg == INDEX_NONE)
+	{
+		return INDEX_NONE;
+	}
+
+	MaterialCompilationOutput.bUsesGlobalDistanceField = true;
+
+	return AddCodeChunk(MCT_Float, TEXT("GetDistanceToNearestSurfaceGlobal(%s)"), *GetParameterCode(PositionArg));
+}
+
+int32 FHLSLMaterialTranslator::DistanceFieldGradient(int32 PositionArg)
+{
+	if (ErrorUnlessFeatureLevelSupported(ERHIFeatureLevel::SM5) == INDEX_NONE)
+	{
+		return INDEX_NONE;
+	}
+
+	if (PositionArg == INDEX_NONE)
+	{
+		return INDEX_NONE;
+	}
+
+	MaterialCompilationOutput.bUsesGlobalDistanceField = true;
+
+	return AddCodeChunk(MCT_Float3, TEXT("GetDistanceFieldGradientGlobal(%s)"), *GetParameterCode(PositionArg));
+}
+
+int32 FHLSLMaterialTranslator::AtmosphericFogColor( int32 WorldPosition )
+{
+	if (ErrorUnlessFeatureLevelSupported(ERHIFeatureLevel::SM4) == INDEX_NONE)
+	{
+		return INDEX_NONE;
+	}
+
+	bUsesAtmosphericFog = true;
+	if( WorldPosition == INDEX_NONE )
+	{
+		return AddCodeChunk( MCT_Float4, TEXT("MaterialExpressionAtmosphericFog(Parameters, Parameters.AbsoluteWorldPosition)"));
+	}
+	else
+	{
+		return AddCodeChunk( MCT_Float4, TEXT("MaterialExpressionAtmosphericFog(Parameters, %s)"), *GetParameterCode(WorldPosition) );
+	}
+}
+
+int32 FHLSLMaterialTranslator::AtmosphericLightVector()
+{
+	bUsesAtmosphericFog = true;
+
+	return AddCodeChunk(MCT_Float3, TEXT("MaterialExpressionAtmosphericLightVector(Parameters)"));
+
+}
+
+int32 FHLSLMaterialTranslator::AtmosphericLightColor()
+{
+	bUsesAtmosphericFog = true;
+
+	return AddCodeChunk(MCT_Float3, TEXT("MaterialExpressionAtmosphericLightColor(Parameters)"));
+
+}
+
+int32 FHLSLMaterialTranslator::SkyAtmosphereLightIlluminance(int32 WorldPosition, int32 LightIndex)
+{
+	bUsesSkyAtmosphere = true;
+	return AddCodeChunk(MCT_Float3, TEXT("MaterialExpressionSkyAtmosphereLightIlluminance(Parameters, %s, %d)"), *GetParameterCode(WorldPosition), LightIndex);
+}
+
+int32 FHLSLMaterialTranslator::SkyAtmosphereLightDirection(int32 LightIndex)
+{
+	bUsesSkyAtmosphere = true;
+	return AddCodeChunk(MCT_Float3, TEXT("MaterialExpressionSkyAtmosphereLightDirection(Parameters, %d)"), LightIndex);
+}
+
+int32 FHLSLMaterialTranslator::SkyAtmosphereLightDiskLuminance(int32 LightIndex)
+{
+	bUsesSkyAtmosphere = true;
+	return AddCodeChunk(MCT_Float3, TEXT("MaterialExpressionSkyAtmosphereLightDiskLuminance(Parameters, %d)"), LightIndex);
+}
+
+int32 FHLSLMaterialTranslator::SkyAtmosphereViewLuminance()
+{
+	bUsesSkyAtmosphere = true;
+	return AddCodeChunk(MCT_Float3, TEXT("MaterialExpressionSkyAtmosphereViewLuminance(Parameters)"));
+}
+
+int32 FHLSLMaterialTranslator::SkyAtmosphereAerialPerspective(int32 WorldPosition)
+{
+	bUsesSkyAtmosphere = true;
+	return AddCodeChunk(MCT_Float4, TEXT("MaterialExpressionSkyAtmosphereAerialPerspective(Parameters, %s)"), *GetParameterCode(WorldPosition));
+}
+
+int32 FHLSLMaterialTranslator::SkyAtmosphereDistantLightScatteredLuminance()
+{
+	bUsesSkyAtmosphere = true;
+	return AddCodeChunk(MCT_Float3, TEXT("MaterialExpressionSkyAtmosphereDistantLightScatteredLuminance(Parameters)"));
+}
+
+int32 FHLSLMaterialTranslator::CustomPrimitiveData(int32 OutputIndex, EMaterialValueType Type)
+{
+	check(OutputIndex < FCustomPrimitiveData::NumCustomPrimitiveDataFloats);
+
+	const int32 NumComponents = GetNumComponents(Type);
+
+	FString HlslCode;
+			
+	// Only float2, float3 and float4 need this
+	if (NumComponents > 1)
+	{
+		HlslCode.Append(FString::Printf(TEXT("float%d("), NumComponents));
+	}
+
+	for (int i = 0; i < NumComponents; i++)
+	{
+		const int32 CurrentOutputIndex = OutputIndex + i;
+
+		// Check if we are accessing inside the array, otherwise default to 0
+		if (CurrentOutputIndex < FCustomPrimitiveData::NumCustomPrimitiveDataFloats)
+		{
+			const int32 CustomDataIndex = CurrentOutputIndex / 4;
+			const int32 ElementIndex = CurrentOutputIndex % 4; // Index x, y, z or w
+
+			HlslCode.Append(FString::Printf(TEXT("GetPrimitiveData(Parameters.PrimitiveId).CustomPrimitiveData[%d][%d]"), CustomDataIndex, ElementIndex));
+		}
+		else
+		{
+			HlslCode.Append(TEXT("0.0f"));
+		}
+
+		if (i+1 < NumComponents)
+		{
+			HlslCode.Append(", ");
+		}
+	}
+
+	// This is the matching parenthesis to the first append
+	if (NumComponents > 1)
+	{
+		HlslCode.AppendChar(')');
+	}
+
+	return AddCodeChunk(Type, TEXT("%s"), *HlslCode);
+}
+
+int32 FHLSLMaterialTranslator::ShadingModel(EMaterialShadingModel InSelectedShadingModel)
+{
+	ShadingModelsFromCompilation.AddShadingModel(InSelectedShadingModel);
+	return AddInlinedCodeChunk(MCT_ShadingModel, TEXT("%d"), InSelectedShadingModel);
+}
+
+int32 FHLSLMaterialTranslator::MapARPassthroughCameraUV(int32 UV)
+{
+	if (UV == INDEX_NONE)
+	{
+		return INDEX_NONE;
+	}
+
+	int32 UVPair0 = AddInlinedCodeChunk(MCT_Float4, TEXT("ResolvedView.XRPassthroughCameraUVs[0]"));
+	int32 UVPair1 = AddInlinedCodeChunk(MCT_Float4, TEXT("ResolvedView.XRPassthroughCameraUVs[1]"));
+
+	int32 ULerp = Lerp(UVPair0, UVPair1, ComponentMask(UV, 1, 0, 0, 0));
+	return Lerp(ComponentMask(ULerp, 1, 1, 0, 0), ComponentMask(ULerp, 0, 0, 1, 1), ComponentMask(UV, 0, 1, 0, 0));
+}
+
+int32 FHLSLMaterialTranslator::CustomExpression( class UMaterialExpressionCustom* Custom, TArray<int32>& CompiledInputs )
+{
+	int32 ResultIdx = INDEX_NONE;
+
+	FString OutputTypeString;
+	EMaterialValueType OutputType;
+	switch( Custom->OutputType )
+	{
+	case CMOT_Float2:
+		OutputType = MCT_Float2;
+		OutputTypeString = TEXT("MaterialFloat2");
+		break;
+	case CMOT_Float3:
+		OutputType = MCT_Float3;
+		OutputTypeString = TEXT("MaterialFloat3");
+		break;
+	case CMOT_Float4:
+		OutputType = MCT_Float4;
+		OutputTypeString = TEXT("MaterialFloat4");
+		break;
+	default:
+		OutputType = MCT_Float;
+		OutputTypeString = TEXT("MaterialFloat");
+		break;
+	}
+
+	// Declare implementation function
+	FString InputParamDecl;
+	check( Custom->Inputs.Num() == CompiledInputs.Num() );
+	for( int32 i = 0; i < Custom->Inputs.Num(); i++ )
+	{
+		// skip over unnamed inputs
+		if( Custom->Inputs[i].InputName.IsNone() )
+		{
+			continue;
+		}
+		InputParamDecl += TEXT(",");
+		const FString InputNameStr = Custom->Inputs[i].InputName.ToString();
+		switch(GetParameterType(CompiledInputs[i]))
+		{
+		case MCT_Float:
+		case MCT_Float1:
+			InputParamDecl += TEXT("MaterialFloat ");
+			InputParamDecl += InputNameStr;
+			break;
+		case MCT_Float2:
+			InputParamDecl += TEXT("MaterialFloat2 ");
+			InputParamDecl += InputNameStr;
+			break;
+		case MCT_Float3:
+			InputParamDecl += TEXT("MaterialFloat3 ");
+			InputParamDecl += InputNameStr;
+			break;
+		case MCT_Float4:
+			InputParamDecl += TEXT("MaterialFloat4 ");
+			InputParamDecl += InputNameStr;
+			break;
+		case MCT_Texture2D:
+			InputParamDecl += TEXT("Texture2D ");
+			InputParamDecl += InputNameStr;
+			InputParamDecl += TEXT(", SamplerState ");
+			InputParamDecl += InputNameStr;
+			InputParamDecl += TEXT("Sampler ");
+			break;
+		case MCT_TextureCube:
+			InputParamDecl += TEXT("TextureCube ");
+			InputParamDecl += InputNameStr;
+			InputParamDecl += TEXT(", SamplerState ");
+			InputParamDecl += InputNameStr;
+			InputParamDecl += TEXT("Sampler ");
+			break;
+		case MCT_TextureExternal:
+			InputParamDecl += TEXT("TextureExternal ");
+			InputParamDecl += InputNameStr;
+			InputParamDecl += TEXT(", SamplerState ");
+			InputParamDecl += InputNameStr;
+			InputParamDecl += TEXT("Sampler ");
+			break;
+		case MCT_VolumeTexture:
+			InputParamDecl += TEXT("Texture3D ");
+			InputParamDecl += InputNameStr;
+			InputParamDecl += TEXT(", SamplerState ");
+			InputParamDecl += InputNameStr;
+			InputParamDecl += TEXT("Sampler ");
+			break;
+		default:
+			return Errorf(TEXT("Bad type %s for %s input %s"),DescribeType(GetParameterType(CompiledInputs[i])), *Custom->Description, *InputNameStr);
+			break;
+		}
+	}
+	int32 CustomExpressionIndex = CustomExpressionImplementations.Num();
+	FString Code = Custom->Code;
+	if( !Code.Contains(TEXT("return")) )
+	{
+		Code = FString(TEXT("return "))+Code+TEXT(";");
+	}
+	Code.ReplaceInline(TEXT("\n"),TEXT("\r\n"), ESearchCase::CaseSensitive);
+
+	FString ParametersType = ShaderFrequency == SF_Vertex ? TEXT("Vertex") : (ShaderFrequency == SF_Domain ? TEXT("Tessellation") : TEXT("Pixel"));
+
+	FString ImplementationCode = FString::Printf(TEXT("%s CustomExpression%d(FMaterial%sParameters Parameters%s)\r\n{\r\n%s\r\n}\r\n"), *OutputTypeString, CustomExpressionIndex, *ParametersType, *InputParamDecl, *Code);
+	CustomExpressionImplementations.Add( ImplementationCode );
+
+	// Add call to implementation function
+	FString CodeChunk = FString::Printf(TEXT("CustomExpression%d(Parameters"), CustomExpressionIndex);
+	for( int32 i = 0; i < CompiledInputs.Num(); i++ )
+	{
+		// skip over unnamed inputs
+		if( Custom->Inputs[i].InputName.IsNone() )
+		{
+			continue;
+		}
+
+		FString ParamCode = GetParameterCode(CompiledInputs[i]);
+		EMaterialValueType ParamType = GetParameterType(CompiledInputs[i]);
+
+		CodeChunk += TEXT(",");
+		CodeChunk += *ParamCode;
+		if (ParamType == MCT_Texture2D || ParamType == MCT_TextureCube || ParamType == MCT_TextureExternal || ParamType == MCT_VolumeTexture)
+		{
+			CodeChunk += TEXT(",");
+			CodeChunk += *ParamCode;
+			CodeChunk += TEXT("Sampler");
+		}
+	}
+	CodeChunk += TEXT(")");
+
+	ResultIdx = AddCodeChunk(
+		OutputType,
+		*CodeChunk
+		);
+	return ResultIdx;
+}
+
+int32 FHLSLMaterialTranslator::CustomOutput(class UMaterialExpressionCustomOutput* Custom, int32 OutputIndex, int32 OutputCode)
+{
+	if (MaterialProperty != MP_MAX)
+	{
+		return Errorf(TEXT("A Custom Output node should not be attached to the %s material property"), *FMaterialAttributeDefinitionMap::GetDisplayName(MaterialProperty));
+	}
+
+	if (OutputCode == INDEX_NONE)
+	{
+		return INDEX_NONE;
+	}
+
+	EMaterialValueType OutputType = GetParameterType(OutputCode);
+	FString OutputTypeString;
+	switch (OutputType)
+	{
+		case MCT_Float:
+		case MCT_Float1:
+			OutputTypeString = TEXT("MaterialFloat");
+			break;
+		case MCT_Float2:
+			OutputTypeString += TEXT("MaterialFloat2");
+			break;
+		case MCT_Float3:
+			OutputTypeString += TEXT("MaterialFloat3");
+			break;
+		case MCT_Float4:
+			OutputTypeString += TEXT("MaterialFloat4");
+			break;
+		default:
+			return Errorf(TEXT("Bad type %s for %s"), DescribeType(GetParameterType(OutputCode)), *Custom->GetDescription());
+			break;
+	}
+
+	FString Definitions;
+	FString Body;
+
+	if ((*CurrentScopeChunks)[OutputCode].UniformExpression && !(*CurrentScopeChunks)[OutputCode].UniformExpression->IsConstant())
+	{
+		Body = GetParameterCode(OutputCode);
+	}
+	else
+	{
+		GetFixedParameterCode(OutputCode, *CurrentScopeChunks, Definitions, Body);
+	}
+
+	FString ImplementationCode = FString::Printf(TEXT("%s %s%d(FMaterial%sParameters Parameters)\r\n{\r\n%s return %s;\r\n}\r\n"), *OutputTypeString, *Custom->GetFunctionName(), OutputIndex, ShaderFrequency == SF_Vertex ? TEXT("Vertex") : TEXT("Pixel"), *Definitions, *Body);
+	CustomOutputImplementations.Add(ImplementationCode);
+
+	// return value is not used
+	return INDEX_NONE;
+}
+
+int32 FHLSLMaterialTranslator::VirtualTextureOutput()
+{
+	MaterialCompilationOutput.bHasRuntimeVirtualTextureOutput = true;
+		
+	// return value is not used
+	return INDEX_NONE;
+}
+
+#if HANDLE_CUSTOM_OUTPUTS_AS_MATERIAL_ATTRIBUTES
+/** Used to translate code for custom output attributes such as ClearCoatBottomNormal */
+void FHLSLMaterialTranslator::GenerateCustomAttributeCode(int32 OutputIndex, int32 OutputCode, EMaterialValueType OutputType, FString& DisplayName)
+{
+	check(MaterialProperty == MP_CustomOutput);
+	check(OutputIndex >= 0 && OutputCode != INDEX_NONE);
+
+	FString OutputTypeString;
+	switch (OutputType)
+	{
+		case MCT_Float:
+		case MCT_Float1:
+			OutputTypeString = TEXT("MaterialFloat");
+			break;
+		case MCT_Float2:
+			OutputTypeString += TEXT("MaterialFloat2");
+			break;
+		case MCT_Float3:
+			OutputTypeString += TEXT("MaterialFloat3");
+			break;
+		case MCT_Float4:
+			OutputTypeString += TEXT("MaterialFloat4");
+			break;
+		default:
+			check(0);
+	}
+
+	FString Definitions;
+	FString Body;
+
+	if ((*CurrentScopeChunks)[OutputCode].UniformExpression && !(*CurrentScopeChunks)[OutputCode].UniformExpression->IsConstant())
+	{
+		Body = GetParameterCode(OutputCode);
+	}
+	else
+	{
+		GetFixedParameterCode(OutputCode, *CurrentScopeChunks, Definitions, Body);
+	}
+
+	FString ImplementationCode = FString::Printf(TEXT("%s %s%d(FMaterial%sParameters Parameters)\r\n{\r\n%s return %s;\r\n}\r\n"), *OutputTypeString, *DisplayName, OutputIndex, ShaderFrequency == SF_Vertex ? TEXT("Vertex") : TEXT("Pixel"), *Definitions, *Body);
+	CustomOutputImplementations.Add(ImplementationCode);
+}
+#endif
+
+/**
+	* Adds code to return a random value shared by all geometry for any given instanced static mesh
+	*
+	* @return	Code index
+	*/
+int32 FHLSLMaterialTranslator::PerInstanceRandom()
+{
+	if (ShaderFrequency != SF_Pixel && ShaderFrequency != SF_Vertex)
+	{
+		return NonVertexOrPixelShaderExpressionError();
+	}
+	else
+	{
+		return AddInlinedCodeChunk(MCT_Float, TEXT("GetPerInstanceRandom(Parameters)"));
+	}
+}
+
+/**
+	* Returns a mask that either enables or disables selection on a per-instance basis when instancing
+	*
+	* @return	Code index
+	*/
+int32 FHLSLMaterialTranslator::PerInstanceFadeAmount()
+{
+	if (ShaderFrequency != SF_Pixel && ShaderFrequency != SF_Vertex)
+	{
+		return NonVertexOrPixelShaderExpressionError();
+	}
+	else
+	{
+		return AddInlinedCodeChunk(MCT_Float, TEXT("GetPerInstanceFadeAmount(Parameters)"));
+	}
+}
+
+/**
+	* Returns a float2 texture coordinate after 2x2 transform and offset applied
+	*
+	* @return	Code index
+	*/
+int32 FHLSLMaterialTranslator::RotateScaleOffsetTexCoords(int32 TexCoordCodeIndex, int32 RotationScale, int32 Offset)
+{
+	return AddCodeChunk(MCT_Float2,
+		TEXT("RotateScaleOffsetTexCoords(%s, %s, %s.xy)"),
+		*GetParameterCode(TexCoordCodeIndex),
+		*GetParameterCode(RotationScale),
+		*GetParameterCode(Offset));
+}
+
+/**
+* Handles SpeedTree vertex animation (wind, smooth LOD)
+*
+* @return	Code index
+*/
+int32 FHLSLMaterialTranslator::SpeedTree(int32 GeometryArg, int32 WindArg, int32 LODArg, float BillboardThreshold, bool bAccurateWindVelocities, bool bExtraBend, int32 ExtraBendArg)
+{ 
+	if (ErrorUnlessFeatureLevelSupported(ERHIFeatureLevel::ES3_1) == INDEX_NONE)
+	{
+		return INDEX_NONE;
+	}
+
+	if (Material && Material->IsUsedWithSkeletalMesh())
+	{
+		return Error(TEXT("SpeedTree node not currently supported for Skeletal Meshes, please disable usage flag."));
+	}
+
+	if (ShaderFrequency != SF_Vertex)
+	{
+		return NonVertexShaderExpressionError();
+	}
+	else
+	{
+		bUsesSpeedTree = true;
+
+		AllocateSlot(AllocatedUserVertexTexCoords, 2, 6);
+
+		// Only generate previous frame's computations if required and opted-in
+		const bool bEnablePreviousFrameInformation = bCompilingPreviousFrame && bAccurateWindVelocities;
+		return AddCodeChunk(MCT_Float3, TEXT("GetSpeedTreeVertexOffset(Parameters, %s, %s, %s, %g, %s, %s, %s)"), *GetParameterCode(GeometryArg), *GetParameterCode(WindArg), *GetParameterCode(LODArg), BillboardThreshold, bEnablePreviousFrameInformation ? TEXT("true") : TEXT("false"), bExtraBend ? TEXT("true") : TEXT("false"), *GetParameterCode(ExtraBendArg, TEXT("float3(0,0,0)")));
+	}
+}
+
+/**
+	* Adds code for texture coordinate offset to localize large UV
+	*
+	* @return	Code index
+	*/
+int32 FHLSLMaterialTranslator::TextureCoordinateOffset()
+{
+	if (FeatureLevel < ERHIFeatureLevel::SM4 && ShaderFrequency == SF_Vertex)
+	{
+		return AddInlinedCodeChunk(MCT_Float2, TEXT("Parameters.TexCoordOffset"));
+	}
+	else
+	{
+		return Constant(0.f);
+	}
+}
+
+/**Experimental access to the EyeAdaptation RT for Post Process materials. Can be one frame behind depending on the value of BlendableLocation. */
+int32 FHLSLMaterialTranslator::EyeAdaptation()
+{
+	if (ErrorUnlessFeatureLevelSupported(ERHIFeatureLevel::SM5) == INDEX_NONE)
+	{
+		return INDEX_NONE;
+	}
+
+	if( ShaderFrequency != SF_Pixel )
+	{
+		NonPixelShaderExpressionError();
+	}
+
+	MaterialCompilationOutput.bUsesEyeAdaptation = true;
+
+	return AddInlinedCodeChunk(MCT_Float, TEXT("EyeAdaptationLookup()"));
+}
+
+// to only have one piece of code dealing with error handling if the Primitive constant buffer is not used.
+// @param Name e.g. TEXT("ObjectWorldPositionAndRadius.w")
+int32 FHLSLMaterialTranslator::GetPrimitiveProperty(EMaterialValueType Type, const TCHAR* ExpressionName, const TCHAR* HLSLName)
+{
+	const EMaterialDomain Domain = (const EMaterialDomain)Material->GetMaterialDomain();
+
+	if(Domain != MD_Surface && Domain != MD_Volume)
+	{
+		Errorf(TEXT("The material expression '%s' is only supported in the 'Surface' or 'Volume' material domain."), ExpressionName);
+		return INDEX_NONE;
+	}
+
+	return AddInlinedCodeChunk(Type, TEXT("GetPrimitiveData(Parameters.PrimitiveId).%s"), HLSLName);
+}
+
+// The compiler can run in a different state and this affects caching of sub expression, Expressions are different (e.g. View.PrevWorldViewOrigin) when using previous frame's values
+bool FHLSLMaterialTranslator::IsCurrentlyCompilingForPreviousFrame() const
+{
+	return bCompilingPreviousFrame;
+}
+
+bool FHLSLMaterialTranslator::IsDevelopmentFeatureEnabled(const FName& FeatureName) const
+{
+	if (FeatureName == NAME_SelectionColor)
+	{
+		// This is an editor-only feature (see FDefaultMaterialInstance::GetVectorValue).
+
+		// Determine if we're sure the editor will never run using the target shader platform.
+		// The list below may not be comprehensive enough, but it definitely includes platforms which won't use selection color for sure.
+		const bool bEditorMayUseTargetShaderPlatform = IsPCPlatform(Platform);
+		static const auto CVar = IConsoleManager::Get().FindTConsoleVariableDataInt(TEXT("r.CompileShadersForDevelopment"));
+		const bool bCompileShadersForDevelopment = (CVar && CVar->GetValueOnAnyThread() != 0);
+
+		return
+			// Does the material explicitly forbid development features?
+			Material->GetAllowDevelopmentShaderCompile()
+			// Can the editor run using the current shader platform?
+			&& bEditorMayUseTargetShaderPlatform
+			// Are shader development features globally disabled?
+			&& bCompileShadersForDevelopment;
+	}
+
+	return true;
+}
+
+#endif // WITH_EDITORONLY_DATA
