Applied previous changes in HLSLMaterialTranslator.h into HLSLMaterialTranslator.cpp.

#rnx

diff --git a/Engine/Source/Runtime/Engine/Private/Materials/HLSLMaterialTranslator.cpp b/Engine/Source/Runtime/Engine/Private/Materials/HLSLMaterialTranslator.cpp
--- a/Engine/Source/Runtime/Engine/Private/Materials/HLSLMaterialTranslator.cpp
+++ b/Engine/Source/Runtime/Engine/Private/Materials/HLSLMaterialTranslator.cpp
@@ -772,12 +772,6 @@ bool FHLSLMaterialTranslator::Translate()
 			CompileCustomOutputs(CustomOutputExpressions, SeenCustomOutputExpressionsClasses, false);
 		}
 
-		// Output the implementation for any custom expressions we will call below.
-		for (int32 ExpressionIndex = 0; ExpressionIndex < CustomExpressionImplementations.Num(); ExpressionIndex++)
-		{
-			ResourcesString += CustomExpressionImplementations[ExpressionIndex] + "\r\n\r\n";
-		}
-
 		// No more calls to non-vertex shader CompilePropertyAndSetMaterialProperty beyond this point
 		const uint32 SavedNumUserTexCoords = GetNumUserTexCoords();
 
@@ -792,6 +786,12 @@ bool FHLSLMaterialTranslator::Translate()
 			}
 		}
 
+		// Output the implementation for any custom expressions we will call below.
+		for (int32 ExpressionIndex = 0; ExpressionIndex < CustomExpressionImplementations.Num(); ExpressionIndex++)
+		{
+			ResourcesString += CustomExpressionImplementations[ExpressionIndex] + "\r\n\r\n";
+		}
+
 		// Translation is designed to have a code chunk generation phase followed by several passes that only has readonly access to the code chunks.
 		// At this point we mark the code chunk generation complete.
 		bAllowCodeChunkGeneration = false;
@@ -1021,7 +1021,7 @@ bool FHLSLMaterialTranslator::Translate()
 			ResourcesString += CustomOutputImplementations[ExpressionIndex] + "\r\n\r\n";
 		}
 
-		LoadShaderSourceFileChecked(TEXT("/Engine/Private/MaterialTemplate.ush"), MaterialTemplate);
+		LoadShaderSourceFileChecked(TEXT("/Engine/Private/MaterialTemplate.ush"), GetShaderPlatform(), MaterialTemplate);
 
 		// Find the string index of the '#line' statement in MaterialTemplate.usf
 		const int32 LineIndex = MaterialTemplate.Find(TEXT("#line"), ESearchCase::CaseSensitive);
@@ -1156,9 +1156,10 @@ void FHLSLMaterialTranslator::GetMaterialEnvironment(EShaderPlatform InPlatform,
 	OutEnvironment.SetDefine(TEXT("USES_DISTORTION"), Material->IsDistorted()); 
 
 	OutEnvironment.SetDefine(TEXT("MATERIAL_ENABLE_TRANSLUCENCY_FOGGING"), Material->ShouldApplyFogging());
-	OutEnvironment.SetDefine(TEXT("MATERIAL_FORCE_SKIP_AERIAL_PERSPECTIVE"), Material->IsSky());
+	OutEnvironment.SetDefine(TEXT("MATERIAL_IS_SKY"), Material->IsSky());
 	OutEnvironment.SetDefine(TEXT("MATERIAL_COMPUTE_FOG_PER_PIXEL"), Material->ComputeFogPerPixel());
 	OutEnvironment.SetDefine(TEXT("MATERIAL_FULLY_ROUGH"), bIsFullyRough || Material->IsFullyRough());
+	OutEnvironment.SetDefine(TEXT("MATERIAL_TWO_SIDED"), Material->IsTwoSided());
 
 	// Count the number of VTStacks (each stack will allocate a feedback slot)
 	OutEnvironment.SetDefine(TEXT("NUM_VIRTUALTEXTURE_SAMPLES"), VTStacks.Num());
@@ -1718,6 +1719,7 @@ const TCHAR* FHLSLMaterialTranslator::DescribeType(EMaterialValueType Type) cons
 	case MCT_Float:					return TEXT("float");
 	case MCT_Texture2D:				return TEXT("texture2D");
 	case MCT_TextureCube:			return TEXT("textureCube");
+	case MCT_Texture2DArray:		return TEXT("texture2DArray");
 	case MCT_VolumeTexture:			return TEXT("volumeTexture");
 	case MCT_StaticBool:			return TEXT("static bool");
 	case MCT_MaterialAttributes:	return TEXT("MaterialAttributes");
@@ -1741,6 +1743,7 @@ const TCHAR* FHLSLMaterialTranslator::HLSLTypeString(EMaterialValueType Type) co
 	case MCT_Float:					return TEXT("MaterialFloat");
 	case MCT_Texture2D:				return TEXT("texture2D");
 	case MCT_TextureCube:			return TEXT("textureCube");
+	case MCT_Texture2DArray:		return TEXT("texture2DArray");
 	case MCT_VolumeTexture:			return TEXT("volumeTexture");
 	case MCT_StaticBool:			return TEXT("static bool");
 	case MCT_MaterialAttributes:	return TEXT("MaterialAttributes");
@@ -2176,6 +2179,10 @@ int32 FHLSLMaterialTranslator::AccessUniformExpression(int32 Index)
 			TextureInputIndex = MaterialCompilationOutput.UniformExpressionSet.UniformCubeTextureExpressions.AddUnique(TextureUniformExpression);
 			BaseName = TEXT("TextureCube");
 			break;
+		case MCT_Texture2DArray:
+			TextureInputIndex = MaterialCompilationOutput.UniformExpressionSet.Uniform2DArrayTextureExpressions.AddUnique(TextureUniformExpression);
+			BaseName = TEXT("Texture2DArray");
+			break;
 		case MCT_VolumeTexture:
 			TextureInputIndex = MaterialCompilationOutput.UniformExpressionSet.UniformVolumeTextureExpressions.AddUnique(TextureUniformExpression);
 			BaseName = TEXT("VolumeTexture");
@@ -2864,6 +2871,7 @@ int32 FHLSLMaterialTranslator::ViewProperty(EMaterialExposedViewProperty Propert
 		{MEVP_TemporalSampleOffset, MCT_Float2, TEXT("View.TemporalAAParams.zw"), nullptr},
 		{MEVP_RuntimeVirtualTextureOutputLevel, MCT_Float1, TEXT("View.VirtualTextureParams.x"), nullptr},
 		{MEVP_RuntimeVirtualTextureOutputDerivative, MCT_Float2, TEXT("View.VirtualTextureParams.zw"), nullptr},
+		{MEVP_PreExposure, MCT_Float1, TEXT("View.PreExposure.x"), TEXT("View.OneOverPreExposure.x")},
 	};
 	static_assert((sizeof(ViewPropertyMetaArray) / sizeof(ViewPropertyMetaArray[0])) == MEVP_MAX, "incoherency between EMaterialExposedViewProperty and ViewPropertyMetaArray");
 
@@ -3432,6 +3440,29 @@ int32 FHLSLMaterialTranslator::ParticleSubUV(int32 TextureIndex, EMaterialSample
 	return ParticleSubUV;
 }
 
+int32 FHLSLMaterialTranslator::ParticleSubUVProperty(int32 PropertyIndex)
+{
+	int32 Result = INDEX_NONE;
+	switch (PropertyIndex)
+	{
+	case 0:
+		Result = AddCodeChunk(MCT_Float2, TEXT("Parameters.Particle.SubUVCoords[0].xy"));
+		break;
+	case 1:
+		Result = AddCodeChunk(MCT_Float2, TEXT("Parameters.Particle.SubUVCoords[1].xy"));
+		break;
+	case 2:
+		Result = AddCodeChunk(MCT_Float, TEXT("Parameters.Particle.SubUVLerp"));
+		break;
+	default:
+		checkNoEntry();
+		break;
+	}
+
+	bUsesParticleSubUVs = true;
+	return Result;
+}
+
 int32 FHLSLMaterialTranslator::ParticleColor()
 {
 	if (ShaderFrequency != SF_Vertex && ShaderFrequency != SF_Pixel && ShaderFrequency != SF_Compute)
@@ -4031,7 +4062,7 @@ int32 FHLSLMaterialTranslator::TextureSample(
 	}
 
 	// If mobile, then disabling AutomaticViewMipBias.
-	if (FeatureLevel < ERHIFeatureLevel::SM4)
+	if (FeatureLevel < ERHIFeatureLevel::SM5)
 	{
 		AutomaticViewMipBias = false;
 	}
@@ -4074,6 +4105,10 @@ int32 FHLSLMaterialTranslator::TextureSample(
 	{
 		SampleCode += TEXT("TextureCubeSample");
 	}
+	else if (TextureType == MCT_Texture2DArray)
+	{
+		SampleCode += TEXT("Texture2DArraySample");
+	}
 	else if (TextureType == MCT_VolumeTexture)
 	{
 		SampleCode += TEXT("Texture3DSample");
@@ -4091,7 +4126,7 @@ int32 FHLSLMaterialTranslator::TextureSample(
 		SampleCode += TEXT("Texture2DSample");
 	}
 		
-	EMaterialValueType UVsType = (TextureType == MCT_TextureCube || TextureType == MCT_VolumeTexture) ? MCT_Float3 : MCT_Float2;
+	EMaterialValueType UVsType = (TextureType == MCT_TextureCube || TextureType == MCT_Texture2DArray || TextureType == MCT_VolumeTexture) ? MCT_Float3 : MCT_Float2;
 	
 	if (RequiresManualViewMipBias)
 	{
@@ -4247,6 +4282,10 @@ int32 FHLSLMaterialTranslator::TextureSample(
 	{
 		TextureName = CoerceParameter(TextureIndex, MCT_TextureCube);
 	}
+	else if (TextureType == MCT_Texture2DArray)
+	{
+		TextureName = CoerceParameter(TextureIndex, MCT_Texture2DArray);
+	}
 	else if (TextureType == MCT_VolumeTexture)
 	{
 		TextureName = CoerceParameter(TextureIndex, MCT_VolumeTexture);
@@ -4567,7 +4606,7 @@ int32 FHLSLMaterialTranslator::SceneTextureLookup(int32 ViewportUV, uint32 InSce
 									SceneTextureId == PPI_SceneDepth ||
 									SceneTextureId == PPI_CustomStencil;
 
-	if (!bSupportedOnMobile	&& ErrorUnlessFeatureLevelSupported(ERHIFeatureLevel::SM4) == INDEX_NONE)
+	if (!bSupportedOnMobile	&& ErrorUnlessFeatureLevelSupported(ERHIFeatureLevel::SM5) == INDEX_NONE)
 	{
 		return INDEX_NONE;
 	}
@@ -4599,13 +4638,22 @@ int32 FHLSLMaterialTranslator::SceneTextureLookup(int32 ViewportUV, uint32 InSce
 
 	AddEstimatedTextureSample();
 
-	if (FeatureLevel >= ERHIFeatureLevel::SM4)
+	if (FeatureLevel >= ERHIFeatureLevel::SM5)
 	{
-		return AddCodeChunk(
+		int32 LookUp = AddCodeChunk(
 			MCT_Float4,
 			TEXT("SceneTextureLookup(%s, %d, %s)"),
 			*CoerceParameter(BufferUV, MCT_Float2), (int)SceneTextureId, bFiltered ? TEXT("true") : TEXT("false")
-			);
+		);
+
+		if (SceneTextureId == PPI_PostProcessInput0 && Material->GetMaterialDomain() == MD_PostProcess && Material->GetBlendableLocation() != BL_AfterTonemapping)
+		{
+			return AddInlinedCodeChunk(MCT_Float4, TEXT("(float4(View.OneOverPreExposure.xxx, 1) * %s)"), *CoerceParameter(LookUp, MCT_Float4));
+		}
+		else
+		{
+			return LookUp;
+		}
 	}
 	else // mobile
 	{
@@ -4671,7 +4719,7 @@ void FHLSLMaterialTranslator::UseSceneTextureId(ESceneTextureId SceneTextureId,
 
 		if (bRequiresSM5)
 		{
-			ErrorUnlessFeatureLevelSupported(ERHIFeatureLevel::SM4);
+			ErrorUnlessFeatureLevelSupported(ERHIFeatureLevel::SM5);
 		}
 	}
 
@@ -4737,7 +4785,7 @@ int32 FHLSLMaterialTranslator::SceneColor(int32 Offset, int32 ViewportUV, bool b
 		Errorf(TEXT("SceneColor lookups are only available when MaterialDomain = Surface."));
 	}
 
-	if (ErrorUnlessFeatureLevelSupported(ERHIFeatureLevel::SM4) == INDEX_NONE)
+	if (ErrorUnlessFeatureLevelSupported(ERHIFeatureLevel::SM5) == INDEX_NONE)
 	{
 		return INDEX_NONE;
 	}
@@ -4907,7 +4955,7 @@ int32 FHLSLMaterialTranslator::VirtualTextureUnpack(int32 CodeIndex0, int32 Code
 
 int32 FHLSLMaterialTranslator::ExternalTexture(const FGuid& ExternalTextureGuid)
 {
-	bool bOnlyInPixelShader = GetFeatureLevel() < ERHIFeatureLevel::SM4;
+	bool bOnlyInPixelShader = GetFeatureLevel() < ERHIFeatureLevel::SM5;
 
 	if (bOnlyInPixelShader && ShaderFrequency != SF_Pixel)
 	{
@@ -4919,7 +4967,7 @@ int32 FHLSLMaterialTranslator::ExternalTexture(const FGuid& ExternalTextureGuid)
 
 int32 FHLSLMaterialTranslator::ExternalTexture(UTexture* InTexture, int32& TextureReferenceIndex)
 {
-	bool bOnlyInPixelShader = GetFeatureLevel() < ERHIFeatureLevel::SM4;
+	bool bOnlyInPixelShader = GetFeatureLevel() < ERHIFeatureLevel::SM5;
 
 	if (bOnlyInPixelShader && ShaderFrequency != SF_Pixel)
 	{
@@ -4934,7 +4982,7 @@ int32 FHLSLMaterialTranslator::ExternalTexture(UTexture* InTexture, int32& Textu
 
 int32 FHLSLMaterialTranslator::ExternalTextureParameter(FName ParameterName, UTexture* DefaultValue, int32& TextureReferenceIndex)
 {
-	bool bOnlyInPixelShader = GetFeatureLevel() < ERHIFeatureLevel::SM4;
+	bool bOnlyInPixelShader = GetFeatureLevel() < ERHIFeatureLevel::SM5;
 
 	if (bOnlyInPixelShader && ShaderFrequency != SF_Pixel)
 	{
@@ -5918,7 +5966,7 @@ int32 FHLSLMaterialTranslator::LightmapUVs()
 		return NonPixelShaderExpressionError();
 	}
 
-	if (ErrorUnlessFeatureLevelSupported(ERHIFeatureLevel::SM4) == INDEX_NONE)
+	if (ErrorUnlessFeatureLevelSupported(ERHIFeatureLevel::SM5) == INDEX_NONE)
 	{
 		return INDEX_NONE;
 	}
@@ -5941,7 +5989,7 @@ int32 FHLSLMaterialTranslator::PrecomputedAOMask()
 		return NonPixelShaderExpressionError();
 	}
 
-	if (ErrorUnlessFeatureLevelSupported(ERHIFeatureLevel::SM4) == INDEX_NONE)
+	if (ErrorUnlessFeatureLevelSupported(ERHIFeatureLevel::SM5) == INDEX_NONE)
 	{
 		return INDEX_NONE;
 	}
@@ -5957,11 +6005,6 @@ int32 FHLSLMaterialTranslator::PrecomputedAOMask()
 	return ResultIdx;
 }
 
-int32 FHLSLMaterialTranslator::LightmassReplace(int32 Realtime, int32 Lightmass)
-{
-	return Realtime;
-}
-
 int32 FHLSLMaterialTranslator::GIReplace(int32 Direct, int32 StaticIndirect, int32 DynamicIndirect)
 { 
 	if(Direct == INDEX_NONE || DynamicIndirect == INDEX_NONE)
@@ -5996,11 +6039,6 @@ int32 FHLSLMaterialTranslator::RayTracingQualitySwitchReplace(int32 Normal, int3
 	return AddCodeChunk(ResultType, TEXT("(GetRayTracingQualitySwitch() ? (%s) : (%s))"), *GetParameterCode(RayTraced), *GetParameterCode(Normal));
 }
 
-int32 FHLSLMaterialTranslator::MaterialProxyReplace(int32 Realtime, int32 MaterialProxy)
-{
-	return Realtime;
-}
-
 int32 FHLSLMaterialTranslator::VirtualTextureOutputReplace(int32 Default, int32 VirtualTexture)
 {
 	if (Default == INDEX_NONE || VirtualTexture == INDEX_NONE)
@@ -6127,7 +6165,7 @@ int32 FHLSLMaterialTranslator::DDY( int32 X )
 
 int32 FHLSLMaterialTranslator::AntialiasedTextureMask(int32 Tex, int32 UV, float Threshold, uint8 Channel)
 {
-	if (ErrorUnlessFeatureLevelSupported(ERHIFeatureLevel::SM4) == INDEX_NONE)
+	if (ErrorUnlessFeatureLevelSupported(ERHIFeatureLevel::SM5) == INDEX_NONE)
 	{
 		return INDEX_NONE;
 	}
@@ -6204,7 +6242,7 @@ int32 FHLSLMaterialTranslator::Noise(int32 Position, float Scale, int32 Quality,
 	// GradientTex3D uses 3D texturing, which is not available on ES2
 	if (NoiseFunction == NOISEFUNCTION_GradientTex3D)
 	{
-		if (ErrorUnlessFeatureLevelSupported(ERHIFeatureLevel::SM4) == INDEX_NONE)
+		if (ErrorUnlessFeatureLevelSupported(ERHIFeatureLevel::SM5) == INDEX_NONE)
 		{
 			Errorf(TEXT("3D textures are not supported for ES2"));
 			return INDEX_NONE;
@@ -6307,6 +6345,41 @@ int32 FHLSLMaterialTranslator::BlackBody( int32 Temp )
 	return AddCodeChunk( MCT_Float3, TEXT("MaterialExpressionBlackBody(%s)"), *GetParameterCode(Temp) );
 }
 
+int32 FHLSLMaterialTranslator::GetHairUV()
+{
+	return AddCodeChunk(MCT_Float2, TEXT("MaterialExpressionGetHairUV(Parameters)"));
+}
+
+int32 FHLSLMaterialTranslator::GetHairDimensions()
+{
+	return AddCodeChunk(MCT_Float2, TEXT("MaterialExpressionGetHairDimensions(Parameters)"));
+}
+
+int32 FHLSLMaterialTranslator::GetHairSeed()
+{
+	return AddCodeChunk(MCT_Float1, TEXT("MaterialExpressionGetHairSeed(Parameters)"));
+}
+
+int32 FHLSLMaterialTranslator::GetHairTangent()
+{
+	return AddCodeChunk(MCT_Float3, TEXT("MaterialExpressionGetHairTangent(Parameters)"));
+}
+
+int32 FHLSLMaterialTranslator::GetHairRootUV()
+{
+	return AddCodeChunk(MCT_Float2, TEXT("MaterialExpressionGetHairRootUV(Parameters)"));
+}
+
+int32 FHLSLMaterialTranslator::GetHairBaseColor()
+{
+	return AddCodeChunk(MCT_Float3, TEXT("MaterialExpressionGetHairBaseColor(Parameters)"));
+}
+
+int32 FHLSLMaterialTranslator::GetHairRoughness()
+{
+	return AddCodeChunk(MCT_Float1, TEXT("MaterialExpressionGetHairRoughness(Parameters)"));
+}
+
 int32 FHLSLMaterialTranslator::DistanceToNearestSurface(int32 PositionArg)
 {
 	if (ErrorUnlessFeatureLevelSupported(ERHIFeatureLevel::SM5) == INDEX_NONE)
@@ -6343,7 +6416,7 @@ int32 FHLSLMaterialTranslator::DistanceFieldGradient(int32 PositionArg)
 
 int32 FHLSLMaterialTranslator::AtmosphericFogColor( int32 WorldPosition )
 {
-	if (ErrorUnlessFeatureLevelSupported(ERHIFeatureLevel::SM4) == INDEX_NONE)
+	if (ErrorUnlessFeatureLevelSupported(ERHIFeatureLevel::SM5) == INDEX_NONE)
 	{
 		return INDEX_NONE;
 	}
@@ -6548,6 +6621,13 @@ int32 FHLSLMaterialTranslator::CustomExpression( class UMaterialExpressionCustom
 			InputParamDecl += InputNameStr;
 			InputParamDecl += TEXT("Sampler ");
 			break;
+		case MCT_Texture2DArray:
+			InputParamDecl += TEXT("Texture2DArray ");
+			InputParamDecl += InputNameStr;
+			InputParamDecl += TEXT(", SamplerState ");
+			InputParamDecl += InputNameStr;
+			InputParamDecl += TEXT("Sampler ");
+			break;
 		case MCT_TextureExternal:
 			InputParamDecl += TEXT("TextureExternal ");
 			InputParamDecl += InputNameStr;
@@ -6595,7 +6675,7 @@ int32 FHLSLMaterialTranslator::CustomExpression( class UMaterialExpressionCustom
 
 		CodeChunk += TEXT(",");
 		CodeChunk += *ParamCode;
-		if (ParamType == MCT_Texture2D || ParamType == MCT_TextureCube || ParamType == MCT_TextureExternal || ParamType == MCT_VolumeTexture)
+		if (ParamType == MCT_Texture2D || ParamType == MCT_TextureCube || ParamType == MCT_Texture2DArray || ParamType == MCT_TextureExternal || ParamType == MCT_VolumeTexture)
 		{
 			CodeChunk += TEXT(",");
 			CodeChunk += *ParamCode;
@@ -6666,8 +6746,17 @@ int32 FHLSLMaterialTranslator::CustomOutput(class UMaterialExpressionCustomOutpu
 
 int32 FHLSLMaterialTranslator::VirtualTextureOutput()
 {
-	MaterialCompilationOutput.bHasRuntimeVirtualTextureOutput = true;
-		
+	if (Material->GetMaterialDomain() == MD_RuntimeVirtualTexture)
+	{
+		// RuntimeVirtualTextureOutput would priority over the output material attributes here
+		// But that could be considered confusing for the user so we error instead
+		Errorf(TEXT("RuntimeVirtualTextureOutput nodes are not used when the Material Domain is set to 'Virtual Texture'"));
+	}
+	else
+	{
+		MaterialCompilationOutput.bHasRuntimeVirtualTextureOutput = true;
+	}
+
 	// return value is not used
 	return INDEX_NONE;
 }
@@ -6804,7 +6893,7 @@ int32 FHLSLMaterialTranslator::SpeedTree(int32 GeometryArg, int32 WindArg, int32
 	*/
 int32 FHLSLMaterialTranslator::TextureCoordinateOffset()
 {
-	if (FeatureLevel < ERHIFeatureLevel::SM4 && ShaderFrequency == SF_Vertex)
+	if (FeatureLevel < ERHIFeatureLevel::SM5 && ShaderFrequency == SF_Vertex)
 	{
 		return AddInlinedCodeChunk(MCT_Float2, TEXT("Parameters.TexCoordOffset"));
 	}
@@ -6824,7 +6913,7 @@ int32 FHLSLMaterialTranslator::EyeAdaptation()
 
 	if( ShaderFrequency != SF_Pixel )
 	{
-		NonPixelShaderExpressionError();
+		return NonPixelShaderExpressionError();
 	}
 
 	MaterialCompilationOutput.bUsesEyeAdaptation = true;
diff --git a/Engine/Source/Runtime/Engine/Private/Materials/HLSLMaterialTranslator.h b/Engine/Source/Runtime/Engine/Private/Materials/HLSLMaterialTranslator.h
--- a/Engine/Source/Runtime/Engine/Private/Materials/HLSLMaterialTranslator.h
+++ b/Engine/Source/Runtime/Engine/Private/Materials/HLSLMaterialTranslator.h
@@ -516,7 +516,7 @@ protected:
 
 	virtual int32 ParticleMacroUV() override;
 	virtual int32 ParticleSubUV(int32 TextureIndex, EMaterialSamplerType SamplerType, bool bBlend) override;
-
+	virtual int32 ParticleSubUVProperty(int32 PropertyIndex) override;
 	virtual int32 ParticleColor() override;
 	virtual int32 ParticlePosition() override;
 	virtual int32 ParticleRadius() override;
@@ -653,14 +653,11 @@ protected:
 	virtual int32 DynamicParameter(FLinearColor& DefaultValue, uint32 ParameterIndex = 0) override;
 	virtual int32 LightmapUVs() override;
 	virtual int32 PrecomputedAOMask() override;
-	virtual int32 LightmassReplace(int32 Realtime, int32 Lightmass) override;
 	virtual int32 GIReplace(int32 Direct, int32 StaticIndirect, int32 DynamicIndirect) override;
 	virtual int32 ShadowReplace(int32 Default, int32 Shadow) override;
 
 	virtual int32 RayTracingQualitySwitchReplace(int32 Normal, int32 RayTraced);
 
-	virtual int32 MaterialProxyReplace(int32 Realtime, int32 MaterialProxy) override;
-
 	virtual int32 VirtualTextureOutputReplace(int32 Default, int32 VirtualTexture) override;
 
 	virtual int32 ObjectOrientation() override;
@@ -681,6 +678,14 @@ protected:
 	virtual int32 VectorNoise(int32 Position, int32 Quality, uint8 NoiseFunction, bool bTiling, uint32 TileSize) override;
 
 	virtual int32 BlackBody(int32 Temp) override;
+
+	virtual int32 GetHairUV() override;
+	virtual int32 GetHairDimensions() override;
+	virtual int32 GetHairSeed() override;
+	virtual int32 GetHairTangent() override;
+	virtual int32 GetHairRootUV() override;
+	virtual int32 GetHairBaseColor() override;
+	virtual int32 GetHairRoughness() override;
 	virtual int32 DistanceToNearestSurface(int32 PositionArg) override;
 	virtual int32 DistanceFieldGradient(int32 PositionArg) override;
 	virtual int32 AtmosphericFogColor(int32 WorldPosition) override;
