Add additional debugging information for ULevelStreaming::UpdateStreaming ensure
#rnx

diff --git a/Engine/Source/Runtime/Engine/Classes/Engine/LevelStreaming.h b/Engine/Source/Runtime/Engine/Classes/Engine/LevelStreaming.h
--- a/Engine/Source/Runtime/Engine/Classes/Engine/LevelStreaming.h
+++ b/Engine/Source/Runtime/Engine/Classes/Engine/LevelStreaming.h
@@ -94,6 +94,8 @@ class ENGINE_API ULevelStreaming : public UObject
 		MakingInvisible
 	};
 
+	static const TCHAR* EnumToString(ECurrentState InCurrentState);
+
 private:
 	enum class ETargetState : uint8
 	{
@@ -103,6 +105,8 @@ private:
 		LoadedVisible,
 	};
 
+	static const TCHAR* EnumToString(ETargetState InTargetState);
+
 public:
 
 #if WITH_EDITORONLY_DATA
diff --git a/Engine/Source/Runtime/Engine/Private/LevelStreaming.cpp b/Engine/Source/Runtime/Engine/Private/LevelStreaming.cpp
--- a/Engine/Source/Runtime/Engine/Private/LevelStreaming.cpp
+++ b/Engine/Source/Runtime/Engine/Private/LevelStreaming.cpp
@@ -689,10 +689,51 @@ void ULevelStreaming::UpdateStreamingState(bool& bOutUpdateAgain, bool& bOutRede
 		break;
 
 	default:
-		ensure(false);
+		ensureMsgf(false, TEXT("Unexpected state in ULevelStreaming::UpdateStreamingState for '%s'. CurrentState='%s' TargetState='%s'"), *GetPathName(), EnumToString(CurrentState), EnumToString(TargetState));
+	}
+}
+
+const TCHAR* ULevelStreaming::EnumToString(ECurrentState InCurrentState)
+{
+	switch (InCurrentState)
+	{
+	case ECurrentState::Removed:
+		return TEXT("Removed");
+	case ECurrentState::Unloaded:
+		return TEXT("Unloaded");
+	case ECurrentState::FailedToLoad:
+		return TEXT("FailedToLoad");
+	case ECurrentState::Loading:
+		return TEXT("Loading");
+	case ECurrentState::LoadedNotVisible:
+		return TEXT("LoadedNotVisible");
+	case ECurrentState::MakingVisible:
+		return TEXT("MakingVisible");
+	case ECurrentState::LoadedVisible:
+		return TEXT("LoadedVisible");
+	case ECurrentState::MakingInvisible:
+		return TEXT("MakingInvisible");
 	}
+	ensure(false);
+	return TEXT("Unknown");
 }
 
+const TCHAR* ULevelStreaming::EnumToString(ETargetState InTargetState)
+{
+	switch (InTargetState)
+	{
+	case ETargetState::Unloaded:
+		return TEXT("Unloaded");
+	case ETargetState::UnloadedAndRemoved:
+		return TEXT("UnloadedAndRemoved");
+	case ETargetState::LoadedNotVisible:
+		return TEXT("LoadedNotVisible");
+	case ETargetState::LoadedVisible:
+		return TEXT("LoadedVisible");
+	}
+	ensure(false);
+	return TEXT("Unknown");
+}
 
 FName ULevelStreaming::GetLODPackageName() const
 {
