Modified Chaos DDC builder to persist face remap table, needed for physical material masks.

diff --git a/Engine/Source/Runtime/Engine/Classes/PhysicsEngine/BodySetup.h b/Engine/Source/Runtime/Engine/Classes/PhysicsEngine/BodySetup.h
--- a/Engine/Source/Runtime/Engine/Classes/PhysicsEngine/BodySetup.h
+++ b/Engine/Source/Runtime/Engine/Classes/PhysicsEngine/BodySetup.h
@@ -121,6 +121,9 @@ struct ENGINE_API FCookBodySetupInfo
 	/** Whether to support UV from hit results */
 	bool bSupportUVFromHitResults;
 
+	/** Whether to support face remap, needed for physical material masks */
+	bool bSupportFaceRemap;
+
 	/** Error generating cook info for trimesh*/
 	bool bTriMeshError;
 };
@@ -269,6 +272,9 @@ public:
 	/** Additional UV info, if available. Used for determining UV for a line trace impact. */
 	FBodySetupUVInfo UVInfo;
 
+	/** Additional face remap table, if available. Used for determining face index mapping from collision mesh to static mesh, for use with physical material masks */
+	TArray<int32> FaceRemap;
+
 	/** Default properties of the body instance, copied into objects on instantiation, was URB_BodyInstance */
 	UPROPERTY(EditAnywhere, Category=Collision, meta=(FullyExpand = "true"))
 	FBodyInstance DefaultInstance;
diff --git a/Engine/Source/Runtime/Engine/Private/PhysicsEngine/BodySetup.cpp b/Engine/Source/Runtime/Engine/Private/PhysicsEngine/BodySetup.cpp
--- a/Engine/Source/Runtime/Engine/Private/PhysicsEngine/BodySetup.cpp
+++ b/Engine/Source/Runtime/Engine/Private/PhysicsEngine/BodySetup.cpp
@@ -97,6 +97,7 @@ FCookBodySetupInfo::FCookBodySetupInfo() :
 	bConvexDeformableMesh(false),
 	bCookTriMesh(false),
 	bSupportUVFromHitResults(false),
+	bSupportFaceRemap(false),
 	bTriMeshError(false)
 {
 }
@@ -350,6 +351,14 @@ void UBodySetup::GetCookInfo(FCookBodySetupInfo& OutCookInfo, EPhysXMeshCookFlag
 			}
 
 			OutCookInfo.TriMeshCookFlags = CookFlags;
+
+			// If outer is a static mesh with physical material mask enabled, set retain UV and face remap table to true
+			/*
+			if (UStaticMesh* StaticMesh = Cast<UStaticMesh>(CDPObj))
+			{
+				OutCookInfo.bSupportFaceRemap = StaticMesh->GetEnablePhysicalMaterialMask();
+			}
+			*/
 		}
 		else
 		{
@@ -357,7 +366,7 @@ void UBodySetup::GetCookInfo(FCookBodySetupInfo& OutCookInfo, EPhysXMeshCookFlag
 		}
 	}
 
-	OutCookInfo.bSupportUVFromHitResults = UPhysicsSettings::Get()->bSupportUVFromHitResults;
+	OutCookInfo.bSupportUVFromHitResults = UPhysicsSettings::Get()->bSupportUVFromHitResults || OutCookInfo.bSupportFaceRemap;
 
 #endif // WITH_PHYSX
 }
@@ -723,6 +732,8 @@ void UBodySetup::FinishCreatingPhysicsMeshes_Chaos(FChaosDerivedDataReader<float
 	}
 
 	ChaosTriMeshes = MoveTemp(InReader.TrimeshImplicitObjects);
+	UVInfo = MoveTemp(InReader.UVInfo);
+	FaceRemap = MoveTemp(InReader.FaceRemap);
 
 	// Clear the cooked data
 	if (!GIsEditor && !bSharedCookedData)
diff --git a/Engine/Source/Runtime/Engine/Private/PhysicsEngine/Experimental/ChaosDerivedData.cpp b/Engine/Source/Runtime/Engine/Private/PhysicsEngine/Experimental/ChaosDerivedData.cpp
--- a/Engine/Source/Runtime/Engine/Private/PhysicsEngine/Experimental/ChaosDerivedData.cpp
+++ b/Engine/Source/Runtime/Engine/Private/PhysicsEngine/Experimental/ChaosDerivedData.cpp
@@ -21,7 +21,7 @@ const TCHAR* FChaosDerivedDataCooker::GetPluginName() const
 
 const TCHAR* FChaosDerivedDataCooker::GetVersionString() const
 {
-	return TEXT("DB9F6EB601DA4E709F5927C32485080D");
+	return TEXT("307A91E898C442A1A0650CD258F76505");
 }
 
 FString FChaosDerivedDataCooker::GetPluginSpecificCacheKeySuffix() const
@@ -85,7 +85,7 @@ FChaosDerivedDataCooker::FChaosDerivedDataCooker(UBodySetup* InSetup, FName InFo
 }
 
 template<typename Precision>
-void FChaosDerivedDataCooker::BuildTriangleMeshes(TArray<TUniquePtr<Chaos::TTriangleMeshImplicitObject<Precision>>>& OutTriangleMeshes, const FCookBodySetupInfo& InParams)
+void FChaosDerivedDataCooker::BuildTriangleMeshes(TArray<TUniquePtr<Chaos::TTriangleMeshImplicitObject<Precision>>>& OutTriangleMeshes, TArray<int32>& OutFaceRemap, const FCookBodySetupInfo& InParams)
 {
 	if(!InParams.bCookTriMesh)
 	{
@@ -106,7 +106,7 @@ void FChaosDerivedDataCooker::BuildTriangleMeshes(TArray<TUniquePtr<Chaos::TTria
 	}
 
 	TArray<int32> Remap;
-	Chaos::CleanTrimesh(FinalVerts, FinalIndices, &Remap);
+	Chaos::CleanTrimesh(FinalVerts, FinalIndices, &OutFaceRemap);
 
 	// Build particle list #BG Maybe allow TParticles to copy vectors?
 	Chaos::TParticles<Precision, 3> TriMeshParticles;
@@ -147,14 +147,14 @@ void FChaosDerivedDataCooker::BuildTriangleMeshes(TArray<TUniquePtr<Chaos::TTria
 
 			if(bHasMaterials)
 			{
-				if(!ensure(Remap.IsValidIndex(TriangleIndex)))
+				if(!ensure(OutFaceRemap.IsValidIndex(TriangleIndex)))
 				{
 					MaterialIndices.Empty();
 					bHasMaterials = false;
 				}
 				else
 				{
-					const int32 OriginalIndex = Remap[TriangleIndex];
+					const int32 OriginalIndex = OutFaceRemap[TriangleIndex];
 
 					if(ensure(InParams.TriangleMeshDesc.MaterialIndices.IsValidIndex(OriginalIndex)))
 					{
@@ -173,9 +173,9 @@ void FChaosDerivedDataCooker::BuildTriangleMeshes(TArray<TUniquePtr<Chaos::TTria
 	OutTriangleMeshes.Emplace(new Chaos::TTriangleMeshImplicitObject<Precision>(MoveTemp(TriMeshParticles), MoveTemp(Triangles), MoveTemp(MaterialIndices)));
 }
 
-template void FChaosDerivedDataCooker::BuildTriangleMeshes(TArray<TUniquePtr<Chaos::TTriangleMeshImplicitObject<float>>>& OutTriangleMeshes, const FCookBodySetupInfo& InParams);
+template void FChaosDerivedDataCooker::BuildTriangleMeshes(TArray<TUniquePtr<Chaos::TTriangleMeshImplicitObject<float>>>& OutTriangleMeshes, TArray<int32>& OutFaceRemap, const FCookBodySetupInfo& InParams);
 //#BGTODO When it's possible to build with doubles, re-enable this. (Currently at least TRigidTransform cannot build with double precision because we don't have a base transform implementation using them)
-//template void FChaosDerivedDataCooker::BuildTriangleMeshes(TArray<TUniquePtr<Chaos::FImplicitObject>>& OutTriangleMeshes, const FCookBodySetupInfo& InParams);
+//template void FChaosDerivedDataCooker::BuildTriangleMeshes(TArray<TUniquePtr<Chaos::FImplicitObject>>& OutTriangleMeshes, TArray<int32>& OutFaceRemap, const FCookBodySetupInfo& InParams);
 
 void FChaosDerivedDataCooker::BuildConvexMeshes(TArray<TUniquePtr<Chaos::FImplicitObject>>& OutConvexMeshes, const FCookBodySetupInfo& InParams)
 {
@@ -250,10 +250,11 @@ void FChaosDerivedDataCooker::BuildInternal(Chaos::FChaosArchive& Ar, FCookBodyS
 {
 	TArray<TUniquePtr<Chaos::FImplicitObject>> SimpleImplicits;
 	TArray<TUniquePtr<Chaos::TTriangleMeshImplicitObject<Precision>>> ComplexImplicits;
+	TArray<int32> FaceRemap;
 
 	//BuildSimpleShapes(SimpleImplicits, Setup);
 	BuildConvexMeshes(SimpleImplicits, InInfo);
-	BuildTriangleMeshes(ComplexImplicits, InInfo);
+	BuildTriangleMeshes(ComplexImplicits, FaceRemap, InInfo);
 
 	//TUniquePtr<Chaos::TImplicitObjectUnion<Precision, 3>> SimpleUnion = MakeUnique<Chaos::TImplicitObjectUnion<Precision, 3>>(MoveTemp(SimpleImplicits));
 	//TUniquePtr<Chaos::TImplicitObjectUnion<Precision, 3>> ComplexUnion = MakeUnique<Chaos::TImplicitObjectUnion<Precision, 3>>(MoveTemp(ComplexImplicits));
@@ -264,7 +265,12 @@ void FChaosDerivedDataCooker::BuildInternal(Chaos::FChaosArchive& Ar, FCookBodyS
 		UVInfo.FillFromTriMesh(InInfo.TriangleMeshDesc);
 	}
 
-	Ar << SimpleImplicits << ComplexImplicits << UVInfo;
+	if (!InInfo.bSupportFaceRemap)
+	{
+		FaceRemap.Empty();
+	}
+
+	Ar << SimpleImplicits << ComplexImplicits << UVInfo << FaceRemap;
 }
 
 template void FChaosDerivedDataCooker::BuildInternal<float>(Chaos::FChaosArchive& Ar, FCookBodySetupInfo& InInfo);
diff --git a/Engine/Source/Runtime/Engine/Private/PhysicsEngine/Experimental/ChaosDerivedData.h b/Engine/Source/Runtime/Engine/Private/PhysicsEngine/Experimental/ChaosDerivedData.h
--- a/Engine/Source/Runtime/Engine/Private/PhysicsEngine/Experimental/ChaosDerivedData.h
+++ b/Engine/Source/Runtime/Engine/Private/PhysicsEngine/Experimental/ChaosDerivedData.h
@@ -69,7 +69,7 @@ private:
 	void BuildInternal(Chaos::FChaosArchive& Ar, FCookBodySetupInfo& InInfo);
 
 	template<typename Precision>
-	void BuildTriangleMeshes(TArray<TUniquePtr<Chaos::TTriangleMeshImplicitObject<Precision>>>& OutTriangleMeshes, const FCookBodySetupInfo& InParams);
+	void BuildTriangleMeshes(TArray<TUniquePtr<Chaos::TTriangleMeshImplicitObject<Precision>>>& OutTriangleMeshes, TArray<int32>& OutFaceRemap, const FCookBodySetupInfo& InParams);
 
 	void BuildConvexMeshes(TArray<TUniquePtr<Chaos::FImplicitObject>>& OutTriangleMeshes, const FCookBodySetupInfo& InParams);
 
@@ -77,7 +77,7 @@ private:
 	FName RequestedFormat;
 };
 
-extern template void FChaosDerivedDataCooker::BuildTriangleMeshes(TArray<TUniquePtr<Chaos::TTriangleMeshImplicitObject<float>>>& OutTriangleMeshes, const FCookBodySetupInfo& InParams);
+extern template void FChaosDerivedDataCooker::BuildTriangleMeshes(TArray<TUniquePtr<Chaos::TTriangleMeshImplicitObject<float>>>& OutTriangleMeshes, TArray<int32>& OutFaceRemap, const FCookBodySetupInfo& InParams);
 // #BGTODO When it's possible to build with doubles, re-enable this. (Currently at least TRigidTransform cannot build with double precision because we don't have a base transform implementation using them)
 //extern template void FChaosDerivedDataCooker::BuildTriangleMeshes(TArray<Chaos::TCollisionTriangleMesh<double>>& OutTriangleMeshes, const FCookBodySetupInfo& InParams);
 
diff --git a/Engine/Source/Runtime/Engine/Private/PhysicsEngine/Experimental/ChaosDerivedDataReader.cpp b/Engine/Source/Runtime/Engine/Private/PhysicsEngine/Experimental/ChaosDerivedDataReader.cpp
--- a/Engine/Source/Runtime/Engine/Private/PhysicsEngine/Experimental/ChaosDerivedDataReader.cpp
+++ b/Engine/Source/Runtime/Engine/Private/PhysicsEngine/Experimental/ChaosDerivedDataReader.cpp
@@ -26,7 +26,7 @@ FChaosDerivedDataReader<T, d>::FChaosDerivedDataReader(FUntypedBulkData* InBulkD
 	}
 	else
 	{
-		ChaosAr << ConvexImplicitObjects << TrimeshImplicitObjects << UVInfo;
+		ChaosAr << ConvexImplicitObjects << TrimeshImplicitObjects << UVInfo << FaceRemap;
 
 		bReadSuccessful = true;
 	}
diff --git a/Engine/Source/Runtime/Engine/Private/StaticMesh.cpp b/Engine/Source/Runtime/Engine/Private/StaticMesh.cpp
--- a/Engine/Source/Runtime/Engine/Private/StaticMesh.cpp
+++ b/Engine/Source/Runtime/Engine/Private/StaticMesh.cpp
@@ -5642,6 +5642,8 @@ bool UStaticMesh::GetPhysicsTriMeshData(struct FTriMeshCollisionData* CollisionD
 
 	TMap<int32, int32> MeshToCollisionVertMap; // map of static mesh verts to collision verts
 
+	// If the mesh enables physical material masks, override the physics setting since UVs are always required 
+	// bool bCopyUVs = GetEnablePhysicalMaterialMask() || UPhysicsSettings::Get()->bSupportUVFromHitResults; // See if we should copy UVs
 	bool bCopyUVs = UPhysicsSettings::Get()->bSupportUVFromHitResults; // See if we should copy UVs
 
 	// If copying UVs, allocate array for storing them
diff --git a/Engine/Source/Runtime/Engine/Public/Physics/Experimental/ChaosDerivedDataReader.h b/Engine/Source/Runtime/Engine/Public/Physics/Experimental/ChaosDerivedDataReader.h
--- a/Engine/Source/Runtime/Engine/Public/Physics/Experimental/ChaosDerivedDataReader.h
+++ b/Engine/Source/Runtime/Engine/Public/Physics/Experimental/ChaosDerivedDataReader.h
@@ -27,6 +27,7 @@ public:
 	TArray<TUniquePtr<Chaos::FConvex>> ConvexImplicitObjects;
 	TArray<TUniquePtr<Chaos::TTriangleMeshImplicitObject<T>>> TrimeshImplicitObjects;
 	FBodySetupUVInfo UVInfo;
+	TArray<int32> FaceRemap;
 
 private:
 	FChaosDerivedDataReader() = delete;
