#UE4 Lower severity of message about cooking an asset that was set to only cook in development builds

diff --git a/Engine/Source/Runtime/Engine/Private/AssetManager.cpp b/Engine/Source/Runtime/Engine/Private/AssetManager.cpp
--- a/Engine/Source/Runtime/Engine/Private/AssetManager.cpp
+++ b/Engine/Source/Runtime/Engine/Private/AssetManager.cpp
@@ -3317,7 +3317,7 @@ bool UAssetManager::VerifyCanCookPackage(FName PackageName, bool bLogError) cons
 	{
 		if (bLogError)
 		{
-			UE_LOG(LogAssetManager, Error, TEXT("Package %s is set to Development, but bOnlyCookProductionAssets is true!"), *PackageName.ToString());
+			UE_LOG(LogAssetManager, Warning, TEXT("Package %s is set to Development, but bOnlyCookProductionAssets is true!"), *PackageName.ToString());
 		}
 
 		return false;
