Fix characters moving farther than they should during movement corrections involving additive RootMotionSources; prevent unnecessarily delayed ServerMoves for moves containing root motion by setting its bForceNoCombine as early as possible
- Fixed additive RootMotionSource-containing SavedMoves in CharacterMovementComponent incorrectly polluting future SavedMoves with their additive velocity during correction-induced SavedMove replay/resimulation by restoring saved root motion to the character's root motion prior to move replay in cases where root motion has just ended, so that the correct "remove additive root motion from character when the root motion ends" clean-up code gets run (see UDN https://udn.unrealengine.com/questions/443525/root-motion-source-error-correction-and-variable-c.html )
- Fixed ClampVelocity/SetVelocity finish parameters not applying to Velocity during move replays in cases where we are not fully resimulating the root motion sources (like when the correction is triggered via ClientAdjustPosition right after a root motion source ended and the root motion ends during one of our SavedMoves). (see UDN https://udn.unrealengine.com/questions/536618/root-motion-source-high-latency.html )
- We now also set bForceNoCombine explicitly for SavedMoves immediately after the move has been done in cases involving RootMotionSources, so that CanDelaySendingMove() will return false and allow the ServerMove to be sent immediately, reducing latency and avoiding mispredictions

#rnx

diff --git a/Engine/Source/Runtime/Engine/Private/Components/CharacterMovementComponent.cpp b/Engine/Source/Runtime/Engine/Private/Components/CharacterMovementComponent.cpp
--- a/Engine/Source/Runtime/Engine/Private/Components/CharacterMovementComponent.cpp
+++ b/Engine/Source/Runtime/Engine/Private/Components/CharacterMovementComponent.cpp
@@ -10712,6 +10712,7 @@ void FSavedMove_Character::PostUpdate(ACharacter* Character, FSavedMove_Characte
 		if( Character->SavedRootMotion.HasActiveRootMotionSources() )
 		{
 			SavedRootMotion = Character->SavedRootMotion;
+			bForceNoCombine = true;
 		}
 
 		// Don't want to combine moves that trigger overlaps, because by moving back and replaying the move we could retrigger overlaps.
@@ -11100,7 +11101,7 @@ void FSavedMove_Character::PrepMoveFor(ACharacter* Character)
 	}
 
 	// Resimulate Root Motion Sources if we need to - occurs after server RPCs over a correction during root motion sources.
-	if( SavedRootMotion.HasActiveRootMotionSources() )
+	if( SavedRootMotion.HasActiveRootMotionSources() || Character->SavedRootMotion.HasActiveRootMotionSources() )
 	{
 		if( Character->bClientResimulateRootMotionSources )
 		{
@@ -11114,6 +11115,18 @@ void FSavedMove_Character::PrepMoveFor(ACharacter* Character)
 			SavedRootMotion.CleanUpInvalidRootMotion(DeltaTime, *Character, *Character->GetCharacterMovement());
 			SavedRootMotion.PrepareRootMotion(DeltaTime, *Character, *Character->GetCharacterMovement());
 		}
+		else
+		{
+			// If this is not the first SavedMove we are replaying, clean up any currently applied root motion sources so that if they have
+			// SetVelocity/ClampVelocity finish settings they get applied correctly before CurrentRootMotion gets stomped below
+			if (FNetworkPredictionData_Client_Character* ClientData = Character->GetCharacterMovement()->GetPredictionData_Client_Character())
+			{
+				if (ClientData->SavedMoves[0].Get() != this)
+				{
+					Character->GetCharacterMovement()->CurrentRootMotion.CleanUpInvalidRootMotion(DeltaTime, *Character, *Character->GetCharacterMovement());
+				}
+			}
+		}
 
 		// Restore root motion to that of this SavedMove to be used during replaying the Move
 		Character->GetCharacterMovement()->CurrentRootMotion = SavedRootMotion;
@@ -11204,4 +11217,4 @@ void UCharacterMovementComponent::FlushServerMoves()
 ETeleportType UCharacterMovementComponent::GetTeleportType() const
 { 
 	return bJustTeleported || bNetworkLargeClientCorrection ? ETeleportType::TeleportPhysics : ETeleportType::None;
-}
\ No newline at end of file
+}
