Deescalate GetLoadedChunk logs when seeked out of bounds.

diff --git a/Engine/Source/Runtime/Engine/Private/ADPCMAudioInfo.cpp b/Engine/Source/Runtime/Engine/Private/ADPCMAudioInfo.cpp
--- a/Engine/Source/Runtime/Engine/Private/ADPCMAudioInfo.cpp
+++ b/Engine/Source/Runtime/Engine/Private/ADPCMAudioInfo.cpp
@@ -812,7 +812,7 @@ const uint8* FADPCMAudioInfo::GetLoadedChunk(USoundWave* InSoundWave, uint32 Chu
 {
 	if (!InSoundWave || ChunkIndex >= InSoundWave->GetNumChunks())
 	{
-		UE_LOG(LogAudio, Error, TEXT("Error calling GetLoadedChunk for ChunkIndex %d!"), ChunkIndex);
+		UE_LOG(LogAudio, Verbose, TEXT("Error calling GetLoadedChunk on wave with %d chunks. ChunkIndex: %d. Name: %s"), InSoundWave->GetNumChunks(), ChunkIndex, *InSoundWave->GetFullName());
 		OutChunkSize = 0;
 		return nullptr;
 	}
diff --git a/Engine/Source/Runtime/Engine/Private/AudioDecompress.cpp b/Engine/Source/Runtime/Engine/Private/AudioDecompress.cpp
--- a/Engine/Source/Runtime/Engine/Private/AudioDecompress.cpp
+++ b/Engine/Source/Runtime/Engine/Private/AudioDecompress.cpp
@@ -361,7 +361,7 @@ const uint8* IStreamedCompressedInfo::GetLoadedChunk(USoundWave* InSoundWave, ui
 {
 	if (!InSoundWave || ChunkIndex >= InSoundWave->GetNumChunks())
 	{
-		UE_LOG(LogAudio, Error, TEXT("Error calling GetLoadedChunk for ChunkIndex %d!"), ChunkIndex);
+		UE_LOG(LogAudio, Verbose, TEXT("Error calling GetLoadedChunk on wave with %d chunks. ChunkIndex: %d. Name: %s"), InSoundWave->GetNumChunks(), ChunkIndex, *InSoundWave->GetFullName());
 		OutChunkSize = 0;
 		return nullptr;
 	}
diff --git a/Engine/Source/Runtime/Engine/Private/VorbisAudioInfo.cpp b/Engine/Source/Runtime/Engine/Private/VorbisAudioInfo.cpp
--- a/Engine/Source/Runtime/Engine/Private/VorbisAudioInfo.cpp
+++ b/Engine/Source/Runtime/Engine/Private/VorbisAudioInfo.cpp
@@ -576,7 +576,7 @@ const uint8* FVorbisAudioInfo::GetLoadedChunk(USoundWave* InSoundWave, uint32 Ch
 {
 	if (!InSoundWave || ChunkIndex >= InSoundWave->GetNumChunks())
 	{
-		UE_LOG(LogAudio, Error, TEXT("Error calling GetLoadedChunk for ChunkIndex %d!"), ChunkIndex);
+		UE_LOG(LogAudio, Verbose, TEXT("Error calling GetLoadedChunk on wave with %d chunks. ChunkIndex: %d. Name: %s"), InSoundWave->GetNumChunks(), ChunkIndex, *InSoundWave->GetFullName());
 		OutChunkSize = 0;
 		return nullptr;
 	}
