Add OnWorldBeginPlay delegate for use by upcoming world systems

diff --git a/Engine/Source/Runtime/Engine/Classes/Engine/World.h b/Engine/Source/Runtime/Engine/Classes/Engine/World.h
--- a/Engine/Source/Runtime/Engine/Classes/Engine/World.h
+++ b/Engine/Source/Runtime/Engine/Classes/Engine/World.h
@@ -2207,6 +2207,9 @@ public:
 	DECLARE_MULTICAST_DELEGATE_OneParam(FOnWorldInitializedActors, const FActorsInitializedParams&);
 	FOnWorldInitializedActors OnActorsInitialized;
 
+	DECLARE_MULTICAST_DELEGATE(FOnWorldBeginPlay);
+	FOnWorldBeginPlay OnWorldBeginPlay;
+
 	/** Returns true if gameplay has already started, false otherwise. */
 	bool HasBegunPlay() const;
 
diff --git a/Engine/Source/Runtime/Engine/Private/World.cpp b/Engine/Source/Runtime/Engine/Private/World.cpp
--- a/Engine/Source/Runtime/Engine/Private/World.cpp
+++ b/Engine/Source/Runtime/Engine/Private/World.cpp
@@ -4112,6 +4112,8 @@ void UWorld::BeginPlay()
 		}
 	}
 
+	OnWorldBeginPlay.Broadcast();
+
 #if WITH_CHAOS
 	if(PhysicsScene)
 	{
