Added a general texture streaming framework to allow custom implementations to fill in mip data (from other source than the cooked mips).
User must derive a class from UTextureMipDataProviderFactory and add it to the UTexture::AssetUserData.
This one has to allocate a custom implementation of FTextureMipDataProvider to fill the streamed in mips with it's own strategy.

diff --git a/Engine/Source/Runtime/Engine/Classes/Engine/TextureMipDataProviderFactory.h b/Engine/Source/Runtime/Engine/Classes/Engine/TextureMipDataProviderFactory.h
new file mode 100644
--- /dev/null
+++ b/Engine/Source/Runtime/Engine/Classes/Engine/TextureMipDataProviderFactory.h
@@ -0,0 +1,37 @@
+// Copyright 1998-2019 Epic Games, Inc. All Rights Reserved.
+
+/*=============================================================================
+TextureMipDataProviderFactory.h: base class to create custom FTextureMipDataProvider.
+=============================================================================*/
+
+#pragma once
+#include "CoreMinimal.h"
+#include "Engine/AssetUserData.h"
+#include "TextureMipDataProviderFactory.generated.h"
+
+class UTexture;
+class FTextureMipDataProvider;
+
+/**
+* UTextureMipDataProviderFactory defines an interface to create instances of FTextureMipDataProvider.
+* Derived classes from UTextureMipDataProviderFactory can be attached to UTexture::MipDataProviderFactory
+* to define a new source for mip content (instead of the default disk file or ddc mips). 
+* Usecases include dynamic textures that need to be driven by the texture streaming or textures that 
+* get they data over the network.
+*/
+UCLASS(abstract, hidecategories=Object)
+class ENGINE_API UTextureMipDataProviderFactory : public UAssetUserData
+{
+	GENERATED_UCLASS_BODY()
+
+public:
+
+	/**
+	* Create a FTextureMipDataProvider to handle a single StreamIn mip operation.
+	* The object lifetime will be managed by FRenderAssetUpdate.
+	* 
+	* @param InAsset - the texture on which the stream in operation will be performed.
+	* @param NewMipCount - the total number mips after the stream in request.
+	*/
+	FTextureMipDataProvider* AllocateMipDataProvider(UTexture* Asset, int32 NewMipCount) { return nullptr; }
+};
diff --git a/Engine/Source/Runtime/Engine/Private/Streaming/RenderAssetUpdate.inl b/Engine/Source/Runtime/Engine/Private/Streaming/RenderAssetUpdate.inl
--- a/Engine/Source/Runtime/Engine/Private/Streaming/RenderAssetUpdate.inl
+++ b/Engine/Source/Runtime/Engine/Private/Streaming/RenderAssetUpdate.inl
@@ -1,7 +1,7 @@
 // Copyright 1998-2019 Epic Games, Inc. All Rights Reserved.
 
 /*=============================================================================
-RenderAssetUpdate.h: Base class of helpers to stream in and out texture/mesh LODs
+RenderAssetUpdate.inl: Base class of helpers to stream in and out texture/mesh LODs
 =============================================================================*/
 
 #pragma once
@@ -71,7 +71,9 @@ FRenderAssetUpdate::ETaskState TRenderAssetUpdate<TContext>::TickInternal(EThrea
 	else if (RelevantThread == InCurrentThread)
 	{
 		ClearCallbacks();
-		RelevantCallback(FContext(RelevantAsset, InCurrentThread));
+		FContext Context(RelevantAsset, InCurrentThread);
+		InitContext(Context);
+		RelevantCallback(Context);
 		return TS_Locked;
 	} 
 	else if (RelevantThread == FRenderAssetUpdate::TT_GameThread && !ScheduledGTTasks)
diff --git a/Engine/Source/Runtime/Engine/Private/Streaming/Texture2DMipAllocator_AsyncCreate.cpp b/Engine/Source/Runtime/Engine/Private/Streaming/Texture2DMipAllocator_AsyncCreate.cpp
new file mode 100644
--- /dev/null
+++ b/Engine/Source/Runtime/Engine/Private/Streaming/Texture2DMipAllocator_AsyncCreate.cpp
@@ -0,0 +1,161 @@
+// Copyright 1998-2019 Epic Games, Inc. All Rights Reserved.
+
+/*=============================================================================
+Texture2DStreamIn_AsyncCreate.cpp: Implementation of FTextureMipAllocator using RHIAsyncCreateTexture2D
+=============================================================================*/
+
+#include "Texture2DMipAllocator_AsyncCreate.h"
+#include "RenderUtils.h"
+#include "Containers/ResourceArray.h"
+
+FTexture2DMipAllocator_AsyncCreate::FTexture2DMipAllocator_AsyncCreate()
+	: FTextureMipAllocator(ETickState::AllocateMips, ETickThread::Async)
+{
+}
+
+FTexture2DMipAllocator_AsyncCreate::~FTexture2DMipAllocator_AsyncCreate()
+{
+	check(!FinalMipData.Num());
+}
+
+bool FTexture2DMipAllocator_AsyncCreate::AllocateMips(
+	const FTextureUpdateContext& Context, 
+	FTextureMipInfoArray& OutMipInfos, 
+	const FTextureUpdateSyncOptions& SyncOptions)
+{
+	check(Context.PendingFirstMipIndex < Context.CurrentFirstMipIndex);
+
+	UTexture2D* Texture2D = CastChecked<UTexture2D>(Context.Texture, ECastCheckedType::NullChecked);
+	FTexture2DResource* Resource = static_cast<FTexture2DResource*>(Texture2D->Resource);
+	FTexture2DRHIRef Texture2DRHI = Resource ? Resource->GetTexture2DRHI() : FTexture2DRHIRef();
+	if (!Texture2DRHI)
+	{
+		return false;
+	}
+
+	OutMipInfos.AddDefaulted(Context.CurrentFirstMipIndex);
+
+	// Allocate the mip memory as temporary buffers so that the FTextureMipDataProvider implementation can write to it.
+	const TIndirectArray<FTexture2DMipMap>& OwnerMips = Texture2D->GetPlatformMips();
+	for (int32 MipIndex = Context.PendingFirstMipIndex; MipIndex < Context.CurrentFirstMipIndex; ++MipIndex)
+	{
+		const FTexture2DMipMap& OwnerMip = OwnerMips[MipIndex];
+		FTextureMipInfo& MipInfo = OutMipInfos[MipIndex];
+
+		MipInfo.Format = Texture2DRHI->GetFormat();
+		MipInfo.SizeX = OwnerMip.SizeX;
+		MipInfo.SizeY = OwnerMip.SizeY;
+		MipInfo.DataSize = CalcTextureMipMapSize(MipInfo.SizeX, MipInfo.SizeY, MipInfo.Format, 0);
+		// Allocate the mip in main memory. It will later be used to create the mips with proper initial states (without going through lock/unlock).
+		MipInfo.DestData = FMemory::Malloc(MipInfo.DataSize);
+
+		// Backup the allocated memory so that it can safely be freed.
+		FinalMipData.Add(MipInfo.DestData);
+	}
+
+	// Backup size and format.
+	if (OutMipInfos.IsValidIndex(Context.PendingFirstMipIndex))
+	{
+		FinalSizeX = OutMipInfos[Context.PendingFirstMipIndex].SizeX;
+		FinalSizeY = OutMipInfos[Context.PendingFirstMipIndex].SizeY;
+		FinalFormat = OutMipInfos[Context.PendingFirstMipIndex].Format;
+
+		// Once the FTextureMipDataProvider has set the mip data, FinalizeMips can then create the texture in it's step (1).
+		AdvanceTo(ETickState::FinalizeMips, ETickThread::Async);
+		return true;
+	}
+	else // No new mips? something is wrong.
+	{
+		return false;
+	}
+}
+
+// This gets called 2 times :
+// - Async : create the texture with the mips data
+// - Render : swap the results
+bool FTexture2DMipAllocator_AsyncCreate::FinalizeMips(const FTextureUpdateContext& Context, const FTextureUpdateSyncOptions& SyncOptions)
+{
+	UTexture2D* Texture2D = CastChecked<UTexture2D>(Context.Texture, ECastCheckedType::NullChecked);
+	FTexture2DResource* Resource = static_cast<FTexture2DResource*>(Texture2D->Resource);
+	FTexture2DRHIRef Texture2DRHI = Resource ? Resource->GetTexture2DRHI() : FTexture2DRHIRef();
+	if (!Texture2DRHI)
+	{
+		return false;
+	}
+
+	// Step (1) : Create the texture on the async thread, having the new mip data as reference so that it can be initialized correctly.
+	if (!IntermediateTextureRHI)
+	{
+		const uint32 Flags = (Texture2D->SRGB ? TexCreate_SRGB : 0) | TexCreate_DisableAutoDefrag;
+
+		// Create the intermediate texture.
+		IntermediateTextureRHI = RHIAsyncCreateTexture2D(FinalSizeX, FinalSizeY, FinalFormat, Context.NumRequestedMips, Flags, FinalMipData.GetData(), FinalMipData.Num());
+		// Free the temporary mip data, since copy is now in the RHIAsyncCreateTexture2D command.
+		ReleaseAllocatedMipData();
+
+		// Go to next step, on the renderthread.
+		AdvanceTo(ETickState::FinalizeMips, ETickThread::Render);
+	}
+	// Step (2) : Copy the non initialized mips on the using RHICopySharedMips, must run on the renderthread.
+	else
+	{
+		// Copy the mips.
+		RHICopySharedMips(IntermediateTextureRHI, Texture2DRHI);
+		// Use the new texture resource for the texture asset, must run on the renderthread.
+		Resource->UpdateTexture(IntermediateTextureRHI, Context.PendingFirstMipIndex);
+		// No need for the intermediate texture anymore.
+		IntermediateTextureRHI.SafeRelease();
+
+		// Update complete, nothing more to do.
+		AdvanceTo(ETickState::Done, ETickThread::None);
+	}
+	return true;
+}
+
+void FTexture2DMipAllocator_AsyncCreate::Cancel(const FTextureUpdateSyncOptions& SyncOptions)
+{
+	// Release the intermediate texture. If not null, this will be on the renderthread.
+	IntermediateTextureRHI.SafeRelease();
+	// Release the temporary mip data. Can be run on either renderthread or async threads.
+	ReleaseAllocatedMipData();
+}
+
+FTextureMipAllocator::ETickThread FTexture2DMipAllocator_AsyncCreate::GetCancelThread() const
+{
+	// If there is an  intermediate texture, it is safer to released on the renderthread.
+	if (IntermediateTextureRHI)
+	{
+		return ETickThread::Render;
+	}
+	// Otherwise, if there are only temporary mip data, they can be freed on any threads.
+	else if (FinalMipData.Num())
+	{
+		return ETickThread::Async;
+	}
+	// Nothing to do.
+	else
+	{
+		return ETickThread::None;
+	}
+}
+
+int32 FTexture2DMipAllocator_AsyncCreate::GetCurrentFirstMip(UTexture* Texture) const
+{
+	UTexture2D* Texture2D = CastChecked<UTexture2D>(Texture, ECastCheckedType::NullChecked);
+	FTexture2DResource* Resource = static_cast<FTexture2DResource*>(Texture2D->Resource);
+	return Resource ? Resource->GetCurrentFirstMip() : INDEX_NONE;
+}
+
+
+void FTexture2DMipAllocator_AsyncCreate::ReleaseAllocatedMipData()
+{
+	// Release the temporary mip data.
+	for (void* NewData : FinalMipData)
+	{
+		if (NewData)
+		{
+			FMemory::Free(NewData);
+		}
+	}
+	FinalMipData.Empty();
+}
diff --git a/Engine/Source/Runtime/Engine/Private/Streaming/Texture2DMipAllocator_AsyncCreate.h b/Engine/Source/Runtime/Engine/Private/Streaming/Texture2DMipAllocator_AsyncCreate.h
new file mode 100644
--- /dev/null
+++ b/Engine/Source/Runtime/Engine/Private/Streaming/Texture2DMipAllocator_AsyncCreate.h
@@ -0,0 +1,54 @@
+// Copyright 1998-2019 Epic Games, Inc. All Rights Reserved.
+
+/*=============================================================================
+Texture2DStreamIn_AsyncCreate.h: Implementation of FTextureMipAllocator using RHIAsyncCreateTexture2D
+=============================================================================*/
+
+#pragma once
+
+#include "CoreMinimal.h"
+#include "TextureMipAllocator.h"
+#include "Engine/Texture2D.h"
+
+class UTexture2D;
+
+/**
+* FTexture2DMipAllocator_AsyncCreate implements FTextureMipAllocator using RHIAsyncCreateTexture2D.
+* It uses a more efficient path than FTexture2DMipAllocator_AsyncReallocate because it can run some
+* steps on async threads (and not just on the render thread). 
+* Can only be used when GRHISupportsAsyncTextureCreation is true.
+*/
+class FTexture2DMipAllocator_AsyncCreate : public FTextureMipAllocator
+{
+public:
+
+	FTexture2DMipAllocator_AsyncCreate();
+	~FTexture2DMipAllocator_AsyncCreate();
+
+	// ********************************************************
+	// ********* FTextureMipAllocator implementation **********
+	// ********************************************************
+
+	bool AllocateMips(const FTextureUpdateContext& Context, FTextureMipInfoArray& OutMipInfos, const FTextureUpdateSyncOptions& SyncOptions) final override;
+	bool FinalizeMips(const FTextureUpdateContext& Context, const FTextureUpdateSyncOptions& SyncOptions) final override;
+	void Cancel(const FTextureUpdateSyncOptions& SyncOptions) final override;
+	ETickThread GetCancelThread() const final override;
+	int32 GetCurrentFirstMip(UTexture* Texture) const final override;
+
+protected:
+
+	// Release the temporary buffers referenced in FinalMipData.
+	void ReleaseAllocatedMipData();
+
+	// The intermediate texture async created in the update process.
+	FTexture2DRHIRef IntermediateTextureRHI;
+
+	// The final resolution x and y of the texture 2d.
+	int32 FinalSizeX = 0;
+	int32 FinalSizeY = 0;
+	// The initial and final format of the texture.
+	EPixelFormat FinalFormat = EPixelFormat::PF_Unknown;
+	// The temporary main memory allocations holding the mip data
+	TArray<void*, TInlineAllocator<MAX_TEXTURE_MIP_COUNT>> FinalMipData;
+};
+
diff --git a/Engine/Source/Runtime/Engine/Private/Streaming/Texture2DMipAllocator_AsyncReallocate.cpp b/Engine/Source/Runtime/Engine/Private/Streaming/Texture2DMipAllocator_AsyncReallocate.cpp
new file mode 100644
--- /dev/null
+++ b/Engine/Source/Runtime/Engine/Private/Streaming/Texture2DMipAllocator_AsyncReallocate.cpp
@@ -0,0 +1,169 @@
+// Copyright 1998-2019 Epic Games, Inc. All Rights Reserved.
+
+/*=============================================================================
+Texture2DStreamIn_AsyncReallocate.cpp: Load texture 2D mips using ITextureMipDataProvider
+=============================================================================*/
+
+#include "Texture2DMipAllocator_AsyncReallocate.h"
+#include "RenderUtils.h"
+#include "Containers/ResourceArray.h"
+
+extern TAutoConsoleVariable<int32> CVarFlushRHIThreadOnSTreamingTextureLocks;
+
+FTexture2DMipAllocator_AsyncReallocate::FTexture2DMipAllocator_AsyncReallocate()
+	: FTextureMipAllocator(ETickState::AllocateMips, ETickThread::Render)
+{
+}
+
+FTexture2DMipAllocator_AsyncReallocate::~FTexture2DMipAllocator_AsyncReallocate()
+{
+	check(!LockedMipIndices.Num());
+}
+
+// ********************************************************
+// ********* FTextureMipAllocator implementation **********
+// ********************************************************
+
+bool FTexture2DMipAllocator_AsyncReallocate::AllocateMips(
+	const FTextureUpdateContext& Context, 
+	FTextureMipInfoArray& OutMipInfos, 
+	const FTextureUpdateSyncOptions& SyncOptions)
+{
+	check(Context.PendingFirstMipIndex < Context.CurrentFirstMipIndex);
+
+	UTexture2D* Texture2D = CastChecked<UTexture2D>(Context.Texture, ECastCheckedType::NullChecked);
+	FTexture2DResource* Resource = static_cast<FTexture2DResource*>(Texture2D->Resource);
+	FTexture2DRHIRef Texture2DRHI = Resource ? Resource->GetTexture2DRHI() : FTexture2DRHIRef();
+	if (!Texture2DRHI)
+	{
+		return false;
+	}
+
+	// Step (1) : Create the texture on the renderthread using RHIAsyncReallocateTexture2D. Wait for the RHI to signal the operation is completed through the FThreadSafeCounter.
+	if (!IntermediateTextureRHI)
+	{
+		const TIndirectArray<FTexture2DMipMap>& OwnerMips = Texture2D->GetPlatformMips();
+		const FTexture2DMipMap& OwnerMip = OwnerMips[Context.PendingFirstMipIndex];
+		check(SyncOptions.Counter);
+
+		SyncOptions.Counter->Increment();
+
+		IntermediateTextureRHI = RHIAsyncReallocateTexture2D(
+			Texture2DRHI,
+			Context.NumRequestedMips,
+			OwnerMip.SizeX,
+			OwnerMip.SizeY,
+			SyncOptions.Counter);
+
+		// Run the next step, when IntermediateTextureRHI will be somewhat valid (after synchronization).
+		AdvanceTo(ETickState::AllocateMips, ETickThread::Render);
+		return true;
+	}
+	// Step (2) : Finalize the texture state through RHIFinalizeAsyncReallocateTexture2D and lock the new mips.
+	else
+	{
+		const bool bFlushRHIThread = CVarFlushRHIThreadOnSTreamingTextureLocks.GetValueOnAnyThread() > 0;
+
+		RHIFinalizeAsyncReallocateTexture2D(IntermediateTextureRHI, true);
+
+		OutMipInfos.AddDefaulted(Context.CurrentFirstMipIndex);
+
+		const TIndirectArray<FTexture2DMipMap>& OwnerMips = Texture2D->GetPlatformMips();
+		for (int32 MipIndex = Context.PendingFirstMipIndex; MipIndex < Context.CurrentFirstMipIndex; ++MipIndex)
+		{
+			const FTexture2DMipMap& OwnerMip = OwnerMips[MipIndex];
+			FTextureMipInfo& MipInfo = OutMipInfos[MipIndex];
+
+			MipInfo.Format = Texture2DRHI->GetFormat();
+			MipInfo.SizeX = OwnerMip.SizeX;
+			MipInfo.SizeY = OwnerMip.SizeY;
+#if WITH_EDITORONLY_DATA
+			MipInfo.DataSize = CalcTextureMipMapSize(MipInfo.SizeX, MipInfo.SizeY, MipInfo.Format, 0);
+#else // Hasn't really been used on console. To investigate!
+			MipInfo.DataSize = 0;
+#endif
+			MipInfo.DestData = RHILockTexture2D(IntermediateTextureRHI, MipIndex - Context.PendingFirstMipIndex, RLM_WriteOnly, MipInfo.RowPitch, false, bFlushRHIThread);
+
+			// Add this mip in the locked list of mips so that it can safely be unlocked when needed.
+			LockedMipIndices.Add(MipIndex - Context.PendingFirstMipIndex);
+		}
+
+		// New mips are ready to be unlocked by the FTextureMipDataProvider implementation.
+		AdvanceTo(ETickState::FinalizeMips, ETickThread::Render);
+		return true;
+	}
+}
+
+bool FTexture2DMipAllocator_AsyncReallocate::FinalizeMips(const FTextureUpdateContext& Context, const FTextureUpdateSyncOptions& SyncOptions)
+{
+	UTexture2D* Texture2D = CastChecked<UTexture2D>(Context.Texture, ECastCheckedType::NullChecked);
+	FTexture2DResource* Resource = static_cast<FTexture2DResource*>(Texture2D->Resource);
+	FTexture2DRHIRef Texture2DRHI = Resource ? Resource->GetTexture2DRHI() : FTexture2DRHIRef();
+	if (!Texture2DRHI)
+	{
+		return false;
+	}
+
+	if (!IntermediateTextureRHI)
+	{
+		return false;
+	}
+
+	// Unlock the mips so that the texture can be updated.
+	UnlockNewMips();
+	// Use the new texture resource for the texture asset, must run on the renderthread.
+	Resource->UpdateTexture(IntermediateTextureRHI, Context.PendingFirstMipIndex);
+	// No need for the intermediate texture anymore.
+	IntermediateTextureRHI.SafeRelease();
+
+	// Update complete, nothing more to do.
+	AdvanceTo(ETickState::Done, ETickThread::None);
+	return true;
+}
+
+void FTexture2DMipAllocator_AsyncReallocate::Cancel(const FTextureUpdateSyncOptions& SyncOptions)
+{
+	// Unlock any locked mips.
+	UnlockNewMips();
+	// Release the intermediate texture.
+	IntermediateTextureRHI.SafeRelease();
+}
+
+FTextureMipAllocator::ETickThread FTexture2DMipAllocator_AsyncReallocate::GetCancelThread() const
+{
+	// If there is an intermediate texture and possibly locked mips. Unlock them and release it on the renderthread.
+	if (IntermediateTextureRHI)
+	{
+		return ETickThread::Render;
+	}
+	// Nothing to do.
+	else
+	{
+		return ETickThread::None;
+	}
+}
+
+int32 FTexture2DMipAllocator_AsyncReallocate::GetCurrentFirstMip(UTexture* Texture) const
+{
+	UTexture2D* Texture2D = CastChecked<UTexture2D>(Texture, ECastCheckedType::NullChecked);
+	FTexture2DResource* Resource = static_cast<FTexture2DResource*>(Texture2D->Resource);
+	return Resource ? Resource->GetCurrentFirstMip() : INDEX_NONE;
+}
+
+// ****************************
+// ********* Helpers **********
+// ****************************
+
+void FTexture2DMipAllocator_AsyncReallocate::UnlockNewMips()
+{
+	// Unlock any locked mips.
+	if (IntermediateTextureRHI)
+	{
+		const bool bFlushRHIThread = CVarFlushRHIThreadOnSTreamingTextureLocks.GetValueOnAnyThread() > 0;
+		for (int32 MipIndex : LockedMipIndices)
+		{
+			RHIUnlockTexture2D(IntermediateTextureRHI, MipIndex, false, CVarFlushRHIThreadOnSTreamingTextureLocks.GetValueOnAnyThread() > 0 );
+		}
+		LockedMipIndices.Empty();
+	}
+}
diff --git a/Engine/Source/Runtime/Engine/Private/Streaming/Texture2DMipAllocator_AsyncReallocate.h b/Engine/Source/Runtime/Engine/Private/Streaming/Texture2DMipAllocator_AsyncReallocate.h
new file mode 100644
--- /dev/null
+++ b/Engine/Source/Runtime/Engine/Private/Streaming/Texture2DMipAllocator_AsyncReallocate.h
@@ -0,0 +1,45 @@
+// Copyright 1998-2019 Epic Games, Inc. All Rights Reserved.
+
+/*=============================================================================
+Texture2DMipAllocator_AsyncReallocate.h: Implementation of FTextureMipAllocator using RHIAsyncReallocateTexture2D
+=============================================================================*/
+
+#pragma once
+
+#include "CoreMinimal.h"
+#include "TextureMipAllocator.h"
+#include "Engine/Texture2D.h"
+
+class UTexture2D;
+
+/**
+* FTexture2DMipAllocator_AsyncReallocate is the default implementation of  FTextureMipAllocator.
+* It uses RHIAsyncReallocateTexture2D which creates a new texture on the renderthread and copies the mips if required.
+*/
+class FTexture2DMipAllocator_AsyncReallocate : public FTextureMipAllocator
+{
+public:
+
+	FTexture2DMipAllocator_AsyncReallocate();
+	~FTexture2DMipAllocator_AsyncReallocate();
+
+	// ********************************************************
+	// ********* FTextureMipAllocator implementation **********
+	// ********************************************************
+
+	bool AllocateMips(const FTextureUpdateContext& Context, FTextureMipInfoArray& OutMipInfos, const FTextureUpdateSyncOptions& SyncOptions) final override;
+	bool FinalizeMips(const FTextureUpdateContext& Context, const FTextureUpdateSyncOptions& SyncOptions) final override;
+	void Cancel(const FTextureUpdateSyncOptions& SyncOptions) final override;
+	ETickThread GetCancelThread() const final override;
+	int32 GetCurrentFirstMip(UTexture* Texture) const final override;
+
+protected:
+
+	// Unlock the mips referenced in LockedMipIndices.
+	void UnlockNewMips();
+
+	// The intermediate texture created with RHIAsyncReallocateTexture2D.
+	FTexture2DRHIRef IntermediateTextureRHI;
+	// The list of mips that are currently locked.
+	TArray<int32, TInlineAllocator<MAX_TEXTURE_MIP_COUNT>> LockedMipIndices;
+};
diff --git a/Engine/Source/Runtime/Engine/Private/Streaming/Texture2DMipDataProvider_DDC.cpp b/Engine/Source/Runtime/Engine/Private/Streaming/Texture2DMipDataProvider_DDC.cpp
new file mode 100644
--- /dev/null
+++ b/Engine/Source/Runtime/Engine/Private/Streaming/Texture2DMipDataProvider_DDC.cpp
@@ -0,0 +1,141 @@
+// Copyright 1998-2019 Epic Games, Inc. All Rights Reserved.
+
+/*=============================================================================
+Texture2DMipDataProvider_DDC.cpp : Implementation of FTextureMipDataProvider using DDC requests.
+=============================================================================*/
+
+#include "Texture2DMipDataProvider_DDC.h"
+#include "DerivedDataCacheInterface.h"
+#include "Engine/Texture2D.h"
+#include "TextureResource.h"
+#include "Streaming/Texture2DStreamIn_DDC.h" // for FAbandonedDDCHandleManager
+
+#if WITH_EDITORONLY_DATA
+
+FTexture2DMipDataProvider_DDC::FTexture2DMipDataProvider_DDC()
+	: FTextureMipDataProvider(ETickState::Init, ETickThread::Async)
+{
+}
+
+FTexture2DMipDataProvider_DDC::~FTexture2DMipDataProvider_DDC()
+{
+	check(!DDCHandles.Num());
+}
+
+void FTexture2DMipDataProvider_DDC::Init(const FTextureUpdateContext& Context, const FTextureUpdateSyncOptions& SyncOptions)
+{
+	if (!DDCHandles.Num())
+	{
+		UTexture2D* Texture2D = CastChecked<UTexture2D>(Context.Texture, ECastCheckedType::NullChecked);
+		const TIndirectArray<FTexture2DMipMap>& OwnerMips = Texture2D->GetPlatformMips();
+		DDCHandles.AddZeroed(Context.CurrentFirstMipIndex);
+
+		for (int32 MipIndex = Context.PendingFirstMipIndex; MipIndex < Context.CurrentFirstMipIndex; ++MipIndex)
+		{
+			const FTexture2DMipMap& OwnerMip = OwnerMips[MipIndex];
+			if (!OwnerMip.DerivedDataKey.IsEmpty())
+			{
+				DDCHandles[MipIndex] = GetDerivedDataCacheRef().GetAsynchronous(*OwnerMip.DerivedDataKey);
+			}
+		}
+		*SyncOptions.bSnooze = true;
+	}
+	else // The DDC request have been issued, only check whether they are ready (since no good sync option is used).
+	{
+		for (uint32 Handle : DDCHandles)
+		{
+			if (Handle && !GetDerivedDataCacheRef().PollAsynchronousCompletion(Handle))
+			{
+				*SyncOptions.bSnooze = true;
+				return;
+			}
+		}
+		AdvanceTo(ETickState::GetMips, ETickThread::Async);
+	}
+}
+
+int32 FTexture2DMipDataProvider_DDC::GetMips(
+	const FTextureUpdateContext& Context,
+	int32 StartingMipIndex,
+	const FTextureMipInfoArray& MipInfos, 
+	const FTextureUpdateSyncOptions& SyncOptions)
+{
+	for (int32 MipIndex = StartingMipIndex; MipIndex < Context.CurrentFirstMipIndex; ++MipIndex)
+	{
+		const uint32 Handle = DDCHandles[MipIndex];
+		bool bSuccess = false;
+		if (Handle)
+		{
+			DDCHandles[MipIndex] = 0; // Clear the handle.
+
+			TArray<uint8> DerivedMipData;
+			if (GetDerivedDataCacheRef().GetAsynchronousResults(Handle, DerivedMipData))
+			{
+				const FTextureMipInfo& MipInfo = MipInfos[MipIndex];
+
+				// The result must be read from a memory reader!
+				FMemoryReader Ar(DerivedMipData, true);
+				int32 MipSize = 0;
+				Ar << MipSize;
+
+				if (MipSize == MipInfo.DataSize)
+				{
+					Ar.Serialize(MipInfo.DestData, MipSize);
+					bSuccess = true;
+				}
+			}
+		}
+		if (!bSuccess)
+		{
+			AdvanceTo(ETickState::CleanUp, ETickThread::Async);
+			return MipIndex; // We failed at getting this mip. Cancel will be called.
+		}
+	}
+
+	AdvanceTo(ETickState::CleanUp, ETickThread::Async);
+	return Context.CurrentFirstMipIndex;
+}
+
+bool FTexture2DMipDataProvider_DDC::PollMips(const FTextureUpdateSyncOptions& SyncOptions)
+{
+	AdvanceTo(ETickState::CleanUp, ETickThread::Async);
+	return true;
+}
+
+void FTexture2DMipDataProvider_DDC::CleanUp(const FTextureUpdateSyncOptions& SyncOptions)
+{
+	ReleaseDDCHandles();
+	AdvanceTo(ETickState::Done, ETickThread::None);
+}
+
+void FTexture2DMipDataProvider_DDC::Cancel(const FTextureUpdateSyncOptions& SyncOptions)
+{
+	ReleaseDDCHandles();
+}
+
+FTextureMipDataProvider::ETickThread FTexture2DMipDataProvider_DDC::GetCancelThread() const
+{
+	if (DDCHandles.Num())
+	{
+		return ETickThread::Async;
+	}
+	else
+	{
+		return ETickThread::None;
+	}
+}
+
+void FTexture2DMipDataProvider_DDC::ReleaseDDCHandles()
+{
+	for (uint32& Handle : DDCHandles)
+	{
+		if (Handle)
+		{
+			GAbandonedDDCHandleManager.Add(Handle);
+			Handle = 0;
+		}
+	}
+	DDCHandles.Empty();
+}
+
+#endif //WITH_EDITORONLY_DATA
diff --git a/Engine/Source/Runtime/Engine/Private/Streaming/Texture2DMipDataProvider_DDC.h b/Engine/Source/Runtime/Engine/Private/Streaming/Texture2DMipDataProvider_DDC.h
new file mode 100644
--- /dev/null
+++ b/Engine/Source/Runtime/Engine/Private/Streaming/Texture2DMipDataProvider_DDC.h
@@ -0,0 +1,47 @@
+// Copyright 1998-2019 Epic Games, Inc. All Rights Reserved.
+
+/*=============================================================================
+Texture2DMipDataProvider_DDC.h : Implementation of FTextureMipDataProvider using DDC requests.
+=============================================================================*/
+
+#pragma once
+
+#include "CoreMinimal.h"
+#include "Streaming/TextureMipDataProvider.h"
+
+#if WITH_EDITORONLY_DATA
+
+/**
+* FTexture2DMipDataProvider_DDC implements FTextureMipAllocator using DDC requests from FDerivedDataCacheInterface.
+*/
+class FTexture2DMipDataProvider_DDC : public FTextureMipDataProvider
+{
+public:
+
+	FTexture2DMipDataProvider_DDC();
+	~FTexture2DMipDataProvider_DDC();
+
+
+	// ********************************************************
+	// ********* FTextureMipDataProvider implementation **********
+	// ********************************************************
+
+	void Init(const FTextureUpdateContext& Context, const FTextureUpdateSyncOptions& SyncOptions) final override;
+	int32 GetMips(const FTextureUpdateContext& Context, int32 StartingMipIndex, const FTextureMipInfoArray& MipInfos, const FTextureUpdateSyncOptions& SyncOptions) final override;
+	bool PollMips(const FTextureUpdateSyncOptions& SyncOptions) final override;
+	void CleanUp(const FTextureUpdateSyncOptions& SyncOptions) final override;
+	void Cancel(const FTextureUpdateSyncOptions& SyncOptions) final override;
+	ETickThread GetCancelThread() const final override;
+
+protected:
+
+	// Helper to route incomplete DDC requests to the FAbandonedDDCHandleManager.
+	void ReleaseDDCHandles();
+
+private:
+
+	// The list of DDC handle for each mips (0 based on the highest mip).
+	TArray<uint32, TInlineAllocator<MAX_TEXTURE_MIP_COUNT>> DDCHandles;
+};
+
+#endif //WITH_EDITORONLY_DATA
diff --git a/Engine/Source/Runtime/Engine/Private/Streaming/Texture2DMipDataProvider_IO.cpp b/Engine/Source/Runtime/Engine/Private/Streaming/Texture2DMipDataProvider_IO.cpp
new file mode 100644
--- /dev/null
+++ b/Engine/Source/Runtime/Engine/Private/Streaming/Texture2DMipDataProvider_IO.cpp
@@ -0,0 +1,234 @@
+// Copyright 1998-2019 Epic Games, Inc. All Rights Reserved.FTexture2DMipDataProvider_IO
+
+/*=============================================================================
+Texture2DMipDataProvider_IO.cpp : Implementation of FTextureMipDataProvider using cooked file IO.
+=============================================================================*/
+
+#include "Texture2DMipDataProvider_IO.h"
+#include "Engine/Texture2D.h"
+#include "TextureResource.h"
+#include "HAL/PlatformFilemanager.h"
+#include "HAL/FileManager.h"
+#include "Misc/Paths.h"
+#include "Streaming/TextureStreamingHelpers.h"
+
+FTexture2DMipDataProvider_IO::FTexture2DMipDataProvider_IO(bool InPrioritizedIORequest)
+	: FTextureMipDataProvider(ETickState::Init, ETickThread::Async)
+	, bPrioritizedIORequest(InPrioritizedIORequest)
+{
+}
+
+FTexture2DMipDataProvider_IO::~FTexture2DMipDataProvider_IO()
+{
+	check(!IORequests.Num());
+}
+
+void FTexture2DMipDataProvider_IO::Init(const FTextureUpdateContext& Context, const FTextureUpdateSyncOptions& SyncOptions)
+{
+	UTexture2D* Texture2D = CastChecked<UTexture2D>(Context.Texture, ECastCheckedType::NullChecked);
+	const TIndirectArray<FTexture2DMipMap>& OwnerMips = Texture2D->GetPlatformMips();
+
+	int32 CurrentFileIndex = INDEX_NONE;
+
+	for (int32 MipIndex = Context.PendingFirstMipIndex; MipIndex < Context.CurrentFirstMipIndex; ++MipIndex)
+	{
+		const FTexture2DMipMap& OwnerMip = OwnerMips[MipIndex];
+		if (OwnerMip.BulkData.IsStoredCompressedOnDisk())
+		{
+			// Compression at the package level is no longer supported
+			continue;
+		}
+		else if (OwnerMip.BulkData.GetBulkDataSize() <= 0)
+		{
+			// Invalid bulk data size.
+			continue;
+		}
+
+		FString IOFilename;
+#if !TEXTURE2DMIPMAP_USE_COMPACT_BULKDATA
+		IOFilename = OwnerMip.BulkData.GetFilename();
+#else
+		if (!Texture2D->GetMipDataFilename(MipIndex, IOFilename))
+		{
+			continue;
+		}
+#endif
+		if (FileInfos.IsValidIndex(CurrentFileIndex))
+		{
+			FFileInfo& FileInfo = FileInfos[CurrentFileIndex];
+			if (FileInfo.IOFilename == IOFilename && FileInfo.LastMipIndex + 1 == MipIndex)
+			{
+				FileInfo.LastMipIndex = MipIndex;
+			}
+			else // Otherwise create a new entry
+			{
+				CurrentFileIndex = INDEX_NONE;
+			}
+		}
+
+		if (CurrentFileIndex == INDEX_NONE)
+		{
+			int64 IOFileOffset = 0;
+			if (GEventDrivenLoaderEnabled)
+			{
+				if (IOFilename.EndsWith(TEXT(".uasset")) || IOFilename.EndsWith(TEXT(".umap")))
+				{
+					IOFileOffset = -IFileManager::Get().FileSize(*IOFilename);
+					check(IOFileOffset < 0);
+					IOFilename = FPaths::GetBaseFilename(IOFilename, false) + TEXT(".uexp");
+					UE_LOG(LogTexture, Error, TEXT("Streaming from the .uexp file '%s' this MUST be in a ubulk instead for best performance."), *IOFilename);
+				}
+			}
+
+			CurrentFileIndex = FileInfos.AddDefaulted();
+			FFileInfo& FileInfo = FileInfos[CurrentFileIndex];
+			FileInfo.IOFilename = MoveTemp(IOFilename);
+			FileInfo.IOFileOffset = IOFileOffset;
+			FileInfo.FirstMipIndex = MipIndex;
+			FileInfo.LastMipIndex = MipIndex;
+		}
+	}
+
+	AdvanceTo(ETickState::GetMips, ETickThread::Async);
+}
+
+int32 FTexture2DMipDataProvider_IO::GetMips(
+	const FTextureUpdateContext& Context, 
+	int32 StartingMipIndex, 
+	const FTextureMipInfoArray& MipInfos, 
+	const FTextureUpdateSyncOptions& SyncOptions)
+{
+	UTexture2D* Texture2D = CastChecked<UTexture2D>(Context.Texture, ECastCheckedType::NullChecked);
+	const TIndirectArray<FTexture2DMipMap>& OwnerMips = Texture2D->GetPlatformMips();
+
+	SetAsyncFileCallback(SyncOptions);
+	check(SyncOptions.Counter && !IORequests.Num());
+	
+	IORequests.AddDefaulted(Context.CurrentFirstMipIndex);
+
+	for (FFileInfo& FileInfo : FileInfos)
+	{
+		while (StartingMipIndex >= FileInfo.FirstMipIndex && StartingMipIndex <= FileInfo.LastMipIndex && StartingMipIndex < Context.CurrentFirstMipIndex)
+		{
+			if (!MipInfos.IsValidIndex(StartingMipIndex))
+			{
+				break;
+			}
+			const FTextureMipInfo& MipInfo = MipInfos[StartingMipIndex];
+
+			if (!FileInfo.IOFileHandle)
+			{
+				FileInfo.IOFileHandle.Reset(FPlatformFileManager::Get().GetPlatformFile().OpenAsyncRead(*FileInfo.IOFilename));
+				if (!FileInfo.IOFileHandle)
+				{
+					break;
+				}
+			}
+
+			const FTexture2DMipMap& OwnerMip = OwnerMips[StartingMipIndex];
+			// If Data size is specified check compatibility for safety
+			if (MipInfo.DataSize && OwnerMip.BulkData.GetBulkDataSize() > MipInfo.DataSize)
+			{
+				break;
+			}
+
+			// Increment as we push the requests. If a requests complete immediately, then it will call the callback
+			// but that won't do anything because the tick would not try to acquire the lock since it is already locked.
+			SyncOptions.Counter->Increment();
+
+			IORequests[StartingMipIndex].Reset(FileInfo.IOFileHandle->ReadRequest(
+				OwnerMip.BulkData.GetBulkDataOffsetInFile() + FileInfo.IOFileOffset, 
+				OwnerMip.BulkData.GetBulkDataSize(), 
+				bPrioritizedIORequest ? AIOP_BelowNormal : AIOP_Low, 
+				&AsyncFileCallBack, 
+				(uint8*)MipInfo.DestData));
+
+			++StartingMipIndex;
+		}
+	}
+
+	AdvanceTo(ETickState::PollMips, ETickThread::Async);
+	return StartingMipIndex;
+}
+
+bool FTexture2DMipDataProvider_IO::PollMips(const FTextureUpdateSyncOptions& SyncOptions)
+{
+	ClearIORequests();
+	AdvanceTo(ETickState::Done, ETickThread::None);
+	return !bIORequestCancelled;
+}
+
+void FTexture2DMipDataProvider_IO::AbortPollMips() 
+{
+	for (TUniquePtr<IAsyncReadRequest>& IORequest : IORequests)
+	{
+		if (IORequest)
+		{
+			// Calling IAsyncReadRequest::cancel() here will trigger the AsyncFileCallBack and precipitate the execution of Cancel().
+			IORequest->Cancel();
+		}
+	}
+}
+
+
+void FTexture2DMipDataProvider_IO::CleanUp(const FTextureUpdateSyncOptions& SyncOptions)
+{
+	AdvanceTo(ETickState::Done, ETickThread::None);
+}
+
+void FTexture2DMipDataProvider_IO::Cancel(const FTextureUpdateSyncOptions& SyncOptions)
+{
+	ClearIORequests();
+}
+
+FTextureMipDataProvider::ETickThread FTexture2DMipDataProvider_IO::GetCancelThread() const
+{
+	return IORequests.Num() ? FTextureMipDataProvider::ETickThread::Async : FTextureMipDataProvider::ETickThread::None;
+}
+
+void FTexture2DMipDataProvider_IO::SetAsyncFileCallback(const FTextureUpdateSyncOptions& SyncOptions)
+{
+	FThreadSafeCounter* Counter = SyncOptions.Counter;
+	FTextureUpdateSyncOptions::FCallback RescheduleCallback = SyncOptions.RescheduleCallback;
+	check(Counter && RescheduleCallback);
+
+	AsyncFileCallBack = [this, Counter, RescheduleCallback](bool bWasCancelled, IAsyncReadRequest* Req)
+	{
+		// At this point task synchronization would hold the number of pending requests.
+		Counter->Decrement();
+		
+		if (bWasCancelled)
+		{
+			bIORequestCancelled = true;
+		}
+
+		if (Counter->GetValue() == 0)
+		{
+#if !UE_BUILD_SHIPPING
+			// On some platforms the IO is too fast to test cancelation requests timing issues.
+			if (FRenderAssetStreamingSettings::ExtraIOLatency > 0)
+			{
+				FPlatformProcess::Sleep(FRenderAssetStreamingSettings::ExtraIOLatency * .001f); // Slow down the streaming.
+			}
+#endif
+			RescheduleCallback();
+		}
+	};
+}
+
+void FTexture2DMipDataProvider_IO::ClearIORequests()
+{
+	for (TUniquePtr<IAsyncReadRequest>& IORequest : IORequests)
+	{
+		if (IORequest)
+		{
+			// If clearing requests not yet completed, cancel and wait.
+			if (!IORequest->PollCompletion())
+			{
+				IORequest->Cancel();
+				IORequest->WaitCompletion();
+			}
+		}
+	}
+	IORequests.Empty();
+}
diff --git a/Engine/Source/Runtime/Engine/Private/Streaming/Texture2DMipDataProvider_IO.h b/Engine/Source/Runtime/Engine/Private/Streaming/Texture2DMipDataProvider_IO.h
new file mode 100644
--- /dev/null
+++ b/Engine/Source/Runtime/Engine/Private/Streaming/Texture2DMipDataProvider_IO.h
@@ -0,0 +1,66 @@
+// Copyright 1998-2019 Epic Games, Inc. All Rights Reserved.
+
+/*=============================================================================
+Texture2DMipDataProvider_IO.h : Implementation of FTextureMipDataProvider using cooked file IO.
+=============================================================================*/
+
+#pragma once
+
+#include "CoreMinimal.h"
+#include "Streaming/TextureMipDataProvider.h"
+#include "Async/AsyncFileHandle.h"
+
+/**
+* FTexture2DMipDataProvider_IO implements FTextureMipAllocator using file IO (cooked content).
+* It support having mips stored in different files contrary to FTexture2DStreamIn_IO.
+*/
+class FTexture2DMipDataProvider_IO : public FTextureMipDataProvider
+{
+public:
+
+	FTexture2DMipDataProvider_IO(bool InPrioritizedIORequest);
+	~FTexture2DMipDataProvider_IO();
+
+	// ********************************************************
+	// ********* FTextureMipDataProvider implementation **********
+	// ********************************************************
+
+	void Init(const FTextureUpdateContext& Context, const FTextureUpdateSyncOptions& SyncOptions) final override;
+	int32 GetMips(const FTextureUpdateContext& Context, int32 StartingMipIndex, const FTextureMipInfoArray& MipInfos, const FTextureUpdateSyncOptions& SyncOptions) final override;
+	bool PollMips(const FTextureUpdateSyncOptions& SyncOptions) final override;
+	void AbortPollMips() final override;
+	void CleanUp(const FTextureUpdateSyncOptions& SyncOptions) final override;
+	void Cancel(const FTextureUpdateSyncOptions& SyncOptions) final override;
+	ETickThread GetCancelThread() const final override;
+
+protected:
+
+	// A structured with information about which file contains which mips.
+	struct FFileInfo
+	{
+		FString IOFilename;
+		TUniquePtr<IAsyncReadFileHandle> IOFileHandle;
+		int64 IOFileOffset = 0;
+		int32 FirstMipIndex = INDEX_NONE;
+		int32 LastMipIndex = INDEX_NONE;
+	};
+
+	// Pending async requests created in GetMips().
+	TArray<TUniquePtr<IAsyncReadRequest>, TInlineAllocator<MAX_TEXTURE_MIP_COUNT>> IORequests;
+	// The list of relevant files used for reading texture mips.
+	TArray<FFileInfo, TInlineAllocator<2>> FileInfos;
+
+	// Whether async read requests must be created with high priority (executes faster). 
+	bool bPrioritizedIORequest = false;
+	// Whether async read requests where cancelled for any reasons.
+	bool bIORequestCancelled = false;
+
+	// A callback to be executed once all IO pending requests are completed.
+	FAsyncFileCallBack AsyncFileCallBack;
+
+	// Helper to configure the AsyncFileCallBack.
+	void SetAsyncFileCallback(const FTextureUpdateSyncOptions& SyncOptions);
+
+	// Release / cancel any pending async file requests.
+	void ClearIORequests();
+};
diff --git a/Engine/Source/Runtime/Engine/Private/Streaming/Texture2DStreamIn_DDC.cpp b/Engine/Source/Runtime/Engine/Private/Streaming/Texture2DStreamIn_DDC.cpp
--- a/Engine/Source/Runtime/Engine/Private/Streaming/Texture2DStreamIn_DDC.cpp
+++ b/Engine/Source/Runtime/Engine/Private/Streaming/Texture2DStreamIn_DDC.cpp
@@ -28,71 +28,60 @@ static FAutoConsoleVariableRef CVarStreamingAbandonedDDCHandlePurgeFrequency(
 	ECVF_Default
 );
 
-namespace
-{
-	class FAbandonedDDCHandleManager
-	{
-	public:
-		void Add(uint32 InHandle);
-		void Purge();
-	private:
-		TArray<uint32> Handles;	
-		FCriticalSection CS;
-		uint32 TotalAdd = 0;
-	};
+// ******************************************
+// ******* FAbandonedDDCHandleManager *******
+// ******************************************
 
+FAbandonedDDCHandleManager GAbandonedDDCHandleManager;
 
-	void FAbandonedDDCHandleManager::Add(uint32 InHandle)
-	{
-		check(InHandle);
+void FAbandonedDDCHandleManager::Add(uint32 InHandle)
+{
+	check(InHandle);
 
-		bool bPurge = false;
-		{
-			FScopeLock ScopeLock(&CS);
-			Handles.Add(InHandle);
+	bool bPurge = false;
+	{
+		FScopeLock ScopeLock(&CS);
+		Handles.Add(InHandle);
 
-			++TotalAdd;
-			bPurge = (TotalAdd % GStreamingAbandonedDDCHandlePurgeFrequency == 0);
-		}
+		++TotalAdd;
+		bPurge = (TotalAdd % GStreamingAbandonedDDCHandlePurgeFrequency == 0);
+	}
 
-		if (bPurge)
-		{
-			Purge();
-		}
+	if (bPurge)
+	{
+		Purge();
 	}
+}
 
-	void FAbandonedDDCHandleManager::Purge()
+void FAbandonedDDCHandleManager::Purge()
+{
+	TArray<uint32> TempHandles;	
 	{
-		TArray<uint32> TempHandles;	
-		{
-			FScopeLock ScopeLock(&CS);
-			FMemory::Memswap(&TempHandles, &Handles, sizeof(TempHandles));
-		}
+		FScopeLock ScopeLock(&CS);
+		FMemory::Memswap(&TempHandles, &Handles, sizeof(TempHandles));
+	}
 
-		FDerivedDataCacheInterface& DDCInterface = GetDerivedDataCacheRef();
-		TArray<uint8> Data;
+	FDerivedDataCacheInterface& DDCInterface = GetDerivedDataCacheRef();
+	TArray<uint8> Data;
 
-		for (int32 Index = 0; Index < TempHandles.Num(); ++Index)
+	for (int32 Index = 0; Index < TempHandles.Num(); ++Index)
+	{
+		uint32 Handle = TempHandles[Index];
+		if (DDCInterface.PollAsynchronousCompletion(Handle))
 		{
-			uint32 Handle = TempHandles[Index];
-			if (DDCInterface.PollAsynchronousCompletion(Handle))
-			{
-				DDCInterface.GetAsynchronousResults(Handle, Data);
-				Data.Reset();
-
-				TempHandles.RemoveAtSwap(Index);
-				--Index;
-			}
-		}
+			DDCInterface.GetAsynchronousResults(Handle, Data);
+			Data.Reset();
 
-		if (TempHandles.Num())
-		{
-			FScopeLock ScopeLock(&CS);
-			Handles.Append(TempHandles);
+			TempHandles.RemoveAtSwap(Index);
+			--Index;
 		}
 	}
 
-	FAbandonedDDCHandleManager GAbandonedDDCHandleManager;
+	if (TempHandles.Num())
+	{
+		FScopeLock ScopeLock(&CS);
+		Handles.Append(TempHandles);
+	}
 }
 
 void PurgeAbandonedDDCHandles()
@@ -100,6 +89,10 @@ void PurgeAbandonedDDCHandles()
 	GAbandonedDDCHandleManager.Purge();
 }
 
+// ******************************************
+// ********* FTexture2DStreamIn_DDC *********
+// ******************************************
+
 FTexture2DStreamIn_DDC::FTexture2DStreamIn_DDC(UTexture2D* InTexture, int32 InRequestedMips)
 	: FTexture2DStreamIn(InTexture, InRequestedMips)
 	, bDDCIsInvalid(false)
diff --git a/Engine/Source/Runtime/Engine/Private/Streaming/Texture2DStreamIn_DDC.h b/Engine/Source/Runtime/Engine/Private/Streaming/Texture2DStreamIn_DDC.h
--- a/Engine/Source/Runtime/Engine/Private/Streaming/Texture2DStreamIn_DDC.h
+++ b/Engine/Source/Runtime/Engine/Private/Streaming/Texture2DStreamIn_DDC.h
@@ -45,4 +45,22 @@ protected:
 	void DoLoadNewMipsFromDDC(const FContext& Context);
 };
 
+/**
+* This class provides a helper to release DDC handles that haven't been waited for.
+* This is to get around limitations of FDerivedDataCacheInterface.
+*/
+
+class FAbandonedDDCHandleManager
+{
+public:
+	void Add(uint32 InHandle);
+	void Purge();
+private:
+	TArray<uint32> Handles;	
+	FCriticalSection CS;
+	uint32 TotalAdd = 0;
+};
+
+extern FAbandonedDDCHandleManager GAbandonedDDCHandleManager;
+
 #endif // WITH_EDITORONLY_DATA
diff --git a/Engine/Source/Runtime/Engine/Private/Streaming/TextureMipAllocator.h b/Engine/Source/Runtime/Engine/Private/Streaming/TextureMipAllocator.h
new file mode 100644
--- /dev/null
+++ b/Engine/Source/Runtime/Engine/Private/Streaming/TextureMipAllocator.h
@@ -0,0 +1,100 @@
+// Copyright 1998-2019 Epic Games, Inc. All Rights Reserved.
+
+/*=============================================================================
+TextureMipAllocator.h: Base class for implementating a mip allocation strategy used by FTextureStreamIn.
+=============================================================================*/
+
+#pragma once
+
+#include "CoreMinimal.h"
+#include "Streaming/TextureMipDataProvider.h"
+
+/**
+* FTextureMipAllocator defines the update steps and interface to implement the mip allocation strategy
+* used in FTextureStreamIn. It allows to decouples the texture creation steps from the texture update.
+* Typical implementations are using RHIAsyncCreateTexture2D, RHIAsyncReallocateTexture2D or virtual textures.
+*/
+class FTextureMipAllocator
+{
+public:
+
+	enum class ETickState : uint32
+	{
+		AllocateMips,	// Call AllocateMips(), which must provide destination memory for the mips to be written to.
+		FinalizeMips,	// Call FinalizeMips(), which must release temporary buffers and update the texture asset with the new mips.
+		Done			// Nothing to do.
+	};
+
+	enum class ETickThread : uint32
+	{
+		Async,			// Tick must run on an async thread.
+		Render,			// Tick must run on the renderthread.
+		None,			// No tick must run, only valid in ETickState::Done. 
+	};
+
+	// Constructor, defining the first tick step and thread.
+	FTextureMipAllocator(ETickState InTickState, ETickThread InTickThread) : NextTickState(InTickState), NextTickThread(InTickThread) {}
+	virtual ~FTextureMipAllocator() {}
+
+	// Get the next tick state and thread for the mip allocator. Used by FTextureStreamIn to schedule correctly the update between FTextureMipAllocator and FTextureMipDataProvider.
+	FORCEINLINE ETickState GetNextTickState() const { return NextTickState; }
+	FORCEINLINE ETickThread GetNextTickThread() const { return NextTickThread; }
+
+	/**
+	* Allocate the FTextureMipInfoArray and configure the entries for the new (streamed in) mips. 
+	* In particular, the FTextureMipAllocator must set valid FTextureMipInfo::DestData so that the mip data provider can write the mip data to.
+	*
+	* @param Context - An update context constant throughout the FTextureStreamIn update. Gives things like which texture asset is updated and what mips are streamed in.
+	* @param OutMipInfos - The array of FTextureMipInfo that must be correctly configured so that FTextureMipDataProvider know what is being requested and where to write data to.
+	* @param SyncOptions - Different sync options to control when the next tick of FTextureStreamIn can be scheduled.
+	*
+	* Return true unless the texture update needs to be aborted because of invalid data or setup.
+	*/
+	virtual bool AllocateMips(const FTextureUpdateContext& Context, FTextureMipInfoArray& OutMipInfos, const FTextureUpdateSyncOptions& SyncOptions) = 0;
+	
+	/**
+	* Upload the mip data to the new mips and update the texture asset. Update is completed after this step and can not be cancelled afterward..
+	*
+	* @param Context - An update context constant throughout the FTextureStreamIn update. Gives things like which texture asset is updated and what mips are streamed in.
+	* @param SyncOptions - Different sync options to control when the next tick of FTextureStreamIn can be scheduled.
+	*
+	* Return true unless the texture update needs to be aborted because of invalid data or setup.
+	*/
+	virtual bool FinalizeMips(const FTextureUpdateContext& Context, const FTextureUpdateSyncOptions& SyncOptions) = 0;
+
+	/**
+	* Cancel the progression and release any temporary resources. 
+	* Called within the FTextureStreamIn update when the stream in request is aborted or cannot complete correctly
+	* The Cancel() function is called on the thread returned by GetCancelThread().
+	*
+	* @param SyncOptions - Different sync options to control when the next tick of FTextureStreamIn can be scheduled.
+	*/
+	virtual void Cancel(const FTextureUpdateSyncOptions& SyncOptions) = 0;
+	
+	// Returns on which thread Cancel() must be called in the texture update (from FTextureStreamIn) to release ressources safely and correctly. ETickThread::None when ready to delete this.
+	virtual ETickThread GetCancelThread() const = 0;
+
+	// A little helper to help FTextureStreamIn figure out the current first mip of the UTexture resource (must access derived classes).
+	virtual int32 GetCurrentFirstMip(UTexture* Texture) const = 0;
+
+protected:
+
+	// Helper to set the next tick state and thread. Validates that the progression is coherent.
+	FORCEINLINE void AdvanceTo(ETickState InState, ETickThread InThread)
+	{
+		// State must advance or at least stay to the same step.
+		check((uint32)NextTickState <= (uint32)InState);
+		// Thread must be valid unless in done state.
+		check((InState == ETickState::Done) == (InThread == ETickThread::None));
+
+		NextTickState = InState;
+		NextTickThread = InThread;
+	}
+
+private:
+
+	// The next tick function that should be called.
+	ETickState NextTickState = ETickState::Done;
+	// The next tick thread from where to call the tick function.
+	ETickThread NextTickThread = ETickThread::None;
+};
diff --git a/Engine/Source/Runtime/Engine/Private/Streaming/TextureMipDataProviderFactory.cpp b/Engine/Source/Runtime/Engine/Private/Streaming/TextureMipDataProviderFactory.cpp
new file mode 100644
--- /dev/null
+++ b/Engine/Source/Runtime/Engine/Private/Streaming/TextureMipDataProviderFactory.cpp
@@ -0,0 +1,12 @@
+// Copyright 1998-2019 Epic Games, Inc. All Rights Reserved.
+
+/*=============================================================================
+	TextureMipDataProviderFactory.cpp: Implementation of TextureMipDataProviderFactory.
+=============================================================================*/
+
+#include "Engine/TextureMipDataProviderFactory.h"
+
+UTextureMipDataProviderFactory::UTextureMipDataProviderFactory(const FObjectInitializer& ObjectInitializer)
+	: Super(ObjectInitializer)
+{
+}
diff --git a/Engine/Source/Runtime/Engine/Private/Streaming/TextureStreamIn.cpp b/Engine/Source/Runtime/Engine/Private/Streaming/TextureStreamIn.cpp
new file mode 100644
--- /dev/null
+++ b/Engine/Source/Runtime/Engine/Private/Streaming/TextureStreamIn.cpp
@@ -0,0 +1,546 @@
+// Copyright 1998-2019 Epic Games, Inc. All Rights Reserved.
+
+/*=============================================================================
+FTextureStreamIn.cpp : Implement a generic texture stream in strategy.
+=============================================================================*/
+
+#include "Streaming/TextureStreamIn.h"
+#include "Engine/Texture.h"
+#include "Streaming/RenderAssetUpdate.inl"
+
+template class TRenderAssetUpdate<FTextureUpdateContext>;
+
+FTextureUpdateContext::FTextureUpdateContext(UTexture* InTexture, EThreadType InCurrentThread) 
+	: Texture(InTexture)
+	, CurrentThread(InCurrentThread) 
+{
+}
+
+FTextureUpdateContext::FTextureUpdateContext(UStreamableRenderAsset* InTexture, EThreadType InCurrentThread)
+	: Texture(CastChecked<UTexture>(InTexture, ECastCheckedType::NullChecked))
+	, CurrentThread(InCurrentThread) 
+{
+}
+
+bool FTextureStreamIn::IsSameThread(FTextureMipAllocator::ETickThread TickThread, int32 TaskThread)
+{
+	if (TaskThread == TT_Async)
+	{
+		return TickThread == FTextureMipAllocator::ETickThread::Async;
+	}
+	else // Expected to be called from either the async thread or the renderthread.
+	{
+		check(TaskThread == TT_Render);
+		return TickThread == FTextureMipAllocator::ETickThread::Render;
+	}
+}
+
+bool FTextureStreamIn::IsSameThread(FTextureMipDataProvider::ETickThread TickThread, int32 TaskThread)
+{
+	if (TaskThread == TT_Async)
+	{
+		return TickThread == FTextureMipDataProvider::ETickThread::Async;
+	}
+	else // Expected to be called from either the async thread or the renderthread.
+	{
+		check(TaskThread == TT_Render);
+		return TickThread == FTextureMipDataProvider::ETickThread::Render;
+	}
+}
+
+FTextureStreamIn::FTextureStreamIn(
+	UTexture* InTexture, 
+	int32 InRequestedMips, 
+	FTextureMipAllocator* InMipAllocator, 
+	FTextureMipDataProvider* InCustomMipDataProvider, 
+	FTextureMipDataProvider* InDefaultMipDataProvider
+)
+	: TRenderAssetUpdate<FTextureUpdateContext>(InTexture, InRequestedMips)
+{
+	check(InMipAllocator);
+	CurrentFirstMip = InMipAllocator->GetCurrentFirstMip(InTexture);
+	StartingMipIndex = PendingFirstMip;
+
+	// Init the allocator and provider.
+	MipAllocator.Reset(InMipAllocator);
+	if (InCustomMipDataProvider)
+	{
+		MipDataProviders.Emplace(InCustomMipDataProvider);
+	}
+	check(InDefaultMipDataProvider);
+	MipDataProviders.Emplace(InDefaultMipDataProvider);
+
+	// Init sync options
+	SyncOptions.bSnooze = &bDeferExecution;
+	SyncOptions.Counter = &TaskSynchronization;
+	SyncOptions.RescheduleCallback = [this]() 
+	{ 
+		if (!IsLocked())
+		{
+			Tick(TT_None); 
+		}
+	};
+
+	// Schedule the first update step.
+	FContext Context(InTexture, TT_None);
+	InitContext(Context);
+	const EThreadType NextThread = GetMipDataProviderThread(FTextureMipDataProvider::ETickState::Init);
+	if (NextThread != TT_None)
+	{
+		PushTask(Context, NextThread, SRA_UPDATE_CALLBACK(InitMipDataProviders), TT_None, nullptr);
+	}
+	else // Otherwise nothing to execute but schedule next phase (edge case).
+	{
+		InitMipDataProviders(Context);
+	}
+}
+
+FTextureStreamIn::~FTextureStreamIn()
+{
+}
+
+FRenderAssetUpdate::EThreadType FTextureStreamIn::GetMipDataProviderThread(FTextureMipDataProvider::ETickState TickState) const
+{
+	if (!IsCancelled())
+	{
+		for (const TUniquePtr<FTextureMipDataProvider>& MipDataProvider : MipDataProviders)
+		{
+			check(MipDataProvider);
+			if (MipDataProvider->GetNextTickState() == TickState)
+			{
+				switch (MipDataProvider->GetNextTickThread())
+				{
+				case FTextureMipDataProvider::ETickThread::Async:
+					return TT_Async;
+				case FTextureMipDataProvider::ETickThread::Render:
+					return TT_Render;
+				default:
+					check(false);
+				}
+			}
+		}
+	}
+	return TT_None;
+}
+
+FRenderAssetUpdate::EThreadType FTextureStreamIn::GetMipAllocatorThread(FTextureMipAllocator::ETickState TickState) const
+{
+	check(MipAllocator);
+	if (!IsCancelled() && MipAllocator->GetNextTickState() == TickState)
+	{
+		switch (MipAllocator->GetNextTickThread())
+		{
+		case FTextureMipAllocator::ETickThread::Async:
+			return TT_Async;
+		case FTextureMipAllocator::ETickThread::Render:
+			return TT_Render;
+		default:
+			check(false);
+		}
+	}
+	return TT_None;
+}
+
+FRenderAssetUpdate::EThreadType FTextureStreamIn::GetCancelThread() const
+{
+	// The mip data provider must be cancelled before the mip allocator since it writes into the mip data.
+	for (const TUniquePtr<FTextureMipDataProvider>& MipDataProvider : MipDataProviders)
+	{
+		check(MipDataProvider);
+		switch (MipDataProvider->GetCancelThread())
+		{
+		case FTextureMipDataProvider::ETickThread::Async:
+			return TT_Async;
+		case FTextureMipDataProvider::ETickThread::Render:
+			return TT_Render;
+		}
+	}
+
+	// If the mip data provider is fully cancelled, then cancel the mip allocator.
+	if (MipAllocator)
+	{
+		switch (MipAllocator->GetCancelThread())
+		{
+		case FTextureMipAllocator::ETickThread::Async:
+			return TT_Async;
+		case FTextureMipAllocator::ETickThread::Render:
+			return TT_Render;
+		}
+	}
+
+
+	// If both are cancelled, then run the final cleanup on the async thread.
+	return TT_Async;
+}
+
+void FTextureStreamIn::InitContext(FContext& Context) const
+{
+	Context.PendingFirstMipIndex = PendingFirstMip;
+	Context.CurrentFirstMipIndex = CurrentFirstMip;
+	Context.NumRequestedMips = RequestedMips;
+}
+
+// ****************************
+// **** Update Steps Work *****
+// ****************************
+
+void FTextureStreamIn::DoInitMipDataProviders(const FContext& Context)
+{
+	for (TUniquePtr<FTextureMipDataProvider>& MipDataProvider : MipDataProviders)
+	{
+		check(MipDataProvider);
+		if (IsSameThread(MipDataProvider->GetNextTickThread(), Context.CurrentThread))
+		{
+			MipDataProvider->Init(Context, SyncOptions);
+		}
+	}
+}
+
+bool FTextureStreamIn::DoAllocateNewMips(const FContext& Context)
+{
+	check(MipAllocator && IsSameThread(MipAllocator->GetNextTickThread(), Context.CurrentThread));
+	return MipAllocator->AllocateMips(Context, MipInfos, SyncOptions);
+}
+
+void FTextureStreamIn::DoGetMipData(const FContext& Context)
+{
+	for (TUniquePtr<FTextureMipDataProvider>& MipDataProvider : MipDataProviders)
+	{
+		check(MipDataProvider);
+		if (IsSameThread(MipDataProvider->GetNextTickThread(), Context.CurrentThread))
+		{
+			StartingMipIndex = MipDataProvider->GetMips(Context, StartingMipIndex, MipInfos, SyncOptions);
+		}
+	}
+}
+
+bool FTextureStreamIn::DoPollMipData(const FContext& Context)
+{
+	for (TUniquePtr<FTextureMipDataProvider>& MipDataProvider : MipDataProviders)
+	{
+		check(MipDataProvider);
+		if (IsSameThread(MipDataProvider->GetNextTickThread(), Context.CurrentThread))
+		{
+			if (!MipDataProvider->PollMips(SyncOptions))
+			{
+				return false;
+			}
+		}
+	}
+	return true;
+}
+
+bool FTextureStreamIn::DoFinalizeNewMips(const FContext& Context)
+{
+	check(MipAllocator && IsSameThread(MipAllocator->GetNextTickThread(), Context.CurrentThread));
+	return MipAllocator->FinalizeMips(Context, SyncOptions);
+}
+
+void FTextureStreamIn::DoCleanUpMipDataProviders(const FContext& Context)
+{
+	for (TUniquePtr<FTextureMipDataProvider>& MipDataProvider : MipDataProviders)
+	{
+		check(MipDataProvider);
+		if (IsSameThread(MipDataProvider->GetNextTickThread(), Context.CurrentThread))
+		{
+			MipDataProvider->CleanUp(SyncOptions);
+		}
+	}
+}
+
+// ****************************
+// ******* Update Steps *******
+// ****************************
+
+void FTextureStreamIn::InitMipDataProviders(const FContext& Context)
+{
+	DECLARE_SCOPE_CYCLE_COUNTER(TEXT("FTextureStreamIn::InitMipDataProviders"), STAT_TextureStreamIn_InitMipDataProviders, STATGROUP_StreamingDetails);
+
+	// Execute
+	DoInitMipDataProviders(Context);
+
+	// Schedule the next update step.
+	EThreadType NextThread = GetMipDataProviderThread(FTextureMipDataProvider::ETickState::Init);
+	if (NextThread != TT_None) // Loop on this state.
+	{
+		PushTask(Context, NextThread, SRA_UPDATE_CALLBACK(InitMipDataProviders), GetCancelThread(), SRA_UPDATE_CALLBACK(Cancel));
+	}
+	else
+	{
+		NextThread = GetMipAllocatorThread(FTextureMipAllocator::ETickState::AllocateMips);
+		if (NextThread != TT_None)
+		{
+			PushTask(Context, NextThread, SRA_UPDATE_CALLBACK(AllocateNewMips), GetCancelThread(), SRA_UPDATE_CALLBACK(Cancel));
+		}
+		else // Otherwise if it is impossible to allocate the new mips, abort.
+		{
+			MarkAsCancelled();
+			PushTask(Context, TT_None, nullptr, GetCancelThread(), SRA_UPDATE_CALLBACK(Cancel));
+		}
+	}
+}
+
+void FTextureStreamIn::AllocateNewMips(const FContext& Context)
+{
+	DECLARE_SCOPE_CYCLE_COUNTER(TEXT("FTextureStreamIn::AllocateNewMips"), STAT_TextureStreamIn_AllocateNewMips, STATGROUP_StreamingDetails);
+
+	// Execute
+	if (!DoAllocateNewMips(Context))
+	{
+		MarkAsCancelled();
+	}
+
+	// Schedule the next update step.
+	EThreadType NextThread = GetMipAllocatorThread(FTextureMipAllocator::ETickState::AllocateMips);
+	if (NextThread != TT_None) // Loop on this state.
+	{
+		PushTask(Context, NextThread, SRA_UPDATE_CALLBACK(AllocateNewMips), GetCancelThread(), SRA_UPDATE_CALLBACK(Cancel));
+	}
+	else
+	{
+		// Check that all mips have data available. If not, something went wrong in DoAllocateNewMips() and everything must be cancelled.
+		for (int32 MipIndex = Context.PendingFirstMipIndex; MipIndex < Context.CurrentFirstMipIndex; ++MipIndex)
+		{
+			if (!MipInfos.IsValidIndex(MipIndex) || !MipInfos[MipIndex].DestData)
+			{
+				MarkAsCancelled();
+			}
+		}
+
+		NextThread = GetMipDataProviderThread(FTextureMipDataProvider::ETickState::GetMips);
+		if (NextThread != TT_None)
+		{
+			PushTask(Context, NextThread, SRA_UPDATE_CALLBACK(GetMipData), GetCancelThread(), SRA_UPDATE_CALLBACK(Cancel));
+		}
+		else // Otherwise if it is impossible the get the mip data abort.
+		{
+			MarkAsCancelled();
+			PushTask(Context, TT_None, nullptr, GetCancelThread(), SRA_UPDATE_CALLBACK(Cancel));
+		}
+	}
+}
+
+void FTextureStreamIn::GetMipData(const FContext& Context)
+{
+	DECLARE_SCOPE_CYCLE_COUNTER(TEXT("FTextureStreamIn::GetMipData"), STAT_TextureStreamIn_GetMipData, STATGROUP_StreamingDetails);
+
+	// Execute
+	DoGetMipData(Context);
+
+	// Schedule the next update step.
+	EThreadType NextThread = GetMipDataProviderThread(FTextureMipDataProvider::ETickState::GetMips);
+	if (NextThread != TT_None) // Loop on this state.
+	{
+		PushTask(Context, NextThread, SRA_UPDATE_CALLBACK(GetMipData), GetCancelThread(), SRA_UPDATE_CALLBACK(Cancel));
+	}
+	else
+	{
+		// All mips must have been updated correctly.
+		if (StartingMipIndex != CurrentFirstMip)
+		{
+			MarkAsCancelled();
+		}
+
+		if (!IsCancelled())
+		{
+			NextThread = GetMipDataProviderThread(FTextureMipDataProvider::ETickState::PollMips);
+			// All mips must be handled before moving to next stage.
+			if (NextThread != TT_None)
+			{
+				bIsPollingMipData = true;
+				PushTask(Context, NextThread, SRA_UPDATE_CALLBACK(PollMipData), GetCancelThread(), SRA_UPDATE_CALLBACK(Cancel));
+			}
+			else // Otherwise nothing to execute but schedule next phase (edge case).
+			{
+				PollMipData(Context);
+			}
+		}
+		else
+		{
+			PushTask(Context, TT_None, nullptr, GetCancelThread(), SRA_UPDATE_CALLBACK(Cancel));
+		}
+	}
+}
+
+void FTextureStreamIn::PollMipData(const FContext& Context)
+{
+	DECLARE_SCOPE_CYCLE_COUNTER(TEXT("FTextureStreamIn::PollMipData"), STAT_TextureStreamIn_PollMipData, STATGROUP_StreamingDetails);
+
+	// Execute
+	if (!DoPollMipData(Context))
+	{
+		MarkAsCancelled();
+	}
+
+	// Schedule the next update step.
+	EThreadType NextThread = GetMipDataProviderThread(FTextureMipDataProvider::ETickState::PollMips);
+	if (NextThread != TT_None) // Loop on this state.
+	{
+		PushTask(Context, NextThread, SRA_UPDATE_CALLBACK(PollMipData), GetCancelThread(), SRA_UPDATE_CALLBACK(Cancel));
+	}
+	else
+	{
+		bIsPollingMipData = false;
+
+		if (!IsCancelled())
+		{
+			NextThread = GetMipAllocatorThread(FTextureMipAllocator::ETickState::FinalizeMips);
+			// All mips must be handled before moving to next stage.
+			if (NextThread != TT_None)
+			{
+				PushTask(Context, NextThread, SRA_UPDATE_CALLBACK(FinalizeNewMips), GetCancelThread(), SRA_UPDATE_CALLBACK(Cancel));
+			}
+			else // Otherwise nothing to execute but schedule next phase (edge case).
+			{
+				FinalizeNewMips(Context);
+			}
+		}
+		else
+		{
+			PushTask(Context, TT_None, nullptr, GetCancelThread(), SRA_UPDATE_CALLBACK(Cancel));
+		}
+	}
+}
+
+void FTextureStreamIn::FinalizeNewMips(const FContext& Context)
+{
+	DECLARE_SCOPE_CYCLE_COUNTER(TEXT("FTextureStreamIn::FinalizeNewMips"), STAT_TextureStreamIn_FinalizeNewMips, STATGROUP_StreamingDetails);
+
+	// Execute
+	if (!DoFinalizeNewMips(Context))
+	{
+		MarkAsCancelled();
+	}
+
+	// Schedule the next update step.
+	EThreadType NextThread = GetMipAllocatorThread(FTextureMipAllocator::ETickState::FinalizeMips);
+	if (NextThread != TT_None) // Loop on this state.
+	{
+		PushTask(Context, NextThread, SRA_UPDATE_CALLBACK(FinalizeNewMips), GetCancelThread(), SRA_UPDATE_CALLBACK(Cancel));
+	}
+	else if (!IsCancelled())
+	{
+		// Release the mip allocator.
+		MipAllocator.Reset();
+		MipInfos.Empty();
+
+		NextThread = GetMipDataProviderThread(FTextureMipDataProvider::ETickState::CleanUp);
+		// All mips must be handled before moving to next stage.
+		if (NextThread != TT_None)
+		{
+			PushTask(Context, NextThread, SRA_UPDATE_CALLBACK(CleanUpMipDataProviders), GetCancelThread(), SRA_UPDATE_CALLBACK(Cancel));
+		}
+		else // Otherwise nothing to execute but schedule next phase (edge case).
+		{
+			CleanUpMipDataProviders(Context);
+		}
+	}
+	else
+	{
+		PushTask(Context, TT_None, nullptr, GetCancelThread(), SRA_UPDATE_CALLBACK(Cancel));
+	}
+}
+
+void FTextureStreamIn::CleanUpMipDataProviders(const FContext& Context)
+{
+	DECLARE_SCOPE_CYCLE_COUNTER(TEXT("FTextureStreamIn::CleanUpMipDataProviders"), STAT_TextureStreamIn_CleanUpMipDataProviders, STATGROUP_StreamingDetails);
+
+	// Execute
+	DoCleanUpMipDataProviders(Context);
+
+	// Schedule the next update step.
+	EThreadType NextThread = GetMipDataProviderThread(FTextureMipDataProvider::ETickState::CleanUp);
+	if (NextThread != TT_None) // Loop on this state.
+	{
+		PushTask(Context, NextThread, SRA_UPDATE_CALLBACK(CleanUpMipDataProviders), GetCancelThread(), SRA_UPDATE_CALLBACK(Cancel));
+	}
+	else
+	{
+		// Release the mip data providers.
+		MipDataProviders.Empty();
+	}
+}
+
+void FTextureStreamIn::Cancel(const FContext& Context)
+{
+	DECLARE_SCOPE_CYCLE_COUNTER(TEXT("FTextureStreamIn::Cancel"), STAT_TextureStreamIn_Cancel, STATGROUP_StreamingDetails);
+
+	EThreadType NextThread = TT_None;
+
+	// The mip data provider must be cancelled before the mip allocator since it writes into the mip data.
+	for (TUniquePtr<FTextureMipDataProvider>& MipDataProvider : MipDataProviders)
+	{
+		check(MipDataProvider);
+		if (IsSameThread(MipDataProvider->GetCancelThread(), Context.CurrentThread))
+		{
+			MipDataProvider->Cancel(SyncOptions);
+		}
+		switch (MipDataProvider->GetCancelThread())
+		{
+		case FTextureMipDataProvider::ETickThread::Async:
+			NextThread = TT_Async;
+			break;
+		case FTextureMipDataProvider::ETickThread::Render:
+			NextThread = TT_Render;
+			break;
+		}
+	}
+	if (NextThread != TT_None)
+	{
+		PushTask(Context, TT_None, nullptr, NextThread, SRA_UPDATE_CALLBACK(Cancel));
+		return;
+	}
+	MipDataProviders.Empty();
+
+	// Then cancel the mip allocator.
+	if (MipAllocator)
+	{
+		if (IsSameThread(MipAllocator->GetCancelThread(), Context.CurrentThread))
+		{
+			MipAllocator->Cancel(SyncOptions);
+		}
+		switch (MipAllocator->GetCancelThread())
+		{
+		case FTextureMipAllocator::ETickThread::Async:
+			NextThread = TT_Async;
+			break;
+		case FTextureMipAllocator::ETickThread::Render:
+			NextThread = TT_Render;
+			break;
+		}
+		if (NextThread != TT_None)
+		{
+			PushTask(Context, TT_None, nullptr, NextThread, SRA_UPDATE_CALLBACK(Cancel));
+			return;
+		}
+		MipAllocator.Reset();
+		MipInfos.Empty();
+	}
+}
+
+void FTextureStreamIn::Abort()
+{
+	MarkAsCancelled();
+
+	if (bIsPollingMipData)
+	{
+		// Prevent the update from being considered done before this is finished.
+		// By checking that it was not already cancelled, we make sure this doesn't get called twice.
+		(new FAsyncAbortPollMipsTask(this))->StartBackgroundTask();
+	}
+}
+
+void FTextureStreamIn::FAbortPollMipsTask::DoWork()
+{
+	check(PendingUpdate);
+	
+	// Acquire the lock of this object in order to cancel any pending IO.
+	// If the object is currently being ticked, wait.
+	const ETaskState PreviousTaskState = PendingUpdate->DoLock();
+	for (TUniquePtr<FTextureMipDataProvider>& MipDataProvider : PendingUpdate->MipDataProviders)
+	{
+		check(MipDataProvider);
+		MipDataProvider->AbortPollMips();
+	}
+	PendingUpdate->DoUnlock(PreviousTaskState);
+}
diff --git a/Engine/Source/Runtime/Engine/Private/Streaming/TextureStreamIn.h b/Engine/Source/Runtime/Engine/Private/Streaming/TextureStreamIn.h
new file mode 100644
--- /dev/null
+++ b/Engine/Source/Runtime/Engine/Private/Streaming/TextureStreamIn.h
@@ -0,0 +1,112 @@
+// Copyright 1998-2019 Epic Games, Inc. All Rights Reserved.
+
+/*=============================================================================
+FTextureStreamIn.h: Implement a generic texture stream in strategy.
+=============================================================================*/
+
+#pragma once
+
+#include "CoreMinimal.h"
+#include "RenderAssetUpdate.h"
+#include "TextureMipAllocator.h"
+#include "Streaming/TextureMipDataProvider.h"
+
+class UTexture;
+
+/**
+* FTextureStreamIn implements a generic framework to stream in textures while decoupling
+* the source of the mip data (FTextureMipDataProvider) and the creation steps (FTextureMipAllocator) of the new texture.
+* It is also texture type independent so it can be used as a base class for any streaming implementation.
+*
+* The engine uses some private implementations to handle DDC and cooked content files, as well as the different texture creation strategy,
+* but allows for custom implementations of FTextureMipDataProvider (which is an ENGINE_API) through UTextureMipDataProviderFactory (of type UAssetUserData). 
+* Those implementation can choose to override the default mip content for specific mips (see FTextureMipDataProvider::GetMips())
+*/
+class FTextureStreamIn : public TRenderAssetUpdate<FTextureUpdateContext>
+{
+public:
+
+	FTextureStreamIn(UTexture* InTexture, int32 InRequestedMips, FTextureMipAllocator* InMipAllocator, FTextureMipDataProvider* InCustomMipDataProvider, FTextureMipDataProvider* InDefaultMipDataProvider);
+	~FTextureStreamIn();
+
+	// Route FTextureMipDataProvider::Init()
+	void InitMipDataProviders(const FContext& Context);
+	
+	// Route FTextureMipAllocator::AllocateMips()
+	void AllocateNewMips(const FContext& Context);
+	
+	// Route FTextureMipDataProvider::GetMips()
+	void GetMipData(const FContext& Context);
+
+	// Route FTextureMipDataProvider::PollMips()
+	void PollMipData(const FContext& Context);
+
+	// Route FTextureMipAllocator::FinalizeMips()
+	void FinalizeNewMips(const FContext& Context);
+
+	// Route FTextureMipDataProvider::CleanUp()
+	void CleanUpMipDataProviders(const FContext& Context);
+
+	// Execute the cancellation for both MipAllocator and MipDataProviders.
+	void Cancel(const FContext& Context);
+
+	// Start an async task to cancel pending IO requests.
+	void Abort() final override;
+
+protected:
+
+	// The current first mip. Used in InitContext()
+	int32 CurrentFirstMip = INDEX_NONE;
+
+	// The first mip to be processed for DoGetMipData(). Starts at PendingFirstMip and goes up to CurrentFirstMip.
+	int32 StartingMipIndex = INDEX_NONE;
+
+	TUniquePtr<FTextureMipAllocator> MipAllocator;
+	FTextureMipDataProviderArray MipDataProviders;
+
+	FTextureUpdateSyncOptions SyncOptions;
+	FTextureMipInfoArray MipInfos;
+
+	// Set additional context parameters (like CurrentFirstMipIndex) so that they become accessible to FTextureMipAllocator and FTextureMipDataProvider.
+	void InitContext(FContext& Context) const final override;
+
+	// The actual implementation of the update steps which redirect to FTextureMipAllocator and FTextureMipDataProvider.
+	void DoInitMipDataProviders(const FContext& Context);
+	bool DoAllocateNewMips(const FContext& Context);
+	void DoGetMipData(const FContext& Context);
+	bool DoPollMipData(const FContext& Context);
+	bool DoFinalizeNewMips(const FContext& Context);
+	void DoCleanUpMipDataProviders(const FContext& Context);
+
+	EThreadType GetMipDataProviderThread(FTextureMipDataProvider::ETickState TickState) const; 
+	EThreadType GetMipAllocatorThread(FTextureMipAllocator::ETickState TickState) const;
+
+	EThreadType GetCancelThread() const;
+
+private:
+
+	// An async task that will call FTextureMipDataProvider::AbortPollMips() on each mip data provider upon cancellation.
+	class FAbortPollMipsTask : public FNonAbandonableTask
+	{
+	public:
+		FAbortPollMipsTask(FTextureStreamIn* InPendingUpdate) : PendingUpdate(InPendingUpdate) {}
+		void DoWork();
+		FORCEINLINE TStatId GetStatId() const
+		{
+			RETURN_QUICK_DECLARE_CYCLE_STAT(FAbortPollMipsTask, STATGROUP_ThreadPoolAsyncTasks);
+		}
+	protected:
+		TRefCountPtr<FTextureStreamIn> PendingUpdate;
+	};
+
+	bool bIsPollingMipData = false; // Used to prevent indirecting any pointer in the abort task.
+	typedef FAutoDeleteAsyncTask<FAbortPollMipsTask> FAsyncAbortPollMipsTask;
+	friend class FAbortPollMipsTask;
+
+	// Whether FTextureMipAllocator::ETickThread is the same thread as EThreadType (as int32).
+	FORCEINLINE static bool IsSameThread(FTextureMipAllocator::ETickThread TickThread, int32 TaskThread);
+
+	// Whether FTextureMipDataProvider::ETickThread is the same thread as EThreadType (as int32).
+	FORCEINLINE static bool IsSameThread(FTextureMipDataProvider::ETickThread TickThread, int32 TaskThread);
+};
+
diff --git a/Engine/Source/Runtime/Engine/Private/Texture2D.cpp b/Engine/Source/Runtime/Engine/Private/Texture2D.cpp
--- a/Engine/Source/Runtime/Engine/Private/Texture2D.cpp
+++ b/Engine/Source/Runtime/Engine/Private/Texture2D.cpp
@@ -32,6 +32,14 @@
 #include "Streaming/Texture2DStreamIn_IO_AsyncCreate.h"
 #include "Streaming/Texture2DStreamIn_IO_AsyncReallocate.h"
 #include "Streaming/Texture2DStreamIn_IO_Virtual.h"
+// Generic path
+#include "Streaming/TextureStreamIn.h"
+#include "Streaming/Texture2DMipAllocator_AsyncCreate.h"
+#include "Streaming/Texture2DMipAllocator_AsyncReallocate.h"
+#include "Streaming/Texture2DMipDataProvider_DDC.h"
+#include "Streaming/Texture2DMipDataProvider_IO.h"
+#include "Engine/TextureMipDataProviderFactory.h"
+
 #include "Async/AsyncFileHandle.h"
 #include "EngineModule.h"
 #include "Engine/Texture2DArray.h"
@@ -90,6 +98,19 @@ static TAutoConsoleVariable<int32> CVarMobileMaxLoadedMips(
 	TEXT("Maximum number of loaded mips for nonstreaming mobile platforms.\n"),
 	ECVF_RenderThreadSafe);
 
+
+int32 GUseGenericStreamingPath = 0;
+static FAutoConsoleVariableRef CVarUseGenericStreamingPath(
+	TEXT("r.Streaming.UseGenericStreamingPath"),
+	GUseGenericStreamingPath,
+	TEXT("Control when to use the mip data provider implementation: (default=0)\n")
+	TEXT("0 to use it when there is a custom asset override.\n")
+	TEXT("1 to always use it.\n")
+	TEXT("2 to never use it."),
+	ECVF_Default
+);
+
+
 static int32 MobileReduceLoadedMips(int32 NumTotalMips)
 {
 	int32 NumReduceMips = FMath::Max(0, CVarMobileReduceLoadedMips.GetValueOnAnyThread());
@@ -2193,36 +2214,84 @@ bool UTexture2D::StreamIn(int32 NewMipCount, bool bHighPrio)
 	FTexture2DResource* Texture2DResource = (FTexture2DResource*)Resource;
 	if (bIsStreamable && !PendingUpdate && Texture2DResource && Texture2DResource->bReadyForStreaming && NewMipCount > GetNumResidentMips())
 	{
-#if WITH_EDITORONLY_DATA
-		if (FPlatformProperties::HasEditorOnlyData() && !GetOutermost()->bIsCookedForEditor)
+		FTextureMipDataProvider* CustomMipDataProvider = nullptr;
+		if (GUseGenericStreamingPath != 2)
 		{
-			if (GRHISupportsAsyncTextureCreation)
+			for (UAssetUserData* UserData : AssetUserData)
 			{
-				PendingUpdate = new FTexture2DStreamIn_DDC_AsyncCreate(this, NewMipCount);
+				UTextureMipDataProviderFactory* CustomMipDataProviderFactory = Cast<UTextureMipDataProviderFactory>(UserData);
+				if (CustomMipDataProviderFactory)
+				{
+					CustomMipDataProvider = CustomMipDataProviderFactory->AllocateMipDataProvider(this, NewMipCount);
+					if (CustomMipDataProvider)
+					{
+						break;
+					}
+				}
+			}
+		}
+
+		if (!CustomMipDataProvider && GUseGenericStreamingPath != 1)
+		{
+	#if WITH_EDITORONLY_DATA
+			if (FPlatformProperties::HasEditorOnlyData() && !GetOutermost()->bIsCookedForEditor)
+			{
+				if (GRHISupportsAsyncTextureCreation)
+				{
+					PendingUpdate = new FTexture2DStreamIn_DDC_AsyncCreate(this, NewMipCount);
+				}
+				else
+				{
+					PendingUpdate = new FTexture2DStreamIn_DDC_AsyncReallocate(this, NewMipCount);
+				}
 			}
 			else
+	#endif
 			{
-				PendingUpdate = new FTexture2DStreamIn_DDC_AsyncReallocate(this, NewMipCount);
+				// If the future texture is to be a virtual texture, use the virtual stream in path.
+				if (Texture2DResource->bUseVirtualUpdatePath)
+				{
+					PendingUpdate = new FTexture2DStreamIn_IO_Virtual(this, NewMipCount, bHighPrio);
+				}
+				// If the platform supports creating the new texture on an async thread, use that path.
+				else if (GRHISupportsAsyncTextureCreation)
+				{
+					PendingUpdate = new FTexture2DStreamIn_IO_AsyncCreate(this, NewMipCount,bHighPrio);
+				}
+				// Otherwise use the default path.
+				else
+				{
+					PendingUpdate = new FTexture2DStreamIn_IO_AsyncReallocate(this, NewMipCount, bHighPrio);
+				}
 			}
 		}
-		else
-#endif
+		else // Generic path
 		{
-			// If the future texture is to be a virtual texture, use the virtual stream in path.
-			if (Texture2DResource->bUseVirtualUpdatePath)
+			FTextureMipAllocator* MipAllocator = nullptr;
+			FTextureMipDataProvider* DefaultMipDataProvider = nullptr;
+
+	#if WITH_EDITORONLY_DATA
+			if (FPlatformProperties::HasEditorOnlyData() && !GetOutermost()->bIsCookedForEditor)
 			{
-				PendingUpdate = new FTexture2DStreamIn_IO_Virtual(this, NewMipCount, bHighPrio);
+				DefaultMipDataProvider = new FTexture2DMipDataProvider_DDC();
 			}
-			// If the platform supports creating the new texture on an async thread, use that path.
-			else if (GRHISupportsAsyncTextureCreation)
+			else 
+#endif
 			{
-				PendingUpdate = new FTexture2DStreamIn_IO_AsyncCreate(this, NewMipCount,bHighPrio);
+				DefaultMipDataProvider = new FTexture2DMipDataProvider_IO(bHighPrio);
+			}
+
+			// FTexture2DMipAllocator_Virtual?
+			if (GRHISupportsAsyncTextureCreation)
+			{
+				MipAllocator = new FTexture2DMipAllocator_AsyncCreate();
 			}
-			// Otherwise use the default path.
 			else
 			{
-				PendingUpdate = new FTexture2DStreamIn_IO_AsyncReallocate(this, NewMipCount, bHighPrio);
+				MipAllocator = new FTexture2DMipAllocator_AsyncReallocate();
 			}
+
+			PendingUpdate = new FTextureStreamIn(this, NewMipCount, MipAllocator, CustomMipDataProvider, DefaultMipDataProvider);
 		}
 		return !PendingUpdate->IsCancelled();
 	}
diff --git a/Engine/Source/Runtime/Engine/Public/RenderAssetUpdate.h b/Engine/Source/Runtime/Engine/Public/RenderAssetUpdate.h
--- a/Engine/Source/Runtime/Engine/Public/RenderAssetUpdate.h
+++ b/Engine/Source/Runtime/Engine/Public/RenderAssetUpdate.h
@@ -249,6 +249,8 @@ protected:
 		CancelationThread = TT_None;
 		CancelationCallback = nullptr;
 	}
+
+	virtual void InitContext(FContext& Context) const {}
 };
 
 void SuspendRenderAssetStreaming();
diff --git a/Engine/Source/Runtime/Engine/Public/Streaming/TextureMipDataProvider.h b/Engine/Source/Runtime/Engine/Public/Streaming/TextureMipDataProvider.h
new file mode 100644
--- /dev/null
+++ b/Engine/Source/Runtime/Engine/Public/Streaming/TextureMipDataProvider.h
@@ -0,0 +1,193 @@
+// Copyright 1998-2019 Epic Games, Inc. All Rights Reserved.
+
+/*=============================================================================
+TextureMipDataProvider.h: Base class for providing the mip data used by FTextureStreamIn.
+=============================================================================*/
+
+#pragma once
+
+#include "CoreMinimal.h"
+#include "HAL/ThreadSafeCounter.h"
+#include "PixelFormat.h"
+#include "RHIDefinitions.h"
+
+// Describes a mip that gets updated through FTextureStreamIn.
+struct FTextureMipInfo
+{
+	void* DestData = nullptr; // The mip data where the content must be copied too.
+	uint32 DataSize = 0; // Optional, might be 0 in cooked. Mostly for safeguard.
+	EPixelFormat Format = PF_Unknown;
+	uint32 SizeX = 0;
+	uint32 SizeY = 0;
+	uint32 SizeZ = 0;
+	uint32 RowPitch = 0;
+	uint32 DepthPitch = 0;	// For volume texture, texture array only.
+	ECubeFace CubeFace = CubeFace_MAX; // Cubemap only.
+};
+
+// An array of mips mapped to a given texture asset. 0 based on the highest mip count for this texture.
+typedef TArray<FTextureMipInfo, TInlineAllocator<MAX_TEXTURE_MIP_COUNT>> FTextureMipInfoArray;
+
+/**
+* Different options to schedule correctly the next tick in FTextureStreamIn.
+* It allows the asset update logic to postpone the next tick until the effect of the last tick is completed.
+* This helps for things like async operations.
+*
+* If no sync options are used in the tick of either FTextureMipDataProvider or FTextureMipAllocator,
+* the tick must at least advance to the next state through AdvanceTo() to prevent infinite loops.
+*/
+struct FTextureUpdateSyncOptions
+{
+	typedef TFunction<void()> FCallback;
+
+	// A bool to tell whether the tick should be postponed by a small duration (undefined by design).
+	bool* bSnooze = nullptr;
+	// A counter that will be incremented before doing an operation and then decremented once has completed.
+	FThreadSafeCounter* Counter = nullptr;
+
+	// A callback to force a new tick on the relevant thread.
+	FCallback RescheduleCallback;
+};
+
+/**
+* A constant context used throughout the texture stream in update tick.
+* Hold useful information for either FTextureMipDataProvider or FTextureMipAllocator.
+*/
+struct FTextureUpdateContext
+{
+	typedef int32 EThreadType;
+
+	FTextureUpdateContext(UTexture* InTexture, EThreadType InCurrentThread);
+	FTextureUpdateContext(UStreamableRenderAsset* InTexture, EThreadType InCurrentThread);
+
+	UStreamableRenderAsset* GetRenderAsset() const { return Texture; }
+	EThreadType GetCurrentThread() const { return CurrentThread; }
+
+	// The texture to update, this must be the same one as the one used when creating the FTextureUpdate object.
+	UTexture* Texture;
+	// The current executing thread.
+	EThreadType CurrentThread;
+
+	// The requested first mip after the update completes (smaller mip index are bigger, 0 being the biggest possible mips for that texture).
+	int32 PendingFirstMipIndex = INDEX_NONE;
+	// The current first mip before the update was requested.
+	int32 CurrentFirstMipIndex = INDEX_NONE;
+	// The number of requested mips after the update completes (PendingFirstMipIndex + NumRequestedMips - 1) being the smallest possible mip for this texture (typically 1x1).
+	int32 NumRequestedMips = INDEX_NONE;
+};
+
+/**
+* FTextureMipDataProvider defines the update steps and interface to implement the mip data strategy
+* used in FTextureStreamIn. It allows to decouples where the texture mip data source from the texture update.
+* Typical implementations are using DDC, disk files, internet server or dynamically generated.
+*/
+class ENGINE_API FTextureMipDataProvider
+{
+public:
+
+	enum class ETickState : uint32
+	{
+		Init,			// Call Init(), which must initialize the mip data provider so that it can efficently provide the mips in GetMips.
+		GetMips,		// Call GetMips(), which must fill mip data from highest to lowest (all mips might not be available thought).
+		PollMips,		// Call PollMips(), which stalls the mip update until the mip have been set.
+		CleanUp,		// Call CleanUp(), which must release all temporary resource created in the update process.
+		Done
+	};
+
+	enum class ETickThread : uint32
+	{
+		Async,			// Tick must run on an async thread.
+		Render,			// Tick must run on the renderthread.
+		None,			// No tick must run, only valid in ETickState::Done. 
+	};
+
+	// Constructor, defining the first tick step and thread.
+	FTextureMipDataProvider(ETickState InTickState, ETickThread InTickThread) : NextTickState(InTickState), NextTickThread(InTickThread) {}
+	virtual ~FTextureMipDataProvider() {}
+
+	// Get the next tick state and thread for the mip allocator. Used by FTextureStreamIn to schedule correctly the update between FTextureMipAllocator and FTextureMipDataProvider.
+	FORCEINLINE ETickState GetNextTickState() const { return NextTickState; }
+	FORCEINLINE ETickThread GetNextTickThread() const { return NextTickThread; }
+
+	/**
+	* Initialize data prelimary to the GetMips() step. Can be called several time. Mostly useful to simplify the logic in GetMips().
+	* This is because GetMips is a chained call between all mip data providers, each taking some mips to handle, and is not compatible with multi step process.
+	* This means that GetMips() must return immediately and can not postpone or delay return by not advancing to the next steps.
+	*
+	* @param Context - An update context constant throughout the FTextureStreamIn update. Gives things like which texture asset is updated and what mips are streamed in.
+	* @param SyncOptions - Different sync options to control when the next tick of FTextureStreamIn can be scheduled.
+	*/
+	virtual void Init(const FTextureUpdateContext& Context, const FTextureUpdateSyncOptions& SyncOptions) = 0;
+
+	/**
+	* Acquire the mips this provider will handle. Non handled mips must be handled by the next mip data provider or the update will be cancelled altogether.
+	* GetMips() must typically advance to PollMips() if it wants to be able to notify FTextureStreamIn that something went wrong and that the update must be cancelled.
+	*
+	* @param Context - An update context constant throughout the FTextureStreamIn update. Gives things like which texture asset is updated and what mips are streamed in.
+	* @param StartingMipIndex - The current starting mip index, somewhere between FTextureUpdateContext::PendingFirstMipIndex and FTextureUpdateContext::CurrentFirstMipIndex inclusively.
+	* @param MipInfos - The array of FTextureMipInfo that hold the information relative to each mip for which data must be provided.
+	* @param SyncOptions - Different sync options to control when the next tick of FTextureStreamIn can be scheduled.
+	*
+	* Return the next value StartingMipIndex (for the next provider). Must be FTextureUpdateContext::CurrentFirstMipIndex to indicate that all mips have been handled.
+	*/
+	virtual int32 GetMips(const FTextureUpdateContext& Context, int32 StartingMipIndex, const FTextureMipInfoArray& MipInfos, const FTextureUpdateSyncOptions& SyncOptions) = 0;
+
+	/**
+	* Check if each mip handled by this mip data provider have been updated correctly. Must move to CleanUp() when done.
+	*
+	* @param SyncOptions - Different sync options to control when the next tick of FTextureStreamIn can be scheduled.
+	*
+	* Return true unless the texture update needs to be aborted because the mip won't be updated correctly (for example IO error).
+	*/
+	virtual bool PollMips(const FTextureUpdateSyncOptions& SyncOptions) = 0;
+
+	/**
+	* Abort anything that could be stalling the update in PollMips(). Called from an async task.
+	* Expected behavior is to cancel any pending IO operations. 
+	* Called when the FTextureStreamIn update is cancelled and the current step is to be waiting on upon PollMips().
+	* This is because Cancel() can only be executed when any pending operations are completed (see FTextureUpdateSyncOptions).
+	*/
+	virtual void AbortPollMips() {}
+
+	/**
+	* Release any temporary data and objects that where used for the update. Final step executed after the texture has been updated correctly.
+	*
+	* @param SyncOptions - Different sync options to control when the next tick of FTextureStreamIn can be scheduled.
+	*/
+	virtual void CleanUp(const FTextureUpdateSyncOptions& SyncOptions) = 0;
+
+	/**
+	* Cancel the progression and release any temporary resources. 
+	* Called within the FTextureStreamIn update when the stream in request is aborted or cannot complete correctly
+	* The Cancel() function is called on the thread returned by GetCancelThread().
+	*
+	* @param SyncOptions - Different sync options to control when the next tick of FTextureStreamIn can be scheduled.
+	*/
+	virtual void Cancel(const FTextureUpdateSyncOptions& SyncOptions) = 0;
+
+	// Returns on which thread Cancel() must be called in the texture update (from FTextureStreamIn) to release ressources safely and correctly. ETickThread::None when ready to delete this.
+	virtual ETickThread GetCancelThread() const = 0;
+
+protected:
+
+	// Helper to set the next tick state and thread. Validates that the progression is coherent.
+	FORCEINLINE void AdvanceTo(ETickState InState, ETickThread InThread)
+	{
+		// State must advance or at least stay to the same step.
+		check((uint32)NextTickState <= (uint32)InState);
+		// Thread must be valid unless in done state.
+		check((InState == ETickState::Done) == (InThread == ETickThread::None));
+
+		NextTickState = InState;
+		NextTickThread = InThread;
+	}
+
+private:
+
+	// The next tick function that should be called.
+	ETickState NextTickState = ETickState::Done;
+	// The next tick thread from where to call the tick function.
+	ETickThread NextTickThread = ETickThread::None;
+};
+
+typedef TArray<TUniquePtr<FTextureMipDataProvider>, TInlineAllocator<2>> FTextureMipDataProviderArray;
