Chaos - Fix Uninitialized value being used to set BodyInstance Scale in UpdateBodyScale.
Change error message to warning and include skipped types in logging.

diff --git a/Engine/Source/Runtime/Engine/Private/PhysicsEngine/BodyInstance.cpp b/Engine/Source/Runtime/Engine/Private/PhysicsEngine/BodyInstance.cpp
--- a/Engine/Source/Runtime/Engine/Private/PhysicsEngine/BodyInstance.cpp
+++ b/Engine/Source/Runtime/Engine/Private/PhysicsEngine/BodyInstance.cpp
@@ -1646,6 +1646,8 @@ bool FBodyInstance::UpdateBodyScale(const FVector& InScale3D, bool bForceUpdate)
 		// Apply scaling
 		ComputeScalingVectors(ScaleMode, InScale3D, AdjustedScale3D, AdjustedScale3DAbs);
 
+		UpdatedScale3D = AdjustedScale3D;
+
 		TArray<TUniquePtr<FImplicitObject>> NewGeometry;
 		NewGeometry.Reserve(Shapes.Num());
 
@@ -1799,7 +1801,7 @@ bool FBodyInstance::UpdateBodyScale(const FVector& InScale3D, bool bForceUpdate)
 				default:
 				{
 					CHAOS_ENSURE(false);
-					UE_LOG(LogPhysics, Error, TEXT("Unimplemented ImplicitObject skipped in UpdateBodyScale."));
+					UE_LOG(LogPhysics, Warning, TEXT("UpdateBodyScale: Unimplemented ImplicitObject of type: %d skipped."), ImplicitType);
 				}
 			}// end switch
 		}
