Remove #pragma once from cpp
#rnx

diff --git a/Engine/Source/Runtime/Engine/Private/Materials/HLSLMaterialTranslator.cpp b/Engine/Source/Runtime/Engine/Private/Materials/HLSLMaterialTranslator.cpp
--- a/Engine/Source/Runtime/Engine/Private/Materials/HLSLMaterialTranslator.cpp
+++ b/Engine/Source/Runtime/Engine/Private/Materials/HLSLMaterialTranslator.cpp
@@ -3,8 +3,6 @@
 	HLSLMaterialTranslator.cpp: Translates material expressions into HLSL code.
 =============================================================================*/
 
-#pragma once
-
 #include "HLSLMaterialTranslator.h"
 
 class Error;
