Fix for cook non-determinism in FStaticMeshLODResources caused by saving an uninitialized FStaticMeshBuffersSize when bulk data is discarded.

diff --git a/Engine/Source/Runtime/Engine/Public/StaticMeshResources.h b/Engine/Source/Runtime/Engine/Public/StaticMeshResources.h
--- a/Engine/Source/Runtime/Engine/Public/StaticMeshResources.h
+++ b/Engine/Source/Runtime/Engine/Public/StaticMeshResources.h
@@ -45,7 +45,6 @@ class UBodySetup;
 #define STATICMESH_ENABLE_DEBUG_RENDERING (!(UE_BUILD_SHIPPING || UE_BUILD_TEST) || WITH_EDITOR)
 
 struct FStaticMaterial;
-struct FStaticMeshBuffersSize;
 
 /**
  * The LOD settings to use for a group of static meshes.
@@ -431,9 +430,9 @@ private:
 	 */
 	struct FStaticMeshBuffersSize
 	{
-		uint32 SerializedBuffersSize;
-		uint32 DepthOnlyIBSize;
-		uint32 ReversedIBsSize;
+		uint32 SerializedBuffersSize = 0;
+		uint32 DepthOnlyIBSize       = 0;
+		uint32 ReversedIBsSize       = 0;
 
 		void Clear()
 		{
