Fixed missing forward declarions for CL 10516726 (FTextureMipDataProvider)

diff --git a/Engine/Source/Runtime/Engine/Public/Streaming/TextureMipDataProvider.h b/Engine/Source/Runtime/Engine/Public/Streaming/TextureMipDataProvider.h
--- a/Engine/Source/Runtime/Engine/Public/Streaming/TextureMipDataProvider.h
+++ b/Engine/Source/Runtime/Engine/Public/Streaming/TextureMipDataProvider.h
@@ -11,6 +11,9 @@ TextureMipDataProvider.h: Base class for providing the mip data used by FTexture
 #include "PixelFormat.h"
 #include "RHIDefinitions.h"
 
+class UTexture;
+class UStreamableRenderAsset;
+
 // Describes a mip that gets updated through FTextureStreamIn.
 struct FTextureMipInfo
 {
