Fixed compilation error with FTextureUpdateContext

diff --git a/Engine/Source/Runtime/Engine/Private/Streaming/TextureStreamIn.cpp b/Engine/Source/Runtime/Engine/Private/Streaming/TextureStreamIn.cpp
--- a/Engine/Source/Runtime/Engine/Private/Streaming/TextureStreamIn.cpp
+++ b/Engine/Source/Runtime/Engine/Private/Streaming/TextureStreamIn.cpp
@@ -22,6 +22,11 @@ FTextureUpdateContext::FTextureUpdateContext(UStreamableRenderAsset* InTexture,
 {
 }
 
+UStreamableRenderAsset* FTextureUpdateContext::GetRenderAsset() const
+{
+	return Texture;
+}
+
 bool FTextureStreamIn::IsSameThread(FTextureMipAllocator::ETickThread TickThread, int32 TaskThread)
 {
 	if (TaskThread == TT_Async)
diff --git a/Engine/Source/Runtime/Engine/Public/Streaming/TextureMipDataProvider.h b/Engine/Source/Runtime/Engine/Public/Streaming/TextureMipDataProvider.h
--- a/Engine/Source/Runtime/Engine/Public/Streaming/TextureMipDataProvider.h
+++ b/Engine/Source/Runtime/Engine/Public/Streaming/TextureMipDataProvider.h
@@ -63,7 +63,7 @@ struct FTextureUpdateContext
 	FTextureUpdateContext(UTexture* InTexture, EThreadType InCurrentThread);
 	FTextureUpdateContext(UStreamableRenderAsset* InTexture, EThreadType InCurrentThread);
 
-	UStreamableRenderAsset* GetRenderAsset() const { return Texture; }
+	UStreamableRenderAsset* GetRenderAsset() const;
 	EThreadType GetCurrentThread() const { return CurrentThread; }
 
 	// The texture to update, this must be the same one as the one used when creating the FTextureUpdate object.
