Chaos - BodyInstance::UpdateBodyScale support for Spheres and Transformed Boxes.

diff --git a/Engine/Source/Runtime/Engine/Private/PhysicsEngine/BodyInstance.cpp b/Engine/Source/Runtime/Engine/Private/PhysicsEngine/BodyInstance.cpp
--- a/Engine/Source/Runtime/Engine/Private/PhysicsEngine/BodyInstance.cpp
+++ b/Engine/Source/Runtime/Engine/Private/PhysicsEngine/BodyInstance.cpp
@@ -33,6 +33,7 @@
 #include "Chaos/Convex.h"
 #include "Chaos/ImplicitObject.h"
 #include "Chaos/ImplicitObjectScaled.h"
+#include "Chaos/ImplicitObjectTransformed.h"
 #include "Chaos/ParticleHandle.h"
 #include "Chaos/TriangleMeshImplicitObject.h"
 
@@ -1660,14 +1661,46 @@ bool FBodyInstance::UpdateBodyScale(const FVector& InScale3D, bool bForceUpdate)
 
 			const FTransform& RelativeTM = GetRelativeBodyTransform(ShapeHandle);
 
-			// TODO: support instanced and transformed.
+			// TODO: support instanced
 			CHAOS_ENSURE(!IsInstanced(ImplicitType));
-			CHAOS_ENSURE(ImplicitType != ImplicitObjectType::Transformed);
+
+			bool bIsTransformed = false;
+			if(ImplicitType == ImplicitObjectType::Transformed)
+			{
+				bIsTransformed = true;
+
+				// Get GeomType that is transformed
+				const TImplicitObjectTransformed<FReal, 3>& ImplicitObjectTransformed = static_cast<const TImplicitObjectTransformed<FReal, 3>&>(ImplicitObject);
+				GeomType = ImplicitObjectTransformed.GetTransformedObject()->GetType(true);
+				CHAOS_ENSURE(!IsInstanced(GeomType));
+			}
 
 			FKShapeElem* ShapeElem = FPhysxUserData::Get<FKShapeElem>(FPhysicsInterface::GetUserData(ShapeHandle));
 
 			switch (GeomType)
 			{
+				case ImplicitObjectType::Sphere:
+				{
+					FKSphereElem* SphereElem = ShapeElem->GetShapeCheck<FKSphereElem>();
+					ensure(ScaleMode == EScaleMode::LockedXYZ);
+
+					FReal Radius = FMath::Max(SphereElem->Radius * AdjustedScale3DAbs.X, FCollisionShape::MinSphereRadius());
+
+
+					if (!CHAOS_ENSURE(!IsScaled(ImplicitType) || !CHAOS_ENSURE(!bIsTransformed)))
+					{
+						// No support for Scaled sphere or transformed.
+						break;
+					}
+
+					FVec3 Center = RelativeTM.TransformPosition(SphereElem->Center) * AdjustedScale3D.X;
+
+					TUniquePtr<TSphere<FReal, 3>> NewSphere = MakeUnique<TSphere<FReal, 3>>(Center, Radius);
+
+					NewGeometry.Emplace(MoveTemp(NewSphere));
+					bSuccess = true;
+					break;
+				}
 				case ImplicitObjectType::Box:
 				{
 					if (!CHAOS_ENSURE(!IsScaled(ImplicitType)))
@@ -1685,21 +1718,29 @@ bool FBodyInstance::UpdateBodyScale(const FVector& InScale3D, bool bForceUpdate)
 					HalfExtents.Y = FMath::Max((0.5f * BoxElem->Y * AdjustedScale3DAbs.Y), FCollisionShape::MinBoxExtent());
 					HalfExtents.Z = FMath::Max((0.5f * BoxElem->Z * AdjustedScale3DAbs.Z), FCollisionShape::MinBoxExtent());
 
-
-					// TODO: For Transformed implicit, do not bake this in. Set Transform instead.
 					FRigidTransform3 LocalTransform = BoxElem->GetTransform() * RelativeTM;
 					LocalTransform.ScaleTranslation(AdjustedScale3D);
 
-					// Any Box with initial Rotation should've been initialized as Transformed implicit.
-					// If we hit this it means user set rotation on initially unrotated box.
-					// TODO: Handle this.
-					CHAOS_ENSURE(LocalTransform.GetRotation() == FQuat::Identity);
+					// If not already transformed and has a rotation, must convert to transformed geometry.
+					if (bIsTransformed || (LocalTransform.GetRotation() == FQuat::Identity))
+					{
+						// Center at origin, transform holds translation
+						const FVec3 Min = -HalfExtents;
+						const FVec3 Max =  HalfExtents;
 
-					const FVec3 Min = LocalTransform.GetLocation() - HalfExtents;
-					const FVec3 Max = LocalTransform.GetLocation() + HalfExtents;
+						TUniquePtr<TBox<FReal, 3>> NewBox = MakeUnique<TBox<FReal, 3>>(Min, Max);
+						TUniquePtr<TImplicitObjectTransformed<FReal, 3>> NewTransformedBox = MakeUnique<TImplicitObjectTransformed<FReal, 3>>(MoveTemp(NewBox), LocalTransform);
+						NewGeometry.Emplace(MoveTemp(NewTransformedBox));
+					}
+					else
+					{
+						// Bake in transformed position
+						const FVec3 Min = LocalTransform.GetLocation() - HalfExtents;
+						const FVec3 Max = LocalTransform.GetLocation() + HalfExtents;
 
-					TUniquePtr<TBox<FReal, 3>> NewBox = MakeUnique<TBox<FReal, 3>>(Min, Max);
-					NewGeometry.Emplace(MoveTemp(NewBox));
+						TUniquePtr<TBox<FReal, 3>> NewBox = MakeUnique<TBox<FReal, 3>>(Min, Max);
+						NewGeometry.Emplace(MoveTemp(NewBox));
+					}
 
 					bSuccess = true;
 
@@ -1716,9 +1757,9 @@ bool FBodyInstance::UpdateBodyScale(const FVector& InScale3D, bool bForceUpdate)
 
 					FKSphylElem* SphylElem = ShapeElem->GetShapeCheck<FKSphylElem>();
 
-					if (!CHAOS_ENSURE(!IsScaled(ImplicitType)))
+					if (!CHAOS_ENSURE(!IsScaled(ImplicitType) || !CHAOS_ENSURE(!bIsTransformed)))
 					{
-						// No support for ScaledImplicit Capsule yet
+						// No support for Scaled or Transformed Capsule yet
 						break;
 					}
 
