Fix sign in GetWorldVelocityAtPoint_AssumesLocked

diff --git a/Engine/Source/Runtime/Engine/Private/PhysicsEngine/Experimental/PhysInterface_Chaos.cpp b/Engine/Source/Runtime/Engine/Private/PhysicsEngine/Experimental/PhysInterface_Chaos.cpp
--- a/Engine/Source/Runtime/Engine/Private/PhysicsEngine/Experimental/PhysInterface_Chaos.cpp
+++ b/Engine/Source/Runtime/Engine/Private/PhysicsEngine/Experimental/PhysInterface_Chaos.cpp
@@ -589,7 +589,7 @@ FVector FPhysInterface_Chaos::GetWorldVelocityAtPoint_AssumesLocked(const FPhysi
 		{
 			const Chaos::FVec3 COM = Chaos::FParticleUtilitiesGT::GetCoMWorldPosition(Kinematic);
 			const Chaos::FVec3 Diff = InPoint - COM;
-			return Kinematic->V() + Chaos::FVec3::CrossProduct(Diff, Kinematic->W());
+			return Kinematic->V() - Chaos::FVec3::CrossProduct(Diff, Kinematic->W());
 		}
 	}
 	return FVector(0);
