Cameras: use non-replicated functions to start spatialized shakes.

The expectation is to not support replicated shake source components for now.

diff --git a/Engine/Source/Runtime/Engine/Classes/GameFramework/PlayerController.h b/Engine/Source/Runtime/Engine/Classes/GameFramework/PlayerController.h
--- a/Engine/Source/Runtime/Engine/Classes/GameFramework/PlayerController.h
+++ b/Engine/Source/Runtime/Engine/Classes/GameFramework/PlayerController.h
@@ -937,7 +937,7 @@ public:
 	 * @param Shake - Camera shake animation to play
 	 * @param SourceComponent - The source from which the camera shakes originates
 	 */
-	UFUNCTION(unreliable, client, BlueprintCallable, Category="Game|Feedback")
+	UFUNCTION(BlueprintCallable, Category="Game|Feedback")
 	void ClientPlayCameraShakeFromSource(TSubclassOf<class UCameraShake> Shake, class UCameraShakeSourceComponent* SourceComponent);
 
 	/**
@@ -1078,7 +1078,7 @@ public:
 	void ClientStopCameraShake(TSubclassOf<class UCameraShake> Shake, bool bImmediately = true);
 
 	/** Stop camera shake on client.  */
-	UFUNCTION(reliable, client, BlueprintCallable, Category="Game|Feedback")
+	UFUNCTION(BlueprintCallable, Category="Game|Feedback")
 	void ClientStopCameraShakesFromSource(class UCameraShakeSourceComponent* SourceComponent, bool bImmediately = true);
 
 	/** 
diff --git a/Engine/Source/Runtime/Engine/Private/Camera/CameraShakeSourceComponent.cpp b/Engine/Source/Runtime/Engine/Private/Camera/CameraShakeSourceComponent.cpp
--- a/Engine/Source/Runtime/Engine/Private/Camera/CameraShakeSourceComponent.cpp
+++ b/Engine/Source/Runtime/Engine/Private/Camera/CameraShakeSourceComponent.cpp
@@ -79,7 +79,7 @@ void UCameraShakeSourceComponent::PlayCameraShake(TSubclassOf<UCameraShake> InCa
 		APlayerController* PlayerController = Iterator->Get();
 		if (PlayerController != nullptr && PlayerController->PlayerCameraManager != nullptr)
 		{
-			PlayerController->ClientPlayCameraShakeFromSource(InCameraShake, this);
+			PlayerController->PlayerCameraManager->PlayCameraShakeFromSource(InCameraShake, this);
 		}
 	}
 }
@@ -91,7 +91,7 @@ void UCameraShakeSourceComponent::StopAllCameraShakes(bool bImmediately)
 		APlayerController* PlayerController = Iterator->Get();
 		if (PlayerController != nullptr && PlayerController->PlayerCameraManager != nullptr)
 		{
-			PlayerController->ClientStopCameraShakesFromSource(this, bImmediately);
+			PlayerController->PlayerCameraManager->StopAllInstancesOfCameraShakeFromSource(this, bImmediately);
 		}
 	}
 }
diff --git a/Engine/Source/Runtime/Engine/Private/PlayerController.cpp b/Engine/Source/Runtime/Engine/Private/PlayerController.cpp
--- a/Engine/Source/Runtime/Engine/Private/PlayerController.cpp
+++ b/Engine/Source/Runtime/Engine/Private/PlayerController.cpp
@@ -4337,7 +4337,7 @@ void APlayerController::ClientStopCameraShake_Implementation( TSubclassOf<class
 	}
 }
 
-void APlayerController::ClientPlayCameraShakeFromSource_Implementation(TSubclassOf<class UCameraShake> Shake, class UCameraShakeSourceComponent* SourceComponent)
+void APlayerController::ClientPlayCameraShakeFromSource(TSubclassOf<class UCameraShake> Shake, class UCameraShakeSourceComponent* SourceComponent)
 {
 	if (PlayerCameraManager != NULL)
 	{
@@ -4345,7 +4345,7 @@ void APlayerController::ClientPlayCameraShakeFromSource_Implementation(TSubclass
 	}
 }
 
-void APlayerController::ClientStopCameraShakesFromSource_Implementation(class UCameraShakeSourceComponent* SourceComponent, bool bImmediately)
+void APlayerController::ClientStopCameraShakesFromSource(class UCameraShakeSourceComponent* SourceComponent, bool bImmediately)
 {
 	if (PlayerCameraManager != NULL)
 	{
