- Add VertexFactory array per View for Cascade particle systems
- This fixes issues around split screen where particles are orientated to view 0's camera for example

diff --git a/Engine/Source/Runtime/Engine/Private/Particles/ParticleSystemRender.cpp b/Engine/Source/Runtime/Engine/Private/Particles/ParticleSystemRender.cpp
--- a/Engine/Source/Runtime/Engine/Private/Particles/ParticleSystemRender.cpp
+++ b/Engine/Source/Runtime/Engine/Private/Particles/ParticleSystemRender.cpp
@@ -7000,8 +7000,6 @@ void FParticleSystemSceneProxy::GetDynamicMeshElements(const TArray<const FScene
 				}
 				FScopeCycleCounter AdditionalScope(Data->StatID);
 
-				FParticleVertexFactoryBase *VertexFactory = EmitterVertexFactoryArray[Data->EmitterIndex];
-
 				//hold on to the emitter index in case we need to access any of its properties
 				DynamicData->EmitterIndex = Index;
 
@@ -7009,6 +7007,17 @@ void FParticleSystemSceneProxy::GetDynamicMeshElements(const TArray<const FScene
 				{
 					if (VisibilityMap & (1 << ViewIndex))
 					{
+						if ((ViewIndex >= EmitterVertexFactoryArray.Num()) || (Data->EmitterIndex >= EmitterVertexFactoryArray[ViewIndex].Num()))
+						{
+							AddEmitterVertexFactory(Data, ViewIndex);
+						}
+						FParticleVertexFactoryBase* VertexFactory = EmitterVertexFactoryArray[ViewIndex][Data->EmitterIndex];
+						if (VertexFactory == nullptr)
+						{
+							AddEmitterVertexFactory(Data, ViewIndex);
+							VertexFactory = EmitterVertexFactoryArray[ViewIndex][Data->EmitterIndex];
+						}
+
 						const FSceneView* View = Views[ViewIndex];
 						Data->GetDynamicMeshElementsEmitter(this, View, ViewFamily, ViewIndex, Collector, VertexFactory);
 						NumDraws++;
diff --git a/Engine/Source/Runtime/Engine/Public/ParticleHelper.h b/Engine/Source/Runtime/Engine/Public/ParticleHelper.h
--- a/Engine/Source/Runtime/Engine/Public/ParticleHelper.h
+++ b/Engine/Source/Runtime/Engine/Public/ParticleHelper.h
@@ -2590,8 +2590,16 @@ public:
 
 	void MarkEmitterVertexFactoryDirty(uint32 EmitterIndex) const
 	{
-		check(EmitterVertexFactoryArray[EmitterIndex]);
-		EmitterVertexFactoryArray[EmitterIndex]->SetDirty();
+		for (TArray<FParticleVertexFactoryBase*>& PerViewEmitterVFs : EmitterVertexFactoryArray)
+		{
+			if ( PerViewEmitterVFs.IsValidIndex(EmitterIndex) )
+			{
+				if ( FParticleVertexFactoryBase *VertexFactory = PerViewEmitterVFs[EmitterIndex] )
+				{
+					VertexFactory->SetDirty();
+				}
+			}
+		}
 	}
 
 	void ClearVertexFactoriesIfDirty() const
@@ -2604,23 +2612,29 @@ public:
 		else
 		{
 			// selectively clear
-			for (int32 Index = 0; Index < EmitterVertexFactoryArray.Num(); Index++)
+			for (TArray<FParticleVertexFactoryBase*>& PerViewEmitterVFs : EmitterVertexFactoryArray)
 			{
-				ClearEmitterVertexFactoryIfDirty(Index);
+				for (int32 Index = 0; Index < PerViewEmitterVFs.Num(); Index++)
+				{
+					ClearEmitterVertexFactoryIfDirty(Index);
+				}
 			}
 		}
 	}
 
 	void ClearVertexFactories() const
 	{
-		for (int32 Index = 0; Index < EmitterVertexFactoryArray.Num(); Index++)
+		for (TArray<FParticleVertexFactoryBase*>& PerViewEmitterVFs : EmitterVertexFactoryArray)
 		{
-			FParticleVertexFactoryBase *VertexFactory = EmitterVertexFactoryArray[Index];
-			if (VertexFactory)
+			for (int32 Index = 0; Index < PerViewEmitterVFs.Num(); Index++)
 			{
-				VertexFactory->ReleaseResource();
-				delete VertexFactory;
-				EmitterVertexFactoryArray[Index] = nullptr;
+				FParticleVertexFactoryBase *VertexFactory = PerViewEmitterVFs[Index];
+				if (VertexFactory)
+				{
+					VertexFactory->ReleaseResource();
+					delete VertexFactory;
+					PerViewEmitterVFs[Index] = nullptr;
+				}
 			}
 		}
 		bVertexFactoriesDirty = false;
@@ -2628,25 +2642,38 @@ public:
 
 	void ClearEmitterVertexFactoryIfDirty(uint32 EmitterIndex) const
 	{
-		FParticleVertexFactoryBase *VertexFactory = EmitterVertexFactoryArray[EmitterIndex];
-		if (VertexFactory && VertexFactory->IsDirty())
+		for (TArray<FParticleVertexFactoryBase*>& PerViewEmitterVFs : EmitterVertexFactoryArray)
 		{
-			VertexFactory->ReleaseResource();
-			delete VertexFactory;
-			EmitterVertexFactoryArray[EmitterIndex] = nullptr;
+			FParticleVertexFactoryBase *VertexFactory = PerViewEmitterVFs[EmitterIndex];
+			if (VertexFactory && VertexFactory->IsDirty())
+			{
+				VertexFactory->ReleaseResource();
+				delete VertexFactory;
+				PerViewEmitterVFs[EmitterIndex] = nullptr;
+			}
 		}
 	}
 
-	void AddEmitterVertexFactory(FDynamicEmitterDataBase *InDynamicData) const
+	void AddEmitterVertexFactory(FDynamicEmitterDataBase *InDynamicData, int32 ViewIndex) const
 	{
-		while(InDynamicData->EmitterIndex >= EmitterVertexFactoryArray.Num())
+		check(ViewIndex >= 0);
+
+		// Ensure view is allocated
+		while (ViewIndex >= EmitterVertexFactoryArray.Num())
+		{
+			EmitterVertexFactoryArray.AddDefaulted();
+		}
+
+		// Allocate space for emitters
+		TArray<FParticleVertexFactoryBase*>& PerViewEmitterVFs = EmitterVertexFactoryArray[ViewIndex];
+		while (InDynamicData->EmitterIndex >= PerViewEmitterVFs.Num())
 		{
-			EmitterVertexFactoryArray.Add(nullptr);
+			PerViewEmitterVFs.Add(nullptr);
 		}
 
-		if (EmitterVertexFactoryArray[InDynamicData->EmitterIndex] == nullptr)
+		if (PerViewEmitterVFs[InDynamicData->EmitterIndex] == nullptr)
 		{
-			EmitterVertexFactoryArray[InDynamicData->EmitterIndex] = InDynamicData->CreateVertexFactory(FeatureLevel, this);
+			PerViewEmitterVFs[InDynamicData->EmitterIndex] = InDynamicData->CreateVertexFactory(FeatureLevel, this);
 		}
 	}
 
@@ -2659,7 +2686,7 @@ public:
 	{
 		for (int32 Index = 0; Index < DynamicDataForThisFrame.Num(); Index++)
 		{
-			AddEmitterVertexFactory(DynamicDataForThisFrame[Index]);
+			AddEmitterVertexFactory(DynamicDataForThisFrame[Index], 0);
 		}
 		DynamicDataForThisFrame.Empty();
 	}
@@ -2725,8 +2752,8 @@ private:
 
 protected:
 	/** vertex factories for all emitters */
-	mutable TArray<FParticleVertexFactoryBase*> EmitterVertexFactoryArray;
-	mutable TArray<FDynamicEmitterDataBase*> DynamicDataForThisFrame; 
+	mutable TArray<TArray<FParticleVertexFactoryBase*>, TInlineAllocator<1>> EmitterVertexFactoryArray;
+	mutable TArray<FDynamicEmitterDataBase*> DynamicDataForThisFrame;
 
 	friend struct FDynamicSpriteEmitterDataBase;
 };
