Use implicit bounds in heightfield. 83% memory reduction

diff --git a/Engine/Source/Runtime/Core/Public/UObject/ExternalPhysicsCustomObjectVersion.h b/Engine/Source/Runtime/Core/Public/UObject/ExternalPhysicsCustomObjectVersion.h
--- a/Engine/Source/Runtime/Core/Public/UObject/ExternalPhysicsCustomObjectVersion.h
+++ b/Engine/Source/Runtime/Core/Public/UObject/ExternalPhysicsCustomObjectVersion.h
@@ -77,6 +77,9 @@ struct CORE_API FExternalPhysicsCustomObjectVersion
 
 		// Added ability to remove shapes from collision resolution (will not construct constraints when one or more shapes removed)
 		AddShapeCollisionDisable,
+
+		//Heightfield cell bounds are implicit
+		HeightfieldImplicitBounds,
 		
 		// -----<new versions can be added above this line>-------------------------------------------------
 		VersionPlusOne,
diff --git a/Engine/Source/Runtime/Engine/Private/PhysicsEngine/Experimental/ChaosDerivedData.cpp b/Engine/Source/Runtime/Engine/Private/PhysicsEngine/Experimental/ChaosDerivedData.cpp
--- a/Engine/Source/Runtime/Engine/Private/PhysicsEngine/Experimental/ChaosDerivedData.cpp
+++ b/Engine/Source/Runtime/Engine/Private/PhysicsEngine/Experimental/ChaosDerivedData.cpp
@@ -21,7 +21,7 @@ const TCHAR* FChaosDerivedDataCooker::GetPluginName() const
 
 const TCHAR* FChaosDerivedDataCooker::GetVersionString() const
 {
-	return TEXT("307A91E898C442A1A0650CD258F76505");
+	return TEXT("8AEC8C7BFB0A4906A6EB99E602E9594B");
 }
 
 FString FChaosDerivedDataCooker::GetPluginSpecificCacheKeySuffix() const
@@ -85,7 +85,7 @@ FChaosDerivedDataCooker::FChaosDerivedDataCooker(UBodySetup* InSetup, FName InFo
 }
 
 template<typename Precision>
-void FChaosDerivedDataCooker::BuildTriangleMeshes(TArray<TUniquePtr<Chaos::TTriangleMeshImplicitObject<Precision>>>& OutTriangleMeshes, TArray<int32>& OutFaceRemap, const FCookBodySetupInfo& InParams)
+void FChaosDerivedDataCooker::BuildTriangleMeshes(TArray<TUniquePtr<Chaos::TTriangleMeshImplicitObject<Precision>>>& OutTriangleMeshes, const FCookBodySetupInfo& InParams)
 {
 	if(!InParams.bCookTriMesh)
 	{
@@ -106,7 +106,7 @@ void FChaosDerivedDataCooker::BuildTriangleMeshes(TArray<TUniquePtr<Chaos::TTria
 	}
 
 	TArray<int32> Remap;
-	Chaos::CleanTrimesh(FinalVerts, FinalIndices, &OutFaceRemap);
+	Chaos::CleanTrimesh(FinalVerts, FinalIndices, &Remap);
 
 	// Build particle list #BG Maybe allow TParticles to copy vectors?
 	Chaos::TParticles<Precision, 3> TriMeshParticles;
@@ -147,14 +147,14 @@ void FChaosDerivedDataCooker::BuildTriangleMeshes(TArray<TUniquePtr<Chaos::TTria
 
 			if(bHasMaterials)
 			{
-				if(!ensure(OutFaceRemap.IsValidIndex(TriangleIndex)))
+				if(!ensure(Remap.IsValidIndex(TriangleIndex)))
 				{
 					MaterialIndices.Empty();
 					bHasMaterials = false;
 				}
 				else
 				{
-					const int32 OriginalIndex = OutFaceRemap[TriangleIndex];
+					const int32 OriginalIndex = Remap[TriangleIndex];
 
 					if(ensure(InParams.TriangleMeshDesc.MaterialIndices.IsValidIndex(OriginalIndex)))
 					{
@@ -173,9 +173,9 @@ void FChaosDerivedDataCooker::BuildTriangleMeshes(TArray<TUniquePtr<Chaos::TTria
 	OutTriangleMeshes.Emplace(new Chaos::TTriangleMeshImplicitObject<Precision>(MoveTemp(TriMeshParticles), MoveTemp(Triangles), MoveTemp(MaterialIndices)));
 }
 
-template void FChaosDerivedDataCooker::BuildTriangleMeshes(TArray<TUniquePtr<Chaos::TTriangleMeshImplicitObject<float>>>& OutTriangleMeshes, TArray<int32>& OutFaceRemap, const FCookBodySetupInfo& InParams);
+template void FChaosDerivedDataCooker::BuildTriangleMeshes(TArray<TUniquePtr<Chaos::TTriangleMeshImplicitObject<float>>>& OutTriangleMeshes, const FCookBodySetupInfo& InParams);
 //#BGTODO When it's possible to build with doubles, re-enable this. (Currently at least TRigidTransform cannot build with double precision because we don't have a base transform implementation using them)
-//template void FChaosDerivedDataCooker::BuildTriangleMeshes(TArray<TUniquePtr<Chaos::FImplicitObject>>& OutTriangleMeshes, TArray<int32>& OutFaceRemap, const FCookBodySetupInfo& InParams);
+//template void FChaosDerivedDataCooker::BuildTriangleMeshes(TArray<TUniquePtr<Chaos::FImplicitObject>>& OutTriangleMeshes, const FCookBodySetupInfo& InParams);
 
 void FChaosDerivedDataCooker::BuildConvexMeshes(TArray<TUniquePtr<Chaos::FImplicitObject>>& OutConvexMeshes, const FCookBodySetupInfo& InParams)
 {
@@ -250,11 +250,10 @@ void FChaosDerivedDataCooker::BuildInternal(Chaos::FChaosArchive& Ar, FCookBodyS
 {
 	TArray<TUniquePtr<Chaos::FImplicitObject>> SimpleImplicits;
 	TArray<TUniquePtr<Chaos::TTriangleMeshImplicitObject<Precision>>> ComplexImplicits;
-	TArray<int32> FaceRemap;
 
 	//BuildSimpleShapes(SimpleImplicits, Setup);
 	BuildConvexMeshes(SimpleImplicits, InInfo);
-	BuildTriangleMeshes(ComplexImplicits, FaceRemap, InInfo);
+	BuildTriangleMeshes(ComplexImplicits, InInfo);
 
 	//TUniquePtr<Chaos::TImplicitObjectUnion<Precision, 3>> SimpleUnion = MakeUnique<Chaos::TImplicitObjectUnion<Precision, 3>>(MoveTemp(SimpleImplicits));
 	//TUniquePtr<Chaos::TImplicitObjectUnion<Precision, 3>> ComplexUnion = MakeUnique<Chaos::TImplicitObjectUnion<Precision, 3>>(MoveTemp(ComplexImplicits));
@@ -265,12 +264,7 @@ void FChaosDerivedDataCooker::BuildInternal(Chaos::FChaosArchive& Ar, FCookBodyS
 		UVInfo.FillFromTriMesh(InInfo.TriangleMeshDesc);
 	}
 
-	if (!InInfo.bSupportFaceRemap)
-	{
-		FaceRemap.Empty();
-	}
-
-	Ar << SimpleImplicits << ComplexImplicits << UVInfo << FaceRemap;
+	Ar << SimpleImplicits << ComplexImplicits << UVInfo;
 }
 
 template void FChaosDerivedDataCooker::BuildInternal<float>(Chaos::FChaosArchive& Ar, FCookBodySetupInfo& InInfo);
diff --git a/Engine/Source/Runtime/Experimental/Chaos/Private/Chaos/HeightField.cpp b/Engine/Source/Runtime/Experimental/Chaos/Private/Chaos/HeightField.cpp
--- a/Engine/Source/Runtime/Experimental/Chaos/Private/Chaos/HeightField.cpp
+++ b/Engine/Source/Runtime/Experimental/Chaos/Private/Chaos/HeightField.cpp
@@ -492,11 +492,10 @@ namespace Chaos
 	template<typename T>
 	bool Chaos::THeightField<T>::GetCellBounds2D(const TVector<int32, 2> InCoord, FBounds2D& OutBounds, const TVector<T, 2>& InInflate /*= {0}*/) const
 	{
-		if(FlatGrid.IsValid(InCoord))
+		if (FlatGrid.IsValid(InCoord))
 		{
-			const TBox<T, 3>& Bound = GeomData.CellBounds[InCoord[1] * (GeomData.NumCols - 1) + InCoord[0]];
-			OutBounds.Min = {Bound.Min()[0], Bound.Min()[1]};
-			OutBounds.Max = {Bound.Max()[0], Bound.Max()[1]};
+			OutBounds.Min = TVector<T, 2>(InCoord[0], InCoord[1]);
+			OutBounds.Max = TVector<T, 2>(InCoord[0] + 1, InCoord[1] + 1);
 			OutBounds.Min -= InInflate;
 			OutBounds.Max += InInflate;
 
@@ -509,11 +508,12 @@ namespace Chaos
 	template<typename T>
 	bool Chaos::THeightField<T>::GetCellBounds3D(const TVector<int32, 2> InCoord, TVector<T, 3>& OutMin, TVector<T, 3>& OutMax, const TVector<T, 3>& InInflate /*= {0}*/) const
 	{
-		if(FlatGrid.IsValid(InCoord))
+		if (FlatGrid.IsValid(InCoord))
 		{
-			const TBox<T, 3>& Bound = GeomData.CellBounds[InCoord[1] * (GeomData.NumCols - 1) + InCoord[0]];
-			OutMin = Bound.Min() - InInflate;
-			OutMax = Bound.Max() + InInflate;
+			OutMin = TVec3<T>(InCoord[0], InCoord[1], 0);
+			OutMax = TVec3<T>(InCoord[0] + 1, InCoord[1] + 1, GeomData.CellHeights[InCoord[1] * (GeomData.NumCols - 1) + InCoord[0]]);
+			OutMin = OutMin - InInflate;
+			OutMax = OutMax + InInflate;
 
 			return true;
 		}
@@ -524,11 +524,10 @@ namespace Chaos
 	template<typename T>
 	bool Chaos::THeightField<T>::GetCellBounds2DScaled(const TVector<int32, 2> InCoord, FBounds2D& OutBounds, const TVector<T, 2>& InInflate /*= {0}*/) const
 	{
-		if(FlatGrid.IsValid(InCoord))
+		if (FlatGrid.IsValid(InCoord))
 		{
-			const TBox<T, 3>& Bound = GeomData.CellBounds[InCoord[1] * (GeomData.NumCols - 1) + InCoord[0]];
-			OutBounds.Min = {Bound.Min()[0], Bound.Min()[1]};
-			OutBounds.Max = {Bound.Max()[0], Bound.Max()[1]};
+			OutBounds.Min = TVector<T, 2>(InCoord[0], InCoord[1]);
+			OutBounds.Max = TVector<T, 2>(InCoord[0] + 1, InCoord[1] + 1);
 			OutBounds.Min -= InInflate;
 			OutBounds.Max += InInflate;
 			const TVector<float, 2> Scale2D = TVector<T, 2>(GeomData.Scale[0], GeomData.Scale[1]);
@@ -543,11 +542,10 @@ namespace Chaos
 	template<typename T>
 	bool Chaos::THeightField<T>::GetCellBounds3DScaled(const TVector<int32, 2> InCoord, TVector<T, 3>& OutMin, TVector<T, 3>& OutMax, const TVector<T, 3>& InInflate /*= {0}*/) const
 	{
-		if(FlatGrid.IsValid(InCoord))
+		if (FlatGrid.IsValid(InCoord))
 		{
-			const TBox<T, 3>& Bound = GeomData.CellBounds[InCoord[1] * (GeomData.NumCols - 1) + InCoord[0]];
-			OutMin = Bound.Min();
-			OutMax = Bound.Max();
+			OutMin = TVec3<T>(InCoord[0], InCoord[1], 0);
+			OutMax = TVec3<T>(InCoord[0] + 1, InCoord[1] + 1, GeomData.CellHeights[InCoord[1] * (GeomData.NumCols - 1) + InCoord[0]]);
 			OutMin = OutMin * GeomData.Scale - InInflate;
 			OutMax = OutMax * GeomData.Scale + InInflate;
 			return true;
@@ -1350,7 +1348,8 @@ namespace Chaos
 		// Cache per-cell bounds
 		const int32 NumX = GeomData.NumCols - 1;
 		const int32 NumY = GeomData.NumRows - 1;
-		GeomData.CellBounds.SetNum(NumX * NumY);
+		//GeomData.CellBounds.SetNum(NumX * NumY);
+		GeomData.CellHeights.SetNum(NumX*NumY);
 		for(int32 XIndex = 0; XIndex < NumX; ++XIndex)
 		{
 			for(int32 YIndex = 0; YIndex < NumY; ++YIndex)
@@ -1358,7 +1357,8 @@ namespace Chaos
 				const TVector<int32, 2> Cell(XIndex, YIndex);
 				TVector<T, 3> Min, Max;
 				CalcCellBounds3D(Cell, Min, Max);
-				GeomData.CellBounds[YIndex * NumX + XIndex] = TBox<T, 3>(Min, Max);
+				//GeomData.CellBounds[YIndex * NumX + XIndex] = TBox<T, 3>(Min, Max);
+				GeomData.CellHeights[YIndex * NumX + XIndex] = Max.Z;
 			}
 		}
 	}
diff --git a/Engine/Source/Runtime/Experimental/Chaos/Public/Chaos/HeightField.h b/Engine/Source/Runtime/Experimental/Chaos/Public/Chaos/HeightField.h
--- a/Engine/Source/Runtime/Experimental/Chaos/Public/Chaos/HeightField.h
+++ b/Engine/Source/Runtime/Experimental/Chaos/Public/Chaos/HeightField.h
@@ -183,7 +183,7 @@ namespace Chaos
 			uint16 NumCols;
 			RealType Range;
 			RealType HeightPerUnit;
-			TArray<TBox<RealType, 3>> CellBounds;
+			TArray<RealType> CellHeights; //todo: remove this and use heights directly
 
 			constexpr float GetCellWidth() const
 			{
@@ -247,7 +247,20 @@ namespace Chaos
 				{
 					Ar << Range;
 					Ar << HeightPerUnit;
-					Ar << CellBounds;
+
+					if (Ar.CustomVer(FExternalPhysicsCustomObjectVersion::GUID) < FExternalPhysicsCustomObjectVersion::HeightfieldImplicitBounds)
+					{
+						TArray<TBox<RealType, 3>> CellBounds;
+						Ar << CellBounds;
+						for (const TBox<RealType, 3>& Box : CellBounds)
+						{
+							CellHeights.Add(Box.Extents().Z);
+						}
+					}
+					else
+					{
+						Ar << CellHeights;
+					}
 				}
 
 				if(Ar.CustomVer(FExternalPhysicsCustomObjectVersion::GUID) >= FExternalPhysicsCustomObjectVersion::AddedMaterialManager)
