Change UActorChannel::ProcessQueuedBunches to better respect timeslicing.

Prior to this change, every channel that had queued bunches remaining would log warnings regardless of whether or not the channel was simply ignored for processing.
That meant a single hitch could cause a cascading effect where all (or most) channels would spam logs, which could then cause additional hitches, etc.

diff --git a/Engine/Source/Runtime/Engine/Classes/Engine/ActorChannel.h b/Engine/Source/Runtime/Engine/Classes/Engine/ActorChannel.h
--- a/Engine/Source/Runtime/Engine/Classes/Engine/ActorChannel.h
+++ b/Engine/Source/Runtime/Engine/Classes/Engine/ActorChannel.h
@@ -87,6 +87,12 @@ private:
 	/** Tracks whether or not our actor has been seen as pending kill. */
 	uint32 bActorIsPendingKill : 1;
 
+	/**
+	 * Whether or not our NetDriver detected a hitch somewhere else in the engine.
+	 * This is used by ProcessQueuedBunches to prevent erroneous log spam.
+	 */
+	uint32 bSuppressQueuedBunchWarningsDueToHitches : 1;
+
 public:
 	bool GetSkipRoleSwap() const { return !!bSkipRoleSwap; }
 	void SetSkipRoleSwap(const bool bShouldSkip) { bSkipRoleSwap = bShouldSkip; }
diff --git a/Engine/Source/Runtime/Engine/Classes/Engine/NetDriver.h b/Engine/Source/Runtime/Engine/Classes/Engine/NetDriver.h
--- a/Engine/Source/Runtime/Engine/Classes/Engine/NetDriver.h
+++ b/Engine/Source/Runtime/Engine/Classes/Engine/NetDriver.h
@@ -1581,6 +1581,8 @@ public:
 	 */
 	inline void IncreaseOutTotalNotifiedPackets() { ++OutTotalNotifiedPackets; }
 
+	bool DidHitchLastFrame() const;
+
 protected:
 
 	bool bMaySendProperties;
@@ -1626,4 +1628,5 @@ private:
 	bool bForcedPacketSettings;
 #endif 
 
+	bool bDidHitchLastFrame = false;
 };
\ No newline at end of file
diff --git a/Engine/Source/Runtime/Engine/Private/DataChannel.cpp b/Engine/Source/Runtime/Engine/Private/DataChannel.cpp
--- a/Engine/Source/Runtime/Engine/Private/DataChannel.cpp
+++ b/Engine/Source/Runtime/Engine/Private/DataChannel.cpp
@@ -2275,25 +2275,24 @@ bool UActorChannel::CanStopTicking() const
 
 bool UActorChannel::ProcessQueuedBunches()
 {
-
 	const uint32 QueueBunchStartCycles = FPlatformTime::Cycles();
 
 	// Try to resolve any guids that are holding up the network stream on this channel
-	for ( auto It = PendingGuidResolves.CreateIterator(); It; ++It )
+	// TODO: This could take a non-trivial amount of time since both GetObjectFromNetGUID
+	// and IsGUIDBroken may do Map Lookups and GetObjectFromNetGUID will attempt
+	// to resolve weak objects.
+	for (auto It = PendingGuidResolves.CreateIterator(); It; ++It)
 	{
-		if ( Connection->Driver->GuidCache->GetObjectFromNetGUID( *It, true ) != NULL )
+		if (Connection->Driver->GuidCache->GetObjectFromNetGUID(*It, true))
 		{
 			// This guid is now resolved, we can remove it from the pending guid list
 			It.RemoveCurrent();
-			continue;
 		}
-
-		if ( Connection->Driver->GuidCache->IsGUIDBroken( *It, true ) )
+		else if (Connection->Driver->GuidCache->IsGUIDBroken(*It, true))
 		{
 			// This guid is broken, remove it, and warn
-			UE_LOG( LogNet, Warning, TEXT( "UActorChannel::ProcessQueuedBunches: Guid is broken. NetGUID: %s, ChIndex: %i, Actor: %s" ), *It->ToString(), ChIndex, Actor != NULL ? *Actor->GetPathName() : TEXT( "NULL" ) );
+			UE_LOG(LogNet, Warning, TEXT("UActorChannel::ProcessQueuedBunches: Guid is broken. NetGUID: %s, ChIndex: %i, Actor: %s"), *It->ToString(), ChIndex, *GetPathNameSafe(Actor));
 			It.RemoveCurrent();
-			continue;
 		}
 	}
 
@@ -2302,6 +2301,9 @@ bool UActorChannel::ProcessQueuedBunches()
 		return true;
 	}
 
+	// Always update this when there are bunches remaining, because we may not hitch on the frame that triggers a warning.
+	bSuppressQueuedBunchWarningsDueToHitches |= Connection->Driver->DidHitchLastFrame();
+
 	// Instant replays are played back in a duplicated level collection, so if this is instant replay
 	// playback, the driver's DuplicateLevelID will be something other than INDEX_NONE.
 	const int BunchTimeLimit = Connection->Driver->GetDuplicateLevelID() == INDEX_NONE ?
@@ -2310,69 +2312,74 @@ bool UActorChannel::ProcessQueuedBunches()
 
 	const bool bHasTimeToProcess = BunchTimeLimit == 0 || Connection->Driver->ProcessQueuedBunchesCurrentFrameMilliseconds < BunchTimeLimit;
 
-	// We can process all of the queued up bunches if ALL of these are true:
-	//	1. We no longer have any pending guids to load
-	//	2. We aren't still processing bunches on another channel that this actor was previously on
-	//	3. We haven't spent too much time yet this frame processing queued bunches
-	//	4. The driver isn't requesting queuing for this GUID
-	if ( PendingGuidResolves.Num() == 0 && ( ChIndex == -1 || !Connection->KeepProcessingActorChannelBunchesMap.Contains( ActorNetGUID ) ) &&
-		 bHasTimeToProcess && !Connection->Driver->ShouldQueueBunchesForActorGUID( ActorNetGUID ) )
+	// If we don't have any time, then don't bother doing anything (including warning) as that may make things worse.
+	if (bHasTimeToProcess)
 	{
-		DECLARE_SCOPE_CYCLE_COUNTER(TEXT("ProcessQueuedBunches time"), STAT_ProcessQueuedBunchesTime, STATGROUP_Net);
-		for ( int32 i = 0; i < QueuedBunches.Num(); i++ )
+		// We can process all of the queued up bunches if ALL of these are true:
+		//	1. We no longer have any pending guids to load
+		//	2. We aren't still processing bunches on another channel that this actor was previously on
+		//	3. We haven't spent too much time yet this frame processing queued bunches
+		//	4. The driver isn't requesting queuing for this GUID
+		if (PendingGuidResolves.Num() == 0
+			&& (ChIndex == -1 || !Connection->KeepProcessingActorChannelBunchesMap.Contains(ActorNetGUID))
+			&& !Connection->Driver->ShouldQueueBunchesForActorGUID(ActorNetGUID))
 		{
-			ProcessBunch( *QueuedBunches[i] );
-			delete QueuedBunches[i];
-		}
+			DECLARE_SCOPE_CYCLE_COUNTER(TEXT("ProcessQueuedBunches time"), STAT_ProcessQueuedBunchesTime, STATGROUP_Net);
+			for ( int32 i = 0; i < QueuedBunches.Num(); i++ )
+			{
+				ProcessBunch( *QueuedBunches[i] );
+				delete QueuedBunches[i];
+			}
 
-		UE_LOG( LogNet, VeryVerbose, TEXT( "UActorChannel::ProcessQueuedBunches: Flushing queued bunches. ChIndex: %i, Actor: %s, Queued: %i" ), ChIndex, Actor != NULL ? *Actor->GetPathName() : TEXT( "NULL" ), QueuedBunches.Num() );
+			UE_LOG(LogNet, VeryVerbose, TEXT("UActorChannel::ProcessQueuedBunches: Flushing queued bunches. ChIndex: %i, Actor: %s, Queued: %i"), ChIndex, Actor != NULL ? *Actor->GetPathName() : TEXT("NULL"), QueuedBunches.Num());
 
-		QueuedBunches.Empty();
+			QueuedBunches.Empty();
 
-		// Call any onreps that were delayed because we were queuing bunches
-		for (auto& ReplicatorPair : ReplicationMap)
-		{
-			ReplicatorPair.Value->CallRepNotifies(true);
-		}
-	}
+			// Call any onreps that were delayed because we were queuing bunches
+			for (auto& ReplicatorPair : ReplicationMap)
+			{
+				ReplicatorPair.Value->CallRepNotifies(true);
+			}
 
-	// Warn when we have queued bunches for a very long time
-	if ( QueuedBunches.Num() > 0 )
-	{
-		const double QUEUED_BUNCH_TIMEOUT_IN_SECONDS = 30;
+			if (UPackageMapClient * PackageMapClient = Cast< UPackageMapClient >(Connection->PackageMap))
+			{
+				PackageMapClient->SetHasQueuedBunches(ActorNetGUID, false);
+			}
 
-		if ( FPlatformTime::Seconds() - QueuedBunchStartTime > QUEUED_BUNCH_TIMEOUT_IN_SECONDS )
+			QueuedBunchObjectReferences.Empty();
+		}
+		else
 		{
-			UE_CLOG(FPlatformProperties::RequiresCookedData(), LogNet, Warning, TEXT( "UActorChannel::ProcessQueuedBunches: Queued bunches for longer than normal. ChIndex: %i, Actor: %s, Queued: %i, PendingGuidResolves: %i" ), ChIndex, Actor != NULL ? *Actor->GetPathName() : TEXT( "NULL" ), QueuedBunches.Num(), PendingGuidResolves.Num() );
-			for (const FNetworkGUID& Guid : PendingGuidResolves)
+			const double QUEUED_BUNCH_TIMEOUT_IN_SECONDS = 30;
+			if ((FPlatformTime::Seconds() - QueuedBunchStartTime) > QUEUED_BUNCH_TIMEOUT_IN_SECONDS)
 			{
-				const bool bIsPending = Connection->Driver->GuidCache->IsGUIDPending(Guid);
-				const FString GuidPath = Connection->Driver->GuidCache->FullNetGUIDPath(Guid);
-				UE_CLOG(FPlatformProperties::RequiresCookedData(), LogNet, Log, TEXT( "  PendingGuidResolve %s. Pending in cache: %d, path: %s" ), *Guid.ToString(), bIsPending ? 1 : 0, *GuidPath );
+				if (!bSuppressQueuedBunchWarningsDueToHitches && FPlatformProperties::RequiresCookedData())
+				{
+					UE_LOG(LogNet, Warning, TEXT("UActorChannel::ProcessQueuedBunches: Queued bunches for longer than normal. ChIndex: %i, Actor: %s, Queued: %i, PendingGuidResolves: %i"), ChIndex, *GetPathNameSafe(Actor), QueuedBunches.Num(), PendingGuidResolves.Num());
+
+					if (UE_LOG_ACTIVE(LogNet, Log))
+					{
+						for (const FNetworkGUID& Guid : PendingGuidResolves)
+						{
+							const bool bIsPending = Connection->Driver->GuidCache->IsGUIDPending(Guid);
+							const FString GuidPath = Connection->Driver->GuidCache->FullNetGUIDPath(Guid);
+							UE_LOG(LogNet, Log, TEXT("  PendingGuidResolve %s. Pending in cache: %d, path: %s"), *Guid.ToString(), !!bIsPending, *GuidPath);
+						}
+					}
+				}
+
+				QueuedBunchStartTime = FPlatformTime::Seconds();
 			}
-			QueuedBunchStartTime = FPlatformTime::Seconds();
 		}
-	}
-	else
-	{
-		// Processed all bunches
-		UPackageMapClient * PackageMapClient = Cast< UPackageMapClient >(Connection->PackageMap);
 
-		if (PackageMapClient)
-		{
-			PackageMapClient->SetHasQueuedBunches(ActorNetGUID, false);
-		}
+		// Update the driver with our time spent
+		const uint32 QueueBunchEndCycles = FPlatformTime::Cycles();
+		const uint32 QueueBunchDeltaCycles = QueueBunchEndCycles - QueueBunchStartCycles;
+		const float QueueBunchDeltaMilliseconds = FPlatformTime::ToMilliseconds(QueueBunchDeltaCycles);
 
-		QueuedBunchObjectReferences.Empty();
+		Connection->Driver->ProcessQueuedBunchesCurrentFrameMilliseconds += QueueBunchDeltaMilliseconds;
 	}
 
-	// Update the driver with our time spent
-	const uint32 QueueBunchEndCycles = FPlatformTime::Cycles();
-	const uint32 QueueBunchDeltaCycles = QueueBunchEndCycles - QueueBunchStartCycles;
-	const float QueueBunchDeltaMilliseconds = FPlatformTime::ToMilliseconds(QueueBunchDeltaCycles);
-
-	Connection->Driver->ProcessQueuedBunchesCurrentFrameMilliseconds += QueueBunchDeltaMilliseconds;
-
 	// Return true if we are done processing queued bunches
 	return QueuedBunches.Num() == 0;
 }
@@ -2489,6 +2496,7 @@ void UActorChannel::ReceivedBunch( FInBunch & Bunch )
 			{
 				// Remember when we first started queuing
 				QueuedBunchStartTime = FPlatformTime::Seconds();
+				bSuppressQueuedBunchWarningsDueToHitches = false;
 			}
 
 			QueuedBunches.Add(new FInBunch(Bunch));
@@ -3890,6 +3898,7 @@ void UActorChannel::AddedToChannelPool()
 	bSkipRoleSwap = false;
 	bClearRecentActorRefs = true;
 	QueuedBunchStartTime = 0;
+	bSuppressQueuedBunchWarningsDueToHitches = false;
 	CreateSubObjects.Empty();
 #if !UE_BUILD_SHIPPING
 	bBlockChannelFailure = false;
diff --git a/Engine/Source/Runtime/Engine/Private/NetworkDriver.cpp b/Engine/Source/Runtime/Engine/Private/NetworkDriver.cpp
--- a/Engine/Source/Runtime/Engine/Private/NetworkDriver.cpp
+++ b/Engine/Source/Runtime/Engine/Private/NetworkDriver.cpp
@@ -1711,7 +1711,9 @@ void UNetDriver::TickDispatch( float DeltaTime )
 	// Setting this to somewhat large value for now, but small enough to catch blocking calls that are causing timeouts
 	const float TickLogThreshold = 5.0f;
 
-	if ( DeltaTime > TickLogThreshold || DeltaRealtime > TickLogThreshold )
+	bDidHitchLastFrame = (DeltaTime > TickLogThreshold || DeltaRealtime > TickLogThreshold);
+
+	if (bDidHitchLastFrame)
 	{
 		UE_LOG( LogNet, Log, TEXT( "UNetDriver::TickDispatch: Very long time between ticks. DeltaTime: %2.2f, Realtime: %2.2f. %s" ), DeltaTime, DeltaRealtime, *GetName() );
 	}
@@ -5639,6 +5641,11 @@ void UNetDriver::ResetAsyncLoadDelinquencyAnalytics()
 	}
 }
 
+bool UNetDriver::DidHitchLastFrame() const
+{
+	return bDidHitchLastFrame;
+}
+
 ECreateReplicationChangelistMgrFlags UNetDriver::GetCreateReplicationChangelistMgrFlags() const
 {
 	return ECreateReplicationChangelistMgrFlags::None;
