Fixed convex collisions no longer rendering in the static mesh tool or in the world when the collision show flag is enabled.

diff --git a/Engine/Source/Editor/UnrealEd/Private/ConvexDecompTool.cpp b/Engine/Source/Editor/UnrealEd/Private/ConvexDecompTool.cpp
--- a/Engine/Source/Editor/UnrealEd/Private/ConvexDecompTool.cpp
+++ b/Engine/Source/Editor/UnrealEd/Private/ConvexDecompTool.cpp
@@ -208,6 +208,14 @@ public:
 
 						ConvexElem.VertexData.Add(V);
 					}
+
+					const uint32 NumIndexEntries = Hull.m_nTriangles * 3;
+					ConvexElem.IndexData.Reset(Hull.m_nTriangles * 3);
+					for(uint32 Index = 0; Index < NumIndexEntries; ++Index)
+					{
+						ConvexElem.IndexData.Add(Hull.m_triangles[Index]);
+					}
+
 					ConvexElem.UpdateElemBox();
 					InBodySetup->AggGeom.ConvexElems.Add(ConvexElem);
 				}
diff --git a/Engine/Source/Runtime/Engine/Classes/PhysicsEngine/ConvexElem.h b/Engine/Source/Runtime/Engine/Classes/PhysicsEngine/ConvexElem.h
--- a/Engine/Source/Runtime/Engine/Classes/PhysicsEngine/ConvexElem.h
+++ b/Engine/Source/Runtime/Engine/Classes/PhysicsEngine/ConvexElem.h
@@ -38,6 +38,9 @@ struct FKConvexElem : public FKShapeElem
 	UPROPERTY()
 	TArray<FVector> VertexData;
 
+	UPROPERTY()
+	TArray<int32> IndexData;
+
 	/** Bounding box of this convex hull. */
 	UPROPERTY()
 	FBox ElemBox;
diff --git a/Engine/Source/Runtime/Engine/Private/Collision/KAggregateGeom.cpp b/Engine/Source/Runtime/Engine/Private/Collision/KAggregateGeom.cpp
--- a/Engine/Source/Runtime/Engine/Private/Collision/KAggregateGeom.cpp
+++ b/Engine/Source/Runtime/Engine/Private/Collision/KAggregateGeom.cpp
@@ -141,23 +141,28 @@ void FKAggregateGeom::CalcBoxSphereBounds(FBoxSphereBounds& Output, const FTrans
 }
 
 
-static void AddVertexIfNotPresent(TArray<FVector> &Vertices, const FVector& NewVertex)
+static int32 AddVertexIfNotPresent(TArray<FVector> &Vertices, const FVector& NewVertex)
 {
 	bool bIsPresent = 0;
+	int32 Index = INDEX_NONE;
 
-	for(int32 i=0; i<Vertices.Num() && !bIsPresent; i++)
+	for(int32 i = 0; i < Vertices.Num() && !bIsPresent; i++)
 	{
 		const float DiffSqr = (NewVertex - Vertices[i]).SizeSquared();
 		if(DiffSqr < MIN_HULL_VERT_DISTANCE * MIN_HULL_VERT_DISTANCE)
 		{
 			bIsPresent = true;
+			Index = i;
 		}
 	}
 
 	if(!bIsPresent)
 	{
 		Vertices.Add(NewVertex);
+		Index = Vertices.Num() - 1;
 	}
+
+	return Index;
 }
 
 static void RemoveDuplicateVerts(TArray<FVector>& InVerts)
@@ -536,16 +541,20 @@ bool FKConvexElem::HullFromPlanes(const TArray<FPlane>& InPlanes, const TArray<F
 		// Do nothing if poly was completely clipped away.
 		if(Polygon.Vertices.Num() > 0)
 		{
+			const int32 BaseIndex = VertexData.Num();
+			TArray<int32> Remap;
+			Remap.AddUninitialized(Polygon.Vertices.Num());
+
 			TotalPolyArea += Polygon.Area();
 
 			// Add vertices of polygon to convex primitive.
-			for(int32 j=0; j<Polygon.Vertices.Num(); j++)
+			for(int32 j = 0; j < Polygon.Vertices.Num() ; j++)
 			{
 				// We try and snap the vert to on of the ones supplied.
 				int32 NearestVert = INDEX_NONE;
 				float NearestDistSqr = BIG_NUMBER;
 
-				for(int32 k=0; k<SnapVerts.Num(); k++)
+				for(int32 k = 0; k < SnapVerts.Num(); k++)
 				{
 					const float DistSquared = (Polygon.Vertices[j] - SnapVerts[k]).SizeSquared();
 
@@ -560,14 +569,22 @@ bool FKConvexElem::HullFromPlanes(const TArray<FPlane>& InPlanes, const TArray<F
 				if( NearestVert != INDEX_NONE && NearestDistSqr < LOCAL_EPS )
 				{
 					const FVector localVert = SnapVerts[NearestVert];
-					AddVertexIfNotPresent(VertexData, localVert);
+					Remap[j] = AddVertexIfNotPresent(VertexData, localVert);
 				}
 				else
 				{
 					const FVector localVert = Polygon.Vertices[j];
-					AddVertexIfNotPresent(VertexData, localVert);
+					Remap[j] = AddVertexIfNotPresent(VertexData, localVert);
 				}
 			}
+
+			int32 NumTriangles = Polygon.Vertices.Num() - 2;
+			for(int32 Index = 1; Index < NumTriangles; ++Index)
+			{
+				IndexData.Add(BaseIndex);
+				IndexData.Add(Remap[Index]);
+				IndexData.Add(Remap[Index + 1]);
+			}
 		}
 	}
 
@@ -660,6 +677,12 @@ void FKConvexElem::ConvexFromBoxElem(const FKBoxElem& InBox)
 	{
 		for (int32 j = 0; j < 2; j++)
 		{
+			const int32 Base = VertexData.Num();
+			for(int32 k = 0; k < 6; ++k)
+			{
+				IndexData.Add(Base + k);
+			}
+
 			P.X = B[i].X; Q.X = B[i].X;
 			P.Y = B[j].Y; Q.Y = B[j].Y;
 			P.Z = B[0].Z; Q.Z = B[1].Z;
diff --git a/Engine/Source/Runtime/Engine/Private/PhysicsEngine/BodySetup.cpp b/Engine/Source/Runtime/Engine/Private/PhysicsEngine/BodySetup.cpp
--- a/Engine/Source/Runtime/Engine/Private/PhysicsEngine/BodySetup.cpp
+++ b/Engine/Source/Runtime/Engine/Private/PhysicsEngine/BodySetup.cpp
@@ -49,6 +49,7 @@
 #if WITH_CHAOS
 	#include "Experimental/ChaosDerivedData.h"
 	#include "Physics/Experimental/ChaosDerivedDataReader.h"
+	#include "Chaos/CollisionConvexMesh.h"
 #endif
 
 /** Helper for enum output... */
@@ -1090,6 +1091,39 @@ void UBodySetup::PostLoad()
 			}
 		}
 	}
+
+#if WITH_CHAOS
+	// For drawing of convex elements we require an index buffer, previously we could
+	// get this from a PxConvexMesh but Chaos doesn't maintain that data. Instead now
+	// it is a part of the element rather than the physics geometry, if we load in an
+	// element without that data present, generate a convex hull from the convex vert
+	// data and extract the index data from there.
+	for(FKConvexElem& Convex : AggGeom.ConvexElems)
+	{
+		const int32 NumVerts = Convex.VertexData.Num();
+		if(NumVerts > 0 && Convex.IndexData.Num() == 0)
+		{
+			Chaos::TParticles<Chaos::FReal, 3> ConvexParticles;
+			ConvexParticles.AddParticles(NumVerts);
+
+			for(int32 VertIndex = 0; VertIndex < NumVerts; ++VertIndex)
+			{
+				ConvexParticles.X(VertIndex) = Convex.VertexData[VertIndex];
+			}
+
+			TArray<Chaos::TVector<int32, 3>> Triangles;
+			Chaos::FConvexBuilder::BuildConvexHull(ConvexParticles, Triangles);
+
+			Convex.IndexData.Reset(Triangles.Num() * 3);
+			for(Chaos::TVector<int32, 3> Tri : Triangles)
+			{
+				Convex.IndexData.Add(Tri[0]);
+				Convex.IndexData.Add(Tri[1]);
+				Convex.IndexData.Add(Tri[2]);
+			}
+		}
+	}
+#endif
 }
 
 void UBodySetup::UpdateTriMeshVertices(const TArray<FVector> & NewPositions)
@@ -1603,10 +1637,9 @@ void FKConvexElem::CloneElem(const FKConvexElem& Other)
 {
 	Super::CloneElem(Other);
 	VertexData = Other.VertexData;
+	IndexData = Other.IndexData;
 	ElemBox = Other.ElemBox;
 	Transform = Other.Transform;
-
-	// TODO: Should this also copy the ChaosConvexMesh?
 }
 
 void FKConvexElem::ScaleElem(FVector DeltaSize, float MinSize)
diff --git a/Engine/Source/Runtime/Engine/Private/PhysicsEngine/PhysDrawing.cpp b/Engine/Source/Runtime/Engine/Private/PhysicsEngine/PhysDrawing.cpp
--- a/Engine/Source/Runtime/Engine/Private/PhysicsEngine/PhysDrawing.cpp
+++ b/Engine/Source/Runtime/Engine/Private/PhysicsEngine/PhysDrawing.cpp
@@ -637,7 +637,7 @@ void FKTaperedCapsuleElem::DrawElemSolid(FPrimitiveDrawInterface* PDI, const FTr
 
 void FKConvexElem::DrawElemWire(FPrimitiveDrawInterface* PDI, const FTransform& ElemTM, const float Scale, const FColor Color) const
 {
-#if WITH_PHYSX
+#if WITH_PHYSX && PHYSICS_INTERFACE_PHYSX
 
 	PxConvexMesh* Mesh = ConvexMesh;
 
@@ -688,12 +688,37 @@ void FKConvexElem::DrawElemWire(FPrimitiveDrawInterface* PDI, const FTransform&
 	{
 		UE_LOG(LogPhysics, Log, TEXT("FKConvexElem::DrawElemWire : No ConvexMesh, so unable to draw."));
 	}
+#elif WITH_CHAOS
+	const int32 NumIndices = IndexData.Num();
+	if(NumIndices > 0 && ensure(NumIndices % 3 == 0))
+	{
+		const int32 NumVerts = VertexData.Num();
+		TArray<FVector> TransformedVerts;
+		TransformedVerts.Reserve(NumVerts);
+		FTransform FinalTransform = Transform.GetRelativeTransformReverse(ElemTM);
+
+		for(const FVector& Vert : VertexData)
+		{
+			TransformedVerts.Add(FinalTransform.TransformPosition(Vert));
+		}
+
+		for(int32 Base = 0; Base < NumIndices; Base += 3)
+		{
+			PDI->DrawLine(TransformedVerts[IndexData[Base]], TransformedVerts[IndexData[Base + 1]], Color, SDPG_World);
+			PDI->DrawLine(TransformedVerts[IndexData[Base + 1]], TransformedVerts[IndexData[Base + 2]], Color, SDPG_World);
+			PDI->DrawLine(TransformedVerts[IndexData[Base + 2]], TransformedVerts[IndexData[Base]], Color, SDPG_World);
+		}
+	}
+	else
+	{
+		UE_LOG(LogPhysics, Log, TEXT("FKConvexElem::AddCachedSolidConvexGeom : No ConvexMesh, so unable to draw."));
+	}
 #endif // WITH_PHYSX
 }
 
 void FKConvexElem::AddCachedSolidConvexGeom(TArray<FDynamicMeshVertex>& VertexBuffer, TArray<uint32>& IndexBuffer, const FColor VertexColor) const
 {
-#if WITH_PHYSX
+#if WITH_PHYSX && PHYSICS_INTERFACE_PHYSX
 	// We always want to generate 'non-mirrored geometry', so if all we have is flipped, we have to un-flip it in this function
 	bool bIsMirrored = false;
 	const PxConvexMesh* ConvexMeshToUse = nullptr; 
@@ -772,6 +797,62 @@ void FKConvexElem::AddCachedSolidConvexGeom(TArray<FDynamicMeshVertex>& VertexBu
 	{
 		UE_LOG(LogPhysics, Log, TEXT("FKConvexElem::AddCachedSolidConvexGeom : No ConvexMesh, so unable to draw."));
 	}
+#elif WITH_CHAOS
+	const int32 NumIndices = IndexData.Num();
+	if(NumIndices > 0 && ensure(NumIndices % 3 == 0))
+	{
+		int32 BeginOffset = VertexBuffer.Num();
+
+		const int32 NumVerts = VertexData.Num();
+		const int32 NumTriangles = NumIndices / 3;
+
+		// Generate Tangents
+		struct LocalTangents
+		{
+			FVector T[3];
+		};
+
+		TArray<LocalTangents> VertTangents;
+		VertTangents.Init({ FVector::ZeroVector, FVector::ZeroVector, FVector::ZeroVector }, VertexData.Num());
+		for(int32 TriIndex = 0; TriIndex < NumTriangles; ++TriIndex)
+		{
+			const int32 Base = TriIndex * 3;
+
+			const int32 I0 = IndexData[Base];
+			const int32 I1 = IndexData[Base + 1];
+			const int32 I2 = IndexData[Base + 2];
+
+			const FVector TangentX = (VertexData[I1] - VertexData[I0]).GetSafeNormal();
+			const FVector TangentZ = FPlane(VertexData[I0], VertexData[I1], VertexData[I2]).GetSafeNormal();
+			const FVector TangentY = FVector::CrossProduct(TangentX, TangentZ).GetSafeNormal();
+
+			for(int32 TriVertIndex = 0; TriVertIndex < 3; ++TriVertIndex)
+			{
+				const int32 Index = IndexData[Base + TriVertIndex];
+				VertTangents[IndexData[Base + TriVertIndex]].T[0] = TangentX;
+				VertTangents[IndexData[Base + TriVertIndex]].T[1] = TangentX;
+				VertTangents[IndexData[Base + TriVertIndex]].T[2] = TangentX;
+			}
+		}
+
+		for(int32 VertIndex = 0; VertIndex < VertexData.Num(); ++VertIndex)
+		{
+			FDynamicMeshVertex Vert;
+			Vert.Position = VertexData[VertIndex];
+			Vert.Color = VertexColor;
+			Vert.SetTangents(VertTangents[VertIndex].T[0], VertTangents[VertIndex].T[1], VertTangents[VertIndex].T[2]);
+			VertexBuffer.Add(Vert);
+		}
+
+		for(int32 Index : IndexData)
+		{
+			IndexBuffer.Add(BeginOffset + Index);
+		}
+	}
+	else
+	{
+		UE_LOG(LogPhysics, Log, TEXT("FKConvexElem::AddCachedSolidConvexGeom : No ConvexMesh, so unable to draw."));
+	}
 #endif // WITH_PHYSX
 }
 
