Prevent crash if a local player has no world context when checking if the player is Primary

diff --git a/Engine/Source/Runtime/Engine/Private/LocalPlayer.cpp b/Engine/Source/Runtime/Engine/Private/LocalPlayer.cpp
--- a/Engine/Source/Runtime/Engine/Private/LocalPlayer.cpp
+++ b/Engine/Source/Runtime/Engine/Private/LocalPlayer.cpp
@@ -1650,8 +1650,12 @@ void ULocalPlayer::AddReferencedObjects(UObject* InThis, FReferenceCollector& Co
 
 bool ULocalPlayer::IsPrimaryPlayer() const
 {
-	ULocalPlayer* const PrimaryPlayer = GetOuterUEngine()->GetFirstGamePlayer(GetWorld());
-	return (this == PrimaryPlayer);
+	if (UWorld* World = GetWorld())
+	{
+		ULocalPlayer* const PrimaryPlayer = GetOuterUEngine()->GetFirstGamePlayer(World);
+		return (this == PrimaryPlayer);
+	}
+	return false;
 }
 
 void ULocalPlayer::CleanupViewState()
