Fixed query only shapes generating collisions in the simulation for Chaos

diff --git a/Engine/Source/Runtime/Engine/Private/PhysicsEngine/BodyInstance.cpp b/Engine/Source/Runtime/Engine/Private/PhysicsEngine/BodyInstance.cpp
--- a/Engine/Source/Runtime/Engine/Private/PhysicsEngine/BodyInstance.cpp
+++ b/Engine/Source/Runtime/Engine/Private/PhysicsEngine/BodyInstance.cpp
@@ -1208,6 +1208,18 @@ struct FInitBodiesHelper
 						{
 							ActorHandles.Add(ActorHandle);
 
+#if WITH_CHAOS
+							// If this shape shouldn't collide in the sim we disable it here until we support
+							// a separation of unions for these shapes
+							if(BI->GetCollisionEnabled() == ECollisionEnabled::QueryOnly)
+							{
+								const int32 NumShapes = FPhysicsInterface::GetNumShapes(ActorHandle);
+								for(int32 ShapeIndex = 0; ShapeIndex < NumShapes; ++ShapeIndex)
+								{
+									ActorHandle->SetShapeCollisionDisable(ShapeIndex, true);
+								}
+							}
+#endif
 /*
 							With the implementation of AddActorsToScene_AssumesLocked,
 							this call duplicates the AddToSolver operation and is not necessary
