- Fix for potential out of bounds array access

#rnx

diff --git a/Engine/Source/Runtime/Engine/Public/ParticleHelper.h b/Engine/Source/Runtime/Engine/Public/ParticleHelper.h
--- a/Engine/Source/Runtime/Engine/Public/ParticleHelper.h
+++ b/Engine/Source/Runtime/Engine/Public/ParticleHelper.h
@@ -2644,12 +2644,15 @@ public:
 	{
 		for (TArray<FParticleVertexFactoryBase*>& PerViewEmitterVFs : EmitterVertexFactoryArray)
 		{
-			FParticleVertexFactoryBase *VertexFactory = PerViewEmitterVFs[EmitterIndex];
-			if (VertexFactory && VertexFactory->IsDirty())
+			if (PerViewEmitterVFs.IsValidIndex(EmitterIndex))
 			{
-				VertexFactory->ReleaseResource();
-				delete VertexFactory;
-				PerViewEmitterVFs[EmitterIndex] = nullptr;
+				FParticleVertexFactoryBase *VertexFactory = PerViewEmitterVFs[EmitterIndex];
+				if (VertexFactory && VertexFactory->IsDirty())
+				{
+					VertexFactory->ReleaseResource();
+					delete VertexFactory;
+					PerViewEmitterVFs[EmitterIndex] = nullptr;
+				}
 			}
 		}
 	}
