Chaos - Add bSimulate flag to PerShapeData and initialize based on collision flags.

diff --git a/Engine/Source/Runtime/Engine/Private/PhysicsEngine/Experimental/ChaosInterfaceUtils.cpp b/Engine/Source/Runtime/Engine/Private/PhysicsEngine/Experimental/ChaosInterfaceUtils.cpp
--- a/Engine/Source/Runtime/Engine/Private/PhysicsEngine/Experimental/ChaosInterfaceUtils.cpp
+++ b/Engine/Source/Runtime/Engine/Private/PhysicsEngine/Experimental/ChaosInterfaceUtils.cpp
@@ -149,6 +149,7 @@ namespace ChaosInterface
 			NewShape->SimData = InParams.CollisionData.CollisionFilterData.SimFilter;
 			NewShape->UpdateShapeBounds(InParams.WorldTransform);
 			NewShape->UserData = UserData;
+			NewShape->bSimulate = bComplexShape ? InParams.CollisionData.CollisionFlags.bEnableSimCollisionComplex : InParams.CollisionData.CollisionFlags.bEnableSimCollisionSimple;
 			return NewShape;
 		};
 
diff --git a/Engine/Source/Runtime/Engine/Private/PhysicsEngine/Experimental/PhysInterface_Chaos.cpp b/Engine/Source/Runtime/Engine/Private/PhysicsEngine/Experimental/PhysInterface_Chaos.cpp
--- a/Engine/Source/Runtime/Engine/Private/PhysicsEngine/Experimental/PhysInterface_Chaos.cpp
+++ b/Engine/Source/Runtime/Engine/Private/PhysicsEngine/Experimental/PhysInterface_Chaos.cpp
@@ -1920,7 +1920,7 @@ int32 GetAllShapesInternal_AssumedLocked(const FPhysicsActorHandle& InActorHandl
 	//todo: can we avoid this construction?
 	for (const TUniquePtr<Chaos::TPerShapeData<float, 3>>& Shape : ShapesArray)
 	{
-		OutShapes.Add(FPhysicsShapeReference_Chaos(Shape.Get(), true, true, InActorHandle));
+		OutShapes.Add(FPhysicsShapeReference_Chaos(Shape.Get(), Shape->bSimulate, true, InActorHandle));
 	}
 	return OutShapes.Num();
 }
diff --git a/Engine/Source/Runtime/Experimental/Chaos/Private/Chaos/GeometryParticles.cpp b/Engine/Source/Runtime/Experimental/Chaos/Private/Chaos/GeometryParticles.cpp
--- a/Engine/Source/Runtime/Experimental/Chaos/Private/Chaos/GeometryParticles.cpp
+++ b/Engine/Source/Runtime/Experimental/Chaos/Private/Chaos/GeometryParticles.cpp
@@ -59,6 +59,7 @@ namespace Chaos
 		, WorldSpaceInflatedShapeBounds(TAABB<FReal, 3>(FVec3(0), FVec3(0)))
 		, Materials()
 		, bDisable(false)
+		, bSimulate(true)
 	{
 	}
 
diff --git a/Engine/Source/Runtime/Experimental/Chaos/Public/Chaos/GeometryParticles.h b/Engine/Source/Runtime/Experimental/Chaos/Public/Chaos/GeometryParticles.h
--- a/Engine/Source/Runtime/Experimental/Chaos/Public/Chaos/GeometryParticles.h
+++ b/Engine/Source/Runtime/Experimental/Chaos/Public/Chaos/GeometryParticles.h
@@ -35,7 +35,10 @@ namespace Chaos
 		TSerializablePtr<FImplicitObject> Geometry;
 		TAABB<FReal, 3> WorldSpaceInflatedShapeBounds;
 		TArray<FMaterialHandle> Materials;
+
+		// TODO: Bitfields?
 		bool bDisable;
+		bool bSimulate;
 
 		void UpdateShapeBounds(const TRigidTransform<T, d>& WorldTM);
 
