Fix build.

diff --git a/Engine/Source/Runtime/Engine/Private/PhysicsEngine/Experimental/ChaosDerivedData.cpp b/Engine/Source/Runtime/Engine/Private/PhysicsEngine/Experimental/ChaosDerivedData.cpp
--- a/Engine/Source/Runtime/Engine/Private/PhysicsEngine/Experimental/ChaosDerivedData.cpp
+++ b/Engine/Source/Runtime/Engine/Private/PhysicsEngine/Experimental/ChaosDerivedData.cpp
@@ -85,7 +85,7 @@ FChaosDerivedDataCooker::FChaosDerivedDataCooker(UBodySetup* InSetup, FName InFo
 }
 
 template<typename Precision>
-void FChaosDerivedDataCooker::BuildTriangleMeshes(TArray<TUniquePtr<Chaos::TTriangleMeshImplicitObject<Precision>>>& OutTriangleMeshes, const FCookBodySetupInfo& InParams)
+void FChaosDerivedDataCooker::BuildTriangleMeshes(TArray<TUniquePtr<Chaos::TTriangleMeshImplicitObject<Precision>>>& OutTriangleMeshes, TArray<int32>& OutFaceRemap, const FCookBodySetupInfo& InParams)
 {
 	if(!InParams.bCookTriMesh)
 	{
@@ -105,8 +105,7 @@ void FChaosDerivedDataCooker::BuildTriangleMeshes(TArray<TUniquePtr<Chaos::TTria
 		FinalIndices.Add(Tri.v2);
 	}
 
-	TArray<int32> Remap;
-	Chaos::CleanTrimesh(FinalVerts, FinalIndices, &Remap);
+	Chaos::CleanTrimesh(FinalVerts, FinalIndices, &OutFaceRemap);
 
 	// Build particle list #BG Maybe allow TParticles to copy vectors?
 	Chaos::TParticles<Precision, 3> TriMeshParticles;
@@ -147,14 +146,14 @@ void FChaosDerivedDataCooker::BuildTriangleMeshes(TArray<TUniquePtr<Chaos::TTria
 
 			if(bHasMaterials)
 			{
-				if(!ensure(Remap.IsValidIndex(TriangleIndex)))
+				if(!ensure(OutFaceRemap.IsValidIndex(TriangleIndex)))
 				{
 					MaterialIndices.Empty();
 					bHasMaterials = false;
 				}
 				else
 				{
-					const int32 OriginalIndex = Remap[TriangleIndex];
+					const int32 OriginalIndex = OutFaceRemap[TriangleIndex];
 
 					if(ensure(InParams.TriangleMeshDesc.MaterialIndices.IsValidIndex(OriginalIndex)))
 					{
@@ -173,7 +172,7 @@ void FChaosDerivedDataCooker::BuildTriangleMeshes(TArray<TUniquePtr<Chaos::TTria
 	OutTriangleMeshes.Emplace(new Chaos::TTriangleMeshImplicitObject<Precision>(MoveTemp(TriMeshParticles), MoveTemp(Triangles), MoveTemp(MaterialIndices)));
 }
 
-template void FChaosDerivedDataCooker::BuildTriangleMeshes(TArray<TUniquePtr<Chaos::TTriangleMeshImplicitObject<float>>>& OutTriangleMeshes, const FCookBodySetupInfo& InParams);
+template void FChaosDerivedDataCooker::BuildTriangleMeshes(TArray<TUniquePtr<Chaos::TTriangleMeshImplicitObject<float>>>& OutTriangleMeshes, TArray<int32>& OutFaceRemap, const FCookBodySetupInfo& InParams);
 //#BGTODO When it's possible to build with doubles, re-enable this. (Currently at least TRigidTransform cannot build with double precision because we don't have a base transform implementation using them)
 //template void FChaosDerivedDataCooker::BuildTriangleMeshes(TArray<TUniquePtr<Chaos::FImplicitObject>>& OutTriangleMeshes, const FCookBodySetupInfo& InParams);
 
@@ -251,9 +250,10 @@ void FChaosDerivedDataCooker::BuildInternal(Chaos::FChaosArchive& Ar, FCookBodyS
 	TArray<TUniquePtr<Chaos::FImplicitObject>> SimpleImplicits;
 	TArray<TUniquePtr<Chaos::TTriangleMeshImplicitObject<Precision>>> ComplexImplicits;
 
+	TArray<int32> FaceRemap;
 	//BuildSimpleShapes(SimpleImplicits, Setup);
 	BuildConvexMeshes(SimpleImplicits, InInfo);
-	BuildTriangleMeshes(ComplexImplicits, InInfo);
+	BuildTriangleMeshes(ComplexImplicits, FaceRemap, InInfo);
 
 	//TUniquePtr<Chaos::TImplicitObjectUnion<Precision, 3>> SimpleUnion = MakeUnique<Chaos::TImplicitObjectUnion<Precision, 3>>(MoveTemp(SimpleImplicits));
 	//TUniquePtr<Chaos::TImplicitObjectUnion<Precision, 3>> ComplexUnion = MakeUnique<Chaos::TImplicitObjectUnion<Precision, 3>>(MoveTemp(ComplexImplicits));
@@ -263,8 +263,12 @@ void FChaosDerivedDataCooker::BuildInternal(Chaos::FChaosArchive& Ar, FCookBodyS
 	{
 		UVInfo.FillFromTriMesh(InInfo.TriangleMeshDesc);
 	}
+	if (!InInfo.bSupportFaceRemap)
+	{
+		FaceRemap.Empty();
+	}
 
-	Ar << SimpleImplicits << ComplexImplicits << UVInfo;
+	Ar << SimpleImplicits << ComplexImplicits << UVInfo << FaceRemap;
 }
 
 template void FChaosDerivedDataCooker::BuildInternal<float>(Chaos::FChaosArchive& Ar, FCookBodySetupInfo& InInfo);
