Fix for build warnings when source audio hasn't been loaded yet.

diff --git a/Engine/Source/Runtime/Engine/Private/AudioDerivedData.cpp b/Engine/Source/Runtime/Engine/Private/AudioDerivedData.cpp
--- a/Engine/Source/Runtime/Engine/Private/AudioDerivedData.cpp
+++ b/Engine/Source/Runtime/Engine/Private/AudioDerivedData.cpp
@@ -1015,6 +1015,8 @@ static void CookSimpleWave(USoundWave* SoundWave, FName FormatName, const IAudio
 	FScopeLock ScopeLock(&SoundWave->RawDataCriticalSection);
 #endif
 
+	SoundWave->RawData.ForceBulkDataResident();
+
 	// check if there is any raw sound data
 	if( SoundWave->RawData.GetBulkDataSize() > 0 )
 	{
@@ -1038,7 +1040,7 @@ static void CookSimpleWave(USoundWave* SoundWave, FName FormatName, const IAudio
 
 	if(!Input.Num())
 	{
-		UE_LOG(LogAudioDerivedData, Warning, TEXT( "Can't cook %s because there is no source compressed or uncompressed PC sound data" ), *SoundWave->GetFullName() );
+		UE_LOG(LogAudioDerivedData, Warning, TEXT( "Can't cook %s because there is no source LPCM data" ), *SoundWave->GetFullName() );
 	}
 	else
 	{
