Add flag to FUpdateLevelVisibilityLevelInfo to control error behavior

diff --git a/Engine/Source/Runtime/Engine/Classes/GameFramework/PlayerController.h b/Engine/Source/Runtime/Engine/Classes/GameFramework/PlayerController.h
--- a/Engine/Source/Runtime/Engine/Classes/GameFramework/PlayerController.h
+++ b/Engine/Source/Runtime/Engine/Classes/GameFramework/PlayerController.h
@@ -172,6 +172,7 @@ struct ENGINE_API FUpdateLevelVisibilityLevelInfo
 		: PackageName(NAME_None)
 		, FileName(NAME_None)
 		, bIsVisible(false)
+		, bSkipCloseOnError(false)
 	{
 	}
 
@@ -193,6 +194,9 @@ struct ENGINE_API FUpdateLevelVisibilityLevelInfo
 	UPROPERTY()
 	uint32 bIsVisible : 1;
 
+	/** Skip connection close if level can't be found (not net serialized) */
+	uint32 bSkipCloseOnError : 1;
+
 	bool NetSerialize(FArchive& Ar, UPackageMap* PackageMap, bool& bOutSuccess);
 };
 
diff --git a/Engine/Source/Runtime/Engine/Private/NetConnection.cpp b/Engine/Source/Runtime/Engine/Private/NetConnection.cpp
--- a/Engine/Source/Runtime/Engine/Private/NetConnection.cpp
+++ b/Engine/Source/Runtime/Engine/Private/NetConnection.cpp
@@ -1071,7 +1071,11 @@ void UNetConnection::UpdateLevelVisibility(const FUpdateLevelVisibilityLevelInfo
 		else
 		{
 			UE_LOG(LogPlayerController, Warning, TEXT("ServerUpdateLevelVisibility() ignored non-existant package. PackageName='%s', FileName='%s'"), *LevelVisibility.PackageName.ToString(), *LevelVisibility.FileName.ToString());
-			Close();
+
+			if (!LevelVisibility.bSkipCloseOnError)
+			{
+				Close();
+			}
 		}
 	}
 	else
diff --git a/Engine/Source/Runtime/Engine/Private/NetworkDriver.cpp b/Engine/Source/Runtime/Engine/Private/NetworkDriver.cpp
--- a/Engine/Source/Runtime/Engine/Private/NetworkDriver.cpp
+++ b/Engine/Source/Runtime/Engine/Private/NetworkDriver.cpp
@@ -454,10 +454,6 @@ void UNetDriver::SetRoleSwapOnReplicate(AActor* Actor, bool bSwapRoles)
 	{
 		InfoPtr->bSwapRolesOnReplicate = bSwapRoles;
 	}
-	else
-	{
-		UE_LOG(LogNet, Warning, TEXT("SetRoleSwapOnReplicate could not find network object info for: %s"), *GetFullNameSafe(Actor));
-	}
 
 	if (ReplicationDriver)
 	{
diff --git a/Engine/Source/Runtime/Engine/Private/PlayerController.cpp b/Engine/Source/Runtime/Engine/Private/PlayerController.cpp
--- a/Engine/Source/Runtime/Engine/Private/PlayerController.cpp
+++ b/Engine/Source/Runtime/Engine/Private/PlayerController.cpp
@@ -104,6 +104,7 @@ const float ForceRetryClientRestartTime = -100.0f;
 
 FUpdateLevelVisibilityLevelInfo::FUpdateLevelVisibilityLevelInfo(const ULevel* const Level, const bool bInIsVisible)
 	: bIsVisible(bInIsVisible)
+	, bSkipCloseOnError(false)
 {
 	const UPackage* const LevelPackage = Level->GetOutermost();
 	PackageName = LevelPackage->GetFName();
