Fix ULineBatchComponent not being able to be inherited from on Mac due to the vtable not being exported
#rnx

diff --git a/Engine/Source/Runtime/Engine/Classes/Components/LineBatchComponent.h b/Engine/Source/Runtime/Engine/Classes/Components/LineBatchComponent.h
--- a/Engine/Source/Runtime/Engine/Classes/Components/LineBatchComponent.h
+++ b/Engine/Source/Runtime/Engine/Classes/Components/LineBatchComponent.h
@@ -117,7 +117,7 @@ public:
  *	This can be useful for debug drawing, but is not very performant for runtime use.
  */
 UCLASS(MinimalAPI)
-class ULineBatchComponent : public UPrimitiveComponent
+class ENGINE_VTABLE ULineBatchComponent : public UPrimitiveComponent
 {
 	GENERATED_UCLASS_BODY()
 
@@ -166,14 +166,14 @@ class ULineBatchComponent : public UPrimitiveComponent
 	ENGINE_API void DrawMesh(TArray<FVector> const& Verts, TArray<int32> const& Indices, FColor const& Color, uint8 DepthPriority, float LifeTime);
 
 	//~ Begin UPrimitiveComponent Interface.
-	ENGINE_API virtual FPrimitiveSceneProxy* CreateSceneProxy() override;
-	ENGINE_API virtual FBoxSphereBounds CalcBounds(const FTransform& LocalToWorld) const override;
+	virtual FPrimitiveSceneProxy* CreateSceneProxy() override;
+	virtual FBoxSphereBounds CalcBounds(const FTransform& LocalToWorld) const override;
 	//~ End UPrimitiveComponent Interface.
 	
 	
 	//~ Begin UActorComponent Interface.
-	ENGINE_API virtual void TickComponent(float DeltaTime, enum ELevelTick TickType, FActorComponentTickFunction *ThisTickFunction) override;
-	ENGINE_API virtual void ApplyWorldOffset(const FVector& InOffset, bool bWorldShift) override;
+	virtual void TickComponent(float DeltaTime, enum ELevelTick TickType, FActorComponentTickFunction *ThisTickFunction) override;
+	virtual void ApplyWorldOffset(const FVector& InOffset, bool bWorldShift) override;
 	//~ End UActorComponent Interface.
 
 	/** Clear all batched lines, points and meshes */
