Fix clang error

diff --git a/Engine/Source/Runtime/Engine/Private/Collision/SceneQueryLowLevel.cpp b/Engine/Source/Runtime/Engine/Private/Collision/SceneQueryLowLevel.cpp
--- a/Engine/Source/Runtime/Engine/Private/Collision/SceneQueryLowLevel.cpp
+++ b/Engine/Source/Runtime/Engine/Private/Collision/SceneQueryLowLevel.cpp
@@ -115,7 +115,7 @@ namespace
 	void RaycastSQCaptureHelper(float QueryDurationSeconds, const FChaosSQAccelerator& SQAccelerator, FPhysScene& Scene, const FVector& Start, const FVector& Dir, float DeltaMag, const FPhysicsHitCallback<FHitRaycast>& HitBuffer, EHitFlags OutputFlags, FQueryFlags QueryFlags, const FCollisionFilterData& Filter, const FQueryFilterData& QueryFilterData, ICollisionQueryFilterCallbackBase* QueryCallback, const FQueryDebugParams& DebugParams)
 	{
 		float QueryDurationMicro = QueryDurationSeconds * 1000.0 * 1000.0;
-		if (((SerializeSQs && QueryDurationMicro > SerializeSQs)) && IsInGameThread())
+		if (((!!SerializeSQs && QueryDurationMicro > SerializeSQs)) && IsInGameThread())
 		{
 			// Measure average time of query over multiple samples to reduce fluke from context switches or that kind of thing.
 			uint32 Cycles = 0.0;
@@ -149,7 +149,7 @@ namespace
 	void OverlapSQCaptureHelper(float QueryDurationSeconds, const FChaosSQAccelerator& SQAccelerator, FPhysScene& Scene, const FPhysicsGeometry& QueryGeom, const FTransform& GeomPose, const FPhysicsHitCallback<FHitOverlap>& HitBuffer, FQueryFlags QueryFlags, const FCollisionFilterData& Filter, const FQueryFilterData& QueryFilterData, ICollisionQueryFilterCallbackBase* QueryCallback, const FQueryDebugParams& DebugParams)
 	{
 		float QueryDurationMicro = QueryDurationSeconds * 1000.0 * 1000.0;
-		if (((SerializeSQs && QueryDurationMicro > SerializeSQs)) && IsInGameThread())
+		if (((!!SerializeSQs && QueryDurationMicro > SerializeSQs)) && IsInGameThread())
 		{
 			// Measure average time of query over multiple samples to reduce fluke from context switches or that kind of thing.
 			uint32 Cycles = 0.0;
@@ -195,13 +195,13 @@ void LowLevelRaycast(FPhysScene& Scene, const FVector& Start, const FVector& Dir
 			SQAccelerator.Raycast(Start, Dir, DeltaMag, HitBuffer, OutputFlags, QueryFilterData, *QueryCallback, DebugParams);
 		}
 
-		if (SerializeSQs && EnableRaycastSQCapture)
+		if (!!SerializeSQs && !!EnableRaycastSQCapture)
 		{
 			RaycastSQCaptureHelper(Time, SQAccelerator, Scene, Start, Dir, DeltaMag, HitBuffer, OutputFlags, QueryFlags, Filter, QueryFilterData, QueryCallback, DebugParams);
 		}
 	}
 #else
-	if (SerializeSQs | ReplaySQs)
+	if (!!SerializeSQs || !!ReplaySQs)
 	{
 		FPhysTestSerializer Serializer;
 		Serializer.SetPhysicsData(*Scene.GetPxScene());
@@ -233,14 +233,14 @@ void LowLevelSweep(FPhysScene& Scene, const FPhysicsGeometry& QueryGeom, const F
 				SQAccelerator.Sweep(QueryGeom, StartTM, Dir, DeltaMag, HitBuffer, OutputFlags, QueryFilterData, *QueryCallback, DebugParams);
 			}
 
-			if (SerializeSQs && EnableSweepSQCapture)
+			if (!!SerializeSQs && !!EnableSweepSQCapture)
 			{
 				SweepSQCaptureHelper(Time, SQAccelerator, Scene, QueryGeom, StartTM, Dir, DeltaMag, HitBuffer, OutputFlags, QueryFlags, Filter, QueryFilterData, QueryCallback, DebugParams);
 			}
 		}
 	}
 #else
-	if (SerializeSQs | ReplaySQs)
+	if (!!SerializeSQs || !!ReplaySQs)
 	{
 		FPhysTestSerializer Serializer;
 		Serializer.SetPhysicsData(*Scene.GetPxScene());
@@ -270,13 +270,13 @@ void LowLevelOverlap(FPhysScene& Scene, const FPhysicsGeometry& QueryGeom, const
 			SQAccelerator.Overlap(QueryGeom, GeomPose, HitBuffer, QueryFilterData, *QueryCallback);
 		}
 
-		if (SerializeSQs && EnableOverlapSQCapture)
+		if (!!SerializeSQs && !!EnableOverlapSQCapture)
 		{
 			OverlapSQCaptureHelper(Time, SQAccelerator, Scene, QueryGeom, GeomPose, HitBuffer, QueryFlags, Filter, QueryFilterData, QueryCallback, DebugParams);
 		}
 	}
 #else
-	if (SerializeSQs | ReplaySQs)
+	if (!!SerializeSQs || !!ReplaySQs)
 	{
 		FPhysTestSerializer Serializer;
 		Serializer.SetPhysicsData(*Scene.GetPxScene());
