Chaos - Fix moved objects not immediately updating in physics scene acceleration structure.

diff --git a/Engine/Source/Runtime/Engine/Private/PhysicsEngine/Experimental/PhysInterface_Chaos.cpp b/Engine/Source/Runtime/Engine/Private/PhysicsEngine/Experimental/PhysInterface_Chaos.cpp
--- a/Engine/Source/Runtime/Engine/Private/PhysicsEngine/Experimental/PhysInterface_Chaos.cpp
+++ b/Engine/Source/Runtime/Engine/Private/PhysicsEngine/Experimental/PhysInterface_Chaos.cpp
@@ -467,6 +467,9 @@ void FPhysInterface_Chaos::SetGlobalPose_AssumesLocked(const FPhysicsActorHandle
 {
 	InActorReference->SetX(InNewPose.GetLocation());
 	InActorReference->SetR(InNewPose.GetRotation());
+
+	FPhysScene* Scene = GetCurrentScene(InActorReference);
+	Scene->GetScene().UpdateActorInAccelerationStructure(InActorReference);
 }
 
 FTransform FPhysInterface_Chaos::GetTransform_AssumesLocked(const FPhysicsActorHandle& InRef, bool bForceGlobalPose /*= false*/)
diff --git a/Engine/Source/Runtime/Engine/Private/PhysicsEngine/Experimental/PhysScene_Chaos.cpp b/Engine/Source/Runtime/Engine/Private/PhysicsEngine/Experimental/PhysScene_Chaos.cpp
--- a/Engine/Source/Runtime/Engine/Private/PhysicsEngine/Experimental/PhysScene_Chaos.cpp
+++ b/Engine/Source/Runtime/Engine/Private/PhysicsEngine/Experimental/PhysScene_Chaos.cpp
@@ -587,7 +587,7 @@ void FPhysScene_Chaos::RemoveActorFromAccelerationStructure(FPhysicsActorHandle&
 #endif
 }
 
-void FPhysScene_Chaos::UpdateActorInAccelerationStructure(FPhysicsActorHandle& Actor)
+void FPhysScene_Chaos::UpdateActorInAccelerationStructure(const FPhysicsActorHandle& Actor)
 {
 #if WITH_CHAOS
 	using namespace Chaos;
diff --git a/Engine/Source/Runtime/Engine/Public/Physics/Experimental/PhysScene_Chaos.h b/Engine/Source/Runtime/Engine/Public/Physics/Experimental/PhysScene_Chaos.h
--- a/Engine/Source/Runtime/Engine/Public/Physics/Experimental/PhysScene_Chaos.h
+++ b/Engine/Source/Runtime/Engine/Public/Physics/Experimental/PhysScene_Chaos.h
@@ -126,7 +126,7 @@ public:
 	void RemoveObject(FFieldSystemPhysicsProxy* InObject);	
 
 	void RemoveActorFromAccelerationStructure(FPhysicsActorHandle& Actor);
-	void UpdateActorInAccelerationStructure(FPhysicsActorHandle& Actor);
+	void UpdateActorInAccelerationStructure(const FPhysicsActorHandle& Actor);
 
 #if XGE_FIXED
 	template<typename PayloadType>
