Fix CIS -  Missing WITH_EDITOR

diff --git a/Engine/Source/Runtime/Engine/Private/PhysicsEngine/PhysLevel.cpp b/Engine/Source/Runtime/Engine/Private/PhysicsEngine/PhysLevel.cpp
--- a/Engine/Source/Runtime/Engine/Private/PhysicsEngine/PhysLevel.cpp
+++ b/Engine/Source/Runtime/Engine/Private/PhysicsEngine/PhysLevel.cpp
@@ -106,7 +106,7 @@ void UWorld::SetupPhysicsTickFunctions(float DeltaSeconds)
 	EndPhysicsTickFunction.Target = this;
 
 // Chaos ticks solver for trace collisions
-#if WITH_CHAOS
+#if (WITH_CHAOS && WITH_EDITOR)
 	bool bEnablePhysics = (bShouldSimulatePhysics || bEnableTraceCollision);
 #else
 	bool bEnablePhysics = bShouldSimulatePhysics;
