Remark virtual overrides as ENGINE_API as apparently that is necessary
#rnx

diff --git a/Engine/Source/Runtime/Engine/Classes/Components/LineBatchComponent.h b/Engine/Source/Runtime/Engine/Classes/Components/LineBatchComponent.h
--- a/Engine/Source/Runtime/Engine/Classes/Components/LineBatchComponent.h
+++ b/Engine/Source/Runtime/Engine/Classes/Components/LineBatchComponent.h
@@ -166,14 +166,14 @@ class ENGINE_VTABLE ULineBatchComponent : public UPrimitiveComponent
 	ENGINE_API void DrawMesh(TArray<FVector> const& Verts, TArray<int32> const& Indices, FColor const& Color, uint8 DepthPriority, float LifeTime);
 
 	//~ Begin UPrimitiveComponent Interface.
-	virtual FPrimitiveSceneProxy* CreateSceneProxy() override;
-	virtual FBoxSphereBounds CalcBounds(const FTransform& LocalToWorld) const override;
+	ENGINE_API virtual FPrimitiveSceneProxy* CreateSceneProxy() override;
+	ENGINE_API virtual FBoxSphereBounds CalcBounds(const FTransform& LocalToWorld) const override;
 	//~ End UPrimitiveComponent Interface.
 	
 	
 	//~ Begin UActorComponent Interface.
-	virtual void TickComponent(float DeltaTime, enum ELevelTick TickType, FActorComponentTickFunction *ThisTickFunction) override;
-	virtual void ApplyWorldOffset(const FVector& InOffset, bool bWorldShift) override;
+	ENGINE_API virtual void TickComponent(float DeltaTime, enum ELevelTick TickType, FActorComponentTickFunction *ThisTickFunction) override;
+	ENGINE_API virtual void ApplyWorldOffset(const FVector& InOffset, bool bWorldShift) override;
 	//~ End UActorComponent Interface.
 
 	/** Clear all batched lines, points and meshes */
