Fix CIS

diff --git a/Engine/Source/Runtime/Engine/Private/PhysicsEngine/PhysDrawing.cpp b/Engine/Source/Runtime/Engine/Private/PhysicsEngine/PhysDrawing.cpp
--- a/Engine/Source/Runtime/Engine/Private/PhysicsEngine/PhysDrawing.cpp
+++ b/Engine/Source/Runtime/Engine/Private/PhysicsEngine/PhysDrawing.cpp
@@ -813,7 +813,7 @@ void FKConvexElem::AddCachedSolidConvexGeom(TArray<FDynamicMeshVertex>& VertexBu
 		};
 
 		TArray<LocalTangents> VertTangents;
-		VertTangents.Init({ FVector::ZeroVector, FVector::ZeroVector, FVector::ZeroVector }, VertexData.Num());
+		VertTangents.Init({ { FVector::ZeroVector, FVector::ZeroVector, FVector::ZeroVector } }, VertexData.Num());
 		for(int32 TriIndex = 0; TriIndex < NumTriangles; ++TriIndex)
 		{
 			const int32 Base = TriIndex * 3;
