HLOD marked as "Unbuilt" after saving sub-levels with a transform

Prevent  level transform from messing with the CRC key computation
Round components location/rotation to avoid issues with floating point accuracy

diff --git a/Engine/Source/Runtime/Engine/Classes/Engine/HLODProxy.h b/Engine/Source/Runtime/Engine/Classes/Engine/HLODProxy.h
--- a/Engine/Source/Runtime/Engine/Classes/Engine/HLODProxy.h
+++ b/Engine/Source/Runtime/Engine/Classes/Engine/HLODProxy.h
@@ -86,12 +86,12 @@ public:
 	static void ExtractComponents(const ALODActor* LODActor, TArray<UPrimitiveComponent*>& InOutComponents);
 
 	/** Build a unique key for the LOD actor, used to determine if the actor needs rebuilding */
-	static FName GenerateKeyForActor(const ALODActor* LODActor);
+	static FName GenerateKeyForActor(const ALODActor* LODActor, bool bMustUndoLevelTransform = true);
 
 	static uint32 GetCRC(UMaterialInterface* InMaterialInterface, uint32 InCRC = 0);
 	static uint32 GetCRC(UTexture* InTexture, uint32 InCRC = 0);
 	static uint32 GetCRC(UStaticMesh* InStaticMesh, uint32 InCRC = 0);
-	static uint32 GetCRC(UStaticMeshComponent* InComponent, uint32 InCRC = 0);
+	static uint32 GetCRC(UStaticMeshComponent* InComponent, uint32 InCRC = 0, const FTransform& TransformComponents = FTransform::Identity);
 
 #endif
 
diff --git a/Engine/Source/Runtime/Engine/Private/HLODProxy.cpp b/Engine/Source/Runtime/Engine/Private/HLODProxy.cpp
--- a/Engine/Source/Runtime/Engine/Private/HLODProxy.cpp
+++ b/Engine/Source/Runtime/Engine/Private/HLODProxy.cpp
@@ -13,6 +13,9 @@
 #include "Engine/Texture.h"
 #include "Engine/StaticMesh.h"
 #include "PhysicsEngine/BodySetup.h"
+#include "LevelUtils.h"
+#include "Engine/LevelStreaming.h"
+#include "Math/UnrealMathUtility.h"
 
 #if WITH_EDITOR
 
@@ -219,12 +222,26 @@ uint32 UHLODProxy::GetCRC(UStaticMesh* InStaticMesh, uint32 InCRC)
 	return FCrc::MemCrc32(KeyBuffer.GetData(), KeyBuffer.Num(), InCRC);
 }
 
-uint32 UHLODProxy::GetCRC(UStaticMeshComponent* InComponent, uint32 InCRC)
+uint32 UHLODProxy::GetCRC(UStaticMeshComponent* InComponent, uint32 InCRC, const FTransform& TransformComponents)
 {
 	TArray<uint8> KeyBuffer;
 
-	// incorporate transform & other relevant properties
-	KeyBuffer.Append((uint8*)&InComponent->GetComponentTransform(), sizeof(FTransform));
+	FVector  ComponentLocation = InComponent->GetComponentLocation();
+	FRotator ComponentRotation = InComponent->GetComponentRotation();
+	FVector  ComponentScale = InComponent->GetComponentScale();
+
+	ComponentLocation = TransformComponents.TransformPosition(ComponentLocation);
+	ComponentRotation = TransformComponents.TransformRotation(ComponentRotation.Quaternion()).Rotator();
+
+	// Include transform - round sufficiently to ensure stability
+	FIntVector Location(ComponentLocation / THRESH_POINTS_ARE_NEAR);
+	KeyBuffer.Append((uint8*)&Location, sizeof(Location));
+	FIntVector Rotation(ComponentRotation.GetNormalized().Vector() / THRESH_POINTS_ARE_NEAR);
+	KeyBuffer.Append((uint8*)&Rotation, sizeof(Rotation));
+	FIntVector Scale(ComponentScale / THRESH_POINTS_ARE_NEAR);
+	KeyBuffer.Append((uint8*)&Scale, sizeof(Scale));	
+
+	// Include other relevant properties
 	KeyBuffer.Append((uint8*)&InComponent->ForcedLodModel, sizeof(int32));
 	bool bUseMaxLODAsImposter = InComponent->bUseMaxLODAsImposter;
 	KeyBuffer.Append((uint8*)&bUseMaxLODAsImposter, sizeof(bool));
@@ -256,7 +273,7 @@ uint32 UHLODProxy::GetCRC(UStaticMeshComponent* InComponent, uint32 InCRC)
 // Key that forms the basis of the HLOD proxy key. Bump this key (i.e. generate a new GUID) when you want to force a rebuild of ALL HLOD proxies
 #define HLOD_PROXY_BASE_KEY		TEXT("76927B120C6645ACB9200E7FB8896AC3")
 
-FName UHLODProxy::GenerateKeyForActor(const ALODActor* LODActor)
+FName UHLODProxy::GenerateKeyForActor(const ALODActor* LODActor, bool bMustUndoLevelTransform)
 {
 	FString Key = HLOD_PROXY_BASE_KEY;
 
@@ -320,6 +337,18 @@ FName UHLODProxy::GenerateKeyForActor(const ALODActor* LODActor)
 
 		TArray<UPrimitiveComponent*> Components;
 		ExtractComponents(LODActor, Components);
+		
+		// Components can be offset by their streaming level transform. Undo that transform to have the same signature
+		// when computing CRC for a sub level or a persistent level.
+		FTransform TransformComponents = FTransform::Identity;
+		if (bMustUndoLevelTransform)
+		{
+			ULevelStreaming* SteamingLevel = FLevelUtils::FindStreamingLevel(LODActor->GetLevel());
+			if (SteamingLevel)
+			{
+				TransformComponents = SteamingLevel->LevelTransform.Inverse();
+			}
+		}
 
 		// We get the CRC of each component and combine them
 		for(UPrimitiveComponent* Component : Components)
@@ -327,7 +356,7 @@ FName UHLODProxy::GenerateKeyForActor(const ALODActor* LODActor)
 			if(UStaticMeshComponent* StaticMeshComponent = Cast<UStaticMeshComponent>(Component))
 			{
 				// CRC component
-				CRC = GetCRC(StaticMeshComponent, CRC);
+				CRC = GetCRC(StaticMeshComponent, CRC, TransformComponents);
 
 				if(UStaticMesh* StaticMesh = StaticMeshComponent->GetStaticMesh())
 				{
diff --git a/Engine/Source/Runtime/Engine/Private/LODActor.cpp b/Engine/Source/Runtime/Engine/Private/LODActor.cpp
--- a/Engine/Source/Runtime/Engine/Private/LODActor.cpp
+++ b/Engine/Source/Runtime/Engine/Private/LODActor.cpp
@@ -1335,10 +1335,12 @@ void ALODActor::PreSave(const class ITargetPlatform* TargetPlatform)
 {
 	Super::PreSave(TargetPlatform);	
 
+	// Always rebuild key on save here.
+	// We don't do this while cooking as keys rely on platform derived data which is context-dependent during cook
 	if(!GIsCookerLoadingPackage)
 	{
-		// Always rebuild key on save here. We dont do this while cooking as keys rely on platform derived data which is context-dependent during cook
-		Key = UHLODProxy::GenerateKeyForActor(this);
+		const bool bMustUndoLevelTransform = false; // In the save process, the level transform is already removed
+		Key = UHLODProxy::GenerateKeyForActor(this, bMustUndoLevelTransform);
 	}
 
 	// check & warn if we need building
