Fix mac editor build

diff --git a/Engine/Source/Runtime/Engine/Classes/Engine/Level.h b/Engine/Source/Runtime/Engine/Classes/Engine/Level.h
--- a/Engine/Source/Runtime/Engine/Classes/Engine/Level.h
+++ b/Engine/Source/Runtime/Engine/Classes/Engine/Level.h
@@ -51,7 +51,7 @@ public:
 
 private:
 #if WITH_EDITOR
-	virtual void BeginOperation(FLevelPartitionOperationScope* Scope) = 0;
+	virtual void BeginOperation(class FLevelPartitionOperationScope* Scope) = 0;
 	virtual void EndOperation() = 0;
 #endif
 };
