Chaos - Remove uneeded CHAOS_ENSURE from BodyInstance::UpdateMassProperties.

diff --git a/Engine/Source/Runtime/Engine/Private/PhysicsEngine/BodyInstance.cpp b/Engine/Source/Runtime/Engine/Private/PhysicsEngine/BodyInstance.cpp
--- a/Engine/Source/Runtime/Engine/Private/PhysicsEngine/BodyInstance.cpp
+++ b/Engine/Source/Runtime/Engine/Private/PhysicsEngine/BodyInstance.cpp
@@ -2800,13 +2800,6 @@ void FBodyInstance::UpdateMassProperties()
 {
 	UPhysicalMaterial* PhysMat = GetSimplePhysicalMaterial();
 
-#if WITH_CHAOS
-	if (ActorHandle->CastToRigidParticle() != nullptr)
-	{
-		CHAOS_ENSURE(false);
-	}
-#endif
-
 #if WITH_PHYSX
 	if (FPhysicsInterface::IsValid(ActorHandle) && FPhysicsInterface::IsRigidBody(ActorHandle))
 	{
