Fix Chaos CIS errors

diff --git a/Engine/Source/Runtime/Engine/Public/Physics/PhysicsInterfaceUtils.h b/Engine/Source/Runtime/Engine/Public/Physics/PhysicsInterfaceUtils.h
--- a/Engine/Source/Runtime/Engine/Public/Physics/PhysicsInterfaceUtils.h
+++ b/Engine/Source/Runtime/Engine/Public/Physics/PhysicsInterfaceUtils.h
@@ -6,7 +6,6 @@
 #include "Containers/Union.h"
 #include "Physics/PhysicsInterfaceTypes.h"
 #include "PhysicsInterfaceUtilsCore.h"
-#include "PhysicsEngine/ConstraintInstance.h"
 
 class FPhysScene_PhysX;
 
