Rotate Chaos Derived Data Key.

diff --git a/Engine/Source/Runtime/Engine/Private/PhysicsEngine/Experimental/ChaosDerivedData.cpp b/Engine/Source/Runtime/Engine/Private/PhysicsEngine/Experimental/ChaosDerivedData.cpp
--- a/Engine/Source/Runtime/Engine/Private/PhysicsEngine/Experimental/ChaosDerivedData.cpp
+++ b/Engine/Source/Runtime/Engine/Private/PhysicsEngine/Experimental/ChaosDerivedData.cpp
@@ -21,7 +21,7 @@ const TCHAR* FChaosDerivedDataCooker::GetPluginName() const
 
 const TCHAR* FChaosDerivedDataCooker::GetVersionString() const
 {
-	return TEXT("D809C44D3C23455296C587E17D4C36C7");
+	return TEXT("193F8697B0744E288BD34D164C7E3B22");
 }
 
 FString FChaosDerivedDataCooker::GetPluginSpecificCacheKeySuffix() const
