- Reset OldPosition & PartSysVelocity when activating a system
- Don't set PartSysVelocity on the first frame as it can result in incorrect velocities

#rnx

diff --git a/Engine/Source/Runtime/Engine/Classes/Particles/ParticleSystemComponent.h b/Engine/Source/Runtime/Engine/Classes/Particles/ParticleSystemComponent.h
--- a/Engine/Source/Runtime/Engine/Classes/Particles/ParticleSystemComponent.h
+++ b/Engine/Source/Runtime/Engine/Classes/Particles/ParticleSystemComponent.h
@@ -609,6 +609,9 @@ public:
 	UPROPERTY(BlueprintAssignable)
 	FParticleCollisionSignature OnParticleCollide;
 
+	UPROPERTY()
+	bool bOldPositionValid;
+
 	UPROPERTY()
 	FVector OldPosition;
 
diff --git a/Engine/Source/Runtime/Engine/Private/Particles/ParticleComponents.cpp b/Engine/Source/Runtime/Engine/Private/Particles/ParticleComponents.cpp
--- a/Engine/Source/Runtime/Engine/Private/Particles/ParticleComponents.cpp
+++ b/Engine/Source/Runtime/Engine/Private/Particles/ParticleComponents.cpp
@@ -3358,11 +3358,12 @@ UParticleSystemComponent::UParticleSystemComponent(const FObjectInitializer& Obj
 	MaxTimeBeforeForceUpdateTransform = 5.0f;
 	bAutoActivate = true;
 	bResetOnDetach = false;
-	OldPosition = FVector(0.0f, 0.0f, 0.0f);
+	bOldPositionValid = false;
+	OldPosition = FVector::ZeroVector;
 
 	RandomStream.Initialize(FApp::bUseFixedSeed ? GetFName() : NAME_None);
 
-	PartSysVelocity = FVector(0.0f, 0.0f, 0.0f);
+	PartSysVelocity = FVector::ZeroVector;
 
 	WarmupTime = 0.0f;
 	SecondsBeforeInactive = 1.0f;
@@ -5500,8 +5501,16 @@ void UParticleSystemComponent::FinalizeTickComponent()
 		bIsTransformDirty = false;
 	}
 
-	const float InvDeltaTime = (DeltaTimeTick > 0.0f) ? 1.0f / DeltaTimeTick : 0.0f;
-	PartSysVelocity = (GetComponentLocation() - OldPosition) * InvDeltaTime;
+	if (bOldPositionValid)
+	{
+		const float InvDeltaTime = (DeltaTimeTick > 0.0f) ? 1.0f / DeltaTimeTick : 0.0f;
+		PartSysVelocity = (GetComponentLocation() - OldPosition) * InvDeltaTime;
+	}
+	else
+	{
+		PartSysVelocity = FVector::ZeroVector;
+	}
+	bOldPositionValid = true;
 	OldPosition = GetComponentLocation();
 
 	if (bIsViewRelevanceDirty)
@@ -5883,6 +5892,10 @@ void UParticleSystemComponent::ActivateSystem(bool bFlagAsJustAttached)
 		return;
 	}
 
+	bOldPositionValid = false;
+	OldPosition = FVector::ZeroVector;
+	PartSysVelocity = FVector::ZeroVector;
+
 	UWorld* World = GetWorld();
 	check(World);
 	UE_LOG(LogParticles,Verbose,
