FIX CIS - Add missing declaration to PhysicsInterfaceUtils.h

diff --git a/Engine/Source/Runtime/Engine/Public/Physics/PhysicsInterfaceUtils.h b/Engine/Source/Runtime/Engine/Public/Physics/PhysicsInterfaceUtils.h
--- a/Engine/Source/Runtime/Engine/Public/Physics/PhysicsInterfaceUtils.h
+++ b/Engine/Source/Runtime/Engine/Public/Physics/PhysicsInterfaceUtils.h
@@ -8,6 +8,7 @@
 #include "PhysicsInterfaceUtilsCore.h"
 
 class FPhysScene_PhysX;
+class FConstraintInstance;
 
 // FILTER DATA
 
