FIX CIS - Use struct instead of class in FConstraintIndex forward declaration.

diff --git a/Engine/Source/Runtime/Engine/Public/Physics/PhysicsInterfaceUtils.h b/Engine/Source/Runtime/Engine/Public/Physics/PhysicsInterfaceUtils.h
--- a/Engine/Source/Runtime/Engine/Public/Physics/PhysicsInterfaceUtils.h
+++ b/Engine/Source/Runtime/Engine/Public/Physics/PhysicsInterfaceUtils.h
@@ -8,7 +8,7 @@
 #include "PhysicsInterfaceUtilsCore.h"
 
 class FPhysScene_PhysX;
-class FConstraintInstance;
+struct FConstraintInstance;
 
 // FILTER DATA
 
