Chaos - fix CIS (ImmedoatePhysicsSimulation_Chaos.cpp)

diff --git a/Engine/Source/Runtime/Engine/Private/PhysicsEngine/ImmediatePhysics/ImmediatePhysicsChaos/ImmediatePhysicsSimulation_Chaos.cpp b/Engine/Source/Runtime/Engine/Private/PhysicsEngine/ImmediatePhysics/ImmediatePhysicsChaos/ImmediatePhysicsSimulation_Chaos.cpp
--- a/Engine/Source/Runtime/Engine/Private/PhysicsEngine/ImmediatePhysics/ImmediatePhysicsChaos/ImmediatePhysicsSimulation_Chaos.cpp
+++ b/Engine/Source/Runtime/Engine/Private/PhysicsEngine/ImmediatePhysics/ImmediatePhysicsChaos/ImmediatePhysicsSimulation_Chaos.cpp
@@ -71,8 +71,8 @@ FAutoConsoleVariableRef CVarChaosImmPhysDriveStiffness(TEXT("p.Chaos.ImmPhys.Joi
 FAutoConsoleVariableRef CVarChaosImmPhysMinParentMassRatio(TEXT("p.Chaos.ImmPhys.Joint.MinParentMassRatio"), ChaosImmediate_Joint_MinParentMassRatio, TEXT("6Dof joint MinParentMassRatio (if > 0)"));
 FAutoConsoleVariableRef CVarChaosImmPhysMaxInertiaRatio(TEXT("p.Chaos.ImmPhys.Joint.MaxInertiaRatio"), ChaosImmediate_Joint_MaxInertiaRatio, TEXT("6Dof joint MaxInertiaRatio (if > 0)"));
 
-int32 ChaosImmediate_Material_Resitution = 0.0f;
-int32 ChaosImmediate_Material_Friction = 0.5f;
+float ChaosImmediate_Material_Resitution = 0.0f;
+float ChaosImmediate_Material_Friction = 0.5f;
 
 //
 // end remove when finished
