Copy of 10496919 from Dev-Editor. Reverted CL#9961146, as it caused a problem with copy/pasted nodes losing their connections

diff --git a/Engine/Source/Runtime/Engine/Private/Materials/MaterialExpressions.cpp b/Engine/Source/Runtime/Engine/Private/Materials/MaterialExpressions.cpp
--- a/Engine/Source/Runtime/Engine/Private/Materials/MaterialExpressions.cpp
+++ b/Engine/Source/Runtime/Engine/Private/Materials/MaterialExpressions.cpp
@@ -12956,14 +12956,6 @@ void UMaterialExpressionMaterialFunctionCall::UpdateFromFunctionResource(bool bR
 				CurrentInput.Input = OriginalInput->Input;
 				CurrentInput.Input.InputName = TempInputName;
 			}
-
-			if (GraphNode)
-			{
-				if (CurrentInput.ExpressionInput->bUsePreviewValueAsDefault)
-				{
-					bRecreateAndLinkNode = true;
-				}
-			}
 		}
 
 		for (int32 OutputIndex = 0; OutputIndex < FunctionOutputs.Num(); OutputIndex++)
