Fix missing include in ChaosDerivedDataUtil.cpp

diff --git a/Engine/Source/Runtime/Engine/Private/PhysicsEngine/Experimental/ChaosDerivedDataUtil.cpp b/Engine/Source/Runtime/Engine/Private/PhysicsEngine/Experimental/ChaosDerivedDataUtil.cpp
--- a/Engine/Source/Runtime/Engine/Private/PhysicsEngine/Experimental/ChaosDerivedDataUtil.cpp
+++ b/Engine/Source/Runtime/Engine/Private/PhysicsEngine/Experimental/ChaosDerivedDataUtil.cpp
@@ -4,6 +4,7 @@
 #include "CoreMinimal.h"
 
 #include "ChaosLog.h"
+#include "ChaosCheck.h"
 
 namespace Chaos
 {
