Following up 10614944, don't allow UWorld::Async* trace functions to be run off game thread as it will result in memory corruption

diff --git a/Engine/Source/Runtime/Engine/Private/Collision/WorldCollisionAsync.cpp b/Engine/Source/Runtime/Engine/Private/Collision/WorldCollisionAsync.cpp
--- a/Engine/Source/Runtime/Engine/Private/Collision/WorldCollisionAsync.cpp
+++ b/Engine/Source/Runtime/Engine/Private/Collision/WorldCollisionAsync.cpp
@@ -318,6 +318,9 @@ namespace
 	template <typename DatumType>
 	FTraceHandle StartNewTrace(FWorldAsyncTraceState& State, const DatumType& Val)
 	{
+		// Using async traces outside of the game thread can cause memory corruption
+		check(IsInGameThread());
+
 		// Get the buffer for the current frame
 		AsyncTraceData& DataBuffer = State.GetBufferForCurrentFrame();
 
