#DUPE merge - Source CL 10635668
Fix crash of ControlRig

diff --git a/Engine/Source/Runtime/Engine/Classes/Animation/AnimInstance.h b/Engine/Source/Runtime/Engine/Classes/Animation/AnimInstance.h
--- a/Engine/Source/Runtime/Engine/Classes/Animation/AnimInstance.h
+++ b/Engine/Source/Runtime/Engine/Classes/Animation/AnimInstance.h
@@ -423,6 +423,9 @@ class ENGINE_API UAnimInstance : public UObject
 	/** Flag to check back on the game thread that indicates we need to run PostUpdateAnimation() in the post-eval call */
 	uint8 bNeedsUpdate : 1;
 
+	/** Flag to check if created by LinkedAnimGraph in ReinitializeLinkedAnimInstance */
+	uint8 bCreatedByLinkedAnimGraph : 1;
+
 private:
 	/** True when Montages are being ticked, and Montage Events should be queued. 
 	 * When Montage are being ticked, we queue AnimNotifies and Events. We trigger notifies first, then Montage events. */
diff --git a/Engine/Source/Runtime/Engine/Private/Animation/AnimInstance.cpp b/Engine/Source/Runtime/Engine/Private/Animation/AnimInstance.cpp
--- a/Engine/Source/Runtime/Engine/Private/Animation/AnimInstance.cpp
+++ b/Engine/Source/Runtime/Engine/Private/Animation/AnimInstance.cpp
@@ -107,7 +107,7 @@ UAnimInstance::UAnimInstance(const FObjectInitializer& ObjectInitializer)
 
 	// Default to using threaded animation update.
 	bUseMultiThreadedAnimationUpdate = true;
-
+	bCreatedByLinkedAnimGraph = false;
 	PendingDynamicResetTeleportType = ETeleportType::None;
 }
 
diff --git a/Engine/Source/Runtime/Engine/Private/Animation/AnimNode_LinkedAnimGraph.cpp b/Engine/Source/Runtime/Engine/Private/Animation/AnimNode_LinkedAnimGraph.cpp
--- a/Engine/Source/Runtime/Engine/Private/Animation/AnimNode_LinkedAnimGraph.cpp
+++ b/Engine/Source/Runtime/Engine/Private/Animation/AnimNode_LinkedAnimGraph.cpp
@@ -241,6 +241,14 @@ void FAnimNode_LinkedAnimGraph::ReinitializeLinkedAnimInstance(const UAnimInstan
 
 		// Need an instance to run, so create it now
 		InstanceToRun = InNewAnimInstance ? InNewAnimInstance : NewObject<UAnimInstance>(MeshComp, InstanceClass);
+		if (InNewAnimInstance == nullptr)
+		{
+			// if incoming AnimInstance was null, it was created by this function
+			// we mark them as created by linked anim graph
+			// this is to know who owns memory instance
+			InstanceToRun->bCreatedByLinkedAnimGraph = true;
+		}
+
 		SetTargetInstance(InstanceToRun);
 
 		// Link before we call InitializeAnimation() so we propgate the call to linked input poses
diff --git a/Engine/Source/Runtime/Engine/Private/Components/SkeletalMeshComponent.cpp b/Engine/Source/Runtime/Engine/Private/Components/SkeletalMeshComponent.cpp
--- a/Engine/Source/Runtime/Engine/Private/Components/SkeletalMeshComponent.cpp
+++ b/Engine/Source/Runtime/Engine/Private/Components/SkeletalMeshComponent.cpp
@@ -2722,7 +2722,7 @@ void USkeletalMeshComponent::ResetLinkedAnimInstances()
 {
 	for(UAnimInstance* LinkedInstance : LinkedInstances)
 	{
-		if(LinkedInstance)
+		if(LinkedInstance && LinkedInstance->bCreatedByLinkedAnimGraph)
 		{
 			LinkedInstance->MarkPendingKill();
 			LinkedInstance = nullptr;
