- Particle Performance Stats runtime gathering across various parts
- Used for automated testing of performance

#rnx

diff --git a/Engine/Plugins/FX/Niagara/Source/Niagara/Classes/NiagaraSystem.h b/Engine/Plugins/FX/Niagara/Source/Niagara/Classes/NiagaraSystem.h
--- a/Engine/Plugins/FX/Niagara/Source/Niagara/Classes/NiagaraSystem.h
+++ b/Engine/Plugins/FX/Niagara/Source/Niagara/Classes/NiagaraSystem.h
@@ -6,6 +6,7 @@
 #include "UObject/ObjectMacros.h"
 #include "UObject/Object.h"
 #include "Particles/ParticleSystem.h"
+#include "Particles/ParticlePerfStats.h"
 #include "NiagaraCommon.h"
 #include "NiagaraDataSet.h"
 #include "NiagaraEmitterInstance.h"
@@ -338,7 +339,6 @@ private:
 	void UpdatePostCompileDIInfo();
 
 protected:
-
 	UPROPERTY(EditAnywhere, Category = "System")
 	UNiagaraEffectType* EffectType;
 
diff --git a/Engine/Plugins/FX/Niagara/Source/Niagara/Private/NiagaraComponent.cpp b/Engine/Plugins/FX/Niagara/Source/Niagara/Private/NiagaraComponent.cpp
--- a/Engine/Plugins/FX/Niagara/Source/Niagara/Private/NiagaraComponent.cpp
+++ b/Engine/Plugins/FX/Niagara/Source/Niagara/Private/NiagaraComponent.cpp
@@ -114,6 +114,9 @@ FNiagaraSceneProxy::FNiagaraSceneProxy(const UNiagaraComponent* InComponent)
 		: FPrimitiveSceneProxy(InComponent, InComponent->GetAsset() ? InComponent->GetAsset()->GetFName() : FName())
 		, bRenderingEnabled(true)
 		, RuntimeCycleCount(nullptr)
+#if WITH_PARTICLE_PERF_STATS
+		, PerfAsset(InComponent->GetAsset())
+#endif
 {
 	// In this case only, update the System renderers on the game thread.
 	check(IsInGameThread());
@@ -1142,6 +1145,7 @@ void UNiagaraComponent::SendRenderDynamicData_Concurrent()
 	LLM_SCOPE(ELLMTag::Niagara);
 	CSV_SCOPED_TIMING_STAT_EXCLUSIVE(Niagara);
 	SCOPE_CYCLE_COUNTER(STAT_NiagaraComponentSendRenderData);
+	PARTICLE_PERF_STAT_CYCLES(Asset, EndOfFrame);
 
 	Super::SendRenderDynamicData_Concurrent();
 
@@ -1206,9 +1210,10 @@ void UNiagaraComponent::SendRenderDynamicData_Concurrent()
 #endif
 		
 		ENQUEUE_RENDER_COMMAND(NiagaraSetDynamicData)(
-			[NiagaraProxy, DynamicData = MoveTemp(NewDynamicData)](FRHICommandListImmediate& RHICmdList)
+			[NiagaraProxy, DynamicData = MoveTemp(NewDynamicData), PerfAsset=Asset](FRHICommandListImmediate& RHICmdList)
 		{
 			SCOPE_CYCLE_COUNTER(STAT_NiagaraSetDynamicData);
+			PARTICLE_PERF_STAT_CYCLES(PerfAsset, RenderUpdate);
 
 			const TArray<FNiagaraRenderer*>& EmitterRenderers_RT = NiagaraProxy->GetEmitterRenderers();
 			for (int32 i = 0; i < EmitterRenderers_RT.Num(); ++i)
diff --git a/Engine/Plugins/FX/Niagara/Source/Niagara/Private/NiagaraRendererMeshes.cpp b/Engine/Plugins/FX/Niagara/Source/Niagara/Private/NiagaraRendererMeshes.cpp
--- a/Engine/Plugins/FX/Niagara/Source/Niagara/Private/NiagaraRendererMeshes.cpp
+++ b/Engine/Plugins/FX/Niagara/Source/Niagara/Private/NiagaraRendererMeshes.cpp
@@ -155,9 +155,11 @@ void FNiagaraRendererMeshes::SetupVertexFactory(FNiagaraMeshVertexFactory *InVer
 
 void FNiagaraRendererMeshes::GetDynamicMeshElements(const TArray<const FSceneView*>& Views, const FSceneViewFamily& ViewFamily, uint32 VisibilityMap, FMeshElementCollector& Collector, const FNiagaraSceneProxy *SceneProxy) const
 {
+	check(SceneProxy);
+
 	SCOPE_CYCLE_COUNTER(STAT_NiagaraRender);
 	SCOPE_CYCLE_COUNTER(STAT_NiagaraRenderMeshes);
-	check(SceneProxy);
+	PARTICLE_PERF_STAT_CYCLES(SceneProxy->PerfAsset, GetDynamicMeshElements);
 
 	NiagaraEmitterInstanceBatcher* Batcher = SceneProxy->GetBatcher();
 	FNiagaraDynamicDataMesh *DynamicDataMesh = (static_cast<FNiagaraDynamicDataMesh*>(DynamicDataRender));
diff --git a/Engine/Plugins/FX/Niagara/Source/Niagara/Private/NiagaraRendererRibbons.cpp b/Engine/Plugins/FX/Niagara/Source/Niagara/Private/NiagaraRendererRibbons.cpp
--- a/Engine/Plugins/FX/Niagara/Source/Niagara/Private/NiagaraRendererRibbons.cpp
+++ b/Engine/Plugins/FX/Niagara/Source/Niagara/Private/NiagaraRendererRibbons.cpp
@@ -282,6 +282,7 @@ void FNiagaraRendererRibbons::GetDynamicMeshElements(const TArray<const FSceneVi
 {
 	SCOPE_CYCLE_COUNTER(STAT_NiagaraRender);
 	SCOPE_CYCLE_COUNTER(STAT_NiagaraRenderRibbons);
+	PARTICLE_PERF_STAT_CYCLES(SceneProxy->PerfAsset, GetDynamicMeshElements);
 
 	FNiagaraDynamicDataRibbon *DynamicDataRibbon = static_cast<FNiagaraDynamicDataRibbon*>(DynamicDataRender);
 	if (!DynamicDataRibbon)
diff --git a/Engine/Plugins/FX/Niagara/Source/Niagara/Private/NiagaraRendererSprites.cpp b/Engine/Plugins/FX/Niagara/Source/Niagara/Private/NiagaraRendererSprites.cpp
--- a/Engine/Plugins/FX/Niagara/Source/Niagara/Private/NiagaraRendererSprites.cpp
+++ b/Engine/Plugins/FX/Niagara/Source/Niagara/Private/NiagaraRendererSprites.cpp
@@ -544,6 +544,7 @@ void FNiagaraRendererSprites::GetDynamicMeshElements(const TArray<const FSceneVi
 {
 	SCOPE_CYCLE_COUNTER(STAT_NiagaraRender);
 	SCOPE_CYCLE_COUNTER(STAT_NiagaraRenderSprites);
+	PARTICLE_PERF_STAT_CYCLES(SceneProxy->PerfAsset, GetDynamicMeshElements);
 	check(SceneProxy);
 
 	//check(DynamicDataRender)
diff --git a/Engine/Plugins/FX/Niagara/Source/Niagara/Private/NiagaraSystemSimulation.cpp b/Engine/Plugins/FX/Niagara/Source/Niagara/Private/NiagaraSystemSimulation.cpp
--- a/Engine/Plugins/FX/Niagara/Source/Niagara/Private/NiagaraSystemSimulation.cpp
+++ b/Engine/Plugins/FX/Niagara/Source/Niagara/Private/NiagaraSystemSimulation.cpp
@@ -160,9 +160,12 @@ public:
 
 	void DoTask(ENamedThreads::Type CurrentThread, const FGraphEventRef& MyCompletionGraphEvent)
 	{
-		Context.MyCompletionGraphEvent = MyCompletionGraphEvent;
-		Context.Owner->Tick_Concurrent(Context);
-		Context.FinalizeEvents = nullptr;
+		{
+			PARTICLE_PERF_STAT_CYCLES(Context.System, TickConcurrent);
+			Context.MyCompletionGraphEvent = MyCompletionGraphEvent;
+			Context.Owner->Tick_Concurrent(Context);
+			Context.FinalizeEvents = nullptr;
+		}
 		WaitAllFinalizeTask->Unlock();
 	}
 };
@@ -194,6 +197,7 @@ public:
 		check(CurrentThread == ENamedThreads::GameThread);
 		FNiagaraScopedRuntimeCycleCounter RuntimeScope(SystemSim->GetSystem(), true, false);
 
+		PARTICLE_PERF_STAT_CYCLES(SystemSim->GetSystem(), Finalize);
 		for (FNiagaraSystemInstance* Inst : Batch)
 		{
 			Inst->FinalizeTick_GameThread();
@@ -231,7 +235,7 @@ public:
 
 	void DoTask(ENamedThreads::Type CurrentThread, const FGraphEventRef& MyCompletionGraphEvent)
 	{
-		FNiagaraScopedRuntimeCycleCounter RuntimeScope(SystemSim->GetSystem(), true, true);
+		PARTICLE_PERF_STAT_CYCLES(Batch[0]->GetSystem(), TickConcurrent);
 		for (FNiagaraSystemInstance* Inst : Batch)
 		{
 			Inst->Tick_Concurrent();
@@ -626,16 +630,18 @@ void FNiagaraSystemSimulation::Tick_GameThread(float DeltaSeconds, const FGraphE
 	LLM_SCOPE(ELLMTag::Niagara);
 	FScopeCycleCounterUObject AdditionalScope(GetSystem(), GET_STATID(STAT_NiagaraOverview_GT_CNC));
 
+	UNiagaraSystem* System = WeakSystem.Get();
+	FScopeCycleCounter SystemStatCounter(System->GetStatID(true, false));
+	PARTICLE_PERF_STAT_INSTANCE_COUNT(System, SystemInstances.Num());
+	PARTICLE_PERF_STAT_CYCLES(System, TickGameThread);
+
 	bHasEverTicked = true;
 
 	SystemTickGraphEvent = nullptr;
 
 	check(SystemInstances.Num() == MainDataSet.GetCurrentDataChecked().GetNumInstances());
 	check(PausedSystemInstances.Num() == PausedInstanceData.GetCurrentDataChecked().GetNumInstances());
-
-	UNiagaraSystem* System = WeakSystem.Get();
 	FNiagaraScopedRuntimeCycleCounter RuntimeScope(System, true, false);
-	FScopeCycleCounter SystemStatCounter(System->GetStatID(true, false));
 
 	if (MaxDeltaTime.IsSet())
 	{
diff --git a/Engine/Plugins/FX/Niagara/Source/Niagara/Private/NiagaraWorldManager.cpp b/Engine/Plugins/FX/Niagara/Source/Niagara/Private/NiagaraWorldManager.cpp
--- a/Engine/Plugins/FX/Niagara/Source/Niagara/Private/NiagaraWorldManager.cpp
+++ b/Engine/Plugins/FX/Niagara/Source/Niagara/Private/NiagaraWorldManager.cpp
@@ -10,6 +10,7 @@
 #include "NiagaraSystemInstance.h"
 #include "Scalability.h"
 #include "Misc/ConfigCacheIni.h"
+#include "Misc/CoreDelegates.h"
 #include "NiagaraDataInterfaceSkeletalMesh.h"
 #include "GameFramework/PlayerController.h"
 #include "EngineModule.h"
diff --git a/Engine/Plugins/FX/Niagara/Source/Niagara/Public/NiagaraComponent.h b/Engine/Plugins/FX/Niagara/Source/Niagara/Public/NiagaraComponent.h
--- a/Engine/Plugins/FX/Niagara/Source/Niagara/Public/NiagaraComponent.h
+++ b/Engine/Plugins/FX/Niagara/Source/Niagara/Public/NiagaraComponent.h
@@ -11,6 +11,7 @@
 #include "NiagaraUserRedirectionParameterStore.h"
 #include "NiagaraSystemInstance.h"
 #include "NiagaraComponentPool.h"
+#include "Particles/ParticlePerfStats.h"
 
 #include "NiagaraComponent.generated.h"
 
@@ -648,6 +649,10 @@ private:
 #if STATS
 	TStatId SystemStatID;
 #endif
+#if WITH_PARTICLE_PERF_STATS
+public:
+	class UNiagaraSystem* PerfAsset;
+#endif
 };
 
 extern float GLastRenderTimeSafetyBias;
diff --git a/Engine/Plugins/FX/Niagara/Source/Niagara/Public/NiagaraWorldManager.h b/Engine/Plugins/FX/Niagara/Source/Niagara/Public/NiagaraWorldManager.h
--- a/Engine/Plugins/FX/Niagara/Source/Niagara/Public/NiagaraWorldManager.h
+++ b/Engine/Plugins/FX/Niagara/Source/Niagara/Public/NiagaraWorldManager.h
@@ -7,6 +7,7 @@
 #include "Engine/EngineTypes.h"
 #include "Engine/EngineBaseTypes.h"
 #include "Engine/World.h"
+#include "Particles/ParticlePerfStats.h"
 #include "NiagaraParameterCollection.h"
 #include "UObject/GCObject.h"
 #include "NiagaraDataSet.h"
diff --git a/Engine/Source/Runtime/Engine/Classes/Particles/ParticlePerfStats.h b/Engine/Source/Runtime/Engine/Classes/Particles/ParticlePerfStats.h
new file mode 100644
--- /dev/null
+++ b/Engine/Source/Runtime/Engine/Classes/Particles/ParticlePerfStats.h
@@ -0,0 +1,87 @@
+// Copyright 1998-2019 Epic Games, Inc. All Rights Reserved.
+
+#pragma once
+
+#include "Templates/Atomic.h"
+
+#ifndef WITH_PARTICLE_PERF_STATS
+	#define WITH_PARTICLE_PERF_STATS !UE_BUILD_SHIPPING
+#endif
+
+#if WITH_PARTICLE_PERF_STATS
+
+struct ENGINE_API FParticlePerfStats
+{
+	struct FGroupedStats
+	{
+		FGroupedStats()
+			: TickConcurrentCycles(0)
+			, EndOfFrameCycles(0)
+		{
+			Reset();
+		}
+
+		void Reset()
+		{
+			NumInstances = 0;
+
+			TickGameThreadCycles = 0;
+			TickConcurrentCycles = 0;
+			FinalizeCycles = 0;
+			EndOfFrameCycles = 0;
+
+			RenderUpdateCycles = 0;
+			GetDynamicMeshElementsCycles = 0;
+		}
+
+		uint32			NumInstances;					// Total instances that this is for
+
+		uint64			TickGameThreadCycles;			// Total cycles spent during game thread update (Note can include concurrent tick / finalize if we don't tick async)
+		TAtomic<uint64>	TickConcurrentCycles;			// Total cycles spent during concurrent thread update
+		uint64			FinalizeCycles;					// Total cycles spent finalizing
+		TAtomic<uint64>	EndOfFrameCycles;				// Total cycles spent during end of frame updates
+
+		uint64			RenderUpdateCycles;				// Total cycles spent doing render updates
+		uint64			GetDynamicMeshElementsCycles;	// Total cycles spent setting up dynamic mesh elements
+	};
+
+	uint32			AccumulatedNumFrames = 0;
+	FGroupedStats	AccumulatedStats;
+
+	FGroupedStats	CurrentFrameStats;
+
+	void Reset();
+	void Tick();
+
+	static void OnStartup();
+	static void OnShutdown();
+
+	static FParticlePerfStats Dummy;
+};
+
+//-TODO: Need to remove the task graph WakeUp cost otherwise it skews results!
+#define PARTICLE_PERF_STAT_INSTANCE_COUNT(STATS, COUNT) \
+	(STATS) != nullptr ? (STATS)->ParticlePerfStats.CurrentFrameStats.NumInstances += COUNT : 0
+
+#define PARTICLE_PERF_STAT_CYCLES(STATS, NAME) \
+	struct FScopedParticleStat_##NAME \
+	{ \
+		FORCEINLINE FScopedParticleStat_ ## NAME(FParticlePerfStats& InStats) \
+			: Stats(InStats) \
+			, StartCycles(FPlatformTime::Cycles64()) \
+		{ \
+		} \
+		FORCEINLINE ~FScopedParticleStat_##NAME() \
+		{ \
+			Stats.CurrentFrameStats.NAME##Cycles += FPlatformTime::Cycles64() - StartCycles; \
+		} \
+		FParticlePerfStats& Stats; \
+		uint64 StartCycles; \
+	} NAME##Stat((STATS) ? (STATS)->ParticlePerfStats : FParticlePerfStats::Dummy);
+
+#else //WITH_PARTICLE_PERF_STATS
+
+#define PARTICLE_PERF_STAT_INSTANCE_COUNT(STATS, COUNT)
+#define PARTICLE_PERF_STAT_CYCLES(STATS, NAME, INSTANCECOUNT)
+
+#endif //WITH_PARTICLE_PERF_STATS
diff --git a/Engine/Source/Runtime/Engine/Classes/Particles/ParticleSystem.h b/Engine/Source/Runtime/Engine/Classes/Particles/ParticleSystem.h
--- a/Engine/Source/Runtime/Engine/Classes/Particles/ParticleSystem.h
+++ b/Engine/Source/Runtime/Engine/Classes/Particles/ParticleSystem.h
@@ -7,6 +7,7 @@
 #include "UObject/ObjectMacros.h"
 #include "UObject/Object.h"
 #include "UObject/ScriptMacros.h"
+#include "Particles/ParticlePerfStats.h"
 
 #include "ParticleSystem.generated.h"
 
@@ -118,6 +119,10 @@ public:
 	UPROPERTY(EditAnywhere, Category = Performance)
 	uint32 MaxPoolSize;
 	//TODO: Allow pool size overriding per world and possibly implement some preallocation too.
+
+#if WITH_PARTICLE_PERF_STATS
+	FParticlePerfStats ParticlePerfStats;
+#endif
 };
 
 /**
diff --git a/Engine/Source/Runtime/Engine/Private/Particles/ParticleComponents.cpp b/Engine/Source/Runtime/Engine/Private/Particles/ParticleComponents.cpp
--- a/Engine/Source/Runtime/Engine/Private/Particles/ParticleComponents.cpp
+++ b/Engine/Source/Runtime/Engine/Private/Particles/ParticleComponents.cpp
@@ -3984,6 +3984,7 @@ void UParticleSystemComponent::SendRenderDynamicData_Concurrent()
 	SCOPE_CYCLE_COUNTER(STAT_ParticleSystemComponent_SendRenderDynamicData_Concurrent);
 	SCOPE_CYCLE_COUNTER(STAT_ParticlesOverview_GT_CNC);
 	CSV_SCOPED_TIMING_STAT_EXCLUSIVE(Effects);
+	PARTICLE_PERF_STAT_CYCLES(Template, EndOfFrame);
 
 	ForceAsyncWorkCompletion(ENSURE_AND_STALL, false, true);
 	Super::SendRenderDynamicData_Concurrent();
@@ -5045,6 +5046,8 @@ void UParticleSystemComponent::TickComponent(float DeltaTime, enum ELevelTick Ti
 		SetComponentTickEnabled(false);
 		return;
 	}
+	PARTICLE_PERF_STAT_CYCLES(Template, TickGameThread);
+	PARTICLE_PERF_STAT_INSTANCE_COUNT(Template, 1);
 
 	checkf(!IsTickManaged() || !PrimaryComponentTick.IsTickFunctionEnabled(), TEXT("PSC has enabled tick funciton and is also ticking via the tick manager.\nTemplate:%s\nPSC: %s\nParent:%s")
 	, *Template->GetFullName(), *GetFullName(), GetAttachParent() ? *GetAttachParent()->GetFullName() : TEXT("nullptr"));
@@ -5302,6 +5305,7 @@ void UParticleSystemComponent::ComputeTickComponent_Concurrent()
 	SCOPE_CYCLE_COUNTER(STAT_ParticleComputeTickTime);
 	FScopeCycleCounterUObject AdditionalScope(AdditionalStatObject(), GET_STATID(STAT_ParticleComputeTickTime));
 	SCOPE_CYCLE_COUNTER(STAT_ParticlesOverview_GT_CNC);
+	PARTICLE_PERF_STAT_CYCLES(Template, TickConcurrent);
 
 	// Tick Subemitters.
 	int32 EmitterIndex;
@@ -5377,6 +5381,7 @@ void UParticleSystemComponent::FinalizeTickComponent()
 
 	SCOPE_CYCLE_COUNTER(STAT_ParticleFinalizeTickTime);
 	SCOPE_CYCLE_COUNTER(STAT_ParticlesOverview_GT);
+	PARTICLE_PERF_STAT_CYCLES(Template, Finalize);
 
 	if(bAsyncDataCopyIsValid)
 	{
diff --git a/Engine/Source/Runtime/Engine/Private/Particles/ParticlePerfStats.cpp b/Engine/Source/Runtime/Engine/Private/Particles/ParticlePerfStats.cpp
new file mode 100644
--- /dev/null
+++ b/Engine/Source/Runtime/Engine/Private/Particles/ParticlePerfStats.cpp
@@ -0,0 +1,187 @@
+// Copyright 1998-2019 Epic Games, Inc. All Rights Reserved.
+
+#include "Particles/ParticlePerfStats.h"
+#include "CoreMinimal.h"
+
+#if WITH_PARTICLE_PERF_STATS
+
+namespace ParticlePerfStatsLocal
+{
+	static FDelegateHandle BeginFrameHandle;
+	static int32 GbParticlePerfStatsEnabled = 0;
+	static int32 GbParticlePerfStatsEnabledLatch = 0;
+	static int32 GParticlePerfStatsFramesRemaining = 0;
+
+	static void DumpParticlePerfStats(FOutputDevice& Ar)
+	{
+		FlushRenderingCommands();
+
+		FString tempString;
+
+		Ar.Logf(TEXT("**** Particle Performance Stats"));
+		Ar.Logf(TEXT("Name,Average PerFrame GameThread,Average PerInstance GameThread,Average PerFrame RenderThread,Average PerInstance RenderThread,NumFrames,Total Instances,Total Tick GameThread,Total Tick Concurrent,Total Finalize,Total End Of Frame,Total Render Update,Total Get Dynamic Mesh Elements"));
+		for (TObjectIterator<UFXSystemAsset> SystemIt; SystemIt; ++SystemIt)
+		{
+			const FParticlePerfStats& PerfStats = SystemIt->ParticlePerfStats;
+			if (PerfStats.AccumulatedStats.NumInstances == 0)
+				continue;
+
+			const uint64 TotalGameThread = PerfStats.AccumulatedStats.TickGameThreadCycles + PerfStats.AccumulatedStats.TickConcurrentCycles + PerfStats.AccumulatedStats.FinalizeCycles + PerfStats.AccumulatedStats.EndOfFrameCycles;
+			const uint64 TotalRenderThread = PerfStats.AccumulatedStats.RenderUpdateCycles + PerfStats.AccumulatedStats.GetDynamicMeshElementsCycles;
+
+			tempString.Reset();
+			tempString.Append(*SystemIt->GetFullName());
+			tempString.Appendf(TEXT(",%u"), uint32(FPlatformTime::ToMilliseconds64(TotalGameThread) * 1000.0 / double(PerfStats.AccumulatedNumFrames)));
+			tempString.Appendf(TEXT(",%u"), uint32(FPlatformTime::ToMilliseconds64(TotalGameThread) * 1000.0 / double(PerfStats.AccumulatedStats.NumInstances)));
+			tempString.Appendf(TEXT(",%u"), uint32(FPlatformTime::ToMilliseconds64(TotalRenderThread) * 1000.0 / double(PerfStats.AccumulatedNumFrames)));
+			tempString.Appendf(TEXT(",%u"), uint32(FPlatformTime::ToMilliseconds64(TotalRenderThread) * 1000.0 / double(PerfStats.AccumulatedStats.NumInstances)));
+			tempString.Appendf(TEXT(",%u"), uint32(PerfStats.AccumulatedNumFrames));
+			tempString.Appendf(TEXT(",%u"), uint32(PerfStats.AccumulatedStats.NumInstances));
+			tempString.Appendf(TEXT(",%u"), uint32(FPlatformTime::ToMilliseconds64(PerfStats.AccumulatedStats.TickGameThreadCycles) * 1000.0));
+			tempString.Appendf(TEXT(",%u"), uint32(FPlatformTime::ToMilliseconds64(PerfStats.AccumulatedStats.TickConcurrentCycles) * 1000.0));
+			tempString.Appendf(TEXT(",%u"), uint32(FPlatformTime::ToMilliseconds64(PerfStats.AccumulatedStats.FinalizeCycles) * 1000.0));
+			tempString.Appendf(TEXT(",%u"), uint32(FPlatformTime::ToMilliseconds64(PerfStats.AccumulatedStats.EndOfFrameCycles) * 1000.0));
+			tempString.Appendf(TEXT(",%u"), uint32(FPlatformTime::ToMilliseconds64(PerfStats.AccumulatedStats.RenderUpdateCycles) * 1000.0));
+			tempString.Appendf(TEXT(",%u"), uint32(FPlatformTime::ToMilliseconds64(PerfStats.AccumulatedStats.GetDynamicMeshElementsCycles) * 1000.0));
+			Ar.Log(*tempString);
+		}
+	}
+
+	static void ResetParticlePerfStats()
+	{
+		FlushRenderingCommands();
+
+		for (TObjectIterator<UFXSystemAsset> SystemIt; SystemIt; ++SystemIt)
+		{
+			SystemIt->ParticlePerfStats.Reset();
+		}
+	}
+
+	static void TickParticlePerfStats()
+	{
+		if (GbParticlePerfStatsEnabledLatch != 0)
+		{
+			for (TObjectIterator<UFXSystemAsset> SystemIt; SystemIt; ++SystemIt)
+			{
+				SystemIt->ParticlePerfStats.Tick();
+			}
+
+			bool bDisableGathering = GbParticlePerfStatsEnabled == 0;
+			if (GParticlePerfStatsFramesRemaining > 0)
+			{
+				--GParticlePerfStatsFramesRemaining;
+				bDisableGathering = GParticlePerfStatsFramesRemaining == 0;
+			}
+
+			if (bDisableGathering)
+			{
+				DumpParticlePerfStats(*GLog);
+				GbParticlePerfStatsEnabledLatch = 0;
+			}
+		}
+		else
+		{
+			if (GbParticlePerfStatsEnabled != 0 || GParticlePerfStatsFramesRemaining > 0)
+			{
+				GbParticlePerfStatsEnabledLatch = true;
+				ResetParticlePerfStats();
+			}
+		}
+	}
+
+	static FAutoConsoleVariableRef CVarParticlePerfStatsEnabled(
+		TEXT("fx.ParticlePerfStats.Enabled"),
+		GbParticlePerfStatsEnabled,
+		TEXT("Used to control if stat gathering is enabled or not.\n"),
+		ECVF_Default
+	);
+
+	static FAutoConsoleCommandWithOutputDevice GParticlePerfStatsDump(
+		TEXT("fx.ParticlePerfStats.Dump"),
+		TEXT("Dumps current particle perf stats to output"),
+		FConsoleCommandWithOutputDeviceDelegate::CreateStatic(DumpParticlePerfStats)
+	);
+
+	static FAutoConsoleCommand GParticlePerfStatsReset(
+		TEXT("fx.ParticlePerfStats.Reset"),
+		TEXT("Resets all particle perf stats to zero"),
+		FConsoleCommandDelegate::CreateStatic(ResetParticlePerfStats)
+	);
+
+	static FAutoConsoleCommandWithWorldAndArgs GParticlePerfStatsRunTest(
+		TEXT("fx.ParticlePerfStats.RunTest"),
+		TEXT("Runs for a number of frames then logs out the results"),
+		FConsoleCommandWithWorldAndArgsDelegate::CreateLambda(
+			[](const TArray<FString>& Args, UWorld* World)
+			{
+				if ( Args.Num() != 1 )
+				{
+					return;
+				}
+				const int32 NumFrames = FCString::Atoi(*Args[0]);
+				if (NumFrames <= 0)
+				{
+					return;
+				}
+
+				GbParticlePerfStatsEnabledLatch = 0;
+				GParticlePerfStatsFramesRemaining = NumFrames;
+			}
+		)
+	);
+}
+
+FParticlePerfStats	FParticlePerfStats::Dummy;
+
+void FParticlePerfStats::Reset()
+{
+	AccumulatedNumFrames = 0;
+	AccumulatedStats.Reset();
+	CurrentFrameStats.Reset();
+}
+
+void FParticlePerfStats::Tick()
+{
+	// Nothing to do if we have no instances
+	if (CurrentFrameStats.NumInstances == 0)
+		return;
+
+	++AccumulatedNumFrames;
+
+	AccumulatedStats.NumInstances += CurrentFrameStats.NumInstances;
+	AccumulatedStats.TickGameThreadCycles += CurrentFrameStats.TickGameThreadCycles;
+	AccumulatedStats.TickConcurrentCycles += CurrentFrameStats.TickConcurrentCycles;
+	AccumulatedStats.FinalizeCycles += CurrentFrameStats.FinalizeCycles;
+	AccumulatedStats.EndOfFrameCycles += CurrentFrameStats.EndOfFrameCycles;
+
+	CurrentFrameStats.NumInstances = 0;
+	CurrentFrameStats.TickGameThreadCycles = 0;
+	CurrentFrameStats.TickConcurrentCycles = 0;
+	CurrentFrameStats.FinalizeCycles = 0;
+	CurrentFrameStats.EndOfFrameCycles = 0;
+
+	// Update RenderThread stats
+	ENQUEUE_RENDER_COMMAND(FGatherParticlePerfStats)
+		(
+			[RTPerfStats=this](FRHICommandListImmediate& RHICmdList)
+			{
+				RTPerfStats->AccumulatedStats.RenderUpdateCycles += RTPerfStats->CurrentFrameStats.RenderUpdateCycles;
+				RTPerfStats->AccumulatedStats.GetDynamicMeshElementsCycles += RTPerfStats->CurrentFrameStats.GetDynamicMeshElementsCycles;
+
+				RTPerfStats->CurrentFrameStats.RenderUpdateCycles = 0;
+				RTPerfStats->CurrentFrameStats.GetDynamicMeshElementsCycles = 0;
+	}
+		);
+}
+
+void FParticlePerfStats::OnStartup()
+{
+	ParticlePerfStatsLocal::BeginFrameHandle = FCoreDelegates::OnBeginFrame.AddStatic(ParticlePerfStatsLocal::TickParticlePerfStats);
+}
+
+void FParticlePerfStats::OnShutdown()
+{
+	FCoreDelegates::OnBeginFrame.Remove(ParticlePerfStatsLocal::BeginFrameHandle);
+}
+
+#endif //WITH_PARTICLE_PERF_STATS
diff --git a/Engine/Source/Runtime/Engine/Private/Particles/ParticleSystemManager.cpp b/Engine/Source/Runtime/Engine/Private/Particles/ParticleSystemManager.cpp
--- a/Engine/Source/Runtime/Engine/Private/Particles/ParticleSystemManager.cpp
+++ b/Engine/Source/Runtime/Engine/Private/Particles/ParticleSystemManager.cpp
@@ -146,6 +146,9 @@ void FParticleSystemWorldManager::OnStartup()
 	FWorldDelegates::OnPostWorldCleanup.AddStatic(&OnWorldCleanup);
 	FWorldDelegates::OnPreWorldFinishDestroy.AddStatic(&OnPreWorldFinishDestroy);
 	FWorldDelegates::OnWorldBeginTearDown.AddStatic(&OnWorldBeginTearDown);
+#if WITH_PARTICLE_PERF_STATS
+	FParticlePerfStats::OnStartup();
+#endif
 }
 
 void FParticleSystemWorldManager::OnShutdown()
@@ -158,6 +161,9 @@ void FParticleSystemWorldManager::OnShutdown()
 		}
 	}
 	WorldManagers.Empty();
+#if WITH_PARTICLE_PERF_STATS
+	FParticlePerfStats::OnShutdown();
+#endif
 }
 
 void FParticleSystemWorldManager::OnWorldInit(UWorld* World, const UWorld::InitializationValues IVS)
diff --git a/Engine/Source/Runtime/Engine/Private/Particles/ParticleSystemRender.cpp b/Engine/Source/Runtime/Engine/Private/Particles/ParticleSystemRender.cpp
--- a/Engine/Source/Runtime/Engine/Private/Particles/ParticleSystemRender.cpp
+++ b/Engine/Source/Runtime/Engine/Private/Particles/ParticleSystemRender.cpp
@@ -6910,6 +6910,9 @@ FParticleSystemSceneProxy::FParticleSystemSceneProxy(const UParticleSystemCompon
 	, VisualizeLODIndex(Component->GetCurrentLODIndex())
 	, LastFramePreRendered(-1)
 	, FirstFreeMeshBatch(0)
+#if WITH_PARTICLE_PERF_STATS
+	, PerfAsset(Component->Template)
+#endif
 {
 	SetWireframeColor(FLinearColor(3.0f, 0.0f, 0.0f));
 	SetLevelColor(FLinearColor(1.0f, 1.0f, 0.0f));
@@ -6967,6 +6970,7 @@ void FParticleSystemSceneProxy::GetDynamicMeshElements(const TArray<const FScene
 
 	SCOPE_CYCLE_COUNTER(STAT_FParticleSystemSceneProxy_GetMeshElements);
 	SCOPE_CYCLE_COUNTER(STAT_ParticlesOverview_RT);
+	PARTICLE_PERF_STAT_CYCLES(PerfAsset, GetDynamicMeshElements);
 
 	if ((GIsEditor == true) || (GbEnableGameThreadLODCalculation == false))
 	{
@@ -7100,6 +7104,7 @@ void FParticleSystemSceneProxy::UpdateData(FParticleDynamicData* NewDynamicData)
 			CSV_SCOPED_TIMING_STAT_EXCLUSIVE(ParticleUpdate);
 			SCOPE_CYCLE_COUNTER(STAT_ParticleUpdateRTTime);
 			STAT(FScopeCycleCounter Context(Proxy->GetStatId());)
+				PARTICLE_PERF_STAT_CYCLES(Proxy->PerfAsset, RenderUpdate);
 			if (NewDynamicData)
 			{
 				for (int32 Index = 0; Index < NewDynamicData->DynamicEmitterDataArray.Num(); Index++)
diff --git a/Engine/Source/Runtime/Engine/Public/ParticleHelper.h b/Engine/Source/Runtime/Engine/Public/ParticleHelper.h
--- a/Engine/Source/Runtime/Engine/Public/ParticleHelper.h
+++ b/Engine/Source/Runtime/Engine/Public/ParticleHelper.h
@@ -21,6 +21,7 @@
 #include "MeshBatch.h"
 #include "MeshParticleVertexFactory.h"
 #include "PrimitiveSceneProxy.h"
+#include "Particles/ParticlePerfStats.h"
 
 #define _ENABLE_PARTICLE_LOD_INGAME_
 
@@ -2759,6 +2760,11 @@ protected:
 	mutable TArray<FDynamicEmitterDataBase*> DynamicDataForThisFrame;
 
 	friend struct FDynamicSpriteEmitterDataBase;
+
+#if WITH_PARTICLE_PERF_STATS
+public:
+	UParticleSystem* PerfAsset;
+#endif
 };
 
 #if STATS
