Fix Chaos Scene Queries not working in editor
(Quick fix, pretty hacky.)

diff --git a/Engine/Source/Runtime/Engine/Private/PhysicsEngine/Experimental/PhysInterface_Chaos.cpp b/Engine/Source/Runtime/Engine/Private/PhysicsEngine/Experimental/PhysInterface_Chaos.cpp
--- a/Engine/Source/Runtime/Engine/Private/PhysicsEngine/Experimental/PhysInterface_Chaos.cpp
+++ b/Engine/Source/Runtime/Engine/Private/PhysicsEngine/Experimental/PhysInterface_Chaos.cpp
@@ -359,6 +359,14 @@ FPhysScene* FPhysInterface_Chaos::GetCurrentScene(const FPhysicsActorHandle& InH
 	return nullptr;
 }
 
+void FPhysInterface_Chaos::FlushScene(FPhysScene* InScene)
+{
+	FPhysicsCommand::ExecuteWrite(InScene, [&]()
+	{
+		InScene->Flush_AssumesLocked();
+	});
+}
+
 bool FPhysInterface_Chaos::CanSimulate_AssumesLocked(const FPhysicsActorHandle& InActorReference)
 {
 	// #todo : Implement
diff --git a/Engine/Source/Runtime/Engine/Private/PhysicsEngine/Experimental/PhysScene_Chaos.cpp b/Engine/Source/Runtime/Engine/Private/PhysicsEngine/Experimental/PhysScene_Chaos.cpp
--- a/Engine/Source/Runtime/Engine/Private/PhysicsEngine/Experimental/PhysScene_Chaos.cpp
+++ b/Engine/Source/Runtime/Engine/Private/PhysicsEngine/Experimental/PhysScene_Chaos.cpp
@@ -417,6 +417,29 @@ FPhysScene_Chaos::~FPhysScene_Chaos()
 #endif
 }
 
+#if WITH_EDITOR
+bool FPhysScene_ChaosInterface::IsOwningWorldEditor() const
+{
+	const UWorld* WorldPtr = GetOwningWorld();
+	const TIndirectArray<FWorldContext>& WorldContexts = GEngine->GetWorldContexts();
+	for (const FWorldContext& Context : WorldContexts)
+	{
+		if (WorldPtr)
+		{
+			if (WorldPtr == Context.World())
+			{
+				if (Context.WorldType == EWorldType::Editor)
+				{
+					return true;
+				}
+			}
+		}
+	}
+
+	return false;
+}
+#endif
+
 bool FPhysScene_Chaos::IsTickable() const
 {
 	const bool bDedicatedThread = ChaosModule->IsPersistentTaskRunning();
@@ -936,6 +959,7 @@ FPhysScene_ChaosInterface::FPhysScene_ChaosInterface(const AWorldSettings* InSet
 #endif
 
 	Scene.GetSolver()->PhysSceneHack = this;
+	
 
 }
 
@@ -946,6 +970,31 @@ void FPhysScene_ChaosInterface::OnWorldBeginPlay()
 	{
 		Solver->SetEnabled(true);
 	}
+
+#if WITH_EDITOR
+	const UWorld* WorldPtr = GetOwningWorld();
+	const TIndirectArray<FWorldContext>& WorldContexts = GEngine->GetWorldContexts();
+	for (const FWorldContext& Context : WorldContexts)
+	{
+		if (Context.WorldType == EWorldType::Editor)
+		{
+			UWorld* World = Context.World();
+			if (World)
+			{
+				auto PhysScene = World->GetPhysicsScene();
+				if (PhysScene)
+				{
+					auto InnerSolver = PhysScene->GetSolver();
+					if (InnerSolver)
+					{
+						InnerSolver->SetEnabled(false);
+					}
+				}
+			}
+		}
+	}
+#endif
+
 }
 
 void FPhysScene_ChaosInterface::OnWorldEndPlay()
@@ -954,7 +1003,33 @@ void FPhysScene_ChaosInterface::OnWorldEndPlay()
 	if (Solver)
 	{
 		Solver->SetEnabled(false);
+
 	}
+
+#if WITH_EDITOR
+	const UWorld* WorldPtr = GetOwningWorld();
+	const TIndirectArray<FWorldContext>& WorldContexts = GEngine->GetWorldContexts();
+	for (const FWorldContext& Context : WorldContexts)
+	{
+		if (Context.WorldType == EWorldType::Editor)
+		{
+			UWorld* World = Context.World();
+			if (World)
+			{
+				auto PhysScene = World->GetPhysicsScene();
+				if (PhysScene)
+				{
+					auto InnerSolver = PhysScene->GetSolver();
+					if (InnerSolver)
+					{
+						InnerSolver->SetEnabled(true);
+					}
+				}
+			}
+		}
+	}
+#endif
+
 }
 
 void FPhysScene_ChaosInterface::AddActorsToScene_AssumesLocked(TArray<FPhysicsActorHandle>& InHandles, const bool bImmediate)
@@ -993,6 +1068,14 @@ void FPhysScene_ChaosInterface::AddAggregateToScene(const FPhysicsAggregateHandl
 void FPhysScene_ChaosInterface::SetOwningWorld(UWorld* InOwningWorld)
 {
 	MOwningWorld = InOwningWorld;
+
+#if WITH_EDITOR
+	if (IsOwningWorldEditor())
+	{
+		GetScene().GetSolver()->SetEnabled(true);
+	}
+#endif
+
 }
 
 UWorld* FPhysScene_ChaosInterface::GetOwningWorld()
@@ -1279,6 +1362,21 @@ void FPhysScene_ChaosInterface::StartFrame()
 	FChaosSolversModule* SolverModule = FChaosSolversModule::GetModule();
 	checkSlow(SolverModule);
 
+	float Dt = MDeltaTime;
+
+#if WITH_EDITOR
+	if (IsOwningWorldEditor())
+	{
+		// Ensure editor solver is enabled
+		if (GetSolver()->Enabled() == false)
+		{
+			GetSolver()->SetEnabled(true);
+		}
+
+		Dt = 0.0f;
+	}
+#endif
+
 	if (Chaos::IDispatcher* Dispatcher = SolverModule->GetDispatcher())
 	{
 		for (auto * Solver : SolverModule->GetSolvers())
@@ -1290,27 +1388,27 @@ void FPhysScene_ChaosInterface::StartFrame()
 		{
 		case EChaosThreadingMode::SingleThread:
 		{
-			OnPhysScenePreTick.Broadcast(this, MDeltaTime);
-			OnPhysSceneStep.Broadcast(this, MDeltaTime);
+			OnPhysScenePreTick.Broadcast(this, Dt);
+			OnPhysSceneStep.Broadcast(this, Dt);
 
 			// Here we can directly tick the scene. Single threaded mode doesn't buffer any commands
 			// that would require pumping here - everything is done on demand.
-			Scene.Tick(MDeltaTime);
+			Scene.Tick(Dt);
 		}
 		break;
 		case EChaosThreadingMode::TaskGraph:
 		{
 			check(!CompletionEvent.GetReference())
 
-				OnPhysScenePreTick.Broadcast(this, MDeltaTime);
-			OnPhysSceneStep.Broadcast(this, MDeltaTime);
+			OnPhysScenePreTick.Broadcast(this, Dt);
+			OnPhysSceneStep.Broadcast(this, Dt);
 
 			FGraphEventRef SimulationCompleteEvent = FGraphEvent::CreateGraphEvent();
 
 			// Need to fire off a parallel task to handle running physics commands and
 			// ticking the scene while the engine continues on until TG_EndPhysics
 			// (this should happen in TG_StartPhysics)
-			PhysicsTickTask = TGraphTask<FPhysicsTickTask>::CreateTask(nullptr, ENamedThreads::GameThread).ConstructAndDispatchWhenReady(SimulationCompleteEvent, MDeltaTime);
+			PhysicsTickTask = TGraphTask<FPhysicsTickTask>::CreateTask(nullptr, ENamedThreads::GameThread).ConstructAndDispatchWhenReady(SimulationCompleteEvent, Dt);
 
 			// Setup post simulate tasks
 			if (PhysicsTickTask.GetReference())
diff --git a/Engine/Source/Runtime/Engine/Private/PhysicsEngine/PhysLevel.cpp b/Engine/Source/Runtime/Engine/Private/PhysicsEngine/PhysLevel.cpp
--- a/Engine/Source/Runtime/Engine/Private/PhysicsEngine/PhysLevel.cpp
+++ b/Engine/Source/Runtime/Engine/Private/PhysicsEngine/PhysLevel.cpp
@@ -104,30 +104,37 @@ void UWorld::SetupPhysicsTickFunctions(float DeltaSeconds)
 	
 	EndPhysicsTickFunction.bCanEverTick = true;
 	EndPhysicsTickFunction.Target = this;
+
+// Chaos ticks solver for trace collisions
+#if (WITH_CHAOS && WITH_EDITOR)
+	bool bEnablePhysics = (bShouldSimulatePhysics || bEnableTraceCollision);
+#else
+	bool bEnablePhysics = bShouldSimulatePhysics;
+#endif
 	
-	// see if we need to update tick registration
-	bool bNeedToUpdateTickRegistration = (bShouldSimulatePhysics != StartPhysicsTickFunction.IsTickFunctionRegistered())
-		|| (bShouldSimulatePhysics != EndPhysicsTickFunction.IsTickFunctionRegistered());
+	// see if we need to update tick registration;
+	bool bNeedToUpdateTickRegistration = (bEnablePhysics != StartPhysicsTickFunction.IsTickFunctionRegistered())
+		|| (bEnablePhysics != EndPhysicsTickFunction.IsTickFunctionRegistered());
 
 	if (bNeedToUpdateTickRegistration && PersistentLevel)
 	{
-		if (bShouldSimulatePhysics && !StartPhysicsTickFunction.IsTickFunctionRegistered())
+		if (bEnablePhysics && !StartPhysicsTickFunction.IsTickFunctionRegistered())
 		{
 			StartPhysicsTickFunction.TickGroup = TG_StartPhysics;
 			StartPhysicsTickFunction.RegisterTickFunction(PersistentLevel);
 		}
-		else if (!bShouldSimulatePhysics && StartPhysicsTickFunction.IsTickFunctionRegistered())
+		else if (!bEnablePhysics && StartPhysicsTickFunction.IsTickFunctionRegistered())
 		{
 			StartPhysicsTickFunction.UnRegisterTickFunction();
 		}
 
-		if (bShouldSimulatePhysics && !EndPhysicsTickFunction.IsTickFunctionRegistered())
+		if (bEnablePhysics && !EndPhysicsTickFunction.IsTickFunctionRegistered())
 		{
 			EndPhysicsTickFunction.TickGroup = TG_EndPhysics;
 			EndPhysicsTickFunction.RegisterTickFunction(PersistentLevel);
 			EndPhysicsTickFunction.AddPrerequisite(this, StartPhysicsTickFunction);
 		}
-		else if (!bShouldSimulatePhysics && EndPhysicsTickFunction.IsTickFunctionRegistered())
+		else if (!bEnablePhysics && EndPhysicsTickFunction.IsTickFunctionRegistered())
 		{
 			EndPhysicsTickFunction.RemovePrerequisite(this, StartPhysicsTickFunction);
 			EndPhysicsTickFunction.UnRegisterTickFunction();
diff --git a/Engine/Source/Runtime/Engine/Private/World.cpp b/Engine/Source/Runtime/Engine/Private/World.cpp
--- a/Engine/Source/Runtime/Engine/Private/World.cpp
+++ b/Engine/Source/Runtime/Engine/Private/World.cpp
@@ -1356,6 +1356,12 @@ void UWorld::InitWorld(const InitializationValues IVS)
 	AWorldSettings* WorldSettings = GetWorldSettings();
 	if (IVS.bInitializeScenes)
 	{
+
+	#if WITH_EDITOR
+		bEnableTraceCollision = IVS.bEnableTraceCollision;
+	#endif
+
+
 		if (IVS.bCreatePhysicsScene)
 		{
 			// Create the physics scene
@@ -1390,9 +1396,6 @@ void UWorld::InitWorld(const InitializationValues IVS)
 		AvoidanceManager = NewObject<UAvoidanceManager>(this, GEngine->AvoidanceManagerClass);
 	}
 
-#if WITH_EDITOR
-	bEnableTraceCollision = IVS.bEnableTraceCollision;
-#endif
 
 	SetupParameterCollectionInstances();
 
@@ -4523,10 +4526,14 @@ void UWorld::SetPhysicsScene(FPhysScene* InScene)
 	// Assign pointer
 	PhysicsScene = InScene;
 
-	// Set pointer in scene to know which world its coming from
 	if(PhysicsScene != NULL)
 	{
+		// Set pointer in scene to know which world its coming from
 		PhysicsScene->SetOwningWorld(this);
+
+#if WITH_CHAOS
+		FPhysInterface_Chaos::FlushScene(PhysicsScene);
+#endif
 	}
 }
 
diff --git a/Engine/Source/Runtime/Engine/Public/Physics/Experimental/PhysInterface_Chaos.h b/Engine/Source/Runtime/Engine/Public/Physics/Experimental/PhysInterface_Chaos.h
--- a/Engine/Source/Runtime/Engine/Public/Physics/Experimental/PhysInterface_Chaos.h
+++ b/Engine/Source/Runtime/Engine/Public/Physics/Experimental/PhysInterface_Chaos.h
@@ -226,6 +226,7 @@ public:
     static bool HasAsyncSceneData(const FPhysicsActorHandle& InHandle) { return false; }
 	static bool IsInScene(const FPhysicsActorHandle& InActorReference);
 	static FPhysScene* GetCurrentScene(const FPhysicsActorHandle& InHandle);
+	static void FlushScene(FPhysScene* InScene);
 	static bool CanSimulate_AssumesLocked(const FPhysicsActorHandle& InActorReference);
 	static float GetMass_AssumesLocked(const FPhysicsActorHandle& InActorReference);
 
diff --git a/Engine/Source/Runtime/Engine/Public/Physics/Experimental/PhysScene_Chaos.h b/Engine/Source/Runtime/Engine/Public/Physics/Experimental/PhysScene_Chaos.h
--- a/Engine/Source/Runtime/Engine/Public/Physics/Experimental/PhysScene_Chaos.h
+++ b/Engine/Source/Runtime/Engine/Public/Physics/Experimental/PhysScene_Chaos.h
@@ -190,6 +190,7 @@ private:
 	 */
 	void OnWorldEndPlay();
 
+
 	// List of objects that we modified during a PIE run for physics simulation caching.
 	TArray<UObject*> PieModifiedObjects;
 #endif
@@ -387,6 +388,11 @@ public:
 
 private:
 
+
+#if WITH_EDITOR
+	bool IsOwningWorldEditor() const;
+#endif
+
 	void SyncBodies(Chaos::FPhysicsSolver* Solver);
 
 	void SetKinematicTransform(FPhysicsActorHandle& InActorReference, const Chaos::TRigidTransform<float, 3>& NewTransform)
