- Fix some build configs when no stats are enabled

#rnx

diff --git a/Engine/Source/Runtime/Engine/Classes/Particles/ParticlePerfStats.h b/Engine/Source/Runtime/Engine/Classes/Particles/ParticlePerfStats.h
--- a/Engine/Source/Runtime/Engine/Classes/Particles/ParticlePerfStats.h
+++ b/Engine/Source/Runtime/Engine/Classes/Particles/ParticlePerfStats.h
@@ -82,6 +82,6 @@ struct ENGINE_API FParticlePerfStats
 #else //WITH_PARTICLE_PERF_STATS
 
 #define PARTICLE_PERF_STAT_INSTANCE_COUNT(STATS, COUNT)
-#define PARTICLE_PERF_STAT_CYCLES(STATS, NAME, INSTANCECOUNT)
+#define PARTICLE_PERF_STAT_CYCLES(STATS, NAME)
 
 #endif //WITH_PARTICLE_PERF_STATS
