Fix CIS PhysScene_Chaos.cpp

diff --git a/Engine/Source/Runtime/Engine/Private/PhysicsEngine/Experimental/PhysScene_Chaos.cpp b/Engine/Source/Runtime/Engine/Private/PhysicsEngine/Experimental/PhysScene_Chaos.cpp
--- a/Engine/Source/Runtime/Engine/Private/PhysicsEngine/Experimental/PhysScene_Chaos.cpp
+++ b/Engine/Source/Runtime/Engine/Private/PhysicsEngine/Experimental/PhysScene_Chaos.cpp
@@ -417,7 +417,7 @@ FPhysScene_Chaos::~FPhysScene_Chaos()
 #endif
 }
 
-#if WITH_EDITOR
+#if WITH_EDITOR && WITH_CHAOS
 bool FPhysScene_ChaosInterface::IsOwningWorldEditor() const
 {
 	const UWorld* WorldPtr = GetOwningWorld();
