- Disable perf stats by default, enabled via build config

#rnx

diff --git a/Engine/Source/Runtime/Engine/Classes/Particles/ParticlePerfStats.h b/Engine/Source/Runtime/Engine/Classes/Particles/ParticlePerfStats.h
--- a/Engine/Source/Runtime/Engine/Classes/Particles/ParticlePerfStats.h
+++ b/Engine/Source/Runtime/Engine/Classes/Particles/ParticlePerfStats.h
@@ -5,7 +5,7 @@
 #include "Templates/Atomic.h"
 
 #ifndef WITH_PARTICLE_PERF_STATS
-	#define WITH_PARTICLE_PERF_STATS !UE_BUILD_SHIPPING
+	#define WITH_PARTICLE_PERF_STATS 0
 #endif
 
 #if WITH_PARTICLE_PERF_STATS
