Fix BodyInstance with NoCollision creating shapes that can collide. Partially fixes Fall Test EngineTest.

diff --git a/Engine/Source/Runtime/Engine/Private/PhysicsEngine/BodyInstance.cpp b/Engine/Source/Runtime/Engine/Private/PhysicsEngine/BodyInstance.cpp
--- a/Engine/Source/Runtime/Engine/Private/PhysicsEngine/BodyInstance.cpp
+++ b/Engine/Source/Runtime/Engine/Private/PhysicsEngine/BodyInstance.cpp
@@ -1211,7 +1211,7 @@ struct FInitBodiesHelper
 #if WITH_CHAOS
 							// If this shape shouldn't collide in the sim we disable it here until we support
 							// a separation of unions for these shapes
-							if(BI->GetCollisionEnabled() == ECollisionEnabled::QueryOnly)
+							if(BI->GetCollisionEnabled() == ECollisionEnabled::QueryOnly || BI->GetCollisionEnabled() == ECollisionEnabled::NoCollision)
 							{
 								const int32 NumShapes = FPhysicsInterface::GetNumShapes(ActorHandle);
 								for(int32 ShapeIndex = 0; ShapeIndex < NumShapes; ++ShapeIndex)
