dont crash shadercookstats on deprecated shader platforms.

diff --git a/Engine/Source/Editor/StatsViewer/Private/StatsPages/ShaderCookerStatsPage.cpp b/Engine/Source/Editor/StatsViewer/Private/StatsPages/ShaderCookerStatsPage.cpp
--- a/Engine/Source/Editor/StatsViewer/Private/StatsPages/ShaderCookerStatsPage.cpp
+++ b/Engine/Source/Editor/StatsViewer/Private/StatsPages/ShaderCookerStatsPage.cpp
@@ -83,12 +83,15 @@ void FShaderCookerStats::Initialize(uint32 Index)
 	TArray<FString> PlatformNames;
 	for (int32 Platform = 0; Platform < SP_NumPlatforms; ++Platform)
 	{
-		FString FormatName = ShaderPlatformToShaderFormatName((EShaderPlatform)Platform).ToString();
-		if (FormatName.StartsWith(TEXT("SF_")))
+		if(!IsDeprecatedShaderPlatform((EShaderPlatform)Platform))
 		{
-			FormatName.MidInline(3, MAX_int32, false);
+			FString FormatName = ShaderPlatformToShaderFormatName((EShaderPlatform)Platform).ToString();
+			if (FormatName.StartsWith(TEXT("SF_")))
+			{
+				FormatName.MidInline(3, MAX_int32, false);
+			}
+			PlatformNames.Add(MoveTemp(FormatName));
 		}
-		PlatformNames.Add(MoveTemp(FormatName));
 	}
 	FShaderCookerStatsSet& Set = StatSets[Index];
 	FString CSVData;
diff --git a/Engine/Source/Runtime/Engine/Private/ShaderCompiler/ShaderCompiler.cpp b/Engine/Source/Runtime/Engine/Private/ShaderCompiler/ShaderCompiler.cpp
--- a/Engine/Source/Runtime/Engine/Private/ShaderCompiler/ShaderCompiler.cpp
+++ b/Engine/Source/Runtime/Engine/Private/ShaderCompiler/ShaderCompiler.cpp
@@ -4691,8 +4691,8 @@ void RecompileShadersForRemote(
 
 	UE_LOG(LogShaders, Display, TEXT("  Done!"))
 
-		// figure out which shaders are out of date
-		TArray<FShaderType*> OutdatedShaderTypes;
+	// figure out which shaders are out of date
+	TArray<FShaderType*> OutdatedShaderTypes;
 	TArray<const FVertexFactoryType*> OutdatedFactoryTypes;
 	TArray<const FShaderPipelineType*> OutdatedShaderPipelineTypes;
 
@@ -4715,11 +4715,8 @@ void RecompileShadersForRemote(
 			// Only compile for the desired platform if requested
 			if (ShaderPlatform == ShaderPlatformToCompile || ShaderPlatformToCompile == SP_NumPlatforms)
 			{
-				// These platforms are deprecated and we should warn about their use
-				if (ShaderPlatform == SP_OPENGL_SM5 || ShaderPlatform == SP_PCD3D_SM4_DEPRECATED || ShaderPlatform == SP_OPENGL_ES2_IOS_DEPRECATED ||
-					ShaderPlatform == SP_PCD3D_ES2 || ShaderPlatform == SP_METAL_MACES2 || ShaderPlatform == SP_OPENGL_PCES2 ||
-					ShaderPlatform == SP_OPENGL_ES2_ANDROID || ShaderPlatform == SP_OPENGL_ES2_WEBGL ||
-					ShaderPlatform == SP_VULKAN_SM4_DEPRECATED)
+				// Warn about outdated shader platforms.
+				if (IsDeprecatedShaderPlatform(ShaderPlatform))
 				{
 					UE_LOG(LogShaderCompilers, Warning, TEXT("You are compiling shaders for a deprecated platform '%s'"), *LegacyShaderPlatformToShaderFormat(ShaderPlatform).ToString());
 				}
diff --git a/Engine/Source/Runtime/RHI/Public/RHIDefinitions.h b/Engine/Source/Runtime/RHI/Public/RHIDefinitions.h
--- a/Engine/Source/Runtime/RHI/Public/RHIDefinitions.h
+++ b/Engine/Source/Runtime/RHI/Public/RHIDefinitions.h
@@ -1155,6 +1155,14 @@ inline bool IsHlslccShaderPlatform(const EShaderPlatform Platform)
 	return IsMetalPlatform(Platform) || IsVulkanPlatform(Platform) || IsSwitchPlatform(Platform) || IsOpenGLPlatform(Platform);
 }
 
+inline bool IsDeprecatedShaderPlatform(const EShaderPlatform ShaderPlatform)
+{
+	return ShaderPlatform == SP_OPENGL_SM5 || ShaderPlatform == SP_PCD3D_SM4_DEPRECATED || ShaderPlatform == SP_OPENGL_ES2_IOS_DEPRECATED ||
+		ShaderPlatform == SP_PCD3D_ES2 || ShaderPlatform == SP_METAL_MACES2 || ShaderPlatform == SP_OPENGL_PCES2 ||
+		ShaderPlatform == SP_OPENGL_ES2_ANDROID || ShaderPlatform == SP_OPENGL_ES2_WEBGL ||
+		ShaderPlatform == SP_VULKAN_SM4_DEPRECATED;
+}
+
 inline ERHIFeatureLevel::Type GetMaxSupportedFeatureLevel(EShaderPlatform InShaderPlatform)
 {
 	switch (InShaderPlatform)
