Store basic meta data of shader compile jobs on stack memory when a SCW crash is handled; should help analyzing crash dumps.

#rnx

diff --git a/Engine/Source/Runtime/Engine/Private/ShaderCompiler/ShaderCompiler.cpp b/Engine/Source/Runtime/Engine/Private/ShaderCompiler/ShaderCompiler.cpp
--- a/Engine/Source/Runtime/Engine/Private/ShaderCompiler/ShaderCompiler.cpp
+++ b/Engine/Source/Runtime/Engine/Private/ShaderCompiler/ShaderCompiler.cpp
@@ -630,6 +630,143 @@ static void ReadSingleJob(FShaderCompileJob* CurrentJob, FArchive& OutputFile)
 	}
 };
 
+// Disable optimization for this crash handler to get full access to the entire stack frame when debugging a crash dump
+PRAGMA_DISABLE_OPTIMIZATION;
+static void HandleWorkerCrash(const TArray<FShaderCommonCompileJob*>& QueuedJobs, FArchive& OutputFile, int32 OutputVersion, int64 FileSize, int32 ErrorCode, int32 NumProcessedJobs, int32 CallstackLength, int32 ExceptionInfoLength)
+{
+	TArray<TCHAR> Callstack;
+	Callstack.AddUninitialized(CallstackLength + 1);
+	OutputFile.Serialize(Callstack.GetData(), CallstackLength * sizeof(TCHAR));
+	Callstack[CallstackLength] = 0;
+
+	TArray<TCHAR> ExceptionInfo;
+	ExceptionInfo.AddUninitialized(ExceptionInfoLength + 1);
+	OutputFile.Serialize(ExceptionInfo.GetData(), ExceptionInfoLength * sizeof(TCHAR));
+	ExceptionInfo[ExceptionInfoLength] = 0;
+
+	// Store primary job information onto stack to make it part of a crash dump
+	static const int32 MaxNumCharsForSourcePaths = 8192;
+	int32 JobInputSourcePathsLength = 0;
+	ANSICHAR JobInputSourcePaths[MaxNumCharsForSourcePaths];
+	JobInputSourcePaths[0] = 0;
+
+	auto WriteInputSourcePathOntoStack = [&JobInputSourcePathsLength, &JobInputSourcePaths](const ANSICHAR* InputSourcePath)
+	{
+		if (InputSourcePath != nullptr && JobInputSourcePathsLength + 3 < MaxNumCharsForSourcePaths)
+		{
+			// Copy input source path into stack buffer
+			int32 InputSourcePathLength = FMath::Min(FCStringAnsi::Strlen(InputSourcePath), (MaxNumCharsForSourcePaths - JobInputSourcePathsLength - 2));
+			FMemory::Memcpy(JobInputSourcePaths + JobInputSourcePathsLength, InputSourcePath, InputSourcePathLength);
+
+			// Write newline character and put NUL character at the end
+			JobInputSourcePathsLength += InputSourcePathLength;
+			JobInputSourcePaths[JobInputSourcePathsLength] = TEXT('\n');
+			++JobInputSourcePathsLength;
+			JobInputSourcePaths[JobInputSourcePathsLength] = 0;
+		}
+	};
+
+	auto StoreInputDebugInfo = [&WriteInputSourcePathOntoStack, &JobInputSourcePathsLength, &JobInputSourcePaths](const FShaderCompilerInput& Input)
+	{
+		FString DebugInfo = FString::Printf(TEXT("%s:%s"), *Input.VirtualSourceFilePath, *Input.EntryPointName);
+		WriteInputSourcePathOntoStack(TCHAR_TO_UTF8(*DebugInfo));
+	};
+
+	for (FShaderCommonCompileJob* CommonJob : QueuedJobs)
+	{
+		if (FShaderCompileJob* SingleJob = CommonJob->GetSingleShaderJob())
+		{
+			StoreInputDebugInfo(SingleJob->Input);
+		}
+		else if (FShaderPipelineCompileJob* PipelineJob = CommonJob->GetShaderPipelineJob())
+		{
+			for (int32 Job = 0; Job < PipelineJob->StageJobs.Num(); ++Job)
+			{
+				if (FShaderCompileJob* SingleStageJob = PipelineJob->StageJobs[Job]->GetSingleShaderJob())
+				{
+					StoreInputDebugInfo(SingleStageJob->Input);
+				}
+			}
+		}
+	}
+
+	// One entry per error code as we want to have different callstacks for crash reporter...
+	switch (ErrorCode)
+	{
+	default:
+	case SCWErrorCode::GeneralCrash:
+	{
+		if (GDumpSCWJobInfoOnCrash != 0 || GIsBuildMachine)
+		{
+			auto DumpSingleJob = [](FShaderCompileJob* SingleJob) -> FString
+			{
+				if (!SingleJob)
+				{
+					return TEXT("Internal error, not a Job!");
+				}
+				FString String = SingleJob->Input.GenerateShaderName();
+				if (SingleJob->VFType)
+				{
+					String += FString::Printf(TEXT(" VF '%s'"), SingleJob->VFType->GetName());
+				}
+				String += FString::Printf(TEXT(" Type '%s'"), SingleJob->ShaderType->GetName());
+				String += FString::Printf(TEXT(" '%s' Entry '%s' Permutation %i "), *SingleJob->Input.VirtualSourceFilePath, *SingleJob->Input.EntryPointName, SingleJob->PermutationId);
+				return String;
+			};
+			UE_LOG(LogShaderCompilers, Error, TEXT("SCW %d Queued Jobs, Finished %d single jobs"), QueuedJobs.Num(), NumProcessedJobs);
+			for (int32 Index = 0; Index < QueuedJobs.Num(); ++Index)
+			{
+				FShaderCommonCompileJob* CommonJob = QueuedJobs[Index];
+				FShaderCompileJob* SingleJob = CommonJob->GetSingleShaderJob();
+				if (SingleJob)
+				{
+					UE_LOG(LogShaderCompilers, Error, TEXT("Job %d [Single] %s"), Index, *DumpSingleJob(SingleJob));
+				}
+				else
+				{
+					FShaderPipelineCompileJob* PipelineJob = CommonJob->GetShaderPipelineJob();
+					UE_LOG(LogShaderCompilers, Error, TEXT("Job %d: Pipeline %s "), Index, PipelineJob->ShaderPipeline->GetName());
+					for (int32 Job = 0; Job < PipelineJob->StageJobs.Num(); ++Job)
+					{
+						UE_LOG(LogShaderCompilers, Error, TEXT("PipelineJob %d %s"), Job, *DumpSingleJob(PipelineJob->StageJobs[Job]->GetSingleShaderJob()));
+					}
+				}
+			}
+		}
+		SCWErrorCode::HandleGeneralCrash(ExceptionInfo.GetData(), Callstack.GetData());
+	}
+	break;
+	case SCWErrorCode::BadShaderFormatVersion:
+		SCWErrorCode::HandleBadShaderFormatVersion(ExceptionInfo.GetData());
+		break;
+	case SCWErrorCode::BadInputVersion:
+		SCWErrorCode::HandleBadInputVersion(ExceptionInfo.GetData());
+		break;
+	case SCWErrorCode::BadSingleJobHeader:
+		SCWErrorCode::HandleBadSingleJobHeader(ExceptionInfo.GetData());
+		break;
+	case SCWErrorCode::BadPipelineJobHeader:
+		SCWErrorCode::HandleBadPipelineJobHeader(ExceptionInfo.GetData());
+		break;
+	case SCWErrorCode::CantDeleteInputFile:
+		SCWErrorCode::HandleCantDeleteInputFile(ExceptionInfo.GetData());
+		break;
+	case SCWErrorCode::CantSaveOutputFile:
+		SCWErrorCode::HandleCantSaveOutputFile(ExceptionInfo.GetData());
+		break;
+	case SCWErrorCode::NoTargetShaderFormatsFound:
+		SCWErrorCode::HandleNoTargetShaderFormatsFound(ExceptionInfo.GetData());
+		break;
+	case SCWErrorCode::CantCompileForSpecificFormat:
+		SCWErrorCode::HandleCantCompileForSpecificFormat(ExceptionInfo.GetData());
+		break;
+	case SCWErrorCode::Success:
+		// Can't get here...
+		break;
+	}
+}
+PRAGMA_ENABLE_OPTIMIZATION;
+
 // Process results from Worker Process
 void FShaderCompileUtilities::DoReadTaskResults(const TArray<FShaderCommonCompileJob*>& QueuedJobs, FArchive& OutputFile)
 {
@@ -666,90 +803,7 @@ void FShaderCompileUtilities::DoReadTaskResults(const TArray<FShaderCommonCompil
 	// Worker crashed
 	if (ErrorCode != SCWErrorCode::Success)
 	{
-		TArray<TCHAR> Callstack;
-		Callstack.AddUninitialized(CallstackLength + 1);
-		OutputFile.Serialize(Callstack.GetData(), CallstackLength * sizeof(TCHAR));
-		Callstack[CallstackLength] = 0;
-
-		TArray<TCHAR> ExceptionInfo;
-		ExceptionInfo.AddUninitialized(ExceptionInfoLength + 1);
-		OutputFile.Serialize(ExceptionInfo.GetData(), ExceptionInfoLength * sizeof(TCHAR));
-		ExceptionInfo[ExceptionInfoLength] = 0;
-
-		// One entry per error code as we want to have different callstacks for crash reporter...
-		switch (ErrorCode)
-		{
-		default:
-		case SCWErrorCode::GeneralCrash:
-			{
-				if (GDumpSCWJobInfoOnCrash != 0 || GIsBuildMachine)
-				{
-					auto DumpSingleJob = [](FShaderCompileJob* SingleJob) -> FString
-					{
-						if (!SingleJob)
-						{
-							return TEXT("Internal error, not a Job!");
-						}
-						FString String = SingleJob->Input.GenerateShaderName();
-						if (SingleJob->VFType)
-						{
-							String += FString::Printf(TEXT(" VF '%s'"), SingleJob->VFType->GetName());
-						}
-						String += FString::Printf(TEXT(" Type '%s'"), SingleJob->ShaderType->GetName());
-						String += FString::Printf(TEXT(" '%s' Entry '%s' Permutation %i "), *SingleJob->Input.VirtualSourceFilePath, *SingleJob->Input.EntryPointName, SingleJob->PermutationId);
-						return String;
-					};
-					UE_LOG(LogShaderCompilers, Error, TEXT("SCW %d Queued Jobs, Finished %d single jobs"), QueuedJobs.Num(), NumProcessedJobs);
-					for (int32 Index = 0; Index < QueuedJobs.Num(); ++Index)
-					{
-						FShaderCommonCompileJob* CommonJob = QueuedJobs[Index];
-						FShaderCompileJob* SingleJob = CommonJob->GetSingleShaderJob();
-						if (SingleJob)
-						{
-							UE_LOG(LogShaderCompilers, Error, TEXT("Job %d [Single] %s"), Index, *DumpSingleJob(SingleJob));
-						}
-						else
-						{
-							FShaderPipelineCompileJob* PipelineJob = CommonJob->GetShaderPipelineJob();
-							UE_LOG(LogShaderCompilers, Error, TEXT("Job %d: Pipeline %s "), Index, PipelineJob->ShaderPipeline->GetName());
-							for (int32 Job = 0; Job < PipelineJob->StageJobs.Num(); ++Job)
-							{
-								UE_LOG(LogShaderCompilers, Error, TEXT("PipelineJob %d %s"), Job, *DumpSingleJob(PipelineJob->StageJobs[Job]->GetSingleShaderJob()));
-							}
-						}
-					}
-				}
-				SCWErrorCode::HandleGeneralCrash(ExceptionInfo.GetData(), Callstack.GetData());
-			}
-			break;
-		case SCWErrorCode::BadShaderFormatVersion:
-			SCWErrorCode::HandleBadShaderFormatVersion(ExceptionInfo.GetData());
-			break;
-		case SCWErrorCode::BadInputVersion:
-			SCWErrorCode::HandleBadInputVersion(ExceptionInfo.GetData());
-			break;
-		case SCWErrorCode::BadSingleJobHeader:
-			SCWErrorCode::HandleBadSingleJobHeader(ExceptionInfo.GetData());
-			break;
-		case SCWErrorCode::BadPipelineJobHeader:
-			SCWErrorCode::HandleBadPipelineJobHeader(ExceptionInfo.GetData());
-			break;
-		case SCWErrorCode::CantDeleteInputFile:
-			SCWErrorCode::HandleCantDeleteInputFile(ExceptionInfo.GetData());
-			break;
-		case SCWErrorCode::CantSaveOutputFile:
-			SCWErrorCode::HandleCantSaveOutputFile(ExceptionInfo.GetData());
-			break;
-		case SCWErrorCode::NoTargetShaderFormatsFound:
-			SCWErrorCode::HandleNoTargetShaderFormatsFound(ExceptionInfo.GetData());
-			break;
-		case SCWErrorCode::CantCompileForSpecificFormat:
-			SCWErrorCode::HandleCantCompileForSpecificFormat(ExceptionInfo.GetData());
-			break;
-		case SCWErrorCode::Success:
-			// Can't get here...
-			break;
-		}
+		HandleWorkerCrash(QueuedJobs, OutputFile, OutputVersion, FileSize, ErrorCode, NumProcessedJobs, CallstackLength, ExceptionInfoLength);
 	}
 
 	TArray<FShaderCompileJob*> QueuedSingleJobs;
