Use SerializeInt to insert JitterClockTime in the to packet info instead of SerializeBits in order to be endian-safe.
Fill DummyPacketInfo with 0 since SerializeInt assumes the internal buffer is 0 and will skip bit values 0 during write.
Added non-shipping checks to ensure Write/Read JitterClock are the assumed bit count.

#rnx

diff --git a/Engine/Source/Runtime/Engine/Classes/Engine/NetConnection.h b/Engine/Source/Runtime/Engine/Classes/Engine/NetConnection.h
--- a/Engine/Source/Runtime/Engine/Classes/Engine/NetConnection.h
+++ b/Engine/Source/Runtime/Engine/Classes/Engine/NetConnection.h
@@ -1306,7 +1306,7 @@ private:
 	void ReceivedAck(int32 AckPacketId);
 
 	/** Calculate the average jitter while adding the new packet's jitter value */
-	void ProcessJitter(int32 PacketJitterClockTimeMS);
+	void ProcessJitter(uint32 PacketJitterClockTimeMS);
 
 	/**
 	 * on the server, the world the client has told us it has loaded
diff --git a/Engine/Source/Runtime/Engine/Private/NetConnection.cpp b/Engine/Source/Runtime/Engine/Private/NetConnection.cpp
--- a/Engine/Source/Runtime/Engine/Private/NetConnection.cpp
+++ b/Engine/Source/Runtime/Engine/Private/NetConnection.cpp
@@ -114,9 +114,8 @@ namespace UE4_NetConnectionPrivate
 		return Result;
 	}
 
-
 	/** Maximum possible clock time value with available bits. If that value is sent than the jitter clock time is ignored */
-	constexpr int32 MaxJitterClockTimeValue = (1 << NetConnectionHelper::NumBitsForJitterClockTimeInHeader) - 1;
+	constexpr uint32 MaxJitterClockTimeValue = (1 << NetConnectionHelper::NumBitsForJitterClockTimeInHeader) - 1;
 
 	/** Maximum possible precision of the jitter stat. Set to 1 second since we only send the milliseconds of the local clock time */
 	constexpr int32 MaxJitterPrecisionInMS = 1000;
@@ -1747,9 +1746,14 @@ void UNetConnection::WriteDummyPacketInfo(FBitWriter& Writer)
 	{
 		// Pre-insert the bits since the final time values will be calculated and inserted right before LowLevelSend
 		HeaderMarkForPacketInfo.Init(Writer);
-		int32 DummyJitterClockTime(UE4_NetConnectionPrivate::MaxJitterClockTimeValue);
+		uint32 DummyJitterClockTime(0);
 		Writer.SerializeBits(&DummyJitterClockTime, NetConnectionHelper::NumBitsForJitterClockTimeInHeader);
 
+#if !UE_BUILD_SHIPPING
+		checkf(Writer.GetNumBits() - BitsWrittenPrePacketInfo == (1 + NetConnectionHelper::NumBitsForJitterClockTimeInHeader), TEXT("WriteDummyPacketInfo did not write the expected nb of bits: Wrote %d, Expected %d"),
+			   Writer.GetNumBits() - BitsWrittenPrePacketInfo, (1 + NetConnectionHelper::NumBitsForJitterClockTimeInHeader));
+#endif
+
 		const uint8 bHasServerFrameTime = Driver->IsServer() ? bLastHasServerFrameTime : (CVarPingExcludeFrameTime.GetValueOnGameThread() > 0 ? 1u : 0u);
 		Writer.WriteBit(bHasServerFrameTime);
 
@@ -1782,11 +1786,15 @@ void UNetConnection::WriteFinalPacketInfo(FBitWriter& Writer, const double Packe
 	// Go back to write over the dummy bits
 	HeaderMarkForPacketInfo.PopWithoutClear(Writer);
 
+#if !UE_BUILD_SHIPPING
+	const int32 BitsWrittenPreJitterClock = Writer.GetNumBits();
+#endif
+
 	// Write Jitter clock time
 	{
 		const double DeltaSendTimeInMS = (PacketSentTimeInS - PreviousPacketSendTimeInS) * 1000.0;
 
-		int32 ClockTimeMilliseconds = 0;
+		uint32 ClockTimeMilliseconds = 0;
 
 		// If the delta is over our max precision, we send MAX value and jitter will be ignored by the receiver.
 		if (DeltaSendTimeInMS >= UE4_NetConnectionPrivate::MaxJitterPrecisionInMS)
@@ -1802,7 +1810,12 @@ void UNetConnection::WriteFinalPacketInfo(FBitWriter& Writer, const double Packe
 			ClockTimeMilliseconds &= UE4_NetConnectionPrivate::MaxJitterClockTimeValue;
 		}
 
-		Writer.SerializeBits(&ClockTimeMilliseconds, NetConnectionHelper::NumBitsForJitterClockTimeInHeader);
+		Writer.SerializeInt(ClockTimeMilliseconds, UE4_NetConnectionPrivate::MaxJitterClockTimeValue + 1);
+
+#if !UE_BUILD_SHIPPING
+		checkf(Writer.GetNumBits() - BitsWrittenPreJitterClock == NetConnectionHelper::NumBitsForJitterClockTimeInHeader, TEXT("WriteFinalPacketInfo did not write the expected nb of bits: Wrote %d, Expected %d"),
+			   Writer.GetNumBits() - BitsWrittenPreJitterClock, NetConnectionHelper::NumBitsForJitterClockTimeInHeader);
+#endif
 
 		PreviousPacketSendTimeInS = PacketSentTimeInS;
 	}
@@ -2024,9 +2037,17 @@ void UNetConnection::ReceivedPacket( FBitReader& Reader, bool bIsReinjectedPacke
 
 			if (bHasPacketInfoPayload)
 			{
+#if !UE_BUILD_SHIPPING
+				const int32 BitsReadPreJitterClock = Reader.GetPosBits();
+#endif
 				// Read jitter clock time from the packet header
-				int32 PacketJitterClockTimeMS = 0;
-				Reader.SerializeBits(&PacketJitterClockTimeMS, NetConnectionHelper::NumBitsForJitterClockTimeInHeader);
+				uint32 PacketJitterClockTimeMS = 0;
+				Reader.SerializeInt(PacketJitterClockTimeMS, UE4_NetConnectionPrivate::MaxJitterClockTimeValue + 1);
+
+#if !UE_BUILD_SHIPPING
+				checkf(Reader.GetPosBits() - BitsReadPreJitterClock == NetConnectionHelper::NumBitsForJitterClockTimeInHeader, TEXT("JitterClockTime did not read the expected nb of bits. Read %d, Expected %d"),
+					   Reader.GetPosBits() - BitsReadPreJitterClock, NetConnectionHelper::NumBitsForJitterClockTimeInHeader);
+#endif
 
 				if (!bIsReinjectedPacket)
 				{
@@ -3910,7 +3931,7 @@ void UNetConnection::TrackReplicationForAnalytics(const bool bWasSaturated)
 	SaturationAnalytics.TrackReplication(bWasSaturated);
 }
 
-void UNetConnection::ProcessJitter(int32 PacketJitterClockTimeMS)
+void UNetConnection::ProcessJitter(uint32 PacketJitterClockTimeMS)
 {
 	if (PacketJitterClockTimeMS >= UE4_NetConnectionPrivate::MaxJitterClockTimeValue)
 	{
@@ -3921,7 +3942,7 @@ void UNetConnection::ProcessJitter(int32 PacketJitterClockTimeMS)
 	const int32 CurrentClockTimeMilliseconds = UE4_NetConnectionPrivate::GetClockTimeMilliseconds(LastReceiveRealtime);
 
 	// Get the delta between the sent and receive clock time.
-	int32 CurrentJitterTimeDelta = CurrentClockTimeMilliseconds - PacketJitterClockTimeMS;
+	int32 CurrentJitterTimeDelta = CurrentClockTimeMilliseconds - (int32)PacketJitterClockTimeMS;
 	
 	if (CurrentJitterTimeDelta < 0)
 	{
