Prevent forcing low/medium quality usage for platforms that cannot lower the shader quality.

- A project can configure quality overrides for such a shader platform.
- Reduces the number of shadermaps we're compiling for PC which end up being identical.

diff --git a/Engine/Source/Runtime/Engine/Private/Materials/Material.cpp b/Engine/Source/Runtime/Engine/Private/Materials/Material.cpp
--- a/Engine/Source/Runtime/Engine/Private/Materials/Material.cpp
+++ b/Engine/Source/Runtime/Engine/Private/Materials/Material.cpp
@@ -3608,15 +3608,19 @@ void UMaterial::GetQualityLevelUsage(TArray<bool, TInlineAllocator<EMaterialQual
 
 	if (ShaderPlatform != SP_NumPlatforms)
 	{
-		// OR in the quality overrides if we're a valid shader platform
 		const UShaderPlatformQualitySettings* MaterialQualitySettings = UMaterialShaderQualitySettings::Get()->GetShaderPlatformQualitySettings(ShaderPlatform);
-		OutQualityLevelsUsed[EMaterialQualityLevel::Low] |= MaterialQualitySettings->GetQualityOverrides(EMaterialQualityLevel::Low).bEnableOverride;
-		OutQualityLevelsUsed[EMaterialQualityLevel::Medium] |= MaterialQualitySettings->GetQualityOverrides(EMaterialQualityLevel::Medium).bEnableOverride;
 
+		// OR in the quality overrides if possible on this shader platform, then
 		// AND in the quality allowances
 		bool bAnyQualityEnabled = false;
 		for (int32 Quality = 0; Quality < EMaterialQualityLevel::Num; ++Quality)
 		{
+			const FMaterialQualityOverrides& QualityOverrides = MaterialQualitySettings->GetQualityOverrides((EMaterialQualityLevel::Type)Quality);			
+			if (Quality != EMaterialQualityLevel::High && QualityOverrides.CanOverride(ShaderPlatform))
+			{
+				OutQualityLevelsUsed[Quality] |= QualityOverrides.bEnableOverride;
+			}
+
 			OutQualityLevelsUsed[Quality] &= !MaterialQualitySettings->GetQualityOverrides((EMaterialQualityLevel::Type)Quality).bDiscardQualityDuringCook;
 			bAnyQualityEnabled |= OutQualityLevelsUsed[Quality];
 		}
diff --git a/Engine/Source/Runtime/MaterialShaderQualitySettings/Classes/ShaderPlatformQualitySettings.h b/Engine/Source/Runtime/MaterialShaderQualitySettings/Classes/ShaderPlatformQualitySettings.h
--- a/Engine/Source/Runtime/MaterialShaderQualitySettings/Classes/ShaderPlatformQualitySettings.h
+++ b/Engine/Source/Runtime/MaterialShaderQualitySettings/Classes/ShaderPlatformQualitySettings.h
@@ -65,6 +65,7 @@ public:
 	UPROPERTY(EditAnywhere, Config, Meta = (DisplayName = "Cascade shadow mapping quality"), Category = "Quality")
 	EMobileCSMQuality MobileCSMQuality;
 
+	MATERIALSHADERQUALITYSETTINGS_API bool CanOverride(enum EShaderPlatform ShaderPlatform) const;
 	bool HasAnyOverridesSet() const;
 };
 
diff --git a/Engine/Source/Runtime/MaterialShaderQualitySettings/Private/MaterialShaderQualitySettings.cpp b/Engine/Source/Runtime/MaterialShaderQualitySettings/Private/MaterialShaderQualitySettings.cpp
--- a/Engine/Source/Runtime/MaterialShaderQualitySettings/Private/MaterialShaderQualitySettings.cpp
+++ b/Engine/Source/Runtime/MaterialShaderQualitySettings/Private/MaterialShaderQualitySettings.cpp
@@ -137,6 +137,13 @@ void UShaderPlatformQualitySettings::AppendToHashState(EMaterialQualityLevel::Ty
 
 //////////////////////////////////////////////////////////////////////////
 
+bool FMaterialQualityOverrides::CanOverride(EShaderPlatform ShaderPlatform) const
+{
+	// Only mobile renderer can lower the quality of a shader even without quality level nodes in the material (see TMobileBasePassPSPolicyParamType<>::ModifyCompilationEnvironmentForQualityLevel).
+	// Whitelist the platforms here that are going to use it.
+	return IsMobilePlatform(ShaderPlatform);
+}
+
 bool FMaterialQualityOverrides::HasAnyOverridesSet() const
 {
 	static const FMaterialQualityOverrides DefaultOverrides;
diff --git a/Engine/Source/Runtime/Renderer/Private/MobileBasePassRendering.cpp b/Engine/Source/Runtime/Renderer/Private/MobileBasePassRendering.cpp
--- a/Engine/Source/Runtime/Renderer/Private/MobileBasePassRendering.cpp
+++ b/Engine/Source/Runtime/Renderer/Private/MobileBasePassRendering.cpp
@@ -95,6 +95,8 @@ bool TMobileBasePassPSPolicyParamType<LightMapPolicyType>::ModifyCompilationEnvi
 	const UShaderPlatformQualitySettings* MaterialShadingQuality = UMaterialShaderQualitySettings::Get()->GetShaderPlatformQualitySettings(Platform);
 	const FMaterialQualityOverrides& QualityOverrides = MaterialShadingQuality->GetQualityOverrides(QualityLevel);
 
+	// the point of this check is to keep the logic between enabling overrides here and in UMaterial::GetQualityLevelUsage() in sync
+	checkf(QualityOverrides.CanOverride(Platform), TEXT("ShaderPlatform %d was not marked as being able to use quality overrides! Include it in CanOverride() and recook."), static_cast<int32>(Platform));
 	OutEnvironment.SetDefine(TEXT("MOBILE_QL_FORCE_FULLY_ROUGH"), QualityOverrides.bEnableOverride && QualityOverrides.bForceFullyRough != 0 ? 1u : 0u);
 	OutEnvironment.SetDefine(TEXT("MOBILE_QL_FORCE_NONMETAL"), QualityOverrides.bEnableOverride && QualityOverrides.bForceNonMetal != 0 ? 1u : 0u);
 	OutEnvironment.SetDefine(TEXT("QL_FORCEDISABLE_LM_DIRECTIONALITY"), QualityOverrides.bEnableOverride && QualityOverrides.bForceDisableLMDirectionality != 0 ? 1u : 0u);
