Fix ensures on VolumetricLightmap
#rnx

diff --git a/Engine/Source/Runtime/Engine/Private/PrecomputedVolumetricLightmap.cpp b/Engine/Source/Runtime/Engine/Private/PrecomputedVolumetricLightmap.cpp
--- a/Engine/Source/Runtime/Engine/Private/PrecomputedVolumetricLightmap.cpp
+++ b/Engine/Source/Runtime/Engine/Private/PrecomputedVolumetricLightmap.cpp
@@ -444,6 +444,9 @@ ENGINE_API void FPrecomputedVolumetricLightmapData::HandleDataMovementInAtlas(in
 			);
 
 			SceneData->IndirectionTexture = NewIndirectionTexture;
+
+ 			FRHIUnorderedAccessView* UAV = NewIndirectionTexture.UAV;
+ 			RHICmdList.TransitionResources(EResourceTransitionAccess::EReadable, EResourceTransitionPipeline::EComputeToGfx, &UAV, 1);
 		}
 	}
 	else
@@ -467,6 +470,9 @@ ENGINE_API void FPrecomputedVolumetricLightmapData::HandleDataMovementInAtlas(in
 				Parameters.SubLevelBrickPositions = SubLevelBrickPositionsSRV;
 
 				FComputeShaderUtils::Dispatch(RHICmdList, *ComputeShader, Parameters, FIntVector(FMath::DivideAndRoundUp(NumBricks, 64), 1, 1));
+
+				FRHIUnorderedAccessView* UAV = SceneData->IndirectionTexture.UAV;
+				RHICmdList.TransitionResources(EResourceTransitionAccess::EReadable, EResourceTransitionPipeline::EComputeToGfx, &UAV, 1);
 			}
 		}
 
@@ -538,6 +544,9 @@ ENGINE_API void FPrecomputedVolumetricLightmapData::AddToSceneData(FPrecomputedV
 				// Steal the indirection texture. When the sublevels are unloaded the values will be restored.
 				IndirectionTexture = SceneData->IndirectionTexture;
 			}
+
+			FRHIUnorderedAccessView* UAV = NewIndirectionTexture.UAV;
+			RHICmdList.TransitionResources(EResourceTransitionAccess::EReadable, EResourceTransitionPipeline::EComputeToGfx, &UAV, 1);
 		}
 		else
 		{
@@ -572,6 +581,9 @@ ENGINE_API void FPrecomputedVolumetricLightmapData::AddToSceneData(FPrecomputedV
 
 				FComputeShaderUtils::Dispatch(RHICmdList, *ComputeShader, Parameters, FIntVector(FMath::DivideAndRoundUp(NumBricks, 64), 1, 1));
 
+				FRHIUnorderedAccessView* UAV = SceneData->IndirectionTexture.UAV;
+				RHICmdList.TransitionResources(EResourceTransitionAccess::EReadable, EResourceTransitionPipeline::EComputeToGfx, &UAV, 1);
+
 				ReleaseRHIForSubLevelResources();
 			}
 		}
@@ -655,6 +667,9 @@ ENGINE_API void FPrecomputedVolumetricLightmapData::RemoveFromSceneData(FPrecomp
 
 				FComputeShaderUtils::Dispatch(RHICmdList, *ComputeShader, Parameters, FIntVector(FMath::DivideAndRoundUp(SubLevelBrickPositions.Num(), 64), 1, 1));
 
+				FRHIUnorderedAccessView* UAV = SceneData->IndirectionTexture.UAV;
+				RHICmdList.TransitionResources(EResourceTransitionAccess::EReadable, EResourceTransitionPipeline::EComputeToGfx, &UAV, 1);
+
 				ReleaseRHIForSubLevelResources();
 			}
 		}
@@ -909,6 +924,12 @@ void CopyDataIntoAtlas(FRHICommandList& RHICmdList, ERHIFeatureLevel::Type Featu
 		Parameters.OutDirectionalLightShadowing = DestTextureSet.DirectionalLightShadowing.UAV;
 
 		FComputeShaderUtils::Dispatch(RHICmdList, *ComputeShader, Parameters, FIntVector(NumBricks, 1, 1));
+
+		FRHIUnorderedAccessView* UAVs[3];
+		UAVs[0] = Parameters.OutAmbientVector;
+		UAVs[1] = Parameters.OutSkyBentNormal;
+		UAVs[2] = Parameters.OutDirectionalLightShadowing;
+		RHICmdList.TransitionResources(EResourceTransitionAccess::EReadable, EResourceTransitionPipeline::EComputeToGfx, UAVs, 3);
 	}
 
 	for (int32 i = 0; i < UE_ARRAY_COUNT(SrcData.SHCoefficients); i++)
@@ -926,6 +947,9 @@ void CopyDataIntoAtlas(FRHICommandList& RHICmdList, ERHIFeatureLevel::Type Featu
 		Parameters.OutSHCoefficients = DestTextureSet.SHCoefficients[i].UAV;
 
 		FComputeShaderUtils::Dispatch(RHICmdList, *ComputeShader, Parameters, FIntVector(NumBricks, 1, 1));
+
+		FRHIUnorderedAccessView* UAV = Parameters.OutSHCoefficients;
+		RHICmdList.TransitionResources(EResourceTransitionAccess::EReadable, EResourceTransitionPipeline::EComputeToGfx, &UAV, 1);
 	}
 }
 
