Chaos - Fix CIS - Don't combine mass properties in immediate physics if no shapes are present,

diff --git a/Engine/Source/Runtime/Engine/Private/PhysicsEngine/ImmediatePhysics/ImmediatePhysicsChaos/ImmediatePhysicsActorHandle_Chaos.cpp b/Engine/Source/Runtime/Engine/Private/PhysicsEngine/ImmediatePhysics/ImmediatePhysicsChaos/ImmediatePhysicsActorHandle_Chaos.cpp
--- a/Engine/Source/Runtime/Engine/Private/PhysicsEngine/ImmediatePhysics/ImmediatePhysicsChaos/ImmediatePhysicsActorHandle_Chaos.cpp
+++ b/Engine/Source/Runtime/Engine/Private/PhysicsEngine/ImmediatePhysics/ImmediatePhysicsChaos/ImmediatePhysicsActorHandle_Chaos.cpp
@@ -136,7 +136,19 @@ namespace ImmediatePhysics_Chaos
 			bModeThanOne = true;
 		}
 
-		OutMassProperties = Combine(AllMassProperties);
+		if (CHAOS_ENSURE(AllMassProperties.Num() > 0))
+		{
+			OutMassProperties = Combine(AllMassProperties);
+		}
+		else 
+		{
+			// @todo : Add support for all types, but for now just hard code a unit sphere tensor {r:50cm} if the type was not processed
+			OutMassProperties.CenterOfMass = FVec3(0.0f);
+			OutMassProperties.Mass = 523.f;
+			OutMassProperties.Volume = 5.24e5;
+			OutMassProperties.RotationOfMass = TRotation<FReal, 3>::FromIdentity();
+			OutMassProperties.InertiaTensor = PMatrix<FReal, 3, 3>(5.24e5, 5.24e5, 5.24e5);
+		}
 	}
 
 
