Fix crash on asserting audio thread when attempting to query a dynamic volume using command line in non-shipped build outside of editor.

diff --git a/Engine/Source/Runtime/Engine/Private/AudioDevice.cpp b/Engine/Source/Runtime/Engine/Private/AudioDevice.cpp
--- a/Engine/Source/Runtime/Engine/Private/AudioDevice.cpp
+++ b/Engine/Source/Runtime/Engine/Private/AudioDevice.cpp
@@ -1726,7 +1726,12 @@ bool FAudioDevice::HandleResetDynamicSoundVolumeCommand(const TCHAR* Cmd, FOutpu
 
 bool FAudioDevice::HandleGetDynamicSoundVolumeCommand(const TCHAR* Cmd, FOutputDevice& Ar)
 {
-	if (FAudioDeviceManager* DeviceManager = GEngine->GetAudioDeviceManager())
+	if (!GEngine)
+	{
+		return false;
+	}
+
+	if (const FAudioDeviceManager* DeviceManager = GEngine->GetAudioDeviceManager())
 	{
 		FName SoundName;
 		if (!FParse::Value(Cmd, TEXT("Name="), SoundName))
@@ -1749,9 +1754,31 @@ bool FAudioDevice::HandleGetDynamicSoundVolumeCommand(const TCHAR* Cmd, FOutputD
 			}
 		}
 
-		const float Volume = DeviceManager->GetDynamicSoundVolume(SoundType, SoundName);
-		FString Msg = FString::Printf(TEXT("'%s' Dynamic Volume: %.4f"), *SoundName.GetPlainNameString(), Volume);
-		Ar.Logf(TEXT("%s"), *Msg);
+		if (!IsInAudioThread())
+		{
+			DECLARE_CYCLE_STAT(TEXT("FAudioThreadTask.GetDynamicSoundVolume"), STAT_AudioGetDynamicSoundVolume, STATGROUP_AudioThreadCommands);
+
+			const ESoundType InSoundType = SoundType;
+			const FName InSoundName = SoundName;
+			FAudioThread::RunCommandOnAudioThread([InSoundType, InSoundName]()
+			{
+				if (!GEngine)
+				{
+					return;
+				}
+				if (const FAudioDeviceManager* InDeviceManager = GEngine->GetAudioDeviceManager())
+				{
+					const float Volume = InDeviceManager->GetDynamicSoundVolume(InSoundType, InSoundName);
+					UE_LOG(LogAudio, Display, TEXT("'%s' Dynamic Volume: %.4f"), *InSoundName.GetPlainNameString(), Volume);
+				}
+			}, GET_STATID(STAT_AudioGetDynamicSoundVolume));
+		}
+		else
+		{
+			const float Volume = DeviceManager->GetDynamicSoundVolume(SoundType, SoundName);
+			FString Msg = FString::Printf(TEXT("'%s' Dynamic Volume: %.4f"), *SoundName.GetPlainNameString(), Volume);
+			Ar.Logf(TEXT("%s"), *Msg);
+		}
 	}
 	return true;
 }
