Fixes mip-tail handling on cubemaps
#rnx

diff --git a/Engine/Source/Runtime/Engine/Private/TextureCube.cpp b/Engine/Source/Runtime/Engine/Private/TextureCube.cpp
--- a/Engine/Source/Runtime/Engine/Private/TextureCube.cpp
+++ b/Engine/Source/Runtime/Engine/Private/TextureCube.cpp
@@ -148,8 +148,8 @@ public:
 		check(Owner->GetNumMips() > 0);
 
 		TIndirectArray<FTexture2DMipMap>& Mips = InOwner->PlatformData->Mips;
-		const int32 FirstMipTailIndex = Mips.Num() - InOwner->PlatformData->GetNumMipsInTail();
-		for (int32 MipIndex = 0; MipIndex < FirstMipTailIndex; MipIndex++)
+		const int32 FirstMipTailIndex = Mips.Num() - FMath::Max(1, InOwner->PlatformData->GetNumMipsInTail());
+		for (int32 MipIndex = 0; MipIndex <= FirstMipTailIndex; MipIndex++)
 		{
 			FTexture2DMipMap& Mip = Mips[MipIndex];
 			if( Mip.BulkData.GetBulkDataSize() <= 0 )
