Use FDataDrivenPlatformInfoRegistry in RemovePlatformPrefixFromName instead of hard-coded list of platforms.
Added run-time sanity check of FVTUploadTileBuffer allocated memory size, before data is copied into it.
#rnx

diff --git a/Engine/Source/Runtime/Engine/Private/TextureDerivedData.cpp b/Engine/Source/Runtime/Engine/Private/TextureDerivedData.cpp
--- a/Engine/Source/Runtime/Engine/Private/TextureDerivedData.cpp
+++ b/Engine/Source/Runtime/Engine/Private/TextureDerivedData.cpp
@@ -51,7 +51,7 @@
 
 // This GUID is mixed into DDC version for virtual textures only, this allows updating DDC version for VT without invalidating DDC for all textures
 // This is useful during development, but once large numbers of VT are present in shipped content, it will have the same problem as TEXTURE_DERIVEDDATA_VER
-#define TEXTURE_VT_DERIVEDDATA_VER	TEXT("1F1EB5C2A6054EAEB31E38E85AAD1A33")
+#define TEXTURE_VT_DERIVEDDATA_VER	TEXT("AD48F7C9E1A64C3F8D8119A7582B4A34")
 
 #if ENABLE_COOK_STATS
 namespace TextureCookStats
diff --git a/Engine/Source/Runtime/Engine/Private/VT/VirtualTextureDataBuilder.cpp b/Engine/Source/Runtime/Engine/Private/VT/VirtualTextureDataBuilder.cpp
--- a/Engine/Source/Runtime/Engine/Private/VT/VirtualTextureDataBuilder.cpp
+++ b/Engine/Source/Runtime/Engine/Private/VT/VirtualTextureDataBuilder.cpp
@@ -9,6 +9,7 @@
 #include "CrunchCompression.h"
 #include "Async/TaskGraphInterfaces.h"
 #include "Async/ParallelFor.h"
+#include "Misc/DataDrivenPlatformInfoRegistry.h"
 #include "Misc/ScopedSlowTask.h"
 #include "TextureCompressorModule.h"
 #include "TextureDerivedDataTask.h"
@@ -785,17 +786,13 @@ int32 FVirtualTextureDataBuilder::FindSourceBlockIndex(int32 MipIndex, int32 Blo
 static const FName RemovePlatformPrefixFromName(FName const& InName)
 {
 	FString NameString = InName.ToString();
-	if (NameString.StartsWith(TEXT("PS4_")))
-	{
-		return *NameString.Replace(TEXT("PS4_"), TEXT(""));
-	}
-	else if (NameString.StartsWith("XBOXONE_"))
-	{
-		return *NameString.Replace(TEXT("XBOXONE_"), TEXT(""));
-	}
-	else if (NameString.StartsWith("SWITCH_"))
+
+	for (const auto& PlatformInfo : FDataDrivenPlatformInfoRegistry::GetAllPlatformInfos())
 	{
-		return *NameString.Replace(TEXT("SWITCH_"), TEXT(""));
+		if (NameString.StartsWith(PlatformInfo.Key, ESearchCase::IgnoreCase))
+		{
+			return *NameString.RightChop(PlatformInfo.Key.Len() + 1); // Platform name followed by `_`
+		}
 	}
 
 	return InName;
diff --git a/Engine/Source/Runtime/Engine/Private/VT/VirtualTextureTranscodeCache.cpp b/Engine/Source/Runtime/Engine/Private/VT/VirtualTextureTranscodeCache.cpp
--- a/Engine/Source/Runtime/Engine/Private/VT/VirtualTextureTranscodeCache.cpp
+++ b/Engine/Source/Runtime/Engine/Private/VT/VirtualTextureTranscodeCache.cpp
@@ -111,11 +111,14 @@ struct FTranscodeTask
 			case EVirtualTextureCodec::RawGPU:
 				if (StagingBufferForLayer.Stride == PackedStride)
 				{
+					check(TileLayerSize <= StagingBufferForLayer.MemorySize);
 					Params.Data->CopyTo(StagingBufferForLayer.Memory, DataOffset, TileLayerSize);
 				}
 				else
 				{
 					check(PackedStride <= StagingBufferForLayer.Stride);
+					check(TileLayerSize <= PackedOutputSize);
+					check(TileHeightInBlocks * StagingBufferForLayer.Stride <= StagingBufferForLayer.MemorySize);
 					TempBuffer.SetNumUninitialized(PackedOutputSize, false);
 					Params.Data->CopyTo(TempBuffer.GetData(), DataOffset, TileLayerSize);
 					for (uint32 y = 0; y < TileHeightInBlocks; ++y)
@@ -141,6 +144,7 @@ struct FTranscodeTask
 
 				check(Params.Codec);
 				const uint32 StagingBufferSize = StagingBufferForLayer.Stride * TileHeightInBlocks;
+				check(StagingBufferSize <= StagingBufferForLayer.MemorySize);
 				bResult = CrunchCompression::Decode(Params.Codec->Contexts[LayerIndex],
 					CompressedTile, TileLayerSize,
 					StagingBufferForLayer.Memory, StagingBufferSize, StagingBufferForLayer.Stride);
@@ -152,6 +156,7 @@ struct FTranscodeTask
 				if (StagingBufferForLayer.Stride == PackedStride)
 				{
 					// output buffer is tightly packed, can decompress directly
+					check(PackedOutputSize <= StagingBufferForLayer.MemorySize);
 					FCompression::UncompressMemoryStream(NAME_Zlib, StagingBufferForLayer.Memory, PackedOutputSize, Params.Data, DataOffset, TileLayerSize);
 				}
 				else
@@ -160,6 +165,7 @@ struct FTranscodeTask
 					check(PackedStride <= StagingBufferForLayer.Stride);
 					TempBuffer.SetNumUninitialized(PackedOutputSize, false);
 					FCompression::UncompressMemoryStream(NAME_Zlib, TempBuffer.GetData(), PackedOutputSize, Params.Data, DataOffset, TileLayerSize);
+					check(TileHeightInBlocks*StagingBufferForLayer.Stride <= StagingBufferForLayer.MemorySize);
 					for (uint32 y = 0; y < TileHeightInBlocks; ++y)
 					{
 						FMemory::Memcpy((uint8*)StagingBufferForLayer.Memory + y * StagingBufferForLayer.Stride, TempBuffer.GetData() + y * PackedStride, PackedStride);
diff --git a/Engine/Source/Runtime/Engine/Private/VT/VirtualTextureUploadCache.cpp b/Engine/Source/Runtime/Engine/Private/VT/VirtualTextureUploadCache.cpp
--- a/Engine/Source/Runtime/Engine/Private/VT/VirtualTextureUploadCache.cpp
+++ b/Engine/Source/Runtime/Engine/Private/VT/VirtualTextureUploadCache.cpp
@@ -208,6 +208,7 @@ FVTUploadTileHandle FVirtualTextureUploadCache::PrepareTileForUpload(FVTUploadTi
 		// Otherwise, CPU writes texture data to temp buffer, then this is copied to GPU via a batched staging texture...this involves more copying, but is best method under default D3D11
 		// Can potentially write each tile to a separate staging texture, but this has too much lock/unlock overhead
 		NewEntry.Stride = Stride;
+		NewEntry.MemorySize = MemorySize;
 		if (ALLOW_COPY_FROM_BUFFER)
 		{
 			FRHIResourceCreateInfo CreateInfo;
@@ -235,6 +236,7 @@ FVTUploadTileHandle FVirtualTextureUploadCache::PrepareTileForUpload(FVTUploadTi
 	++NumPendingTiles;
 	
 	OutBuffer.Memory = Entry.Memory;
+	OutBuffer.MemorySize = Entry.MemorySize;
 	OutBuffer.Stride = Entry.Stride;
 	return FVTUploadTileHandle(Index);
 }
diff --git a/Engine/Source/Runtime/Engine/Private/VT/VirtualTextureUploadCache.h b/Engine/Source/Runtime/Engine/Private/VT/VirtualTextureUploadCache.h
--- a/Engine/Source/Runtime/Engine/Private/VT/VirtualTextureUploadCache.h
+++ b/Engine/Source/Runtime/Engine/Private/VT/VirtualTextureUploadCache.h
@@ -13,6 +13,7 @@ class FRHIStructuredBuffer;
 struct FVTUploadTileBuffer
 {
 	void* Memory = nullptr;
+	uint32 MemorySize = 0u;
 	uint32 Stride = 0u;
 };
 
@@ -79,6 +80,7 @@ private:
 		TRefCountPtr<FRHIStructuredBuffer> RHIStagingBuffer;
 		FRHITexture2D* RHISubmitTexture = nullptr;
 		void* Memory = nullptr;
+		uint32 MemorySize = 0u;
 		uint32 Stride = 0u;
 		uint32 FrameSubmitted = 0u;
 		uint32 SubmitBatchIndex = 0u;
