Ensure dynamic meshes are using the ray tracing LOD that matches what was requested.

diff --git a/Engine/Source/Runtime/Engine/Private/SkeletalMesh.cpp b/Engine/Source/Runtime/Engine/Private/SkeletalMesh.cpp
--- a/Engine/Source/Runtime/Engine/Private/SkeletalMesh.cpp
+++ b/Engine/Source/Runtime/Engine/Private/SkeletalMesh.cpp
@@ -4792,6 +4792,13 @@ void FSkeletalMeshSceneProxy::GetDynamicRayTracingInstances(FRayTracingMaterialG
 				ensure(LODSections.Num() > 0);
 				const FLODSectionElements& LODSection = LODSections[LODIndex];
 				check(LODSection.SectionElements.Num() == LODData.RenderSections.Num());
+				
+				//#dxr_todo: verify why this condition is not fulfilled sometimes
+				verify(LODSection.SectionElements.Num() == MeshObject->GetRayTracingGeometry()->Initializer.Segments.Num());
+				if(LODSection.SectionElements.Num() != MeshObject->GetRayTracingGeometry()->Initializer.Segments.Num())
+				{
+					return;
+				}
 
 			#if WITH_EDITORONLY_DATA
 				int32 SectionIndexPreview = MeshObject->SectionIndexPreview;
diff --git a/Engine/Source/Runtime/Engine/Private/StaticMeshRender.cpp b/Engine/Source/Runtime/Engine/Private/StaticMeshRender.cpp
--- a/Engine/Source/Runtime/Engine/Private/StaticMeshRender.cpp
+++ b/Engine/Source/Runtime/Engine/Private/StaticMeshRender.cpp
@@ -1576,8 +1576,12 @@ void FStaticMeshSceneProxy::GetDynamicRayTracingInstances(FRayTracingMaterialGat
 		
 		RayTracingInstance.BuildInstanceMaskAndFlags();
 
-		check(RayTracingInstance.Geometry->Initializer.Segments.Num() == RayTracingInstance.Materials.Num());
-		OutRayTracingInstances.Add(RayTracingInstance);
+		//#dxr_todo: verify why this condition is not fulfilled sometimes
+		verify(RayTracingInstance.Geometry->Initializer.Segments.Num() == RayTracingInstance.Materials.Num());
+		if (RayTracingInstance.Geometry->Initializer.Segments.Num() == RayTracingInstance.Materials.Num())
+		{
+			OutRayTracingInstances.Add(RayTracingInstance);
+		}
 	}
 }
 #endif
