Fixed crash when trying to create a sublevel by pasting another level in the level's window

(on behalf of Yujiang Wang)

diff --git a/Engine/Source/Runtime/Engine/Private/PrecomputedVolumetricLightmap.cpp b/Engine/Source/Runtime/Engine/Private/PrecomputedVolumetricLightmap.cpp
--- a/Engine/Source/Runtime/Engine/Private/PrecomputedVolumetricLightmap.cpp
+++ b/Engine/Source/Runtime/Engine/Private/PrecomputedVolumetricLightmap.cpp
@@ -492,7 +492,10 @@ inline FIntVector ComputeBrickLayoutPosition(int32 BrickLayoutAllocation, FIntVe
 
 ENGINE_API void FPrecomputedVolumetricLightmapData::AddToSceneData(FPrecomputedVolumetricLightmapData* SceneData)
 {
-	check(SceneDataAdded.Find(SceneData) == INDEX_NONE);
+	if (SceneDataAdded.Find(SceneData) != INDEX_NONE)
+	{
+		return;
+	}
 
 	SceneDataAdded.Add(SceneData);
 
@@ -634,7 +637,10 @@ ENGINE_API void FPrecomputedVolumetricLightmapData::AddToSceneData(FPrecomputedV
 
 ENGINE_API void FPrecomputedVolumetricLightmapData::RemoveFromSceneData(FPrecomputedVolumetricLightmapData* SceneData, int32 PersistentLevelBrickDataBaseOffset)
 {
-	check(SceneDataAdded.Find(SceneData) != INDEX_NONE);
+	if (SceneDataAdded.Find(SceneData) == INDEX_NONE)
+	{
+		return;
+	}
 
 	SceneDataAdded.Remove(SceneData);
 
