- Move StopLowestPriority/StopLowestPriorityThenPreventNew to use Culling logic instead of Eviction due to it requiring evaluation of priority via a call to parse
- Fix Priority reporting in AudioDebug to show highest priority of waves managed for a particular active sound
- Remove AudioDisableConcurrencyStopSilentForLoopsCvar that was temporarily added for testing virtualization and is no longer needed

diff --git a/Engine/Source/Runtime/Engine/Classes/Sound/SoundConcurrency.h b/Engine/Source/Runtime/Engine/Classes/Sound/SoundConcurrency.h
--- a/Engine/Source/Runtime/Engine/Classes/Sound/SoundConcurrency.h
+++ b/Engine/Source/Runtime/Engine/Classes/Sound/SoundConcurrency.h
@@ -50,6 +50,8 @@ namespace EMaxConcurrentResolutionRule
 
 		/** Stop the lowest priority sound in the group. If all sounds are the same priority, then it won't play a new sound. */
 		StopLowestPriorityThenPreventNew,
+
+		Count UMETA(Hidden)
 	};
 }
 
@@ -122,6 +124,12 @@ public:
 	 * Retrieves the volume scale
 	 */
 	float GetVolumeScale() const;
+
+	/** Whether or not ResolutionRule supports eviction, wherein eviction is the ability to keep a sound
+	  * from playing prior to start and culling is the requirement of a sound to initialize and actively parse
+	  * prior to being removed from a concurrency group.
+	  */
+	bool IsEvictionSupported() const;
 };
 
 UCLASS(BlueprintType, hidecategories=Object, editinlinenew, MinimalAPI)
@@ -233,8 +241,8 @@ public:
 	/** Removes an active sound from the active sound array. */
 	void RemoveActiveSound(FActiveSound& ActiveSound);
 
-	/** Sorts the active sound array by volume */
-	void StopQuietSoundsDueToMaxConcurrency();
+	/** Sorts the active sound if concurrency settings require culling post playback */
+	void CullSoundsDueToMaxConcurrency();
 };
 
 typedef TMap<FConcurrencyGroupID, FConcurrencyGroup*> FConcurrencyGroups;
@@ -287,8 +295,8 @@ public:
 	/** Stops sound, applying concurrency rules for how to stop. */
 	void StopDueToVoiceStealing(FActiveSound& ActiveSound);
 
-	/** Stops any active sounds due to max concurrency quietest sound resolution rule */
-	void UpdateQuietSoundsToStop();
+	/** Culls any active sounds due to max concurrency sound resolution rule constraints being met */
+	void UpdateSoundsToCull();
 
 private: // Methods
 	/** Evaluates whether or not the sound can play given the concurrency group's rules. Appends permissible
@@ -313,7 +321,6 @@ private: // Methods
 
 	/** Helper functions for finding evictable sounds for the given concurrency rule */
 	FActiveSound* GetEvictableSoundStopFarthest(const FActiveSound& NewActiveSound, const FConcurrencyGroup& ConcurrencyGroup, bool bIsRetriggering) const;
-	FActiveSound* GetEvictableSoundStopLowestPriority(const FActiveSound& NewActiveSound, const FConcurrencyGroup& ConcurrencyGroup, bool bIsRetriggering) const;
 	FActiveSound* GetEvictableSoundStopOldest(const FActiveSound& NewActiveSound, const FConcurrencyGroup& ConcurrencyGroup, bool bIsRetriggering) const;
 
 	/** Handles concurrency evaluation that happens per USoundConcurrencyObject */
diff --git a/Engine/Source/Runtime/Engine/Private/ActiveSound.cpp b/Engine/Source/Runtime/Engine/Private/ActiveSound.cpp
--- a/Engine/Source/Runtime/Engine/Private/ActiveSound.cpp
+++ b/Engine/Source/Runtime/Engine/Private/ActiveSound.cpp
@@ -18,13 +18,6 @@ FAutoConsoleVariableRef CVarAudioOcclusionEnabled(
 	TEXT("Disables (1) or enables (0) audio occlusion.\n"),
 	ECVF_Default);
 
-static int32 AudioDisableConcurrencyStopSilentForLoopsCvar = 0;
-FAutoConsoleVariableRef CVarAudioDisableConcurrencyStopSilentForLoops(
-	TEXT("au.DisableConcurrencyStopSilentForLoops"),
-	AudioDisableConcurrencyStopSilentForLoopsCvar,
-	TEXT("Disables (1) or enables (0) audio concurrency for loops subscribing to stop silent.\n"),
-	ECVF_Default);
-
 FTraceDelegate FActiveSound::ActiveSoundTraceDelegate;
 TMap<FTraceHandle, FActiveSound::FAsyncTraceDetails> FActiveSound::TraceToActiveSoundMap;
 
@@ -101,7 +94,7 @@ FActiveSound::FActiveSound()
 	, EnvelopeFollowerReleaseTime(100)
 #if ENABLE_AUDIO_DEBUG
 	, DebugColor(FColor::Black)
-#endif // !(UE_BUILD_SHIPPING || UE_BUILD_TEST)
+#endif // ENABLE_AUDIO_DEBUG
 {
 	if (!ActiveSoundTraceDelegate.IsBound())
 	{
@@ -455,20 +448,6 @@ void FActiveSound::GetConcurrencyHandles(TArray<FConcurrencyHandle>& OutConcurre
 			}
 		}
 	}
-
-	if (IsLooping() && AudioDisableConcurrencyStopSilentForLoopsCvar)
-	{
-		for (int32 i = OutConcurrencyHandles.Num() - 1; i >= 0; --i)
-		{
-			if (OutConcurrencyHandles[i].Settings.ResolutionRule == EMaxConcurrentResolutionRule::StopQuietest)
-			{
-				UE_LOG(LogAudio, Warning,
-					TEXT("Concurrency not observed for '%s': StopQuietest concurrency disabled for looping sounds"),
-					Sound ? *Sound->GetName() : TEXT("N/A"));
-				OutConcurrencyHandles.RemoveAtSwap(i, 1, false);
-			}
-		}
-	}
 }
 
 bool FActiveSound::GetConcurrencyFadeDuration(float& OutFadeDuration) const
@@ -478,12 +457,12 @@ bool FActiveSound::GetConcurrencyFadeDuration(float& OutFadeDuration) const
 	GetConcurrencyHandles(Handles);
 	for (FConcurrencyHandle& Handle : Handles)
 	{
-		// StopQuietest can spam if a looping ActiveSound isn't active longer than a virtualization update period, which
+		// Resolution rules that don't support eviction (effectively requiring a sound to start before culling)
+		// can spam if a looping ActiveSound isn't active longer than a virtualization update period, which
 		// can happen when a concurrency group is maxed and constantly evicting.  If the voice steal fade time is particularly
 		// long, this can flood the active sound count. Therefore, only use the voice steal fade time if the sound has been
 		// active for a sufficient period of time.
-		const bool bStopQuietest = Handle.Settings.ResolutionRule == EMaxConcurrentResolutionRule::StopQuietest;
-		if (bStopQuietest && IsLooping() && FMath::IsNearlyZero(PlaybackTimeNonVirtualized, 0.1f))
+		if (!Handle.Settings.IsEvictionSupported() && IsLooping() && FMath::IsNearlyZero(PlaybackTimeNonVirtualized, 0.1f))
 		{
 			OutFadeDuration = 0.0f;
 			return false;
@@ -525,6 +504,8 @@ void FActiveSound::UpdateWaveInstances(TArray<FWaveInstance*> &InWaveInstances,
 		ClosestListenerIndex = AudioDevice->FindClosestListenerIndex(Transform);
 	}
 
+	FocusData.PriorityHighest = 1.0f;
+
 	FSoundParseParameters ParseParams;
 	ParseParams.Transform = Transform;
 	ParseParams.StartTime = RequestedStartTime;
@@ -639,7 +620,7 @@ void FActiveSound::UpdateWaveInstances(TArray<FWaveInstance*> &InWaveInstances,
 			}
 
 			// Remove concurrency volume scalars as this can cause ping-ponging to occur with virtualization and loops
-			// utilizing concurrency with StopQuietest.
+			// utilizing concurrency with rules that don't support eviction (removal from concurrency system prior to playback).
 			const float VolumeScale = GetTotalConcurrencyVolumeScale();
 			if (VolumeScale > SMALL_NUMBER)
 			{
@@ -1530,6 +1511,7 @@ void FActiveSound::UpdateAttenuation(float DeltaTime, FSoundParseParameters& Par
 
 	// Reset Focus data and recompute if necessary
 	FAttenuationFocusData FocusDataToApply;
+	FocusDataToApply.PriorityHighest = FocusData.PriorityHighest;
 
 	if (Settings->bEnableReverbSend)
 	{
@@ -1579,6 +1561,10 @@ void FActiveSound::UpdateAttenuation(float DeltaTime, FSoundParseParameters& Par
 		// in and out of focus values.
 		UpdateFocusData(DeltaTime, ListenerData, &FocusDataToApply);
 
+		// Update FocusData's highest priority copy prior to applying cached scalar immediately following
+		// to avoid applying scalar twice
+		FocusDataToApply.PriorityHighest = FMath::Max(FocusDataToApply.PriorityHighest, ParseParams.Priority);
+
 		ParseParams.Volume *= FocusDataToApply.VolumeScale;
 		ParseParams.Priority *= FocusDataToApply.PriorityScale;
 		if (Settings->bSpatialize)
@@ -1728,4 +1714,9 @@ void FActiveSound::UpdateAttenuation(float DeltaTime, FSoundParseParameters& Par
 	{
 		FocusData = FocusDataToApply;
 	}
+	// Make sure to always update highest priority
+	else
+	{
+		FocusData.PriorityHighest = FocusDataToApply.PriorityHighest;
+	}
 }
diff --git a/Engine/Source/Runtime/Engine/Private/AudioDebug.cpp b/Engine/Source/Runtime/Engine/Private/AudioDebug.cpp
--- a/Engine/Source/Runtime/Engine/Private/AudioDebug.cpp
+++ b/Engine/Source/Runtime/Engine/Private/AudioDebug.cpp
@@ -95,7 +95,7 @@ FAutoConsoleVariableRef CVarAudioDebugSoundSortCVar(
 	TEXT("au.Debug.Sounds.Sort"),
 	AudioDebugSoundSortCVarCVar,
 	TEXT("Value to sort by and display when sound stats are active. \n")
-	TEXT("Class, Distance, Name (Default), Priority, Time, Waves, Volume"),
+	TEXT("Class, Distance, Name (Default), Priority (Highest of wave instances per sound), Time, Waves, Volume"),
 	ECVF_Default);
 
 static FString AudioDebugStatSoundTextColorCVar = TEXT("White");
@@ -591,7 +591,7 @@ void FAudioDebugger::DrawDebugInfo(const FActiveSound& ActiveSound, const TArray
 	const float PlaybackTime = ActiveSound.PlaybackTime;
 	const float PlaybackTimeNonVirtualized = ActiveSound.PlaybackTimeNonVirtualized;
 
-	// StopQuietest sounds can start and immediately stop repeatedly when subscribed
+	// Sounds requiring culling can start and immediately stop repeatedly when subscribed
 	// concurrency is flooded, so don't show the initial frame.
 	if (FMath::IsNearlyZero(PlaybackTimeNonVirtualized))
 	{
@@ -1766,7 +1766,7 @@ void FAudioDebugger::SendUpdateResultsToGameThread(const FAudioDevice& AudioDevi
 					StatSoundInfo.SoundPath = SoundBase->GetPathName();
 					StatSoundInfo.Distance = AudioDevice.GetDistanceToNearestListener(ActiveSound->Transform.GetTranslation());
 					StatSoundInfo.PlaybackTime = ActiveSound->PlaybackTime;
-					StatSoundInfo.Priority = ActiveSound->GetPriority();
+					StatSoundInfo.Priority = ActiveSound->GetHighestPriority();
 					StatSoundInfo.PlaybackTimeNonVirtualized = ActiveSound->PlaybackTimeNonVirtualized;
 					
 					StatSoundInfo.Volume = 0.0f;
diff --git a/Engine/Source/Runtime/Engine/Private/AudioDevice.cpp b/Engine/Source/Runtime/Engine/Private/AudioDevice.cpp
--- a/Engine/Source/Runtime/Engine/Private/AudioDevice.cpp
+++ b/Engine/Source/Runtime/Engine/Private/AudioDevice.cpp
@@ -2230,12 +2230,12 @@ void FAudioDevice::RecursiveApplyAdjuster(const FSoundClassAdjuster& InAdjuster,
 	}
 }
 
-void FAudioDevice::StopQuietSoundsDueToMaxConcurrency(TArray<FWaveInstance*>& WaveInstances, TArray<FActiveSound *>& ActiveSoundsCopy)
+void FAudioDevice::CullSoundsDueToMaxConcurrency(TArray<FWaveInstance*>& WaveInstances, TArray<FActiveSound *>& ActiveSoundsCopy)
 {
 	// Now stop any sounds that are active that are in concurrency resolution groups that resolve by stopping quietest
 	{
 		SCOPE_CYCLE_COUNTER(STAT_AudioEvaluateConcurrency);
-		ConcurrencyManager.UpdateQuietSoundsToStop();
+		ConcurrencyManager.UpdateSoundsToCull();
 	}
 
 	for (int32 i = ActiveSoundsCopy.Num() - 1; i >= 0; --i)
@@ -3659,7 +3659,7 @@ int32 FAudioDevice::GetSortedActiveWaveInstances(TArray<FWaveInstance*>& WaveIns
 
 	if (GetType != ESortedActiveWaveGetType::QueryOnly)
 	{
-		StopQuietSoundsDueToMaxConcurrency(WaveInstances, ActiveSoundsCopy);
+		CullSoundsDueToMaxConcurrency(WaveInstances, ActiveSoundsCopy);
 	}
 
 	// Must be completed after StopQuietSoundsDueToMaxConcurrency as it avoids an issue
diff --git a/Engine/Source/Runtime/Engine/Private/SoundConcurrency.cpp b/Engine/Source/Runtime/Engine/Private/SoundConcurrency.cpp
--- a/Engine/Source/Runtime/Engine/Private/SoundConcurrency.cpp
+++ b/Engine/Source/Runtime/Engine/Private/SoundConcurrency.cpp
@@ -39,27 +39,40 @@ namespace
 } // namespace <>
 
 
-/************************************************************************/
-/* USoundConcurrency													*/
-/************************************************************************/
-
 USoundConcurrency::USoundConcurrency(const FObjectInitializer& ObjectInitializer)
 	: Super(ObjectInitializer)
 {
 }
 
-/************************************************************************/
-/* FSoundConcurrencySettings											*/
-/************************************************************************/
 float FSoundConcurrencySettings::GetVolumeScale() const
 {
 	return FMath::Clamp(VolumeScale, 0.0f, 1.0f);;
 }
 
+bool FSoundConcurrencySettings::IsEvictionSupported() const
+{
+	switch (ResolutionRule)
+	{
+		case EMaxConcurrentResolutionRule::StopLowestPriority:
+		case EMaxConcurrentResolutionRule::StopLowestPriorityThenPreventNew:
+		case EMaxConcurrentResolutionRule::StopQuietest:
+		{
+			return false;
+		}
+
+		case EMaxConcurrentResolutionRule::PreventNew:
+		case EMaxConcurrentResolutionRule::StopOldest:
+		case EMaxConcurrentResolutionRule::StopFarthestThenPreventNew:
+		case EMaxConcurrentResolutionRule::StopFarthestThenOldest:
+		default:
+		{
+			static_assert(static_cast<int32>(EMaxConcurrentResolutionRule::Count) == 7,
+				"Possible Missing EMaxConcurrentResolutionRule switch case coverage");
+			return true;
+		}
+	}
+}
 
-/************************************************************************/
-/* USoundConcurrency													*/
-/************************************************************************/
 
 FConcurrencyHandle::FConcurrencyHandle(const FSoundConcurrencySettings& InSettings)
 	: Settings(InSettings)
@@ -119,10 +132,6 @@ void FConcurrencySoundData::SetTarget(float InTargetVolume, float InLerpTime)
 	DbTargetVolume = Audio::ConvertToDecibels(InTargetVolume, KINDA_SMALL_NUMBER);
 }
 
-/************************************************************************/
-/* FConcurrencyGroup												*/
-/************************************************************************/
-
 FConcurrencyGroup::FConcurrencyGroup(FConcurrencyGroupID InGroupID, const FConcurrencyHandle& ConcurrencyHandle)
 	: GroupID(InGroupID)
 	, ObjectID(ConcurrencyHandle.ObjectID)
@@ -187,30 +196,74 @@ void FConcurrencyGroup::RemoveActiveSound(FActiveSound& ActiveSound)
 	}
 }
 
-void FConcurrencyGroup::StopQuietSoundsDueToMaxConcurrency()
+void FConcurrencyGroup::CullSoundsDueToMaxConcurrency()
 {
 	// Nothing to do if our active sound count is less than or equal to our max active sounds
-	if (Settings.ResolutionRule != EMaxConcurrentResolutionRule::StopQuietest || ActiveSounds.Num() <= Settings.MaxCount)
+	// or if eviction is supported
+	if (Settings.IsEvictionSupported() || ActiveSounds.Num() <= Settings.MaxCount)
 	{
 		return;
 	}
 
-	// Comparator for sorting group's ActiveSounds according to their "volume" concurrency. Quieter sounds will be
-	//  at the front of the array. If they share the same volume, newer sounds will be sorted first to avoid loop
-	// realization ping-ponging 
+	// Comparator for sorting group's ActiveSounds according to their dynamic concurrency nature.
+	// Quieter/lower priority sounds will be at the front of the array.
 	struct FCompareActiveSounds
 	{
+		EMaxConcurrentResolutionRule::Type ResolutionRule;
+
+		FCompareActiveSounds(EMaxConcurrentResolutionRule::Type InResolutionRule)
+			: ResolutionRule(InResolutionRule)
+		{
+		}
+
 		FORCEINLINE bool operator()(const FActiveSound& A, const FActiveSound& B) const
 		{
-			if (FMath::IsNearlyEqual(A.VolumeConcurrency, B.VolumeConcurrency, KINDA_SMALL_NUMBER))
+			switch (ResolutionRule)
 			{
-				return A.PlaybackTime > B.PlaybackTime;
+				case EMaxConcurrentResolutionRule::StopQuietest:
+				{
+					// If sounds share the same volume, newer sounds will be sorted first to avoid loop realization ping-ponging 
+					if (FMath::IsNearlyEqual(A.VolumeConcurrency, B.VolumeConcurrency, KINDA_SMALL_NUMBER))
+					{
+						return A.PlaybackTime > B.PlaybackTime;
+					}
+					return A.VolumeConcurrency < B.VolumeConcurrency;
+				}
+				break;
+
+				case EMaxConcurrentResolutionRule::StopLowestPriority:
+				case EMaxConcurrentResolutionRule::StopLowestPriorityThenPreventNew:
+				{
+					const int32 APriority = A.GetHighestPriority();
+					const int32 BPriority = B.GetHighestPriority();
+					if (!FMath::IsNearlyEqual(APriority, BPriority, KINDA_SMALL_NUMBER))
+					{
+						return APriority < BPriority;
+					}
+
+					// Newer sounds pushed forward in sort to make them more likely to be culled if PreventNew
+					return ResolutionRule == EMaxConcurrentResolutionRule::StopLowestPriorityThenPreventNew
+						? A.PlaybackTime < B.PlaybackTime
+						: A.PlaybackTime > B.PlaybackTime;
+				}
+				break;
+
+				default:
+				case EMaxConcurrentResolutionRule::PreventNew:
+				case EMaxConcurrentResolutionRule::StopOldest:
+				case EMaxConcurrentResolutionRule::StopFarthestThenPreventNew:
+				case EMaxConcurrentResolutionRule::StopFarthestThenOldest:
+				{
+					static_assert(static_cast<int32>(EMaxConcurrentResolutionRule::Count) == 7,
+						"Possible Missing EMaxConcurrentResolutionRule switch case coverage");
+					return true;
+				}
+				break;
 			}
-			return A.VolumeConcurrency < B.VolumeConcurrency;
 		}
 	};
 
-	ActiveSounds.Sort(FCompareActiveSounds());
+	ActiveSounds.Sort(FCompareActiveSounds(Settings.ResolutionRule));
 
 	const int32 NumActiveSounds = ActiveSounds.Num();
 	const int32 NumSoundsToStop = NumActiveSounds - Settings.MaxCount;
@@ -233,11 +286,6 @@ void FConcurrencyGroup::StopQuietSoundsDueToMaxConcurrency()
 
 }
 
-
-/************************************************************************/
-/* FSoundConcurrencyManager												*/
-/************************************************************************/
-
 FSoundConcurrencyManager::FSoundConcurrencyManager(class FAudioDevice* InAudioDevice)
 	: AudioDevice(InAudioDevice)
 {
@@ -429,15 +477,10 @@ FActiveSound* FSoundConcurrencyManager::GetEvictableSound(const FActiveSound& Ne
 		}
 		break;
 
+		// Eviction not supported by following rules due to requiring the sound to be initialized in order to calculate.
+		// Therefore, it is culled later and not evicted.
 		case EMaxConcurrentResolutionRule::StopLowestPriority:
 		case EMaxConcurrentResolutionRule::StopLowestPriorityThenPreventNew:
-		{
-			return GetEvictableSoundStopLowestPriority(NewActiveSound, ConcurrencyGroup, bIsRetriggering);
-		}
-		break;
-
-		// Eviction not supported by StopQuietest due to it requiring the sound to be initialized in order to calculate.
-		// Therefore, it is culled later but not evicted.
 		case EMaxConcurrentResolutionRule::StopQuietest:
 		{
 			return nullptr;
@@ -446,7 +489,8 @@ FActiveSound* FSoundConcurrencyManager::GetEvictableSound(const FActiveSound& Ne
 
 		default:
 		{
-			checkf(false, TEXT("Unknown EMaxConcurrentResolutionRule enumeration."));
+			static_assert(static_cast<int32>(EMaxConcurrentResolutionRule::Count) == 7,
+				"Possible Missing EMaxConcurrentResolutionRule switch case coverage");
 		}
 		break;
 	}
@@ -520,43 +564,6 @@ FActiveSound* FSoundConcurrencyManager::GetEvictableSoundStopFarthest(const FAct
 	return EvictableSound;
 }
 
-FActiveSound* FSoundConcurrencyManager::GetEvictableSoundStopLowestPriority(const FActiveSound& NewActiveSound, const FConcurrencyGroup& ConcurrencyGroup, bool bIsRetriggering) const
-{
-	// Find oldest and oldest lowest priority sound in the group
-	FActiveSound* EvictableSound = nullptr;
-	const TArray<FActiveSound*>& ActiveSounds = ConcurrencyGroup.GetActiveSounds();
-	for (FActiveSound* ActiveSound : ActiveSounds)
-	{
-		check(ActiveSound);
-
-		if (EvictableSound == nullptr
-			|| (ActiveSound->GetPriority() < EvictableSound->GetPriority())
-			|| (ActiveSound->GetPriority() == EvictableSound->GetPriority() && ActiveSound->PlaybackTime > EvictableSound->PlaybackTime))
-		{
-			EvictableSound = ActiveSound;
-		}
-	}
-
-	if (EvictableSound)
-	{
-		// Drop request as same priority and preventing new
-		const EMaxConcurrentResolutionRule::Type Rule = ConcurrencyGroup.GetSettings().ResolutionRule;
-		if (Rule == EMaxConcurrentResolutionRule::StopLowestPriorityThenPreventNew
-			&& EvictableSound->GetPriority() == NewActiveSound.GetPriority())
-		{
-			return nullptr;
-		}
-
-		// Drop request as NewActiveSound's priority is lower than the lowest priority sound playing
-		else if (EvictableSound->GetPriority() > NewActiveSound.GetPriority())
-		{
-			return nullptr;
-		}
-	}
-
-	return EvictableSound;
-}
-
 FActiveSound* FSoundConcurrencyManager::EvaluateConcurrency(const FActiveSound& NewActiveSound, TArray<FConcurrencyHandle>& ConcurrencyHandles, bool bIsRetriggering)
 {
 	check(NewActiveSound.GetSound());
@@ -830,12 +837,12 @@ void FSoundConcurrencyManager::StopDueToVoiceStealing(FActiveSound& ActiveSound)
 	}
 }
 
-void FSoundConcurrencyManager::UpdateQuietSoundsToStop()
+void FSoundConcurrencyManager::UpdateSoundsToCull()
 {
 	check(IsInAudioThread());
 
 	for (auto& ConcurrenyGroupEntry : ConcurrencyGroups)
 	{
-		ConcurrenyGroupEntry.Value->StopQuietSoundsDueToMaxConcurrency();
+		ConcurrenyGroupEntry.Value->CullSoundsDueToMaxConcurrency();
 	}
 }
\ No newline at end of file
diff --git a/Engine/Source/Runtime/Engine/Public/ActiveSound.h b/Engine/Source/Runtime/Engine/Public/ActiveSound.h
--- a/Engine/Source/Runtime/Engine/Public/ActiveSound.h
+++ b/Engine/Source/Runtime/Engine/Public/ActiveSound.h
@@ -47,6 +47,9 @@ struct FAttenuationFocusData
 	/** The amount priority is scaled due to focus */
 	float PriorityScale;
 
+	/** Cached highest priority of the parent active sound's wave instances. */
+	float PriorityHighest;
+
 	/** The amount volume is scaled due to focus */
 	float VolumeScale;
 
@@ -56,6 +59,7 @@ struct FAttenuationFocusData
 		, FocusFactor(1.0f)
 		, DistanceScale(1.0f)
 		, PriorityScale(1.0f)
+		, PriorityHighest(1.0f)
 		, VolumeScale(1.0f)
 	{
 	}
@@ -657,8 +661,8 @@ public:
 	/** Applies the active sound's attenuation settings to the input parse params using the given listener */
 	void ParseAttenuation(FSoundParseParameters& OutParseParams, int32 ListenerIndex, const FSoundAttenuationSettings& InAttenuationSettings);
 
-	/** Returns the effective priority of the active sound */
-	float GetPriority() const { return Priority * FocusData.PriorityScale; }
+	/** Returns the highest effective priority of the child wave instances */
+	float GetHighestPriority() const { return Priority * FocusData.PriorityHighest * FocusData.PriorityScale; }
 
 	/** Sets the amount of audio from this active sound to send to the submix. */
 	void SetSubmixSend(const FSoundSubmixSendInfo& SubmixSendInfo);
