Missing file

diff --git a/Engine/Source/Runtime/Engine/Public/AudioDevice.h b/Engine/Source/Runtime/Engine/Public/AudioDevice.h
--- a/Engine/Source/Runtime/Engine/Public/AudioDevice.h
+++ b/Engine/Source/Runtime/Engine/Public/AudioDevice.h
@@ -1335,9 +1335,9 @@ private:
 	 */
 	void ParseSoundClasses(float InDeltaTime);
 
-	/** Stops quiet sounds due to being evaluated as not fulfilling concurrency requirements
+	/** Stops quiet/low priority sounds due to being evaluated as not fulfilling concurrency requirements
 	 */
-	void StopQuietSoundsDueToMaxConcurrency(TArray<FWaveInstance*>& WaveInstances, TArray<FActiveSound*>& ActiveSoundsCopy);
+	void CullSoundsDueToMaxConcurrency(TArray<FWaveInstance*>& WaveInstances, TArray<FActiveSound*>& ActiveSoundsCopy);
 
 	/**
 	 * Checks if the given sound would be audible.
