Initial version of physics replication for Chaos.

diff --git a/Engine/Source/Runtime/Engine/Private/PhysicsEngine/Experimental/PhysScene_Chaos.cpp b/Engine/Source/Runtime/Engine/Private/PhysicsEngine/Experimental/PhysScene_Chaos.cpp
--- a/Engine/Source/Runtime/Engine/Private/PhysicsEngine/Experimental/PhysScene_Chaos.cpp
+++ b/Engine/Source/Runtime/Engine/Private/PhysicsEngine/Experimental/PhysScene_Chaos.cpp
@@ -35,7 +35,7 @@
 #include "PBDRigidActiveParticlesBuffer.h"
 #include "Chaos/GeometryParticlesfwd.h"
 #include "Chaos/Box.h"
-
+#include "PhysicsReplication.h"
 
 #if !UE_BUILD_SHIPPING
 #include "Engine/World.h"
@@ -401,6 +401,17 @@ FPhysScene_Chaos::FPhysScene_Chaos(AActor* InSolverActor
 
 FPhysScene_Chaos::~FPhysScene_Chaos()
 {
+#if WITH_CHAOS
+	if (IPhysicsReplicationFactory* RawReplicationFactory = FPhysScene_ChaosInterface::PhysicsReplicationFactory.Get())
+	{
+		RawReplicationFactory->Destroy(PhysicsReplication);
+	}
+	else
+	{
+		delete PhysicsReplication;
+	}
+#endif
+
 	Shutdown();
 	
 	FCoreDelegates::OnPreExit.RemoveAll(this);
@@ -813,6 +824,16 @@ void FPhysScene_Chaos::Shutdown()
 	ComponentToPhysicsProxyMap.Reset();
 }
 
+FPhysicsReplication* FPhysScene_Chaos::GetPhysicsReplication()
+{
+	return PhysicsReplication;
+}
+
+void FPhysScene_Chaos::SetPhysicsReplication(FPhysicsReplication* InPhysicsReplication)
+{
+	PhysicsReplication = InPhysicsReplication;
+}
+
 void FPhysScene_Chaos::AddReferencedObjects(FReferenceCollector& Collector)
 {
 #if WITH_EDITOR
@@ -958,6 +979,10 @@ FPhysScene_ChaosInterface::FPhysScene_ChaosInterface(const AWorldSettings* InSet
 	Scene.GetSolver()->GetEvolution()->GetParticles().AddArray(&BodyInstances);
 #endif
 
+	// Create replication manager
+	FPhysicsReplication* PhysicsReplication = PhysicsReplicationFactory.IsValid() ? PhysicsReplicationFactory->Create(this) : new FPhysicsReplication(this);
+	Scene.SetPhysicsReplication(PhysicsReplication);
+
 	Scene.GetSolver()->PhysSceneHack = this;
 	
 
@@ -1135,7 +1160,7 @@ void FPhysScene_ChaosInterface::Flush_AssumesLocked()
 
 FPhysicsReplication* FPhysScene_ChaosInterface::GetPhysicsReplication()
 {
-	return nullptr;
+	return Scene.GetPhysicsReplication();
 }
 
 void FPhysScene_ChaosInterface::RemoveBodyInstanceFromPendingLists_AssumesLocked(FBodyInstance* BodyInstance, int32 SceneType)
@@ -1431,6 +1456,11 @@ void FPhysScene_ChaosInterface::StartFrame()
 			break;
 		}
 	}
+
+	if (FPhysicsReplication* PhysicsReplication = Scene.GetPhysicsReplication())
+	{
+		PhysicsReplication->Tick(Dt);
+	}
 }
 
 void FPhysScene_ChaosInterface::EndFrame(ULineBatchComponent* InLineBatcher)
diff --git a/Engine/Source/Runtime/Engine/Public/Physics/Experimental/PhysScene_Chaos.h b/Engine/Source/Runtime/Engine/Public/Physics/Experimental/PhysScene_Chaos.h
--- a/Engine/Source/Runtime/Engine/Public/Physics/Experimental/PhysScene_Chaos.h
+++ b/Engine/Source/Runtime/Engine/Public/Physics/Experimental/PhysScene_Chaos.h
@@ -27,6 +27,7 @@
 class UPrimitiveComponent;
 
 class AdvanceOneTimeStepTask;
+class FPhysicsReplication;
 class FPhysInterface_Chaos;
 class FChaosSolversModule;
 struct FForceFieldProxy;
@@ -140,6 +141,9 @@ public:
 
 	void Shutdown();
 
+	FPhysicsReplication* GetPhysicsReplication();
+	void SetPhysicsReplication(FPhysicsReplication* InPhysicsReplication);
+
 #if WITH_EDITOR
 	void AddPieModifiedObject(UObject* InObj);
 #endif
@@ -175,6 +179,9 @@ public:
 private:
 	TUniquePtr<Chaos::ISpatialAccelerationCollection<Chaos::TAccelerationStructureHandle<float, 3>, float, 3>> SolverAccelerationStructure;
 
+	/** Replication manager that updates physics bodies towards replicated physics state */
+	FPhysicsReplication* PhysicsReplication;
+
 #if CHAOS_WITH_PAUSABLE_SOLVER
 	/** Callback that checks the status of the world settings for this scene before pausing/unpausing its solver. */
 	void OnUpdateWorldPause();
@@ -266,7 +273,6 @@ void FPhysScene_Chaos::RegisterEventHandler(const Chaos::EEventType& EventID, Ha
 
 class UWorld;
 class AWorldSettings;
-class FPhysicsReplication;
 class FPhysicsReplicationFactory;
 class FContactModifyCallbackFactory;
 
