Fix for memory leak that triggered check(!TotalBytesRequested) in FD3D12SegListAllocator::Destroy(). FVirtualTextureChunkStreamingManager was never destroyed so the resources it held thus never released.

diff --git a/Engine/Source/Runtime/Engine/Private/ContentStreaming.cpp b/Engine/Source/Runtime/Engine/Private/ContentStreaming.cpp
--- a/Engine/Source/Runtime/Engine/Private/ContentStreaming.cpp
+++ b/Engine/Source/Runtime/Engine/Private/ContentStreaming.cpp
@@ -744,7 +744,8 @@ FStreamingManagerCollection::FStreamingManagerCollection()
 :	NumIterations(1)
 ,	DisableResourceStreamingCount(0)
 ,	LoadMapTimeLimit( 5.0f )
-,   TextureStreamingManager( NULL )
+,   TextureStreamingManager(nullptr)
+, VirtualTextureChunkStreamingManager(nullptr)
 {
 #if PLATFORM_SUPPORTS_TEXTURE_STREAMING
 	// Disable texture streaming if that was requested (needs to happen before the call to ProcessNewlyLoadedUObjects, as that can load textures)
@@ -772,6 +773,27 @@ FStreamingManagerCollection::FStreamingManagerCollection()
 	AddStreamingManager(AnimationStreamingManager);
 }
 
+FStreamingManagerCollection::~FStreamingManagerCollection()
+{
+	RemoveStreamingManager(AnimationStreamingManager);
+	delete AnimationStreamingManager;
+	AnimationStreamingManager = nullptr;
+
+	RemoveStreamingManager(AudioStreamingManager);
+	delete AudioStreamingManager;
+	AudioStreamingManager = nullptr;
+
+	RemoveStreamingManager(TextureStreamingManager);
+	delete TextureStreamingManager;
+	TextureStreamingManager = nullptr;
+
+	RemoveStreamingManager(VirtualTextureChunkStreamingManager);
+	delete VirtualTextureChunkStreamingManager;
+	VirtualTextureChunkStreamingManager = nullptr;
+
+	check(StreamingManagers.Num() == 0);
+}
+
 /**
  * Sets the number of iterations to use for the next time UpdateResourceStreaming is being called. This 
  * is being reset to 1 afterwards.
diff --git a/Engine/Source/Runtime/Engine/Private/VT/VirtualTextureChunkManager.cpp b/Engine/Source/Runtime/Engine/Private/VT/VirtualTextureChunkManager.cpp
--- a/Engine/Source/Runtime/Engine/Private/VT/VirtualTextureChunkManager.cpp
+++ b/Engine/Source/Runtime/Engine/Private/VT/VirtualTextureChunkManager.cpp
@@ -35,7 +35,7 @@ struct FVirtualTextureChunkStreamingManager& FVirtualTextureChunkStreamingManage
 
 FVirtualTextureChunkStreamingManager::FVirtualTextureChunkStreamingManager()
 {
-	IStreamingManager::Get().AddStreamingManager(this);
+	IStreamingManager::Get().SetVirtualTextureChunkStreamingManager(this);
 #if WITH_EDITOR
 	GetVirtualTextureChunkDDCCache()->Initialize();
 #endif
@@ -46,7 +46,7 @@ FVirtualTextureChunkStreamingManager::~FVirtualTextureChunkStreamingManager()
 #if WITH_EDITOR
 	GetVirtualTextureChunkDDCCache()->ShutDown();
 #endif
-	IStreamingManager::Get().RemoveStreamingManager(this);
+	IStreamingManager::Get().SetVirtualTextureChunkStreamingManager(nullptr);
 }
 
 void FVirtualTextureChunkStreamingManager::UpdateResourceStreaming(float DeltaTime, bool bProcessEverything /*= false*/)
diff --git a/Engine/Source/Runtime/Engine/Public/ContentStreaming.h b/Engine/Source/Runtime/Engine/Public/ContentStreaming.h
--- a/Engine/Source/Runtime/Engine/Public/ContentStreaming.h
+++ b/Engine/Source/Runtime/Engine/Public/ContentStreaming.h
@@ -586,6 +586,7 @@ struct ENGINE_VTABLE FStreamingManagerCollection : public IStreamingManager
 {
 	/** Default constructor, initializing all member variables. */
 	ENGINE_API FStreamingManagerCollection();
+	ENGINE_API ~FStreamingManagerCollection();
 
 	/**
 	 * Calls UpdateResourceStreaming(), and does per-frame cleaning. Call once per frame.
@@ -750,6 +751,11 @@ struct ENGINE_VTABLE FStreamingManagerCollection : public IStreamingManager
 	virtual void OnAudioStreamingParamsChanged() override;
 #endif
 
+	void SetVirtualTextureChunkStreamingManager(IStreamingManager* InVirtualTextureChunkStreamingManager)
+	{
+		VirtualTextureChunkStreamingManager = InVirtualTextureChunkStreamingManager;
+	}
+
 protected:
 
 	virtual void AddOrRemoveTextureStreamingManagerIfNeeded(bool bIsInit=false);
@@ -774,6 +780,8 @@ protected:
 	/** The animation streaming manager, should always exist */
 	IAnimationStreamingManager* AnimationStreamingManager;
 
+	IStreamingManager* VirtualTextureChunkStreamingManager;
+
 #if WITH_EDITOR
 	// Locks out any audio streaming manager call when we are re-initializing the audio streaming manager.
 	mutable FCriticalSection AudioStreamingManagerCriticalSection;
