Change check() to UE_CLOG() since it's not a fatal condition by any means

diff --git a/Engine/Source/Runtime/Engine/Private/ContentStreaming.cpp b/Engine/Source/Runtime/Engine/Private/ContentStreaming.cpp
--- a/Engine/Source/Runtime/Engine/Private/ContentStreaming.cpp
+++ b/Engine/Source/Runtime/Engine/Private/ContentStreaming.cpp
@@ -791,7 +791,7 @@ FStreamingManagerCollection::~FStreamingManagerCollection()
 	delete VirtualTextureChunkStreamingManager;
 	VirtualTextureChunkStreamingManager = nullptr;
 
-	check(StreamingManagers.Num() == 0);
+	UE_CLOG(StreamingManagers.Num() > 0, LogContentStreaming, Display, TEXT("There are %d unreleased StreamingManagers"), StreamingManagers.Num());
 }
 
 /**
