Fix Cook Crash in PhysScene_Chaos destructoin

diff --git a/Engine/Source/Runtime/Engine/Private/PhysicsEngine/Experimental/PhysScene_Chaos.cpp b/Engine/Source/Runtime/Engine/Private/PhysicsEngine/Experimental/PhysScene_Chaos.cpp
--- a/Engine/Source/Runtime/Engine/Private/PhysicsEngine/Experimental/PhysScene_Chaos.cpp
+++ b/Engine/Source/Runtime/Engine/Private/PhysicsEngine/Experimental/PhysScene_Chaos.cpp
@@ -342,7 +342,8 @@ FPhysScene_Chaos::FPhysScene_Chaos(AActor* InSolverActor
 	, const FName& DebugName
 #endif
 )
-	: ChaosModule(nullptr)
+	: PhysicsReplication(nullptr)
+	, ChaosModule(nullptr)
 	, SceneSolver(nullptr)
 	, SolverActor(InSolverActor)
 #if WITH_EDITOR
@@ -406,7 +407,7 @@ FPhysScene_Chaos::~FPhysScene_Chaos()
 	{
 		RawReplicationFactory->Destroy(PhysicsReplication);
 	}
-	else
+	else if(PhysicsReplication)
 	{
 		delete PhysicsReplication;
 	}
