Do not create another shadermap for the same Id if a cached one exists, even if incomplete.

- Without this, material instances being cooked at the same time can result in duplicate shadermaps (and unnecessary shader compilation).

diff --git a/Engine/Source/Runtime/Engine/Private/Materials/MaterialShared.cpp b/Engine/Source/Runtime/Engine/Private/Materials/MaterialShared.cpp
--- a/Engine/Source/Runtime/Engine/Private/Materials/MaterialShared.cpp
+++ b/Engine/Source/Runtime/Engine/Private/Materials/MaterialShared.cpp
@@ -1809,8 +1809,9 @@ bool FMaterial::CacheShaders(const FMaterialShaderMapId& ShaderMapId, const FSta
 		// See FMaterialShaderMap::ProcessCompilationResults().
 		if  (GetMaterialShaderMapUsage() != EMaterialShaderMapUsage::DebugViewMode)
 		{
-			// Attempt to load from the derived data cache if we are uncooked
-			if ((!ShaderMap || !ShaderMap->IsComplete(this, true)) && !FPlatformProperties::RequiresCookedData())
+			// Attempt to load from the derived data cache if we are uncooked and don't have any shadermap.
+			// If we have an incomplete shadermap, continue with it to prevent creation of duplicate shadermaps for the same ShaderMapId
+			if (!ShaderMap && !FPlatformProperties::RequiresCookedData()) 
 			{
 				TRefCountPtr<FMaterialShaderMap> LoadedShaderMap;
 				FMaterialShaderMap::LoadFromDerivedDataCache(this, ShaderMapId, Platform, LoadedShaderMap);
