GameInstance - Adding events when new local players are added or removed from the game instance.

diff --git a/Engine/Source/Runtime/Engine/Classes/Engine/GameInstance.h b/Engine/Source/Runtime/Engine/Classes/Engine/GameInstance.h
--- a/Engine/Source/Runtime/Engine/Classes/Engine/GameInstance.h
+++ b/Engine/Source/Runtime/Engine/Classes/Engine/GameInstance.h
@@ -108,6 +108,8 @@ struct FGameInstancePIEParameters
 
 #endif
 
+DECLARE_EVENT_OneParam(UGameInstance, FOnLocalPlayerEvent, ULocalPlayer*);
+
 /**
  * GameInstance: high-level manager object for an instance of the running game.
  * Spawned at game creation and not destroyed until game instance is shut down.
@@ -232,6 +234,9 @@ public:
 
 	/** Local player access */
 
+	FOnLocalPlayerEvent OnLocalPlayerAddedEvent;
+	FOnLocalPlayerEvent OnLocalPlayerRemovedEvent;
+
 	/**
 	 * Debug console command to create a player.
 	 * @param ControllerId - The controller ID the player should accept input from.
@@ -265,7 +270,6 @@ public:
 	 */
 	virtual int32			AddLocalPlayer(ULocalPlayer* NewPlayer, int32 ControllerId);
 
-
 	/**
 	 * Removes a player.
 	 * @param Player - The player to remove.
diff --git a/Engine/Source/Runtime/Engine/Private/GameInstance.cpp b/Engine/Source/Runtime/Engine/Private/GameInstance.cpp
--- a/Engine/Source/Runtime/Engine/Private/GameInstance.cpp
+++ b/Engine/Source/Runtime/Engine/Private/GameInstance.cpp
@@ -715,7 +715,7 @@ ULocalPlayer* UGameInstance::CreateLocalPlayer(int32 ControllerId, FString& OutE
 
 int32 UGameInstance::AddLocalPlayer(ULocalPlayer* NewLocalPlayer, int32 ControllerId)
 {
-	if (NewLocalPlayer == NULL)
+	if (NewLocalPlayer == nullptr)
 	{
 		return INDEX_NONE;
 	}
@@ -729,11 +729,15 @@ int32 UGameInstance::AddLocalPlayer(ULocalPlayer* NewLocalPlayer, int32 Controll
 	NewLocalPlayer->PlayerAdded(GetGameViewportClient(), ControllerId);
 
 	// Notify the viewport that we added a player (so it can update splitscreen settings, etc)
-	if ( GetGameViewportClient() != NULL )
+	if ( GetGameViewportClient() != nullptr)
 	{
 		GetGameViewportClient()->NotifyPlayerAdded(InsertIndex, NewLocalPlayer);
 	}
 
+	UE_LOG(LogPlayerManagement, Log, TEXT("UGameInstance::AddLocalPlayer: Added player %s with ControllerId %d at index %d (%d remaining players)"), *NewLocalPlayer->GetName(), NewLocalPlayer->GetControllerId(), InsertIndex, LocalPlayers.Num());
+
+	OnLocalPlayerAddedEvent.Broadcast(NewLocalPlayer);
+
 	return InsertIndex;
 }
 
@@ -787,6 +791,8 @@ bool UGameInstance::RemoveLocalPlayer(ULocalPlayer* ExistingPlayer)
 
 	UE_LOG(LogPlayerManagement, Log, TEXT("UGameInstance::RemovePlayer: Removed player %s with ControllerId %i at index %i (%i remaining players)"), *ExistingPlayer->GetName(), ExistingPlayer->GetControllerId(), OldIndex, LocalPlayers.Num());
 
+	OnLocalPlayerRemovedEvent.Broadcast(ExistingPlayer);
+
 	return true;
 }
 
