AnimNode_SaveCachedPose.cpp: Fixed crash when an Inertialization node is placed before a Save Cached Pose node and then there are multiple uses of that same cached pose later in the graph

AnimNodeBase.h: Fixed crash when an inertial blend request is issued without an ancestor Inertialization node after a previous inertial blend request that did succeed earlier in the same graph execution path

#rnx

diff --git a/Engine/Source/Runtime/Engine/Classes/Animation/AnimNodeBase.h b/Engine/Source/Runtime/Engine/Classes/Animation/AnimNodeBase.h
--- a/Engine/Source/Runtime/Engine/Classes/Animation/AnimNodeBase.h
+++ b/Engine/Source/Runtime/Engine/Classes/Animation/AnimNodeBase.h
@@ -70,7 +70,7 @@ struct FAnimNodeTracker
 	NodeType* Top() const
 	{
 		const FNodeStack* Stack = Map.Find(GetKey<NodeType>());
-		return Stack ? static_cast<NodeType*>(Stack->Top()) : nullptr;
+		return (Stack && Stack->Num() != 0) ? static_cast<NodeType*>(Stack->Top()) : nullptr;
 	}
 
 	void CopyTopsOnly(const FAnimNodeTracker& Source)
diff --git a/Engine/Source/Runtime/Engine/Private/Animation/AnimNode_SaveCachedPose.cpp b/Engine/Source/Runtime/Engine/Private/Animation/AnimNode_SaveCachedPose.cpp
--- a/Engine/Source/Runtime/Engine/Private/Animation/AnimNode_SaveCachedPose.cpp
+++ b/Engine/Source/Runtime/Engine/Private/Animation/AnimNode_SaveCachedPose.cpp
@@ -117,7 +117,7 @@ void FAnimNode_SaveCachedPose::PostGraphUpdate()
 			{
 				for (auto Iter : AncestorTracker->Map)
 				{
-					FAnimNode_Base* AncestorNode = Iter.Value.Top();
+					FAnimNode_Base* AncestorNode = Iter.Value.Num() ? Iter.Value.Top() : nullptr;
 					if (AncestorNode && AncestorNode->WantsSkippedUpdates())
 					{
 						AncestorsWithSkippedUpdateHandlers.Add(AncestorNode);
