Random determinism
- Replaced uses of FMath::Rand() with instances of FRandomStream
- Enforce determinism when running with -FixedSeed
#rnx

diff --git a/Engine/Source/Editor/Cascade/Private/Cascade.cpp b/Engine/Source/Editor/Cascade/Private/Cascade.cpp
--- a/Engine/Source/Editor/Cascade/Private/Cascade.cpp
+++ b/Engine/Source/Editor/Cascade/Private/Cascade.cpp
@@ -72,6 +72,8 @@ static const FName Cascade_CurveEditorTab("Cascade_CurveEditor");
 
 DEFINE_LOG_CATEGORY(LogCascade);
 
+FRandomStream FCascade::RandomStream(FPlatformTime::Cycles());
+
 FCascade::FCascade()
 	: ParticleSystem(nullptr)
 	, ParticleSystemComponent(nullptr)
@@ -4756,7 +4758,7 @@ void FCascade::OnSetRandomSeed()
 		ParticleSystem->PreEditChange(NULL);
 		ParticleSystemComponent->PreEditChange(NULL);
 
-		int32 RandomSeed = FMath::RoundToInt(RAND_MAX * FMath::SRand());
+		int32 RandomSeed = FMath::RoundToInt(RAND_MAX * RandomStream.FRand());
 		if (SelectedModule->SetRandomSeedEntry(0, RandomSeed) == false)
 		{
 			UE_LOG(LogCascade, Warning, TEXT("Failed to set random seed entry on module %s"), *(SelectedModule->GetClass()->GetName()));
@@ -4898,7 +4900,7 @@ bool FCascade::ConvertModuleToSeeded(UParticleSystem* ParticleSystem, UParticleE
 			if (RandSeedInfo != NULL)
 			{
 				RandSeedInfo->bResetSeedOnEmitterLooping = true;
-				RandSeedInfo->RandomSeeds.Add(FMath::TruncToInt(FMath::Rand() * UINT_MAX));
+				RandSeedInfo->RandomSeeds.Add(FMath::TruncToInt(RandomStream.FRand() * UINT_MAX));
 			}
 		}
 
diff --git a/Engine/Source/Editor/Cascade/Private/Cascade.h b/Engine/Source/Editor/Cascade/Private/Cascade.h
--- a/Engine/Source/Editor/Cascade/Private/Cascade.h
+++ b/Engine/Source/Editor/Cascade/Private/Cascade.h
@@ -16,6 +16,7 @@
 #include "ICascade.h"
 #include "IDistCurveEditor.h"
 #include "Particles/ParticleEmitter.h"
+#include "Math/RandomStream.h"
 
 class FFXSystemInterface;
 class IDetailsView;
@@ -496,4 +497,7 @@ private:
 
 	/** Maps particle emitters to rendered thumbnails for the emitter UI */
 	FParticleEmitterThumbnailMap EmitterToThumbnailMap;
+
+	/** Random stream used to seed particle modules */
+	static FRandomStream RandomStream;
 };
diff --git a/Engine/Source/Runtime/AudioMixer/Private/DSP/Noise.cpp b/Engine/Source/Runtime/AudioMixer/Private/DSP/Noise.cpp
--- a/Engine/Source/Runtime/AudioMixer/Private/DSP/Noise.cpp
+++ b/Engine/Source/Runtime/AudioMixer/Private/DSP/Noise.cpp
@@ -5,21 +5,21 @@
 namespace Audio
 {
 	FWhiteNoise::FWhiteNoise(const float InScale, const float InAdd)
-		: Scale(InScale)
-		, Add(InAdd)
 	{
-
+		SetScaleAdd(InScale, InAdd);
 	}
 
 	void FWhiteNoise::SetScaleAdd(const float InScale, const float InAdd)
 	{
 		Scale = InScale;
 		Add = InAdd;
+
+		RandomStream.Initialize(HashCombine(GetTypeHash(Scale), GetTypeHash(Add)));
 	}
 
 	float FWhiteNoise::Generate()
 	{
-		return Add + Scale * FMath::SRand() * 2 - 1.0f;
+		return Add + Scale * RandomStream.FRand() * 2 - 1.0f;
 	}
 
 	FPinkNoise::FPinkNoise(const float InScale, const float InAdd)
diff --git a/Engine/Source/Runtime/AudioMixer/Public/DSP/Noise.h b/Engine/Source/Runtime/AudioMixer/Public/DSP/Noise.h
--- a/Engine/Source/Runtime/AudioMixer/Public/DSP/Noise.h
+++ b/Engine/Source/Runtime/AudioMixer/Public/DSP/Noise.h
@@ -3,6 +3,7 @@
 #pragma once
 
 #include "CoreMinimal.h"
+#include "Math/RandomStream.h"
 
 namespace Audio
 {
@@ -24,6 +25,8 @@ namespace Audio
 	private:
 		float Scale;
 		float Add;
+
+		FRandomStream RandomStream;
 	};
 
 	/** 
diff --git a/Engine/Source/Runtime/Core/Public/Math/RandomStream.h b/Engine/Source/Runtime/Core/Public/Math/RandomStream.h
--- a/Engine/Source/Runtime/Core/Public/Math/RandomStream.h
+++ b/Engine/Source/Runtime/Core/Public/Math/RandomStream.h
@@ -8,6 +8,7 @@
 #include "Math/Matrix.h"
 #include "Math/RotationMatrix.h"
 #include "Math/Transform.h"
+#include "HAL/PlatformTime.h"
 
 /**
  * Implements a thread-safe SRand based RNG.
@@ -36,9 +37,20 @@ public:
 	 * @param InSeed The seed value.
 	 */
 	FRandomStream( int32 InSeed )
-		: InitialSeed(InSeed)
-		, Seed(InSeed)
-	{ }
+	{ 
+		Initialize(InSeed);
+	}
+
+	/**
+	 * Creates and initializes a new random stream from the specified name.
+	 *
+	 * @note If NAME_None is provided, the stream will be seeded using the current time.
+	 * @param InName The name value from which the stream will be initialized.
+	 */
+	FRandomStream( FName InName )
+	{
+		Initialize(InName);
+	}
 
 public:
 
@@ -53,6 +65,26 @@ public:
 		Seed = InSeed;
 	}
 
+	/**
+	 * Initializes this random stream using the specified name.
+	 *
+	 * @note If NAME_None is provided, the stream will be seeded using the current time.
+	 * @param InName The name value from which the stream will be initialized.
+	 */
+	void Initialize( FName InName )
+	{
+		if (InName != NAME_None)
+		{
+			InitialSeed = GetTypeHash(InName.ToString());
+		}
+		else
+		{
+			InitialSeed = FPlatformTime::Cycles();
+		}
+
+		Seed = InitialSeed;
+	}
+
 	/**
 	 * Resets this random stream to the initial seed value.
 	 */
diff --git a/Engine/Source/Runtime/Engine/Classes/Components/AudioComponent.h b/Engine/Source/Runtime/Engine/Classes/Components/AudioComponent.h
--- a/Engine/Source/Runtime/Engine/Classes/Components/AudioComponent.h
+++ b/Engine/Source/Runtime/Engine/Classes/Components/AudioComponent.h
@@ -9,6 +9,7 @@
 #include "Components/SceneComponent.h"
 #include "Sound/SoundAttenuation.h"
 #include "Sound/SoundWave.h"
+#include "Math/RandomStream.h"
 
 #include "AudioComponent.generated.h"
 
@@ -595,6 +596,8 @@ private:
 	void UpdateSpriteTexture();
 #endif
 
+	FRandomStream RandomStream;
+
 	static uint64 AudioComponentIDCounter;
 	static TMap<uint64, UAudioComponent*> AudioIDToComponentMap;
 	static FCriticalSection AudioIDToComponentMapLock;
diff --git a/Engine/Source/Runtime/Engine/Classes/Engine/NetDriver.h b/Engine/Source/Runtime/Engine/Classes/Engine/NetDriver.h
--- a/Engine/Source/Runtime/Engine/Classes/Engine/NetDriver.h
+++ b/Engine/Source/Runtime/Engine/Classes/Engine/NetDriver.h
@@ -5,6 +5,7 @@
 #pragma once
 
 #include "CoreMinimal.h"
+#include "Math/RandomStream.h"
 #include "UObject/ObjectMacros.h"
 #include "UObject/UObjectGlobals.h"
 #include "UObject/Object.h"
@@ -1215,6 +1216,9 @@ protected:
 
 	bool bMaySendProperties;
 
+	/** Stream of random numbers to be used by this instance of UNetDriver */
+	FRandomStream UpdateDelayRandomStream;
+
 private:
 
 	FDelegateHandle PostGarbageCollectHandle;
diff --git a/Engine/Source/Runtime/Engine/Classes/GameFramework/CharacterMovementComponent.h b/Engine/Source/Runtime/Engine/Classes/GameFramework/CharacterMovementComponent.h
--- a/Engine/Source/Runtime/Engine/Classes/GameFramework/CharacterMovementComponent.h
+++ b/Engine/Source/Runtime/Engine/Classes/GameFramework/CharacterMovementComponent.h
@@ -3,6 +3,7 @@
 #pragma once
 
 #include "CoreMinimal.h"
+#include "Math/RandomStream.h"
 #include "UObject/ObjectMacros.h"
 #include "UObject/UObjectGlobals.h"
 #include "Engine/NetSerialization.h"
@@ -2110,6 +2111,8 @@ protected:
 	class FNetworkPredictionData_Client_Character* ClientPredictionData;
 	class FNetworkPredictionData_Server_Character* ServerPredictionData;
 
+	FRandomStream RandomStream;
+
 	/**
 	 * Smooth mesh location for network interpolation, based on values set up by SmoothCorrection.
 	 * Internally this simply calls SmoothClientPosition_Interpolate() then SmoothClientPosition_UpdateVisuals().
diff --git a/Engine/Source/Runtime/Engine/Classes/Sound/SoundNode.h b/Engine/Source/Runtime/Engine/Classes/Sound/SoundNode.h
--- a/Engine/Source/Runtime/Engine/Classes/Sound/SoundNode.h
+++ b/Engine/Source/Runtime/Engine/Classes/Sound/SoundNode.h
@@ -4,6 +4,7 @@
 #pragma once
 
 #include "CoreMinimal.h"
+#include "Math/RandomStream.h"
 #include "UObject/ObjectMacros.h"
 #include "UObject/Object.h"
 #include "UObject/Class.h"
@@ -69,6 +70,9 @@ class ENGINE_API USoundNode : public UObject
 	class UEdGraphNode* GetGraphNode() const;
 #endif
 
+	/** Stream of random numbers to be used by this instance of USoundNode */
+	FRandomStream RandomStream;
+
 public:
 	//~ Begin UObject Interface
 #if WITH_EDITOR
diff --git a/Engine/Source/Runtime/Engine/Private/Components/AudioComponent.cpp b/Engine/Source/Runtime/Engine/Private/Components/AudioComponent.cpp
--- a/Engine/Source/Runtime/Engine/Private/Components/AudioComponent.cpp
+++ b/Engine/Source/Runtime/Engine/Private/Components/AudioComponent.cpp
@@ -10,6 +10,7 @@
 #include "Sound/SoundCue.h"
 #include "Components/BillboardComponent.h"
 #include "UObject/FrameworkObjectVersion.h"
+#include "Misc/App.h"
 
 DECLARE_CYCLE_STAT(TEXT("AudioComponent Play"), STAT_AudioComp_Play, STATGROUP_Audio);
 
@@ -67,6 +68,8 @@ UAudioComponent::UAudioComponent(const FObjectInitializer& ObjectInitializer)
 	AudioDeviceHandle = INDEX_NONE;
 	AudioComponentID = FPlatformAtomics::InterlockedIncrement(reinterpret_cast<volatile int64*>(&AudioComponentIDCounter));
 
+	RandomStream.Initialize(FApp::bUseFixedSeed ? GetFName() : NAME_None);
+
 	// TODO: Consider only putting played/active components in to the map
 	FScopeLock Lock(&AudioIDToComponentMapLock);
 	AudioIDToComponentMap.Add(AudioComponentID, this);
@@ -370,7 +373,7 @@ void UAudioComponent::PlayInternal(const float StartTime, const float FadeInDura
 			NewActiveSound.SetSoundClass(SoundClassOverride);
 			NewActiveSound.ConcurrencySet = ConcurrencySet;
 
-			NewActiveSound.VolumeMultiplier = (VolumeModulationMax + ((VolumeModulationMin - VolumeModulationMax) * FMath::SRand())) * VolumeMultiplier;
+			NewActiveSound.VolumeMultiplier = (VolumeModulationMax + ((VolumeModulationMin - VolumeModulationMax) * RandomStream.FRand())) * VolumeMultiplier;
 			// The priority used for the active sound is the audio component's priority scaled with the sound's priority
 			if (bOverridePriority)
 			{
@@ -381,7 +384,7 @@ void UAudioComponent::PlayInternal(const float StartTime, const float FadeInDura
 				NewActiveSound.Priority = Sound->Priority;
 			}
 
-			NewActiveSound.PitchMultiplier = (PitchModulationMax + ((PitchModulationMin - PitchModulationMax) * FMath::SRand())) * PitchMultiplier;
+			NewActiveSound.PitchMultiplier = (PitchModulationMax + ((PitchModulationMin - PitchModulationMax) * RandomStream.FRand())) * PitchMultiplier;
 			NewActiveSound.bEnableLowPassFilter = bEnableLowPassFilter;
 			NewActiveSound.LowPassFilterFrequency = LowPassFilterFrequency;
 			NewActiveSound.RequestedStartTime = FMath::Max(0.f, StartTime);
diff --git a/Engine/Source/Runtime/Engine/Private/Components/CharacterMovementComponent.cpp b/Engine/Source/Runtime/Engine/Private/Components/CharacterMovementComponent.cpp
--- a/Engine/Source/Runtime/Engine/Private/Components/CharacterMovementComponent.cpp
+++ b/Engine/Source/Runtime/Engine/Private/Components/CharacterMovementComponent.cpp
@@ -24,6 +24,7 @@
 #include "AI/Navigation/PathFollowingAgentInterface.h"
 #include "AI/Navigation/AvoidanceManager.h"
 #include "Components/BrushComponent.h"
+#include "Misc/App.h"
 
 #include "Engine/DemoNetDriver.h"
 #include "Engine/NetworkObjectList.h"
@@ -379,6 +380,8 @@ FName FCharacterMovementComponentPostPhysicsTickFunction::DiagnosticContext(bool
 UCharacterMovementComponent::UCharacterMovementComponent(const FObjectInitializer& ObjectInitializer)
 	: Super(ObjectInitializer)
 {
+	RandomStream.Initialize(FApp::bUseFixedSeed ? GetFName() : NAME_None);
+
 	PostPhysicsTickFunction.bCanEverTick = true;
 	PostPhysicsTickFunction.bStartWithTickEnabled = false;
 	PostPhysicsTickFunction.SetTickFunctionEnable(false);
@@ -941,7 +944,7 @@ FVector UCharacterMovementComponent::GetBestDirectionOffActor(AActor* BaseActor)
 
 float UCharacterMovementComponent::GetNetworkSafeRandomAngleDegrees() const
 {
-	float Angle = FMath::SRand() * 360.f;
+	float Angle = RandomStream.FRand() * 360.f;
 
 	if (!IsNetMode(NM_Standalone))
 	{
@@ -4401,8 +4404,8 @@ void UCharacterMovementComponent::PhysFalling(float deltaTime, int32 Iterations)
 							const float MovedDist2DSq = (PawnLocation - OldLocation).SizeSquared2D();
 							if (ZMovedDist <= 0.2f * timeTick && MovedDist2DSq <= 4.f * timeTick)
 							{
-								Velocity.X += 0.25f * GetMaxSpeed() * (FMath::FRand() - 0.5f);
-								Velocity.Y += 0.25f * GetMaxSpeed() * (FMath::FRand() - 0.5f);
+								Velocity.X += 0.25f * GetMaxSpeed() * (RandomStream.FRand() - 0.5f);
+								Velocity.Y += 0.25f * GetMaxSpeed() * (RandomStream.FRand() - 0.5f);
 								Velocity.Z = FMath::Max<float>(JumpZVelocity * 0.25f, 1.f);
 								Delta = Velocity * timeTick;
 								SafeMoveUpdatedComponent(Delta, PawnRotation, true, Hit);
@@ -8073,7 +8076,7 @@ void UCharacterMovementComponent::ReplicateMoveToServer(float DeltaTime, const F
 			bSendServerMove = false;
 			UE_LOG(LogNetPlayerMovement, Log, TEXT("Drop ServerMove, %.2f time remains"), CharacterMovementCVars::NetForceClientServerMoveLossDuration - TimeSinceLossStart);
 		}
-		else if (CharacterMovementCVars::NetForceClientServerMoveLossPercent != 0.f && (FMath::SRand() < CharacterMovementCVars::NetForceClientServerMoveLossPercent))
+		else if (CharacterMovementCVars::NetForceClientServerMoveLossPercent != 0.f && (RandomStream.FRand() < CharacterMovementCVars::NetForceClientServerMoveLossPercent))
 		{
 			bSendServerMove = false;
 			ClientData->DebugForcedPacketLossTimerStart = (CharacterMovementCVars::NetForceClientServerMoveLossDuration > 0) ? MyWorld->RealTimeSeconds : 0.0f;
@@ -8817,7 +8820,7 @@ bool UCharacterMovementComponent::ServerCheckClientError(float ClientTimeStamp,
 #if !(UE_BUILD_SHIPPING || UE_BUILD_TEST)
 		if (CharacterMovementCVars::NetForceClientAdjustmentPercent > SMALL_NUMBER)
 		{
-			if (FMath::SRand() < CharacterMovementCVars::NetForceClientAdjustmentPercent)
+			if (RandomStream.FRand() < CharacterMovementCVars::NetForceClientAdjustmentPercent)
 			{
 				UE_LOG(LogNetPlayerMovement, VeryVerbose, TEXT("** ServerCheckClientError forced by p.NetForceClientAdjustmentPercent"));
 				return true;
diff --git a/Engine/Source/Runtime/Engine/Private/DemoNetDriver.cpp b/Engine/Source/Runtime/Engine/Private/DemoNetDriver.cpp
--- a/Engine/Source/Runtime/Engine/Private/DemoNetDriver.cpp
+++ b/Engine/Source/Runtime/Engine/Private/DemoNetDriver.cpp
@@ -2830,7 +2830,7 @@ bool UDemoNetDriver::ReplicatePrioritizedActor(const FActorPriority& ActorPriori
 		const double ClampedExtraTime = FMath::Clamp(ExtraTime, 0.0, NetUpdateDelay);
 
 		// Try to spread the updates across multiple frames to smooth out spikes.
-		ActorInfo->NextUpdateTime = (DemoCurrentTime + NextUpdateDelta - ClampedExtraTime + ((FMath::SRand() - 0.5) * Params.ServerTickTime));
+		ActorInfo->NextUpdateTime = (DemoCurrentTime + NextUpdateDelta - ClampedExtraTime + ((UpdateDelayRandomStream.FRand() - 0.5) * Params.ServerTickTime));
 
 		Actor->CallPreReplication(this);
 
diff --git a/Engine/Source/Runtime/Engine/Private/NetworkDriver.cpp b/Engine/Source/Runtime/Engine/Private/NetworkDriver.cpp
--- a/Engine/Source/Runtime/Engine/Private/NetworkDriver.cpp
+++ b/Engine/Source/Runtime/Engine/Private/NetworkDriver.cpp
@@ -9,6 +9,7 @@
 #include "Misc/CommandLine.h"
 #include "Misc/NetworkGuid.h"
 #include "Stats/Stats.h"
+#include "Misc/App.h"
 #include "Misc/MemStack.h"
 #include "HAL/IConsoleManager.h"
 #include "UObject/ObjectMacros.h"
@@ -276,6 +277,8 @@ PRAGMA_DISABLE_DEPRECATION_WARNINGS
 	ChannelClasses[CHTYPE_Actor]	= UActorChannel::StaticClass();
 	ChannelClasses[CHTYPE_Voice]	= UVoiceChannel::StaticClass();
 PRAGMA_ENABLE_DEPRECATION_WARNINGS
+
+	UpdateDelayRandomStream.Initialize(FApp::bUseFixedSeed ? GetFName() : NAME_None);
 }
 
 void UNetDriver::InitPacketSimulationSettings()
@@ -3659,7 +3662,7 @@ void UNetDriver::ServerReplicateActors_BuildConsiderList( TArray<FNetworkObjectI
 			const float NextUpdateDelta = bUseAdapativeNetFrequency ? ActorInfo->OptimalNetUpdateDelta : 1.0f / Actor->NetUpdateFrequency;
 
 			// then set the next update time
-			ActorInfo->NextUpdateTime = World->TimeSeconds + FMath::SRand() * ServerTickTime + NextUpdateDelta;
+			ActorInfo->NextUpdateTime = World->TimeSeconds + UpdateDelayRandomStream.FRand() * ServerTickTime + NextUpdateDelta;
 
 			// and mark when the actor first requested an update
 			//@note: using Time because it's compared against UActorChannel.LastUpdateTime which also uses that value
@@ -4023,7 +4026,7 @@ int32 UNetDriver::ServerReplicateActors_ProcessPrioritizedActors( UNetConnection
 					// if it is relevant then mark the channel as relevant for a short amount of time
 					if ( bIsRelevant )
 					{
-						Channel->RelevantTime = Time + 0.5f * FMath::SRand();
+						Channel->RelevantTime = Time + 0.5f * UpdateDelayRandomStream.FRand();
 					}
 					// if the channel isn't saturated
 					if ( Channel->IsNetReady( 0 ) )
@@ -4422,7 +4425,7 @@ int32 UNetDriver::ServerReplicateActors(float DeltaSeconds)
 					PriorityActors[k]->ActorInfo->bPendingNetUpdate = true;
 					if ( Channel != NULL )
 					{
-						Channel->RelevantTime = Time + 0.5f * FMath::SRand();
+						Channel->RelevantTime = Time + 0.5f * UpdateDelayRandomStream.FRand();
 					}
 				}
 			}
diff --git a/Engine/Source/Runtime/Engine/Private/SoundNode.cpp b/Engine/Source/Runtime/Engine/Private/SoundNode.cpp
--- a/Engine/Source/Runtime/Engine/Private/SoundNode.cpp
+++ b/Engine/Source/Runtime/Engine/Private/SoundNode.cpp
@@ -4,6 +4,7 @@
 #include "Sound/SoundNode.h"
 #include "EngineUtils.h"
 #include "Sound/SoundCue.h"
+#include "Misc/App.h"
 
 /*-----------------------------------------------------------------------------
 	USoundNode implementation.
@@ -11,6 +12,7 @@
 USoundNode::USoundNode(const FObjectInitializer& ObjectInitializer)
 	: Super(ObjectInitializer)
 {
+	RandomStream.Initialize(FApp::bUseFixedSeed ? GetFName() : NAME_None);
 }
 
 
diff --git a/Engine/Source/Runtime/Engine/Private/SoundNodeDelay.cpp b/Engine/Source/Runtime/Engine/Private/SoundNodeDelay.cpp
--- a/Engine/Source/Runtime/Engine/Private/SoundNodeDelay.cpp
+++ b/Engine/Source/Runtime/Engine/Private/SoundNodeDelay.cpp
@@ -31,7 +31,7 @@ void USoundNodeDelay::ParseNodes( FAudioDevice* AudioDevice, const UPTRINT NodeW
 	{
 		*RequiresInitialization = false;
 
-		const float ActualDelay = FMath::Max(0.f, DelayMax + ( ( DelayMin - DelayMax ) * FMath::SRand() ));
+		const float ActualDelay = FMath::Max(0.f, DelayMax + ( ( DelayMin - DelayMax ) * RandomStream.FRand() ));
 
 		if (ActualDelay > 0.0f && ParseParams.StartTime >= ActualDelay)
 		{
diff --git a/Engine/Source/Runtime/Engine/Private/SoundNodeEnveloper.cpp b/Engine/Source/Runtime/Engine/Private/SoundNodeEnveloper.cpp
--- a/Engine/Source/Runtime/Engine/Private/SoundNodeEnveloper.cpp
+++ b/Engine/Source/Runtime/Engine/Private/SoundNodeEnveloper.cpp
@@ -68,8 +68,8 @@ void USoundNodeEnveloper::ParseNodes( FAudioDevice* AudioDevice, const UPTRINT N
 	if( *RequiresInitialization )
 	{
 		StartTime = ActiveSound.PlaybackTime - ParseParams.StartTime;
-		UsedVolumeModulation = VolumeMax + ( ( VolumeMin - VolumeMax ) * FMath::SRand() );
-		UsedPitchModulation = PitchMax + ( ( PitchMin - PitchMax ) * FMath::SRand() );
+		UsedVolumeModulation = VolumeMax + ( ( VolumeMin - VolumeMax ) * RandomStream.FRand() );
+		UsedPitchModulation = PitchMax + ( ( PitchMin - PitchMax ) * RandomStream.FRand() );
 		LastLoopCount = -1;
 
 		*RequiresInitialization = false;
@@ -95,8 +95,8 @@ void USoundNodeEnveloper::ParseNodes( FAudioDevice* AudioDevice, const UPTRINT N
 		else if ( CurrentLoopCount != LastLoopCount )
 		{
 			// randomize multipliers for the new repeat
-			UsedVolumeModulation = VolumeMax + ( ( VolumeMin - VolumeMax ) * FMath::SRand() );
-			UsedPitchModulation = PitchMax + ( ( PitchMin - PitchMax ) * FMath::SRand() );
+			UsedVolumeModulation = VolumeMax + ( ( VolumeMin - VolumeMax ) * RandomStream.FRand() );
+			UsedPitchModulation = PitchMax + ( ( PitchMin - PitchMax ) * RandomStream.FRand() );
 			LastLoopCount = CurrentLoopCount;
 		}
 	}
diff --git a/Engine/Source/Runtime/Engine/Private/SoundNodeModulator.cpp b/Engine/Source/Runtime/Engine/Private/SoundNodeModulator.cpp
--- a/Engine/Source/Runtime/Engine/Private/SoundNodeModulator.cpp
+++ b/Engine/Source/Runtime/Engine/Private/SoundNodeModulator.cpp
@@ -24,8 +24,8 @@ void USoundNodeModulator::ParseNodes( FAudioDevice* AudioDevice, const UPTRINT N
 
 	if( *RequiresInitialization )
 	{
-		UsedVolumeModulation = VolumeMax + ( ( VolumeMin - VolumeMax ) * FMath::SRand() );
-		UsedPitchModulation = PitchMax + ( ( PitchMin - PitchMax ) * FMath::SRand() );
+		UsedVolumeModulation = VolumeMax + ( ( VolumeMin - VolumeMax ) * RandomStream.FRand() );
+		UsedPitchModulation = PitchMax + ( ( PitchMin - PitchMax ) * RandomStream.FRand() );
 
 		*RequiresInitialization = 0;
 	}
diff --git a/Engine/Source/Runtime/Engine/Private/SoundNodeOscillator.cpp b/Engine/Source/Runtime/Engine/Private/SoundNodeOscillator.cpp
--- a/Engine/Source/Runtime/Engine/Private/SoundNodeOscillator.cpp
+++ b/Engine/Source/Runtime/Engine/Private/SoundNodeOscillator.cpp
@@ -32,10 +32,10 @@ void USoundNodeOscillator::ParseNodes( FAudioDevice* AudioDevice, const UPTRINT
 
 	if( *RequiresInitialization )
 	{
-		UsedAmplitude = AmplitudeMax + ( ( AmplitudeMin - AmplitudeMax ) * FMath::SRand() );
-		UsedFrequency = FrequencyMax + ( ( FrequencyMin - FrequencyMax ) * FMath::SRand() );
-		UsedOffset = OffsetMax + ( ( OffsetMin - OffsetMax ) * FMath::SRand() );
-		UsedCenter = CenterMax + ( ( CenterMin - CenterMax ) * FMath::SRand() );
+		UsedAmplitude = AmplitudeMax + ( ( AmplitudeMin - AmplitudeMax ) * RandomStream.FRand() );
+		UsedFrequency = FrequencyMax + ( ( FrequencyMin - FrequencyMax ) * RandomStream.FRand() );
+		UsedOffset = OffsetMax + ( ( OffsetMin - OffsetMax ) * RandomStream.FRand() );
+		UsedCenter = CenterMax + ( ( CenterMin - CenterMax ) * RandomStream.FRand() );
 
 		*RequiresInitialization = 0;
 	}
