PR #5650: : Add missing parameter/default values for debug draw functions (Contributed by [REDACTED_PERSON])

diff --git a/Engine/Source/Runtime/Engine/Public/DrawDebugHelpers.h b/Engine/Source/Runtime/Engine/Public/DrawDebugHelpers.h
--- a/Engine/Source/Runtime/Engine/Public/DrawDebugHelpers.h
+++ b/Engine/Source/Runtime/Engine/Public/DrawDebugHelpers.h
@@ -153,9 +153,9 @@ FORCEINLINE void DrawDebugSphere(const UWorld* InWorld, FVector const& Center, f
 FORCEINLINE void DrawDebugCylinder(const UWorld* InWorld, FVector const& Start, FVector const& End, float Radius, int32 Segments, FColor const& Color, bool bPersistentLines = false, float LifeTime = -1.f, uint8 DepthPriority = 0, float Thickness = 0.f) {}
 FORCEINLINE void DrawDebugCone(const UWorld* InWorld, FVector const& Origin, FVector const& Direction, float Length, float AngleWidth, float AngleHeight, int32 NumSides, FColor const& Color, bool bPersistentLines = false, float LifeTime = -1.f, uint8 DepthPriority = 0, float Thickness = 0.f) {}
 FORCEINLINE void DrawDebugAltCone(const UWorld* InWorld, FVector const& Origin, FRotator const& Rotation, float Length, float AngleWidth, float AngleHeight, FColor const& DrawColor, bool bPersistentLines = false, float LifeTime = -1.f, uint8 DepthPriority = 0, float Thickness = 0.f) {}
-FORCEINLINE void DrawDebugString(const UWorld* InWorld, FVector const& TextLocation, const FString& Text, class AActor* TestBaseActor = NULL, FColor const& TextColor = FColor::White, float Duration = -1.000000, bool bDrawShadow = false) {}
+FORCEINLINE void DrawDebugString(const UWorld* InWorld, FVector const& TextLocation, const FString& Text, class AActor* TestBaseActor = NULL, FColor const& TextColor = FColor::White, float Duration = -1.000000, bool bDrawShadow = false, float FontScale = 1.f) {}
 FORCEINLINE void DrawDebugFrustum(const UWorld* InWorld, const FMatrix& FrustumToWorld, FColor const& Color, bool bPersistentLines = false, float LifeTime = -1.f, uint8 DepthPriority = 0, float Thickness = 0.f) {}
-FORCEINLINE void DrawCircle(const UWorld* InWorld, const FVector& Base, const FVector& X, const FVector& Y, const FColor& Color, float Radius, int32 NumSides, bool bPersistentLines, float LifeTime, uint8 DepthPriority, float Thickness) {}
+FORCEINLINE void DrawCircle(const UWorld* InWorld, const FVector& Base, const FVector& X, const FVector& Y, const FColor& Color, float Radius, int32 NumSides, bool bPersistentLines = false, float LifeTime = -1.f, uint8 DepthPriority = 0, float Thickness = 0) {}
 FORCEINLINE void DrawDebugCapsule(const UWorld* InWorld, FVector const& Center, float HalfHeight, float Radius, const FQuat& Rotation, FColor const& Color, bool bPersistentLines = false, float LifeTime = -1.f, uint8 DepthPriority = 0, float Thickness = 0) {}
 FORCEINLINE void DrawDebugCamera(const UWorld* InWorld, FVector const& Location, FRotator const& Rotation, float FOVDeg, float Scale = 1.f, FColor const& Color = FColor::White, bool bPersistentLines = false, float LifeTime = -1.f, uint8 DepthPriority = 0) {}
 FORCEINLINE void FlushDebugStrings(const UWorld* InWorld) {}
