Particle random determinism
- When running with -FixedSeed, automatically add a seed payload to each instance, for each module
- Refactored code to always use a FRandomStream, obtained through the UParticleModule::GetRandomStream() method
- Cleaned up the _Seeded class to remove redundant code that is not needed anymore with the addition of these changes.
#rnx

diff --git a/Engine/Source/Runtime/Engine/Classes/Particles/Color/ParticleModuleColor_Seeded.h b/Engine/Source/Runtime/Engine/Classes/Particles/Color/ParticleModuleColor_Seeded.h
--- a/Engine/Source/Runtime/Engine/Classes/Particles/Color/ParticleModuleColor_Seeded.h
+++ b/Engine/Source/Runtime/Engine/Classes/Particles/Color/ParticleModuleColor_Seeded.h
@@ -21,9 +21,6 @@ class UParticleModuleColor_Seeded : public UParticleModuleColor
 
 
 	//Begin UParticleModule Interface
-	virtual void Spawn(FParticleEmitterInstance* Owner, int32 Offset, float SpawnTime, FBaseParticle* ParticleBase) override;
-	virtual uint32	RequiredBytesPerInstance() override;
-	virtual uint32	PrepPerInstanceBlock(FParticleEmitterInstance* Owner, void* InstData) override;
 	virtual FParticleRandomSeedInfo* GetRandomSeedInfo() override
 	{
 		return &RandomSeedInfo;
diff --git a/Engine/Source/Runtime/Engine/Classes/Particles/Lifetime/ParticleModuleLifetime_Seeded.h b/Engine/Source/Runtime/Engine/Classes/Particles/Lifetime/ParticleModuleLifetime_Seeded.h
--- a/Engine/Source/Runtime/Engine/Classes/Particles/Lifetime/ParticleModuleLifetime_Seeded.h
+++ b/Engine/Source/Runtime/Engine/Classes/Particles/Lifetime/ParticleModuleLifetime_Seeded.h
@@ -21,9 +21,6 @@ class UParticleModuleLifetime_Seeded : public UParticleModuleLifetime
 
 
 	//Begin UParticleModule Interface
-	virtual void Spawn(FParticleEmitterInstance* Owner, int32 Offset, float SpawnTime, FBaseParticle* ParticleBase) override;
-	virtual uint32	RequiredBytesPerInstance() override;
-	virtual uint32	PrepPerInstanceBlock(FParticleEmitterInstance* Owner, void* InstData) override;
 	virtual FParticleRandomSeedInfo* GetRandomSeedInfo() override
 	{
 		return &RandomSeedInfo;
diff --git a/Engine/Source/Runtime/Engine/Classes/Particles/Light/ParticleModuleLight_Seeded.h b/Engine/Source/Runtime/Engine/Classes/Particles/Light/ParticleModuleLight_Seeded.h
--- a/Engine/Source/Runtime/Engine/Classes/Particles/Light/ParticleModuleLight_Seeded.h
+++ b/Engine/Source/Runtime/Engine/Classes/Particles/Light/ParticleModuleLight_Seeded.h
@@ -21,9 +21,6 @@ class UParticleModuleLight_Seeded : public UParticleModuleLight
 
 
 	//Begin UParticleModule Interface
-	virtual void Spawn(FParticleEmitterInstance* Owner, int32 Offset, float SpawnTime, FBaseParticle* ParticleBase) override;
-	virtual uint32	RequiredBytesPerInstance() override;
-	virtual uint32	PrepPerInstanceBlock(FParticleEmitterInstance* Owner, void* InstData) override;
 	virtual FParticleRandomSeedInfo* GetRandomSeedInfo() override
 	{
 		return &RandomSeedInfo;
diff --git a/Engine/Source/Runtime/Engine/Classes/Particles/Location/ParticleModuleLocationBoneSocket.h b/Engine/Source/Runtime/Engine/Classes/Particles/Location/ParticleModuleLocationBoneSocket.h
--- a/Engine/Source/Runtime/Engine/Classes/Particles/Location/ParticleModuleLocationBoneSocket.h
+++ b/Engine/Source/Runtime/Engine/Classes/Particles/Location/ParticleModuleLocationBoneSocket.h
@@ -179,8 +179,8 @@ class ENGINE_API UParticleModuleLocationBoneSocket : public UParticleModuleLocat
 	bool GetBoneInfoForSourceIndex(FModuleLocationBoneSocketInstancePayload* InstancePayload, USkeletalMeshComponent* SourceComponent, int32 SourceIndex, FMatrix& OutBoneMatrix, FVector& OutOffset)const;
 	
 	/** Selects the next socket or bone index to spawn from. */
-	int32 SelectNextSpawnIndex(FModuleLocationBoneSocketInstancePayload* InstancePayload, USkeletalMeshComponent* SourceComponent);
-	void RegeneratePreSelectedIndices(FModuleLocationBoneSocketInstancePayload* InstancePayload, USkeletalMeshComponent* SourceComponent);
+	int32 SelectNextSpawnIndex(FModuleLocationBoneSocketInstancePayload* InstancePayload, USkeletalMeshComponent* SourceComponent, FRandomStream& InRandomStream);
+	void RegeneratePreSelectedIndices(FModuleLocationBoneSocketInstancePayload* InstancePayload, USkeletalMeshComponent* SourceComponent, FRandomStream& InRandomStream);
 
 	void UpdatePrevBoneLocationsAndVelocities(FModuleLocationBoneSocketInstancePayload* InstancePayload, USkeletalMeshComponent* SourceComponent, float DeltaTime);
 
diff --git a/Engine/Source/Runtime/Engine/Classes/Particles/Location/ParticleModuleLocationPrimitiveCylinder_Seeded.h b/Engine/Source/Runtime/Engine/Classes/Particles/Location/ParticleModuleLocationPrimitiveCylinder_Seeded.h
--- a/Engine/Source/Runtime/Engine/Classes/Particles/Location/ParticleModuleLocationPrimitiveCylinder_Seeded.h
+++ b/Engine/Source/Runtime/Engine/Classes/Particles/Location/ParticleModuleLocationPrimitiveCylinder_Seeded.h
@@ -21,9 +21,6 @@ class ENGINE_API UParticleModuleLocationPrimitiveCylinder_Seeded : public UParti
 
 
 	//Begin UParticleModule Interface
-	virtual void Spawn(FParticleEmitterInstance* Owner, int32 Offset, float SpawnTime, FBaseParticle* ParticleBase) override;
-	virtual uint32	RequiredBytesPerInstance() override;
-	virtual uint32	PrepPerInstanceBlock(FParticleEmitterInstance* Owner, void* InstData) override;
 	virtual FParticleRandomSeedInfo* GetRandomSeedInfo() override
 	{
 		return &RandomSeedInfo;
diff --git a/Engine/Source/Runtime/Engine/Classes/Particles/Location/ParticleModuleLocationPrimitiveSphere_Seeded.h b/Engine/Source/Runtime/Engine/Classes/Particles/Location/ParticleModuleLocationPrimitiveSphere_Seeded.h
--- a/Engine/Source/Runtime/Engine/Classes/Particles/Location/ParticleModuleLocationPrimitiveSphere_Seeded.h
+++ b/Engine/Source/Runtime/Engine/Classes/Particles/Location/ParticleModuleLocationPrimitiveSphere_Seeded.h
@@ -21,9 +21,6 @@ class ENGINE_API UParticleModuleLocationPrimitiveSphere_Seeded : public UParticl
 
 
 	//Begin UParticleModule Interface
-	virtual void Spawn(FParticleEmitterInstance* Owner, int32 Offset, float SpawnTime, FBaseParticle* ParticleBase) override;
-	virtual uint32	RequiredBytesPerInstance() override;
-	virtual uint32	PrepPerInstanceBlock(FParticleEmitterInstance* Owner, void* InstData) override;
 	virtual FParticleRandomSeedInfo* GetRandomSeedInfo() override
 	{
 		return &RandomSeedInfo;
diff --git a/Engine/Source/Runtime/Engine/Classes/Particles/Location/ParticleModuleLocationWorldOffset_Seeded.h b/Engine/Source/Runtime/Engine/Classes/Particles/Location/ParticleModuleLocationWorldOffset_Seeded.h
--- a/Engine/Source/Runtime/Engine/Classes/Particles/Location/ParticleModuleLocationWorldOffset_Seeded.h
+++ b/Engine/Source/Runtime/Engine/Classes/Particles/Location/ParticleModuleLocationWorldOffset_Seeded.h
@@ -21,9 +21,6 @@ class ENGINE_API UParticleModuleLocationWorldOffset_Seeded : public UParticleMod
 
 
 	//Begin UParticleModule Interface
-	virtual void Spawn(FParticleEmitterInstance* Owner, int32 Offset, float SpawnTime, FBaseParticle* ParticleBase) override;
-	virtual uint32	RequiredBytesPerInstance() override;
-	virtual uint32	PrepPerInstanceBlock(FParticleEmitterInstance* Owner, void* InstData) override;
 	virtual FParticleRandomSeedInfo* GetRandomSeedInfo() override
 	{
 		return &RandomSeedInfo;
diff --git a/Engine/Source/Runtime/Engine/Classes/Particles/Location/ParticleModuleLocation_Seeded.h b/Engine/Source/Runtime/Engine/Classes/Particles/Location/ParticleModuleLocation_Seeded.h
--- a/Engine/Source/Runtime/Engine/Classes/Particles/Location/ParticleModuleLocation_Seeded.h
+++ b/Engine/Source/Runtime/Engine/Classes/Particles/Location/ParticleModuleLocation_Seeded.h
@@ -21,9 +21,6 @@ class ENGINE_API UParticleModuleLocation_Seeded : public UParticleModuleLocation
 
 
 	//Begin UParticleModule Interface
-	virtual void Spawn(FParticleEmitterInstance* Owner, int32 Offset, float SpawnTime, FBaseParticle* ParticleBase) override;
-	virtual uint32 RequiredBytesPerInstance() override;
-	virtual uint32 PrepPerInstanceBlock(FParticleEmitterInstance* Owner, void* InstData) override;
 	virtual FParticleRandomSeedInfo* GetRandomSeedInfo() override
 	{
 		return &RandomSeedInfo;
diff --git a/Engine/Source/Runtime/Engine/Classes/Particles/Parameter/ParticleModuleParameterDynamic_Seeded.h b/Engine/Source/Runtime/Engine/Classes/Particles/Parameter/ParticleModuleParameterDynamic_Seeded.h
--- a/Engine/Source/Runtime/Engine/Classes/Particles/Parameter/ParticleModuleParameterDynamic_Seeded.h
+++ b/Engine/Source/Runtime/Engine/Classes/Particles/Parameter/ParticleModuleParameterDynamic_Seeded.h
@@ -19,9 +19,6 @@ class UParticleModuleParameterDynamic_Seeded : public UParticleModuleParameterDy
 
 
 	//Begin UParticleModule Interface
-	virtual void Spawn(FParticleEmitterInstance* Owner, int32 Offset, float SpawnTime, FBaseParticle* ParticleBase) override;
-	virtual uint32	RequiredBytesPerInstance() override;
-	virtual uint32	PrepPerInstanceBlock(FParticleEmitterInstance* Owner, void* InstData) override;
 	virtual FParticleRandomSeedInfo* GetRandomSeedInfo() override
 	{
 		return &RandomSeedInfo;
diff --git a/Engine/Source/Runtime/Engine/Classes/Particles/ParticleEmitter.h b/Engine/Source/Runtime/Engine/Classes/Particles/ParticleEmitter.h
--- a/Engine/Source/Runtime/Engine/Classes/Particles/ParticleEmitter.h
+++ b/Engine/Source/Runtime/Engine/Classes/Particles/ParticleEmitter.h
@@ -217,6 +217,9 @@ class UParticleEmitter : public UObject
 	/** Map module pointers to their offset into the instance data.		*/
 	TMap<UParticleModule*, uint32> ModuleInstanceOffsetMap;
 
+	/** Map module pointers to their offset into the instance data.		*/
+	TMap<UParticleModule*, uint32> ModuleRandomSeedInstanceOffsetMap;
+	
 	/** Materials collected from any MeshMaterial modules */
 	TArray<class UMaterialInterface*> MeshMaterials;
 
@@ -236,6 +239,9 @@ class UParticleEmitter : public UObject
 	// Array of modules that want emitter instance data
 	TArray<UParticleModule*> ModulesNeedingInstanceData;
 
+	// Array of modules that want emitter random seed instance data
+	TArray<UParticleModule*> ModulesNeedingRandomSeedInstanceData;
+
 	/** SubUV animation asset to use for cutout geometry. */
 	class USubUVAnimation* RESTRICT SubUVAnimation;
 
diff --git a/Engine/Source/Runtime/Engine/Classes/Particles/ParticleModule.h b/Engine/Source/Runtime/Engine/Classes/Particles/ParticleModule.h
--- a/Engine/Source/Runtime/Engine/Classes/Particles/ParticleModule.h
+++ b/Engine/Source/Runtime/Engine/Classes/Particles/ParticleModule.h
@@ -545,6 +545,15 @@ class ENGINE_API UParticleModule : public UObject
 	/** Returns whether this module is used in any GPU emitters. */
 	bool IsUsedInGPUEmitter()const;
 
+	/**
+	 * Retreive the random stream that should be used for the provided instance.
+	 *
+	 * @param	Owner			The emitter instance that owns this module
+	 *
+	 * @return FRandomStream&	The random stream to use for the provided instance.
+	 */
+	FRandomStream& GetRandomStream(FParticleEmitterInstance* Owner);
+
 #if WITH_EDITOR
 	virtual void PostLoadSubobjects( FObjectInstancingGraph* OuterInstanceGraph ) override;
 
diff --git a/Engine/Source/Runtime/Engine/Classes/Particles/ParticleModuleRequired.h b/Engine/Source/Runtime/Engine/Classes/Particles/ParticleModuleRequired.h
--- a/Engine/Source/Runtime/Engine/Classes/Particles/ParticleModuleRequired.h
+++ b/Engine/Source/Runtime/Engine/Classes/Particles/ParticleModuleRequired.h
@@ -37,24 +37,48 @@ enum class EParticleUVFlipMode : uint8
 };
 
 /** Flips the sign of a particle's base size based on it's UV flip mode. */
-FORCEINLINE void AdjustParticleBaseSizeForUVFlipping(FVector& OutSize, EParticleUVFlipMode FlipMode)
+FORCEINLINE void AdjustParticleBaseSizeForUVFlipping(FVector& OutSize, EParticleUVFlipMode FlipMode, FRandomStream& InRandomStream)
 {
-	static const int32 HalfRandMax = RAND_MAX / 2;
+	static const float HalfRandMax = 0.5f;
+
 	switch (FlipMode)
 	{
-	case EParticleUVFlipMode::FlipUV:						OutSize = -OutSize;			return;
-	case EParticleUVFlipMode::FlipUOnly:					OutSize.X = -OutSize.X;		return;
-	case EParticleUVFlipMode::FlipVOnly:					OutSize.Y = -OutSize.Y;		return;
-	case EParticleUVFlipMode::RandomFlipUV:					OutSize = FMath::Rand() > HalfRandMax ? -OutSize : OutSize;			return;
-	case EParticleUVFlipMode::RandomFlipUOnly:				OutSize.X = FMath::Rand() > HalfRandMax ? -OutSize.X : OutSize.X;	return;
-	case EParticleUVFlipMode::RandomFlipVOnly:				OutSize.Y = FMath::Rand() > HalfRandMax ? -OutSize.Y : OutSize.Y;	return;
-	case EParticleUVFlipMode::RandomFlipUVIndependent:
-	{
-		OutSize.X = FMath::Rand() > HalfRandMax ? -OutSize.X : OutSize.X;		
-		OutSize.Y = FMath::Rand() > HalfRandMax ? -OutSize.Y : OutSize.Y;
-		return;
+		case EParticleUVFlipMode::None:
+			return;
+
+		case EParticleUVFlipMode::FlipUV:
+			OutSize = -OutSize;
+			return;
+		
+		case EParticleUVFlipMode::FlipUOnly:
+			OutSize.X = -OutSize.X;
+			return;
+		
+		case EParticleUVFlipMode::FlipVOnly:
+			OutSize.Y = -OutSize.Y;
+			return;
+		
+		case EParticleUVFlipMode::RandomFlipUV:
+			OutSize = InRandomStream.FRand() > HalfRandMax ? -OutSize : OutSize;
+			return;
+
+		case EParticleUVFlipMode::RandomFlipUOnly:
+			OutSize.X = InRandomStream.FRand() > HalfRandMax ? -OutSize.X : OutSize.X;
+			return;
+
+		case EParticleUVFlipMode::RandomFlipVOnly:
+			OutSize.Y = InRandomStream.FRand() > HalfRandMax ? -OutSize.Y : OutSize.Y;
+			return;
+
+		case EParticleUVFlipMode::RandomFlipUVIndependent:
+			OutSize.X = InRandomStream.FRand() > HalfRandMax ? -OutSize.X : OutSize.X;
+			OutSize.Y = InRandomStream.FRand() > HalfRandMax ? -OutSize.Y : OutSize.Y;
+			return;
+
+		default:
+			checkNoEntry();
+			break;
 	}
-	};
 }
 
 UENUM()
diff --git a/Engine/Source/Runtime/Engine/Classes/Particles/ParticleSystemComponent.h b/Engine/Source/Runtime/Engine/Classes/Particles/ParticleSystemComponent.h
--- a/Engine/Source/Runtime/Engine/Classes/Particles/ParticleSystemComponent.h
+++ b/Engine/Source/Runtime/Engine/Classes/Particles/ParticleSystemComponent.h
@@ -1133,6 +1133,9 @@ public:
 	/** Static delegate called for all systems on an activation change. */
 	static FOnSystemPreActivationChange OnSystemPreActivationChange;
 
+	/** Stream of random values to use with this component */
+	FRandomStream RandomStream;
+
 private:
 	/** In some cases the async work for this PSC can be created externally by the manager. */
 	FORCEINLINE void SetAsyncWork(FGraphEventRef& InAsyncWork) { AsyncWork = InAsyncWork; }
diff --git a/Engine/Source/Runtime/Engine/Classes/Particles/Rotation/ParticleModuleMeshRotation_Seeded.h b/Engine/Source/Runtime/Engine/Classes/Particles/Rotation/ParticleModuleMeshRotation_Seeded.h
--- a/Engine/Source/Runtime/Engine/Classes/Particles/Rotation/ParticleModuleMeshRotation_Seeded.h
+++ b/Engine/Source/Runtime/Engine/Classes/Particles/Rotation/ParticleModuleMeshRotation_Seeded.h
@@ -21,9 +21,6 @@ class UParticleModuleMeshRotation_Seeded : public UParticleModuleMeshRotation
 
 
 	//Begin UParticleModule Interface
-	virtual void Spawn(FParticleEmitterInstance* Owner, int32 Offset, float SpawnTime, FBaseParticle* ParticleBase) override;
-	virtual uint32	RequiredBytesPerInstance() override;
-	virtual uint32	PrepPerInstanceBlock(FParticleEmitterInstance* Owner, void* InstData) override;
 	virtual FParticleRandomSeedInfo* GetRandomSeedInfo() override
 	{
 		return &RandomSeedInfo;
diff --git a/Engine/Source/Runtime/Engine/Classes/Particles/Rotation/ParticleModuleRotation_Seeded.h b/Engine/Source/Runtime/Engine/Classes/Particles/Rotation/ParticleModuleRotation_Seeded.h
--- a/Engine/Source/Runtime/Engine/Classes/Particles/Rotation/ParticleModuleRotation_Seeded.h
+++ b/Engine/Source/Runtime/Engine/Classes/Particles/Rotation/ParticleModuleRotation_Seeded.h
@@ -21,9 +21,6 @@ class UParticleModuleRotation_Seeded : public UParticleModuleRotation
 
 
 	//Begin UParticleModule Interface
-	virtual void Spawn(FParticleEmitterInstance* Owner, int32 Offset, float SpawnTime, FBaseParticle* ParticleBase) override;
-	virtual uint32	RequiredBytesPerInstance() override;
-	virtual uint32	PrepPerInstanceBlock(FParticleEmitterInstance* Owner, void* InstData) override;
 	virtual FParticleRandomSeedInfo* GetRandomSeedInfo() override
 	{
 		return &RandomSeedInfo;
diff --git a/Engine/Source/Runtime/Engine/Classes/Particles/RotationRate/ParticleModuleMeshRotationRate_Seeded.h b/Engine/Source/Runtime/Engine/Classes/Particles/RotationRate/ParticleModuleMeshRotationRate_Seeded.h
--- a/Engine/Source/Runtime/Engine/Classes/Particles/RotationRate/ParticleModuleMeshRotationRate_Seeded.h
+++ b/Engine/Source/Runtime/Engine/Classes/Particles/RotationRate/ParticleModuleMeshRotationRate_Seeded.h
@@ -21,9 +21,6 @@ class UParticleModuleMeshRotationRate_Seeded : public UParticleModuleMeshRotatio
 
 
 	//Begin UParticleModule Interface
-	virtual void Spawn(FParticleEmitterInstance* Owner, int32 Offset, float SpawnTime, FBaseParticle* ParticleBase) override;
-	virtual uint32	RequiredBytesPerInstance() override;
-	virtual uint32	PrepPerInstanceBlock(FParticleEmitterInstance* Owner, void* InstData) override;
 	virtual FParticleRandomSeedInfo* GetRandomSeedInfo() override
 	{
 		return &RandomSeedInfo;
diff --git a/Engine/Source/Runtime/Engine/Classes/Particles/RotationRate/ParticleModuleRotationRate_Seeded.h b/Engine/Source/Runtime/Engine/Classes/Particles/RotationRate/ParticleModuleRotationRate_Seeded.h
--- a/Engine/Source/Runtime/Engine/Classes/Particles/RotationRate/ParticleModuleRotationRate_Seeded.h
+++ b/Engine/Source/Runtime/Engine/Classes/Particles/RotationRate/ParticleModuleRotationRate_Seeded.h
@@ -21,9 +21,6 @@ class UParticleModuleRotationRate_Seeded : public UParticleModuleRotationRate
 
 
 	//Begin UParticleModule Interface
-	virtual void Spawn(FParticleEmitterInstance* Owner, int32 Offset, float SpawnTime, FBaseParticle* ParticleBase) override;
-	virtual uint32	RequiredBytesPerInstance() override;
-	virtual uint32	PrepPerInstanceBlock(FParticleEmitterInstance* Owner, void* InstData) override;
 	virtual FParticleRandomSeedInfo* GetRandomSeedInfo() override
 	{
 		return &RandomSeedInfo;
diff --git a/Engine/Source/Runtime/Engine/Classes/Particles/Size/ParticleModuleSize_Seeded.h b/Engine/Source/Runtime/Engine/Classes/Particles/Size/ParticleModuleSize_Seeded.h
--- a/Engine/Source/Runtime/Engine/Classes/Particles/Size/ParticleModuleSize_Seeded.h
+++ b/Engine/Source/Runtime/Engine/Classes/Particles/Size/ParticleModuleSize_Seeded.h
@@ -21,9 +21,6 @@ class UParticleModuleSize_Seeded : public UParticleModuleSize
 
 
 	//Begin UParticleModule Interface
-	virtual void Spawn(FParticleEmitterInstance* Owner, int32 Offset, float SpawnTime, FBaseParticle* ParticleBase) override;
-	virtual uint32	RequiredBytesPerInstance() override;
-	virtual uint32	PrepPerInstanceBlock(FParticleEmitterInstance* Owner, void* InstData) override;
 	virtual FParticleRandomSeedInfo* GetRandomSeedInfo() override
 	{
 		return &RandomSeedInfo;
diff --git a/Engine/Source/Runtime/Engine/Classes/Particles/Velocity/ParticleModuleVelocity_Seeded.h b/Engine/Source/Runtime/Engine/Classes/Particles/Velocity/ParticleModuleVelocity_Seeded.h
--- a/Engine/Source/Runtime/Engine/Classes/Particles/Velocity/ParticleModuleVelocity_Seeded.h
+++ b/Engine/Source/Runtime/Engine/Classes/Particles/Velocity/ParticleModuleVelocity_Seeded.h
@@ -21,9 +21,6 @@ class UParticleModuleVelocity_Seeded : public UParticleModuleVelocity
 
 
 	//Begin UParticleModule Interface
-	virtual void Spawn(FParticleEmitterInstance* Owner, int32 Offset, float SpawnTime, FBaseParticle* ParticleBase) override;
-	virtual uint32	RequiredBytesPerInstance() override;
-	virtual uint32	PrepPerInstanceBlock(FParticleEmitterInstance* Owner, void* InstData) override;
 	virtual FParticleRandomSeedInfo* GetRandomSeedInfo() override
 	{
 		return &RandomSeedInfo;
diff --git a/Engine/Source/Runtime/Engine/Private/Particles/ParticleBeamModules.cpp b/Engine/Source/Runtime/Engine/Private/Particles/ParticleBeamModules.cpp
--- a/Engine/Source/Runtime/Engine/Private/Particles/ParticleBeamModules.cpp
+++ b/Engine/Source/Runtime/Engine/Private/Particles/ParticleBeamModules.cpp
@@ -1363,6 +1363,8 @@ void UParticleModuleBeamNoise::Spawn(FParticleEmitterInstance* Owner, int32 Offs
 	UParticleSystemComponent*		Component	= Owner->Component;
 	UParticleModuleTypeDataBeam2*	BeamTD		= BeamInst->BeamTypeData;
 
+	FRandomStream& RandomStream = GetRandomStream(Owner);
+
 	SPAWN_INIT;
 
 	FBeam2TypeDataPayload*	BeamData			= NULL;
@@ -1395,7 +1397,7 @@ void UParticleModuleBeamNoise::Spawn(FParticleEmitterInstance* Owner, int32 Offs
 	int32 CalcFreq = Frequency;
 	if (Frequency_LowRange > 0)
 	{
-		CalcFreq = FMath::TruncToInt((FMath::SRand() * (Frequency - Frequency_LowRange)) + Frequency_LowRange);
+		CalcFreq = FMath::TruncToInt((RandomStream.FRand() * (Frequency - Frequency_LowRange)) + Frequency_LowRange);
 	}
 	BEAM2_TYPEDATA_SETFREQUENCY(BeamData->Lock_Max_NumNoisePoints, CalcFreq);
 	
diff --git a/Engine/Source/Runtime/Engine/Private/Particles/ParticleComponents.cpp b/Engine/Source/Runtime/Engine/Private/Particles/ParticleComponents.cpp
--- a/Engine/Source/Runtime/Engine/Private/Particles/ParticleComponents.cpp
+++ b/Engine/Source/Runtime/Engine/Private/Particles/ParticleComponents.cpp
@@ -1699,7 +1699,9 @@ void UParticleEmitter::CacheEmitterModuleInfo()
 	LockAxisFlags = EPAL_NONE;
 	ModuleOffsetMap.Empty();
 	ModuleInstanceOffsetMap.Empty();
+	ModuleRandomSeedInstanceOffsetMap.Empty();
 	ModulesNeedingInstanceData.Empty();
+	ModulesNeedingRandomSeedInstanceData.Empty();
 	MeshMaterials.Empty();
 	DynamicParameterDataOffset = 0;
 	LightDataOffset = 0;
@@ -1789,6 +1791,25 @@ void UParticleEmitter::CacheEmitterModuleInfo()
 				}
 				ReqInstanceBytes += TempInstanceBytes;
 			}
+
+			// Add space for per instance random seed value if required
+			if (FApp::bUseFixedSeed || ParticleModule->bSupportsRandomSeed)
+			{
+				// Add the high-lodlevel offset to the lookup map
+				ModuleRandomSeedInstanceOffsetMap.Add(ParticleModule, ReqInstanceBytes);
+				// Remember that this module has emitter-instance data
+				ModulesNeedingRandomSeedInstanceData.Add(ParticleModule);
+
+				// Add all the other LODLevel modules, using the same offset.
+				// This removes the need to always also grab the HighestLODLevel pointer.
+				for (int32 LODIdx = 1; LODIdx < LODLevels.Num(); LODIdx++)
+				{
+					UParticleLODLevel* CurLODLevel = LODLevels[LODIdx];
+					ModuleRandomSeedInstanceOffsetMap.Add(CurLODLevel->Modules[ModuleIdx], ReqInstanceBytes);
+				}
+
+				ReqInstanceBytes += sizeof(FParticleRandomSeedInstancePayload);
+			}
 		}
 
 		if (ParticleModule->IsA(UParticleModuleOrientationAxisLock::StaticClass()))
@@ -3311,6 +3332,8 @@ UParticleSystemComponent::UParticleSystemComponent(const FObjectInitializer& Obj
 	bResetOnDetach = false;
 	OldPosition = FVector(0.0f, 0.0f, 0.0f);
 
+	RandomStream.Initialize(FApp::bUseFixedSeed ? GetFName() : NAME_None);
+
 	PartSysVelocity = FVector(0.0f, 0.0f, 0.0f);
 
 	WarmupTime = 0.0f;
@@ -6528,7 +6551,7 @@ void UParticleSystemComponent::InitializeSystem()
 
 			if( Template->bUseDelayRange )
 			{
-				const float	Rand = FMath::FRand();
+				const float	Rand = RandomStream.FRand();
 				EmitterDelay	 = Template->DelayLow + ((Template->Delay - Template->DelayLow) * Rand);
 			}
 		}
@@ -7187,8 +7210,7 @@ bool UParticleSystemComponent::GetFloatParameter(const FName InName,float& OutFl
 			}
 			else if (Param.ParamType == PSPT_ScalarRand)
 			{
-				// check(IsInGameThread()); this isn't exactly cool to call from multiple threads, but it isn't terrible.
-				OutFloat = Param.Scalar + (Param.Scalar_Low - Param.Scalar) * FMath::SRand();
+				OutFloat = Param.Scalar + (Param.Scalar_Low - Param.Scalar) * RandomStream.FRand();
 				return true;
 			}
 		}
@@ -7219,9 +7241,7 @@ bool UParticleSystemComponent::GetVectorParameter(const FName InName,FVector& Ou
 			}
 			else if (Param.ParamType == PSPT_VectorRand)
 			{
-				check(IsInGameThread());
-				FVector RandValue(FMath::SRand(), FMath::SRand(), FMath::SRand());
-				OutVector = Param.Vector + (Param.Vector_Low - Param.Vector) * RandValue;
+				OutVector = Param.Vector + (Param.Vector_Low - Param.Vector) * RandomStream.VRand();
 				return true;
 			}
 		}
@@ -7251,9 +7271,7 @@ bool UParticleSystemComponent::GetAnyVectorParameter(const FName InName,FVector&
 			}
 			if (Param.ParamType == PSPT_VectorRand)
 			{
-				//check(IsInGameThread());
-				FVector RandValue(FMath::SRand(), FMath::SRand(), FMath::SRand());
-				OutVector = Param.Vector + (Param.Vector_Low - Param.Vector) * RandValue;
+				OutVector = Param.Vector + (Param.Vector_Low - Param.Vector) * RandomStream.VRand();
 				return true;
 			}
 			if (Param.ParamType == PSPT_Scalar)
@@ -7264,8 +7282,7 @@ bool UParticleSystemComponent::GetAnyVectorParameter(const FName InName,FVector&
 			}
 			if (Param.ParamType == PSPT_ScalarRand)
 			{
-				// check(IsInGameThread()); this isn't exactly cool to call from multiple threads, but it isn't terrible.
-				float OutFloat = Param.Scalar + (Param.Scalar_Low - Param.Scalar) * FMath::SRand();
+				float OutFloat = Param.Scalar + (Param.Scalar_Low - Param.Scalar) * RandomStream.FRand();
 				OutVector = FVector(OutFloat, OutFloat, OutFloat);
 				return true;
 			}
diff --git a/Engine/Source/Runtime/Engine/Private/Particles/ParticleEmitterInstances.cpp b/Engine/Source/Runtime/Engine/Private/Particles/ParticleEmitterInstances.cpp
--- a/Engine/Source/Runtime/Engine/Private/Particles/ParticleEmitterInstances.cpp
+++ b/Engine/Source/Runtime/Engine/Private/Particles/ParticleEmitterInstances.cpp
@@ -528,12 +528,21 @@ void FParticleEmitterInstance::Init()
     
 	    for (UParticleModule* ParticleModule : SpriteTemplate->ModulesNeedingInstanceData)
 	    {
-		    check(ParticleModule);
-		    uint8* PrepInstData = GetModuleInstanceData(ParticleModule);
-			    check(PrepInstData != nullptr); // Shouldn't be in the list if it doesn't have data
-			    ParticleModule->PrepPerInstanceBlock(this, (void*)PrepInstData);
+			check(ParticleModule);
+			uint8* PrepInstData = GetModuleInstanceData(ParticleModule);
+			check(PrepInstData != nullptr); // Shouldn't be in the list if it doesn't have data
+			ParticleModule->PrepPerInstanceBlock(this, (void*)PrepInstData);
 	    }
-    
+
+		for (UParticleModule* ParticleModule : SpriteTemplate->ModulesNeedingRandomSeedInstanceData)
+		{
+			check(ParticleModule);
+			FParticleRandomSeedInstancePayload* SeedInstancePayload = GetModuleRandomSeedInstanceData(ParticleModule);
+			check(SeedInstancePayload != nullptr); // Shouldn't be in the list if it doesn't have data
+			FParticleRandomSeedInfo* RandomSeedInfo = ParticleModule->GetRandomSeedInfo();
+			ParticleModule->PrepRandomSeedInstancePayload(this, SeedInstancePayload, RandomSeedInfo ? *RandomSeedInfo : FParticleRandomSeedInfo());
+		}
+
 	    // Offset into emitter specific payload (e.g. TrailComponent requires extra bytes).
 	    PayloadOffset = ParticleSize;
 	    
@@ -544,7 +553,7 @@ void FParticleEmitterInstance::Init()
 	    ParticleSize = Align(ParticleSize, 16);
     
 	    // E.g. trail emitters store trailing particles directly after leading one.
-	    ParticleStride			= CalculateParticleStride(ParticleSize);
+	    ParticleStride = CalculateParticleStride(ParticleSize);
 	}
 
 	// Setup the emitter instance material array...
@@ -1555,9 +1564,25 @@ uint8* FParticleEmitterInstance::GetModuleInstanceData(UParticleModule* Module)
 		if (Offset)
 		{
 			check(*Offset < (uint32)InstancePayloadSize);
-				return &(InstanceData[*Offset]);
-			}
+			return &(InstanceData[*Offset]);
 		}
+	}
+	return NULL;
+}
+
+/** Get pointer to emitter instance random seed payload data for a particular module */
+FParticleRandomSeedInstancePayload* FParticleEmitterInstance::GetModuleRandomSeedInstanceData(UParticleModule* Module)
+{
+	// If there is instance data present, look up the modules offset
+	if (InstanceData)
+	{
+		uint32* Offset = SpriteTemplate->ModuleRandomSeedInstanceOffsetMap.Find(Module);
+		if (Offset)
+		{
+			check(*Offset < (uint32)InstancePayloadSize);
+			return (FParticleRandomSeedInstancePayload*)&(InstanceData[*Offset]);
+		}
+	}
 	return NULL;
 }
 
@@ -1620,6 +1645,8 @@ float FParticleEmitterInstance::GetCurrentBurstRateOffset(float& DeltaTime, int3
 	UParticleLODLevel* LODLevel	= GetCurrentLODLevelChecked();
 	if (LODLevel->SpawnModule->BurstList.Num() > 0)
 	{
+		FRandomStream& RandomStream = LODLevel->SpawnModule->GetRandomStream(this);
+
 		// For each burst in the list
 		for (int32 BurstIdx = 0; BurstIdx < LODLevel->SpawnModule->BurstList.Num(); BurstIdx++)
 		{
@@ -1644,7 +1671,7 @@ float FParticleEmitterInstance::GetCurrentBurstRateOffset(float& DeltaTime, int3
 							int32 Count = BurstEntry->Count;
 							if (BurstEntry->CountLow > -1)
 							{
-								Count = BurstEntry->CountLow + FMath::RoundToInt(FMath::SRand() * (float)(BurstEntry->Count - BurstEntry->CountLow));
+								Count = RandomStream.RandRange(BurstEntry->CountLow, BurstEntry->Count);
 							}
 							// Take in to account scale.
 							float Scale = LODLevel->SpawnModule->BurstScale.GetValue(EmitterTime, Component);
@@ -2630,10 +2657,12 @@ void FParticleEmitterInstance::SetupEmitterDuration()
 		UParticleLODLevel* TempLOD = SpriteTemplate->LODLevels[LODIndex];
 		UParticleModuleRequired* RequiredModule = TempLOD->RequiredModule;
 
+		FRandomStream& RandomStream = RequiredModule->GetRandomStream(this);
+
 		CurrentDelay = RequiredModule->EmitterDelay + Component->EmitterDelay;
 		if (RequiredModule->bEmitterDelayUseRange)
 		{
-			const float	Rand	= FMath::FRand();
+			const float	Rand	= RandomStream.FRand();
 			CurrentDelay	    = RequiredModule->EmitterDelayLow + 
 				((RequiredModule->EmitterDelay - RequiredModule->EmitterDelayLow) * Rand) + Component->EmitterDelay;
 		}
@@ -2641,7 +2670,7 @@ void FParticleEmitterInstance::SetupEmitterDuration()
 
 		if (RequiredModule->bEmitterDurationUseRange)
 		{
-			const float	Rand		= FMath::FRand();
+			const float	Rand		= RandomStream.FRand();
 			const float	Duration	= RequiredModule->EmitterDurationLow + 
 				((RequiredModule->EmitterDuration - RequiredModule->EmitterDurationLow) * Rand);
 			EmitterDurations[TempLOD->Level] = Duration + CurrentDelay;
diff --git a/Engine/Source/Runtime/Engine/Private/Particles/ParticleGpuSimulation.cpp b/Engine/Source/Runtime/Engine/Private/Particles/ParticleGpuSimulation.cpp
--- a/Engine/Source/Runtime/Engine/Private/Particles/ParticleGpuSimulation.cpp
+++ b/Engine/Source/Runtime/Engine/Private/Particles/ParticleGpuSimulation.cpp
@@ -3548,7 +3548,7 @@ FGPUSpriteParticleEmitterInstance(FFXSystem* InFXSystem, FGPUSpriteEmitterInfo&
 		check(AllocatedTiles.Num() == TileTimeOfDeath.Num());
 		FreeParticlesInTile = 0;
 
-		RandomStream.Initialize( FMath::Rand() );
+		RandomStream.Initialize(Component->RandomStream.FRand());
 		EmitterInstRandom = RandomStream.GetFraction();
 
 		FParticleSimulationResources* ParticleSimulationResources = FXSystem->GetParticleSimulationResources();
diff --git a/Engine/Source/Runtime/Engine/Private/Particles/ParticleModules.cpp b/Engine/Source/Runtime/Engine/Private/Particles/ParticleModules.cpp
--- a/Engine/Source/Runtime/Engine/Private/Particles/ParticleModules.cpp
+++ b/Engine/Source/Runtime/Engine/Private/Particles/ParticleModules.cpp
@@ -13,6 +13,7 @@
 #include "UObject/RenderingObjectVersion.h"
 #include "UObject/UObjectHash.h"
 #include "UObject/Package.h"
+#include "Misc/App.h"
 #include "GameFramework/WorldSettings.h"
 #include "Particles/ParticleSystem.h"
 #include "ParticleHelper.h"
@@ -667,7 +668,7 @@ uint32 UParticleModule::PrepRandomSeedInstancePayload(FParticleEmitterInstance*
 {
 	if (InRandSeedPayload != NULL)
 	{
-		FMemory::Memzero(InRandSeedPayload, sizeof(FParticleRandomSeedInstancePayload));
+		new(InRandSeedPayload) FParticleRandomSeedInstancePayload();
 
 		// See if the parameter is set on the instance...
 		if ((Owner != NULL) && (Owner->Component != NULL) && (InRandSeedInfo.bGetSeedFromInstance == true))
@@ -701,7 +702,7 @@ uint32 UParticleModule::PrepRandomSeedInstancePayload(FParticleEmitterInstance*
 		{
 			if (InRandSeedInfo.bRandomlySelectSeedArray)
 			{
-				int32 Index = FMath::RandRange(0, InRandSeedInfo.RandomSeeds.Num() - 1);
+				int32 Index = Owner->Component->RandomStream.RandHelper(InRandSeedInfo.RandomSeeds.Num());
 				InRandSeedPayload->RandomStream.Initialize(InRandSeedInfo.RandomSeeds[Index]);
 			}
 			else
@@ -710,6 +711,11 @@ uint32 UParticleModule::PrepRandomSeedInstancePayload(FParticleEmitterInstance*
 			}
 			return 0;
 		}
+		else if (FApp::bUseFixedSeed)
+		{
+			InRandSeedPayload->RandomStream.Initialize(GetFName());
+			return 0;
+		}
 	}
 	return 0xffffffff;
 }
@@ -788,6 +794,13 @@ bool UParticleModule::IsUsedInGPUEmitter()const
 	return false;
 }
 
+FRandomStream& UParticleModule::GetRandomStream(FParticleEmitterInstance* Owner)
+{
+	FParticleRandomSeedInstancePayload* Payload = Owner->GetModuleRandomSeedInstanceData(this);
+	FRandomStream& RandomStream = (Payload != nullptr) ? Payload->RandomStream : Owner->Component->RandomStream;
+	return RandomStream;
+}
+
 #if WITH_EDITOR
 void UParticleModule::SetTransactionFlag()
 {
@@ -1403,7 +1416,7 @@ void UParticleModuleMeshRotation::PostEditChangeProperty(FPropertyChangedEvent&
 
 void UParticleModuleMeshRotation::Spawn(FParticleEmitterInstance* Owner, int32 Offset, float SpawnTime, FBaseParticle* ParticleBase)
 {
-	SpawnEx(Owner, Offset, SpawnTime, NULL, ParticleBase);
+	SpawnEx(Owner, Offset, SpawnTime, &GetRandomStream(Owner), ParticleBase);
 }
 
 
@@ -1444,27 +1457,11 @@ UParticleModuleMeshRotation_Seeded::UParticleModuleMeshRotation_Seeded(const FOb
 	bRequiresLoopingNotification = true;
 }
 
-void UParticleModuleMeshRotation_Seeded::Spawn(FParticleEmitterInstance* Owner, int32 Offset, float SpawnTime, FBaseParticle* ParticleBase)
-{
-	FParticleRandomSeedInstancePayload* Payload = (FParticleRandomSeedInstancePayload*)(Owner->GetModuleInstanceData(this));
-	SpawnEx(Owner, Offset, SpawnTime, (Payload != NULL) ? &(Payload->RandomStream) : NULL, ParticleBase);
-}
-
-uint32 UParticleModuleMeshRotation_Seeded::RequiredBytesPerInstance()
-{
-	return RandomSeedInfo.GetInstancePayloadSize();
-}
-
-uint32 UParticleModuleMeshRotation_Seeded::PrepPerInstanceBlock(FParticleEmitterInstance* Owner, void* InstData)
-{
-	return PrepRandomSeedInstancePayload(Owner, (FParticleRandomSeedInstancePayload*)InstData, RandomSeedInfo);
-}
-
 void UParticleModuleMeshRotation_Seeded::EmitterLoopingNotify(FParticleEmitterInstance* Owner)
 {
 	if (RandomSeedInfo.bResetSeedOnEmitterLooping == true)
 	{
-		FParticleRandomSeedInstancePayload* Payload = (FParticleRandomSeedInstancePayload*)(Owner->GetModuleInstanceData(this));
+		FParticleRandomSeedInstancePayload* Payload = Owner->GetModuleRandomSeedInstanceData(this);
 		PrepRandomSeedInstancePayload(Owner, Payload, RandomSeedInfo);
 	}
 }
@@ -1517,7 +1514,7 @@ void UParticleModuleMeshRotationRate::PostEditChangeProperty(FPropertyChangedEve
 
 void UParticleModuleMeshRotationRate::Spawn(FParticleEmitterInstance* Owner, int32 Offset, float SpawnTime, FBaseParticle* ParticleBase)
 {
-	SpawnEx(Owner, Offset, SpawnTime, NULL, ParticleBase);
+	SpawnEx(Owner, Offset, SpawnTime, &GetRandomStream(Owner), ParticleBase);
 }
 
 void UParticleModuleMeshRotationRate::SpawnEx(FParticleEmitterInstance* Owner, int32 Offset, float SpawnTime, struct FRandomStream* InRandomStream, FBaseParticle* ParticleBase)
@@ -1562,28 +1559,11 @@ UParticleModuleMeshRotationRate_Seeded::UParticleModuleMeshRotationRate_Seeded(c
 	bRequiresLoopingNotification = true;
 }
 
-void UParticleModuleMeshRotationRate_Seeded::Spawn(FParticleEmitterInstance* Owner, int32 Offset, float SpawnTime, FBaseParticle* ParticleBase)
-{
-	FParticleRandomSeedInstancePayload* Payload = (FParticleRandomSeedInstancePayload*)(Owner->GetModuleInstanceData(this));
-	SpawnEx(Owner, Offset, SpawnTime, (Payload != NULL) ? &(Payload->RandomStream) : NULL, ParticleBase);
-}
-
-uint32 UParticleModuleMeshRotationRate_Seeded::RequiredBytesPerInstance()
-{
-	return RandomSeedInfo.GetInstancePayloadSize();
-}
-
-
-uint32 UParticleModuleMeshRotationRate_Seeded::PrepPerInstanceBlock(FParticleEmitterInstance* Owner, void* InstData)
-{
-	return PrepRandomSeedInstancePayload(Owner, (FParticleRandomSeedInstancePayload*)InstData, RandomSeedInfo);
-}
-
 void UParticleModuleMeshRotationRate_Seeded::EmitterLoopingNotify(FParticleEmitterInstance* Owner)
 {
 	if (RandomSeedInfo.bResetSeedOnEmitterLooping == true)
 	{
-		FParticleRandomSeedInstancePayload* Payload = (FParticleRandomSeedInstancePayload*)(Owner->GetModuleInstanceData(this));
+		FParticleRandomSeedInstancePayload* Payload = Owner->GetModuleRandomSeedInstanceData(this);
 		PrepRandomSeedInstancePayload(Owner, Payload, RandomSeedInfo);
 	}
 }
@@ -1809,7 +1789,7 @@ void UParticleModuleRotation::PostEditChangeProperty(FPropertyChangedEvent& Prop
 
 void UParticleModuleRotation::Spawn(FParticleEmitterInstance* Owner, int32 Offset, float SpawnTime, FBaseParticle* ParticleBase)
 {
-	SpawnEx(Owner, Offset, SpawnTime, NULL, ParticleBase);
+	SpawnEx(Owner, Offset, SpawnTime, &GetRandomStream(Owner), ParticleBase);
 }
 
 void UParticleModuleRotation::SpawnEx(FParticleEmitterInstance* Owner, int32 Offset, float SpawnTime, struct FRandomStream* InRandomStream, FBaseParticle* ParticleBase)
@@ -1831,27 +1811,11 @@ UParticleModuleRotation_Seeded::UParticleModuleRotation_Seeded(const FObjectInit
 	bRequiresLoopingNotification = true;
 }
 
-void UParticleModuleRotation_Seeded::Spawn(FParticleEmitterInstance* Owner, int32 Offset, float SpawnTime, FBaseParticle* ParticleBase)
-{
-	FParticleRandomSeedInstancePayload* Payload = (FParticleRandomSeedInstancePayload*)(Owner->GetModuleInstanceData(this));
-	SpawnEx(Owner, Offset, SpawnTime, (Payload != NULL) ? &(Payload->RandomStream) : NULL, ParticleBase);
-}
-
-uint32 UParticleModuleRotation_Seeded::RequiredBytesPerInstance()
-{
-	return RandomSeedInfo.GetInstancePayloadSize();
-}
-
-uint32 UParticleModuleRotation_Seeded::PrepPerInstanceBlock(FParticleEmitterInstance* Owner, void* InstData)
-{
-	return PrepRandomSeedInstancePayload(Owner, (FParticleRandomSeedInstancePayload*)InstData, RandomSeedInfo);
-}
-
 void UParticleModuleRotation_Seeded::EmitterLoopingNotify(FParticleEmitterInstance* Owner)
 {
 	if (RandomSeedInfo.bResetSeedOnEmitterLooping == true)
 	{
-		FParticleRandomSeedInstancePayload* Payload = (FParticleRandomSeedInstancePayload*)(Owner->GetModuleInstanceData(this));
+		FParticleRandomSeedInstancePayload* Payload = Owner->GetModuleRandomSeedInstanceData(this);
 		PrepRandomSeedInstancePayload(Owner, Payload, RandomSeedInfo);
 	}
 }
@@ -1904,7 +1868,7 @@ void UParticleModuleRotationRate::PostEditChangeProperty(FPropertyChangedEvent&
 
 void UParticleModuleRotationRate::Spawn(FParticleEmitterInstance* Owner, int32 Offset, float SpawnTime, FBaseParticle* ParticleBase)
 {
-	SpawnEx(Owner, Offset, SpawnTime, NULL, ParticleBase);
+	SpawnEx(Owner, Offset, SpawnTime, &GetRandomStream(Owner), ParticleBase);
 }
 
 void UParticleModuleRotationRate::SpawnEx(FParticleEmitterInstance* Owner, int32 Offset, float SpawnTime, struct FRandomStream* InRandomStream, FBaseParticle* ParticleBase)
@@ -1940,27 +1904,11 @@ UParticleModuleRotationRate_Seeded::UParticleModuleRotationRate_Seeded(const FOb
 	bRequiresLoopingNotification = true;
 }
 
-void UParticleModuleRotationRate_Seeded::Spawn(FParticleEmitterInstance* Owner, int32 Offset, float SpawnTime, FBaseParticle* ParticleBase)
-{
-	FParticleRandomSeedInstancePayload* Payload = (FParticleRandomSeedInstancePayload*)(Owner->GetModuleInstanceData(this));
-	SpawnEx(Owner, Offset, SpawnTime, (Payload != NULL) ? &(Payload->RandomStream) : NULL, ParticleBase);
-}
-
-uint32 UParticleModuleRotationRate_Seeded::RequiredBytesPerInstance()
-{
-	return RandomSeedInfo.GetInstancePayloadSize();
-}
-
-uint32 UParticleModuleRotationRate_Seeded::PrepPerInstanceBlock(FParticleEmitterInstance* Owner, void* InstData)
-{
-	return PrepRandomSeedInstancePayload(Owner, (FParticleRandomSeedInstancePayload*)InstData, RandomSeedInfo);
-}
-
 void UParticleModuleRotationRate_Seeded::EmitterLoopingNotify(FParticleEmitterInstance* Owner)
 {
 	if (RandomSeedInfo.bResetSeedOnEmitterLooping == true)
 	{
-		FParticleRandomSeedInstancePayload* Payload = (FParticleRandomSeedInstancePayload*)(Owner->GetModuleInstanceData(this));
+		FParticleRandomSeedInstancePayload* Payload = Owner->GetModuleRandomSeedInstanceData(this);
 		PrepRandomSeedInstancePayload(Owner, Payload, RandomSeedInfo);
 	}
 }
@@ -2247,8 +2195,7 @@ float UParticleModuleSubUV::DetermineImageIndex(FParticleEmitterInstance* Owner,
 			((Particle->RelativeTime - SubUVPayload.RandomImageTime) > LODLevel->RequiredModule->RandomImageTime) ||
 			(SubUVPayload.RandomImageTime == 0.0f))
 		{
-			const float RandomNumber = FMath::SRand();
-			ImageIndex = FMath::TruncToInt(RandomNumber * TotalSubImages);
+			ImageIndex = GetRandomStream(Owner).RandHelper(TotalSubImages);
 			SubUVPayload.RandomImageTime	= Particle->RelativeTime;
 		}
 
@@ -2380,7 +2327,8 @@ void UParticleModuleSubUVMovie::Spawn(FParticleEmitterInstance* Owner, int32 Off
 			}
 			else if (StartingFrame == 0)
 			{
-				MoviePayload.Time = FMath::TruncToFloat(FMath::SRand() * (TotalSubImages-1));
+				FRandomStream& RandomStream = GetRandomStream(Owner);
+				MoviePayload.Time = FMath::TruncToFloat(RandomStream.FRand() * (TotalSubImages-1));
 			}
 
 			// Update the payload
@@ -3116,7 +3064,7 @@ void UParticleModuleLight::SpawnEx(FParticleEmitterInstance* Owner, int32 Offset
 		LightData.RadiusScale = RadiusScale.GetValue(Owner->EmitterTime, Owner->Component, InRandomStream);
 		// Exponent of 0 is interpreted by renderer as inverse squared falloff
 		LightData.LightExponent = bUseInverseSquaredFalloff ? 0 : LightExponent.GetValue(Owner->EmitterTime, Owner->Component, InRandomStream);
-		const float RandomNumber = InRandomStream ? InRandomStream->GetFraction() : FMath::SRand();
+		const float RandomNumber = InRandomStream->GetFraction();
 		LightData.bValid = RandomNumber < SpawnFraction;
 		LightData.bAffectsTranslucency = bAffectsTranslucency;
 		LightData.bHighQuality = bHighQualityLights;
@@ -3239,7 +3187,7 @@ void UParticleModuleLight::UpdateHQLight(UPointLightComponent* PointLightCompone
 
 void UParticleModuleLight::Spawn(FParticleEmitterInstance* Owner, int32 Offset, float SpawnTime, FBaseParticle* ParticleBase)
 {
-	SpawnEx(Owner, Offset, SpawnTime, NULL, ParticleBase);
+	SpawnEx(Owner, Offset, SpawnTime, &GetRandomStream(Owner), ParticleBase);
 }
 
 void UParticleModuleLight::Update(FParticleEmitterInstance* Owner, int32 Offset, float DeltaTime)
@@ -3388,27 +3336,11 @@ UParticleModuleLight_Seeded::UParticleModuleLight_Seeded(const FObjectInitialize
 	bRequiresLoopingNotification = true;
 }
 
-void UParticleModuleLight_Seeded::Spawn(FParticleEmitterInstance* Owner, int32 Offset, float SpawnTime, FBaseParticle* ParticleBase)
-{
-	FParticleRandomSeedInstancePayload* Payload = (FParticleRandomSeedInstancePayload*)(Owner->GetModuleInstanceData(this) + Super::RequiredBytesPerInstance());
-	SpawnEx(Owner, Offset, SpawnTime, (Payload != NULL) ? &(Payload->RandomStream) : NULL, ParticleBase);
-}
-
-uint32 UParticleModuleLight_Seeded::RequiredBytesPerInstance()
-{
-	return Super::RequiredBytesPerInstance() + RandomSeedInfo.GetInstancePayloadSize();
-}
-
-uint32 UParticleModuleLight_Seeded::PrepPerInstanceBlock(FParticleEmitterInstance* Owner, void* InstData)
-{
-	return PrepRandomSeedInstancePayload(Owner, (FParticleRandomSeedInstancePayload*)((uint8*)InstData + Super::RequiredBytesPerInstance()), RandomSeedInfo);
-}
-
 void UParticleModuleLight_Seeded::EmitterLoopingNotify(FParticleEmitterInstance* Owner)
 {
 	if (RandomSeedInfo.bResetSeedOnEmitterLooping == true)
 	{
-		FParticleRandomSeedInstancePayload* Payload = (FParticleRandomSeedInstancePayload*)(Owner->GetModuleInstanceData(this) + Super::RequiredBytesPerInstance());
+		FParticleRandomSeedInstancePayload* Payload = Owner->GetModuleRandomSeedInstanceData(this);
 		PrepRandomSeedInstancePayload(Owner, Payload, RandomSeedInfo);
 	}
 }
@@ -3875,7 +3807,7 @@ void UParticleModuleLifetime::CompileModule( struct FParticleEmitterBuildInfo& E
 
 void UParticleModuleLifetime::Spawn(FParticleEmitterInstance* Owner, int32 Offset, float SpawnTime, FBaseParticle* ParticleBase)
 {
-	SpawnEx(Owner, Offset, SpawnTime, NULL, ParticleBase);
+	SpawnEx(Owner, Offset, SpawnTime, &GetRandomStream(Owner), ParticleBase);
 }
 
 void UParticleModuleLifetime::SpawnEx(FParticleEmitterInstance* Owner, int32 Offset, float SpawnTime, struct FRandomStream* InRandomStream, FBaseParticle* ParticleBase)
@@ -3941,39 +3873,18 @@ UParticleModuleLifetime_Seeded::UParticleModuleLifetime_Seeded(const FObjectInit
 	bRequiresLoopingNotification = true;
 }
 
-void UParticleModuleLifetime_Seeded::Spawn(FParticleEmitterInstance* Owner, int32 Offset, float SpawnTime, FBaseParticle* ParticleBase)
-{
-	FParticleRandomSeedInstancePayload* Payload = (FParticleRandomSeedInstancePayload*)(Owner->GetModuleInstanceData(this));
-	SpawnEx(Owner, Offset, SpawnTime, (Payload != NULL) ? &(Payload->RandomStream) : NULL, ParticleBase);
-}
-
-uint32 UParticleModuleLifetime_Seeded::RequiredBytesPerInstance()
-{
-	return RandomSeedInfo.GetInstancePayloadSize();
-}
-
-uint32 UParticleModuleLifetime_Seeded::PrepPerInstanceBlock(FParticleEmitterInstance* Owner, void* InstData)
-{
-	return PrepRandomSeedInstancePayload(Owner, (FParticleRandomSeedInstancePayload*)InstData, RandomSeedInfo);
-}
-
 void UParticleModuleLifetime_Seeded::EmitterLoopingNotify(FParticleEmitterInstance* Owner)
 {
 	if (RandomSeedInfo.bResetSeedOnEmitterLooping == true)
 	{
-		FParticleRandomSeedInstancePayload* Payload = (FParticleRandomSeedInstancePayload*)(Owner->GetModuleInstanceData(this));
+		FParticleRandomSeedInstancePayload* Payload = Owner->GetModuleRandomSeedInstanceData(this);
 		PrepRandomSeedInstancePayload(Owner, Payload, RandomSeedInfo);
 	}
 }
 
 float UParticleModuleLifetime_Seeded::GetLifetimeValue(FParticleEmitterInstance* Owner, float InTime, UObject* Data )
 {
-	FParticleRandomSeedInstancePayload* Payload = (FParticleRandomSeedInstancePayload*)(Owner->GetModuleInstanceData(this));
-	if (Payload != NULL)
-	{		
-		return Lifetime.GetValue(InTime, Data, &(Payload->RandomStream));
-	}
-	return UParticleModuleLifetime::GetLifetimeValue(Owner, InTime, Data);
+	return Lifetime.GetValue(InTime, Data, &GetRandomStream(Owner));
 }
 
 /*-----------------------------------------------------------------------------
@@ -4207,7 +4118,7 @@ void UParticleModuleAttractorParticle::Spawn(FParticleEmitterInstance* Owner, in
 			switch (SelectionMethod)
 			{
 			case EAPSM_Random:
-				LastSelIndex		= FMath::TruncToInt(FMath::SRand() * AttractorEmitterInst->ActiveParticles);
+				LastSelIndex		= GetRandomStream(Owner).RandHelper(AttractorEmitterInst->ActiveParticles);
 				Data.SourceIndex	= LastSelIndex;
 				break;
 			case EAPSM_Sequential:
diff --git a/Engine/Source/Runtime/Engine/Private/Particles/ParticleModules_Color.cpp b/Engine/Source/Runtime/Engine/Private/Particles/ParticleModules_Color.cpp
--- a/Engine/Source/Runtime/Engine/Private/Particles/ParticleModules_Color.cpp
+++ b/Engine/Source/Runtime/Engine/Private/Particles/ParticleModules_Color.cpp
@@ -155,7 +155,7 @@ bool UParticleModuleColor::AddModuleCurvesToEditor(UInterpCurveEdSetup* EdSetup,
 
 void UParticleModuleColor::Spawn(FParticleEmitterInstance* Owner, int32 Offset, float SpawnTime, FBaseParticle* ParticleBase)
 {
-	SpawnEx(Owner, Offset, SpawnTime, NULL, ParticleBase);
+	SpawnEx(Owner, Offset, SpawnTime, &GetRandomStream(Owner), ParticleBase);
 }
 
 
@@ -192,27 +192,11 @@ UParticleModuleColor_Seeded::UParticleModuleColor_Seeded(const FObjectInitialize
 	bRequiresLoopingNotification = true;
 }
 
-void UParticleModuleColor_Seeded::Spawn(FParticleEmitterInstance* Owner, int32 Offset, float SpawnTime, FBaseParticle* ParticleBase)
-{
-	FParticleRandomSeedInstancePayload* Payload = (FParticleRandomSeedInstancePayload*)(Owner->GetModuleInstanceData(this));
-	SpawnEx(Owner, Offset, SpawnTime, (Payload != NULL) ? &(Payload->RandomStream) : NULL, ParticleBase);
-}
-
-uint32 UParticleModuleColor_Seeded::RequiredBytesPerInstance()
-{
-	return RandomSeedInfo.GetInstancePayloadSize();
-}
-
-uint32 UParticleModuleColor_Seeded::PrepPerInstanceBlock(FParticleEmitterInstance* Owner, void* InstData)
-{
-	return PrepRandomSeedInstancePayload(Owner, (FParticleRandomSeedInstancePayload*)InstData, RandomSeedInfo);
-}
-
 void UParticleModuleColor_Seeded::EmitterLoopingNotify(FParticleEmitterInstance* Owner)
 {
 	if (RandomSeedInfo.bResetSeedOnEmitterLooping == true)
 	{
-		FParticleRandomSeedInstancePayload* Payload = (FParticleRandomSeedInstancePayload*)(Owner->GetModuleInstanceData(this));
+		FParticleRandomSeedInstancePayload* Payload = Owner->GetModuleRandomSeedInstanceData(this);
 		PrepRandomSeedInstancePayload(Owner, Payload, RandomSeedInfo);
 	}
 }
diff --git a/Engine/Source/Runtime/Engine/Private/Particles/ParticleModules_Location.cpp b/Engine/Source/Runtime/Engine/Private/Particles/ParticleModules_Location.cpp
--- a/Engine/Source/Runtime/Engine/Private/Particles/ParticleModules_Location.cpp
+++ b/Engine/Source/Runtime/Engine/Private/Particles/ParticleModules_Location.cpp
@@ -90,7 +90,7 @@ void UParticleModuleLocation::PostEditChangeProperty(FPropertyChangedEvent& Prop
 
 void UParticleModuleLocation::Spawn(FParticleEmitterInstance* Owner, int32 Offset, float SpawnTime, FBaseParticle* ParticleBase)
 {
-	SpawnEx(Owner, Offset, SpawnTime, NULL, ParticleBase);
+	SpawnEx(Owner, Offset, SpawnTime, &GetRandomStream(Owner), ParticleBase);
 }
 
 void UParticleModuleLocation::SpawnEx(FParticleEmitterInstance* Owner, int32 Offset, float SpawnTime, struct FRandomStream* InRandomStream, FBaseParticle* ParticleBase)
@@ -103,7 +103,7 @@ void UParticleModuleLocation::SpawnEx(FParticleEmitterInstance* Owner, int32 Off
 	// Avoid divide by zero.
 	if ((DistributeOverNPoints != 0.0f) && (DistributeOverNPoints != 1.f))
 	{
-		float RandomNum = FMath::SRand() * FMath::Fractional(Owner->EmitterTime);
+		float RandomNum = InRandomStream->FRand() * FMath::Fractional(Owner->EmitterTime);
 
 		if(RandomNum > DistributeThreshold)
 		{
@@ -113,7 +113,7 @@ void UParticleModuleLocation::SpawnEx(FParticleEmitterInstance* Owner, int32 Off
 		{
 			FVector Min, Max;
 			StartLocation.GetRange(Min, Max);
-			FVector Lerped = FMath::Lerp(Min, Max, FMath::TruncToFloat((FMath::SRand() * (DistributeOverNPoints - 1.0f)) + 0.5f)/(DistributeOverNPoints - 1.0f));
+			FVector Lerped = FMath::Lerp(Min, Max, FMath::TruncToFloat((InRandomStream->FRand() * (DistributeOverNPoints - 1.0f)) + 0.5f)/(DistributeOverNPoints - 1.0f));
 			LocationOffset.Set(Lerped.X, Lerped.Y, Lerped.Z);
 		}
 	}
@@ -187,27 +187,11 @@ UParticleModuleLocation_Seeded::UParticleModuleLocation_Seeded(const FObjectInit
 	bRequiresLoopingNotification = true;
 }
 
-void UParticleModuleLocation_Seeded::Spawn(FParticleEmitterInstance* Owner, int32 Offset, float SpawnTime, FBaseParticle* ParticleBase)
-{
-	FParticleRandomSeedInstancePayload* Payload = (FParticleRandomSeedInstancePayload*)(Owner->GetModuleInstanceData(this));
-	SpawnEx(Owner, Offset, SpawnTime, (Payload != NULL) ? &(Payload->RandomStream) : NULL, ParticleBase);
-}
-
-uint32 UParticleModuleLocation_Seeded::RequiredBytesPerInstance()
-{
-	return RandomSeedInfo.GetInstancePayloadSize();
-}
-
-uint32 UParticleModuleLocation_Seeded::PrepPerInstanceBlock(FParticleEmitterInstance* Owner, void* InstData)
-{
-	return PrepRandomSeedInstancePayload(Owner, (FParticleRandomSeedInstancePayload*)InstData, RandomSeedInfo);
-}
-
 void UParticleModuleLocation_Seeded::EmitterLoopingNotify(FParticleEmitterInstance* Owner)
 {
 	if (RandomSeedInfo.bResetSeedOnEmitterLooping == true)
 	{
-		FParticleRandomSeedInstancePayload* Payload = (FParticleRandomSeedInstancePayload*)(Owner->GetModuleInstanceData(this));
+		FParticleRandomSeedInstancePayload* Payload = Owner->GetModuleRandomSeedInstanceData(this);
 		PrepRandomSeedInstancePayload(Owner, Payload, RandomSeedInfo);
 	}
 }
@@ -251,27 +235,11 @@ UParticleModuleLocationWorldOffset_Seeded::UParticleModuleLocationWorldOffset_Se
 	bRequiresLoopingNotification = true;
 }
 
-void UParticleModuleLocationWorldOffset_Seeded::Spawn(FParticleEmitterInstance* Owner, int32 Offset, float SpawnTime, FBaseParticle* ParticleBase)
-{
-	FParticleRandomSeedInstancePayload* Payload = (FParticleRandomSeedInstancePayload*)(Owner->GetModuleInstanceData(this));
-	SpawnEx(Owner, Offset, SpawnTime, (Payload != NULL) ? &(Payload->RandomStream) : NULL, ParticleBase);
-}
-
-uint32 UParticleModuleLocationWorldOffset_Seeded::RequiredBytesPerInstance()
-{
-	return RandomSeedInfo.GetInstancePayloadSize();
-}
-
-uint32 UParticleModuleLocationWorldOffset_Seeded::PrepPerInstanceBlock(FParticleEmitterInstance* Owner, void* InstData)
-{
-	return PrepRandomSeedInstancePayload(Owner, (FParticleRandomSeedInstancePayload*)InstData, RandomSeedInfo);
-}
-
 void UParticleModuleLocationWorldOffset_Seeded::EmitterLoopingNotify(FParticleEmitterInstance* Owner)
 {
 	if (RandomSeedInfo.bResetSeedOnEmitterLooping == true)
 	{
-		FParticleRandomSeedInstancePayload* Payload = (FParticleRandomSeedInstancePayload*)(Owner->GetModuleInstanceData(this));
+		FParticleRandomSeedInstancePayload* Payload = Owner->GetModuleRandomSeedInstanceData(this);
 		PrepRandomSeedInstancePayload(Owner, Payload, RandomSeedInfo);
 	}
 }
@@ -448,6 +416,8 @@ void UParticleModuleLocationEmitter::Spawn(FParticleEmitterInstance* Owner, int3
 	bool bSourceIsInLocalSpace = LocationEmitterInst->CurrentLODLevel->RequiredModule->bUseLocalSpace;
 	bool bInLocalSpace = Owner->CurrentLODLevel->RequiredModule->bUseLocalSpace;
 
+	FRandomStream& RandomStream = GetRandomStream(Owner);
+
 	SPAWN_INIT;
 		{
 			int32 Index = 0;
@@ -456,11 +426,7 @@ void UParticleModuleLocationEmitter::Spawn(FParticleEmitterInstance* Owner, int3
 			{
 			case ELESM_Random:
 				{
-					Index = FMath::TruncToInt(FMath::SRand() * LocationEmitterInst->ActiveParticles);
-					if (Index >= LocationEmitterInst->ActiveParticles)
-					{
-						Index = LocationEmitterInst->ActiveParticles - 1;
-					}
+					Index = RandomStream.RandHelper(LocationEmitterInst->ActiveParticles);
 				}
 				break;
 			case ELESM_Sequential:
@@ -719,18 +685,9 @@ void UParticleModuleLocationPrimitiveBase::DetermineUnitDirection(FParticleEmitt
 	FVector vRand;
 
 	// Grab 3 random numbers for the axes
-	if (InRandomStream == NULL)
-	{
-		vRand.X	= FMath::SRand();
-		vRand.Y = FMath::SRand();
-		vRand.Z = FMath::SRand();
-	}
-	else
-	{
-		vRand.X	= InRandomStream->GetFraction();
-		vRand.Y = InRandomStream->GetFraction();
-		vRand.Z = InRandomStream->GetFraction();
-	}
+	vRand.X	= InRandomStream->GetFraction();
+	vRand.Y = InRandomStream->GetFraction();
+	vRand.Z = InRandomStream->GetFraction();
 
 	// Set the unit dir
 	if (Positive_X && Negative_X)
@@ -846,7 +803,7 @@ void UParticleModuleLocationPrimitiveTriangle::PostEditChangeProperty(FPropertyC
 
 void UParticleModuleLocationPrimitiveTriangle::Spawn(FParticleEmitterInstance* Owner, int32 Offset, float SpawnTime, FBaseParticle* ParticleBase)
 {
-	SpawnEx(Owner, Offset, SpawnTime, NULL, ParticleBase);
+	SpawnEx(Owner, Offset, SpawnTime, &GetRandomStream(Owner), ParticleBase);
 }
 
 void UParticleModuleLocationPrimitiveTriangle::SpawnEx(FParticleEmitterInstance* Owner, int32 Offset, float SpawnTime, struct FRandomStream* InRandomStream, FBaseParticle* ParticleBase)
@@ -868,21 +825,12 @@ void UParticleModuleLocationPrimitiveTriangle::SpawnEx(FParticleEmitterInstance*
 
 	float BarycentricCoords[3] = {0};
 	float ZPos = 0.0f;
-	if (InRandomStream)
-	{
-		BarycentricCoords[0] = InRandomStream->GetFraction();
-		BarycentricCoords[1] = InRandomStream->GetFraction();
-		BarycentricCoords[2] = InRandomStream->GetFraction();
-		ZPos = InRandomStream->GetFraction();
-	}
-	else
-	{
-		BarycentricCoords[0] = FMath::SRand();
-		BarycentricCoords[1] = FMath::SRand();
-		BarycentricCoords[2] = FMath::SRand();
-		ZPos = FMath::SRand();
-	}
 
+	BarycentricCoords[0] = InRandomStream->GetFraction();
+	BarycentricCoords[1] = InRandomStream->GetFraction();
+	BarycentricCoords[2] = InRandomStream->GetFraction();
+	ZPos = InRandomStream->GetFraction();
+	
 	FVector LocationOffset = FVector::ZeroVector;
 	float Sum = FMath::Max<float>(KINDA_SMALL_NUMBER, BarycentricCoords[0] + BarycentricCoords[1] + BarycentricCoords[2]);
 	for (int32 i = 0; i < 3; i++)
@@ -1000,7 +948,7 @@ void UParticleModuleLocationPrimitiveCylinder::PostEditChangeProperty(FPropertyC
 
 void UParticleModuleLocationPrimitiveCylinder::Spawn(FParticleEmitterInstance* Owner, int32 Offset, float SpawnTime, FBaseParticle* ParticleBase)
 {
-	SpawnEx(Owner, Offset, SpawnTime, NULL, ParticleBase);
+	SpawnEx(Owner, Offset, SpawnTime, &GetRandomStream(Owner), ParticleBase);
 }
 
 void UParticleModuleLocationPrimitiveCylinder::SpawnEx(FParticleEmitterInstance* Owner, int32 Offset, float SpawnTime, struct FRandomStream* InRandomStream, FBaseParticle* ParticleBase)
@@ -1240,27 +1188,11 @@ UParticleModuleLocationPrimitiveCylinder_Seeded::UParticleModuleLocationPrimitiv
 	bRequiresLoopingNotification = true;
 }
 
-void UParticleModuleLocationPrimitiveCylinder_Seeded::Spawn(FParticleEmitterInstance* Owner, int32 Offset, float SpawnTime, FBaseParticle* ParticleBase)
-{
-	FParticleRandomSeedInstancePayload* Payload = (FParticleRandomSeedInstancePayload*)(Owner->GetModuleInstanceData(this));
-	SpawnEx(Owner, Offset, SpawnTime, (Payload != NULL) ? &(Payload->RandomStream) : NULL, ParticleBase);
-}
-
-uint32 UParticleModuleLocationPrimitiveCylinder_Seeded::RequiredBytesPerInstance()
-{
-	return RandomSeedInfo.GetInstancePayloadSize();
-}
-
-uint32 UParticleModuleLocationPrimitiveCylinder_Seeded::PrepPerInstanceBlock(FParticleEmitterInstance* Owner, void* InstData)
-{
-	return PrepRandomSeedInstancePayload(Owner, (FParticleRandomSeedInstancePayload*)InstData, RandomSeedInfo);
-}
-
 void UParticleModuleLocationPrimitiveCylinder_Seeded::EmitterLoopingNotify(FParticleEmitterInstance* Owner)
 {
 	if (RandomSeedInfo.bResetSeedOnEmitterLooping == true)
 	{
-		FParticleRandomSeedInstancePayload* Payload = (FParticleRandomSeedInstancePayload*)(Owner->GetModuleInstanceData(this));
+		FParticleRandomSeedInstancePayload* Payload = Owner->GetModuleRandomSeedInstanceData(this);
 		PrepRandomSeedInstancePayload(Owner, Payload, RandomSeedInfo);
 	}
 }
@@ -1304,7 +1236,7 @@ void UParticleModuleLocationPrimitiveSphere::PostEditChangeProperty(FPropertyCha
 
 void UParticleModuleLocationPrimitiveSphere::Spawn(FParticleEmitterInstance* Owner, int32 Offset, float SpawnTime, FBaseParticle* ParticleBase)
 {
-	SpawnEx(Owner, Offset, SpawnTime, NULL, ParticleBase);
+	SpawnEx(Owner, Offset, SpawnTime, &GetRandomStream(Owner), ParticleBase);
 }
 
 void UParticleModuleLocationPrimitiveSphere::SpawnEx(FParticleEmitterInstance* Owner, int32 Offset, float SpawnTime, struct FRandomStream* InRandomStream, FBaseParticle* ParticleBase)
@@ -1456,27 +1388,11 @@ UParticleModuleLocationPrimitiveSphere_Seeded::UParticleModuleLocationPrimitiveS
 	bRequiresLoopingNotification = true;
 }
 
-void UParticleModuleLocationPrimitiveSphere_Seeded::Spawn(FParticleEmitterInstance* Owner, int32 Offset, float SpawnTime, FBaseParticle* ParticleBase)
-{
-	FParticleRandomSeedInstancePayload* Payload = (FParticleRandomSeedInstancePayload*)(Owner->GetModuleInstanceData(this));
-	SpawnEx(Owner, Offset, SpawnTime, (Payload != NULL) ? &(Payload->RandomStream) : NULL, ParticleBase);
-}
-
-uint32 UParticleModuleLocationPrimitiveSphere_Seeded::RequiredBytesPerInstance()
-{
-	return RandomSeedInfo.GetInstancePayloadSize();
-}
-
-uint32 UParticleModuleLocationPrimitiveSphere_Seeded::PrepPerInstanceBlock(FParticleEmitterInstance* Owner, void* InstData)
-{
-	return PrepRandomSeedInstancePayload(Owner, (FParticleRandomSeedInstancePayload*)InstData, RandomSeedInfo);
-}
-
 void UParticleModuleLocationPrimitiveSphere_Seeded::EmitterLoopingNotify(FParticleEmitterInstance* Owner)
 {
 	if (RandomSeedInfo.bResetSeedOnEmitterLooping == true)
 	{
-		FParticleRandomSeedInstancePayload* Payload = (FParticleRandomSeedInstancePayload*)(Owner->GetModuleInstanceData(this));
+		FParticleRandomSeedInstancePayload* Payload = Owner->GetModuleRandomSeedInstanceData(this);
 		PrepRandomSeedInstancePayload(Owner, Payload, RandomSeedInfo);
 	}
 }
@@ -1511,7 +1427,7 @@ UParticleModuleLocationBoneSocket::UParticleModuleLocationBoneSocket(const FObje
 	InheritVelocityScale = 1.0f;
 }
 
-int32 UParticleModuleLocationBoneSocket::SelectNextSpawnIndex(FModuleLocationBoneSocketInstancePayload* InstancePayload, USkeletalMeshComponent* SourceComponent)
+int32 UParticleModuleLocationBoneSocket::SelectNextSpawnIndex(FModuleLocationBoneSocketInstancePayload* InstancePayload, USkeletalMeshComponent* SourceComponent, FRandomStream& InRandomStream)
 {
 	check(InstancePayload && SourceComponent);
 
@@ -1531,7 +1447,7 @@ int32 UParticleModuleLocationBoneSocket::SelectNextSpawnIndex(FModuleLocationBon
 	else if (SelectionMethod == BONESOCKETSEL_Random)
 	{
 		// Note: This can select the same socket over and over...
-		SourceIndex = FMath::TruncToInt(FMath::SRand() * ((float)MaxIndex - 0.5f));
+		SourceIndex = FMath::TruncToInt(InRandomStream.FRand() * ((float)MaxIndex - 0.5f));
 		InstancePayload->LastSelectedIndex = SourceIndex;
 	}
 
@@ -1547,7 +1463,7 @@ int32 UParticleModuleLocationBoneSocket::SelectNextSpawnIndex(FModuleLocationBon
 	return SourceIndex;
 }
 
-void UParticleModuleLocationBoneSocket::RegeneratePreSelectedIndices(FModuleLocationBoneSocketInstancePayload* InstancePayload, USkeletalMeshComponent* SourceComponent)
+void UParticleModuleLocationBoneSocket::RegeneratePreSelectedIndices(FModuleLocationBoneSocketInstancePayload* InstancePayload, USkeletalMeshComponent* SourceComponent, FRandomStream& InRandomStream)
 {
 	if (SourceIndexMode == EBoneSocketSourceIndexMode::PreSelectedIndices)
 	{
@@ -1555,7 +1471,7 @@ void UParticleModuleLocationBoneSocket::RegeneratePreSelectedIndices(FModuleLoca
 		for (int32 i = 0; i < NumPreSelectedIndices; ++i)
 		{
 			//Should we provide sequential selection here? Does that make sense for the pre selected list?
-			InstancePayload->PreSelectedBoneSocketIndices[i] = FMath::TruncToInt(FMath::SRand() * ((float)MaxIndex - 0.5f));
+			InstancePayload->PreSelectedBoneSocketIndices[i] = FMath::TruncToInt(InRandomStream.FRand() * ((float)MaxIndex - 0.5f));
 		}
 
 		if (InheritingBoneVelocity())
@@ -1600,6 +1516,8 @@ void UParticleModuleLocationBoneSocket::Spawn(FParticleEmitterInstance* Owner, i
 		return;
 	}
 
+	FRandomStream& RandomStream = GetRandomStream(Owner);
+
 	// Setup the source skeletal mesh component...
 	GetSkeletalMeshComponentSource(Owner, InstancePayload);
 
@@ -1610,7 +1528,7 @@ void UParticleModuleLocationBoneSocket::Spawn(FParticleEmitterInstance* Owner, i
 	}
 	USkeletalMeshComponent* SourceComponent = InstancePayload->SourceComponent.Get();
 
-	int32 SourceIndex = SelectNextSpawnIndex(InstancePayload, SourceComponent);
+	int32 SourceIndex = SelectNextSpawnIndex(InstancePayload, SourceComponent, RandomStream);
 	if (SourceIndex == INDEX_NONE)
 	{
 		return;
@@ -1784,7 +1702,7 @@ void UParticleModuleLocationBoneSocket::FinalUpdate(FParticleEmitterInstance* Ow
 	}
 
 	//Select a new set of bones to spawn from next frame.
-	RegeneratePreSelectedIndices(InstancePayload, SourceComponent);
+	RegeneratePreSelectedIndices(InstancePayload, SourceComponent, GetRandomStream(Owner));
 }
 
 uint32 UParticleModuleLocationBoneSocket::RequiredBytes(UParticleModuleTypeDataBase* TypeData)
@@ -2033,7 +1951,7 @@ void UParticleModuleLocationBoneSocket::GetSkeletalMeshComponentSource(FParticle
 		if (NewSkelComp)
 		{
 			InstancePayload->SourceComponent = NewSkelComp;
-			RegeneratePreSelectedIndices(InstancePayload, NewSkelComp);
+			RegeneratePreSelectedIndices(InstancePayload, NewSkelComp, GetRandomStream(Owner));
 		}
 	}
 }
@@ -2287,6 +2205,8 @@ void UParticleModuleLocationSkelVertSurface::Spawn(FParticleEmitterInstance* Own
 	{
 		return;
 	}
+
+	FRandomStream& RandomStream = GetRandomStream(Owner);
 	
 	GetSkeletalMeshComponentSource(Owner, InstancePayload);
 
@@ -2316,7 +2236,7 @@ void UParticleModuleLocationSkelVertSurface::Spawn(FParticleEmitterInstance* Own
 	{
 		int32 SourceLocationsCount(LODData.GetNumVertices());
 
-		SourceIndex = FMath::TruncToInt(FMath::SRand() * ((float)SourceLocationsCount) - 1);
+		SourceIndex = FMath::TruncToInt(RandomStream.FRand() * ((float)SourceLocationsCount) - 1);
 		InstancePayload->VertIndex = SourceIndex;
 
 		if(SourceIndex != -1)
@@ -2335,10 +2255,10 @@ void UParticleModuleLocationSkelVertSurface::Spawn(FParticleEmitterInstance* Own
 	else if(SourceType == VERTSURFACESOURCE_Surface)
 	{
 		int32 SectionCount = LODData.RenderSections.Num();
-		int32 RandomSection = FMath::RoundToInt(FMath::SRand() * ((float)SectionCount-1));
+		int32 RandomSection = FMath::RoundToInt(RandomStream.FRand() * ((float)SectionCount-1));
 
 		SourceIndex = LODData.RenderSections[RandomSection].BaseIndex +
-			(FMath::TruncToInt(FMath::SRand() * ((float)LODData.RenderSections[RandomSection].NumTriangles))*3);
+			(FMath::TruncToInt(RandomStream.FRand() * ((float)LODData.RenderSections[RandomSection].NumTriangles))*3);
 
 		InstancePayload->VertIndex = SourceIndex;
 
@@ -2469,10 +2389,7 @@ void UParticleModuleLocationSkelVertSurface::Spawn(FParticleEmitterInstance* Own
 				{
 					//We have the mesh oriented to the normal of the triangle it's on but this looks fugly as particles on each triangle are facing the same way.
 					//The only valid orientation reference should be the normal. So add an additional random rotation around it.
-					int32 OldRandSeed = FMath::GetRandSeed();
-					FMath::SRandInit((int32)((intptr_t)ParticleBase));
-					SourceRotation = SourceRotation * FQuat(FVector::UpVector, FMath::SRand()*(PI*2.0f));
-					FMath::SRandInit(OldRandSeed);
+					SourceRotation = SourceRotation * FQuat(FVector::UpVector, RandomStream.FRand()*(PI*2.0f));
 				}
 
 				FVector Rot = SourceRotation.Euler();
@@ -2544,6 +2461,8 @@ void UParticleModuleLocationSkelVertSurface::Update(FParticleEmitterInstance* Ow
 	const bool bMeshRotationActive = MeshRotationOffset > 0 && Owner->IsMeshRotationActive();
 	const FTransform& OwnerTM = Owner->Component->GetAsyncComponentToWorld();
 
+	FRandomStream& RandomStream = GetRandomStream(Owner);
+
 	BEGIN_UPDATE_LOOP;
 	{
 		FModuleLocationVertSurfaceParticlePayload* ParticlePayload = (FModuleLocationVertSurfaceParticlePayload*)((uint8*)&Particle + Offset);
@@ -2557,11 +2476,7 @@ void UParticleModuleLocationSkelVertSurface::Update(FParticleEmitterInstance* Ow
 				{
 					//We have the mesh oriented to the normal of the triangle it's on but this looks fugly as particles on each triangle are facing the same way.
 					//The only valid orientation reference should be the normal. So add an additional random rotation around it.
-
-					int32 OldRandSeed = FMath::GetRandSeed();
-					FMath::SRandInit((int32)((intptr_t)ParticleBase));
-					SourceRotation = SourceRotation * FQuat(FVector::UpVector, FMath::SRand()*(PI*2.0f));
-					FMath::SRandInit(OldRandSeed);
+					SourceRotation = SourceRotation * FQuat(FVector::UpVector, RandomStream.FRand()*(PI*2.0f));
 				}
 
 				FMeshRotationPayloadData* PayloadData = (FMeshRotationPayloadData*)((uint8*)&Particle + MeshRotationOffset);
diff --git a/Engine/Source/Runtime/Engine/Private/Particles/ParticleModules_Parameter.cpp b/Engine/Source/Runtime/Engine/Private/Particles/ParticleModules_Parameter.cpp
--- a/Engine/Source/Runtime/Engine/Private/Particles/ParticleModules_Parameter.cpp
+++ b/Engine/Source/Runtime/Engine/Private/Particles/ParticleModules_Parameter.cpp
@@ -137,7 +137,7 @@ bool UParticleModuleParameterDynamic::CanTickInAnyThread()
 
 void UParticleModuleParameterDynamic::Spawn(FParticleEmitterInstance* Owner, int32 Offset, float SpawnTime, FBaseParticle* ParticleBase)
 {
-	SpawnEx(Owner, Offset, SpawnTime, NULL, ParticleBase);
+	SpawnEx(Owner, Offset, SpawnTime, &GetRandomStream(Owner), ParticleBase);
 }
 
 void UParticleModuleParameterDynamic::SpawnEx(FParticleEmitterInstance* Owner, int32 Offset, float SpawnTime, struct FRandomStream* InRandomStream, FBaseParticle* ParticleBase)
@@ -178,6 +178,8 @@ void UParticleModuleParameterDynamic::Update(FParticleEmitterInstance* Owner, in
 		return;
 	}
 
+	FRandomStream& RandomStream = GetRandomStream(Owner);
+
 	// 
 	int32 ParameterIndex = ParticleDynamicParameter_GetIndexFromFlag(UpdateFlags);
 
@@ -201,7 +203,7 @@ void UParticleModuleParameterDynamic::Update(FParticleEmitterInstance* Owner, in
 					FEmitterDynamicParameterPayload& DynamicPayload = *((FEmitterDynamicParameterPayload*)(ParticleBase + CurrentOffset));
 					FPlatformMisc::Prefetch(ParticleData, (ParticleIndices[i+1] * ParticleStride));
 					FPlatformMisc::Prefetch(ParticleData, (ParticleIndices[i+1] * ParticleStride) + PLATFORM_CACHE_LINE_SIZE);
-					DynamicPayload.DynamicParameterValue[ParameterIndex] = GetParameterValue_UserSet(DynParam, Particle, Owner, NULL);
+					DynamicPayload.DynamicParameterValue[ParameterIndex] = GetParameterValue_UserSet(DynParam, Particle, Owner, &RandomStream);
 				}
 				END_UPDATE_LOOP;
 			}
@@ -212,7 +214,7 @@ void UParticleModuleParameterDynamic::Update(FParticleEmitterInstance* Owner, in
 					FEmitterDynamicParameterPayload& DynamicPayload = *((FEmitterDynamicParameterPayload*)(ParticleBase + CurrentOffset));
 					FPlatformMisc::Prefetch(ParticleData, (ParticleIndices[i+1] * ParticleStride));
 					FPlatformMisc::Prefetch(ParticleData, (ParticleIndices[i+1] * ParticleStride) + PLATFORM_CACHE_LINE_SIZE);
-					DynamicPayload.DynamicParameterValue[ParameterIndex] = GetParameterValue(DynParam, Particle, Owner, NULL);
+					DynamicPayload.DynamicParameterValue[ParameterIndex] = GetParameterValue(DynParam, Particle, Owner, &RandomStream);
 				}
 				END_UPDATE_LOOP;
 			}
@@ -230,8 +232,8 @@ void UParticleModuleParameterDynamic::Update(FParticleEmitterInstance* Owner, in
 					FEmitterDynamicParameterPayload& DynamicPayload = *((FEmitterDynamicParameterPayload*)(ParticleBase + CurrentOffset));
 					FPlatformMisc::Prefetch(ParticleData, (ParticleIndices[i+1] * ParticleStride));
 					FPlatformMisc::Prefetch(ParticleData, (ParticleIndices[i+1] * ParticleStride) + PLATFORM_CACHE_LINE_SIZE);
-					DynamicPayload.DynamicParameterValue[0] = GetParameterValue_UserSet(DynParam0, Particle, Owner, NULL);
-					DynamicPayload.DynamicParameterValue[1] = GetParameterValue_UserSet(DynParam1, Particle, Owner, NULL);
+					DynamicPayload.DynamicParameterValue[0] = GetParameterValue_UserSet(DynParam0, Particle, Owner, &RandomStream);
+					DynamicPayload.DynamicParameterValue[1] = GetParameterValue_UserSet(DynParam1, Particle, Owner, &RandomStream);
 				}
 				END_UPDATE_LOOP;
 			}
@@ -242,8 +244,8 @@ void UParticleModuleParameterDynamic::Update(FParticleEmitterInstance* Owner, in
 					FEmitterDynamicParameterPayload& DynamicPayload = *((FEmitterDynamicParameterPayload*)(ParticleBase + CurrentOffset));
 					FPlatformMisc::Prefetch(ParticleData, (ParticleIndices[i+1] * ParticleStride));
 					FPlatformMisc::Prefetch(ParticleData, (ParticleIndices[i+1] * ParticleStride) + PLATFORM_CACHE_LINE_SIZE);
-					DynamicPayload.DynamicParameterValue[0] = GetParameterValue(DynParam0, Particle, Owner, NULL);
-					DynamicPayload.DynamicParameterValue[1] = GetParameterValue(DynParam1, Particle, Owner, NULL);
+					DynamicPayload.DynamicParameterValue[0] = GetParameterValue(DynParam0, Particle, Owner, &RandomStream);
+					DynamicPayload.DynamicParameterValue[1] = GetParameterValue(DynParam1, Particle, Owner, &RandomStream);
 				}
 				END_UPDATE_LOOP;
 			}
@@ -262,9 +264,9 @@ void UParticleModuleParameterDynamic::Update(FParticleEmitterInstance* Owner, in
 					FEmitterDynamicParameterPayload& DynamicPayload = *((FEmitterDynamicParameterPayload*)(ParticleBase + CurrentOffset));
 					FPlatformMisc::Prefetch(ParticleData, (ParticleIndices[i+1] * ParticleStride));
 					FPlatformMisc::Prefetch(ParticleData, (ParticleIndices[i+1] * ParticleStride) + PLATFORM_CACHE_LINE_SIZE);
-					DynamicPayload.DynamicParameterValue[0] = GetParameterValue_UserSet(DynParam0, Particle, Owner, NULL);
-					DynamicPayload.DynamicParameterValue[1] = GetParameterValue_UserSet(DynParam1, Particle, Owner, NULL);
-					DynamicPayload.DynamicParameterValue[2] = GetParameterValue_UserSet(DynParam2, Particle, Owner, NULL);
+					DynamicPayload.DynamicParameterValue[0] = GetParameterValue_UserSet(DynParam0, Particle, Owner, &RandomStream);
+					DynamicPayload.DynamicParameterValue[1] = GetParameterValue_UserSet(DynParam1, Particle, Owner, &RandomStream);
+					DynamicPayload.DynamicParameterValue[2] = GetParameterValue_UserSet(DynParam2, Particle, Owner, &RandomStream);
 				}
 				END_UPDATE_LOOP;
 			}
@@ -275,9 +277,9 @@ void UParticleModuleParameterDynamic::Update(FParticleEmitterInstance* Owner, in
 					FEmitterDynamicParameterPayload& DynamicPayload = *((FEmitterDynamicParameterPayload*)(ParticleBase + CurrentOffset));
 					FPlatformMisc::Prefetch(ParticleData, (ParticleIndices[i+1] * ParticleStride));
 					FPlatformMisc::Prefetch(ParticleData, (ParticleIndices[i+1] * ParticleStride) + PLATFORM_CACHE_LINE_SIZE);
-					DynamicPayload.DynamicParameterValue[0] = GetParameterValue(DynParam0, Particle, Owner, NULL);
-					DynamicPayload.DynamicParameterValue[1] = GetParameterValue(DynParam1, Particle, Owner, NULL);
-					DynamicPayload.DynamicParameterValue[2] = GetParameterValue(DynParam2, Particle, Owner, NULL);
+					DynamicPayload.DynamicParameterValue[0] = GetParameterValue(DynParam0, Particle, Owner, &RandomStream);
+					DynamicPayload.DynamicParameterValue[1] = GetParameterValue(DynParam1, Particle, Owner, &RandomStream);
+					DynamicPayload.DynamicParameterValue[2] = GetParameterValue(DynParam2, Particle, Owner, &RandomStream);
 				}
 				END_UPDATE_LOOP;
 			}
@@ -296,10 +298,10 @@ void UParticleModuleParameterDynamic::Update(FParticleEmitterInstance* Owner, in
 					FEmitterDynamicParameterPayload& DynamicPayload = *((FEmitterDynamicParameterPayload*)(ParticleBase + CurrentOffset));
 					FPlatformMisc::Prefetch(ParticleData, (ParticleIndices[i+1] * ParticleStride));
 					FPlatformMisc::Prefetch(ParticleData, (ParticleIndices[i+1] * ParticleStride) + PLATFORM_CACHE_LINE_SIZE);
-					DynamicPayload.DynamicParameterValue[0] = GetParameterValue_UserSet(DynParam0, Particle, Owner, NULL);
-					DynamicPayload.DynamicParameterValue[1] = GetParameterValue_UserSet(DynParam1, Particle, Owner, NULL);
-					DynamicPayload.DynamicParameterValue[2] = GetParameterValue_UserSet(DynParam2, Particle, Owner, NULL);
-					DynamicPayload.DynamicParameterValue[3] = GetParameterValue_UserSet(DynParam3, Particle, Owner, NULL);
+					DynamicPayload.DynamicParameterValue[0] = GetParameterValue_UserSet(DynParam0, Particle, Owner, &RandomStream);
+					DynamicPayload.DynamicParameterValue[1] = GetParameterValue_UserSet(DynParam1, Particle, Owner, &RandomStream);
+					DynamicPayload.DynamicParameterValue[2] = GetParameterValue_UserSet(DynParam2, Particle, Owner, &RandomStream);
+					DynamicPayload.DynamicParameterValue[3] = GetParameterValue_UserSet(DynParam3, Particle, Owner, &RandomStream);
 				}
 				END_UPDATE_LOOP;
 			}
@@ -310,10 +312,10 @@ void UParticleModuleParameterDynamic::Update(FParticleEmitterInstance* Owner, in
 					FEmitterDynamicParameterPayload& DynamicPayload = *((FEmitterDynamicParameterPayload*)(ParticleBase + CurrentOffset));
 					FPlatformMisc::Prefetch(ParticleData, (ParticleIndices[i+1] * ParticleStride));
 					FPlatformMisc::Prefetch(ParticleData, (ParticleIndices[i+1] * ParticleStride) + PLATFORM_CACHE_LINE_SIZE);
-					DynamicPayload.DynamicParameterValue[0] = GetParameterValue(DynParam0, Particle, Owner, NULL);
-					DynamicPayload.DynamicParameterValue[1] = GetParameterValue(DynParam1, Particle, Owner, NULL);
-					DynamicPayload.DynamicParameterValue[2] = GetParameterValue(DynParam2, Particle, Owner, NULL);
-					DynamicPayload.DynamicParameterValue[3] = GetParameterValue(DynParam3, Particle, Owner, NULL);
+					DynamicPayload.DynamicParameterValue[0] = GetParameterValue(DynParam0, Particle, Owner, &RandomStream);
+					DynamicPayload.DynamicParameterValue[1] = GetParameterValue(DynParam1, Particle, Owner, &RandomStream);
+					DynamicPayload.DynamicParameterValue[2] = GetParameterValue(DynParam2, Particle, Owner, &RandomStream);
+					DynamicPayload.DynamicParameterValue[3] = GetParameterValue(DynParam3, Particle, Owner, &RandomStream);
 				}
 				END_UPDATE_LOOP;
 			}
@@ -332,10 +334,10 @@ void UParticleModuleParameterDynamic::Update(FParticleEmitterInstance* Owner, in
 					FEmitterDynamicParameterPayload& DynamicPayload = *((FEmitterDynamicParameterPayload*)(ParticleBase + CurrentOffset));
 					FPlatformMisc::Prefetch(ParticleData, (ParticleIndices[i+1] * ParticleStride));
 					FPlatformMisc::Prefetch(ParticleData, (ParticleIndices[i+1] * ParticleStride) + PLATFORM_CACHE_LINE_SIZE);
-					DynamicPayload.DynamicParameterValue[0] = (UpdateFlags & EDPU_UPDATE_0) ? GetParameterValue_UserSet(DynParam0, Particle, Owner, NULL) : DynamicPayload.DynamicParameterValue[0];
-					DynamicPayload.DynamicParameterValue[1] = (UpdateFlags & EDPU_UPDATE_1) ? GetParameterValue_UserSet(DynParam1, Particle, Owner, NULL) : DynamicPayload.DynamicParameterValue[1];
-					DynamicPayload.DynamicParameterValue[2] = (UpdateFlags & EDPU_UPDATE_2) ? GetParameterValue_UserSet(DynParam2, Particle, Owner, NULL) : DynamicPayload.DynamicParameterValue[2];
-					DynamicPayload.DynamicParameterValue[3] = (UpdateFlags & EDPU_UPDATE_3) ? GetParameterValue_UserSet(DynParam3, Particle, Owner, NULL) : DynamicPayload.DynamicParameterValue[3];
+					DynamicPayload.DynamicParameterValue[0] = (UpdateFlags & EDPU_UPDATE_0) ? GetParameterValue_UserSet(DynParam0, Particle, Owner, &RandomStream) : DynamicPayload.DynamicParameterValue[0];
+					DynamicPayload.DynamicParameterValue[1] = (UpdateFlags & EDPU_UPDATE_1) ? GetParameterValue_UserSet(DynParam1, Particle, Owner, &RandomStream) : DynamicPayload.DynamicParameterValue[1];
+					DynamicPayload.DynamicParameterValue[2] = (UpdateFlags & EDPU_UPDATE_2) ? GetParameterValue_UserSet(DynParam2, Particle, Owner, &RandomStream) : DynamicPayload.DynamicParameterValue[2];
+					DynamicPayload.DynamicParameterValue[3] = (UpdateFlags & EDPU_UPDATE_3) ? GetParameterValue_UserSet(DynParam3, Particle, Owner, &RandomStream) : DynamicPayload.DynamicParameterValue[3];
 				}
 				END_UPDATE_LOOP;
 			}
@@ -346,10 +348,10 @@ void UParticleModuleParameterDynamic::Update(FParticleEmitterInstance* Owner, in
 					FEmitterDynamicParameterPayload& DynamicPayload = *((FEmitterDynamicParameterPayload*)(ParticleBase + CurrentOffset));
 					FPlatformMisc::Prefetch(ParticleData, (ParticleIndices[i+1] * ParticleStride));
 					FPlatformMisc::Prefetch(ParticleData, (ParticleIndices[i+1] * ParticleStride) + PLATFORM_CACHE_LINE_SIZE);
-					DynamicPayload.DynamicParameterValue[0] = (UpdateFlags & EDPU_UPDATE_0) ? GetParameterValue(DynParam0, Particle, Owner, NULL) : DynamicPayload.DynamicParameterValue[0];
-					DynamicPayload.DynamicParameterValue[1] = (UpdateFlags & EDPU_UPDATE_1) ? GetParameterValue(DynParam1, Particle, Owner, NULL) : DynamicPayload.DynamicParameterValue[1];
-					DynamicPayload.DynamicParameterValue[2] = (UpdateFlags & EDPU_UPDATE_2) ? GetParameterValue(DynParam2, Particle, Owner, NULL) : DynamicPayload.DynamicParameterValue[2];
-					DynamicPayload.DynamicParameterValue[3] = (UpdateFlags & EDPU_UPDATE_3) ? GetParameterValue(DynParam3, Particle, Owner, NULL) : DynamicPayload.DynamicParameterValue[3];
+					DynamicPayload.DynamicParameterValue[0] = (UpdateFlags & EDPU_UPDATE_0) ? GetParameterValue(DynParam0, Particle, Owner, &RandomStream) : DynamicPayload.DynamicParameterValue[0];
+					DynamicPayload.DynamicParameterValue[1] = (UpdateFlags & EDPU_UPDATE_1) ? GetParameterValue(DynParam1, Particle, Owner, &RandomStream) : DynamicPayload.DynamicParameterValue[1];
+					DynamicPayload.DynamicParameterValue[2] = (UpdateFlags & EDPU_UPDATE_2) ? GetParameterValue(DynParam2, Particle, Owner, &RandomStream) : DynamicPayload.DynamicParameterValue[2];
+					DynamicPayload.DynamicParameterValue[3] = (UpdateFlags & EDPU_UPDATE_3) ? GetParameterValue(DynParam3, Particle, Owner, &RandomStream) : DynamicPayload.DynamicParameterValue[3];
 				}
 				END_UPDATE_LOOP;
 			}
@@ -575,22 +577,6 @@ UParticleModuleParameterDynamic_Seeded::UParticleModuleParameterDynamic_Seeded(c
 	bRequiresLoopingNotification = true;
 }
 
-void UParticleModuleParameterDynamic_Seeded::Spawn(FParticleEmitterInstance* Owner, int32 Offset, float SpawnTime, FBaseParticle* ParticleBase)
-{
-	FParticleRandomSeedInstancePayload* Payload = (FParticleRandomSeedInstancePayload*)(Owner->GetModuleInstanceData(this));
-	SpawnEx(Owner, Offset, SpawnTime, (Payload != NULL) ? &(Payload->RandomStream) : NULL, ParticleBase);
-}
-
-uint32 UParticleModuleParameterDynamic_Seeded::RequiredBytesPerInstance()
-{
-	return RandomSeedInfo.GetInstancePayloadSize();
-}
-
-uint32 UParticleModuleParameterDynamic_Seeded::PrepPerInstanceBlock(FParticleEmitterInstance* Owner, void* InstData)
-{
-	return PrepRandomSeedInstancePayload(Owner, (FParticleRandomSeedInstancePayload*)InstData, RandomSeedInfo);
-}
-
 void UParticleModuleParameterDynamic_Seeded::EmitterLoopingNotify(FParticleEmitterInstance* Owner)
 {
 	if (RandomSeedInfo.bResetSeedOnEmitterLooping == true)
diff --git a/Engine/Source/Runtime/Engine/Private/Particles/ParticleModules_Size.cpp b/Engine/Source/Runtime/Engine/Private/Particles/ParticleModules_Size.cpp
--- a/Engine/Source/Runtime/Engine/Private/Particles/ParticleModules_Size.cpp
+++ b/Engine/Source/Runtime/Engine/Private/Particles/ParticleModules_Size.cpp
@@ -84,7 +84,7 @@ void UParticleModuleSize::PostEditChangeProperty(FPropertyChangedEvent& Property
 
 void UParticleModuleSize::Spawn(FParticleEmitterInstance* Owner, int32 Offset, float SpawnTime, FBaseParticle* ParticleBase)
 {
-	SpawnEx(Owner, Offset, SpawnTime, NULL, ParticleBase);
+	SpawnEx(Owner, Offset, SpawnTime, &GetRandomStream(Owner), ParticleBase);
 }
 
 void UParticleModuleSize::SpawnEx(FParticleEmitterInstance* Owner, int32 Offset, float SpawnTime, struct FRandomStream* InRandomStream, FBaseParticle* ParticleBase)
@@ -93,7 +93,7 @@ void UParticleModuleSize::SpawnEx(FParticleEmitterInstance* Owner, int32 Offset,
 	FVector Size		 = StartSize.GetValue(Owner->EmitterTime, Owner->Component, 0, InRandomStream);
 	Particle.Size	+= Size;
 
-	AdjustParticleBaseSizeForUVFlipping(Size, Owner->CurrentLODLevel->RequiredModule->UVFlippingMode);
+	AdjustParticleBaseSizeForUVFlipping(Size, Owner->CurrentLODLevel->RequiredModule->UVFlippingMode, *InRandomStream);
 	Particle.BaseSize += Size;
 }
 
@@ -108,27 +108,11 @@ UParticleModuleSize_Seeded::UParticleModuleSize_Seeded(const FObjectInitializer&
 	bRequiresLoopingNotification = true;
 }
 
-void UParticleModuleSize_Seeded::Spawn(FParticleEmitterInstance* Owner, int32 Offset, float SpawnTime, FBaseParticle* ParticleBase)
-{
-	FParticleRandomSeedInstancePayload* Payload = (FParticleRandomSeedInstancePayload*)(Owner->GetModuleInstanceData(this));
-	SpawnEx(Owner, Offset, SpawnTime, (Payload != NULL) ? &(Payload->RandomStream) : NULL, ParticleBase);
-}
-
-uint32 UParticleModuleSize_Seeded::RequiredBytesPerInstance()
-{
-	return RandomSeedInfo.GetInstancePayloadSize();
-}
-
-uint32 UParticleModuleSize_Seeded::PrepPerInstanceBlock(FParticleEmitterInstance* Owner, void* InstData)
-{
-	return PrepRandomSeedInstancePayload(Owner, (FParticleRandomSeedInstancePayload*)InstData, RandomSeedInfo);
-}
-
 void UParticleModuleSize_Seeded::EmitterLoopingNotify(FParticleEmitterInstance* Owner)
 {
 	if (RandomSeedInfo.bResetSeedOnEmitterLooping == true)
 	{
-		FParticleRandomSeedInstancePayload* Payload = (FParticleRandomSeedInstancePayload*)(Owner->GetModuleInstanceData(this));
+		FParticleRandomSeedInstancePayload* Payload = Owner->GetModuleRandomSeedInstanceData(this);
 		PrepRandomSeedInstancePayload(Owner, Payload, RandomSeedInfo);
 	}
 }
diff --git a/Engine/Source/Runtime/Engine/Private/Particles/ParticleModules_Velocity.cpp b/Engine/Source/Runtime/Engine/Private/Particles/ParticleModules_Velocity.cpp
--- a/Engine/Source/Runtime/Engine/Private/Particles/ParticleModules_Velocity.cpp
+++ b/Engine/Source/Runtime/Engine/Private/Particles/ParticleModules_Velocity.cpp
@@ -73,7 +73,7 @@ void UParticleModuleVelocity::PostEditChangeProperty(FPropertyChangedEvent& Prop
 
 void UParticleModuleVelocity::Spawn(FParticleEmitterInstance* Owner, int32 Offset, float SpawnTime, FBaseParticle* ParticleBase)
 {
-	SpawnEx(Owner, Offset, SpawnTime, NULL, ParticleBase);
+	SpawnEx(Owner, Offset, SpawnTime, &GetRandomStream(Owner), ParticleBase);
 }
 
 void UParticleModuleVelocity::SpawnEx(FParticleEmitterInstance* Owner, int32 Offset, float SpawnTime, struct FRandomStream* InRandomStream, FBaseParticle* ParticleBase)
@@ -124,22 +124,6 @@ UParticleModuleVelocity_Seeded::UParticleModuleVelocity_Seeded(const FObjectInit
 	bRequiresLoopingNotification = true;
 }
 
-void UParticleModuleVelocity_Seeded::Spawn(FParticleEmitterInstance* Owner, int32 Offset, float SpawnTime, FBaseParticle* ParticleBase)
-{
-	FParticleRandomSeedInstancePayload* Payload = (FParticleRandomSeedInstancePayload*)(Owner->GetModuleInstanceData(this));
-	SpawnEx(Owner, Offset, SpawnTime, (Payload != NULL) ? &(Payload->RandomStream) : NULL, ParticleBase);
-}
-
-uint32 UParticleModuleVelocity_Seeded::RequiredBytesPerInstance()
-{
-	return RandomSeedInfo.GetInstancePayloadSize();
-}
-
-uint32 UParticleModuleVelocity_Seeded::PrepPerInstanceBlock(FParticleEmitterInstance* Owner, void* InstData)
-{
-	return PrepRandomSeedInstancePayload(Owner, (FParticleRandomSeedInstancePayload*)InstData, RandomSeedInfo);
-}
-
 void UParticleModuleVelocity_Seeded::EmitterLoopingNotify(FParticleEmitterInstance* Owner)
 {
 	if (RandomSeedInfo.bResetSeedOnEmitterLooping == true)
@@ -415,7 +399,7 @@ void UParticleModuleVelocityCone::PostEditChangeProperty(FPropertyChangedEvent&
 
 void UParticleModuleVelocityCone::Spawn(FParticleEmitterInstance* Owner, int32 Offset, float SpawnTime, FBaseParticle* ParticleBase)
 {
-	SpawnEx(Owner, Offset, SpawnTime, NULL, ParticleBase);
+	SpawnEx(Owner, Offset, SpawnTime, &GetRandomStream(Owner), ParticleBase);
 }
 
 void UParticleModuleVelocityCone::SpawnEx(FParticleEmitterInstance* Owner, int32 Offset, float SpawnTime, struct FRandomStream* InRandomStream, FBaseParticle* ParticleBase)
@@ -440,7 +424,7 @@ void UParticleModuleVelocityCone::SpawnEx(FParticleEmitterInstance* Owner, int32
 		// Calculate the default position (prior to the Direction vector being factored in)
 		const float SpawnAngle = Angle.GetValue(Owner->EmitterTime, Owner->Component, InRandomStream);
 		const float SpawnVelocity = Velocity.GetValue(Owner->EmitterTime, Owner->Component, InRandomStream);
-		const float LatheAngle = FMath::SRand() * TwoPI;
+		const float LatheAngle = InRandomStream->FRand() * TwoPI;
 		const FRotator DefaultDirectionRotater((int32)(SpawnAngle * ToRads * UUPerRad), (int32)(LatheAngle * UUPerRad), 0);
 		const FRotationMatrix DefaultDirectionRotation(DefaultDirectionRotater);
 		const FVector DefaultSpawnDirection = DefaultDirectionRotation.TransformVector(DefaultDirection);
diff --git a/Engine/Source/Runtime/Engine/Private/Particles/ParticleTrail2EmitterInstance.cpp b/Engine/Source/Runtime/Engine/Private/Particles/ParticleTrail2EmitterInstance.cpp
--- a/Engine/Source/Runtime/Engine/Private/Particles/ParticleTrail2EmitterInstance.cpp
+++ b/Engine/Source/Runtime/Engine/Private/Particles/ParticleTrail2EmitterInstance.cpp
@@ -2223,7 +2223,8 @@ bool FParticleRibbonEmitterInstance::ResolveSourcePoint(int32 InTrailIdx,
 						{
 							case EPSSM_Random:
 							{
-								Index = FMath::TruncToInt(FMath::FRand() * SourceEmitter->ActiveParticles);
+								FRandomStream& RandomStream = SourceModule->GetRandomStream(this);
+								Index = RandomStream.RandHelper(SourceEmitter->ActiveParticles);
 							}
 								break;
 							case EPSSM_Sequential:
diff --git a/Engine/Source/Runtime/Engine/Public/ParticleEmitterInstances.h b/Engine/Source/Runtime/Engine/Public/ParticleEmitterInstances.h
--- a/Engine/Source/Runtime/Engine/Public/ParticleEmitterInstances.h
+++ b/Engine/Source/Runtime/Engine/Public/ParticleEmitterInstances.h
@@ -438,6 +438,8 @@ public:
 	uint32 GetModuleDataOffset(UParticleModule* Module);
 	/** Get pointer to emitter instance payload data for a particular module */
 	uint8* GetModuleInstanceData(UParticleModule* Module);
+	/** Get pointer to emitter instance random seed payload data for a particular module */
+	FParticleRandomSeedInstancePayload* GetModuleRandomSeedInstanceData(UParticleModule* Module);
 	virtual uint8* GetTypeDataModuleInstanceData();
 	virtual uint32 CalculateParticleStride(uint32 ParticleSize);
 	virtual void ResetBurstList();
