Fixed static analysis warnings
#rnx

diff --git a/Engine/Source/Runtime/Engine/Private/Particles/ParticleModules.cpp b/Engine/Source/Runtime/Engine/Private/Particles/ParticleModules.cpp
--- a/Engine/Source/Runtime/Engine/Private/Particles/ParticleModules.cpp
+++ b/Engine/Source/Runtime/Engine/Private/Particles/ParticleModules.cpp
@@ -666,12 +666,18 @@ void UParticleModule::GetParticleParametersUtilized(TArray<FString>& ParticlePar
 
 uint32 UParticleModule::PrepRandomSeedInstancePayload(FParticleEmitterInstance* Owner, FParticleRandomSeedInstancePayload* InRandSeedPayload, const FParticleRandomSeedInfo& InRandSeedInfo)
 {
-	if (InRandSeedPayload != NULL)
+    // These should never be null
+	if (!ensure( Owner != nullptr && Owner->Component != nullptr ))
+	{
+		return 0xffffffff;
+	}
+
+	if (InRandSeedPayload != nullptr)
 	{
 		new(InRandSeedPayload) FParticleRandomSeedInstancePayload();
 
 		// See if the parameter is set on the instance...
-		if ((Owner != NULL) && (Owner->Component != NULL) && (InRandSeedInfo.bGetSeedFromInstance == true))
+		if (InRandSeedInfo.bGetSeedFromInstance == true)
 		{
 			float SeedValue;
 			if (Owner->Component->GetFloatParameter(InRandSeedInfo.ParameterName, SeedValue) == true)
@@ -700,15 +706,14 @@ uint32 UParticleModule::PrepRandomSeedInstancePayload(FParticleEmitterInstance*
 		// Pick a seed to use and initialize it!!!!
 		if (InRandSeedInfo.RandomSeeds.Num() > 0)
 		{
+			int32 Index = 0;
+
 			if (InRandSeedInfo.bRandomlySelectSeedArray)
 			{
-				int32 Index = Owner->Component->RandomStream.RandHelper(InRandSeedInfo.RandomSeeds.Num());
-				InRandSeedPayload->RandomStream.Initialize(InRandSeedInfo.RandomSeeds[Index]);
-			}
-			else
-			{
-				InRandSeedPayload->RandomStream.Initialize(InRandSeedInfo.RandomSeeds[0]);
+				Index = Owner->Component->RandomStream.RandHelper(InRandSeedInfo.RandomSeeds.Num());
 			}
+
+			InRandSeedPayload->RandomStream.Initialize(InRandSeedInfo.RandomSeeds[Index]);
 			return 0;
 		}
 		else if (FApp::bUseFixedSeed)
