Cached level name is now invalidated when the world asset is renamed from the content browser.

diff --git a/Engine/Source/Runtime/Engine/Classes/Engine/LevelStreaming.h b/Engine/Source/Runtime/Engine/Classes/Engine/LevelStreaming.h
--- a/Engine/Source/Runtime/Engine/Classes/Engine/LevelStreaming.h
+++ b/Engine/Source/Runtime/Engine/Classes/Engine/LevelStreaming.h
@@ -150,11 +150,14 @@ private:
 	ETargetState TargetState;
 
 	/** Whether this level streaming object's level should be unloaded and the object be removed from the level list.			*/
-	uint8 bIsRequestingUnloadAndRemoval : 1;
+	uint8 bIsRequestingUnloadAndRemoval:1;
 
 	/* Whether CachedWorldAssetPackageFName is valid */
 	mutable uint8 bHasCachedWorldAssetPackageFName:1;
 
+	/* Whether CachedLoadedLevelPackageName is valid */
+	mutable uint8 bHasCachedLoadedLevelPackageName:1;
+
 #if WITH_EDITORONLY_DATA
 	/** Whether this level should be visible in the Editor																		*/
 	UPROPERTY()
@@ -477,6 +480,9 @@ private:
 	/** @return Name of the LOD package on disk to load to the new package named PackageName, Name_None otherwise					*/
 	FName GetLODPackageNameToLoad() const;
 
+	/** @return Name of the level package that is currently loaded.																	*/
+	FName GetLoadedLevelPackageName() const;
+
 	/** 
 	 * Try to find loaded level in memory, issue a loading request otherwise
 	 *
@@ -518,7 +524,8 @@ private:
 	/** The cached package name of the world asset that is loaded by the levelstreaming */
 	mutable FName CachedWorldAssetPackageFName;
 
-	FName CachedLoadedLevelPackageName;
+	/** The cached package name of the currently loaded level. */
+	mutable FName CachedLoadedLevelPackageName;
 
 	friend struct FStreamingLevelPrivateAccessor;
 };
diff --git a/Engine/Source/Runtime/Engine/Private/LevelStreaming.cpp b/Engine/Source/Runtime/Engine/Private/LevelStreaming.cpp
--- a/Engine/Source/Runtime/Engine/Private/LevelStreaming.cpp
+++ b/Engine/Source/Runtime/Engine/Private/LevelStreaming.cpp
@@ -743,7 +743,7 @@ void ULevelStreaming::SetLoadedLevel(ULevel* Level)
 	check(PendingUnloadLevel == nullptr);
 	PendingUnloadLevel = LoadedLevel;
 	LoadedLevel = Level;
-	CachedLoadedLevelPackageName = (LoadedLevel ? LoadedLevel->GetOutermost()->GetFName() : NAME_None);
+	bHasCachedLoadedLevelPackageName = false;
 
 	// Cancel unloading for this level, in case it was queued for it
 	FLevelStreamingGCHelper::CancelUnloadRequest(LoadedLevel);
@@ -800,7 +800,9 @@ bool ULevelStreaming::IsDesiredLevelLoaded() const
 	{
 		const bool bIsGameWorld = GetWorld()->IsGameWorld();
 		const FName DesiredPackageName = bIsGameWorld ? GetLODPackageName() : GetWorldAssetPackageFName();
-		return (CachedLoadedLevelPackageName == DesiredPackageName);
+		const FName LoadedLevelPackageName = GetLoadedLevelPackageName();
+
+		return (LoadedLevelPackageName == DesiredPackageName);
 	}
 
 	return false;
@@ -826,9 +828,10 @@ bool ULevelStreaming::RequestLevel(UWorld* PersistentWorld, bool bAllowLevelLoad
 	// Package name we want to load
 	const bool bIsGameWorld = PersistentWorld->IsGameWorld();
 	const FName DesiredPackageName = bIsGameWorld ? GetLODPackageName() : GetWorldAssetPackageFName();
+	const FName LoadedLevelPackageName = GetLoadedLevelPackageName();
 
 	// Check if currently loaded level is what we want right now
-	if (LoadedLevel && CachedLoadedLevelPackageName == DesiredPackageName)
+	if (LoadedLevel && LoadedLevelPackageName == DesiredPackageName)
 	{
 		return true;
 	}
@@ -1355,6 +1358,7 @@ void ULevelStreaming::SetWorldAsset(const TSoftObjectPtr<UWorld>& NewWorldAsset)
 	{
 		WorldAsset = NewWorldAsset;
 		bHasCachedWorldAssetPackageFName = false;
+		bHasCachedLoadedLevelPackageName = false;
 
 		if (CurrentState == ECurrentState::FailedToLoad)
 		{
@@ -1383,6 +1387,17 @@ FName ULevelStreaming::GetWorldAssetPackageFName() const
 	return CachedWorldAssetPackageFName;
 }
 
+FName ULevelStreaming::GetLoadedLevelPackageName() const
+{
+	if( !bHasCachedLoadedLevelPackageName )
+	{
+		CachedLoadedLevelPackageName = (LoadedLevel ? LoadedLevel->GetOutermost()->GetFName() : NAME_None);
+		bHasCachedLoadedLevelPackageName = true;
+	}
+
+	return CachedLoadedLevelPackageName;
+}
+
 void ULevelStreaming::SetWorldAssetByPackageName(FName InPackageName)
 {
 	const FString TargetWorldPackageName = InPackageName.ToString();
@@ -1549,6 +1564,7 @@ void ULevelStreaming::PostEditChangeProperty(FPropertyChangedEvent& PropertyChan
 		else if (PropertyName == GET_MEMBER_NAME_CHECKED(ULevelStreaming, WorldAsset))
 		{
 			bHasCachedWorldAssetPackageFName = false;
+			bHasCachedLoadedLevelPackageName = false;
 		}
 	}
 
