Passing the result of GetMousePosition directly into SetMousePosition in a blueprint no longer causes the cursor to drift to the top-left corner.

diff --git a/Engine/Source/Runtime/Engine/Private/Slate/SceneViewport.cpp b/Engine/Source/Runtime/Engine/Private/Slate/SceneViewport.cpp
--- a/Engine/Source/Runtime/Engine/Private/Slate/SceneViewport.cpp
+++ b/Engine/Source/Runtime/Engine/Private/Slate/SceneViewport.cpp
@@ -198,7 +198,7 @@ void FSceneViewport::SetMouse( int32 X, int32 Y )
 {
 	const FVector2D NormalizedLocalMousePosition = FVector2D(X, Y) / GetSizeXY();
 	FVector2D AbsolutePos = CachedGeometry.LocalToAbsolute(NormalizedLocalMousePosition * CachedGeometry.GetLocalSize());
-	FSlateApplication::Get().SetCursorPos( AbsolutePos );
+	FSlateApplication::Get().SetCursorPos( AbsolutePos.RoundToVector() );
 	CachedCursorPos = FIntPoint(X, Y);
 }
 
