When rebuilding geometry, calling CleanupWorld had the side-effect of setting the bCleanedUpWorld on the worlds (which is used to avoid calling CleanupWorld multiple times on streaming levels), preventing proper cleanup when unloading this world later on. Reset the bCleanedUpWorld flag after every call to CleanupWorld.

#rnx

diff --git a/Engine/Source/Runtime/Engine/Classes/Engine/World.h b/Engine/Source/Runtime/Engine/Classes/Engine/World.h
--- a/Engine/Source/Runtime/Engine/Classes/Engine/World.h
+++ b/Engine/Source/Runtime/Engine/Classes/Engine/World.h
@@ -2324,10 +2324,11 @@ public:
 
 	/**
 	 * Cleans up components, streaming data and assorted other intermediate data.
-	 * @param bSessionEnded whether to notify the viewport that the game session has ended
+	 * @param bSessionEnded whether to notify the viewport that the game session has ended.
 	 * @param NewWorld Optional new world that will be loaded after this world is cleaned up. Specify a new world to prevent it and it's sublevels from being GCed during map transitions.
+	 * @param bResetCleanedUpFlag wheter to reset the bCleanedUpWorld flag or not.
 	 */
-	void CleanupWorld(bool bSessionEnded = true, bool bCleanupResources = true, UWorld* NewWorld = nullptr);
+	void CleanupWorld(bool bSessionEnded = true, bool bCleanupResources = true, UWorld* NewWorld = nullptr, bool bResetCleanedUpFlag = true);
 	
 	/**
 	 * Invalidates the cached data used to render the levels' UModel.
diff --git a/Engine/Source/Runtime/Engine/Private/World.cpp b/Engine/Source/Runtime/Engine/Private/World.cpp
--- a/Engine/Source/Runtime/Engine/Private/World.cpp
+++ b/Engine/Source/Runtime/Engine/Private/World.cpp
@@ -3965,13 +3965,22 @@ bool UWorld::IsNavigationRebuilt() const
 	return GetNavigationSystem() == NULL || GetNavigationSystem()->IsNavigationBuilt(GetWorldSettings());
 }
 
-void UWorld::CleanupWorld(bool bSessionEnded, bool bCleanupResources, UWorld* NewWorld)
+void UWorld::CleanupWorld(bool bSessionEnded, bool bCleanupResources, UWorld* NewWorld, bool bResetCleanedUpFlag)
 {
 	UE_LOG(LogWorld, Log, TEXT("UWorld::CleanupWorld for %s, bSessionEnded=%s, bCleanupResources=%s"), *GetName(), bSessionEnded ? TEXT("true") : TEXT("false"), bCleanupResources ? TEXT("true") : TEXT("false"));
 
+	TArray<UWorld*> WorldsToResetCleanedUpFlag;
+
 	check(IsVisibilityRequestPending() == false);
+	
+	check(!bCleanedUpWorld);
 	bCleanedUpWorld = true;
 
+	if (bResetCleanedUpFlag)
+	{
+		WorldsToResetCleanedUpFlag.Add(this);
+	}
+
 	// Wait on current physics scenes if they are processing
 	if(FPhysScene* CurrPhysicsScene = GetPhysicsScene())
 	{
@@ -4065,7 +4074,12 @@ void UWorld::CleanupWorld(bool bSessionEnded, bool bCleanupResources, UWorld* Ne
 		UWorld* World = CastChecked<UWorld>(GetLevel(LevelIndex)->GetOuter());
 		if (!World->bCleanedUpWorld)
 		{
-			World->CleanupWorld(bSessionEnded, bCleanupResources, NewWorld);
+			World->CleanupWorld(bSessionEnded, bCleanupResources, NewWorld, false);
+
+			if (bResetCleanedUpFlag)
+			{
+				WorldsToResetCleanedUpFlag.Add(World);
+			}
 		}
 	}
 
@@ -4076,7 +4090,12 @@ void UWorld::CleanupWorld(bool bSessionEnded, bool bCleanupResources, UWorld* Ne
 			UWorld* World = CastChecked<UWorld>(Level->GetOuter());
 			if (!World->bCleanedUpWorld)
 			{
-				World->CleanupWorld(bSessionEnded, bCleanupResources, NewWorld);
+				World->CleanupWorld(bSessionEnded, bCleanupResources, NewWorld, false);
+
+				if (bResetCleanedUpFlag)
+				{
+					WorldsToResetCleanedUpFlag.Add(World);
+				}
 			}
 		}
 	}
@@ -4095,7 +4114,12 @@ void UWorld::CleanupWorld(bool bSessionEnded, bool bCleanupResources, UWorld* Ne
 			UWorld* const LevelWorld = CastChecked<UWorld>(Level->GetOuter());
 			if (!LevelWorld->bCleanedUpWorld)
 			{
-				LevelWorld->CleanupWorld(bSessionEnded, bCleanupResources, NewWorld);
+				LevelWorld->CleanupWorld(bSessionEnded, bCleanupResources, NewWorld, false);
+
+				if (bResetCleanedUpFlag)
+				{
+					WorldsToResetCleanedUpFlag.Add(LevelWorld);
+				}
 			}
 		}
 	}
@@ -4103,6 +4127,12 @@ void UWorld::CleanupWorld(bool bSessionEnded, bool bCleanupResources, UWorld* Ne
 	PSCPool.Cleanup();
 
 	FWorldDelegates::OnPostWorldCleanup.Broadcast(this, bSessionEnded, bCleanupResources);
+
+	for (UWorld* WorldToResetCleanedUpFlag: WorldsToResetCleanedUpFlag)
+	{
+		check(WorldToResetCleanedUpFlag->bCleanedUpWorld);
+		WorldToResetCleanedUpFlag->bCleanedUpWorld = false;
+	}
 }
 
 UGameViewportClient* UWorld::GetGameViewport() const
