PR #5833: Incorrect values of TextureSize and GradientPixelSize in cooked builds. (Contributed by [REDACTED_PERSON])

diff --git a/Engine/Source/Runtime/Engine/Private/Curves/CurveLinearColorAtlas.cpp b/Engine/Source/Runtime/Engine/Private/Curves/CurveLinearColorAtlas.cpp
--- a/Engine/Source/Runtime/Engine/Private/Curves/CurveLinearColorAtlas.cpp
+++ b/Engine/Source/Runtime/Engine/Private/Curves/CurveLinearColorAtlas.cpp
@@ -13,9 +13,9 @@
 UCurveLinearColorAtlas::UCurveLinearColorAtlas(const FObjectInitializer& ObjectInitializer)
 	: Super(ObjectInitializer)
 {
-#if WITH_EDITOR
 	TextureSize = 256;
 	GradientPixelSize = 1;
+#if WITH_EDITORONLY_DATA
 	bHasAnyDirtyTextures = false;
 	bShowDebugColorsForNullGradients = false;
 	SizeXY = { (float)TextureSize, (float)GradientPixelSize };
@@ -209,4 +209,4 @@ bool UCurveLinearColorAtlas::GetCurvePosition(UCurveLinearColor* InCurve, float&
 		return true;
 	}
 	return false;
-}
\ No newline at end of file
+}
