Make sure that the reorder table match the PerInstanceSMData, otherwise we should regenerate the data and it's not valid.

diff --git a/Engine/Source/Runtime/Engine/Private/HierarchicalInstancedStaticMesh.cpp b/Engine/Source/Runtime/Engine/Private/HierarchicalInstancedStaticMesh.cpp
--- a/Engine/Source/Runtime/Engine/Private/HierarchicalInstancedStaticMesh.cpp
+++ b/Engine/Source/Runtime/Engine/Private/HierarchicalInstancedStaticMesh.cpp
@@ -2980,9 +2980,10 @@ void UHierarchicalInstancedStaticMeshComponent::OnPostLoadPerInstanceData()
 			ULevel* OwnerLevel = Owner->GetLevel();
 			UWorld* OwnerWorld = OwnerLevel ? OwnerLevel->OwningWorld : nullptr;
 
-			if (OwnerWorld && OwnerWorld->GetActiveLightingScenario() != nullptr && OwnerWorld->GetActiveLightingScenario() != OwnerLevel)
+			//update the instance data if the lighting scenario isn't the owner level or if the reorder table do not match the per instance sm data
+			if (OwnerWorld && OwnerWorld->GetActiveLightingScenario() != nullptr && OwnerWorld->GetActiveLightingScenario() != OwnerLevel
+				|| (PerInstanceSMData.Num() > 0 && PerInstanceSMData.Num() != InstanceReorderTable.Num()))
 			{
-				//update the instance data if the lighting scenario isn't the owner level
 				bForceTreeBuild = true;
 			}
 
diff --git a/Engine/Source/Runtime/Engine/Private/InstancedStaticMesh.cpp b/Engine/Source/Runtime/Engine/Private/InstancedStaticMesh.cpp
--- a/Engine/Source/Runtime/Engine/Private/InstancedStaticMesh.cpp
+++ b/Engine/Source/Runtime/Engine/Private/InstancedStaticMesh.cpp
@@ -2224,7 +2224,8 @@ void UInstancedStaticMeshComponent::OnComponentCreated()
 	{
 		// if we are pasting/duplicating this component, it may be created with some instances already in place
 		// in this case, need to ensure that the instance render data is properly created
-		const bool InitializeFromCurrentData = PerInstanceSMData.Num() > 0;
+		// We only need to only init from current data if the reorder table == per instance data, but only for the HISM Component, in the case of ISM, the reorder table is never used.
+		const bool InitializeFromCurrentData = PerInstanceSMData.Num() > 0 && (InstanceReorderTable.Num() == PerInstanceSMData.Num() || InstanceReorderTable.Num() == 0);
 		InitPerInstanceRenderData(InitializeFromCurrentData);
 	}
 }
