Attach to Pawn is no longer an Advanced property and is thus visible in the details view of a Controller blueprint again.

diff --git a/Engine/Source/Runtime/Engine/Classes/GameFramework/Controller.h b/Engine/Source/Runtime/Engine/Classes/GameFramework/Controller.h
--- a/Engine/Source/Runtime/Engine/Classes/GameFramework/Controller.h
+++ b/Engine/Source/Runtime/Engine/Classes/GameFramework/Controller.h
@@ -92,7 +92,7 @@ protected:
 	 * an Actor that follows the possessed Pawn location, but that still has the full aim rotation (since a Pawn might
 	 * update only some components of the rotation).
 	 */
-	UPROPERTY(EditDefaultsOnly, AdvancedDisplay, Category="Controller|Transform")
+	UPROPERTY(EditDefaultsOnly, Category="Controller|Transform")
 	uint8 bAttachToPawn:1;
 
 	/** Whether this controller is a PlayerController. */
