Rename IsFunctionImplementedInBlueprint to IsFunctionImplementedInScript

diff --git a/Engine/Plugins/Experimental/PythonScriptPlugin/Source/PythonScriptPlugin/Private/PyWrapperObject.cpp b/Engine/Plugins/Experimental/PythonScriptPlugin/Source/PythonScriptPlugin/Private/PyWrapperObject.cpp
--- a/Engine/Plugins/Experimental/PythonScriptPlugin/Source/PythonScriptPlugin/Private/PyWrapperObject.cpp
+++ b/Engine/Plugins/Experimental/PythonScriptPlugin/Source/PythonScriptPlugin/Private/PyWrapperObject.cpp
@@ -1687,7 +1687,7 @@ void UPythonGeneratedClass::PostInitInstance(UObject* InObj)
 	}
 }
 
-bool UPythonGeneratedClass::IsFunctionImplementedInBlueprint(FName InFunctionName) const
+bool UPythonGeneratedClass::IsFunctionImplementedInScript(FName InFunctionName) const
 {
 	UFunction* Function = FindFunctionByName(InFunctionName);
 	return Function && Function->GetOuter() && Function->GetOuter()->IsA(UPythonGeneratedClass::StaticClass());
diff --git a/Engine/Plugins/Experimental/PythonScriptPlugin/Source/PythonScriptPlugin/Private/PyWrapperObject.h b/Engine/Plugins/Experimental/PythonScriptPlugin/Source/PythonScriptPlugin/Private/PyWrapperObject.h
--- a/Engine/Plugins/Experimental/PythonScriptPlugin/Source/PythonScriptPlugin/Private/PyWrapperObject.h
+++ b/Engine/Plugins/Experimental/PythonScriptPlugin/Source/PythonScriptPlugin/Private/PyWrapperObject.h
@@ -193,7 +193,7 @@ public:
 	//~ UClass interface
 	virtual void PostInitInstance(UObject* InObj) override;
 
-	virtual bool IsFunctionImplementedInBlueprint(FName InFunctionName) const override;
+	virtual bool IsFunctionImplementedInScript(FName InFunctionName) const override;
 
 	/** Generate an Unreal class from the given Python type */
 	static UPythonGeneratedClass* GenerateClass(PyTypeObject* InPyType);
diff --git a/Engine/Source/Runtime/CoreUObject/Private/UObject/Class.cpp b/Engine/Source/Runtime/CoreUObject/Private/UObject/Class.cpp
--- a/Engine/Source/Runtime/CoreUObject/Private/UObject/Class.cpp
+++ b/Engine/Source/Runtime/CoreUObject/Private/UObject/Class.cpp
@@ -4008,9 +4008,9 @@ UClass* UClass::FindCommonBase(const TArray<UClass*>& InClasses)
 	return CommonClass;
 }
 
-bool UClass::IsFunctionImplementedInBlueprint(FName InFunctionName) const
+bool UClass::IsFunctionImplementedInScript(FName InFunctionName) const
 {
-	// Implemented in UBlueprintGeneratedClass
+	// Implemented in classes such as UBlueprintGeneratedClass
 	return false;
 }
 
diff --git a/Engine/Source/Runtime/CoreUObject/Public/UObject/Class.h b/Engine/Source/Runtime/CoreUObject/Public/UObject/Class.h
--- a/Engine/Source/Runtime/CoreUObject/Public/UObject/Class.h
+++ b/Engine/Source/Runtime/CoreUObject/Public/UObject/Class.h
@@ -2725,7 +2725,10 @@ public:
 	 * @param InFunctionName	The name of the function to test
 	 * @return					True if the specified function exists and is implemented in a blueprint generated class
 	 */
-	virtual bool IsFunctionImplementedInBlueprint(FName InFunctionName) const;
+	virtual bool IsFunctionImplementedInScript(FName InFunctionName) const;
+
+	UE_DEPRECATED(4.23, "IsFunctionImplementedInBlueprint is deprecated, call IsFunctionImplementedInScript instead")
+	bool IsFunctionImplementedInBlueprint(FName InFunctionName) const { return IsFunctionImplementedInScript(InFunctionName); }
 
 	/**
 	 * Checks if the property exists on this class or a parent class.
diff --git a/Engine/Source/Runtime/Engine/Classes/Engine/BlueprintGeneratedClass.h b/Engine/Source/Runtime/Engine/Classes/Engine/BlueprintGeneratedClass.h
--- a/Engine/Source/Runtime/Engine/Classes/Engine/BlueprintGeneratedClass.h
+++ b/Engine/Source/Runtime/Engine/Classes/Engine/BlueprintGeneratedClass.h
@@ -704,7 +704,7 @@ public:
 #endif //WITH_EDITOR
 	virtual void SerializeDefaultObject(UObject* Object, FArchive& Ar) override;
 	virtual void PostLoadDefaultObject(UObject* Object) override;
-	virtual bool IsFunctionImplementedInBlueprint(FName InFunctionName) const override;
+	virtual bool IsFunctionImplementedInScript(FName InFunctionName) const override;
 	virtual uint8* GetPersistentUberGraphFrame(UObject* Obj, UFunction* FuncToCheck) const override;
 	virtual void CreatePersistentUberGraphFrame(UObject* Obj, bool bCreateOnlyIfEmpty = false, bool bSkipSuperClass = false, UClass* OldClass = nullptr) const override;
 	virtual void DestroyPersistentUberGraphFrame(UObject* Obj, bool bSkipSuperClass = false) const override;
diff --git a/Engine/Source/Runtime/Engine/Private/BlueprintGeneratedClass.cpp b/Engine/Source/Runtime/Engine/Private/BlueprintGeneratedClass.cpp
--- a/Engine/Source/Runtime/Engine/Private/BlueprintGeneratedClass.cpp
+++ b/Engine/Source/Runtime/Engine/Private/BlueprintGeneratedClass.cpp
@@ -684,7 +684,7 @@ void UBlueprintGeneratedClass::InitArrayPropertyFromCustomList(const UArrayPrope
 	}
 }
 
-bool UBlueprintGeneratedClass::IsFunctionImplementedInBlueprint(FName InFunctionName) const
+bool UBlueprintGeneratedClass::IsFunctionImplementedInScript(FName InFunctionName) const
 {
 	UFunction* Function = FindFunctionByName(InFunctionName);
 	return Function && Function->GetOuter() && Function->GetOuter()->IsA(UBlueprintGeneratedClass::StaticClass());
