Remove edit condition on UPhysicsSettings's bEnableShapeSharing as bEnableAsyncScene no longer exists.

diff --git a/Engine/Source/Runtime/Engine/Classes/PhysicsEngine/PhysicsSettings.h b/Engine/Source/Runtime/Engine/Classes/PhysicsEngine/PhysicsSettings.h
--- a/Engine/Source/Runtime/Engine/Classes/PhysicsEngine/PhysicsSettings.h
+++ b/Engine/Source/Runtime/Engine/Classes/PhysicsEngine/PhysicsSettings.h
@@ -106,7 +106,7 @@ class ENGINE_API UPhysicsSettings : public UDeveloperSettings
 	float TriangleMeshTriangleMinAreaThreshold;
 	
 	/** Enables shape sharing between sync and async scene for static rigid actors */
-	UPROPERTY(config, EditAnywhere, AdvancedDisplay, Category = Simulation, meta = (editcondition="bEnableAsyncScene"))
+	UPROPERTY(config, EditAnywhere, AdvancedDisplay, Category = Simulation)
 	bool bEnableShapeSharing;
 
 	/** Enables persistent contact manifolds. This will generate fewer contact points, but with more accuracy. Reduces stability of stacking, but can help energy conservation.*/
