Fix EditConditions on FBroadphaseSettings

diff --git a/Engine/Source/Runtime/Engine/Classes/GameFramework/WorldSettings.h b/Engine/Source/Runtime/Engine/Classes/GameFramework/WorldSettings.h
--- a/Engine/Source/Runtime/Engine/Classes/GameFramework/WorldSettings.h
+++ b/Engine/Source/Runtime/Engine/Classes/GameFramework/WorldSettings.h
@@ -386,15 +386,15 @@ struct FBroadphaseSettings
 	bool bUseMBPOuterBounds;
 
 	/** Total bounds for MBP, must cover the game world or collisions are disabled for out of bounds actors */
-	UPROPERTY(EditAnywhere, Category = Broadphase, meta = (EditCondition = bUseMBP))
+	UPROPERTY(EditAnywhere, Category = Broadphase, meta = (EditCondition = "bUseMBPOnClient || bUseMBPOnServer"))
 	FBox MBPBounds;
 
 	/** Total bounds for MBP, should cover absolute maximum bounds of the game world where physics is required */
-	UPROPERTY(EditAnywhere, Category = Broadphase, meta = (EditCondition = bUseMBP))
+	UPROPERTY(EditAnywhere, Category = Broadphase, meta = (EditCondition = "bUseMBPOnClient || bUseMBPOnServer"))
 	FBox MBPOuterBounds;
 
 	/** Number of times to subdivide the MBP bounds, final number of regions is MBPNumSubdivs^2 */
-	UPROPERTY(EditAnywhere, Category = Broadphase, meta = (EditCondition = bUseMBP, ClampMin=1, ClampMax=16))
+	UPROPERTY(EditAnywhere, Category = Broadphase, meta = (EditCondition = "bUseMBPOnClient || bUseMBPOnServer", ClampMin=1, ClampMax=16))
 	uint32 MBPNumSubdivs;
 };
 
