Do not add hide instance command if they do not have corresponding instance on the render thread, this is to fix a edge case that a user is/can be calling multiple time in the same frame Clear, Add, Clear Add, ...

diff --git a/Engine/Source/Runtime/Engine/Private/HierarchicalInstancedStaticMesh.cpp b/Engine/Source/Runtime/Engine/Private/HierarchicalInstancedStaticMesh.cpp
--- a/Engine/Source/Runtime/Engine/Private/HierarchicalInstancedStaticMesh.cpp
+++ b/Engine/Source/Runtime/Engine/Private/HierarchicalInstancedStaticMesh.cpp
@@ -2324,19 +2324,25 @@ void UHierarchicalInstancedStaticMeshComponent::ClearInstances()
 
 	InstanceUpdateCmdBuffer.Reset();
 
-	// Hide all instance until the build tree is completed
-	int32 NumInstances = PerInstanceSMData.Num();
-
-	for (int32 Index = 0; Index < NumInstances; ++Index)
+	// Don't try to queue hide command if there is no render instances
+	if (PerInstanceRenderData.IsValid() && PerInstanceRenderData->InstanceBuffer_GameThread->GetNumInstances() > 0)
 	{
-		int32 RenderIndex = InstanceReorderTable.IsValidIndex(Index) ? InstanceReorderTable[Index] : Index;
-		if (RenderIndex == INDEX_NONE)
+		// Hide all instance until the build tree is completed but if there is a mismatch between game thread data and render thread data, only add command matching render thread data, 
+		// this can happen in a case where you perform many time add, clear, add, clear, in the same frame, so you might get a mismatch between both thread data
+		int32 NumInstances = FMath::Clamp(PerInstanceSMData.Num(), 0, PerInstanceRenderData->InstanceBuffer_GameThread->GetNumInstances());
+
+		for (int32 Index = 0; Index < NumInstances; ++Index)
 		{
-			// could be skipped by density settings
-			continue;
-		}
+			int32 RenderIndex = InstanceReorderTable.IsValidIndex(Index) ? InstanceReorderTable[Index] : Index;
 
-		InstanceUpdateCmdBuffer.HideInstance(RenderIndex);
+			if (RenderIndex == INDEX_NONE) 
+			{
+				// could be skipped by density settings
+				continue;
+			}
+
+			InstanceUpdateCmdBuffer.HideInstance(RenderIndex);
+		}
 	}
 
 	// Clear all the per-instance data
