Fix compile error in HierarchicalInstancedStaticMesh.cpp
#rnx

diff --git a/Engine/Source/Runtime/Engine/Private/HierarchicalInstancedStaticMesh.cpp b/Engine/Source/Runtime/Engine/Private/HierarchicalInstancedStaticMesh.cpp
--- a/Engine/Source/Runtime/Engine/Private/HierarchicalInstancedStaticMesh.cpp
+++ b/Engine/Source/Runtime/Engine/Private/HierarchicalInstancedStaticMesh.cpp
@@ -2987,7 +2987,7 @@ void UHierarchicalInstancedStaticMeshComponent::OnPostLoadPerInstanceData()
 			UWorld* OwnerWorld = OwnerLevel ? OwnerLevel->OwningWorld : nullptr;
 
 			//update the instance data if the lighting scenario isn't the owner level or if the reorder table do not match the per instance sm data
-			if (OwnerWorld && OwnerWorld->GetActiveLightingScenario() != nullptr && OwnerWorld->GetActiveLightingScenario() != OwnerLevel
+			if ((OwnerWorld && OwnerWorld->GetActiveLightingScenario() != nullptr && OwnerWorld->GetActiveLightingScenario() != OwnerLevel)
 				|| (PerInstanceSMData.Num() > 0 && PerInstanceSMData.Num() != InstanceReorderTable.Num()))
 			{
 				bForceTreeBuild = true;
