Fix a broken EditCondition

#rnx

diff --git a/Engine/Source/Runtime/Engine/Classes/Components/LocalLightComponent.h b/Engine/Source/Runtime/Engine/Classes/Components/LocalLightComponent.h
--- a/Engine/Source/Runtime/Engine/Classes/Components/LocalLightComponent.h
+++ b/Engine/Source/Runtime/Engine/Classes/Components/LocalLightComponent.h
@@ -23,7 +23,7 @@ class ENGINE_API ULocalLightComponent : public ULightComponent
 	 * The peak luminous intensity is measured in candelas,
 	 * while the luminous power is measured in lumens.
 	 */
-	UPROPERTY(EditAnywhere, BlueprintReadOnly, Category=Light, AdvancedDisplay, meta=(DisplayName="Intensity Units", EditCondition="bUseInverseSquaredFalloff"))
+	UPROPERTY(EditAnywhere, BlueprintReadOnly, Category=Light, AdvancedDisplay, meta=(DisplayName="Intensity Units"))
 	ELightUnits IntensityUnits;
 
 	UPROPERTY()
