PR #5912: Fixing DPI scaling issues for the ReporterGraph (Contributed by [REDACTED_PERSON])

diff --git a/Engine/Source/Runtime/Engine/Private/Debug/ReporterBase.cpp b/Engine/Source/Runtime/Engine/Private/Debug/ReporterBase.cpp
--- a/Engine/Source/Runtime/Engine/Private/Debug/ReporterBase.cpp
+++ b/Engine/Source/Runtime/Engine/Private/Debug/ReporterBase.cpp
@@ -17,9 +17,11 @@ UReporterBase::UReporterBase(const FObjectInitializer& ObjectInitializer) :
 
 FVector2D UReporterBase::ToScreenSpace(const FVector2D& InVector, UCanvas* Canvas)
 {
+	const float DPIScale = Canvas->GetDPIScale();
+
 	FVector2D OutVector = InVector;
-	OutVector.X *= Canvas->SizeX;
-	OutVector.Y *= Canvas->SizeY;
+	OutVector.X *= (Canvas->SizeX / DPIScale);
+	OutVector.Y *= (Canvas->SizeY / DPIScale);
 
 	return OutVector;
 }
diff --git a/Engine/Source/Runtime/Engine/Private/Debug/ReporterGraph.cpp b/Engine/Source/Runtime/Engine/Private/Debug/ReporterGraph.cpp
--- a/Engine/Source/Runtime/Engine/Private/Debug/ReporterGraph.cpp
+++ b/Engine/Source/Runtime/Engine/Private/Debug/ReporterGraph.cpp
@@ -142,8 +142,8 @@ void UReporterGraph::SetLegendPosition(ELegendPosition::Type Position)
 
 void UReporterGraph::DrawBackground(UCanvas* Canvas)
 {
-	FVector2D Min = FVector2D(GraphScreenSize.Min.X * Canvas->SizeX, Canvas->SizeY - GraphScreenSize.Min.Y * Canvas->SizeY);
-	FVector2D Max = FVector2D(GraphScreenSize.Max.X * Canvas->SizeX, Canvas->SizeY - GraphScreenSize.Max.Y * Canvas->SizeY);
+	FVector2D Min = ToScreenSpace(GraphScreenSize.Min, Canvas);
+	FVector2D Max = ToScreenSpace(GraphScreenSize.Max, Canvas);
 
 	FCanvasTileItem TileItem(Min, GWhiteTexture, Max-Min, BackgroundColor);
 	TileItem.BlendMode = SE_BLEND_Translucent;
@@ -186,6 +186,7 @@ void UReporterGraph::DrawLegend(UCanvas* Canvas)
 
 void UReporterGraph::DrawAxes(UCanvas* Canvas)
 {
+	const float DPIScale = Canvas->GetDPIScale();
 	FVector2D Min = GraphScreenSize.Min;
 	
 	FVector2D XMax = GraphScreenSize.Min;
@@ -198,13 +199,13 @@ void UReporterGraph::DrawAxes(UCanvas* Canvas)
 	int32 StringSizeX, StringSizeY;
 	// Draw the X axis
 	StringSize(Font, StringSizeX, StringSizeY, *FString::Printf(TEXT("%.2f"), GraphMinMaxData.Max.X) );
-	const float SizeX = (XMax.X - Min.X) * Canvas->SizeX;
+	const float SizeX = ((XMax.X - Min.X) * Canvas->SizeX) / DPIScale;
 	NumXNotches = FMath::CeilToInt(SizeX * 0.7 / StringSizeX);
 	DrawAxis(Canvas, Min, XMax, NumXNotches, false);
 
 	// Draw the Y axis
 	StringSize(Font, StringSizeX, StringSizeY, *FString::Printf(TEXT("%.2f"), GraphMinMaxData.Max.Y));
-	float SizeY = (YMax.Y - Min.Y) * Canvas->SizeY;
+	float SizeY = ((YMax.Y - Min.Y) * Canvas->SizeY) / DPIScale;
 	NumYNotches = FMath::CeilToInt(SizeY * 0.7 / StringSizeY);
 	DrawAxis(Canvas, Min, YMax, NumYNotches, true);
 }
@@ -361,8 +362,8 @@ void UReporterGraph::DrawData(UCanvas* Canvas)
 {
 	FVector2D Start, End;
 	
-	const FVector2D Min = FVector2D(GraphScreenSize.Min.X * Canvas->SizeX, Canvas->SizeY - GraphScreenSize.Min.Y * Canvas->SizeY);
-	const FVector2D Max = FVector2D(GraphScreenSize.Max.X * Canvas->SizeX, Canvas->SizeY - GraphScreenSize.Max.Y * Canvas->SizeY);
+	const FVector2D Min = ToScreenSpace(GraphScreenSize.Min, Canvas);
+	const FVector2D Max = ToScreenSpace(GraphScreenSize.Max, Canvas);
 	const float Height = GraphScreenSize.Max.Y - GraphScreenSize.Min.Y;
 	const float dx = Height / FMath::Abs(Max.Y - Min.Y);
 	UFont* Font = GetDefaultFont();
@@ -441,9 +442,11 @@ void UReporterGraph::DrawData(UCanvas* Canvas)
 
 FVector2D UReporterGraph::ToScreenSpace(const FVector2D& InVector, UCanvas* Canvas)
 {
+	const float DPIScale = Canvas->GetDPIScale();
+
 	FVector2D OutVector = InVector;
-	OutVector.X *= Canvas->SizeX;
-	OutVector.Y = Canvas->SizeY - (OutVector.Y * Canvas->SizeY);
+	OutVector.X = (OutVector.X * Canvas->SizeX) / DPIScale;
+	OutVector.Y = (Canvas->SizeY - (OutVector.Y * Canvas->SizeY)) / DPIScale;
 
 	return OutVector;
 }
