Fixing highlight for some meshes/actors don't highlight in orthographic view when selected

diff --git a/Engine/Source/Runtime/Engine/Private/Components/TextRenderComponent.cpp b/Engine/Source/Runtime/Engine/Private/Components/TextRenderComponent.cpp
--- a/Engine/Source/Runtime/Engine/Private/Components/TextRenderComponent.cpp
+++ b/Engine/Source/Runtime/Engine/Private/Components/TextRenderComponent.cpp
@@ -465,11 +465,6 @@ public:
 		}
 	}
 
-	virtual FString GetReferencerName() const override
-	{
-		return TEXT("FTextRenderComponentMIDCache");
-	}
-
 private:
 	/** Key identifying an array of MIDs */
 	struct FKey
@@ -743,6 +738,7 @@ void FTextRenderSceneProxy::GetDynamicMeshElements(const TArray<const FSceneView
 					Mesh.MaterialRenderProxy = TextBatch.Material->GetRenderProxy();
 					Mesh.bCanApplyViewModeOverrides = !bAlwaysRenderAsText;
 					Mesh.LODIndex = 0;
+					Mesh.bUseWireframeSelectionColoring = IsSelected() ? 1 : 0;
 
 					Collector.AddMesh(ViewIndex, Mesh);
 				}
diff --git a/Engine/Source/Runtime/Engine/Private/StaticMeshRender.cpp b/Engine/Source/Runtime/Engine/Private/StaticMeshRender.cpp
--- a/Engine/Source/Runtime/Engine/Private/StaticMeshRender.cpp
+++ b/Engine/Source/Runtime/Engine/Private/StaticMeshRender.cpp
@@ -1111,7 +1111,7 @@ void FStaticMeshSceneProxy::GetDynamicMeshElements(const TArray<const FSceneView
 									{
 										const FLODInfo::FSectionInfo& Section = LODs[LODIndex].Sections[SectionIndex];
 
-										bSectionIsSelected = Section.bSelected;
+										bSectionIsSelected = Section.bSelected || (bIsWireframeView && bProxyIsSelected);
 										MeshElement.BatchHitProxyId = Section.HitProxy ? Section.HitProxy->Id : FHitProxyId();
 									}
 	#endif // WITH_EDITOR
