PR #5919: Fixing DPI scaling issues for the Console (Contributed by [REDACTED_PERSON])

diff --git a/Engine/Source/Runtime/Engine/Private/UserInterface/Console.cpp b/Engine/Source/Runtime/Engine/Private/UserInterface/Console.cpp
--- a/Engine/Source/Runtime/Engine/Private/UserInterface/Console.cpp
+++ b/Engine/Source/Runtime/Engine/Private/UserInterface/Console.cpp
@@ -1203,9 +1203,14 @@ void UConsole::PostRender_Console_Open(UCanvas* Canvas)
 
 	UFont* Font = GEngine->GetSmallFont();
 
+	const float DPIScale = Canvas->GetDPIScale();
+	const bool bDPIAwareStringMeasurement = true;
+
 	// determine the height of the text
 	float xl, yl;
-	Canvas->StrLen(Font, TEXT("M"),xl,yl);
+	Canvas->StrLen(Font, TEXT("M"),xl,yl, bDPIAwareStringMeasurement);
+	xl /= DPIScale;
+	yl /= DPIScale;
 	// Background
 	FLinearColor BackgroundColor = ConsoleDefs::AutocompleteBackgroundColor.ReinterpretAsLinear();
 	BackgroundColor.A = ConsoleSettings->BackgroundOpacityPercentage / 100.0f;
@@ -1240,7 +1245,9 @@ void UConsole::PostRender_Console_Open(UCanvas* Canvas)
 			if (idx < Scrollback.Num())
 			{
 				float ScrollLineXL, ScrollLineYL;
-				Canvas->StrLen(Font, Scrollback[idx], ScrollLineXL, ScrollLineYL);
+				Canvas->StrLen(Font, Scrollback[idx], ScrollLineXL, ScrollLineYL, bDPIAwareStringMeasurement);
+				ScrollLineXL /= DPIScale;
+				ScrollLineYL /= DPIScale;
 				if (ScrollLineYL > yl)
 				{
 					y -= (ScrollLineYL - yl);
@@ -1334,10 +1341,15 @@ void UConsole::PostRender_InputLine(UCanvas* Canvas, FIntPoint UserInputLinePos)
 	const FString TypedInputText = FString::Printf(TEXT("%s%s"), *ConsoleDefs::LeadingInputText, *TypedStr);
 	const FString PrecompletedInputText = !PrecompletedInputLine.IsEmpty() ? PrecompletedInputLine.RightChop(TypedStr.Len()) : FString();
 
+	const float DPIScale = Canvas->GetDPIScale();
+	const bool bDPIAwareStringMeasurement = true;
+
 	// use the smallest font
 	UFont* Font = GEngine->GetSmallFont();
 	// determine the size of the input line
-	Canvas->StrLen(Font, TypedInputText, xl, yl);
+	Canvas->StrLen(Font, TypedInputText, xl, yl, bDPIAwareStringMeasurement);
+	xl /= DPIScale;
+	yl /= DPIScale;
 
 	float ClipX = Canvas->ClipX;
 	float ClipY = Canvas->ClipY;
@@ -1393,7 +1405,10 @@ void UConsole::PostRender_InputLine(UCanvas* Canvas, FIntPoint UserInputLinePos)
 			StartIdx = FMath::Max(0, AutoComplete.Num() + StartIdx);
 		}
 
-		Canvas->StrLen(Font, *ConsoleDefs::LeadingInputText, xl, yl);
+		Canvas->StrLen(Font, *ConsoleDefs::LeadingInputText, xl, yl, bDPIAwareStringMeasurement);
+		xl /= DPIScale;
+		yl /= DPIScale;
+
 		float y = UserInputLinePos.Y - 6.0f - (yl * 2.0f);
 
 		// Set the background color/texture of the auto-complete section
@@ -1416,12 +1431,16 @@ void UConsole::PostRender_InputLine(UCanvas* Canvas, FIntPoint UserInputLinePos)
 
 			// Find the longest command and the longest description for left-justification of the descriptions
 			float CmdLenX, CmdLenY;
-			Canvas->StrLen(Font, Cmd.GetLeft(), CmdLenX, CmdLenY);
+			Canvas->StrLen(Font, Cmd.GetLeft(), CmdLenX, CmdLenY, bDPIAwareStringMeasurement);
+			CmdLenX /= DPIScale;
+			CmdLenY /= DPIScale;
 			MaxLeftWidth = FMath::Max(MaxLeftWidth, CmdLenX);
 			if (!Cmd.Desc.IsEmpty())
 			{
 				float DescLenX, DescLenY;
-				Canvas->StrLen(Font, Cmd.GetRight(), DescLenX, DescLenY);
+				Canvas->StrLen(Font, Cmd.GetRight(), DescLenX, DescLenY, bDPIAwareStringMeasurement);
+				DescLenX /= DPIScale;
+				DescLenY /= DPIScale;
 				MaxRightWidth = FMath::Max(MaxRightWidth, ConsoleDefs::AutocompleteGap + DescLenX);
 			}
 		}
@@ -1483,7 +1502,9 @@ void UConsole::PostRender_InputLine(UCanvas* Canvas, FIntPoint UserInputLinePos)
 			}
 
 			float CommandWidth, CommandHeight;
-			Canvas->StrLen(Font, AutoCompleteElement.Command, CommandWidth, CommandHeight);
+			Canvas->StrLen(Font, AutoCompleteElement.Command, CommandWidth, CommandHeight, bDPIAwareStringMeasurement);
+			CommandWidth /= DPIScale;
+			CommandHeight /= DPIScale;
 
 			if (bMoreMatches)
 			{
@@ -1496,7 +1517,10 @@ void UConsole::PostRender_InputLine(UCanvas* Canvas, FIntPoint UserInputLinePos)
 			Canvas->DrawItem( ConsoleText, UserInputLinePos.X + CmdXOffset + xl, y );
 
 			float DescriptionWidth, DescriptionHeight;
-			Canvas->StrLen(Font, AutoCompleteElement.GetRight(), DescriptionWidth, DescriptionHeight);
+			Canvas->StrLen(Font, AutoCompleteElement.GetRight(), DescriptionWidth, DescriptionHeight, bDPIAwareStringMeasurement);
+			DescriptionWidth /= DPIScale;
+			DescriptionHeight /= DPIScale;
+			
 			float DescriptionX = UserInputLinePos.X + xl + ConsoleDefs::AutocompleteGap;
 			float DescriptionOverflow = DescriptionX + MaxLeftWidth + DescriptionWidth - Canvas->SizeX;
 
@@ -1532,15 +1556,15 @@ void UConsole::PostRender_InputLine(UCanvas* Canvas, FIntPoint UserInputLinePos)
 		}
 	}
 
-	const float DPIScale = Canvas->GetDPIScale();
-	const bool bDPIAwareStringMeasurement = true;
 	// determine the cursor position
 	const FString TypedInputTextUpToCursor = FString::Printf(TEXT("%s%s"), *ConsoleDefs::LeadingInputText, *TypedStr.Left(TypedStrPos));
 	Canvas->StrLen(Font, TypedInputTextUpToCursor, xl, yl, bDPIAwareStringMeasurement);
+	xl /= DPIScale;
+	yl /= DPIScale;
 	// draw the cursor
 	ConsoleText.SetColor(ConsoleDefs::CursorColor);
 	ConsoleText.Text = FText::FromString(FString(TEXT("_")));
-	Canvas->DrawItem(ConsoleText, UserInputLinePos.X + xl / DPIScale, UserInputLinePos.Y - 1.0f - yl / DPIScale);
+	Canvas->DrawItem(ConsoleText, UserInputLinePos.X + xl, UserInputLinePos.Y - 1.0f - yl);
 
 }
 
