1. Fixed Tooltips incorrect when changing FExpressionInput of UMaterialExpressionTextureSample.
2. Default UMaterialExpression::GetConnectorToolTip is slightly more efficient. Previously, if the hovered object did not have a tooltip, the for loop would wrongly iterate over all elements rather than stopping after checking the hovered object.
3. Additionally added some missing "const" keywords to some arguments, and replaced "NULL" by "nullptr".

diff --git a/Engine/Source/Runtime/Engine/Classes/Materials/MaterialExpressionTextureSample.h b/Engine/Source/Runtime/Engine/Classes/Materials/MaterialExpressionTextureSample.h
--- a/Engine/Source/Runtime/Engine/Classes/Materials/MaterialExpressionTextureSample.h
+++ b/Engine/Source/Runtime/Engine/Classes/Materials/MaterialExpressionTextureSample.h
@@ -99,6 +99,7 @@ class ENGINE_API UMaterialExpressionTextureSample : public UMaterialExpressionTe
 	virtual FExpressionInput* GetInput(int32 InputIndex) override;
 	virtual FName GetInputName(int32 InputIndex) const override;
 	virtual int32 GetWidth() const override;
+	virtual void GetConnectorToolTip(int32 InputIndex, int32 OutputIndex, TArray<FString>& OutToolTip) override;
 	virtual int32 GetLabelPadding() override { return 8; }
 	virtual int32 Compile(class FMaterialCompiler* Compiler, int32 OutputIndex) override;
 	virtual void GetCaption(TArray<FString>& OutCaptions) const override;
diff --git a/Engine/Source/Runtime/Engine/Private/Materials/MaterialExpressions.cpp b/Engine/Source/Runtime/Engine/Private/Materials/MaterialExpressions.cpp
--- a/Engine/Source/Runtime/Engine/Private/Materials/MaterialExpressions.cpp
+++ b/Engine/Source/Runtime/Engine/Private/Materials/MaterialExpressions.cpp
@@ -238,7 +238,7 @@ FUObjectAnnotationSparseBool GMaterialFunctionsThatNeedSamplerFixup;
 #endif // #if WITH_EDITOR
 
 /** Returns whether the given expression class is allowed. */
-bool IsAllowedExpressionType(UClass* Class, bool bMaterialFunction)
+bool IsAllowedExpressionType(const UClass* const Class, const bool bMaterialFunction)
 {
 	static const auto AllowVolumeTextureAssetCreationVar = IConsoleManager::Get().FindTConsoleVariableDataInt(TEXT("r.AllowVolumeTextureAssetCreation"));
 
@@ -260,7 +260,7 @@ bool IsAllowedExpressionType(UClass* Class, bool bMaterialFunction)
 }
 
 /** Parses a string into multiple lines, for use with tooltips. */
-void ConvertToMultilineToolTip(const FString& InToolTip, int32 TargetLineLength, TArray<FString>& OutToolTip)
+void ConvertToMultilineToolTip(const FString& InToolTip, const int32 TargetLineLength, TArray<FString>& OutToolTip)
 {
 	int32 CurrentPosition = 0;
 	int32 LastPosition = 0;
@@ -297,7 +297,7 @@ void ConvertToMultilineToolTip(const FString& InToolTip, int32 TargetLineLength,
 	}
 }
 
-void GetMaterialValueTypeDescriptions(uint32 MaterialValueType, TArray<FText>& OutDescriptions)
+void GetMaterialValueTypeDescriptions(const uint32 MaterialValueType, TArray<FText>& OutDescriptions)
 {
 	// Get exact float type if possible
 	uint32 MaskedFloatType = MaterialValueType & MCT_Float;
@@ -356,7 +356,7 @@ void GetMaterialValueTypeDescriptions(uint32 MaterialValueType, TArray<FText>& O
 		OutDescriptions.Add(LOCTEXT("Unknown", "Unknown"));
 }
 
-bool CanConnectMaterialValueTypes(uint32 InputType, uint32 OutputType)
+bool CanConnectMaterialValueTypes(const uint32 InputType, const uint32 OutputType)
 {
 	if (InputType & MCT_Unknown)
 	{
@@ -413,7 +413,7 @@ void ValidateParameterNameInternal(class UMaterialExpression* ExpressionToValida
 
 				bFoundValidName = true;
 
-				for (UMaterialExpression* Expression : OwningMaterial->Expressions)
+				for (const UMaterialExpression* Expression : OwningMaterial->Expressions)
 				{
 					if (Expression != nullptr && Expression->HasAParameterName())
 					{
@@ -464,15 +464,15 @@ void ValidateParameterNameInternal(class UMaterialExpression* ExpressionToValida
  * @param ParameterName           (Optional) Parameter name of the texture parameter that's assigned to the sample (used to access the external texture transform at runtime)
  * @return Index to a new code chunk that supplies the transformed UV coordinates
  */
-int32 CompileExternalTextureCoordinates(FMaterialCompiler* Compiler, int32 TexCoordCodeIndex, int32 TextureReferenceIndex, TOptional<FName> ParameterName = TOptional<FName>())
+int32 CompileExternalTextureCoordinates(FMaterialCompiler* Compiler, const int32 TexCoordCodeIndex, const int32 TextureReferenceIndex, const TOptional<FName> ParameterName = TOptional<FName>())
 {
 	if (TexCoordCodeIndex == INDEX_NONE)
 	{
 		return INDEX_NONE;
 	}
 
-	int32 ScaleRotationCode = Compiler->ExternalTextureCoordinateScaleRotation(TextureReferenceIndex, ParameterName);
-	int32 OffsetCode = Compiler->ExternalTextureCoordinateOffset(TextureReferenceIndex, ParameterName);
+	const int32 ScaleRotationCode = Compiler->ExternalTextureCoordinateScaleRotation(TextureReferenceIndex, ParameterName);
+	const int32 OffsetCode = Compiler->ExternalTextureCoordinateOffset(TextureReferenceIndex, ParameterName);
 
 	return Compiler->RotateScaleOffsetTexCoords(TexCoordCodeIndex, ScaleRotationCode, OffsetCode);
 }
@@ -481,6 +481,7 @@ int32 CompileExternalTextureCoordinates(FMaterialCompiler* Compiler, int32 TexCo
  * Compile a texture sample taking into consideration external textures (which may use different sampling code in the shader on some platforms)
  *
  * @param Compiler                The compiler to add code to
+ * @param Texture                 UTexture pointer used for the compiler
  * @param TexCoordCodeIndex       Index to the code chunk that supplies the vanilla texture coordinates
  * @param SamplerType             The type of sampler that is to be used
  * @param ParameterName           (Optional) Parameter name of the texture parameter that's assigned to the sample
@@ -494,13 +495,13 @@ int32 CompileTextureSample(
 	FMaterialCompiler* Compiler,
 	UTexture* Texture,
 	int32 TexCoordCodeIndex,
-	EMaterialSamplerType SamplerType,
-	TOptional<FName> ParameterName = TOptional<FName>(),
-	int32 MipValue0Index=INDEX_NONE,
-	int32 MipValue1Index=INDEX_NONE,
-	ETextureMipValueMode MipValueMode=TMVM_None,
-	ESamplerSourceMode SamplerSource=SSM_FromTextureAsset,
-	bool AutomaticViewMipBias=false
+	const EMaterialSamplerType SamplerType,
+	const TOptional<FName> ParameterName = TOptional<FName>(),
+	const int32 MipValue0Index=INDEX_NONE,
+	const int32 MipValue1Index=INDEX_NONE,
+	const ETextureMipValueMode MipValueMode=TMVM_None,
+	const ESamplerSourceMode SamplerSource=SSM_FromTextureAsset,
+	const bool AutomaticViewMipBias=false
 	)
 {
 	int32 TextureReferenceIndex = INDEX_NONE;
@@ -540,14 +541,14 @@ int32 CompileTextureSample(
  * @param Alpha					Bland factor [0..1]
  * @return						Index to a new code chunk
  */
-int32 CompileShadingModelBlendFunction(FMaterialCompiler* Compiler, int32 A, int32 B, int32 Alpha)
+int32 CompileShadingModelBlendFunction(FMaterialCompiler* Compiler, const int32 A, const int32 B, const int32 Alpha)
 {
 	if (A == INDEX_NONE || B == INDEX_NONE || Alpha == INDEX_NONE)
 	{
 		return INDEX_NONE;
 	}
 
-	int32  MidPoint = Compiler->Constant(0.5f);
+	const int32 MidPoint = Compiler->Constant(0.5f);
 
 	return Compiler->If(Alpha, MidPoint, B, INDEX_NONE, A, INDEX_NONE);
 }
@@ -556,7 +557,7 @@ int32 CompileShadingModelBlendFunction(FMaterialCompiler* Compiler, int32 A, int
 UMaterialExpression::UMaterialExpression(const FObjectInitializer& ObjectInitializer)
 	: Super(ObjectInitializer)
 #if WITH_EDITORONLY_DATA
-	, GraphNode(NULL)
+	, GraphNode(nullptr)
 #endif // WITH_EDITORONLY_DATA
 {
 #if WITH_EDITORONLY_DATA
@@ -600,12 +601,12 @@ void UMaterialExpression::CopyMaterialExpressions(const TArray<UMaterialExpressi
 			bIsValidFunctionExpression = false;
 		}
 
-		if (bIsValidFunctionExpression && IsAllowedExpressionType(SrcExpression->GetClass(), EditFunction != NULL))
+		if (bIsValidFunctionExpression && IsAllowedExpressionType(SrcExpression->GetClass(), EditFunction != nullptr))
 		{
 			UMaterialExpression* NewExpression = Cast<UMaterialExpression>(StaticDuplicateObject( SrcExpression, ExpressionOuter, NAME_None, RF_Transactional ));
 			NewExpression->Material = Material;
 			// Make sure we remove any references to functions the nodes came from
-			NewExpression->Function = NULL;
+			NewExpression->Function = nullptr;
 
 			SrcToDestMap.Add( SrcExpression, NewExpression );
 
@@ -661,7 +662,7 @@ void UMaterialExpression::CopyMaterialExpressions(const TArray<UMaterialExpressi
 				}
 				else
 				{
-					Input->Expression = NULL;
+					Input->Expression = nullptr;
 				}
 			}
 		}
@@ -768,7 +769,7 @@ void UMaterialExpression::PostEditChangeProperty(FPropertyChangedEvent& Property
 	}
 
 	UProperty* PropertyThatChanged = PropertyChangedEvent.Property;
-	if( PropertyThatChanged != NULL )
+	if( PropertyThatChanged != nullptr )
 	{
 		// Update the preview for this node if we adjusted a property
 		bNeedToUpdatePreview = true;
@@ -804,7 +805,7 @@ void UMaterialExpression::PostEditImport()
 bool UMaterialExpression::CanEditChange(const UProperty* InProperty) const
 {
 	bool bIsEditable = Super::CanEditChange(InProperty);
-	if (bIsEditable && InProperty != NULL)
+	if (bIsEditable && InProperty != nullptr)
 	{
 		// Automatically set property as non-editable if it has OverridingInputProperty metadata
 		// pointing to an FExpressionInput property which is hooked up as an input.
@@ -894,7 +895,7 @@ FExpressionInput* UMaterialExpression::GetInput(int32 InputIndex)
 	}
 	}
 
-	return NULL;
+	return nullptr;
 }
 
 
@@ -1072,9 +1073,6 @@ void UMaterialExpression::GetConnectorToolTip(int32 InputIndex, int32 OutputInde
 {
 	if (InputIndex != INDEX_NONE)
 	{
-		const TArray<FExpressionInput*> Inputs = GetInputs();
-
-		int32 Index = 0;
 		for( TFieldIterator<UStructProperty> InputIt(GetClass()) ; InputIt ; ++InputIt )
 		{
 			UStructProperty* StructProp = *InputIt;
@@ -1082,13 +1080,16 @@ void UMaterialExpression::GetConnectorToolTip(int32 InputIndex, int32 OutputInde
 			{
 				for (int32 ArrayIndex = 0; ArrayIndex < StructProp->ArrayDim; ArrayIndex++)
 				{
-					if( Index == InputIndex && StructProp->HasMetaData(TEXT("tooltip")) )
+					if (!InputIndex)
 					{
-						// Set the tooltip from the .h comments
-						ConvertToMultilineToolTip(StructProp->GetToolTipText().ToString(), 40, OutToolTip);
+						if (StructProp->HasMetaData(TEXT("tooltip")))
+						{
+							// Set the tooltip from the .h comments
+							ConvertToMultilineToolTip(StructProp->GetToolTipText().ToString(), 40, OutToolTip);
+						}
 						return;
 					}
-					Index++;
+					InputIndex--;
 				}
 			}
 		}
@@ -1111,7 +1112,7 @@ bool UMaterialExpression::Modify( bool bAlwaysMarkDirty/*=true*/ )
 
 bool UMaterialExpression::MatchesSearchQuery( const TCHAR* SearchQuery )
 {
-	if (FCString::Stristr(SearchQuery, TEXT("NAME=")) != NULL)
+	if (FCString::Stristr(SearchQuery, TEXT("NAME=")) != nullptr)
 	{
 		FString SearchString(SearchQuery);
 		SearchString = SearchString.Right(SearchString.Len() - 5);
@@ -1343,12 +1344,12 @@ void UMaterialExpressionTextureBase::PostEditChangeProperty(FPropertyChangedEven
 			for (auto ItExpressions = Expressions.CreateConstIterator(); ItExpressions; ItExpressions++)
 			{
 				UMaterialExpressionTextureBase* TextureSample = Cast<UMaterialExpressionTextureBase>(*ItExpressions);
-				if (TextureSample != NULL && TextureSample != this)
+				if (TextureSample != nullptr && TextureSample != this)
 				{
 					if(TextureSample->IsDefaultMeshpaintTexture)
 					{
 						FText ErrorMessage = LOCTEXT("MeshPaintDefaultTextureErrorDefault","Only one texture can be set as the Mesh Paint Default Texture, disabling previous default");
-						if (TextureSample->Texture != NULL)
+						if (TextureSample->Texture != nullptr)
 						{
 							FFormatNamedArguments Args;
 							Args.Add( TEXT("TextureName"), FText::FromString( TextureSample->Texture->GetName() ) );
@@ -1461,7 +1462,7 @@ UMaterialExpressionTextureSample::UMaterialExpressionTextureSample(const FObject
 bool UMaterialExpressionTextureSample::CanEditChange(const UProperty* InProperty) const
 {
 	bool bIsEditable = Super::CanEditChange(InProperty);
-	if (bIsEditable && InProperty != NULL)
+	if (bIsEditable && InProperty != nullptr)
 	{
 		FName PropertyFName = InProperty->GetFName();
 
@@ -1476,11 +1477,11 @@ bool UMaterialExpressionTextureSample::CanEditChange(const UProperty* InProperty
 		else if (PropertyFName == GET_MEMBER_NAME_CHECKED(UMaterialExpressionTextureSample, Texture))
 		{
 			// The Texture property is overridden by a connection to TextureObject
-			bIsEditable = TextureObject.GetTracedInput().Expression == NULL;
+			bIsEditable = TextureObject.GetTracedInput().Expression == nullptr;
 		}
 		else if (PropertyFName == GET_MEMBER_NAME_CHECKED(UMaterialExpressionTextureSample, AutomaticViewMipBias))
 		{
-			bIsEditable = AutomaticViewMipBiasValue.GetTracedInput().Expression == NULL;
+			bIsEditable = AutomaticViewMipBiasValue.GetTracedInput().Expression == nullptr;
 		}
 	}
 
@@ -1540,14 +1541,14 @@ const TArray<FExpressionInput*> UMaterialExpressionTextureSample::GetInputs()
 }
 
 // this define is only used for the following function
-#define IF_INPUT_RETURN(Item) if(!InputIndex) return &Item;	--InputIndex
+#define IF_INPUT_RETURN(Item) if(!InputIndex) return &Item; --InputIndex
 FExpressionInput* UMaterialExpressionTextureSample::GetInput(int32 InputIndex)
 {
 	IF_INPUT_RETURN(Coordinates);
 
 	if (bShowTextureInputPin)
 	{
-	IF_INPUT_RETURN(TextureObject);
+		IF_INPUT_RETURN(TextureObject);
 	}
 
 	if(MipValueMode == TMVM_Derivative)
@@ -1562,36 +1563,43 @@ FExpressionInput* UMaterialExpressionTextureSample::GetInput(int32 InputIndex)
 
 	IF_INPUT_RETURN(AutomaticViewMipBiasValue);
 
-	return NULL;
+	return nullptr;
 }
 #undef IF_INPUT_RETURN
 
 // this define is only used for the following function
-#define IF_INPUT_RETURN(Item, Name) if(!InputIndex) return Name; --InputIndex
+#define IF_INPUT_RETURN(Name) if(!InputIndex) return Name; --InputIndex
 FName UMaterialExpressionTextureSample::GetInputName(int32 InputIndex) const
 {
-	IF_INPUT_RETURN(Coordinates, TEXT("Coordinates"));
+	// Coordinates
+	IF_INPUT_RETURN(TEXT("Coordinates"));
 
 	if (bShowTextureInputPin)
 	{
-	IF_INPUT_RETURN(TextureObject, TEXT("TextureObject"));
+		// TextureObject
+		IF_INPUT_RETURN(TEXT("TextureObject"));
 	}
 
 	if(MipValueMode == TMVM_MipLevel)
 	{
-		IF_INPUT_RETURN(MipValue, TEXT("MipLevel"));
+		// MipValue
+		IF_INPUT_RETURN(TEXT("MipLevel"));
 	}
 	else if(MipValueMode == TMVM_MipBias)
 	{
-		IF_INPUT_RETURN(MipValue, TEXT("MipBias"));
+		// MipValue
+		IF_INPUT_RETURN(TEXT("MipBias"));
 	}
 	else if(MipValueMode == TMVM_Derivative)
 	{
-		IF_INPUT_RETURN(CoordinatesDX, TEXT("DDX(UVs)"));
-		IF_INPUT_RETURN(CoordinatesDY, TEXT("DDY(UVs)"));
+		// CoordinatesDX
+		IF_INPUT_RETURN(TEXT("DDX(UVs)"));
+		// CoordinatesDY
+		IF_INPUT_RETURN(TEXT("DDY(UVs)"));
 	}
 
-	IF_INPUT_RETURN(AutomaticViewMipBiasValue, TEXT("View MipBias"));
+	// AutomaticViewMipBiasValue
+	IF_INPUT_RETURN(TEXT("View MipBias"));
 
 	return TEXT("");
 }
@@ -1605,7 +1613,7 @@ FName UMaterialExpressionTextureSample::GetInputName(int32 InputIndex) const
  * they do not.
  * @param Compiler - The material compiler to which errors will be reported.
  * @param ExpressionDesc - Description of the expression verifying the sampler type.
- * @param Texture - The texture to verify. A NULL texture is considered valid!
+ * @param Texture - The texture to verify. A nullptr texture is considered valid!
  * @param SamplerType - The sampler type to verify.
  */
 static bool VerifySamplerType(
@@ -1768,7 +1776,7 @@ int32 UMaterialExpressionTextureSample::Compile(class FMaterialCompiler* Compile
 		}
 		else
 		{
-			// TextureObject.Expression is responsible for generating the error message, since it had a NULL texture value
+			// TextureObject.Expression is responsible for generating the error message, since it had a nullptr texture value
 			return INDEX_NONE;
 		}
 	}
@@ -1783,6 +1791,43 @@ int32 UMaterialExpressionTextureSample::GetWidth() const
 	return ME_STD_THUMBNAIL_SZ+(ME_STD_BORDER*2);
 }
 
+#if WITH_EDITOR
+#define IF_INPUT_RETURN(Value) if(!InputIndex) return (Value); --InputIndex
+int32 GetAbsoluteIndex(int32 InputIndex, const bool bShowTextureInputPin, const TEnumAsByte<enum ETextureMipValueMode>& MipValueMode)
+{
+	// Coordinates
+	IF_INPUT_RETURN(0);
+	if (bShowTextureInputPin)
+	{
+		// TextureObject
+		IF_INPUT_RETURN(1);
+	}
+	if(MipValueMode == TMVM_Derivative)
+	{
+		// CoordinatesDX
+		IF_INPUT_RETURN(3);
+		// CoordinatesDY
+		IF_INPUT_RETURN(4);
+	}
+	else if(MipValueMode != TMVM_None)
+	{
+		// MipValue
+		IF_INPUT_RETURN(2);
+	}
+	// AutomaticViewMipBiasValue
+	IF_INPUT_RETURN(5);
+	// If not found
+	return INDEX_NONE;
+}
+#undef IF_INPUT_RETURN
+
+void UMaterialExpressionTextureSample::GetConnectorToolTip(int32 InputIndex, int32 OutputIndex, TArray<FString>& OutToolTip)
+{
+	const int32 AbsoluteIndex = GetAbsoluteIndex(InputIndex, bShowTextureInputPin, MipValueMode);
+	Super::GetConnectorToolTip(AbsoluteIndex, OutputIndex, OutToolTip);
+}
+#endif // WITH_EDITOR
+
 void UMaterialExpressionTextureSample::GetCaption(TArray<FString>& OutCaptions) const
 {
 	OutCaptions.Add(TEXT("Texture Sample"));
@@ -1790,7 +1835,7 @@ void UMaterialExpressionTextureSample::GetCaption(TArray<FString>& OutCaptions)
 
 bool UMaterialExpressionTextureSample::MatchesSearchQuery( const TCHAR* SearchQuery )
 {
-	if( Texture!=NULL && Texture->GetName().Contains(SearchQuery) )
+	if( Texture!=nullptr && Texture->GetName().Contains(SearchQuery) )
 	{
 		return true;
 	}
@@ -1799,28 +1844,34 @@ bool UMaterialExpressionTextureSample::MatchesSearchQuery( const TCHAR* SearchQu
 }
 
 // this define is only used for the following function
-#define IF_INPUT_RETURN(Item, Type) if(!InputIndex) return Type; --InputIndex
+#define IF_INPUT_RETURN(Type) if(!InputIndex) return (Type); --InputIndex
 uint32 UMaterialExpressionTextureSample::GetInputType(int32 InputIndex)
 {
-	IF_INPUT_RETURN(Coordinates, MCT_Float);
+	// Coordinates
+	IF_INPUT_RETURN(MCT_Float);
 
 	if (bShowTextureInputPin)
 	{
+		// TextureObject
 		// TODO: Only show the TextureObject input inside a material function, since that's the only place it is useful
-	IF_INPUT_RETURN(TextureObject, MCT_Texture);
+		IF_INPUT_RETURN(MCT_Texture);
 	}
 	
 	if(MipValueMode == TMVM_MipLevel || MipValueMode == TMVM_MipBias)
 	{
-		IF_INPUT_RETURN(MipValue, MCT_Float);
+		// MipValue
+		IF_INPUT_RETURN(MCT_Float);
 	}
 	else if(MipValueMode == TMVM_Derivative)
 	{
-		IF_INPUT_RETURN(CoordinatesDX, MCT_Float);
-		IF_INPUT_RETURN(CoordinatesDY, MCT_Float);
+		// CoordinatesDX
+		IF_INPUT_RETURN(MCT_Float);
+		// CoordinatesDY
+		IF_INPUT_RETURN(MCT_Float);
 	}
 
-	IF_INPUT_RETURN(AutomaticViewMipBiasValue, MCT_StaticBool);
+	// AutomaticViewMipBiasValue
+	IF_INPUT_RETURN(MCT_StaticBool);
 
 	return MCT_Unknown;
 }
@@ -1909,7 +1960,7 @@ UMaterialExpressionTextureSampleParameter::UMaterialExpressionTextureSampleParam
 #if WITH_EDITOR
 int32 UMaterialExpressionTextureSampleParameter::Compile(class FMaterialCompiler* Compiler, int32 OutputIndex)
 {
-	if (Texture == NULL)
+	if (Texture == nullptr)
 	{
 		return CompilerError(Compiler, GetRequirements());
 	}
@@ -2189,11 +2240,11 @@ int32 UMaterialExpressionTextureObject::CompilePreview(class FMaterialCompiler*
 
 uint32 UMaterialExpressionTextureObject::GetOutputType(int32 OutputIndex)
 {
-	if (Cast<UTextureCube>(Texture) != NULL)
+	if (Cast<UTextureCube>(Texture) != nullptr)
 	{
 		return MCT_TextureCube;
 	}
-	else if (Cast<UVolumeTexture>(Texture) != NULL)
+	else if (Cast<UVolumeTexture>(Texture) != nullptr)
 	{
 		return MCT_VolumeTexture;
 	}
@@ -2280,10 +2331,11 @@ void UMaterialExpressionTextureProperty::GetCaption(TArray<FString>& OutCaptions
 }
 
 // this define is only used for the following function
-#define IF_INPUT_RETURN(Item, Type) if(!InputIndex) return Type; --InputIndex
+#define IF_INPUT_RETURN(Type) if(!InputIndex) return Type; --InputIndex
 uint32 UMaterialExpressionTextureProperty::GetInputType(int32 InputIndex)
 {
-	IF_INPUT_RETURN(TextureObject, MCT_Texture);
+	// TextureObject
+	IF_INPUT_RETURN(MCT_Texture);
 	return MCT_Unknown;
 }
 #undef IF_INPUT_RETURN
@@ -2377,7 +2429,7 @@ const TCHAR* UMaterialExpressionTextureSampleParameter2D::GetRequirements()
 
 void UMaterialExpressionTextureSampleParameter2D::SetDefaultTexture()
 {
-	Texture = LoadObject<UTexture2D>(NULL, TEXT("/Engine/EngineResources/DefaultTexture.DefaultTexture"), NULL, LOAD_None, NULL);
+	Texture = LoadObject<UTexture2D>(nullptr, TEXT("/Engine/EngineResources/DefaultTexture.DefaultTexture"), nullptr, LOAD_None, nullptr);
 }
 
 #if WITH_EDITOR
@@ -2474,7 +2526,7 @@ const TCHAR* UMaterialExpressionTextureSampleParameterCube::GetRequirements()
 
 void UMaterialExpressionTextureSampleParameterCube::SetDefaultTexture()
 {
-	Texture = LoadObject<UTextureCube>(NULL, TEXT("/Engine/EngineResources/DefaultTextureCube.DefaultTextureCube"), NULL, LOAD_None, NULL);
+	Texture = LoadObject<UTextureCube>(nullptr, TEXT("/Engine/EngineResources/DefaultTextureCube.DefaultTextureCube"), nullptr, LOAD_None, nullptr);
 }
 
 //
@@ -2545,7 +2597,7 @@ const TCHAR* UMaterialExpressionTextureSampleParameterVolume::GetRequirements()
 
 void UMaterialExpressionTextureSampleParameterVolume::SetDefaultTexture()
 {
-	Texture = LoadObject<UVolumeTexture>(NULL, TEXT("/Engine/EngineResources/DefaultVolumeTexture.DefaultVolumeTexture"), NULL, LOAD_None, NULL);
+	Texture = LoadObject<UVolumeTexture>(nullptr, TEXT("/Engine/EngineResources/DefaultVolumeTexture.DefaultVolumeTexture"), nullptr, LOAD_None, nullptr);
 }
 
 /** 
@@ -2584,7 +2636,7 @@ UMaterialExpressionTextureSampleParameterSubUV::UMaterialExpressionTextureSample
 #if WITH_EDITOR
 int32 UMaterialExpressionTextureSampleParameterSubUV::Compile(class FMaterialCompiler* Compiler, int32 OutputIndex)
 {
-	if (Texture == NULL)
+	if (Texture == nullptr)
 	{
 		return CompilerError(Compiler, GetRequirements());
 	}
@@ -4538,7 +4590,7 @@ void UMaterialExpressionMakeMaterialAttributes::Serialize(FStructuredArchive::FR
 int32 UMaterialExpressionMakeMaterialAttributes::Compile(class FMaterialCompiler* Compiler, int32 OutputIndex) 
 {
 	int32 Ret = INDEX_NONE;
-	UMaterialExpression* Expression = NULL;
+	UMaterialExpression* Expression = nullptr;
 
  	static_assert(MP_MAX == 30, 
 		"New material properties should be added to the end of the inputs for this expression. \
@@ -9244,7 +9296,7 @@ void UMaterialExpressionFontSample::GetCaption(TArray<FString>& OutCaptions) con
 
 bool UMaterialExpressionFontSample::MatchesSearchQuery( const TCHAR* SearchQuery )
 {
-	if( Font!=NULL && Font->GetName().Contains(SearchQuery) )
+	if( Font != nullptr && Font->GetName().Contains(SearchQuery) )
 	{
 		return true;
 	}
@@ -9261,7 +9313,7 @@ UTexture* UMaterialExpressionFontSample::GetReferencedTexture()
 		return Texture;
 	}
 
-	return NULL;
+	return nullptr;
 }
 
 /*-----------------------------------------------------------------------------
@@ -9959,7 +10011,7 @@ FExpressionInput* UMaterialExpressionCustom::GetInput(int32 InputIndex)
 	{
 		return &Inputs[InputIndex].Input;
 	}
-	return NULL;
+	return nullptr;
 }
 
 FName UMaterialExpressionCustom::GetInputName(int32 InputIndex) const
@@ -10218,15 +10270,15 @@ UMaterialFunction::UMaterialFunction(const FObjectInitializer& ObjectInitializer
 	LibraryCategoriesText.Add(LOCTEXT("Misc", "Misc"));
 #endif
 #if WITH_EDITORONLY_DATA
-	PreviewMaterial = NULL;
-	ThumbnailInfo = NULL;
+	PreviewMaterial = nullptr;
+	ThumbnailInfo = nullptr;
 #endif
 }
 
 #if WITH_EDITOR
 UMaterialInterface* UMaterialFunction::GetPreviewMaterial()
 {
-	if( NULL == PreviewMaterial )
+	if( nullptr == PreviewMaterial )
 	{
 		PreviewMaterial = NewObject<UMaterial>(GetTransientPackage(), NAME_None, RF_Transient | RF_Public);
 
@@ -10243,7 +10295,7 @@ UMaterialInterface* UMaterialFunction::GetPreviewMaterial()
 		}
 
 		//Compile the material.
-		PreviewMaterial->PreEditChange( NULL );
+		PreviewMaterial->PreEditChange(nullptr);
 		PreviewMaterial->PostEditChange();
 	}
 	return PreviewMaterial;
@@ -10337,7 +10389,7 @@ void UMaterialFunction::PostLoad()
 
 	for (int32 ExpressionIndex = 0; ExpressionIndex < FunctionExpressions.Num(); ExpressionIndex++)
 	{
-		// Expressions whose type was removed can be NULL
+		// Expressions whose type was removed can be nullptr
 		if (FunctionExpressions[ExpressionIndex])
 		{
 			FunctionExpressions[ExpressionIndex]->ConditionalPostLoad();
@@ -10353,7 +10405,7 @@ void UMaterialFunction::PostLoad()
 	if (GIsEditor)
 	{
 		// Clean up any removed material expression classes	
-		if (FunctionExpressions.Remove(NULL) != 0)
+		if (FunctionExpressions.Remove(nullptr) != 0)
 		{
 			// Force this function to recompile because its expressions have changed
 			// Warning: any content taking this path will recompile every load until saved!
@@ -10493,7 +10545,7 @@ static const FFunctionExpressionInput* FindInputById(const FGuid& Id, const TArr
 			return &CurrentInput;
 		}
 	}
-	return NULL;
+	return nullptr;
 }
 
 /** Finds an input in the passed in array with a matching name. */
@@ -10507,7 +10559,7 @@ static const FFunctionExpressionInput* FindInputByName(const FName& Name, const
 			return &CurrentInput;
 		}
 	}
-	return NULL;
+	return nullptr;
 }
 
 /** Finds an input in the passed in array with a matching expression object. */
@@ -10521,7 +10573,7 @@ static const FExpressionInput* FindInputByExpression(UMaterialExpressionFunction
 			return &CurrentInput.Input;
 		}
 	}
-	return NULL;
+	return nullptr;
 }
 
 /** Finds an output in the passed in array with a matching Id. */
@@ -10706,7 +10758,7 @@ void UMaterialFunction::UnlinkFromCaller()
 			// Restore the default value
 			InputExpression->bCompilingFunctionPreview = true;
 			// Clear the reference to make stale accesses obvious
-			InputExpression->EffectivePreviewDuringCompile.Expression = NULL;
+			InputExpression->EffectivePreviewDuringCompile.Expression = nullptr;
 		}
 	}
 }
@@ -11074,8 +11126,8 @@ UMaterialFunctionInstance::UMaterialFunctionInstance(const FObjectInitializer& O
 	: Super(ObjectInitializer)
 {
 #if WITH_EDITORONLY_DATA
-	PreviewMaterial = NULL;
-	ThumbnailInfo = NULL;
+	PreviewMaterial = nullptr;
+	ThumbnailInfo = nullptr;
 #endif
 }
 
@@ -11265,12 +11317,12 @@ void UMaterialFunctionInstance::AppendReferencedTextures(TArray<UTexture*>& InOu
 #if WITH_EDITOR
 UMaterialInterface* UMaterialFunctionInstance::GetPreviewMaterial()
 {
-	if (NULL == PreviewMaterial)
+	if (nullptr == PreviewMaterial)
 	{
 		PreviewMaterial = NewObject<UMaterialInstanceConstant>((UObject*)GetTransientPackage(), FName(TEXT("None")), RF_Transient);
 		PreviewMaterial->SetParentEditorOnly(Parent->GetPreviewMaterial());
 		OverrideMaterialInstanceParameterValues(PreviewMaterial);
-		PreviewMaterial->PreEditChange(NULL);
+		PreviewMaterial->PreEditChange(nullptr);
 		PreviewMaterial->PostEditChange();
 
 	}
@@ -11439,7 +11491,7 @@ void FMaterialLayersFunctions::SerializeForDDC(FArchive& Ar)
 // UMaterialExpressionMaterialFunctionCall
 ///////////////////////////////////////////////////////////////////////////////
 
-UMaterialFunctionInterface* SavedMaterialFunction = NULL;
+UMaterialFunctionInterface* SavedMaterialFunction = nullptr;
 
 UMaterialExpressionMaterialFunctionCall::UMaterialExpressionMaterialFunctionCall(const FObjectInitializer& ObjectInitializer)
 	: Super(ObjectInitializer)
@@ -11511,7 +11563,7 @@ void UMaterialExpressionMaterialFunctionCall::PostEditChangeProperty(FPropertyCh
 	{
 		// Set the new material function
 		SetMaterialFunctionEx(SavedMaterialFunction, MaterialFunction);
-		SavedMaterialFunction = NULL;
+		SavedMaterialFunction = nullptr;
 	}
 
 	Super::PostEditChangeProperty(PropertyChangedEvent);
@@ -11596,7 +11648,7 @@ FExpressionInput* UMaterialExpressionMaterialFunctionCall::GetInput(int32 InputI
 	{
 		return &FunctionInputs[InputIndex].Input;
 	}
-	return NULL;
+	return nullptr;
 }
 
 
@@ -11643,7 +11695,7 @@ FName UMaterialExpressionMaterialFunctionCall::GetInputName(int32 InputIndex) co
 
 bool UMaterialExpressionMaterialFunctionCall::IsInputConnectionRequired(int32 InputIndex) const
 {
-	if (InputIndex < FunctionInputs.Num() && FunctionInputs[InputIndex].ExpressionInput != NULL)
+	if (InputIndex < FunctionInputs.Num() && FunctionInputs[InputIndex].ExpressionInput != nullptr)
 	{
 		return !FunctionInputs[InputIndex].ExpressionInput->bUsePreviewValueAsDefault;
 	}
@@ -11758,7 +11810,7 @@ bool UMaterialExpressionMaterialFunctionCall::SetMaterialFunctionEx(
 	{
 		// Prevent recursive function call graphs
 		FMessageDialog::Open( EAppMsgType::Ok, NSLOCTEXT("MaterialExpressions", "Error_CircularFunctionDependency", "Can't use that material function as it would cause a circular dependency.") );
-		NewFunctionResource = NULL;
+		NewFunctionResource = nullptr;
 	}
 
 	MaterialFunction = NewFunctionResource;
@@ -11795,7 +11847,7 @@ bool UMaterialExpressionMaterialFunctionCall::SetMaterialFunctionEx(
 		}
 	}
 
-	// Fixup even if NewFunctionResource is NULL, because we have to clear old connections
+	// Fixup even if NewFunctionResource is nullptr, because we have to clear old connections
 	if (OldFunctionResource && OldFunctionResource != NewFunctionResource)
 	{
 		TArray<FExpressionInput*> MaterialInputs;
@@ -11828,7 +11880,7 @@ bool UMaterialExpressionMaterialFunctionCall::SetMaterialFunctionEx(
 		CastChecked<UMaterialGraphNode>(GraphNode)->RecreateAndLinkNode();
 	}
 
-	return NewFunctionResource != NULL;
+	return NewFunctionResource != nullptr;
 }
 
 void UMaterialExpressionMaterialFunctionCall::UpdateFromFunctionResource(bool bRecreateAndLinkNode)
@@ -11955,14 +12007,14 @@ static void FixupReferencingInputs(
 				if (CurrentInput->OutputIndex == INDEX_NONE)
 				{
 					// The output that this input was connected to no longer exists, break the connection
-					CurrentInput->Expression = NULL;
+					CurrentInput->Expression = nullptr;
 				}
 			}
 			else
 			{
 				// The output that this input was connected to no longer exists, break the connection
 				CurrentInput->OutputIndex = INDEX_NONE;
-				CurrentInput->Expression = NULL;
+				CurrentInput->Expression = nullptr;
 			}
 		}
 	}
@@ -13492,7 +13544,7 @@ UMaterialExpressionNoise::UMaterialExpressionNoise(const FObjectInitializer& Obj
 bool UMaterialExpressionNoise::CanEditChange(const UProperty* InProperty) const
 {
 	bool bIsEditable = Super::CanEditChange(InProperty);
-	if (bIsEditable && InProperty != NULL)
+	if (bIsEditable && InProperty != nullptr)
 	{
 		FName PropertyFName = InProperty->GetFName();
 
@@ -13586,7 +13638,7 @@ UMaterialExpressionVectorNoise::UMaterialExpressionVectorNoise(const FObjectInit
 bool UMaterialExpressionVectorNoise::CanEditChange(const UProperty* InProperty) const
 {
 	bool bIsEditable = Super::CanEditChange(InProperty);
-	if (bIsEditable && InProperty != NULL)
+	if (bIsEditable && InProperty != nullptr)
 	{
 		FName PropertyFName = InProperty->GetFName();
 
@@ -14099,7 +14151,7 @@ const TCHAR* UMaterialExpressionAntialiasedTextureMask::GetRequirements()
 
 void UMaterialExpressionAntialiasedTextureMask::SetDefaultTexture()
 {
-	Texture = LoadObject<UTexture2D>(NULL, TEXT("/Engine/EngineResources/DefaultTexture.DefaultTexture"), NULL, LOAD_None, NULL);
+	Texture = LoadObject<UTexture2D>(nullptr, TEXT("/Engine/EngineResources/DefaultTexture.DefaultTexture"), nullptr, LOAD_None, nullptr);
 }
 
 //
@@ -14803,7 +14855,7 @@ int32 UMaterialExpressionSpeedTree::Compile(class FMaterialCompiler* Compiler, i
 	int32 WindArg = (WindInput.GetTracedInput().Expression ? WindInput.Compile(Compiler) : Compiler->Constant(WindType));
 	int32 LODArg = (LODInput.GetTracedInput().Expression ? LODInput.Compile(Compiler) : Compiler->Constant(LODType));
 	
-	bool bExtraBend = (ExtraBendWS.GetTracedInput().Expression != NULL);
+	bool bExtraBend = (ExtraBendWS.GetTracedInput().Expression != nullptr);
 	int32 ExtraBendArg = (ExtraBendWS.GetTracedInput().Expression ? ExtraBendWS.Compile(Compiler) : Compiler->Constant3(0.0f, 0.0f, 0.0f));
 	 
 	return Compiler->SpeedTree(GeometryArg, WindArg, LODArg, BillboardThreshold, bAccurateWindVelocities, bExtraBend, ExtraBendArg);
@@ -15446,7 +15498,7 @@ int32 UMaterialExpressionShadingModel::Compile(class FMaterialCompiler* Compiler
 
 void UMaterialExpressionShadingModel::GetCaption(TArray<FString>& OutCaptions) const
 {
-	const UEnum* ShadingModelEnum = FindObject<UEnum>(NULL, TEXT("Engine.EMaterialShadingModel"));
+	const UEnum* ShadingModelEnum = FindObject<UEnum>(nullptr, TEXT("Engine.EMaterialShadingModel"));
 	check(ShadingModelEnum);
 
 	const FString ShadingModelDisplayName = ShadingModelEnum->GetDisplayNameTextByValue(ShadingModel).ToString();
diff --git a/Engine/Source/Runtime/Engine/Public/MaterialShared.h b/Engine/Source/Runtime/Engine/Public/MaterialShared.h
--- a/Engine/Source/Runtime/Engine/Public/MaterialShared.h
+++ b/Engine/Source/Runtime/Engine/Public/MaterialShared.h
@@ -1172,16 +1172,16 @@ private:
 };
 
 /** Returns whether the given expression class is allowed. */
-extern ENGINE_API bool IsAllowedExpressionType(UClass* Class, bool bMaterialFunction);
+extern ENGINE_API bool IsAllowedExpressionType(const UClass* const Class, const bool bMaterialFunction);
 
 /** Parses a string into multiple lines, for use with tooltips. */
-extern ENGINE_API void ConvertToMultilineToolTip(const FString& InToolTip, int32 TargetLineLength, TArray<FString>& OutToolTip);
+extern ENGINE_API void ConvertToMultilineToolTip(const FString& InToolTip, const int32 TargetLineLength, TArray<FString>& OutToolTip);
 
 /** Given a combination of EMaterialValueType flags, get text descriptions of all types */
-extern ENGINE_API void GetMaterialValueTypeDescriptions(uint32 MaterialValueType, TArray<FText>& OutDescriptions);
+extern ENGINE_API void GetMaterialValueTypeDescriptions(const uint32 MaterialValueType, TArray<FText>& OutDescriptions);
 
 /** Check whether a combination of EMaterialValueType flags can be connected */
-extern ENGINE_API bool CanConnectMaterialValueTypes(uint32 InputType, uint32 OutputType);
+extern ENGINE_API bool CanConnectMaterialValueTypes(const uint32 InputType, const uint32 OutputType);
 
 /**
  * FMaterial serves 3 intertwined purposes:
