PR #5486: Added QuitEditor function that can also be called in python (Contributed by [REDACTED_PERSON])

diff --git a/Engine/Source/Runtime/Engine/Classes/Kismet/KismetSystemLibrary.h b/Engine/Source/Runtime/Engine/Classes/Kismet/KismetSystemLibrary.h
--- a/Engine/Source/Runtime/Engine/Classes/Kismet/KismetSystemLibrary.h
+++ b/Engine/Source/Runtime/Engine/Classes/Kismet/KismetSystemLibrary.h
@@ -434,6 +434,14 @@ class ENGINE_API UKismetSystemLibrary : public UBlueprintFunctionLibrary
 	 */
 	UFUNCTION(BlueprintCallable, Category="Game",meta=(WorldContext="WorldContextObject"))
 	static void QuitGame(UObject* WorldContextObject, class APlayerController* SpecificPlayer, TEnumAsByte<EQuitPreference::Type> QuitPreference, bool bIgnorePlatformRestrictions);
+	
+#if WITH_EDITOR
+	/**
+	 *	Exit the editor
+	 */
+	UFUNCTION(BlueprintCallable, Category="Development")
+	static void QuitEditor();
+#endif	// WITH_EDITOR
 
 	//=============================================================================
 	// Latent Actions
diff --git a/Engine/Source/Runtime/Engine/Private/KismetSystemLibrary.cpp b/Engine/Source/Runtime/Engine/Private/KismetSystemLibrary.cpp
--- a/Engine/Source/Runtime/Engine/Private/KismetSystemLibrary.cpp
+++ b/Engine/Source/Runtime/Engine/Private/KismetSystemLibrary.cpp
@@ -381,6 +381,13 @@ void UKismetSystemLibrary::QuitGame(UObject* WorldContextObject, class APlayerCo
 	}
 }
 
+#if WITH_EDITOR
+void UKismetSystemLibrary::QuitEditor()
+{
+	GEngine->Exec(nullptr, TEXT("QUIT_EDITOR"), *GLog);
+}
+#endif	// WITH_EDITOR
+
 bool UKismetSystemLibrary::K2_IsValidTimerHandle(FTimerHandle TimerHandle)
 {
 	return TimerHandle.IsValid();
