PR #6046: Use given Viewport instead of ViewportClient->Viewport in ULocalPlayer::GetProjectionData (Contributed by [REDACTED_PERSON])

diff --git a/Engine/Source/Runtime/Engine/Private/LocalPlayer.cpp b/Engine/Source/Runtime/Engine/Private/LocalPlayer.cpp
--- a/Engine/Source/Runtime/Engine/Private/LocalPlayer.cpp
+++ b/Engine/Source/Runtime/Engine/Private/LocalPlayer.cpp
@@ -1112,7 +1112,7 @@ bool ULocalPlayer::GetProjectionData(FViewport* Viewport, EStereoscopicPass Ster
 	if (!bNeedStereo)
 	{
 		// Create the projection matrix (and possibly constrain the view rectangle)
-		FMinimalViewInfo::CalculateProjectionMatrixGivenView(ViewInfo, AspectRatioAxisConstraint, ViewportClient->Viewport, /*inout*/ ProjectionData);
+		FMinimalViewInfo::CalculateProjectionMatrixGivenView(ViewInfo, AspectRatioAxisConstraint, Viewport, /*inout*/ ProjectionData);
 
 		for (auto& ViewExt : GEngine->ViewExtensions->GatherActiveExtensions())
         {
