Always check bInitializedRecords before trying to call FPlatformMisc::SetStoredValue in FEngineSessionManager.

diff --git a/Engine/Source/Runtime/Engine/Private/EngineSessionManager.cpp b/Engine/Source/Runtime/Engine/Private/EngineSessionManager.cpp
--- a/Engine/Source/Runtime/Engine/Private/EngineSessionManager.cpp
+++ b/Engine/Source/Runtime/Engine/Private/EngineSessionManager.cpp
@@ -318,16 +318,16 @@ void FEngineSessionManager::Shutdown()
 	FCoreDelegates::ApplicationWillTerminateDelegate.RemoveAll(this);
 	FCoreDelegates::IsVanillaProductChanged.RemoveAll(this);
 
-	if (!CurrentSession.bIsTerminating) // Skip Slate if terminating, since we can't guarantee which thread called us.
-	{
-		FSlateApplication::Get().GetOnModalLoopTickEvent().RemoveAll(this);
-
-		FPlatformMisc::SetStoredValue(SessionManagerDefs::StoreId, CurrentSessionSectionName, SessionManagerDefs::WasShutdownStoreKey, SessionManagerDefs::TrueValueString);
-	}
-
 	// Clear the session record for this session
 	if (bInitializedRecords)
 	{
+		if (!CurrentSession.bIsTerminating) // Skip Slate if terminating, since we can't guarantee which thread called us.
+		{
+			FSlateApplication::Get().GetOnModalLoopTickEvent().RemoveAll(this);
+
+			FPlatformMisc::SetStoredValue(SessionManagerDefs::StoreId, CurrentSessionSectionName, SessionManagerDefs::WasShutdownStoreKey, SessionManagerDefs::TrueValueString);
+		}
+
 		if (!CurrentSession.bCrashed)
 		{
 #if PLATFORM_SUPPORTS_WATCHDOG
@@ -768,7 +768,7 @@ void FEngineSessionManager::OnCrashing()
 
 void FEngineSessionManager::OnAppReactivate()
 {
-	if (CurrentSession.bIsDeactivated)
+	if (CurrentSession.bIsDeactivated && bInitializedRecords)
 	{
 		CurrentSession.bIsDeactivated = false;
 		FPlatformMisc::SetStoredValue(SessionManagerDefs::StoreId, CurrentSessionSectionName, SessionManagerDefs::DeactivatedStoreKey, SessionManagerDefs::FalseValueString);
@@ -777,7 +777,7 @@ void FEngineSessionManager::OnAppReactivate()
 
 void FEngineSessionManager::OnAppDeactivate()
 {
-	if (!CurrentSession.bIsDeactivated)
+	if (!CurrentSession.bIsDeactivated && bInitializedRecords)
 	{
 		CurrentSession.bIsDeactivated = true;
 		FPlatformMisc::SetStoredValue(SessionManagerDefs::StoreId, CurrentSessionSectionName, SessionManagerDefs::DeactivatedStoreKey, SessionManagerDefs::TrueValueString);
@@ -786,7 +786,7 @@ void FEngineSessionManager::OnAppDeactivate()
 
 void FEngineSessionManager::OnAppBackground()
 {
-	if (!CurrentSession.bIsInBackground)
+	if (!CurrentSession.bIsInBackground && bInitializedRecords)
 	{
 		CurrentSession.bIsInBackground = true;
 		FPlatformMisc::SetStoredValue(SessionManagerDefs::StoreId, CurrentSessionSectionName, SessionManagerDefs::BackgroundStoreKey, SessionManagerDefs::TrueValueString);
@@ -795,7 +795,7 @@ void FEngineSessionManager::OnAppBackground()
 
 void FEngineSessionManager::OnAppForeground()
 {
-	if (CurrentSession.bIsInBackground)
+	if (CurrentSession.bIsInBackground && bInitializedRecords)
 	{
 		CurrentSession.bIsInBackground = false;
 		FPlatformMisc::SetStoredValue(SessionManagerDefs::StoreId, CurrentSessionSectionName, SessionManagerDefs::BackgroundStoreKey, SessionManagerDefs::FalseValueString);
@@ -804,7 +804,7 @@ void FEngineSessionManager::OnAppForeground()
 
 void FEngineSessionManager::OnTerminate()
 {
-	if (!CurrentSession.bIsTerminating)
+	if (!CurrentSession.bIsTerminating && bInitializedRecords)
 	{
 		CurrentSession.bIsTerminating = true;
 		FPlatformMisc::SetStoredValue(SessionManagerDefs::StoreId, CurrentSessionSectionName, SessionManagerDefs::TerminatingKey, SessionManagerDefs::TrueValueString);
