Update parameter name validation to work for material functions as well as materials

diff --git a/Engine/Source/Runtime/Engine/Private/Materials/MaterialExpressions.cpp b/Engine/Source/Runtime/Engine/Private/Materials/MaterialExpressions.cpp
--- a/Engine/Source/Runtime/Engine/Private/Materials/MaterialExpressions.cpp
+++ b/Engine/Source/Runtime/Engine/Private/Materials/MaterialExpressions.cpp
@@ -6953,7 +6953,17 @@ void UMaterialExpressionVectorParameter::ApplyChannelNames()
 void UMaterialExpressionVectorParameter::ValidateParameterName(const bool bAllowDuplicateName)
 {
 	bool bOverrideDuplicateBehavior = false;
-	for (UMaterialExpression* Expression : Material->Expressions)
+	TArray<UMaterialExpression*> Expressions;
+	if (Material)
+	{
+		Expressions = Material->Expressions;
+	}
+	else if (Function)
+	{
+		Expressions = Function->FunctionExpressions;
+	}
+
+	for (UMaterialExpression* Expression : Expressions)
 	{
 		if (Expression != nullptr && Expression->HasAParameterName())
 		{
@@ -7244,7 +7254,17 @@ void UMaterialExpressionScalarParameter::PostEditChangeProperty(FPropertyChanged
 void UMaterialExpressionScalarParameter::ValidateParameterName(const bool bAllowDuplicateName)
 {
 	bool bOverrideDuplicateBehavior = false;
-	for (UMaterialExpression* Expression : Material->Expressions)
+	TArray<UMaterialExpression*> Expressions;
+	if (Material)
+	{
+		Expressions = Material->Expressions;
+	}
+	else if (Function)
+	{
+		Expressions = Function->FunctionExpressions;
+	}
+
+	for (UMaterialExpression* Expression : Expressions)
 	{
 		if (Expression != nullptr && Expression->HasAParameterName())
 		{
