PR #6007: Feature date time based naming for the high res screenshots (Contributed by [REDACTED_PERSON])

diff --git a/Engine/Source/Editor/LevelEditor/Private/HighResScreenshotUI.cpp b/Engine/Source/Editor/LevelEditor/Private/HighResScreenshotUI.cpp
--- a/Engine/Source/Editor/LevelEditor/Private/HighResScreenshotUI.cpp
+++ b/Engine/Source/Editor/LevelEditor/Private/HighResScreenshotUI.cpp
@@ -51,7 +51,13 @@ void SHighResScreenshotDialog::Construct( const FArguments& InArgs )
 							.VAlign(VAlign_Center)
 							[
 								SNew( STextBlock )
-								.Text( NSLOCTEXT("HighResScreenshot", "IncludeBufferVisTargets", "Include Buffer Visualization Targets") )
+								.Text( NSLOCTEXT("HighResScreenshot", "UseDateTimeAsImageName", "Use Date & Timestamp as Image name") )
+							]
+							+ SVerticalBox::Slot()
+							.VAlign(VAlign_Center)
+							[
+								SNew(STextBlock)
+								.Text(NSLOCTEXT("HighResScreenshot", "IncludeBufferVisTargets", "Include Buffer Visualization Targets"))
 							]
 							+ SVerticalBox::Slot()
 							.VAlign(VAlign_Center)
@@ -97,6 +103,13 @@ void SHighResScreenshotDialog::Construct( const FArguments& InArgs )
 									.OnValueChanged(this, &SHighResScreenshotDialog::OnResolutionMultiplierSliderChanged)
 								]
 							]
+							+ SVerticalBox::Slot()
+							.VAlign(VAlign_Center)
+							[
+								SNew(SCheckBox)
+								.OnCheckStateChanged(this, &SHighResScreenshotDialog::OnDateTimeBasedNamingEnabledChanged)
+								.IsChecked(this, &SHighResScreenshotDialog::GetDateTimeBasedNamingEnabled)
+							]
 							+SVerticalBox::Slot()
 							.VAlign(VAlign_Center)
 							[
diff --git a/Engine/Source/Editor/LevelEditor/Private/HighresScreenshotUI.h b/Engine/Source/Editor/LevelEditor/Private/HighresScreenshotUI.h
--- a/Engine/Source/Editor/LevelEditor/Private/HighresScreenshotUI.h
+++ b/Engine/Source/Editor/LevelEditor/Private/HighresScreenshotUI.h
@@ -119,6 +119,12 @@ private:
 		}
 	}
 
+	void OnDateTimeBasedNamingEnabledChanged(ECheckBoxState NewValue)
+	{
+		bool bEnabled = (NewValue == ECheckBoxState::Checked);
+		Config.bDateTimeBasedNaming = bEnabled;
+	}
+
 	void OnBufferVisualizationDumpEnabledChanged(ECheckBoxState NewValue)
 	{
 		bool bEnabled = (NewValue == ECheckBoxState::Checked);
@@ -172,6 +178,11 @@ private:
 		return Config.bDumpBufferVisualizationTargets ? ECheckBoxState::Checked : ECheckBoxState::Unchecked;
 	}
 
+	ECheckBoxState GetDateTimeBasedNamingEnabled() const
+	{
+		return Config.bDateTimeBasedNaming ? ECheckBoxState::Checked : ECheckBoxState::Unchecked;
+	}
+
 	bool IsCaptureRegionEditingAvailable() const
 	{
 		return CaptureRegionWidget.IsValid();
diff --git a/Engine/Source/Runtime/Core/Private/Misc/FileHelper.cpp b/Engine/Source/Runtime/Core/Private/Misc/FileHelper.cpp
--- a/Engine/Source/Runtime/Core/Private/Misc/FileHelper.cpp
+++ b/Engine/Source/Runtime/Core/Private/Misc/FileHelper.cpp
@@ -311,6 +311,28 @@ bool FFileHelper::GenerateNextBitmapFilename( const FString& Pattern, const FStr
 	return bSuccess;
 }
 
+/**
+ * Generates the next unique bitmap filename with a specified extension
+ *
+ * @param Pattern		Filename with path, but without extension.
+ * @oaran Extension		File extension to be appended
+ * @param OutFilename	Reference to an FString where the newly generated filename will be placed
+ *
+ * @return true if success
+ */
+void FFileHelper::GenerateDateTimeBasedBitmapFilename(const FString& Pattern, const FString& Extension, FString& OutFilename)
+{
+	// Use current date & time to obtain more organized screenshot libraries
+	// There is no need to check for file duplicate, as two certian moments, can't occure twice in the world!
+	
+	OutFilename = "";
+
+	static int32 LastScreenShotIndex = 0;
+	int32 SearchIndex = 0;
+
+	OutFilename = FString::Printf(TEXT("%s_%s.%s"), *Pattern, *FDateTime::Now().ToString(), *Extension);
+}
+
 /**
  * Saves a 24Bit BMP file to disk
  * 
diff --git a/Engine/Source/Runtime/Core/Public/Misc/FileHelper.h b/Engine/Source/Runtime/Core/Public/Misc/FileHelper.h
--- a/Engine/Source/Runtime/Core/Public/Misc/FileHelper.h
+++ b/Engine/Source/Runtime/Core/Public/Misc/FileHelper.h
@@ -108,6 +108,17 @@ struct CORE_API FFileHelper
 	 * @return true if success
 	 */
 	static bool GenerateNextBitmapFilename(const FString& Pattern, const FString& Extension, FString& OutFilename, IFileManager* FileManager = &IFileManager::Get());
+
+	/**
+	 * Generates the next unique bitmap filename with a specified extension
+	 *
+	 * @param Pattern		Filename with path, but without extension.
+	 * @param Extension		File extension to be appended
+	 * @param OutFilename	Reference to an FString where the newly generated filename will be placed
+	 *
+	 * @return true if success
+	 */
+	static void GenerateDateTimeBasedBitmapFilename(const FString& Pattern, const FString& Extension, FString& OutFilename);
 	
 	/**
 	 *	Load the given ANSI text file to an array of strings - one FString per line of the file.
diff --git a/Engine/Source/Runtime/Engine/Private/HighResScreenshot.cpp b/Engine/Source/Runtime/Engine/Private/HighResScreenshot.cpp
--- a/Engine/Source/Runtime/Engine/Private/HighResScreenshot.cpp
+++ b/Engine/Source/Runtime/Engine/Private/HighResScreenshot.cpp
@@ -32,6 +32,7 @@ FHighResScreenshotConfig::FHighResScreenshotConfig()
 	: ResolutionMultiplier(FHighResScreenshotConfig::MinResolutionMultipler)
 	, ResolutionMultiplierScale(0.0f)
 	, bMaskEnabled(false)
+	, bDateTimeBasedNaming(false)
 	, bDumpBufferVisualizationTargets(false)
 {
 	ChangeViewport(TWeakPtr<FSceneViewport>());
@@ -91,6 +92,7 @@ void FHighResScreenshotConfig::ChangeViewport(TWeakPtr<FSceneViewport> InViewpor
 	UnscaledCaptureRegion = FIntRect(0, 0, 0, 0);
 	CaptureRegion = UnscaledCaptureRegion;
 	bMaskEnabled = false;
+	bDateTimeBasedNaming = false;
 	bDumpBufferVisualizationTargets = false;
 	ResolutionMultiplier = FHighResScreenshotConfig::MinResolutionMultipler;
 	ResolutionMultiplierScale = 0.0f;
@@ -104,7 +106,7 @@ bool FHighResScreenshotConfig::ParseConsoleCommand(const FString& InCmd, FOutput
 	ResolutionMultiplier = FHighResScreenshotConfig::MinResolutionMultipler;
 	ResolutionMultiplierScale = 0.0f;
 
-	if( GetHighResScreenShotInput(*InCmd, Ar, GScreenshotResolutionX, GScreenshotResolutionY, ResolutionMultiplier, CaptureRegion, bMaskEnabled, bDumpBufferVisualizationTargets, bCaptureHDR, FilenameOverride) )
+	if( GetHighResScreenShotInput(*InCmd, Ar, GScreenshotResolutionX, GScreenshotResolutionY, ResolutionMultiplier, CaptureRegion, bMaskEnabled, bDumpBufferVisualizationTargets, bCaptureHDR, FilenameOverride, bDateTimeBasedNaming) )
 	{
 		GScreenshotResolutionX *= ResolutionMultiplier;
 		GScreenshotResolutionY *= ResolutionMultiplier;
diff --git a/Engine/Source/Runtime/Engine/Private/UnrealClient.cpp b/Engine/Source/Runtime/Engine/Private/UnrealClient.cpp
--- a/Engine/Source/Runtime/Engine/Private/UnrealClient.cpp
+++ b/Engine/Source/Runtime/Engine/Private/UnrealClient.cpp
@@ -291,7 +291,14 @@ void FScreenshotRequest::RequestScreenshot(const FString& InFilename, bool bInSh
 	{
 		const bool bRemovePath = false;
 		GeneratedFilename = FPaths::GetBaseFilename(GeneratedFilename, bRemovePath);
-		FFileHelper::GenerateNextBitmapFilename(GeneratedFilename, TEXT("png"), Filename);
+		if (GetHighResScreenshotConfig().bDateTimeBasedNaming)
+		{
+			FFileHelper::GenerateDateTimeBasedBitmapFilename(GeneratedFilename, TEXT("png"), Filename);
+		}
+		else
+		{
+			FFileHelper::GenerateNextBitmapFilename(GeneratedFilename, TEXT("png"), Filename);
+		}
 	}
 	else
 	{
@@ -2133,7 +2140,7 @@ ENGINE_API bool GetViewportScreenShot(FViewport* Viewport, TArray<FColor>& Bitma
 
 extern bool ParseResolution( const TCHAR* InResolution, uint32& OutX, uint32& OutY, int32& WindowMode );
 
-ENGINE_API bool GetHighResScreenShotInput(const TCHAR* Cmd, FOutputDevice& Ar, uint32& OutXRes, uint32& OutYRes, float& OutResMult, FIntRect& OutCaptureRegion, bool& OutShouldEnableMask, bool& OutDumpBufferVisualizationTargets, bool& OutCaptureHDR, FString& OutFilenameOverride)
+ENGINE_API bool GetHighResScreenShotInput(const TCHAR* Cmd, FOutputDevice& Ar, uint32& OutXRes, uint32& OutYRes, float& OutResMult, FIntRect& OutCaptureRegion, bool& OutShouldEnableMask, bool& OutDumpBufferVisualizationTargets, bool& OutCaptureHDR, FString& OutFilenameOverride, bool& OutUseDateTimeAsFileName)
 {
 	FString CmdString = Cmd;
 	int32 SeperatorPos = -1;
@@ -2204,6 +2211,8 @@ ENGINE_API bool GetHighResScreenShotInput(const TCHAR* Cmd, FOutputDevice& Ar, u
 		OutShouldEnableMask = NumArguments > 5 ? FCString::Atoi(*Arguments[5]) != 0 : false;
 		OutDumpBufferVisualizationTargets = NumArguments > 6 ? FCString::Atoi(*Arguments[6]) != 0 : false;
 		OutCaptureHDR = NumArguments > 7 ? FCString::Atoi(*Arguments[7]) != 0 : false;
+		OutUseDateTimeAsFileName = NumArguments > 8 ? FCString::Atoi(*Arguments[8]) != 0 : false;
+
 
 		return true;
 	}
diff --git a/Engine/Source/Runtime/Engine/Public/HighResScreenshot.h b/Engine/Source/Runtime/Engine/Public/HighResScreenshot.h
--- a/Engine/Source/Runtime/Engine/Public/HighResScreenshot.h
+++ b/Engine/Source/Runtime/Engine/Public/HighResScreenshot.h
@@ -21,6 +21,7 @@ struct ENGINE_API FHighResScreenshotConfig
 	float ResolutionMultiplier;
 	float ResolutionMultiplierScale;
 	bool bMaskEnabled;
+	bool bDateTimeBasedNaming;
 	bool bDumpBufferVisualizationTargets;
 	TWeakPtr<FSceneViewport> TargetViewport;
 	bool bDisplayCaptureRegion;
diff --git a/Engine/Source/Runtime/Engine/Public/UnrealClient.h b/Engine/Source/Runtime/Engine/Public/UnrealClient.h
--- a/Engine/Source/Runtime/Engine/Public/UnrealClient.h
+++ b/Engine/Source/Runtime/Engine/Public/UnrealClient.h
@@ -718,7 +718,7 @@ extern ENGINE_API bool IsShiftDown(FViewport* Viewport);
 extern ENGINE_API bool IsAltDown(FViewport* Viewport);
 
 extern ENGINE_API bool GetViewportScreenShot(FViewport* Viewport, TArray<FColor>& Bitmap, const FIntRect& ViewRect = FIntRect());
-extern ENGINE_API bool GetHighResScreenShotInput(const TCHAR* Cmd, FOutputDevice& Ar, uint32& OutXRes, uint32& OutYRes, float& OutResMult, FIntRect& OutCaptureRegion, bool& OutShouldEnableMask, bool& OutDumpBufferVisualizationTargets, bool& OutCaptureHDR, FString& OutFilenameOverride);
+extern ENGINE_API bool GetHighResScreenShotInput(const TCHAR* Cmd, FOutputDevice& Ar, uint32& OutXRes, uint32& OutYRes, float& OutResMult, FIntRect& OutCaptureRegion, bool& OutShouldEnableMask, bool& OutDumpBufferVisualizationTargets, bool& OutCaptureHDR, FString& OutFilenameOverride, bool& OutUseDateTimeAsFileName);
 
 struct FInputKeyEventArgs
 {
