Removed clearing of viewport focus for all users when the primary keyboard user focuses something other than the viewport. Updated hard-coded checks to User 0 to instead fetch the keyboard user ID.

diff --git a/Engine/Source/Runtime/Engine/Private/Slate/SceneViewport.cpp b/Engine/Source/Runtime/Engine/Private/Slate/SceneViewport.cpp
--- a/Engine/Source/Runtime/Engine/Private/Slate/SceneViewport.cpp
+++ b/Engine/Source/Runtime/Engine/Private/Slate/SceneViewport.cpp
@@ -1058,7 +1058,7 @@ FReply FSceneViewport::OnFocusReceived(const FFocusEvent& InFocusEvent)
 {
 	CurrentReplyState = FReply::Handled(); 
 
-	if ( InFocusEvent.GetUser() == 0 )
+	if ( InFocusEvent.GetUser() == FSlateApplication::Get().GetUserIndexForKeyboard() )
 	{
 		if ( ViewportClient != nullptr )
 		{
@@ -1107,7 +1107,7 @@ FReply FSceneViewport::OnFocusReceived(const FFocusEvent& InFocusEvent)
 void FSceneViewport::OnFocusLost( const FFocusEvent& InFocusEvent )
 {
 	// If the focus loss event isn't the for the primary 'keyboard' user, don't worry about it.
-	if ( InFocusEvent.GetUser() != 0 )
+	if ( InFocusEvent.GetUser() != FSlateApplication::Get().GetUserIndexForKeyboard() )
 	{
 		return;
 	}
@@ -1119,19 +1119,6 @@ void FSceneViewport::OnFocusLost( const FFocusEvent& InFocusEvent )
 	{
 		FScopedConditionalWorldSwitcher WorldSwitcher(ViewportClient);
 		ViewportClient->LostFocus(this);
-
-		TSharedPtr<SWidget> ViewportWidgetPin = ViewportWidget.Pin();
-		if ( ViewportWidgetPin.IsValid() )
-		{
-			FSlateApplication::Get().ForEachUser([&] (FSlateUser* User) {
-				if ( User->GetFocusedWidget() == ViewportWidgetPin )
-				{
-					FSlateApplication::Get().ClearUserFocus(User->GetUserIndex());
-				}
-			});
-		}
-
-		UE_LOG(LogSlate, Log, TEXT("FSceneViewport::OnFocusLost() reason %d"), (int32)InFocusEvent.GetCause());
 	}
 }
 
