Added EditorSessionSummarySender/Writer. These classes will read/write to the platform-specific key-value store (registry on Windows) with statistics about previous sessions.

By and large this is the same information as EngineSessionManager will send, but in a format that is more easily digestible to queries.

diff --git a/Engine/Source/Runtime/Engine/Private/Analytics/EditorSessionSummary.cpp b/Engine/Source/Runtime/Engine/Private/Analytics/EditorSessionSummary.cpp
new file mode 100644
--- /dev/null
+++ b/Engine/Source/Runtime/Engine/Private/Analytics/EditorSessionSummary.cpp
@@ -0,0 +1,780 @@
+// Copyright 1998-2019 Epic Games, Inc. All Rights Reserved.
+
+#include "EditorSessionSummary.h"
+
+#include "AnalyticsEventAttribute.h"
+#include "EngineAnalytics.h"
+#include "EngineGlobals.h"
+#include "GeneralProjectSettings.h"
+#include "RHI.h"
+#include "UserActivityTracking.h"
+
+#include "Engine/Engine.h"
+#include "Framework/Application/SlateApplication.h"
+#include "HAL/PlatformProcess.h"
+#include "Interfaces/IAnalyticsProvider.h"
+#include "Interfaces/IPluginManager.h"
+#include "Interfaces/IProjectManager.h"
+#include "Logging/LogMacros.h"
+#include "Misc/App.h"
+#include "Misc/CommandLine.h"
+#include "Misc/ConfigCacheIni.h"
+#include "Misc/CoreDelegates.h"
+#include "Misc/EngineBuildSettings.h"
+#include "Misc/EngineVersion.h"
+#include "Misc/Guid.h"
+
+#if WITH_EDITOR
+#include "IVREditorModule.h"
+#include "Kismet2/DebuggerCommands.h"
+#endif // WITH_EDITOR
+
+#define LOCTEXT_NAMESPACE "SessionSummary"
+
+DEFINE_LOG_CATEGORY_STATIC(LogEditorSessionSummary, Verbose, All);
+
+namespace SessionSummaryDefs
+{
+	static const FTimespan SessionRecordExpiration = FTimespan::FromDays(30.0);
+	static const FTimespan SessionRecordTimeout = FTimespan::FromMinutes(10.0);
+	static const FTimespan GlobalLockWaitTimeout = FTimespan::FromSeconds(0.5);
+	static const int HeartbeatPeriodSeconds(60);
+
+	static const FString FalseValueString(TEXT("0"));
+	static const FString TrueValueString(TEXT("1"));
+
+	// shutdown types
+	static const FString RunningSessionToken(TEXT("Running"));
+	static const FString ShutdownSessionToken(TEXT("Shutdown"));
+	static const FString CrashSessionToken(TEXT("Crashed"));
+	static const FString TerminatedSessionToken(TEXT("Terminated"));
+	static const FString DebuggerSessionToken(TEXT("Debugger"));
+	static const FString AbnormalSessionToken(TEXT("AbnormalShutdown"));
+
+	static const FString DefaultUserActivity(TEXT("Unknown"));
+	static const FString UnknownProjectValueString(TEXT("UnknownProject"));
+
+	// storage location
+	static const FString StoreId(TEXT("Epic Games"));
+	static const FString SessionSummarySection(TEXT("Unreal Engine/Session Summary/1_0"));
+	static const FString GlobalLockName(TEXT("UE4_SessionSummary_Lock"));
+	static const FString SessionListStoreKey(TEXT("SessionList"));
+
+	// general values
+	static const FString ProjectNameStoreKey(TEXT("ProjectName"));
+	static const FString SessionIdStoreKey(TEXT("SessionId"));
+	static const FString PlatformProcessIDKey(TEXT("PlatformProcessID"));
+	static const FString EngineVersionStoreKey(TEXT("EngineVersion"));
+	static const FString StatusStoreKey(TEXT("LastExecutionState"));
+	static const FString UserActivityStoreKey(TEXT("CurrentUserActivity"));
+	static const FString PluginsStoreKey(TEXT("Plugins"));
+	static const FString AverageFPSStoreKey(TEXT("AverageFPS"));
+
+	// timestamps
+	static const FString TimestampStoreKey(TEXT("Timestamp"));
+	static const FString StartupTimestampStoreKey(TEXT("StartupTimestamp"));
+	static const FString SessionDurationStoreKey(TEXT("SessionDuration"));
+	static const FString Idle1MinStoreKey(TEXT("Idle1Min"));
+	static const FString Idle5MinStoreKey(TEXT("Idle5Min"));
+	static const FString Idle30MinStoreKey(TEXT("Idle30Min"));
+	
+	// boolean flags
+	static const FString IsTerminatingKey(TEXT("Terminating"));
+	static const FString WasDebuggerStoreKey(TEXT("WasEverDebugger"));
+	static const FString IsDebuggerStoreKey(TEXT("IsDebugger"));
+	static const FString IsVanillaStoreKey(TEXT("IsVanilla"));
+	static const FString IsCrashStoreKey(TEXT("IsCrash"));
+	static const FString WasShutdownStoreKey(TEXT("WasShutdown"));
+	static const FString IsGPUCrashStoreKey(TEXT("IsGPUCrash"));
+	static const FString IsInVRModeStoreKey(TEXT("IsInVRMode"));
+	static const FString IsInEnterpriseStoreKey(TEXT("IsInEnterprise"));
+	static const FString IsInPIEStoreKey(TEXT("IsInPIE"));
+}
+
+struct FSessionRecord
+{
+	FString SessionId;
+	FString ProjectName;
+	FString EngineVersion;
+	int32 PlatformProcessID;
+	FDateTime StartupTimestamp;
+	FDateTime Timestamp;
+	int32 Idle1Min;
+	int32 Idle5Min;
+	int32 Idle30Min;
+	FString CurrentUserActivity;
+	TArray<FString> Plugins;
+	float AverageFPS;
+
+	bool bCrashed : 1;
+	bool bGPUCrashed : 1;
+	bool bIsDebugger : 1;
+	bool bWasEverDebugger : 1;
+	bool bIsVanilla : 1;
+	bool bIsTerminating : 1;
+	bool bWasShutdown : 1;
+	bool bIsInPIE : 1;
+	bool bIsInEnterprise : 1;
+	bool bIsInVRMode : 1;
+
+	FSessionRecord()
+	{
+		PlatformProcessID = 0;
+		StartupTimestamp = FDateTime::MinValue();
+		Timestamp = FDateTime::MinValue();
+		Idle1Min = 0;
+		Idle5Min = 0;
+		Idle30Min = 0;
+		AverageFPS = 0;
+		bCrashed = false;
+		bGPUCrashed = false;
+		bIsDebugger = false;
+		bWasEverDebugger = false;
+		bIsVanilla = false;
+		bIsTerminating = false;
+		bWasShutdown = false;
+		bIsInPIE = false;
+		bIsInEnterprise = false;
+		bIsInVRMode = false;
+	}
+};
+
+// Utilities for writing to stored values
+namespace
+{
+	FString TimestampToString(FDateTime InTimestamp)
+	{
+		return LexToString(InTimestamp.ToUnixTimestamp());
+	}
+
+	FDateTime StringToTimestamp(FString InString)
+	{
+		int64 TimestampUnix;
+		if (LexTryParseString(TimestampUnix, *InString))
+		{
+			return FDateTime::FromUnixTimestamp(TimestampUnix);
+		}
+		return FDateTime::MinValue();
+	}
+
+	FString BoolToStoredString(bool bInValue)
+	{
+		return bInValue ? SessionSummaryDefs::TrueValueString : SessionSummaryDefs::FalseValueString;
+	}
+
+	bool GetStoredBool(const FString& SectionName, const FString& StoredKey)
+	{
+		FString StoredString = SessionSummaryDefs::FalseValueString;
+		FPlatformMisc::GetStoredValue(SessionSummaryDefs::StoreId, SectionName, StoredKey, StoredString);
+
+		return StoredString == SessionSummaryDefs::TrueValueString;
+	}
+
+	FString GetSessionStorageLocation(const FSessionRecord& Record)
+	{
+		return SessionSummaryDefs::SessionSummarySection + TEXT("/") + Record.SessionId;
+	}
+}
+
+/* FEditorSessionSummaryWriter */
+
+FEditorSessionSummaryWriter::FEditorSessionSummaryWriter() :
+	bInitializedRecords(false)
+	, CurrentSession(nullptr)
+	, HeartbeatTimeElapsed(0.0f)
+	, bShutdown(false)
+{
+
+}
+
+void FEditorSessionSummaryWriter::Initialize()
+{
+	// Register for crash and app state callbacks
+	FCoreDelegates::OnHandleSystemError.AddRaw(this, &FEditorSessionSummaryWriter::OnCrashing);
+	FCoreDelegates::ApplicationWillTerminateDelegate.AddRaw(this, &FEditorSessionSummaryWriter::OnTerminate);
+	FUserActivityTracking::OnActivityChanged.AddRaw(this, &FEditorSessionSummaryWriter::OnUserActivity);
+	FCoreDelegates::IsVanillaProductChanged.AddRaw(this, &FEditorSessionSummaryWriter::OnVanillaStateChanged);
+	FSlateApplication::Get().GetOnModalLoopTickEvent().AddRaw(this, &FEditorSessionSummaryWriter::Tick);
+
+	const bool bFirstInitAttempt = true;
+	InitializeRecords(bFirstInitAttempt);
+}
+
+void FEditorSessionSummaryWriter::InitializeRecords(bool bFirstAttempt)
+{
+	if (!FEngineAnalytics::IsAvailable())
+	{
+		return;
+	}
+
+	if (bInitializedRecords)
+	{
+		return;
+	}
+
+	{
+		// Scoped lock
+		FSystemWideCriticalSection StoredValuesLock(SessionSummaryDefs::GlobalLockName, bFirstAttempt ? SessionSummaryDefs::GlobalLockWaitTimeout : FTimespan::Zero());
+
+		if (StoredValuesLock.IsValid())
+		{
+			UE_LOG(LogEditorSessionSummary, Verbose, TEXT("Initializing EditorSessionSummaryWriter for editor session tracking"));
+
+			// Create a session record for this session
+			CurrentSession = CreateRecordForCurrentSession();
+			CurrentSessionSectionName = GetSessionStorageLocation(*CurrentSession);
+			WriteStoredRecord(*CurrentSession);
+
+			bInitializedRecords = true;
+
+			UE_LOG(LogEditorSessionSummary, Log, TEXT("EditorSessionSummaryWriter initialized"));
+		}
+
+		// update session list string
+		FString SessionListString;
+		FPlatformMisc::GetStoredValue(SessionSummaryDefs::StoreId, SessionSummaryDefs::SessionSummarySection, SessionSummaryDefs::SessionListStoreKey, SessionListString);
+		
+		if (!SessionListString.IsEmpty())
+		{
+			SessionListString.Append(TEXT(","));
+		}
+		SessionListString.Append(CurrentSession->SessionId);
+
+		FPlatformMisc::SetStoredValue(SessionSummaryDefs::StoreId, SessionSummaryDefs::SessionSummarySection, SessionSummaryDefs::SessionListStoreKey, SessionListString);
+	}
+}
+
+void FEditorSessionSummaryWriter::UpdateTimestamps()
+{
+	CurrentSession->Timestamp = FDateTime::UtcNow();
+	FPlatformMisc::SetStoredValue(SessionSummaryDefs::StoreId, CurrentSessionSectionName, SessionSummaryDefs::TimestampStoreKey, TimestampToString(CurrentSession->Timestamp));
+
+	const float IdleSeconds = FPlatformTime::Seconds() - FSlateApplication::Get().GetLastUserInteractionTime();
+
+	// 1 + 1 minutes
+	if (IdleSeconds > (60 + 60))
+	{
+		CurrentSession->Idle1Min += 1;
+		FPlatformMisc::SetStoredValue(SessionSummaryDefs::StoreId, CurrentSessionSectionName, SessionSummaryDefs::Idle1MinStoreKey, FString::FromInt(CurrentSession->Idle1Min));
+	}
+
+	// 5 + 1 minutes
+	if (IdleSeconds > (5 * 60 + 60))
+	{
+		CurrentSession->Idle5Min += 1;
+		FPlatformMisc::SetStoredValue(SessionSummaryDefs::StoreId, CurrentSessionSectionName, SessionSummaryDefs::Idle5MinStoreKey, FString::FromInt(CurrentSession->Idle5Min));
+	}
+
+	// 30 + 1 minutes
+	if (IdleSeconds > (30 * 60 + 60))
+	{
+		CurrentSession->Idle30Min += 1;
+		FPlatformMisc::SetStoredValue(SessionSummaryDefs::StoreId, CurrentSessionSectionName, SessionSummaryDefs::Idle30MinStoreKey, FString::FromInt(CurrentSession->Idle30Min));
+	}
+}
+
+void FEditorSessionSummaryWriter::Tick(float DeltaTime)
+{
+	HeartbeatTimeElapsed += DeltaTime;
+
+	if (HeartbeatTimeElapsed > (float) SessionSummaryDefs::HeartbeatPeriodSeconds && !bShutdown)
+	{
+		HeartbeatTimeElapsed = 0.0f;
+
+		// Try late initialization
+		const bool bFirstInitAttempt = false;
+		InitializeRecords(bFirstInitAttempt);
+
+		if (bInitializedRecords)
+		{
+			// check if the debugger is present
+			bool bIsDebuggerPresent = FPlatformMisc::IsDebuggerPresent();
+			if (CurrentSession->bIsDebugger != bIsDebuggerPresent)
+			{
+				CurrentSession->bIsDebugger = bIsDebuggerPresent;
+				FPlatformMisc::SetStoredValue(SessionSummaryDefs::StoreId, CurrentSessionSectionName, SessionSummaryDefs::IsDebuggerStoreKey, BoolToStoredString(CurrentSession->bIsDebugger));
+
+				if (!CurrentSession->bWasEverDebugger && CurrentSession->bIsDebugger)
+				{
+					CurrentSession->bWasEverDebugger = true;
+					FPlatformMisc::SetStoredValue(SessionSummaryDefs::StoreId, CurrentSessionSectionName, SessionSummaryDefs::WasDebuggerStoreKey, SessionSummaryDefs::TrueValueString);
+				}
+			}
+
+			UpdateTimestamps();
+
+#if WITH_EDITOR
+			extern ENGINE_API float GAverageFPS;
+
+			CurrentSession->AverageFPS = GAverageFPS;
+			const FString AverageFPSString = FString::SanitizeFloat(GAverageFPS);
+			FPlatformMisc::SetStoredValue(SessionSummaryDefs::StoreId, CurrentSessionSectionName, SessionSummaryDefs::AverageFPSStoreKey, AverageFPSString);
+
+			CurrentSession->bIsInVRMode = IVREditorModule::Get().IsVREditorModeActive();
+			const FString IsInVRModeString = BoolToStoredString(CurrentSession->bIsInVRMode);
+			FPlatformMisc::SetStoredValue(SessionSummaryDefs::StoreId, CurrentSessionSectionName, SessionSummaryDefs::IsInVRModeStoreKey, IsInVRModeString);
+
+			CurrentSession->bIsInEnterprise = IProjectManager::Get().IsEnterpriseProject();
+			const FString IsInEnterpriseString = BoolToStoredString(CurrentSession->bIsInEnterprise);
+			FPlatformMisc::SetStoredValue(SessionSummaryDefs::StoreId, CurrentSessionSectionName, SessionSummaryDefs::IsInEnterpriseStoreKey, IsInEnterpriseString);
+
+			CurrentSession->bIsInPIE = FPlayWorldCommandCallbacks::IsInPIE();
+			const FString IsInPIEString = BoolToStoredString(CurrentSession->bIsInPIE);
+			FPlatformMisc::SetStoredValue(SessionSummaryDefs::StoreId, CurrentSessionSectionName, SessionSummaryDefs::IsInPIEStoreKey, IsInPIEString);
+#endif
+		}
+	}
+}
+
+void FEditorSessionSummaryWriter::Shutdown()
+{
+	FCoreDelegates::OnHandleSystemError.RemoveAll(this);
+	FCoreDelegates::ApplicationHasReactivatedDelegate.RemoveAll(this);
+	FCoreDelegates::ApplicationWillDeactivateDelegate.RemoveAll(this);
+	FCoreDelegates::ApplicationWillEnterBackgroundDelegate.RemoveAll(this);
+	FCoreDelegates::ApplicationHasEnteredForegroundDelegate.RemoveAll(this);
+	FCoreDelegates::ApplicationWillTerminateDelegate.RemoveAll(this);
+	FCoreDelegates::IsVanillaProductChanged.RemoveAll(this);
+
+	// Skip Slate if terminating, since we can't guarantee which thread called us.
+	if (GIsRequestingExit) 
+	{
+		FSlateApplication::Get().GetOnModalLoopTickEvent().RemoveAll(this);
+	}
+
+	// Clear the session record for this session
+	if (bInitializedRecords)
+	{
+		if (!CurrentSession->bIsTerminating && !CurrentSession->bCrashed)
+		{
+			CurrentSession->bWasShutdown = true;
+			FPlatformMisc::SetStoredValue(SessionSummaryDefs::StoreId, CurrentSessionSectionName, SessionSummaryDefs::WasShutdownStoreKey, SessionSummaryDefs::TrueValueString);
+		}
+
+		delete CurrentSession;
+		bInitializedRecords = false;
+		bShutdown = true;
+	}
+}
+
+void FEditorSessionSummaryWriter::WriteStoredRecord(const FSessionRecord& Record) const
+{
+	FString StorageLocation = GetSessionStorageLocation(Record);
+	FString PluginsString = FString::Join(Record.Plugins, TEXT(","));
+
+	FPlatformMisc::SetStoredValue(SessionSummaryDefs::StoreId, StorageLocation, SessionSummaryDefs::ProjectNameStoreKey, Record.ProjectName);
+	FPlatformMisc::SetStoredValue(SessionSummaryDefs::StoreId, StorageLocation, SessionSummaryDefs::IsCrashStoreKey, SessionSummaryDefs::FalseValueString);
+	FPlatformMisc::SetStoredValue(SessionSummaryDefs::StoreId, StorageLocation, SessionSummaryDefs::EngineVersionStoreKey, Record.EngineVersion);
+	FPlatformMisc::SetStoredValue(SessionSummaryDefs::StoreId, StorageLocation, SessionSummaryDefs::StartupTimestampStoreKey, TimestampToString(Record.StartupTimestamp));
+	FPlatformMisc::SetStoredValue(SessionSummaryDefs::StoreId, StorageLocation, SessionSummaryDefs::TimestampStoreKey, TimestampToString(Record.Timestamp));
+	FPlatformMisc::SetStoredValue(SessionSummaryDefs::StoreId, StorageLocation, SessionSummaryDefs::Idle1MinStoreKey, FString::FromInt(Record.Idle1Min));
+	FPlatformMisc::SetStoredValue(SessionSummaryDefs::StoreId, StorageLocation, SessionSummaryDefs::Idle5MinStoreKey, FString::FromInt(Record.Idle5Min));
+	FPlatformMisc::SetStoredValue(SessionSummaryDefs::StoreId, StorageLocation, SessionSummaryDefs::Idle30MinStoreKey, FString::FromInt(Record.Idle30Min));
+	FPlatformMisc::SetStoredValue(SessionSummaryDefs::StoreId, StorageLocation, SessionSummaryDefs::UserActivityStoreKey, Record.CurrentUserActivity);
+	FPlatformMisc::SetStoredValue(SessionSummaryDefs::StoreId, StorageLocation, SessionSummaryDefs::IsVanillaStoreKey, BoolToStoredString(Record.bIsVanilla));
+	FPlatformMisc::SetStoredValue(SessionSummaryDefs::StoreId, StorageLocation, SessionSummaryDefs::IsTerminatingKey, BoolToStoredString(Record.bIsTerminating));
+	FPlatformMisc::SetStoredValue(SessionSummaryDefs::StoreId, StorageLocation, SessionSummaryDefs::PlatformProcessIDKey, FString::FromInt(Record.PlatformProcessID));
+	FPlatformMisc::SetStoredValue(SessionSummaryDefs::StoreId, StorageLocation, SessionSummaryDefs::PluginsStoreKey, PluginsString);
+	FPlatformMisc::SetStoredValue(SessionSummaryDefs::StoreId, StorageLocation, SessionSummaryDefs::AverageFPSStoreKey, FString::SanitizeFloat(Record.AverageFPS));
+	FPlatformMisc::SetStoredValue(SessionSummaryDefs::StoreId, StorageLocation, SessionSummaryDefs::IsDebuggerStoreKey, BoolToStoredString(Record.bIsDebugger));
+	FPlatformMisc::SetStoredValue(SessionSummaryDefs::StoreId, StorageLocation, SessionSummaryDefs::WasDebuggerStoreKey, BoolToStoredString(Record.bWasEverDebugger));
+	FPlatformMisc::SetStoredValue(SessionSummaryDefs::StoreId, StorageLocation, SessionSummaryDefs::WasShutdownStoreKey, BoolToStoredString(Record.bWasShutdown));
+	FPlatformMisc::SetStoredValue(SessionSummaryDefs::StoreId, StorageLocation, SessionSummaryDefs::IsInPIEStoreKey, BoolToStoredString(Record.bIsInPIE));
+	FPlatformMisc::SetStoredValue(SessionSummaryDefs::StoreId, StorageLocation, SessionSummaryDefs::IsInEnterpriseStoreKey, BoolToStoredString(Record.bIsInEnterprise));
+	FPlatformMisc::SetStoredValue(SessionSummaryDefs::StoreId, StorageLocation, SessionSummaryDefs::IsInVRModeStoreKey, BoolToStoredString(Record.bIsInVRMode));
+}
+
+FSessionRecord* FEditorSessionSummaryWriter::CreateRecordForCurrentSession() const
+{
+	FSessionRecord* Record = new FSessionRecord();
+
+	FGuid SessionId;
+	if (FGuid::Parse(FEngineAnalytics::GetProvider().GetSessionID(), SessionId))
+	{
+		// convert session guid to one without braces or other chars that might not be suitable for storage
+		Record->SessionId = SessionId.ToString(EGuidFormats::DigitsWithHyphens);
+	}
+	else
+	{
+		Record->SessionId = FEngineAnalytics::GetProvider().GetSessionID();
+	}
+
+	const uint32 ProcId = FPlatformProcess::GetCurrentProcessId();
+	FString CurrentProcessIDString = FString::FromInt(ProcId);
+
+	const UGeneralProjectSettings& ProjectSettings = *GetDefault<UGeneralProjectSettings>();
+
+	Record->ProjectName = ProjectSettings.ProjectName;
+	Record->EngineVersion = FEngineVersion::Current().ToString(EVersionComponent::Changelist);
+	Record->Timestamp = Record->StartupTimestamp = FDateTime::UtcNow();
+	Record->bIsDebugger = FPlatformMisc::IsDebuggerPresent();
+	Record->bWasEverDebugger = FPlatformMisc::IsDebuggerPresent();
+	Record->CurrentUserActivity = GetUserActivityString();
+	Record->bIsVanilla = GEngine && GEngine->IsVanillaProduct();
+
+	TArray<TSharedRef<IPlugin>> Plugins = IPluginManager::Get().GetEnabledPlugins();
+	for (const TSharedRef<IPlugin>& Plugin : Plugins)
+	{
+		Record->Plugins.Add(Plugin->GetName());
+	}
+
+	Record->Plugins.Sort();
+
+	return Record;
+}
+
+extern CORE_API bool GIsGPUCrashed;
+void FEditorSessionSummaryWriter::OnCrashing()
+{
+	if (bInitializedRecords && !CurrentSession->bCrashed)
+	{
+		UpdateTimestamps();
+
+		CurrentSession->bCrashed = true;
+		FPlatformMisc::SetStoredValue(SessionSummaryDefs::StoreId, CurrentSessionSectionName, SessionSummaryDefs::IsCrashStoreKey, SessionSummaryDefs::TrueValueString);
+		CurrentSession->bGPUCrashed = GIsGPUCrashed;
+		FPlatformMisc::SetStoredValue(SessionSummaryDefs::StoreId, CurrentSessionSectionName, SessionSummaryDefs::IsGPUCrashStoreKey, BoolToStoredString(CurrentSession->bGPUCrashed));
+	}
+}
+
+void FEditorSessionSummaryWriter::OnTerminate()
+{
+	if (bInitializedRecords && !CurrentSession->bIsTerminating)
+	{
+		UpdateTimestamps();
+
+		CurrentSession->bIsTerminating = true;
+		FPlatformMisc::SetStoredValue(SessionSummaryDefs::StoreId, CurrentSessionSectionName, SessionSummaryDefs::IsTerminatingKey, SessionSummaryDefs::TrueValueString);
+
+		if (GIsRequestingExit)
+		{
+			Shutdown();
+		}
+	}
+}
+
+void FEditorSessionSummaryWriter::OnVanillaStateChanged(bool bIsVanilla)
+{
+	if (bInitializedRecords && CurrentSession->bIsVanilla != bIsVanilla)
+	{
+		CurrentSession->bIsVanilla = bIsVanilla;
+		FPlatformMisc::SetStoredValue(SessionSummaryDefs::StoreId, CurrentSessionSectionName, SessionSummaryDefs::IsVanillaStoreKey, BoolToStoredString(CurrentSession->bIsVanilla));
+	}
+}
+
+void FEditorSessionSummaryWriter::OnUserActivity(const FUserActivity& UserActivity)
+{
+	if (bInitializedRecords && !CurrentSession->bCrashed)
+	{
+		CurrentSession->CurrentUserActivity = GetUserActivityString();
+		FPlatformMisc::SetStoredValue(SessionSummaryDefs::StoreId, CurrentSessionSectionName, SessionSummaryDefs::UserActivityStoreKey, CurrentSession->CurrentUserActivity);
+	}
+}
+
+FString FEditorSessionSummaryWriter::GetUserActivityString() const
+{
+	const FUserActivity& UserActivity = FUserActivityTracking::GetUserActivity();
+
+	if (UserActivity.ActionName.IsEmpty())
+	{
+		return SessionSummaryDefs::DefaultUserActivity;
+	}
+
+	return UserActivity.ActionName;
+}
+
+/* FEditorSessionSummarySender */
+
+FEditorSessionSummarySender::FEditorSessionSummarySender() :
+	HeartbeatTimeElapsed(0.0f)
+{
+
+}
+
+void FEditorSessionSummarySender::Initialize()
+{
+	SendStoredRecords(SessionSummaryDefs::GlobalLockWaitTimeout);
+}
+
+void FEditorSessionSummarySender::Tick(float DeltaTime)
+{
+	HeartbeatTimeElapsed += DeltaTime;
+
+	if (HeartbeatTimeElapsed > (float) SessionSummaryDefs::HeartbeatPeriodSeconds)
+	{
+		HeartbeatTimeElapsed = 0.0f;
+
+		SendStoredRecords(FTimespan::Zero());
+	}
+}
+
+void FEditorSessionSummarySender::SendStoredRecords(FTimespan Timeout) const
+{
+	if (!FEngineAnalytics::IsAvailable())
+	{
+		return;
+	}
+
+	TArray<FSessionRecord> SessionRecordsToReport;
+
+	{
+		// Scoped lock
+		FSystemWideCriticalSection StoredValuesLock(SessionSummaryDefs::GlobalLockName, Timeout);
+
+		if (StoredValuesLock.IsValid())
+		{
+			// Get list of sessions in storage
+			TArray<FSessionRecord> ExistingRecords = ReadStoredRecords();
+
+			TArray<FSessionRecord> SessionRecordsToDelete;
+
+			// Check each stored session to see if they should be sent or not 
+			for (FSessionRecord& Record : ExistingRecords)
+			{
+				if (IsSessionProcessRunning(Record))
+				{
+					// Skip processes that are still running
+					continue;
+				}
+
+				FTimespan RecordAge = FDateTime::UtcNow() - Record.Timestamp;
+
+				if (RecordAge > SessionSummaryDefs::SessionRecordTimeout)
+				{
+					if (RecordAge < SessionSummaryDefs::SessionRecordExpiration)
+					{
+						SessionRecordsToReport.Add(Record);
+					}
+					SessionRecordsToDelete.Add(Record);
+				}
+			}
+
+			for (const FSessionRecord& ToDelete : SessionRecordsToDelete)
+			{
+				DeleteStoredRecord(ToDelete);
+				ExistingRecords.RemoveAll([&ToDelete](const FSessionRecord& Record)
+					{
+						return Record.SessionId == ToDelete.SessionId;
+					});
+			}
+
+			// build up a new SessionList string
+			FString SessionListString;
+
+			for (const FSessionRecord& Remaining : ExistingRecords)
+			{
+				if (!SessionListString.IsEmpty())
+				{
+					SessionListString.Append(TEXT(","));
+				}
+
+				SessionListString.Append(Remaining.SessionId);
+			}
+
+			FPlatformMisc::SetStoredValue(SessionSummaryDefs::StoreId, SessionSummaryDefs::SessionSummarySection, SessionSummaryDefs::SessionListStoreKey, SessionListString);
+		}
+	}
+
+	for (const FSessionRecord& Record : SessionRecordsToReport)
+	{
+		SendSessionSummaryEvent(Record);
+	}
+}
+
+void FEditorSessionSummarySender::DeleteStoredRecord(const FSessionRecord& Record) const
+{
+	// Delete the session record in storage
+	FString SectionName = GetSessionStorageLocation(Record);
+
+	FPlatformMisc::DeleteStoredValue(SessionSummaryDefs::StoreId, SectionName, SessionSummaryDefs::ProjectNameStoreKey);
+	FPlatformMisc::DeleteStoredValue(SessionSummaryDefs::StoreId, SectionName, SessionSummaryDefs::IsCrashStoreKey);
+	FPlatformMisc::DeleteStoredValue(SessionSummaryDefs::StoreId, SectionName, SessionSummaryDefs::IsGPUCrashStoreKey);
+	FPlatformMisc::DeleteStoredValue(SessionSummaryDefs::StoreId, SectionName, SessionSummaryDefs::EngineVersionStoreKey);
+	FPlatformMisc::DeleteStoredValue(SessionSummaryDefs::StoreId, SectionName, SessionSummaryDefs::StartupTimestampStoreKey);
+	FPlatformMisc::DeleteStoredValue(SessionSummaryDefs::StoreId, SectionName, SessionSummaryDefs::TimestampStoreKey);
+	FPlatformMisc::DeleteStoredValue(SessionSummaryDefs::StoreId, SectionName, SessionSummaryDefs::Idle1MinStoreKey);
+	FPlatformMisc::DeleteStoredValue(SessionSummaryDefs::StoreId, SectionName, SessionSummaryDefs::Idle5MinStoreKey);
+	FPlatformMisc::DeleteStoredValue(SessionSummaryDefs::StoreId, SectionName, SessionSummaryDefs::Idle30MinStoreKey);
+	FPlatformMisc::DeleteStoredValue(SessionSummaryDefs::StoreId, SectionName, SessionSummaryDefs::IsDebuggerStoreKey);
+	FPlatformMisc::DeleteStoredValue(SessionSummaryDefs::StoreId, SectionName, SessionSummaryDefs::WasDebuggerStoreKey);
+	FPlatformMisc::DeleteStoredValue(SessionSummaryDefs::StoreId, SectionName, SessionSummaryDefs::WasShutdownStoreKey);
+	FPlatformMisc::DeleteStoredValue(SessionSummaryDefs::StoreId, SectionName, SessionSummaryDefs::UserActivityStoreKey);
+	FPlatformMisc::DeleteStoredValue(SessionSummaryDefs::StoreId, SectionName, SessionSummaryDefs::IsVanillaStoreKey);
+	FPlatformMisc::DeleteStoredValue(SessionSummaryDefs::StoreId, SectionName, SessionSummaryDefs::IsTerminatingKey);
+	FPlatformMisc::DeleteStoredValue(SessionSummaryDefs::StoreId, SectionName, SessionSummaryDefs::PlatformProcessIDKey);
+	FPlatformMisc::DeleteStoredValue(SessionSummaryDefs::StoreId, SectionName, SessionSummaryDefs::PluginsStoreKey);
+	FPlatformMisc::DeleteStoredValue(SessionSummaryDefs::StoreId, SectionName, SessionSummaryDefs::AverageFPSStoreKey);
+	FPlatformMisc::DeleteStoredValue(SessionSummaryDefs::StoreId, SectionName, SessionSummaryDefs::IsInPIEStoreKey);
+	FPlatformMisc::DeleteStoredValue(SessionSummaryDefs::StoreId, SectionName, SessionSummaryDefs::IsInEnterpriseStoreKey);
+	FPlatformMisc::DeleteStoredValue(SessionSummaryDefs::StoreId, SectionName, SessionSummaryDefs::IsInVRModeStoreKey);
+}
+
+bool FEditorSessionSummarySender::IsSessionProcessRunning(const FSessionRecord& Record) const
+{
+	FProcHandle Handle = FPlatformProcess::OpenProcess((uint32) Record.PlatformProcessID);
+	if (Handle.IsValid())
+	{
+		const bool bIsRunning = FPlatformProcess::IsProcRunning(Handle);
+		FPlatformProcess::CloseProc(Handle);
+		if (bIsRunning)
+		{
+			return true;
+		}
+	}
+
+	return false;
+}
+
+TArray<FSessionRecord> FEditorSessionSummarySender::ReadStoredRecords() const
+{
+	TArray<FSessionRecord> Records;
+
+	FString SessionListString;
+	FPlatformMisc::GetStoredValue(SessionSummaryDefs::StoreId, SessionSummaryDefs::SessionSummarySection, SessionSummaryDefs::SessionListStoreKey, SessionListString);
+
+	TArray<FString> SessionIDs;
+	SessionListString.ParseIntoArray(SessionIDs, TEXT(","));
+
+	// Retrieve all the sessions in the list from storage
+	for (const FString& SessionId : SessionIDs)
+	{
+		FSessionRecord NewRecord;
+		NewRecord.SessionId = SessionId;
+
+		FString SectionName = GetSessionStorageLocation(NewRecord);
+
+		// Read values
+		FString PlatformProcessIDString;
+		FPlatformMisc::GetStoredValue(SessionSummaryDefs::StoreId, SectionName, SessionSummaryDefs::PlatformProcessIDKey, PlatformProcessIDString);
+		FString EngineVersionString;
+		FPlatformMisc::GetStoredValue(SessionSummaryDefs::StoreId, SectionName, SessionSummaryDefs::EngineVersionStoreKey, EngineVersionString);
+		FString TimestampString;
+		FPlatformMisc::GetStoredValue(SessionSummaryDefs::StoreId, SectionName, SessionSummaryDefs::TimestampStoreKey, TimestampString);
+		FString ProjectName = SessionSummaryDefs::UnknownProjectValueString;
+		FPlatformMisc::GetStoredValue(SessionSummaryDefs::StoreId, SectionName, SessionSummaryDefs::ProjectNameStoreKey, ProjectName);
+		FString UserActivityString = SessionSummaryDefs::DefaultUserActivity;
+		FPlatformMisc::GetStoredValue(SessionSummaryDefs::StoreId, SectionName, SessionSummaryDefs::UserActivityStoreKey, UserActivityString);
+		FString StartupTimestampString;
+		FPlatformMisc::GetStoredValue(SessionSummaryDefs::StoreId, SectionName, SessionSummaryDefs::StartupTimestampStoreKey, StartupTimestampString);
+		FString Idle1MinString;
+		FPlatformMisc::GetStoredValue(SessionSummaryDefs::StoreId, SectionName, SessionSummaryDefs::Idle1MinStoreKey, Idle1MinString);
+		FString Idle5MinString;
+		FPlatformMisc::GetStoredValue(SessionSummaryDefs::StoreId, SectionName, SessionSummaryDefs::Idle5MinStoreKey, Idle5MinString);
+		FString Idle30MinString;
+		FPlatformMisc::GetStoredValue(SessionSummaryDefs::StoreId, SectionName, SessionSummaryDefs::Idle30MinStoreKey, Idle30MinString);
+		FString PluginsString;
+		FPlatformMisc::GetStoredValue(SessionSummaryDefs::StoreId, SectionName, SessionSummaryDefs::PluginsStoreKey, PluginsString);
+		FString AverageFPSString;
+		FPlatformMisc::GetStoredValue(SessionSummaryDefs::StoreId, SectionName, SessionSummaryDefs::AverageFPSStoreKey, AverageFPSString);
+
+		// Create new record from the read values
+		NewRecord.PlatformProcessID = FCString::Atoi(*PlatformProcessIDString);
+		NewRecord.ProjectName = ProjectName;
+		NewRecord.EngineVersion = EngineVersionString;
+		NewRecord.StartupTimestamp = StringToTimestamp(StartupTimestampString);
+		NewRecord.Timestamp = StringToTimestamp(TimestampString);
+		NewRecord.Idle1Min = FCString::Atod(*Idle1MinString);
+		NewRecord.Idle5Min = FCString::Atod(*Idle5MinString);
+		NewRecord.Idle30Min = FCString::Atod(*Idle30MinString);
+		NewRecord.AverageFPS = FCString::Atof(*AverageFPSString);
+		NewRecord.CurrentUserActivity = UserActivityString;
+		NewRecord.bCrashed = GetStoredBool(SectionName, SessionSummaryDefs::IsCrashStoreKey);
+		NewRecord.bGPUCrashed = GetStoredBool(SectionName, SessionSummaryDefs::IsGPUCrashStoreKey);
+		NewRecord.bIsDebugger = GetStoredBool(SectionName, SessionSummaryDefs::IsDebuggerStoreKey);
+		NewRecord.bWasEverDebugger = GetStoredBool(SectionName, SessionSummaryDefs::WasDebuggerStoreKey);
+		NewRecord.bIsVanilla = GetStoredBool(SectionName, SessionSummaryDefs::IsVanillaStoreKey);
+		NewRecord.bIsTerminating = GetStoredBool(SectionName, SessionSummaryDefs::IsTerminatingKey);
+		NewRecord.bWasShutdown = GetStoredBool(SectionName, SessionSummaryDefs::WasShutdownStoreKey);
+		NewRecord.bIsInPIE = GetStoredBool(SectionName, SessionSummaryDefs::IsInPIEStoreKey);
+		NewRecord.bIsInVRMode = GetStoredBool(SectionName, SessionSummaryDefs::IsInVRModeStoreKey);
+		NewRecord.bIsInEnterprise = GetStoredBool(SectionName, SessionSummaryDefs::IsInEnterpriseStoreKey);
+
+		PluginsString.ParseIntoArray(NewRecord.Plugins, TEXT(","));
+
+		Records.Add(NewRecord);
+	}
+
+	return MoveTemp(Records);
+}
+
+void FEditorSessionSummarySender::SendSessionSummaryEvent(const FSessionRecord& Record) const
+{
+	FGuid SessionId;
+	FString SessionIdString = Record.SessionId;
+	if (FGuid::Parse(SessionIdString, SessionId))
+	{
+		// convert session guid to one with braces for sending to analytics
+		SessionIdString = SessionId.ToString(EGuidFormats::DigitsWithHyphensInBraces);
+	}
+
+#if !PLATFORM_PS4
+	FString ShutdownTypeString = Record.bCrashed ? SessionSummaryDefs::CrashSessionToken :
+		(Record.bWasEverDebugger ? SessionSummaryDefs::DebuggerSessionToken :
+		(Record.bIsTerminating ? SessionSummaryDefs::TerminatedSessionToken :
+		(Record.bWasShutdown ? SessionSummaryDefs::ShutdownSessionToken : SessionSummaryDefs::AbnormalSessionToken)));
+#else
+	// PS4 cannot set the crash flag so report abnormal shutdowns with a specific token meaning "crash or abnormal shutdown".
+	FString ShutdownTypeString = Record.bWasEverDebugger ? SessionSummaryDefs::DebuggerSessionToken : SessionSummaryDefs::PS4SessionToken;
+#endif
+
+	FString PluginsString = FString::Join(Record.Plugins, TEXT(","));
+
+	TArray< FAnalyticsEventAttribute > AnalyticsAttributes;
+	AnalyticsAttributes.Emplace(TEXT("ProjectName"), Record.ProjectName);
+	AnalyticsAttributes.Emplace(TEXT("Platform"), FPlatformProperties::PlatformName());
+	AnalyticsAttributes.Emplace(TEXT("SessionId"), SessionIdString);
+	AnalyticsAttributes.Emplace(TEXT("EngineVersion"), Record.EngineVersion);
+	AnalyticsAttributes.Emplace(TEXT("ShutdownType"), ShutdownTypeString);
+	AnalyticsAttributes.Emplace(TEXT("Timestamp"), Record.Timestamp.ToIso8601());
+	AnalyticsAttributes.Emplace(TEXT("CurrentUserActivity"), Record.CurrentUserActivity);
+	AnalyticsAttributes.Emplace(TEXT("IsVanilla"), Record.bIsVanilla);
+	AnalyticsAttributes.Emplace(TEXT("WasDebugged"), Record.bWasEverDebugger);
+	AnalyticsAttributes.Emplace(TEXT("GPUCrash"), Record.bGPUCrashed);
+	AnalyticsAttributes.Emplace(SessionSummaryDefs::WasShutdownStoreKey, Record.bWasShutdown);
+	AnalyticsAttributes.Emplace(SessionSummaryDefs::StartupTimestampStoreKey, Record.StartupTimestamp.ToIso8601());
+	AnalyticsAttributes.Emplace(SessionSummaryDefs::AverageFPSStoreKey, Record.AverageFPS);
+	AnalyticsAttributes.Emplace(SessionSummaryDefs::IsInPIEStoreKey, Record.bIsInPIE);
+	AnalyticsAttributes.Emplace(SessionSummaryDefs::IsInEnterpriseStoreKey, Record.bIsInEnterprise);
+	AnalyticsAttributes.Emplace(SessionSummaryDefs::IsInVRModeStoreKey, Record.bIsInVRMode);
+
+	double SessionDuration = (Record.Timestamp - Record.StartupTimestamp).GetTotalSeconds();
+	AnalyticsAttributes.Emplace(SessionSummaryDefs::SessionDurationStoreKey, SessionDuration);
+
+	AnalyticsAttributes.Emplace(TEXT("1MinIdle"), Record.Idle1Min);
+	AnalyticsAttributes.Emplace(TEXT("5MinIdle"), Record.Idle5Min);
+	AnalyticsAttributes.Emplace(TEXT("30MinIdle"), Record.Idle30Min);
+
+	// Add project info whether we are in editor or game.
+	const UGeneralProjectSettings& ProjectSettings = *GetDefault<UGeneralProjectSettings>();
+	FString OSMajor;
+	FString OSMinor;
+	FPlatformMisc::GetOSVersions(/*out*/ OSMajor, /*out*/ OSMinor);
+
+	FPlatformMemoryStats Stats = FPlatformMemory::GetStats();
+	AnalyticsAttributes.Emplace(TEXT("ProjectName"), ProjectSettings.ProjectName);
+	AnalyticsAttributes.Emplace(TEXT("ProjectID"), ProjectSettings.ProjectID);
+	AnalyticsAttributes.Emplace(TEXT("ProjectDescription"), ProjectSettings.Description);
+	AnalyticsAttributes.Emplace(TEXT("ProjectVersion"), ProjectSettings.ProjectVersion);
+	AnalyticsAttributes.Emplace(TEXT("GPUVendorID"), GRHIVendorId);
+	AnalyticsAttributes.Emplace(TEXT("GPUDeviceID"), GRHIDeviceId);
+	AnalyticsAttributes.Emplace(TEXT("GRHIDeviceRevision"), GRHIDeviceRevision);
+	AnalyticsAttributes.Emplace(TEXT("GRHIAdapterInternalDriverVersion"), GRHIAdapterInternalDriverVersion);
+	AnalyticsAttributes.Emplace(TEXT("GRHIAdapterUserDriverVersion"), GRHIAdapterUserDriverVersion);
+	AnalyticsAttributes.Emplace(TEXT("TotalPhysicalRAM"), Stats.TotalPhysical);
+	AnalyticsAttributes.Emplace(TEXT("CPUPhysicalCores"), FPlatformMisc::NumberOfCores());
+	AnalyticsAttributes.Emplace(TEXT("CPULogicalCores"), FPlatformMisc::NumberOfCoresIncludingHyperthreads());
+	AnalyticsAttributes.Emplace(TEXT("DesktopGPUAdapter"), FPlatformMisc::GetPrimaryGPUBrand());
+	AnalyticsAttributes.Emplace(TEXT("RenderingGPUAdapter"), GRHIAdapterName);
+	AnalyticsAttributes.Emplace(TEXT("CPUVendor"), FPlatformMisc::GetCPUVendor());
+	AnalyticsAttributes.Emplace(TEXT("CPUBrand"), FPlatformMisc::GetCPUBrand());
+	AnalyticsAttributes.Emplace(TEXT("OSMajor"), OSMajor);
+	AnalyticsAttributes.Emplace(TEXT("OSMinor"), OSMinor);
+	AnalyticsAttributes.Emplace(TEXT("OSVersion"), FPlatformMisc::GetOSVersion());
+	AnalyticsAttributes.Emplace(TEXT("Is64BitOS"), FPlatformMisc::Is64bitOperatingSystem());
+
+	FEngineAnalytics::GetProvider().RecordEvent(TEXT("SessionSummary"), AnalyticsAttributes);
+
+	UE_LOG(LogEditorSessionSummary, Log, TEXT("EditorSessionSummary sent report. Type=%s, SessionId=%s"), *ShutdownTypeString, *SessionIdString);
+}
+
+#undef LOCTEXT_NAMESPACE
diff --git a/Engine/Source/Runtime/Engine/Private/Analytics/EditorSessionSummary.h b/Engine/Source/Runtime/Engine/Private/Analytics/EditorSessionSummary.h
new file mode 100644
--- /dev/null
+++ b/Engine/Source/Runtime/Engine/Private/Analytics/EditorSessionSummary.h
@@ -0,0 +1,63 @@
+// Copyright 1998-2019 Epic Games, Inc. All Rights Reserved.
+
+#pragma once
+
+#include "CoreMinimal.h"
+
+struct FUserActivity;
+struct FSessionRecord;
+
+/** Writer for SessionSummary events to track all editor sessions. */
+class FEditorSessionSummaryWriter
+{
+public:
+	FEditorSessionSummaryWriter();
+
+	void Initialize();
+	void Tick(float DeltaTime);
+	void Shutdown();
+
+private:
+	void InitializeRecords(bool bFirstAttempt);
+
+	FSessionRecord* CreateRecordForCurrentSession() const;
+	void WriteStoredRecord(const FSessionRecord& Record) const;
+
+	void OnCrashing();
+	void OnTerminate();
+	void OnUserActivity(const FUserActivity& UserActivity);
+	void OnVanillaStateChanged(bool bIsVanilla);
+
+	FString GetUserActivityString() const;
+	void UpdateTimestamps();
+
+private:
+	bool bInitializedRecords;
+	FSessionRecord* CurrentSession;
+	FString CurrentSessionSectionName;
+	float HeartbeatTimeElapsed;
+	bool bShutdown;
+};
+
+/** Sender of SessionSummary events from all editor sessions in-between runs. 
+  * Separated from Writer to make it easier to run it out-of-process.
+  */
+class FEditorSessionSummarySender
+{
+public:
+	FEditorSessionSummarySender();
+
+	void Initialize();
+	void Tick(float DeltaTime);
+
+private:
+	/** Send any stored records. Returns true if we successfully recorded any outstanding events. */
+	void SendStoredRecords(FTimespan Timeout) const;
+	TArray<FSessionRecord> ReadStoredRecords() const;
+	void DeleteStoredRecord(const FSessionRecord& Record) const;
+	void SendSessionSummaryEvent(const FSessionRecord& Record) const;
+	bool IsSessionProcessRunning(const FSessionRecord& Record) const;
+
+private:
+	float HeartbeatTimeElapsed;
+};
\ No newline at end of file
diff --git a/Engine/Source/Runtime/Engine/Private/EngineAnalytics.cpp b/Engine/Source/Runtime/Engine/Private/EngineAnalytics.cpp
--- a/Engine/Source/Runtime/Engine/Private/EngineAnalytics.cpp
+++ b/Engine/Source/Runtime/Engine/Private/EngineAnalytics.cpp
@@ -15,11 +15,15 @@
 #include "EngineSessionManager.h"
 #include "Misc/EngineVersion.h"
 #include "RHI.h"
+#include "Analytics/EditorSessionSummary.h"
 
 bool FEngineAnalytics::bIsInitialized;
 TSharedPtr<IAnalyticsProviderET> FEngineAnalytics::Analytics;
 TSharedPtr<FEngineSessionManager> FEngineAnalytics::SessionManager;
 
+static TSharedPtr<FEditorSessionSummaryWriter> SessionSummaryWriter;
+static TSharedPtr<FEditorSessionSummarySender> SessionSummarySender;
+
 /**
 * Default config func.
 */
@@ -150,10 +154,20 @@ void FEngineAnalytics::Initialize()
 			SessionManager = MakeShareable(new FEngineSessionManager(EEngineSessionManagerMode::Editor));
 			SessionManager->Initialize();
 		}
+
+		if (!SessionSummaryWriter.IsValid())
+		{
+			SessionSummaryWriter = MakeShareable(new FEditorSessionSummaryWriter());
+			SessionSummaryWriter->Initialize();
+		}
+
+		if (!SessionSummarySender.IsValid())
+		{
+			SessionSummarySender = MakeShareable(new FEditorSessionSummarySender());
+		}
 	}
 }
 
-
 void FEngineAnalytics::Shutdown(bool bIsEngineShutdown)
 {
 	ensure(!Analytics.IsValid() || Analytics.IsUnique());
@@ -166,6 +180,14 @@ void FEngineAnalytics::Shutdown(bool bIsEngineShutdown)
 		SessionManager->Shutdown();
 		SessionManager.Reset();
 	}
+
+	if (SessionSummaryWriter.IsValid())
+	{
+		SessionSummaryWriter->Shutdown();
+		SessionSummaryWriter.Reset();
+	}
+
+	SessionSummarySender.Reset();
 }
 
 void FEngineAnalytics::Tick(float DeltaTime)
@@ -176,4 +198,14 @@ void FEngineAnalytics::Tick(float DeltaTime)
 	{
 		SessionManager->Tick(DeltaTime);
 	}
+
+	if (SessionSummaryWriter.IsValid())
+	{
+		SessionSummaryWriter->Tick(DeltaTime);
+	}
+
+	if (SessionSummarySender.IsValid())
+	{
+		SessionSummarySender->Tick(DeltaTime);
+	}
 }
