Compile fix for unity builds merging two files with identically-named utility functions in anonymous namespaces. Named the namespaces to avoid the collision.

Renamed the FSessionRecord in EditorSessionSummary.cpp to FEditorSessionRecord to disambiguate it from FEngineSessionManager::FSessionRecord. Though not technically ambiguous since it does compile, it feels like poor form.

Also had an unused platform compile flag that shouldn't have been there in the first place.

diff --git a/Engine/Source/Runtime/Engine/Private/Analytics/EditorSessionSummary.cpp b/Engine/Source/Runtime/Engine/Private/Analytics/EditorSessionSummary.cpp
--- a/Engine/Source/Runtime/Engine/Private/Analytics/EditorSessionSummary.cpp
+++ b/Engine/Source/Runtime/Engine/Private/Analytics/EditorSessionSummary.cpp
@@ -91,7 +91,7 @@ namespace SessionSummaryDefs
 	static const FString IsInPIEStoreKey(TEXT("IsInPIE"));
 }
 
-struct FSessionRecord
+struct FEditorSessionRecord
 {
 	FString SessionId;
 	FString ProjectName;
@@ -117,7 +117,7 @@ struct FSessionRecord
 	bool bIsInEnterprise : 1;
 	bool bIsInVRMode : 1;
 
-	FSessionRecord()
+	FEditorSessionRecord()
 	{
 		PlatformProcessID = 0;
 		StartupTimestamp = FDateTime::MinValue();
@@ -140,7 +140,7 @@ struct FSessionRecord
 };
 
 // Utilities for writing to stored values
-namespace
+namespace EditorSessionSummaryUtils
 {
 	FString TimestampToString(FDateTime InTimestamp)
 	{
@@ -170,7 +170,7 @@ namespace
 		return StoredString == SessionSummaryDefs::TrueValueString;
 	}
 
-	FString GetSessionStorageLocation(const FSessionRecord& Record)
+	FString GetSessionStorageLocation(const FEditorSessionRecord& Record)
 	{
 		return SessionSummaryDefs::SessionSummarySection + TEXT("/") + Record.SessionId;
 	}
@@ -222,7 +222,7 @@ void FEditorSessionSummaryWriter::InitializeRecords(bool bFirstAttempt)
 
 			// Create a session record for this session
 			CurrentSession = CreateRecordForCurrentSession();
-			CurrentSessionSectionName = GetSessionStorageLocation(*CurrentSession);
+			CurrentSessionSectionName = EditorSessionSummaryUtils::GetSessionStorageLocation(*CurrentSession);
 			WriteStoredRecord(*CurrentSession);
 
 			bInitializedRecords = true;
@@ -247,7 +247,7 @@ void FEditorSessionSummaryWriter::InitializeRecords(bool bFirstAttempt)
 void FEditorSessionSummaryWriter::UpdateTimestamps()
 {
 	CurrentSession->Timestamp = FDateTime::UtcNow();
-	FPlatformMisc::SetStoredValue(SessionSummaryDefs::StoreId, CurrentSessionSectionName, SessionSummaryDefs::TimestampStoreKey, TimestampToString(CurrentSession->Timestamp));
+	FPlatformMisc::SetStoredValue(SessionSummaryDefs::StoreId, CurrentSessionSectionName, SessionSummaryDefs::TimestampStoreKey, EditorSessionSummaryUtils::TimestampToString(CurrentSession->Timestamp));
 
 	const float IdleSeconds = FPlatformTime::Seconds() - FSlateApplication::Get().GetLastUserInteractionTime();
 
@@ -277,7 +277,7 @@ void FEditorSessionSummaryWriter::Tick(float DeltaTime)
 {
 	HeartbeatTimeElapsed += DeltaTime;
 
-	if (HeartbeatTimeElapsed > (float) SessionSummaryDefs::HeartbeatPeriodSeconds && !bShutdown)
+	if (HeartbeatTimeElapsed > (float) SessionSummaryDefs::HeartbeatPeriodSeconds)
 	{
 		HeartbeatTimeElapsed = 0.0f;
 
@@ -292,7 +292,7 @@ void FEditorSessionSummaryWriter::Tick(float DeltaTime)
 			if (CurrentSession->bIsDebugger != bIsDebuggerPresent)
 			{
 				CurrentSession->bIsDebugger = bIsDebuggerPresent;
-				FPlatformMisc::SetStoredValue(SessionSummaryDefs::StoreId, CurrentSessionSectionName, SessionSummaryDefs::IsDebuggerStoreKey, BoolToStoredString(CurrentSession->bIsDebugger));
+				FPlatformMisc::SetStoredValue(SessionSummaryDefs::StoreId, CurrentSessionSectionName, SessionSummaryDefs::IsDebuggerStoreKey, EditorSessionSummaryUtils::BoolToStoredString(CurrentSession->bIsDebugger));
 
 				if (!CurrentSession->bWasEverDebugger && CurrentSession->bIsDebugger)
 				{
@@ -311,15 +311,15 @@ void FEditorSessionSummaryWriter::Tick(float DeltaTime)
 			FPlatformMisc::SetStoredValue(SessionSummaryDefs::StoreId, CurrentSessionSectionName, SessionSummaryDefs::AverageFPSStoreKey, AverageFPSString);
 
 			CurrentSession->bIsInVRMode = IVREditorModule::Get().IsVREditorModeActive();
-			const FString IsInVRModeString = BoolToStoredString(CurrentSession->bIsInVRMode);
+			const FString IsInVRModeString = EditorSessionSummaryUtils::BoolToStoredString(CurrentSession->bIsInVRMode);
 			FPlatformMisc::SetStoredValue(SessionSummaryDefs::StoreId, CurrentSessionSectionName, SessionSummaryDefs::IsInVRModeStoreKey, IsInVRModeString);
 
 			CurrentSession->bIsInEnterprise = IProjectManager::Get().IsEnterpriseProject();
-			const FString IsInEnterpriseString = BoolToStoredString(CurrentSession->bIsInEnterprise);
+			const FString IsInEnterpriseString = EditorSessionSummaryUtils::BoolToStoredString(CurrentSession->bIsInEnterprise);
 			FPlatformMisc::SetStoredValue(SessionSummaryDefs::StoreId, CurrentSessionSectionName, SessionSummaryDefs::IsInEnterpriseStoreKey, IsInEnterpriseString);
 
 			CurrentSession->bIsInPIE = FPlayWorldCommandCallbacks::IsInPIE();
-			const FString IsInPIEString = BoolToStoredString(CurrentSession->bIsInPIE);
+			const FString IsInPIEString = EditorSessionSummaryUtils::BoolToStoredString(CurrentSession->bIsInPIE);
 			FPlatformMisc::SetStoredValue(SessionSummaryDefs::StoreId, CurrentSessionSectionName, SessionSummaryDefs::IsInPIEStoreKey, IsInPIEString);
 #endif
 		}
@@ -357,36 +357,36 @@ void FEditorSessionSummaryWriter::Shutdown()
 	}
 }
 
-void FEditorSessionSummaryWriter::WriteStoredRecord(const FSessionRecord& Record) const
+void FEditorSessionSummaryWriter::WriteStoredRecord(const FEditorSessionRecord& Record) const
 {
-	FString StorageLocation = GetSessionStorageLocation(Record);
+	FString StorageLocation = EditorSessionSummaryUtils::GetSessionStorageLocation(Record);
 	FString PluginsString = FString::Join(Record.Plugins, TEXT(","));
 
 	FPlatformMisc::SetStoredValue(SessionSummaryDefs::StoreId, StorageLocation, SessionSummaryDefs::ProjectNameStoreKey, Record.ProjectName);
 	FPlatformMisc::SetStoredValue(SessionSummaryDefs::StoreId, StorageLocation, SessionSummaryDefs::IsCrashStoreKey, SessionSummaryDefs::FalseValueString);
 	FPlatformMisc::SetStoredValue(SessionSummaryDefs::StoreId, StorageLocation, SessionSummaryDefs::EngineVersionStoreKey, Record.EngineVersion);
-	FPlatformMisc::SetStoredValue(SessionSummaryDefs::StoreId, StorageLocation, SessionSummaryDefs::StartupTimestampStoreKey, TimestampToString(Record.StartupTimestamp));
-	FPlatformMisc::SetStoredValue(SessionSummaryDefs::StoreId, StorageLocation, SessionSummaryDefs::TimestampStoreKey, TimestampToString(Record.Timestamp));
+	FPlatformMisc::SetStoredValue(SessionSummaryDefs::StoreId, StorageLocation, SessionSummaryDefs::StartupTimestampStoreKey, EditorSessionSummaryUtils::TimestampToString(Record.StartupTimestamp));
+	FPlatformMisc::SetStoredValue(SessionSummaryDefs::StoreId, StorageLocation, SessionSummaryDefs::TimestampStoreKey, EditorSessionSummaryUtils::TimestampToString(Record.Timestamp));
 	FPlatformMisc::SetStoredValue(SessionSummaryDefs::StoreId, StorageLocation, SessionSummaryDefs::Idle1MinStoreKey, FString::FromInt(Record.Idle1Min));
 	FPlatformMisc::SetStoredValue(SessionSummaryDefs::StoreId, StorageLocation, SessionSummaryDefs::Idle5MinStoreKey, FString::FromInt(Record.Idle5Min));
 	FPlatformMisc::SetStoredValue(SessionSummaryDefs::StoreId, StorageLocation, SessionSummaryDefs::Idle30MinStoreKey, FString::FromInt(Record.Idle30Min));
 	FPlatformMisc::SetStoredValue(SessionSummaryDefs::StoreId, StorageLocation, SessionSummaryDefs::UserActivityStoreKey, Record.CurrentUserActivity);
-	FPlatformMisc::SetStoredValue(SessionSummaryDefs::StoreId, StorageLocation, SessionSummaryDefs::IsVanillaStoreKey, BoolToStoredString(Record.bIsVanilla));
-	FPlatformMisc::SetStoredValue(SessionSummaryDefs::StoreId, StorageLocation, SessionSummaryDefs::IsTerminatingKey, BoolToStoredString(Record.bIsTerminating));
+	FPlatformMisc::SetStoredValue(SessionSummaryDefs::StoreId, StorageLocation, SessionSummaryDefs::IsVanillaStoreKey, EditorSessionSummaryUtils::BoolToStoredString(Record.bIsVanilla));
+	FPlatformMisc::SetStoredValue(SessionSummaryDefs::StoreId, StorageLocation, SessionSummaryDefs::IsTerminatingKey, EditorSessionSummaryUtils::BoolToStoredString(Record.bIsTerminating));
 	FPlatformMisc::SetStoredValue(SessionSummaryDefs::StoreId, StorageLocation, SessionSummaryDefs::PlatformProcessIDKey, FString::FromInt(Record.PlatformProcessID));
 	FPlatformMisc::SetStoredValue(SessionSummaryDefs::StoreId, StorageLocation, SessionSummaryDefs::PluginsStoreKey, PluginsString);
 	FPlatformMisc::SetStoredValue(SessionSummaryDefs::StoreId, StorageLocation, SessionSummaryDefs::AverageFPSStoreKey, FString::SanitizeFloat(Record.AverageFPS));
-	FPlatformMisc::SetStoredValue(SessionSummaryDefs::StoreId, StorageLocation, SessionSummaryDefs::IsDebuggerStoreKey, BoolToStoredString(Record.bIsDebugger));
-	FPlatformMisc::SetStoredValue(SessionSummaryDefs::StoreId, StorageLocation, SessionSummaryDefs::WasDebuggerStoreKey, BoolToStoredString(Record.bWasEverDebugger));
-	FPlatformMisc::SetStoredValue(SessionSummaryDefs::StoreId, StorageLocation, SessionSummaryDefs::WasShutdownStoreKey, BoolToStoredString(Record.bWasShutdown));
-	FPlatformMisc::SetStoredValue(SessionSummaryDefs::StoreId, StorageLocation, SessionSummaryDefs::IsInPIEStoreKey, BoolToStoredString(Record.bIsInPIE));
-	FPlatformMisc::SetStoredValue(SessionSummaryDefs::StoreId, StorageLocation, SessionSummaryDefs::IsInEnterpriseStoreKey, BoolToStoredString(Record.bIsInEnterprise));
-	FPlatformMisc::SetStoredValue(SessionSummaryDefs::StoreId, StorageLocation, SessionSummaryDefs::IsInVRModeStoreKey, BoolToStoredString(Record.bIsInVRMode));
+	FPlatformMisc::SetStoredValue(SessionSummaryDefs::StoreId, StorageLocation, SessionSummaryDefs::IsDebuggerStoreKey, EditorSessionSummaryUtils::BoolToStoredString(Record.bIsDebugger));
+	FPlatformMisc::SetStoredValue(SessionSummaryDefs::StoreId, StorageLocation, SessionSummaryDefs::WasDebuggerStoreKey, EditorSessionSummaryUtils::BoolToStoredString(Record.bWasEverDebugger));
+	FPlatformMisc::SetStoredValue(SessionSummaryDefs::StoreId, StorageLocation, SessionSummaryDefs::WasShutdownStoreKey, EditorSessionSummaryUtils::BoolToStoredString(Record.bWasShutdown));
+	FPlatformMisc::SetStoredValue(SessionSummaryDefs::StoreId, StorageLocation, SessionSummaryDefs::IsInPIEStoreKey, EditorSessionSummaryUtils::BoolToStoredString(Record.bIsInPIE));
+	FPlatformMisc::SetStoredValue(SessionSummaryDefs::StoreId, StorageLocation, SessionSummaryDefs::IsInEnterpriseStoreKey, EditorSessionSummaryUtils::BoolToStoredString(Record.bIsInEnterprise));
+	FPlatformMisc::SetStoredValue(SessionSummaryDefs::StoreId, StorageLocation, SessionSummaryDefs::IsInVRModeStoreKey, EditorSessionSummaryUtils::BoolToStoredString(Record.bIsInVRMode));
 }
 
-FSessionRecord* FEditorSessionSummaryWriter::CreateRecordForCurrentSession() const
+FEditorSessionRecord* FEditorSessionSummaryWriter::CreateRecordForCurrentSession() const
 {
-	FSessionRecord* Record = new FSessionRecord();
+	FEditorSessionRecord* Record = new FEditorSessionRecord();
 
 	FGuid SessionId;
 	if (FGuid::Parse(FEngineAnalytics::GetProvider().GetSessionID(), SessionId))
@@ -433,7 +433,7 @@ void FEditorSessionSummaryWriter::OnCrashing()
 		CurrentSession->bCrashed = true;
 		FPlatformMisc::SetStoredValue(SessionSummaryDefs::StoreId, CurrentSessionSectionName, SessionSummaryDefs::IsCrashStoreKey, SessionSummaryDefs::TrueValueString);
 		CurrentSession->bGPUCrashed = GIsGPUCrashed;
-		FPlatformMisc::SetStoredValue(SessionSummaryDefs::StoreId, CurrentSessionSectionName, SessionSummaryDefs::IsGPUCrashStoreKey, BoolToStoredString(CurrentSession->bGPUCrashed));
+		FPlatformMisc::SetStoredValue(SessionSummaryDefs::StoreId, CurrentSessionSectionName, SessionSummaryDefs::IsGPUCrashStoreKey, EditorSessionSummaryUtils::BoolToStoredString(CurrentSession->bGPUCrashed));
 	}
 }
 
@@ -458,7 +458,7 @@ void FEditorSessionSummaryWriter::OnVanillaStateChanged(bool bIsVanilla)
 	if (bInitializedRecords && CurrentSession->bIsVanilla != bIsVanilla)
 	{
 		CurrentSession->bIsVanilla = bIsVanilla;
-		FPlatformMisc::SetStoredValue(SessionSummaryDefs::StoreId, CurrentSessionSectionName, SessionSummaryDefs::IsVanillaStoreKey, BoolToStoredString(CurrentSession->bIsVanilla));
+		FPlatformMisc::SetStoredValue(SessionSummaryDefs::StoreId, CurrentSessionSectionName, SessionSummaryDefs::IsVanillaStoreKey, EditorSessionSummaryUtils::BoolToStoredString(CurrentSession->bIsVanilla));
 	}
 }
 
@@ -515,7 +515,7 @@ void FEditorSessionSummarySender::SendStoredRecords(FTimespan Timeout) const
 		return;
 	}
 
-	TArray<FSessionRecord> SessionRecordsToReport;
+	TArray<FEditorSessionRecord> SessionRecordsToReport;
 
 	{
 		// Scoped lock
@@ -524,12 +524,12 @@ void FEditorSessionSummarySender::SendStoredRecords(FTimespan Timeout) const
 		if (StoredValuesLock.IsValid())
 		{
 			// Get list of sessions in storage
-			TArray<FSessionRecord> ExistingRecords = ReadStoredRecords();
+			TArray<FEditorSessionRecord> ExistingRecords = ReadStoredRecords();
 
-			TArray<FSessionRecord> SessionRecordsToDelete;
+			TArray<FEditorSessionRecord> SessionRecordsToDelete;
 
 			// Check each stored session to see if they should be sent or not 
-			for (FSessionRecord& Record : ExistingRecords)
+			for (FEditorSessionRecord& Record : ExistingRecords)
 			{
 				if (IsSessionProcessRunning(Record))
 				{
@@ -549,10 +549,10 @@ void FEditorSessionSummarySender::SendStoredRecords(FTimespan Timeout) const
 				}
 			}
 
-			for (const FSessionRecord& ToDelete : SessionRecordsToDelete)
+			for (const FEditorSessionRecord& ToDelete : SessionRecordsToDelete)
 			{
 				DeleteStoredRecord(ToDelete);
-				ExistingRecords.RemoveAll([&ToDelete](const FSessionRecord& Record)
+				ExistingRecords.RemoveAll([&ToDelete](const FEditorSessionRecord& Record)
 					{
 						return Record.SessionId == ToDelete.SessionId;
 					});
@@ -561,7 +561,7 @@ void FEditorSessionSummarySender::SendStoredRecords(FTimespan Timeout) const
 			// build up a new SessionList string
 			FString SessionListString;
 
-			for (const FSessionRecord& Remaining : ExistingRecords)
+			for (const FEditorSessionRecord& Remaining : ExistingRecords)
 			{
 				if (!SessionListString.IsEmpty())
 				{
@@ -575,16 +575,16 @@ void FEditorSessionSummarySender::SendStoredRecords(FTimespan Timeout) const
 		}
 	}
 
-	for (const FSessionRecord& Record : SessionRecordsToReport)
+	for (const FEditorSessionRecord& Record : SessionRecordsToReport)
 	{
 		SendSessionSummaryEvent(Record);
 	}
 }
 
-void FEditorSessionSummarySender::DeleteStoredRecord(const FSessionRecord& Record) const
+void FEditorSessionSummarySender::DeleteStoredRecord(const FEditorSessionRecord& Record) const
 {
 	// Delete the session record in storage
-	FString SectionName = GetSessionStorageLocation(Record);
+	FString SectionName = EditorSessionSummaryUtils::GetSessionStorageLocation(Record);
 
 	FPlatformMisc::DeleteStoredValue(SessionSummaryDefs::StoreId, SectionName, SessionSummaryDefs::ProjectNameStoreKey);
 	FPlatformMisc::DeleteStoredValue(SessionSummaryDefs::StoreId, SectionName, SessionSummaryDefs::IsCrashStoreKey);
@@ -609,7 +609,7 @@ void FEditorSessionSummarySender::DeleteStoredRecord(const FSessionRecord& Recor
 	FPlatformMisc::DeleteStoredValue(SessionSummaryDefs::StoreId, SectionName, SessionSummaryDefs::IsInVRModeStoreKey);
 }
 
-bool FEditorSessionSummarySender::IsSessionProcessRunning(const FSessionRecord& Record) const
+bool FEditorSessionSummarySender::IsSessionProcessRunning(const FEditorSessionRecord& Record) const
 {
 	FProcHandle Handle = FPlatformProcess::OpenProcess((uint32) Record.PlatformProcessID);
 	if (Handle.IsValid())
@@ -625,9 +625,9 @@ bool FEditorSessionSummarySender::IsSessionProcessRunning(const FSessionRecord&
 	return false;
 }
 
-TArray<FSessionRecord> FEditorSessionSummarySender::ReadStoredRecords() const
+TArray<FEditorSessionRecord> FEditorSessionSummarySender::ReadStoredRecords() const
 {
-	TArray<FSessionRecord> Records;
+	TArray<FEditorSessionRecord> Records;
 
 	FString SessionListString;
 	FPlatformMisc::GetStoredValue(SessionSummaryDefs::StoreId, SessionSummaryDefs::SessionSummarySection, SessionSummaryDefs::SessionListStoreKey, SessionListString);
@@ -638,10 +638,10 @@ TArray<FSessionRecord> FEditorSessionSummarySender::ReadStoredRecords() const
 	// Retrieve all the sessions in the list from storage
 	for (const FString& SessionId : SessionIDs)
 	{
-		FSessionRecord NewRecord;
+		FEditorSessionRecord NewRecord;
 		NewRecord.SessionId = SessionId;
 
-		FString SectionName = GetSessionStorageLocation(NewRecord);
+		FString SectionName = EditorSessionSummaryUtils::GetSessionStorageLocation(NewRecord);
 
 		// Read values
 		FString PlatformProcessIDString;
@@ -671,23 +671,23 @@ TArray<FSessionRecord> FEditorSessionSummarySender::ReadStoredRecords() const
 		NewRecord.PlatformProcessID = FCString::Atoi(*PlatformProcessIDString);
 		NewRecord.ProjectName = ProjectName;
 		NewRecord.EngineVersion = EngineVersionString;
-		NewRecord.StartupTimestamp = StringToTimestamp(StartupTimestampString);
-		NewRecord.Timestamp = StringToTimestamp(TimestampString);
+		NewRecord.StartupTimestamp = EditorSessionSummaryUtils::StringToTimestamp(StartupTimestampString);
+		NewRecord.Timestamp = EditorSessionSummaryUtils::StringToTimestamp(TimestampString);
 		NewRecord.Idle1Min = FCString::Atod(*Idle1MinString);
 		NewRecord.Idle5Min = FCString::Atod(*Idle5MinString);
 		NewRecord.Idle30Min = FCString::Atod(*Idle30MinString);
 		NewRecord.AverageFPS = FCString::Atof(*AverageFPSString);
 		NewRecord.CurrentUserActivity = UserActivityString;
-		NewRecord.bCrashed = GetStoredBool(SectionName, SessionSummaryDefs::IsCrashStoreKey);
-		NewRecord.bGPUCrashed = GetStoredBool(SectionName, SessionSummaryDefs::IsGPUCrashStoreKey);
-		NewRecord.bIsDebugger = GetStoredBool(SectionName, SessionSummaryDefs::IsDebuggerStoreKey);
-		NewRecord.bWasEverDebugger = GetStoredBool(SectionName, SessionSummaryDefs::WasDebuggerStoreKey);
-		NewRecord.bIsVanilla = GetStoredBool(SectionName, SessionSummaryDefs::IsVanillaStoreKey);
-		NewRecord.bIsTerminating = GetStoredBool(SectionName, SessionSummaryDefs::IsTerminatingKey);
-		NewRecord.bWasShutdown = GetStoredBool(SectionName, SessionSummaryDefs::WasShutdownStoreKey);
-		NewRecord.bIsInPIE = GetStoredBool(SectionName, SessionSummaryDefs::IsInPIEStoreKey);
-		NewRecord.bIsInVRMode = GetStoredBool(SectionName, SessionSummaryDefs::IsInVRModeStoreKey);
-		NewRecord.bIsInEnterprise = GetStoredBool(SectionName, SessionSummaryDefs::IsInEnterpriseStoreKey);
+		NewRecord.bCrashed = EditorSessionSummaryUtils::GetStoredBool(SectionName, SessionSummaryDefs::IsCrashStoreKey);
+		NewRecord.bGPUCrashed = EditorSessionSummaryUtils::GetStoredBool(SectionName, SessionSummaryDefs::IsGPUCrashStoreKey);
+		NewRecord.bIsDebugger = EditorSessionSummaryUtils::GetStoredBool(SectionName, SessionSummaryDefs::IsDebuggerStoreKey);
+		NewRecord.bWasEverDebugger = EditorSessionSummaryUtils::GetStoredBool(SectionName, SessionSummaryDefs::WasDebuggerStoreKey);
+		NewRecord.bIsVanilla = EditorSessionSummaryUtils::GetStoredBool(SectionName, SessionSummaryDefs::IsVanillaStoreKey);
+		NewRecord.bIsTerminating = EditorSessionSummaryUtils::GetStoredBool(SectionName, SessionSummaryDefs::IsTerminatingKey);
+		NewRecord.bWasShutdown = EditorSessionSummaryUtils::GetStoredBool(SectionName, SessionSummaryDefs::WasShutdownStoreKey);
+		NewRecord.bIsInPIE = EditorSessionSummaryUtils::GetStoredBool(SectionName, SessionSummaryDefs::IsInPIEStoreKey);
+		NewRecord.bIsInVRMode = EditorSessionSummaryUtils::GetStoredBool(SectionName, SessionSummaryDefs::IsInVRModeStoreKey);
+		NewRecord.bIsInEnterprise = EditorSessionSummaryUtils::GetStoredBool(SectionName, SessionSummaryDefs::IsInEnterpriseStoreKey);
 
 		PluginsString.ParseIntoArray(NewRecord.Plugins, TEXT(","));
 
@@ -697,7 +697,7 @@ TArray<FSessionRecord> FEditorSessionSummarySender::ReadStoredRecords() const
 	return MoveTemp(Records);
 }
 
-void FEditorSessionSummarySender::SendSessionSummaryEvent(const FSessionRecord& Record) const
+void FEditorSessionSummarySender::SendSessionSummaryEvent(const FEditorSessionRecord& Record) const
 {
 	FGuid SessionId;
 	FString SessionIdString = Record.SessionId;
@@ -707,15 +707,10 @@ void FEditorSessionSummarySender::SendSessionSummaryEvent(const FSessionRecord&
 		SessionIdString = SessionId.ToString(EGuidFormats::DigitsWithHyphensInBraces);
 	}
 
-#if !PLATFORM_PS4
 	FString ShutdownTypeString = Record.bCrashed ? SessionSummaryDefs::CrashSessionToken :
 		(Record.bWasEverDebugger ? SessionSummaryDefs::DebuggerSessionToken :
 		(Record.bIsTerminating ? SessionSummaryDefs::TerminatedSessionToken :
 		(Record.bWasShutdown ? SessionSummaryDefs::ShutdownSessionToken : SessionSummaryDefs::AbnormalSessionToken)));
-#else
-	// PS4 cannot set the crash flag so report abnormal shutdowns with a specific token meaning "crash or abnormal shutdown".
-	FString ShutdownTypeString = Record.bWasEverDebugger ? SessionSummaryDefs::DebuggerSessionToken : SessionSummaryDefs::PS4SessionToken;
-#endif
 
 	FString PluginsString = FString::Join(Record.Plugins, TEXT(","));
 
diff --git a/Engine/Source/Runtime/Engine/Private/Analytics/EditorSessionSummary.h b/Engine/Source/Runtime/Engine/Private/Analytics/EditorSessionSummary.h
--- a/Engine/Source/Runtime/Engine/Private/Analytics/EditorSessionSummary.h
+++ b/Engine/Source/Runtime/Engine/Private/Analytics/EditorSessionSummary.h
@@ -5,7 +5,7 @@
 #include "CoreMinimal.h"
 
 struct FUserActivity;
-struct FSessionRecord;
+struct FEditorSessionRecord;
 
 /** Writer for SessionSummary events to track all editor sessions. */
 class FEditorSessionSummaryWriter
@@ -20,8 +20,8 @@ public:
 private:
 	void InitializeRecords(bool bFirstAttempt);
 
-	FSessionRecord* CreateRecordForCurrentSession() const;
-	void WriteStoredRecord(const FSessionRecord& Record) const;
+	FEditorSessionRecord* CreateRecordForCurrentSession() const;
+	void WriteStoredRecord(const FEditorSessionRecord& Record) const;
 
 	void OnCrashing();
 	void OnTerminate();
@@ -33,7 +33,7 @@ private:
 
 private:
 	bool bInitializedRecords;
-	FSessionRecord* CurrentSession;
+	FEditorSessionRecord* CurrentSession;
 	FString CurrentSessionSectionName;
 	float HeartbeatTimeElapsed;
 	bool bShutdown;
@@ -53,10 +53,10 @@ public:
 private:
 	/** Send any stored records. Returns true if we successfully recorded any outstanding events. */
 	void SendStoredRecords(FTimespan Timeout) const;
-	TArray<FSessionRecord> ReadStoredRecords() const;
-	void DeleteStoredRecord(const FSessionRecord& Record) const;
-	void SendSessionSummaryEvent(const FSessionRecord& Record) const;
-	bool IsSessionProcessRunning(const FSessionRecord& Record) const;
+	TArray<FEditorSessionRecord> ReadStoredRecords() const;
+	void DeleteStoredRecord(const FEditorSessionRecord& Record) const;
+	void SendSessionSummaryEvent(const FEditorSessionRecord& Record) const;
+	bool IsSessionProcessRunning(const FEditorSessionRecord& Record) const;
 
 private:
 	float HeartbeatTimeElapsed;
diff --git a/Engine/Source/Runtime/Engine/Private/EngineSessionManager.cpp b/Engine/Source/Runtime/Engine/Private/EngineSessionManager.cpp
--- a/Engine/Source/Runtime/Engine/Private/EngineSessionManager.cpp
+++ b/Engine/Source/Runtime/Engine/Private/EngineSessionManager.cpp
@@ -74,7 +74,7 @@ namespace SessionManagerDefs
 	static const FString UnknownProjectValueString(TEXT("UnknownProject"));
 }
 
-namespace
+namespace EngineSessionManagerUtils
 {
 	FString TimestampToString(FDateTime InTimestamp)
 	{
@@ -220,7 +220,7 @@ void FEngineSessionManager::Tick(float DeltaTime)
 				}
 			}
 
-			const FString TimestampString = TimestampToString(FDateTime::UtcNow());
+			const FString TimestampString = EngineSessionManagerUtils::TimestampToString(FDateTime::UtcNow());
 			FPlatformMisc::SetStoredValue(SessionManagerDefs::StoreId, CurrentSessionSectionName, SessionManagerDefs::TimestampStoreKey, TimestampString);
 
 #if PLATFORM_SUPPORTS_WATCHDOG
@@ -260,7 +260,7 @@ void FEngineSessionManager::Shutdown()
 			{
 				const FString& ShutdownType = CurrentSession.bIsTerminating ? SessionManagerDefs::TerminatedSessionToken : SessionManagerDefs::ShutdownSessionToken;
 				FPlatformMisc::SetStoredValue(SessionManagerDefs::StoreId, WatchdogSectionName, SessionManagerDefs::StatusStoreKey, ShutdownType);
-				FPlatformMisc::SetStoredValue(SessionManagerDefs::StoreId, WatchdogSectionName, SessionManagerDefs::TimestampStoreKey, TimestampToString(FDateTime::UtcNow()));
+				FPlatformMisc::SetStoredValue(SessionManagerDefs::StoreId, WatchdogSectionName, SessionManagerDefs::TimestampStoreKey, EngineSessionManagerUtils::TimestampToString(FDateTime::UtcNow()));
 				WatchdogSectionName.Empty();
 			}
 #endif
@@ -345,7 +345,7 @@ bool FEngineSessionManager::BeginReadWriteRecords()
 			NewRecord.Mode = ModeString == SessionManagerDefs::EditorValueString ? EEngineSessionManagerMode::Editor : EEngineSessionManagerMode::Game;
 			NewRecord.ProjectName = ProjectName;
 			NewRecord.EngineVersion = EngineVersionString;
-			NewRecord.Timestamp = StringToTimestamp(TimestampString);
+			NewRecord.Timestamp = EngineSessionManagerUtils::StringToTimestamp(TimestampString);
 			NewRecord.bCrashed = IsCrashString == SessionManagerDefs::TrueValueString;
 			NewRecord.bGPUCrashed = IsGPUCrashString == SessionManagerDefs::TrueValueString;
 			NewRecord.bIsDebugger = IsDebuggerString == SessionManagerDefs::TrueValueString;
@@ -557,7 +557,7 @@ void FEngineSessionManager::CreateAndWriteRecordForSession()
 	FPlatformMisc::SetStoredValue(SessionManagerDefs::StoreId, CurrentSessionSectionName, SessionManagerDefs::ProjectNameStoreKey, CurrentSession.ProjectName);
 	FPlatformMisc::SetStoredValue(SessionManagerDefs::StoreId, CurrentSessionSectionName, SessionManagerDefs::CrashStoreKey, SessionManagerDefs::FalseValueString);
 	FPlatformMisc::SetStoredValue(SessionManagerDefs::StoreId, CurrentSessionSectionName, SessionManagerDefs::EngineVersionStoreKey, CurrentSession.EngineVersion);
-	FPlatformMisc::SetStoredValue(SessionManagerDefs::StoreId, CurrentSessionSectionName, SessionManagerDefs::TimestampStoreKey, TimestampToString(CurrentSession.Timestamp));
+	FPlatformMisc::SetStoredValue(SessionManagerDefs::StoreId, CurrentSessionSectionName, SessionManagerDefs::TimestampStoreKey, EngineSessionManagerUtils::TimestampToString(CurrentSession.Timestamp));
 	FPlatformMisc::SetStoredValue(SessionManagerDefs::StoreId, CurrentSessionSectionName, SessionManagerDefs::DebuggerStoreKey, IsDebuggerString);
 	FPlatformMisc::SetStoredValue(SessionManagerDefs::StoreId, CurrentSessionSectionName, SessionManagerDefs::WasDebuggerStoreKey, WasDebuggerString);
 	FPlatformMisc::SetStoredValue(SessionManagerDefs::StoreId, CurrentSessionSectionName, SessionManagerDefs::DeactivatedStoreKey, IsDeactivatedString);
@@ -593,7 +593,7 @@ void FEngineSessionManager::OnCrashing()
 		if (!WatchdogSectionName.IsEmpty())
 		{
 			FPlatformMisc::SetStoredValue(SessionManagerDefs::StoreId, WatchdogSectionName, SessionManagerDefs::StatusStoreKey, SessionManagerDefs::CrashSessionToken);
-			FPlatformMisc::SetStoredValue(SessionManagerDefs::StoreId, WatchdogSectionName, SessionManagerDefs::TimestampStoreKey, TimestampToString(FDateTime::UtcNow()));
+			FPlatformMisc::SetStoredValue(SessionManagerDefs::StoreId, WatchdogSectionName, SessionManagerDefs::TimestampStoreKey, EngineSessionManagerUtils::TimestampToString(FDateTime::UtcNow()));
 		}
 #endif
 	}
@@ -652,7 +652,7 @@ void FEngineSessionManager::OnTerminate()
 		else if (!WatchdogSectionName.IsEmpty())
 		{
 			FPlatformMisc::SetStoredValue(SessionManagerDefs::StoreId, WatchdogSectionName, SessionManagerDefs::StatusStoreKey, SessionManagerDefs::TerminatedSessionToken);
-			FPlatformMisc::SetStoredValue(SessionManagerDefs::StoreId, WatchdogSectionName, SessionManagerDefs::TimestampStoreKey, TimestampToString(FDateTime::UtcNow()));
+			FPlatformMisc::SetStoredValue(SessionManagerDefs::StoreId, WatchdogSectionName, SessionManagerDefs::TimestampStoreKey, EngineSessionManagerUtils::TimestampToString(FDateTime::UtcNow()));
 		}
 #endif
 	}
@@ -694,7 +694,7 @@ void FEngineSessionManager::OnUserActivity(const FUserActivity& UserActivity)
 		if (!WatchdogSectionName.IsEmpty())
 		{
 			FPlatformMisc::SetStoredValue(SessionManagerDefs::StoreId, WatchdogSectionName, SessionManagerDefs::UserActivityStoreKey, CurrentSession.CurrentUserActivity);
-			FPlatformMisc::SetStoredValue(SessionManagerDefs::StoreId, WatchdogSectionName, SessionManagerDefs::TimestampStoreKey, TimestampToString(FDateTime::UtcNow()));
+			FPlatformMisc::SetStoredValue(SessionManagerDefs::StoreId, WatchdogSectionName, SessionManagerDefs::TimestampStoreKey, EngineSessionManagerUtils::TimestampToString(FDateTime::UtcNow()));
 		}
 #endif
 	}
@@ -786,7 +786,7 @@ void FEngineSessionManager::StartWatchdog(const FString& RunType, const FString&
 
 		if (WatchdogProcessHandle.IsValid())
 		{
-			FString WatchdogStartTimeString = TimestampToString(FDateTime::UtcNow());
+			FString WatchdogStartTimeString = EngineSessionManagerUtils::TimestampToString(FDateTime::UtcNow());
 			FString WasDebuggerString = CurrentSession.bWasEverDebugger ? SessionManagerDefs::TrueValueString : SessionManagerDefs::FalseValueString;
 
 			WatchdogStartedAttributes.Add(FAnalyticsEventAttribute(TEXT("Outcome"), TEXT("Succeeded")));
