Added check for NullRHI before drawing canvas elements in Flush_GameThread. This guards against a crash when using NullRHI in conjunction with "Constrain Aspect Ratio"

diff --git a/Engine/Source/Runtime/Engine/Private/UserInterface/Canvas.cpp b/Engine/Source/Runtime/Engine/Private/UserInterface/Canvas.cpp
--- a/Engine/Source/Runtime/Engine/Private/UserInterface/Canvas.cpp
+++ b/Engine/Source/Runtime/Engine/Private/UserInterface/Canvas.cpp
@@ -809,7 +809,7 @@ void FCanvas::Flush_GameThread(bool bForce)
 	bool bEmitCanvasDrawEvents = GetEmitDrawEvents();
 
 	// Only create these render commands if we actually have something to draw.
-	if(SortedElements.Num() > 0)
+	if(SortedElements.Num() > 0 && !GUsingNullRHI)
 	{
 		FRenderThreadScope RenderThreadScope;
 		RenderThreadScope.EnqueueRenderCommand(
