Avoid Niagara graph nodes being gathered for game localization

diff --git a/Engine/Plugins/FX/Niagara/Source/NiagaraEditor/Private/NiagaraNode.cpp b/Engine/Plugins/FX/Niagara/Source/NiagaraEditor/Private/NiagaraNode.cpp
--- a/Engine/Plugins/FX/Niagara/Source/NiagaraEditor/Private/NiagaraNode.cpp
+++ b/Engine/Plugins/FX/Niagara/Source/NiagaraEditor/Private/NiagaraNode.cpp
@@ -7,6 +7,7 @@
 #include "GraphEditAction.h"
 #include "SNiagaraGraphNode.h"
 #include "Misc/SecureHash.h"
+#include "Serialization/PropertyLocalizationDataGathering.h"
 
 #define LOCTEXT_NAMESPACE "NiagaraNode"
 
@@ -611,5 +612,12 @@ void UNiagaraNode::UpdateCompileHashForNode(FSHA1& HashState) const
 	HashState.Update((const uint8*)&ChangeId, sizeof(FGuid));
 }
 
+#if WITH_EDITORONLY_DATA
+void UNiagaraNode::GatherForLocalization(FPropertyLocalizationDataGatherer& PropertyLocalizationDataGatherer, const EPropertyLocalizationGathererTextFlags GatherTextFlags) const
+{
+	// Niagara nodes only contain editor-only text data
+	Super::GatherForLocalization(PropertyLocalizationDataGatherer, GatherTextFlags | EPropertyLocalizationGathererTextFlags::ForceEditorOnly);
+}
+#endif
 
 #undef LOCTEXT_NAMESPACE
\ No newline at end of file
diff --git a/Engine/Plugins/FX/Niagara/Source/NiagaraEditor/Private/NiagaraNodeEmitter.cpp b/Engine/Plugins/FX/Niagara/Source/NiagaraEditor/Private/NiagaraNodeEmitter.cpp
--- a/Engine/Plugins/FX/Niagara/Source/NiagaraEditor/Private/NiagaraNodeEmitter.cpp
+++ b/Engine/Plugins/FX/Niagara/Source/NiagaraEditor/Private/NiagaraNodeEmitter.cpp
@@ -267,13 +267,13 @@ FText UNiagaraNodeEmitter::GetNameFromEmitter()
 		{
 			if (EmitterHandle.GetId() == EmitterHandleId)
 			{
-				return FText::FromName(EmitterHandle.GetName());
+				return FText::AsCultureInvariant(EmitterHandle.GetName().ToString());
 			}
 		}
 	}
 	else if (CachedUniqueName.IsValid())
 	{
-		return FText::FromName(CachedUniqueName);
+		return FText::AsCultureInvariant(CachedUniqueName.ToString());
 	}
 	return FText();
 }
diff --git a/Engine/Plugins/FX/Niagara/Source/NiagaraEditor/Public/NiagaraNode.h b/Engine/Plugins/FX/Niagara/Source/NiagaraEditor/Public/NiagaraNode.h
--- a/Engine/Plugins/FX/Niagara/Source/NiagaraEditor/Public/NiagaraNode.h
+++ b/Engine/Plugins/FX/Niagara/Source/NiagaraEditor/Public/NiagaraNode.h
@@ -138,6 +138,10 @@ protected:
 		the input pins should have been visited already.*/
 	virtual void RouteParameterMapAroundMe(FNiagaraParameterMapHistoryBuilder& OutHistory, bool bRecursive) const;
 	
+#if WITH_EDITORONLY_DATA
+	virtual void GatherForLocalization(FPropertyLocalizationDataGatherer& PropertyLocalizationDataGatherer, const EPropertyLocalizationGathererTextFlags GatherTextFlags) const override;
+#endif
+
 	/** The current change identifier for this node. Used to sync status with UNiagaraScripts.*/
 	UPROPERTY()
 	FGuid ChangeId;
diff --git a/Engine/Source/Runtime/Engine/Classes/EdGraph/EdGraphNode.h b/Engine/Source/Runtime/Engine/Classes/EdGraph/EdGraphNode.h
--- a/Engine/Source/Runtime/Engine/Classes/EdGraph/EdGraphNode.h
+++ b/Engine/Source/Runtime/Engine/Classes/EdGraph/EdGraphNode.h
@@ -23,6 +23,9 @@ struct FSlateIcon;
 struct FDiffResults;
 struct FDiffSingleResult;
 
+class FPropertyLocalizationDataGatherer;
+enum class EPropertyLocalizationGathererTextFlags : uint8;
+
 /**
   * Struct used to define information for terminal types, e.g. types that can be contained
   * by a container. Currently can represent strong/weak references to a type (only UObjects), 
@@ -922,6 +925,12 @@ public:
 	void ForEachNodeDirectlyConnectedToOutputs(TFunctionRef<void(UEdGraphNode*)> Func);
 	
 protected:
+#if WITH_EDITORONLY_DATA
+	/** Internal function used to gather pins from a graph node for localization */
+	friend void GatherGraphNodeForLocalization(const UObject* const Object, FPropertyLocalizationDataGatherer& PropertyLocalizationDataGatherer, const EPropertyLocalizationGathererTextFlags GatherTextFlags);
+	virtual void GatherForLocalization(FPropertyLocalizationDataGatherer& PropertyLocalizationDataGatherer, const EPropertyLocalizationGathererTextFlags GatherTextFlags) const;
+#endif
+
 	/**
 	 * Finds the difference in properties of node instance, for subobjects
 	 *
diff --git a/Engine/Source/Runtime/Engine/Private/EdGraph/EdGraphNode.cpp b/Engine/Source/Runtime/Engine/Private/EdGraph/EdGraphNode.cpp
--- a/Engine/Source/Runtime/Engine/Private/EdGraph/EdGraphNode.cpp
+++ b/Engine/Source/Runtime/Engine/Private/EdGraph/EdGraphNode.cpp
@@ -124,32 +124,33 @@ void UGraphNodeContextMenuContext::Init(const UEdGraph* InGraph, const UEdGraphN
 // UEdGraphNode
 
 #if WITH_EDITORONLY_DATA
-namespace
+void GatherGraphNodeForLocalization(const UObject* const Object, FPropertyLocalizationDataGatherer& PropertyLocalizationDataGatherer, const EPropertyLocalizationGathererTextFlags GatherTextFlags)
 {
-	void GatherGraphNodeForLocalization(const UObject* const Object, FPropertyLocalizationDataGatherer& PropertyLocalizationDataGatherer, const EPropertyLocalizationGathererTextFlags GatherTextFlags)
-	{
-		const UEdGraphNode* const GraphNode = CastChecked<UEdGraphNode>(Object);
+	const UEdGraphNode* const GraphNode = CastChecked<UEdGraphNode>(Object);
+	GraphNode->GatherForLocalization(PropertyLocalizationDataGatherer, GatherTextFlags);
+}
 
-		// We need to gather graph pins separately as they're no longer UObjects or UProperties
-		// plus they have custom logic for working out whether they're the default value
-		const FString PathToObject = GraphNode->GetPathName();
-		for (const UEdGraphPin* Pin : GraphNode->Pins)
+void UEdGraphNode::GatherForLocalization(FPropertyLocalizationDataGatherer& PropertyLocalizationDataGatherer, const EPropertyLocalizationGathererTextFlags GatherTextFlags) const
+{
+	// We need to gather graph pins separately as they're no longer UObjects or UProperties
+	// plus they have custom logic for working out whether they're the default value
+	const FString PathToObject = GetPathName();
+	for (const UEdGraphPin* Pin : Pins)
+	{
+		if (!Pin->DefaultTextValue.IsEmpty())
 		{
-			if (!Pin->DefaultTextValue.IsEmpty())
+			if (Pin->DoesDefaultValueMatchAutogenerated())
 			{
-				if (Pin->DoesDefaultValueMatchAutogenerated())
-				{
-					PropertyLocalizationDataGatherer.MarkDefaultTextInstance(Pin->DefaultTextValue);
-				}
-				else
-				{
-					PropertyLocalizationDataGatherer.GatherTextInstance(Pin->DefaultTextValue, FString::Printf(TEXT("%s.%s"), *PathToObject, *Pin->GetName()), /*bIsEditorOnly*/true);
-				}
+				PropertyLocalizationDataGatherer.MarkDefaultTextInstance(Pin->DefaultTextValue);
+			}
+			else
+			{
+				PropertyLocalizationDataGatherer.GatherTextInstance(Pin->DefaultTextValue, FString::Printf(TEXT("%s.%s"), *PathToObject, *Pin->GetName()), /*bIsEditorOnly*/true);
 			}
 		}
-
-		PropertyLocalizationDataGatherer.GatherLocalizationDataFromObject(GraphNode, GatherTextFlags);
 	}
+
+	PropertyLocalizationDataGatherer.GatherLocalizationDataFromObject(this, GatherTextFlags);
 }
 #endif
 
