Fix tooltip.

diff --git a/Engine/Source/Runtime/Engine/Classes/GameFramework/WorldSettings.h b/Engine/Source/Runtime/Engine/Classes/GameFramework/WorldSettings.h
--- a/Engine/Source/Runtime/Engine/Classes/GameFramework/WorldSettings.h
+++ b/Engine/Source/Runtime/Engine/Classes/GameFramework/WorldSettings.h
@@ -287,7 +287,7 @@ struct ENGINE_API FHierarchicalSimplification
 	UPROPERTY(Category = FHierarchicalSimplification, EditAnywhere)
 	uint8 bSimplifyMesh:1;
 
-	/** Min number of actors to build LODActor */
+	/** Only generate clusters for HLOD volumes */
 	UPROPERTY(EditAnywhere, Category = FHierarchicalSimplification, AdvancedDisplay, meta = (editcondition = "!bReusePreviousLevelClusters", DisplayAfter="MinNumberOfActorsToBuild"))
 	uint8 bOnlyGenerateClustersForVolumes:1;
 
