add missing headers
#rnx

diff --git a/Engine/Source/Editor/Persona/Private/LODInfoUILayout.h b/Engine/Source/Editor/Persona/Private/LODInfoUILayout.h
--- a/Engine/Source/Editor/Persona/Private/LODInfoUILayout.h
+++ b/Engine/Source/Editor/Persona/Private/LODInfoUILayout.h
@@ -10,6 +10,7 @@
 #include "UObject/ObjectMacros.h"
 #include "UObject/Object.h"
 #include "IPersonaToolkit.h"
+#include "Engine/SkeletalMesh.h"
 #include "LODInfoUILayout.generated.h"
 
 UCLASS(HideCategories=Object, MinimalAPI)
diff --git a/Engine/Source/Runtime/Engine/Classes/Engine/SkeletalMeshLODSettings.h b/Engine/Source/Runtime/Engine/Classes/Engine/SkeletalMeshLODSettings.h
--- a/Engine/Source/Runtime/Engine/Classes/Engine/SkeletalMeshLODSettings.h
+++ b/Engine/Source/Runtime/Engine/Classes/Engine/SkeletalMeshLODSettings.h
@@ -5,6 +5,7 @@
 #include "CoreMinimal.h"
 #include "UObject/ObjectMacros.h"
 #include "UObject/Object.h"
+#include "Engine/EngineTypes.h"
 #include "SkeletalMeshReductionSettings.h"
 #include "DataAsset.h"
 #include "PerPlatformProperties.h"
