Fixed crash due to accessing CurrentSession after failing to create it in FEditorSessionSummaryWriter::Initialize().

diff --git a/Engine/Source/Runtime/Engine/Private/Analytics/EditorSessionSummary.cpp b/Engine/Source/Runtime/Engine/Private/Analytics/EditorSessionSummary.cpp
--- a/Engine/Source/Runtime/Engine/Private/Analytics/EditorSessionSummary.cpp
+++ b/Engine/Source/Runtime/Engine/Private/Analytics/EditorSessionSummary.cpp
@@ -213,28 +213,26 @@ void FEditorSessionSummaryWriter::InitializeRecords(bool bFirstAttempt)
 		return;
 	}
 
-	{
-		// Scoped lock
-		FSystemWideCriticalSection StoredValuesLock(SessionSummaryDefs::GlobalLockName, bFirstAttempt ? SessionSummaryDefs::GlobalLockWaitTimeout : FTimespan::Zero());
+	// Scoped lock
+	FSystemWideCriticalSection StoredValuesLock(SessionSummaryDefs::GlobalLockName, bFirstAttempt ? SessionSummaryDefs::GlobalLockWaitTimeout : FTimespan::Zero());
 
-		if (StoredValuesLock.IsValid())
-		{
-			UE_LOG(LogEditorSessionSummary, Verbose, TEXT("Initializing EditorSessionSummaryWriter for editor session tracking"));
+	if (StoredValuesLock.IsValid())
+	{
+		UE_LOG(LogEditorSessionSummary, Verbose, TEXT("Initializing EditorSessionSummaryWriter for editor session tracking"));
 
-			// Create a session record for this session
-			CurrentSession = CreateRecordForCurrentSession();
-			CurrentSessionSectionName = EditorSessionSummaryUtils::GetSessionStorageLocation(*CurrentSession);
-			WriteStoredRecord(*CurrentSession);
+		// Create a session record for this session
+		CurrentSession = CreateRecordForCurrentSession();
+		CurrentSessionSectionName = EditorSessionSummaryUtils::GetSessionStorageLocation(*CurrentSession);
+		WriteStoredRecord(*CurrentSession);
 
-			bInitializedRecords = true;
+		bInitializedRecords = true;
 
-			UE_LOG(LogEditorSessionSummary, Log, TEXT("EditorSessionSummaryWriter initialized"));
-		}
+		UE_LOG(LogEditorSessionSummary, Log, TEXT("EditorSessionSummaryWriter initialized"));
 
 		// update session list string
 		FString SessionListString;
 		FPlatformMisc::GetStoredValue(SessionSummaryDefs::StoreId, SessionSummaryDefs::SessionSummarySection, SessionSummaryDefs::SessionListStoreKey, SessionListString);
-		
+
 		if (!SessionListString.IsEmpty())
 		{
 			SessionListString.Append(TEXT(","));
@@ -247,6 +245,11 @@ void FEditorSessionSummaryWriter::InitializeRecords(bool bFirstAttempt)
 
 void FEditorSessionSummaryWriter::UpdateTimestamps()
 {
+	if (!ensure(bInitializedRecords))
+	{
+		return;
+	}
+
 	CurrentSession->Timestamp = FDateTime::UtcNow();
 	FPlatformMisc::SetStoredValue(SessionSummaryDefs::StoreId, CurrentSessionSectionName, SessionSummaryDefs::TimestampStoreKey, EditorSessionSummaryUtils::TimestampToString(CurrentSession->Timestamp));
 
@@ -345,12 +348,6 @@ void FEditorSessionSummaryWriter::Shutdown()
 	FCoreDelegates::ApplicationWillTerminateDelegate.RemoveAll(this);
 	FCoreDelegates::IsVanillaProductChanged.RemoveAll(this);
 
-	// Skip Slate if terminating, since we can't guarantee which thread called us.
-	if (!CurrentSession->bIsTerminating) 
-	{
-		FSlateApplication::Get().GetOnModalLoopTickEvent().RemoveAll(this);
-	}
-
 	// Clear the session record for this session
 	if (bInitializedRecords)
 	{
