Fix old asset that is not reduced to make sure there reduction settings are inactive.
#rnx

diff --git a/Engine/Source/Runtime/Engine/Private/SkeletalMesh.cpp b/Engine/Source/Runtime/Engine/Private/SkeletalMesh.cpp
--- a/Engine/Source/Runtime/Engine/Private/SkeletalMesh.cpp
+++ b/Engine/Source/Runtime/Engine/Private/SkeletalMesh.cpp
@@ -2120,6 +2120,38 @@ void USkeletalMesh::PostLoad()
 				int32 LastMaterialIndex = INDEX_NONE;
 				int32 CurrentParentChunkIndex = INDEX_NONE;
 				int32 OriginalIndex = 0;
+				
+				//Reset the reduction setting to a non active state if the asset has active reduction but have no RawSkeletalMeshBulkData (we cannot reduce it)
+				if (IsReductionActive(LodIndex) && !ThisLODInfo->bHasBeenSimplified && ThisLODModel.RawSkeletalMeshBulkData.IsEmpty())
+				{
+					if (LodIndex > ThisLODInfo->ReductionSettings.BaseLOD)
+					{
+						ThisLODInfo->bHasBeenSimplified = true;
+					}
+					else if (LodIndex == ThisLODInfo->ReductionSettings.BaseLOD)
+					{
+						if (ThisLODInfo->ReductionSettings.TerminationCriterion == SkeletalMeshTerminationCriterion::SMTC_AbsNumOfTriangles
+							|| ThisLODInfo->ReductionSettings.TerminationCriterion == SkeletalMeshTerminationCriterion::SMTC_AbsNumOfVerts
+							|| ThisLODInfo->ReductionSettings.TerminationCriterion == SkeletalMeshTerminationCriterion::SMTC_AbsTriangleOrVert)
+						{
+							//MaxNum.... cannot be inactive, switch to NumOfTriangle
+							ThisLODInfo->ReductionSettings.TerminationCriterion = SMTC_NumOfTriangles;
+						}
+
+						//Now that we use trinagle or vert num, set an inactive value
+						if (ThisLODInfo->ReductionSettings.TerminationCriterion == SkeletalMeshTerminationCriterion::SMTC_NumOfTriangles
+							|| ThisLODInfo->ReductionSettings.TerminationCriterion == SkeletalMeshTerminationCriterion::SMTC_TriangleOrVert)
+						{
+							ThisLODInfo->ReductionSettings.NumOfTrianglesPercentage = 1.0f;
+						}
+						if (ThisLODInfo->ReductionSettings.TerminationCriterion == SkeletalMeshTerminationCriterion::SMTC_NumOfVerts
+							|| ThisLODInfo->ReductionSettings.TerminationCriterion == SkeletalMeshTerminationCriterion::SMTC_TriangleOrVert)
+						{
+							ThisLODInfo->ReductionSettings.NumOfVertPercentage = 1.0f;
+						}
+					}
+				}
+
 				for (int32 LODModelSectionIndex = 0; LODModelSectionIndex < LODModelSectionNum; ++LODModelSectionIndex)
 				{
 					//If we have cloth on a chunked section we treat the chunked section has a parent section (this is to get the same result has before the refactor)
