Prevent bumping the skeletalmesh ddc key guid
Only change the ddc when the asset are change.

diff --git a/Engine/Source/Runtime/Engine/Classes/Engine/SkeletalMesh.h b/Engine/Source/Runtime/Engine/Classes/Engine/SkeletalMesh.h
--- a/Engine/Source/Runtime/Engine/Classes/Engine/SkeletalMesh.h
+++ b/Engine/Source/Runtime/Engine/Classes/Engine/SkeletalMesh.h
@@ -1350,6 +1350,11 @@ public:
 	/** Releases all allocated Skin Weight Profile resources, assumes none are currently in use */
 	void ReleaseSkinWeightProfileResources();
 
+#if WITH_EDITORONLY_DATA
+	/*Transient data use when we postload an old asset to use legacy ddc key, it is turn off so if the user change the asset it go back to the latest ddc code*/
+	bool UseLegacyMeshDerivedDataKey = false;
+#endif
+
 protected:
 	/** Set of skin weight profiles associated with this mesh */
 	UPROPERTY(EditAnywhere, Category = SkinWeights, EditFixedSize, Meta=(NoResetToDefault))
diff --git a/Engine/Source/Runtime/Engine/Private/SkeletalMesh.cpp b/Engine/Source/Runtime/Engine/Private/SkeletalMesh.cpp
--- a/Engine/Source/Runtime/Engine/Private/SkeletalMesh.cpp
+++ b/Engine/Source/Runtime/Engine/Private/SkeletalMesh.cpp
@@ -2108,6 +2108,9 @@ void USkeletalMesh::PostLoad()
 
 		if (GetLinkerCustomVersion(FEditorObjectVersion::GUID) < FEditorObjectVersion::SkeletalMeshBuildRefactor)
 		{
+			//We want to avoid changing the ddc if we load an old asset.
+			//This bool should be put to false at the end of the postload, if there is another posteditchange call after a new ddc will be created
+			UseLegacyMeshDerivedDataKey = true;
 			//Fill up the Section ChunkedParentSectionIndex and OriginalDataSectionIndex
 			//We also want to create the UserSectionsData structure so the user can change the section data
 			for (int32 LodIndex = 0; LodIndex < LODInfo.Num(); LodIndex++)
@@ -2300,6 +2303,10 @@ void USkeletalMesh::PostLoad()
 #if !WITH_EDITOR
 	RebuildSocketMap();
 #endif // !WITH_EDITOR
+#if WITH_EDITORONLY_DATA
+	//Next postedit change will use the new ddc key scheme
+	UseLegacyMeshDerivedDataKey = false;
+#endif
 }
 
 #if WITH_EDITORONLY_DATA
diff --git a/Engine/Source/Runtime/Engine/Private/SkeletalMeshRenderData.cpp b/Engine/Source/Runtime/Engine/Private/SkeletalMeshRenderData.cpp
--- a/Engine/Source/Runtime/Engine/Private/SkeletalMeshRenderData.cpp
+++ b/Engine/Source/Runtime/Engine/Private/SkeletalMeshRenderData.cpp
@@ -59,7 +59,7 @@ static void SerializeLODInfoForDDC(USkeletalMesh* SkeletalMesh, FString& KeySuff
 // differences, etc.) replace the version GUID below with a new one.
 // In case of merge conflicts with DDC versions, you MUST generate a new GUID
 // and set this new GUID as the version.                                       
-#define SKELETALMESH_DERIVEDDATA_VER TEXT("E33979BDFD4B42C9A4B97BE162A8C0FF")
+#define SKELETALMESH_DERIVEDDATA_VER TEXT("51BDAF351B6941548818C2F8B33397D7")
 
 static const FString& GetSkeletalMeshDerivedDataVersion()
 {
@@ -74,30 +74,37 @@ static const FString& GetSkeletalMeshDerivedDataVersion()
 static FString BuildSkeletalMeshDerivedDataKey(USkeletalMesh* SkelMesh)
 {
 	FString KeySuffix(TEXT(""));
-	FString tmpDebugString;
 
-	//Synchronize the user data that are part of the key
-	SkelMesh->GetImportedModel()->SyncronizeLODUserSectionsData();
-
-	tmpDebugString = SkelMesh->GetImportedModel()->GetIdString();
-	KeySuffix += tmpDebugString;
-	tmpDebugString = SkelMesh->GetImportedModel()->GetLODModelIdString();
-	KeySuffix += tmpDebugString;
-
-	//Add the max gpu bone per section
-	const uint32 MaxGPUSkinBones = FGPUBaseSkinVertexFactory::GetMaxGPUSkinBones();
-	KeySuffix += FString::FromInt((int32)MaxGPUSkinBones);
-
-	tmpDebugString = TEXT("");
-	SerializeLODInfoForDDC(SkelMesh, tmpDebugString);
-	KeySuffix += tmpDebugString;
-
-	KeySuffix += SkelMesh->bHasVertexColors ? "1" : "0";
-	if (SkelMesh->bHasVertexColors)
+	if (SkelMesh->UseLegacyMeshDerivedDataKey )
+	{
+		//Old asset will have the same LOD settings for bUseFullPrecisionUVs. We can use the LOD 0
+		const FSkeletalMeshLODInfo* BaseLODInfo = SkelMesh->GetLODInfo(0);
+		bool bUseFullPrecisionUVs = BaseLODInfo ? BaseLODInfo->BuildSettings.bUseFullPrecisionUVs : false;
+		KeySuffix += SkelMesh->GetImportedModel()->GetIdString();
+		KeySuffix += (bUseFullPrecisionUVs || !GVertexElementTypeSupport.IsSupported(VET_Half2)) ? "1" : "0";
+	}
+	else
 	{
-		KeySuffix += SkelMesh->VertexColorGuid.ToString(EGuidFormats::Digits);
+		FString tmpDebugString;
+		//Synchronize the user data that are part of the key
+		SkelMesh->GetImportedModel()->SyncronizeLODUserSectionsData();
+		tmpDebugString = SkelMesh->GetImportedModel()->GetIdString();
+		KeySuffix += tmpDebugString;
+		tmpDebugString = SkelMesh->GetImportedModel()->GetLODModelIdString();
+		KeySuffix += tmpDebugString;
+		
+		//Add the max gpu bone per section
+		const uint32 MaxGPUSkinBones = FGPUBaseSkinVertexFactory::GetMaxGPUSkinBones();
+		KeySuffix += FString::FromInt((int32)MaxGPUSkinBones);
+
+		tmpDebugString = TEXT("");
+		SerializeLODInfoForDDC(SkelMesh, tmpDebugString);
+		KeySuffix += tmpDebugString;
 	}
 
+	KeySuffix += SkelMesh->bHasVertexColors ? "1" : "0";
+	KeySuffix += SkelMesh->VertexColorGuid.ToString(EGuidFormats::Digits);
+
 	const ITargetPlatform* TargetPlatform = GetTargetPlatformManagerRef().GetRunningTargetPlatform();
 	check(TargetPlatform);
 	const FName PlatformGroupName = TargetPlatform->GetPlatformInfo().PlatformGroupName;
@@ -145,24 +152,31 @@ void FSkeletalMeshRenderData::Cache(USkeletalMesh* Owner)
 		{
 			COOK_STAT(Timer.AddHit(DerivedData.Num()));
 			
-			FSkeletalMeshModel* SkelMeshModel = Owner->GetImportedModel();
-			check(SkelMeshModel);
 
 			FMemoryReader Ar(DerivedData, /*bIsPersistent=*/ true);
-			
-			//Serialize the LODModel sections since they are dependent on the reduction
-			for (int32 LODIndex = 0; LODIndex < SkelMeshModel->LODModels.Num(); LODIndex++)
+
+			//With skeletal mesh build refactor we serialize the LODModel sections into the DDC
+			//We need to store those so we do not have to rerun the reduction to make them up to date
+			//with the serialize renderdata. This allow to use DDC when changing the reduction settings.
+			//The old workflow has to reduce the LODModel before getting the render data DDC.
+			if (!Owner->UseLegacyMeshDerivedDataKey)
 			{
-				FSkeletalMeshLODModel* LODModel = &(SkelMeshModel->LODModels[LODIndex]);
-				int32 SectionNum = 0;
-				Ar << SectionNum;
-				LODModel->Sections.Reset(SectionNum);
-				LODModel->Sections.AddDefaulted(SectionNum);
-				for (int32 SectionIndex = 0; SectionIndex < LODModel->Sections.Num(); ++SectionIndex)
+				FSkeletalMeshModel* SkelMeshModel = Owner->GetImportedModel();
+				check(SkelMeshModel);
+				//Serialize the LODModel sections since they are dependent on the reduction
+				for (int32 LODIndex = 0; LODIndex < SkelMeshModel->LODModels.Num(); LODIndex++)
 				{
-					Ar << LODModel->Sections[SectionIndex];
+					FSkeletalMeshLODModel* LODModel = &(SkelMeshModel->LODModels[LODIndex]);
+					int32 SectionNum = 0;
+					Ar << SectionNum;
+					LODModel->Sections.Reset(SectionNum);
+					LODModel->Sections.AddDefaulted(SectionNum);
+					for (int32 SectionIndex = 0; SectionIndex < LODModel->Sections.Num(); ++SectionIndex)
+					{
+						Ar << LODModel->Sections[SectionIndex];
+					}
+					LODModel->SyncronizeUserSectionsDataArray();
 				}
-				LODModel->SyncronizeUserSectionsDataArray();
 			}
 
 			Serialize(Ar, Owner);
@@ -247,15 +261,20 @@ void FSkeletalMeshRenderData::Cache(USkeletalMesh* Owner)
 
 			FMemoryWriter Ar(DerivedData, /*bIsPersistent=*/ true);
 			
-			//Serialize the LODModel sections since they are dependent on the reduction
-			for (int32 LODIndex = 0; LODIndex < SkelMeshModel->LODModels.Num(); LODIndex++)
+			//If we load an old asset we want to be sure the serialize ddc will be the same has before the skeletalmesh build refactor
+			//So we do not serialize the LODModel sections.
+			if (!Owner->UseLegacyMeshDerivedDataKey)
 			{
-				FSkeletalMeshLODModel* LODModel = &(SkelMeshModel->LODModels[LODIndex]);
-				int32 SectionNum = LODModel->Sections.Num();
-				Ar << SectionNum;
-				for (int32 SectionIndex = 0; SectionIndex < LODModel->Sections.Num(); ++SectionIndex)
+				//Serialize the LODModel sections since they are dependent on the reduction
+				for (int32 LODIndex = 0; LODIndex < SkelMeshModel->LODModels.Num(); LODIndex++)
 				{
-					Ar << LODModel->Sections[SectionIndex];
+					FSkeletalMeshLODModel* LODModel = &(SkelMeshModel->LODModels[LODIndex]);
+					int32 SectionNum = LODModel->Sections.Num();
+					Ar << SectionNum;
+					for (int32 SectionIndex = 0; SectionIndex < LODModel->Sections.Num(); ++SectionIndex)
+					{
+						Ar << LODModel->Sections[SectionIndex];
+					}
 				}
 			}
 
