Level packages should always have PIE prefix added during PIE

diff --git a/Engine/Source/Runtime/Engine/Private/World.cpp b/Engine/Source/Runtime/Engine/Private/World.cpp
--- a/Engine/Source/Runtime/Engine/Private/World.cpp
+++ b/Engine/Source/Runtime/Engine/Private/World.cpp
@@ -5892,11 +5892,7 @@ void FSeamlessTravelHandler::StartLoadingDestination()
 				PackageFlags |= PKG_PlayInEditor;
 			}
 			PIEInstanceID = WorldContext.PIEInstance;
-			UPackage* EditorLevelPackage = (UPackage*)StaticFindObjectFast(UPackage::StaticClass(), NULL, URLMapFName, 0, 0, RF_NoFlags, EInternalObjectFlags::PendingKill);
-			if (EditorLevelPackage)
-			{
-				URLMapPackageName = UWorld::ConvertToPIEPackageName(URLMapPackageName, PIEInstanceID);
-			}
+			URLMapPackageName = UWorld::ConvertToPIEPackageName(URLMapPackageName, PIEInstanceID);
 		}
 #endif
 
