Editor session summary now contains a field for whether or not the user was notified of having low disk space, which is suspected of contributing to some abnormal shutdowns.

diff --git a/Engine/Source/Editor/UnrealEd/Private/UnrealEdEngine.cpp b/Engine/Source/Editor/UnrealEd/Private/UnrealEdEngine.cpp
--- a/Engine/Source/Editor/UnrealEd/Private/UnrealEdEngine.cpp
+++ b/Engine/Source/Editor/UnrealEd/Private/UnrealEdEngine.cpp
@@ -63,9 +63,10 @@
 #include "AutoReimport/AssetSourceFilenameCache.h"
 #include "UObject/UObjectThreadContext.h"
 #include "EngineUtils.h"
-
-
+#include "EngineAnalytics.h"
 #include "CookerSettings.h"
+#include "Misc/MessageDialog.h"
+
 DEFINE_LOG_CATEGORY_STATIC(LogUnrealEdEngine, Log, All);
 
 
@@ -73,6 +74,22 @@ void UUnrealEdEngine::Init(IEngineLoop* InEngineLoop)
 {
 	Super::Init(InEngineLoop);
 
+	uint64 TotalDiskSpace = 0;
+	uint64 FreeDiskSpace = 0;
+	const TCHAR* BaseDir = FPlatformProcess::BaseDir();
+	if (FPlatformMisc::GetDiskTotalAndFreeSpace(BaseDir, TotalDiskSpace, FreeDiskSpace))
+	{
+		const uint64 HardDriveFreeMB = FreeDiskSpace / (1024 * 1024);
+		if (HardDriveFreeMB < 512)
+		{
+			FEngineAnalytics::LowDriveSpaceDetected();
+			
+			const FText Title = NSLOCTEXT("DriveSpaceDialog", "LowHardDriveSpaceMsgTitle", "Low drive space warning");
+			const FText Message = NSLOCTEXT("DriveSpaceDialog", "LowHardDriveSpaceMsg", "The available space on the drive containing this project is low.\nIt is recommended that you free some space to avoid issues.");
+			FMessageDialog::Open(EAppMsgType::Ok, Message, &Title);
+		}
+	}
+
 	// Build databases used by source code navigation
 	FSourceCodeNavigation::Initialize();
 
diff --git a/Engine/Source/Runtime/Engine/Private/Analytics/EditorSessionSummary.cpp b/Engine/Source/Runtime/Engine/Private/Analytics/EditorSessionSummary.cpp
--- a/Engine/Source/Runtime/Engine/Private/Analytics/EditorSessionSummary.cpp
+++ b/Engine/Source/Runtime/Engine/Private/Analytics/EditorSessionSummary.cpp
@@ -79,16 +79,17 @@ namespace SessionSummaryDefs
 	static const FString Idle30MinStoreKey(TEXT("Idle30Min"));
 	
 	// boolean flags
-	static const FString IsTerminatingKey(TEXT("Terminating"));
-	static const FString WasDebuggerStoreKey(TEXT("WasEverDebugger"));
+	static const FString IsCrashStoreKey(TEXT("IsCrash"));
+	static const FString IsGPUCrashStoreKey(TEXT("IsGPUCrash"));
 	static const FString IsDebuggerStoreKey(TEXT("IsDebugger"));
+	static const FString WasDebuggerStoreKey(TEXT("WasEverDebugger"));
 	static const FString IsVanillaStoreKey(TEXT("IsVanilla"));
-	static const FString IsCrashStoreKey(TEXT("IsCrash"));
+	static const FString IsTerminatingKey(TEXT("Terminating"));
 	static const FString WasShutdownStoreKey(TEXT("WasShutdown"));
-	static const FString IsGPUCrashStoreKey(TEXT("IsGPUCrash"));
-	static const FString IsInVRModeStoreKey(TEXT("IsInVRMode"));
-	static const FString IsInEnterpriseStoreKey(TEXT("IsInEnterprise"));
 	static const FString IsInPIEStoreKey(TEXT("IsInPIE"));
+	static const FString IsInEnterpriseStoreKey(TEXT("IsInEnterprise"));
+	static const FString IsInVRModeStoreKey(TEXT("IsInVRMode"));
+	static const FString IsLowDriveSpaceDetectedStoreKey(TEXT("LowDriveSpaceDetected"));
 }
 
 struct FEditorSessionRecord
@@ -116,6 +117,7 @@ struct FEditorSessionRecord
 	bool bIsInPIE : 1;
 	bool bIsInEnterprise : 1;
 	bool bIsInVRMode : 1;
+	bool bLowDriveSpaceDetected : 1;
 
 	FEditorSessionRecord()
 	{
@@ -136,6 +138,7 @@ struct FEditorSessionRecord
 		bIsInPIE = false;
 		bIsInEnterprise = false;
 		bIsInVRMode = false;
+		bLowDriveSpaceDetected = false;
 	}
 };
 
@@ -322,22 +325,33 @@ void FEditorSessionSummaryWriter::Tick(float DeltaTime)
 			const FString AverageFPSString = FString::SanitizeFloat(GAverageFPS);
 			FPlatformMisc::SetStoredValue(SessionSummaryDefs::StoreId, CurrentSessionSectionName, SessionSummaryDefs::AverageFPSStoreKey, AverageFPSString);
 
-			CurrentSession->bIsInVRMode = IVREditorModule::Get().IsVREditorModeActive();
-			const FString IsInVRModeString = EditorSessionSummaryUtils::BoolToStoredString(CurrentSession->bIsInVRMode);
-			FPlatformMisc::SetStoredValue(SessionSummaryDefs::StoreId, CurrentSessionSectionName, SessionSummaryDefs::IsInVRModeStoreKey, IsInVRModeString);
+			CurrentSession->bIsInPIE = FPlayWorldCommandCallbacks::IsInPIE();
+			const FString IsInPIEString = EditorSessionSummaryUtils::BoolToStoredString(CurrentSession->bIsInPIE);
+			FPlatformMisc::SetStoredValue(SessionSummaryDefs::StoreId, CurrentSessionSectionName, SessionSummaryDefs::IsInPIEStoreKey, IsInPIEString);
 
 			CurrentSession->bIsInEnterprise = IProjectManager::Get().IsEnterpriseProject();
 			const FString IsInEnterpriseString = EditorSessionSummaryUtils::BoolToStoredString(CurrentSession->bIsInEnterprise);
 			FPlatformMisc::SetStoredValue(SessionSummaryDefs::StoreId, CurrentSessionSectionName, SessionSummaryDefs::IsInEnterpriseStoreKey, IsInEnterpriseString);
 
-			CurrentSession->bIsInPIE = FPlayWorldCommandCallbacks::IsInPIE();
-			const FString IsInPIEString = EditorSessionSummaryUtils::BoolToStoredString(CurrentSession->bIsInPIE);
-			FPlatformMisc::SetStoredValue(SessionSummaryDefs::StoreId, CurrentSessionSectionName, SessionSummaryDefs::IsInPIEStoreKey, IsInPIEString);
+			CurrentSession->bIsInVRMode = IVREditorModule::Get().IsVREditorModeActive();
+			const FString IsInVRModeString = EditorSessionSummaryUtils::BoolToStoredString(CurrentSession->bIsInVRMode);
+			FPlatformMisc::SetStoredValue(SessionSummaryDefs::StoreId, CurrentSessionSectionName, SessionSummaryDefs::IsInVRModeStoreKey, IsInVRModeString);
+
 #endif
 		}
 	}
 }
 
+void FEditorSessionSummaryWriter::LowDriveSpaceDetected()
+{
+	if (bInitializedRecords && !CurrentSession->bLowDriveSpaceDetected)
+	{
+		CurrentSession->bLowDriveSpaceDetected = true;
+
+		FPlatformMisc::SetStoredValue(SessionSummaryDefs::StoreId, CurrentSessionSectionName, SessionSummaryDefs::IsLowDriveSpaceDetectedStoreKey, SessionSummaryDefs::TrueValueString);
+	}
+}
+
 void FEditorSessionSummaryWriter::Shutdown()
 {
 	FCoreDelegates::OnHandleSystemError.RemoveAll(this);
@@ -388,6 +402,7 @@ void FEditorSessionSummaryWriter::WriteStoredRecord(const FEditorSessionRecord&
 	FPlatformMisc::SetStoredValue(SessionSummaryDefs::StoreId, StorageLocation, SessionSummaryDefs::IsInPIEStoreKey, EditorSessionSummaryUtils::BoolToStoredString(Record.bIsInPIE));
 	FPlatformMisc::SetStoredValue(SessionSummaryDefs::StoreId, StorageLocation, SessionSummaryDefs::IsInEnterpriseStoreKey, EditorSessionSummaryUtils::BoolToStoredString(Record.bIsInEnterprise));
 	FPlatformMisc::SetStoredValue(SessionSummaryDefs::StoreId, StorageLocation, SessionSummaryDefs::IsInVRModeStoreKey, EditorSessionSummaryUtils::BoolToStoredString(Record.bIsInVRMode));
+	FPlatformMisc::SetStoredValue(SessionSummaryDefs::StoreId, StorageLocation, SessionSummaryDefs::IsLowDriveSpaceDetectedStoreKey, EditorSessionSummaryUtils::BoolToStoredString(Record.bLowDriveSpaceDetected));
 }
 
 FEditorSessionRecord* FEditorSessionSummaryWriter::CreateRecordForCurrentSession() const
@@ -613,6 +628,7 @@ void FEditorSessionSummarySender::DeleteStoredRecord(const FEditorSessionRecord&
 	FPlatformMisc::DeleteStoredValue(SessionSummaryDefs::StoreId, SectionName, SessionSummaryDefs::IsInPIEStoreKey);
 	FPlatformMisc::DeleteStoredValue(SessionSummaryDefs::StoreId, SectionName, SessionSummaryDefs::IsInEnterpriseStoreKey);
 	FPlatformMisc::DeleteStoredValue(SessionSummaryDefs::StoreId, SectionName, SessionSummaryDefs::IsInVRModeStoreKey);
+	FPlatformMisc::DeleteStoredValue(SessionSummaryDefs::StoreId, SectionName, SessionSummaryDefs::IsLowDriveSpaceDetectedStoreKey);
 }
 
 bool FEditorSessionSummarySender::IsSessionProcessRunning(const FEditorSessionRecord& Record) const
@@ -692,8 +708,9 @@ TArray<FEditorSessionRecord> FEditorSessionSummarySender::ReadStoredRecords() co
 		NewRecord.bIsTerminating = EditorSessionSummaryUtils::GetStoredBool(SectionName, SessionSummaryDefs::IsTerminatingKey);
 		NewRecord.bWasShutdown = EditorSessionSummaryUtils::GetStoredBool(SectionName, SessionSummaryDefs::WasShutdownStoreKey);
 		NewRecord.bIsInPIE = EditorSessionSummaryUtils::GetStoredBool(SectionName, SessionSummaryDefs::IsInPIEStoreKey);
-		NewRecord.bIsInVRMode = EditorSessionSummaryUtils::GetStoredBool(SectionName, SessionSummaryDefs::IsInVRModeStoreKey);
 		NewRecord.bIsInEnterprise = EditorSessionSummaryUtils::GetStoredBool(SectionName, SessionSummaryDefs::IsInEnterpriseStoreKey);
+		NewRecord.bIsInVRMode = EditorSessionSummaryUtils::GetStoredBool(SectionName, SessionSummaryDefs::IsInVRModeStoreKey);
+		NewRecord.bLowDriveSpaceDetected = EditorSessionSummaryUtils::GetStoredBool(SectionName, SessionSummaryDefs::IsLowDriveSpaceDetectedStoreKey);
 
 		PluginsString.ParseIntoArray(NewRecord.Plugins, TEXT(","));
 
@@ -737,6 +754,7 @@ void FEditorSessionSummarySender::SendSessionSummaryEvent(const FEditorSessionRe
 	AnalyticsAttributes.Emplace(SessionSummaryDefs::IsInPIEStoreKey, Record.bIsInPIE);
 	AnalyticsAttributes.Emplace(SessionSummaryDefs::IsInEnterpriseStoreKey, Record.bIsInEnterprise);
 	AnalyticsAttributes.Emplace(SessionSummaryDefs::IsInVRModeStoreKey, Record.bIsInVRMode);
+	AnalyticsAttributes.Emplace(SessionSummaryDefs::IsLowDriveSpaceDetectedStoreKey, Record.bLowDriveSpaceDetected);
 
 	double SessionDuration = (Record.Timestamp - Record.StartupTimestamp).GetTotalSeconds();
 	AnalyticsAttributes.Emplace(SessionSummaryDefs::SessionDurationStoreKey, SessionDuration);
diff --git a/Engine/Source/Runtime/Engine/Private/Analytics/EditorSessionSummary.h b/Engine/Source/Runtime/Engine/Private/Analytics/EditorSessionSummary.h
--- a/Engine/Source/Runtime/Engine/Private/Analytics/EditorSessionSummary.h
+++ b/Engine/Source/Runtime/Engine/Private/Analytics/EditorSessionSummary.h
@@ -15,6 +15,7 @@ public:
 
 	void Initialize();
 	void Tick(float DeltaTime);
+	void LowDriveSpaceDetected();
 	void Shutdown();
 
 private:
diff --git a/Engine/Source/Runtime/Engine/Private/EngineAnalytics.cpp b/Engine/Source/Runtime/Engine/Private/EngineAnalytics.cpp
--- a/Engine/Source/Runtime/Engine/Private/EngineAnalytics.cpp
+++ b/Engine/Source/Runtime/Engine/Private/EngineAnalytics.cpp
@@ -209,3 +209,11 @@ void FEngineAnalytics::Tick(float DeltaTime)
 		SessionSummarySender->Tick(DeltaTime);
 	}
 }
+
+void FEngineAnalytics::LowDriveSpaceDetected()
+{
+	if (SessionSummaryWriter.IsValid())
+	{
+		SessionSummaryWriter->LowDriveSpaceDetected();
+	}
+}
diff --git a/Engine/Source/Runtime/Engine/Public/EngineAnalytics.h b/Engine/Source/Runtime/Engine/Public/EngineAnalytics.h
--- a/Engine/Source/Runtime/Engine/Public/EngineAnalytics.h
+++ b/Engine/Source/Runtime/Engine/Public/EngineAnalytics.h
@@ -35,6 +35,8 @@ public:
 
 	static ENGINE_API void Tick(float DeltaTime);
 
+	static ENGINE_API void LowDriveSpaceDetected();
+
 private:
 	static bool bIsInitialized;
 	static ENGINE_API TSharedPtr<IAnalyticsProviderET> Analytics;
