CIS compile fix for non-editor builds.

diff --git a/Engine/Source/Runtime/Engine/Private/EngineAnalytics.cpp b/Engine/Source/Runtime/Engine/Private/EngineAnalytics.cpp
--- a/Engine/Source/Runtime/Engine/Private/EngineAnalytics.cpp
+++ b/Engine/Source/Runtime/Engine/Private/EngineAnalytics.cpp
@@ -236,8 +236,10 @@ void FEngineAnalytics::Tick(float DeltaTime)
 
 void FEngineAnalytics::LowDriveSpaceDetected()
 {
+#if WITH_EDITOR
 	if (SessionSummaryWriter.IsValid())
 	{
 		SessionSummaryWriter->LowDriveSpaceDetected();
 	}
+#endif
 }
