When computing instances dithering LODs, offset instances by the mesh pivot to match CPU LOD logic more closely.
The proposed changes from the PR wasn't kept, since it was modifying the node's bounding boxes.
Note: didn't test the USE_INSTANCING_EMULATED path as I didn't find any RHI supporting it, so old behavior is preserved (only a problem is dithering is enabled in this path).

diff --git a/Engine/Shaders/Private/LocalVertexFactory.ush b/Engine/Shaders/Private/LocalVertexFactory.ush
--- a/Engine/Shaders/Private/LocalVertexFactory.ush
+++ b/Engine/Shaders/Private/LocalVertexFactory.ush
@@ -36,6 +36,7 @@
 	float4 InstancingWorldViewOriginOne;
 #endif
 
+float4 InstancingOffset;
 float4 InstancingFadeOutParams;
 
 #ifndef USE_INSTANCING_EMULATED
@@ -441,7 +442,7 @@ float GetInstanceRandom(FVertexFactoryIntermediates Intermediates)
 float3 GetInstanceOrigin(FVertexFactoryIntermediates Intermediates)
 {
 	#if !USE_INSTANCING_EMULATED
-		float3 Origin = Intermediates.InstanceOrigin.xyz;
+		float3 Origin = Intermediates.InstanceOrigin.xyz + InstancingOffset;
 	#else
 		float3 Origin = CPUInstanceOrigin.xyz;
 	#endif	// !USE_INSTANCING_EMULATED
diff --git a/Engine/Source/Runtime/Engine/Private/InstancedStaticMesh.cpp b/Engine/Source/Runtime/Engine/Private/InstancedStaticMesh.cpp
--- a/Engine/Source/Runtime/Engine/Private/InstancedStaticMesh.cpp
+++ b/Engine/Source/Runtime/Engine/Private/InstancedStaticMesh.cpp
@@ -886,6 +886,7 @@ void FInstancedStaticMeshSceneProxy::SetupProxy(UInstancedStaticMeshComponent* I
 	UserData_AllInstances.MeshRenderData = InComponent->GetStaticMesh()->RenderData.Get();
 	UserData_AllInstances.StartCullDistance = InComponent->InstanceStartCullDistance;
 	UserData_AllInstances.EndCullDistance = InComponent->InstanceEndCullDistance;
+	UserData_AllInstances.InstancingOffset = InComponent->GetStaticMesh()->GetBoundingBox().GetCenter();
 	UserData_AllInstances.MinLOD = ClampedMinLOD;
 	UserData_AllInstances.bRenderSelected = true;
 	UserData_AllInstances.bRenderUnselected = true;
@@ -2986,6 +2987,7 @@ void FInstancedStaticMeshVertexFactoryShaderParameters::GetElementShaderBindings
 		FVector4 InstancingViewZCompareZero(MIN_flt, MIN_flt, MAX_flt, 1.0f);
 		FVector4 InstancingViewZCompareOne(MIN_flt, MIN_flt, MAX_flt, 0.0f);
 		FVector4 InstancingViewZConstant(ForceInit);
+		FVector4 InstancingOffset(InstancingUserData->InstancingOffset);
 		FVector4 InstancingWorldViewOriginZero(ForceInit);
 		FVector4 InstancingWorldViewOriginOne(ForceInit);
 		InstancingWorldViewOriginOne.W = 1.0f;
@@ -3081,6 +3083,7 @@ void FInstancedStaticMeshVertexFactoryShaderParameters::GetElementShaderBindings
 		ShaderBindings.Add(InstancingViewZCompareZeroParameter, InstancingViewZCompareZero);
 		ShaderBindings.Add(InstancingViewZCompareOneParameter, InstancingViewZCompareOne);
 		ShaderBindings.Add(InstancingViewZConstantParameter, InstancingViewZConstant);
+		ShaderBindings.Add(InstancingOffsetParameter, InstancingOffset);
 		ShaderBindings.Add(InstancingWorldViewOriginZeroParameter, InstancingWorldViewOriginZero);
 		ShaderBindings.Add(InstancingWorldViewOriginOneParameter, InstancingWorldViewOriginOne);
 	}
diff --git a/Engine/Source/Runtime/Engine/Private/InstancedStaticMesh.h b/Engine/Source/Runtime/Engine/Private/InstancedStaticMesh.h
--- a/Engine/Source/Runtime/Engine/Private/InstancedStaticMesh.h
+++ b/Engine/Source/Runtime/Engine/Private/InstancedStaticMesh.h
@@ -153,6 +153,7 @@ struct FInstancingUserData
 	bool bRenderSelected;
 	bool bRenderUnselected;
 	FVector AverageInstancesScale;
+	FVector InstancingOffset;
 };
 
 struct FInstancedStaticMeshDataType
@@ -333,6 +334,7 @@ class FInstancedStaticMeshVertexFactoryShaderParameters : public FLocalVertexFac
 		InstancingViewZCompareZeroParameter.Bind(ParameterMap, TEXT("InstancingViewZCompareZero"));
 		InstancingViewZCompareOneParameter.Bind(ParameterMap, TEXT("InstancingViewZCompareOne"));
 		InstancingViewZConstantParameter.Bind(ParameterMap, TEXT("InstancingViewZConstant"));
+		InstancingOffsetParameter.Bind(ParameterMap, TEXT("InstancingOffset"));
 		InstancingWorldViewOriginZeroParameter.Bind(ParameterMap, TEXT("InstancingWorldViewOriginZero"));
 		InstancingWorldViewOriginOneParameter.Bind(ParameterMap, TEXT("InstancingWorldViewOriginOne"));
 		CPUInstanceOrigin.Bind(ParameterMap, TEXT("CPUInstanceOrigin"));
@@ -363,6 +365,7 @@ class FInstancedStaticMeshVertexFactoryShaderParameters : public FLocalVertexFac
 		Ar << InstancingViewZCompareZeroParameter;
 		Ar << InstancingViewZCompareOneParameter;
 		Ar << InstancingViewZConstantParameter;
+		Ar << InstancingOffsetParameter;
 		Ar << InstancingWorldViewOriginZeroParameter;
 		Ar << InstancingWorldViewOriginOneParameter;
 		Ar << CPUInstanceOrigin;
@@ -381,6 +384,7 @@ private:
 	FShaderParameter InstancingViewZCompareZeroParameter;
 	FShaderParameter InstancingViewZCompareOneParameter;
 	FShaderParameter InstancingViewZConstantParameter;
+	FShaderParameter InstancingOffsetParameter;
 	FShaderParameter InstancingWorldViewOriginZeroParameter;
 	FShaderParameter InstancingWorldViewOriginOneParameter;
 
