Switch HighResScreenshot args to use FParse::Token

diff --git a/Engine/Source/Runtime/Engine/Private/UnrealClient.cpp b/Engine/Source/Runtime/Engine/Private/UnrealClient.cpp
--- a/Engine/Source/Runtime/Engine/Private/UnrealClient.cpp
+++ b/Engine/Source/Runtime/Engine/Private/UnrealClient.cpp
@@ -2149,26 +2149,20 @@ extern bool ParseResolution( const TCHAR* InResolution, uint32& OutX, uint32& Ou
 ENGINE_API bool GetHighResScreenShotInput(const TCHAR* Cmd, FOutputDevice& Ar, uint32& OutXRes, uint32& OutYRes, float& OutResMult, FIntRect& OutCaptureRegion, bool& OutShouldEnableMask, bool& OutDumpBufferVisualizationTargets, bool& OutCaptureHDR, FString& OutFilenameOverride, bool& OutUseDateTimeAsFileName)
 {
 	FString CmdString = Cmd;
-	int32 SeperatorPos = -1;
-	int32 LastSeperatorPos = 0;
 	TArray<FString> Arguments;
+	const FString FilenameSearchString = TEXT("filename=");
 
-	// Look for an optional filename to override from the default filename and strip it if found.
-	FString FilenameSearchString = TEXT("filename=");
-	int32 FilenamePos = CmdString.Find(FilenameSearchString, ESearchCase::IgnoreCase);
-	if (FilenamePos != INDEX_NONE)
-	{
-		FString FilenameOverride;
-		FParse::Value(Cmd, TEXT("filename="), FilenameOverride);
-		OutFilenameOverride = FilenameOverride;
-		CmdString.RemoveAt(FilenamePos, FilenameSearchString.Len() + FilenameOverride.Len());
-		CmdString.TrimStartAndEndInline(); 
-	}
+	// FParse::Value has better handling of escape characters than FParse::Token
+	FParse::Value(Cmd, *FilenameSearchString, OutFilenameOverride);
 
-	while (CmdString.FindChar(TCHAR(' '), SeperatorPos))
+	FString Arg;
+	while (FParse::Token(Cmd, Arg, true))
 	{
-		Arguments.Add(CmdString.Mid(LastSeperatorPos, SeperatorPos));
-		CmdString = CmdString.Mid(SeperatorPos + 1);
+		// Now skip filename since we already processed it
+		if (!Arg.StartsWith(FilenameSearchString))
+		{
+			Arguments.Add(Arg);
+		}
 	}
 
 	if (CmdString.Len() > 0)
