Updated warning on UActorComponent::SetIsReplicated(bool bShouldReplicate) to emit the correct function name to use instead during construction.

#rnx

diff --git a/Engine/Source/Runtime/Engine/Private/Components/ActorComponent.cpp b/Engine/Source/Runtime/Engine/Private/Components/ActorComponent.cpp
--- a/Engine/Source/Runtime/Engine/Private/Components/ActorComponent.cpp
+++ b/Engine/Source/Runtime/Engine/Private/Components/ActorComponent.cpp
@@ -1738,7 +1738,7 @@ void UActorComponent::SetIsReplicated(bool bShouldReplicate)
 {
 	if (GetIsReplicated() != bShouldReplicate)
 	{
-		ensureMsgf(!NeedsInitialization(), TEXT("SetReplicatedByDefault is preferred during Component Construction."));
+		ensureMsgf(!NeedsInitialization(), TEXT("SetIsReplicatedByDefault is preferred during Component Construction."));
 
 		if (GetComponentClassCanReplicate())
 		{
