Fixed issue where LoadLevelInstance location offset wasn't affecting the BSP geometry in a dynamically loaded level. Now a level's ModelComponents also have the same transform applied: this fixes their rendering, physics and static lighting.

diff --git a/Engine/Source/Runtime/Engine/Private/LevelUtils.cpp b/Engine/Source/Runtime/Engine/Private/LevelUtils.cpp
--- a/Engine/Source/Runtime/Engine/Private/LevelUtils.cpp
+++ b/Engine/Source/Runtime/Engine/Private/LevelUtils.cpp
@@ -11,6 +11,7 @@
 #include "EngineGlobals.h"
 #include "Misc/FeedbackContext.h"
 #include "GameFramework/WorldSettings.h"
+#include "Components/ModelComponent.h"
 
 #if WITH_EDITOR
 #include "ScopedTransaction.h"
@@ -398,6 +399,16 @@ void FLevelUtils::ApplyLevelTransform(const FLevelUtils::FApplyLevelTransformPar
 
 		if (TransformParams.bSetRelativeTransformDirectly)
 		{
+			// Iterate over all model components to transform BSP geometry accordingly
+			for (UModelComponent* ModelComponent : TransformParams.Level->ModelComponents)
+			{
+				if (ModelComponent)
+				{
+					ModelComponent->SetRelativeLocation_Direct(TransformParams.LevelTransform.TransformPosition(ModelComponent->GetRelativeLocation()));
+					ModelComponent->SetRelativeRotation_Direct(TransformParams.LevelTransform.TransformRotation(ModelComponent->GetRelativeRotation().Quaternion()).Rotator());
+				}
+			}
+
 			// Iterate over all actors in the level and transform them
 			for (AActor* Actor : TransformParams.Level->Actors)
 			{
@@ -416,6 +427,15 @@ void FLevelUtils::ApplyLevelTransform(const FLevelUtils::FApplyLevelTransformPar
 		}
 		else
 		{
+			// Iterate over all model components to transform BSP geometry accordingly
+			for (UModelComponent* ModelComponent : TransformParams.Level->ModelComponents)
+			{
+				if (ModelComponent)
+				{
+					ModelComponent->SetRelativeLocationAndRotation(TransformParams.LevelTransform.TransformPosition(ModelComponent->GetRelativeLocation()), TransformParams.LevelTransform.TransformRotation(ModelComponent->GetRelativeRotation().Quaternion()));
+				}
+			}
+
 			// Iterate over all actors in the level and transform them
 			for (AActor* Actor : TransformParams.Level->Actors)
 			{
