Fixed static analysis warning.

diff --git a/Engine/Source/Runtime/Engine/Private/InstancedStaticMesh.cpp b/Engine/Source/Runtime/Engine/Private/InstancedStaticMesh.cpp
--- a/Engine/Source/Runtime/Engine/Private/InstancedStaticMesh.cpp
+++ b/Engine/Source/Runtime/Engine/Private/InstancedStaticMesh.cpp
@@ -2987,7 +2987,7 @@ void FInstancedStaticMeshVertexFactoryShaderParameters::GetElementShaderBindings
 		FVector4 InstancingViewZCompareZero(MIN_flt, MIN_flt, MAX_flt, 1.0f);
 		FVector4 InstancingViewZCompareOne(MIN_flt, MIN_flt, MAX_flt, 0.0f);
 		FVector4 InstancingViewZConstant(ForceInit);
-		FVector4 InstancingOffset(InstancingUserData->InstancingOffset);
+		FVector4 InstancingOffset(ForceInit);
 		FVector4 InstancingWorldViewOriginZero(ForceInit);
 		FVector4 InstancingWorldViewOriginOne(ForceInit);
 		InstancingWorldViewOriginOne.W = 1.0f;
@@ -3069,7 +3069,7 @@ void FInstancedStaticMeshVertexFactoryShaderParameters::GetElementShaderBindings
 				}
 			}
 
-
+			InstancingOffset = InstancingUserData->InstancingOffset;
 			InstancingWorldViewOriginZero = View->GetTemporalLODOrigin(0);
 			InstancingWorldViewOriginOne = View->GetTemporalLODOrigin(1);
 
