Fixed deprecation warnings.

diff --git a/Engine/Source/Runtime/Engine/Private/Components/SceneComponent.cpp b/Engine/Source/Runtime/Engine/Private/Components/SceneComponent.cpp
--- a/Engine/Source/Runtime/Engine/Private/Components/SceneComponent.cpp
+++ b/Engine/Source/Runtime/Engine/Private/Components/SceneComponent.cpp
@@ -2766,10 +2766,10 @@ bool USceneComponent::InternalSetWorldLocationAndRotation(FVector NewLocation, c
 		}
 
 #if ENABLE_NAN_DIAGNOSTIC
-		if (RelativeRotation.ContainsNaN())
+		if (GetRelativeRotation().ContainsNaN())
 		{
-			logOrEnsureNanError(TEXT("USceneComponent:InternalSetWorldLocationAndRotation found NaN in RelativeRotation: %s"), *RelativeRotation.ToString());
-			RelativeRotation = FRotator::ZeroRotator;
+			logOrEnsureNanError(TEXT("USceneComponent:InternalSetWorldLocationAndRotation found NaN in RelativeRotation: %s"), *GetRelativeRotation().ToString());
+			SetRelativeRotation_Direct(FRotator::ZeroRotator);
 		}
 #endif
 		UpdateComponentToWorldWithParent(GetAttachParent(),GetAttachSocketName(), SkipPhysicsToEnum(bNoPhysics), RelativeRotationCache.GetCachedQuat(), Teleport);
