Adding FrameIndex override to FSceneView.

diff --git a/Engine/Source/Runtime/Engine/Public/SceneView.h b/Engine/Source/Runtime/Engine/Public/SceneView.h
--- a/Engine/Source/Runtime/Engine/Public/SceneView.h
+++ b/Engine/Source/Runtime/Engine/Public/SceneView.h
@@ -1032,6 +1032,9 @@ public:
 	/** Translucent sort mode */
 	TEnumAsByte<ETranslucentSortPolicy::Type> TranslucentSortPolicy;
 	
+	/** The frame index to override, useful for keeping determinism when rendering sequences. **/
+	TOptional<uint32> OverrideFrameIndexValue;
+
 #if WITH_EDITOR
 	/** The set of (the first 64) groups' visibility info for this view */
 	uint64 EditorViewBitflag;
diff --git a/Engine/Source/Runtime/Renderer/Private/SceneVisibility.cpp b/Engine/Source/Runtime/Renderer/Private/SceneVisibility.cpp
--- a/Engine/Source/Runtime/Renderer/Private/SceneVisibility.cpp
+++ b/Engine/Source/Runtime/Renderer/Private/SceneVisibility.cpp
@@ -3114,6 +3114,11 @@ void FSceneRenderer::PreVisibilityFrameSetup(FRHICommandListImmediate& RHICmdLis
 			{
 				ViewState->FrameIndex++;
 			}
+
+			if (View.OverrideFrameIndexValue.IsSet())
+			{
+				ViewState->FrameIndex = View.OverrideFrameIndexValue.GetValue();
+			}
 		}
 		
 		// Subpixel jitter for temporal AA
