Copy up from Dev-Editor [REDACTED_HANDLE]

diff --git a/Engine/Source/Runtime/Engine/Classes/Engine/SkeletalMesh.h b/Engine/Source/Runtime/Engine/Classes/Engine/SkeletalMesh.h
--- a/Engine/Source/Runtime/Engine/Classes/Engine/SkeletalMesh.h
+++ b/Engine/Source/Runtime/Engine/Classes/Engine/SkeletalMesh.h
@@ -47,6 +47,10 @@ class FSkeletalMeshLODRenderData;
 class FSkinWeightVertexBuffer;
 struct FSkinWeightProfileInfo;
 class FSkeletalMeshUpdate;
+class USkeletalMeshEditorData;
+class FSkeletalMeshImportData;
+enum class ESkeletalMeshGeoImportVersions : uint8;
+enum class ESkeletalMeshSkinningImportVersions : uint8;
 
 #if WITH_APEX_CLOTHING
 
@@ -492,11 +496,74 @@ private:
 	/** Rendering resources used at runtime */
 	TUniquePtr<FSkeletalMeshRenderData> SkeletalMeshRenderData;
 
+#if WITH_EDITORONLY_DATA
+	/*
+	 * This editor data asset is save in the same package has the skeletalmesh, the editor data asset is always loaded.
+	 * If the skeletal mesh is rename the editor data asset will also be rename: the name is SkeletalMeshName_USkeletalMeshEditorData
+	 * If the skeletal mesh is duplicate the editor data asset will also be duplicate
+	 * There is only one editor data asset possible per skeletalmesh.
+	 * The reason we store the editor data in a separate asset is because the size of it can be very big and affect the editor performance. (undo/redo transactions)
+	 */
+	UPROPERTY()
+	mutable USkeletalMeshEditorData* MeshEditorDataObject;
+
+	/*
+	 * Return a valid USkeletalMeshEditorData, if the MeshEditorDataPath is invalid it will create the USkeletalMeshEditorData and set the MeshEditorDataPath to point on it.
+	 */
+	USkeletalMeshEditorData& GetMeshEditorData() const;
+
+#endif //WITH_EDITORONLY_DATA
+
 public:
+
 #if WITH_EDITORONLY_DATA
+
+	//////////////////////////////////////////////////////////////////////////
+	// USkeletalMeshEditorData public skeletalmesh API
+	// We do not want skeletal mesh client to use directly the asset(function GetMeshEditorData)
+	// We have to maintain some sync between the LODModels and the asset to avoid loading the asset when
+	// building the DDC key. That is why the asset accessor are private. the data we keep in sync in the LODModels is:
+	// IsLODImportedDataBuildAvailable
+	// IsLODImportedDataEmpty
+	// Raw mesh data DDC string ID, there is no API to retrieve it, since only the LODModels need this value
+	
+
+	/* Fill the OutMesh with the imported data */
+	void LoadLODImportedData(const int32 LODIndex, FSkeletalMeshImportData& OutMesh) const;
+	
+	/* Fill the asset LOD entry with the InMesh. */
+	void SaveLODImportedData(const int32 LODIndex, FSkeletalMeshImportData& InMesh);
+	
+	/* Return true if the imported data has all the necessary data to use the skeletalmesh builder. Return False otherwise.
+	 * Old asset before the refactor will not be able to be build until it get fully re-import.
+	 * This value is cache in the LODModel and update when we call SaveLODImportedData.
+	 */
+	bool IsLODImportedDataBuildAvailable(const int32 LODIndex) const;
+	
+	/* Return true if the imported data is present. Return false otherwise.
+	 * Old asset before the split workflow will not have this data and will not support import geo only or skinning only.
+	 * This value is cache in the LODModel and update when we call SaveLODImportedData.
+	 */
+	bool IsLODImportedDataEmpty(const int32 LODIndex) const;
+
+	/* Get the Versions of the geo and skinning data. We use those versions to answer to IsLODImportedDataBuildAvailable function. */
+	void GetLODImportedDataVersions(const int32 LODIndex, ESkeletalMeshGeoImportVersions& OutGeoImportVersion, ESkeletalMeshSkinningImportVersions& OutSkinningImportVersion) const;
+
+	/* Set the Versions of the geo and skinning data. We use those versions to answer to IsLODImportedDataBuildAvailable function. */
+	void SetLODImportedDataVersions(const int32 LODIndex, const ESkeletalMeshGeoImportVersions& InGeoImportVersion, const ESkeletalMeshSkinningImportVersions& InSkinningImportVersion);
+
+	/* Static function that copy the LOD import data from a source skeletal mesh to a destination skeletal mesh*/
+	static void CopyImportedData(int32 SrcLODIndex, USkeletalMesh* SrcSkeletalMesh, int32 DestLODIndex, USkeletalMesh* DestSkeletalMesh);
+
+	/* Allocate the space we need. Use this before calling this API in multithreaded. */
+	void ReserveLODImportData(int32 MaxLODIndex);
+	
+	// End USkeletalMeshEditorData public skeletalmesh API
+	//////////////////////////////////////////////////////////////////////////
+
 	/** Get the imported data for this skeletal mesh. */
 	FORCEINLINE FSkeletalMeshModel* GetImportedModel() const { return ImportedModel.Get(); }
-#endif
+#endif //WITH_EDITORONLY_DATA
 
 	/** Get the data to use for rendering. */
 	FORCEINLINE FSkeletalMeshRenderData* GetResourceForRendering() const { return SkeletalMeshRenderData.Get(); }
@@ -990,6 +1057,8 @@ public:
 	virtual FString GetDetailedInfoInternal() const override;
 	virtual void GetResourceSizeEx(FResourceSizeEx& CumulativeResourceSize) override;
 	virtual void GetPreloadDependencies(TArray<UObject*>& OutDeps) override;
+	virtual void PostDuplicate(bool bDuplicateForPIE) override;
+	virtual void PostRename(UObject* OldOuter, const FName OldName) override;
 	//~ End UObject Interface.
 
 	//~ Begin UStreamableRenderAsset Interface.
diff --git a/Engine/Source/Runtime/Engine/Classes/Engine/WorldComposition.h b/Engine/Source/Runtime/Engine/Classes/Engine/WorldComposition.h
--- a/Engine/Source/Runtime/Engine/Classes/Engine/WorldComposition.h
+++ b/Engine/Source/Runtime/Engine/Classes/Engine/WorldComposition.h
@@ -199,7 +199,7 @@ private:
 	bool IsDistanceDependentLevel(int32 TileIdx) const;
 
 	/** Attempts to set new streaming state for a particular tile, could be rejected if state change on 'cooldown' */
-	bool CommitTileStreamingState(UWorld* PersistenWorld, int32 TileIdx, bool bShouldBeLoaded, bool bShouldBeVisible, bool bShouldBlock, int32 LODIdx);
+	bool CommitTileStreamingState(UWorld* PersistentWorld, int32 TileIdx, bool bShouldBeLoaded, bool bShouldBeVisible, bool bShouldBlock, int32 LODIdx);
 
 public:
 #if WITH_EDITOR
diff --git a/Engine/Source/Runtime/Engine/Private/Analytics/EditorSessionSummaryWriter.h b/Engine/Source/Runtime/Engine/Private/Analytics/EditorSessionSummaryWriter.h
--- a/Engine/Source/Runtime/Engine/Private/Analytics/EditorSessionSummaryWriter.h
+++ b/Engine/Source/Runtime/Engine/Private/Analytics/EditorSessionSummaryWriter.h
@@ -17,6 +17,7 @@ public:
 
 	void Initialize();
 	void Tick(float DeltaTime);
+	void LowDriveSpaceDetected();
 	void Shutdown();
 
 private:
@@ -31,6 +32,7 @@ private:
 
 	FString GetUserActivityString() const;
 	void UpdateTimestamps();
+	void TrySaveCurrentSession();
 
 private:
 	FEditorAnalyticsSession* CurrentSession;
diff --git a/Engine/Source/Runtime/Engine/Private/AudioDevice.cpp b/Engine/Source/Runtime/Engine/Private/AudioDevice.cpp
--- a/Engine/Source/Runtime/Engine/Private/AudioDevice.cpp
+++ b/Engine/Source/Runtime/Engine/Private/AudioDevice.cpp
@@ -6195,7 +6195,7 @@ bool FAudioDevice::CanUseVRAudioDevice()
 	if (GIsEditor)
 	{
 		UEditorEngine* EdEngine = Cast<UEditorEngine>(GEngine);
-		return EdEngine->bUseVRPreviewForPlayWorld;
+		return EdEngine->IsVRPreviewActive();
 	}
 	else
 #endif
diff --git a/Engine/Source/Runtime/Engine/Private/GameInstance.cpp b/Engine/Source/Runtime/Engine/Private/GameInstance.cpp
--- a/Engine/Source/Runtime/Engine/Private/GameInstance.cpp
+++ b/Engine/Source/Runtime/Engine/Private/GameInstance.cpp
@@ -202,6 +202,35 @@ void UGameInstance::CreateMinimalNetRPCWorld(const FName InPackageName, UPackage
 }
 
 #if WITH_EDITOR
+static ENetMode GetNetModeFromPlayNetMode(const EPlayNetMode InPlayNetMode, const bool bInDedicatedServer)
+{
+	switch (InPlayNetMode)
+	{
+		case EPlayNetMode::PIE_Client:
+		{
+			return NM_Client;
+		}
+		case EPlayNetMode::PIE_ListenServer:
+		{
+			if (bInDedicatedServer)
+			{
+				return NM_DedicatedServer;
+			}
+
+			return NM_ListenServer;
+		}
+		case EPlayNetMode::PIE_Standalone:
+		{
+			return NM_Standalone;
+		}
+		default:
+		{
+			break;
+		}
+	}
+
+	return NM_Standalone;
+}
 
 FGameInstancePIEResult UGameInstance::InitializeForPlayInEditor(int32 PIEInstanceIndex, const FGameInstancePIEParameters& Params)
 {
@@ -229,15 +258,13 @@ FGameInstancePIEResult UGameInstance::InitializeForPlayInEditor(int32 PIEInstanc
 	WorldContext->LastURL.Map = WorldPackageName;
 	WorldContext->PIEPrefix = WorldContext->PIEInstance != INDEX_NONE ? UWorld::BuildPIEPackagePrefix(WorldContext->PIEInstance) : FString();
 
-	const ULevelEditorPlaySettings* PlayInSettings = GetDefault<ULevelEditorPlaySettings>();
-
+	
 	// We always need to create a new PIE world unless we're using the editor world for SIE
 	UWorld* NewWorld = nullptr;
 
 	bool bNeedsGarbageCollection = false;
-	const EPlayNetMode PlayNetMode = [&PlayInSettings]{ EPlayNetMode NetMode(PIE_Standalone); return (PlayInSettings->GetPlayNetMode(NetMode) ? NetMode : PIE_Standalone); }();
-	const bool CanRunUnderOneProcess = [&PlayInSettings]{ bool RunUnderOneProcess(false); return (PlayInSettings->GetRunUnderOneProcess(RunUnderOneProcess) && RunUnderOneProcess); }();
-	if (PlayNetMode == PIE_Client)
+
+	if (Params.NetMode == EPlayNetMode::PIE_Client)
 	{
 		// We are going to connect, so just load an empty world
 		NewWorld = EditorEngine->CreatePIEWorldFromEntry(*WorldContext, EditorEngine->EditorWorld, PIEMapName);
@@ -257,6 +284,7 @@ FGameInstancePIEResult UGameInstance::InitializeForPlayInEditor(int32 PIEInstanc
 		return FGameInstancePIEResult::Failure(NSLOCTEXT("UnrealEd", "Error_FailedCreateEditorPreviewWorld", "Failed to create editor preview world."));
 	}
 
+	NewWorld->SetPlayInEditorInitialNetMode(GetNetModeFromPlayNetMode(Params.NetMode, Params.bRunAsDedicated));
 	NewWorld->SetGameInstance(this);
 	WorldContext->SetCurrentWorld(NewWorld);
 	WorldContext->AddRef(EditorEngine->PlayWorld);	// Tie this context to this UEngine::PlayWorld*		// @fixme, needed still?
@@ -300,28 +328,25 @@ FGameInstancePIEResult UGameInstance::StartPlayInEditorGameInstance(ULocalPlayer
 	OnStart();
 
 	UEditorEngine* const EditorEngine = CastChecked<UEditorEngine>(GetEngine());
-	ULevelEditorPlaySettings const* PlayInSettings = GetDefault<ULevelEditorPlaySettings>();
-
-	const EPlayNetMode PlayNetMode = [&PlayInSettings]{ EPlayNetMode NetMode(PIE_Standalone); return (PlayInSettings->GetPlayNetMode(NetMode) ? NetMode : PIE_Standalone); }();
 
 	// for clients, just connect to the server
-	if (PlayNetMode == PIE_Client)
+	if (Params.NetMode == PIE_Client)
 	{
 		FString Error;
 		FURL BaseURL = WorldContext->LastURL;
 
 		FString URLString(TEXT("127.0.0.1"));
 		uint16 ServerPort = 0;
-		if (PlayInSettings->GetServerPort(ServerPort))
+		if (Params.EditorPlaySettings->GetServerPort(ServerPort))
 		{
 			URLString += FString::Printf(TEXT(":%hu"), ServerPort);
 		}
 
-		if (PlayInSettings->IsNetworkEmulationEnabled())
+		if (Params.EditorPlaySettings->IsNetworkEmulationEnabled())
 		{
-			if (PlayInSettings->NetworkEmulationSettings.IsEmulationEnabledForTarget(NetworkEmulationTarget::Client))
+			if (Params.EditorPlaySettings->NetworkEmulationSettings.IsEmulationEnabledForTarget(NetworkEmulationTarget::Client))
 			{
-				URLString += PlayInSettings->NetworkEmulationSettings.BuildPacketSettingsForURL();
+				URLString += Params.EditorPlaySettings->NetworkEmulationSettings.BuildPacketSettingsForURL();
 			}
 		}
 
@@ -340,21 +365,21 @@ FGameInstancePIEResult UGameInstance::StartPlayInEditorGameInstance(ULocalPlayer
 		UWorld* const PlayWorld = GetWorld();
 
 		FString ExtraURLOptions;
-		if (PlayInSettings->IsNetworkEmulationEnabled())
+		if (Params.EditorPlaySettings->IsNetworkEmulationEnabled())
 		{
-			NetworkEmulationTarget CurrentTarget = PlayNetMode == PIE_ListenServer ? NetworkEmulationTarget::Server : NetworkEmulationTarget::Client;
-			if (PlayInSettings->NetworkEmulationSettings.IsEmulationEnabledForTarget(CurrentTarget))
+			NetworkEmulationTarget CurrentTarget = Params.NetMode == PIE_ListenServer ? NetworkEmulationTarget::Server : NetworkEmulationTarget::Client;
+			if (Params.EditorPlaySettings->NetworkEmulationSettings.IsEmulationEnabledForTarget(CurrentTarget))
 			{
-				ExtraURLOptions += PlayInSettings->NetworkEmulationSettings.BuildPacketSettingsForURL();
+				ExtraURLOptions += Params.EditorPlaySettings->NetworkEmulationSettings.BuildPacketSettingsForURL();
 			}
 		}
 
 		// make a URL
 		FURL URL;
 		// If the user wants to start in spectator mode, do not use the custom play world for now
-		if (EditorEngine->UserEditedPlayWorldURL.Len() > 0 && !Params.bStartInSpectatorMode)
+		if (EditorEngine->UserEditedPlayWorldURL.Len() > 0 || Params.OverrideMapURL.Len() > 0)
 		{
-			FString UserURL = EditorEngine->UserEditedPlayWorldURL;
+			FString UserURL = EditorEngine->UserEditedPlayWorldURL.Len() > 0 ? EditorEngine->UserEditedPlayWorldURL : Params.OverrideMapURL;
 			UserURL += ExtraURLOptions;
 
 			// If the user edited the play world url. Verify that the map name is the same as the currently loaded map.
@@ -373,7 +398,7 @@ FGameInstancePIEResult UGameInstance::StartPlayInEditorGameInstance(ULocalPlayer
 
 		// If a start location is specified, spawn a temporary PlayerStartPIE actor at the start location and use it as the portal.
 		AActor* PlayerStart = NULL;
-		if (EditorEngine->SpawnPlayFromHereStart(PlayWorld, PlayerStart, EditorEngine->PlayWorldLocation, EditorEngine->PlayWorldRotation) == false)
+		if (!EditorEngine->SpawnPlayFromHereStart(PlayWorld, PlayerStart))
 		{
 			// failed to create "play from here" playerstart
 			return FGameInstancePIEResult::Failure(NSLOCTEXT("UnrealEd", "Error_FailedCreatePlayFromHerePlayerStart", "Failed to create PlayerStart at desired starting location."));
@@ -417,11 +442,11 @@ FGameInstancePIEResult UGameInstance::StartPlayInEditorGameInstance(ULocalPlayer
 			GEngine->BlockTillLevelStreamingCompleted(PlayWorld);
 		}
 		
-		if (PlayNetMode == PIE_ListenServer)
+		if (Params.NetMode == PIE_ListenServer)
 		{
 			// Add port
 			uint16 ServerPort = 0;
-			if (PlayInSettings->GetServerPort(ServerPort))
+			if (Params.EditorPlaySettings->GetServerPort(ServerPort))
 			{
 				URL.Port = ServerPort;
 			}
diff --git a/Engine/Source/Runtime/Engine/Private/HierarchicalInstancedStaticMesh.cpp b/Engine/Source/Runtime/Engine/Private/HierarchicalInstancedStaticMesh.cpp
--- a/Engine/Source/Runtime/Engine/Private/HierarchicalInstancedStaticMesh.cpp
+++ b/Engine/Source/Runtime/Engine/Private/HierarchicalInstancedStaticMesh.cpp
@@ -1974,6 +1974,17 @@ UHierarchicalInstancedStaticMeshComponent::~UHierarchicalInstancedStaticMeshComp
 #if WITH_EDITOR
 void UHierarchicalInstancedStaticMeshComponent::PostEditChangeChainProperty(FPropertyChangedChainEvent& PropertyChangedEvent)
 {
+	if ((PropertyChangedEvent.Property != NULL && PropertyChangedEvent.Property->GetFName() == "InstancingRandomSeed"))
+	{
+		// If we don't have a random seed for this instanced static mesh component yet, then go ahead and
+		// generate one now.  This will be saved with the static mesh component and used for future generation
+		// of random numbers for this component's instances. (Used by the PerInstanceRandom material expression)
+		while (InstancingRandomSeed == 0)
+		{
+			InstancingRandomSeed = FMath::Rand();
+		}
+	}
+
 	Super::PostEditChangeChainProperty(PropertyChangedEvent);
 
 	if ((PropertyChangedEvent.Property != NULL && PropertyChangedEvent.Property->GetFName() == "PerInstanceSMData") ||
@@ -2252,6 +2263,7 @@ bool UHierarchicalInstancedStaticMeshComponent::UpdateInstanceTransform(int32 In
 		if (!bIsOmittedInstance)
 		{
 			InstanceUpdateCmdBuffer.UpdateInstance(RenderIndex, NewLocalTransform.ToMatrixWithScale());
+			bMarkRenderStateDirty = true;
 		}
 		
 		if (bDoInPlaceUpdate)
@@ -2261,11 +2273,11 @@ bool UHierarchicalInstancedStaticMeshComponent::UpdateInstanceTransform(int32 In
 			if (!OldInstanceBounds.IsInside(NewInstanceBounds))
 			{
 				BuiltInstanceBounds += NewInstanceBounds;
+			}
 
-				if (bMarkRenderStateDirty)
-				{
-					MarkRenderStateDirty();
-				}
+			if (bMarkRenderStateDirty)
+			{
+				MarkRenderStateDirty();
 			}
 		}
 		else
diff --git a/Engine/Source/Runtime/Engine/Private/Materials/HLSLMaterialTranslator.cpp b/Engine/Source/Runtime/Engine/Private/Materials/HLSLMaterialTranslator.cpp
--- a/Engine/Source/Runtime/Engine/Private/Materials/HLSLMaterialTranslator.cpp
+++ b/Engine/Source/Runtime/Engine/Private/Materials/HLSLMaterialTranslator.cpp
@@ -1372,7 +1372,7 @@ void FHLSLMaterialTranslator::GetSharedInputsMaterialCode(FString& PixelMembersD
 			const EMaterialProperty Property = (EMaterialProperty)PropertyIndex;
 			check(FMaterialAttributeDefinitionMap::GetShaderFrequency(Property) == SF_Pixel);
 			// Special case MP_SubsurfaceColor as the actual property is a combination of the color and the profile but we don't want to expose the profile
-			const FString PropertyName = Property == MP_SubsurfaceColor ? "Subsurface" : FMaterialAttributeDefinitionMap::GetDisplayName(Property);
+			const FString PropertyName = Property == MP_SubsurfaceColor ? "Subsurface" : FMaterialAttributeDefinitionMap::GetAttributeName(Property);
 			check(PropertyName.Len() > 0);				
 			const EMaterialValueType Type = Property == MP_SubsurfaceColor ? MCT_Float4 : FMaterialAttributeDefinitionMap::GetValueType(Property);
 
@@ -6693,7 +6693,7 @@ int32 FHLSLMaterialTranslator::CustomOutput(class UMaterialExpressionCustomOutpu
 {
 	if (MaterialProperty != MP_MAX)
 	{
-		return Errorf(TEXT("A Custom Output node should not be attached to the %s material property"), *FMaterialAttributeDefinitionMap::GetDisplayName(MaterialProperty));
+		return Errorf(TEXT("A Custom Output node should not be attached to the %s material property"), *FMaterialAttributeDefinitionMap::GetAttributeName(MaterialProperty));
 	}
 
 	if (OutputCode == INDEX_NONE)
diff --git a/Engine/Source/Runtime/Engine/Private/Materials/MaterialExpressions.cpp b/Engine/Source/Runtime/Engine/Private/Materials/MaterialExpressions.cpp
--- a/Engine/Source/Runtime/Engine/Private/Materials/MaterialExpressions.cpp
+++ b/Engine/Source/Runtime/Engine/Private/Materials/MaterialExpressions.cpp
@@ -5424,7 +5424,7 @@ int32 UMaterialExpressionMakeMaterialAttributes::Compile(class FMaterialCompiler
 	//If we've connected an expression but its still returned INDEX_NONE, flag the error. This also catches reroute nodes to nowhere.
 	if (Expression && INDEX_NONE == Ret)
 	{
-		Compiler->Errorf(TEXT("Error on property %s"), *FMaterialAttributeDefinitionMap::GetDisplayName(Property));
+		Compiler->Errorf(TEXT("Error on property %s"), *FMaterialAttributeDefinitionMap::GetAttributeName(Property));
 	}
 
 	return Ret;
@@ -5799,7 +5799,7 @@ void UMaterialExpressionGetMaterialAttributes::PostEditChangeProperty(FPropertyC
 			}
 		
 			// Copy final defaults to new output
-			FString AttributeName = FMaterialAttributeDefinitionMap::GetDisplayName(AttributeGetTypes.Last());
+			FString AttributeName = FMaterialAttributeDefinitionMap::GetDisplayNameForMaterial(AttributeGetTypes.Last(), Material).ToString();
 			Outputs.Add(FExpressionOutput(*AttributeName, 0, 0, 0, 0, 0));
 
 			GraphNode->ReconstructNode();
@@ -5845,7 +5845,7 @@ void UMaterialExpressionGetMaterialAttributes::PostEditChangeProperty(FPropertyC
 			// Type changed, update pin names
 			for (int i = 1; i < Outputs.Num(); ++i)
 			{
-				Outputs[i].OutputName = *FMaterialAttributeDefinitionMap::GetDisplayName(AttributeGetTypes[i-1]);
+				Outputs[i].OutputName = *FMaterialAttributeDefinitionMap::GetDisplayNameForMaterial(AttributeGetTypes[i-1], Material).ToString();
 			}
 
 			GraphNode->ReconstructNode();
@@ -5864,7 +5864,7 @@ void UMaterialExpressionGetMaterialAttributes::PostLoad()
 
 	for (int i = 1; i < Outputs.Num(); ++i)
 	{
-		const FString DisplayName = FMaterialAttributeDefinitionMap::GetDisplayName(AttributeGetTypes[i-1]);
+		const FString DisplayName = FMaterialAttributeDefinitionMap::GetDisplayNameForMaterial(AttributeGetTypes[i-1],Material).ToString();
 		if (Outputs[i].OutputName.ToString() != DisplayName)
 		{
 			FString MaterialName;
@@ -5985,7 +5985,7 @@ FName UMaterialExpressionSetMaterialAttributes::GetInputName(int32 InputIndex) c
 	}
 	else if (InputIndex > 0)
 	{
-		Name = *FMaterialAttributeDefinitionMap::GetDisplayName(AttributeSetTypes[InputIndex-1]);
+		Name = *FMaterialAttributeDefinitionMap::GetDisplayNameForMaterial(AttributeSetTypes[InputIndex-1], Material).ToString();
 	}
 
 	return Name;
@@ -6055,7 +6055,7 @@ void UMaterialExpressionSetMaterialAttributes::PostEditChangeProperty(FPropertyC
 		
 			// Copy final defaults to new input
 			Inputs.Add(FExpressionInput());
-			Inputs.Last().InputName = FName(*FMaterialAttributeDefinitionMap::GetDisplayName(AttributeSetTypes.Last()));
+			Inputs.Last().InputName = FName(*FMaterialAttributeDefinitionMap::GetDisplayNameForMaterial(AttributeSetTypes.Last(), Material).ToString());
 			GraphNode->ReconstructNode();
 		}	 
 		else if (PreEditAttributeSetTypes.Num() > AttributeSetTypes.Num())
@@ -6099,7 +6099,7 @@ void UMaterialExpressionSetMaterialAttributes::PostEditChangeProperty(FPropertyC
 			// Type changed, update pin names
 			for (int i = 1; i < Inputs.Num(); ++i)
 			{
-				Inputs[i].InputName = *FMaterialAttributeDefinitionMap::GetDisplayName(AttributeSetTypes[i - 1]);
+				Inputs[i].InputName = FName(*FMaterialAttributeDefinitionMap::GetDisplayNameForMaterial(AttributeSetTypes[i - 1], Material).ToString());
 			}
 			GraphNode->ReconstructNode();
 		}
diff --git a/Engine/Source/Runtime/Engine/Private/Materials/MaterialShared.cpp b/Engine/Source/Runtime/Engine/Private/Materials/MaterialShared.cpp
--- a/Engine/Source/Runtime/Engine/Private/Materials/MaterialShared.cpp
+++ b/Engine/Source/Runtime/Engine/Private/Materials/MaterialShared.cpp
@@ -43,6 +43,8 @@
 #include "UObject/CoreRedirects.h"
 #include "RayTracingDefinitions.h"
 
+#define LOCTEXT_NAMESPACE "MaterialShared"
+
 DEFINE_LOG_CATEGORY(LogMaterial);
 
 int32 GDeferUniformExpressionCaching = 1;
@@ -3542,11 +3544,11 @@ bool FMaterialShaderMapId::ContainsVertexFactoryType(const FVertexFactoryType* V
 //////////////////////////////////////////////////////////////////////////
 
 FMaterialAttributeDefintion::FMaterialAttributeDefintion(
-		const FGuid& InAttributeID, const FString& InDisplayName, EMaterialProperty InProperty,
+		const FGuid& InAttributeID, const FString& InAttributeName, EMaterialProperty InProperty,
 		EMaterialValueType InValueType, const FVector4& InDefaultValue, EShaderFrequency InShaderFrequency,
 		int32 InTexCoordIndex /*= INDEX_NONE*/, bool bInIsHidden /*= false*/, MaterialAttributeBlendFunction InBlendFunction /*= nullptr*/)
 	: AttributeID(InAttributeID)
-	, DisplayName(InDisplayName)
+	, AttributeName(InAttributeName)
 	, Property(InProperty)
 	, ValueType(InValueType)
 	, DefaultValue(InDefaultValue)
@@ -3600,9 +3602,9 @@ int32 FMaterialAttributeDefintion::CompileDefaultValue(FMaterialCompiler* Compil
 //////////////////////////////////////////////////////////////////////////
 
 FMaterialCustomOutputAttributeDefintion::FMaterialCustomOutputAttributeDefintion(
-		const FGuid& InAttributeID, const FString& InDisplayName, const FString& InFunctionName, EMaterialProperty InProperty,
+		const FGuid& InAttributeID, const FString& InAttributeName, const FString& InFunctionName, EMaterialProperty InProperty,
 		EMaterialValueType InValueType, const FVector4& InDefaultValue, EShaderFrequency InShaderFrequency, MaterialAttributeBlendFunction InBlendFunction /*= nullptr*/)
-	: FMaterialAttributeDefintion(InAttributeID, InDisplayName, InProperty, InValueType, InDefaultValue, InShaderFrequency, INDEX_NONE, false, InBlendFunction)
+	: FMaterialAttributeDefintion(InAttributeID, InAttributeName, InProperty, InValueType, InDefaultValue, InShaderFrequency, INDEX_NONE, false, InBlendFunction)
 	, FunctionName(InFunctionName)
 {
 }
@@ -3664,12 +3666,13 @@ void FMaterialAttributeDefinitionMap::InitializeAttributeMap()
 	AddCustomAttribute(FGuid(0x8EAB2CB2, 0x73634A24, 0x8CD14F47, 0x3F9C8E55), "CustomEyeTangent", "GetTangentOutput", MCT_Float3, FVector4(0, 0, 0, 0));
 }
 
-void FMaterialAttributeDefinitionMap::Add(const FGuid& AttributeID, const FString& DisplayName, EMaterialProperty Property,
+
+void FMaterialAttributeDefinitionMap::Add(const FGuid& AttributeID, const FString& AttributeName, EMaterialProperty Property,
 	EMaterialValueType ValueType, const FVector4& DefaultValue, EShaderFrequency ShaderFrequency,
 	int32 TexCoordIndex /*= INDEX_NONE*/, bool bIsHidden /*= false*/, MaterialAttributeBlendFunction BlendFunction /*= nullptr*/)
 {
 	checkf(!AttributeMap.Contains(Property), TEXT("Tried to add duplicate material property."));
-	AttributeMap.Add(Property, FMaterialAttributeDefintion(AttributeID, DisplayName, Property, ValueType, DefaultValue, ShaderFrequency, TexCoordIndex, bIsHidden, BlendFunction));
+	AttributeMap.Add(Property, FMaterialAttributeDefintion(AttributeID, AttributeName, Property, ValueType, DefaultValue, ShaderFrequency, TexCoordIndex, bIsHidden, BlendFunction));
 	if (!bIsHidden)
 	{
 		OrderedVisibleAttributeList.Add(AttributeID);
@@ -3709,6 +3712,120 @@ FMaterialAttributeDefintion* FMaterialAttributeDefinitionMap::Find(EMaterialProp
 	return Find(MP_MAX);
 }
 
+FText FMaterialAttributeDefinitionMap::GetAttributeOverrideForMaterial(const FGuid& AttributeID, UMaterial* Material)
+{
+	TArray<TKeyValuePair<EMaterialShadingModel, FString>> CustomPinNames;
+	EMaterialProperty Property = GMaterialPropertyAttributesMap.Find(AttributeID)->Property;
+
+	switch (Property)
+	{
+	case MP_EmissiveColor:
+		return Material->IsUIMaterial() ? LOCTEXT("UIOutputColor", "Final Color") : LOCTEXT("EmissiveColor", "Emissive Color");
+	case MP_Opacity:
+		return Material->MaterialDomain == MD_Volume ? LOCTEXT("Extinction", "Extinction") : LOCTEXT("Opacity", "Opacity");
+	case MP_OpacityMask:
+		return LOCTEXT("OpacityMask", "Opacity Mask");
+	case MP_DiffuseColor:
+		return LOCTEXT("DiffuseColor", "Diffuse Color");
+	case MP_SpecularColor:
+		return LOCTEXT("SpecularColor", "Specular Color");
+	case MP_BaseColor:
+		return Material->MaterialDomain == MD_Volume ? LOCTEXT("Albedo", "Albedo") : LOCTEXT("BaseColor", "Base Color");
+	case MP_Metallic:
+		CustomPinNames.Add({MSM_Hair, "Scatter"});
+		return FText::FromString(GetPinNameFromShadingModelField(Material->GetShadingModels(), CustomPinNames, "Metallic"));
+	case MP_Specular:
+		return LOCTEXT("Specular", "Specular");
+	case MP_Roughness:
+		return LOCTEXT("Roughness", "Roughness");
+	case MP_Normal:
+		CustomPinNames.Add({MSM_Hair, "Tangent"});
+		return FText::FromString(GetPinNameFromShadingModelField(Material->GetShadingModels(), CustomPinNames, "Normal"));
+	case MP_WorldPositionOffset:
+		return Material->IsUIMaterial() ? LOCTEXT("ScreenPosition", "Screen Position") : LOCTEXT("WorldPositionOffset", "World Position Offset");
+	case MP_WorldDisplacement:
+		return LOCTEXT("WorldDisplacement", "World Displacement");
+	case MP_TessellationMultiplier:
+		return LOCTEXT("TessellationMultiplier", "Tessellation Multiplier");
+	case MP_SubsurfaceColor:
+		CustomPinNames.Add({MSM_Cloth, "Fuzz Color"});
+		return FText::FromString(GetPinNameFromShadingModelField(Material->GetShadingModels(), CustomPinNames, "Subsurface Color"));
+	case MP_CustomData0:	
+		CustomPinNames.Add({ MSM_ClearCoat, "Clear Coat" });
+		CustomPinNames.Add({MSM_Hair, "Backlit"});
+		CustomPinNames.Add({MSM_Cloth, "Cloth"});
+		CustomPinNames.Add({MSM_Eye, "Iris Mask"});
+		return FText::FromString(GetPinNameFromShadingModelField(Material->GetShadingModels(), CustomPinNames, "Custom Data 0"));
+	case MP_CustomData1:
+		CustomPinNames.Add({ MSM_ClearCoat, "Clear Coat Roughness" });
+		CustomPinNames.Add({MSM_Eye, "Iris Distance"});
+		return FText::FromString(GetPinNameFromShadingModelField(Material->GetShadingModels(), CustomPinNames, "Custom Data 1"));
+	case MP_AmbientOcclusion:
+		return LOCTEXT("AmbientOcclusion", "Ambient Occlusion");
+	case MP_Refraction:
+		return LOCTEXT("Refraction", "Refraction");
+	case MP_CustomizedUVs0:
+		return LOCTEXT("CustomizedUV0", "Customized UV 0");
+	case MP_CustomizedUVs1:
+		return LOCTEXT("CustomizedUV1", "Customized UV 1");
+	case MP_CustomizedUVs2:
+		return LOCTEXT("CustomizedUV2", "Customized UV 2");
+	case MP_CustomizedUVs3:
+		return LOCTEXT("CustomizedUV3", "Customized UV 3");
+	case MP_CustomizedUVs4:
+		return LOCTEXT("CustomizedUV4", "Customized UV 4");
+	case MP_CustomizedUVs5:
+		return LOCTEXT("CustomizedUV5", "Customized UV 5");
+	case MP_CustomizedUVs6:
+		return LOCTEXT("CustomizedUV6", "Customized UV 6");
+	case MP_CustomizedUVs7:
+		return LOCTEXT("CustomizedUV7", "Customized UV 7");
+	case MP_PixelDepthOffset:
+		return LOCTEXT("PixelDepthOffset", "Pixel Depth Offset");
+	case MP_ShadingModel:
+		return LOCTEXT("ShadingModel", "Shading Model");
+	case MP_CustomOutput:
+		return FText::FromString(GetAttributeName(AttributeID));
+		
+	}
+	return  LOCTEXT("Missing", "Missing");
+}
+
+FString FMaterialAttributeDefinitionMap::GetPinNameFromShadingModelField(FMaterialShadingModelField InShadingModels, const TArray<TKeyValuePair<EMaterialShadingModel, FString>>& InCustomShadingModelPinNames, const FString& InDefaultPinName) 
+{
+	FString OutPinName;
+	for (const TKeyValuePair<EMaterialShadingModel, FString>& CustomShadingModelPinName : InCustomShadingModelPinNames)
+	{
+		if (InShadingModels.HasShadingModel(CustomShadingModelPinName.Key))
+		{
+			// Add delimiter
+			if (!OutPinName.IsEmpty())
+			{
+				OutPinName.Append(" or ");
+			}
+
+			// Append the name and remove the shading model from the temp field
+			OutPinName.Append(CustomShadingModelPinName.Value);
+			InShadingModels.RemoveShadingModel(CustomShadingModelPinName.Key);
+		}
+	}
+
+	// There are other shading models present, these don't have their own specific name for this pin, so use a default one
+	if (InShadingModels.CountShadingModels() != 0)
+	{
+		// Add delimiter
+		if (!OutPinName.IsEmpty())
+		{
+			OutPinName.Append(" or ");
+		}
+
+		OutPinName.Append(InDefaultPinName);
+	}
+
+	ensure(!OutPinName.IsEmpty());
+	return OutPinName;
+}
+
 void FMaterialAttributeDefinitionMap::AppendDDCKeyString(FString& String)
 {
 	FString& DDCString = GMaterialPropertyAttributesMap.AttributeDDCString;
@@ -3743,18 +3860,18 @@ void FMaterialAttributeDefinitionMap::AppendDDCKeyString(FString& String)
 	String.Append(DDCString);
 }
 
-void FMaterialAttributeDefinitionMap::AddCustomAttribute(const FGuid& AttributeID, const FString& DisplayName, const FString& FunctionName, EMaterialValueType ValueType, const FVector4& DefaultValue, MaterialAttributeBlendFunction BlendFunction /*= nullptr*/)
+void FMaterialAttributeDefinitionMap::AddCustomAttribute(const FGuid& AttributeID, const FString& AttributeName, const FString& FunctionName, EMaterialValueType ValueType, const FVector4& DefaultValue, MaterialAttributeBlendFunction BlendFunction /*= nullptr*/)
 {
 	// Make sure that we init CustomAttributes before DDCString is initialized (before first shader load)
 	check(GMaterialPropertyAttributesMap.AttributeDDCString.Len() == 0);
 
-	FMaterialCustomOutputAttributeDefintion UserAttribute(AttributeID, DisplayName, FunctionName, MP_CustomOutput, ValueType, DefaultValue, SF_Pixel, BlendFunction);
+	FMaterialCustomOutputAttributeDefintion UserAttribute(AttributeID, AttributeName, FunctionName, MP_CustomOutput, ValueType, DefaultValue, SF_Pixel, BlendFunction);
 #if DO_CHECK
 	for (auto& Attribute : GMaterialPropertyAttributesMap.AttributeMap)
 	{
-		checkf(Attribute.Value.AttributeID != AttributeID, TEXT("Tried to add duplicate custom output attribute (%s) already in base attributes (%s)."), *DisplayName, *(Attribute.Value.DisplayName));
+		checkf(Attribute.Value.AttributeID != AttributeID, TEXT("Tried to add duplicate custom output attribute (%s) already in base attributes (%s)."), *AttributeName, *(Attribute.Value.AttributeName));
 	}
-	checkf(!GMaterialPropertyAttributesMap.CustomAttributes.Contains(UserAttribute), TEXT("Tried to add duplicate custom output attribute (%s)."), *DisplayName);
+	checkf(!GMaterialPropertyAttributesMap.CustomAttributes.Contains(UserAttribute), TEXT("Tried to add duplicate custom output attribute (%s)."), *AttributeName);
 #endif
 	GMaterialPropertyAttributesMap.CustomAttributes.Add(UserAttribute);
 
@@ -3773,13 +3890,13 @@ void FMaterialAttributeDefinitionMap::GetCustomAttributeList(TArray<FMaterialCus
 	}
 }
 
-void FMaterialAttributeDefinitionMap::GetDisplayNameToIDList(TArray<TPair<FString, FGuid>>& NameToIDList)
+void FMaterialAttributeDefinitionMap::GetAttributeNameToIDList(TArray<TPair<FString, FGuid>>& NameToIDList)
 {
 	NameToIDList.Empty(GMaterialPropertyAttributesMap.OrderedVisibleAttributeList.Num());
 	for (const FGuid& AttributeID : GMaterialPropertyAttributesMap.OrderedVisibleAttributeList)
 	{
 		FMaterialAttributeDefintion* Attribute = GMaterialPropertyAttributesMap.Find(AttributeID);
-		NameToIDList.Emplace(Attribute->DisplayName, AttributeID);
+		NameToIDList.Emplace(Attribute->AttributeName, AttributeID);
 	}
 }
 
@@ -4008,3 +4125,5 @@ void SetShaderMapsOnMaterialResources(const TMap<FMaterial*, FMaterialShaderMap*
 		SetShaderMapsOnMaterialResources_RenderThread(RHICmdList, InMaterialsToUpdate);
 	});
 }
+
+#undef LOCTEXT_NAMESPACE
\ No newline at end of file
diff --git a/Engine/Source/Runtime/Engine/Private/Rendering/SkeletalMeshLODImporterData.cpp b/Engine/Source/Runtime/Engine/Private/Rendering/SkeletalMeshLODImporterData.cpp
--- a/Engine/Source/Runtime/Engine/Private/Rendering/SkeletalMeshLODImporterData.cpp
+++ b/Engine/Source/Runtime/Engine/Private/Rendering/SkeletalMeshLODImporterData.cpp
@@ -15,7 +15,9 @@ DEFINE_LOG_CATEGORY_STATIC(LogSkeletalMeshLODImporterData, Log, All);
 
 void FSkeletalMeshImportData::CopyDataNeedByMorphTargetImport(FSkeletalMeshImportData& Other) const
 {
+	//The points array is the only data we need to compute the morph target in the skeletalmesh build
 	Other.Points = Points;
+	//PointToRawMap should not be save when saving morph target data, we only need it temporary to gather the point from the fbx shape
 	Other.PointToRawMap = PointToRawMap;
 	Other.bDiffPose = bDiffPose;
 	Other.bUseT0AsRefPose = bUseT0AsRefPose;
@@ -135,7 +137,7 @@ bool FSkeletalMeshImportData::ReplaceSkeletalMeshGeometryImportData(const USkele
 
 	//Load the original skeletal mesh import data
 	FSkeletalMeshImportData OriginalSkeletalMeshImportData;
-	SkeletalMeshLODModel.RawSkeletalMeshBulkData.LoadRawMesh(OriginalSkeletalMeshImportData);
+	SkeletalMesh->LoadLODImportedData(LodIndex, OriginalSkeletalMeshImportData);
 
 	//Backup the new geometry and rig to be able to apply the rig to the old geometry
 	FSkeletalMeshImportData NewGeometryAndRigData = *ImportData;
@@ -184,7 +186,7 @@ bool FSkeletalMeshImportData::ReplaceSkeletalMeshRigImportData(const USkeletalMe
 
 	//Load the original skeletal mesh import data
 	FSkeletalMeshImportData OriginalSkeletalMeshImportData;
-	SkeletalMeshLODModel.RawSkeletalMeshBulkData.LoadRawMesh(OriginalSkeletalMeshImportData);
+	SkeletalMesh->LoadLODImportedData(LodIndex, OriginalSkeletalMeshImportData);
 
 	ImportData->bDiffPose = OriginalSkeletalMeshImportData.bDiffPose;
 	ImportData->bUseT0AsRefPose = OriginalSkeletalMeshImportData.bUseT0AsRefPose;
@@ -652,6 +654,34 @@ FArchive& operator<<(FArchive& Ar, FSkeletalMeshImportData& RawMesh)
 	return Ar;
 }
 
+void FRawSkeletalMeshBulkData::Serialize(FArchive& Ar, TArray<FRawSkeletalMeshBulkData>& RawSkeltalMeshBulkDatas, UObject* Owner)
+{
+	Ar.CountBytes(RawSkeltalMeshBulkDatas.Num() * sizeof(FRawSkeletalMeshBulkData), RawSkeltalMeshBulkDatas.Num() * sizeof(FRawSkeletalMeshBulkData));
+	if (Ar.IsLoading())
+	{
+		// Load array.
+		int32 NewNum;
+		Ar << NewNum;
+		RawSkeltalMeshBulkDatas.Empty(NewNum);
+		for (int32 Index = 0; Index < NewNum; Index++)
+		{
+			int32 NewEntryIndex = RawSkeltalMeshBulkDatas.AddDefaulted();
+			check(NewEntryIndex == Index);
+			RawSkeltalMeshBulkDatas[Index].Serialize(Ar, Owner);
+		}
+	}
+	else
+	{
+		// Save array.
+		int32 Num = RawSkeltalMeshBulkDatas.Num();
+		Ar << Num;
+		for (int32 Index = 0; Index < Num; Index++)
+		{
+			RawSkeltalMeshBulkDatas[Index].Serialize(Ar, Owner);
+		}
+	}
+}
+
 void FRawSkeletalMeshBulkData::Serialize(FArchive& Ar, UObject* Owner)
 {
 	if (Ar.IsTransacting())
diff --git a/Engine/Source/Runtime/Engine/Private/SceneView.cpp b/Engine/Source/Runtime/Engine/Private/SceneView.cpp
--- a/Engine/Source/Runtime/Engine/Private/SceneView.cpp
+++ b/Engine/Source/Runtime/Engine/Private/SceneView.cpp
@@ -635,6 +635,7 @@ FSceneView::FSceneView(const FSceneViewInitOptions& InitOptions)
 	, SpecularOverrideParameter(FVector4(0,0,0,1))
 	, NormalOverrideParameter(FVector4(0,0,0,1))
 	, RoughnessOverrideParameter(FVector2D(0,1))
+	, MaterialTextureMipBias(0.f)
 	, HiddenPrimitives(InitOptions.HiddenPrimitives)
 	, ShowOnlyPrimitives(InitOptions.ShowOnlyPrimitives)
 	, OriginOffsetThisFrame(InitOptions.OriginOffsetThisFrame)
@@ -647,6 +648,7 @@ FSceneView::FSceneView(const FSceneViewInitOptions& InitOptions)
 	, bIsSceneCapture(false)
 	, bIsReflectionCapture(false)
 	, bIsPlanarReflection(false)
+	, bIsOfflineRender(false)
 	, bRenderSceneTwoSided(false)
 	, bIsLocked(false)
 	, bStaticSceneOnly(false)
@@ -655,7 +657,9 @@ FSceneView::FSceneView(const FSceneViewInitOptions& InitOptions)
 	, bIsMobileMultiViewEnabled(false)
 	, bIsMobileMultiViewDirectEnabled(false)
 	, bShouldBindInstancedViewUB(false)
+	, bForceCameraVisibilityReset(false)
 	, GlobalClippingPlane(FPlane(0, 0, 0, 0))
+	, LensPrincipalPointOffsetScale(0.0f, 0.0f, 1.0f, 1.0f)
 #if WITH_EDITOR
 	, OverrideLODViewOrigin(InitOptions.OverrideLODViewOrigin)
 	, bAllowTranslucentPrimitivesInHitProxy( true )
diff --git a/Engine/Source/Runtime/Engine/Private/SkeletalMesh.cpp b/Engine/Source/Runtime/Engine/Private/SkeletalMesh.cpp
--- a/Engine/Source/Runtime/Engine/Private/SkeletalMesh.cpp
+++ b/Engine/Source/Runtime/Engine/Private/SkeletalMesh.cpp
@@ -57,9 +57,11 @@
 #include "HAL/FileManager.h"
 
 #if WITH_EDITOR
+#include "Async/ParallelFor.h"
 #include "Rendering/SkeletalMeshModel.h"
 #include "Rendering/SkeletalMeshLODImporterData.h"
 #include "MeshUtilities.h"
+#include "Engine/SkeletalMeshEditorData.h"
 
 #if WITH_APEX_CLOTHING
 #include "ApexClothingUtils.h"
@@ -1675,6 +1677,61 @@ void USkeletalMesh::GetPreloadDependencies(TArray<UObject*>& OutDeps)
 	OutDeps.Add(Skeleton);
 }
 
+void USkeletalMesh::PostDuplicate(bool bDuplicateForPIE)
+{
+	Super::PostDuplicate(bDuplicateForPIE);
+#if WITH_EDITORONLY_DATA
+	//We must not use the same asset for the imported data and duplicate it in the destination skeletalmesh package
+	if (MeshEditorDataObject)
+	{
+		USkeletalMeshEditorData* SourceMeshEditorData = MeshEditorDataObject;
+		check(SourceMeshEditorData);
+		MeshEditorDataObject = nullptr;
+		FSkeletalMeshModel& ImportedModels = *GetImportedModel();
+		//allocate the necessary data before going multi thread
+		ReserveLODImportData(ImportedModels.LODModels.Num() - 1);
+		USkeletalMeshEditorData& DestMeshEditorData = GetMeshEditorData();
+		//We should have a brand new created asset
+		check(MeshEditorDataObject);
+		//Lets duplicate the data
+		ParallelFor(ImportedModels.LODModels.Num(), [&ImportedModels, &SourceMeshEditorData, &DestMeshEditorData](int32 LODIndex)
+		{
+			if (!SourceMeshEditorData->IsLODImportDataValid(LODIndex) || SourceMeshEditorData->GetLODImportedData(LODIndex).IsEmpty())
+			{
+				return;
+			}
+			FSkeletalMeshLODModel& ThisLODModel = ImportedModels.LODModels[LODIndex];
+			FRawSkeletalMeshBulkData& SourceRawSkeletalMeshBulkData = SourceMeshEditorData->GetLODImportedData(LODIndex);
+			FSkeletalMeshImportData RawMesh;
+			SourceRawSkeletalMeshBulkData.LoadRawMesh(RawMesh);
+			//This will create the asset into the correct package and assign the transient pointer
+			FRawSkeletalMeshBulkData& RawSkeletalMeshBulkData = DestMeshEditorData.GetLODImportedData(LODIndex);
+			RawSkeletalMeshBulkData.SaveRawMesh(RawMesh);
+			RawSkeletalMeshBulkData.GeoImportVersion = SourceRawSkeletalMeshBulkData.GeoImportVersion;
+			RawSkeletalMeshBulkData.SkinningImportVersion = SourceRawSkeletalMeshBulkData.SkinningImportVersion;
+			//Set all the cache data to avoid loading the asset if we do not edit the asset
+			ThisLODModel.bIsBuildDataAvailable = RawSkeletalMeshBulkData.IsBuildDataAvailable();
+			ThisLODModel.bIsRawSkeletalMeshBulkDataEmpty = RawSkeletalMeshBulkData.IsEmpty();
+			ThisLODModel.BuildStringID = RawSkeletalMeshBulkData.GetIdString();
+		});
+	}
+#endif //WITH_EDITORONLY_DATA
+}
+
+void USkeletalMesh::PostRename(UObject* OldOuter, const FName OldName)
+{
+	Super::PostRename(OldOuter, OldName);
+#if WITH_EDITORONLY_DATA
+	//We must not use the same asset for the imported data and duplicate it in the destination skeletalmesh package
+	if (MeshEditorDataObject)
+	{
+		FString MeshEditorDataString = GetName() + TEXT("_USkeletalMeshEditorData");
+		//Do a soft rename avoid: dirty, redirector, transaction and reset of the loaders
+		MeshEditorDataObject->Rename(*MeshEditorDataString, GetOutermost(), (REN_DoNotDirty | REN_DontCreateRedirectors | REN_NonTransactional | REN_ForceNoResetLoaders));
+	}
+#endif //WITH_EDITORONLY_DATA
+}
+
 void USkeletalMesh::FlushRenderState()
 {
 	//TComponentReregisterContext<USkeletalMeshComponent> ReregisterContext;
@@ -2030,6 +2087,106 @@ void USkeletalMesh::RemoveLegacyClothingSections()
 	}
 }
 
+USkeletalMeshEditorData& USkeletalMesh::GetMeshEditorData() const
+{
+	if (MeshEditorDataObject == nullptr)
+	{
+		FString MeshEditorDataString = GetName() + TEXT("_USkeletalMeshEditorData");
+		FName MeshEditorDataName = *MeshEditorDataString;
+		//We can have only one USkeletalMeshEditorData asset per skeletalmesh. Find if there is already one existing 
+		UObject* ExistingObject = StaticFindObjectFast(USkeletalMeshEditorData::StaticClass(), GetOutermost(), MeshEditorDataName, true, false, RF_NoFlags, EInternalObjectFlags::PendingKill);
+		if (ExistingObject)
+		{
+			MeshEditorDataObject = Cast<USkeletalMeshEditorData>(ExistingObject);
+		}
+		if (MeshEditorDataObject == nullptr)
+		{
+			//The asset is created in the skeletalmesh package. We keep it private so the user cannot see it in the content browser
+			//StandAlone make sure the asset is save when we save the package(i.e. the skeletalmesh)
+			MeshEditorDataObject = NewObject<USkeletalMeshEditorData>(GetOutermost(), MeshEditorDataName, RF_Standalone);
+		}
+	}
+	//Make sure we have a valid pointer
+	check(MeshEditorDataObject != nullptr);
+	return *MeshEditorDataObject;
+}
+
+void USkeletalMesh::LoadLODImportedData(const int32 LODIndex, FSkeletalMeshImportData& OutMesh) const
+{
+	GetMeshEditorData().GetLODImportedData(LODIndex).LoadRawMesh(OutMesh);
+}
+
+void USkeletalMesh::SaveLODImportedData(const int32 LODIndex, FSkeletalMeshImportData& InMesh)
+{
+	FRawSkeletalMeshBulkData& RawSkeletalMeshBulkData = GetMeshEditorData().GetLODImportedData(LODIndex);
+	RawSkeletalMeshBulkData.SaveRawMesh(InMesh);
+	//Update the cache
+	check(ImportedModel->LODModels.IsValidIndex(LODIndex));
+	ImportedModel->LODModels[LODIndex].RawSkeletalMeshBulkDataID = RawSkeletalMeshBulkData.GetIdString();
+	ImportedModel->LODModels[LODIndex].bIsBuildDataAvailable = RawSkeletalMeshBulkData.IsBuildDataAvailable();
+	ImportedModel->LODModels[LODIndex].bIsRawSkeletalMeshBulkDataEmpty = RawSkeletalMeshBulkData.IsEmpty();
+}
+
+bool USkeletalMesh::IsLODImportedDataBuildAvailable(const int32 LODIndex) const
+{
+	if (!ImportedModel->LODModels.IsValidIndex(LODIndex))
+	{
+		return false;
+	}
+	return ImportedModel->LODModels[LODIndex].bIsBuildDataAvailable;
+}
+
+bool USkeletalMesh::IsLODImportedDataEmpty(const int32 LODIndex) const
+{
+	if (!ImportedModel->LODModels.IsValidIndex(LODIndex))
+	{
+		return false;
+	}
+	return ImportedModel->LODModels[LODIndex].bIsRawSkeletalMeshBulkDataEmpty;
+}
+
+void USkeletalMesh::GetLODImportedDataVersions(const int32 LODIndex, ESkeletalMeshGeoImportVersions& OutGeoImportVersion, ESkeletalMeshSkinningImportVersions& OutSkinningImportVersion) const
+{
+	const FRawSkeletalMeshBulkData& RawSkeletalMeshBulkData = GetMeshEditorData().GetLODImportedData(LODIndex);
+	OutGeoImportVersion = RawSkeletalMeshBulkData.GeoImportVersion;
+	OutSkinningImportVersion = RawSkeletalMeshBulkData.SkinningImportVersion;
+}
+
+void USkeletalMesh::SetLODImportedDataVersions(const int32 LODIndex, const ESkeletalMeshGeoImportVersions& InGeoImportVersion, const ESkeletalMeshSkinningImportVersions& InSkinningImportVersion)
+{
+	FRawSkeletalMeshBulkData& RawSkeletalMeshBulkData = GetMeshEditorData().GetLODImportedData(LODIndex);
+	RawSkeletalMeshBulkData.GeoImportVersion = InGeoImportVersion;
+	RawSkeletalMeshBulkData.SkinningImportVersion = InSkinningImportVersion;
+	//Update the cache
+	check(ImportedModel->LODModels.IsValidIndex(LODIndex));
+	ImportedModel->LODModels[LODIndex].RawSkeletalMeshBulkDataID = RawSkeletalMeshBulkData.GetIdString();
+	ImportedModel->LODModels[LODIndex].bIsBuildDataAvailable = RawSkeletalMeshBulkData.IsBuildDataAvailable();
+	ImportedModel->LODModels[LODIndex].bIsRawSkeletalMeshBulkDataEmpty = RawSkeletalMeshBulkData.IsEmpty();
+}
+
+void USkeletalMesh::CopyImportedData(int32 SrcLODIndex, USkeletalMesh* SrcSkeletalMesh, int32 DestLODIndex, USkeletalMesh* DestSkeletalMesh)
+{
+	check(DestSkeletalMesh->ImportedModel->LODModels.IsValidIndex(DestLODIndex));
+	FRawSkeletalMeshBulkData& SrcRawMesh = SrcSkeletalMesh->GetMeshEditorData().GetLODImportedData(SrcLODIndex);
+	FRawSkeletalMeshBulkData& DestRawMesh = DestSkeletalMesh->GetMeshEditorData().GetLODImportedData(DestLODIndex);
+	FSkeletalMeshImportData SrcImportData;
+	SrcRawMesh.LoadRawMesh(SrcImportData);
+	DestRawMesh.SaveRawMesh(SrcImportData);
+	DestRawMesh.GeoImportVersion = SrcRawMesh.GeoImportVersion;
+	DestRawMesh.SkinningImportVersion = SrcRawMesh.SkinningImportVersion;
+	
+	FSkeletalMeshLODModel& DestLODModel = DestSkeletalMesh->ImportedModel->LODModels[DestLODIndex];
+	DestLODModel.RawSkeletalMeshBulkDataID = DestRawMesh.GetIdString();
+	DestLODModel.bIsBuildDataAvailable = DestRawMesh.IsBuildDataAvailable();
+	DestLODModel.bIsRawSkeletalMeshBulkDataEmpty = DestRawMesh.IsEmpty();
+}
+
+void USkeletalMesh::ReserveLODImportData(int32 MaxLODIndex)
+{
+	//Getting the LODImportedData will allocate the data to default value.
+	GetMeshEditorData().GetLODImportedData(MaxLODIndex);
+}
+
 #endif // WITH_EDITOR
 
 void USkeletalMesh::PostLoad()
@@ -2135,6 +2292,39 @@ void USkeletalMesh::PostLoad()
 
 		UpdateGenerateUpToData();
 
+		if (GetLinkerCustomVersion(FEditorObjectVersion::GUID) < FEditorObjectVersion::SkeletalMeshMoveEditorSourceDataToPrivateAsset)
+		{
+			ReserveLODImportData(ImportedModel->LODModels.Num() - 1);
+			for (int32 LODIndex = 0; LODIndex < ImportedModel->LODModels.Num(); ++LODIndex)
+			{
+				FSkeletalMeshLODModel& ThisLODModel = ImportedModel->LODModels[LODIndex];
+				//We can have partial data if the asset was save after the split workflow implementation
+				//Use the deprecated member to retrieve this data
+				if (GetLinkerCustomVersion(FFortniteMainBranchObjectVersion::GUID) >= FFortniteMainBranchObjectVersion::NewSkeletalMeshImporterWorkflow)
+				{
+					if (!ThisLODModel.RawSkeletalMeshBulkData_DEPRECATED.IsEmpty())
+					{
+						FSkeletalMeshImportData SerializeMeshData;
+						ThisLODModel.RawSkeletalMeshBulkData_DEPRECATED.LoadRawMesh(SerializeMeshData);
+						SaveLODImportedData(LODIndex, SerializeMeshData);
+					}
+					//Get the FRawSkeletalMeshBulkData to set the geo and skinning version
+					FRawSkeletalMeshBulkData& RawSkeletalMeshBulkData = GetMeshEditorData().GetLODImportedData(LODIndex);
+					RawSkeletalMeshBulkData.GeoImportVersion = ThisLODModel.RawSkeletalMeshBulkData_DEPRECATED.GeoImportVersion;
+					RawSkeletalMeshBulkData.SkinningImportVersion = ThisLODModel.RawSkeletalMeshBulkData_DEPRECATED.SkinningImportVersion;
+					//Empty the DEPRECATED member
+					FSkeletalMeshImportData EmptyMeshData;
+					ThisLODModel.RawSkeletalMeshBulkData_DEPRECATED.SaveRawMesh(EmptyMeshData);
+					ThisLODModel.RawSkeletalMeshBulkData_DEPRECATED.EmptyBulkData();
+				}
+				//Set the cache data into the LODModel
+				FRawSkeletalMeshBulkData& RawSkeletalMeshBulkData = GetMeshEditorData().GetLODImportedData(LODIndex);
+				ThisLODModel.bIsRawSkeletalMeshBulkDataEmpty = RawSkeletalMeshBulkData.IsEmpty();
+				ThisLODModel.bIsBuildDataAvailable = RawSkeletalMeshBulkData.IsBuildDataAvailable();
+				ThisLODModel.RawSkeletalMeshBulkDataID = RawSkeletalMeshBulkData.GetIdString();
+			}
+		}
+
 		if (GetLinkerCustomVersion(FEditorObjectVersion::GUID) < FEditorObjectVersion::SkeletalMeshBuildRefactor)
 		{
 			//We want to avoid changing the ddc if we load an old asset.
@@ -2147,9 +2337,8 @@ void USkeletalMesh::PostLoad()
 				FSkeletalMeshLODModel& ThisLODModel = ImportedModel->LODModels[LodIndex];
 				FSkeletalMeshLODInfo* ThisLODInfo = GetLODInfo(LodIndex);
 				check(ThisLODInfo);
-				
 				//Reset the reduction setting to a non active state if the asset has active reduction but have no RawSkeletalMeshBulkData (we cannot reduce it)
-				if (IsReductionActive(LodIndex) && !ThisLODInfo->bHasBeenSimplified && ThisLODModel.RawSkeletalMeshBulkData.IsEmpty())
+				if (IsReductionActive(LodIndex) && !ThisLODInfo->bHasBeenSimplified && IsLODImportedDataEmpty(LodIndex))
 				{
 					if (LodIndex > ThisLODInfo->ReductionSettings.BaseLOD)
 					{
@@ -2178,7 +2367,10 @@ void USkeletalMesh::PostLoad()
 						}
 					}
 				}
-				ThisLODModel.UpdateChunkedSectionInfo(GetName(), ThisLODInfo->LODMaterialMap);
+				if (!IsLODImportedDataBuildAvailable(LodIndex))
+				{
+					ThisLODModel.UpdateChunkedSectionInfo(GetName(), ThisLODInfo->LODMaterialMap);
+				}
 			}
 		}
 
@@ -3814,6 +4006,7 @@ static void SerializeBuildSettingsForDDC(FArchive& Ar, FSkeletalMeshBuildSetting
 	Ar << BuildSettings.ThresholdPosition;
 	Ar << BuildSettings.ThresholdTangentNormal;
 	Ar << BuildSettings.ThresholdUV;
+	Ar << BuildSettings.MorphThresholdPosition;
 }
 
 
diff --git a/Engine/Source/Runtime/Engine/Private/SkeletalMeshEditorData.cpp b/Engine/Source/Runtime/Engine/Private/SkeletalMeshEditorData.cpp
new file mode 100644
--- /dev/null
+++ b/Engine/Source/Runtime/Engine/Private/SkeletalMeshEditorData.cpp
@@ -0,0 +1,60 @@
+// Copyright 1998-2019 Epic Games, Inc. All Rights Reserved.
+#include "Engine/SkeletalMeshEditorData.h"
+
+#include "CoreMinimal.h"
+#include "UObject/ObjectMacros.h"
+#include "UObject/Object.h"
+#include "Templates/SubclassOf.h"
+#include "EngineDefines.h"
+#include "Engine/EngineTypes.h"
+
+#if WITH_EDITORONLY_DATA
+
+#include "Rendering/SkeletalMeshLODImporterData.h"
+
+#endif //WITH_EDITORONLY_DATA
+
+#define LOCTEXT_NAMESPACE "SkeltalMeshEditorData"
+
+DEFINE_LOG_CATEGORY_STATIC(LogSkeletalMeshEditorData, Log, All);
+
+USkeletalMeshEditorData::USkeletalMeshEditorData(const FObjectInitializer& ObjectInitializer)
+	: Super(ObjectInitializer)
+{}
+
+#if WITH_EDITORONLY_DATA
+
+FRawSkeletalMeshBulkData& USkeletalMeshEditorData::GetLODImportedData(int32 LODIndex)
+{
+
+	//Create missing default item
+	check(LODIndex >= 0);
+	if (LODIndex >= RawSkeletalMeshBulkDatas.Num())
+	{
+		//Avoid changing the array outside of the main thread
+		//The allocation must be done before going multi thread
+		//TArray is not thread safe when allocating
+		check(IsInGameThread());
+
+		int32 AddItemCount = 1 + (LODIndex - RawSkeletalMeshBulkDatas.Num());
+		RawSkeletalMeshBulkDatas.AddDefaulted(AddItemCount);
+	}
+	check(RawSkeletalMeshBulkDatas.IsValidIndex(LODIndex));
+	//Return the Data
+	return RawSkeletalMeshBulkDatas[LODIndex];
+}
+
+bool USkeletalMeshEditorData::IsLODImportDataValid(int32 LODIndex)
+{
+	return RawSkeletalMeshBulkDatas.IsValidIndex(LODIndex);
+}
+
+void USkeletalMeshEditorData::Serialize(FArchive& Ar)
+{
+	Super::Serialize(Ar);
+
+	//Serialize all LODs Raw imported source data
+	FRawSkeletalMeshBulkData::Serialize(Ar, RawSkeletalMeshBulkDatas, this);
+}
+
+#endif //WITH_EDITOR
\ No newline at end of file
diff --git a/Engine/Source/Runtime/Engine/Private/SkeletalMeshLODModel.cpp b/Engine/Source/Runtime/Engine/Private/SkeletalMeshLODModel.cpp
--- a/Engine/Source/Runtime/Engine/Private/SkeletalMeshLODModel.cpp
+++ b/Engine/Source/Runtime/Engine/Private/SkeletalMeshLODModel.cpp
@@ -688,9 +688,20 @@ void FSkeletalMeshLODModel::Serialize(FArchive& Ar, UObject* Owner, int32 Idx)
 	if (!StripFlags.IsEditorDataStripped())
 	{
 		RawPointIndices.Serialize(Ar, Owner);
-		if (Ar.IsSaving() || (Ar.CustomVer(FFortniteMainBranchObjectVersion::GUID) >= FFortniteMainBranchObjectVersion::NewSkeletalMeshImporterWorkflow))
+		if (Ar.IsLoading()
+			&& (Ar.CustomVer(FFortniteMainBranchObjectVersion::GUID) >= FFortniteMainBranchObjectVersion::NewSkeletalMeshImporterWorkflow)
+			&& (Ar.CustomVer(FEditorObjectVersion::GUID) < FEditorObjectVersion::SkeletalMeshMoveEditorSourceDataToPrivateAsset))
 		{
-			RawSkeletalMeshBulkData.Serialize(Ar, Owner);
+			RawSkeletalMeshBulkData_DEPRECATED.Serialize(Ar, Owner);
+			RawSkeletalMeshBulkDataID = RawSkeletalMeshBulkData_DEPRECATED.GetIdString();
+			bIsBuildDataAvailable = RawSkeletalMeshBulkData_DEPRECATED.IsBuildDataAvailable();
+			bIsRawSkeletalMeshBulkDataEmpty = RawSkeletalMeshBulkData_DEPRECATED.IsEmpty();
+		}
+		if (Ar.CustomVer(FEditorObjectVersion::GUID) >= FEditorObjectVersion::SkeletalMeshMoveEditorSourceDataToPrivateAsset)
+		{
+			Ar << RawSkeletalMeshBulkDataID;
+			Ar << bIsBuildDataAvailable;
+			Ar << bIsRawSkeletalMeshBulkDataEmpty;
 		}
 	}
 
@@ -1003,7 +1014,7 @@ FString FSkeletalMeshLODModel::GetLODModelDeriveDataKey() const
 	FMemoryWriter Ar(ByteData, true);
 
 	//Add the bulk data ID (if someone modify the original imported data, this ID will change)
-	FString BulkDatIDString = RawSkeletalMeshBulkData.GetIdString();
+	FString BulkDatIDString = RawSkeletalMeshBulkDataID; //Need to re-assign to tmp var because function is const
 	Ar << BulkDatIDString;
 	int32 UserSectionCount = UserSectionsData.Num();
 	Ar << UserSectionCount;
@@ -1026,12 +1037,6 @@ FString FSkeletalMeshLODModel::GetLODModelDeriveDataKey() const
 
 void FSkeletalMeshLODModel::UpdateChunkedSectionInfo(const FString& SkeletalMeshName, TArray<int32>& LODMaterialMap)
 {
-	//Look if the data is pre skeletalmesh build refactor
-	if (RawSkeletalMeshBulkData.IsBuildDataAvailable())
-	{
-		//If the data is up to date just return and do nothing
-		return;
-	}
 	int32 LODModelSectionNum = Sections.Num();
 	//Fill the ChunkedParentSectionIndex data, we assume that every section using the same material are chunked
 	int32 LastMaterialIndex = INDEX_NONE;
@@ -1087,16 +1092,16 @@ void FSkeletalMeshLODModel::UpdateChunkedSectionInfo(const FString& SkeletalMesh
 
 bool FSkeletalMeshLODModel::CopyStructure(FSkeletalMeshLODModel* Destination, FSkeletalMeshLODModel* Source)
 {
-	if (Source->RawPointIndices.IsLocked() || Source->LegacyRawPointIndices.IsLocked() || Source->RawSkeletalMeshBulkData.GetBulkData().IsLocked())
+	if (Source->RawPointIndices.IsLocked() || Source->LegacyRawPointIndices.IsLocked() || Source->RawSkeletalMeshBulkData_DEPRECATED.GetBulkData().IsLocked())
 	{
 		return false;
 	}
 	// Bulk data arrays need to be locked before a copy can be made.
 	Source->RawPointIndices.Lock(LOCK_READ_ONLY);
 	Source->LegacyRawPointIndices.Lock(LOCK_READ_ONLY);
-	Source->RawSkeletalMeshBulkData.GetBulkData().Lock(LOCK_READ_ONLY);
+	Source->RawSkeletalMeshBulkData_DEPRECATED.GetBulkData().Lock(LOCK_READ_ONLY);
 	*Destination = *Source;
-	Source->RawSkeletalMeshBulkData.GetBulkData().Unlock();
+	Source->RawSkeletalMeshBulkData_DEPRECATED.GetBulkData().Unlock();
 	Source->RawPointIndices.Unlock();
 	Source->LegacyRawPointIndices.Unlock();
 
diff --git a/Engine/Source/Runtime/Engine/Private/SkeletalMeshRenderData.cpp b/Engine/Source/Runtime/Engine/Private/SkeletalMeshRenderData.cpp
--- a/Engine/Source/Runtime/Engine/Private/SkeletalMeshRenderData.cpp
+++ b/Engine/Source/Runtime/Engine/Private/SkeletalMeshRenderData.cpp
@@ -58,8 +58,8 @@ static void SerializeLODInfoForDDC(USkeletalMesh* SkeletalMesh, FString& KeySuff
 // If skeletal mesh derived data needs to be rebuilt (new format, serialization
 // differences, etc.) replace the version GUID below with a new one.
 // In case of merge conflicts with DDC versions, you MUST generate a new GUID
-// and set this new GUID as the version.                                       
-#define SKELETALMESH_DERIVEDDATA_VER TEXT("29B218D50FE444AF8EFC8E4234F5D4E1")
+// and set this new GUID as the version.
+#define SKELETALMESH_DERIVEDDATA_VER TEXT("7208E3E18CB047589DD75CD8E66EA1D6")
 
 static const FString& GetSkeletalMeshDerivedDataVersion()
 {
@@ -151,7 +151,6 @@ void FSkeletalMeshRenderData::Cache(USkeletalMesh* Owner)
 		if (GetDerivedDataCacheRef().GetSynchronous(*DerivedDataKey, DerivedData))
 		{
 			COOK_STAT(Timer.AddHit(DerivedData.Num()));
-			
 
 			FMemoryReader Ar(DerivedData, /*bIsPersistent=*/ true);
 
@@ -163,6 +162,50 @@ void FSkeletalMeshRenderData::Cache(USkeletalMesh* Owner)
 			{
 				FSkeletalMeshModel* SkelMeshModel = Owner->GetImportedModel();
 				check(SkelMeshModel);
+
+				int32 MorphTargetNumber = 0;
+				Ar << MorphTargetNumber;
+				TArray<UMorphTarget*> ToDeleteMorphTargets;
+				ToDeleteMorphTargets.Append(Owner->MorphTargets);
+				Owner->MorphTargets.Empty();
+				//Rebuild the MorphTarget object
+				//We cannot serialize directly the UMorphTarget with a FMemoryArchive. This is not supported.
+				for (int32 MorphTargetIndex = 0; MorphTargetIndex < MorphTargetNumber; ++MorphTargetIndex)
+				{
+					FName MorphTargetName = NAME_None;
+					Ar << MorphTargetName;
+					UMorphTarget* MorphTarget = Cast<UMorphTarget>(StaticFindObjectFast(nullptr, Owner, MorphTargetName));
+					if (!MorphTarget)
+					{
+						MorphTarget = NewObject<UMorphTarget>(Owner, MorphTargetName);
+						check(MorphTarget);
+					}
+					else
+					{
+						ToDeleteMorphTargets.Remove(MorphTarget);
+					}
+					MorphTarget->MorphLODModels.Empty();
+					Owner->MorphTargets.Add(MorphTarget);
+					check(MorphTargetIndex == Owner->MorphTargets.Num() - 1);
+					int32 MorphLODModelNumber = 0;
+					Ar << MorphLODModelNumber;
+					MorphTarget->MorphLODModels.AddDefaulted(MorphLODModelNumber);
+					for (int32 MorphDataIndex = 0; MorphDataIndex < MorphLODModelNumber; ++MorphDataIndex)
+					{
+						Ar << MorphTarget->MorphLODModels[MorphDataIndex];
+					}
+				}
+				//Rebuild the mapping and rehook the curve data
+				Owner->InitMorphTargets();
+				for (int32 DeleteMorphIndex = 0; DeleteMorphIndex < ToDeleteMorphTargets.Num(); ++DeleteMorphIndex)
+				{
+					ToDeleteMorphTargets[DeleteMorphIndex]->BaseSkelMesh = nullptr;
+					ToDeleteMorphTargets[DeleteMorphIndex]->MorphLODModels.Empty();
+					//Move the unused asset in the transient package and mark it pending kill
+					ToDeleteMorphTargets[DeleteMorphIndex]->Rename(nullptr, GetTransientPackage(), REN_ForceNoResetLoaders | REN_DoNotDirty | REN_DontCreateRedirectors | REN_NonTransactional);
+					ToDeleteMorphTargets[DeleteMorphIndex]->MarkPendingKill();
+				}
+
 				//Serialize the LODModel sections since they are dependent on the reduction
 				for (int32 LODIndex = 0; LODIndex < SkelMeshModel->LODModels.Num(); LODIndex++)
 				{
@@ -211,9 +254,11 @@ void FSkeletalMeshRenderData::Cache(USkeletalMesh* Owner)
 				FSkeletalMeshLODModel* LODModel = &(SkelMeshModel->LODModels[LODIndex]);
 				FSkeletalMeshLODInfo* LODInfo = Owner->GetLODInfo(LODIndex);
 				check(LODInfo);
+				bool bRawDataEmpty = Owner->IsLODImportedDataEmpty(LODIndex);
+				bool bRawBuildDataAvailable = Owner->IsLODImportedDataBuildAvailable(LODIndex);
 				//Build the source model before the render data, if we are a purely generated LOD we do not need to be build
 				IMeshBuilderModule& MeshBuilderModule = FModuleManager::Get().LoadModuleChecked<IMeshBuilderModule>("MeshBuilder");
-				if (!LODModel->RawSkeletalMeshBulkData.IsEmpty() && LODModel->RawSkeletalMeshBulkData.IsBuildDataAvailable())
+				if (!bRawDataEmpty && bRawBuildDataAvailable)
 				{
 					MeshBuilderModule.BuildSkeletalMesh(Owner, LODIndex, true);
 					LODModel = &(SkelMeshModel->LODModels[LODIndex]);
@@ -254,6 +299,21 @@ void FSkeletalMeshRenderData::Cache(USkeletalMesh* Owner)
 			//So we do not serialize the LODModel sections.
 			if (!Owner->UseLegacyMeshDerivedDataKey)
 			{
+				int32 MorphTargetNumber = Owner->MorphTargets.Num();
+				Ar << MorphTargetNumber;
+				for (int32 MorphTargetIndex = 0; MorphTargetIndex < MorphTargetNumber; ++MorphTargetIndex)
+				{
+					FName MorphTargetName = Owner->MorphTargets[MorphTargetIndex]->GetFName();
+					Ar << MorphTargetName;
+					int32 MorphLODModelNumber = Owner->MorphTargets[MorphTargetIndex]->MorphLODModels.Num();
+					Ar << MorphLODModelNumber;
+					for (int32 MorphIndex = 0; MorphIndex < MorphLODModelNumber; ++MorphIndex)
+					{
+						Ar << Owner->MorphTargets[MorphTargetIndex]->MorphLODModels[MorphIndex];
+					}
+				}
+				//No need to serialize the morph target mapping since we will rebuild the mapping when loading a ddc
+
 				//Serialize the LODModel sections since they are dependent on the reduction
 				for (int32 LODIndex = 0; LODIndex < SkelMeshModel->LODModels.Num(); LODIndex++)
 				{
diff --git a/Engine/Source/Runtime/Engine/Private/TextureRenderTarget2D.cpp b/Engine/Source/Runtime/Engine/Private/TextureRenderTarget2D.cpp
--- a/Engine/Source/Runtime/Engine/Private/TextureRenderTarget2D.cpp
+++ b/Engine/Source/Runtime/Engine/Private/TextureRenderTarget2D.cpp
@@ -24,6 +24,8 @@ int32 GTextureRenderTarget2DMaxSizeY = 999999999;
 UTextureRenderTarget2D::UTextureRenderTarget2D(const FObjectInitializer& ObjectInitializer)
 	: Super(ObjectInitializer)
 {
+	SizeX = 1;
+	SizeY = 1;
 	bHDR_DEPRECATED = true;
 	RenderTargetFormat = RTF_RGBA16f;
 	bAutoGenerateMips = false;
diff --git a/Engine/Source/Runtime/Engine/Private/WorldComposition.cpp b/Engine/Source/Runtime/Engine/Private/WorldComposition.cpp
--- a/Engine/Source/Runtime/Engine/Private/WorldComposition.cpp
+++ b/Engine/Source/Runtime/Engine/Private/WorldComposition.cpp
@@ -825,7 +825,7 @@ bool UWorldComposition::IsDistanceDependentLevel(FName PackageName) const
 	return IsDistanceDependentLevel(TileIdx);
 }
 
-bool UWorldComposition::CommitTileStreamingState(UWorld* PersistenWorld, int32 TileIdx, bool bShouldBeLoaded, bool bShouldBeVisible, bool bShouldBlock, int32 LODIdx)
+bool UWorldComposition::CommitTileStreamingState(UWorld* PersistentWorld, int32 TileIdx, bool bShouldBeLoaded, bool bShouldBeVisible, bool bShouldBlock, int32 LODIdx)
 {
 	if (!Tiles.IsValidIndex(TileIdx))
 	{
@@ -846,7 +846,7 @@ bool UWorldComposition::CommitTileStreamingState(UWorld* PersistenWorld, int32 T
 	}
 
 	// Quit early in case we have cooldown on streaming state changes
-	const bool bUseStreamingStateCooldown = (PersistenWorld->IsGameWorld() && PersistenWorld->FlushLevelStreamingType == EFlushLevelStreamingType::None);
+	const bool bUseStreamingStateCooldown = (PersistentWorld->IsGameWorld() && PersistentWorld->FlushLevelStreamingType == EFlushLevelStreamingType::None);
 	if (bUseStreamingStateCooldown && TilesStreamingTimeThreshold > 0.0)
 	{
 		const double CurrentTime = FPlatformTime::Seconds();
diff --git a/Engine/Source/Runtime/Engine/Public/MaterialShared.h b/Engine/Source/Runtime/Engine/Public/MaterialShared.h
--- a/Engine/Source/Runtime/Engine/Public/MaterialShared.h
+++ b/Engine/Source/Runtime/Engine/Public/MaterialShared.h
@@ -2415,7 +2415,7 @@ typedef int32 (*MaterialAttributeBlendFunction)(FMaterialCompiler* Compiler, int
 class FMaterialAttributeDefintion
 {
 public:
-	FMaterialAttributeDefintion(const FGuid& InGUID, const FString& InDisplayName, EMaterialProperty InProperty,
+	FMaterialAttributeDefintion(const FGuid& InGUID, const FString& AttributeName, EMaterialProperty InProperty,
 		EMaterialValueType InValueType, const FVector4& InDefaultValue, EShaderFrequency InShaderFrequency,
 		int32 InTexCoordIndex = INDEX_NONE, bool bInIsHidden = false, MaterialAttributeBlendFunction InBlendFunction = nullptr);
 
@@ -2427,7 +2427,7 @@ public:
 	}
 
 	FGuid				AttributeID;
-	FString				DisplayName;
+	FString				AttributeName;
 	EMaterialProperty	Property;	
 	EMaterialValueType	ValueType;
 	FVector4			DefaultValue;
@@ -2447,7 +2447,7 @@ public:
 class FMaterialCustomOutputAttributeDefintion : public FMaterialAttributeDefintion
 {
 public:
-	FMaterialCustomOutputAttributeDefintion(const FGuid& InGUID, const FString& InDisplayName, const FString& InFunctionName, EMaterialProperty InProperty,
+	FMaterialCustomOutputAttributeDefintion(const FGuid& InGUID, const FString& InAttributeName, const FString& InFunctionName, EMaterialProperty InProperty,
 		EMaterialValueType InValueType, const FVector4& InDefaultValue, EShaderFrequency InShaderFrequency, MaterialAttributeBlendFunction InBlendFunction = nullptr);
 
 	bool operator==(const FMaterialCustomOutputAttributeDefintion& Other) const
@@ -2488,17 +2488,31 @@ public:
 	}
 
 	/** Returns the display name of a material attribute */
-	ENGINE_API static FString GetDisplayName(EMaterialProperty Property)
+	ENGINE_API static FString GetAttributeName(EMaterialProperty Property)
 	{
 		FMaterialAttributeDefintion* Attribute = GMaterialPropertyAttributesMap.Find(Property);
-		return Attribute->DisplayName;
+		return Attribute->AttributeName;
 	}
 
 	/** Returns the display name of a material attribute */
-	ENGINE_API static FString GetDisplayName(const FGuid& AttributeID)
+	ENGINE_API static FString GetAttributeName(const FGuid& AttributeID)
 	{
 		FMaterialAttributeDefintion* Attribute = GMaterialPropertyAttributesMap.Find(AttributeID);
-		return Attribute->DisplayName;
+		return Attribute->AttributeName;
+	}
+
+	/** Returns the display name of a material attribute, accounting for overrides based on properties of a given material */
+	ENGINE_API static FText GetDisplayNameForMaterial(EMaterialProperty Property, UMaterial* Material)
+	{
+		FMaterialAttributeDefintion* Attribute = GMaterialPropertyAttributesMap.Find(Property);
+		return GetAttributeOverrideForMaterial(Attribute->AttributeID, Material);
+	}
+
+	/** Returns the display name of a material attribute, accounting for overrides based on properties of a given material */
+	ENGINE_API static FText GetDisplayNameForMaterial(const FGuid& AttributeID, UMaterial* Material)
+	{
+		FMaterialAttributeDefintion* Attribute = GMaterialPropertyAttributesMap.Find(AttributeID);
+		return GetAttributeOverrideForMaterial(AttributeID, Material);
 	}
 
 	/** Returns the value type of a material attribute */
@@ -2563,7 +2577,7 @@ public:
 	ENGINE_API static void AppendDDCKeyString(FString& String);
 
 	/** Appends a new attribute definition to the custom output list */
-	ENGINE_API static void AddCustomAttribute(const FGuid& AttributeID, const FString& DisplayName, const FString& FunctionName, EMaterialValueType ValueType, const FVector4& DefaultValue, MaterialAttributeBlendFunction BlendFunction = nullptr);
+	ENGINE_API static void AddCustomAttribute(const FGuid& AttributeID, const FString& AttributeName, const FString& FunctionName, EMaterialValueType ValueType, const FVector4& DefaultValue, MaterialAttributeBlendFunction BlendFunction = nullptr);
 
 	/** Returns a list of registered custom attributes */
 	ENGINE_API static void GetCustomAttributeList(TArray<FMaterialCustomOutputAttributeDefintion>& CustomAttributeList);
@@ -2578,18 +2592,22 @@ private:
 	friend class FMaterialAttributePropertyDetails;
 
 	/** Returns a list of display names and their associated GUIDs for material properties */
-	ENGINE_API static void GetDisplayNameToIDList(TArray<TPair<FString, FGuid>>& NameToIDList);
+	ENGINE_API static void GetAttributeNameToIDList(TArray<TPair<FString, FGuid>>& NameToIDList);
 
 	// Internal map management
 	void InitializeAttributeMap();
 
-	void Add(const FGuid& AttributeID, const FString& DisplayName, EMaterialProperty Property,
+	void Add(const FGuid& AttributeID, const FString& AttributeName, EMaterialProperty Property,
 		EMaterialValueType ValueType, const FVector4& DefaultValue, EShaderFrequency ShaderFrequency,
 		int32 TexCoordIndex = INDEX_NONE, bool bIsHidden = false, MaterialAttributeBlendFunction BlendFunction = nullptr);
 
 	ENGINE_API FMaterialAttributeDefintion* Find(const FGuid& AttributeID);
 	ENGINE_API FMaterialAttributeDefintion* Find(EMaterialProperty Property);
 
+	// Helper functions to determine display name based on shader model, material domain, etc.
+	static FText GetAttributeOverrideForMaterial(const FGuid& AttributeID, UMaterial* Material);
+	static FString GetPinNameFromShadingModelField(FMaterialShadingModelField InShadingModels, const TArray<TKeyValuePair<EMaterialShadingModel, FString>>& InCustomShadingModelPinNames, const FString& InDefaultPinName);
+
 	ENGINE_API static FMaterialAttributeDefinitionMap GMaterialPropertyAttributesMap;
 
 	TMap<EMaterialProperty, FMaterialAttributeDefintion>	AttributeMap; // Fixed map of compile-time definitions
diff --git a/Engine/Source/Runtime/Engine/Public/Rendering/SkeletalMeshLODImporterData.h b/Engine/Source/Runtime/Engine/Public/Rendering/SkeletalMeshLODImporterData.h
--- a/Engine/Source/Runtime/Engine/Public/Rendering/SkeletalMeshLODImporterData.h
+++ b/Engine/Source/Runtime/Engine/Public/Rendering/SkeletalMeshLODImporterData.h
@@ -427,7 +427,9 @@ public:
 	/** Default constructor. */
 	ENGINE_API FReductionBaseSkeletalMeshBulkData();
 
+	/*Static function helper to serialize an array of reduction data. */
 	ENGINE_API static void Serialize(FArchive& Ar, TArray<FReductionBaseSkeletalMeshBulkData*>& ReductionBaseSkeletalMeshDatas, UObject* Owner);
+
 	/** Serialization. */
 	ENGINE_API void Serialize(class FArchive& Ar, class UObject* Owner);
 
@@ -480,6 +482,9 @@ public:
 	/** Default constructor. */
 	ENGINE_API FRawSkeletalMeshBulkData();
 
+	/*Static function helper to serialize an array of Raw source data. */
+	ENGINE_API static void Serialize(FArchive& Ar, TArray<FRawSkeletalMeshBulkData>& RawSkeltalMeshBulkDatas, UObject* Owner);
+	
 	/** Serialization. */
 	ENGINE_API void Serialize(class FArchive& Ar, class UObject* Owner);
 
@@ -498,7 +503,17 @@ public:
 	ENGINE_API FByteBulkData& GetBulkData();
 	
 	/** Empty the bulk data. */
-	ENGINE_API void EmptyBulkData() { BulkData.RemoveBulkData(); }
+	ENGINE_API void EmptyBulkData()
+	{
+		//Clear all the data
+		BulkData.RemoveBulkData();
+		Guid.Invalidate();
+		bGuidIsHash = false;
+		SerializeLoadingCustomVersionContainer.Empty();
+		bUseSerializeLoadingCustomVersion = false;
+		GeoImportVersion = ESkeletalMeshGeoImportVersions::Before_Versionning;
+		SkinningImportVersion = ESkeletalMeshSkinningImportVersions::Before_Versionning;
+	}
 
 	/** Returns true if no bulk data is available for this mesh. */
 	FORCEINLINE bool IsEmpty() const { return BulkData.GetBulkDataSize() == 0; }
diff --git a/Engine/Source/Runtime/Engine/Public/Rendering/SkeletalMeshLODModel.h b/Engine/Source/Runtime/Engine/Public/Rendering/SkeletalMeshLODModel.h
--- a/Engine/Source/Runtime/Engine/Public/Rendering/SkeletalMeshLODModel.h
+++ b/Engine/Source/Runtime/Engine/Public/Rendering/SkeletalMeshLODModel.h
@@ -313,13 +313,20 @@ public:
 	FWordBulkData				LegacyRawPointIndices;
 
 	/** Imported raw mesh data. Optional, only the imported mesh LOD has this, generated LOD or old asset will be null. */
-	FRawSkeletalMeshBulkData	RawSkeletalMeshBulkData;
+	FRawSkeletalMeshBulkData	RawSkeletalMeshBulkData_DEPRECATED;
+	/** This ID is use to create the DDC key, it must be set when we save the FRawSkeletalMeshBulkData. */
+	FString						RawSkeletalMeshBulkDataID;
+	bool						bIsBuildDataAvailable;
+	bool						bIsRawSkeletalMeshBulkDataEmpty;
 
 	/** Constructor (default) */
 	FSkeletalMeshLODModel()
 		: NumVertices(0)
 		, NumTexCoords(0)
 		, MaxImportVertex(-1)
+		, RawSkeletalMeshBulkDataID(TEXT(""))
+		, bIsBuildDataAvailable(false)
+		, bIsRawSkeletalMeshBulkDataEmpty(true)
 		, BuildStringID(TEXT(""))
 	{
 	}
diff --git a/Engine/Source/Runtime/Experimental/InteractiveToolsFramework/Private/BaseBehaviors/SingleClickBehavior.cpp b/Engine/Source/Runtime/Experimental/InteractiveToolsFramework/Private/BaseBehaviors/SingleClickBehavior.cpp
--- a/Engine/Source/Runtime/Experimental/InteractiveToolsFramework/Private/BaseBehaviors/SingleClickBehavior.cpp
+++ b/Engine/Source/Runtime/Experimental/InteractiveToolsFramework/Private/BaseBehaviors/SingleClickBehavior.cpp
@@ -16,33 +16,36 @@ void USingleClickInputBehavior::Initialize(IClickBehaviorTarget* TargetIn)
 }
 
 
-FInputCaptureRequest USingleClickInputBehavior::WantsCapture(const FInputDeviceState& input)
+FInputCaptureRequest USingleClickInputBehavior::WantsCapture(const FInputDeviceState& Input)
 {
-	if (IsPressed(input) && (ModifierCheckFunc == nullptr || ModifierCheckFunc(input)) )
+	if (IsPressed(Input) && (ModifierCheckFunc == nullptr || ModifierCheckFunc(Input)) )
 	{
-		if ( Target->IsHitByClick(GetDeviceRay(input)) )
+		FInputRayHit HitResult = Target->IsHitByClick(GetDeviceRay(Input));
+		if (HitResult.bHit)
 		{
-			return FInputCaptureRequest::Begin(this, EInputCaptureSide::Any);
+			return FInputCaptureRequest::Begin(this, EInputCaptureSide::Any, HitResult.HitDepth);
 		}
 	}
 	return FInputCaptureRequest::Ignore();
 }
 
 
-FInputCaptureUpdate USingleClickInputBehavior::BeginCapture(const FInputDeviceState& input, EInputCaptureSide eSide)
+FInputCaptureUpdate USingleClickInputBehavior::BeginCapture(const FInputDeviceState& Input, EInputCaptureSide eSide)
 {
+	Modifiers.UpdateModifiers(Input, Target);
 	return FInputCaptureUpdate::Begin(this, EInputCaptureSide::Any);
 }
 
 
-FInputCaptureUpdate USingleClickInputBehavior::UpdateCapture(const FInputDeviceState& input, const FInputCaptureData& data)
+FInputCaptureUpdate USingleClickInputBehavior::UpdateCapture(const FInputDeviceState& Input, const FInputCaptureData& Data)
 {
-	if (IsReleased(input)) 
+	Modifiers.UpdateModifiers(Input, Target);
+	if (IsReleased(Input))
 	{
 		if (HitTestOnRelease == false || 
-			Target->IsHitByClick(GetDeviceRay(input)) )
+			Target->IsHitByClick(GetDeviceRay(Input)).bHit )
 		{
-			Clicked(input, data);
+			Clicked(Input, Data);
 		}
 
 		return FInputCaptureUpdate::End();
diff --git a/Engine/Source/Runtime/Experimental/InteractiveToolsFramework/Private/BaseGizmos/TransformProxy.cpp b/Engine/Source/Runtime/Experimental/InteractiveToolsFramework/Private/BaseGizmos/TransformProxy.cpp
--- a/Engine/Source/Runtime/Experimental/InteractiveToolsFramework/Private/BaseGizmos/TransformProxy.cpp
+++ b/Engine/Source/Runtime/Experimental/InteractiveToolsFramework/Private/BaseGizmos/TransformProxy.cpp
@@ -45,6 +45,18 @@ void UTransformProxy::SetTransform(const FTransform& TransformIn)
 }
 
 
+void UTransformProxy::BeginTransformEditSequence()
+{
+	OnBeginTransformEdit.Broadcast(this);
+}
+
+void UTransformProxy::EndTransformEditSequence()
+{
+	OnEndTransformEdit.Broadcast(this);
+}
+
+
+
 
 
 void UTransformProxy::UpdateObjects()
@@ -151,6 +163,7 @@ void FTransformProxyChangeSource::BeginChange()
 	{
 		ActiveChange = MakeUnique<FTransformProxyChange>();
 		ActiveChange->From = Proxy->GetTransform();
+		Proxy->BeginTransformEditSequence();
 	}
 }
 
@@ -158,6 +171,7 @@ TUniquePtr<FToolCommandChange> FTransformProxyChangeSource::EndChange()
 {
 	if (Proxy.IsValid())
 	{
+		Proxy->EndTransformEditSequence();
 		ActiveChange->To = Proxy->GetTransform();
 		return MoveTemp(ActiveChange);
 	}
diff --git a/Engine/Source/Runtime/Experimental/InteractiveToolsFramework/Private/BaseTools/MeshSurfacePointTool.cpp b/Engine/Source/Runtime/Experimental/InteractiveToolsFramework/Private/BaseTools/MeshSurfacePointTool.cpp
--- a/Engine/Source/Runtime/Experimental/InteractiveToolsFramework/Private/BaseTools/MeshSurfacePointTool.cpp
+++ b/Engine/Source/Runtime/Experimental/InteractiveToolsFramework/Private/BaseTools/MeshSurfacePointTool.cpp
@@ -2,6 +2,7 @@
 
 #include "BaseTools/MeshSurfacePointTool.h"
 #include "BaseBehaviors/MouseHoverBehavior.h"
+#include "BaseBehaviors/ClickDragBehavior.h"
 #include "InteractiveToolManager.h"
 #include "ToolBuilderUtil.h"
 #include "Components/PrimitiveComponent.h"
@@ -35,10 +36,15 @@ void UMeshSurfacePointToolBuilder::InitializeNewTool(UMeshSurfacePointTool* NewT
 	UActorComponent* ActorComponent = ToolBuilderUtil::FindFirstComponent(SceneState, CanMakeComponentTarget);
 	UPrimitiveComponent* MeshComponent = Cast<UPrimitiveComponent>(ActorComponent);
 	check(MeshComponent != nullptr);
+	NewTool->SetStylusAPI(this->StylusAPI);
 	NewTool->SetSelection( MakeComponentTarget(MeshComponent) );
 }
 
 
+static const int UMeshSurfacePointTool_ShiftModifier = 1;
+static const int UMeshSurfacePointTool_CtrlModifier = 2;
+
+
 /*
  * Tool
  */
@@ -50,15 +56,21 @@ void UMeshSurfacePointTool::Setup()
 	bCtrlToggle = false;
 
 	// add input behaviors
-	UMeshSurfacePointToolMouseBehavior* mouseBehavior = NewObject<UMeshSurfacePointToolMouseBehavior>();
-	mouseBehavior->Initialize(this);
-	AddInputBehavior(mouseBehavior);
+	UClickDragInputBehavior* DragBehavior = NewObject<UClickDragInputBehavior>();
+	DragBehavior->Initialize(this);
+	AddInputBehavior(DragBehavior);
 
-	UMouseHoverBehavior* hoverBehavior = NewObject<UMouseHoverBehavior>();
-	hoverBehavior->Initialize(this);
-	AddInputBehavior(hoverBehavior);
+	UMouseHoverBehavior* HoverBehavior = NewObject<UMouseHoverBehavior>();
+	HoverBehavior->Modifiers.RegisterModifier(UMeshSurfacePointTool_ShiftModifier, FInputDeviceState::IsShiftKeyDown);
+	HoverBehavior->Modifiers.RegisterModifier(UMeshSurfacePointTool_CtrlModifier, FInputDeviceState::IsCtrlKeyDown);
+	HoverBehavior->Initialize(this);
+	AddInputBehavior(HoverBehavior);
 }
 
+void UMeshSurfacePointTool::SetStylusAPI(IToolStylusStateProviderAPI* StylusAPIIn)
+{
+	this->StylusAPI = StylusAPIIn;
+}
 
 bool UMeshSurfacePointTool::HitTest(const FRay& Ray, FHitResult& OutHit)
 {
@@ -101,7 +113,21 @@ void UMeshSurfacePointTool::SetCtrlToggle(bool bCtrlDown)
 }
 
 
-FInputRayHit UMeshSurfacePointTool::BeginHoverSequenceHitTest(const FInputDeviceRay& PressPos)
+
+void UMeshSurfacePointTool::OnUpdateModifierState(int ModifierID, bool bIsOn)
+{
+	if (ModifierID == UMeshSurfacePointTool_ShiftModifier)
+	{
+		bShiftToggle = bIsOn;
+	}
+	else if (ModifierID == UMeshSurfacePointTool_CtrlModifier)
+	{
+		bCtrlToggle = bIsOn;
+	}
+}
+
+
+FInputRayHit UMeshSurfacePointTool::CanBeginClickDragSequence(const FInputDeviceRay& PressPos)
 {
 	FHitResult OutHit;
 	if (HitTest(PressPos.WorldRay, OutHit))
@@ -111,73 +137,48 @@ FInputRayHit UMeshSurfacePointTool::BeginHoverSequenceHitTest(const FInputDevice
 	return FInputRayHit();
 }
 
-
-
-
-
-
-/*
- * Mouse Input Behavior
- */
-
-
-void UMeshSurfacePointToolMouseBehavior::Initialize(UMeshSurfacePointTool* ToolIn)
+void UMeshSurfacePointTool::OnClickPress(const FInputDeviceRay& PressPos)
 {
-	this->Tool = ToolIn;
-	bInDragCapture = false;
+	LastWorldRay = PressPos.WorldRay;
+	OnBeginDrag(PressPos.WorldRay);
 }
 
-
-FInputCaptureRequest UMeshSurfacePointToolMouseBehavior::WantsCapture(const FInputDeviceState& input)
+void UMeshSurfacePointTool::OnClickDrag(const FInputDeviceRay& DragPos)
 {
-	if ( IsPressed(input) )
-	{
-		FHitResult OutHit;
-		if (Tool->HitTest(input.Mouse.WorldRay, OutHit))
-		{
-			return FInputCaptureRequest::Begin(this, EInputCaptureSide::Any, OutHit.Distance);
-		}
-	}
-	return FInputCaptureRequest::Ignore();
+	LastWorldRay = DragPos.WorldRay;
+	OnUpdateDrag(DragPos.WorldRay);
 }
 
-FInputCaptureUpdate UMeshSurfacePointToolMouseBehavior::BeginCapture(const FInputDeviceState& input, EInputCaptureSide eSide)
+void UMeshSurfacePointTool::OnClickRelease(const FInputDeviceRay& ReleasePos)
 {
-	Tool->SetShiftToggle(input.bShiftKeyDown);
-	Tool->SetCtrlToggle(input.bCtrlKeyDown);
-	Tool->OnBeginDrag(input.Mouse.WorldRay);
-	LastWorldRay = input.Mouse.WorldRay;
-	bInDragCapture = true;
-	return FInputCaptureUpdate::Begin(this, EInputCaptureSide::Any);
+	LastWorldRay = ReleasePos.WorldRay;
+	OnEndDrag(ReleasePos.WorldRay);
 }
 
-FInputCaptureUpdate UMeshSurfacePointToolMouseBehavior::UpdateCapture(const FInputDeviceState& input, const FInputCaptureData& data)
+void UMeshSurfacePointTool::OnTerminateDragSequence()
 {
-	LastWorldRay = input.Mouse.WorldRay;
+	OnEndDrag(LastWorldRay);
+}
 
-	if ( IsReleased(input) ) 
-	{
-		Tool->OnEndDrag(input.Mouse.WorldRay);
-		bInDragCapture = false;
-		return FInputCaptureUpdate::End();
-	}
 
-	Tool->OnUpdateDrag(input.Mouse.WorldRay);
-	return FInputCaptureUpdate::Continue();
-}
 
-void UMeshSurfacePointToolMouseBehavior::ForceEndCapture(const FInputCaptureData& data)
+
+FInputRayHit UMeshSurfacePointTool::BeginHoverSequenceHitTest(const FInputDeviceRay& PressPos)
 {
-	if (bInDragCapture)
+	FHitResult OutHit;
+	if (HitTest(PressPos.WorldRay, OutHit))
 	{
-		Tool->OnEndDrag(LastWorldRay);
-		bInDragCapture = false;
+		return FInputRayHit(OutHit.Distance);
 	}
-
-	// nothing to do
+	return FInputRayHit();
 }
 
 
+float UMeshSurfacePointTool::GetCurrentDevicePressure() const
+{
+	return (StylusAPI != nullptr) ? FMath::Clamp(StylusAPI->GetCurrentPressure(), 0.0f, 1.0f) : 1.0f;
+}
+
 
 
 #undef LOCTEXT_NAMESPACE
\ No newline at end of file
diff --git a/Engine/Source/Runtime/Experimental/InteractiveToolsFramework/Private/BaseTools/SingleClickTool.cpp b/Engine/Source/Runtime/Experimental/InteractiveToolsFramework/Private/BaseTools/SingleClickTool.cpp
--- a/Engine/Source/Runtime/Experimental/InteractiveToolsFramework/Private/BaseTools/SingleClickTool.cpp
+++ b/Engine/Source/Runtime/Experimental/InteractiveToolsFramework/Private/BaseTools/SingleClickTool.cpp
@@ -40,9 +40,9 @@ void USingleClickTool::Setup()
 }
 
 
-bool USingleClickTool::IsHitByClick(const FInputDeviceRay& ClickPos)
+FInputRayHit USingleClickTool::IsHitByClick(const FInputDeviceRay& ClickPos)
 {
-	return true;
+	return FInputRayHit(0.0f);
 }
 
 
diff --git a/Engine/Source/Runtime/Experimental/InteractiveToolsFramework/Private/ComponentSourceInterfaces.cpp b/Engine/Source/Runtime/Experimental/InteractiveToolsFramework/Private/ComponentSourceInterfaces.cpp
--- a/Engine/Source/Runtime/Experimental/InteractiveToolsFramework/Private/ComponentSourceInterfaces.cpp
+++ b/Engine/Source/Runtime/Experimental/InteractiveToolsFramework/Private/ComponentSourceInterfaces.cpp
@@ -10,14 +10,13 @@ namespace
 	TArray<TUniquePtr<FComponentTargetFactory>> Factories;
 }
 
-void
-AddComponentTargetFactory( TUniquePtr<FComponentTargetFactory> Factory )
+
+void AddComponentTargetFactory( TUniquePtr<FComponentTargetFactory> Factory )
 {
 	Factories.Push( MoveTemp(Factory) );
 }
 
-bool
-CanMakeComponentTarget(UActorComponent* Component)
+bool CanMakeComponentTarget(UActorComponent* Component)
 {
 	for ( const auto& Factory : Factories )
 	{
@@ -29,8 +28,7 @@ CanMakeComponentTarget(UActorComponent* Component)
 	return false;
 }
 
-TUniquePtr<FPrimitiveComponentTarget>
-MakeComponentTarget(UPrimitiveComponent* Component)
+TUniquePtr<FPrimitiveComponentTarget> MakeComponentTarget(UPrimitiveComponent* Component)
 {
 	for ( const auto& Factory : Factories )
 	{
@@ -42,21 +40,39 @@ MakeComponentTarget(UPrimitiveComponent* Component)
 	return {};
 }
 
-AActor*
-FPrimitiveComponentTarget::GetOwnerActor() const
+
+
+
+bool FComponentMaterialSet::operator!=(const FComponentMaterialSet& Other) const
+{
+	int32 Num = Materials.Num();
+	if (Other.Materials.Num() != Num)
+	{
+		return true;
+	}
+	for (int32 j = 0; j < Num; ++j)
+	{
+		if (Other.Materials[j] != Materials[j])
+		{
+			return true;
+		}
+	}
+	return false;
+}
+
+
+AActor* FPrimitiveComponentTarget::GetOwnerActor() const
 {
 	return Component->GetOwner();
 }
 
-UPrimitiveComponent*
-FPrimitiveComponentTarget::GetOwnerComponent() const
+UPrimitiveComponent* FPrimitiveComponentTarget::GetOwnerComponent() const
 {
 	return Component;
 }
 
 
-void
-FPrimitiveComponentTarget::SetOwnerVisibility(bool bVisible) const
+void FPrimitiveComponentTarget::SetOwnerVisibility(bool bVisible) const
 {
 	Component->SetVisibility(bVisible);
 }
@@ -67,21 +83,37 @@ int32 FPrimitiveComponentTarget::GetNumMaterials() const
 	return Component->GetNumMaterials();
 }
 
-UMaterialInterface*
-FPrimitiveComponentTarget::GetMaterial(int32 MaterialIndex) const
+UMaterialInterface* FPrimitiveComponentTarget::GetMaterial(int32 MaterialIndex) const
 {
 	return Component->GetMaterial(MaterialIndex);
 }
 
-FTransform
-FPrimitiveComponentTarget::GetWorldTransform() const
+void FPrimitiveComponentTarget::GetMaterialSet(FComponentMaterialSet& MaterialSetOut, bool bAssetMaterials) const
+{
+	int32 NumMaterials = Component->GetNumMaterials(); 
+	MaterialSetOut.Materials.SetNum(NumMaterials);
+	for (int32 k = 0; k < NumMaterials; ++k)
+	{
+		MaterialSetOut.Materials[k] = Component->GetMaterial(k);
+	}
+}
+
+
+void FPrimitiveComponentTarget::CommitMaterialSetUpdate(const FComponentMaterialSet& MaterialSet, bool bApplyToAsset)
+{
+	check(false);		// not implemented
+}
+
+
+
+
+FTransform FPrimitiveComponentTarget::GetWorldTransform() const
 {
 	//return Component->GetOwner()->GetActorTransform();
 	return Component->GetComponentTransform();
 }
 
-bool
-FPrimitiveComponentTarget::HitTest(const FRay& WorldRay, FHitResult& OutHit) const
+bool FPrimitiveComponentTarget::HitTest(const FRay& WorldRay, FHitResult& OutHit) const
 {
 	FVector End = WorldRay.PointAt(HALF_WORLD_MAX);
 	if (Component->LineTraceComponent(OutHit, WorldRay.Origin, End, FCollisionQueryParams(SCENE_QUERY_STAT(HitTest), true)))
diff --git a/Engine/Source/Runtime/Experimental/InteractiveToolsFramework/Public/BaseBehaviors/BehaviorTargetInterfaces.h b/Engine/Source/Runtime/Experimental/InteractiveToolsFramework/Public/BaseBehaviors/BehaviorTargetInterfaces.h
--- a/Engine/Source/Runtime/Experimental/InteractiveToolsFramework/Public/BaseBehaviors/BehaviorTargetInterfaces.h
+++ b/Engine/Source/Runtime/Experimental/InteractiveToolsFramework/Public/BaseBehaviors/BehaviorTargetInterfaces.h
@@ -38,9 +38,9 @@ public:
 	/**
 	 * Test if target is hit by a click
 	 * @param ClickPos device position/ray at click point
-	 * @return true if target was hit by click ray/point
+	 * @return hit information at this point
 	 */
-	virtual bool IsHitByClick(const FInputDeviceRay& ClickPos) = 0;
+	virtual FInputRayHit IsHitByClick(const FInputDeviceRay& ClickPos) = 0;
 
 
 	/**
diff --git a/Engine/Source/Runtime/Experimental/InteractiveToolsFramework/Public/BaseBehaviors/SingleClickBehavior.h b/Engine/Source/Runtime/Experimental/InteractiveToolsFramework/Public/BaseBehaviors/SingleClickBehavior.h
--- a/Engine/Source/Runtime/Experimental/InteractiveToolsFramework/Public/BaseBehaviors/SingleClickBehavior.h
+++ b/Engine/Source/Runtime/Experimental/InteractiveToolsFramework/Public/BaseBehaviors/SingleClickBehavior.h
@@ -5,6 +5,7 @@
 #include "CoreMinimal.h"
 #include "BaseBehaviors/AnyButtonInputBehavior.h"
 #include "BehaviorTargetInterfaces.h"
+#include "InputBehaviorModifierStates.h"
 #include "SingleClickBehavior.generated.h"
 
 
@@ -54,6 +55,11 @@ public:
 	UPROPERTY()
 	bool HitTestOnRelease;
 
+	/**
+	 * The modifier set for this behavior
+	 */
+	FInputBehaviorModifierStates Modifiers;
+
 
 protected:
 	/** Click Target object */
diff --git a/Engine/Source/Runtime/Experimental/InteractiveToolsFramework/Public/BaseGizmos/TransformProxy.h b/Engine/Source/Runtime/Experimental/InteractiveToolsFramework/Public/BaseGizmos/TransformProxy.h
--- a/Engine/Source/Runtime/Experimental/InteractiveToolsFramework/Public/BaseGizmos/TransformProxy.h
+++ b/Engine/Source/Runtime/Experimental/InteractiveToolsFramework/Public/BaseGizmos/TransformProxy.h
@@ -45,6 +45,20 @@ public:
 	virtual void SetTransform(const FTransform& Transform);
 
 
+	/**
+	 * In some use cases SetTransform() will be called repeatedly (eg during an interactive gizmo edit). External
+	 * code may know and/or need to know when such a sequence starts/ends. The OnBeginTransformEdit/OnEndTransformEdit delegates
+	 * can provide these notifications, however client code must call BeginTransformEditSequence()/EndTransformEditSequence() to fire those delegates 
+	 * as the TransformProxy doesn't know about this external state. 
+	 * @note FTransformProxyChangeSource will call these functions on Begin/End
+	 */
+	virtual void BeginTransformEditSequence();
+
+	/**
+	 * External clients should call this when done a sequence of SetTransform calls (see BeginTransformEditSequence)
+	 */
+	virtual void EndTransformEditSequence();
+
 public:
 	/**
 	 * This delegate is fired whenever the internal transform changes, ie
@@ -53,6 +67,15 @@ public:
 	DECLARE_MULTICAST_DELEGATE_TwoParams(FOnTransformChanged, UTransformProxy*, FTransform);
 	FOnTransformChanged OnTransformChanged;
 
+	/** This delegate is fired when BeginTransformEditSequence() is called to indicate that a transform change has started */
+	DECLARE_MULTICAST_DELEGATE_OneParam(FOnBeginTransformEdit, UTransformProxy*);
+	FOnBeginTransformEdit OnBeginTransformEdit;
+
+	/** This delegate is fired when EndTransformEditSequence() is called to indicate that a transform change has ended */
+	DECLARE_MULTICAST_DELEGATE_OneParam(FOnEndTransformEdit, UTransformProxy*);
+	FOnEndTransformEdit OnEndTransformEdit;
+
+
 	/**
 	 * If true, relative rotation of shared transform is applied to objects before relative translation (ie they rotate in place)
 	 */
diff --git a/Engine/Source/Runtime/Experimental/InteractiveToolsFramework/Public/BaseTools/MeshSurfacePointTool.h b/Engine/Source/Runtime/Experimental/InteractiveToolsFramework/Public/BaseTools/MeshSurfacePointTool.h
--- a/Engine/Source/Runtime/Experimental/InteractiveToolsFramework/Public/BaseTools/MeshSurfacePointTool.h
+++ b/Engine/Source/Runtime/Experimental/InteractiveToolsFramework/Public/BaseTools/MeshSurfacePointTool.h
@@ -14,6 +14,14 @@
 
 class UMeshSurfacePointTool;
 
+// This is a temporary interface to provide stylus pressure, currently necessary due to limitations
+// in the stylus plugin architecture. Should be removed once InputState/InputBehavior can support stylus events
+class INTERACTIVETOOLSFRAMEWORK_API IToolStylusStateProviderAPI
+{
+public:
+	virtual float GetCurrentPressure() const = 0;
+};
+
 /**
  * 
  */
@@ -23,6 +31,8 @@ class INTERACTIVETOOLSFRAMEWORK_API UMeshSurfacePointToolBuilder : public UInter
 	GENERATED_BODY()
 
 public:
+	IToolStylusStateProviderAPI* StylusAPI = nullptr;
+
 	/** @return true if a single mesh source can be found in the active selection */
 	virtual bool CanBuildTool(const FToolBuilderState& SceneState) const override;
 
@@ -47,7 +57,7 @@ public:
  * to implement custom behavior.
  */
 UCLASS()
-class INTERACTIVETOOLSFRAMEWORK_API UMeshSurfacePointTool : public USingleSelectionTool, public IHoverBehaviorTarget
+class INTERACTIVETOOLSFRAMEWORK_API UMeshSurfacePointTool : public USingleSelectionTool, public IClickDragBehaviorTarget, public IHoverBehaviorTarget
 {
 	GENERATED_BODY()
 
@@ -57,6 +67,8 @@ public:
 	/** Register InputBehaviors, etc */
 	virtual void Setup() override;
 
+	/** Set current stlyus API source */
+	virtual void SetStylusAPI(IToolStylusStateProviderAPI* StylusAPI);
 
 	// UMeshSurfacePointTool API
 
@@ -65,24 +77,33 @@ public:
 	 */
 	virtual bool HitTest(const FRay& Ray, FHitResult& OutHit);
 
-	
 	/**
-	 * This function is called by registered InputBehaviors when the user begins a click-drag-release interaction
+	 * This function is called when the user begins a click-drag-release interaction
 	 */
 	virtual void OnBeginDrag(const FRay& Ray);
 
 	/**
-	 * This function is called by registered InputBehaviorseach frame that the user is in a click-drag-release interaction
+	 * This function is called each frame that the user is in a click-drag-release interaction
 	 */
 	virtual void OnUpdateDrag(const FRay& Ray);
 
 	/**
-	 * This function is called by registered InputBehaviors when the user releases the button driving a click-drag-release interaction
+	 * This function is called when the user releases the button driving a click-drag-release interaction
 	 */
 	virtual void OnEndDrag(const FRay& Ray);
 
 
 
+
+	// IClickDragBehaviorTarget implementation
+	virtual FInputRayHit CanBeginClickDragSequence(const FInputDeviceRay& PressPos) override;
+	virtual void OnClickPress(const FInputDeviceRay& PressPos) override;
+	virtual void OnClickDrag(const FInputDeviceRay& DragPos) override;
+	virtual void OnClickRelease(const FInputDeviceRay& ReleasePos) override;
+	virtual void OnTerminateDragSequence() override;
+	virtual void OnUpdateModifierState(int ModifierID, bool bIsOn) override;
+
+
 	// IHoverBehaviorTarget implementation
 	
 	virtual FInputRayHit BeginHoverSequenceHitTest(const FInputDeviceRay& PressPos) override;
@@ -105,6 +126,8 @@ public:
 	/** @return current state of the shift toggle */
 	virtual bool GetCtrlToggle() const { return bCtrlToggle; }
 
+	/** @return current input device pressure in range 0-1 */
+	virtual float GetCurrentDevicePressure() const;
 
 protected:
 	/** Current state of the shift modifier toggle */
@@ -112,36 +135,9 @@ protected:
 
 	/** Current state of the ctrl modifier toggle */
 	bool bCtrlToggle = false;
-};
-
-
-
-
-
 
-
-
-/**
- * UMeshSurfacePointToolMouseBehavior implements mouse press-drag-release interaction behavior for Mouse devices.
- * You can configure the base UAnyButtonInputBehavior to change the mouse button in use (default = left mouse)
- */
-UCLASS()
-class INTERACTIVETOOLSFRAMEWORK_API UMeshSurfacePointToolMouseBehavior : public UAnyButtonInputBehavior
-{
-	GENERATED_BODY()
-
-public:
-	virtual void Initialize(UMeshSurfacePointTool* Tool);
-
-	virtual FInputCaptureRequest WantsCapture(const FInputDeviceState& input) override;
-	virtual FInputCaptureUpdate BeginCapture(const FInputDeviceState& input, EInputCaptureSide eSide) override;
-	virtual FInputCaptureUpdate UpdateCapture(const FInputDeviceState& input, const FInputCaptureData& data) override;
-	virtual void ForceEndCapture(const FInputCaptureData& data) override;
-
-protected:
-	UMeshSurfacePointTool* Tool;
 	FRay LastWorldRay;
-	bool bInDragCapture;
-};
 
+	IToolStylusStateProviderAPI* StylusAPI = nullptr;
+};
 
diff --git a/Engine/Source/Runtime/Experimental/InteractiveToolsFramework/Public/BaseTools/SingleClickTool.h b/Engine/Source/Runtime/Experimental/InteractiveToolsFramework/Public/BaseTools/SingleClickTool.h
--- a/Engine/Source/Runtime/Experimental/InteractiveToolsFramework/Public/BaseTools/SingleClickTool.h
+++ b/Engine/Source/Runtime/Experimental/InteractiveToolsFramework/Public/BaseTools/SingleClickTool.h
@@ -51,7 +51,7 @@ public:
 	/**
 	 * Test if the Target is hit at this 2D position / 3D ray
 	 */
-	virtual bool IsHitByClick(const FInputDeviceRay& ClickPos);
+	virtual FInputRayHit IsHitByClick(const FInputDeviceRay& ClickPos);
 
 	/**
 	 * Click the Target at this 2D position / 3D ray. Default behavior is to print debug string.
diff --git a/Engine/Source/Runtime/Experimental/InteractiveToolsFramework/Public/ComponentSourceInterfaces.h b/Engine/Source/Runtime/Experimental/InteractiveToolsFramework/Public/ComponentSourceInterfaces.h
--- a/Engine/Source/Runtime/Experimental/InteractiveToolsFramework/Public/ComponentSourceInterfaces.h
+++ b/Engine/Source/Runtime/Experimental/InteractiveToolsFramework/Public/ComponentSourceInterfaces.h
@@ -15,6 +15,16 @@ class UPrimitiveComponent;
 struct FMeshDescription;
 class UMaterialInterface;
 
+/**
+ * FComponentMaterialSet is the set of materials assigned to a component (ie Material Slots on a StaticMesh)
+ */
+struct INTERACTIVETOOLSFRAMEWORK_API FComponentMaterialSet
+{
+	TArray<UMaterialInterface*> Materials;
+	bool operator!=(const FComponentMaterialSet& Other) const;
+};
+
+
 /**
  * Wrapper around a UObject Component that can provide a MeshDescription, and
  * (optionally) bake a modified MeshDescription back to this Component.
@@ -49,6 +59,13 @@ public:
 	 */
 	UMaterialInterface* GetMaterial(int32 MaterialIndex) const;
 
+	/**
+	 * Get material set provided by this source
+	 * @param MaterialSetOut returned material set
+	 * @param bAssetMaterials if an underlying asset exists, return the Asset-level material assignment instead of the component materials
+	 */
+	virtual void GetMaterialSet(FComponentMaterialSet& MaterialSetOut, bool bAssetMaterials = false) const;
+
 	/**
 	 * @return the transform on this component
 	 * @todo Do we need to return a list of transforms here?
@@ -69,7 +86,20 @@ public:
 	 */
 	void SetOwnerVisibility(bool bVisible) const;
 
-	using  FCommitter  = TFunction< void( FMeshDescription* ) >;
+
+	/**
+	 * Commit an update to the material set. This may generate a transaction.
+	 * @param MaterialSet new list of materials
+	 * @param bApplyToAsset if true, materials of Asset are updated (if Asset exists), rather than Component
+	 */
+	virtual void CommitMaterialSetUpdate(const FComponentMaterialSet& MaterialSet, bool bApplyToAsset);
+
+
+	struct FCommitParams
+	{
+		FMeshDescription* MeshDescription{ nullptr };
+	};
+	using  FCommitter  = TFunction< void( const FCommitParams& ) >;
 	virtual FMeshDescription* GetMesh() = 0;
 	virtual void CommitMesh( const FCommitter& ) = 0;
 
