Added STAT_PerformOverlapQuery to differentiate components that have overlap events enabled

#rnx

diff --git a/Engine/Source/Runtime/Engine/Private/Components/PrimitiveComponent.cpp b/Engine/Source/Runtime/Engine/Private/Components/PrimitiveComponent.cpp
--- a/Engine/Source/Runtime/Engine/Private/Components/PrimitiveComponent.cpp
+++ b/Engine/Source/Runtime/Engine/Private/Components/PrimitiveComponent.cpp
@@ -3107,6 +3107,7 @@ bool UPrimitiveComponent::UpdateOverlapsImpl(const TOverlapArrayView* NewPending
 				}
 				else
 				{
+					SCOPE_CYCLE_COUNTER(STAT_PerformOverlapQuery);
 					UE_LOG(LogPrimitiveComponent, VeryVerbose, TEXT("%s->%s Performing overlaps!"), *GetNameSafe(GetOwner()), *GetName());
 					UWorld* const MyWorld = GetWorld();
 					TArray<FOverlapResult> Overlaps;
diff --git a/Engine/Source/Runtime/Engine/Private/LevelTick.cpp b/Engine/Source/Runtime/Engine/Private/LevelTick.cpp
--- a/Engine/Source/Runtime/Engine/Private/LevelTick.cpp
+++ b/Engine/Source/Runtime/Engine/Private/LevelTick.cpp
@@ -98,6 +98,7 @@ DEFINE_STAT(STAT_TeleportToTime);
 DEFINE_STAT(STAT_MoveComponentTime);
 DEFINE_STAT(STAT_MoveComponentSceneComponentTime);
 DEFINE_STAT(STAT_UpdateOverlaps);
+DEFINE_STAT(STAT_PerformOverlapQuery);
 DEFINE_STAT(STAT_UpdatePhysicsVolume);
 DEFINE_STAT(STAT_EndScopedMovementUpdate);
 
diff --git a/Engine/Source/Runtime/Engine/Public/EngineStats.h b/Engine/Source/Runtime/Engine/Public/EngineStats.h
--- a/Engine/Source/Runtime/Engine/Public/EngineStats.h
+++ b/Engine/Source/Runtime/Engine/Public/EngineStats.h
@@ -48,6 +48,7 @@ DECLARE_CYCLE_STAT_EXTERN(TEXT("Actor BeginPlay"), STAT_ActorBeginPlay, STATGROU
 DECLARE_CYCLE_STAT_EXTERN(TEXT("MoveComponent(Primitive) Time"),STAT_MoveComponentTime,STATGROUP_Game, );
 DECLARE_CYCLE_STAT_EXTERN(TEXT("MoveComponent(SceneComp) Time"), STAT_MoveComponentSceneComponentTime, STATGROUP_Game, );
 DECLARE_CYCLE_STAT_EXTERN(TEXT("UpdateOverlaps Time"),STAT_UpdateOverlaps,STATGROUP_Game, );
+DECLARE_CYCLE_STAT_EXTERN(TEXT("PerformOverlapQuery Time"), STAT_PerformOverlapQuery, STATGROUP_Game, );
 DECLARE_CYCLE_STAT_EXTERN(TEXT("UpdatePhysicsVolume Time"),STAT_UpdatePhysicsVolume,STATGROUP_Game, );
 DECLARE_CYCLE_STAT_EXTERN(TEXT("EndScopedMovementUpdate Time"),STAT_EndScopedMovementUpdate,STATGROUP_Game, );
 DECLARE_CYCLE_STAT_EXTERN(TEXT("TeleportTo Time"),STAT_TeleportToTime,STATGROUP_Game, );
