Assetization of Master Submix Graph Prep 2:

UX Updates:
- Clean-up SoundClass Properties
- Bifurcate Effects category on inherited USoundBase types for clarity into Effects|Submix & Effects|Source
- Make ChannelFormat visible but not editable on USoundSubmix

Functional Updates:
- Add ability to set default SubMix for subscribing sound types on SoundClass & deprecate bApplyEffects
- Remove bEQFilterApplied from WaveInstance/ActiveSound in preparation of deprecating (leave old backends directly referencing from SoundClass)
- Set Oculus Audio Plugin to deploy content
- Minimize compile times when editing USoundClass by removing from PCH/Audio.h

Minor Clean-up:
- Remove unused/legacy bEQFilterApplied from UAudioComponent
- Remove unused/legacy bEQFilterApplied & SoundSubmixOverride from FActiveSound

diff --git a/Engine/Plugins/Runtime/Oculus/OculusAudio/OculusAudio.uplugin b/Engine/Plugins/Runtime/Oculus/OculusAudio/OculusAudio.uplugin
--- a/Engine/Plugins/Runtime/Oculus/OculusAudio/OculusAudio.uplugin
+++ b/Engine/Plugins/Runtime/Oculus/OculusAudio/OculusAudio.uplugin
@@ -11,7 +11,7 @@
 	"MarketplaceURL": "",
 	"SupportURL": "https://forums.oculusvr.com/developer/categories/audio-development",
 	"EnabledByDefault": false,
-	"CanContainContent": false,
+	"CanContainContent": true,
 	"IsBetaVersion": false,
 	"Installed": false,
 	"Modules": [
diff --git a/Engine/Source/Editor/UnrealEd/Private/UnrealEdPrivatePCH.h b/Engine/Source/Editor/UnrealEd/Private/UnrealEdPrivatePCH.h
--- a/Engine/Source/Editor/UnrealEd/Private/UnrealEdPrivatePCH.h
+++ b/Engine/Source/Editor/UnrealEd/Private/UnrealEdPrivatePCH.h
@@ -550,7 +550,6 @@
 #include "Curves/CurveBase.h"
 #include "Curves/CurveOwnerInterface.h"
 #include "Camera/CameraTypes.h"
-#include "Sound/SoundClass.h"
 #include "Curves/CurveFloat.h"
 #include "EditorFramework/AssetImportData.h"
 #include "Animation/Skeleton.h"
diff --git a/Engine/Source/Editor/UnrealEd/Public/UnrealEdSharedPCH.h b/Engine/Source/Editor/UnrealEd/Public/UnrealEdSharedPCH.h
--- a/Engine/Source/Editor/UnrealEd/Public/UnrealEdSharedPCH.h
+++ b/Engine/Source/Editor/UnrealEd/Public/UnrealEdSharedPCH.h
@@ -575,7 +575,6 @@
 #include "Materials/MaterialExpressionMaterialFunctionCall.h"
 #include "Materials/MaterialFunction.h"
 #include "PhysicsEngine/PhysicsSettingsEnums.h"
-#include "Sound/SoundClass.h"
 #include "Sound/SoundAttenuation.h"
 #include "Engine/CurveTable.h"
 #include "Audio.h"
diff --git a/Engine/Source/Runtime/AudioMixer/Private/AudioMixerSource.cpp b/Engine/Source/Runtime/AudioMixer/Private/AudioMixerSource.cpp
--- a/Engine/Source/Runtime/AudioMixer/Private/AudioMixerSource.cpp
+++ b/Engine/Source/Runtime/AudioMixer/Private/AudioMixerSource.cpp
@@ -236,7 +236,7 @@ namespace Audio
 							// Send the voice to the EQ submix if it's enabled
 							const bool bIsEQDisabled = GetDefault<UAudioSettings>()->bDisableMasterEQ;
 							bool bUseMaster = true;
-							if (!bIsEQDisabled && IsEQFilterApplied())
+							if (!bIsEQDisabled)
 							{
 								if (MixerDevice->GetMasterEQSubmix().IsValid())
 								{
diff --git a/Engine/Source/Runtime/Engine/Classes/Components/AudioComponent.h b/Engine/Source/Runtime/Engine/Classes/Components/AudioComponent.h
--- a/Engine/Source/Runtime/Engine/Classes/Components/AudioComponent.h
+++ b/Engine/Source/Runtime/Engine/Classes/Components/AudioComponent.h
@@ -213,9 +213,6 @@ class ENGINE_API UAudioComponent : public USceneComponent
 	/** If true, this sound will not be stopped when flushing the audio device. */
 	uint8 bIgnoreForFlushing:1;
 
-	/** Whether audio effects are applied */
-	uint8 bEQFilterApplied:1;
-
 	/** Whether to artificially prioritize the component to play */
 	uint8 bAlwaysPlay:1;
 
diff --git a/Engine/Source/Runtime/Engine/Classes/Sound/AudioOutputTarget.h b/Engine/Source/Runtime/Engine/Classes/Sound/AudioOutputTarget.h
new file mode 100644
--- /dev/null
+++ b/Engine/Source/Runtime/Engine/Classes/Sound/AudioOutputTarget.h
@@ -0,0 +1,14 @@
+#pragma once
+
+#include "UObject/ObjectMacros.h"
+
+UENUM()
+namespace EAudioOutputTarget
+{
+	enum Type
+	{
+		Speaker UMETA(ToolTip = "Sound plays only from speakers."),
+		Controller UMETA(ToolTip = "Sound plays only from controller if present."),
+		ControllerFallbackToSpeaker UMETA(ToolTip = "Sound plays on the controller if present. If not present, it plays from speakers.")
+	};
+}
diff --git a/Engine/Source/Runtime/Engine/Classes/Sound/SoundBase.h b/Engine/Source/Runtime/Engine/Classes/Sound/SoundBase.h
--- a/Engine/Source/Runtime/Engine/Classes/Sound/SoundBase.h
+++ b/Engine/Source/Runtime/Engine/Classes/Sound/SoundBase.h
@@ -70,7 +70,7 @@ public:
 	uint8 bOverrideConcurrency : 1;
 
 	/** Whether or not to only send this audio's output to a bus. If true, will not be this sound won't be audible except through bus sends. */
-	UPROPERTY(EditAnywhere, BlueprintReadWrite, Category = Effects)
+	UPROPERTY(EditAnywhere, BlueprintReadWrite, Category = "Effects|Source")
 	uint8 bOutputToBusOnly : 1;
 
 	/** Whether or not to only send this audio's output to a bus. If true, will not be this sound won't be audible except through bus sends. */
@@ -149,27 +149,25 @@ public:
 	UPROPERTY(EditAnywhere, Category = Modulation)
 	FSoundModulation Modulation;
 
-	/** Sound submix this sound belongs to.
-	  * Audio will play here and traverse through the submix graph.
-	  * A null entry will make the sound obey the default master effects graph.
-	  */
-	UPROPERTY(EditAnywhere, BlueprintReadWrite, Category = Effects, meta = (DisplayName = "Sound Submix"))
+	/** Submix to route sound output to. If unset, falls back to referenced SoundClass submix.
+	  * If SoundClass submix is unset, sends to the 'Master Submix' as set in the 'Audio' category of Project Settings'. */
+	UPROPERTY(EditAnywhere, BlueprintReadWrite, Category = "Effects|Submix", meta = (DisplayName = "Submix"))
 	USoundSubmix* SoundSubmixObject;
 
-	/** An array of submix sends. Audio from this sound will send a portion of its audio to these effects. */
-	UPROPERTY(EditAnywhere, BlueprintReadWrite, Category = Effects)
+	/** Array of submix sends to which a prescribed amount (see 'Send Level') of this sound is sent. */
+	UPROPERTY(EditAnywhere, BlueprintReadWrite, Category = "Effects|Submix", meta = (DisplayName = "Submix Sends"))
 	TArray<FSoundSubmixSendInfo> SoundSubmixSends;
 
 	/** The source effect chain to use for this sound. */
-	UPROPERTY(EditAnywhere, BlueprintReadWrite, Category = Effects)
+	UPROPERTY(EditAnywhere, BlueprintReadWrite, Category = "Effects|Source")
 	USoundEffectSourcePresetChain* SourceEffectChain;
 
 	/** This sound will send its audio output to this list of buses if there are bus instances playing after source effects are processed. */
-	UPROPERTY(EditAnywhere, BlueprintReadWrite, Category = Effects, meta = (DisplayName = "Post-Effect Bus Sends"))
+	UPROPERTY(EditAnywhere, BlueprintReadWrite, Category = "Effects|Source", meta = (DisplayName = "Post-Effect Bus Sends"))
 	TArray<FSoundSourceBusSendInfo> BusSends;
 
 	/** This sound will send its audio output to this list of buses if there are bus instances playing before source effects are processed. */
-	UPROPERTY(EditAnywhere, BlueprintReadWrite, Category = Effects, meta = (DisplayName = "Pre-Effect Bus Sends"))
+	UPROPERTY(EditAnywhere, BlueprintReadWrite, Category = "Effects|Source", meta = (DisplayName = "Pre-Effect Bus Sends"))
 	TArray<FSoundSourceBusSendInfo> PreEffectBusSends;
 
 public:
diff --git a/Engine/Source/Runtime/Engine/Classes/Sound/SoundClass.h b/Engine/Source/Runtime/Engine/Classes/Sound/SoundClass.h
--- a/Engine/Source/Runtime/Engine/Classes/Sound/SoundClass.h
+++ b/Engine/Source/Runtime/Engine/Classes/Sound/SoundClass.h
@@ -2,52 +2,23 @@
 #pragma once
 
 #include "CoreMinimal.h"
+
+#include "AudioDefines.h"
+#include "AudioDynamicParameter.h"
 #include "IAudioExtensionPlugin.h"
+#include "Sound/AudioOutputTarget.h"
 #include "UObject/ObjectMacros.h"
 #include "UObject/Object.h"
-#include "AudioDefines.h"
-#include "AudioDynamicParameter.h"
+
 #if WITH_EDITOR
 #include "EdGraph/EdGraph.h"
-#endif
+#endif // WITH_EDITOR
+#include "SoundWaveLoadingBehavior.h"
+
 #include "SoundClass.generated.h"
 
-UENUM()
-namespace EAudioOutputTarget
-{
-	enum Type
-	{
-		Speaker UMETA(ToolTip = "Sound plays only from speakers."),
-		Controller UMETA(ToolTip = "Sound plays only from controller if present."),
-		ControllerFallbackToSpeaker UMETA(ToolTip = "Sound plays on the controller if present. If not present, it plays from speakers.")
-	};
-}
 
-/**
- * Only used when stream caching is enabled. Determines how we are going to load or retain a given audio asset.
- * A USoundWave's loading behavior can be overridden in the USoundWave itself, the sound wave's USoundClass, or by cvars.
- * The order of priority is defined as:
- * 1) The loading behavior set on the USoundWave
- * 2) The loading behavior set on the USoundWave's USoundClass.
- * 3) The loading behavior set on the nearest parent of a USoundWave's USoundClass.
- * 4) The loading behavior set via the au.streamcache cvars.
- */
-UENUM()
-enum class ESoundWaveLoadingBehavior : uint8
-{
-	// If set on a USoundWave, use the setting defined by the USoundClass. If set on the next parent USoundClass, or the default behavior defined via the au.streamcache cvars.
-	Inherited = 0,
-	// the first chunk of audio for this asset will be retained in the audio cache until a given USoundWave is either destroyed or USoundWave::ReleaseCompressedAudioData is called.
-	RetainOnLoad = 1,
-	// the first chunk of audio for this asset will be loaded into the cache from disk when this asset is loaded, but may be evicted to make room for other audio if it isn't played for a while.
-	PrimeOnLoad = 2,
-	// the first chunk of audio for this asset will not be loaded until this asset is played or primed.
-	LoadOnDemand = 3,
-	// Force all audio data for this audio asset to live outside of the cache and use the non-streaming decode pathways. Only usable if set on the USoundWave.
-	ForceInline = 4 UMETA(DisplayName = "Force Inline"),
-	// This value is used to delineate when the value of ESoundWaveLoadingBehavior hasn't been cached on a USoundWave yet.
-	Uninitialized = 0xff UMETA(Hidden)
-};
+
 
 USTRUCT()
 struct FSoundClassEditorData
@@ -81,81 +52,88 @@ struct FSoundClassProperties
 	GENERATED_USTRUCT_BODY()
 
 	/** Volume multiplier. */
-	UPROPERTY(EditAnywhere, BlueprintReadOnly, Category=SoundClassProperties)
+	UPROPERTY(EditAnywhere, BlueprintReadOnly, Category = General)
 	float Volume;
 
 	/** Pitch multiplier. */
-	UPROPERTY(EditAnywhere, BlueprintReadOnly, Category=SoundClassProperties)
+	UPROPERTY(EditAnywhere, BlueprintReadOnly, Category = General)
 	float Pitch;
 
-	/** Distance scale to apply to sounds that play with this sound class. Sounds will have their attenuation distance scaled by this amount. Allows adjusting attenuation settings dynamically. **/
-	UPROPERTY(EditAnywhere, BlueprintReadOnly, Category = SoundClassProperties)
-	float AttenuationDistanceScale;
-
 	/** Lowpass filter frequency */
-	UPROPERTY(EditAnywhere, BlueprintReadOnly, Category = SoundClassProperties)
+	UPROPERTY(EditAnywhere, BlueprintReadOnly, Category = General)
 	float LowPassFilterFrequency;
 
+	/** Distance scale to apply to sounds that play with this sound class.
+	  * Sounds will have their attenuation distance scaled by this amount.
+	  * Allows adjusting attenuation settings dynamically. */
+	UPROPERTY(EditAnywhere, BlueprintReadOnly, Category = General)
+	float AttenuationDistanceScale;
+
 	/** The amount of stereo sounds to bleed to the rear speakers */
-	UPROPERTY(EditAnywhere, BlueprintReadOnly, Category=SoundClassProperties)
+	UPROPERTY(EditAnywhere, BlueprintReadOnly, Category = Routing)
 	float StereoBleed;
 
 	/** The amount of a sound to bleed to the LFE channel */
-	UPROPERTY(EditAnywhere, BlueprintReadOnly, Category=SoundClassProperties)
+	UPROPERTY(EditAnywhere, BlueprintReadOnly, Category = Routing, meta = (DisplayName = "LFE Bleed"))
 	float LFEBleed;
 
-	/** Voice center channel volume - Not a multiplier (no propagation)	*/
-	UPROPERTY(EditAnywhere, BlueprintReadOnly, Category=SoundClassProperties)
+	/** Voice center channel volume - Not a multiplier (does not propagate to child classes) */
+	UPROPERTY(EditAnywhere, BlueprintReadOnly, Category = Routing)
 	float VoiceCenterChannelVolume;
 
-	/** Volume of the radio filter effect */
-	UPROPERTY(EditAnywhere, BlueprintReadOnly, Category=SoundClassProperties)
+	/** Volume of the radio filter effect. */
+	UPROPERTY(EditAnywhere, BlueprintReadOnly, AdvancedDisplay, Category = Legacy)
 	float RadioFilterVolume;
 
 	/** Volume at which the radio filter kicks in */
-	UPROPERTY(EditAnywhere, BlueprintReadOnly, Category=SoundClassProperties)
+	UPROPERTY(EditAnywhere, BlueprintReadOnly, AdvancedDisplay, Category = Legacy)
 	float RadioFilterVolumeThreshold;
 
-	/** Sound mix voice - whether to apply audio effects */
-	UPROPERTY(EditAnywhere, BlueprintReadOnly, Category=SoundClassProperties)
-	uint32 bApplyEffects:1;
+	/** Whether to use 'Master EQ Submix' as set in the 'Audio' category of Project Settings as the default submix for referencing sounds. */
+	UPROPERTY(EditAnywhere, BlueprintReadOnly, AdvancedDisplay, Category = Legacy, meta = (DisplayName = "Output to Master EQ Submix"))
+	uint8 bApplyEffects:1;
 
-	/** Whether to artificially prioritise the component to play */
-	UPROPERTY(EditAnywhere, BlueprintReadOnly, Category=SoundClassProperties)
-	uint32 bAlwaysPlay:1;
+	/** Whether to inflate referencing sound's priority to always play. */
+	UPROPERTY(EditAnywhere, BlueprintReadOnly, Category = General)
+	uint8 bAlwaysPlay:1;
 
 	/** Whether or not this sound plays when the game is paused in the UI */
-	UPROPERTY(EditAnywhere, BlueprintReadOnly, Category=SoundClassProperties)
-	uint32 bIsUISound:1;
-
-	/** Whether or not this is music (propagates only if parent is true) */
-	UPROPERTY(EditAnywhere, BlueprintReadOnly, Category=SoundClassProperties)
-	uint32 bIsMusic:1;
+	UPROPERTY(EditAnywhere, BlueprintReadOnly, AdvancedDisplay, Category = Legacy)
+	uint8 bIsUISound:1;
 
-	/** Whether or not this sound class has reverb applied */
-	UPROPERTY(EditAnywhere, BlueprintReadOnly, Category=SoundClassProperties)
-	uint32 bReverb:1;
-
-	/** Amount of audio to send to master reverb effect for 2D sounds played with this sound class. */
-	UPROPERTY(EditAnywhere, BlueprintReadOnly, Category = SoundClassProperties)
-	float Default2DReverbSendAmount;
+	/** Whether or not this is music (propagates to child classes only if parent is true) */
+	UPROPERTY(EditAnywhere, BlueprintReadOnly, AdvancedDisplay, Category = Legacy)
+	uint8 bIsMusic:1;
 
 	/** Whether or not this sound class forces sounds to the center channel */
-	UPROPERTY(EditAnywhere, BlueprintReadOnly, Category=SoundClassProperties)
-	uint32 bCenterChannelOnly:1;
+	UPROPERTY(EditAnywhere, BlueprintReadOnly, Category = Routing)
+	uint8 bCenterChannelOnly:1;
 
 	/** Whether the Interior/Exterior volume and LPF modifiers should be applied */
-	UPROPERTY(EditAnywhere, BlueprintReadOnly, Category=SoundClassProperties)
-	uint32 bApplyAmbientVolumes:1;
+	UPROPERTY(EditAnywhere, BlueprintReadOnly, Category = Routing)
+	uint8 bApplyAmbientVolumes:1;
+
+	/** Whether or not sounds referencing this class send to the reverb submix */
+	UPROPERTY(EditAnywhere, BlueprintReadOnly, Category = Submix, meta = (DisplayName = "Send to Master Reverb Submix"))
+	uint8 bReverb:1;
+
+	/** Send amount to master reverb effect for referencing, unattenuated (2D) sounds. */
+	UPROPERTY(EditAnywhere, BlueprintReadOnly, Category = Submix)
+	float Default2DReverbSendAmount;
 
 	/** Which output target the sound should be played through */
-	UPROPERTY(EditAnywhere, BlueprintReadOnly, Category=SoundClassProperties)
+	UPROPERTY(EditAnywhere, BlueprintReadOnly, Category = Routing)
 	TEnumAsByte<EAudioOutputTarget::Type> OutputTarget;
 
-	/** The loading behavior used for USoundWaves of this USoundClass. */
-	UPROPERTY(EditAnywhere, BlueprintReadOnly, Category = SoundClassProperties, meta = (DisplayName = "Loading Behavior Override"))
+	/** Specifies how and when compressed audio data is loaded for asset if stream caching is enabled. */
+	UPROPERTY(EditAnywhere, BlueprintReadOnly, Category = Loading, meta = (DisplayName = "Loading Behavior Override"))
 	ESoundWaveLoadingBehavior LoadingBehavior;
 
+	/** Default output submix of referencing sounds. If unset, falls back to the 'Master Submix' as set in the 'Audio' category of Project Settings. 
+	  * (Unavailable if legacy 'Output to Master EQ Submix' is set) */
+	UPROPERTY(EditAnywhere, BlueprintReadOnly, Category = Submix, meta = (EditCondition = "!bApplyEffects"))
+	USoundSubmix* DefaultSubmix;
+
 	// Sets the attenuation scale of the sound class in the given amount of time
 	void SetAttenuationDistanceScale(float InAttenuationDistanceScale, float InTime);
 
@@ -168,30 +146,7 @@ struct FSoundClassProperties
 	// Updates any dynamic sound class properties
 	void UpdateSoundClassProperties(float DeltaTime);
 
-	FSoundClassProperties()
-		: Volume(1.0f)
-		, Pitch(1.0f)
-		, AttenuationDistanceScale(1.0f)
-		, LowPassFilterFrequency(MAX_FILTER_FREQUENCY)
-		, StereoBleed(0.25f)
-		, LFEBleed(0.5f)
-		, VoiceCenterChannelVolume(0.0f)
-		, RadioFilterVolume(0.0f)
-		, RadioFilterVolumeThreshold(0.0f)
-		, bApplyEffects(false)
-		, bAlwaysPlay(false)
-		, bIsUISound(false)
-		, bIsMusic(false)
-		, bReverb(true)
-		, Default2DReverbSendAmount(0.0f)
-		, bCenterChannelOnly(false)
-		, bApplyAmbientVolumes(false)
-		, OutputTarget(EAudioOutputTarget::Speaker)
-		, LoadingBehavior(ESoundWaveLoadingBehavior::Inherited)
-		, AttenuationScaleParam(AttenuationDistanceScale)
-		, ParentAttenuationScale(1.0f)
-		{
-		}
+	FSoundClassProperties();
 
 private:
 	FDynamicParameter AttenuationScaleParam;
@@ -248,15 +203,15 @@ class ENGINE_API USoundClass : public UObject
 	GENERATED_UCLASS_BODY()
 
 	/** Configurable properties like volume and priority. */
-	UPROPERTY(EditAnywhere, BlueprintReadOnly, Category=SoundClass)
-	struct FSoundClassProperties Properties;
+	UPROPERTY(EditAnywhere, BlueprintReadOnly, Category = General, meta = (ShowOnlyInnerProperties))
+	FSoundClassProperties Properties;
 
-	UPROPERTY(EditAnywhere, BlueprintReadOnly, Category=SoundClass)
+	UPROPERTY(EditAnywhere, BlueprintReadOnly, Category = General)
 	TArray<USoundClass*> ChildClasses;
 
 	/** SoundMix Modifiers to activate automatically when a sound of this class is playing. */
-	UPROPERTY(EditAnywhere, BlueprintReadOnly, Category=SoundClass)
-	TArray<struct FPassiveSoundMixModifier> PassiveSoundMixModifiers;
+	UPROPERTY(EditAnywhere, BlueprintReadOnly, Category = General)
+	TArray<FPassiveSoundMixModifier> PassiveSoundMixModifiers;
 
 	/**
 	  * Modulation for the sound class. If not set on sound directly, settings
@@ -272,7 +227,7 @@ public:
 #if WITH_EDITORONLY_DATA
 	/** EdGraph based representation of the SoundClass */
 	class UEdGraph* SoundClassGraph;
-#endif
+#endif // WITH_EDITORONLY_DATA
 
 protected:
 
diff --git a/Engine/Source/Runtime/Engine/Classes/Sound/SoundSubmix.h b/Engine/Source/Runtime/Engine/Classes/Sound/SoundSubmix.h
--- a/Engine/Source/Runtime/Engine/Classes/Sound/SoundSubmix.h
+++ b/Engine/Source/Runtime/Engine/Classes/Sound/SoundSubmix.h
@@ -94,7 +94,7 @@ class ENGINE_API USoundSubmix : public UObject
 #endif
 
 	// Experimental! Specifies the channel format for the submix. Sources will be mixed at the specified format. Useful for specific effects that need to operate on a specific format.
-	UPROPERTY()
+	UPROPERTY(VisibleAnywhere, BlueprintReadOnly, Category = SoundSubmix)
 	ESubmixChannelFormat ChannelFormat;
 
 	/** Mute this submix when the application is muted or in the background. Used to prevent submix effect tails from continuing when tabbing out of application or if application is muted. */
diff --git a/Engine/Source/Runtime/Engine/Classes/Sound/SoundWave.h b/Engine/Source/Runtime/Engine/Classes/Sound/SoundWave.h
--- a/Engine/Source/Runtime/Engine/Classes/Sound/SoundWave.h
+++ b/Engine/Source/Runtime/Engine/Classes/Sound/SoundWave.h
@@ -12,10 +12,10 @@
 #include "Misc/Guid.h"
 #include "Async/AsyncWork.h"
 #include "Sound/SoundBase.h"
-#include "Sound/SoundClass.h"
 #include "Serialization/BulkData.h"
 #include "Serialization/BulkDataBuffer.h"
 #include "Sound/SoundGroups.h"
+#include "Sound/SoundWaveLoadingBehavior.h"
 #include "AudioMixerTypes.h"
 #include "AudioCompressionSettings.h"
 #include "PerPlatformProperties.h"
@@ -328,7 +328,7 @@ public:
 	UPROPERTY(EditAnywhere, Category = "Playback|Streaming", meta = (DisplayName = "Seekable", EditCondition = "bStreaming"))
 	uint8 bSeekableStreaming:1;
 
-	/** If stream caching is enabled, this can be used to specify how and when compressed audio data is loaded for this asset. */
+	/** Specifies how and when compressed audio data is loaded for asset if stream caching is enabled. */
 	UPROPERTY(EditAnywhere, Category = "Loading", meta = (DisplayName = "Loading Behavior Override"))
 	ESoundWaveLoadingBehavior LoadingBehavior;
 
@@ -361,7 +361,9 @@ public:
 	uint8 bVirtualizeWhenSilent_DEPRECATED:1;
 #endif // WITH_EDITORONLY_DATA
 
-	/** Whether or not this source is ambisonics file format. */
+	/** Whether or not this source is ambisonics file format. If set, sound always uses the 
+	  * 'Master Ambisonics Submix' as set in the 'Audio' category of Project Settings'
+	  * and ignores submix if provided locally or in the referenced SoundClass. */
 	UPROPERTY(EditAnywhere, Category = Format)
 	uint8 bIsAmbisonics : 1;
 
diff --git a/Engine/Source/Runtime/Engine/Classes/Sound/SoundWaveLoadingBehavior.h b/Engine/Source/Runtime/Engine/Classes/Sound/SoundWaveLoadingBehavior.h
new file mode 100644
--- /dev/null
+++ b/Engine/Source/Runtime/Engine/Classes/Sound/SoundWaveLoadingBehavior.h
@@ -0,0 +1,27 @@
+#pragma once
+
+/**
+ * Only used when stream caching is enabled. Determines how we are going to load or retain a given audio asset.
+ * A USoundWave's loading behavior can be overridden in the USoundWave itself, the sound wave's USoundClass, or by cvars.
+ * The order of priority is defined as:
+ * 1) The loading behavior set on the USoundWave
+ * 2) The loading behavior set on the USoundWave's USoundClass.
+ * 3) The loading behavior set on the nearest parent of a USoundWave's USoundClass.
+ * 4) The loading behavior set via the au.streamcache cvars.
+ */
+UENUM()
+enum class ESoundWaveLoadingBehavior : uint8
+{
+	// If set on a USoundWave, use the setting defined by the USoundClass. If set on the next parent USoundClass, or the default behavior defined via the au.streamcache cvars.
+	Inherited = 0,
+	// the first chunk of audio for this asset will be retained in the audio cache until a given USoundWave is either destroyed or USoundWave::ReleaseCompressedAudioData is called.
+	RetainOnLoad = 1,
+	// the first chunk of audio for this asset will be loaded into the cache from disk when this asset is loaded, but may be evicted to make room for other audio if it isn't played for a while.
+	PrimeOnLoad = 2,
+	// the first chunk of audio for this asset will not be loaded until this asset is played or primed.
+	LoadOnDemand = 3,
+	// Force all audio data for this audio asset to live outside of the cache and use the non-streaming decode pathways. Only usable if set on the USoundWave.
+	ForceInline = 4 UMETA(DisplayName = "Force Inline"),
+	// This value is used to delineate when the value of ESoundWaveLoadingBehavior hasn't been cached on a USoundWave yet.
+	Uninitialized = 0xff UMETA(Hidden)
+};
diff --git a/Engine/Source/Runtime/Engine/Private/ActiveSound.cpp b/Engine/Source/Runtime/Engine/Private/ActiveSound.cpp
--- a/Engine/Source/Runtime/Engine/Private/ActiveSound.cpp
+++ b/Engine/Source/Runtime/Engine/Private/ActiveSound.cpp
@@ -6,11 +6,16 @@
 #include "Misc/App.h"
 #include "AudioThread.h"
 #include "AudioDevice.h"
+#include "IAudioExtensionPlugin.h"
+#include "Sound/AudioSettings.h"
+#include "Sound/SoundClass.h"
 #include "Sound/SoundCue.h"
 #include "Sound/SoundWave.h"
 #include "Sound/SoundNodeAttenuation.h"
+#include "Sound/SoundSubmix.h"
 #include "SubtitleManager.h"
 
+
 static int32 AudioOcclusionDisabledCvar = 0;
 FAutoConsoleVariableRef CVarAudioOcclusionEnabled(
 	TEXT("au.DisableOcclusion"),
@@ -30,7 +35,6 @@ FActiveSound::FActiveSound()
 	, OwnerID(0)
 	, AudioDevice(nullptr)
 	, SoundClassOverride(nullptr)
-	, SoundSubmixOverride(nullptr)
 	, bHasCheckedOcclusion(false)
 	, bAllowSpatialization(true)
 	, bHasAttenuationSettings(false)
@@ -45,7 +49,6 @@ FActiveSound::FActiveSound()
 	, bHasExternalSubtitles(false)
 	, bLocationDefined(false)
 	, bIgnoreForFlushing(false)
-	, bEQFilterApplied(false)
 	, bAlwaysPlay(false)
 	, bIsUISound(false)
 	, bIsMusic(false)
@@ -321,15 +324,7 @@ USoundClass* FActiveSound::GetSoundClass() const
 
 USoundSubmix* FActiveSound::GetSoundSubmix() const
 {
-	if (SoundSubmixOverride)
-	{
-		return SoundSubmixOverride;
-	}
-	else if (Sound)
-	{
-		return Sound->GetSoundSubmix();
-	}
-	return nullptr;
+	return Sound ? Sound->GetSoundSubmix() : nullptr;
 }
 
 void FActiveSound::SetSubmixSend(const FSoundSubmixSendInfo& SubmixSendInfo)
@@ -566,15 +561,15 @@ void FActiveSound::UpdateWaveInstances(TArray<FWaveInstance*> &InWaveInstances,
 		}
 		else
 		{
-			ParseParams.ReverbSendMethod = EReverbSendMethod::Manual;		
+			ParseParams.ReverbSendMethod = EReverbSendMethod::Manual;
 			if (ParseParams.SoundClass)
 			{
 				ParseParams.ManualReverbSendLevel = ParseParams.SoundClass->Properties.Default2DReverbSendAmount;
-		}
-		else
-		{
-			ParseParams.ManualReverbSendLevel = AudioDevice->GetDefaultReverbSendLevel();
-		}
+			}
+			else
+			{
+				ParseParams.ManualReverbSendLevel = AudioDevice->GetDefaultReverbSendLevel();
+			}
 		}
 
 		ParseParams.ModulationPluginSettings = FindModulationSettings();
diff --git a/Engine/Source/Runtime/Engine/Private/Audio.cpp b/Engine/Source/Runtime/Engine/Private/Audio.cpp
--- a/Engine/Source/Runtime/Engine/Private/Audio.cpp
+++ b/Engine/Source/Runtime/Engine/Private/Audio.cpp
@@ -813,7 +813,6 @@ FWaveInstance::FWaveInstance(const UPTRINT InWaveInstanceHash, FActiveSound& InA
 	, bUseSpatialization(false)
 	, bEnableLowPassFilter(false)
 	, bIsOccluded(false)
-	, bEQFilterApplied(false)
 	, bIsUISound(false)
 	, bIsMusic(false)
 	, bReverb(true)
diff --git a/Engine/Source/Runtime/Engine/Private/AudioEffect.cpp b/Engine/Source/Runtime/Engine/Private/AudioEffect.cpp
--- a/Engine/Source/Runtime/Engine/Private/AudioEffect.cpp
+++ b/Engine/Source/Runtime/Engine/Private/AudioEffect.cpp
@@ -250,7 +250,7 @@ FAudioEffectsManager::FAudioEffectsManager( FAudioDevice* InDevice )
 	, CurrentEQMix(nullptr)
 	, bReverbActive(false)
 	, bEQActive(false)
-	, bReverbChanged(true) // Setting to true to catch the first default reverb setting
+	, bReverbChanged(true) // Initially true to update from default reverb setting
 	, bEQChanged(false)
 	, bUseLegacyReverb(false)
 {
diff --git a/Engine/Source/Runtime/Engine/Private/EnginePrivatePCH.h b/Engine/Source/Runtime/Engine/Private/EnginePrivatePCH.h
--- a/Engine/Source/Runtime/Engine/Private/EnginePrivatePCH.h
+++ b/Engine/Source/Runtime/Engine/Private/EnginePrivatePCH.h
@@ -445,7 +445,6 @@
 #include "BatchedElements.h"
 #include "GameFramework/Volume.h"
 #include "EngineGlobals.h"
-#include "Sound/SoundClass.h"
 #include "Audio.h"
 #include "Sound/SoundAttenuation.h"
 #include "IAudioExtensionPlugin.h"
diff --git a/Engine/Source/Runtime/Engine/Private/SoundClass.cpp b/Engine/Source/Runtime/Engine/Private/SoundClass.cpp
--- a/Engine/Source/Runtime/Engine/Private/SoundClass.cpp
+++ b/Engine/Source/Runtime/Engine/Private/SoundClass.cpp
@@ -11,11 +11,38 @@
 #include "UObject/UObjectHash.h"
 #include "UObject/UObjectIterator.h"
 #include "Sound/SoundMix.h"
+
 #if WITH_EDITOR
 #include "SoundClassGraph/SoundClassGraph.h"
 #include "Framework/Notifications/NotificationManager.h"
 #include "Widgets/Notifications/SNotificationList.h"
-#endif
+#endif // WITH_EDITOR
+
+FSoundClassProperties::FSoundClassProperties()
+	: Volume(1.0f)
+	, Pitch(1.0f)
+	, LowPassFilterFrequency(MAX_FILTER_FREQUENCY)
+	, AttenuationDistanceScale(1.0f)
+	, StereoBleed(0.25f)
+	, LFEBleed(0.5f)
+	, VoiceCenterChannelVolume(0.0f)
+	, RadioFilterVolume(0.0f)
+	, RadioFilterVolumeThreshold(0.0f)
+	, bApplyEffects(false)
+	, bAlwaysPlay(false)
+	, bIsUISound(false)
+	, bIsMusic(false)
+	, bCenterChannelOnly(false)
+	, bApplyAmbientVolumes(false)
+	, bReverb(true)
+	, Default2DReverbSendAmount(0.0f)
+	, OutputTarget(EAudioOutputTarget::Speaker)
+	, LoadingBehavior(ESoundWaveLoadingBehavior::Inherited)
+	, DefaultSubmix(nullptr)
+	, AttenuationScaleParam(AttenuationDistanceScale)
+	, ParentAttenuationScale(1.0f)
+{
+}
 
 void FSoundClassProperties::SetAttenuationDistanceScale(float InAttenuationScale, float InTime)
 {
diff --git a/Engine/Source/Runtime/Engine/Private/SoundWave.cpp b/Engine/Source/Runtime/Engine/Private/SoundWave.cpp
--- a/Engine/Source/Runtime/Engine/Private/SoundWave.cpp
+++ b/Engine/Source/Runtime/Engine/Private/SoundWave.cpp
@@ -1,20 +1,27 @@
 // Copyright 1998-2019 Epic Games, Inc. All Rights Reserved.
 
 #include "Sound/SoundWave.h"
+
+#include "ActiveSound.h"
+#include "Audio.h"
+#include "AudioCompressionSettingsUtils.h"
+#include "AudioDecompress.h"
+#include "AudioDerivedData.h"
+#include "AudioDevice.h"
+#include "AudioThread.h"
+#include "SampleBuffer.h"
 #include "Serialization/MemoryWriter.h"
+#include "Sound/AudioSettings.h"
+#include "Sound/SoundSubmix.h"
 #include "UObject/FrameworkObjectVersion.h"
 #include "UObject/Package.h"
 #include "UObject/UObjectIterator.h"
 #include "EngineDefines.h"
 #include "Components/AudioComponent.h"
 #include "ContentStreaming.h"
-#include "ActiveSound.h"
-#include "AudioThread.h"
-#include "AudioDevice.h"
-#include "AudioDecompress.h"
 #include "Interfaces/ITargetPlatform.h"
 #include "Interfaces/ITargetPlatformManagerModule.h"
-#include "AudioDerivedData.h"
+#include "Sound/SoundClass.h"
 #include "SubtitleManager.h"
 #include "DerivedDataCacheInterface.h"
 #include "EditorFramework/AssetImportData.h"
@@ -22,12 +29,10 @@
 #include "HAL/LowLevelMemTracker.h"
 #include "HAL/IConsoleManager.h"
 #include "HAL/FileManager.h"
-#include "AudioCompressionSettingsUtils.h"
 #include "DSP/SpectrumAnalyzer.h"
 #include "DSP/EnvelopeFollower.h"
 #include "DSP/BufferVectorOperations.h"
 #include "Misc/OutputDeviceArchiveWrapper.h"
-#include "SampleBuffer.h"
 
 #include "Misc/CommandLine.h"
 
@@ -1725,7 +1730,9 @@ void USoundWave::Parse(FAudioDevice* AudioDevice, const UPTRINT NodeWaveInstance
 	WaveInstance->SoundClass = ParseParams.SoundClass;
 	if (ParseParams.SoundClass)
 	{
-		FSoundClassProperties* SoundClassProperties = AudioDevice->GetSoundClassCurrentProperties(ParseParams.SoundClass);
+		const FSoundClassProperties* SoundClassProperties = AudioDevice->GetSoundClassCurrentProperties(ParseParams.SoundClass);
+		check(SoundClassProperties);
+
 		// Use values from "parsed/ propagated" sound class properties
 		float VolumeMultiplier = WaveInstance->GetVolumeMultiplier();
 		WaveInstance->SetVolumeMultiplier(VolumeMultiplier* SoundClassProperties->Volume);
@@ -1743,10 +1750,14 @@ void USoundWave::Parse(FAudioDevice* AudioDevice, const UPTRINT NodeWaveInstance
 		WaveInstance->bIsUISound = ActiveSound.bIsUISound || SoundClassProperties->bIsUISound;
 		WaveInstance->bIsMusic = ActiveSound.bIsMusic || SoundClassProperties->bIsMusic;
 		WaveInstance->bCenterChannelOnly = ActiveSound.bCenterChannelOnly || SoundClassProperties->bCenterChannelOnly;
-		WaveInstance->bEQFilterApplied = ActiveSound.bEQFilterApplied || SoundClassProperties->bApplyEffects;
 		WaveInstance->bReverb = ActiveSound.bReverb || SoundClassProperties->bReverb;
 		WaveInstance->OutputTarget = SoundClassProperties->OutputTarget;
 
+		if (SoundClassProperties->DefaultSubmix)
+		{
+			WaveInstance->SoundSubmix = SoundClassProperties->DefaultSubmix;
+		}
+
 		if (SoundClassProperties->bApplyAmbientVolumes)
 		{
 			VolumeMultiplier = WaveInstance->GetVolumeMultiplier();
@@ -1764,7 +1775,7 @@ void USoundWave::Parse(FAudioDevice* AudioDevice, const UPTRINT NodeWaveInstance
 		WaveInstance->RadioFilterVolumeThreshold = 0.f;
 		WaveInstance->StereoBleed = 0.f;
 		WaveInstance->LFEBleed = 0.f;
-		WaveInstance->bEQFilterApplied = ActiveSound.bEQFilterApplied;
+		WaveInstance->SoundSubmix = ActiveSound.GetSoundSubmix();
 		WaveInstance->bIsUISound = ActiveSound.bIsUISound;
 		WaveInstance->bIsMusic = ActiveSound.bIsMusic;
 		WaveInstance->bReverb = ActiveSound.bReverb;
@@ -1773,6 +1784,12 @@ void USoundWave::Parse(FAudioDevice* AudioDevice, const UPTRINT NodeWaveInstance
 		bAlwaysPlay = ActiveSound.bAlwaysPlay;
 	}
 
+	WaveInstance->bIsAmbisonics = bIsAmbisonics;
+	if (ParseParams.SoundSubmix)
+	{
+		WaveInstance->SoundSubmix = ParseParams.SoundSubmix;
+	}
+
 	// If set to bAlwaysPlay, increase the current sound's priority scale by 10x. This will still result in a possible 0-priority output if the sound has 0 actual volume
 	if (bAlwaysPlay)
 	{
@@ -1814,10 +1831,9 @@ void USoundWave::Parse(FAudioDevice* AudioDevice, const UPTRINT NodeWaveInstance
 		// Apply stereo normalization to wave instances if enabled
 		if (ParseParams.bApplyNormalizationToStereoSounds && NumChannels == 2)
 		{
-			const float ThisVolumeMultiplier = WaveInstance->GetVolumeMultiplier();
-			WaveInstance->SetVolumeMultiplier(ThisVolumeMultiplier * 0.5f);
+			float WaveInstanceVolume = WaveInstance->GetVolume();
+			WaveInstance->SetVolume(WaveInstanceVolume * 0.5f);
 		}
-
 	}
 
 	// Copy reverb send settings
@@ -1832,7 +1848,6 @@ void USoundWave::Parse(FAudioDevice* AudioDevice, const UPTRINT NodeWaveInstance
 	WaveInstance->EnvelopeFollowerReleaseTime = ParseParams.EnvelopeFollowerReleaseTime;
 
 	// Copy over the submix sends.
-	WaveInstance->SoundSubmix = ParseParams.SoundSubmix;
 	WaveInstance->SoundSubmixSends = ParseParams.SoundSubmixSends;
 
 	// Copy over the source bus send and data
@@ -1852,8 +1867,6 @@ void USoundWave::Parse(FAudioDevice* AudioDevice, const UPTRINT NodeWaveInstance
 	WaveInstance->ReverbPluginSettings = ParseParams.ReverbPluginSettings;
 	WaveInstance->ModulationPluginSettings = ParseParams.ModulationPluginSettings;
 
-	WaveInstance->bIsAmbisonics = bIsAmbisonics;
-
 	if (WaveInstance->IsPlaying())
 	{
 		WaveInstances.Add(WaveInstance);
diff --git a/Engine/Source/Runtime/Engine/Public/ActiveSound.h b/Engine/Source/Runtime/Engine/Public/ActiveSound.h
--- a/Engine/Source/Runtime/Engine/Public/ActiveSound.h
+++ b/Engine/Source/Runtime/Engine/Public/ActiveSound.h
@@ -346,9 +346,6 @@ private:
 	/** Optional SoundClass to override for the sound. */
 	USoundClass* SoundClassOverride;
 
-	/** Optional SoundSubmix to override for the sound. */
-	USoundSubmix* SoundSubmixOverride;
-
 	/** Optional override the submix sends for the sound. */
 	TArray<FSoundSubmixSendInfo> SoundSubmixSendsOverride;
 
@@ -412,9 +409,6 @@ public:
 	/** If true, this sound will not be stopped when flushing the audio device. */
 	uint8 bIgnoreForFlushing:1;
 
-	/** Whether audio effects are applied */
-	uint8 bEQFilterApplied:1;
-
 	/** Whether to artificially prioritize the component to play */
 	uint8 bAlwaysPlay:1;
 
diff --git a/Engine/Source/Runtime/Engine/Public/Audio.h b/Engine/Source/Runtime/Engine/Public/Audio.h
--- a/Engine/Source/Runtime/Engine/Public/Audio.h
+++ b/Engine/Source/Runtime/Engine/Public/Audio.h
@@ -10,7 +10,7 @@
 #include "AudioDefines.h"
 #include "Stats/Stats.h"
 #include "HAL/ThreadSafeBool.h"
-#include "Sound/SoundClass.h"
+#include "Sound/AudioOutputTarget.h"
 #include "Sound/SoundAttenuation.h"
 #include "Sound/SoundEffectSource.h"
 #include "Sound/SoundSubmixSend.h"
@@ -262,9 +262,6 @@ public:
 	/** Whether or not the sound is occluded. */
 	uint32 bIsOccluded:1;
 
-	/** Whether to apply audio effects */
-	uint32 bEQFilterApplied:1;
-
 	/** Whether or not this sound plays when the game is paused in the UI */
 	uint32 bIsUISound:1;
 
@@ -632,9 +629,6 @@ public:
 	/** Returns true if reverb should be applied. */
 	bool IsReverbApplied() const { return bReverbApplied; }
 
-	/** Returns true if EQ should be applied. */
-	bool IsEQFilterApplied() const  { return WaveInstance->bEQFilterApplied; }
-
 	/** Set the bReverbApplied variable. */
 	ENGINE_API bool SetReverbApplied(bool bHardwareAvailable);
 
diff --git a/Engine/Source/Runtime/Engine/Public/AudioDevice.h b/Engine/Source/Runtime/Engine/Public/AudioDevice.h
--- a/Engine/Source/Runtime/Engine/Public/AudioDevice.h
+++ b/Engine/Source/Runtime/Engine/Public/AudioDevice.h
@@ -40,7 +40,6 @@ class ICompressedAudioInfo;
 class UReverbEffect;
 class USoundAttenuation;
 class USoundBase;
-class USoundClass;
 class USoundConcurrency;
 class USoundEffectSourcePreset;
 class USoundEffectSubmixPreset;
@@ -1758,7 +1757,7 @@ private:
 	TMap<FWaveInstance*, FSoundSource*>	WaveInstanceSourceMap;
 
 	/** Current properties of all sound classes */
-	TMap<USoundClass*, FSoundClassProperties>	SoundClasses;
+	TMap<USoundClass*, FSoundClassProperties> SoundClasses;
 
 	/** The Base SoundMix that's currently active */
 	USoundMix* BaseSoundMix;
diff --git a/Engine/Source/Runtime/Engine/Public/Engine.h b/Engine/Source/Runtime/Engine/Public/Engine.h
--- a/Engine/Source/Runtime/Engine/Public/Engine.h
+++ b/Engine/Source/Runtime/Engine/Public/Engine.h
@@ -295,7 +295,6 @@ MONOLITHIC_HEADER_BOILERPLATE()
 #include "Engine/SkeletalMeshSocket.h"
 #include "Animation/Skeleton.h"
 #include "Sound/DialogueSoundWaveProxy.h"
-#include "Sound/SoundClass.h"
 #include "Sound/SoundCue.h"
 #include "Sound/SoundGroups.h"
 #include "Sound/SoundMix.h"
diff --git a/Engine/Source/Runtime/Engine/Public/EngineSharedPCH.h b/Engine/Source/Runtime/Engine/Public/EngineSharedPCH.h
--- a/Engine/Source/Runtime/Engine/Public/EngineSharedPCH.h
+++ b/Engine/Source/Runtime/Engine/Public/EngineSharedPCH.h
@@ -560,7 +560,6 @@
 #include "Math/GenericOctreePublic.h"
 #include "GameFramework/OnlineReplStructs.h"
 #include "Audio.h"
-#include "Sound/SoundClass.h"
 #include "Sound/SoundAttenuation.h"
 #include "IAudioExtensionPlugin.h"
 #include "Tickable.h"
diff --git a/Engine/Source/Runtime/Mac/CoreAudio/Private/CoreAudioSource.cpp b/Engine/Source/Runtime/Mac/CoreAudio/Private/CoreAudioSource.cpp
--- a/Engine/Source/Runtime/Mac/CoreAudio/Private/CoreAudioSource.cpp
+++ b/Engine/Source/Runtime/Mac/CoreAudio/Private/CoreAudioSource.cpp
@@ -851,7 +851,7 @@ bool FCoreAudioSoundSource::AttachToAUGraph()
 	bool bNeedReverbFilter = false;
 
 #if CORE_AUDIO_EQ_ENABLED
-	bNeedEQFilter = IsEQFilterApplied();
+	bNeedEQFilter = WaveInstance->SoundClass ? WaveInstance->SoundClass->Properties.bApplyEffects : false;
 #endif
 
 #if CORE_AUDIO_RADIO_ENABLED
diff --git a/Engine/Source/Runtime/MediaAssets/Public/MediaSoundComponent.h b/Engine/Source/Runtime/MediaAssets/Public/MediaSoundComponent.h
--- a/Engine/Source/Runtime/MediaAssets/Public/MediaSoundComponent.h
+++ b/Engine/Source/Runtime/MediaAssets/Public/MediaSoundComponent.h
@@ -15,7 +15,7 @@
 #include "DSP/SpectrumAnalyzer.h"
 #include "DSP/BufferVectorOperations.h"
 #include "DSP/EnvelopeFollower.h"
-
+#include "Sound/SoundClass.h"
 
 #include "MediaSoundComponent.generated.h"
 
@@ -24,7 +24,6 @@ class FMediaPlayerFacade;
 class IMediaAudioSample;
 class IMediaPlayer;
 class UMediaPlayer;
-class USoundClass;
 
 
 /**
diff --git a/Engine/Source/Runtime/Windows/XAudio2/Private/XAudio2Source.cpp b/Engine/Source/Runtime/Windows/XAudio2/Private/XAudio2Source.cpp
--- a/Engine/Source/Runtime/Windows/XAudio2/Private/XAudio2Source.cpp
+++ b/Engine/Source/Runtime/Windows/XAudio2/Private/XAudio2Source.cpp
@@ -388,7 +388,8 @@ bool FXAudio2SoundSource::CreateSource( void )
 	// EQFilter Causes some sound devices to lag and starve important game threads. Hack disable until a long term solution is put into place.
 
 	const bool bIsEQDisabled = GetDefault<UAudioSettings>()->bDisableMasterEQ;
-	if (!bIsEQDisabled && IsEQFilterApplied())
+	const bool bIsEQFilterApplied = WaveInstance->SoundClass ? WaveInstance->SoundClass->Properties.bApplyEffects : false;
+	if (!bIsEQDisabled && bIsEQFilterApplied)
 	{
 		Destinations[NumSends].pOutputVoice = Effects->EQPremasterVoice;
 	}
