Added delegate version for AssetManagerRequireResource that returns a list of still missing chunks.
Optimize UAssetManager to move delegate values around rather than copying when possible.

diff --git a/Engine/Source/Runtime/Engine/Classes/Engine/AssetManager.h b/Engine/Source/Runtime/Engine/Classes/Engine/AssetManager.h
--- a/Engine/Source/Runtime/Engine/Classes/Engine/AssetManager.h
+++ b/Engine/Source/Runtime/Engine/Classes/Engine/AssetManager.h
@@ -19,8 +19,12 @@ struct FPrimaryAssetTypeData;
 struct FPrimaryAssetData;
 struct FPrimaryAssetRulesCustomOverride;
 
+/** Delegate called when acquiring resources/chunks for assets, first parameter will be true if all resources were acquired, false if any failed.
+	Second parameter will contain a list of chunks not yet downloaded */
+DECLARE_DELEGATE_TwoParams(FAssetManagerAcquireResourceDelegateEx, bool /* bSuccess */, const TArray<int32>& /* MissingChunks */);
+
 /** Delegate called when acquiring resources/chunks for assets, parameter will be true if all resources were acquired, false if any failed */
-DECLARE_DELEGATE_OneParam(FAssetManagerAcquireResourceDelegate, bool);
+DECLARE_DELEGATE_OneParam(FAssetManagerAcquireResourceDelegate, bool /* bSuccess */);
 
 /** 
  * A singleton UObject that is responsible for loading and unloading PrimaryAssets, and maintaining game-specific asset references
@@ -298,6 +302,7 @@ public:
 	 * @param CompleteDelegate		Delegate called when chunks have been acquired or failed. If any chunks fail the entire operation is considered a failure
 	 * @param Priority				Priority to use when acquiring chunks
 	 */
+	virtual void AcquireResourcesForAssetList(const TArray<FSoftObjectPath>& AssetList, FAssetManagerAcquireResourceDelegateEx CompleteDelegate, EChunkPriority::Type Priority = EChunkPriority::Immediate);
 	virtual void AcquireResourcesForAssetList(const TArray<FSoftObjectPath>& AssetList, FAssetManagerAcquireResourceDelegate CompleteDelegate, EChunkPriority::Type Priority = EChunkPriority::Immediate);
 
 	/** 
diff --git a/Engine/Source/Runtime/Engine/Private/AssetManager.cpp b/Engine/Source/Runtime/Engine/Private/AssetManager.cpp
--- a/Engine/Source/Runtime/Engine/Private/AssetManager.cpp
+++ b/Engine/Source/Runtime/Engine/Private/AssetManager.cpp
@@ -1308,7 +1308,7 @@ TSharedPtr<FStreamableHandle> UAssetManager::ChangeBundleStateForMatchingPrimary
 	if (GetPrimaryAssetsWithBundleState(AssetsToChange, TArray<FPrimaryAssetType>(), OldBundles))
 	{
 		// This will call delegate when done
-		return ChangeBundleStateForPrimaryAssets(AssetsToChange, NewBundles, OldBundles, false, DelegateToCall, Priority);
+		return ChangeBundleStateForPrimaryAssets(AssetsToChange, NewBundles, OldBundles, false, MoveTemp(DelegateToCall), Priority);
 	}
 
 	// Nothing to transition, call delegate now
@@ -1376,7 +1376,7 @@ TSharedPtr<FStreamableHandle> UAssetManager::PreloadPrimaryAssets(const TArray<F
 		}
 	}
 
-	ReturnHandle = LoadAssetList(PathsToLoad.Array(), DelegateToCall, Priority, DebugName);
+	ReturnHandle = LoadAssetList(PathsToLoad.Array(), MoveTemp(DelegateToCall), Priority, DebugName);
 
 	if (!ensureMsgf(ReturnHandle.IsValid(), TEXT("Requested preload of Primary Asset with no referenced assets!")))
 	{
@@ -1415,19 +1415,19 @@ void UAssetManager::OnAssetStateChangeCompleted(FPrimaryAssetId PrimaryAssetId,
 
 TSharedPtr<FStreamableHandle> UAssetManager::LoadPrimaryAssets(const TArray<FPrimaryAssetId>& AssetsToLoad, const TArray<FName>& LoadBundles, FStreamableDelegate DelegateToCall, TAsyncLoadPriority Priority)
 {
-	return ChangeBundleStateForPrimaryAssets(AssetsToLoad, LoadBundles, TArray<FName>(), true, DelegateToCall, Priority);
+	return ChangeBundleStateForPrimaryAssets(AssetsToLoad, LoadBundles, TArray<FName>(), true, MoveTemp(DelegateToCall), Priority);
 }
 
 TSharedPtr<FStreamableHandle> UAssetManager::LoadPrimaryAsset(const FPrimaryAssetId& AssetToLoad, const TArray<FName>& LoadBundles, FStreamableDelegate DelegateToCall, TAsyncLoadPriority Priority)
 {
-	return LoadPrimaryAssets(TArray<FPrimaryAssetId>{AssetToLoad}, LoadBundles, DelegateToCall, Priority);
+	return LoadPrimaryAssets(TArray<FPrimaryAssetId>{AssetToLoad}, LoadBundles, MoveTemp(DelegateToCall), Priority);
 }
 
 TSharedPtr<FStreamableHandle> UAssetManager::LoadPrimaryAssetsWithType(FPrimaryAssetType PrimaryAssetType, const TArray<FName>& LoadBundles, FStreamableDelegate DelegateToCall, TAsyncLoadPriority Priority)
 {
 	TArray<FPrimaryAssetId> Assets;
 	GetPrimaryAssetIdList(PrimaryAssetType, Assets);
-	return LoadPrimaryAssets(Assets, LoadBundles, DelegateToCall, Priority);
+	return LoadPrimaryAssets(Assets, LoadBundles, MoveTemp(DelegateToCall), Priority);
 }
 
 TSharedPtr<FStreamableHandle> UAssetManager::GetPrimaryAssetHandle(const FPrimaryAssetId& PrimaryAssetId, bool bForceCurrent, TArray<FName>* Bundles) const
@@ -1592,7 +1592,7 @@ TSharedPtr<FStreamableHandle> UAssetManager::LoadAssetList(const TArray<FSoftObj
 	}
 	else
 	{
-		NewHandle = StreamableManager.RequestAsyncLoad(AssetList, DelegateToCall, Priority, false, MissingChunks.Num() > 0, DebugName);
+		NewHandle = StreamableManager.RequestAsyncLoad(AssetList, MoveTemp(DelegateToCall), Priority, false, MissingChunks.Num() > 0, DebugName);
 
 		if (MissingChunks.Num() > 0 && NewHandle.IsValid())
 		{
@@ -1737,7 +1737,7 @@ bool UAssetManager::FindMissingChunkList(const TArray<FSoftObjectPath>& AssetLis
 void UAssetManager::AcquireChunkList(const TArray<int32>& ChunkList, FAssetManagerAcquireResourceDelegate CompleteDelegate, EChunkPriority::Type Priority, TSharedPtr<FStreamableHandle> StalledHandle)
 {
 	FPendingChunkInstall* PendingChunkInstall = new(PendingChunkInstalls) FPendingChunkInstall;
-	PendingChunkInstall->ManualCallback = CompleteDelegate;
+	PendingChunkInstall->ManualCallback = MoveTemp(CompleteDelegate);
 	PendingChunkInstall->RequestedChunks = ChunkList;
 	PendingChunkInstall->PendingChunks = ChunkList;
 	PendingChunkInstall->StalledStreamableHandle = StalledHandle;
@@ -1757,29 +1757,29 @@ void UAssetManager::AcquireChunkList(const TArray<int32>& ChunkList, FAssetManag
 
 void UAssetManager::AcquireResourcesForAssetList(const TArray<FSoftObjectPath>& AssetList, FAssetManagerAcquireResourceDelegate CompleteDelegate, EChunkPriority::Type Priority)
 {
-	TArray<int32> MissingChunks, ErrorChunks;
+	AcquireResourcesForAssetList(AssetList, FAssetManagerAcquireResourceDelegateEx::CreateLambda([CompleteDelegate = MoveTemp(CompleteDelegate)](bool bSuccess, const TArray<int32>& Unused) { CompleteDelegate.ExecuteIfBound(bSuccess); }), Priority);
+}
 
+void UAssetManager::AcquireResourcesForAssetList(const TArray<FSoftObjectPath>& AssetList, FAssetManagerAcquireResourceDelegateEx CompleteDelegate, EChunkPriority::Type Priority)
+{
+	TArray<int32> MissingChunks, ErrorChunks;
 	FindMissingChunkList(AssetList, MissingChunks, ErrorChunks);
-
 	if (ErrorChunks.Num() > 0)
 	{
 		// At least one chunk doesn't exist, fail
-		FStreamableDelegate TempDelegate = FStreamableDelegate::CreateLambda([CompleteDelegate]() {CompleteDelegate.ExecuteIfBound(false); });
+		FStreamableDelegate TempDelegate = FStreamableDelegate::CreateLambda([CompleteDelegate = MoveTemp(CompleteDelegate), MissingChunks]() { CompleteDelegate.ExecuteIfBound(false, MissingChunks); });
 		FStreamableHandle::ExecuteDelegate(TempDelegate);
-
-		return;
 	}
-
-	if (MissingChunks.Num() == 0)
+	else if (MissingChunks.Num() == 0)
 	{
 		// All here, schedule the callback
-		FStreamableDelegate TempDelegate = FStreamableDelegate::CreateLambda([CompleteDelegate]() {CompleteDelegate.ExecuteIfBound(true); });
+		FStreamableDelegate TempDelegate = FStreamableDelegate::CreateLambda([CompleteDelegate = MoveTemp(CompleteDelegate)]() { CompleteDelegate.ExecuteIfBound(true, TArray<int32>()); });
 		FStreamableHandle::ExecuteDelegate(TempDelegate);
-
-		return;
 	}
-	
-	AcquireChunkList(MissingChunks, CompleteDelegate, Priority, nullptr);
+	else
+	{
+		AcquireChunkList(MissingChunks, FAssetManagerAcquireResourceDelegate::CreateLambda([CompleteDelegate = MoveTemp(CompleteDelegate), MissingChunks](bool bSuccess) { CompleteDelegate.ExecuteIfBound(bSuccess, MissingChunks); }), Priority, nullptr);
+	}
 }
 
 void UAssetManager::AcquireResourcesForPrimaryAssetList(const TArray<FPrimaryAssetId>& PrimaryAssetList, FAssetManagerAcquireResourceDelegate CompleteDelegate, EChunkPriority::Type Priority)
@@ -1814,7 +1814,7 @@ void UAssetManager::AcquireResourcesForPrimaryAssetList(const TArray<FPrimaryAss
 		}
 	}
 
-	AcquireResourcesForAssetList(PathsToLoad.Array(), CompleteDelegate, Priority);
+	AcquireResourcesForAssetList(PathsToLoad.Array(), MoveTemp(CompleteDelegate), Priority);
 }
 
 bool UAssetManager::GetResourceAcquireProgress(int32& OutAcquiredCount, int32& OutRequestedCount) const
