Missing copyright notices
#rnx

diff --git a/Engine/Source/Runtime/Engine/Classes/Sound/AudioOutputTarget.h b/Engine/Source/Runtime/Engine/Classes/Sound/AudioOutputTarget.h
--- a/Engine/Source/Runtime/Engine/Classes/Sound/AudioOutputTarget.h
+++ b/Engine/Source/Runtime/Engine/Classes/Sound/AudioOutputTarget.h
@@ -1,3 +1,4 @@
+// Copyright 1998-2019 Epic Games, Inc. All Rights Reserved.
 #pragma once
 
 #include "UObject/ObjectMacros.h"
diff --git a/Engine/Source/Runtime/Engine/Classes/Sound/SoundWaveLoadingBehavior.h b/Engine/Source/Runtime/Engine/Classes/Sound/SoundWaveLoadingBehavior.h
--- a/Engine/Source/Runtime/Engine/Classes/Sound/SoundWaveLoadingBehavior.h
+++ b/Engine/Source/Runtime/Engine/Classes/Sound/SoundWaveLoadingBehavior.h
@@ -1,3 +1,4 @@
+// Copyright 1998-2019 Epic Games, Inc. All Rights Reserved.
 #pragma once
 
 /**
