One more strict compile error straggler
#rnx

diff --git a/Engine/Source/Runtime/Engine/Private/SoundNodeQualityLevel.cpp b/Engine/Source/Runtime/Engine/Private/SoundNodeQualityLevel.cpp
--- a/Engine/Source/Runtime/Engine/Private/SoundNodeQualityLevel.cpp
+++ b/Engine/Source/Runtime/Engine/Private/SoundNodeQualityLevel.cpp
@@ -11,8 +11,9 @@
 #include "Engine/Engine.h"
 
 #if WITH_EDITORONLY_DATA
-#include "Settings/LevelEditorPlaySettings.h"
+#include "EdGraph/EdGraph.h"
 #include "Editor.h"
+#include "Settings/LevelEditorPlaySettings.h"
 #endif
 
 #if WITH_EDITOR
