Back out changelist 10680686, 10682670, 10687665, 10690466, and 10690798

Possible cause of EdTest errors

diff --git a/Engine/Source/Runtime/Engine/Private/PhysicsEngine/Experimental/PhysInterface_Chaos.cpp b/Engine/Source/Runtime/Engine/Private/PhysicsEngine/Experimental/PhysInterface_Chaos.cpp
--- a/Engine/Source/Runtime/Engine/Private/PhysicsEngine/Experimental/PhysInterface_Chaos.cpp
+++ b/Engine/Source/Runtime/Engine/Private/PhysicsEngine/Experimental/PhysInterface_Chaos.cpp
@@ -1399,9 +1399,9 @@ const Chaos::FConvex& FPhysicsGeometryCollection_Chaos::GetConvexGeometry() cons
 	return Geom.GetObjectChecked<Chaos::FConvex>();
 }
 
-const Chaos::FTriangleMeshImplicitObject& FPhysicsGeometryCollection_Chaos::GetTriMeshGeometry() const
+const Chaos::TTriangleMeshImplicitObject<float>& FPhysicsGeometryCollection_Chaos::GetTriMeshGeometry() const
 {
-	return Geom.GetObjectChecked<Chaos::FTriangleMeshImplicitObject>();
+	return Geom.GetObjectChecked<Chaos::TTriangleMeshImplicitObject<float>>();
 }
 
 FPhysicsGeometryCollection_Chaos::FPhysicsGeometryCollection_Chaos(const FPhysicsShapeReference_Chaos& InShape)
