Rename GetActiveProfileName to a more suited name, since it rebuild the device profile name ech time. Made GetDeviceProfileName deprecated, and added a new function that get the current device profile name and fallback on the platform device name.

#rnx

diff --git a/Engine/Source/Runtime/Engine/Classes/DeviceProfiles/DeviceProfileManager.h b/Engine/Source/Runtime/Engine/Classes/DeviceProfiles/DeviceProfileManager.h
--- a/Engine/Source/Runtime/Engine/Classes/DeviceProfiles/DeviceProfileManager.h
+++ b/Engine/Source/Runtime/Engine/Classes/DeviceProfiles/DeviceProfileManager.h
@@ -121,13 +121,29 @@ public:
 	*/
 	void GetAllPossibleParentProfiles(const UDeviceProfile* ChildProfile, OUT TArray<UDeviceProfile*>& PossibleParentProfiles) const;
 
+	/**
+	* Get the current active profile name.
+	*
+	* @return The selected profile.
+	*/
+	const FString GetActiveDeviceProfileName();
+
 	/**
 	* Get the selected device profile name, either the platform name, or the name
 	* provided by a Device Profile Selector Module.
 	*
 	* @return The selected profile.
 	*/
+	UE_DEPRECATED(4.24, "Use either GetActiveDeviceProfileName to have the current active device profile or GetPlatformDeviceProfileName to have the default one. Note, GetActiveDeviceProfileName will fallback on GetPlatformDeviceProfileName, if there is no active device profile ")
 	static const FString GetActiveProfileName();
+
+	/**
+	* Get the selected device profile name, either the platform name, or the name
+	* provided by a Device Profile Selector Module.
+	*
+	* @return The selected profile.
+	*/
+	static const FString GetPlatformDeviceProfileName();
 	
 	/** Retrieves the value of a scalability group cvar if it was set by the active device profile. */
 	static bool GetScalabilityCVar(const FString& CvarName, int32& OutValue);
diff --git a/Engine/Source/Runtime/Engine/Private/ChartCreation.cpp b/Engine/Source/Runtime/Engine/Private/ChartCreation.cpp
--- a/Engine/Source/Runtime/Engine/Private/ChartCreation.cpp
+++ b/Engine/Source/Runtime/Engine/Private/ChartCreation.cpp
@@ -661,7 +661,7 @@ void FPerformanceTrackingChart::Reset(const FDateTime& InStartTime)
 
 	StartBatteryLevel = -1;
 	StopBatteryLevel = -1;
-	DeviceProfileName = UDeviceProfileManager::GetActiveProfileName();
+	DeviceProfileName = UDeviceProfileManager::Get().GetActiveDeviceProfileName();
 	bIsChartingPaused = false;
 }
 
@@ -712,7 +712,7 @@ void FPerformanceTrackingChart::StartCharting()
 {
 	StartTemperatureLevel = FPlatformMisc::GetDeviceTemperatureLevel();
 	StartBatteryLevel = FPlatformMisc::GetBatteryLevel();
-	DeviceProfileName = UDeviceProfileManager::GetActiveProfileName();
+	DeviceProfileName = UDeviceProfileManager::Get().GetActiveDeviceProfileName();
 	bIsChartingPaused = false;
 }
 
@@ -734,7 +734,7 @@ void FPerformanceTrackingChart::ResumeCharting()
 
 void FPerformanceTrackingChart::OnDeviceProfileManagerUpdated()
 {
-	FString CurrentDeviceProfileName = UDeviceProfileManager::GetActiveProfileName();
+	FString CurrentDeviceProfileName = UDeviceProfileManager::Get().GetActiveDeviceProfileName();
 	if (CurrentDeviceProfileName != DeviceProfileName)
 	{
 		DeviceProfileName = TEXT("Mixed");
