Fixed mac compile issue

diff --git a/Engine/Source/Runtime/Engine/Public/MaterialShared.h b/Engine/Source/Runtime/Engine/Public/MaterialShared.h
--- a/Engine/Source/Runtime/Engine/Public/MaterialShared.h
+++ b/Engine/Source/Runtime/Engine/Public/MaterialShared.h
@@ -2605,8 +2605,8 @@ private:
 	ENGINE_API FMaterialAttributeDefintion* Find(EMaterialProperty Property);
 
 	// Helper functions to determine display name based on shader model, material domain, etc.
-	static FText GetAttributeOverrideForMaterial(const FGuid& AttributeID, UMaterial* Material);
-	static FString GetPinNameFromShadingModelField(FMaterialShadingModelField InShadingModels, const TArray<TKeyValuePair<EMaterialShadingModel, FString>>& InCustomShadingModelPinNames, const FString& InDefaultPinName);
+	ENGINE_API static FText GetAttributeOverrideForMaterial(const FGuid& AttributeID, UMaterial* Material);
+	ENGINE_API static FString GetPinNameFromShadingModelField(FMaterialShadingModelField InShadingModels, const TArray<TKeyValuePair<EMaterialShadingModel, FString>>& InCustomShadingModelPinNames, const FString& InDefaultPinName);
 
 	ENGINE_API static FMaterialAttributeDefinitionMap GMaterialPropertyAttributesMap;
 
