Fix forward declaring old TTriangleMeshImplicitObject type instead of FTriangleMeshImplicitObject.

diff --git a/Engine/Source/Runtime/Engine/Public/AI/NavigationSystemHelpers.h b/Engine/Source/Runtime/Engine/Public/AI/NavigationSystemHelpers.h
--- a/Engine/Source/Runtime/Engine/Public/AI/NavigationSystemHelpers.h
+++ b/Engine/Source/Runtime/Engine/Public/AI/NavigationSystemHelpers.h
@@ -30,8 +30,7 @@ namespace Chaos
 	template<typename T>
 	class THeightField;
 
-	template<typename T>
-	class TTriangleMeshImplicitObject;
+	class FTriangleMeshImplicitObject;
 }
 #endif
 
