Missing forward declare

diff --git a/Engine/Source/Runtime/Engine/Public/AI/NavigationSystemHelpers.h b/Engine/Source/Runtime/Engine/Public/AI/NavigationSystemHelpers.h
--- a/Engine/Source/Runtime/Engine/Public/AI/NavigationSystemHelpers.h
+++ b/Engine/Source/Runtime/Engine/Public/AI/NavigationSystemHelpers.h
@@ -11,6 +11,7 @@ class AActor;
 class UBodySetup;
 struct FCompositeNavModifier;
 struct FKAggregateGeom;
+struct FKConvexElem;
 struct FNavHeightfieldSamples;
 
 template<typename InElementType> class TNavStatArray;
