Remove the redundant macro MATERIAL_TWO_SIDED

diff --git a/Engine/Shaders/Private/MobileBasePassPixelShader.usf b/Engine/Shaders/Private/MobileBasePassPixelShader.usf
--- a/Engine/Shaders/Private/MobileBasePassPixelShader.usf
+++ b/Engine/Shaders/Private/MobileBasePassPixelShader.usf
@@ -634,7 +634,7 @@ void Main(
 		#if ENABLE_SKY_LIGHT
 			//@mw todo
 			// TODO: Also need to do specular.
-#if MATERIAL_TWO_SIDED && LQ_TEXTURE_LIGHTMAP
+#if MATERIAL_TWOSIDED && LQ_TEXTURE_LIGHTMAP
 			if (NoL == 0)
 			{
 #endif
@@ -644,7 +644,7 @@ void Main(
 			WaterDiffuseIndirectLuminance += SkyDiffuseLighting;
 		#endif
 			Color += SkyDiffuseLighting * half3(ResolvedView.SkyLightColor.rgb) * DiffuseColor * MaterialAO;
-#if MATERIAL_TWO_SIDED && LQ_TEXTURE_LIGHTMAP
+#if MATERIAL_TWOSIDED && LQ_TEXTURE_LIGHTMAP
 			
 			}
 #endif
diff --git a/Engine/Source/Runtime/Engine/Private/Materials/HLSLMaterialTranslator.cpp b/Engine/Source/Runtime/Engine/Private/Materials/HLSLMaterialTranslator.cpp
--- a/Engine/Source/Runtime/Engine/Private/Materials/HLSLMaterialTranslator.cpp
+++ b/Engine/Source/Runtime/Engine/Private/Materials/HLSLMaterialTranslator.cpp
@@ -1157,7 +1157,6 @@ void FHLSLMaterialTranslator::GetMaterialEnvironment(EShaderPlatform InPlatform,
 	OutEnvironment.SetDefine(TEXT("MATERIAL_IS_SKY"), Material->IsSky());
 	OutEnvironment.SetDefine(TEXT("MATERIAL_COMPUTE_FOG_PER_PIXEL"), Material->ComputeFogPerPixel());
 	OutEnvironment.SetDefine(TEXT("MATERIAL_FULLY_ROUGH"), bIsFullyRough || Material->IsFullyRough());
-	OutEnvironment.SetDefine(TEXT("MATERIAL_TWO_SIDED"), Material->IsTwoSided());
 
 	// Count the number of VTStacks (each stack will allocate a feedback slot)
 	OutEnvironment.SetDefine(TEXT("NUM_VIRTUALTEXTURE_SAMPLES"), VTStacks.Num());
