Added stat for canvas tile rendering
#rnx

diff --git a/Engine/Source/Runtime/Engine/Private/TileRendering.cpp b/Engine/Source/Runtime/Engine/Private/TileRendering.cpp
--- a/Engine/Source/Runtime/Engine/Private/TileRendering.cpp
+++ b/Engine/Source/Runtime/Engine/Private/TileRendering.cpp
@@ -174,6 +174,7 @@ void FCanvasTileRendererItem::FRenderData::RenderTiles(
 	InitTileMesh(View, bNeedsToSwitchVerticalAxis);
 
 	SCOPED_DRAW_EVENTF(RHICmdList, CanvasDrawTile, *MaterialRenderProxy->GetMaterial(GMaxRHIFeatureLevel)->GetFriendlyName());
+	QUICK_SCOPE_CYCLE_COUNTER(STAT_CanvasDrawTile)
 
 	for (int32 TileIdx = 0; TileIdx < Tiles.Num(); TileIdx++)
 	{
