Fixed crash caused by null materials being passed to GetDisplayNameForMaterial. Copy of 10890781 from Dev-Editor

diff --git a/Engine/Source/Runtime/Engine/Public/MaterialShared.h b/Engine/Source/Runtime/Engine/Public/MaterialShared.h
--- a/Engine/Source/Runtime/Engine/Public/MaterialShared.h
+++ b/Engine/Source/Runtime/Engine/Public/MaterialShared.h
@@ -2504,6 +2504,11 @@ public:
 	/** Returns the display name of a material attribute, accounting for overrides based on properties of a given material */
 	ENGINE_API static FText GetDisplayNameForMaterial(EMaterialProperty Property, UMaterial* Material)
 	{
+		if (!Material)
+		{
+			return FText::FromString(GetAttributeName(Property));
+		}
+
 		FMaterialAttributeDefintion* Attribute = GMaterialPropertyAttributesMap.Find(Property);
 		return GetAttributeOverrideForMaterial(Attribute->AttributeID, Material);
 	}
@@ -2511,6 +2516,11 @@ public:
 	/** Returns the display name of a material attribute, accounting for overrides based on properties of a given material */
 	ENGINE_API static FText GetDisplayNameForMaterial(const FGuid& AttributeID, UMaterial* Material)
 	{
+		if (!Material)
+		{
+			return FText::FromString(GetAttributeName(AttributeID));
+		}
+
 		FMaterialAttributeDefintion* Attribute = GMaterialPropertyAttributesMap.Find(AttributeID);
 		return GetAttributeOverrideForMaterial(AttributeID, Material);
 	}
