Fixing bulk data CIS errors

diff --git a/Engine/Source/Runtime/Engine/Private/PhysicsEngine/Experimental/ChaosDerivedDataReader.cpp b/Engine/Source/Runtime/Engine/Private/PhysicsEngine/Experimental/ChaosDerivedDataReader.cpp
--- a/Engine/Source/Runtime/Engine/Private/PhysicsEngine/Experimental/ChaosDerivedDataReader.cpp
+++ b/Engine/Source/Runtime/Engine/Private/PhysicsEngine/Experimental/ChaosDerivedDataReader.cpp
@@ -6,7 +6,7 @@
 #include "Chaos/Convex.h"
 
 template<typename T, int d>
-FChaosDerivedDataReader<T, d>::FChaosDerivedDataReader(FUntypedBulkData* InBulkData)
+FChaosDerivedDataReader<T, d>::FChaosDerivedDataReader(FBulkDataInterface* InBulkData)
 	: bReadSuccessful(false)
 {
 	const int32 DataTypeSize = sizeof(T);
diff --git a/Engine/Source/Runtime/Engine/Public/Physics/Experimental/ChaosDerivedDataReader.h b/Engine/Source/Runtime/Engine/Public/Physics/Experimental/ChaosDerivedDataReader.h
--- a/Engine/Source/Runtime/Engine/Public/Physics/Experimental/ChaosDerivedDataReader.h
+++ b/Engine/Source/Runtime/Engine/Public/Physics/Experimental/ChaosDerivedDataReader.h
@@ -21,7 +21,7 @@ class FChaosDerivedDataReader
 public:
 
 	// Only valid use is to explicitly read chaos bulk data
-	explicit FChaosDerivedDataReader(FUntypedBulkData* InBulkData);
+	explicit FChaosDerivedDataReader(FBulkDataInterface* InBulkData);
 
 	TArray<TUniquePtr<Chaos::FConvex>> ConvexImplicitObjects;
 	TArray<TSharedPtr<Chaos::FTriangleMeshImplicitObject, ESPMode::ThreadSafe>> TrimeshImplicitObjects;
