Make passed in parameter const correct for GetFirstLocalPlayerController
#rnx

diff --git a/Engine/Source/Runtime/Engine/Classes/Engine/Engine.h b/Engine/Source/Runtime/Engine/Classes/Engine/Engine.h
--- a/Engine/Source/Runtime/Engine/Classes/Engine/Engine.h
+++ b/Engine/Source/Runtime/Engine/Classes/Engine/Engine.h
@@ -2543,13 +2543,13 @@ public:
 	 * Find a Local Player Controller, which may not exist at all if this is a server.
 	 * @return first found LocalPlayerController. Fine for single player, in split screen, one will be picked. 
 	 */
-	class APlayerController* GetFirstLocalPlayerController(UWorld *InWorld);
+	class APlayerController* GetFirstLocalPlayerController(const UWorld* InWorld);
 
 	/** Gets all local players associated with the engine. 
 	 *	This function should only be used in rare cases where no UWorld* is available to get a player list associated with the world.
 	 *  E.g, - use GetFirstLocalPlayerController(UWorld *InWorld) when possible!
 	 */
-	void GetAllLocalPlayerControllers(TArray<APlayerController*>	& PlayerList);
+	void GetAllLocalPlayerControllers(TArray<APlayerController*>& PlayerList);
 
 	/** Returns the GameViewport widget */
 	virtual TSharedPtr<class SViewport> GetGameViewportWidget() const
diff --git a/Engine/Source/Runtime/Engine/Classes/Engine/GameInstance.h b/Engine/Source/Runtime/Engine/Classes/Engine/GameInstance.h
--- a/Engine/Source/Runtime/Engine/Classes/Engine/GameInstance.h
+++ b/Engine/Source/Runtime/Engine/Classes/Engine/GameInstance.h
@@ -279,7 +279,7 @@ public:
 
 	int32					GetNumLocalPlayers() const;
 	ULocalPlayer*			GetLocalPlayerByIndex(const int32 Index) const;
-	APlayerController*		GetFirstLocalPlayerController(UWorld* World = nullptr) const;
+	APlayerController*		GetFirstLocalPlayerController(const UWorld* World = nullptr) const;
 	ULocalPlayer*			FindLocalPlayerFromControllerId(const int32 ControllerId) const;
 	ULocalPlayer*			FindLocalPlayerFromUniqueNetId(TSharedPtr<const FUniqueNetId> UniqueNetId) const;
 	ULocalPlayer*			FindLocalPlayerFromUniqueNetId(const FUniqueNetId& UniqueNetId) const;
diff --git a/Engine/Source/Runtime/Engine/Private/GameInstance.cpp b/Engine/Source/Runtime/Engine/Private/GameInstance.cpp
--- a/Engine/Source/Runtime/Engine/Private/GameInstance.cpp
+++ b/Engine/Source/Runtime/Engine/Private/GameInstance.cpp
@@ -835,7 +835,7 @@ ULocalPlayer* UGameInstance::GetLocalPlayerByIndex(const int32 Index) const
 	return nullptr;
 }
 
-APlayerController* UGameInstance::GetFirstLocalPlayerController(UWorld* World) const
+APlayerController* UGameInstance::GetFirstLocalPlayerController(const UWorld* World) const
 {
 	if (World == nullptr)
 	{
diff --git a/Engine/Source/Runtime/Engine/Private/UnrealEngine.cpp b/Engine/Source/Runtime/Engine/Private/UnrealEngine.cpp
--- a/Engine/Source/Runtime/Engine/Private/UnrealEngine.cpp
+++ b/Engine/Source/Runtime/Engine/Private/UnrealEngine.cpp
@@ -3343,11 +3343,11 @@ void UEngine::SwapControllerId(ULocalPlayer *NewPlayer, const int32 CurrentContr
 	}
 }
 
-APlayerController* UEngine::GetFirstLocalPlayerController(UWorld *InWorld)
+APlayerController* UEngine::GetFirstLocalPlayerController(const UWorld* InWorld)
 {
-	const FWorldContext &Context = GetWorldContextFromWorldChecked(InWorld);
+	const FWorldContext& Context = GetWorldContextFromWorldChecked(InWorld);
 
-	return ( Context.OwningGameInstance != NULL ) ? Context.OwningGameInstance->GetFirstLocalPlayerController(InWorld) : NULL;
+	return (Context.OwningGameInstance ? Context.OwningGameInstance->GetFirstLocalPlayerController(InWorld) : nullptr);
 }
 
 void UEngine::GetAllLocalPlayerControllers(TArray<APlayerController*> & PlayerList)
