Use tree to acceleration trimesh cleanup

diff --git a/Engine/Source/Runtime/Engine/Private/PhysicsEngine/Experimental/ChaosDerivedDataUtil.cpp b/Engine/Source/Runtime/Engine/Private/PhysicsEngine/Experimental/ChaosDerivedDataUtil.cpp
--- a/Engine/Source/Runtime/Engine/Private/PhysicsEngine/Experimental/ChaosDerivedDataUtil.cpp
+++ b/Engine/Source/Runtime/Engine/Private/PhysicsEngine/Experimental/ChaosDerivedDataUtil.cpp
@@ -5,10 +5,22 @@
 
 #include "ChaosLog.h"
 #include "ChaosCheck.h"
+#include "Chaos/AABBTree.h"
 
 namespace Chaos
 {
 
+	struct FCleanMeshWrapper
+	{
+		FVector Vert;
+
+		template <typename TPayloadType>
+		TPayloadType GetPayload(int32 Idx) const { return Idx; }
+
+		bool HasBoundingBox() const { return true; }
+		TAABB<FReal, 3> BoundingBox() const { return TAABB<FReal, 3>(Vert, Vert); }
+	};
+
 	void CleanTrimesh(TArray<FVector>& InOutVertices, TArray<int32>& InOutIndices, TArray<int32>* OutOptFaceRemap)
 	{
 		TArray<FVector> LocalSourceVerts = InOutVertices;
@@ -59,38 +71,42 @@ namespace Chaos
 
 		float WeldThresholdSq = 0.0f;// SMALL_NUMBER * SMALL_NUMBER;
 
-		for(int32 SourceVertIndex = 0; SourceVertIndex < NumSourceVerts; ++SourceVertIndex)
+		TArray<FCleanMeshWrapper> WrapperVerts;
+
+		for (int i = 0; i < LocalSourceVerts.Num(); ++i)
 		{
-			const FVector& SourceVert = LocalSourceVerts[SourceVertIndex];
+			WrapperVerts.Add(FCleanMeshWrapper{LocalSourceVerts[i] });
+		}
 
-			bool bUnique = true; // assume the vertex is unique until we find otherwise
-			int32 RemapIndex = INDEX_NONE; // if the vertex isn't unique this will be set
+		TAABBTree<int32, TAABBTreeLeafArray<int32, FReal>, FReal> Accel(WrapperVerts);
+		TSet<int32> NonUnique;
 
-			const int32 NumUniqueVerts = LocalUniqueVerts.Num();
-			for(int32 UniqueVertIndex = 0; UniqueVertIndex < NumUniqueVerts; ++UniqueVertIndex)
+		for(int32 SourceVertIndex = 0; SourceVertIndex < NumSourceVerts; ++SourceVertIndex)
+		{
+			if (NonUnique.Contains(SourceVertIndex))
 			{
-				const FVector& UniqueVert = LocalUniqueVerts[UniqueVertIndex];
+				continue;
+			}
 
-				if((UniqueVert - SourceVert).SizeSquared() <= WeldThresholdSq)
-				{
-					//This vertex isn't unique and needs to be merged with the unique one we found
-					bUnique = false;
-					RemapIndex = UniqueVertIndex;
+			const FVector& SourceVert = LocalSourceVerts[SourceVertIndex];
 
-					// Done searching
-					break;
-				}
-			}
+			TArray<int32> Duplicates = Accel.FindAllIntersections(TAABB<FReal, 3>(SourceVert - WeldThresholdSq, SourceVert + WeldThresholdSq));
+			ensure(Duplicates.Num() > 0);	//Should always find at least original vert
 
-			if(bUnique)
-			{
-				LocalUniqueVerts.Add(SourceVert);
-				LocalUniqueToSourceIndices.Add(SourceVertIndex);
-				LocalVertexRemap[SourceVertIndex] = LocalUniqueVerts.Num() - 1;
-			}
-			else
+			//first index is always considered unique
+			LocalUniqueVerts.Add(SourceVert);
+			LocalUniqueToSourceIndices.Add(SourceVertIndex);
+			LocalVertexRemap[SourceVertIndex] = LocalUniqueVerts.Num() - 1;
+
+
+			for (int32 Idx : Duplicates)
 			{
-				LocalVertexRemap[SourceVertIndex] = RemapIndex;
+				if (Idx != SourceVertIndex)
+				{
+					ensure(Idx > SourceVertIndex);	//shouldn't be here if a smaller idx already found these duplicates
+					NonUnique.Add(Idx);
+					LocalVertexRemap[Idx] = LocalVertexRemap[SourceVertIndex];
+				}
 			}
 		}
 
