Fixed DF mem leak - Integrated from dev-rendering

diff --git a/Engine/Source/Runtime/Engine/Public/DistanceFieldAtlas.h b/Engine/Source/Runtime/Engine/Public/DistanceFieldAtlas.h
--- a/Engine/Source/Runtime/Engine/Public/DistanceFieldAtlas.h
+++ b/Engine/Source/Runtime/Engine/Public/DistanceFieldAtlas.h
@@ -94,6 +94,7 @@ public:
 
 	virtual void ReleaseRHI() override
 	{
+		VolumeTextureUAVRHI.SafeRelease();
 		VolumeTextureRHI.SafeRelease();
 	}
 
