Rename GetActiveProfileName to a more suited name, since it rebuild the device profile name ech time. Made GetDeviceProfileName deprecated, and added a new function that get the current device profile name and fallback on the platform device name.

#rnx

diff --git a/Engine/Source/Runtime/Engine/Classes/DeviceProfiles/DeviceProfileManager.h b/Engine/Source/Runtime/Engine/Classes/DeviceProfiles/DeviceProfileManager.h
--- a/Engine/Source/Runtime/Engine/Classes/DeviceProfiles/DeviceProfileManager.h
+++ b/Engine/Source/Runtime/Engine/Classes/DeviceProfiles/DeviceProfileManager.h
@@ -121,13 +121,29 @@ public:
 	*/
 	void GetAllPossibleParentProfiles(const UDeviceProfile* ChildProfile, OUT TArray<UDeviceProfile*>& PossibleParentProfiles) const;
 
+	/**
+	* Get the current active profile name.
+	*
+	* @return The selected profile.
+	*/
+	const FString GetActiveDeviceProfileName();
+
 	/**
 	* Get the selected device profile name, either the platform name, or the name
 	* provided by a Device Profile Selector Module.
 	*
 	* @return The selected profile.
 	*/
+	UE_DEPRECATED(4.24, "Use either GetActiveDeviceProfileName to have the current active device profile or GetPlatformDeviceProfileName to have the default one. Note, GetActiveDeviceProfileName will fallback on GetPlatformDeviceProfileName, if there is no active device profile ")
 	static const FString GetActiveProfileName();
+
+	/**
+	* Get the selected device profile name, either the platform name, or the name
+	* provided by a Device Profile Selector Module.
+	*
+	* @return The selected profile.
+	*/
+	static const FString GetPlatformDeviceProfileName();
 	
 	/** Retrieves the value of a scalability group cvar if it was set by the active device profile. */
 	static bool GetScalabilityCVar(const FString& CvarName, int32& OutValue);
diff --git a/Engine/Source/Runtime/Engine/Private/ChartCreation.cpp b/Engine/Source/Runtime/Engine/Private/ChartCreation.cpp
--- a/Engine/Source/Runtime/Engine/Private/ChartCreation.cpp
+++ b/Engine/Source/Runtime/Engine/Private/ChartCreation.cpp
@@ -661,7 +661,7 @@ void FPerformanceTrackingChart::Reset(const FDateTime& InStartTime)
 
 	StartBatteryLevel = -1;
 	StopBatteryLevel = -1;
-	DeviceProfileName = UDeviceProfileManager::GetActiveProfileName();
+	DeviceProfileName = UDeviceProfileManager::Get().GetActiveDeviceProfileName();
 	bIsChartingPaused = false;
 }
 
@@ -712,7 +712,7 @@ void FPerformanceTrackingChart::StartCharting()
 {
 	StartTemperatureLevel = FPlatformMisc::GetDeviceTemperatureLevel();
 	StartBatteryLevel = FPlatformMisc::GetBatteryLevel();
-	DeviceProfileName = UDeviceProfileManager::GetActiveProfileName();
+	DeviceProfileName = UDeviceProfileManager::Get().GetActiveDeviceProfileName();
 	bIsChartingPaused = false;
 }
 
@@ -734,7 +734,7 @@ void FPerformanceTrackingChart::ResumeCharting()
 
 void FPerformanceTrackingChart::OnDeviceProfileManagerUpdated()
 {
-	FString CurrentDeviceProfileName = UDeviceProfileManager::GetActiveProfileName();
+	FString CurrentDeviceProfileName = UDeviceProfileManager::Get().GetActiveDeviceProfileName();
 	if (CurrentDeviceProfileName != DeviceProfileName)
 	{
 		DeviceProfileName = TEXT("Mixed");
diff --git a/Engine/Source/Runtime/Engine/Private/DeviceProfiles/DeviceProfileManager.cpp b/Engine/Source/Runtime/Engine/Private/DeviceProfiles/DeviceProfileManager.cpp
--- a/Engine/Source/Runtime/Engine/Private/DeviceProfiles/DeviceProfileManager.cpp
+++ b/Engine/Source/Runtime/Engine/Private/DeviceProfiles/DeviceProfileManager.cpp
@@ -58,7 +58,7 @@ UDeviceProfileManager& UDeviceProfileManager::Get(bool bFromPostCDOContruct)
 		}
 
 		// always start with an active profile, even if we create it on the spot
-		UDeviceProfile* ActiveProfile = DeviceProfileManagerSingleton->FindProfile(GetActiveProfileName());
+		UDeviceProfile* ActiveProfile = DeviceProfileManagerSingleton->FindProfile(GetPlatformDeviceProfileName());
 		DeviceProfileManagerSingleton->SetActiveDeviceProfile(ActiveProfile);
 
 		// now we allow the cvar changes to be acknowledged
@@ -85,7 +85,7 @@ UDeviceProfileManager& UDeviceProfileManager::Get(bool bFromPostCDOContruct)
 
 void UDeviceProfileManager::InitializeCVarsForActiveDeviceProfile(bool bPushSettings)
 {
-	FString ActiveProfileName = DeviceProfileManagerSingleton ? DeviceProfileManagerSingleton->ActiveDeviceProfile->GetName() : GetActiveProfileName();
+	FString ActiveProfileName = DeviceProfileManagerSingleton ? DeviceProfileManagerSingleton->ActiveDeviceProfile->GetName() : GetPlatformDeviceProfileName();
 
 	UE_LOG(LogInit, Log, TEXT("Applying CVar settings loaded from the selected device profile: [%s]"), *ActiveProfileName);
 
@@ -661,7 +661,7 @@ void UDeviceProfileManager::HandleDeviceProfileOverridePop()
 	RestoreDefaultDeviceProfile();
 }
 
-const FString UDeviceProfileManager::GetActiveProfileName()
+const FString UDeviceProfileManager::GetPlatformDeviceProfileName()
 {
 	FString ActiveProfileName = FPlatformProperties::PlatformName();
 
@@ -706,6 +706,24 @@ const FString UDeviceProfileManager::GetActiveProfileName()
 	return ActiveProfileName;
 }
 
+
+const FString UDeviceProfileManager::GetActiveDeviceProfileName()
+{
+	if(ActiveDeviceProfile != nullptr)
+	{
+		return ActiveDeviceProfile->GetName();
+	}
+	else
+	{
+		return GetPlatformDeviceProfileName();
+	}
+}
+
+const FString UDeviceProfileManager::GetActiveProfileName()
+{
+	return GetPlatformDeviceProfileName();
+}
+
 bool UDeviceProfileManager::GetScalabilityCVar(const FString& CVarName, int32& OutValue)
 {
 	if (const FString* CVarValue = DeviceProfileScalabilityCVars.Find(CVarName))
