Provide deprecation path and function forwarding for AudioDevice & ActiveSound changes.
*** Note These files will need to be compared to CL 9966924 for a more accurate view of what the changes would look like between shipped revisions ***

diff --git a/Engine/Source/Runtime/Engine/Private/ActiveSound.cpp b/Engine/Source/Runtime/Engine/Private/ActiveSound.cpp
--- a/Engine/Source/Runtime/Engine/Private/ActiveSound.cpp
+++ b/Engine/Source/Runtime/Engine/Private/ActiveSound.cpp
@@ -425,6 +425,16 @@ void FActiveSound::GetSoundSourceBusSends(EBusSendType BusSendType, TArray<FSoun
 	}
 }
 
+int32 FActiveSound::FindClosestListener(const TArray<struct FListener>& InListeners) const
+{
+	return AudioDevice ? AudioDevice->FindClosestListenerIndex(Transform, InListeners) : INDEX_NONE;
+}
+
+int32 FActiveSound::FindClosestListener() const
+{
+	return AudioDevice ? AudioDevice->FindClosestListenerIndex(Transform) : INDEX_NONE;
+}
+
 void FActiveSound::GetConcurrencyHandles(TArray<FConcurrencyHandle>& OutConcurrencyHandles) const
 {
 	OutConcurrencyHandles.Reset();
@@ -1494,11 +1504,21 @@ void FActiveSound::UpdateFocusData(float DeltaTime, const FAttenuationListenerDa
 	FocusDataToUpdate->VolumeScale = ListenerData.AttenuationSettings->GetFocusAttenuation(FocusSettings, FocusDataToUpdate->FocusFactor);
 }
 
+void FActiveSound::ParseAttenuation(FSoundParseParameters& OutParseParams, const FListener& InListener, const FSoundAttenuationSettings& InAttenuationSettings)
+{
+	UpdateAttenuation(0.0f, OutParseParams, InListener.ListenerIndex, &InAttenuationSettings);
+}
+
 void FActiveSound::ParseAttenuation(FSoundParseParameters& OutParseParams, int32 ListenerIndex, const FSoundAttenuationSettings& InAttenuationSettings)
 {
 	UpdateAttenuation(0.0f, OutParseParams, ListenerIndex, &InAttenuationSettings);
 }
 
+void FActiveSound::UpdateAttenuation(float DeltaTime, FSoundParseParameters& ParseParams, const FListener& Listener, const FSoundAttenuationSettings* SettingsAttenuationNode)
+{
+	UpdateAttenuation(DeltaTime, ParseParams, Listener.ListenerIndex, SettingsAttenuationNode);
+}
+
 void FActiveSound::UpdateAttenuation(float DeltaTime, FSoundParseParameters& ParseParams, int32 ListenerIndex, const FSoundAttenuationSettings* SettingsAttenuationNode)
 {
 	// Get the attenuation settings to use for this application to the active sound
diff --git a/Engine/Source/Runtime/Engine/Private/AudioDevice.cpp b/Engine/Source/Runtime/Engine/Private/AudioDevice.cpp
--- a/Engine/Source/Runtime/Engine/Private/AudioDevice.cpp
+++ b/Engine/Source/Runtime/Engine/Private/AudioDevice.cpp
@@ -167,6 +167,37 @@ namespace
 #endif // !UE_BUILD_SHIPPING
 } // namespace <>
 
+FAttenuationListenerData FAttenuationListenerData::Create(const FAudioDevice& AudioDevice, const FTransform& InListenerTransform, const FTransform& InSoundTransform, const FSoundAttenuationSettings& InAttenuationSettings)
+{
+	// This function is deprecated as we assume the listener transform is from listener 0.  
+	FAttenuationListenerData ListenerData(InListenerTransform, InSoundTransform, InAttenuationSettings);
+
+	const FVector SoundTranslation = InSoundTransform.GetTranslation();
+	FVector ListenerToSound = SoundTranslation - InListenerTransform.GetTranslation();
+	ListenerToSound.ToDirectionAndLength(ListenerData.ListenerToSoundDir, ListenerData.ListenerToSoundDistance);
+
+	// Store the actual distance for surround-panning sources with spread (AudioMixer)
+	ListenerData.ListenerToSoundDistanceForPanning = ListenerData.ListenerToSoundDistance;
+
+	// Calculating override listener-to-sound distance and transform must
+	// be applied after distance used for panning value is calculated.
+	FVector ListenerPosition;
+	const bool bAllowAttenuationOverride = true;
+
+	if (AudioDevice.GetListenerPosition(0, ListenerPosition, bAllowAttenuationOverride))
+	{
+		ListenerData.ListenerToSoundDistance = (SoundTranslation - ListenerPosition).Size();
+		ListenerData.ListenerTransform.SetTranslation(ListenerPosition);
+	}
+
+	const FSoundAttenuationSettings& AttenuationSettings = *ListenerData.AttenuationSettings;
+	if ((AttenuationSettings.bAttenuate && AttenuationSettings.AttenuationShape == EAttenuationShape::Sphere) || AttenuationSettings.bAttenuateWithLPF)
+	{
+		ListenerData.AttenuationDistance = FMath::Max(ListenerData.ListenerToSoundDistance - AttenuationSettings.AttenuationShapeExtents.X, 0.0f);
+	}
+
+	return MoveTemp(ListenerData);
+}
 
 FAttenuationListenerData FAttenuationListenerData::Create(const FAudioDevice& AudioDevice, int32 ListenerIndex, const FTransform& InSoundTransform, const FSoundAttenuationSettings& InAttenuationSettings)
 {
@@ -2990,13 +3021,23 @@ void FAudioDevice::SetListener(UWorld* World, const int32 InViewportIndex, const
 			PluginManager->OnListenerUpdated(this, InViewportIndex, ListenerTransformCopy, InDeltaSeconds);
 		}
 
+		const int32 StartingListenerCount = Listeners.Num();
+
 		TArray<FListener>& AudioThreadListeners = Listeners;
 		if (InViewportIndex >= AudioThreadListeners.Num())
 		{
-			const int32 NumListeners = InViewportIndex - AudioThreadListeners.Num() + 1;
-			for (int32 i = 0; i < NumListeners; ++i)
+			const int32 NumListenersToAdd = InViewportIndex - AudioThreadListeners.Num() + 1;
+			for (int32 i = 0; i < NumListenersToAdd; ++i)
 			{
 				AudioThreadListeners.Add(FListener(this));
+
+				// While we're going through the process of moving from raw listener access to access by index, 
+				// we're going to store our current index inside the listener to help in deprecation and backwards compat.
+				const int32 CurrentIndex = i + StartingListenerCount;
+				if (ensure(AudioThreadListeners.IsValidIndex(CurrentIndex)))
+				{
+					AudioThreadListeners[CurrentIndex].ListenerIndex = CurrentIndex;
+				}
 			}
 		}
 
@@ -4912,6 +4953,27 @@ bool FAudioDevice::LocationIsAudible(const FVector& Location, const float MaxDis
 	return false;
 }
 
+bool FAudioDevice::LocationIsAudible(const FVector& Location, const FTransform& ListenerTransform, const float MaxDistance) const
+{
+	// This function is deprecated as it assumes listener 0.  
+	// To check if a location is audible by any listener, use FAudioDevice::LocationIsAudible that takes a location and max distance.
+	// To check if a location is audible by a specific listener, use FAudioDevice::LocationIsAudible that additionally takes a listener index. 
+	if (MaxDistance >= WORLD_MAX)
+	{
+		return true;
+	}
+
+	FVector ListenerTranslation;
+	const bool bAllowOverride = true;
+	if (!GetListenerPosition(0, ListenerTranslation, bAllowOverride))
+	{
+		return false;
+	}
+
+	const float MaxDistanceSquared = MaxDistance * MaxDistance;
+	return (ListenerTranslation - Location).SizeSquared() < MaxDistanceSquared;
+}
+
 bool FAudioDevice::LocationIsAudible(const FVector& Location, int32 ListenerIndex, float MaxDistance) const
 {
 	if (MaxDistance >= WORLD_MAX)
@@ -4947,14 +5009,14 @@ float FAudioDevice::GetDistanceToNearestListener(const FVector& Location) const
 	return FMath::Sqrt(DistSquared);
 }
 
-float FAudioDevice::GetDistanceSquaredToListener(const FVector& Location, int32 ListenerIndex) const
+float FAudioDevice::GetSquaredDistanceToListener(const FVector& Location, const FTransform& ListenerTransform) const
 {
-	float OutSqDistance;
-	GetSquaredDistanceToListener(Location, ListenerIndex, OutSqDistance);
-	return OutSqDistance;
+	// This function is deprecated as it does not take into account listener attenuation override position
+	FVector ListenerTranslation = ListenerTransform.GetTranslation();
+	return (ListenerTranslation - Location).SizeSquared();
 }
 
-bool FAudioDevice::GetSquaredDistanceToListener(const FVector& Location, int32 ListenerIndex, float& OutSqDistance) const
+bool FAudioDevice::GetDistanceSquaredToListener(const FVector& Location, int32 ListenerIndex, float& OutSqDistance) const
 {
 	OutSqDistance = TNumericLimits<float>::Max();
 	const int32 ListenerCount = IsInAudioThread() ? Listeners.Num() : ListenerProxies.Num();
@@ -5141,6 +5203,28 @@ bool FAudioDevice::SoundIsAudible(const FActiveSound& NewActiveSound)
 
 	return false;
 }
+int32 FAudioDevice::FindClosestListenerIndex(const FTransform& SoundTransform, const TArray<FListener>& InListeners)
+{
+	check(IsInAudioThread());
+	int32 ClosestListenerIndex = 0;
+	const bool bAllowAttenuationOverride = true;
+	if (InListeners.Num() > 0)
+	{
+		float ClosestDistSq = FVector::DistSquared(SoundTransform.GetTranslation(), InListeners[0].GetPosition(bAllowAttenuationOverride));
+
+		for (int32 i = 1; i < InListeners.Num(); i++)
+		{
+			const float DistSq = FVector::DistSquared(SoundTransform.GetTranslation(), InListeners[i].GetPosition(bAllowAttenuationOverride));
+			if (DistSq < ClosestDistSq)
+			{
+				ClosestListenerIndex = i;
+				ClosestDistSq = DistSq;
+			}
+		}
+	}
+
+	return ClosestListenerIndex;
+}
 
 int32 FAudioDevice::FindClosestListenerIndex(const FTransform& SoundTransform) const
 {
@@ -6102,6 +6186,37 @@ float FAudioDevice::GetGameDeltaTime() const
 	return FMath::Min(DeltaTime, 0.5f);
 }
 
+bool FAudioDevice::IsUsingListenerAttenuationOverride(int32 ListenerIndex) const
+{
+	const bool bInAudioThread = IsInAudioThread();
+	const int32 ListenerCount = bInAudioThread ? Listeners.Num() : ListenerProxies.Num();
+	if (ListenerIndex >= ListenerCount)
+	{
+		return false;
+	}
+
+	if (bInAudioThread)
+	{
+		return Listeners[ListenerIndex].bUseAttenuationOverride;
+	}
+
+	return ListenerProxies[ListenerIndex].bUseAttenuationOverride;
+}
+
+const FVector& FAudioDevice::GetListenerAttenuationOverride(int32 ListenerIndex) const
+{
+	const bool bInAudioThread = IsInAudioThread();
+	const int32 ListenerCount = bInAudioThread ? Listeners.Num() : ListenerProxies.Num();
+	check(ListenerIndex < ListenerCount);
+
+	if (bInAudioThread)
+	{
+		return Listeners[ListenerIndex].AttenuationOverride;
+	}
+
+	return ListenerProxies[ListenerIndex].AttenuationOverride;
+}
+
 void FAudioDevice::UpdateVirtualLoops(bool bForceUpdate)
 {
 	check(IsInAudioThread());
diff --git a/Engine/Source/Runtime/Engine/Public/ActiveSound.h b/Engine/Source/Runtime/Engine/Public/ActiveSound.h
--- a/Engine/Source/Runtime/Engine/Public/ActiveSound.h
+++ b/Engine/Source/Runtime/Engine/Public/ActiveSound.h
@@ -641,6 +641,12 @@ public:
 	/** Gets the sound source bus sends to use for this sound instance. */
 	void GetSoundSourceBusSends(EBusSendType BusSendType, TArray<FSoundSourceBusSendInfo>& OutSends) const;
 
+	/* Determines which of the provided listeners is the closest to the sound */
+	int32 FindClosestListener( const TArray<struct FListener>& InListeners ) const;
+
+	/* Determines which listener is the closest to the sound */
+	int32 FindClosestListener() const;
+
 	/** Returns the unique ID of the active sound's owner if it exists. Returns 0 if the sound doesn't have an owner. */
 	FSoundOwnerObjectID GetOwnerID() const { return OwnerID; }
 
@@ -652,6 +658,10 @@ public:
 	/** Delegate callback function when an async occlusion trace completes */
 	static void OcclusionTraceDone(const FTraceHandle& TraceHandle, FTraceDatum& TraceDatum);
 
+	/** Applies the active sound's attenuation settings to the input parse params using the given listener */
+	UE_DEPRECATED(4.25, "Use ParseAttenuation that passes a ListenerIndex instead")
+	void ParseAttenuation(FSoundParseParameters& OutParseParams, const FListener& InListener, const FSoundAttenuationSettings& InAttenuationSettings);
+
 	/** Applies the active sound's attenuation settings to the input parse params using the given listener */
 	void ParseAttenuation(FSoundParseParameters& OutParseParams, int32 ListenerIndex, const FSoundAttenuationSettings& InAttenuationSettings);
 
@@ -664,6 +674,10 @@ public:
 	/** Sets the amount of audio from this active sound to send to the source bus. */
 	void SetSourceBusSend(EBusSendType BusSendTyoe, const FSoundSourceBusSendInfo& SourceBusSendInfo);
 
+	/** Updates the active sound's attenuation settings to the input parse params using the given listener */
+	UE_DEPRECATED(4.25, "Use UpdateAttenuation that passes a ListenerIndex instead")
+	void UpdateAttenuation(float DeltaTime, FSoundParseParameters& ParseParams, const FListener& Listener, const FSoundAttenuationSettings* SettingsAttenuationNode = nullptr);
+
 	/** Updates the active sound's attenuation settings to the input parse params using the given listener */
 	void UpdateAttenuation(float DeltaTime, FSoundParseParameters& ParseParams, int32 ListenerIndex, const FSoundAttenuationSettings* SettingsAttenuationNode = nullptr);
 
diff --git a/Engine/Source/Runtime/Engine/Public/AudioDevice.h b/Engine/Source/Runtime/Engine/Public/AudioDevice.h
--- a/Engine/Source/Runtime/Engine/Public/AudioDevice.h
+++ b/Engine/Source/Runtime/Engine/Public/AudioDevice.h
@@ -150,6 +150,9 @@ struct FListener
 	/** The ID of the world the listener resides in */
 	uint32 WorldID;
 
+	/** Index of this listener inside the AudioDevice's listener array */
+	int32 ListenerIndex;
+
 	/** The times of interior volumes fading in and out */
 	double InteriorStartTime;
 	double InteriorEndTime;
@@ -192,6 +195,7 @@ struct FListener
 		, AttenuationOverride(ForceInit)
 		, bUseAttenuationOverride(false)
 		, AudioVolumeID(0)
+		, ListenerIndex(0)
 		, InteriorStartTime(0.0)
 		, InteriorEndTime(0.0)
 		, ExteriorEndTime(0.0)
@@ -325,6 +329,10 @@ struct FAttenuationListenerData
 	const FTransform SoundTransform;
 	const FSoundAttenuationSettings* AttenuationSettings;
 
+	/** Computes and returns some geometry related to the listener and the given sound transform. */
+	UE_DEPRECATED(4.25, "Use FAttenuationListenerData::Create that passes a ListenerIndex")
+	static FAttenuationListenerData Create(const FAudioDevice& AudioDevice, const FTransform& InListenerTransform, const FTransform& InSoundTransform, const FSoundAttenuationSettings& InAttenuationSettings);
+
 	/** Computes and returns some geometry related to the listener and the given sound transform. */
 	static FAttenuationListenerData Create(const FAudioDevice& AudioDevice, int32 ListenerIndex, const FTransform& InSoundTransform, const FSoundAttenuationSettings& InAttenuationSettings);
 
@@ -641,9 +649,17 @@ public:
 	void SetListener(UWorld* World, int32 InListenerIndex, const FTransform& ListenerTransform, float InDeltaSeconds);
 
 	/** Sets an override for the listener to do attenuation calculations. */
+	UE_DEPRECATED(4.25, "Use SetListenerAttenuationOverride that passes a ListenerIndex instead")
+	void SetListenerAttenuationOverride(const FVector AttenuationPosition) { SetListenerAttenuationOverride(0, AttenuationPosition); }
+
+	/** Sets an override position for the specified listener to do attenuation calculations. */
 	void SetListenerAttenuationOverride(int32 ListenerIndex, const FVector AttenuationPosition);
 
 	/** Removes a listener attenuation override. */
+	UE_DEPRECATED(4.25, "Use ClearListenerAttenuationOverride that passes a ListenerIndex instead")
+	void ClearListenerAttenuationOverride() { ClearListenerAttenuationOverride(0); }
+
+	/** Removes a listener attenuation override for the specified listener. */
 	void ClearListenerAttenuationOverride(int32 ListenerIndex);
 
 	const TArray<FListener>& GetListeners() const { check(IsInAudioThread()); return Listeners; }
@@ -858,6 +874,12 @@ public:
 	/**
 	* Checks to see if a coordinate is within a distance of the given listener
 	*/
+	UE_DEPRECATED(4.25, "Use LocationIsAudible that passes a ListenerIndex to check against a specific Listener")
+	bool LocationIsAudible(const FVector& Location, const FTransform& ListenerTransform, const float MaxDistance) const;
+
+	/**
+	* Checks to see if a coordinate is within a distance of a specific listener
+	*/
 	bool LocationIsAudible(const FVector& Location, int32 ListenerIndex, const float MaxDistance) const;
 
 	/**
@@ -865,16 +887,13 @@ public:
 	*/
 	float GetDistanceToNearestListener(const FVector& Location) const;
 
-	/**
-	* Returns the distance from location to the appropriate listener representation, depending on calling thread
-	*/
-	UE_DEPRECATED(4.25, "Use GetSquaredDistanceToListener instead")
-	float GetDistanceSquaredToListener(const FVector& Location, int32 ListenerIndex) const;
+	UE_DEPRECATED(4.25, "Use GetDistanceSquaredToListener to check against a specific Listener")
+	float GetSquaredDistanceToListener(const FVector& Location, const FTransform& ListenerTransform) const;
 
 	/**
 	* Returns the distance from location to the appropriate listener representation, depending on calling thread
 	*/
-	bool GetSquaredDistanceToListener(const FVector& Location, int32 ListenerIndex, float& OutSqDistance) const;
+	bool GetDistanceSquaredToListener(const FVector& Location, int32 ListenerIndex, float& OutSqDistance) const;
 
 	/**
 	* Returns a position from the appropriate listener representation, depending on calling thread.
@@ -1094,6 +1113,9 @@ public:
 	*/
 	bool SoundIsAudible(USoundBase* Sound, const UWorld* World, const FVector& Location, const FSoundAttenuationSettings* AttenuationSettingsToApply, float MaxDistance, float FocusFactor) const;
 
+	/** Returns the index of the listener closest to the given sound transform */
+	static int32 FindClosestListenerIndex(const FTransform& SoundTransform, const TArray<FListener>& InListeners);
+
 	/** Returns the index of the listener closest to the given sound transform */
 	int32 FindClosestListenerIndex(const FTransform& SoundTransform) const;
 	int32 FindClosestListenerIndex(const FVector& Position, float& OutSqDistance, bool AllowAttenuationOverrides) const;
@@ -1574,6 +1596,20 @@ public:
 	/** Returns the game's delta time */
 	float GetGameDeltaTime() const;
 
+	/** Whether device is using listener attenuation override or not. */
+	UE_DEPRECATED(4.25, "Use ParseAttenuation that passes a ListenerIndex instead")
+	bool IsUsingListenerAttenuationOverride() const { return IsUsingListenerAttenuationOverride(0); }
+
+	/** Returns if the specific listener is using an attenuation override position. */
+	bool IsUsingListenerAttenuationOverride(int32 ListenerIndex) const;
+
+	/** Returns the listener attenuation override */
+	UE_DEPRECATED(4.25, "Use ParseAttenuation that passes a ListenerIndex instead")
+	const FVector& GetListenerAttenuationOverride() const { return GetListenerAttenuationOverride(0); }
+
+	/** Returns the listener attenuation override for the specified listener */
+	const FVector& GetListenerAttenuationOverride(int32 ListenerIndex) const;
+
 	void UpdateVirtualLoops(bool bForceUpdate);
 
 	/** Sets the update delta time for the audio frame */
