#UE4 Quick fix for a crash during cook. Needs more investigation.

diff --git a/Engine/Source/Runtime/Engine/Private/InstancedStaticMesh.cpp b/Engine/Source/Runtime/Engine/Private/InstancedStaticMesh.cpp
--- a/Engine/Source/Runtime/Engine/Private/InstancedStaticMesh.cpp
+++ b/Engine/Source/Runtime/Engine/Private/InstancedStaticMesh.cpp
@@ -1953,7 +1953,7 @@ void UInstancedStaticMeshComponent::SerializeRenderData(FArchive& Ar)
 		{
 			uint64 RenderDataPos = Ar.Tell();
 
-			if (PerInstanceSMData.Num() > 0)
+			if (PerInstanceSMData.Num() > 0 && PerInstanceRenderData.IsValid())
 			{
 				check(PerInstanceRenderData.IsValid());
 
