Change deprecated version to 4.25.

diff --git a/Engine/Source/Runtime/Engine/Classes/DeviceProfiles/DeviceProfileManager.h b/Engine/Source/Runtime/Engine/Classes/DeviceProfiles/DeviceProfileManager.h
--- a/Engine/Source/Runtime/Engine/Classes/DeviceProfiles/DeviceProfileManager.h
+++ b/Engine/Source/Runtime/Engine/Classes/DeviceProfiles/DeviceProfileManager.h
@@ -134,7 +134,7 @@ public:
 	*
 	* @return The selected profile.
 	*/
-	UE_DEPRECATED(4.24, "Use either GetActiveDeviceProfileName to have the current active device profile or GetPlatformDeviceProfileName to have the default one. Note, GetActiveDeviceProfileName will fallback on GetPlatformDeviceProfileName, if there is no active device profile ")
+	UE_DEPRECATED(4.25, "Use either GetActiveDeviceProfileName to have the current active device profile or GetPlatformDeviceProfileName to have the default one. Note, GetActiveDeviceProfileName will fallback on GetPlatformDeviceProfileName, if there is no active device profile ")
 	static const FString GetActiveProfileName();
 
 	/**
