Fix crash if SkeletalReduction plugin is disabled

diff --git a/Engine/Source/Runtime/Engine/Private/SkeletalMesh.cpp b/Engine/Source/Runtime/Engine/Private/SkeletalMesh.cpp
--- a/Engine/Source/Runtime/Engine/Private/SkeletalMesh.cpp
+++ b/Engine/Source/Runtime/Engine/Private/SkeletalMesh.cpp
@@ -470,41 +470,45 @@ void USkeletalMesh::ValidateBoundsExtension()
 /* Return true if the reduction settings are setup to reduce a LOD*/
 bool USkeletalMesh::IsReductionActive(int32 LODIndex) const
 {
-	FSkeletalMeshOptimizationSettings ReductionSettings = GetReductionSettings(LODIndex);
-	IMeshReduction* ReductionModule = FModuleManager::Get().LoadModuleChecked<IMeshReductionManagerModule>("MeshReductionInterface").GetSkeletalMeshReductionInterface();
-	uint32 LODVertexNumber = MAX_uint32;
-	uint32 LODTriNumber = MAX_uint32;
-	const FSkeletalMeshLODInfo* LODInfoPtr = GetLODInfo(LODIndex);
-	bool bLODHasBeenSimplified = LODInfoPtr && LODInfoPtr->bHasBeenSimplified;
-	if (GetImportedModel() && GetImportedModel()->LODModels.IsValidIndex(LODIndex))
-	{
-		if (!bLODHasBeenSimplified)
-		{
-			LODVertexNumber = 0;
-			LODTriNumber = 0;
-			const FSkeletalMeshLODModel& LODModel = GetImportedModel()->LODModels[LODIndex];
-			//We can take the vertices and triangles count from the source model
-			for (int32 SectionIndex = 0; SectionIndex < LODModel.Sections.Num(); ++SectionIndex)
+	bool bReductionActive = false;
+	if (IMeshReduction* ReductionModule = FModuleManager::Get().LoadModuleChecked<IMeshReductionManagerModule>("MeshReductionInterface").GetSkeletalMeshReductionInterface())
+	{
+		FSkeletalMeshOptimizationSettings ReductionSettings = GetReductionSettings(LODIndex);
+		uint32 LODVertexNumber = MAX_uint32;
+		uint32 LODTriNumber = MAX_uint32;
+		const FSkeletalMeshLODInfo* LODInfoPtr = GetLODInfo(LODIndex);
+		bool bLODHasBeenSimplified = LODInfoPtr && LODInfoPtr->bHasBeenSimplified;
+		if (GetImportedModel() && GetImportedModel()->LODModels.IsValidIndex(LODIndex))
+		{
+			if (!bLODHasBeenSimplified)
 			{
-				const FSkelMeshSection& Section = LODModel.Sections[SectionIndex];
-
-				if (!Section.bDisabled)
+				LODVertexNumber = 0;
+				LODTriNumber = 0;
+				const FSkeletalMeshLODModel& LODModel = GetImportedModel()->LODModels[LODIndex];
+				//We can take the vertices and triangles count from the source model
+				for (int32 SectionIndex = 0; SectionIndex < LODModel.Sections.Num(); ++SectionIndex)
 				{
-					//Make sure the count fit in a uint32
-					LODVertexNumber += Section.NumVertices < 0 ? 0 : Section.NumVertices;
-					LODTriNumber += Section.NumTriangles;
+					const FSkelMeshSection& Section = LODModel.Sections[SectionIndex];
+
+					if (!Section.bDisabled)
+					{
+						//Make sure the count fit in a uint32
+						LODVertexNumber += Section.NumVertices < 0 ? 0 : Section.NumVertices;
+						LODTriNumber += Section.NumTriangles;
+					}
 				}
 			}
+			else if (GetImportedModel()->OriginalReductionSourceMeshData.IsValidIndex(LODIndex)
+				&& !GetImportedModel()->OriginalReductionSourceMeshData[LODIndex]->IsEmpty())
+			{
+				//In this case we have to use the stored reduction source data to know how many vertices/triangles we have before the reduction
+				USkeletalMesh* MutableSkeletalMesh = const_cast<USkeletalMesh*>(this);
+				GetImportedModel()->OriginalReductionSourceMeshData[LODIndex]->GetGeometryInfo(LODVertexNumber, LODTriNumber, MutableSkeletalMesh);
+			}
 		}
-		else if (GetImportedModel()->OriginalReductionSourceMeshData.IsValidIndex(LODIndex)
-			&& !GetImportedModel()->OriginalReductionSourceMeshData[LODIndex]->IsEmpty())
-		{
-			//In this case we have to use the stored reduction source data to know how many vertices/triangles we have before the reduction
-			USkeletalMesh* MutableSkeletalMesh = const_cast<USkeletalMesh*>(this);
-			GetImportedModel()->OriginalReductionSourceMeshData[LODIndex]->GetGeometryInfo(LODVertexNumber, LODTriNumber, MutableSkeletalMesh);
-		}
+		bReductionActive = ReductionModule->IsReductionActive(ReductionSettings, LODVertexNumber, LODTriNumber);
 	}
-	return ReductionModule->IsReductionActive(ReductionSettings, LODVertexNumber, LODTriNumber);
+	return bReductionActive;
 }
 
 /* Get a copy of the reduction settings for a specified LOD index. */
