(4.24) Fix RandomStreams not initializing to a random value if the default stream seed is 0
#rnx

diff --git a/Engine/Source/Runtime/Engine/Private/ActorConstruction.cpp b/Engine/Source/Runtime/Engine/Private/ActorConstruction.cpp
--- a/Engine/Source/Runtime/Engine/Private/ActorConstruction.cpp
+++ b/Engine/Source/Runtime/Engine/Private/ActorConstruction.cpp
@@ -43,29 +43,28 @@ DECLARE_CYCLE_STAT(TEXT("InstanceActorComponent"), STAT_InstanceActorComponent,
 // AActor Blueprint Stuff
 
 #if WITH_EDITOR
-static TArray<FRandomStream*> FindRandomStreams(AActor* InActor)
+void AActor::SeedAllRandomStreams()
 {
-	check(InActor);
-	TArray<FRandomStream*> OutStreams;
 	UScriptStruct* RandomStreamStruct = TBaseStructure<FRandomStream>::Get();
-	for( TFieldIterator<FStructProperty> It(InActor->GetClass()) ; It ; ++It )
+	UObject* Archetype = GetArchetype();
+
+	for (TFieldIterator<FStructProperty> It(GetClass()); It; ++It)
 	{
 		FStructProperty* StructProp = *It;
-		if( StructProp->Struct == RandomStreamStruct )
+		if (StructProp->Struct == RandomStreamStruct)
 		{
-			FRandomStream* StreamPtr = StructProp->ContainerPtrToValuePtr<FRandomStream>(InActor);
-			OutStreams.Add(StreamPtr);
-		}
-	}
-	return OutStreams;
-}
+			FRandomStream* ArchetypeStreamPtr = StructProp->ContainerPtrToValuePtr<FRandomStream>(Archetype);
+			FRandomStream* StreamPtr = StructProp->ContainerPtrToValuePtr<FRandomStream>(this);
 
-void AActor::SeedAllRandomStreams()
-{
-	TArray<FRandomStream*> Streams = FindRandomStreams(this);
-	for(int32 i=0; i<Streams.Num(); i++)
-	{
-		Streams[i]->Reset();
+			if (ArchetypeStreamPtr->GetInitialSeed() == 0)
+			{
+				StreamPtr->GenerateNewSeed();
+			}
+			else
+			{
+				StreamPtr->Reset();
+			}
+		}
 	}
 }
 #endif //WITH_EDITOR
