Fix issue where FVirtualTextureChunkStreamingManager wasn't added to the FStreamingManagerCollection::StreamingManagers
Manage lifetime of FVirtualTextureChunkStreamingManager to ensure that it is destroyed at shutdown

diff --git a/Engine/Source/Runtime/Engine/Private/ContentStreaming.cpp b/Engine/Source/Runtime/Engine/Private/ContentStreaming.cpp
--- a/Engine/Source/Runtime/Engine/Private/ContentStreaming.cpp
+++ b/Engine/Source/Runtime/Engine/Private/ContentStreaming.cpp
@@ -741,11 +741,10 @@ void IRenderAssetStreamingManager::PauseTextureStreaming(bool bInShouldPause)
 -----------------------------------------------------------------------------*/
 
 FStreamingManagerCollection::FStreamingManagerCollection()
-:	NumIterations(1)
-,	DisableResourceStreamingCount(0)
-,	LoadMapTimeLimit( 5.0f )
-,   TextureStreamingManager(nullptr)
-, VirtualTextureChunkStreamingManager(nullptr)
+	: NumIterations(1)
+	, DisableResourceStreamingCount(0)
+	, LoadMapTimeLimit(5.0f)
+	, TextureStreamingManager(nullptr)
 {
 #if PLATFORM_SUPPORTS_TEXTURE_STREAMING
 	// Disable texture streaming if that was requested (needs to happen before the call to ProcessNewlyLoadedUObjects, as that can load textures)
@@ -787,10 +786,6 @@ FStreamingManagerCollection::~FStreamingManagerCollection()
 	delete TextureStreamingManager;
 	TextureStreamingManager = nullptr;
 
-	RemoveStreamingManager(VirtualTextureChunkStreamingManager);
-	delete VirtualTextureChunkStreamingManager;
-	VirtualTextureChunkStreamingManager = nullptr;
-
 	UE_CLOG(StreamingManagers.Num() > 0, LogContentStreaming, Display, TEXT("There are %d unreleased StreamingManagers"), StreamingManagers.Num());
 }
 
diff --git a/Engine/Source/Runtime/Engine/Private/VT/UploadingVirtualTexture.cpp b/Engine/Source/Runtime/Engine/Private/VT/UploadingVirtualTexture.cpp
--- a/Engine/Source/Runtime/Engine/Private/VT/UploadingVirtualTexture.cpp
+++ b/Engine/Source/Runtime/Engine/Private/VT/UploadingVirtualTexture.cpp
@@ -38,6 +38,7 @@ FUploadingVirtualTexture::FUploadingVirtualTexture(FVirtualTextureBuiltData* InD
 	HandlePerChunk.AddDefaulted(InData->Chunks.Num());
 	CodecPerChunk.AddDefaulted(InData->Chunks.Num());
 	InvalidChunks.Init(false, InData->Chunks.Num());
+	StreamingManager = FVirtualTextureChunkStreamingManager::AddRef();
 
 	INC_MEMORY_STAT_BY(STAT_TotalHeaderSize, InData->GetMemoryFootprint());
 	INC_MEMORY_STAT_BY(STAT_TileHeaderSize, InData->GetTileMemoryFootprint());
@@ -55,7 +56,8 @@ FUploadingVirtualTexture::~FUploadingVirtualTexture()
 	DEC_DWORD_STAT_BY(STAT_NumTileHeaders, Data->GetNumTileHeaders());
 
 	// Complete all open transcode requests before deleting IFileCacheHandle objects
-	FVirtualTextureChunkStreamingManager::Get().WaitTasksFinished();
+	StreamingManager->WaitTasksFinished();
+	FVirtualTextureChunkStreamingManager::DecRef();
 
 	for (TUniquePtr<FVirtualTextureCodec>& Codec : CodecPerChunk)
 	{
@@ -97,7 +99,7 @@ uint32 FUploadingVirtualTexture::GetLocalMipBias(uint8 vLevel, uint32 vAddress)
 
 FVTRequestPageResult FUploadingVirtualTexture::RequestPageData(const FVirtualTextureProducerHandle& ProducerHandle, uint8 LayerMask, uint8 vLevel, uint32 vAddress, EVTRequestPagePriority Priority)
 {
-	return FVirtualTextureChunkStreamingManager::Get().RequestTile(this, ProducerHandle, LayerMask, FirstMipOffset + vLevel, vAddress, Priority);
+	return StreamingManager->RequestTile(this, ProducerHandle, LayerMask, FirstMipOffset + vLevel, vAddress, Priority);
 }
 
 IVirtualTextureFinalizer* FUploadingVirtualTexture::ProducePageData(FRHICommandListImmediate& RHICmdList,
@@ -110,7 +112,7 @@ IVirtualTextureFinalizer* FUploadingVirtualTexture::ProducePageData(FRHICommandL
 	INC_DWORD_STAT(STAT_VTP_NumUploads);
 
 	const uint32 SkipBorderSize = (Flags & EVTProducePageFlags::SkipPageBorders) != EVTProducePageFlags::None ? Data->TileBorderSize : 0;
-	return FVirtualTextureChunkStreamingManager::Get().ProduceTile(RHICmdList, SkipBorderSize, Data->GetNumLayers(), LayerMask, RequestHandle, TargetLayers);
+	return StreamingManager->ProduceTile(RHICmdList, SkipBorderSize, Data->GetNumLayers(), LayerMask, RequestHandle, TargetLayers);
 }
 
 void FVirtualTextureCodec::RetireOldCodecs()
diff --git a/Engine/Source/Runtime/Engine/Private/VT/UploadingVirtualTexture.h b/Engine/Source/Runtime/Engine/Private/VT/UploadingVirtualTexture.h
--- a/Engine/Source/Runtime/Engine/Private/VT/UploadingVirtualTexture.h
+++ b/Engine/Source/Runtime/Engine/Private/VT/UploadingVirtualTexture.h
@@ -86,4 +86,6 @@ private:
 	TArray< TUniquePtr<FVirtualTextureCodec> > CodecPerChunk;
 	TBitArray<> InvalidChunks; // marks chunks that failed to load
 	int32 FirstMipOffset;
+
+	struct FVirtualTextureChunkStreamingManager* StreamingManager;
 };
diff --git a/Engine/Source/Runtime/Engine/Private/VT/VirtualTextureChunkManager.cpp b/Engine/Source/Runtime/Engine/Private/VT/VirtualTextureChunkManager.cpp
--- a/Engine/Source/Runtime/Engine/Private/VT/VirtualTextureChunkManager.cpp
+++ b/Engine/Source/Runtime/Engine/Private/VT/VirtualTextureChunkManager.cpp
@@ -23,19 +23,28 @@ static FAutoConsoleVariableRef CVarNumTranscodeRequests(
 	TEXT("Number of transcode request that can be in flight. default 32\n"),
 	ECVF_Default);
 
-static FVirtualTextureChunkStreamingManager* VirtualTexturePageStreamingManager = nullptr;
-struct FVirtualTextureChunkStreamingManager& FVirtualTextureChunkStreamingManager::Get()
+
+/** 
+ * FVirtualTextureChunkStreamingManager is a treated as a singleton referenced by multiple FUploadingVirtualTexture objects.
+ * Use RefCount to ensure it is unregistered and destroyed on shutdown.
+ */
+struct FVirtualTextureChunkStreamingManagerSingleton
 {
-	if (VirtualTexturePageStreamingManager == nullptr)
+	FVirtualTextureChunkStreamingManager* Ptr = nullptr;
+	int32 RefCount = 0;
+
+	~FVirtualTextureChunkStreamingManagerSingleton()
 	{
-		VirtualTexturePageStreamingManager = new FVirtualTextureChunkStreamingManager();
+		check(Ptr == nullptr);
+		check(RefCount == 0);
 	}
-	return *VirtualTexturePageStreamingManager;
-}
+
+} GVirtualTextureChunkStreamingManager;
+
 
 FVirtualTextureChunkStreamingManager::FVirtualTextureChunkStreamingManager()
 {
-	IStreamingManager::Get().SetVirtualTextureChunkStreamingManager(this);
+	IStreamingManager::Get().AddStreamingManager(this);
 #if WITH_EDITOR
 	GetVirtualTextureChunkDDCCache()->Initialize();
 #endif
@@ -46,7 +55,26 @@ FVirtualTextureChunkStreamingManager::~FVirtualTextureChunkStreamingManager()
 #if WITH_EDITOR
 	GetVirtualTextureChunkDDCCache()->ShutDown();
 #endif
-	IStreamingManager::Get().SetVirtualTextureChunkStreamingManager(nullptr);
+	IStreamingManager::Get().RemoveStreamingManager(this);
+}
+
+FVirtualTextureChunkStreamingManager* FVirtualTextureChunkStreamingManager::AddRef()
+{
+	if (GVirtualTextureChunkStreamingManager.RefCount++ == 0)
+	{
+		GVirtualTextureChunkStreamingManager.Ptr = new FVirtualTextureChunkStreamingManager();
+	}
+	return GVirtualTextureChunkStreamingManager.Ptr;
+}
+
+void FVirtualTextureChunkStreamingManager::DecRef()
+{
+	check(GVirtualTextureChunkStreamingManager.RefCount > 0);
+	if (--GVirtualTextureChunkStreamingManager.RefCount == 0)
+	{
+		delete GVirtualTextureChunkStreamingManager.Ptr;
+		GVirtualTextureChunkStreamingManager.Ptr = nullptr;
+	}
 }
 
 void FVirtualTextureChunkStreamingManager::UpdateResourceStreaming(float DeltaTime, bool bProcessEverything /*= false*/)
diff --git a/Engine/Source/Runtime/Engine/Private/VT/VirtualTextureChunkManager.h b/Engine/Source/Runtime/Engine/Private/VT/VirtualTextureChunkManager.h
--- a/Engine/Source/Runtime/Engine/Private/VT/VirtualTextureChunkManager.h
+++ b/Engine/Source/Runtime/Engine/Private/VT/VirtualTextureChunkManager.h
@@ -33,7 +33,9 @@ DECLARE_DWORD_COUNTER_STAT(TEXT("Num uploads"), STAT_VTP_NumUploads, STATGROUP_V
 
 struct FVirtualTextureChunkStreamingManager final  : public IStreamingManager
 {
-	static struct FVirtualTextureChunkStreamingManager& Get();
+	static FVirtualTextureChunkStreamingManager* AddRef();
+	static void DecRef();
+
 private:
 	FVirtualTextureChunkStreamingManager();
 	virtual ~FVirtualTextureChunkStreamingManager();
diff --git a/Engine/Source/Runtime/Engine/Public/ContentStreaming.h b/Engine/Source/Runtime/Engine/Public/ContentStreaming.h
--- a/Engine/Source/Runtime/Engine/Public/ContentStreaming.h
+++ b/Engine/Source/Runtime/Engine/Public/ContentStreaming.h
@@ -751,11 +751,6 @@ struct ENGINE_VTABLE FStreamingManagerCollection : public IStreamingManager
 	virtual void OnAudioStreamingParamsChanged() override;
 #endif
 
-	void SetVirtualTextureChunkStreamingManager(IStreamingManager* InVirtualTextureChunkStreamingManager)
-	{
-		VirtualTextureChunkStreamingManager = InVirtualTextureChunkStreamingManager;
-	}
-
 protected:
 
 	virtual void AddOrRemoveTextureStreamingManagerIfNeeded(bool bIsInit=false);
@@ -780,8 +775,6 @@ protected:
 	/** The animation streaming manager, should always exist */
 	IAnimationStreamingManager* AnimationStreamingManager;
 
-	IStreamingManager* VirtualTextureChunkStreamingManager;
-
 #if WITH_EDITOR
 	// Locks out any audio streaming manager call when we are re-initializing the audio streaming manager.
 	mutable FCriticalSection AudioStreamingManagerCriticalSection;
