Chaos - Move Convex index computation to helper function, compute indices on set convex.
Ensure in navmesh export that convex have indices.

diff --git a/Engine/Source/Runtime/Engine/Classes/PhysicsEngine/ConvexElem.h b/Engine/Source/Runtime/Engine/Classes/PhysicsEngine/ConvexElem.h
--- a/Engine/Source/Runtime/Engine/Classes/PhysicsEngine/ConvexElem.h
+++ b/Engine/Source/Runtime/Engine/Classes/PhysicsEngine/ConvexElem.h
@@ -115,6 +115,8 @@ public:
 	ENGINE_API void SetChaosConvexMesh(TUniquePtr<Chaos::FConvex>&& InChaosConvex);
 
 	ENGINE_API void ResetChaosConvexMesh();
+
+	ENGINE_API void ComputeChaosConvexIndices();
 #endif
 
 	/** Get current transform applied to convex mesh vertices */
diff --git a/Engine/Source/Runtime/Engine/Private/PhysicsEngine/BodySetup.cpp b/Engine/Source/Runtime/Engine/Private/PhysicsEngine/BodySetup.cpp
--- a/Engine/Source/Runtime/Engine/Private/PhysicsEngine/BodySetup.cpp
+++ b/Engine/Source/Runtime/Engine/Private/PhysicsEngine/BodySetup.cpp
@@ -1106,28 +1106,7 @@ void UBodySetup::PostLoad()
 	// data and extract the index data from there.
 	for(FKConvexElem& Convex : AggGeom.ConvexElems)
 	{
-		const int32 NumVerts = Convex.VertexData.Num();
-		if(NumVerts > 0 && Convex.IndexData.Num() == 0)
-		{
-			Chaos::TParticles<Chaos::FReal, 3> ConvexParticles;
-			ConvexParticles.AddParticles(NumVerts);
-
-			for(int32 VertIndex = 0; VertIndex < NumVerts; ++VertIndex)
-			{
-				ConvexParticles.X(VertIndex) = Convex.VertexData[VertIndex];
-			}
-
-			TArray<Chaos::TVector<int32, 3>> Triangles;
-			Chaos::FConvexBuilder::BuildConvexHull(ConvexParticles, Triangles);
-
-			Convex.IndexData.Reset(Triangles.Num() * 3);
-			for(Chaos::TVector<int32, 3> Tri : Triangles)
-			{
-				Convex.IndexData.Add(Tri[0]);
-				Convex.IndexData.Add(Tri[1]);
-				Convex.IndexData.Add(Tri[2]);
-			}
-		}
+		Convex.ComputeChaosConvexIndices();
 	}
 #endif
 }
@@ -1732,12 +1711,39 @@ const TUniquePtr<Chaos::FConvex>& FKConvexElem::GetChaosConvexMesh() const
 void FKConvexElem::SetChaosConvexMesh(TUniquePtr<Chaos::FConvex>&& InChaosConvex)
 {
 	ChaosConvex = MoveTemp(InChaosConvex);
+	ComputeChaosConvexIndices();
 }
 
 void FKConvexElem::ResetChaosConvexMesh()
 {
 	ChaosConvex.Reset();
 }
+
+ENGINE_API void FKConvexElem::ComputeChaosConvexIndices()
+{
+	const int32 NumVerts = VertexData.Num();
+	if(NumVerts > 0 && IndexData.Num() == 0)
+	{
+		Chaos::TParticles<Chaos::FReal, 3> ConvexParticles;
+		ConvexParticles.AddParticles(NumVerts);
+
+		for(int32 VertIndex = 0; VertIndex < NumVerts; ++VertIndex)
+		{
+			ConvexParticles.X(VertIndex) = VertexData[VertIndex];
+		}
+
+		TArray<Chaos::TVector<int32, 3>> Triangles;
+		Chaos::FConvexBuilder::BuildConvexHull(ConvexParticles, Triangles);
+
+		IndexData.Reset(Triangles.Num() * 3);
+		for(Chaos::TVector<int32, 3> Tri : Triangles)
+		{
+			IndexData.Add(Tri[0]);
+			IndexData.Add(Tri[1]);
+			IndexData.Add(Tri[2]);
+		}
+	}
+}
 #endif
 
 #if WITH_EDITORONLY_DATA
diff --git a/Engine/Source/Runtime/NavigationSystem/Private/NavMesh/RecastNavMeshGenerator.cpp b/Engine/Source/Runtime/NavigationSystem/Private/NavMesh/RecastNavMeshGenerator.cpp
--- a/Engine/Source/Runtime/NavigationSystem/Private/NavMesh/RecastNavMeshGenerator.cpp
+++ b/Engine/Source/Runtime/NavigationSystem/Private/NavMesh/RecastNavMeshGenerator.cpp
@@ -633,6 +633,11 @@ void ExportChaosConvexMesh(const FKConvexElem* const Convex, const FTransform& L
 		VertexBuffer.Add(UnrealCoord.Z);
 	}
 
+	if (Convex->VertexData.Num())
+	{
+		ensure(Convex->IndexData.Num());
+	}
+
 	for (int32 i = 0; i < Convex->IndexData.Num(); ++i)
 	{
 		IndexBuffer.Add(VertOffset + Convex->IndexData[i]);
