Chaos - RBAN support for external forces

diff --git a/Engine/Source/Runtime/Engine/Private/PhysicsEngine/ImmediatePhysics/ImmediatePhysicsChaos/ImmediatePhysicsActorHandle_Chaos.cpp b/Engine/Source/Runtime/Engine/Private/PhysicsEngine/ImmediatePhysics/ImmediatePhysicsChaos/ImmediatePhysicsActorHandle_Chaos.cpp
--- a/Engine/Source/Runtime/Engine/Private/PhysicsEngine/ImmediatePhysics/ImmediatePhysicsChaos/ImmediatePhysicsActorHandle_Chaos.cpp
+++ b/Engine/Source/Runtime/Engine/Private/PhysicsEngine/ImmediatePhysics/ImmediatePhysicsChaos/ImmediatePhysicsActorHandle_Chaos.cpp
@@ -429,39 +429,89 @@ namespace ImmediatePhysics_Chaos
 
 	void FActorHandle::AddForce(const FVector& Force)
 	{
-#if IMMEDIATEPHYSICS_CHAOS_TODO
-#endif
+		using namespace Chaos;
+
+		if (TPBDRigidParticleHandle<FReal, 3>* Rigid = Handle()->CastToRigidParticle())
+		{
+			Rigid->ExternalForce() += Force;
+		}
 	}
 
 	void FActorHandle::AddRadialForce(const FVector& Origin, float Strength, float Radius, ERadialImpulseFalloff Falloff, EForceType ForceType)
 	{
-#if IMMEDIATEPHYSICS_CHAOS_TODO
-#endif
+		using namespace Chaos;
+
+		if (TPBDRigidParticleHandle<FReal, 3>* Rigid = Handle()->CastToRigidParticle())
+		{
+			const FRigidTransform3& PCOMTransform = FParticleUtilities::GetCoMWorldTransform(Rigid);
+			FVec3 Delta = PCOMTransform.GetTranslation() - Origin;
+
+			const float Mag = Delta.Size();
+			if (Mag > Radius)
+			{
+				return;
+			}
+			Delta.Normalize();
+
+			float ImpulseMag = Strength;
+			if (Falloff == RIF_Linear)
+			{
+				ImpulseMag *= (1.0f - (Mag / Radius));
+			}
+
+			const FVec3 PImpulse = Delta * ImpulseMag;
+			const FVec3 ApplyDelta = (ForceType == EForceType::AddAcceleration || ForceType == EForceType::AddVelocity) ? PImpulse : PImpulse * Rigid->InvM();
+
+			if (ForceType == EForceType::AddImpulse || ForceType == EForceType::AddVelocity)
+			{
+				Rigid->V() += ApplyDelta;
+			}
+			else
+			{
+				Rigid->F() += ApplyDelta;
+			}
+		}
 	}
 
 	void FActorHandle::SetLinearDamping(float NewLinearDamping)
 	{
-#if IMMEDIATEPHYSICS_CHAOS_TODO
-#endif
+		using namespace Chaos;
+
+		if (TPBDRigidParticleHandle<FReal, 3>* Rigid = Handle()->CastToRigidParticle())
+		{
+			Rigid->LinearEtherDrag() = NewLinearDamping;
+		}
 	}
 
 	float FActorHandle::GetLinearDamping() const
 	{
-#if IMMEDIATEPHYSICS_CHAOS_TODO
-#endif
+		using namespace Chaos;
+
+		if (TPBDRigidParticleHandle<FReal, 3>* Rigid = Handle()->CastToRigidParticle())
+		{
+			return Rigid->LinearEtherDrag();
+		}
 		return 0.0f;
 	}
 
 	void FActorHandle::SetAngularDamping(float NewAngularDamping)
 	{
-#if IMMEDIATEPHYSICS_CHAOS_TODO
-#endif
+		using namespace Chaos;
+
+		if (TPBDRigidParticleHandle<FReal, 3>* Rigid = Handle()->CastToRigidParticle())
+		{
+			Rigid->AngularEtherDrag() = NewAngularDamping;
+		}
 	}
 
 	float FActorHandle::GetAngularDamping() const
 	{
-#if IMMEDIATEPHYSICS_CHAOS_TODO
-#endif
+		using namespace Chaos;
+
+		if (TPBDRigidParticleHandle<FReal, 3>* Rigid = Handle()->CastToRigidParticle())
+		{
+			return Rigid->AngularEtherDrag();
+		}
 		return 0.0f;
 	}
 
diff --git a/Engine/Source/Runtime/Experimental/Chaos/Private/Chaos/Evolution/PBDMinEvolution.cpp b/Engine/Source/Runtime/Experimental/Chaos/Private/Chaos/Evolution/PBDMinEvolution.cpp
--- a/Engine/Source/Runtime/Experimental/Chaos/Private/Chaos/Evolution/PBDMinEvolution.cpp
+++ b/Engine/Source/Runtime/Experimental/Chaos/Private/Chaos/Evolution/PBDMinEvolution.cpp
@@ -135,6 +135,11 @@ namespace Chaos
 
 				EulerStepPositionRule.Apply(Particle, Dt);
 
+				Particle.F() = FVec3(0);
+				Particle.Torque() = FVec3(0);
+				Particle.ExternalForce() = FVec3(0);
+				Particle.ExternalTorque() = FVec3(0);
+
 				if (Particle.HasBounds())
 				{
 					const FReal BoundsThickness = 1.0f;
