Tickable: prevent attempt to add TickableObjects multiple times from the pending list when an object is created/deleted/created before processing pending add.

diff --git a/Engine/Source/Runtime/Engine/Private/Tickable.cpp b/Engine/Source/Runtime/Engine/Private/Tickable.cpp
--- a/Engine/Source/Runtime/Engine/Private/Tickable.cpp
+++ b/Engine/Source/Runtime/Engine/Private/Tickable.cpp
@@ -104,11 +104,13 @@ void FTickableGameObject::TickObjects(UWorld* World, const int32 InTickType, con
 	// It's a long lock but it's ok, the only thing we can block here is the GC worker thread that destroys UObjects
 	FScopeLock LockTickableObjects(&Statics.TickableObjectsCritical);
 
+	TSet<FTickableGameObject*> AddedTickableObjects;
 	while (FTickableGameObject* NewTickableObject = Statics.NewTickableObjects.Pop())
 	{
 		// If the new tickable object has already been deleted, don't add it to the tickable objects list
-		if (!Statics.DeletedTickableObjects.Contains(NewTickableObject))
+		if (!Statics.DeletedTickableObjects.Contains(NewTickableObject) && !AddedTickableObjects.Contains(NewTickableObject))
 		{
+			AddedTickableObjects.Add(NewTickableObject);
 			AddTickableObject(Statics.TickableObjects, NewTickableObject);
 		}
 	}
