Chaos - Improved Joint Drive Stiffness conversion from PhysX

diff --git a/Engine/Source/Runtime/Engine/Private/PhysicsEngine/ImmediatePhysics/ImmediatePhysicsChaos/ImmediatePhysicsJointHandle_Chaos.cpp b/Engine/Source/Runtime/Engine/Private/PhysicsEngine/ImmediatePhysics/ImmediatePhysicsChaos/ImmediatePhysicsJointHandle_Chaos.cpp
--- a/Engine/Source/Runtime/Engine/Private/PhysicsEngine/ImmediatePhysics/ImmediatePhysicsChaos/ImmediatePhysicsJointHandle_Chaos.cpp
+++ b/Engine/Source/Runtime/Engine/Private/PhysicsEngine/ImmediatePhysics/ImmediatePhysicsChaos/ImmediatePhysicsJointHandle_Chaos.cpp
@@ -22,7 +22,10 @@ namespace ImmediatePhysics_Chaos
 	bool ChaosImmediate_StiffnessUseMass = false;
 	FAutoConsoleVariableRef CVarStiffnessUseMass(TEXT("p.Chaos.ImmPhys.StiffnessUseMass"), ChaosImmediate_StiffnessUseMass, TEXT("Whether to use mass to scale stiffness in the conversion."));
 
-	float ChaosImmediate_DriveStiffnessScale = 8.0f;
+	bool ChaosImmediate_StiffnessUseDistance = true;
+	FAutoConsoleVariableRef CVarStiffnessUseDistance(TEXT("p.Chaos.ImmPhys.StiffnessUseDistance"), ChaosImmediate_StiffnessUseDistance, TEXT("Whether to use joint connector length to scale stiffness in the conversion."));
+
+	float ChaosImmediate_DriveStiffnessScale = 30.0f;
 	FAutoConsoleVariableRef CVarDriveStiffnessScale(TEXT("p.Chaos.ImmPhys.DriveStiffnessScale"), ChaosImmediate_DriveStiffnessScale, TEXT("Conversion factor for drive stiffness."));
 
 	float ChaosImmediate_JointMinProjection = 0.0f;
@@ -57,16 +60,28 @@ namespace ImmediatePhysics_Chaos
 	{
 		using namespace Chaos;
 
-		// Calculate inertia of the system about the joint connector using parallel axis theorem
-		float XI = (FReal)1.0;
+		float InvStiffnessScale = (FReal)1.0;
 		if (ChaosImmediate_StiffnessUseMass)
 		{
+			// Calculate inertia of the system about the joint connector using parallel axis theorem
 			const FVec3 XI0 = I0 + M0 * (X0 * X0);
 			const FVec3 XI1 = I1 + M1 * (X1 * X1);
-			XI = FMath::Max(KINDA_SMALL_NUMBER, XI0.Max() + XI1.Max());
+			InvStiffnessScale = FMath::Max(KINDA_SMALL_NUMBER, XI0.Max() + XI1.Max());
+		}
+		if (ChaosImmediate_StiffnessUseDistance)
+		{
+			// Scale stiffness by distance to account for the fact that our rotational drives are applied at
+			// the center of mass, not the connector...although maybe that should change
+			float Distance0 = (M0 > 0)? X0.Size() : 0;
+			float Distance1 = (M1 > 0)? X1.Size() : 0;
+			float Distance = FMath::Max(Distance0, Distance1);
+			if (Distance > 1)
+			{
+				InvStiffnessScale = InvStiffnessScale * Distance;
+			}
 		}
 
-		return ChaosImmediate_DriveStiffnessScale * InStiffness / XI;
+		return ChaosImmediate_DriveStiffnessScale * InStiffness / InvStiffnessScale;
 	}
 
 	float ConvertSoftLinearStiffness(const float InStiffness, const float InvM0, const float InvM1)
@@ -156,7 +171,7 @@ namespace ImmediatePhysics_Chaos
 			ConstraintSettings.Motion.bAngularSwingDriveEnabled = ConstraintInstance->ProfileInstance.AngularDrive.SwingDrive.bEnablePositionDrive;
 		}
 
-		ConstraintSettings.Motion.AngularDriveStiffness = ConvertDriveAngularStiffness(ConstraintInstance->ProfileInstance.AngularDrive.TwistDrive.Stiffness, ConstraintSettings.ConstraintFrames[0].GetTranslation(), Actor1->GetMass(), Actor1->GetInertia(), ConstraintSettings.ConstraintFrames[1].GetTranslation(), Actor2->GetMass(), Actor2->GetInertia());
+		ConstraintSettings.Motion.AngularDriveStiffness = ConvertDriveAngularStiffness(ConstraintInstance->ProfileInstance.AngularDrive.TwistDrive.Stiffness, ConstraintSettings.ConstraintFrames[0].GetTranslation() - Actor1->GetLocalCoMTransform().GetTranslation(), Actor1->GetMass(), Actor1->GetInertia(), ConstraintSettings.ConstraintFrames[1].GetTranslation() - Actor2->GetLocalCoMTransform().GetTranslation(), Actor2->GetMass(), Actor2->GetInertia());
 
 		ConstraintSettings.Motion.Sanitize();
 
diff --git a/Engine/Source/Runtime/Experimental/Chaos/Private/Chaos/Joint/PBDJointSolverGaussSeidel.cpp b/Engine/Source/Runtime/Experimental/Chaos/Private/Chaos/Joint/PBDJointSolverGaussSeidel.cpp
--- a/Engine/Source/Runtime/Experimental/Chaos/Private/Chaos/Joint/PBDJointSolverGaussSeidel.cpp
+++ b/Engine/Source/Runtime/Experimental/Chaos/Private/Chaos/Joint/PBDJointSolverGaussSeidel.cpp
@@ -62,6 +62,7 @@ namespace Chaos
 		bEnableTwistLimits = SolverSettings.bEnableTwistLimits;
 		bEnableSwingLimits = SolverSettings.bEnableSwingLimits;
 		bEnableDrives = SolverSettings.bEnableDrives;
+		bAccelerationDriveMode = true;
 	}
 
 
@@ -275,43 +276,45 @@ namespace Chaos
 		const FReal IM1 = FVec3::DotProduct(Axis1, IA1);
 		const FReal IM = (IM0 + IM1);
 
-		FReal JV = 0;
+		// Damping angular velocity
+		FReal JW = 0;
 		if (Damping > KINDA_SMALL_NUMBER)
 		{
 			const FVec3 W0 = FRotation3::CalculateAngularVelocity(PrevQs[0], Qs[0], Dt);
 			const FVec3 W1 = FRotation3::CalculateAngularVelocity(PrevQs[1], Qs[1], Dt);
-			JV = FVec3::DotProduct(Axis0, W0) - FVec3::DotProduct(Axis1, W1);
+			JW = FVec3::DotProduct(Axis0, W0) - FVec3::DotProduct(Axis1, W1);
 		}
 
-		//const FReal MassScale = (FReal)1;	// For force springs. Make this an option...
-		const FReal MassScale = IM;			// For acceleration springs. Make this an option...
+		// Cancel out the mass for acceleration drives
+		const FReal MassScale = (bAccelerationDriveMode)? IM : (FReal)1.0;
+
 		FReal DLambda = 0;
 		if (Stiffness > KINDA_SMALL_NUMBER)
 		{
 			const FReal Alpha = MassScale / (Stiffness * Dt * Dt);
 			const FReal AlphaBeta = Damping / Stiffness;
-			const FReal DLambdaNumerator = (Angle - Alpha * Lambda - AlphaBeta * JV);
+			const FReal DLambdaNumerator = (Angle - Alpha * Lambda - AlphaBeta * JW);
 			const FReal DLambdaDenominator = (((FReal)1 + AlphaBeta / Dt) * IM + Alpha);
 			DLambda = DLambdaNumerator / DLambdaDenominator;
 		}
 		else
 		{
 			const FReal Beta = Damping / (MassScale * Dt * Dt);
-			const FReal DLambdaNumerator = -Beta * JV;
+			const FReal DLambdaNumerator = -Beta * JW;
 			const FReal DLambdaDenominator = IM;
 			DLambda = DLambdaNumerator / DLambdaDenominator;
 		}
 
-		const FVec3 DR0 = IA0 * DLambda;
-		const FVec3 DR1 = IA1 * -DLambda;
-		//const FVec3 DR0 = Axis0 * (DLambda * IM0);
-		//const FVec3 DR1 = Axis1 * -(DLambda * IM1);
-
+		//const FVec3 DR0 = IA0 * DLambda;
+		//const FVec3 DR1 = IA1 * -DLambda;
+		const FVec3 DR0 = Axis0 * (DLambda * IM0);
+		const FVec3 DR1 = Axis1 * -(DLambda * IM1);
 		const FRotation3 DQ0 = (FRotation3::FromElements(DR0, 0) * Qs[0]) * (FReal)0.5;
 		const FRotation3 DQ1 = (FRotation3::FromElements(DR1, 0) * Qs[1]) * (FReal)0.5;
 		Qs[0] = (Qs[0] + DQ0).GetNormalized();
 		Qs[1] = (Qs[1] + DQ1).GetNormalized();
 		Qs[1].EnforceShortestArcWith(Qs[0]);
+
 		Lambda += DLambda;
 	}
 
diff --git a/Engine/Source/Runtime/Experimental/Chaos/Public/Chaos/Joint/PBDJointSolverGaussSeidel.h b/Engine/Source/Runtime/Experimental/Chaos/Public/Chaos/Joint/PBDJointSolverGaussSeidel.h
--- a/Engine/Source/Runtime/Experimental/Chaos/Public/Chaos/Joint/PBDJointSolverGaussSeidel.h
+++ b/Engine/Source/Runtime/Experimental/Chaos/Public/Chaos/Joint/PBDJointSolverGaussSeidel.h
@@ -211,6 +211,7 @@ namespace Chaos
 		bool bEnableTwistLimits;
 		bool bEnableSwingLimits;
 		bool bEnableDrives;
+		bool bAccelerationDriveMode;
 	};
 
 }
\ No newline at end of file
