Back out changelist 10737651. Deadlock in async loading thread

diff --git a/Engine/Source/Runtime/Engine/Private/Tickable.cpp b/Engine/Source/Runtime/Engine/Private/Tickable.cpp
--- a/Engine/Source/Runtime/Engine/Private/Tickable.cpp
+++ b/Engine/Source/Runtime/Engine/Private/Tickable.cpp
@@ -73,9 +73,6 @@ FTickableGameObject::FTickableGameObject()
 	if (UObjectInitialized())
 	{
 		Statics.NewTickableObjects.Push(this);
-
-		FScopeLock LockTickableObjects(&Statics.TickableObjectsCritical);
-		Statics.DeletedTickableObjects.Remove(this);
 	}
 	else
 	{
