Fixed crash opening viewport for characters that use linked anim layers in their animation blueprint

Recent changes to manage the ownership of linked anim instances missed the grouped/layer case. This change applies the bCreatedByLinkedAnimGraph flag there too.

diff --git a/Engine/Source/Runtime/Engine/Private/Animation/AnimInstance.cpp b/Engine/Source/Runtime/Engine/Private/Animation/AnimInstance.cpp
--- a/Engine/Source/Runtime/Engine/Private/Animation/AnimInstance.cpp
+++ b/Engine/Source/Runtime/Engine/Private/Animation/AnimInstance.cpp
@@ -2632,6 +2632,7 @@ void UAnimInstance::PerformLinkedLayerOverlayOperation(TSubclassOf<UAnimInstance
 						if (TargetInstance == nullptr || ClassToSet != TargetInstance->GetClass())
 						{
 							UAnimInstance* NewLinkedInstance = NewObject<UAnimInstance>(MeshComp, ClassToSet);
+							NewLinkedInstance->bCreatedByLinkedAnimGraph = true;
 							NewLinkedInstance->InitializeAnimation();
 
 							// Unlink any layer nodes in the new linked instance, as they may have been hooked up to self in InitializeAnimation above.
@@ -2687,6 +2688,7 @@ void UAnimInstance::PerformLinkedLayerOverlayOperation(TSubclassOf<UAnimInstance
 					{
 						// Create and add one linked instance for this group
 						UAnimInstance* NewLinkedInstance = NewObject<UAnimInstance>(MeshComp, ClassToSet);
+						NewLinkedInstance->bCreatedByLinkedAnimGraph = true;
 						NewLinkedInstance->InitializeAnimation();
 
 						// Unlink any layer nodes in the new linked instance, as they may have been hooked up to self in InitializeAnimation above.
