Fixing Tickable objects sometimes failing to register when re-allocated with the same address as previously destroyed ones

diff --git a/Engine/Source/Runtime/Engine/Private/Tickable.cpp b/Engine/Source/Runtime/Engine/Private/Tickable.cpp
--- a/Engine/Source/Runtime/Engine/Private/Tickable.cpp
+++ b/Engine/Source/Runtime/Engine/Private/Tickable.cpp
@@ -10,10 +10,24 @@ struct FTickableStatics
 {
 	FCriticalSection TickableObjectsCritical;
 	TArray<FTickableObjectBase::FTickableObjectEntry> TickableObjects;
-	TLockFreePointerListUnordered<FTickableGameObject, PLATFORM_CACHE_LINE_SIZE> NewTickableObjects;
-	TSet<FTickableGameObject*> DeletedTickableObjects;
+
+	FCriticalSection NewTickableObjectsCritical;
+	TSet<FTickableGameObject*> NewTickableObjects;
+
 	bool bIsTickingObjects = false;
 
+	void QueueTickableObjectForAdd(FTickableGameObject* InTickable)
+	{
+		FScopeLock NewTickableObjectsLock(&NewTickableObjectsCritical);
+		NewTickableObjects.Add(InTickable);
+	}
+
+	void RemoveTickableObjectFromNewObjectsQueue(FTickableGameObject* InTickable)
+	{
+		FScopeLock NewTickableObjectsLock(&NewTickableObjectsCritical);
+		NewTickableObjects.Remove(InTickable);
+	}
+
 	static FTickableStatics& Get()
 	{
 		static FTickableStatics Singleton;
@@ -72,7 +86,7 @@ FTickableGameObject::FTickableGameObject()
 
 	if (UObjectInitialized())
 	{
-		Statics.NewTickableObjects.Push(this);
+		Statics.QueueTickableObjectForAdd(this);
 	}
 	else
 	{
@@ -82,10 +96,9 @@ FTickableGameObject::FTickableGameObject()
 
 FTickableGameObject::~FTickableGameObject()
 {
-	FTickableStatics& Statics = FTickableStatics::Get();
-
+	FTickableStatics& Statics = FTickableStatics::Get();	
+	Statics.RemoveTickableObjectFromNewObjectsQueue(this);	
 	FScopeLock LockTickableObjects(&Statics.TickableObjectsCritical);
-	Statics.DeletedTickableObjects.Add(this);
 	RemoveTickableObject(Statics.TickableObjects, this, Statics.bIsTickingObjects);
 }
 
@@ -101,17 +114,14 @@ void FTickableGameObject::TickObjects(UWorld* World, const int32 InTickType, con
 	// It's a long lock but it's ok, the only thing we can block here is the GC worker thread that destroys UObjects
 	FScopeLock LockTickableObjects(&Statics.TickableObjectsCritical);
 
-	TSet<FTickableGameObject*> AddedTickableObjects;
-	while (FTickableGameObject* NewTickableObject = Statics.NewTickableObjects.Pop())
 	{
-		// If the new tickable object has already been deleted, don't add it to the tickable objects list
-		if (!Statics.DeletedTickableObjects.Contains(NewTickableObject) && !AddedTickableObjects.Contains(NewTickableObject))
+		FScopeLock NewTickableObjectsLock(&Statics.NewTickableObjectsCritical);
+		for (FTickableGameObject* NewTickableObject : Statics.NewTickableObjects)
 		{
-			AddedTickableObjects.Add(NewTickableObject);
 			AddTickableObject(Statics.TickableObjects, NewTickableObject);
 		}
+		Statics.NewTickableObjects.Empty();
 	}
-	Statics.DeletedTickableObjects.Empty();
 
 	if (Statics.TickableObjects.Num() > 0)
 	{
