Fix Chaos Build

diff --git a/Engine/Source/Runtime/Engine/Private/PhysicsEngine/Experimental/PhysInterface_Chaos.cpp b/Engine/Source/Runtime/Engine/Private/PhysicsEngine/Experimental/PhysInterface_Chaos.cpp
--- a/Engine/Source/Runtime/Engine/Private/PhysicsEngine/Experimental/PhysInterface_Chaos.cpp
+++ b/Engine/Source/Runtime/Engine/Private/PhysicsEngine/Experimental/PhysInterface_Chaos.cpp
@@ -1347,7 +1347,7 @@ void FPhysInterface_Chaos::AddGeometry(FPhysicsActorHandle& InActor, const FGeom
 		}
 
 		//todo: we should not be creating unique geometry per actor
-		InActor->SetGeometry(MakeUnique<Chaos::TImplicitObjectUnion<float, 3>>(MoveTemp(Geoms)));
+		InActor->SetGeometry(MakeUnique<Chaos::FImplicitObjectUnion>(MoveTemp(Geoms)));
 		InActor->SetShapesArray(MoveTemp(Shapes));
 	}
 #endif
