Make sure that linked instances are cleaned up properly on unlink and reset

This makes sure that notify states are fired off correctly and any playing montages are correctly terminated

diff --git a/Engine/Source/Runtime/Engine/Private/Animation/AnimNode_LinkedAnimGraph.cpp b/Engine/Source/Runtime/Engine/Private/Animation/AnimNode_LinkedAnimGraph.cpp
--- a/Engine/Source/Runtime/Engine/Private/Animation/AnimNode_LinkedAnimGraph.cpp
+++ b/Engine/Source/Runtime/Engine/Private/Animation/AnimNode_LinkedAnimGraph.cpp
@@ -238,6 +238,8 @@ void FAnimNode_LinkedAnimGraph::ReinitializeLinkedAnimInstance(const UAnimInstan
 			// Never delete the owning animation instance
 			if (InstanceToRun != InOwningAnimInstance)
 			{
+				// Only call UninitializeAnimation if we are not the owning anim instance
+				InstanceToRun->UninitializeAnimation();
 				InstanceToRun->MarkPendingKill();
 			}
 			InstanceToRun = nullptr;
diff --git a/Engine/Source/Runtime/Engine/Private/Animation/AnimNode_LinkedAnimLayer.cpp b/Engine/Source/Runtime/Engine/Private/Animation/AnimNode_LinkedAnimLayer.cpp
--- a/Engine/Source/Runtime/Engine/Private/Animation/AnimNode_LinkedAnimLayer.cpp
+++ b/Engine/Source/Runtime/Engine/Private/Animation/AnimNode_LinkedAnimLayer.cpp
@@ -47,6 +47,7 @@ void FAnimNode_LinkedAnimLayer::InitializeSelfLayer(const UAnimInstance* SelfAni
 	// Switch from dynamic external to internal, kill old instance
 	if (CurrentTarget && CurrentTarget != SelfAnimInstance)
 	{
+		CurrentTarget->UninitializeAnimation();
 		MeshComp->GetLinkedAnimInstances().Remove(CurrentTarget);
 		CurrentTarget->MarkPendingKill();
 		CurrentTarget = nullptr;
diff --git a/Engine/Source/Runtime/Engine/Private/Components/SkeletalMeshComponent.cpp b/Engine/Source/Runtime/Engine/Private/Components/SkeletalMeshComponent.cpp
--- a/Engine/Source/Runtime/Engine/Private/Components/SkeletalMeshComponent.cpp
+++ b/Engine/Source/Runtime/Engine/Private/Components/SkeletalMeshComponent.cpp
@@ -2724,6 +2724,7 @@ void USkeletalMeshComponent::ResetLinkedAnimInstances()
 	{
 		if(LinkedInstance && LinkedInstance->bCreatedByLinkedAnimGraph)
 		{
+			LinkedInstance->EndNotifyStates();
 			LinkedInstance->MarkPendingKill();
 			LinkedInstance = nullptr;
 		}
