Fix solver commands enqueued after solver advance task, and before PhysScene's EndFrame not being flushed, giving game thread stale acceleration structure w/ dangling pointers.

diff --git a/Engine/Source/Runtime/Engine/Private/PhysicsEngine/Experimental/PhysScene_Chaos.cpp b/Engine/Source/Runtime/Engine/Private/PhysicsEngine/Experimental/PhysScene_Chaos.cpp
--- a/Engine/Source/Runtime/Engine/Private/PhysicsEngine/Experimental/PhysScene_Chaos.cpp
+++ b/Engine/Source/Runtime/Engine/Private/PhysicsEngine/Experimental/PhysScene_Chaos.cpp
@@ -1495,6 +1495,18 @@ void FPhysScene_ChaosInterface::EndFrame(ULineBatchComponent* InLineBatcher)
 		//flush queue so we can merge the two threads
 		Dispatcher->Execute();
 
+		// flush solver queues
+		const TArray<FPhysicsSolver*>& SolverList = SolverModule->GetSolvers();
+		for (FPhysicsSolver* Solver : SolverList)
+		{
+			TQueue<TFunction<void(Chaos::FPhysicsSolver*)>, EQueueMode::Mpsc>& Queue = Solver->GetCommandQueue();
+			TFunction<void(Chaos::FPhysicsSolver*)> Command;
+			while (Queue.Dequeue(Command))
+			{
+				Command(Solver);
+			}
+		}
+
 		// Flip the buffers over to the game thread and sync
 		{
 			SCOPE_CYCLE_COUNTER(STAT_FlipResults);
@@ -1503,7 +1515,6 @@ void FPhysScene_ChaosInterface::EndFrame(ULineBatchComponent* InLineBatcher)
 			//for now just copy the whole thing, stomping any changes that came from GT
 			Scene.CopySolverAccelerationStructure();
 
-			const TArray<FPhysicsSolver*>& SolverList = SolverModule->GetSolvers();
 			TArray<FPhysicsSolver*> ActiveSolvers;
 			ActiveSolvers.Reserve(SolverList.Num());
 
