Tweaked DumpTextureStreamingStats command to be more intuitive and added auto-complete

#rnx

diff --git a/Engine/Config/BaseInput.ini b/Engine/Config/BaseInput.ini
--- a/Engine/Config/BaseInput.ini
+++ b/Engine/Config/BaseInput.ini
@@ -122,6 +122,7 @@ MaxScrollbackSize=1024
 +ManualAutoCompleteList=(Command="ListSkeletalMeshes",Desc="Lists info for all skeletal meshes")
 +ManualAutoCompleteList=(Command="ListStaticMeshes",Desc="Lists info for all static meshes")
 +ManualAutoCompleteList=(Command="InvestigateTexture",Desc="Shows streaming info about the specified texture")
++ManualAutoCompleteList=(Command="DumpTextureStreamingStats",Desc="Dump current streaming stats (e.g. pool capacity) to the log")
 ;placed here so we can type res<tab> to restartlevel as we often do that
 +ManualAutoCompleteList=(Command="RestartLevel",Desc="Restarts the level")
 +ManualAutoCompleteList=(Command="Module List",Desc="Lists all known modules")
diff --git a/Engine/Source/Runtime/Engine/Private/Streaming/StreamingManagerTexture.cpp b/Engine/Source/Runtime/Engine/Private/Streaming/StreamingManagerTexture.cpp
--- a/Engine/Source/Runtime/Engine/Private/Streaming/StreamingManagerTexture.cpp
+++ b/Engine/Source/Runtime/Engine/Private/Streaming/StreamingManagerTexture.cpp
@@ -1794,10 +1794,14 @@ bool FRenderAssetStreamingManager::HandleDumpTextureStreamingStatsCommand( const
 	float StreamingPoolMB = DisplayedStats.StreamingPool / 1024.f / 1024.f;
 
 	// uses csv stats for access in test builds
-	Ar.Logf(TEXT("Current Texture Streaming Stats") );
-	Ar.Logf(TEXT("  Streaming Assets, Current/Pool = %.2f / %.2f MB (%d%%)"), DisplayedStats.WantedMips / 1024.f / 1024.f, StreamingPoolMB, FMath::RoundToInt(CurrentOccupancyPct));
-	Ar.Logf(TEXT("  Streaming Assets, Target/Pool =  %.2f / %.2f MB (%d%%)"), DisplayedStats.RequiredPool / 1024.f / 1024.f, StreamingPoolMB, FMath::RoundToInt(TargetOccupancyPct));
-	Ar.Logf(TEXT("  Non-Streaming Mips = %.2f MB"), DisplayedStats.NonStreamingMips / 1024.f / 1024.f);
+	Ar.Logf(TEXT("--------------------------------------------------------"));
+	Ar.Logf(TEXT("Texture Streaming Stats:") );
+	Ar.Logf(TEXT("Total Pool Size (aka RenderAssetPool) = %.2f MB"), DisplayedStats.RenderAssetPool / 1024.f / 1024.f);
+	Ar.Logf(TEXT("Non-Streaming Mips = %.2f MB"), DisplayedStats.NonStreamingMips / 1024.f / 1024.f);
+	Ar.Logf(TEXT("Remaining Streaming Pool Size = %.2f MB"), StreamingPoolMB);
+	Ar.Logf(TEXT("Streaming Assets, Current/Pool = %.2f / %.2f MB (%d%%)"), DisplayedStats.WantedMips / 1024.f / 1024.f, StreamingPoolMB, FMath::RoundToInt(CurrentOccupancyPct));
+	Ar.Logf(TEXT("Streaming Assets, Target/Pool =  %.2f / %.2f MB (%d%%)"), DisplayedStats.RequiredPool / 1024.f / 1024.f, StreamingPoolMB, FMath::RoundToInt(TargetOccupancyPct));
+	Ar.Logf(TEXT("--------------------------------------------------------"));
 
 	return true;
 }
