Adding direct access getters on ULocalPlayer to the FSlateUser that corresponds to it, based on its assigned ControllerId.
- Also added a means on FSlateApplication to get a SlateUser from a ControllerId (as opposed to the existing method of getting by user index)
- PlayerController now checks with only the appropriate SlateUser when checking to see if the cursor is in the viewport (and it also only checks the mouse cursor now, not implicitly all pointers as it did before)
- FSceneViewport::UpdateCachedCursorPos ignores all attempts by mouse events that didn't come from the CursorUser. This prevents the faux cursors generated for any additional SlateUsers from confusing the scene viewport. Ulimately we'll probably want a more cogent solution here that allows non-hardware cursors to be functional with FSceneViewport, but atm it's simply not built for that.

diff --git a/Engine/Source/Runtime/Engine/Classes/Engine/LocalPlayer.h b/Engine/Source/Runtime/Engine/Classes/Engine/LocalPlayer.h
--- a/Engine/Source/Runtime/Engine/Classes/Engine/LocalPlayer.h
+++ b/Engine/Source/Runtime/Engine/Classes/Engine/LocalPlayer.h
@@ -27,6 +27,7 @@
 
 class AActor;
 class FSceneView;
+class FSlateUser;
 class UGameInstance;
 class ULocalPlayer;
 struct FMinimalViewInfo;
@@ -264,6 +265,10 @@ public:
 	FReply& GetSlateOperations() { return SlateOperations; }
 	const FReply& GetSlateOperations() const { return SlateOperations; }
 
+	/** Get the SlateUser that this LocalPlayer corresponds to */
+	TSharedPtr<FSlateUser> GetSlateUser();
+	TSharedPtr<const FSlateUser> GetSlateUser() const;
+
 	/**
 	 * Get the world the players actor belongs to
 	 *
diff --git a/Engine/Source/Runtime/Engine/Private/LocalPlayer.cpp b/Engine/Source/Runtime/Engine/Private/LocalPlayer.cpp
--- a/Engine/Source/Runtime/Engine/Private/LocalPlayer.cpp
+++ b/Engine/Source/Runtime/Engine/Private/LocalPlayer.cpp
@@ -15,6 +15,7 @@
 #include "UObject/UObjectIterator.h"
 #include "GameFramework/OnlineReplStructs.h"
 #include "GameFramework/PlayerController.h"
+#include "GameFramework/PlayerState.h"
 #include "Engine/SkeletalMesh.h"
 #include "Components/SkeletalMeshComponent.h"
 #include "UnrealEngine.h"
@@ -28,13 +29,13 @@
 #include "Physics/PhysicsInterfaceCore.h"
 #include "Rendering/SkeletalMeshRenderData.h"
 #include "HAL/PlatformApplicationMisc.h"
+#include "Framework/Application/SlateApplication.h"
 
 #include "IHeadMountedDisplay.h"
 #include "IXRTrackingSystem.h"
 #include "IXRCamera.h"
 #include "SceneViewExtension.h"
 #include "Net/DataChannel.h"
-#include "GameFramework/PlayerState.h"
 
 #include "GameDelegates.h"
 
@@ -1622,6 +1623,16 @@ bool ULocalPlayer::IsCachedUniqueNetIdPairedWithControllerId() const
 	return (CachedUniqueNetId == UniqueIdFromController);
 }
 
+TSharedPtr<FSlateUser> ULocalPlayer::GetSlateUser()
+{
+	return FSlateApplication::Get().GetUserFromControllerId(ControllerId);
+}
+
+TSharedPtr<const FSlateUser> ULocalPlayer::GetSlateUser() const
+{
+	return FSlateApplication::Get().GetUserFromControllerId(ControllerId);
+}
+
 UWorld* ULocalPlayer::GetWorld() const
 {
 	return ViewportClient ? ViewportClient->GetWorld() : nullptr;
diff --git a/Engine/Source/Runtime/Engine/Private/PlayerController.cpp b/Engine/Source/Runtime/Engine/Private/PlayerController.cpp
--- a/Engine/Source/Runtime/Engine/Private/PlayerController.cpp
+++ b/Engine/Source/Runtime/Engine/Private/PlayerController.cpp
@@ -28,6 +28,7 @@
 #include "DrawDebugHelpers.h"
 #include "EngineUtils.h"
 #include "Framework/Application/SlateApplication.h"
+#include "Framework/Application/SlateUser.h"
 #include "Widgets/SViewport.h"
 #include "Engine/Console.h"
 #include "Net/UnrealNetwork.h"
@@ -4501,7 +4502,16 @@ ULocalPlayer* APlayerController::GetLocalPlayer() const
 
 bool APlayerController::IsInViewportClient(UGameViewportClient* ViewportClient) const
 {
-	return ViewportClient && ViewportClient->GetGameViewportWidget().IsValid() && ViewportClient->GetGameViewportWidget()->IsDirectlyHovered();
+	const ULocalPlayer* LocalPlayer = GetLocalPlayer();
+	if (LocalPlayer && ViewportClient)
+	{
+		TSharedPtr<const FSlateUser> SlateUser = LocalPlayer->GetSlateUser();
+		if (SlateUser && SlateUser->IsWidgetDirectlyUnderCursor(ViewportClient->GetGameViewportWidget()))
+		{
+			return true;
+		}
+	}
+	return false;
 }
 
 int32 APlayerController::GetInputIndex() const
diff --git a/Engine/Source/Runtime/Engine/Private/Slate/SceneViewport.cpp b/Engine/Source/Runtime/Engine/Private/Slate/SceneViewport.cpp
--- a/Engine/Source/Runtime/Engine/Private/Slate/SceneViewport.cpp
+++ b/Engine/Source/Runtime/Engine/Private/Slate/SceneViewport.cpp
@@ -213,11 +213,14 @@ void FSceneViewport::ProcessInput( float DeltaTime )
 
 void FSceneViewport::UpdateCachedCursorPos( const FGeometry& InGeometry, const FPointerEvent& InMouseEvent )
 {
-	FVector2D LocalPixelMousePos = InGeometry.AbsoluteToLocal(InMouseEvent.GetScreenSpacePosition());
-	LocalPixelMousePos.X *= CachedGeometry.Scale;
-	LocalPixelMousePos.Y *= CachedGeometry.Scale;
+	if (InMouseEvent.GetUserIndex() == FSlateApplication::CursorUserIndex)
+	{
+		FVector2D LocalPixelMousePos = InGeometry.AbsoluteToLocal(InMouseEvent.GetScreenSpacePosition());
+		LocalPixelMousePos.X *= CachedGeometry.Scale;
+		LocalPixelMousePos.Y *= CachedGeometry.Scale;
 
-	CachedCursorPos = LocalPixelMousePos.IntPoint();
+		CachedCursorPos = LocalPixelMousePos.IntPoint();
+	}
 }
 
 void FSceneViewport::UpdateCachedGeometry( const FGeometry& InGeometry )
diff --git a/Engine/Source/Runtime/Slate/Public/Framework/Application/SlateApplication.h b/Engine/Source/Runtime/Slate/Public/Framework/Application/SlateApplication.h
--- a/Engine/Source/Runtime/Slate/Public/Framework/Application/SlateApplication.h
+++ b/Engine/Source/Runtime/Slate/Public/Framework/Application/SlateApplication.h
@@ -877,6 +877,25 @@ public:
 	}
 	FORCEINLINE TSharedPtr<const FSlateUser> GetUser(const FInputEvent& InputEvent) const { return GetUser(InputEvent.GetUserIndex()); }
 	FORCEINLINE TSharedPtr<FSlateUser> GetUser(const FInputEvent& InputEvent) { return GetUser(InputEvent.GetUserIndex()); }
+	
+	FORCEINLINE TSharedPtr<FSlateUser> GetUserFromControllerId(int32 ControllerId)
+	{
+		TOptional<int32> UserIndex = GetUserIndexForController(ControllerId);
+		if (UserIndex.IsSet())
+		{
+			return GetUser(UserIndex.GetValue());
+		}
+		return nullptr;
+	}
+	FORCEINLINE TSharedPtr<const FSlateUser> GetUserFromControllerId(int32 ControllerId) const
+	{
+		TOptional<int32> UserIndex = GetUserIndexForController(ControllerId);
+		if (UserIndex.IsSet())
+		{
+			return GetUser(UserIndex.GetValue());
+		}
+		return nullptr;
+	}
 
 	/** Get the standard 'default' user (there's always guaranteed to be at least one). */
 	FORCEINLINE TSharedPtr<const FSlateUser> GetCursorUser() const
