Fix CIS (prrama optimizes left in Chaos joint files)

diff --git a/Engine/Source/Runtime/Engine/Private/PhysicsEngine/ImmediatePhysics/ImmediatePhysicsChaos/ImmediatePhysicsJointHandle_Chaos.cpp b/Engine/Source/Runtime/Engine/Private/PhysicsEngine/ImmediatePhysics/ImmediatePhysicsChaos/ImmediatePhysicsJointHandle_Chaos.cpp
--- a/Engine/Source/Runtime/Engine/Private/PhysicsEngine/ImmediatePhysics/ImmediatePhysicsChaos/ImmediatePhysicsJointHandle_Chaos.cpp
+++ b/Engine/Source/Runtime/Engine/Private/PhysicsEngine/ImmediatePhysics/ImmediatePhysicsChaos/ImmediatePhysicsJointHandle_Chaos.cpp
@@ -8,7 +8,7 @@
 
 #include "PhysicsEngine/ConstraintInstance.h"
 
-#pragma optimize("", off)
+//#pragma optimize("", off)
 
 static_assert((int32)Chaos::EJointMotionType::Free == (int32)EAngularConstraintMotion::ACM_Free, "Chaos::EJointMotionType and EAngularConstraintMotion mismatch");
 static_assert((int32)Chaos::EJointMotionType::Limited == (int32)EAngularConstraintMotion::ACM_Limited, "Chaos::EJointMotionType and EAngularConstraintMotion mismatch");
diff --git a/Engine/Source/Runtime/Experimental/Chaos/Private/Chaos/Joint/PBDJointSolverGaussSeidel.cpp b/Engine/Source/Runtime/Experimental/Chaos/Private/Chaos/Joint/PBDJointSolverGaussSeidel.cpp
--- a/Engine/Source/Runtime/Experimental/Chaos/Private/Chaos/Joint/PBDJointSolverGaussSeidel.cpp
+++ b/Engine/Source/Runtime/Experimental/Chaos/Private/Chaos/Joint/PBDJointSolverGaussSeidel.cpp
@@ -6,7 +6,7 @@
 #include "ChaosLog.h"
 #include "ChaosStats.h"
 
-#pragma optimize("", off)
+//#pragma optimize("", off)
 
 namespace Chaos
 {
