Fix duplicate cvar warning.

diff --git a/Engine/Source/Runtime/Engine/Private/PhysicsEngine/ImmediatePhysics/ImmediatePhysicsChaos/ImmediatePhysicsJointHandle_Chaos.cpp b/Engine/Source/Runtime/Engine/Private/PhysicsEngine/ImmediatePhysics/ImmediatePhysicsChaos/ImmediatePhysicsJointHandle_Chaos.cpp
--- a/Engine/Source/Runtime/Engine/Private/PhysicsEngine/ImmediatePhysics/ImmediatePhysicsChaos/ImmediatePhysicsJointHandle_Chaos.cpp
+++ b/Engine/Source/Runtime/Engine/Private/PhysicsEngine/ImmediatePhysics/ImmediatePhysicsChaos/ImmediatePhysicsJointHandle_Chaos.cpp
@@ -32,13 +32,13 @@ namespace ImmediatePhysics_Chaos
 
 	float ChaosImmediate_JointMinLinearProjection = 0.0f;
 	float ChaosImmediate_JointMaxLinearProjection = 0.5f;
-	FAutoConsoleVariableRef CVarJointMinLinearProjection(TEXT("p.Chaos.ImmPhys.JointMinProjection"), ChaosImmediate_JointMinLinearProjection, TEXT("Joint min projection (for joints with projection disabled)."));
-	FAutoConsoleVariableRef CVarJointMaxLinearProjection(TEXT("p.Chaos.ImmPhys.JointMaxProjection"), ChaosImmediate_JointMaxLinearProjection, TEXT("Joint max projection (for joints with projection enabled)."));
+	FAutoConsoleVariableRef CVarJointMinLinearProjection(TEXT("p.Chaos.ImmPhys.JointMinLinearProjection"), ChaosImmediate_JointMinLinearProjection, TEXT("Joint min projection (for joints with projection disabled)."));
+	FAutoConsoleVariableRef CVarJointMaxLinearProjection(TEXT("p.Chaos.ImmPhys.JointMaxLinearProjection"), ChaosImmediate_JointMaxLinearProjection, TEXT("Joint max projection (for joints with projection enabled)."));
 
 	float ChaosImmediate_JointMinAngularProjection = 0.0f;
 	float ChaosImmediate_JointMaxAngularProjection = 0.5f;
-	FAutoConsoleVariableRef CVarJointMinAngularProjection(TEXT("p.Chaos.ImmPhys.JointMinProjection"), ChaosImmediate_JointMinAngularProjection, TEXT("Joint min projection (for joints with projection disabled)."));
-	FAutoConsoleVariableRef CVarJointMaxAngularProjection(TEXT("p.Chaos.ImmPhys.JointMaxProjection"), ChaosImmediate_JointMaxAngularProjection, TEXT("Joint max projection (for joints with projection enabled)."));
+	FAutoConsoleVariableRef CVarJointMinAngularProjection(TEXT("p.Chaos.ImmPhys.JointMinAngularProjection"), ChaosImmediate_JointMinAngularProjection, TEXT("Joint min projection (for joints with projection disabled)."));
+	FAutoConsoleVariableRef CVarJointMaxAngularProjection(TEXT("p.Chaos.ImmPhys.JointMaxAngularProjection"), ChaosImmediate_JointMaxAngularProjection, TEXT("Joint max projection (for joints with projection enabled)."));
 
 	float ChaosImmediate_DriveStiffnessSourceMin = 100.0f;	// PhysX stiffness per inertia that we translate to a Chaos stiffness of ChaosImmediate_DriveStiffnessTargetMin
 	float ChaosImmediate_DriveStiffnessSourceMax = 2000.0f;	// PhysX stiffness per inertia that we translate to a Chaos stiffness of ChaosImmediate_DriveStiffnessTargetMax
