Fix bogus warnings when rendering BSP models with cooked data.

UModelComponent::BuildRenderData() was checking the max vertex index of each element against the size of the VertexBuffer.Vertices array. When running with cooked data, this array is cleared after the GPU vertex buffer is created, which happens before BuildRenderData() runs.

diff --git a/Engine/Source/Runtime/Engine/Private/ModelRender.cpp b/Engine/Source/Runtime/Engine/Private/ModelRender.cpp
--- a/Engine/Source/Runtime/Engine/Private/ModelRender.cpp
+++ b/Engine/Source/Runtime/Engine/Private/ModelRender.cpp
@@ -134,6 +134,12 @@ void UModelComponent::BuildRenderData()
 			if( ensureMsgf( TheModel->Nodes.IsValidIndex( NodeIdx ), TEXT( "Invalid Node Index, Idx:%d, Num:%d" ), NodeIdx, TheModel->Nodes.Num() ) )
 			{
 				const FBspNode& Node = TheModel->Nodes[NodeIdx];
+				if (Node.NumVertices == 0)
+				{
+					Element.MinVertexIndex = 0;
+					continue;
+				}
+
 				if( ensureMsgf( TheModel->Surfs.IsValidIndex(Node.iSurf), TEXT("Invalid Surf Index, Idx:%d, Num:%d"), Node.iSurf, TheModel->Surfs.Num() ) )
 				{
 					const FBspSurf& Surf = TheModel->Surfs[Node.iSurf];
@@ -167,10 +173,11 @@ void UModelComponent::BuildRenderData()
 						Element.MinVertexIndex = FMath::Min(Node.iVertexIndex + Node.NumVertices * BackFace, Element.MinVertexIndex);
 						Element.MaxVertexIndex = FMath::Max(Node.iVertexIndex + Node.NumVertices * BackFace + Node.NumVertices - 1, Element.MaxVertexIndex);
 
-						if (Element.MaxVertexIndex > (uint32)TheModel->VertexBuffer.Vertices.Num())
+						uint32 NumModelVertices = TheModel->VertexBuffer.Buffers.StaticMeshVertexBuffer.GetNumVertices();
+						if (Element.MaxVertexIndex >= NumModelVertices)
 						{
-							UE_LOG(LogModelComponent, Log, TEXT("Model %s has elements that reference missing vertices. MaxVertex=%d, NumVertices=%d"),
-								*GetPathName(), Element.MaxVertexIndex, TheModel->VertexBuffer.Vertices.Num());
+							UE_LOG(LogModelComponent, Log, TEXT("Model %s has elements that reference missing vertices. MaxVertex=%d, NumModelVertices=%d"),
+								*GetPathName(), Element.MaxVertexIndex, NumModelVertices);
 						}
 					}
 				}
