missing file.

diff --git a/Engine/Source/Runtime/Engine/Public/MaterialShared.h b/Engine/Source/Runtime/Engine/Public/MaterialShared.h
--- a/Engine/Source/Runtime/Engine/Public/MaterialShared.h
+++ b/Engine/Source/Runtime/Engine/Public/MaterialShared.h
@@ -1179,6 +1179,7 @@ private:
 	/** The material's user friendly name, typically the object name. */
 	FString FriendlyName;
 	FString MaterialPath;
+	float CompileTime;
 #endif
 
 	/** The static parameter set that this shader map was compiled with */
