Make the lod.TemporalLag cvar usable by scalability configs.  It is recently being used in mobile scalability, but wasn't flagged as a scalability cvar and thus would fire ensures constantly on mobile platforms.

diff --git a/Engine/Source/Runtime/Engine/Private/SceneManagement.cpp b/Engine/Source/Runtime/Engine/Private/SceneManagement.cpp
--- a/Engine/Source/Runtime/Engine/Private/SceneManagement.cpp
+++ b/Engine/Source/Runtime/Engine/Private/SceneManagement.cpp
@@ -20,7 +20,8 @@
 static TAutoConsoleVariable<float> CVarLODTemporalLag(
 	TEXT("lod.TemporalLag"),
 	0.5f,
-	TEXT("This controls the the time lag for temporal LOD, in seconds."));
+	TEXT("This controls the the time lag for temporal LOD, in seconds."),
+	ECVF_Scalability | ECVF_Default);
 
 void FTemporalLODState::UpdateTemporalLODTransition(const FViewInfo& View, float LastRenderTime)
 {
