#UE4 Fix nan in SubsurfaceProfile

diff --git a/Engine/Source/Runtime/Engine/Classes/Engine/SubsurfaceProfile.h b/Engine/Source/Runtime/Engine/Classes/Engine/SubsurfaceProfile.h
--- a/Engine/Source/Runtime/Engine/Classes/Engine/SubsurfaceProfile.h
+++ b/Engine/Source/Runtime/Engine/Classes/Engine/SubsurfaceProfile.h
@@ -135,7 +135,7 @@ struct FSubsurfaceProfileStruct
 		SurfaceAlbedo = FLinearColor(0, 0, 0);
 		MeanFreePathColor = FLinearColor(0, 0, 0);
 		MeanFreePathDistance = 0.0f;
-		WorldUnitScale = 0.0f;
+		WorldUnitScale = 0.1f;
 		TransmissionTintColor = FLinearColor(0, 0, 0);
 	}
 };
