Chaos - RBAN smoothed variable deltatime

diff --git a/Engine/Source/Runtime/Engine/Private/PhysicsEngine/ImmediatePhysics/ImmediatePhysicsChaos/ImmediatePhysicsSimulation_Chaos.cpp b/Engine/Source/Runtime/Engine/Private/PhysicsEngine/ImmediatePhysics/ImmediatePhysicsChaos/ImmediatePhysicsSimulation_Chaos.cpp
--- a/Engine/Source/Runtime/Engine/Private/PhysicsEngine/ImmediatePhysics/ImmediatePhysicsChaos/ImmediatePhysicsSimulation_Chaos.cpp
+++ b/Engine/Source/Runtime/Engine/Private/PhysicsEngine/ImmediatePhysics/ImmediatePhysicsChaos/ImmediatePhysicsSimulation_Chaos.cpp
@@ -20,10 +20,16 @@
 //////////////////////////////////////////////////////////////////////////
 // @todo(ccaulfield): remove when finished
 //
-float ChaosImmediate_Evolution_DeltaTime = 0.03f;
+float ChaosImmediate_Evolution_StepTime = 0.0f;
+int32 ChaosImmediate_Evolution_NumSteps = 0;
+float ChaosImmediate_Evolution_InitialStepTime = 0.033f;
+int32 ChaosImmediate_Evolution_DeltaTimeCount = 100;
 int32 ChaosImmediate_Evolution_Iterations = 4;
 int32 ChaosImmediate_Evolution_PushOutIterations = 1;
-FAutoConsoleVariableRef CVarChaosImmPhysDeltaTime(TEXT("p.Chaos.ImmPhys.DeltaTime"), ChaosImmediate_Evolution_DeltaTime, TEXT("Override chaos immediate physics delta time if non-zero"));
+FAutoConsoleVariableRef CVarChaosImmPhysStepTime(TEXT("p.Chaos.ImmPhys.StepTime"), ChaosImmediate_Evolution_StepTime, TEXT("Override step time (if not zero)"));
+FAutoConsoleVariableRef CVarChaosImmPhysNumSteps(TEXT("p.Chaos.ImmPhys.NumSteps"), ChaosImmediate_Evolution_NumSteps, TEXT("Override num steps (if not zero)"));
+FAutoConsoleVariableRef CVarChaosImmPhysInitialStepTime(TEXT("p.Chaos.ImmPhys.InitialStepTime"), ChaosImmediate_Evolution_InitialStepTime, TEXT("Initial step time (then calculated from rolling average)"));
+FAutoConsoleVariableRef CVarChaosImmPhysDeltaTimeCount(TEXT("p.Chaos.ImmPhys.DeltaTimeCount"), ChaosImmediate_Evolution_DeltaTimeCount, TEXT("The number of ticks over which the moving average is calculated"));
 FAutoConsoleVariableRef CVarChaosImmPhysIterations(TEXT("p.Chaos.ImmPhys.Iterations"), ChaosImmediate_Evolution_Iterations, TEXT("Number of constraint solver loops in immediate physics"));
 FAutoConsoleVariableRef CVarChaosImmPhysPushOutIterations(TEXT("p.Chaos.ImmPhys.PushOutIterations"), ChaosImmediate_Evolution_PushOutIterations, TEXT("Set the ApplyPushOut() (position correction) iteration count"));
 
@@ -172,6 +178,9 @@ namespace ImmediatePhysics_Chaos
 		, Evolution(Particles, CollisionDetector)
 		, NumActiveDynamicActorHandles(0)
 		, SimulationSpaceTransform(FTransform::Identity)
+		, RollingAverageStepTime(ChaosImmediate_Evolution_InitialStepTime)
+		, NumRollingAverageStepTimes(1)
+		, MaxNumRollingAverageStepTimes(ChaosImmediate_Evolution_DeltaTimeCount)
 	{
 		using namespace Chaos;
 
@@ -471,23 +480,41 @@ namespace ImmediatePhysics_Chaos
 		}
 	}
 
+	FReal FSimulation::UpdateStepTime(const FReal DeltaTime, const FReal MaxStepTime)
+	{
+		RollingAverageStepTime = RollingAverageStepTime + (DeltaTime - RollingAverageStepTime) / NumRollingAverageStepTimes;
+		RollingAverageStepTime = FMath::Min(RollingAverageStepTime, MaxStepTime);
+		NumRollingAverageStepTimes = FMath::Min(NumRollingAverageStepTimes + 1, MaxNumRollingAverageStepTimes);
+		return RollingAverageStepTime;
+	}
+
 	DECLARE_CYCLE_STAT(TEXT("FSimulation::Simulate_Chaos"), STAT_ImmediateSimulate_Chaos, STATGROUP_ImmediatePhysics);
 
-	void FSimulation::Simulate(float DeltaTime, float MaxDeltaTime, int32 MaxSubSteps, const FVector& InGravity)
+	void FSimulation::Simulate(float InDeltaTime, float MaxStepTime, int32 MaxSubSteps, const FVector& InGravity)
 	{
 		SCOPE_CYCLE_COUNTER(STAT_ImmediateSimulate_Chaos);
 		using namespace Chaos;
 
+		// Reject DeltaTime outliers
+		const FReal DeltaTime = FMath::Min(InDeltaTime, MaxStepTime * MaxSubSteps);
+
+		// Update rolling average step time - we want a smooth step time from frame-to-frame that is roughly the target frame rate.
+		// @todo(ccaulfield): decouple sim and game delta times and simulate ahead. Add extrapolation of kinematic targets, and interpolation of physics results.
+		FReal StepTime = UpdateStepTime(DeltaTime, MaxStepTime);
+
+		// Calculate number of steps to run
+		int32 NumSteps = FMath::Clamp(FMath::RoundToInt(DeltaTime / StepTime), 1, MaxSubSteps);
+
 		// TEMP: overrides
 		{
-			if (ChaosImmediate_Evolution_DeltaTime > 0)
+			if (ChaosImmediate_Evolution_StepTime > 0)
+			{
+				StepTime = ChaosImmediate_Evolution_StepTime;
+			}
+			if (ChaosImmediate_Evolution_NumSteps > 0)
 			{
-				// Round Dt to the nearest multiple of fixed step size...
-				const float NumDts = FMath::Max((FReal)1, FMath::RoundToFloat(DeltaTime / ChaosImmediate_Evolution_DeltaTime));
-				DeltaTime = NumDts * ChaosImmediate_Evolution_DeltaTime;
-				MaxDeltaTime = ChaosImmediate_Evolution_DeltaTime;
+				NumSteps = ChaosImmediate_Evolution_NumSteps;
 			}
-			UE_LOG(LogChaosJoint, Verbose, TEXT("Simulate dt = %f"), DeltaTime);
 
 			FPBDJointSolverSettings JointsSettings = Joints.GetSettings();
 			JointsSettings.ApplyPairIterations = ChaosImmediate_Joint_PairIterations;
@@ -522,6 +549,7 @@ namespace ImmediatePhysics_Chaos
 			// TEMP until we can remove constraints again, or I add broad-phase filtering. (FilterCollisionConstraints will crash since the persistent collision changes)
 			Collisions.SetCollisionsEnabled(ChaosImmediate_Collision_Enabled != 0);
 		}
+		UE_LOG(LogChaosJoint, Verbose, TEXT("Simulate Dt = %f Steps %d x %f"), DeltaTime, NumSteps, StepTime);
 
 		DebugDrawKinematicParticles(2, 2, FColor(128, 0, 0));
 		DebugDrawDynamicParticles(2, 2, FColor(192, 192, 0));
@@ -531,9 +559,7 @@ namespace ImmediatePhysics_Chaos
 
 		Evolution.SetGravity(InGravity);
 		
-		Evolution.Advance(DeltaTime, MaxDeltaTime, MaxSubSteps);
-		
-		//Evolution.EndFrame(DeltaTime);
+		Evolution.Advance(StepTime, NumSteps);
 
 		DebugDrawKinematicParticles(1, 4, FColor(128, 0, 0));
 		DebugDrawDynamicParticles(1, 4, FColor(255, 255, 0));
diff --git a/Engine/Source/Runtime/Engine/Public/Physics/ImmediatePhysics/ImmediatePhysicsChaos/ImmediatePhysicsSimulation_Chaos.h b/Engine/Source/Runtime/Engine/Public/Physics/ImmediatePhysics/ImmediatePhysicsChaos/ImmediatePhysicsSimulation_Chaos.h
--- a/Engine/Source/Runtime/Engine/Public/Physics/ImmediatePhysics/ImmediatePhysicsChaos/ImmediatePhysicsSimulation_Chaos.h
+++ b/Engine/Source/Runtime/Engine/Public/Physics/ImmediatePhysics/ImmediatePhysicsChaos/ImmediatePhysicsSimulation_Chaos.h
@@ -63,32 +63,15 @@ namespace ImmediatePhysics_Chaos
 		/** Set bodies that require no collision */
 		void SetIgnoreCollisionActors(const TArray<FActorHandle*>& InIgnoreCollisionActors);
 
-		///** Get/Set whether the body is kinematic or not, kinematics do not simulate and move where they're told. They also act as if they have infinite mass */
-		//bool GetIsKinematic(int32 ActorDataIndex);
-		//void SetIsKinematic(int32 ActorDataIndex, bool bKinematic);
-
 		/** Advance the simulation by DeltaTime */
-		void Simulate(float DeltaTime, float MaxDeltaTime, int32 MaxSubSteps, const FVector& InGravity);
-		void Simulate_AssumesLocked(float DeltaTime, float MaxDeltaTime, int32 MaxSubSteps, const FVector& InGravity) { Simulate(DeltaTime, MaxDeltaTime, MaxSubSteps, InGravity); }
-
-		///** Whether or not an entity is simulated */
-		//bool IsSimulated(int32 ActorDataIndex) const;
-
-
-		///** Add a radial impulse to the given actor */
-		//void AddRadialForce(int32 ActorDataIndex, const FVector& Origin, float Strength, float Radius, ERadialImpulseFalloff Falloff, EForceType ForceType);
-
-		///** Add a force to the given actor */
-		//void AddForce(int32 ActorDataIndex, const FVector& Force);
-
-		///* Set solver iteration counts per step */
-		//void SetPositionIterationCount(int32 InIterationCount);
-		//void SetVelocityIterationCount(int32 InIterationCount);
+		void Simulate(float DeltaTime, float MaxStepTime, int32 MaxSubSteps, const FVector& InGravity);
+		void Simulate_AssumesLocked(float DeltaTime, float MaxStepTime, int32 MaxSubSteps, const FVector& InGravity) { Simulate(DeltaTime, MaxStepTime, MaxSubSteps, InGravity); }
 
 		void SetSimulationSpaceTransform(const FTransform& Transform) { SimulationSpaceTransform = Transform; }
 
 	private:
 		void ConditionConstraints();
+		FReal UpdateStepTime(const FReal DeltaTime, const FReal MaxStepTime);
 
 		void DebugDrawKinematicParticles(const int32 MinDebugLevel, const int32 MaxDebugLevel, const FColor& Color);
 		void DebugDrawDynamicParticles(const int32 MinDebugLevel, const int32 MaxDebugLevel, const FColor& Color);
@@ -129,6 +112,10 @@ namespace ImmediatePhysics_Chaos
 		TMap<const FParticleHandle*, TSet<const FParticleHandle*>> IgnoreCollisionParticlePairTable;
 
 		FTransform SimulationSpaceTransform;
+
+		FReal RollingAverageStepTime;
+		int32 NumRollingAverageStepTimes;
+		int32 MaxNumRollingAverageStepTimes;
 	};
 
 }
diff --git a/Engine/Source/Runtime/Experimental/Chaos/Private/Chaos/Evolution/PBDMinEvolution.cpp b/Engine/Source/Runtime/Experimental/Chaos/Private/Chaos/Evolution/PBDMinEvolution.cpp
--- a/Engine/Source/Runtime/Experimental/Chaos/Private/Chaos/Evolution/PBDMinEvolution.cpp
+++ b/Engine/Source/Runtime/Experimental/Chaos/Private/Chaos/Evolution/PBDMinEvolution.cpp
@@ -41,31 +41,17 @@ namespace Chaos
 		ConstraintRules.Add(Rule);
 	}
 
-	// @todo(ccaulfield): dedupe (PBDRigidsEvolutionGBF)
-	void FPBDMinEvolution::Advance(const FReal Dt, const FReal MaxStepDt, const int32 MaxSteps)
+	void FPBDMinEvolution::Advance(const FReal StepDt, const int32 NumSteps)
 	{
-		// Determine how many steps we would like to take
-		int32 NumSteps = FMath::CeilToInt(Dt / MaxStepDt);
-		if (NumSteps > 0)
+		for (int32 Step = 0; Step < NumSteps; ++Step)
 		{
-			// Determine the step time
-			const FReal StepDt = Dt / (FReal)NumSteps;
-
-			// Limit the number of steps
-			// NOTE: This is after step time calculation so sim will appear to slow down for large Dt
-			// but that is preferable to blowing up from a large timestep.
-			NumSteps = FMath::Clamp(NumSteps, 1, MaxSteps);
+			// StepFraction: how much of the remaining time this step represents, used to interpolate kinematic targets
+			// E.g., for 4 steps this will be: 1/4, 1/3, 1/2, 1
+			const float StepFraction = (FReal)1 / (FReal)(NumSteps - Step);
 
-			for (int32 Step = 0; Step < NumSteps; ++Step)
-			{
-				// StepFraction: how much of the remaining time this step represents, used to interpolate kinematic targets
-				// E.g., for 4 steps this will be: 1/4, 1/3, 1/2, 1
-				const float StepFraction = (FReal)1 / (FReal)(NumSteps - Step);
+			UE_LOG(LogChaos, Verbose, TEXT("Advance dt = %f [%d/%d]"), StepDt, Step + 1, NumSteps);
 
-				UE_LOG(LogChaos, Verbose, TEXT("Advance dt = %f [%d/%d]"), StepDt, Step + 1, NumSteps);
-
-				AdvanceOneTimeStep(StepDt, StepFraction);
-			}
+			AdvanceOneTimeStep(StepDt, StepFraction);
 		}
 	}
 
diff --git a/Engine/Source/Runtime/Experimental/Chaos/Public/Chaos/Evolution/PBDMinEvolution.h b/Engine/Source/Runtime/Experimental/Chaos/Public/Chaos/Evolution/PBDMinEvolution.h
--- a/Engine/Source/Runtime/Experimental/Chaos/Public/Chaos/Evolution/PBDMinEvolution.h
+++ b/Engine/Source/Runtime/Experimental/Chaos/Public/Chaos/Evolution/PBDMinEvolution.h
@@ -41,8 +41,8 @@ namespace Chaos
 
 		void AddConstraintRule(FSimpleConstraintRule* Rule);
 
-		void Advance(const FReal Dt, const FReal MaxStepDt, const int32 MaxSteps);
-		void AdvanceOneTimeStep(const FReal dt, const FReal StepFraction);
+		void Advance(const FReal StepDt, const int32 NumSteps);
+		void AdvanceOneTimeStep(const FReal Dt, const FReal StepFraction);
 
 		void SetNumIterations(const int32 NumIts)
 		{
