Use world-space inertia tensor instead of local for torques applied as acceleration.

diff --git a/Engine/Source/Runtime/Engine/Private/PhysicsEngine/Experimental/PhysScene_Chaos.cpp b/Engine/Source/Runtime/Engine/Private/PhysicsEngine/Experimental/PhysScene_Chaos.cpp
--- a/Engine/Source/Runtime/Engine/Private/PhysicsEngine/Experimental/PhysScene_Chaos.cpp
+++ b/Engine/Source/Runtime/Engine/Private/PhysicsEngine/Experimental/PhysScene_Chaos.cpp
@@ -1299,7 +1299,7 @@ void FPhysScene_ChaosInterface::AddTorque_AssumesLocked(FBodyInstance* BodyInsta
 				const Chaos::TVector<float, 3> CurrentTorque = Rigid->Torque();
 				if (bAccelChange)
 				{
-					Rigid->SetTorque(CurrentTorque + (Rigid->I() * Torque));
+					Rigid->SetTorque(CurrentTorque + (Utilities::ComputeWorldSpaceInertia(Rigid->R(), Rigid->I()) * Torque));
 				}
 				else
 				{
