QualityLevelNodeUsage improvements for cooking.

- Only account for the expressions that are actually used by the material when looking for the quality switches.
- Cache the result so material instances do not recalculate it (in commandlets only).

diff --git a/Engine/Source/Runtime/Engine/Private/Materials/Material.cpp b/Engine/Source/Runtime/Engine/Private/Materials/Material.cpp
--- a/Engine/Source/Runtime/Engine/Private/Materials/Material.cpp
+++ b/Engine/Source/Runtime/Engine/Private/Materials/Material.cpp
@@ -3483,6 +3483,11 @@ void UMaterial::Serialize(FArchive& Ar)
 		{
 			Expressions.Remove(nullptr);
 		}
+		else if (IsRunningCommandlet())
+		{
+			// invalidate the cache of quality level to make sure they are rebuilt. GetQualityLevelNodeUsage() is not going to recalculate it otherwise
+			CachedQualityLevelsUsed.Empty();
+		}
 	}
 
 	if (Ar.UE4Ver() >= VER_UE4_PURGED_FMATERIAL_COMPILE_OUTPUTS)
@@ -3620,8 +3625,17 @@ void UMaterial::BackwardsCompatibilityInputConversion()
 void UMaterial::GetQualityLevelUsage(TArray<bool, TInlineAllocator<EMaterialQualityLevel::Num> >& OutQualityLevelsUsed, EShaderPlatform ShaderPlatform)
 {
 #if WITH_EDITOR
-	GetQualityLevelNodeUsage(OutQualityLevelsUsed);
-	CachedQualityLevelsUsed = OutQualityLevelsUsed;
+	// Assume that a commandlet cannot change nodes once Cached levels are set
+	// FIXME: assumption may not work for a commandlet that edits the material graph?
+	if (IsRunningCommandlet() && CachedQualityLevelsUsed.Num() == EMaterialQualityLevel::Num)
+	{
+		OutQualityLevelsUsed = CachedQualityLevelsUsed;
+	}
+	else
+	{
+		GetQualityLevelNodeUsage(OutQualityLevelsUsed);
+		CachedQualityLevelsUsed = OutQualityLevelsUsed;
+	}
 #else
 	if (CachedQualityLevelsUsed.Num() != EMaterialQualityLevel::Num)
 	{
@@ -3666,9 +3680,14 @@ void UMaterial::GetQualityLevelNodeUsage(TArray<bool, TInlineAllocator<EMaterial
 {
 	OutQualityLevelsUsed.AddZeroed(EMaterialQualityLevel::Num);
 
-	for (int32 ExpressionIndex = 0; ExpressionIndex < Expressions.Num(); ExpressionIndex++)
+	// Some nodes can be disconnected, do not account for those.
+	// TODO: do it recursively in functions
+	TArray<UMaterialExpression*> ActuallyReferencedExpressions;
+	GetAllReferencedExpressions(ActuallyReferencedExpressions, nullptr);
+
+	for (int32 ExpressionIndex = 0, NumExpressions = ActuallyReferencedExpressions.Num(); ExpressionIndex < NumExpressions; ExpressionIndex++)
 	{
-		UMaterialExpression* Expression = Expressions[ExpressionIndex];
+		UMaterialExpression* Expression = ActuallyReferencedExpressions[ExpressionIndex];
 		UMaterialExpressionQualitySwitch* QualitySwitchNode = Cast<UMaterialExpressionQualitySwitch>(Expression);
 		UMaterialExpressionMaterialFunctionCall* MaterialFunctionNode = Cast<UMaterialExpressionMaterialFunctionCall>(Expression);
 		UMaterialExpressionMaterialAttributeLayers* MaterialLayersNode = Cast<UMaterialExpressionMaterialAttributeLayers>(Expression);
