Do not traverse primitive components in the cooker.

- Render state can only be created if FApp::CanEverRender().

diff --git a/Engine/Source/Runtime/Engine/Private/ShaderCompiler/ShaderCompiler.cpp b/Engine/Source/Runtime/Engine/Private/ShaderCompiler/ShaderCompiler.cpp
--- a/Engine/Source/Runtime/Engine/Private/ShaderCompiler/ShaderCompiler.cpp
+++ b/Engine/Source/Runtime/Engine/Private/ShaderCompiler/ShaderCompiler.cpp
@@ -2441,9 +2441,12 @@ void FShaderCompilingManager::ProcessCompiledShaderMaps(
 			Material->NotifyCompilationFinished();
 		}
 
-		PropagateMaterialChangesToPrimitives(MaterialsToUpdate);
+		if (FApp::CanEverRender())
+		{
+			PropagateMaterialChangesToPrimitives(MaterialsToUpdate);
 
-		FEditorSupportDelegates::RedrawAllViewports.Broadcast();
+			FEditorSupportDelegates::RedrawAllViewports.Broadcast();
+		}
 	}
 #endif // WITH_EDITOR
 }
