Material: fix compile error (undo checking only referenced nodes for now).

diff --git a/Engine/Source/Runtime/Engine/Private/Materials/Material.cpp b/Engine/Source/Runtime/Engine/Private/Materials/Material.cpp
--- a/Engine/Source/Runtime/Engine/Private/Materials/Material.cpp
+++ b/Engine/Source/Runtime/Engine/Private/Materials/Material.cpp
@@ -3680,14 +3680,9 @@ void UMaterial::GetQualityLevelNodeUsage(TArray<bool, TInlineAllocator<EMaterial
 {
 	OutQualityLevelsUsed.AddZeroed(EMaterialQualityLevel::Num);
 
-	// Some nodes can be disconnected, do not account for those.
-	// TODO: do it recursively in functions
-	TArray<UMaterialExpression*> ActuallyReferencedExpressions;
-	GetAllReferencedExpressions(ActuallyReferencedExpressions, nullptr);
-
-	for (int32 ExpressionIndex = 0, NumExpressions = ActuallyReferencedExpressions.Num(); ExpressionIndex < NumExpressions; ExpressionIndex++)
+	for (int32 ExpressionIndex = 0, NumExpressions = Expressions.Num(); ExpressionIndex < NumExpressions; ExpressionIndex++)
 	{
-		UMaterialExpression* Expression = ActuallyReferencedExpressions[ExpressionIndex];
+		UMaterialExpression* Expression = Expressions[ExpressionIndex];
 		UMaterialExpressionQualitySwitch* QualitySwitchNode = Cast<UMaterialExpressionQualitySwitch>(Expression);
 		UMaterialExpressionMaterialFunctionCall* MaterialFunctionNode = Cast<UMaterialExpressionMaterialFunctionCall>(Expression);
 		UMaterialExpressionMaterialAttributeLayers* MaterialLayersNode = Cast<UMaterialExpressionMaterialAttributeLayers>(Expression);
