Chaos - Joint position drives (WIP but functional)

diff --git a/Engine/Source/Runtime/Engine/Private/PhysicsEngine/ImmediatePhysics/ImmediatePhysicsChaos/ImmediatePhysicsJointHandle_Chaos.cpp b/Engine/Source/Runtime/Engine/Private/PhysicsEngine/ImmediatePhysics/ImmediatePhysicsChaos/ImmediatePhysicsJointHandle_Chaos.cpp
--- a/Engine/Source/Runtime/Engine/Private/PhysicsEngine/ImmediatePhysics/ImmediatePhysicsChaos/ImmediatePhysicsJointHandle_Chaos.cpp
+++ b/Engine/Source/Runtime/Engine/Private/PhysicsEngine/ImmediatePhysics/ImmediatePhysicsChaos/ImmediatePhysicsJointHandle_Chaos.cpp
@@ -150,6 +150,13 @@ namespace ImmediatePhysics_Chaos
 		ConstraintSettings.Motion.SoftSwingStiffness = DriveAngularStiffnessScale * ConstraintInstance->GetSoftSwingLimitStiffness();
 		ConstraintSettings.Motion.SoftSwingDamping = DriveAngularDampingScale * ConstraintInstance->GetSoftSwingLimitDamping();
 
+		ConstraintSettings.Motion.LinearDriveTarget = ConstraintInstance->ProfileInstance.LinearDrive.PositionTarget;
+		ConstraintSettings.Motion.bLinearDriveEnabled[0] = ConstraintInstance->ProfileInstance.LinearDrive.XDrive.bEnablePositionDrive;
+		ConstraintSettings.Motion.bLinearDriveEnabled[1] = ConstraintInstance->ProfileInstance.LinearDrive.YDrive.bEnablePositionDrive;
+		ConstraintSettings.Motion.bLinearDriveEnabled[2] = ConstraintInstance->ProfileInstance.LinearDrive.ZDrive.bEnablePositionDrive;
+		ConstraintSettings.Motion.LinearDriveStiffness = 0.3f;// ConstraintInstance->ProfileInstance.LinearDrive.XDrive.bEnablePositionDrive;
+		ConstraintSettings.Motion.LinearDriveDamping = 0.0f;// ConstraintInstance->ProfileInstance.LinearDrive.XDrive.bEnablePositionDrive;
+
 		ConstraintSettings.Motion.AngularDriveTarget = FQuat(ConstraintInstance->ProfileInstance.AngularDrive.OrientationTarget);
 		// NOTE: Hard dependence on EJointAngularConstraintIndex - the following will break if we change the order
 		ConstraintSettings.Motion.AngularDriveTargetAngles[(int32)EJointAngularConstraintIndex::Twist] = FMath::DegreesToRadians(ConstraintInstance->ProfileInstance.AngularDrive.OrientationTarget.Roll);
diff --git a/Engine/Source/Runtime/Engine/Private/PhysicsEngine/ImmediatePhysics/ImmediatePhysicsChaos/ImmediatePhysicsSimulation_Chaos.cpp b/Engine/Source/Runtime/Engine/Private/PhysicsEngine/ImmediatePhysics/ImmediatePhysicsChaos/ImmediatePhysicsSimulation_Chaos.cpp
--- a/Engine/Source/Runtime/Engine/Private/PhysicsEngine/ImmediatePhysics/ImmediatePhysicsChaos/ImmediatePhysicsSimulation_Chaos.cpp
+++ b/Engine/Source/Runtime/Engine/Private/PhysicsEngine/ImmediatePhysics/ImmediatePhysicsChaos/ImmediatePhysicsSimulation_Chaos.cpp
@@ -59,8 +59,10 @@ float ChaosImmediate_Joint_SoftTwistStiffness = 0.0f;
 float ChaosImmediate_Joint_SoftTwistDamping = 0.0f;
 float ChaosImmediate_Joint_SoftSwingStiffness = 0.0f;
 float ChaosImmediate_Joint_SoftSwingDamping = 0.0f;
-float ChaosImmediate_Joint_DriveStiffness = 0.0f;
-float ChaosImmediate_Joint_DriveDamping = 0.0f;
+float ChaosImmediate_Joint_LinearDriveStiffness = 0.0f;
+float ChaosImmediate_Joint_LinearDriveDamping = 0.0f;
+float ChaosImmediate_Joint_AngularDriveStiffness = 0.0f;
+float ChaosImmediate_Joint_AngularDriveDamping = 0.0f;
 float ChaosImmediate_Joint_MinParentMassRatio = 0.5f;
 float ChaosImmediate_Joint_MaxInertiaRatio = 5.0f;
 FAutoConsoleVariableRef CVarChaosImmPhysPairIterations(TEXT("p.Chaos.ImmPhys.Joint.PairIterations"), ChaosImmediate_Joint_PairIterations, TEXT("PairIterations."));
@@ -77,8 +79,10 @@ FAutoConsoleVariableRef CVarChaosImmPhysSoftTwistStiffness(TEXT("p.Chaos.ImmPhys
 FAutoConsoleVariableRef CVarChaosImmPhysSoftTwistDamping(TEXT("p.Chaos.ImmPhys.Joint.SoftTwistDamping"), ChaosImmediate_Joint_SoftTwistDamping, TEXT("6Dof joint SoftTwist damping override (if > 0)."));
 FAutoConsoleVariableRef CVarChaosImmPhysSoftSwingStiffness(TEXT("p.Chaos.ImmPhys.Joint.SoftSwingStiffness"), ChaosImmediate_Joint_SoftSwingStiffness, TEXT("6Dof joint SoftSwing stiffness override (if > 0)."));
 FAutoConsoleVariableRef CVarChaosImmPhysSoftSwingDamping(TEXT("p.Chaos.ImmPhys.Joint.SoftSwingDamping"), ChaosImmediate_Joint_SoftSwingDamping, TEXT("6Dof joint SoftSwing damping override (if > 0)."));
-FAutoConsoleVariableRef CVarChaosImmPhysDriveStiffness(TEXT("p.Chaos.ImmPhys.Joint.DriveStiffness"), ChaosImmediate_Joint_DriveStiffness, TEXT("6Dof joint drive stiffness override (if > 0)."));
-FAutoConsoleVariableRef CVarChaosImmPhysDriveDamping(TEXT("p.Chaos.ImmPhys.Joint.DriveDamping"), ChaosImmediate_Joint_DriveDamping, TEXT("6Dof joint drive damping override (if > 0)."));
+FAutoConsoleVariableRef CVarChaosImmPhysLinearDriveStiffness(TEXT("p.Chaos.ImmPhys.Joint.LinearDriveStiffness"), ChaosImmediate_Joint_LinearDriveStiffness, TEXT("6Dof joint drive stiffness override (if > 0)."));
+FAutoConsoleVariableRef CVarChaosImmPhysLinearDriveDamping(TEXT("p.Chaos.ImmPhys.Joint.LinearDriveDamping"), ChaosImmediate_Joint_LinearDriveDamping, TEXT("6Dof joint drive damping override (if > 0)."));
+FAutoConsoleVariableRef CVarChaosImmPhysAngularDriveStiffness(TEXT("p.Chaos.ImmPhys.Joint.AngularDriveStiffness"), ChaosImmediate_Joint_AngularDriveStiffness, TEXT("6Dof joint drive stiffness override (if > 0)."));
+FAutoConsoleVariableRef CVarChaosImmPhysAngularDriveDamping(TEXT("p.Chaos.ImmPhys.Joint.AngularDriveDamping"), ChaosImmediate_Joint_AngularDriveDamping, TEXT("6Dof joint drive damping override (if > 0)."));
 FAutoConsoleVariableRef CVarChaosImmPhysMinParentMassRatio(TEXT("p.Chaos.ImmPhys.Joint.MinParentMassRatio"), ChaosImmediate_Joint_MinParentMassRatio, TEXT("6Dof joint MinParentMassRatio (if > 0)"));
 FAutoConsoleVariableRef CVarChaosImmPhysMaxInertiaRatio(TEXT("p.Chaos.ImmPhys.Joint.MaxInertiaRatio"), ChaosImmediate_Joint_MaxInertiaRatio, TEXT("6Dof joint MaxInertiaRatio (if > 0)"));
 
@@ -533,8 +537,10 @@ namespace ImmediatePhysics_Chaos
 			JointsSettings.SoftTwistDamping = ChaosImmediate_Joint_SoftTwistDamping;
 			JointsSettings.SoftSwingStiffness = ChaosImmediate_Joint_SoftSwingStiffness;
 			JointsSettings.SoftSwingDamping = ChaosImmediate_Joint_SoftSwingDamping;
-			JointsSettings.DriveStiffness = ChaosImmediate_Joint_DriveStiffness;
-			JointsSettings.DriveDamping = ChaosImmediate_Joint_DriveDamping;
+			JointsSettings.LinearDriveStiffness = ChaosImmediate_Joint_LinearDriveStiffness;
+			JointsSettings.LinearDriveDamping = ChaosImmediate_Joint_LinearDriveDamping;
+			JointsSettings.AngularDriveStiffness = ChaosImmediate_Joint_AngularDriveStiffness;
+			JointsSettings.AngularDriveDamping = ChaosImmediate_Joint_AngularDriveDamping;
 
 			Joints.SetSettings(JointsSettings);
 
diff --git a/Engine/Source/Runtime/Experimental/Chaos/Private/Chaos/Joint/PBDJointSolverGaussSeidel.cpp b/Engine/Source/Runtime/Experimental/Chaos/Private/Chaos/Joint/PBDJointSolverGaussSeidel.cpp
--- a/Engine/Source/Runtime/Experimental/Chaos/Private/Chaos/Joint/PBDJointSolverGaussSeidel.cpp
+++ b/Engine/Source/Runtime/Experimental/Chaos/Private/Chaos/Joint/PBDJointSolverGaussSeidel.cpp
@@ -266,8 +266,9 @@ namespace Chaos
 		const FPBDJointSolverSettings& SolverSettings,
 		const FPBDJointSettings& JointSettings)
 	{
-		// @todo(ccaulfield): linear and velocity drives
 		ApplyRotationDrives(Dt, SolverSettings, JointSettings);
+
+		ApplyPositionDrives(Dt, SolverSettings, JointSettings);
 	}
 
 
@@ -1070,6 +1071,59 @@ namespace Chaos
 	}
 
 
+	void FJointSolverGaussSeidel::ApplyPositionDrives(
+		const FReal Dt,
+		const FPBDJointSolverSettings& SolverSettings,
+		const FPBDJointSettings& JointSettings)
+	{
+		// @todo(ccaulfield): XPBD for linear drives
+		if (JointSettings.Motion.bLinearDriveEnabled[0] || JointSettings.Motion.bLinearDriveEnabled[1] || JointSettings.Motion.bLinearDriveEnabled[2])
+		{
+			const FVec3 TargetX = Xs[0] + Rs[0] * JointSettings.Motion.LinearDriveTarget;
+			FVec3 CX = Xs[1] - TargetX;
+
+			const FMatrix33 RM0 = Rs[0].ToMatrix();
+			for (int32 IAxis = 0; IAxis < 3; ++IAxis)
+			{
+				if (!JointSettings.Motion.bLinearDriveEnabled[IAxis])
+				{
+					const FVec3 Axis = RM0.GetAxis(IAxis);
+					CX = CX - FVec3::DotProduct(Axis, CX) * Axis;
+				}
+			}
+
+			const FMatrix33 InvI0 = Utilities::ComputeWorldSpaceInertia(Qs[0], InvILs[0]);
+			const FMatrix33 InvI1 = Utilities::ComputeWorldSpaceInertia(Qs[1], InvILs[1]);
+
+			// Calculate constraint correction
+			FMatrix33 M0 = FMatrix33(0, 0, 0);
+			FMatrix33 M1 = FMatrix33(0, 0, 0);
+			if (InvMs[0] > 0)
+			{
+				M0 = Utilities::ComputeJointFactorMatrix(Xs[0] - Ps[0], InvI0, InvMs[0]);
+			}
+			if (InvMs[1] > 0)
+			{
+				M1 = Utilities::ComputeJointFactorMatrix(Xs[1] - Ps[1], InvI1, InvMs[1]);
+			}
+			const FMatrix33 MI = (M0 + M1).Inverse();
+			const FVec3 DX = Utilities::Multiply(MI, CX);
+
+			// Apply constraint correction
+			const FVec3 DP0 = InvMs[0] * DX;
+			const FVec3 DP1 = -InvMs[1] * DX;
+			const FVec3 DR0 = Utilities::Multiply(InvI0, FVec3::CrossProduct(Xs[0] - Ps[0], DX));
+			const FVec3 DR1 = Utilities::Multiply(InvI1, FVec3::CrossProduct(Xs[1] - Ps[1], -DX));
+
+			const FReal LinearStiffness = FPBDJointUtilities::GetLinearDriveStiffness(SolverSettings, JointSettings);
+			const FReal LinearDamping = FPBDJointUtilities::GetLinearDriveDamping(SolverSettings, JointSettings);
+			ApplyPositionDelta(LinearStiffness, DP0, DP1);
+			ApplyRotationDelta(LinearStiffness, DR0, DR1);
+			UpdateDerivedState();
+		}
+	}
+
+
 	void FJointSolverGaussSeidel::ApplyPositionProjection(
 		const FReal Dt,
 		const FPBDJointSolverSettings& SolverSettings,
diff --git a/Engine/Source/Runtime/Experimental/Chaos/Private/Chaos/PBDJointConstraintUtilities.cpp b/Engine/Source/Runtime/Experimental/Chaos/Private/Chaos/PBDJointConstraintUtilities.cpp
--- a/Engine/Source/Runtime/Experimental/Chaos/Private/Chaos/PBDJointConstraintUtilities.cpp
+++ b/Engine/Source/Runtime/Experimental/Chaos/Private/Chaos/PBDJointConstraintUtilities.cpp
@@ -69,19 +69,32 @@ namespace Chaos
 		return (SolverSettings.SoftSwingDamping > 0) ? SolverSettings.SoftSwingDamping : JointSettings.Motion.SoftSwingDamping;
 	}
 
-	FReal FPBDJointUtilities::GetAngularDriveStiffness(
+	FReal FPBDJointUtilities::GetLinearDriveStiffness(
+		const FPBDJointSolverSettings& SolverSettings,
+		const FPBDJointSettings& JointSettings)
+	{
+		return (SolverSettings.LinearDriveStiffness > 0.0f) ? SolverSettings.LinearDriveStiffness : JointSettings.Motion.LinearDriveStiffness;
+	}
+
+	FReal FPBDJointUtilities::GetLinearDriveDamping(
 		const FPBDJointSolverSettings& SolverSettings,
 		const FPBDJointSettings& JointSettings)
 	{
-		return (SolverSettings.DriveStiffness > 0.0f) ? SolverSettings.DriveStiffness : JointSettings.Motion.AngularDriveStiffness;
+		return (SolverSettings.LinearDriveDamping > 0.0f) ? SolverSettings.LinearDriveDamping : JointSettings.Motion.LinearDriveDamping;
 	}
 
+	FReal FPBDJointUtilities::GetAngularDriveStiffness(
+		const FPBDJointSolverSettings& SolverSettings,
+		const FPBDJointSettings& JointSettings)
+	{
+		return (SolverSettings.AngularDriveStiffness > 0.0f) ? SolverSettings.AngularDriveStiffness : JointSettings.Motion.AngularDriveStiffness;
+	}
 
 	FReal FPBDJointUtilities::GetAngularDriveDamping(
 		const FPBDJointSolverSettings& SolverSettings,
 		const FPBDJointSettings& JointSettings)
 	{
-		return (SolverSettings.DriveDamping > 0.0f) ? SolverSettings.DriveDamping: JointSettings.Motion.AngularDriveDamping;
+		return (SolverSettings.AngularDriveDamping > 0.0f) ? SolverSettings.AngularDriveDamping : JointSettings.Motion.AngularDriveDamping;
 	}
 
 	FReal FPBDJointUtilities::GetLinearProjection(
@@ -1407,8 +1420,7 @@ namespace Chaos
 		const FVec3 TwistAxis0 = R0 * TwistAxis01;
 		const FVec3 TwistAxis1 = R1 * TwistAxis01;
 		const FReal TwistAngleTarget = JointSettings.Motion.AngularDriveTargetAngles[(int32)EJointAngularConstraintIndex::Twist];
-		const FReal DriveStiffnessUnclamped = (SolverSettings.DriveStiffness > 0) ? SolverSettings.DriveStiffness : JointSettings.Motion.AngularDriveStiffness;
-		const FReal DriveStiffness = FMath::Clamp(DriveStiffnessUnclamped, (FReal)0, (FReal)1);
+		const FReal DriveStiffness = GetAngularDriveStiffness(SolverSettings, JointSettings);
 		const FReal DTwistAngle = TwistAngle - TwistAngleTarget;
 		const FReal DTwistAngle0 = DTwistAngle;
 		const FReal DTwistAngle1 = -DTwistAngle;
@@ -1464,8 +1476,7 @@ namespace Chaos
 		FReal Swing2Target = JointSettings.Motion.AngularDriveTargetAngles[(int32)EJointAngularConstraintIndex::Swing2];
 		FReal SwingAngleTarget = FMath::Max(Swing1Target, Swing2Target);
 
-		FReal DriveStiffnessUnclamped = (SolverSettings.DriveStiffness > 0) ? SolverSettings.DriveStiffness : JointSettings.Motion.AngularDriveStiffness;
-		FReal DriveStiffness = FMath::Clamp(DriveStiffnessUnclamped, (FReal)0, (FReal)1);
+		const FReal DriveStiffness = GetAngularDriveStiffness(SolverSettings, JointSettings);
 		FReal DSwingAngle = SwingAngle - SwingAngleTarget;
 		FReal DSwingAngle0 = DSwingAngle;
 		FReal DSwingAngle1 = -DSwingAngle;
@@ -1507,8 +1518,7 @@ namespace Chaos
 		const FRotation3 TargetQ0 = DR1.Inverse() * Q0;
 		const FRotation3 TargetQ1 = DR1 * Q1;
 
-		FReal DriveStiffnessUnclamped = (SolverSettings.DriveStiffness > 0) ? SolverSettings.DriveStiffness : JointSettings.Motion.AngularDriveStiffness;
-		FReal DriveStiffness = FMath::Clamp(DriveStiffnessUnclamped, (FReal)0, (FReal)1);
+		const FReal DriveStiffness = GetAngularDriveStiffness(SolverSettings, JointSettings);
 
 		FVec3 SLerpAxis;
 		FReal SLerpAngle;
diff --git a/Engine/Source/Runtime/Experimental/Chaos/Private/Chaos/PBDJointConstraints.cpp b/Engine/Source/Runtime/Experimental/Chaos/Private/Chaos/PBDJointConstraints.cpp
--- a/Engine/Source/Runtime/Experimental/Chaos/Private/Chaos/PBDJointConstraints.cpp
+++ b/Engine/Source/Runtime/Experimental/Chaos/Private/Chaos/PBDJointConstraints.cpp
@@ -78,8 +78,8 @@ namespace Chaos
 	
 	FPBDJointMotionSettings::FPBDJointMotionSettings()
 		: Stiffness((FReal)1)
-		, LinearProjection((FReal)0)
-		, AngularProjection((FReal)0)
+		, LinearProjection(0)
+		, AngularProjection(0)
 		, LinearMotionTypes({ EJointMotionType::Locked, EJointMotionType::Locked, EJointMotionType::Locked })
 		, LinearLimit(FLT_MAX)
 		, AngularMotionTypes({ EJointMotionType::Free, EJointMotionType::Free, EJointMotionType::Free })
@@ -93,6 +93,10 @@ namespace Chaos
 		, SoftTwistDamping(0)
 		, SoftSwingStiffness(0)
 		, SoftSwingDamping(0)
+		, LinearDriveTarget(FVec3(0, 0, 0))
+		, bLinearDriveEnabled({ false, false, false })
+		, LinearDriveStiffness(0)
+		, LinearDriveDamping(0)
 		, AngularDriveTarget(FRotation3::FromIdentity())
 		, AngularDriveTargetAngles(FVec3(0, 0, 0))
 		, bAngularSLerpDriveEnabled(false)
@@ -107,8 +111,8 @@ namespace Chaos
 	
 	FPBDJointMotionSettings::FPBDJointMotionSettings(const TVector<EJointMotionType, 3>& InLinearMotionTypes, const TVector<EJointMotionType, 3>& InAngularMotionTypes)
 		: Stiffness((FReal)1)
-		, LinearProjection((FReal)0)
-		, AngularProjection((FReal)0)
+		, LinearProjection(0)
+		, AngularProjection(0)
 		, LinearMotionTypes(InLinearMotionTypes)
 		, LinearLimit(FLT_MAX)
 		, AngularMotionTypes({ EJointMotionType::Free, EJointMotionType::Free, EJointMotionType::Free })
@@ -122,6 +126,10 @@ namespace Chaos
 		, SoftTwistDamping(0)
 		, SoftSwingStiffness(0)
 		, SoftSwingDamping(0)
+		, LinearDriveTarget(FVec3(0))
+		, bLinearDriveEnabled({ false, false, false })
+		, LinearDriveStiffness(0)
+		, LinearDriveDamping(0)
 		, AngularDriveTarget(FRotation3::FromIdentity())
 		, AngularDriveTargetAngles(FVec3(0, 0, 0))
 		, bAngularSLerpDriveEnabled(false)
@@ -182,13 +190,15 @@ namespace Chaos
 		, bEnableTwistLimits(true)
 		, bEnableSwingLimits(true)
 		, bEnableDrives(true)
-		, LinearProjection((FReal)0)
-		, AngularProjection((FReal)0)
-		, Stiffness((FReal)0)
-		, DriveStiffness((FReal)0)
-		, DriveDamping((FReal)0)
-		, SoftLinearStiffness((FReal)0)
-		, SoftLinearDamping((FReal)0)
+		, LinearProjection(0)
+		, AngularProjection(0)
+		, Stiffness(0)
+		, LinearDriveStiffness(0)
+		, LinearDriveDamping(0)
+		, AngularDriveStiffness(0)
+		, AngularDriveDamping(0)
+		, SoftLinearStiffness(0)
+		, SoftLinearDamping(0)
 		, SoftTwistStiffness(0)
 		, SoftTwistDamping(0)
 		, SoftSwingStiffness(0)
diff --git a/Engine/Source/Runtime/Experimental/Chaos/Public/Chaos/Joint/PBDJointSolverGaussSeidel.h b/Engine/Source/Runtime/Experimental/Chaos/Public/Chaos/Joint/PBDJointSolverGaussSeidel.h
--- a/Engine/Source/Runtime/Experimental/Chaos/Public/Chaos/Joint/PBDJointSolverGaussSeidel.h
+++ b/Engine/Source/Runtime/Experimental/Chaos/Public/Chaos/Joint/PBDJointSolverGaussSeidel.h
@@ -104,6 +104,11 @@ namespace Chaos
 			const FPBDJointSolverSettings& SolverSettings,
 			const FPBDJointSettings& JointSettings);
 
+		void ApplyPositionDrives(
+			const FReal Dt,
+			const FPBDJointSolverSettings& SolverSettings,
+			const FPBDJointSettings& JointSettings);
+
 		void ApplyRotationDrives(
 			const FReal Dt,
 			const FPBDJointSolverSettings& SolverSettings,
diff --git a/Engine/Source/Runtime/Experimental/Chaos/Public/Chaos/PBDJointConstraintTypes.h b/Engine/Source/Runtime/Experimental/Chaos/Public/Chaos/PBDJointConstraintTypes.h
--- a/Engine/Source/Runtime/Experimental/Chaos/Public/Chaos/PBDJointConstraintTypes.h
+++ b/Engine/Source/Runtime/Experimental/Chaos/Public/Chaos/PBDJointConstraintTypes.h
@@ -91,6 +91,11 @@ namespace Chaos
 		FReal SoftSwingStiffness;
 		FReal SoftSwingDamping;
 
+		FVec3 LinearDriveTarget;
+		TVector<bool, 3> bLinearDriveEnabled;
+		FReal LinearDriveStiffness;
+		FReal LinearDriveDamping;
+
 		// @todo(ccaulfield): remove one of these
 		FRotation3 AngularDriveTarget;
 		FVec3 AngularDriveTargetAngles;
@@ -141,8 +146,10 @@ namespace Chaos
 		FReal LinearProjection;
 		FReal AngularProjection;
 		FReal Stiffness;
-		FReal DriveStiffness;
-		FReal DriveDamping;
+		FReal LinearDriveStiffness;
+		FReal LinearDriveDamping;
+		FReal AngularDriveStiffness;
+		FReal AngularDriveDamping;
 		FReal SoftLinearStiffness;
 		FReal SoftLinearDamping;
 		FReal SoftTwistStiffness;
diff --git a/Engine/Source/Runtime/Experimental/Chaos/Public/Chaos/PBDJointConstraintUtilities.h b/Engine/Source/Runtime/Experimental/Chaos/Public/Chaos/PBDJointConstraintUtilities.h
--- a/Engine/Source/Runtime/Experimental/Chaos/Public/Chaos/PBDJointConstraintUtilities.h
+++ b/Engine/Source/Runtime/Experimental/Chaos/Public/Chaos/PBDJointConstraintUtilities.h
@@ -92,6 +92,14 @@ namespace Chaos
 			const FPBDJointSolverSettings& SolverSettings,
 			const FPBDJointSettings& JointSettings);
 
+		static FReal GetLinearDriveStiffness(
+			const FPBDJointSolverSettings& SolverSettings,
+			const FPBDJointSettings& JointSettings);
+
+		static FReal GetLinearDriveDamping(
+			const FPBDJointSolverSettings& SolverSettings,
+			const FPBDJointSettings& JointSettings);
+
 		static FReal GetAngularDriveStiffness(
 			const FPBDJointSolverSettings& SolverSettings,
 			const FPBDJointSettings& JointSettings);
