Chaos : Collisions

- Added support for kinematic implicit scaled<triangle mesh> mass properties.

diff --git a/Engine/Source/Runtime/Engine/Private/PhysicsEngine/Experimental/PhysInterface_Chaos.cpp b/Engine/Source/Runtime/Engine/Private/PhysicsEngine/Experimental/PhysInterface_Chaos.cpp
--- a/Engine/Source/Runtime/Engine/Private/PhysicsEngine/Experimental/PhysInterface_Chaos.cpp
+++ b/Engine/Source/Runtime/Engine/Private/PhysicsEngine/Experimental/PhysInterface_Chaos.cpp
@@ -1522,9 +1522,27 @@ bool CalculateMassPropertiesOfImplicitType(
 
 	if (ImplicitObject)
 	{
-		// Hack to handle Transformed until CastHelper can properly support transformed
+		// Hack to handle Transformed and Scaled<ImplicitObjectTriangleMesh> until CastHelper can properly support transformed
 		Chaos::FRigidTransform3 Transform(FMatrix::Identity);
-		if (ImplicitObject->GetType(true) & Chaos::ImplicitObjectType::Transformed)
+		if (Chaos::IsScaled(ImplicitObject->GetType(true)) && Chaos::GetInnerType(ImplicitObject->GetType(true)) & Chaos::ImplicitObjectType::TriangleMesh)
+		{
+			OutMassProperties.Volume = 0.f;
+			OutMassProperties.Mass = FLT_MAX;
+			OutMassProperties.InertiaTensor = FMatrix33(0,0,0);
+			OutMassProperties.CenterOfMass = FVector(0);
+			OutMassProperties.RotationOfMass = Chaos::TRotation<float, 3>::FromIdentity();
+			return false;
+		}
+		else if (ImplicitObject->GetType(true) & Chaos::ImplicitObjectType::TriangleMesh)
+		{
+			OutMassProperties.Volume = 0.f;
+			OutMassProperties.Mass = FLT_MAX;
+			OutMassProperties.InertiaTensor = FMatrix33(0, 0, 0);
+			OutMassProperties.CenterOfMass = FVector(0);
+			OutMassProperties.RotationOfMass = Chaos::TRotation<float, 3>::FromIdentity();
+			return false;
+		}
+		else if (ImplicitObject->GetType(true) & Chaos::ImplicitObjectType::Transformed)
 		{
 			// TODO:  This all very wrong, but is wrong in the same way as scaled. Rotation/Translation are ignored though. The three methods on Transformed no implemented.
 			// Only adding this to hack around CastHelper, as TransformedImplicit is very not supported in that path.
