Reverting kinematic scaled implicit mass properties as it breaks vehicles.

diff --git a/Engine/Source/Runtime/Engine/Private/PhysicsEngine/Experimental/PhysInterface_Chaos.cpp b/Engine/Source/Runtime/Engine/Private/PhysicsEngine/Experimental/PhysInterface_Chaos.cpp
--- a/Engine/Source/Runtime/Engine/Private/PhysicsEngine/Experimental/PhysInterface_Chaos.cpp
+++ b/Engine/Source/Runtime/Engine/Private/PhysicsEngine/Experimental/PhysInterface_Chaos.cpp
@@ -1524,25 +1524,27 @@ bool CalculateMassPropertiesOfImplicitType(
 	{
 		// Hack to handle Transformed and Scaled<ImplicitObjectTriangleMesh> until CastHelper can properly support transformed
 		Chaos::FRigidTransform3 Transform(FMatrix::Identity);
-		if (Chaos::IsScaled(ImplicitObject->GetType(true)) && Chaos::GetInnerType(ImplicitObject->GetType(true)) & Chaos::ImplicitObjectType::TriangleMesh)
-		{
-			OutMassProperties.Volume = 0.f;
-			OutMassProperties.Mass = FLT_MAX;
-			OutMassProperties.InertiaTensor = FMatrix33(0,0,0);
-			OutMassProperties.CenterOfMass = FVector(0);
-			OutMassProperties.RotationOfMass = Chaos::TRotation<float, 3>::FromIdentity();
-			return false;
-		}
-		else if (ImplicitObject->GetType(true) & Chaos::ImplicitObjectType::TriangleMesh)
-		{
-			OutMassProperties.Volume = 0.f;
-			OutMassProperties.Mass = FLT_MAX;
-			OutMassProperties.InertiaTensor = FMatrix33(0, 0, 0);
-			OutMassProperties.CenterOfMass = FVector(0);
-			OutMassProperties.RotationOfMass = Chaos::TRotation<float, 3>::FromIdentity();
-			return false;
-		}
-		else if (ImplicitObject->GetType(true) & Chaos::ImplicitObjectType::Transformed)
+		// Commenting this out temporarily as it breaks vehicles
+		/*	if (Chaos::IsScaled(ImplicitObject->GetType(true)) && Chaos::GetInnerType(ImplicitObject->GetType(true)) & Chaos::ImplicitObjectType::TriangleMesh)
+			{
+				OutMassProperties.Volume = 0.f;
+				OutMassProperties.Mass = FLT_MAX;
+				OutMassProperties.InertiaTensor = FMatrix33(0, 0, 0);
+				OutMassProperties.CenterOfMass = FVector(0);
+				OutMassProperties.RotationOfMass = Chaos::TRotation<float, 3>::FromIdentity();
+				return false;
+			}
+			else if (ImplicitObject->GetType(true) & Chaos::ImplicitObjectType::TriangleMesh)
+			{
+				OutMassProperties.Volume = 0.f;
+				OutMassProperties.Mass = FLT_MAX;
+				OutMassProperties.InertiaTensor = FMatrix33(0, 0, 0);
+				OutMassProperties.CenterOfMass = FVector(0);
+				OutMassProperties.RotationOfMass = Chaos::TRotation<float, 3>::FromIdentity();
+				return false;
+			}
+		else*/ 
+		if (ImplicitObject->GetType(true) & Chaos::ImplicitObjectType::Transformed)
 		{
 			// TODO:  This all very wrong, but is wrong in the same way as scaled. Rotation/Translation are ignored though. The three methods on Transformed no implemented.
 			// Only adding this to hack around CastHelper, as TransformedImplicit is very not supported in that path.
