Make the foliage.MinLOD cvar usable by scalability configs.
It wasn't flagged as a scalability cvar and thus would fire ensures constantly on mobile platforms.

diff --git a/Engine/Source/Runtime/Engine/Private/InstancedStaticMesh.cpp b/Engine/Source/Runtime/Engine/Private/InstancedStaticMesh.cpp
--- a/Engine/Source/Runtime/Engine/Private/InstancedStaticMesh.cpp
+++ b/Engine/Source/Runtime/Engine/Private/InstancedStaticMesh.cpp
@@ -51,7 +51,8 @@ IMPLEMENT_HIT_PROXY(HInstancedStaticMeshInstance, HHitProxy);
 TAutoConsoleVariable<int32> CVarMinLOD(
 	TEXT("foliage.MinLOD"),
 	-1,
-	TEXT("Used to discard the top LODs for performance evaluation. -1: Disable all effects of this cvar."));
+	TEXT("Used to discard the top LODs for performance evaluation. -1: Disable all effects of this cvar."),
+	ECVF_Scalability | ECVF_Default);
 
 static TAutoConsoleVariable<int32> CVarRayTracingRenderInstances(
 	TEXT("r.RayTracing.InstancedStaticMeshes"),
