Serialized Custom primitive data is no longer used as the instance data directly. It's simply used as the default data and an internal instance data is used for the runtime instead. This should take care of the undeterministic serialization.

diff --git a/Engine/Source/Runtime/Engine/Classes/Components/PrimitiveComponent.h b/Engine/Source/Runtime/Engine/Classes/Components/PrimitiveComponent.h
--- a/Engine/Source/Runtime/Engine/Classes/Components/PrimitiveComponent.h
+++ b/Engine/Source/Runtime/Engine/Classes/Components/PrimitiveComponent.h
@@ -553,10 +553,13 @@ public:
 	int32 CustomDepthStencilValue;
 
 private:
-	/** Custom data that can be read by a material through a material parameter expression. Set data using SetCustomPrimitiveData* functions */
-	UPROPERTY(EditAnywhere, AdvancedDisplay, Category=Rendering)
+	/** Optional user defined default values for the custom primitive data of this primitive */
+	UPROPERTY(EditAnywhere, AdvancedDisplay, Category=Rendering, meta = (DisplayName = "Custom Primitive Data Defaults"))
 	FCustomPrimitiveData CustomPrimitiveData;
 
+	/** Custom data that can be read by a material through a material parameter expression. Set data using SetCustomPrimitiveData* functions */
+	UPROPERTY(Transient)
+	FCustomPrimitiveData CustomPrimitiveDataInternal;
 public:
 
 	/**
@@ -802,7 +805,7 @@ public:
 	 * Get the custom primitive data for this primitive component.
 	 * @return The payload of custom data that will be set on the primitive and accessible in the material through a material expression.
 	 */
-	const FCustomPrimitiveData& GetCustomPrimitiveData() const { return CustomPrimitiveData; }
+	const FCustomPrimitiveData& GetCustomPrimitiveData() const { return CustomPrimitiveDataInternal; }
 
 #if WITH_EDITOR
 	/** Override delegate used for checking the selection state of a component */
@@ -812,6 +815,9 @@ public:
 
 protected:
 
+	/** Reset the custom primitive data of this primitive to the optional user defined default */
+	void ResetCustomPrimitiveData();
+
 	/** Insert an array of floats into the CustomPrimitiveData, starting at the given index */
 	void SetCustomPrimitiveDataInternal(int32 DataIndex, const TArray<float>& Values);
 
diff --git a/Engine/Source/Runtime/Engine/Private/Components/PrimitiveComponent.cpp b/Engine/Source/Runtime/Engine/Private/Components/PrimitiveComponent.cpp
--- a/Engine/Source/Runtime/Engine/Private/Components/PrimitiveComponent.cpp
+++ b/Engine/Source/Runtime/Engine/Private/Components/PrimitiveComponent.cpp
@@ -966,6 +966,7 @@ void UPrimitiveComponent::PostEditChangeProperty(FPropertyChangedEvent& Property
 		// Reregister to get the custom primitive data to the proxy
 		if (MemberPropertyName == GET_MEMBER_NAME_CHECKED(UPrimitiveComponent, CustomPrimitiveData))
 		{
+			ResetCustomPrimitiveData();
 			MarkRenderStateDirty();
 		}
 	}
@@ -1198,6 +1199,9 @@ void UPrimitiveComponent::PostLoad()
 	{
 		LightmapType = ELightmapType::Default;
 	}
+
+	// Setup the default here
+	ResetCustomPrimitiveData();
 }
 
 void UPrimitiveComponent::PostDuplicate(bool bDuplicateForPIE)
@@ -1663,6 +1667,11 @@ UMaterialInstanceDynamic* UPrimitiveComponent::CreateDynamicMaterialInstance(int
 	return MID;
 }
 
+void UPrimitiveComponent::ResetCustomPrimitiveData()
+{
+	CustomPrimitiveDataInternal.Data = CustomPrimitiveData.Data;
+}
+
 void UPrimitiveComponent::SetCustomPrimitiveDataInternal(int32 DataIndex, const TArray<float>& Values)
 {
 	// Can only set data on valid indices and only if there's actually any data to set
@@ -1675,15 +1684,15 @@ void UPrimitiveComponent::SetCustomPrimitiveDataInternal(int32 DataIndex, const
 		const int32 NumValuesToSet = FMath::Min(Values.Num(), FCustomPrimitiveData::NumCustomPrimitiveDataFloats - DataIndex);
 
 		// If trying to set data on an index which doesn't exist yet, allocate up to it
-		if (NeededFloats > CustomPrimitiveData.Data.Num())
+		if (NeededFloats > CustomPrimitiveDataInternal.Data.Num())
 		{
-			CustomPrimitiveData.Data.SetNumZeroed(NeededFloats);
+			CustomPrimitiveDataInternal.Data.SetNumZeroed(NeededFloats);
 		}
 
 		// Only update data if it has changed
-		if (FMemory::Memcmp(&CustomPrimitiveData.Data[DataIndex], Values.GetData(), NumValuesToSet * sizeof(float)) != 0)
+		if (FMemory::Memcmp(&CustomPrimitiveDataInternal.Data[DataIndex], Values.GetData(), NumValuesToSet * sizeof(float)) != 0)
 		{
-			FMemory::Memcpy(&CustomPrimitiveData.Data[DataIndex], Values.GetData(), NumValuesToSet * sizeof(float));
+			FMemory::Memcpy(&CustomPrimitiveDataInternal.Data[DataIndex], Values.GetData(), NumValuesToSet * sizeof(float));
 
 			UWorld* World = GetWorld();
 			if (World && World->Scene)
