Chaos
- low-iteration angular stiffness improvements
- fix soft limit stiffness relative to drive stiffness
- support "parent dominates"

diff --git a/Engine/Source/Runtime/Engine/Private/PhysicsEngine/ImmediatePhysics/ImmediatePhysicsChaos/ImmediatePhysicsActorHandle_Chaos.cpp b/Engine/Source/Runtime/Engine/Private/PhysicsEngine/ImmediatePhysics/ImmediatePhysicsChaos/ImmediatePhysicsActorHandle_Chaos.cpp
--- a/Engine/Source/Runtime/Engine/Private/PhysicsEngine/ImmediatePhysics/ImmediatePhysicsChaos/ImmediatePhysicsActorHandle_Chaos.cpp
+++ b/Engine/Source/Runtime/Engine/Private/PhysicsEngine/ImmediatePhysics/ImmediatePhysicsChaos/ImmediatePhysicsActorHandle_Chaos.cpp
@@ -25,51 +25,35 @@ namespace ImmediatePhysics_Chaos
 	// Utils
 	//
 
-	template<typename T, int d>
-	Chaos::PMatrix<T, d, d> CalculateInertia_Solid(const T Mass, const FKSphereElem& SphereElem)
+	Chaos::FMatrix33 CalculateInertia_Solid(const FReal Mass, const FKSphereElem& SphereElem)
 	{
-		return Chaos::PMatrix<T, d, d>(
-			((T)2 / (T)5) * Mass * SphereElem.Radius * SphereElem.Radius,
-			((T)2 / (T)5) * Mass * SphereElem.Radius * SphereElem.Radius,
-			((T)2 / (T)5) * Mass * SphereElem.Radius * SphereElem.Radius
-			);
+		return Chaos::TSphere<FReal, 3>::GetInertiaTensor(Mass, SphereElem.Radius);
 	}
 
-	template<typename T, int d>
-	Chaos::PMatrix<T, d, d> CalculateInertia_Solid(const T Mass, const FKSphylElem& SphylElem)
+	Chaos::FMatrix33 CalculateInertia_Solid(const FReal Mass, const FKSphylElem& SphylElem)
 	{
-		return Chaos::PMatrix<T, d, d>(
-			((T)1 / (T)12) * Mass * ((T)3 * SphylElem.Radius * SphylElem.Radius + SphylElem.Length * SphylElem.Length),
-			((T)1 / (T)12) * Mass * ((T)3 * SphylElem.Radius * SphylElem.Radius + SphylElem.Length * SphylElem.Length),
-			((T)1 / (T)2) * Mass * SphylElem.Radius * SphylElem.Radius
-			);
+		return Chaos::TCapsule<FReal>::GetInertiaTensor(Mass, SphylElem.Length, SphylElem.Radius);
 	}
 
-	template<typename T, int d>
-	Chaos::PMatrix<T, d, d> CalculateInertia_Solid(const T Mass, const FKBoxElem& BoxElem)
+	Chaos::FMatrix33 CalculateInertia_Solid(const FReal Mass, const FKBoxElem& BoxElem)
 	{
-		return Chaos::PMatrix<T, d, d>(
-			((T)1 / (T)12) * Mass * (BoxElem.Y * BoxElem.Y + BoxElem.Z * BoxElem.Z),
-			((T)1 / (T)12) * Mass * (BoxElem.Z * BoxElem.Z + BoxElem.X * BoxElem.X),
-			((T)1 / (T)12) * Mass * (BoxElem.X * BoxElem.X + BoxElem.Y * BoxElem.Y)
-			);
+		return Chaos::TBox<FReal, 3>::GetInertiaTensor(Mass, FVector(BoxElem.X, BoxElem.Y, BoxElem.Z));
 	}
 
-	template<typename T, int d>
-	void CalculateMassProperties(const FVector& Scale, const FTransform& LocalTransform, const FKAggregateGeom& AggGeom, Chaos::TMassProperties<T, d>& OutMassProperties)
+	void CalculateMassProperties(const FVector& Scale, const FTransform& LocalTransform, const FKAggregateGeom& AggGeom, Chaos::TMassProperties<FReal, 3>& OutMassProperties)
 	{
 		using namespace Chaos;
-		TArray<TMassProperties<T, d>> AllMassProperties;
+		TArray<TMassProperties<FReal, 3>> AllMassProperties;
 
 		for (uint32 i = 0; i < static_cast<uint32>(AggGeom.SphereElems.Num()); ++i)
 		{
 			const FKSphereElem ScaledSphereElem = AggGeom.SphereElems[i].GetFinalScaled(Scale, LocalTransform);
 
-			TMassProperties<T, d> MassProperties;
+			TMassProperties<FReal, 3> MassProperties;
 			MassProperties.CenterOfMass = LocalTransform.GetTranslation() + ScaledSphereElem.Center;
-			MassProperties.RotationOfMass = TRotation<T, d>::FromIdentity();
+			MassProperties.RotationOfMass = FRotation3::FromIdentity();
 			MassProperties.Volume = ScaledSphereElem.GetVolume(FVector::OneVector);
-			MassProperties.InertiaTensor = CalculateInertia_Solid<T, d>(MassProperties.Volume, ScaledSphereElem);
+			MassProperties.InertiaTensor = CalculateInertia_Solid(MassProperties.Volume, ScaledSphereElem);
 
 			AllMassProperties.Add(MassProperties);
 		}
@@ -77,11 +61,11 @@ namespace ImmediatePhysics_Chaos
 		{
 			const auto& BoxElem = AggGeom.BoxElems[i];
 
-			TMassProperties<T, d> MassProperties;
+			TMassProperties<FReal, 3> MassProperties;
 			MassProperties.CenterOfMass = LocalTransform.GetTranslation() + BoxElem.Center;
-			MassProperties.RotationOfMass = LocalTransform.GetRotation() * TRotation<T, d>(FQuat(BoxElem.Rotation));
+			MassProperties.RotationOfMass = LocalTransform.GetRotation() * FRotation3(FQuat(BoxElem.Rotation));
 			MassProperties.Volume = BoxElem.GetVolume(Scale);
-			MassProperties.InertiaTensor = CalculateInertia_Solid<T, d>(MassProperties.Volume, BoxElem);
+			MassProperties.InertiaTensor = CalculateInertia_Solid(MassProperties.Volume, BoxElem);
 
 			AllMassProperties.Add(MassProperties);
 		}
@@ -96,21 +80,21 @@ namespace ImmediatePhysics_Chaos
 				//not a capsule just use a sphere
 				const FKSphereElem ScaledSphereElem = FKSphereElem(Radius);
 
-				TMassProperties<T, d> MassProperties;
+				TMassProperties<FReal, 3> MassProperties;
 				MassProperties.CenterOfMass = LocalTransform.GetTranslation() + ScaledSphereElem.Center;
-				MassProperties.RotationOfMass = TRotation<T, d>::FromIdentity();
+				MassProperties.RotationOfMass = FRotation3::FromIdentity();
 				MassProperties.Volume = ScaledSphereElem.GetVolume(FVector::OneVector);
-				MassProperties.InertiaTensor = CalculateInertia_Solid<T, d>(MassProperties.Volume, ScaledSphereElem);
+				MassProperties.InertiaTensor = CalculateInertia_Solid(MassProperties.Volume, ScaledSphereElem);
 
 				AllMassProperties.Add(MassProperties);
 			}
 			else
 			{
-				TMassProperties<T, d> MassProperties;
+				TMassProperties<FReal, 3> MassProperties;
 				MassProperties.CenterOfMass = LocalTransform.GetTranslation() + ScaledSphylElem.Center;
-				MassProperties.RotationOfMass = LocalTransform.GetRotation() * TRotation<T, d>(FQuat(ScaledSphylElem.Rotation));
+				MassProperties.RotationOfMass = LocalTransform.GetRotation() * FRotation3(FQuat(ScaledSphylElem.Rotation));
 				MassProperties.Volume = ScaledSphylElem.GetVolume(FVector::OneVector);
-				MassProperties.InertiaTensor = CalculateInertia_Solid<T, d>(MassProperties.Volume, ScaledSphylElem);
+				MassProperties.InertiaTensor = CalculateInertia_Solid(MassProperties.Volume, ScaledSphylElem);
 
 				AllMassProperties.Add(MassProperties);
 			}
@@ -158,7 +142,7 @@ namespace ImmediatePhysics_Chaos
 
 #if WITH_CHAOS && !PHYSICS_INTERFACE_PHYSX
 		Chaos::TMassProperties<float, 3> MassProperties;
-		CalculateMassProperties<float, 3>(Scale, FTransform::Identity, BodySetup->AggGeom, MassProperties);
+		CalculateMassProperties(Scale, FTransform::Identity, BodySetup->AggGeom, MassProperties);
 		float Density = 1.e-3f;	// 1g/cm3	@todo(ccaulfield): should come from material
 		if (BodyInstance->bOverrideMass)
 		{
diff --git a/Engine/Source/Runtime/Engine/Private/PhysicsEngine/ImmediatePhysics/ImmediatePhysicsChaos/ImmediatePhysicsJointHandle_Chaos.cpp b/Engine/Source/Runtime/Engine/Private/PhysicsEngine/ImmediatePhysics/ImmediatePhysicsChaos/ImmediatePhysicsJointHandle_Chaos.cpp
--- a/Engine/Source/Runtime/Engine/Private/PhysicsEngine/ImmediatePhysics/ImmediatePhysicsChaos/ImmediatePhysicsJointHandle_Chaos.cpp
+++ b/Engine/Source/Runtime/Engine/Private/PhysicsEngine/ImmediatePhysics/ImmediatePhysicsChaos/ImmediatePhysicsJointHandle_Chaos.cpp
@@ -16,19 +16,18 @@ static_assert((int32)Chaos::EJointMotionType::Locked == (int32)EAngularConstrain
 
 namespace ImmediatePhysics_Chaos
 {
-	const float ChaosImmediate_StiffnessDt = 0.03f;
-
 	float ChaosImmediate_JointStiffness = 1.0f;
-	FAutoConsoleVariableRef CVarJointStiffness(TEXT("p.Chaos.ImmPhys.JointStiffness"), ChaosImmediate_JointStiffness, TEXT("Joint solver stiffness."));
-
-	bool ChaosImmediate_StiffnessUseMass = false;
-	FAutoConsoleVariableRef CVarStiffnessUseMass(TEXT("p.Chaos.ImmPhys.StiffnessUseMass"), ChaosImmediate_StiffnessUseMass, TEXT("Whether to use mass to scale stiffness in the conversion."));
-
-	bool ChaosImmediate_StiffnessUseDistance = true;
-	FAutoConsoleVariableRef CVarStiffnessUseDistance(TEXT("p.Chaos.ImmPhys.StiffnessUseDistance"), ChaosImmediate_StiffnessUseDistance, TEXT("Whether to use joint connector length to scale stiffness in the conversion."));
+	FAutoConsoleVariableRef CVarJointStiffness(TEXT("p.Chaos.ImmPhys.JointStiffness"), ChaosImmediate_JointStiffness, TEXT("Hard-joint solver stiffness."));
 
-	float ChaosImmediate_DriveStiffnessScale = 30.0f;
+	float ChaosImmediate_DriveStiffnessScale = 0.85f;
+	float ChaosImmediate_DriveDampingScale = 1.0f;
 	FAutoConsoleVariableRef CVarDriveStiffnessScale(TEXT("p.Chaos.ImmPhys.DriveStiffnessScale"), ChaosImmediate_DriveStiffnessScale, TEXT("Conversion factor for drive stiffness."));
+	FAutoConsoleVariableRef CVarDriveDampingScale(TEXT("p.Chaos.ImmPhys.DriveDampingScale"), ChaosImmediate_DriveDampingScale, TEXT("Conversion factor for drive damping."));
+
+	float ChaosImmediate_SoftStiffnessScale = 100000.0f;
+	float ChaosImmediate_SoftDampingScale = 100000.0f;
+	FAutoConsoleVariableRef CVarSoftStiffnessScale(TEXT("p.Chaos.ImmPhys.SoftStiffnessScale"), ChaosImmediate_SoftStiffnessScale, TEXT("Conversion factor for soft-joint stiffness."));
+	FAutoConsoleVariableRef CVarSoftDampingScale(TEXT("p.Chaos.ImmPhys.SoftDampingScale"), ChaosImmediate_SoftDampingScale, TEXT("Conversion factor for soft-joint damping."));
 
 	float ChaosImmediate_JointMinLinearProjection = 0.0f;
 	float ChaosImmediate_JointMaxLinearProjection = 0.5f;
@@ -40,58 +39,6 @@ namespace ImmediatePhysics_Chaos
 	FAutoConsoleVariableRef CVarJointMinAngularProjection(TEXT("p.Chaos.ImmPhys.JointMinAngularProjection"), ChaosImmediate_JointMinAngularProjection, TEXT("Joint min projection (for joints with projection disabled)."));
 	FAutoConsoleVariableRef CVarJointMaxAngularProjection(TEXT("p.Chaos.ImmPhys.JointMaxAngularProjection"), ChaosImmediate_JointMaxAngularProjection, TEXT("Joint max projection (for joints with projection enabled)."));
 
-	float ChaosImmediate_DriveStiffnessSourceMin = 100.0f;	// PhysX stiffness per inertia that we translate to a Chaos stiffness of ChaosImmediate_DriveStiffnessTargetMin
-	float ChaosImmediate_DriveStiffnessSourceMax = 2000.0f;	// PhysX stiffness per inertia that we translate to a Chaos stiffness of ChaosImmediate_DriveStiffnessTargetMax
-	float ChaosImmediate_DriveStiffnessTargetMin = 0.6f;
-	float ChaosImmediate_DriveStiffnessTargetMax = 1.0f;
-	FAutoConsoleVariableRef CVarDriveStiffnessSourceMin(TEXT("p.Chaos.ImmPhys.DriveStiffnessSourceMin"), ChaosImmediate_DriveStiffnessSourceMin, TEXT("Conversion factor for drive stiffness."));
-	FAutoConsoleVariableRef CVarDriveStiffnessSourceMax(TEXT("p.Chaos.ImmPhys.DriveStiffnessSourceMax"), ChaosImmediate_DriveStiffnessSourceMax, TEXT("Conversion factor for drive stiffness."));
-
-	float ChaosImmediate_SoftLinearStiffnessSourceMin = 100.0f;
-	float ChaosImmediate_SoftLinearStiffnessSourceMax = 500.0f;
-	float ChaosImmediate_SoftLinearStiffnessTargetMin = 0.5f;
-	float ChaosImmediate_SoftLinearStiffnessTargetMax = 1.0f;
-	FAutoConsoleVariableRef CVarSoftLinearStiffnessSourceMin(TEXT("p.Chaos.ImmPhys.SoftyLinearStiffnessSourceMin"), ChaosImmediate_SoftLinearStiffnessSourceMin, TEXT("Conversion factor for soft linear stiffness."));
-	FAutoConsoleVariableRef CVarSoftLinearStiffnessSourceMax(TEXT("p.Chaos.ImmPhys.SoftyLinearStiffnessSourceMax"), ChaosImmediate_SoftLinearStiffnessSourceMax, TEXT("Conversion factor for soft linear stiffness."));
-
-	float ChaosImmediate_SoftAngularStiffnessSourceMin = 100.0f;
-	float ChaosImmediate_SoftAngularStiffnessSourceMax = 500.0f;
-	float ChaosImmediate_SoftAngularStiffnessTargetMin = 0.5f;
-	float ChaosImmediate_SoftAngularStiffnessTargetMax = 1.0f;
-	FAutoConsoleVariableRef CVarSoftAngularStiffnessSourceMin(TEXT("p.Chaos.ImmPhys.SoftAngularStiffnessSourceMin"), ChaosImmediate_SoftAngularStiffnessSourceMin, TEXT("Conversion factor for soft angular stiffness."));
-	FAutoConsoleVariableRef CVarSoftAngularStiffnessSourceMax(TEXT("p.Chaos.ImmPhys.SoftAngularStiffnessSourceMax"), ChaosImmediate_SoftAngularStiffnessSourceMax, TEXT("Conversion factor for soft angular stiffness."));
-
-	// Convert the UE (PhysX) drive spring stiffness to a joint stiffness [0,1] value for use in the solver.
-	// We linearly map a range of UE stiffness values to a range of Chaos values, optionally scaling by the effective mass.
-	float CalculateDriveAngularStiffnessScale(const Chaos::FVec3& X0, const FReal M0, const Chaos::FVec3& I0, const Chaos::FVec3& X1, const FReal M1, const Chaos::FVec3& I1)
-	{
-		using namespace Chaos;
-
-		float InvStiffnessScale = (FReal)1.0;
-		if (ChaosImmediate_StiffnessUseMass)
-		{
-			// Calculate inertia of the system about the joint connector using parallel axis theorem
-			const FVec3 XI0 = I0 + M0 * (X0 * X0);
-			const FVec3 XI1 = I1 + M1 * (X1 * X1);
-			InvStiffnessScale = FMath::Max(KINDA_SMALL_NUMBER, XI0.Max() + XI1.Max());
-		}
-		if (ChaosImmediate_StiffnessUseDistance)
-		{
-			// Scale stiffness by distance to account for the fact that our rotational drives are applied at
-			// the center of mass, not the connector...although maybe that should change
-			float Distance0 = (M0 > 0)? X0.Size() : 0;
-			float Distance1 = (M1 > 0)? X1.Size() : 0;
-			float Distance = FMath::Max(Distance0, Distance1);
-			if (Distance > 1)
-			{
-				InvStiffnessScale = InvStiffnessScale * Distance;
-			}
-		}
-
-		return ChaosImmediate_DriveStiffnessScale / InvStiffnessScale;
-	}
-
-
 	FJointHandle::FJointHandle(FChaosConstraintContainer* InConstraints, FConstraintInstance* ConstraintInstance, FActorHandle* Actor1, FActorHandle* Actor2)
 		: ActorHandles({ Actor1, Actor2 })
 		, Constraints(InConstraints)
@@ -99,90 +46,76 @@ namespace ImmediatePhysics_Chaos
 		using namespace Chaos;
 
 		FPBDJointSettings ConstraintSettings;
+		TVector<FRigidTransform3, 2> ConstraintFrames;
 
 		// BodyInstance/PhysX has the constraint locations in actor-space, but we need them in Center-of-Mass space
-		FTransform ConstraintFrame1 = FParticleUtilities::ActorLocalToParticleLocal(TGenericParticleHandle<FReal, 3>(Actor1->GetParticle()), ConstraintInstance->GetRefFrame(EConstraintFrame::Frame1));
-		FTransform ConstraintFrame2 = FParticleUtilities::ActorLocalToParticleLocal(TGenericParticleHandle<FReal, 3>(Actor2->GetParticle()), ConstraintInstance->GetRefFrame(EConstraintFrame::Frame2));
-
-		// PhysX to Chaos conversions
-		const float DriveAngularStiffnessScale = CalculateDriveAngularStiffnessScale(
-			ConstraintFrame1.GetTranslation() - Actor1->GetLocalCoMTransform().GetTranslation(),
-			Actor1->GetMass(), 
-			Actor1->GetInertia(), 
-			ConstraintFrame2.GetTranslation() - Actor2->GetLocalCoMTransform().GetTranslation(),
-			Actor2->GetMass(), 
-			Actor2->GetInertia());
-
-		const float DriveAngularDampingScale = DriveAngularStiffnessScale;
-
-		ConstraintSettings.ConstraintFrames = 
-		{ 
-			ConstraintFrame1,
-			ConstraintFrame2,
-		};
-		ConstraintSettings.Motion.Stiffness = ChaosImmediate_JointStiffness;
+		ConstraintFrames[0] = FParticleUtilities::ActorLocalToParticleLocal(TGenericParticleHandle<FReal, 3>(Actor1->GetParticle()), ConstraintInstance->GetRefFrame(EConstraintFrame::Frame1));
+		ConstraintFrames[1] = FParticleUtilities::ActorLocalToParticleLocal(TGenericParticleHandle<FReal, 3>(Actor2->GetParticle()), ConstraintInstance->GetRefFrame(EConstraintFrame::Frame2));
+
+		ConstraintSettings.Stiffness = ChaosImmediate_JointStiffness;
 
-		ConstraintSettings.Motion.LinearMotionTypes = 
+		ConstraintSettings.LinearMotionTypes = 
 		{
 			static_cast<EJointMotionType>(ConstraintInstance->GetLinearXMotion()),
 			static_cast<EJointMotionType>(ConstraintInstance->GetLinearYMotion()),
 			static_cast<EJointMotionType>(ConstraintInstance->GetLinearZMotion()),
 		};
-		ConstraintSettings.Motion.LinearLimit = ConstraintInstance->GetLinearLimit();
-
-		ConstraintSettings.Motion.AngularMotionTypes[(int32)EJointAngularConstraintIndex::Twist] = static_cast<EJointMotionType>(ConstraintInstance->GetAngularTwistMotion());
-		ConstraintSettings.Motion.AngularMotionTypes[(int32)EJointAngularConstraintIndex::Swing1] = static_cast<EJointMotionType>(ConstraintInstance->GetAngularSwing1Motion());
-		ConstraintSettings.Motion.AngularMotionTypes[(int32)EJointAngularConstraintIndex::Swing2] = static_cast<EJointMotionType>(ConstraintInstance->GetAngularSwing2Motion());
-		ConstraintSettings.Motion.AngularLimits[(int32)EJointAngularConstraintIndex::Twist] = FMath::DegreesToRadians(ConstraintInstance->GetAngularTwistLimit());
-		ConstraintSettings.Motion.AngularLimits[(int32)EJointAngularConstraintIndex::Swing1] = FMath::DegreesToRadians(ConstraintInstance->GetAngularSwing1Limit());
-		ConstraintSettings.Motion.AngularLimits[(int32)EJointAngularConstraintIndex::Swing2] = FMath::DegreesToRadians(ConstraintInstance->GetAngularSwing2Limit());
-
-		ConstraintSettings.Motion.LinearProjection = ConstraintInstance->IsProjectionEnabled() ? ChaosImmediate_JointMaxLinearProjection : ChaosImmediate_JointMinLinearProjection;
-		ConstraintSettings.Motion.AngularProjection = ConstraintInstance->IsProjectionEnabled() ? ChaosImmediate_JointMaxAngularProjection : ChaosImmediate_JointMinAngularProjection;
-
-		ConstraintSettings.Motion.bSoftLinearLimitsEnabled = ConstraintInstance->GetIsSoftLinearLimit();
-		ConstraintSettings.Motion.bSoftTwistLimitsEnabled = ConstraintInstance->GetIsSoftTwistLimit();
-		ConstraintSettings.Motion.bSoftSwingLimitsEnabled = ConstraintInstance->GetIsSoftSwingLimit();
-		ConstraintSettings.Motion.SoftLinearStiffness = 1.0f;	// @todo(ccaulfield): xpbd soft linear constraints
-		ConstraintSettings.Motion.SoftLinearDamping = 0.0f;		// @todo(ccaulfield): xpbd soft linear constraints
-		ConstraintSettings.Motion.SoftTwistStiffness = DriveAngularStiffnessScale * ConstraintInstance->GetSoftTwistLimitStiffness();
-		ConstraintSettings.Motion.SoftTwistDamping = DriveAngularDampingScale * ConstraintInstance->GetSoftTwistLimitDamping();
-		ConstraintSettings.Motion.SoftSwingStiffness = DriveAngularStiffnessScale * ConstraintInstance->GetSoftSwingLimitStiffness();
-		ConstraintSettings.Motion.SoftSwingDamping = DriveAngularDampingScale * ConstraintInstance->GetSoftSwingLimitDamping();
-
-		ConstraintSettings.Motion.LinearDriveTarget = ConstraintInstance->ProfileInstance.LinearDrive.PositionTarget;
-		ConstraintSettings.Motion.bLinearDriveEnabled[0] = ConstraintInstance->ProfileInstance.LinearDrive.XDrive.bEnablePositionDrive;
-		ConstraintSettings.Motion.bLinearDriveEnabled[1] = ConstraintInstance->ProfileInstance.LinearDrive.YDrive.bEnablePositionDrive;
-		ConstraintSettings.Motion.bLinearDriveEnabled[2] = ConstraintInstance->ProfileInstance.LinearDrive.ZDrive.bEnablePositionDrive;
-		ConstraintSettings.Motion.LinearDriveStiffness = 0.3f;// ConstraintInstance->ProfileInstance.LinearDrive.XDrive.bEnablePositionDrive;
-		ConstraintSettings.Motion.LinearDriveDamping = 0.0f;// ConstraintInstance->ProfileInstance.LinearDrive.XDrive.bEnablePositionDrive;
-
-		ConstraintSettings.Motion.AngularDriveTarget = FQuat(ConstraintInstance->ProfileInstance.AngularDrive.OrientationTarget);
+		ConstraintSettings.LinearLimit = ConstraintInstance->GetLinearLimit();
+
+		ConstraintSettings.AngularMotionTypes[(int32)EJointAngularConstraintIndex::Twist] = static_cast<EJointMotionType>(ConstraintInstance->GetAngularTwistMotion());
+		ConstraintSettings.AngularMotionTypes[(int32)EJointAngularConstraintIndex::Swing1] = static_cast<EJointMotionType>(ConstraintInstance->GetAngularSwing1Motion());
+		ConstraintSettings.AngularMotionTypes[(int32)EJointAngularConstraintIndex::Swing2] = static_cast<EJointMotionType>(ConstraintInstance->GetAngularSwing2Motion());
+		ConstraintSettings.AngularLimits[(int32)EJointAngularConstraintIndex::Twist] = FMath::DegreesToRadians(ConstraintInstance->GetAngularTwistLimit());
+		ConstraintSettings.AngularLimits[(int32)EJointAngularConstraintIndex::Swing1] = FMath::DegreesToRadians(ConstraintInstance->GetAngularSwing1Limit());
+		ConstraintSettings.AngularLimits[(int32)EJointAngularConstraintIndex::Swing2] = FMath::DegreesToRadians(ConstraintInstance->GetAngularSwing2Limit());
+
+		ConstraintSettings.LinearProjection = ConstraintInstance->IsProjectionEnabled() ? ChaosImmediate_JointMaxLinearProjection : ChaosImmediate_JointMinLinearProjection;
+		ConstraintSettings.AngularProjection = ConstraintInstance->IsProjectionEnabled() ? ChaosImmediate_JointMaxAngularProjection : ChaosImmediate_JointMinAngularProjection;
+		ConstraintSettings.ParentInvMassScale = ConstraintInstance->ProfileInstance.bParentDominates ? (FReal)0 : (FReal)1;
+
+		ConstraintSettings.bSoftLinearLimitsEnabled = ConstraintInstance->GetIsSoftLinearLimit();
+		ConstraintSettings.bSoftTwistLimitsEnabled = ConstraintInstance->GetIsSoftTwistLimit();
+		ConstraintSettings.bSoftSwingLimitsEnabled = ConstraintInstance->GetIsSoftSwingLimit();
+		ConstraintSettings.SoftLinearStiffness = 1.0f;	// @todo(ccaulfield): xpbd soft linear constraints
+		ConstraintSettings.SoftLinearDamping = 0.0f;		// @todo(ccaulfield): xpbd soft linear constraints
+		ConstraintSettings.SoftTwistStiffness = ChaosImmediate_SoftStiffnessScale * ConstraintInstance->GetSoftTwistLimitStiffness();
+		ConstraintSettings.SoftTwistDamping = ChaosImmediate_SoftDampingScale * ConstraintInstance->GetSoftTwistLimitDamping();
+		ConstraintSettings.SoftSwingStiffness = ChaosImmediate_SoftStiffnessScale * ConstraintInstance->GetSoftSwingLimitStiffness();
+		ConstraintSettings.SoftSwingDamping = ChaosImmediate_SoftDampingScale * ConstraintInstance->GetSoftSwingLimitDamping();
+
+		ConstraintSettings.LinearDriveTarget = ConstraintInstance->ProfileInstance.LinearDrive.PositionTarget;
+		ConstraintSettings.bLinearDriveEnabled[0] = ConstraintInstance->ProfileInstance.LinearDrive.XDrive.bEnablePositionDrive;
+		ConstraintSettings.bLinearDriveEnabled[1] = ConstraintInstance->ProfileInstance.LinearDrive.YDrive.bEnablePositionDrive;
+		ConstraintSettings.bLinearDriveEnabled[2] = ConstraintInstance->ProfileInstance.LinearDrive.ZDrive.bEnablePositionDrive;
+		ConstraintSettings.LinearDriveStiffness = 0.3f;// ConstraintInstance->ProfileInstance.LinearDrive.XDrive.bEnablePositionDrive;
+		ConstraintSettings.LinearDriveDamping = 0.0f;// ConstraintInstance->ProfileInstance.LinearDrive.XDrive.bEnablePositionDrive;
+
+		ConstraintSettings.AngularDriveTarget = FQuat(ConstraintInstance->ProfileInstance.AngularDrive.OrientationTarget);
 		// NOTE: Hard dependence on EJointAngularConstraintIndex - the following will break if we change the order
-		ConstraintSettings.Motion.AngularDriveTargetAngles[(int32)EJointAngularConstraintIndex::Twist] = FMath::DegreesToRadians(ConstraintInstance->ProfileInstance.AngularDrive.OrientationTarget.Roll);
-		ConstraintSettings.Motion.AngularDriveTargetAngles[(int32)EJointAngularConstraintIndex::Swing1] = FMath::DegreesToRadians(ConstraintInstance->ProfileInstance.AngularDrive.OrientationTarget.Yaw);
-		ConstraintSettings.Motion.AngularDriveTargetAngles[(int32)EJointAngularConstraintIndex::Swing2] = FMath::DegreesToRadians(ConstraintInstance->ProfileInstance.AngularDrive.OrientationTarget.Pitch);
+		ConstraintSettings.AngularDriveTargetAngles[(int32)EJointAngularConstraintIndex::Twist] = FMath::DegreesToRadians(ConstraintInstance->ProfileInstance.AngularDrive.OrientationTarget.Roll);
+		ConstraintSettings.AngularDriveTargetAngles[(int32)EJointAngularConstraintIndex::Swing1] = FMath::DegreesToRadians(ConstraintInstance->ProfileInstance.AngularDrive.OrientationTarget.Yaw);
+		ConstraintSettings.AngularDriveTargetAngles[(int32)EJointAngularConstraintIndex::Swing2] = FMath::DegreesToRadians(ConstraintInstance->ProfileInstance.AngularDrive.OrientationTarget.Pitch);
 
 		if (ConstraintInstance->ProfileInstance.AngularDrive.AngularDriveMode == EAngularDriveMode::SLERP)
 		{
-			ConstraintSettings.Motion.bAngularSLerpDriveEnabled = ConstraintInstance->ProfileInstance.AngularDrive.SlerpDrive.bEnablePositionDrive;
-			ConstraintSettings.Motion.bAngularTwistDriveEnabled = false;
-			ConstraintSettings.Motion.bAngularSwingDriveEnabled = false;
+			ConstraintSettings.bAngularSLerpDriveEnabled = ConstraintInstance->ProfileInstance.AngularDrive.SlerpDrive.bEnablePositionDrive;
+			ConstraintSettings.bAngularTwistDriveEnabled = false;
+			ConstraintSettings.bAngularSwingDriveEnabled = false;
 			
 		}
 		else
 		{
-			ConstraintSettings.Motion.bAngularSLerpDriveEnabled = false;
-			ConstraintSettings.Motion.bAngularTwistDriveEnabled = ConstraintInstance->ProfileInstance.AngularDrive.TwistDrive.bEnablePositionDrive;
-			ConstraintSettings.Motion.bAngularSwingDriveEnabled = ConstraintInstance->ProfileInstance.AngularDrive.SwingDrive.bEnablePositionDrive;
+			ConstraintSettings.bAngularSLerpDriveEnabled = false;
+			ConstraintSettings.bAngularTwistDriveEnabled = ConstraintInstance->ProfileInstance.AngularDrive.TwistDrive.bEnablePositionDrive;
+			ConstraintSettings.bAngularSwingDriveEnabled = ConstraintInstance->ProfileInstance.AngularDrive.SwingDrive.bEnablePositionDrive;
 		}
 
-		ConstraintSettings.Motion.AngularDriveStiffness = DriveAngularStiffnessScale * ConstraintInstance->ProfileInstance.AngularDrive.TwistDrive.Stiffness;
-		ConstraintSettings.Motion.AngularDriveDamping = DriveAngularDampingScale * ConstraintInstance->ProfileInstance.AngularDrive.TwistDrive.Damping;
+		ConstraintSettings.AngularDriveStiffness = ChaosImmediate_DriveStiffnessScale * ConstraintInstance->ProfileInstance.AngularDrive.TwistDrive.Stiffness;
+		ConstraintSettings.AngularDriveDamping = ChaosImmediate_DriveDampingScale * ConstraintInstance->ProfileInstance.AngularDrive.TwistDrive.Damping;
 
-		ConstraintSettings.Motion.Sanitize();
+		ConstraintSettings.Sanitize();
 
-		ConstraintHandle = Constraints->AddConstraint({ Actor1->ParticleHandle, Actor2->ParticleHandle }, ConstraintSettings);
+		ConstraintHandle = Constraints->AddConstraint({ Actor1->ParticleHandle, Actor2->ParticleHandle }, ConstraintFrames, ConstraintSettings);
 	}
 
 	FJointHandle::~FJointHandle()
diff --git a/Engine/Source/Runtime/Engine/Private/PhysicsEngine/ImmediatePhysics/ImmediatePhysicsChaos/ImmediatePhysicsSimulation_Chaos.cpp b/Engine/Source/Runtime/Engine/Private/PhysicsEngine/ImmediatePhysics/ImmediatePhysicsChaos/ImmediatePhysicsSimulation_Chaos.cpp
--- a/Engine/Source/Runtime/Engine/Private/PhysicsEngine/ImmediatePhysics/ImmediatePhysicsChaos/ImmediatePhysicsSimulation_Chaos.cpp
+++ b/Engine/Source/Runtime/Engine/Private/PhysicsEngine/ImmediatePhysics/ImmediatePhysicsChaos/ImmediatePhysicsSimulation_Chaos.cpp
@@ -50,7 +50,6 @@ float ChaosImmediate_Joint_SwingTwistAngleTolerance = 1.0e-6f;
 int32 ChaosImmediate_Joint_EnableTwistLimits = 1;
 int32 ChaosImmediate_Joint_EnableSwingLimits = 1;
 int32 ChaosImmediate_Joint_EnableDrives = 1;
-int32 ChaosImmediate_Joint_ProjectionPhase = (int32)Chaos::EJointSolverPhase::ApplyPushOut;
 float ChaosImmediate_Joint_LinearProjection = 0.0f;
 float ChaosImmediate_Joint_AngularProjection = 0.0f;
 float ChaosImmediate_Joint_Stiffness = 1.0f;
@@ -125,21 +124,6 @@ FAutoConsoleVariableRef CVarChaosImmPhysDebugDrawJointFeatures(TEXT("p.Chaos.Imm
 
 namespace ImmediatePhysics_Chaos
 {
-	Chaos::EJointSolverPhase ToJointSolverPhase(const int32 Index)
-	{
-		using namespace Chaos;
-
-		if (Index == 1)
-		{
-			return EJointSolverPhase::Apply;
-		}
-		else if (Index == 2)
-		{
-			return EJointSolverPhase::ApplyPushOut;
-		}
-		return EJointSolverPhase::None;
-	}
-
 	template<typename T, int d>
 	bool ShouldIgnoreCollisionConstraint(
 		const Chaos::TGeometryParticleHandle<T, d>* ParticleA,
diff --git a/Engine/Source/Runtime/Experimental/Chaos/Private/Chaos/Joint/PBDJointSolverCholesky.cpp b/Engine/Source/Runtime/Experimental/Chaos/Private/Chaos/Joint/PBDJointSolverCholesky.cpp
--- a/Engine/Source/Runtime/Experimental/Chaos/Private/Chaos/Joint/PBDJointSolverCholesky.cpp
+++ b/Engine/Source/Runtime/Experimental/Chaos/Private/Chaos/Joint/PBDJointSolverCholesky.cpp
@@ -290,7 +290,7 @@ namespace Chaos
 		FDenseMatrix66& J1,
 		FDenseMatrix61& C)
 	{
-		const FReal Limit = JointSettings.Motion.LinearLimit;
+		const FReal Limit = JointSettings.LinearLimit;
 		const FVec3 ConstraintSeparation = Xs[1] - Xs[0];
 		const FReal ConstraintSeparationLen = ConstraintSeparation.Size();
 
@@ -332,7 +332,7 @@ namespace Chaos
 		// Radial Constraint
 		const FVec3 ConstraintSeparationRadial = ConstraintSeparation - ConstraintDistanceAxial * Axis;
 		const FReal ConstraintDistanceRadial = ConstraintSeparationRadial.Size();
-		const FReal RadialLimit = JointSettings.Motion.LinearLimit;
+		const FReal RadialLimit = JointSettings.LinearLimit;
 		bool bRadialConstraintActive = (ConstraintDistanceRadial >= RadialLimit);
 		if (bRadialConstraintActive)
 		{
@@ -352,7 +352,7 @@ namespace Chaos
 		FDenseMatrix66& J1,
 		FDenseMatrix61& C)
 	{
-		const FReal Limit = (AxisMotion == EJointMotionType::Limited) ? JointSettings.Motion.LinearLimit : (FReal)0;
+		const FReal Limit = (AxisMotion == EJointMotionType::Limited) ? JointSettings.LinearLimit : (FReal)0;
 		const FVec3 ConstraintSeparation = Xs[1] - Xs[0];
 
 		// Planar Constraint
@@ -388,7 +388,7 @@ namespace Chaos
 			TwistAngle = -TwistAngle;
 		}
 
-		const FReal TwistAngleMax = JointSettings.Motion.AngularLimits[(int32)EJointAngularConstraintIndex::Twist];
+		const FReal TwistAngleMax = JointSettings.AngularLimits[(int32)EJointAngularConstraintIndex::Twist];
 		bool bConstraintActive = (TwistAngle <= -TwistAngleMax) || (TwistAngle >= TwistAngleMax);
 		if (bConstraintActive)
 		{
@@ -419,8 +419,8 @@ namespace Chaos
 
 		// Calculate swing limit for the current swing axis
 		FReal SwingAngleMax = FLT_MAX;
-		const FReal Swing1Limit = JointSettings.Motion.AngularLimits[(int32)EJointAngularConstraintIndex::Swing1];
-		const FReal Swing2Limit = JointSettings.Motion.AngularLimits[(int32)EJointAngularConstraintIndex::Swing2];
+		const FReal Swing1Limit = JointSettings.AngularLimits[(int32)EJointAngularConstraintIndex::Swing1];
+		const FReal Swing2Limit = JointSettings.AngularLimits[(int32)EJointAngularConstraintIndex::Swing2];
 
 		// Circular swing limit
 		SwingAngleMax = Swing1Limit;
@@ -495,7 +495,7 @@ namespace Chaos
 			const FReal SwingAngle2 = GetSwingAngle(R01SwingYorZ, R01Swing.W);
 
 
-			FReal SwingAngleMax = JointSettings.Motion.AngularLimits[(int32)SwingConstraintIndex];
+			FReal SwingAngleMax = JointSettings.AngularLimits[(int32)SwingConstraintIndex];
 			bool bConstraintActive = (SwingAngle <= -SwingAngleMax) || (SwingAngle >= SwingAngleMax);
 			if (bConstraintActive)
 			{
@@ -521,7 +521,7 @@ namespace Chaos
 	//	const FReal R01SwingYorZ = ((int32)SwingAxisIndex == 2) ? R01Swing.Z : R01Swing.Y;	// Can't index a quat :(
 	//	FReal SwingAngle = GetSwingAngle(R01SwingYorZ, R01Swing.W);
 
-	//	FReal SwingAngleMax = JointSettings.Motion.AngularLimits[(int32)SwingConstraintIndex];
+	//	FReal SwingAngleMax = JointSettings.AngularLimits[(int32)SwingConstraintIndex];
 	//	bool bConstraintActive = (SwingAngle <= -SwingAngleMax) || (SwingAngle >= SwingAngleMax);
 	//	if (bConstraintActive)
 	//	{
@@ -576,7 +576,7 @@ namespace Chaos
 		FDenseMatrix66& J1,
 		FDenseMatrix61& C)
 	{
-		const TVector<EJointMotionType, 3>& Motion = JointSettings.Motion.LinearMotionTypes;
+		const TVector<EJointMotionType, 3>& Motion = JointSettings.LinearMotionTypes;
 		if ((Motion[0] == EJointMotionType::Locked) && (Motion[1] == EJointMotionType::Locked) && (Motion[2] == EJointMotionType::Locked))
 		{
 			AddLinearConstraints_Point(JointSettings, J0, J1, C);
@@ -625,9 +625,9 @@ namespace Chaos
 		FDenseMatrix66& J1,
 		FDenseMatrix61& C)
 	{
-		EJointMotionType TwistMotion = JointSettings.Motion.AngularMotionTypes[(int32)EJointAngularConstraintIndex::Twist];
-		EJointMotionType Swing1Motion = JointSettings.Motion.AngularMotionTypes[(int32)EJointAngularConstraintIndex::Swing1];
-		EJointMotionType Swing2Motion = JointSettings.Motion.AngularMotionTypes[(int32)EJointAngularConstraintIndex::Swing2];
+		EJointMotionType TwistMotion = JointSettings.AngularMotionTypes[(int32)EJointAngularConstraintIndex::Twist];
+		EJointMotionType Swing1Motion = JointSettings.AngularMotionTypes[(int32)EJointAngularConstraintIndex::Swing1];
+		EJointMotionType Swing2Motion = JointSettings.AngularMotionTypes[(int32)EJointAngularConstraintIndex::Swing2];
 
 		bool bAddTwist = bEnableTwistLimits && (TwistMotion != EJointMotionType::Free);
 		bool bAddConeOrSwing = bEnableSwingLimits && ((Swing1Motion != EJointMotionType::Free) || (Swing2Motion != EJointMotionType::Free));
@@ -673,14 +673,14 @@ namespace Chaos
 		FDenseMatrix66& J1,
 		FDenseMatrix61& C)
 	{
-		EJointMotionType TwistMotion = JointSettings.Motion.AngularMotionTypes[(int32)EJointAngularConstraintIndex::Twist];
-		EJointMotionType Swing1Motion = JointSettings.Motion.AngularMotionTypes[(int32)EJointAngularConstraintIndex::Swing1];
-		EJointMotionType Swing2Motion = JointSettings.Motion.AngularMotionTypes[(int32)EJointAngularConstraintIndex::Swing2];
-
-		bool bTwistDriveEnabled = bEnableDrives && JointSettings.Motion.bAngularTwistDriveEnabled && (TwistMotion != EJointMotionType::Locked);
-		bool bSwing1DriveEnabled = bEnableDrives && JointSettings.Motion.bAngularSwingDriveEnabled && (Swing1Motion != EJointMotionType::Locked);
-		bool bSwing2DriveEnabled = bEnableDrives && JointSettings.Motion.bAngularSwingDriveEnabled && (Swing2Motion != EJointMotionType::Locked);
-		bool bSlerpDriveEnabled = bEnableDrives && JointSettings.Motion.bAngularSLerpDriveEnabled && (Swing1Motion != EJointMotionType::Locked) && (Swing2Motion != EJointMotionType::Locked);
+		EJointMotionType TwistMotion = JointSettings.AngularMotionTypes[(int32)EJointAngularConstraintIndex::Twist];
+		EJointMotionType Swing1Motion = JointSettings.AngularMotionTypes[(int32)EJointAngularConstraintIndex::Swing1];
+		EJointMotionType Swing2Motion = JointSettings.AngularMotionTypes[(int32)EJointAngularConstraintIndex::Swing2];
+
+		bool bTwistDriveEnabled = bEnableDrives && JointSettings.bAngularTwistDriveEnabled && (TwistMotion != EJointMotionType::Locked);
+		bool bSwing1DriveEnabled = bEnableDrives && JointSettings.bAngularSwingDriveEnabled && (Swing1Motion != EJointMotionType::Locked);
+		bool bSwing2DriveEnabled = bEnableDrives && JointSettings.bAngularSwingDriveEnabled && (Swing2Motion != EJointMotionType::Locked);
+		bool bSlerpDriveEnabled = bEnableDrives && JointSettings.bAngularSLerpDriveEnabled && (Swing1Motion != EJointMotionType::Locked) && (Swing2Motion != EJointMotionType::Locked);
 		bool bAnyDriveEnabled = (bTwistDriveEnabled || bSwing1DriveEnabled || bSwing2DriveEnabled || bSlerpDriveEnabled);
 
 		if (bAnyDriveEnabled)
diff --git a/Engine/Source/Runtime/Experimental/Chaos/Private/Chaos/Joint/PBDJointSolverGaussSeidel.cpp b/Engine/Source/Runtime/Experimental/Chaos/Private/Chaos/Joint/PBDJointSolverGaussSeidel.cpp
--- a/Engine/Source/Runtime/Experimental/Chaos/Private/Chaos/Joint/PBDJointSolverGaussSeidel.cpp
+++ b/Engine/Source/Runtime/Experimental/Chaos/Private/Chaos/Joint/PBDJointSolverGaussSeidel.cpp
@@ -59,8 +59,8 @@ namespace Chaos
 		const FReal SwingAngle)
 	{
 		// Calculate swing limit for the current swing axis
-		const FReal Swing1Limit = JointSettings.Motion.AngularLimits[(int32)EJointAngularConstraintIndex::Swing1];
-		const FReal Swing2Limit = JointSettings.Motion.AngularLimits[(int32)EJointAngularConstraintIndex::Swing2];
+		const FReal Swing1Limit = JointSettings.AngularLimits[(int32)EJointAngularConstraintIndex::Swing1];
+		const FReal Swing2Limit = JointSettings.AngularLimits[(int32)EJointAngularConstraintIndex::Swing2];
 
 		// Circular swing limit
 		FReal SwingAngleMax = Swing1Limit;
@@ -208,11 +208,12 @@ namespace Chaos
 	{
 		XLs[0] = XL0;
 		XLs[1] = XL1;
-		InvILs[0] = InvIL0;
+
+		// @todo(ccaulfield): mass conditioning
+		InvILs[0] = JointSettings.ParentInvMassScale * InvIL0;
 		InvILs[1] = InvIL1;
-		InvMs[0] = InvM0;
+		InvMs[0] = JointSettings.ParentInvMassScale * InvM0;
 		InvMs[1] = InvM1;
-		// @todo(ccaulfield): mass conditioning
 
 		PrevQs[0] = PrevQ0;
 		PrevQs[1] = PrevQ1;
@@ -294,11 +295,11 @@ namespace Chaos
 	{
 		// @todo(ccaulfield): eliminate branches in solver loop by pre-building list of functions to call?
 
-		EJointMotionType TwistMotion = JointSettings.Motion.AngularMotionTypes[(int32)EJointAngularConstraintIndex::Twist];
-		EJointMotionType Swing1Motion = JointSettings.Motion.AngularMotionTypes[(int32)EJointAngularConstraintIndex::Swing1];
-		EJointMotionType Swing2Motion = JointSettings.Motion.AngularMotionTypes[(int32)EJointAngularConstraintIndex::Swing2];
-		bool bTwistSoft = JointSettings.Motion.bSoftTwistLimitsEnabled;
-		bool bSwingSoft = JointSettings.Motion.bSoftSwingLimitsEnabled;
+		EJointMotionType TwistMotion = JointSettings.AngularMotionTypes[(int32)EJointAngularConstraintIndex::Twist];
+		EJointMotionType Swing1Motion = JointSettings.AngularMotionTypes[(int32)EJointAngularConstraintIndex::Swing1];
+		EJointMotionType Swing2Motion = JointSettings.AngularMotionTypes[(int32)EJointAngularConstraintIndex::Swing2];
+		bool bTwistSoft = JointSettings.bSoftTwistLimitsEnabled;
+		bool bSwingSoft = JointSettings.bSoftSwingLimitsEnabled;
 
 		// Apply twist constraint
 		if (SolverSettings.bEnableTwistLimits)
@@ -357,9 +358,9 @@ namespace Chaos
 		const FPBDJointSolverSettings& SolverSettings,
 		const FPBDJointSettings& JointSettings)
 	{
-		EJointMotionType TwistMotion = JointSettings.Motion.AngularMotionTypes[(int32)EJointAngularConstraintIndex::Twist];
-		EJointMotionType Swing1Motion = JointSettings.Motion.AngularMotionTypes[(int32)EJointAngularConstraintIndex::Swing1];
-		EJointMotionType Swing2Motion = JointSettings.Motion.AngularMotionTypes[(int32)EJointAngularConstraintIndex::Swing2];
+		EJointMotionType TwistMotion = JointSettings.AngularMotionTypes[(int32)EJointAngularConstraintIndex::Twist];
+		EJointMotionType Swing1Motion = JointSettings.AngularMotionTypes[(int32)EJointAngularConstraintIndex::Swing1];
+		EJointMotionType Swing2Motion = JointSettings.AngularMotionTypes[(int32)EJointAngularConstraintIndex::Swing2];
 
 		if (SolverSettings.bEnableDrives)
 		{
@@ -368,26 +369,26 @@ namespace Chaos
 			bool bSwing2Locked = Swing2Motion == EJointMotionType::Locked;
 
 			// No SLerp drive if we have a locked rotation (it will be grayed out in the editor in this case, but could still have been set before the rotation was locked)
-			if (JointSettings.Motion.bAngularSLerpDriveEnabled && !bTwistLocked && !bSwing1Locked && !bSwing2Locked)
+			if (JointSettings.bAngularSLerpDriveEnabled && !bTwistLocked && !bSwing1Locked && !bSwing2Locked)
 			{
 				ApplySLerpDrive(Dt, SolverSettings, JointSettings);
 			}
 			else
 			{
-				if (JointSettings.Motion.bAngularTwistDriveEnabled && !bTwistLocked)
+				if (JointSettings.bAngularTwistDriveEnabled && !bTwistLocked)
 				{
 					ApplyTwistDrive(Dt, SolverSettings, JointSettings);
 				}
 
-				if (JointSettings.Motion.bAngularSwingDriveEnabled && !bSwing1Locked && !bSwing2Locked)
+				if (JointSettings.bAngularSwingDriveEnabled && !bSwing1Locked && !bSwing2Locked)
 				{
 					ApplyConeDrive(Dt, SolverSettings, JointSettings);
 				}
-				else if (JointSettings.Motion.bAngularSwingDriveEnabled && !bSwing1Locked)
+				else if (JointSettings.bAngularSwingDriveEnabled && !bSwing1Locked)
 				{
 					ApplySwingDrive(Dt, SolverSettings, JointSettings, EJointAngularConstraintIndex::Swing1, EJointAngularAxisIndex::Swing1);
 				}
-				else if (JointSettings.Motion.bAngularSwingDriveEnabled && !bSwing2Locked)
+				else if (JointSettings.bAngularSwingDriveEnabled && !bSwing2Locked)
 				{
 					ApplySwingDrive(Dt, SolverSettings, JointSettings, EJointAngularConstraintIndex::Swing2, EJointAngularAxisIndex::Swing2);
 				}
@@ -401,11 +402,11 @@ namespace Chaos
 		const FPBDJointSolverSettings& SolverSettings,
 		const FPBDJointSettings& JointSettings)
 	{
-		EJointMotionType TwistMotion = JointSettings.Motion.AngularMotionTypes[(int32)EJointAngularConstraintIndex::Twist];
-		EJointMotionType Swing1Motion = JointSettings.Motion.AngularMotionTypes[(int32)EJointAngularConstraintIndex::Swing1];
-		EJointMotionType Swing2Motion = JointSettings.Motion.AngularMotionTypes[(int32)EJointAngularConstraintIndex::Swing2];
-		bool bTwistSoft = JointSettings.Motion.bSoftTwistLimitsEnabled;
-		bool bSwingSoft = JointSettings.Motion.bSoftSwingLimitsEnabled;
+		EJointMotionType TwistMotion = JointSettings.AngularMotionTypes[(int32)EJointAngularConstraintIndex::Twist];
+		EJointMotionType Swing1Motion = JointSettings.AngularMotionTypes[(int32)EJointAngularConstraintIndex::Swing1];
+		EJointMotionType Swing2Motion = JointSettings.AngularMotionTypes[(int32)EJointAngularConstraintIndex::Swing2];
+		bool bTwistSoft = JointSettings.bSoftTwistLimitsEnabled;
+		bool bSwingSoft = JointSettings.bSoftSwingLimitsEnabled;
 
 		if (SolverSettings.bEnableTwistLimits)
 		{
@@ -441,7 +442,7 @@ namespace Chaos
 		const FPBDJointSolverSettings& SolverSettings,
 		const FPBDJointSettings& JointSettings)
 	{
-		const TVector<EJointMotionType, 3>& LinearMotion = JointSettings.Motion.LinearMotionTypes;
+		const TVector<EJointMotionType, 3>& LinearMotion = JointSettings.LinearMotionTypes;
 		if ((LinearMotion[0] != EJointMotionType::Free) || (LinearMotion[1] != EJointMotionType::Free) || (LinearMotion[2] != EJointMotionType::Free))
 		{
 			ApplyPositionConstraint(Dt, SolverSettings, JointSettings);
@@ -536,36 +537,20 @@ namespace Chaos
 		// Joint-space inverse mass
 		const FReal II0 = FVec3::DotProduct(Axis0, IA0);
 		const FReal II1 = FVec3::DotProduct(Axis1, IA1);
-		const FReal IM = II0 + II1;
-
-		//FVec3 LinearAxis0 = FVec3::CrossProduct(Xs[0] - Ps[0], Axis0);
-		//FVec3 LinearAxis1 = FVec3::CrossProduct(Xs[1] - Ps[1], Axis1);
-		//FReal LinearAxis0InvLenSq = LinearAxis0.SizeSquared();
-		//FReal LinearAxis1InvLenSq = LinearAxis1.SizeSquared();
-		//if (LinearAxis0InvLenSq > KINDA_SMALL_NUMBER)
-		//{
-		//	LinearAxis0InvLenSq = (FReal)1 / LinearAxis0InvLenSq;
-		//	LinearAxis0 = LinearAxis0 * LinearAxis0InvLenSq;
-		//}
-		//if (LinearAxis1InvLenSq > KINDA_SMALL_NUMBER)
-		//{
-		//	LinearAxis1InvLenSq = (FReal)1 / LinearAxis1InvLenSq;
-		//	LinearAxis1 = LinearAxis1 * LinearAxis1InvLenSq;
-		//}
-		//const FReal IM0 = InvMs[0] * LinearAxis0InvLenSq;
-		//const FReal IM1 = InvMs[1] * LinearAxis1InvLenSq;
-		//const FReal IM = IM0 + II0 + IM1 + II1;
-
-		// This Position delta assumes locked linear constraints...correction in the direction of inactive constraints should be eliminated.
-		// However, it's pretty rare to have angular limits and position limits on the same joint. Should check for this and fix if necessary.
-		//const FVec3 DP0 = LinearAxis0 * (Angle * InvMs[0] / IM);
-		//const FVec3 DP1 = LinearAxis1 * (-Angle * InvMs[1] / IM);
-		const FVec3 DR0 = IA0 * (Angle / IM);
-		const FVec3 DR1 = IA1 * -(Angle / IM);
-		//const FVec3 DR0 = Axis0 * (Angle * II0 / IM);
-		//const FVec3 DR1 = Axis1 * -(Angle * II1 / IM);
-
-		//ApplyPositionDelta(Stiffness, DP0, DP1);
+
+		const FVec3 DR0 = IA0 * (Angle / (II0 + II1));
+		const FVec3 DR1 = IA1 * -(Angle / (II0 + II1));
+		//const FVec3 DR0 = Axis0 * (Angle * II0 / (II0 + II1));
+		//const FVec3 DR1 = Axis1 * -(Angle * II1 / (II0 + II1));
+
+		// @todo(ccaulfield): this position correction needs to have components in direction of inactive position constraints removed
+		const FReal IM0 = InvMs[0];
+		const FReal IM1 = InvMs[1];
+		const FVec3 DX = FVec3::CrossProduct(DR0, Xs[0] - Ps[0]) + FVec3::CrossProduct(DR1, Xs[1] - Ps[1]);
+		const FVec3 DP0 = DX * (IM0 / (IM0 + IM1));
+		const FVec3 DP1 = DX * (-IM1 / (IM0 + IM1));
+
+		ApplyPositionDelta(Stiffness, DP0, DP1);
 		ApplyRotationDelta(Stiffness, DR0, DR1);
 	}
 
@@ -590,61 +575,47 @@ namespace Chaos
 		// Joint-space inverse mass
 		const FReal II0 = FVec3::DotProduct(Axis0, IA0);
 		const FReal II1 = FVec3::DotProduct(Axis1, IA1);
-		const FReal IM = II0 + II1;
-
-		//FVec3 LinearAxis0 = FVec3::CrossProduct(Xs[0] - Ps[0], Axis0);
-		//FVec3 LinearAxis1 = FVec3::CrossProduct(Xs[1] - Ps[1], Axis1);
-		//FReal LinearAxis0InvLenSq = LinearAxis0.SizeSquared();
-		//FReal LinearAxis1InvLenSq = LinearAxis1.SizeSquared();
-		//if (LinearAxis0InvLenSq > KINDA_SMALL_NUMBER)
-		//{
-		//	LinearAxis0InvLenSq = (FReal)1 / LinearAxis0InvLenSq;
-		//}
-		//if (LinearAxis1InvLenSq > KINDA_SMALL_NUMBER)
-		//{
-		//	LinearAxis1InvLenSq = (FReal)1 / LinearAxis1InvLenSq;
-		//}
-		//const FReal IM0 = InvMs[0] * LinearAxis0InvLenSq;
-		//const FReal IM1 = InvMs[1] * LinearAxis1InvLenSq;
-		//const FReal IM = IM0 + II0 + IM1 + II1;
+		const FReal II = (II0 + II1);
 
 		// Damping angular velocity
-		FReal JW = 0;
+		FReal AngVel = 0;
 		if (Damping > KINDA_SMALL_NUMBER)
 		{
 			const FVec3 W0 = FRotation3::CalculateAngularVelocity(PrevQs[0], Qs[0], Dt);
 			const FVec3 W1 = FRotation3::CalculateAngularVelocity(PrevQs[1], Qs[1], Dt);
-			JW = FVec3::DotProduct(Axis0, W0) - FVec3::DotProduct(Axis1, W1);
+			AngVel = FVec3::DotProduct(Axis0, W0) - FVec3::DotProduct(Axis1, W1);
 		}
 
-		// Cancel out the mass for acceleration drives
-		const FReal MassScale = (bAccelerationMode) ? IM : (FReal)1.0;
-
 		FReal DLambda = 0;
 		if (Stiffness > KINDA_SMALL_NUMBER)
 		{
-			const FReal Alpha = MassScale / (Stiffness * Dt * Dt);
+			const FReal Alpha = (FReal)1 / (Stiffness * Dt * Dt);
 			const FReal AlphaBeta = Damping / Stiffness;
-			const FReal DLambdaNumerator = (Angle - Alpha * Lambda - AlphaBeta * JW);
-			const FReal DLambdaDenominator = (((FReal)1 + AlphaBeta / Dt) * IM + Alpha);
-			DLambda = DLambdaNumerator / DLambdaDenominator;
-		}
-		else
-		{
-			const FReal Beta = MassScale * Damping / (Dt * Dt);
-			const FReal DLambdaNumerator = -Lambda - Beta * JW;
-			const FReal DLambdaDenominator = IM;
-			DLambda = DLambdaNumerator / DLambdaDenominator;
+			if (bAccelerationMode)
+			{
+				const FReal Multiplier = (FReal)1 / (((FReal)1 + AlphaBeta / Dt) + Alpha);
+				DLambda = Multiplier * (Angle - AlphaBeta * AngVel) / II - Multiplier * Alpha * Lambda;
+			}
+			else
+			{
+				const FReal Multiplier = (FReal)1 / (((FReal)1 + AlphaBeta / Dt) * II + Alpha);
+				DLambda = Multiplier * (Angle - Alpha * Lambda - AlphaBeta * AngVel);
+			}
 		}
 
-		//const FVec3 DP0 = LinearAxis0 * (DLambda * IM0);
-		//const FVec3 DP1 = LinearAxis1 * (-DLambda * IM1);
 		//const FVec3 DR0 = IA0 * DLambda;
 		//const FVec3 DR1 = IA1 * -DLambda;
 		const FVec3 DR0 = Axis0 * (DLambda * II0);
 		const FVec3 DR1 = Axis1 * -(DLambda * II1);
 
-		//ApplyPositionDelta(1.0f, DP0, DP1);
+		// @todo(ccaulfield): this position correction needs to have components in direction of inactive position constraints removed
+		const FReal IM0 = InvMs[0];
+		const FReal IM1 = InvMs[1];
+		const FVec3 DX = FVec3::CrossProduct(DR0, Xs[0] - Ps[0]) + FVec3::CrossProduct(DR1, Xs[1] - Ps[1]);
+		const FVec3 DP0 = DX * (IM0 / (IM0 + IM1));
+		const FVec3 DP1 = DX * (-IM1 / (IM0 + IM1));
+
+		ApplyPositionDelta(1.0f, DP0, DP1);
 		ApplyRotationDelta(1.0f, DR0, DR1);
 		Lambda += DLambda;
 	}
@@ -670,11 +641,11 @@ namespace Chaos
 		const FVec3 TwistAxis1 = Rs[1] * TwistAxisLocal;
 
 		FReal TwistAngleMax = FLT_MAX;
-		if (JointSettings.Motion.AngularMotionTypes[(int32)EJointAngularConstraintIndex::Twist] == EJointMotionType::Limited)
+		if (JointSettings.AngularMotionTypes[(int32)EJointAngularConstraintIndex::Twist] == EJointMotionType::Limited)
 		{
-			TwistAngleMax = JointSettings.Motion.AngularLimits[(int32)EJointAngularConstraintIndex::Twist];
+			TwistAngleMax = JointSettings.AngularLimits[(int32)EJointAngularConstraintIndex::Twist];
 		}
-		else if (JointSettings.Motion.AngularMotionTypes[(int32)EJointAngularConstraintIndex::Twist] == EJointMotionType::Locked)
+		else if (JointSettings.AngularMotionTypes[(int32)EJointAngularConstraintIndex::Twist] == EJointMotionType::Locked)
 		{
 			TwistAngleMax = 0;
 		}
@@ -713,11 +684,11 @@ namespace Chaos
 		const FVec3 TwistAxis1 = Rs[1] * TwistAxisLocal;
 
 		FReal TwistAngleMax = FLT_MAX;
-		if (JointSettings.Motion.AngularMotionTypes[(int32)EJointAngularConstraintIndex::Twist] == EJointMotionType::Limited)
+		if (JointSettings.AngularMotionTypes[(int32)EJointAngularConstraintIndex::Twist] == EJointMotionType::Limited)
 		{
-			TwistAngleMax = JointSettings.Motion.AngularLimits[(int32)EJointAngularConstraintIndex::Twist];
+			TwistAngleMax = JointSettings.AngularLimits[(int32)EJointAngularConstraintIndex::Twist];
 		}
-		else if (JointSettings.Motion.AngularMotionTypes[(int32)EJointAngularConstraintIndex::Twist] == EJointMotionType::Locked)
+		else if (JointSettings.AngularMotionTypes[(int32)EJointAngularConstraintIndex::Twist] == EJointMotionType::Locked)
 		{
 			TwistAngleMax = 0;
 		}
@@ -767,7 +738,7 @@ namespace Chaos
 
 		const FVec3 TwistAxis0 = Rs[0] * TwistAxis01;
 		const FVec3 TwistAxis1 = Rs[1] * TwistAxis01;
-		const FReal TwistAngleTarget = JointSettings.Motion.AngularDriveTargetAngles[(int32)EJointAngularConstraintIndex::Twist];
+		const FReal TwistAngleTarget = JointSettings.AngularDriveTargetAngles[(int32)EJointAngularConstraintIndex::Twist];
 		const FReal DTwistAngle = TwistAngle - TwistAngleTarget;
 
 		// Apply twist correction
@@ -851,8 +822,8 @@ namespace Chaos
 
 		// Circular swing target (max of Swing1, Swing2 targets)
 		// @todo(ccaulfield): what should cone target really do?
-		const FReal Swing1Target = JointSettings.Motion.AngularDriveTargetAngles[(int32)EJointAngularConstraintIndex::Swing1];
-		const FReal Swing2Target = JointSettings.Motion.AngularDriveTargetAngles[(int32)EJointAngularConstraintIndex::Swing2];
+		const FReal Swing1Target = JointSettings.AngularDriveTargetAngles[(int32)EJointAngularConstraintIndex::Swing1];
+		const FReal Swing2Target = JointSettings.AngularDriveTargetAngles[(int32)EJointAngularConstraintIndex::Swing2];
 		const FReal SwingAngleTarget = FMath::Max(Swing1Target, Swing2Target);
 		const FReal DSwingAngle = SwingAngle - SwingAngleTarget;
 
@@ -910,8 +881,8 @@ namespace Chaos
 		FReal SwingAngle;
 		GetSwingAxisAngle(Rs[0], Rs[1], SolverSettings.SwingTwistAngleTolerance, SwingConstraintIndex, SwingAxisIndex, SwingAxis, SwingAngle);
 
-		const bool isLocked = (JointSettings.Motion.AngularMotionTypes[(int32)SwingConstraintIndex] == EJointMotionType::Locked);
-		const FReal SwingAngleMax = isLocked ? 0 : JointSettings.Motion.AngularLimits[(int32)SwingConstraintIndex];
+		const bool isLocked = (JointSettings.AngularMotionTypes[(int32)SwingConstraintIndex] == EJointMotionType::Locked);
+		const FReal SwingAngleMax = isLocked ? 0 : JointSettings.AngularLimits[(int32)SwingConstraintIndex];
 
 		// Calculate swing error we need to correct
 		FReal DSwingAngle = 0;
@@ -945,7 +916,7 @@ namespace Chaos
 		FReal SwingAngle;
 		GetSwingAxisAngle(Rs[0], Rs[1], SolverSettings.SwingTwistAngleTolerance, SwingConstraintIndex, SwingAxisIndex, SwingAxis, SwingAngle);
 
-		const FReal SwingAngleMax = JointSettings.Motion.AngularLimits[(int32)SwingConstraintIndex];
+		const FReal SwingAngleMax = JointSettings.AngularLimits[(int32)SwingConstraintIndex];
 
 		// Calculate swing error we need to correct
 		FReal DSwingAngle = 0;
@@ -981,7 +952,7 @@ namespace Chaos
 		FReal SwingAngle;
 		GetSwingAxisAngle(Rs[0], Rs[1], SolverSettings.SwingTwistAngleTolerance, SwingConstraintIndex, SwingAxisIndex, SwingAxis, SwingAngle);
 	
-		const FReal SwingAngleTarget = JointSettings.Motion.AngularDriveTargetAngles[(int32)SwingConstraintIndex];
+		const FReal SwingAngleTarget = JointSettings.AngularDriveTargetAngles[(int32)SwingConstraintIndex];
 		const FReal DSwingAngle = SwingAngle - SwingAngleTarget;
 
 		// Apply drive forces to each body
@@ -1010,7 +981,7 @@ namespace Chaos
 		const FPBDJointSettings& JointSettings)
 	{
 		// Calculate the rotation we need to apply to resolve the rotation delta
-		const FRotation3 TargetR1 = Rs[0] * JointSettings.Motion.AngularDriveTarget;
+		const FRotation3 TargetR1 = Rs[0] * JointSettings.AngularDriveTarget;
 		const FRotation3 DR1 = TargetR1 * Rs[1].Inverse();
 
 		FVec3 SLerpAxis;
@@ -1077,15 +1048,15 @@ namespace Chaos
 		const FPBDJointSettings& JointSettings)
 	{
 		// @todo(ccaulfield): XPBD for linear drives
-		if (JointSettings.Motion.bLinearDriveEnabled[0] || JointSettings.Motion.bLinearDriveEnabled[1] || JointSettings.Motion.bLinearDriveEnabled[2])
+		if (JointSettings.bLinearDriveEnabled[0] || JointSettings.bLinearDriveEnabled[1] || JointSettings.bLinearDriveEnabled[2])
 		{
-			const FVec3 TargetX = Xs[0] + Rs[0] * JointSettings.Motion.LinearDriveTarget;
+			const FVec3 TargetX = Xs[0] + Rs[0] * JointSettings.LinearDriveTarget;
 			FVec3 CX = Xs[1] - TargetX;
 
 			const FMatrix33 RM0 = Rs[0].ToMatrix();
 			for (int32 IAxis = 0; IAxis < 3; ++IAxis)
 			{
-				if (!JointSettings.Motion.bLinearDriveEnabled[IAxis])
+				if (!JointSettings.bLinearDriveEnabled[IAxis])
 				{
 					const FVec3 Axis = RM0.GetAxis(IAxis);
 					CX = CX - FVec3::DotProduct(Axis, CX) * Axis;
diff --git a/Engine/Source/Runtime/Experimental/Chaos/Private/Chaos/PBDJointConstraintUtilities.cpp b/Engine/Source/Runtime/Experimental/Chaos/Private/Chaos/PBDJointConstraintUtilities.cpp
--- a/Engine/Source/Runtime/Experimental/Chaos/Private/Chaos/PBDJointConstraintUtilities.cpp
+++ b/Engine/Source/Runtime/Experimental/Chaos/Private/Chaos/PBDJointConstraintUtilities.cpp
@@ -20,9 +20,9 @@ namespace Chaos
 		const FPBDJointSolverSettings& SolverSettings,
 		const FPBDJointSettings& JointSettings)
 	{
-		const FReal SolverStiffness = (SolverSettings.Stiffness > (FReal)0) ? SolverSettings.Stiffness : JointSettings.Motion.Stiffness;
-		const FReal SoftSolverStiffness = (SolverSettings.SoftLinearStiffness > (FReal)0) ? SolverSettings.SoftLinearStiffness : JointSettings.Motion.SoftLinearStiffness;
-		const bool bIsSoft = JointSettings.Motion.bSoftLinearLimitsEnabled && ((JointSettings.Motion.LinearMotionTypes[0] == EJointMotionType::Limited) || (JointSettings.Motion.LinearMotionTypes[1] == EJointMotionType::Limited) || (JointSettings.Motion.LinearMotionTypes[2] == EJointMotionType::Limited));
+		const FReal SolverStiffness = (SolverSettings.Stiffness > (FReal)0) ? SolverSettings.Stiffness : JointSettings.Stiffness;
+		const FReal SoftSolverStiffness = (SolverSettings.SoftLinearStiffness > (FReal)0) ? SolverSettings.SoftLinearStiffness : JointSettings.SoftLinearStiffness;
+		const bool bIsSoft = JointSettings.bSoftLinearLimitsEnabled && ((JointSettings.LinearMotionTypes[0] == EJointMotionType::Limited) || (JointSettings.LinearMotionTypes[1] == EJointMotionType::Limited) || (JointSettings.LinearMotionTypes[2] == EJointMotionType::Limited));
 		const FReal Stiffness = bIsSoft ? SolverStiffness * SoftSolverStiffness : SolverStiffness;
 		return Stiffness;
 	}
@@ -31,98 +31,98 @@ namespace Chaos
 		const FPBDJointSolverSettings& SolverSettings,
 		const FPBDJointSettings& JointSettings)
 	{
-		return (SolverSettings.Stiffness > (FReal)0) ? SolverSettings.Stiffness : JointSettings.Motion.Stiffness;
+		return (SolverSettings.Stiffness > (FReal)0) ? SolverSettings.Stiffness : JointSettings.Stiffness;
 	}
 
 	FReal FPBDJointUtilities::GetSoftTwistStiffness(
 		const FPBDJointSolverSettings& SolverSettings,
 		const FPBDJointSettings& JointSettings)
 	{
-		return (SolverSettings.SoftTwistStiffness > 0)? SolverSettings.SoftTwistStiffness : JointSettings.Motion.SoftTwistStiffness;
+		return (SolverSettings.SoftTwistStiffness > 0)? SolverSettings.SoftTwistStiffness : JointSettings.SoftTwistStiffness;
 	}
 
 	FReal FPBDJointUtilities::GetSoftTwistDamping(
 		const FPBDJointSolverSettings& SolverSettings,
 		const FPBDJointSettings& JointSettings)
 	{
-		return (SolverSettings.SoftTwistDamping > 0) ? SolverSettings.SoftTwistDamping : JointSettings.Motion.SoftTwistDamping;
+		return (SolverSettings.SoftTwistDamping > 0) ? SolverSettings.SoftTwistDamping : JointSettings.SoftTwistDamping;
 	}
 
 	FReal FPBDJointUtilities::GetSwingStiffness(
 		const FPBDJointSolverSettings& SolverSettings,
 		const FPBDJointSettings& JointSettings)
 	{
-		return (SolverSettings.Stiffness > (FReal)0) ? SolverSettings.Stiffness : JointSettings.Motion.Stiffness;
+		return (SolverSettings.Stiffness > (FReal)0) ? SolverSettings.Stiffness : JointSettings.Stiffness;
 	}
 
 	FReal FPBDJointUtilities::GetSoftSwingStiffness(
 		const FPBDJointSolverSettings& SolverSettings,
 		const FPBDJointSettings& JointSettings)
 	{
-		return (SolverSettings.SoftSwingStiffness > 0) ? SolverSettings.SoftSwingStiffness : JointSettings.Motion.SoftSwingStiffness;
+		return (SolverSettings.SoftSwingStiffness > 0) ? SolverSettings.SoftSwingStiffness : JointSettings.SoftSwingStiffness;
 	}
 
 	FReal FPBDJointUtilities::GetSoftSwingDamping(
 		const FPBDJointSolverSettings& SolverSettings,
 		const FPBDJointSettings& JointSettings)
 	{
-		return (SolverSettings.SoftSwingDamping > 0) ? SolverSettings.SoftSwingDamping : JointSettings.Motion.SoftSwingDamping;
+		return (SolverSettings.SoftSwingDamping > 0) ? SolverSettings.SoftSwingDamping : JointSettings.SoftSwingDamping;
 	}
 
 	FReal FPBDJointUtilities::GetLinearDriveStiffness(
 		const FPBDJointSolverSettings& SolverSettings,
 		const FPBDJointSettings& JointSettings)
 	{
-		return (SolverSettings.LinearDriveStiffness > 0.0f) ? SolverSettings.LinearDriveStiffness : JointSettings.Motion.LinearDriveStiffness;
+		return (SolverSettings.LinearDriveStiffness > 0.0f) ? SolverSettings.LinearDriveStiffness : JointSettings.LinearDriveStiffness;
 	}
 
 	FReal FPBDJointUtilities::GetLinearDriveDamping(
 		const FPBDJointSolverSettings& SolverSettings,
 		const FPBDJointSettings& JointSettings)
 	{
-		return (SolverSettings.LinearDriveDamping > 0.0f) ? SolverSettings.LinearDriveDamping : JointSettings.Motion.LinearDriveDamping;
+		return (SolverSettings.LinearDriveDamping > 0.0f) ? SolverSettings.LinearDriveDamping : JointSettings.LinearDriveDamping;
 	}
 
 	FReal FPBDJointUtilities::GetAngularDriveStiffness(
 		const FPBDJointSolverSettings& SolverSettings,
 		const FPBDJointSettings& JointSettings)
 	{
-		return (SolverSettings.AngularDriveStiffness > 0.0f) ? SolverSettings.AngularDriveStiffness : JointSettings.Motion.AngularDriveStiffness;
+		return (SolverSettings.AngularDriveStiffness > 0.0f) ? SolverSettings.AngularDriveStiffness : JointSettings.AngularDriveStiffness;
 	}
 
 	FReal FPBDJointUtilities::GetAngularDriveDamping(
 		const FPBDJointSolverSettings& SolverSettings,
 		const FPBDJointSettings& JointSettings)
 	{
-		return (SolverSettings.AngularDriveDamping > 0.0f) ? SolverSettings.AngularDriveDamping : JointSettings.Motion.AngularDriveDamping;
+		return (SolverSettings.AngularDriveDamping > 0.0f) ? SolverSettings.AngularDriveDamping : JointSettings.AngularDriveDamping;
 	}
 
 	FReal FPBDJointUtilities::GetLinearProjection(
 		const FPBDJointSolverSettings& SolverSettings,
 		const FPBDJointSettings& JointSettings)
 	{
-		return (SolverSettings.LinearProjection > 0.0f) ? SolverSettings.LinearProjection : JointSettings.Motion.LinearProjection;
+		return (SolverSettings.LinearProjection > 0.0f) ? SolverSettings.LinearProjection : JointSettings.LinearProjection;
 	}
 
 	FReal FPBDJointUtilities::GetAngularProjection(
 		const FPBDJointSolverSettings& SolverSettings,
 		const FPBDJointSettings& JointSettings)
 	{
-		return (SolverSettings.AngularProjection > 0.0f) ? SolverSettings.AngularProjection : JointSettings.Motion.AngularProjection;
+		return (SolverSettings.AngularProjection > 0.0f) ? SolverSettings.AngularProjection : JointSettings.AngularProjection;
 	}
 
 	bool FPBDJointUtilities::GetSoftAccelerationMode(
 		const FPBDJointSolverSettings& SolverSettings,
 		const FPBDJointSettings& JointSettings)
 	{
-		return true;
+		return JointSettings.SoftForceMode == EJointForceMode::Acceleration;
 	}
 
 	bool FPBDJointUtilities::GetDriveAccelerationMode(
 		const FPBDJointSolverSettings& SolverSettings,
 		const FPBDJointSettings& JointSettings)
 	{
-		return true;
+		return JointSettings.AngularDriveForceMode == EJointForceMode::Acceleration;
 	}
 
 
@@ -360,7 +360,7 @@ namespace Chaos
 	
 	FVec3 FPBDJointUtilities::GetLimitedPositionError(const FPBDJointSettings& JointSettings, const FRotation3& R0, const FVec3& InCX)
 	{
-		const TVector<EJointMotionType, 3>& Motion = JointSettings.Motion.LinearMotionTypes;
+		const TVector<EJointMotionType, 3>& Motion = JointSettings.LinearMotionTypes;
 		if ((Motion[0] == EJointMotionType::Locked) && (Motion[1] == EJointMotionType::Locked) && (Motion[2] == EJointMotionType::Locked))
 		{
 			return InCX;
@@ -368,25 +368,25 @@ namespace Chaos
 		else if ((Motion[0] == EJointMotionType::Limited) && (Motion[1] == EJointMotionType::Limited) && (Motion[2] == EJointMotionType::Limited))
 		{
 			// Spherical distance constraints
-			return GetSphereLimitedPositionError(InCX, JointSettings.Motion.LinearLimit);
+			return GetSphereLimitedPositionError(InCX, JointSettings.LinearLimit);
 		}
 		else if ((Motion[1] == EJointMotionType::Limited) && (Motion[2] == EJointMotionType::Limited))
 		{
 			// Circular Limit (X Axis)
 			FVec3 Axis = R0 * FVec3(1, 0, 0);
-			return GetCylinderLimitedPositionError(InCX, Axis, JointSettings.Motion.LinearLimit, Motion[0]);
+			return GetCylinderLimitedPositionError(InCX, Axis, JointSettings.LinearLimit, Motion[0]);
 		}
 		else if ((Motion[0] == EJointMotionType::Limited) && (Motion[2] == EJointMotionType::Limited))
 		{
 			// Circular Limit (Y Axis)
 			FVec3 Axis = R0 * FVec3(0, 1, 0);
-			return GetCylinderLimitedPositionError(InCX, Axis, JointSettings.Motion.LinearLimit, Motion[1]);
+			return GetCylinderLimitedPositionError(InCX, Axis, JointSettings.LinearLimit, Motion[1]);
 		}
 		else if ((Motion[0] == EJointMotionType::Limited) && (Motion[1] == EJointMotionType::Limited))
 		{
 			// Circular Limit (Z Axis)
 			FVec3 Axis = R0 * FVec3(0, 0, 1);
-			return GetCylinderLimitedPositionError(InCX, Axis, JointSettings.Motion.LinearLimit, Motion[2]);
+			return GetCylinderLimitedPositionError(InCX, Axis, JointSettings.LinearLimit, Motion[2]);
 		}
 		else
 		{
@@ -395,17 +395,17 @@ namespace Chaos
 			if (Motion[0] != EJointMotionType::Locked)
 			{
 				FVec3 Axis = R0 * FVec3(1, 0, 0);
-				CX = GetLineLimitedPositionError(CX, Axis, JointSettings.Motion.LinearLimit, Motion[0]);
+				CX = GetLineLimitedPositionError(CX, Axis, JointSettings.LinearLimit, Motion[0]);
 			}
 			if (Motion[1] != EJointMotionType::Locked)
 			{
 				FVec3 Axis = R0 * FVec3(0, 1, 0);
-				CX = GetLineLimitedPositionError(CX, Axis, JointSettings.Motion.LinearLimit, Motion[1]);
+				CX = GetLineLimitedPositionError(CX, Axis, JointSettings.LinearLimit, Motion[1]);
 			}
 			if (Motion[2] != EJointMotionType::Locked)
 			{
 				FVec3 Axis = R0 * FVec3(0, 0, 1);
-				CX = GetLineLimitedPositionError(CX, Axis, JointSettings.Motion.LinearLimit, Motion[2]);
+				CX = GetLineLimitedPositionError(CX, Axis, JointSettings.LinearLimit, Motion[2]);
 			}
 			return CX;
 		}
@@ -414,7 +414,7 @@ namespace Chaos
 	
 	FVec3 FPBDJointUtilities::GetLimitedVelocityError(const FPBDJointSettings& JointSettings, const FRotation3& R0, const FVec3& InCX, const FVec3& InCV)
 	{
-		const TVector<EJointMotionType, 3>& Motion = JointSettings.Motion.LinearMotionTypes;
+		const TVector<EJointMotionType, 3>& Motion = JointSettings.LinearMotionTypes;
 		if ((Motion[0] == EJointMotionType::Locked) && (Motion[1] == EJointMotionType::Locked) && (Motion[2] == EJointMotionType::Locked))
 		{
 			return InCV;
@@ -422,25 +422,25 @@ namespace Chaos
 		else if ((Motion[0] == EJointMotionType::Limited) && (Motion[1] == EJointMotionType::Limited) && (Motion[2] == EJointMotionType::Limited))
 		{
 			// Spherical distance constraints
-			return GetSphereLimitedVelocityError(InCX, JointSettings.Motion.LinearLimit, InCV);
+			return GetSphereLimitedVelocityError(InCX, JointSettings.LinearLimit, InCV);
 		}
 		else if ((Motion[1] == EJointMotionType::Limited) && (Motion[2] == EJointMotionType::Limited))
 		{
 			// Circular Limit (X Axis)
 			FVec3 Axis = R0 * FVec3(1, 0, 0);
-			return GetCylinderLimitedVelocityError(InCX, Axis, JointSettings.Motion.LinearLimit, Motion[0], InCV);
+			return GetCylinderLimitedVelocityError(InCX, Axis, JointSettings.LinearLimit, Motion[0], InCV);
 		}
 		else if ((Motion[0] == EJointMotionType::Limited) && (Motion[2] == EJointMotionType::Limited))
 		{
 			// Circular Limit (Y Axis)
 			FVec3 Axis = R0 * FVec3(0, 1, 0);
-			return GetCylinderLimitedVelocityError(InCX, Axis, JointSettings.Motion.LinearLimit, Motion[1], InCV);
+			return GetCylinderLimitedVelocityError(InCX, Axis, JointSettings.LinearLimit, Motion[1], InCV);
 		}
 		else if ((Motion[0] == EJointMotionType::Limited) && (Motion[1] == EJointMotionType::Limited))
 		{
 			// Circular Limit (Z Axis)
 			FVec3 Axis = R0 * FVec3(0, 0, 1);
-			return GetCylinderLimitedVelocityError(InCX, Axis, JointSettings.Motion.LinearLimit, Motion[2], InCV);
+			return GetCylinderLimitedVelocityError(InCX, Axis, JointSettings.LinearLimit, Motion[2], InCV);
 		}
 		else
 		{
@@ -449,17 +449,17 @@ namespace Chaos
 			if (Motion[0] != EJointMotionType::Locked)
 			{
 				FVec3 Axis = R0 * FVec3(1, 0, 0);
-				CV = GetLineLimitedVelocityError(InCX, Axis, JointSettings.Motion.LinearLimit, Motion[0], CV);
+				CV = GetLineLimitedVelocityError(InCX, Axis, JointSettings.LinearLimit, Motion[0], CV);
 			}
 			if (Motion[1] != EJointMotionType::Locked)
 			{
 				FVec3 Axis = R0 * FVec3(0, 1, 0);
-				CV = GetLineLimitedVelocityError(InCX, Axis, JointSettings.Motion.LinearLimit, Motion[1], CV);
+				CV = GetLineLimitedVelocityError(InCX, Axis, JointSettings.LinearLimit, Motion[1], CV);
 			}
 			if (Motion[2] != EJointMotionType::Locked)
 			{
 				FVec3 Axis = R0 * FVec3(0, 0, 1);
-				CV = GetLineLimitedVelocityError(InCX, Axis, JointSettings.Motion.LinearLimit, Motion[2], CV);
+				CV = GetLineLimitedVelocityError(InCX, Axis, JointSettings.LinearLimit, Motion[2], CV);
 			}
 			return CV;
 		}
@@ -923,11 +923,11 @@ namespace Chaos
 		const FVec3 TwistAxis0 = R0 * TwistAxis01;
 		const FVec3 TwistAxis1 = R1 * TwistAxis01;
 		FReal TwistAngleMax = FLT_MAX;
-		if (JointSettings.Motion.AngularMotionTypes[(int32)EJointAngularConstraintIndex::Twist] == EJointMotionType::Limited)
+		if (JointSettings.AngularMotionTypes[(int32)EJointAngularConstraintIndex::Twist] == EJointMotionType::Limited)
 		{
-			TwistAngleMax = JointSettings.Motion.AngularLimits[(int32)EJointAngularConstraintIndex::Twist];
+			TwistAngleMax = JointSettings.AngularLimits[(int32)EJointAngularConstraintIndex::Twist];
 		}
-		else if (JointSettings.Motion.AngularMotionTypes[(int32)EJointAngularConstraintIndex::Twist] == EJointMotionType::Locked)
+		else if (JointSettings.AngularMotionTypes[(int32)EJointAngularConstraintIndex::Twist] == EJointMotionType::Locked)
 		{
 			TwistAngleMax = 0;
 		}
@@ -999,11 +999,11 @@ namespace Chaos
 		const FVec3 TwistAxis1 = R1 * TwistAxis01;
 
 		FReal TwistAngleMax = FLT_MAX;
-		if (JointSettings.Motion.AngularMotionTypes[(int32)EJointAngularConstraintIndex::Twist] == EJointMotionType::Limited)
+		if (JointSettings.AngularMotionTypes[(int32)EJointAngularConstraintIndex::Twist] == EJointMotionType::Limited)
 		{
-			TwistAngleMax = JointSettings.Motion.AngularLimits[(int32)EJointAngularConstraintIndex::Twist];
+			TwistAngleMax = JointSettings.AngularLimits[(int32)EJointAngularConstraintIndex::Twist];
 		}
-		else if (JointSettings.Motion.AngularMotionTypes[(int32)EJointAngularConstraintIndex::Twist] == EJointMotionType::Locked)
+		else if (JointSettings.AngularMotionTypes[(int32)EJointAngularConstraintIndex::Twist] == EJointMotionType::Locked)
 		{
 			TwistAngleMax = 0;
 		}
@@ -1070,8 +1070,8 @@ namespace Chaos
 
 		// Calculate swing limit for the current swing axis
 		FReal SwingAngleMax = FLT_MAX;
-		const FReal Swing1Limit = JointSettings.Motion.AngularLimits[(int32)EJointAngularConstraintIndex::Swing1];
-		const FReal Swing2Limit = JointSettings.Motion.AngularLimits[(int32)EJointAngularConstraintIndex::Swing2];
+		const FReal Swing1Limit = JointSettings.AngularLimits[(int32)EJointAngularConstraintIndex::Swing1];
+		const FReal Swing2Limit = JointSettings.AngularLimits[(int32)EJointAngularConstraintIndex::Swing2];
 
 		// Circular swing limit
 		SwingAngleMax = Swing1Limit;
@@ -1148,8 +1148,8 @@ namespace Chaos
 
 		// Calculate swing limit for the current swing axis
 		FReal SwingAngleMax = FLT_MAX;
-		const FReal Swing1Limit = JointSettings.Motion.AngularLimits[(int32)EJointAngularConstraintIndex::Swing1];
-		const FReal Swing2Limit = JointSettings.Motion.AngularLimits[(int32)EJointAngularConstraintIndex::Swing2];
+		const FReal Swing1Limit = JointSettings.AngularLimits[(int32)EJointAngularConstraintIndex::Swing1];
+		const FReal Swing2Limit = JointSettings.AngularLimits[(int32)EJointAngularConstraintIndex::Swing2];
 
 		// Circular swing limit
 		SwingAngleMax = Swing1Limit;
@@ -1249,12 +1249,12 @@ namespace Chaos
 			}
 
 			FReal SwingAngleMax = FLT_MAX;
-			if (JointSettings.Motion.AngularMotionTypes[(int32)SwingConstraintIndex] == EJointMotionType::Limited)
+			if (JointSettings.AngularMotionTypes[(int32)SwingConstraintIndex] == EJointMotionType::Limited)
 			{
-				FReal SwingLimit = JointSettings.Motion.AngularLimits[(int32)SwingConstraintIndex];
+				FReal SwingLimit = JointSettings.AngularLimits[(int32)SwingConstraintIndex];
 				SwingAngleMax = SwingLimit;
 			}
-			else if (JointSettings.Motion.AngularMotionTypes[(int32)SwingConstraintIndex] == EJointMotionType::Locked)
+			else if (JointSettings.AngularMotionTypes[(int32)SwingConstraintIndex] == EJointMotionType::Locked)
 			{
 				SwingAngleMax = 0;
 			}
@@ -1346,12 +1346,12 @@ namespace Chaos
 			}
 
 			FReal SwingAngleMax = FLT_MAX;
-			if (JointSettings.Motion.AngularMotionTypes[(int32)SwingConstraintIndex] == EJointMotionType::Limited)
+			if (JointSettings.AngularMotionTypes[(int32)SwingConstraintIndex] == EJointMotionType::Limited)
 			{
-				FReal Swing1Limit = JointSettings.Motion.AngularLimits[(int32)SwingConstraintIndex];
+				FReal Swing1Limit = JointSettings.AngularLimits[(int32)SwingConstraintIndex];
 				SwingAngleMax = Swing1Limit;
 			}
-			else if (JointSettings.Motion.AngularMotionTypes[(int32)SwingConstraintIndex] == EJointMotionType::Locked)
+			else if (JointSettings.AngularMotionTypes[(int32)SwingConstraintIndex] == EJointMotionType::Locked)
 			{
 				SwingAngleMax = 0;
 			}
@@ -1419,7 +1419,7 @@ namespace Chaos
 
 		const FVec3 TwistAxis0 = R0 * TwistAxis01;
 		const FVec3 TwistAxis1 = R1 * TwistAxis01;
-		const FReal TwistAngleTarget = JointSettings.Motion.AngularDriveTargetAngles[(int32)EJointAngularConstraintIndex::Twist];
+		const FReal TwistAngleTarget = JointSettings.AngularDriveTargetAngles[(int32)EJointAngularConstraintIndex::Twist];
 		const FReal DriveStiffness = GetAngularDriveStiffness(SolverSettings, JointSettings);
 		const FReal DTwistAngle = TwistAngle - TwistAngleTarget;
 		const FReal DTwistAngle0 = DTwistAngle;
@@ -1472,8 +1472,8 @@ namespace Chaos
 		const FVec3 SwingAxis1 = SwingAxis0;
 
 		// Circular swing target (max of Swing1, Swing2 targets)
-		FReal Swing1Target = JointSettings.Motion.AngularDriveTargetAngles[(int32)EJointAngularConstraintIndex::Swing1];
-		FReal Swing2Target = JointSettings.Motion.AngularDriveTargetAngles[(int32)EJointAngularConstraintIndex::Swing2];
+		FReal Swing1Target = JointSettings.AngularDriveTargetAngles[(int32)EJointAngularConstraintIndex::Swing1];
+		FReal Swing2Target = JointSettings.AngularDriveTargetAngles[(int32)EJointAngularConstraintIndex::Swing2];
 		FReal SwingAngleTarget = FMath::Max(Swing1Target, Swing2Target);
 
 		const FReal DriveStiffness = GetAngularDriveStiffness(SolverSettings, JointSettings);
@@ -1513,7 +1513,7 @@ namespace Chaos
 		const FRotation3 R1 = Q1 * XL1.GetRotation();
 
 		// Calculate the rotation we need to apply to resolve the rotation delta
-		const FRotation3 TargetR1 = R0 * JointSettings.Motion.AngularDriveTarget;
+		const FRotation3 TargetR1 = R0 * JointSettings.AngularDriveTarget;
 		const FRotation3 DR1 = TargetR1 * R1.Inverse();
 		const FRotation3 TargetQ0 = DR1.Inverse() * Q0;
 		const FRotation3 TargetQ1 = DR1 * Q1;
diff --git a/Engine/Source/Runtime/Experimental/Chaos/Private/Chaos/PBDJointConstraints.cpp b/Engine/Source/Runtime/Experimental/Chaos/Private/Chaos/PBDJointConstraints.cpp
--- a/Engine/Source/Runtime/Experimental/Chaos/Private/Chaos/PBDJointConstraints.cpp
+++ b/Engine/Source/Runtime/Experimental/Chaos/Private/Chaos/PBDJointConstraints.cpp
@@ -76,10 +76,11 @@
 	//
 
 	
-	FPBDJointMotionSettings::FPBDJointMotionSettings()
-		: Stiffness((FReal)1)
+	FPBDJointSettings::FPBDJointSettings()
+		: Stiffness(1)
 		, LinearProjection(0)
 		, AngularProjection(0)
+		, ParentInvMassScale(1)
 		, LinearMotionTypes({ EJointMotionType::Locked, EJointMotionType::Locked, EJointMotionType::Locked })
 		, LinearLimit(FLT_MAX)
 		, AngularMotionTypes({ EJointMotionType::Free, EJointMotionType::Free, EJointMotionType::Free })
@@ -87,6 +88,7 @@
 		, bSoftLinearLimitsEnabled(false)
 		, bSoftTwistLimitsEnabled(false)
 		, bSoftSwingLimitsEnabled(false)
+		, SoftForceMode(EJointForceMode::Acceleration)
 		, SoftLinearStiffness(0)
 		, SoftLinearDamping(0)
 		, SoftTwistStiffness(0)
@@ -95,6 +97,7 @@
 		, SoftSwingDamping(0)
 		, LinearDriveTarget(FVec3(0, 0, 0))
 		, bLinearDriveEnabled(TVector<bool, 3>( false, false, false ))
+		, LinearDriveForceMode(EJointForceMode::Acceleration)
 		, LinearDriveStiffness(0)
 		, LinearDriveDamping(0)
 		, AngularDriveTarget(FRotation3::FromIdentity())
@@ -102,17 +105,18 @@
 		, bAngularSLerpDriveEnabled(false)
 		, bAngularTwistDriveEnabled(false)
 		, bAngularSwingDriveEnabled(false)
-		, bAngularAccelerationDriveMode(true)
+		, AngularDriveForceMode(EJointForceMode::Acceleration)
 		, AngularDriveStiffness(0)
 		, AngularDriveDamping(0)
 	{
 	}
 
 	
-	FPBDJointMotionSettings::FPBDJointMotionSettings(const TVector<EJointMotionType, 3>& InLinearMotionTypes, const TVector<EJointMotionType, 3>& InAngularMotionTypes)
-		: Stiffness((FReal)1)
+	FPBDJointSettings::FPBDJointSettings(const TVector<EJointMotionType, 3>& InLinearMotionTypes, const TVector<EJointMotionType, 3>& InAngularMotionTypes)
+		: Stiffness(1)
 		, LinearProjection(0)
 		, AngularProjection(0)
+		, ParentInvMassScale(1)
 		, LinearMotionTypes(InLinearMotionTypes)
 		, LinearLimit(FLT_MAX)
 		, AngularMotionTypes({ EJointMotionType::Free, EJointMotionType::Free, EJointMotionType::Free })
@@ -120,6 +124,7 @@
 		, bSoftLinearLimitsEnabled(false)
 		, bSoftTwistLimitsEnabled(false)
 		, bSoftSwingLimitsEnabled(false)
+		, SoftForceMode(EJointForceMode::Acceleration)
 		, SoftLinearStiffness(0)
 		, SoftLinearDamping(0)
 		, SoftTwistStiffness(0)
@@ -128,6 +133,7 @@
 		, SoftSwingDamping(0)
 		, LinearDriveTarget(FVec3(0))
 		, bLinearDriveEnabled(TVector<bool, 3>(false, false, false))
+		, LinearDriveForceMode(EJointForceMode::Acceleration)
 		, LinearDriveStiffness(0)
 		, LinearDriveDamping(0)
 		, AngularDriveTarget(FRotation3::FromIdentity())
@@ -135,13 +141,14 @@
 		, bAngularSLerpDriveEnabled(false)
 		, bAngularTwistDriveEnabled(false)
 		, bAngularSwingDriveEnabled(false)
-		, bAngularAccelerationDriveMode(true)
+		, AngularDriveForceMode(EJointForceMode::Acceleration)
 		, AngularDriveStiffness(0)
 		, AngularDriveDamping(0)
 	{
 	}
 
-	void FPBDJointMotionSettings::Sanitize()
+
+	void FPBDJointSettings::Sanitize()
 	{
 		// Reset limits if they won;t be used (means we don't have to check if limited/locked in a few cases).
 		// A side effect: if we enable a constraint, we need to reset the value of the limit.
@@ -164,12 +171,6 @@
 	}
 
 	
-	FPBDJointSettings::FPBDJointSettings()
-		: ConstraintFrames({ FTransform::Identity, FTransform::Identity })
-	{
-	}
-
-	
 	FPBDJointState::FPBDJointState()
 		: Level(INDEX_NONE)
 		, ParticleLevels({ INDEX_NONE, INDEX_NONE })
@@ -262,37 +263,32 @@
 
 	typename FPBDJointConstraints::FConstraintContainerHandle* FPBDJointConstraints::AddConstraint(const FParticlePair& InConstrainedParticles, const FRigidTransform3& WorldConstraintFrame)
 	{
-		FTransformPair ConstraintFrames;
-		ConstraintFrames[0] = FRigidTransform3(
+		FTransformPair JointFrames;
+		JointFrames[0] = FRigidTransform3(
 			WorldConstraintFrame.GetTranslation() - InConstrainedParticles[0]->X(),
 			WorldConstraintFrame.GetRotation() * InConstrainedParticles[0]->R().Inverse()
 			);
-		ConstraintFrames[1] = FRigidTransform3(
+		JointFrames[1] = FRigidTransform3(
 			WorldConstraintFrame.GetTranslation() - InConstrainedParticles[1]->X(),
 			WorldConstraintFrame.GetRotation() * InConstrainedParticles[1]->R().Inverse()
 			);
-		return AddConstraint(InConstrainedParticles, ConstraintFrames);
+		return AddConstraint(InConstrainedParticles, JointFrames, FPBDJointSettings());
 	}
 
 	
-	typename FPBDJointConstraints::FConstraintContainerHandle* FPBDJointConstraints::AddConstraint(const FParticlePair& InConstrainedParticles, const FTransformPair& ConstraintFrames)
+	typename FPBDJointConstraints::FConstraintContainerHandle* FPBDJointConstraints::AddConstraint(const FParticlePair& InConstrainedParticles, const FTransformPair& InConstraintFrames)
+	{
+		return AddConstraint(InConstrainedParticles, InConstraintFrames, FPBDJointSettings());
+	}
+
+	
+	typename FPBDJointConstraints::FConstraintContainerHandle* FPBDJointConstraints::AddConstraint(const FParticlePair& InConstrainedParticles, const FTransformPair& InConstraintFrames, const FPBDJointSettings& InConstraintSettings)
 	{
 		int ConstraintIndex = Handles.Num();
 		Handles.Add(HandleAllocator.AllocHandle(this, ConstraintIndex));
 		ConstraintParticles.Add(InConstrainedParticles);
-		ConstraintSettings.Add(FPBDJointSettings());
-		ConstraintSettings[ConstraintIndex].ConstraintFrames = ConstraintFrames;
-		ConstraintStates.Add(FPBDJointState());
-		return Handles.Last();
-	}
-
-	
-	typename FPBDJointConstraints::FConstraintContainerHandle* FPBDJointConstraints::AddConstraint(const FParticlePair& InConstrainedParticles, const FPBDJointSettings& InConstraintSettings)
-	{
-		int ConstraintIndex = Handles.Num();
-		Handles.Add(HandleAllocator.AllocHandle(this, ConstraintIndex));
-		ConstraintParticles.Add(InConstrainedParticles);
 		ConstraintSettings.Add(InConstraintSettings);
+		ConstraintFrames.Add(InConstraintFrames);
 		ConstraintStates.Add(FPBDJointState());
 		return Handles.Last();
 	}
@@ -311,6 +307,7 @@
 		// Swap the last constraint into the gap to keep the array packed
 		ConstraintParticles.RemoveAtSwap(ConstraintIndex);
 		ConstraintSettings.RemoveAtSwap(ConstraintIndex);
+		ConstraintFrames.RemoveAtSwap(ConstraintIndex);
 		ConstraintStates.RemoveAtSwap(ConstraintIndex);
 		Handles.RemoveAtSwap(ConstraintIndex);
 
@@ -327,6 +324,42 @@
 	}
 
 
+	void FPBDJointConstraints::SortConstraints()
+	{
+		// Sort constraints so that constraints with lower level (closer to a kinematic joint) are first
+		// @todo(ccaulfield): should probably also take islands/particle order into account
+		// @todo(ccaulfield): optimize (though isn't called very often)
+		SCOPE_CYCLE_COUNTER(STAT_Joints_Sort);
+
+		FHandles SortedHandles = Handles;
+		SortedHandles.StableSort([](const FConstraintContainerHandle& L, const FConstraintContainerHandle& R)
+			{
+				return L.GetConstraintLevel() < R.GetConstraintLevel();
+			});
+
+		TArray<FPBDJointSettings> SortedConstraintSettings;
+		TArray<FTransformPair> SortedConstraintFrames;
+		TArray<FParticlePair> SortedConstraintParticles;
+		TArray<FPBDJointState> SortedConstraintStates;
+		int32 SortedConstraintIndex = 0;
+		for (FConstraintContainerHandle* Handle : SortedHandles)
+		{
+			int32 UnsortedIndex = Handle->GetConstraintIndex();
+			SortedConstraintSettings.Add(ConstraintSettings[UnsortedIndex]);
+			SortedConstraintFrames.Add(ConstraintFrames[UnsortedIndex]);
+			SortedConstraintParticles.Add(ConstraintParticles[UnsortedIndex]);
+			SortedConstraintStates.Add(ConstraintStates[UnsortedIndex]);
+			SetConstraintIndex(Handle, SortedConstraintIndex++);
+		}
+
+		Swap(ConstraintSettings, SortedConstraintSettings);
+		Swap(ConstraintFrames, SortedConstraintFrames);
+		Swap(ConstraintParticles, SortedConstraintParticles);
+		Swap(ConstraintStates, SortedConstraintStates);
+		Swap(Handles, SortedHandles);
+	}
+
+
 	
 	void FPBDJointConstraints::SetPreApplyCallback(const FJointPreApplyCallback& Callback)
 	{
@@ -404,39 +437,6 @@
 	}
 
 
-	void FPBDJointConstraints::SortConstraints()
-	{
-		// Sort constraints so that constraints with lower level (closer to a kinematic joint) are first
-		// @todo(ccaulfield): should probably also take islands/particle order into account
-		// @todo(ccaulfield): optimize (though isn't called very often)
-		SCOPE_CYCLE_COUNTER(STAT_Joints_Sort);
-
-		FHandles SortedHandles = Handles;
-		SortedHandles.StableSort([](const FConstraintContainerHandle& L, const FConstraintContainerHandle& R)
-			{
-				return L.GetConstraintLevel() < R.GetConstraintLevel();
-			});
-
-		TArray<FPBDJointSettings> SortedConstraintSettings;
-		TArray<FParticlePair> SortedConstraintParticles;
-		TArray<FPBDJointState> SortedConstraintStates;
-		int32 SortedConstraintIndex = 0;
-		for (FConstraintContainerHandle* Handle : SortedHandles)
-		{
-			int32 UnsortedIndex = Handle->GetConstraintIndex();
-			SortedConstraintSettings.Add(ConstraintSettings[UnsortedIndex]);
-			SortedConstraintParticles.Add(ConstraintParticles[UnsortedIndex]);
-			SortedConstraintStates.Add(ConstraintStates[UnsortedIndex]);
-			SetConstraintIndex(Handle, SortedConstraintIndex++);
-		}
-
-		Swap(ConstraintSettings, SortedConstraintSettings);
-		Swap(ConstraintParticles, SortedConstraintParticles);
-		Swap(ConstraintStates, SortedConstraintStates);
-		Swap(Handles, SortedHandles);
-	}
-
-	
 	void FPBDJointConstraints::UpdatePositionBasedState(const FReal Dt)
 	{
 		if (bRequiresSort)
@@ -457,6 +457,7 @@
 			for (int32 ConstraintIndex = 0; ConstraintIndex < NumConstraints(); ++ConstraintIndex)
 			{
 				const FPBDJointSettings& JointSettings = ConstraintSettings[ConstraintIndex];
+				const FTransformPair& JointFrames = ConstraintFrames[ConstraintIndex];
 				FJointSolverGaussSeidel& Solver = ConstraintSolvers[ConstraintIndex];
 
 				int32 Index0, Index1;
@@ -474,8 +475,8 @@
 					Particle0->InvI(),
 					Particle1->InvM(),
 					Particle1->InvI(),
-					FParticleUtilities::ParticleLocalToCoMLocal(Particle0, JointSettings.ConstraintFrames[Index0]),
-					FParticleUtilities::ParticleLocalToCoMLocal(Particle1, JointSettings.ConstraintFrames[Index1]));
+					FParticleUtilities::ParticleLocalToCoMLocal(Particle0, JointFrames[Index0]),
+					FParticleUtilities::ParticleLocalToCoMLocal(Particle1, JointFrames[Index1]));
 			}
 		}
 	}
@@ -491,10 +492,8 @@
 		const FRotation3 Q0 = FParticleUtilities::GetCoMWorldRotation(Particle0);
 		const FVec3 P1 = FParticleUtilities::GetCoMWorldPosition(Particle1);
 		const FRotation3 Q1 = FParticleUtilities::GetCoMWorldRotation(Particle1);
+		const FRigidTransform3& XL0 = FParticleUtilities::ParticleLocalToCoMLocal(Particle0, ConstraintFrames[ConstraintIndex][Index0]);
+		const FRigidTransform3& XL1 = FParticleUtilities::ParticleLocalToCoMLocal(Particle1, ConstraintFrames[ConstraintIndex][Index1]);
 
-		const FPBDJointSettings& JointSettings = ConstraintSettings[ConstraintIndex];
-		const FRigidTransform3 XL0 = FParticleUtilities::ParticleLocalToCoMLocal(Particle0, JointSettings.ConstraintFrames[Index0]);
-		const FRigidTransform3 XL1 = FParticleUtilities::ParticleLocalToCoMLocal(Particle1, JointSettings.ConstraintFrames[Index1]);
-
 		OutX0 = P0 + Q0 * XL0.GetTranslation();
 		OutX1 = P1 + Q1 * XL1.GetTranslation();
@@ -723,6 +722,7 @@
 		UE_LOG(LogChaosJoint, VeryVerbose, TEXT("Solve Joint Constraint %d %s %s (dt = %f; it = %d / %d)"), ConstraintIndex, *Constraint[0]->ToString(), *Constraint[1]->ToString(), Dt, It, NumIts);
 
 		const FPBDJointSettings& JointSettings = ConstraintSettings[ConstraintIndex];
+		const FTransformPair& JointFrames = ConstraintFrames[ConstraintIndex];
 
 		int32 Index0, Index1;
 		GetConstrainedParticleIndices(ConstraintIndex, Index0, Index1);
@@ -742,8 +742,8 @@
 			Particle0->InvI(), 
 			Particle1->InvM(), 
 			Particle1->InvI(), 
-			FParticleUtilities::ParticleLocalToCoMLocal(Particle0, JointSettings.ConstraintFrames[Index0]),
-			FParticleUtilities::ParticleLocalToCoMLocal(Particle1, JointSettings.ConstraintFrames[Index1]));
+			FParticleUtilities::ParticleLocalToCoMLocal(Particle0, JointFrames[Index0]),
+			FParticleUtilities::ParticleLocalToCoMLocal(Particle1, JointFrames[Index1]));
 		
 		for (int32 PairIt = 0; PairIt < NumPairIts; ++PairIt)
 		{
diff --git a/Engine/Source/Runtime/Experimental/Chaos/Public/Chaos/PBDJointConstraintTypes.h b/Engine/Source/Runtime/Experimental/Chaos/Public/Chaos/PBDJointConstraintTypes.h
--- a/Engine/Source/Runtime/Experimental/Chaos/Public/Chaos/PBDJointConstraintTypes.h
+++ b/Engine/Source/Runtime/Experimental/Chaos/Public/Chaos/PBDJointConstraintTypes.h
@@ -24,6 +24,12 @@ namespace Chaos
 		Locked,
 	};
 
+	enum class EJointForceMode : int32
+	{
+		Acceleration,
+		Force,
+	};
+
 	/**
 	 * The order of the angular constraints (for settings held in vectors etc)
 	 */
@@ -44,13 +50,6 @@ namespace Chaos
 		Swing1 = 2,		// Swing1 Axis = Z
 	};
 
-	enum class EJointSolverPhase
-	{
-		None,
-		Apply,
-		ApplyPushOut,
-	};
-
 	struct FJointConstants
 	{
 		/** The constraint-space twist axis (X Axis) */
@@ -63,17 +62,18 @@ namespace Chaos
 		static const FVec3 Swing2Axis() { return FVec3(0, 1, 0); }
 	};
 
-	class CHAOS_API FPBDJointMotionSettings
+	class CHAOS_API FPBDJointSettings
 	{
 	public:
-		FPBDJointMotionSettings();
-		FPBDJointMotionSettings(const TVector<EJointMotionType, 3>& InLinearMotionTypes, const TVector<EJointMotionType, 3>& InAngularMotionTypes);
+		FPBDJointSettings();
+		FPBDJointSettings(const TVector<EJointMotionType, 3>& InLinearMotionTypes, const TVector<EJointMotionType, 3>& InAngularMotionTypes);
 
 		void Sanitize();
 
 		FReal Stiffness;
 		FReal LinearProjection;
 		FReal AngularProjection;
+		FReal ParentInvMassScale;
 
 		TVector<EJointMotionType, 3> LinearMotionTypes;
 		FReal LinearLimit;
@@ -84,6 +84,7 @@ namespace Chaos
 		bool bSoftLinearLimitsEnabled;
 		bool bSoftTwistLimitsEnabled;
 		bool bSoftSwingLimitsEnabled;
+		EJointForceMode SoftForceMode;
 		FReal SoftLinearStiffness;
 		FReal SoftLinearDamping;
 		FReal SoftTwistStiffness;
@@ -93,6 +94,7 @@ namespace Chaos
 
 		FVec3 LinearDriveTarget;
 		TVector<bool, 3> bLinearDriveEnabled;
+		EJointForceMode LinearDriveForceMode;
 		FReal LinearDriveStiffness;
 		FReal LinearDriveDamping;
 
@@ -103,26 +105,11 @@ namespace Chaos
 		bool bAngularSLerpDriveEnabled;
 		bool bAngularTwistDriveEnabled;
 		bool bAngularSwingDriveEnabled;
-		bool bAngularAccelerationDriveMode;
+		EJointForceMode AngularDriveForceMode;
 		FReal AngularDriveStiffness;
 		FReal AngularDriveDamping;
 	};
 
-
-	class CHAOS_API FPBDJointSettings
-	{
-	public:
-		using FTransformPair = TVector<FRigidTransform3, 2>;
-
-		FPBDJointSettings();
-
-		// Particle-relative joint axes and positions
-		FTransformPair ConstraintFrames;
-
-		// How the constraint is allowed to move
-		FPBDJointMotionSettings Motion;
-	};
-
 	class CHAOS_API FPBDJointSolverSettings
 	{
 	public:
diff --git a/Engine/Source/Runtime/Experimental/Chaos/Public/Chaos/PBDJointConstraints.h b/Engine/Source/Runtime/Experimental/Chaos/Public/Chaos/PBDJointConstraints.h
--- a/Engine/Source/Runtime/Experimental/Chaos/Public/Chaos/PBDJointConstraints.h
+++ b/Engine/Source/Runtime/Experimental/Chaos/Public/Chaos/PBDJointConstraints.h
@@ -83,7 +83,7 @@ namespace Chaos
 		 */
 		FConstraintContainerHandle* AddConstraint(const FParticlePair& InConstrainedParticles, const FRigidTransform3& WorldConstraintFrame);
 		FConstraintContainerHandle* AddConstraint(const FParticlePair& InConstrainedParticles, const FTransformPair& ConstraintFrames);
-		FConstraintContainerHandle* AddConstraint(const FParticlePair& InConstrainedParticles, const FPBDJointSettings& InConstraintSettings);
+		FConstraintContainerHandle* AddConstraint(const FParticlePair& InConstrainedParticles, const FTransformPair& ConstraintFrames, const FPBDJointSettings& InConstraintSettings);
 
 		/**
 		 * Remove the specified constraint.
@@ -170,6 +170,7 @@ namespace Chaos
 		FPBDJointSolverSettings Settings;
 
 		TArray<FPBDJointSettings> ConstraintSettings;
+		TArray<FTransformPair> ConstraintFrames;
 		TArray<FParticlePair> ConstraintParticles;
 		TArray<FPBDJointState> ConstraintStates;
 
