CIS fix (Immediate{hysicsJointHandle)

diff --git a/Engine/Source/Runtime/Engine/Private/PhysicsEngine/ImmediatePhysics/ImmediatePhysicsChaos/ImmediatePhysicsJointHandle_Chaos.cpp b/Engine/Source/Runtime/Engine/Private/PhysicsEngine/ImmediatePhysics/ImmediatePhysicsChaos/ImmediatePhysicsJointHandle_Chaos.cpp
--- a/Engine/Source/Runtime/Engine/Private/PhysicsEngine/ImmediatePhysics/ImmediatePhysicsChaos/ImmediatePhysicsJointHandle_Chaos.cpp
+++ b/Engine/Source/Runtime/Engine/Private/PhysicsEngine/ImmediatePhysics/ImmediatePhysicsChaos/ImmediatePhysicsJointHandle_Chaos.cpp
@@ -131,7 +131,7 @@ namespace ImmediatePhysics_Chaos
 		ConstraintSettings.AngularDriveDamping = ChaosImmediate_DriveDampingScale * Profile.AngularDrive.TwistDrive.Damping;
 		ConstraintSettings.AngularDriveForceMode = EJointForceMode::Acceleration;
 
-		ConstraintSettings.Sanitize(Actor1->GetMass(), Actor2->GetMass());
+		ConstraintSettings.Sanitize();
 
 		ConstraintHandle = Constraints->AddConstraint({ Actor1->ParticleHandle, Actor2->ParticleHandle }, ConstraintFrames, ConstraintSettings);
 	}
