Chaos
- Joint parity (disable soft joints if Limit < 0.1)
- Better projection

diff --git a/Engine/Source/Runtime/Engine/Private/PhysicsEngine/ImmediatePhysics/ImmediatePhysicsChaos/ImmediatePhysicsJointHandle_Chaos.cpp b/Engine/Source/Runtime/Engine/Private/PhysicsEngine/ImmediatePhysics/ImmediatePhysicsChaos/ImmediatePhysicsJointHandle_Chaos.cpp
--- a/Engine/Source/Runtime/Engine/Private/PhysicsEngine/ImmediatePhysics/ImmediatePhysicsChaos/ImmediatePhysicsJointHandle_Chaos.cpp
+++ b/Engine/Source/Runtime/Engine/Private/PhysicsEngine/ImmediatePhysics/ImmediatePhysicsChaos/ImmediatePhysicsJointHandle_Chaos.cpp
@@ -33,6 +33,11 @@ namespace ImmediatePhysics_Chaos
 	FAutoConsoleVariableRef CVarSoftLinearStiffnessScale(TEXT("p.Chaos.ImmPhys.SoftLinearStiffnessScale"), ChaosImmediate_SoftLinearStiffnessScale, TEXT("Conversion factor for soft-joint stiffness."));
 	FAutoConsoleVariableRef CVarSoftLinearDampingScale(TEXT("p.Chaos.ImmPhys.SoftLinearDampingScale"), ChaosImmediate_SoftLinearDampingScale, TEXT("Conversion factor for soft-joint damping."));
 
+	int ChaosImmediate_SoftLinearForceMode = 0;
+	int ChaosImmediate_SoftAngularForceMode = 0;
+	FAutoConsoleVariableRef CVarSoftLinearForceMode(TEXT("p.Chaos.ImmPhys.SoftLinearForceMode"), ChaosImmediate_SoftLinearForceMode, TEXT("Soft Linear constraint force mode (0: Acceleration; 1: Force"));
+	FAutoConsoleVariableRef CVarSoftAngularForceMode(TEXT("p.Chaos.ImmPhys.SoftAngularForceMode"), ChaosImmediate_SoftAngularForceMode, TEXT("Soft Angular constraint force mode (0: Acceleration; 1: Force"));
+
 	// This value is from PhysX UE - in Chaos we don't really need Soft constraints with stiffness greater than about 1000
 	// (In PhysX, Soft constraints were implicitly integrated, and essentially equivalent to XPBD, which for large stiffness is just PBD).
 	float ChaosImmediate_SoftAngularStiffnessScale = 100000.0f;
@@ -41,7 +46,7 @@ namespace ImmediatePhysics_Chaos
 	FAutoConsoleVariableRef CVarSoftAngularDampingScale(TEXT("p.Chaos.ImmPhys.SoftAngularDampingScale"), ChaosImmediate_SoftAngularDampingScale, TEXT("Conversion factor for soft-joint damping."));
 
 	float ChaosImmediate_JointMinLinearProjection = 0.0f;
-	float ChaosImmediate_JointMaxLinearProjection = 0.5f;
+	float ChaosImmediate_JointMaxLinearProjection = 0.8f;
 	FAutoConsoleVariableRef CVarJointMinLinearProjection(TEXT("p.Chaos.ImmPhys.JointMinLinearProjection"), ChaosImmediate_JointMinLinearProjection, TEXT("Joint min projection (for joints with projection disabled)."));
 	FAutoConsoleVariableRef CVarJointMaxLinearProjection(TEXT("p.Chaos.ImmPhys.JointMaxLinearProjection"), ChaosImmediate_JointMaxLinearProjection, TEXT("Joint max projection (for joints with projection enabled)."));
 
@@ -95,8 +100,8 @@ namespace ImmediatePhysics_Chaos
 		ConstraintSettings.SoftTwistDamping = ChaosImmediate_SoftAngularDampingScale * ConstraintInstance->GetSoftTwistLimitDamping();
 		ConstraintSettings.SoftSwingStiffness = ChaosImmediate_SoftAngularStiffnessScale * ConstraintInstance->GetSoftSwingLimitStiffness();
 		ConstraintSettings.SoftSwingDamping = ChaosImmediate_SoftAngularDampingScale * ConstraintInstance->GetSoftSwingLimitDamping();
-		ConstraintSettings.LinearSoftForceMode = EJointForceMode::Acceleration;
-		ConstraintSettings.AngularSoftForceMode = EJointForceMode::Acceleration;
+		ConstraintSettings.LinearSoftForceMode = (ChaosImmediate_SoftLinearForceMode == 0)? EJointForceMode::Acceleration : EJointForceMode::Force;
+		ConstraintSettings.AngularSoftForceMode = (ChaosImmediate_SoftAngularForceMode == 0)? EJointForceMode::Acceleration : EJointForceMode::Force;
 
 		ConstraintSettings.LinearDriveTarget = Profile.LinearDrive.PositionTarget;
 		ConstraintSettings.bLinearPositionDriveEnabled[0] = Profile.LinearDrive.XDrive.bEnablePositionDrive;
@@ -131,6 +136,18 @@ namespace ImmediatePhysics_Chaos
 		ConstraintSettings.AngularDriveDamping = ChaosImmediate_DriveDampingScale * Profile.AngularDrive.TwistDrive.Damping;
 		ConstraintSettings.AngularDriveForceMode = EJointForceMode::Acceleration;
 
+		// UE Disables Soft Limits when the Limit is less than some threshold. This is not necessary in Chaos but for now we also do it for parity's sake (See FLinearConstraint::UpdateLinearLimit_AssumesLocked).
+		if (ConstraintSettings.LinearLimit < RB_MinSizeToLockDOF)
+		{
+			for (int32 AxisIndex = 0; AxisIndex < 3; ++AxisIndex)
+			{
+				if (ConstraintSettings.LinearMotionTypes[AxisIndex] == EJointMotionType::Limited)
+				{
+					ConstraintSettings.LinearMotionTypes[AxisIndex] = EJointMotionType::Locked;
+				}
+			}
+		}
+
 		ConstraintSettings.Sanitize();
 
 		ConstraintHandle = Constraints->AddConstraint({ Actor1->ParticleHandle, Actor2->ParticleHandle }, ConstraintFrames, ConstraintSettings);
diff --git a/Engine/Source/Runtime/Experimental/Chaos/Private/Chaos/Joint/PBDJointSolverGaussSeidel.cpp b/Engine/Source/Runtime/Experimental/Chaos/Private/Chaos/Joint/PBDJointSolverGaussSeidel.cpp
--- a/Engine/Source/Runtime/Experimental/Chaos/Private/Chaos/Joint/PBDJointSolverGaussSeidel.cpp
+++ b/Engine/Source/Runtime/Experimental/Chaos/Private/Chaos/Joint/PBDJointSolverGaussSeidel.cpp
@@ -382,8 +382,7 @@ namespace Chaos
 	{
 		InitAccumulatedDeltas();
 
-		// With the parent body set to infinite mass...
-		ApplyRotationProjection(Dt, SolverSettings, JointSettings);
+		//ApplyRotationProjection(Dt, SolverSettings, JointSettings);
 
 		ApplyPositionProjection(Dt, SolverSettings, JointSettings);
 
@@ -1663,32 +1662,48 @@ namespace Chaos
 		const FPBDJointSettings& JointSettings)
 	{
 		// Apply a position correction with the parent body set to infinite mass, then correct the velocity.
+		// Note: Velocity correction does not set parent to infinite mass (if we do, then the velocities
+		// get amplified down the chain).
 		FVec3 CX = FPBDJointUtilities::GetLimitedPositionError(JointSettings, Rs[0], Xs[1] - Xs[0]);
 		const FReal CXLen = CX.Size();
 		if (CXLen > KINDA_SMALL_NUMBER)
 		{
+			FReal LinearProjection = FPBDJointUtilities::GetLinearProjection(SolverSettings, JointSettings);
+			const FReal ParentMassScale = FMath::Max(0.0f, 1.0f - LinearProjection);
+
 			const FVec3 CXDir = CX / CXLen;
 			const FVec3 V0 = Vs[0] + FVec3::CrossProduct(Ws[0], Xs[0] - Ps[0]);
 			const FVec3 V1 = Vs[1] + FVec3::CrossProduct(Ws[1], Xs[1] - Ps[1]);
 			FVec3 CV = FVec3::DotProduct(V1 - V0, CXDir) * CXDir;
 
-			FMatrix33 InvI1 = Utilities::ComputeWorldSpaceInertia(Qs[1], InvILs[1]);
-			FMatrix33 M1 = Utilities::ComputeJointFactorMatrix(Xs[1] - Ps[1], InvI1, InvMs[1]);
-			FMatrix33 MI = M1.Inverse();
-
-			const FVec3 DX = Utilities::Multiply(MI, CX);
-			const FVec3 DV = Utilities::Multiply(MI, CV);
-
-			const FVec3 DP1 = -InvMs[1] * DX;
-			const FVec3 DR1 = Utilities::Multiply(InvI1, FVec3::CrossProduct(Xs[1] - Ps[1], -DX));
-			const FVec3 DV1 = -InvMs[1] * DV;
-			const FVec3 DW1 = Utilities::Multiply(InvI1, FVec3::CrossProduct(Xs[1] - Ps[1], -DV));
-
-			FReal LinearProjection = FPBDJointUtilities::GetLinearProjection(SolverSettings, JointSettings);
-			ApplyPositionDelta(1, LinearProjection, DP1);
-			ApplyRotationDelta(1, LinearProjection, DR1);
-			ApplyVelocityDelta(1, LinearProjection, DV1, DW1);
-			UpdateDerivedState(1);
+			const FReal IM0 = ParentMassScale * InvMs[0];
+			const FReal IM1 = InvMs[1];
+			const FMatrix33 IIL0 = ParentMassScale * InvILs[0];
+			const FMatrix33& IIL1 = InvILs[1];
+			FMatrix33 II0 = Utilities::ComputeWorldSpaceInertia(Qs[0], IIL0);
+			FMatrix33 II1 = Utilities::ComputeWorldSpaceInertia(Qs[1], IIL1);
+			FMatrix33 J0 = (IM0 > 0) ? Utilities::ComputeJointFactorMatrix(Xs[0] - Ps[0], II0, IM0) : FMatrix33(0, 0, 0);
+			FMatrix33 J1 = Utilities::ComputeJointFactorMatrix(Xs[1] - Ps[1], II1, IM1);
+			FMatrix33 IJ = (J0 + J1).Inverse();
+
+			const FVec3 DX = Utilities::Multiply(IJ, CX);
+			const FVec3 DV = Utilities::Multiply(IJ, CV);
+
+			const FVec3 DP0 = IM0 * DX;
+			const FVec3 DP1 = -IM1 * DX;
+			const FVec3 DR0 = Utilities::Multiply(II0, FVec3::CrossProduct(Xs[0] - Ps[0], DX));
+			const FVec3 DR1 = Utilities::Multiply(II1, FVec3::CrossProduct(Xs[1] - Ps[1], -DX));
+
+			const FVec3 DV0 = IM0 * DV;
+			const FVec3 DV1 = -IM1 * DV;
+			const FVec3 DW0 = Utilities::Multiply(II0, FVec3::CrossProduct(Xs[0] - Ps[0], DV));
+			const FVec3 DW1 = Utilities::Multiply(II1, FVec3::CrossProduct(Xs[1] - Ps[1], -DV));
+
+			const FReal Stiffness = FPBDJointUtilities::GetLinearStiffness(SolverSettings, JointSettings);
+			ApplyPositionDelta(Stiffness, DP0, DP1);
+			ApplyRotationDelta(Stiffness, DR0, DR1);
+			ApplyVelocityDelta(Stiffness, DV0, DW0, DV1, DW1);
+			UpdateDerivedState();
 		}
 	}
 }
diff --git a/Engine/Source/Runtime/Experimental/Chaos/Private/Chaos/PBDJointConstraints.cpp b/Engine/Source/Runtime/Experimental/Chaos/Private/Chaos/PBDJointConstraints.cpp
--- a/Engine/Source/Runtime/Experimental/Chaos/Private/Chaos/PBDJointConstraints.cpp
+++ b/Engine/Source/Runtime/Experimental/Chaos/Private/Chaos/PBDJointConstraints.cpp
@@ -138,21 +138,6 @@ namespace Chaos
 		{
 			AngularLimits[(int32)EJointAngularConstraintIndex::Swing2] = 0;
 		}
-
-		// Don't use soft constraints if the stiffness is very high
-		// (Should really use mass if Force Mode is enabled - which is isn't for anything atm).
-		if (bSoftLinearLimitsEnabled && SoftLinearStiffness > 1000)
-		{
-			bSoftLinearLimitsEnabled = false;
-		}
-		if (bSoftTwistLimitsEnabled && SoftTwistStiffness > 1000)
-		{
-			bSoftTwistLimitsEnabled = false;
-		}
-		if (bSoftSwingLimitsEnabled && SoftSwingStiffness > 1000)
-		{
-			bSoftSwingLimitsEnabled = false;
-		}
 	}
 
 	
