#UE4 Allow the warning about a low number of SCWs that are spawned during cooking to be configurable by ini

diff --git a/Engine/Config/BaseEngine.ini b/Engine/Config/BaseEngine.ini
--- a/Engine/Config/BaseEngine.ini
+++ b/Engine/Config/BaseEngine.ini
@@ -1203,6 +1203,7 @@ bUseVirtualCores = False
 ; CookerMemoryUsedInGB = 49
 ; MemoryToLeaveForTheOSInGB = 3
 ; MemoryUsedPerSCWProcessInGB = 0.4
+; MinSCWsToSpawnBeforeWarning = 8
 
 [LogFiles]
 PurgeLogsDays=5
diff --git a/Engine/Source/Runtime/Engine/Private/ShaderCompiler/ShaderCompiler.cpp b/Engine/Source/Runtime/Engine/Private/ShaderCompiler/ShaderCompiler.cpp
--- a/Engine/Source/Runtime/Engine/Private/ShaderCompiler/ShaderCompiler.cpp
+++ b/Engine/Source/Runtime/Engine/Private/ShaderCompiler/ShaderCompiler.cpp
@@ -1860,10 +1860,12 @@ FShaderCompilingManager::FShaderCompilingManager() :
 		float CookerMemoryUsedInGB = 0.0f;
 		float MemoryToLeaveForTheOSInGB = 0.0f;
 		float MemoryUsedPerSCWProcessInGB = 0.0f;
+		int32 MinSCWsToSpawnBeforeWarning = 8; // optional, default to 8
 		bool bFoundEntries = true;
 		bFoundEntries = bFoundEntries && GConfig->GetFloat(TEXT("DevOptions.Shaders"), TEXT("CookerMemoryUsedInGB"), CookerMemoryUsedInGB, GEngineIni);
 		bFoundEntries = bFoundEntries && GConfig->GetFloat(TEXT("DevOptions.Shaders"), TEXT("MemoryToLeaveForTheOSInGB"), MemoryToLeaveForTheOSInGB, GEngineIni);
 		bFoundEntries = bFoundEntries && GConfig->GetFloat(TEXT("DevOptions.Shaders"), TEXT("MemoryUsedPerSCWProcessInGB"), MemoryUsedPerSCWProcessInGB, GEngineIni);
+		GConfig->GetInt(TEXT("DevOptions.Shaders"), TEXT("MinSCWsToSpawnBeforeWarning"), MinSCWsToSpawnBeforeWarning, GEngineIni);
 		if (bFoundEntries)
 		{
 			uint32 PhysicalGBRam = FPlatformMemory::GetPhysicalGBRam();
@@ -1892,7 +1894,7 @@ FShaderCompilingManager::FShaderCompilingManager() :
 				GConfig->GetBool(TEXT("DevOptions.Shaders"), TEXT("bUseVirtualCores"), bUseVirtualCores, GEngineIni);
 				uint32 MaxNumCoresToUse = bUseVirtualCores ? NumVirtualCores : FPlatformMisc::NumberOfCores();
 				NumShaderCompilingThreads = FMath::Clamp<uint32>(NumShaderCompilingThreads, 1, MaxNumCoresToUse - 1);
-				if (NumShaderCompilingThreads < 8)
+				if (NumShaderCompilingThreads < static_cast<uint32>(MinSCWsToSpawnBeforeWarning))
 				{
 					UE_LOG(LogShaderCompilers, Warning, TEXT("Only %d SCWs will be spawned, which will result in longer shader compile times."), NumShaderCompilingThreads);
 				}
