Chaos
- adjustable post-rotation constraint position fixup amount (fixes energy accumulation)

diff --git a/Engine/Source/Runtime/Engine/Private/PhysicsEngine/ImmediatePhysics/ImmediatePhysicsChaos/ImmediatePhysicsSimulation_Chaos.cpp b/Engine/Source/Runtime/Engine/Private/PhysicsEngine/ImmediatePhysics/ImmediatePhysicsChaos/ImmediatePhysicsSimulation_Chaos.cpp
--- a/Engine/Source/Runtime/Engine/Private/PhysicsEngine/ImmediatePhysics/ImmediatePhysicsChaos/ImmediatePhysicsSimulation_Chaos.cpp
+++ b/Engine/Source/Runtime/Engine/Private/PhysicsEngine/ImmediatePhysics/ImmediatePhysicsChaos/ImmediatePhysicsSimulation_Chaos.cpp
@@ -44,7 +44,7 @@ FAutoConsoleVariableRef CVarChaosImmPhysCollisionPushOutPairIterations(TEXT("p.C
 FAutoConsoleVariableRef CVarChaosImmPhysCollisionPriority(TEXT("p.Chaos.ImmPhys.Collision.Priority"), ChaosImmediate_Collision_Priority, TEXT("Set the Collision constraint sort order (Joints have priority 0)"));
 FAutoConsoleVariableRef CVarChaosImmPhysThickness(TEXT("p.Chaos.ImmPhys.Collision.Thickness"), ChaosImmediate_Collision_Thickness, TEXT("ChaosImmediateThickness"));
 
-int32 ChaosImmediate_Joint_PairIterations = 1;
+int32 ChaosImmediate_Joint_PairIterations = 2;
 int32 ChaosImmediate_Joint_PushOutPairIterations = 1;
 float ChaosImmediate_Joint_SwingTwistAngleTolerance = 1.0e-6f;
 int32 ChaosImmediate_Joint_EnableTwistLimits = 1;
@@ -64,6 +64,7 @@ float ChaosImmediate_Joint_AngularDriveStiffness = 0.0f;
 float ChaosImmediate_Joint_AngularDriveDamping = 0.0f;
 float ChaosImmediate_Joint_MinParentMassRatio = 0.2f;
 float ChaosImmediate_Joint_MaxInertiaRatio = 5.0f;
+float ChaosImmediate_Joint_AngularPositionCorrection = 0.7f;
 FAutoConsoleVariableRef CVarChaosImmPhysPairIterations(TEXT("p.Chaos.ImmPhys.Joint.PairIterations"), ChaosImmediate_Joint_PairIterations, TEXT("PairIterations."));
 FAutoConsoleVariableRef CVarChaosImmPhysPushOutPairIterations(TEXT("p.Chaos.ImmPhys.Joint.PushOutPairIterations"), ChaosImmediate_Joint_PushOutPairIterations, TEXT("PushOutPairIterations."));
 FAutoConsoleVariableRef CVarChaosImmPhysSwingTwistAngleTolerance(TEXT("p.Chaos.ImmPhys.Joint.SwingTwistAngleTolerance"), ChaosImmediate_Joint_SwingTwistAngleTolerance, TEXT("SwingTwistAngleTolerance."));
@@ -84,6 +85,7 @@ FAutoConsoleVariableRef CVarChaosImmPhysAngularDriveStiffness(TEXT("p.Chaos.ImmP
 FAutoConsoleVariableRef CVarChaosImmPhysAngularDriveDamping(TEXT("p.Chaos.ImmPhys.Joint.AngularDriveDamping"), ChaosImmediate_Joint_AngularDriveDamping, TEXT("6Dof joint drive damping override (if > 0)."));
 FAutoConsoleVariableRef CVarChaosImmPhysMinParentMassRatio(TEXT("p.Chaos.ImmPhys.Joint.MinParentMassRatio"), ChaosImmediate_Joint_MinParentMassRatio, TEXT("6Dof joint MinParentMassRatio (if > 0)"));
 FAutoConsoleVariableRef CVarChaosImmPhysMaxInertiaRatio(TEXT("p.Chaos.ImmPhys.Joint.MaxInertiaRatio"), ChaosImmediate_Joint_MaxInertiaRatio, TEXT("6Dof joint MaxInertiaRatio (if > 0)"));
+FAutoConsoleVariableRef CVarChaosImmPhysAngularPositionCorrection(TEXT("p.Chaos.ImmPhys.Joint.AngularPositionCorrection"), ChaosImmediate_Joint_AngularPositionCorrection, TEXT("6Dof joint post-rotation constraint position correction amount [0-1]"));
 
 float ChaosImmediate_Material_Resitution = 0.0f;
 float ChaosImmediate_Material_Friction = 0.5f;
@@ -510,6 +512,7 @@ namespace ImmediatePhysics_Chaos
 			JointsSettings.SwingTwistAngleTolerance = ChaosImmediate_Joint_SwingTwistAngleTolerance;
 			JointsSettings.MinParentMassRatio = ChaosImmediate_Joint_MinParentMassRatio;
 			JointsSettings.MaxInertiaRatio = ChaosImmediate_Joint_MaxInertiaRatio;
+			JointsSettings.AngularConstraintPositionCorrection = ChaosImmediate_Joint_AngularPositionCorrection;
 			JointsSettings.bEnableTwistLimits = ChaosImmediate_Joint_EnableTwistLimits != 0;
 			JointsSettings.bEnableSwingLimits = ChaosImmediate_Joint_EnableSwingLimits != 0;
 			JointsSettings.bEnableDrives = ChaosImmediate_Joint_EnableDrives != 0;
diff --git a/Engine/Source/Runtime/Experimental/Chaos/Private/Chaos/Joint/PBDJointSolverGaussSeidel.cpp b/Engine/Source/Runtime/Experimental/Chaos/Private/Chaos/Joint/PBDJointSolverGaussSeidel.cpp
--- a/Engine/Source/Runtime/Experimental/Chaos/Private/Chaos/Joint/PBDJointSolverGaussSeidel.cpp
+++ b/Engine/Source/Runtime/Experimental/Chaos/Private/Chaos/Joint/PBDJointSolverGaussSeidel.cpp
@@ -306,9 +306,7 @@ namespace Chaos
 		TwistDriveLambda = (FReal)0;
 		SwingDriveLambda = (FReal)0;
 
-		// Disable the angular limit hardness improvement if linear limits are set up
-		// @todo(ccaulfield): fix angular constraint position correction in ApplyRotationCorrection and ApplyRotationCorrectionSoft
-		bEnableAngularConstraintPositionCorrection = ((JointSettings.LinearMotionTypes[0] == EJointMotionType::Locked) && (JointSettings.LinearMotionTypes[1] == EJointMotionType::Locked) && (JointSettings.LinearMotionTypes[2] == EJointMotionType::Locked));
+		AngularPositionCorrection = FPBDJointUtilities::GetAngularPositionCorrection(SolverSettings, JointSettings);
 	}
 
 
@@ -947,14 +945,14 @@ namespace Chaos
 		//const FVec3 DR1 = Axis1 * -(Angle * II1 / (II0 + II1));
 
 		// @todo(ccaulfield): this position correction needs to have components in direction of inactive position constraints removed
-		if (bEnableAngularConstraintPositionCorrection)
+		if (AngularPositionCorrection > 0)
 		{
 			const FReal IM0 = InvMs[0];
 			const FReal IM1 = InvMs[1];
 			const FVec3 DX = FVec3::CrossProduct(DR0, Xs[0] - Ps[0]) + FVec3::CrossProduct(DR1, Xs[1] - Ps[1]);
 			const FVec3 DP0 = DX * (IM0 / (IM0 + IM1));
 			const FVec3 DP1 = DX * (-IM1 / (IM0 + IM1));
-			ApplyPositionDelta(Stiffness, DP0, DP1);
+			ApplyPositionDelta(Stiffness * AngularPositionCorrection, DP0, DP1);
 		}
 
 		ApplyRotationDelta(Stiffness, DR0, DR1);
@@ -1033,14 +1031,14 @@ namespace Chaos
 		const FVec3 DR1 = Axis1 * -(DLambda * II1);
 
 		// @todo(ccaulfield): this position correction needs to have components in direction of inactive position constraints removed
-		if (bEnableAngularConstraintPositionCorrection)
+		if (AngularPositionCorrection > 0)
 		{
 			const FReal IM0 = InvMs[0];
 			const FReal IM1 = InvMs[1];
 			const FVec3 DX = FVec3::CrossProduct(DR0, Xs[0] - Ps[0]) + FVec3::CrossProduct(DR1, Xs[1] - Ps[1]);
 			const FVec3 DP0 = DX * (IM0 / (IM0 + IM1));
 			const FVec3 DP1 = DX * (-IM1 / (IM0 + IM1));
-			ApplyPositionDelta(1.0f, DP0, DP1);
+			ApplyPositionDelta(AngularPositionCorrection, DP0, DP1);
 		}
 
 		Lambda += DLambda;
diff --git a/Engine/Source/Runtime/Experimental/Chaos/Private/Chaos/PBDJointConstraintUtilities.cpp b/Engine/Source/Runtime/Experimental/Chaos/Private/Chaos/PBDJointConstraintUtilities.cpp
--- a/Engine/Source/Runtime/Experimental/Chaos/Private/Chaos/PBDJointConstraintUtilities.cpp
+++ b/Engine/Source/Runtime/Experimental/Chaos/Private/Chaos/PBDJointConstraintUtilities.cpp
@@ -194,6 +194,20 @@ namespace Chaos
 		return JointSettings.AngularDriveForceMode == EJointForceMode::Acceleration;
 	}
 
+	FReal FPBDJointUtilities::GetAngularPositionCorrection(
+		const FPBDJointSolverSettings& SolverSettings,
+		const FPBDJointSettings& JointSettings)
+	{
+		// Disable the angular limit hardness improvement if linear limits are set up
+		// @todo(ccaulfield): fix angular constraint position correction in ApplyRotationCorrection and ApplyRotationCorrectionSoft
+		bool bPositionCorrectionEnabled = ((JointSettings.LinearMotionTypes[0] == EJointMotionType::Locked) && (JointSettings.LinearMotionTypes[1] == EJointMotionType::Locked) && (JointSettings.LinearMotionTypes[2] == EJointMotionType::Locked));
+		if (bPositionCorrectionEnabled)
+		{
+			return SolverSettings.AngularConstraintPositionCorrection;
+		}
+		return 0.0f;
+	}
+
 
 	FVec3 FPBDJointUtilities::ConditionInertia(const FVec3& InI, const FReal MaxRatio)
 	{
diff --git a/Engine/Source/Runtime/Experimental/Chaos/Private/Chaos/PBDJointConstraints.cpp b/Engine/Source/Runtime/Experimental/Chaos/Private/Chaos/PBDJointConstraints.cpp
--- a/Engine/Source/Runtime/Experimental/Chaos/Private/Chaos/PBDJointConstraints.cpp
+++ b/Engine/Source/Runtime/Experimental/Chaos/Private/Chaos/PBDJointConstraints.cpp
@@ -158,6 +158,7 @@ namespace Chaos
 		, SwingTwistAngleTolerance((FReal)1.0e-6)
 		, MinParentMassRatio(0)
 		, MaxInertiaRatio(0)
+		, AngularConstraintPositionCorrection(0)
 		, bEnableTwistLimits(true)
 		, bEnableSwingLimits(true)
 		, bEnableDrives(true)
diff --git a/Engine/Source/Runtime/Experimental/Chaos/Public/Chaos/Joint/PBDJointSolverGaussSeidel.h b/Engine/Source/Runtime/Experimental/Chaos/Public/Chaos/Joint/PBDJointSolverGaussSeidel.h
--- a/Engine/Source/Runtime/Experimental/Chaos/Public/Chaos/Joint/PBDJointSolverGaussSeidel.h
+++ b/Engine/Source/Runtime/Experimental/Chaos/Public/Chaos/Joint/PBDJointSolverGaussSeidel.h
@@ -342,7 +342,8 @@ namespace Chaos
 		FReal TwistDriveLambda;
 		FReal SwingDriveLambda;
 
-		bool bEnableAngularConstraintPositionCorrection;
+		// Post-angular constraint position fixup (to reduce iterations required for stiff angular constraints)
+		FReal AngularPositionCorrection;
 	};
 
 }
\ No newline at end of file
diff --git a/Engine/Source/Runtime/Experimental/Chaos/Public/Chaos/PBDJointConstraintTypes.h b/Engine/Source/Runtime/Experimental/Chaos/Public/Chaos/PBDJointConstraintTypes.h
--- a/Engine/Source/Runtime/Experimental/Chaos/Public/Chaos/PBDJointConstraintTypes.h
+++ b/Engine/Source/Runtime/Experimental/Chaos/Public/Chaos/PBDJointConstraintTypes.h
@@ -131,6 +131,9 @@ namespace Chaos
 		FReal MinParentMassRatio;
 		FReal MaxInertiaRatio;
 
+		// Angular stiffness
+		FReal AngularConstraintPositionCorrection;
+
 		// @todo(ccaulfield): remove these TEMP overrides for testing
 		bool bEnableTwistLimits;
 		bool bEnableSwingLimits;
diff --git a/Engine/Source/Runtime/Experimental/Chaos/Public/Chaos/PBDJointConstraintUtilities.h b/Engine/Source/Runtime/Experimental/Chaos/Public/Chaos/PBDJointConstraintUtilities.h
--- a/Engine/Source/Runtime/Experimental/Chaos/Public/Chaos/PBDJointConstraintUtilities.h
+++ b/Engine/Source/Runtime/Experimental/Chaos/Public/Chaos/PBDJointConstraintUtilities.h
@@ -152,6 +152,10 @@ namespace Chaos
 			const FPBDJointSolverSettings& SolverSettings,
 			const FPBDJointSettings& JointSettings);
 
+		static FReal GetAngularPositionCorrection(
+			const FPBDJointSolverSettings& SolverSettings,
+			const FPBDJointSettings& JointSettings);
+
 		static FVec3 GetSphereLimitedPositionError(const FVec3& CX, const FReal Radius);
 		static FVec3 GetCylinderLimitedPositionError(const FVec3& CX, const FVec3& Axis, const FReal Limit, const EJointMotionType AxisMotion);
 		static FVec3 GetLineLimitedPositionError(const FVec3& CX, const FVec3& Axis, const FReal Limit, const EJointMotionType AxisMotion);
