* Added parameter to GetActorBounds / GetComponentsBoundingBox to be able to inspect sub-actors recursively
* Added generic ForEachComponent utilitary function on actor and merged code from GetComponents  / GetActorBounds / GetComponentsBoundingBox / ... to benefit from it (avoids building an array of components in place). Modifying the list of components in the functor is of course prohibited (like in c++ range-based iterators), otherwise, users can still use GetComponents to retrieve components in a local list and perform changes afterwards.

diff --git a/Engine/Source/Runtime/AIModule/Classes/Navigation/NavLinkProxy.h b/Engine/Source/Runtime/AIModule/Classes/Navigation/NavLinkProxy.h
--- a/Engine/Source/Runtime/AIModule/Classes/Navigation/NavLinkProxy.h
+++ b/Engine/Source/Runtime/AIModule/Classes/Navigation/NavLinkProxy.h
@@ -81,7 +81,7 @@ protected:
 public:
 #endif // ENABLE_VISUAL_LOG
 
-	virtual FBox GetComponentsBoundingBox(bool bNonColliding = false) const override;
+	virtual FBox GetComponentsBoundingBox(bool bNonColliding = false, bool bIncludeFromChildActors = false) const override;
 
 	//////////////////////////////////////////////////////////////////////////
 	// Blueprint interface for smart links
diff --git a/Engine/Source/Runtime/AIModule/Private/Navigation/NavLinkProxy.cpp b/Engine/Source/Runtime/AIModule/Private/Navigation/NavLinkProxy.cpp
--- a/Engine/Source/Runtime/AIModule/Private/Navigation/NavLinkProxy.cpp
+++ b/Engine/Source/Runtime/AIModule/Private/Navigation/NavLinkProxy.cpp
@@ -206,7 +206,7 @@ bool ANavLinkProxy::GetNavigationLinksArray(TArray<FNavigationLink>& OutLink, TA
 	return (PointLinks.Num() > 0) || (SegmentLinks.Num() > 0) || bIsSmartLinkActive;
 }
 
-FBox ANavLinkProxy::GetComponentsBoundingBox(bool bNonColliding) const
+FBox ANavLinkProxy::GetComponentsBoundingBox(bool bNonColliding, bool bIncludeFromChildActors) const
 {
 	FBox LinksBB(FVector(0.f, 0.f, -10.f), FVector(0.f,0.f,10.f));
 
diff --git a/Engine/Source/Runtime/Engine/Classes/Engine/LODActor.h b/Engine/Source/Runtime/Engine/Classes/Engine/LODActor.h
--- a/Engine/Source/Runtime/Engine/Classes/Engine/LODActor.h
+++ b/Engine/Source/Runtime/Engine/Classes/Engine/LODActor.h
@@ -72,7 +72,7 @@ public:
 	virtual void EditorApplyScale(const FVector& DeltaScale, const FVector* PivotLocation, bool bAltDown, bool bShiftDown, bool bCtrlDown) override;
 	virtual void EditorApplyMirror(const FVector& MirrorScale, const FVector& PivotLocation) override;
 #endif // WITH_EDITOR	
-	virtual FBox GetComponentsBoundingBox(bool bNonColliding = false) const override;
+	virtual FBox GetComponentsBoundingBox(bool bNonColliding = false, bool bIncludeFromChildActors = false) const override;
 	virtual void PostRegisterAllComponents() override;
 	virtual void Tick(float DeltaSeconds) override;	
 	virtual bool IsLevelBoundsRelevant() const override { return false; }
diff --git a/Engine/Source/Runtime/Engine/Classes/Engine/LevelBounds.h b/Engine/Source/Runtime/Engine/Classes/Engine/LevelBounds.h
--- a/Engine/Source/Runtime/Engine/Classes/Engine/LevelBounds.h
+++ b/Engine/Source/Runtime/Engine/Classes/Engine/LevelBounds.h
@@ -47,7 +47,7 @@ class ENGINE_API ALevelBounds
 	//~ End UObject Interface
 	
 	//~ Begin AActor Interface.
-	virtual FBox GetComponentsBoundingBox(bool bNonColliding = false) const override;
+	virtual FBox GetComponentsBoundingBox(bool bNonColliding = false, bool bIncludeFromChildActors = false) const override;
 	virtual bool IsLevelBoundsRelevant() const override { return false; }
 	//~ End AActor Interface.
 
diff --git a/Engine/Source/Runtime/Engine/Classes/GameFramework/Actor.h b/Engine/Source/Runtime/Engine/Classes/GameFramework/Actor.h
--- a/Engine/Source/Runtime/Engine/Classes/GameFramework/Actor.h
+++ b/Engine/Source/Runtime/Engine/Classes/GameFramework/Actor.h
@@ -948,9 +948,10 @@ public:
 	 * @param	bOnlyCollidingComponents	If true, will only return the bounding box for components with collision enabled.
 	 * @param	Origin						Set to the center of the actor in world space
 	 * @param	BoxExtent					Set to half the actor's size in 3d space
+	 * @param	bIncludeFromChildActors		If true then recurse in to ChildActor components 
 	 */
 	UFUNCTION(BlueprintCallable, Category="Collision", meta=(DisplayName = "GetActorBounds"))
-	void GetActorBounds(bool bOnlyCollidingComponents, FVector& Origin, FVector& BoxExtent) const;
+	void GetActorBounds(bool bOnlyCollidingComponents, FVector& Origin, FVector& BoxExtent, bool bIncludeFromChildActors = false) const;
 
 	/** Returns the RootComponent of this Actor */
 	UFUNCTION(BlueprintGetter)
@@ -2585,17 +2586,23 @@ public:
 	/** 
 	 * Returns the world space bounding box of all components in this Actor.
 	 * @param bNonColliding Indicates that you want to include non-colliding components in the bounding box
+	 * @param bIncludeFromChildActors If true then recurse in to ChildActor components and find components of the appropriate type in those Actors as well
 	 */
-	virtual FBox GetComponentsBoundingBox(bool bNonColliding = false) const;
+	virtual FBox GetComponentsBoundingBox(bool bNonColliding = false, bool bIncludeFromChildActors = false) const;
 
 	/** 
 	 * Calculates the actor space bounding box of all components in this Actor.  This is slower than GetComponentsBoundingBox(), because the local bounds of the components are not cached -- they are recalculated every time this function is called.
 	 * @param bNonColliding Indicates that you want to include non-colliding components in the bounding box
+	 * @param bIncludeFromChildActors If true then recurse in to ChildActor components and find components of the appropriate type in those Actors as well
 	 */
-	virtual FBox CalculateComponentsBoundingBoxInLocalSpace(bool bNonColliding = false) const;
+	virtual FBox CalculateComponentsBoundingBoxInLocalSpace(bool bNonColliding = false, bool bIncludeFromChildActors = false) const;
 
-	/** Get half-height/radius of a big axis-aligned cylinder around this actors registered colliding components, or all registered components if bNonColliding is false. */
-	virtual void GetComponentsBoundingCylinder(float& CollisionRadius, float& CollisionHalfHeight, bool bNonColliding = false) const;
+	/** 
+	 * Get half-height/radius of a big axis-aligned cylinder around this actors registered colliding components, or all registered components if bNonColliding is false. 
+	 * @param bNonColliding Indicates that you want to include non-colliding components in the bounding cylinder
+	 * @param bIncludeFromChildActors If true then recurse in to ChildActor components and find components of the appropriate type in those Actors as well
+	*/
+	virtual void GetComponentsBoundingCylinder(float& CollisionRadius, float& CollisionHalfHeight, bool bNonColliding = false, bool bIncludeFromChildActors = false) const;
 
 	/**
 	 * Get axis-aligned cylinder around this actor, used for simple collision checks (ie Pawns reaching a destination).
@@ -2839,26 +2846,28 @@ public:
 
 private:
 	/**
-	 * Internal helper function to centralize GetComponents logic. 
+	 * Internal helper function to call a compile-time lambda on all components of a given type
 	 * Use template parameter bClassIsActorComponent to avoid doing unnecessary IsA checks when the ComponentClass is exactly UActorComponent
+	 * Use template parameter bIncludeFromChildActors to recurse in to ChildActor components and find components of the appropriate type in those actors as well
 	 */
-	template<bool bClassIsActorComponent, class AllocatorType>
-	void GetComponents_Internal(TSubclassOf<UActorComponent> ComponentClass, TArray<UActorComponent*, AllocatorType>& OutComponents, bool bIncludeFromChildActors = false) const
+	template<class ComponentType, bool bClassIsActorComponent, bool bIncludeFromChildActors, typename Func>
+	void ForEachComponent_Internal(TSubclassOf<UActorComponent> ComponentClass, Func InFunc) const
 	{
 		check(bClassIsActorComponent == false || ComponentClass == UActorComponent::StaticClass());
+		check(ComponentClass->IsChildOf(ComponentType::StaticClass()));
 
-		TArray<AActor*> ChildActors;
-
-		for (UActorComponent* OwnedComponent : OwnedComponents)
+		// static check, so that the most common case (bIncludeFromChildActors) doesn't need to allocate an additional array : 
+		if (bIncludeFromChildActors)
 		{
-			if (OwnedComponent)
+			TArray<AActor*, TInlineAllocator<NumInlinedActorComponents>> ChildActors;
+			for (UActorComponent* OwnedComponent : OwnedComponents)
 			{
-				if (bClassIsActorComponent || OwnedComponent->IsA(ComponentClass))
-				{
-					OutComponents.Add(OwnedComponent);
-				}
-				if (bIncludeFromChildActors)
+				if (OwnedComponent)
 				{
+					if (bClassIsActorComponent || OwnedComponent->IsA(ComponentClass))
+					{
+						InFunc(static_cast<ComponentType*>(OwnedComponent));
+					}
 					if (UChildActorComponent* ChildActorComponent = Cast<UChildActorComponent>(OwnedComponent))
 					{
 						if (AActor* ChildActor = ChildActorComponent->GetChildActor())
@@ -2868,16 +2877,82 @@ private:
 					}
 				}
 			}
+
+			for (AActor* ChildActor : ChildActors)
+			{
+				ChildActor->ForEachComponent_Internal<ComponentType, bClassIsActorComponent, bIncludeFromChildActors>(ComponentClass, InFunc);
+			}
+		}
+		else
+		{
+			for (UActorComponent* OwnedComponent : OwnedComponents)
+			{
+				if (OwnedComponent)
+				{
+					if (bClassIsActorComponent || OwnedComponent->IsA(ComponentClass))
+					{
+						InFunc(static_cast<ComponentType*>(OwnedComponent));
+					}
+				}
+			}
 		}
+	}
 
-		for (AActor* ChildActor : ChildActors)
+	/**
+	 * Internal helper function to call a compile-time lambda on all components of a given type
+	 * Redirects the call to the proper template function specialization for bClassIsActorComponent and bIncludeFromChildActors :
+	 */
+	template<class ComponentType, typename Func>
+	void ForEachComponent_Internal(TSubclassOf<UActorComponent> ComponentClass, bool bIncludeFromChildActors, Func InFunc) const
+	{
+		static_assert(TPointerIsConvertibleFromTo<ComponentType, const UActorComponent>::Value, "'ComponentType' template parameter to ForEachComponent must be derived from UActorComponent");
+		if (ComponentClass == UActorComponent::StaticClass())
+		{
+			if (bIncludeFromChildActors)
+			{
+				ForEachComponent_Internal<ComponentType, true /*bClassIsActorComponent*/, true /*bIncludeFromChildActors*/>(ComponentClass, InFunc);
+			}
+			else
+			{
+				ForEachComponent_Internal<ComponentType, true /*bClassIsActorComponent*/, false /*bIncludeFromChildActors*/>(ComponentClass, InFunc);
+			}
+		}
+		else
 		{
-			ChildActor->GetComponents_Internal<bClassIsActorComponent>(ComponentClass, OutComponents, true);
+			if (bIncludeFromChildActors)
+			{
+				ForEachComponent_Internal<ComponentType, false /*bClassIsActorComponent*/, true /*bIncludeFromChildActors*/>(ComponentClass, InFunc);
+			}
+			else
+			{
+				ForEachComponent_Internal<ComponentType, false /*bClassIsActorComponent*/, false /*bIncludeFromChildActors*/>(ComponentClass, InFunc);
+			}
 		}
 	}
 
 public:
 
+	/**
+	 * Calls the compile-time lambda on each component of the specified type
+	 * @param ComponentType             The component class to find all components of a class derived from
+	 * @param bIncludeFromChildActors   If true then recurse in to ChildActor components and find components of the appropriate type in those Actors as well
+	 */
+	template<class ComponentType, typename Func>
+	void ForEachComponent(bool bIncludeFromChildActors, Func InFunc) const
+	{		
+		ForEachComponent_Internal<ComponentType>(ComponentType::StaticClass(), bIncludeFromChildActors, InFunc);
+	}
+
+	/**
+	 * Calls the compile-time lambda on each valid component 
+	 * @param bIncludeFromChildActors   If true then recurse in to ChildActor components and find components of the appropriate type in those Actors as well
+	 */
+	template<typename Func>
+	void ForEachComponent(bool bIncludeFromChildActors, Func InFunc) const
+	{
+		ForEachComponent_Internal<UActorComponent>(UActorComponent::StaticClass(), bIncludeFromChildActors, InFunc);
+	}
+
 	UE_DEPRECATED(4.24, "Use one of the GetComponents implementations as appropriate")
 	TArray<UActorComponent*> GetComponentsByClass(TSubclassOf<UActorComponent> ComponentClass) const
 	{
@@ -2901,14 +2976,10 @@ public:
 		SCOPE_CYCLE_COUNTER(STAT_GetComponentsTime);
 
 		OutComponents.Reset();
-		if (ComponentClass == UActorComponent::StaticClass())
+		ForEachComponent_Internal<UActorComponent>(ComponentClass, bIncludeFromChildActors, [&](UActorComponent* InComp)
 		{
-			GetComponents_Internal<true>(ComponentClass, OutComponents, bIncludeFromChildActors);
-		}
-		else
-		{
-			GetComponents_Internal<false>(ComponentClass, OutComponents, bIncludeFromChildActors);
-		}
+			OutComponents.Add(InComp);
+		});
 	}
 
 	/**
@@ -2926,9 +2997,11 @@ public:
 	{
 		SCOPE_CYCLE_COUNTER(STAT_GetComponentsTime);
 
-		static_assert(TPointerIsConvertibleFromTo<T, const UActorComponent>::Value, "'T' template parameter to GetComponents must be derived from UActorComponent");
 		OutComponents.Reset();
-		GetComponents_Internal<false>(T::StaticClass(), *reinterpret_cast<TArray<UActorComponent*, AllocatorType>*>(&OutComponents), bIncludeFromChildActors);
+		ForEachComponent_Internal<T>(T::StaticClass(), bIncludeFromChildActors, [&](T* InComp)
+		{
+			OutComponents.Add(InComp);
+		});
 	}
 
 	/**
@@ -2948,7 +3021,10 @@ public:
 		SCOPE_CYCLE_COUNTER(STAT_GetComponentsTime);
 
 		OutComponents.Reset();
-		GetComponents_Internal<true>(UActorComponent::StaticClass(), OutComponents, bIncludeFromChildActors);
+		ForEachComponent_Internal<UActorComponent>(UActorComponent::StaticClass(), bIncludeFromChildActors, [&](UActorComponent* InComp)
+		{
+			OutComponents.Add(InComp);
+		});
 	}
 
 	/**
@@ -3556,7 +3632,7 @@ FORCEINLINE_DEBUGGABLE bool AActor::IsNetMode(ENetMode Mode) const
 	FVector GetActorForwardVector() const { return Super::GetActorForwardVector(); } \
 	FVector GetActorUpVector() const { return Super::GetActorUpVector(); } \
 	FVector GetActorRightVector() const { return Super::GetActorRightVector(); } \
-	void GetActorBounds(bool bOnlyCollidingComponents, FVector& Origin, FVector& BoxExtent) const { return Super::GetActorBounds(bOnlyCollidingComponents, Origin, BoxExtent); } \
+	void GetActorBounds(bool bOnlyCollidingComponents, FVector& Origin, FVector& BoxExtent, bool bIncludeFromChildActors = false) const { return Super::GetActorBounds(bOnlyCollidingComponents, Origin, BoxExtent, bIncludeFromChildActors); } \
 	void SetActorScale3D(FVector NewScale3D) { Super::SetActorScale3D(NewScale3D); } \
 	FVector GetActorScale3D() const { return Super::GetActorScale3D(); } \
 	void SetActorRelativeScale3D(FVector NewRelativeScale) { Super::SetActorRelativeScale3D(NewRelativeScale); } \
diff --git a/Engine/Source/Runtime/Engine/Private/Actor.cpp b/Engine/Source/Runtime/Engine/Private/Actor.cpp
--- a/Engine/Source/Runtime/Engine/Private/Actor.cpp
+++ b/Engine/Source/Runtime/Engine/Private/Actor.cpp
@@ -1129,12 +1129,12 @@ void AActor::PostActorCreated()
 	// nothing at the moment
 }
 
-void AActor::GetComponentsBoundingCylinder(float& OutCollisionRadius, float& OutCollisionHalfHeight, bool bNonColliding) const
+void AActor::GetComponentsBoundingCylinder(float& OutCollisionRadius, float& OutCollisionHalfHeight, bool bNonColliding, bool bIncludeFromChildActors) const
 {
 	bool bIgnoreRegistration = false;
 
 #if WITH_EDITOR
-	if(IsTemplate())
+	if (IsTemplate())
 	{
 		// Editor code calls this function on default objects when placing them in the viewport, so no components will be registered in those cases.
 		UWorld* MyWorld = GetWorld();
@@ -1154,30 +1154,22 @@ void AActor::GetComponentsBoundingCylinder(float& OutCollisionRadius, float& Out
 	}
 #endif
 
-	float Radius = 0.f;
-	float HalfHeight = 0.f;
+	OutCollisionRadius = 0.f;
+	OutCollisionHalfHeight = 0.f;
 
-	for (const UActorComponent* ActorComponent : GetComponents())
+	ForEachComponent<UPrimitiveComponent>(bIncludeFromChildActors, [&](const UPrimitiveComponent* InPrimComp)
 	{
-		const UPrimitiveComponent* PrimComp = Cast<const UPrimitiveComponent>(ActorComponent);
-		if (PrimComp)
+		// Only use collidable components to find collision bounding cylinder
+		if ((bIgnoreRegistration || InPrimComp->IsRegistered()) && (bNonColliding || InPrimComp->IsCollisionEnabled()))
 		{
-			// Only use collidable components to find collision bounding box.
-			if ((bIgnoreRegistration || PrimComp->IsRegistered()) && (bNonColliding || PrimComp->IsCollisionEnabled()))
-			{
-				float TestRadius, TestHalfHeight;
-				PrimComp->CalcBoundingCylinder(TestRadius, TestHalfHeight);
-				Radius = FMath::Max(Radius, TestRadius);
-				HalfHeight = FMath::Max(HalfHeight, TestHalfHeight);
-			}
+			float TestRadius, TestHalfHeight;
+			InPrimComp->CalcBoundingCylinder(TestRadius, TestHalfHeight);
+			OutCollisionRadius = FMath::Max(OutCollisionRadius, TestRadius);
+			OutCollisionHalfHeight = FMath::Max(OutCollisionHalfHeight, TestHalfHeight);
 		}
-	}
-
-	OutCollisionRadius = Radius;
-	OutCollisionHalfHeight = HalfHeight;
+	});
 }
 
-
 void AActor::GetSimpleCollisionCylinder(float& CollisionRadius, float& CollisionHalfHeight) const
 {
 	if (IsRootComponentCollisionRegistered())
@@ -1250,48 +1242,37 @@ bool AActor::Modify( bool bAlwaysMarkDirty/*=true*/ )
 }
 #endif
 
-FBox AActor::GetComponentsBoundingBox(bool bNonColliding) const
+FBox AActor::GetComponentsBoundingBox(bool bNonColliding, bool bIncludeFromChildActors) const
 {
 	FBox Box(ForceInit);
 
-	for (const UActorComponent* ActorComponent : GetComponents())
+	ForEachComponent<UPrimitiveComponent>(bIncludeFromChildActors, [&](const UPrimitiveComponent* InPrimComp)
 	{
-		const UPrimitiveComponent* PrimComp = Cast<const UPrimitiveComponent>(ActorComponent);
-		if (PrimComp)
+		// Only use collidable components to find collision bounding box.
+		if (InPrimComp->IsRegistered() && (bNonColliding || InPrimComp->IsCollisionEnabled()))
 		{
-			// Only use collidable components to find collision bounding box.
-			if (PrimComp->IsRegistered() && (bNonColliding || PrimComp->IsCollisionEnabled()))
-			{
-				Box += PrimComp->Bounds.GetBox();
-			}
+			Box += InPrimComp->Bounds.GetBox();
 		}
-	}
+	});
 
 	return Box;
 }
 
-FBox AActor::CalculateComponentsBoundingBoxInLocalSpace( bool bNonColliding ) const
+FBox AActor::CalculateComponentsBoundingBoxInLocalSpace(bool bNonColliding, bool bIncludeFromChildActors) const
 {
 	FBox Box(ForceInit);
-
 	const FTransform& ActorToWorld = GetTransform();
 	const FTransform WorldToActor = ActorToWorld.Inverse();
 
-	for( const UActorComponent* ActorComponent : GetComponents() )
+	ForEachComponent<UPrimitiveComponent>(bIncludeFromChildActors, [&](const UPrimitiveComponent* InPrimComp)
 	{
-		const UPrimitiveComponent* PrimComp = Cast<const UPrimitiveComponent>( ActorComponent );
-		if( PrimComp )
+		// Only use collidable components to find collision bounding box.
+		if (InPrimComp->IsRegistered() && (bNonColliding || InPrimComp->IsCollisionEnabled()))
 		{
-			// Only use collidable components to find collision bounding box.
-			if( PrimComp->IsRegistered() && ( bNonColliding || PrimComp->IsCollisionEnabled() ) )
-			{
-				const FTransform ComponentToActor = PrimComp->GetComponentTransform() * WorldToActor;
-				FBoxSphereBounds ActorSpaceComponentBounds = PrimComp->CalcBounds( ComponentToActor );
-
-				Box += ActorSpaceComponentBounds.GetBox();
-			}
+			const FTransform ComponentToActor = InPrimComp->GetComponentTransform() * WorldToActor;
+			Box += InPrimComp->CalcBounds(ComponentToActor).GetBox();
 		}
-	}
+	});
 
 	return Box;
 }
@@ -3952,9 +3933,9 @@ bool AActor::SetRootComponent(class USceneComponent* NewRootComponent)
 	return false;
 }
 
-void AActor::GetActorBounds(bool bOnlyCollidingComponents, FVector& Origin, FVector& BoxExtent) const
+void AActor::GetActorBounds(bool bOnlyCollidingComponents, FVector& Origin, FVector& BoxExtent, bool bIncludeFromChildActors) const
 {
-	const FBox Bounds = GetComponentsBoundingBox(!bOnlyCollidingComponents);
+	const FBox Bounds = GetComponentsBoundingBox(!bOnlyCollidingComponents, bIncludeFromChildActors);
 
 	// To keep consistency with the other GetBounds functions, transform our result into an origin / extent formatting
 	Bounds.GetCenterAndExtents(Origin, BoxExtent);
diff --git a/Engine/Source/Runtime/Engine/Private/LODActor.cpp b/Engine/Source/Runtime/Engine/Private/LODActor.cpp
--- a/Engine/Source/Runtime/Engine/Private/LODActor.cpp
+++ b/Engine/Source/Runtime/Engine/Private/LODActor.cpp
@@ -1205,7 +1205,7 @@ UStaticMeshComponent* ALODActor::GetOrCreateLODComponentForActor(const AActor* I
 
 #endif // !(UE_BUILD_SHIPPING || UE_BUILD_TEST)
 
-FBox ALODActor::GetComponentsBoundingBox(bool bNonColliding) const 
+FBox ALODActor::GetComponentsBoundingBox(bool bNonColliding, bool bIncludeFromChildActors) const
 {
 	FBox BoundBox = Super::GetComponentsBoundingBox(bNonColliding);
 
@@ -1240,7 +1240,7 @@ FBox ALODActor::GetComponentsBoundingBox(bool bNonColliding) const
 			{
 				if (Actor)
 				{
-					BoundBox += Actor->GetComponentsBoundingBox(bNonColliding);
+					BoundBox += Actor->GetComponentsBoundingBox(bNonColliding, bIncludeFromChildActors);
 				}
 			}
 		}
diff --git a/Engine/Source/Runtime/Engine/Private/LevelBounds.cpp b/Engine/Source/Runtime/Engine/Private/LevelBounds.cpp
--- a/Engine/Source/Runtime/Engine/Private/LevelBounds.cpp
+++ b/Engine/Source/Runtime/Engine/Private/LevelBounds.cpp
@@ -47,7 +47,7 @@ void ALevelBounds::PostLoad()
 	}
 }
 
-FBox ALevelBounds::GetComponentsBoundingBox(bool bNonColliding) const
+FBox ALevelBounds::GetComponentsBoundingBox(bool bNonColliding, bool bIncludeFromChildActors) const
 {
 	checkf(RootComponent != nullptr, TEXT("LevelBounds actor with null root component: %s"), *GetPathNameSafe(this));
 	FVector BoundsCenter = RootComponent->GetComponentLocation();
