Resolving robomerge conflict - adding okforgithub tag

Adds the console controller movement / input functionality for infiltration mode to 12.3
Use #jira OCN-22003

diff --git a/Engine/Source/Runtime/Engine/Private/GameViewportClient.cpp b/Engine/Source/Runtime/Engine/Private/GameViewportClient.cpp
--- a/Engine/Source/Runtime/Engine/Private/GameViewportClient.cpp
+++ b/Engine/Source/Runtime/Engine/Private/GameViewportClient.cpp
@@ -821,7 +821,11 @@ bool UGameViewportClient::GetMousePosition(FVector2D& MousePosition) const
 {
 	bool bGotMousePosition = false;
 
+#ifdef USE_CONSOLE_CONTROLLER
 	if (Viewport && FSlateApplication::Get().IsMouseAttached())
+#else
+	if (Viewport)
+#endif
 	{
 		FIntPoint MousePos;
 		Viewport->GetMousePos(MousePos);
