Only output sort actor hierarchy warning when running attended. It is unclear why it occasionally appears in cooking on the build machines, but it appears to be a red-herring.
#rnx

diff --git a/Engine/Source/Runtime/Engine/Private/Level.cpp b/Engine/Source/Runtime/Engine/Private/Level.cpp
--- a/Engine/Source/Runtime/Engine/Private/Level.cpp
+++ b/Engine/Source/Runtime/Engine/Private/Level.cpp
@@ -902,7 +902,7 @@ static void SortActorsHierarchy(TArray<AActor*>& Actors, UObject* Level)
 	const double StableSortTime = FPlatformTime::Seconds() - StableSortStartTime;
 
 	const double ElapsedTime = FPlatformTime::Seconds() - StartTime;
-	if (ElapsedTime > 1.0)
+	if (ElapsedTime > 1.0 && !FApp::IsUnattended())
 	{
 		UE_LOG(LogLevel, Warning, TEXT("SortActorsHierarchy(%s) took %f seconds (CalcAttachDepth: %f StableSort: %f)"), Level ? *GetNameSafe(Level->GetOutermost()) : TEXT("??"), ElapsedTime, CalcAttachDepthTime, StableSortTime);
 	}
