Always update the longest run of saturated frames for net analytics.

Previously, we would only update the longest run when we were no longer saturated. This could cause issues where if we were saturated for 100% of the time, the run would report at 0.

diff --git a/Engine/Source/Runtime/Engine/Private/Net/NetAnalyticsTypes.cpp b/Engine/Source/Runtime/Engine/Private/Net/NetAnalyticsTypes.cpp
--- a/Engine/Source/Runtime/Engine/Private/Net/NetAnalyticsTypes.cpp
+++ b/Engine/Source/Runtime/Engine/Private/Net/NetAnalyticsTypes.cpp
@@ -168,10 +168,13 @@ void FNetConnectionSaturationAnalytics::TrackFrame(const bool bIsSaturated)
 	{
 		++NumberOfSaturatedFrames;
 		++CurrentRunOfSaturatedFrames;
+        if (CurrentRunOfSaturatedFrames > LongestRunOfSaturatedFrames)
+        {
+            LongestRunOfSaturatedFrames = CurrentRunOfSaturatedFrames;
+        }
 	}
 	else
 	{
-		LongestRunOfSaturatedFrames = FMath::Max<uint32>(CurrentRunOfSaturatedFrames, LongestRunOfSaturatedFrames);
 		CurrentRunOfSaturatedFrames = 0;
 	}
 }
@@ -184,10 +187,13 @@ void FNetConnectionSaturationAnalytics::TrackReplication(const bool bIsSaturated
 	{
 		++NumberOfSaturatedReplications;
 		++CurrentRunOfSaturatedReplications;
+        if (CurrentRunOfSaturatedReplications > LongestRunOfSaturatedReplications)
+        {
+            LongestRunOfSaturatedReplications = CurrentRunOfSaturatedReplications;
+        }
 	}
 	else
 	{
-		LongestRunOfSaturatedReplications = FMath::Max<uint32>(CurrentRunOfSaturatedReplications, LongestRunOfSaturatedReplications);
 		CurrentRunOfSaturatedReplications = 0;
 	}
 }
