Added another variant of DrawCoordinatesSystem

diff --git a/Engine/Source/Runtime/Engine/Private/PrimitiveDrawingUtils.cpp b/Engine/Source/Runtime/Engine/Private/PrimitiveDrawingUtils.cpp
--- a/Engine/Source/Runtime/Engine/Private/PrimitiveDrawingUtils.cpp
+++ b/Engine/Source/Runtime/Engine/Private/PrimitiveDrawingUtils.cpp
@@ -1200,6 +1200,20 @@ void DrawCoordinateSystem(FPrimitiveDrawInterface* PDI, FVector const& AxisLoc,
 	PDI->DrawLine(AxisLoc, AxisLoc + Z*Scale, FLinearColor::Blue, DepthPriority, Thickness );
 }
 
+
+void DrawCoordinateSystem(FPrimitiveDrawInterface* PDI, FVector const& AxisLoc, FRotator const& AxisRot, float Scale, const FLinearColor& InColor, uint8 DepthPriority, float Thickness)
+{
+	FRotationMatrix R(AxisRot);
+	FVector const X = R.GetScaledAxis(EAxis::X);
+	FVector const Y = R.GetScaledAxis(EAxis::Y);
+	FVector const Z = R.GetScaledAxis(EAxis::Z);
+
+	PDI->DrawLine(AxisLoc, AxisLoc + X * Scale, InColor, DepthPriority, Thickness);
+	PDI->DrawLine(AxisLoc, AxisLoc + Y * Scale, InColor, DepthPriority, Thickness);
+	PDI->DrawLine(AxisLoc, AxisLoc + Z * Scale, InColor, DepthPriority, Thickness);
+}
+
+
 void DrawDirectionalArrow(FPrimitiveDrawInterface* PDI,const FMatrix& ArrowToWorld,const FLinearColor& InColor,float Length,float ArrowSize,uint8 DepthPriority,float Thickness)
 {
 	PDI->DrawLine(ArrowToWorld.TransformPosition(FVector(Length,0,0)),ArrowToWorld.TransformPosition(FVector::ZeroVector),InColor,DepthPriority,Thickness);
diff --git a/Engine/Source/Runtime/Engine/Public/SceneManagement.h b/Engine/Source/Runtime/Engine/Public/SceneManagement.h
--- a/Engine/Source/Runtime/Engine/Public/SceneManagement.h
+++ b/Engine/Source/Runtime/Engine/Public/SceneManagement.h
@@ -2723,8 +2723,31 @@ extern ENGINE_API void DrawDashedLine(class FPrimitiveDrawInterface* PDI, const
  */
 extern ENGINE_API void DrawWireDiamond(class FPrimitiveDrawInterface* PDI, const FMatrix& DiamondMatrix, float Size, const FLinearColor& InColor, uint8 DepthPriority, float Thickness = 0.0f);
 
+/**
+ * Draws a coordinate system (Red for X axis, Green for Y axis, Blue for Z axis).
+ *
+ * @param	PDI				Draw interface.
+ * @param	AxisLoc			Location of the coordinate system.
+ * @param	AxisRot			Location of the coordinate system.
+ * @param	Scale			Scale for the axis lines.
+ * @param	DepthPriority	Depth priority coordinate system.
+ * @param	Thickness		How thick to draw the axis lines
+ */
 extern ENGINE_API void DrawCoordinateSystem(FPrimitiveDrawInterface* PDI, FVector const& AxisLoc, FRotator const& AxisRot, float Scale, uint8 DepthPriority, float Thickness = 0.0f);
 
+/**
+ * Draws a coordinate system with a fixed color.
+ *
+ * @param	PDI				Draw interface.
+ * @param	AxisLoc			Location of the coordinate system.
+ * @param	AxisRot			Location of the coordinate system.
+ * @param	Scale			Scale for the axis lines.
+ * @param	InColor			Color of the axis lines.
+ * @param	DepthPriority	Depth priority coordinate system.
+ * @param	Thickness		How thick to draw the axis lines
+ */
+extern ENGINE_API void DrawCoordinateSystem(FPrimitiveDrawInterface* PDI, FVector const& AxisLoc, FRotator const& AxisRot, float Scale, const FLinearColor& InColor, uint8 DepthPriority, float Thickness = 0.0f);
+
 /**
  * Draws a wireframe of the bounds of a frustum as defined by a transform from clip-space into world-space.
  * @param PDI - The interface to draw the wireframe.
