Fix convex geometry dangling pointer.
Added convex memory tracking

diff --git a/Engine/Source/Runtime/Engine/Classes/PhysicsEngine/ConvexElem.h b/Engine/Source/Runtime/Engine/Classes/PhysicsEngine/ConvexElem.h
--- a/Engine/Source/Runtime/Engine/Classes/PhysicsEngine/ConvexElem.h
+++ b/Engine/Source/Runtime/Engine/Classes/PhysicsEngine/ConvexElem.h
@@ -58,7 +58,7 @@ private:
 	physx::PxConvexMesh*   ConvexMeshNegX;
 
 #if WITH_CHAOS
-	TUniquePtr<Chaos::FConvex> ChaosConvex;
+	TSharedPtr<Chaos::FConvex, ESPMode::ThreadSafe> ChaosConvex;
 #endif
 
 public:
@@ -110,9 +110,12 @@ public:
 	ENGINE_API void SetMirroredConvexMesh(physx::PxConvexMesh* InMesh);
 
 #if WITH_CHAOS
-	ENGINE_API const TUniquePtr<Chaos::FConvex>& GetChaosConvexMesh() const;
+	ENGINE_API const auto& GetChaosConvexMesh() const
+	{
+		return ChaosConvex;
+	}
 
-	ENGINE_API void SetChaosConvexMesh(TUniquePtr<Chaos::FConvex>&& InChaosConvex);
+	ENGINE_API void SetChaosConvexMesh(TSharedPtr<Chaos::FConvex, ESPMode::ThreadSafe>&& InChaosConvex);
 
 	ENGINE_API void ResetChaosConvexMesh();
 
diff --git a/Engine/Source/Runtime/Engine/Private/PhysicsEngine/BodyInstance.cpp b/Engine/Source/Runtime/Engine/Private/PhysicsEngine/BodyInstance.cpp
--- a/Engine/Source/Runtime/Engine/Private/PhysicsEngine/BodyInstance.cpp
+++ b/Engine/Source/Runtime/Engine/Private/PhysicsEngine/BodyInstance.cpp
@@ -1815,7 +1815,7 @@ bool FBodyInstance::UpdateBodyScale(const FVector& InScale3D, bool bForceUpdate)
 					const TImplicitObjectScaled<FConvex>* ConvexGeometry = static_cast<const TImplicitObjectScaled<FConvex>*>(&ImplicitObject);
 
 					FKConvexElem* ConvexElem = ShapeElem->GetShapeCheck<FKConvexElem>();
-					const TUniquePtr<FConvex>& ConvexImplicit = ConvexElem->GetChaosConvexMesh();
+					const auto& ConvexImplicit = ConvexElem->GetChaosConvexMesh();
 
 					// Ensure no rotation/translation in relative transform. PhysX supports this, we currently do not.
 					CHAOS_ENSURE(RelativeTM.GetRotation() == FQuat::Identity);
diff --git a/Engine/Source/Runtime/Engine/Private/PhysicsEngine/BodySetup.cpp b/Engine/Source/Runtime/Engine/Private/PhysicsEngine/BodySetup.cpp
--- a/Engine/Source/Runtime/Engine/Private/PhysicsEngine/BodySetup.cpp
+++ b/Engine/Source/Runtime/Engine/Private/PhysicsEngine/BodySetup.cpp
@@ -713,6 +713,10 @@ void UBodySetup::FinishCreatingPhysicsMeshes_Chaos(FChaosDerivedDataReader<float
 			{
 				ConvexElem.SetChaosConvexMesh(MoveTemp(InReader.ConvexImplicitObjects[ElementIndex]));
 
+#if TRACK_CHAOS_GEOMETRY
+				ConvexElem.GetChaosConvexMesh()->Track(Chaos::MakeSerializable(ConvexElem.GetChaosConvexMesh()),FullName);
+#endif
+
 				if (ConvexElem.GetChaosConvexMesh()->IsPerformanceWarning())
 				{
 					const FString& PerformanceString = ConvexElem.GetChaosConvexMesh()->PerformanceWarningAndSimplifaction();
@@ -1711,12 +1715,8 @@ float FKConvexElem::GetVolume(const FVector& Scale) const
 }
 
 #if WITH_CHAOS
-const TUniquePtr<Chaos::FConvex>& FKConvexElem::GetChaosConvexMesh() const
-{
-	return ChaosConvex;
-}
 
-void FKConvexElem::SetChaosConvexMesh(TUniquePtr<Chaos::FConvex>&& InChaosConvex)
+void FKConvexElem::SetChaosConvexMesh(TSharedPtr<Chaos::FConvex, ESPMode::ThreadSafe>&& InChaosConvex)
 {
 	ChaosConvex = MoveTemp(InChaosConvex);
 	ComputeChaosConvexIndices();
diff --git a/Engine/Source/Runtime/Engine/Public/Physics/Experimental/ChaosDerivedDataReader.h b/Engine/Source/Runtime/Engine/Public/Physics/Experimental/ChaosDerivedDataReader.h
--- a/Engine/Source/Runtime/Engine/Public/Physics/Experimental/ChaosDerivedDataReader.h
+++ b/Engine/Source/Runtime/Engine/Public/Physics/Experimental/ChaosDerivedDataReader.h
@@ -23,7 +23,7 @@ public:
 	// Only valid use is to explicitly read chaos bulk data
 	explicit FChaosDerivedDataReader(FBulkDataInterface* InBulkData);
 
-	TArray<TUniquePtr<Chaos::FConvex>> ConvexImplicitObjects;
+	TArray<TSharedPtr<Chaos::FConvex, ESPMode::ThreadSafe>> ConvexImplicitObjects;
 	TArray<TSharedPtr<Chaos::FTriangleMeshImplicitObject, ESPMode::ThreadSafe>> TrimeshImplicitObjects;
 	FBodySetupUVInfo UVInfo;
 	TArray<int32> FaceRemap;
diff --git a/Engine/Source/Runtime/Experimental/Chaos/Public/Chaos/TrackedGeometryManager.h b/Engine/Source/Runtime/Experimental/Chaos/Public/Chaos/TrackedGeometryManager.h
--- a/Engine/Source/Runtime/Experimental/Chaos/Public/Chaos/TrackedGeometryManager.h
+++ b/Engine/Source/Runtime/Experimental/Chaos/Public/Chaos/TrackedGeometryManager.h
@@ -20,6 +20,7 @@ namespace Chaos
 			{
 				uint32 NumBytes;
 				FString DebugInfo;
+				FString ToString;
 
 				bool operator<(const FMemInfo& Other) const { return NumBytes < Other.NumBytes; }
 			};
@@ -35,6 +36,7 @@ namespace Chaos
 				FChaosArchive ChaosAr(MemAr);
 				FImplicitObject* NonConst = const_cast<FImplicitObject*>(Itr.Key.Get());	//only doing this to write out, serialize is non const for read in
 				NonConst->Serialize(ChaosAr);
+				Info.ToString = NonConst->ToString();
 				Info.NumBytes = Data.Num();
 				MemEntries.Add(Info);
 				TotalBytes += Info.NumBytes;
@@ -48,7 +50,7 @@ namespace Chaos
 
 			for (const FMemInfo& Info : MemEntries)
 			{
-				Ar->Logf(TEXT("%-10d %s"), Info.NumBytes, *Info.DebugInfo);
+				Ar->Logf(TEXT("%-10d %s ToString:%s"), Info.NumBytes, *Info.DebugInfo, *Info.ToString);
 			}
 
 			Ar->Logf(TEXT("%-10d Total"), TotalBytes);
diff --git a/Engine/Source/Runtime/Experimental/Chaos/Public/Chaos/TriangleMeshImplicitObject.h b/Engine/Source/Runtime/Experimental/Chaos/Public/Chaos/TriangleMeshImplicitObject.h
--- a/Engine/Source/Runtime/Experimental/Chaos/Public/Chaos/TriangleMeshImplicitObject.h
+++ b/Engine/Source/Runtime/Experimental/Chaos/Public/Chaos/TriangleMeshImplicitObject.h
@@ -206,6 +206,11 @@ namespace Chaos
 
 		virtual void Serialize(FChaosArchive& Ar) override;
 
+		virtual FString ToString() const
+		{
+			return FString::Printf(TEXT("TriangleMesh"));
+		}
+
 		virtual uint32 GetTypeHash() const override;
 
 		FVec3 GetFaceNormal(const int32 FaceIdx) const;
