Change lock semantic for virtual texture upload structured buffer
Required because we want a persistant pointer to the upload buffer

diff --git a/Engine/Source/Runtime/Engine/Private/VT/VirtualTextureUploadCache.cpp b/Engine/Source/Runtime/Engine/Private/VT/VirtualTextureUploadCache.cpp
--- a/Engine/Source/Runtime/Engine/Private/VT/VirtualTextureUploadCache.cpp
+++ b/Engine/Source/Runtime/Engine/Private/VT/VirtualTextureUploadCache.cpp
@@ -216,7 +216,7 @@ FVTUploadTileHandle FVirtualTextureUploadCache::PrepareTileForUpload(FVTUploadTi
 
 			// Here we bypass 'normal' RHI operations in order to get a persistent pointer to GPU memory, on supported platforms
 			// This should be encapsulated into a proper RHI method at some point
-			NewEntry.Memory = RHICmdList.LockStructuredBuffer(NewEntry.RHIStagingBuffer, 0u, MemorySize, RLM_WriteOnly);
+			NewEntry.Memory = RHICmdList.LockStructuredBuffer(NewEntry.RHIStagingBuffer, 0u, MemorySize, RLM_WriteOnly_NoOverwrite);
 
 			INC_MEMORY_STAT_BY(STAT_TotalGPUUploadSize, MemorySize);
 		}
