Fix p.NetShowCorrections client error logging displaying diff of CURRENT location vs corrected location instead of the MOVE location vs corrected location, and additionally display "same location" corrections as yellow
- Revert to previous "current location" if LastAckedMove is somehow invalid (but should not be, since this function is only called right after acking a move)
- Fixing this shows correctly that with a single client error, the first correction from the server has a different location, and then subsequent corrections are no longer different since those moves have been updated to the correct position after first correction move replay. We now draw these as YELLOW instead of overlapping green/red capsules, to show that there was ONE correction, and then multiple corrections received but that we have already corrected for. This occurred with latency - if you had 300 ping and received a correction, p.NetShowCorrections would show multiple green/red sets of capsules (the server corrects you and then corrects subsequent moves it receives from you, not knowing whether you have corrected or not). This was confusing when debugging, since if there's one error, only one correction should be needed. So now with latency, you should see one set of green/red capsules (your error that is corrected), and then yellow capsules for "the server sent you a correction but you already corrected location for it". Did not change any actual logic, just debugging display.

#rnx

diff --git a/Engine/Source/Runtime/Engine/Private/Components/CharacterMovementComponent.cpp b/Engine/Source/Runtime/Engine/Private/Components/CharacterMovementComponent.cpp
--- a/Engine/Source/Runtime/Engine/Private/Components/CharacterMovementComponent.cpp
+++ b/Engine/Source/Runtime/Engine/Private/Components/CharacterMovementComponent.cpp
@@ -9355,13 +9355,25 @@ void UCharacterMovementComponent::OnClientCorrectionReceived(FNetworkPredictionD
 #if !UE_BUILD_SHIPPING
 	if (CharacterMovementCVars::NetShowCorrections != 0)
 	{
-		const FVector LocDiff = UpdatedComponent->GetComponentLocation() - NewLocation;
+		const FVector ClientLocAtCorrectedMove = ClientData.LastAckedMove.IsValid() ? ClientData.LastAckedMove->SavedLocation : UpdatedComponent->GetComponentLocation();
+		const FVector LocDiff = ClientLocAtCorrectedMove - NewLocation;
 		const FString NewBaseString = NewBase ? NewBase->GetPathName(NewBase->GetOutermost()) : TEXT("None");
 		UE_LOG(LogNetPlayerMovement, Warning, TEXT("*** Client: Error for %s at Time=%.3f is %3.3f LocDiff(%s) ClientLoc(%s) ServerLoc(%s) NewBase: %s NewBone: %s ClientVel(%s) ServerVel(%s) SavedMoves %d"),
-			   *GetNameSafe(CharacterOwner), TimeStamp, LocDiff.Size(), *LocDiff.ToString(), *UpdatedComponent->GetComponentLocation().ToString(), *NewLocation.ToString(), *NewBaseString, *NewBaseBoneName.ToString(), *Velocity.ToString(), *NewVelocity.ToString(), ClientData.SavedMoves.Num());
+			   *GetNameSafe(CharacterOwner), TimeStamp, LocDiff.Size(), *LocDiff.ToString(), *ClientLocAtCorrectedMove.ToString(), *NewLocation.ToString(), *NewBaseString, *NewBaseBoneName.ToString(), *Velocity.ToString(), *NewVelocity.ToString(), ClientData.SavedMoves.Num());
 		const float DebugLifetime = CharacterMovementCVars::NetCorrectionLifetime;
-		DrawDebugCapsule(GetWorld(), UpdatedComponent->GetComponentLocation(), CharacterOwner->GetSimpleCollisionHalfHeight(), CharacterOwner->GetSimpleCollisionRadius(), FQuat::Identity, FColor(255, 100, 100), false, DebugLifetime);
-		DrawDebugCapsule(GetWorld(), NewLocation, CharacterOwner->GetSimpleCollisionHalfHeight(), CharacterOwner->GetSimpleCollisionRadius(), FQuat::Identity, FColor(100, 255, 100), false, DebugLifetime);
+		if (!LocDiff.IsNearlyZero())
+		{
+			// When server corrects us to a new location, draw red at location where client thought they were, green where the server corrected us to
+			DrawDebugCapsule(GetWorld(), ClientLocAtCorrectedMove, CharacterOwner->GetSimpleCollisionHalfHeight(), CharacterOwner->GetSimpleCollisionRadius(), FQuat::Identity, FColor(255, 100, 100), false, DebugLifetime);
+			DrawDebugCapsule(GetWorld(), NewLocation, CharacterOwner->GetSimpleCollisionHalfHeight(), CharacterOwner->GetSimpleCollisionRadius(), FQuat::Identity, FColor(100, 255, 100), false, DebugLifetime);
+		}
+		else
+		{
+			// When we receive a server correction that doesn't change our position from where our client move had us, draw yellow (otherwise would be overlapping)
+			// This occurs when we receive an initial correction, replay moves to get us into the right location, and then receive subsequent corrections by the server (who doesn't know if we corrected already
+			// so continues to send corrections). This is a "no-op" server correction with regards to location since we already corrected (occurs with latency)
+			DrawDebugCapsule(GetWorld(), NewLocation, CharacterOwner->GetSimpleCollisionHalfHeight(), CharacterOwner->GetSimpleCollisionRadius(), FQuat::Identity, FColor(255, 255, 100), false, DebugLifetime);
+		}
 	}
 #endif //!UE_BUILD_SHIPPING
 
