Fix CharacterMovementComponent ClientAdjustRootMotionPosition RPC never being called due to incorrect conditional - moves with animation root motion were always going through the more expensive root motion source + anim root motion ClientAdjustRootMotionSourcePosition RPC

#rnx

diff --git a/Engine/Source/Runtime/Engine/Private/Components/CharacterMovementComponent.cpp b/Engine/Source/Runtime/Engine/Private/Components/CharacterMovementComponent.cpp
--- a/Engine/Source/Runtime/Engine/Private/Components/CharacterMovementComponent.cpp
+++ b/Engine/Source/Runtime/Engine/Private/Components/CharacterMovementComponent.cpp
@@ -9124,7 +9124,7 @@ void UCharacterMovementComponent::SendClientAdjustment()
 			ServerLastClientAdjustmentTime = CurrentTime;
 
 			const bool bIsPlayingNetworkedRootMotionMontage = CharacterOwner->IsPlayingNetworkedRootMotionMontage();
-			if (HasRootMotionSources())
+			if (CurrentRootMotion.HasActiveRootMotionSources())
 			{
 				FRotator Rotation = ServerData->PendingAdjustment.NewRot.GetNormalized();
 				FVector_NetQuantizeNormal CompressedRotation(Rotation.Pitch / 180.f, Rotation.Yaw / 180.f, Rotation.Roll / 180.f);
