Fix bad merge

diff --git a/Engine/Source/Runtime/Engine/Private/PhysicsEngine/Experimental/ChaosDerivedData.cpp b/Engine/Source/Runtime/Engine/Private/PhysicsEngine/Experimental/ChaosDerivedData.cpp
--- a/Engine/Source/Runtime/Engine/Private/PhysicsEngine/Experimental/ChaosDerivedData.cpp
+++ b/Engine/Source/Runtime/Engine/Private/PhysicsEngine/Experimental/ChaosDerivedData.cpp
@@ -1,4 +1,4 @@
-// Copyright 1998-2019 Epic Games, Inc. All Rights Reserved.
+// Copyright Epic Games, Inc. All Rights Reserved.
 
 #include "ChaosDerivedData.h"
 
@@ -84,7 +84,7 @@ FChaosDerivedDataCooker::FChaosDerivedDataCooker(UBodySetup* InSetup, FName InFo
 
 }
 
-void FChaosDerivedDataCooker::BuildTriangleMeshes(TArray<TUniquePtr<Chaos::FTriangleMeshImplicitObject>>& OutTriangleMeshes, const FCookBodySetupInfo& InParams)
+void FChaosDerivedDataCooker::BuildTriangleMeshes(TArray<TUniquePtr<Chaos::FTriangleMeshImplicitObject>>& OutTriangleMeshes, TArray<int32>& OutFaceRemap, const FCookBodySetupInfo& InParams)
 {
 	if(!InParams.bCookTriMesh)
 	{
@@ -104,8 +104,7 @@ void FChaosDerivedDataCooker::BuildTriangleMeshes(TArray<TUniquePtr<Chaos::FTria
 		FinalIndices.Add(Tri.v2);
 	}
 
-	TArray<int32> Remap;
-	Chaos::CleanTrimesh(FinalVerts, FinalIndices, &Remap);
+	Chaos::CleanTrimesh(FinalVerts, FinalIndices, &OutFaceRemap);
 
 	// Build particle list #BG Maybe allow TParticles to copy vectors?
 	Chaos::TParticles<Chaos::FReal, 3> TriMeshParticles;
@@ -147,14 +146,14 @@ void FChaosDerivedDataCooker::BuildTriangleMeshes(TArray<TUniquePtr<Chaos::FTria
 
 				if(bHasMaterials)
 				{
-					if(!ensure(Remap.IsValidIndex(TriangleIndex)))
+				if(!ensure(OutFaceRemap.IsValidIndex(TriangleIndex)))
 					{
 						MaterialIndices.Empty();
 						bHasMaterials = false;
 					}
 					else
 					{
-						const int32 OriginalIndex = Remap[TriangleIndex];
+					const int32 OriginalIndex = OutFaceRemap[TriangleIndex];
 
 						if(ensure(InParams.TriangleMeshDesc.MaterialIndices.IsValidIndex(OriginalIndex)))
 						{
@@ -258,9 +257,10 @@ void FChaosDerivedDataCooker::BuildInternal(Chaos::FChaosArchive& Ar, FCookBodyS
 	TArray<TUniquePtr<Chaos::FImplicitObject>> SimpleImplicits;
 	TArray<TUniquePtr<Chaos::FTriangleMeshImplicitObject>> ComplexImplicits;
 
+	TArray<int32> FaceRemap;
 	//BuildSimpleShapes(SimpleImplicits, Setup);
 	BuildConvexMeshes(SimpleImplicits, InInfo);
-	BuildTriangleMeshes(ComplexImplicits, InInfo);
+	BuildTriangleMeshes(ComplexImplicits, FaceRemap, InInfo);
 
 	//TUniquePtr<Chaos::TImplicitObjectUnion<Precision, 3>> SimpleUnion = MakeUnique<Chaos::TImplicitObjectUnion<Precision, 3>>(MoveTemp(SimpleImplicits));
 	//TUniquePtr<Chaos::TImplicitObjectUnion<Precision, 3>> ComplexUnion = MakeUnique<Chaos::TImplicitObjectUnion<Precision, 3>>(MoveTemp(ComplexImplicits));
@@ -270,8 +270,12 @@ void FChaosDerivedDataCooker::BuildInternal(Chaos::FChaosArchive& Ar, FCookBodyS
 	{
 		UVInfo.FillFromTriMesh(InInfo.TriangleMeshDesc);
 	}
+	if (!InInfo.bSupportFaceRemap)
+	{
+		FaceRemap.Empty();
+	}
 
-	Ar << SimpleImplicits << ComplexImplicits << UVInfo;
+	Ar << SimpleImplicits << ComplexImplicits << UVInfo << FaceRemap;
 }
 
 #endif
