Rendering: Added the full path of the material instance to FMaterialResource::GetDebugName(), making error messages containing it actionable

diff --git a/Engine/Source/Runtime/Engine/Private/Materials/MaterialShared.cpp b/Engine/Source/Runtime/Engine/Private/Materials/MaterialShared.cpp
--- a/Engine/Source/Runtime/Engine/Private/Materials/MaterialShared.cpp
+++ b/Engine/Source/Runtime/Engine/Private/Materials/MaterialShared.cpp
@@ -1138,9 +1138,9 @@ FString FMaterialResource::GetDebugName() const
 {
 	if (MaterialInstance)
 	{
-		return FString::Printf(TEXT("%s (MI:%s)"), *GetBaseMaterialPathName(), *MaterialInstance->GetName());
+		return FString::Printf(TEXT("%s (MI:%s)"), *GetBaseMaterialPathName(), *MaterialInstance->GetPathName());
 	}
-	
+
 	return GetBaseMaterialPathName();
 }
 
