Added missing undef LOCTEXT_NAMESPACE

#rnx

diff --git a/Engine/Source/Runtime/Engine/Private/SkeletalMeshEditorData.cpp b/Engine/Source/Runtime/Engine/Private/SkeletalMeshEditorData.cpp
--- a/Engine/Source/Runtime/Engine/Private/SkeletalMeshEditorData.cpp
+++ b/Engine/Source/Runtime/Engine/Private/SkeletalMeshEditorData.cpp
@@ -57,4 +57,6 @@ void USkeletalMeshEditorData::Serialize(FArchive& Ar)
 	FRawSkeletalMeshBulkData::Serialize(Ar, RawSkeletalMeshBulkDatas, this);
 }
 
-#endif //WITH_EDITOR
\ No newline at end of file
+#endif //WITH_EDITORONLY_DATA
+
+#undef LOCTEXT_NAMESPACE
