Copying 10733543, 10852936, 10903639 to fix crash when importing in some cases

diff --git a/Engine/Source/Runtime/Core/Private/Misc/AssertionMacros.cpp b/Engine/Source/Runtime/Core/Private/Misc/AssertionMacros.cpp
--- a/Engine/Source/Runtime/Core/Private/Misc/AssertionMacros.cpp
+++ b/Engine/Source/Runtime/Core/Private/Misc/AssertionMacros.cpp
@@ -86,13 +86,8 @@ void PrintScriptCallstack()
 
 static void AssertFailedImplV(const ANSICHAR* Expr, const ANSICHAR* File, int32 Line, const TCHAR* Format, va_list Args)
 {
-	if (GIsCriticalError)
-	{
-		return;
-	}
-
 	// This is not perfect because another thread might crash and be handled before this assert
-	// but this static varible will report the crash as an assert. Given complexity of a thread
+	// but this static variable will report the crash as an assert. Given complexity of a thread
 	// aware solution, this should be good enough. If crash reports are obviously wrong we can
 	// look into fixing this.
 	bHasAsserted = true;
diff --git a/Engine/Source/Runtime/Core/Private/Windows/WindowsPlatformCrashContext.cpp b/Engine/Source/Runtime/Core/Private/Windows/WindowsPlatformCrashContext.cpp
--- a/Engine/Source/Runtime/Core/Private/Windows/WindowsPlatformCrashContext.cpp
+++ b/Engine/Source/Runtime/Core/Private/Windows/WindowsPlatformCrashContext.cpp
@@ -1135,6 +1135,8 @@ private:
 		// Generic exception description is stored in GErrorExceptionDescription
 		else if (ExceptionInfo->ExceptionRecord->ExceptionCode != 1)
 		{
+			// When a generic exception is thrown, it is important to get all the stack frames
+			NumStackFramesToIgnore = 0;
 			CreateExceptionInfoString(ExceptionInfo->ExceptionRecord);
 			ErrorMessage = GErrorExceptionDescription;
 		}
diff --git a/Engine/Source/Runtime/Engine/Private/DistanceFieldAtlas.cpp b/Engine/Source/Runtime/Engine/Private/DistanceFieldAtlas.cpp
--- a/Engine/Source/Runtime/Engine/Private/DistanceFieldAtlas.cpp
+++ b/Engine/Source/Runtime/Engine/Private/DistanceFieldAtlas.cpp
@@ -875,7 +875,7 @@ public:
 	}
 
 	// FRunnable interface.
-	virtual bool Init() { bIsRunning = true; return true; }
+	virtual bool Init() { return true; }
 	virtual void Exit() { bIsRunning = false; }
 	virtual void Stop() { bForceFinish = true; }
 	virtual uint32 Run();
@@ -884,8 +884,15 @@ public:
 	{
 		check(!bIsRunning);
 
+		// Calling Reset will call Kill which in turn will call Stop and set bForceFinish to true.
+		Thread.Reset();
+
+		// Now we can set bForceFinish to false without being overwritten by the old thread shutting down.
 		bForceFinish = false;
 		Thread.Reset(FRunnableThread::Create(this, *FString::Printf(TEXT("BuildDistanceFieldThread%u"), NextThreadIndex), 0, TPri_Normal, FPlatformAffinity::GetPoolThreadMask()));
+
+		// Set this now before exiting so that IsRunning() returns true without having to wait on the thread to be completely started.
+		bIsRunning = true;
 		NextThreadIndex++;
 	}
 
@@ -908,6 +915,7 @@ private:
 
 FQueuedThreadPool* CreateWorkerThreadPool()
 {
+	TRACE_CPUPROFILER_EVENT_SCOPE(CreateWorkerThreadPool)
 	const int32 NumThreads = FMath::Max<int32>(FPlatformMisc::NumberOfCoresIncludingHyperthreads() - 2, 1);
 	FQueuedThreadPool* WorkerThreadPool = FQueuedThreadPool::Allocate();
 	WorkerThreadPool->Create(NumThreads, 32 * 1024, TPri_BelowNormal);
@@ -916,9 +924,16 @@ FQueuedThreadPool* CreateWorkerThreadPool()
 
 uint32 FBuildDistanceFieldThreadRunnable::Run()
 {
+	TRACE_CPUPROFILER_EVENT_SCOPE(FBuildDistanceFieldThreadRunnable::Run)
+
 	bool bHasWork = true;
 
-	while (!bForceFinish && bHasWork)
+	// Do not exit right away if no work to do as it often leads to stop and go problems
+	// when tasks are being queued at a slower rate than the processor capability to process them.
+	const uint64 ExitAfterIdleCycle = static_cast<uint64>(10.0 / FPlatformTime::GetSecondsPerCycle64()); // 10s
+
+	uint64 LastWorkCycle = FPlatformTime::Cycles64();
+	while (!bForceFinish &&  (bHasWork || (FPlatformTime::Cycles64() - LastWorkCycle) < ExitAfterIdleCycle))
 	{
 		// LIFO build order, since meshes actually visible in a map are typically loaded last
 		FAsyncDistanceFieldTask* Task = AsyncQueue.TaskQueue.Pop();
@@ -931,11 +946,13 @@ uint32 FBuildDistanceFieldThreadRunnable::Run()
 			}
 
 			AsyncQueue.Build(Task, *WorkerThreadPool);
+			LastWorkCycle = FPlatformTime::Cycles64();
 			bHasWork = true;
 		}
 		else
 		{
 			bHasWork = false;
+			FPlatformProcess::Sleep(.01f);
 		}
 	}
 
@@ -1007,6 +1024,8 @@ void FDistanceFieldAsyncQueue::AddTask(FAsyncDistanceFieldTask* Task)
 
 void FDistanceFieldAsyncQueue::BlockUntilBuildComplete(UStaticMesh* StaticMesh, bool bWarnIfBlocked)
 {
+	TRACE_CPUPROFILER_EVENT_SCOPE(FDistanceFieldAsyncQueue::BlockUntilBuildComplete)
+
 	// We will track the wait time here, but only the cycles used.
 	// This function is called whether or not an async task is pending, 
 	// so we have to look elsewhere to properly count how many resources have actually finished building.
@@ -1056,6 +1075,7 @@ void FDistanceFieldAsyncQueue::BlockUntilBuildComplete(UStaticMesh* StaticMesh,
 
 void FDistanceFieldAsyncQueue::BlockUntilAllBuildsComplete()
 {
+	TRACE_CPUPROFILER_EVENT_SCOPE(FDistanceFieldAsyncQueue::BlockUntilAllBuildsComplete)
 	do 
 	{
 		ProcessAsyncTasks();
@@ -1067,10 +1087,11 @@ void FDistanceFieldAsyncQueue::BlockUntilAllBuildsComplete()
 void FDistanceFieldAsyncQueue::Build(FAsyncDistanceFieldTask* Task, FQueuedThreadPool& ThreadPool)
 {
 #if WITH_EDITOR
-
 	// Editor 'force delete' can null any UObject pointers which are seen by reference collecting (eg FProperty or serialized)
 	if (Task->StaticMesh && Task->GenerateSource)
 	{
+		TRACE_CPUPROFILER_EVENT_SCOPE(FDistanceFieldAsyncQueue::Build)
+
 		const FStaticMeshLODResources& LODModel = Task->GenerateSource->RenderData->LODResources[0];
 
 		MeshUtilities->GenerateSignedDistanceFieldVolumeData(
@@ -1106,6 +1127,7 @@ FString FDistanceFieldAsyncQueue::GetReferencerName() const
 
 void FDistanceFieldAsyncQueue::ProcessAsyncTasks()
 {
+	TRACE_CPUPROFILER_EVENT_SCOPE(FDistanceFieldAsyncQueue::ProcessAsyncTasks)
 #if WITH_EDITOR
 	TArray<FAsyncDistanceFieldTask*> LocalCompletedTasks;
 	CompletedTasks.PopAll(LocalCompletedTasks);
diff --git a/Engine/Source/Runtime/Engine/Private/PhysicsEngine/BodySetup.cpp b/Engine/Source/Runtime/Engine/Private/PhysicsEngine/BodySetup.cpp
--- a/Engine/Source/Runtime/Engine/Private/PhysicsEngine/BodySetup.cpp
+++ b/Engine/Source/Runtime/Engine/Private/PhysicsEngine/BodySetup.cpp
@@ -698,7 +698,6 @@ bool UBodySetup::RuntimeCookPhysics_Chaos()
 
 void UBodySetup::FinishCreatingPhysicsMeshes_Chaos(FChaosDerivedDataReader<float, 3>& InReader)
 {
-	check(IsInGameThread());
 	ClearPhysicsMeshes();
 
 	const FString FullName = GetFullName();
diff --git a/Engine/Source/Runtime/Engine/Private/PhysicsEngine/PhysDrawing.cpp b/Engine/Source/Runtime/Engine/Private/PhysicsEngine/PhysDrawing.cpp
--- a/Engine/Source/Runtime/Engine/Private/PhysicsEngine/PhysDrawing.cpp
+++ b/Engine/Source/Runtime/Engine/Private/PhysicsEngine/PhysDrawing.cpp
@@ -1015,44 +1015,46 @@ void FKAggregateGeom::GetAggGeom(const FTransform& Transform, const FColor Color
 	}
 }
 
-/** Release the RenderInfo (if its there) and safely clean up any resources. Call on the game thread. */
+/** Release the RenderInfo (if its there) and safely clean up any resources. Not thread safe, but can be called from any thread (conditionally safe). */
 void FKAggregateGeom::FreeRenderInfo()
 {
 	// See if we have rendering resources to free
-	if(RenderInfo)
+	if (RenderInfo)
 	{
 		// Should always have these if RenderInfo exists
 		check(RenderInfo->VertexBuffers);
 		check(RenderInfo->IndexBuffer);
 
-		// Fire off commands to free these resources
-		BeginReleaseResource(&RenderInfo->VertexBuffers->ColorVertexBuffer);
-		BeginReleaseResource(&RenderInfo->VertexBuffers->StaticMeshVertexBuffer);
-		BeginReleaseResource(&RenderInfo->VertexBuffers->PositionVertexBuffer);
-		BeginReleaseResource(RenderInfo->IndexBuffer);
+		// Fire off a render command to free these resources
+		ENQUEUE_RENDER_COMMAND(FKAggregateGeomFreeRenderInfo)(
+			[RenderInfoToRelease = RenderInfo](FRHICommandList& RHICmdList)
+			{
+				RenderInfoToRelease->VertexBuffers->ColorVertexBuffer.ReleaseResource();
+				RenderInfoToRelease->VertexBuffers->StaticMeshVertexBuffer.ReleaseResource();
+				RenderInfoToRelease->VertexBuffers->PositionVertexBuffer.ReleaseResource();
+				RenderInfoToRelease->IndexBuffer->ReleaseResource();
 
-		// May not exist if no geometry was available
-		if(RenderInfo->CollisionVertexFactory != NULL)
-		{
-			BeginReleaseResource(RenderInfo->CollisionVertexFactory);
-		}
+				// May not exist if no geometry was available
+				if (RenderInfoToRelease->CollisionVertexFactory != nullptr)
+				{
+					RenderInfoToRelease->CollisionVertexFactory->ReleaseResource();
+				}
 
-		// Wait until those commands have been processed
-		FRenderCommandFence Fence;
-		Fence.BeginFence();
-		Fence.Wait();
+				// Free memory.
+				delete RenderInfoToRelease->VertexBuffers;
+				delete RenderInfoToRelease->IndexBuffer;
 
-		// Release memory.
-		delete RenderInfo->VertexBuffers;
-		delete RenderInfo->IndexBuffer;
+				if (RenderInfoToRelease->CollisionVertexFactory != nullptr)
+				{
+					delete RenderInfoToRelease->CollisionVertexFactory;
+				}
 
-		if (RenderInfo->CollisionVertexFactory != NULL)
-		{
-			delete RenderInfo->CollisionVertexFactory;
-		}
+				delete RenderInfoToRelease;
+			}
+		);
 
-		delete RenderInfo;
-		RenderInfo = NULL;
+		// Reset the pointer as it's been given to the render thread
+		RenderInfo = nullptr;
 	}
 }
 
diff --git a/Engine/Source/Runtime/Engine/Private/StaticMeshBuild.cpp b/Engine/Source/Runtime/Engine/Private/StaticMeshBuild.cpp
--- a/Engine/Source/Runtime/Engine/Private/StaticMeshBuild.cpp
+++ b/Engine/Source/Runtime/Engine/Private/StaticMeshBuild.cpp
@@ -154,6 +154,16 @@ void UStaticMesh::BatchBuild(const TArray<UStaticMesh*>& InStaticMeshes, bool bI
 		ITargetPlatform* RunningPlatform = TargetPlatformManager.GetRunningTargetPlatform();
 		check(RunningPlatform);
 
+		for (UStaticMesh* StaticMesh : StaticMeshesToProcess)
+		{
+			if (StaticMesh->RenderData)
+			{
+				// Finish any previous async builds before modifying RenderData
+				// This can happen during import as the mesh is rebuilt redundantly
+				GDistanceFieldAsyncQueue->BlockUntilBuildComplete(StaticMesh, true);
+			}
+		}
+
 		// Detach all instances of those static meshes from the scene.
 		FStaticMeshComponentRecreateRenderStateContext RecreateRenderStateContext(StaticMeshesToProcess, false);
 
@@ -251,13 +261,6 @@ void UStaticMesh::PreBuildInternal()
 	// Flush the resource release commands to the rendering thread to ensure that the build doesn't occur while a resource is still
 	// allocated, and potentially accessing the UStaticMesh.
 	ReleaseResourcesFence.Wait();
-
-	if (RenderData)
-	{
-		// Finish any previous async builds before modifying RenderData
-		// This can happen during import as the mesh is rebuilt redundantly
-		GDistanceFieldAsyncQueue->BlockUntilBuildComplete(this, true);
-	}
 }
 
 bool UStaticMesh::BuildInternal(bool bInSilent, TArray<FText> * OutErrors)
diff --git a/Engine/Source/Runtime/Engine/Private/UnrealEngine.cpp b/Engine/Source/Runtime/Engine/Private/UnrealEngine.cpp
--- a/Engine/Source/Runtime/Engine/Private/UnrealEngine.cpp
+++ b/Engine/Source/Runtime/Engine/Private/UnrealEngine.cpp
@@ -8260,7 +8260,31 @@ bool UEngine::PerformError(const TCHAR* Cmd, FOutputDevice& Ar)
 		);
 		return true;
 	}
-	if (FParse::Command(&Cmd, TEXT("TWOTHREADSCRASH")))
+
+	if (FParse::Command(&Cmd, TEXT("THREADEDCHECKS")))
+	{
+		UE_LOG(LogEngine, Warning, TEXT("Queuing several multi-thread checks."));
+		for (int32 Index = 0; Index < 16; ++Index)
+		{
+			Async(
+				EAsyncExecution::ThreadPool,
+				[]()
+				{
+					FGenericCrashContext::SetCrashTrigger(ECrashTrigger::Debug);
+
+					// We now test if assert behavior is deterministic when called from different threads or
+					// if one thread could end up with a different behavior due to a race condition.
+					// We expect that the result will always be the assert being reported and not end-up
+					// with an Int 3 being handled by the exception handler.
+					const bool CrashingTheWorkerThreadAtYourRequest = false;
+					check(CrashingTheWorkerThreadAtYourRequest);
+				}
+			);
+		};
+
+		return true;
+	}
+	else if (FParse::Command(&Cmd, TEXT("TWOTHREADSCRASH")))
 	{
 		class FThreadPoolCrash : public IQueuedWork
 		{
