Timecode: Take Recorder:

Temp fix for timecode provider not being set up which breaks Take Recorder and possibly some other fixes. Better fix will come in from Dev-VP.

diff --git a/Engine/Source/Runtime/Engine/Private/UnrealEngine.cpp b/Engine/Source/Runtime/Engine/Private/UnrealEngine.cpp
--- a/Engine/Source/Runtime/Engine/Private/UnrealEngine.cpp
+++ b/Engine/Source/Runtime/Engine/Private/UnrealEngine.cpp
@@ -2410,7 +2410,11 @@ static void LoadTimecodeProvider(UEngine* Engine)
 	}
 	else
 	{
-		Engine->SetTimecodeProvider(nullptr);
+#if PLATFORM_DESKTOP
+		//HACK until Dev-VP comes in!!!
+		UTimecodeProvider* NewTimecodeProvider = NewObject<USystemTimeTimecodeProvider>(Engine, TEXT("SystemTimeTimecodeProvider"));
+		Engine->SetTimecodeProvider(NewTimecodeProvider);
+#endif
 	}
 }
 
