Copied cl 10954136:

Fixing convex collision crash

diff --git a/Engine/Source/Runtime/Engine/Private/Collision/KAggregateGeom.cpp b/Engine/Source/Runtime/Engine/Private/Collision/KAggregateGeom.cpp
--- a/Engine/Source/Runtime/Engine/Private/Collision/KAggregateGeom.cpp
+++ b/Engine/Source/Runtime/Engine/Private/Collision/KAggregateGeom.cpp
@@ -541,7 +541,6 @@ bool FKConvexElem::HullFromPlanes(const TArray<FPlane>& InPlanes, const TArray<F
 		// Do nothing if poly was completely clipped away.
 		if(Polygon.Vertices.Num() > 0)
 		{
-			const int32 BaseIndex = VertexData.Num();
 			TArray<int32> Remap;
 			Remap.AddUninitialized(Polygon.Vertices.Num());
 
@@ -578,12 +577,13 @@ bool FKConvexElem::HullFromPlanes(const TArray<FPlane>& InPlanes, const TArray<F
 				}
 			}
 
-			int32 NumTriangles = Polygon.Vertices.Num() - 2;
-			for(int32 Index = 1; Index < NumTriangles; ++Index)
+			const int32 NumTriangles = Polygon.Vertices.Num() - 2;
+			const int32 BaseIndex = Remap[0];
+			for(int32 Index = 0; Index < NumTriangles; ++Index)
 			{
 				IndexData.Add(BaseIndex);
-				IndexData.Add(Remap[Index]);
 				IndexData.Add(Remap[Index + 1]);
+				IndexData.Add(Remap[Index + 2]);
 			}
 		}
 	}
