#UE4 Fix occational crash related to BP compiling in -Game or -Server

diff --git a/Engine/Source/Runtime/Engine/Private/EdGraph/EdGraphPin.cpp b/Engine/Source/Runtime/Engine/Private/EdGraph/EdGraphPin.cpp
--- a/Engine/Source/Runtime/Engine/Private/EdGraph/EdGraphPin.cpp
+++ b/Engine/Source/Runtime/Engine/Private/EdGraph/EdGraphPin.cpp
@@ -1830,7 +1830,7 @@ void UEdGraphPin::SerializePinArray(FArchive& Ar, TArray<UEdGraphPin*>& ArrayRef
 				{
 					// PeerNode will be null if the pin we were linked to was already thrown away:
 					UEdGraphNode* PeerNode = Peer->GetOwningNodeUnchecked();
-					if (PeerNode && !GEditor->Trans->IsObjectTransacting(PeerNode))
+					if (PeerNode && (!GEditor || !GEditor->Trans->IsObjectTransacting(PeerNode)))
 					{
 						Pin->BreakLinkTo(Peer);
 						--I;
