Add range check to collision debug draw.

diff --git a/Engine/Source/Runtime/Engine/Private/PhysicsEngine/PhysDrawing.cpp b/Engine/Source/Runtime/Engine/Private/PhysicsEngine/PhysDrawing.cpp
--- a/Engine/Source/Runtime/Engine/Private/PhysicsEngine/PhysDrawing.cpp
+++ b/Engine/Source/Runtime/Engine/Private/PhysicsEngine/PhysDrawing.cpp
@@ -704,6 +704,11 @@ void FKConvexElem::DrawElemWire(FPrimitiveDrawInterface* PDI, const FTransform&
 
 		for(int32 Base = 0; Base < NumIndices; Base += 3)
 		{
+			if (IndexData[Base] >= NumVerts || IndexData[Base + 1] >= NumVerts || IndexData[Base + 2] >= NumVerts)
+			{
+				continue;
+			}
+
 			PDI->DrawLine(TransformedVerts[IndexData[Base]], TransformedVerts[IndexData[Base + 1]], Color, SDPG_World);
 			PDI->DrawLine(TransformedVerts[IndexData[Base + 1]], TransformedVerts[IndexData[Base + 2]], Color, SDPG_World);
 			PDI->DrawLine(TransformedVerts[IndexData[Base + 2]], TransformedVerts[IndexData[Base]], Color, SDPG_World);
