Fix an issue with Fast Array Delta Serialization during PIE.

The DeltaFlags field is used to track state of delta serialization, and needs to be a UPROPERTY for reflection reasons.
However, it was not marked Transient and that could cause issues when objects were duplicated for PIE.

diff --git a/Engine/Source/Runtime/Engine/Classes/Engine/NetSerialization.h b/Engine/Source/Runtime/Engine/Classes/Engine/NetSerialization.h
--- a/Engine/Source/Runtime/Engine/Classes/Engine/NetSerialization.h
+++ b/Engine/Source/Runtime/Engine/Classes/Engine/NetSerialization.h
@@ -671,7 +671,7 @@ private:
 	int32 CachedNumItems;
 	int32 CachedNumItemsToConsiderForWriting;
 
-	UPROPERTY(NotReplicated)
+	UPROPERTY(NotReplicated, Transient)
 	EFastArraySerializerDeltaFlags DeltaFlags;
 };
 
