Fix editor crash in FixupOverrideColors if corrupt static mesh has no LODs

diff --git a/Engine/Source/Runtime/Engine/Private/Components/StaticMeshComponent.cpp b/Engine/Source/Runtime/Engine/Private/Components/StaticMeshComponent.cpp
--- a/Engine/Source/Runtime/Engine/Private/Components/StaticMeshComponent.cpp
+++ b/Engine/Source/Runtime/Engine/Private/Components/StaticMeshComponent.cpp
@@ -1252,6 +1252,11 @@ void UStaticMeshComponent::PrivateFixupOverrideColors()
 
 	const uint32 NumLODs = GetStaticMesh()->RenderData->LODResources.Num();
 
+	if (NumLODs == 0)
+	{
+		return;
+	}
+
 	// Initialize override vertex colors on any new LODs which have just been created
 	SetLODDataCount(NumLODs, LODData.Num());
 	bool UpdateStaticMeshDeriveDataKey = false;
