Chaos : CVar for enable/disable mesh cleaning during cook - lets users import massive meshes without having any additional performance hit

diff --git a/Engine/Source/Runtime/Engine/Private/PhysicsEngine/Experimental/ChaosDerivedData.cpp b/Engine/Source/Runtime/Engine/Private/PhysicsEngine/Experimental/ChaosDerivedData.cpp
--- a/Engine/Source/Runtime/Engine/Private/PhysicsEngine/Experimental/ChaosDerivedData.cpp
+++ b/Engine/Source/Runtime/Engine/Private/PhysicsEngine/Experimental/ChaosDerivedData.cpp
@@ -14,6 +14,9 @@
 #include "PhysicsEngine/BodySetup.h"
 #include "Serialization/MemoryWriter.h"
 
+int32 EnableMeshClean = 1;
+FAutoConsoleVariableRef CVarEnableMeshClean(TEXT("p.EnableMeshClean"), EnableMeshClean, TEXT("Enable/Disable mesh cleanup during cook."));
+
 const TCHAR* FChaosDerivedDataCooker::GetPluginName() const
 {
 	return TEXT("ChaosGeometryData");
@@ -104,7 +107,10 @@ void FChaosDerivedDataCooker::BuildTriangleMeshes(TArray<TUniquePtr<Chaos::FTria
 		FinalIndices.Add(Tri.v2);
 	}
 
-	Chaos::CleanTrimesh(FinalVerts, FinalIndices, &OutFaceRemap);
+	if (EnableMeshClean)
+	{
+		Chaos::CleanTrimesh(FinalVerts, FinalIndices, &OutFaceRemap);
+	}
 
 	// Build particle list #BG Maybe allow TParticles to copy vectors?
 	Chaos::TParticles<Chaos::FReal, 3> TriMeshParticles;
@@ -146,18 +152,33 @@ void FChaosDerivedDataCooker::BuildTriangleMeshes(TArray<TUniquePtr<Chaos::FTria
 
 				if(bHasMaterials)
 				{
-				if(!ensure(OutFaceRemap.IsValidIndex(TriangleIndex)))
+					if (EnableMeshClean)				
 					{
-						MaterialIndices.Empty();
-						bHasMaterials = false;
+						if(!ensure(OutFaceRemap.IsValidIndex(TriangleIndex)))
+						{
+							MaterialIndices.Empty();
+							bHasMaterials = false;
+						}
+						else
+						{
+							const int32 OriginalIndex = OutFaceRemap[TriangleIndex];
+
+							if (ensure(InParams.TriangleMeshDesc.MaterialIndices.IsValidIndex(OriginalIndex)))
+							{
+								MaterialIndices.Add(InParams.TriangleMeshDesc.MaterialIndices[OriginalIndex]);
+							}
+							else
+							{
+								MaterialIndices.Empty();
+								bHasMaterials = false;
+							}
+						}
 					}
 					else
 					{
-					const int32 OriginalIndex = OutFaceRemap[TriangleIndex];
-
-						if(ensure(InParams.TriangleMeshDesc.MaterialIndices.IsValidIndex(OriginalIndex)))
+						if (ensure(InParams.TriangleMeshDesc.MaterialIndices.IsValidIndex(TriangleIndex)))
 						{
-							MaterialIndices.Add(InParams.TriangleMeshDesc.MaterialIndices[OriginalIndex]);
+							MaterialIndices.Add(InParams.TriangleMeshDesc.MaterialIndices[TriangleIndex]);
 						}
 						else
 						{
@@ -165,6 +186,7 @@ void FChaosDerivedDataCooker::BuildTriangleMeshes(TArray<TUniquePtr<Chaos::FTria
 							bHasMaterials = false;
 						}
 					}
+
 				}
 			}
 		}
