Fix boolean logic that meant that deleting cull distance volumes in the editor did not update the cached cull distances on meshes

diff --git a/Engine/Source/Runtime/Engine/Private/CullDistanceVolume.cpp b/Engine/Source/Runtime/Engine/Private/CullDistanceVolume.cpp
--- a/Engine/Source/Runtime/Engine/Private/CullDistanceVolume.cpp
+++ b/Engine/Source/Runtime/Engine/Private/CullDistanceVolume.cpp
@@ -22,10 +22,13 @@ void ACullDistanceVolume::Destroyed()
 {
 	Super::Destroyed();
 
-	UWorld* World = GetWorld();
-	if (GIsEditor && World && World->IsGameWorld())
+	if (GIsEditor)
 	{
-		World->bDoDelayedUpdateCullDistanceVolumes = true;
+		UWorld* World = GetWorld();
+		if (World && !World->IsGameWorld())
+		{
+			World->bDoDelayedUpdateCullDistanceVolumes = true;
+		}
 	}
 }
 
