Fixing GC crash when exiting PIE

diff --git a/Engine/Source/Runtime/Engine/Classes/Engine/GameViewportClient.h b/Engine/Source/Runtime/Engine/Classes/Engine/GameViewportClient.h
--- a/Engine/Source/Runtime/Engine/Classes/Engine/GameViewportClient.h
+++ b/Engine/Source/Runtime/Engine/Classes/Engine/GameViewportClient.h
@@ -140,6 +140,7 @@ public:
 	//~ Begin UObject Interface
 	virtual void PostInitProperties() override;
 	virtual void BeginDestroy() override;
+	virtual bool IsDestructionThreadSafe() const override { return false; }
 	//~ End UObject Interface
 
 	//~ Begin FViewportClient Interface.
