Made FTickableGameObject registration and de-registration thread safe (registration is lock free)

diff --git a/Engine/Source/Runtime/Engine/Private/LevelTick.cpp b/Engine/Source/Runtime/Engine/Private/LevelTick.cpp
--- a/Engine/Source/Runtime/Engine/Private/LevelTick.cpp
+++ b/Engine/Source/Runtime/Engine/Private/LevelTick.cpp
@@ -1324,14 +1324,25 @@ void FTickableGameObject::TickObjects(UWorld* World, const int32 InTickType, con
 	SCOPE_CYCLE_COUNTER(STAT_TickableGameObjectsTime);
 	CSV_SCOPED_TIMING_STAT_EXCLUSIVE(Tickables);
 
-	TArray<FTickableGameObject*>& PendingTickableObjects = GetPendingTickableObjects();
+	check(IsInGameThread());
+
+	TLockFreePointerListUnordered<FTickableGameObject, PLATFORM_CACHE_LINE_SIZE>& NewTickableObjects = GetNewTickableObjects();	
+
+	// It's a long lock but it's ok, the only thing we can block here is the GC worker thread that destroys UObjects
+	FScopeLock LockTickableObjects(&GetTickableObjectsCritical());
+
 	TArray<FTickableObjectEntry>& TickableObjects = GetTickableObjects();
+	TSet<FTickableGameObject*>& DeletedTickableObjects = GetDeletedTickableObjects();
 
-	for (FTickableGameObject* PendingTickable : PendingTickableObjects)
+	while (FTickableGameObject* NewTickableObject = NewTickableObjects.Pop())
 	{
-		AddTickableObject(TickableObjects, PendingTickable);
+		// If the new tickable object has already been deleted, don't add it to the tickable objects list
+		if (!DeletedTickableObjects.Contains(NewTickableObject))
+		{
+			AddTickableObject(TickableObjects, NewTickableObject);
+		}
 	}
-	PendingTickableObjects.Empty();
+	DeletedTickableObjects.Empty();
 
 	if (TickableObjects.Num() > 0)
 	{
diff --git a/Engine/Source/Runtime/Engine/Public/Tickable.h b/Engine/Source/Runtime/Engine/Public/Tickable.h
--- a/Engine/Source/Runtime/Engine/Public/Tickable.h
+++ b/Engine/Source/Runtime/Engine/Public/Tickable.h
@@ -135,10 +135,24 @@ private:
 		return TickableObjects;
 	}
 
-	static TArray<FTickableGameObject*>& GetPendingTickableObjects()
+	static FCriticalSection& GetTickableObjectsCritical()
 	{
-		static TArray<FTickableGameObject*> PendingTickableObjects;
-		return PendingTickableObjects;
+		static FCriticalSection TickableObjectsCritical;
+		return TickableObjectsCritical;
+	}
+
+	/** List of newly created objects */
+	static TLockFreePointerListUnordered<FTickableGameObject, PLATFORM_CACHE_LINE_SIZE>& GetNewTickableObjects()
+	{
+		static TLockFreePointerListUnordered<FTickableGameObject, PLATFORM_CACHE_LINE_SIZE> NewTickableObjects;
+		return NewTickableObjects;
+	}
+
+	/** Deleted list is for cases when an object gets deleted before it was added to the tickable objects array from the NewTickableObjects list */
+	static TSet<FTickableGameObject*>& GetDeletedTickableObjects()
+	{
+		static TSet<FTickableGameObject*> DeletedTickableObjects;
+		return DeletedTickableObjects;
 	}
 
 	static bool bIsTickingObjects;
@@ -150,22 +164,17 @@ public:
 	 */
 	FTickableGameObject()
 	{
-		ensure(IsInGameThread() || IsInAsyncLoadingThread());
-		check(!GetPendingTickableObjects().Contains(this));
-		check(!GetTickableObjects().Contains(this));
-		GetPendingTickableObjects().Add(this);
+		GetNewTickableObjects().Push(this);
 	}
 
 	/**
 	 * Removes this instance from the static array of tickable objects.
 	 */
 	virtual ~FTickableGameObject()
-	{
-		ensure(IsInGameThread() || IsInAsyncLoadingThread());
-		if (GetPendingTickableObjects().Remove(this) == 0)
-		{
-			RemoveTickableObject(GetTickableObjects(), this, bIsTickingObjects);
-		}
+	{		
+		FScopeLock LockTickableObjects(&GetTickableObjectsCritical());
+		GetDeletedTickableObjects().Add(this);
+		RemoveTickableObject(GetTickableObjects(), this, bIsTickingObjects);
 	}
 
 	/**
