CIS fix

diff --git a/Engine/Source/Runtime/Engine/Public/Tickable.h b/Engine/Source/Runtime/Engine/Public/Tickable.h
--- a/Engine/Source/Runtime/Engine/Public/Tickable.h
+++ b/Engine/Source/Runtime/Engine/Public/Tickable.h
@@ -9,6 +9,7 @@
 
 #include "CoreMinimal.h"
 #include "Stats/Stats.h"
+#include "Misc/ScopeLock.h"
 
 /**
  * Enum used to convey whether a tickable object will always tick, conditionally tick, or never tick.
