Fixed GC crash

diff --git a/Engine/Source/Runtime/Engine/Classes/Atmosphere/AtmosphericFogComponent.h b/Engine/Source/Runtime/Engine/Classes/Atmosphere/AtmosphericFogComponent.h
--- a/Engine/Source/Runtime/Engine/Classes/Atmosphere/AtmosphericFogComponent.h
+++ b/Engine/Source/Runtime/Engine/Classes/Atmosphere/AtmosphericFogComponent.h
@@ -278,6 +278,11 @@ public:
 private:
 #if WITH_EDITORONLY_DATA
 	class FAtmospherePrecomputeDataHandler* PrecomputeDataHandler;
+
+public:
+	// Because FAtmospherePrecomputeDataHandler is a FTicableEditorObject and its destruction is not thread safe
+	virtual bool IsDestructionThreadSafe() const override { return false; }
+private:
 #endif
 
 	friend class FAtmosphericFogSceneInfo;
