Merge fix for UEngine::ToggleStatSounds which somehow lost its body in the merge down

diff --git a/Engine/Source/Runtime/Engine/Private/UnrealEngine.cpp b/Engine/Source/Runtime/Engine/Private/UnrealEngine.cpp
--- a/Engine/Source/Runtime/Engine/Private/UnrealEngine.cpp
+++ b/Engine/Source/Runtime/Engine/Private/UnrealEngine.cpp
@@ -15552,6 +15552,11 @@ int32 UEngine::RenderStatSoundCues(UWorld* World, FViewport* Viewport, FCanvas*
 // SOUNDS
 bool UEngine::ToggleStatSounds(UWorld* World, FCommonViewportClient* ViewportClient, const TCHAR* Stream)
 {
+#if ENABLE_AUDIO_DEBUG
+	return FAudioDebugger::ToggleStatSounds(World, ViewportClient, Stream);
+#else // !ENABLE_AUDIO_DEBUG
+	return false;
+#endif // !ENABLE_AUDIO_DEBUG
 }
 
 int32 UEngine::RenderStatSounds(UWorld* World, FViewport* Viewport, FCanvas* Canvas, int32 X, int32 Y, const FVector* ViewLocation, const FRotator* ViewRotation)
