Mergedown fixes - FName code upgrade

diff --git a/Engine/Source/Runtime/Engine/Private/AudioDebug.cpp b/Engine/Source/Runtime/Engine/Private/AudioDebug.cpp
--- a/Engine/Source/Runtime/Engine/Private/AudioDebug.cpp
+++ b/Engine/Source/Runtime/Engine/Private/AudioDebug.cpp
@@ -731,7 +731,7 @@ int32 FAudioDebugger::RenderStatSounds(UWorld* World, FViewport* Viewport, FCanv
 	}
 	else if (AudioStats.DisplayFlags & static_cast<uint8>(FAudioStats::EDisplayFlags::Sort_Class))
 	{
-		AudioStats.StatSoundInfos.Sort([](const FAudioStats::FStatSoundInfo& A, const FAudioStats::FStatSoundInfo& B) { return A.SoundClassName < B.SoundClassName; });
+		AudioStats.StatSoundInfos.Sort([](const FAudioStats::FStatSoundInfo& A, const FAudioStats::FStatSoundInfo& B) { return A.SoundClassName.LexicalLess(B.SoundClassName); });
 		SortingName = TEXT("class");
 	}
 	else if (AudioStats.DisplayFlags & static_cast<uint8>(FAudioStats::EDisplayFlags::Sort_WavesNum))
