Added some explicit initialization code to ensure globals get destroyed in the correct order.

Before I added this we would end up with double deletes like these (Valgrind report):

==20863== Invalid free() / delete / delete[] / realloc()
==20863==    at 0x4C30D3B: free (in /usr/lib/valgrind/vgpreload_memcheck-amd64-linux.so)
==20863==    by 0x5461706: FMallocAnsi::Realloc(void*, unsigned long, unsigned int) (MallocAnsi.cpp:148)
==20863==    by 0x54864BD: FMemory::Realloc(void*, unsigned long, unsigned int) (FMemory.inl:50)
==20863==    by 0x8DD7EFD: ResizeAllocation (ContainerAllocationPolicies.h:414)
==20863==    by 0x8DD7EFD: ResizeAllocation (ContainerAllocationPolicies.h:534)
==20863==    by 0x8DD7EFD: TSet<FTickableGameObject*, DefaultKeyFuncs<FTickableGameObject*, false>, FDefaultSetAllocator>::Rehash() const (Set.h:1064)
==20863==    by 0x8DD7DFC: TSet<FTickableGameObject*, DefaultKeyFuncs<FTickableGameObject*, false>, FDefaultSetAllocator>::ConditionalRehash(int, bool) const (Set.h:1051)
==20863==    by 0x8DD79E3: FSetElementId TSet<FTickableGameObject*, DefaultKeyFuncs<FTickableGameObject*, false>, FDefaultSetAllocator>::Emplace<FTickableGameObject*>(FTickableGameObject*&&, bool*) (Set.h:535)
==20863==    by 0x8DD7666: Add (Set.h:491)
==20863==    by 0x8DD7666: FTickableGameObject::~FTickableGameObject() (Tickable.h:177)
==20863==    by 0x11960040: __run_exit_handlers (exit.c:108)
==20863==    by 0x11960139: exit (exit.c:139)
==20863==    by 0x1193EB9D: (below main) (libc-start.c:344)
==20863==  Address 0x5afecac0 is 0 bytes inside a block of size 160 free'd
==20863==    at 0x4C30D3B: free (in /usr/lib/valgrind/vgpreload_memcheck-amd64-linux.so)
==20863==    by 0x8DD7838: ~ForAnyElementType (ContainerAllocationPolicies.h:399)
==20863==    by 0x8DD7838: ~ForElementType (ContainerAllocationPolicies.h:488)
==20863==    by 0x8DD7838: TSet<FTickableGameObject*, DefaultKeyFuncs<FTickableGameObject*, false>, FDefaultSetAllocator>::~TSet() (Set.h:248)
==20863==    by 0x11960040: __run_exit_handlers (exit.c:108)
==20863==    by 0x11960139: exit (exit.c:139)
==20863==    by 0x1193EB9D: (below main) (libc-start.c:344)
==20863==  Block was alloc'd at
==20863==    at 0x4C2FB0F: malloc (in /usr/lib/valgrind/vgpreload_memcheck-amd64-linux.so)
==20863==    by 0x54615DF: FMallocAnsi::Malloc(unsigned long, unsigned int) (MallocAnsi.cpp:60)
==20863==    by 0x546171F: FMallocAnsi::Realloc(void*, unsigned long, unsigned int) (MallocAnsi.cpp:144)
==20863==    by 0x54864BD: FMemory::Realloc(void*, unsigned long, unsigned int) (FMemory.inl:50)
==20863==    by 0x8DD7F4E: ResizeAllocation (ContainerAllocationPolicies.h:414)
==20863==    by 0x8DD7F4E: ResizeAllocation (ContainerAllocationPolicies.h:550)
==20863==    by 0x8DD7F4E: TSet<FTickableGameObject*, DefaultKeyFuncs<FTickableGameObject*, false>, FDefaultSetAllocator>::Rehash() const (Set.h:1071)
==20863==    by 0x8DD7DFC: TSet<FTickableGameObject*, DefaultKeyFuncs<FTickableGameObject*, false>, FDefaultSetAllocator>::ConditionalRehash(int, bool) const (Set.h:1051)
==20863==    by 0x8DD79E3: FSetElementId TSet<FTickableGameObject*, DefaultKeyFuncs<FTickableGameObject*, false>, FDefaultSetAllocator>::Emplace<FTickableGameObject*>(FTickableGameObject*&&, bool*) (Set.h:535)
==20863==    by 0x8DD7666: Add (Set.h:491)
==20863==    by 0x8DD7666: FTickableGameObject::~FTickableGameObject() (Tickable.h:177)
==20863==    by 0x2D36B502: ~UAISubsystem (AISubsystem.h:15)
==20863==    by 0x2D36B502: UAIPerceptionSystem::~UAIPerceptionSystem() (AIPerceptionSystem.h:29)
==20863==    by 0x6481367: TickDestroyGameThreadObjects (GarbageCollection.cpp:316)
==20863==    by 0x6481367: FAsyncPurge::TickPurge(bool, float, double) (GarbageCollection.cpp:412)
==20863==    by 0x6440017: IncrementalDestroyGarbage (GarbageCollection.cpp:1592)
==20863==    by 0x6440017: IncrementalPurgeGarbage(bool, float) (GarbageCollection.cpp:1344)
==20863==    by 0x64D36FD: StaticExit() (Obj.cpp:4409)

diff --git a/Engine/Source/Runtime/Engine/Public/Tickable.h b/Engine/Source/Runtime/Engine/Public/Tickable.h
--- a/Engine/Source/Runtime/Engine/Public/Tickable.h
+++ b/Engine/Source/Runtime/Engine/Public/Tickable.h
@@ -165,6 +165,12 @@ public:
 	 */
 	FTickableGameObject()
 	{
+		// These are called here only to make sure the function level statics they encapsulate
+		// get destroyed in the correct order at shutdown. This was added to fix UE-75429
+		GetTickableObjects();
+		GetDeletedTickableObjects();
+		GetTickableObjectsCritical();
+
 		GetNewTickableObjects().Push(this);
 	}
 
