Fixed use of unitialized member in anim instance proxy (flagged by Valgrind)

diff --git a/Engine/Source/Runtime/Engine/Public/Animation/AnimInstanceProxy.h b/Engine/Source/Runtime/Engine/Public/Animation/AnimInstanceProxy.h
--- a/Engine/Source/Runtime/Engine/Public/Animation/AnimInstanceProxy.h
+++ b/Engine/Source/Runtime/Engine/Public/Animation/AnimInstanceProxy.h
@@ -864,7 +864,7 @@ private:
 	FBoneContainer RequiredBones;
 
 	/** LODLevel used by RequiredBones */
-	int32 LODLevel;
+	int32 LODLevel = 0;
 
 	/** Counter used to control CacheBones recursion behavior - makes sure we cache bones correctly when recursing into different subgraphs */
 	int32 CacheBonesRecursionCounter;
