Out-of-line FTickableGameObject implementation with single set of linked static variables rather than per-module, to better control construction order (and be correct).

diff --git a/Engine/Source/Runtime/Engine/Private/LevelTick.cpp b/Engine/Source/Runtime/Engine/Private/LevelTick.cpp
--- a/Engine/Source/Runtime/Engine/Private/LevelTick.cpp
+++ b/Engine/Source/Runtime/Engine/Private/LevelTick.cpp
@@ -141,8 +141,6 @@ DEFINE_STAT(STAT_NetBroadcastPostTickTime);
 DEFINE_STAT(STAT_NetRebuildConditionalTime);
 DEFINE_STAT(STAT_PackageMap_SerializeObjectTime);
 
-DECLARE_CYCLE_STAT(TEXT("TickableGameObjects Time"), STAT_TickableGameObjectsTime, STATGROUP_Game);
-
 
 /*-----------------------------------------------------------------------------
 	Externs.
@@ -151,13 +149,6 @@ DECLARE_CYCLE_STAT(TEXT("TickableGameObjects Time"), STAT_TickableGameObjectsTim
 extern bool GShouldLogOutAFrameOfMoveComponent;
 extern bool GShouldLogOutAFrameOfSetBodyTransform;
 
-/*-----------------------------------------------------------------------------
-	FTickableGameObject implementation.
------------------------------------------------------------------------------*/
-
-/** Static array of tickable objects */
-bool FTickableGameObject::bIsTickingObjects = false;
-
 #if LOG_DETAILED_PATHFINDING_STATS
 /** Global detailed pathfinding stats. */
 FDetailedTickStats GDetailedPathFindingStats(30, 10, 1, 20, TEXT("pathfinding"));
@@ -1318,79 +1309,6 @@ void EndTickDrawEvent(TDrawEvent<FRHICommandList>* TickDrawEvent)
 		});
 }
 
-
-void FTickableGameObject::TickObjects(UWorld* World, const int32 InTickType, const bool bIsPaused, const float DeltaSeconds)
-{
-	SCOPE_CYCLE_COUNTER(STAT_TickableGameObjectsTime);
-	CSV_SCOPED_TIMING_STAT_EXCLUSIVE(Tickables);
-
-	check(IsInGameThread());
-
-	TLockFreePointerListUnordered<FTickableGameObject, PLATFORM_CACHE_LINE_SIZE>& NewTickableObjects = GetNewTickableObjects();	
-
-	// It's a long lock but it's ok, the only thing we can block here is the GC worker thread that destroys UObjects
-	FScopeLock LockTickableObjects(&GetTickableObjectsCritical());
-
-	TArray<FTickableObjectEntry>& TickableObjects = GetTickableObjects();
-	TSet<FTickableGameObject*>& DeletedTickableObjects = GetDeletedTickableObjects();
-
-	while (FTickableGameObject* NewTickableObject = NewTickableObjects.Pop())
-	{
-		// If the new tickable object has already been deleted, don't add it to the tickable objects list
-		if (!DeletedTickableObjects.Contains(NewTickableObject))
-		{
-			AddTickableObject(TickableObjects, NewTickableObject);
-		}
-	}
-	DeletedTickableObjects.Empty();
-
-	if (TickableObjects.Num() > 0)
-	{
-		check(!bIsTickingObjects);
-		bIsTickingObjects = true;
-
-		bool bNeedsCleanup = false;
-		const ELevelTick TickType = (ELevelTick)InTickType;
-
-		for (const FTickableObjectEntry& TickableEntry : TickableObjects)
-		{
-			if (FTickableGameObject* TickableObject = static_cast<FTickableGameObject*>(TickableEntry.TickableObject))
-			{
-				// If it is tickable and in this world
-				if (((TickableEntry.TickType == ETickableTickType::Always) || TickableObject->IsTickable()) && (TickableObject->GetTickableGameObjectWorld() == World))
-				{
-					const bool bIsGameWorld = InTickType == LEVELTICK_All || (World && World->IsGameWorld());
-					// If we are in editor and it is editor tickable, always tick
-					// If this is a game world then tick if we are not doing a time only (paused) update and we are not paused or the object is tickable when paused
-					if ((GIsEditor && TickableObject->IsTickableInEditor()) ||
-						(bIsGameWorld && ((!bIsPaused && TickType != LEVELTICK_TimeOnly) || (bIsPaused && TickableObject->IsTickableWhenPaused()))))
-					{
-						FScopeCycleCounter Context(TickableObject->GetStatId());
-						TickableObject->Tick(DeltaSeconds);
-
-						// In case it was removed during tick
-						if (TickableEntry.TickableObject == nullptr)
-						{
-							bNeedsCleanup = true;
-						}
-					}
-				}
-			}
-			else
-			{
-				bNeedsCleanup = true;
-			}
-		}
-
-		if (bNeedsCleanup)
-		{
-			TickableObjects.RemoveAll([](const FTickableObjectEntry& Entry) { return Entry.TickableObject == nullptr; });
-		}
-
-		bIsTickingObjects = false;
-	}
-}
-
 /**
  * Update the level after a variable amount of time, DeltaSeconds, has passed.
  * All child actors are ticked after their owners have been ticked.
diff --git a/Engine/Source/Runtime/Engine/Private/Tickable.cpp b/Engine/Source/Runtime/Engine/Private/Tickable.cpp
new file mode 100644
--- /dev/null
+++ b/Engine/Source/Runtime/Engine/Private/Tickable.cpp
@@ -0,0 +1,150 @@
+// Copyright 1998-2019 Epic Games, Inc. All Rights Reserved.
+
+#include "Tickable.h"
+
+DECLARE_CYCLE_STAT(TEXT("TickableGameObjects Time"), STAT_TickableGameObjectsTime, STATGROUP_Game);
+
+struct FTickableStatics
+{
+	FCriticalSection TickableObjectsCritical;
+	TArray<FTickableObjectBase::FTickableObjectEntry> TickableObjects;
+	TLockFreePointerListUnordered<FTickableGameObject, PLATFORM_CACHE_LINE_SIZE> NewTickableObjects;
+	TSet<FTickableGameObject*> DeletedTickableObjects;
+	bool bIsTickingObjects = false;
+
+	static FTickableStatics& Get()
+	{
+		static FTickableStatics Singleton;
+		return Singleton;
+	}
+};
+
+void FTickableObjectBase::AddTickableObject(TArray<FTickableObjectEntry>& TickableObjects, FTickableObjectBase* TickableObject)
+{
+	check(!TickableObjects.Contains(TickableObject));
+	const ETickableTickType TickType = TickableObject->GetTickableTickType();
+	if (TickType != ETickableTickType::Never)
+	{
+		TickableObjects.Add({ TickableObject, TickType });
+	}
+}
+
+void FTickableObjectBase::RemoveTickableObject(TArray<FTickableObjectEntry>& TickableObjects, FTickableObjectBase* TickableObject, const bool bIsTickingObjects)
+{
+	const int32 Pos = TickableObjects.IndexOfByKey(TickableObject);
+
+#if 0 // virtual from destructor doesn't work ... need to rethink how to do warning
+	// ensure that GetTickableTickType did not change over time
+	switch (TickableObject->GetTickableTickType())
+	{
+	case ETickableTickType::Always:
+		ensureMsgf(Pos != INDEX_NONE && TickableObjects[Pos].TickType == ETickableTickType::Always, TEXT("TickType has changed since object was created. Result of GetTickableTickType must be invariant for a given object."));
+		break;
+
+	case ETickableTickType::Conditional:
+		ensureMsgf(Pos != INDEX_NONE && TickableObjects[Pos].TickType == ETickableTickType::Conditional, TEXT("TickType has changed since object was created. Result of GetTickableTickType must be invariant for a given object."));
+		break;
+
+	case ETickableTickType::Never:
+		ensureMsgf(Pos == INDEX_NONE, TEXT("TickType has changed since object was created. Result of GetTickableTickType must be invariant for a given object."));
+		break;
+	}
+#endif
+
+	if (Pos != INDEX_NONE)
+	{
+		if (bIsTickingObjects)
+		{
+			TickableObjects[Pos].TickableObject = nullptr;
+		}
+		else
+		{
+			TickableObjects.RemoveAt(Pos);
+		}
+	}
+}
+
+FTickableGameObject::FTickableGameObject()
+{
+	FTickableStatics& Statics = FTickableStatics::Get();
+
+	Statics.NewTickableObjects.Push(this);
+}
+
+FTickableGameObject::~FTickableGameObject()
+{
+	FTickableStatics& Statics = FTickableStatics::Get();
+
+	FScopeLock LockTickableObjects(&Statics.TickableObjectsCritical);
+	Statics.DeletedTickableObjects.Add(this);
+	RemoveTickableObject(Statics.TickableObjects, this, Statics.bIsTickingObjects);
+}
+
+void FTickableGameObject::TickObjects(UWorld* World, const int32 InTickType, const bool bIsPaused, const float DeltaSeconds)
+{
+	SCOPE_CYCLE_COUNTER(STAT_TickableGameObjectsTime);
+	CSV_SCOPED_TIMING_STAT_EXCLUSIVE(Tickables);
+
+	FTickableStatics& Statics = FTickableStatics::Get();
+
+	check(IsInGameThread());
+
+	// It's a long lock but it's ok, the only thing we can block here is the GC worker thread that destroys UObjects
+	FScopeLock LockTickableObjects(&Statics.TickableObjectsCritical);
+
+	while (FTickableGameObject* NewTickableObject = Statics.NewTickableObjects.Pop())
+	{
+		// If the new tickable object has already been deleted, don't add it to the tickable objects list
+		if (!Statics.DeletedTickableObjects.Contains(NewTickableObject))
+		{
+			AddTickableObject(Statics.TickableObjects, NewTickableObject);
+		}
+	}
+	Statics.DeletedTickableObjects.Empty();
+
+	if (Statics.TickableObjects.Num() > 0)
+	{
+		check(!Statics.bIsTickingObjects);
+		Statics.bIsTickingObjects = true;
+
+		bool bNeedsCleanup = false;
+		const ELevelTick TickType = (ELevelTick)InTickType;
+
+		for (const FTickableObjectEntry& TickableEntry : Statics.TickableObjects)
+		{
+			if (FTickableGameObject* TickableObject = static_cast<FTickableGameObject*>(TickableEntry.TickableObject))
+			{
+				// If it is tickable and in this world
+				if (((TickableEntry.TickType == ETickableTickType::Always) || TickableObject->IsTickable()) && (TickableObject->GetTickableGameObjectWorld() == World))
+				{
+					const bool bIsGameWorld = InTickType == LEVELTICK_All || (World && World->IsGameWorld());
+					// If we are in editor and it is editor tickable, always tick
+					// If this is a game world then tick if we are not doing a time only (paused) update and we are not paused or the object is tickable when paused
+					if ((GIsEditor && TickableObject->IsTickableInEditor()) ||
+						(bIsGameWorld && ((!bIsPaused && TickType != LEVELTICK_TimeOnly) || (bIsPaused && TickableObject->IsTickableWhenPaused()))))
+					{
+						FScopeCycleCounter Context(TickableObject->GetStatId());
+						TickableObject->Tick(DeltaSeconds);
+
+						// In case it was removed during tick
+						if (TickableEntry.TickableObject == nullptr)
+						{
+							bNeedsCleanup = true;
+						}
+					}
+				}
+			}
+			else
+			{
+				bNeedsCleanup = true;
+			}
+		}
+
+		if (bNeedsCleanup)
+		{
+			Statics.TickableObjects.RemoveAll([](const FTickableObjectEntry& Entry) { return Entry.TickableObject == nullptr; });
+		}
+
+		Statics.bIsTickingObjects = false;
+	}
+}
diff --git a/Engine/Source/Runtime/Engine/Public/Tickable.h b/Engine/Source/Runtime/Engine/Public/Tickable.h
--- a/Engine/Source/Runtime/Engine/Public/Tickable.h
+++ b/Engine/Source/Runtime/Engine/Public/Tickable.h
@@ -26,11 +26,15 @@ enum class ETickableTickType : uint8
 	Never        
 };
 
+struct FTickableStatics;
+
 /**
  * Base class for tickable objects
  */
-class FTickableObjectBase
+class ENGINE_API FTickableObjectBase
 {
+	friend struct FTickableStatics;
+
 protected:
 	struct FTickableObjectEntry
 	{
@@ -40,50 +44,9 @@ protected:
 		bool operator==(FTickableObjectBase* OtherObject) const { return TickableObject == OtherObject; }
 	};
 
-	static void AddTickableObject(TArray<FTickableObjectEntry>& TickableObjects, FTickableObjectBase* TickableObject)
-	{
-		check(!TickableObjects.Contains(TickableObject));
-		const ETickableTickType TickType = TickableObject->GetTickableTickType();
-		if (TickType != ETickableTickType::Never)
-		{
-			TickableObjects.Add({ TickableObject, TickType });
-		}
-	}
+	static void AddTickableObject(TArray<FTickableObjectEntry>& TickableObjects, FTickableObjectBase* TickableObject);
 
-	static void RemoveTickableObject(TArray<FTickableObjectEntry>& TickableObjects, FTickableObjectBase* TickableObject, const bool bIsTickingObjects)
-	{
-		const int32 Pos = TickableObjects.IndexOfByKey(TickableObject);
-
-#if 0 // virtual from destructor doesn't work ... need to rethink how to do warning
-		// ensure that GetTickableTickType did not change over time
-		switch (TickableObject->GetTickableTickType())
-		{
-		case ETickableTickType::Always:
-			ensureMsgf(Pos != INDEX_NONE && TickableObjects[Pos].TickType == ETickableTickType::Always, TEXT("TickType has changed since object was created. Result of GetTickableTickType must be invariant for a given object."));
-			break;
-
-		case ETickableTickType::Conditional:
-			ensureMsgf(Pos != INDEX_NONE && TickableObjects[Pos].TickType == ETickableTickType::Conditional, TEXT("TickType has changed since object was created. Result of GetTickableTickType must be invariant for a given object."));
-			break;
-
-		case ETickableTickType::Never:
-			ensureMsgf(Pos == INDEX_NONE, TEXT("TickType has changed since object was created. Result of GetTickableTickType must be invariant for a given object."));
-			break;
-		}
-#endif
-
-		if (Pos != INDEX_NONE)
-		{
-			if (bIsTickingObjects)
-			{
-				TickableObjects[Pos].TickableObject = nullptr;
-			}
-			else
-			{
-				TickableObjects.RemoveAt(Pos);
-			}
-		}
-	}
+	static void RemoveTickableObject(TArray<FTickableObjectEntry>& TickableObjects, FTickableObjectBase* TickableObject, const bool bIsTickingObjects);
 
 public:
 	/**
@@ -128,61 +91,17 @@ public:
  */
 class ENGINE_API FTickableGameObject : public FTickableObjectBase
 {
-private:
-	/** Static array of tickable objects */
-	static TArray<FTickableObjectEntry>& GetTickableObjects()
-	{
-		static TArray<FTickableObjectEntry> TickableObjects;
-		return TickableObjects;
-	}
-
-	static FCriticalSection& GetTickableObjectsCritical()
-	{
-		static FCriticalSection TickableObjectsCritical;
-		return TickableObjectsCritical;
-	}
-
-	/** List of newly created objects */
-	static TLockFreePointerListUnordered<FTickableGameObject, PLATFORM_CACHE_LINE_SIZE>& GetNewTickableObjects()
-	{
-		static TLockFreePointerListUnordered<FTickableGameObject, PLATFORM_CACHE_LINE_SIZE> NewTickableObjects;
-		return NewTickableObjects;
-	}
-
-	/** Deleted list is for cases when an object gets deleted before it was added to the tickable objects array from the NewTickableObjects list */
-	static TSet<FTickableGameObject*>& GetDeletedTickableObjects()
-	{
-		static TSet<FTickableGameObject*> DeletedTickableObjects;
-		return DeletedTickableObjects;
-	}
-
-	static bool bIsTickingObjects;
-
 public:
 	/**
 	 * Registers this instance with the static array of tickable objects.	
 	 *
 	 */
-	FTickableGameObject()
-	{
-		// These are called here only to make sure the function level statics they encapsulate
-		// get destroyed in the correct order at shutdown. This was added to fix UE-75429
-		GetTickableObjects();
-		GetDeletedTickableObjects();
-		GetTickableObjectsCritical();
-
-		GetNewTickableObjects().Push(this);
-	}
+	FTickableGameObject();
 
 	/**
 	 * Removes this instance from the static array of tickable objects.
 	 */
-	virtual ~FTickableGameObject()
-	{		
-		FScopeLock LockTickableObjects(&GetTickableObjectsCritical());
-		GetDeletedTickableObjects().Add(this);
-		RemoveTickableObject(GetTickableObjects(), this, bIsTickingObjects);
-	}
+	virtual ~FTickableGameObject();
 
 	/**
 	 * Used to determine if an object should be ticked when the game is paused.
