Removed redundant USkeletalMesh::AddReferencedObjects

diff --git a/Engine/Source/Runtime/Engine/Classes/Engine/SkeletalMesh.h b/Engine/Source/Runtime/Engine/Classes/Engine/SkeletalMesh.h
--- a/Engine/Source/Runtime/Engine/Classes/Engine/SkeletalMesh.h
+++ b/Engine/Source/Runtime/Engine/Classes/Engine/SkeletalMesh.h
@@ -912,7 +912,6 @@ public:
 	virtual FString GetDesc() override;
 	virtual FString GetDetailedInfoInternal() const override;
 	virtual void GetResourceSizeEx(FResourceSizeEx& CumulativeResourceSize) override;
-	static void AddReferencedObjects(UObject* InThis, FReferenceCollector& Collector);
 	//~ End UObject Interface.
 
 	/** Setup-only routines - not concerned with the instance. */
diff --git a/Engine/Source/Runtime/Engine/Private/SkeletalMesh.cpp b/Engine/Source/Runtime/Engine/Private/SkeletalMesh.cpp
--- a/Engine/Source/Runtime/Engine/Private/SkeletalMesh.cpp
+++ b/Engine/Source/Runtime/Engine/Private/SkeletalMesh.cpp
@@ -1211,22 +1211,6 @@ void USkeletalMesh::Serialize( FArchive& Ar )
 	}
 }
 
-void USkeletalMesh::AddReferencedObjects(UObject* InThis, FReferenceCollector& Collector)
-{	
-	USkeletalMesh* This = CastChecked<USkeletalMesh>(InThis);
-#if WITH_EDITOR
-	if( GIsEditor )
-	{
-		// Required by the unified GC when running in the editor
-		for( int32 Index = 0; Index < This->Materials.Num(); Index++ )
-		{
-			Collector.AddReferencedObject( This->Materials[ Index ].MaterialInterface, This );
-		}
-	}
-#endif
-	Super::AddReferencedObjects( This, Collector );
-}
-
 void USkeletalMesh::FlushRenderState()
 {
 	//TComponentReregisterContext<USkeletalMeshComponent> ReregisterContext;
