Fix use of uninitialized member variable in PhysScene_PhysX

diff --git a/Engine/Source/Runtime/Engine/Public/Physics/PhysScene_PhysX.h b/Engine/Source/Runtime/Engine/Public/Physics/PhysScene_PhysX.h
--- a/Engine/Source/Runtime/Engine/Public/Physics/PhysScene_PhysX.h
+++ b/Engine/Source/Runtime/Engine/Public/Physics/PhysScene_PhysX.h
@@ -302,7 +302,7 @@ private:
 
 #if WITH_PHYSX
 
-	bool bIsSceneSimulating;
+	bool bIsSceneSimulating = false;
 
 	/** Dispatcher for CPU tasks */
 	class PxCpuDispatcher*			CPUDispatcher;
