Optimize UAssetManager::OnObjectPreSave by reducing temporary string allocations.

diff --git a/Engine/Source/Runtime/Engine/Private/AssetManager.cpp b/Engine/Source/Runtime/Engine/Private/AssetManager.cpp
--- a/Engine/Source/Runtime/Engine/Private/AssetManager.cpp
+++ b/Engine/Source/Runtime/Engine/Private/AssetManager.cpp
@@ -3507,9 +3507,14 @@ void UAssetManager::OnInMemoryAssetDeleted(UObject *Object)
 void UAssetManager::OnObjectPreSave(UObject* Object)
 {
 	// If this is in the asset manager dictionary, make sure it actually has a primary asset id that matches
-	FPrimaryAssetId FoundPrimaryAssetId = GetPrimaryAssetIdForPath(FSoftObjectPath(Object));
+	const bool bIsAssetOrClass = Object->IsAsset() || Object->IsA(UClass::StaticClass()); 
+	if (!bIsAssetOrClass)
+	{
+		return;
+	}
 
-	if ((Object->IsAsset() || Object->IsA(UClass::StaticClass())) && FoundPrimaryAssetId.IsValid())
+	FPrimaryAssetId FoundPrimaryAssetId = GetPrimaryAssetIdForPath(*Object->GetPathName());
+	if (FoundPrimaryAssetId.IsValid())
 	{
 		TSharedRef<FPrimaryAssetTypeData>* FoundType = AssetTypeMap.Find(FoundPrimaryAssetId.PrimaryAssetType);
 		FPrimaryAssetId ObjectPrimaryAssetId = Object->GetPrimaryAssetId();
