Optimize RemapPaintedVertexColors (UStaticMeshComponent::PrivateFixupOverrideColors) by reusing temporary array

diff --git a/Engine/Source/Runtime/Engine/Private/StaticMeshBuild.cpp b/Engine/Source/Runtime/Engine/Private/StaticMeshBuild.cpp
--- a/Engine/Source/Runtime/Engine/Private/StaticMeshBuild.cpp
+++ b/Engine/Source/Runtime/Engine/Private/StaticMeshBuild.cpp
@@ -397,10 +397,11 @@ void RemapPaintedVertexColors(const TArray<FPaintedVertex>& InPaintedVertices,
 	// Iterate over each new vertex position, attempting to find the old vertex it is closest to, applying
 	// the color of the old vertex to the new position if possible.
 	OutOverrideColors.Empty(NewPositions.GetNumVertices());
+	TArray<FPaintedVertex> PointsToConsider;
 	const float DistanceOverNormalThreshold = OptionalVertexBuffer ? KINDA_SMALL_NUMBER : 0.0f;
 	for ( uint32 NewVertIndex = 0; NewVertIndex < NewPositions.GetNumVertices(); ++NewVertIndex )
 	{
-		TArray<FPaintedVertex> PointsToConsider;
+		PointsToConsider.Reset();
 		TSMCVertPosOctree::TConstIterator<> OctreeIter( VertPosOctree );
 		const FVector& CurPosition = NewPositions.VertexPosition( NewVertIndex );
 		FVector CurNormal = FVector::ZeroVector;
