Fix FindShaderById and FindOrCreate issues after dev-render copy up.

diff --git a/Engine/Source/Runtime/Engine/Private/Materials/MaterialShader.cpp b/Engine/Source/Runtime/Engine/Private/Materials/MaterialShader.cpp
--- a/Engine/Source/Runtime/Engine/Private/Materials/MaterialShader.cpp
+++ b/Engine/Source/Runtime/Engine/Private/Materials/MaterialShader.cpp
@@ -1012,7 +1012,7 @@ FShader* FMaterialShaderType::FinishCompileShader(
 
 	// Reuse an existing resource with the same key or create a new one based on the compile output
 	// This allows FShaders to share compiled bytecode and RHI shader references
-	FShaderResource* Resource = FShaderResource::FindOrCreateShaderResource(CurrentJob.Output, SpecificType, SpecificPermutationId);
+	TRefCountPtr<FShaderResource> Resource = FShaderResource::FindOrCreate(CurrentJob.Output, SpecificType, SpecificPermutationId);
 
 	if (ShaderPipelineType && !ShaderPipelineType->ShouldOptimizeUnusedOutputs(CurrentJob.Input.Target.GetPlatform()))
 	{
@@ -1026,7 +1026,7 @@ FShader* FMaterialShaderType::FinishCompileShader(
 	// There was no shader with the same key so create a new one with the compile output, which will bind shader parameters
 	if (!Shader)
 	{
-		Shader = (*ConstructCompiledRef)(CompiledShaderInitializerType(this, CurrentJob.PermutationId, CurrentJob.Output, Resource, UniformExpressionSet, MaterialShaderMapHash, ShaderPipelineType, nullptr, InDebugDescription));
+		Shader = (*ConstructCompiledRef)(CompiledShaderInitializerType(this, CurrentJob.PermutationId, CurrentJob.Output, MoveTemp(Resource), UniformExpressionSet, MaterialShaderMapHash, ShaderPipelineType, nullptr, InDebugDescription));
 		CurrentJob.Output.ParameterMap.VerifyBindingsAreComplete(GetName(), CurrentJob.Output.Target, CurrentJob.VFType);
 	}
 
@@ -1215,7 +1215,7 @@ void FMaterialShaderMap::SaveForRemoteRecompile(FArchive& Ar, const TMap<FString
 	UE_LOG(LogMaterial, Display, TEXT("Looking for unique resources, %d were on client"), ClientResourceIds.Num());	
 
 	// first, we look for the unique shader resources
-	TArray<FShaderResource*> UniqueResources;
+	TArray<TRefCountPtr<FShaderResource>> UniqueResources;
 	int32 NumSkippedResources = 0;
 
 	for (TMap<FString, TArray<TRefCountPtr<FMaterialShaderMap> > >::TConstIterator It(CompiledShaderMaps); It; ++It)
@@ -1255,7 +1255,7 @@ void FMaterialShaderMap::SaveForRemoteRecompile(FArchive& Ar, const TMap<FString
 					if (ClientResourceIds.Contains(ShaderId) == false)
 					{
 						// lookup the resource by ID
-						FShaderResource* Resource = FShaderResource::FindShaderResourceById(ShaderId);
+						TRefCountPtr<FShaderResource> Resource = FShaderResource::FindById(ShaderId);
 						// add it if it's unique
 						UniqueResources.AddUnique(Resource);
 					}
@@ -1320,28 +1320,16 @@ void FMaterialShaderMap::LoadForRemoteRecompile(FArchive& Ar, EShaderPlatform Sh
 
 	// KeepAliveReferences keeps resources alive until we are finished serializing in this function
 	TArray<TRefCountPtr<FShaderResource> > KeepAliveReferences;
+	KeepAliveReferences.Reserve(NumResources);
 
 	// load and register the resources
 	for (int32 Index = 0; Index < NumResources; Index++)
 	{
 		// Load the inlined shader resource
-		FShaderResource* Resource = new FShaderResource();
-		Resource->Serialize(Ar, false);
+		FShaderResource ResourceTemp;
+		ResourceTemp.Serialize(Ar, false);
 
-		// if this Id is already in memory, that means that this is a repeated resource and so we skip it
-		if (FShaderResource::FindShaderResourceById(Resource->GetId()) != NULL)
-		{
-			delete Resource;
-		}
-		// otherwise, it's a new resource, so we register it for the maps to find below
-		else
-		{
-			Resource->Register();
-
-			// Keep this guy alive until we finish serializing all the FShaders in
-			// The FShaders which are discarded may cause these resources to be discarded 
-			KeepAliveReferences.Add( Resource );
-		}
+		KeepAliveReferences.Add(FShaderResource::FindOrClone(MoveTemp(ResourceTemp)));
 	}
 
 	int32 MapSize;
diff --git a/Engine/Source/Runtime/Engine/Private/Materials/MeshMaterialShader.cpp b/Engine/Source/Runtime/Engine/Private/Materials/MeshMaterialShader.cpp
--- a/Engine/Source/Runtime/Engine/Private/Materials/MeshMaterialShader.cpp
+++ b/Engine/Source/Runtime/Engine/Private/Materials/MeshMaterialShader.cpp
@@ -151,7 +151,7 @@ FShader* FMeshMaterialShaderType::FinishCompileShader(
 
 	// Reuse an existing resource with the same key or create a new one based on the compile output
 	// This allows FShaders to share compiled bytecode and RHI shader references
-	FShaderResource* Resource = FShaderResource::FindOrCreateShaderResource(CurrentJob.Output, SpecificType, CurrentJob.PermutationId);
+	TRefCountPtr<FShaderResource> Resource = FShaderResource::FindOrCreate(CurrentJob.Output, SpecificType, CurrentJob.PermutationId);
 
 	if (ShaderPipelineType && !ShaderPipelineType->ShouldOptimizeUnusedOutputs(CurrentJob.Input.Target.GetPlatform()))
 	{
@@ -165,7 +165,7 @@ FShader* FMeshMaterialShaderType::FinishCompileShader(
 	// There was no shader with the same key so create a new one with the compile output, which will bind shader parameters
 	if (!Shader)
 	{
-		Shader = (*ConstructCompiledRef)(CompiledShaderInitializerType(this, CurrentJob.PermutationId, CurrentJob.Output, Resource, UniformExpressionSet, MaterialShaderMapHash, InDebugDescription, ShaderPipelineType, CurrentJob.VFType));
+		Shader = (*ConstructCompiledRef)(CompiledShaderInitializerType(this, CurrentJob.PermutationId, CurrentJob.Output, MoveTemp(Resource), UniformExpressionSet, MaterialShaderMapHash, InDebugDescription, ShaderPipelineType, CurrentJob.VFType));
 		ValidateAfterBind((FMeshMaterialShader*)Shader);
 		CurrentJob.Output.ParameterMap.VerifyBindingsAreComplete(GetName(), CurrentJob.Output.Target, CurrentJob.VFType);
 	}
