Removed some rogue text from a file that breaks compilation with certain combinations of physics defines. Already fixed by physics guys in other branches but redone here to help keep CIS happy

#rnx

diff --git a/Engine/Source/Runtime/Engine/Private/Collision/SceneQueryLowLevel.cpp b/Engine/Source/Runtime/Engine/Private/Collision/SceneQueryLowLevel.cpp
--- a/Engine/Source/Runtime/Engine/Private/Collision/SceneQueryLowLevel.cpp
+++ b/Engine/Source/Runtime/Engine/Private/Collision/SceneQueryLowLevel.cpp
@@ -46,7 +46,7 @@ void FinalizeCapture(FPhysTestSerializer& Serializer)
 	if (ReplaySQs)
 	{
 		const bool bReplaySuccess = SQComparisonHelper(Serializer);
-		if (!bReplaySuccess)'ve just'
+		if (!bReplaySuccess)
 		{
 			UE_LOG(LogPhysicsCore, Warning, TEXT("Chaos SQ does not match physx"));
 			if (SerializeBadSQs && !SerializeSQs)
