Making AtmosphericFogComponent not crash when InitResource is executed off the main thread

diff --git a/Engine/Source/Runtime/Engine/Private/Atmosphere/AtmosphericFogComponent.cpp b/Engine/Source/Runtime/Engine/Private/Atmosphere/AtmosphericFogComponent.cpp
--- a/Engine/Source/Runtime/Engine/Private/Atmosphere/AtmosphericFogComponent.cpp
+++ b/Engine/Source/Runtime/Engine/Private/Atmosphere/AtmosphericFogComponent.cpp
@@ -195,7 +195,7 @@ static TAutoConsoleVariable<int32> CVarAtmosphereRender(
 // On CPU
 void UAtmosphericFogComponent::InitResource()
 {
-	bool bNewAtmosphere = CVarAtmosphereRender.GetValueOnGameThread() != 0;
+	bool bNewAtmosphere = CVarAtmosphereRender.GetValueOnAnyThread() != 0;
 
 	if(!bNewAtmosphere)
 	{
