Fix some code that got misaligned during a merge at some point.

#rnx

diff --git a/Engine/Source/Runtime/Engine/Private/TextureDerivedData.cpp b/Engine/Source/Runtime/Engine/Private/TextureDerivedData.cpp
--- a/Engine/Source/Runtime/Engine/Private/TextureDerivedData.cpp
+++ b/Engine/Source/Runtime/Engine/Private/TextureDerivedData.cpp
@@ -1299,8 +1299,8 @@ static void SerializePlatformData(
 			if (!bIsVirtual)
 			{
 				FirstMipToSerialize = FMath::Clamp(FirstMipToSerialize, 0, FMath::Max(NumMips - 1, 0));
-			NumMips -= FirstMipToSerialize;
-		}
+				NumMips -= FirstMipToSerialize;
+			}
 			else
 			{
 				FirstMipToSerialize = FMath::Clamp(FirstMipToSerialize, 0, FMath::Max((int32)PlatformData->VTData->GetNumMips() - 1, 0));
@@ -1413,12 +1413,12 @@ static void SerializePlatformData(
 
 	if (bIsVirtual == false)
 	{
-	for (int32 MipIndex = 0; MipIndex < BulkDataMipFlags.Num(); ++MipIndex)
-	{
+		for (int32 MipIndex = 0; MipIndex < BulkDataMipFlags.Num(); ++MipIndex)
+		{
 			check(Ar.IsSaving());
-		PlatformData->Mips[MipIndex].BulkData.ClearBulkDataFlags(~BulkDataMipFlags[MipIndex]);
-		PlatformData->Mips[MipIndex].BulkData.SetBulkDataFlags(BulkDataMipFlags[MipIndex]);
-	}
+			PlatformData->Mips[MipIndex].BulkData.ClearBulkDataFlags(~BulkDataMipFlags[MipIndex]);
+			PlatformData->Mips[MipIndex].BulkData.SetBulkDataFlags(BulkDataMipFlags[MipIndex]);
+		}
 	}
 	else
 	{
@@ -1445,16 +1445,16 @@ void FTexturePlatformData::SerializeCooked(FArchive& Ar, UTexture* Owner, bool b
 	{
 		// Patch up Size as due to mips being stripped out during cooking it could be wrong.
 		if (Mips.Num() > 0)
-	{
-		SizeX = Mips[0].SizeX;
-		SizeY = Mips[0].SizeY;
-			
-		// SizeZ is not the same as NumSlices for texture arrays and cubemaps.
-		if (Owner && Owner->IsA(UVolumeTexture::StaticClass()))
 		{
-			 NumSlices = Mips[0].SizeZ;
+			SizeX = Mips[0].SizeX;
+			SizeY = Mips[0].SizeY;
+			
+			// SizeZ is not the same as NumSlices for texture arrays and cubemaps.
+			if (Owner && Owner->IsA(UVolumeTexture::StaticClass()))
+			{
+				 NumSlices = Mips[0].SizeZ;
+			}
 		}
-	}
 		else if ( VTData )
 		{
 			SizeX = VTData->Width;
