Fixed FTriangleRasterizer DrawTriangle division by zero case

diff --git a/Engine/Source/Runtime/Engine/Public/Raster.h b/Engine/Source/Runtime/Engine/Public/Raster.h
--- a/Engine/Source/Runtime/Engine/Public/Raster.h
+++ b/Engine/Source/Runtime/Engine/Public/Raster.h
@@ -18,51 +18,60 @@ public:
 
 	typedef typename RasterPolicyType::InterpolantType InterpolantType;
 
-	void DrawTriangle(const InterpolantType& I0,const InterpolantType& I1,const InterpolantType& I2,const FVector2D& P0,const FVector2D& P1,const FVector2D& P2,bool BackFacing)
+	void DrawTriangle(const InterpolantType& I0, const InterpolantType& I1, const InterpolantType& I2, const FVector2D& P0, const FVector2D& P1, const FVector2D& P2, bool BackFacing)
 	{
 		InterpolantType	Interpolants[3] = { I0, I1, I2 };
 		FVector2D		Points[3] = { P0, P1, P2 };
 
-		// Avoid any division by zero
-		if (FMath::IsNearlyEqual(Points[1].Y, Points[0].Y) ||
-			FMath::IsNearlyEqual(Points[2].Y, Points[0].Y) ||
-			FMath::IsNearlyEqual(Points[2].Y, Points[1].Y))
-		{
-			return;
-		}
-
 		// Find the top point.
 
-		if(Points[1].Y < Points[0].Y && Points[1].Y <= Points[2].Y)
+		if (Points[1].Y < Points[0].Y && Points[1].Y <= Points[2].Y)
 		{
-			Exchange(Points[0],Points[1]);
-			Exchange(Interpolants[0],Interpolants[1]);
+			Exchange(Points[0], Points[1]);
+			Exchange(Interpolants[0], Interpolants[1]);
 		}
-		else if(Points[2].Y < Points[0].Y && Points[2].Y <= Points[1].Y)
+		else if (Points[2].Y < Points[0].Y && Points[2].Y <= Points[1].Y)
 		{
-			Exchange(Points[0],Points[2]);
-			Exchange(Interpolants[0],Interpolants[2]);
+			Exchange(Points[0], Points[2]);
+			Exchange(Interpolants[0], Interpolants[2]);
 		}
 
 		// Find the bottom point.
 
-		if(Points[1].Y > Points[2].Y)
+		if (Points[1].Y > Points[2].Y)
+		{
+			Exchange(Points[2], Points[1]);
+			Exchange(Interpolants[2], Interpolants[1]);
+		}
+
+		// Avoid any division by zero
+		float PointDiffY_1_0 = (Points[1].Y - Points[0].Y);
+		if (FMath::IsNearlyZero(PointDiffY_1_0))
+		{
+			PointDiffY_1_0 = PointDiffY_1_0 >= 0.f ? SMALL_NUMBER : -SMALL_NUMBER;
+		}
+		float PointDiffY_2_0 = (Points[2].Y - Points[0].Y);
+		if (FMath::IsNearlyZero(PointDiffY_2_0))
+		{
+			PointDiffY_2_0 = PointDiffY_2_0 >= 0.f ? SMALL_NUMBER : -SMALL_NUMBER;
+		}
+		float PointDiffY_2_1 = (Points[2].Y - Points[1].Y);
+		if (FMath::IsNearlyZero(PointDiffY_2_1))
 		{
-			Exchange(Points[2],Points[1]);
-			Exchange(Interpolants[2],Interpolants[1]);
+			PointDiffY_2_1 = PointDiffY_2_1 >= 0.f ? SMALL_NUMBER : -SMALL_NUMBER;
 		}
 
 		// Calculate the edge gradients.
 
-		float			TopMinDiffX = (Points[1].X - Points[0].X) / (Points[1].Y - Points[0].Y),
-						TopMaxDiffX = (Points[2].X - Points[0].X) / (Points[2].Y - Points[0].Y);
-		InterpolantType	TopMinDiffInterpolant = (Interpolants[1] - Interpolants[0]) / (Points[1].Y - Points[0].Y),
-						TopMaxDiffInterpolant = (Interpolants[2] - Interpolants[0]) / (Points[2].Y - Points[0].Y);
+		float			TopMinDiffX = (Points[1].X - Points[0].X) / PointDiffY_1_0,
+						TopMaxDiffX = (Points[2].X - Points[0].X) / PointDiffY_2_0;
+		InterpolantType	TopMinDiffInterpolant = (Interpolants[1] - Interpolants[0]) / PointDiffY_1_0,
+						TopMaxDiffInterpolant = (Interpolants[2] - Interpolants[0]) / PointDiffY_2_0;
 
-		float			BottomMinDiffX = (Points[2].X - Points[1].X) / (Points[2].Y - Points[1].Y),
-						BottomMaxDiffX = (Points[2].X - Points[0].X) / (Points[2].Y - Points[0].Y);
-		InterpolantType	BottomMinDiffInterpolant = (Interpolants[2] - Interpolants[1]) / (Points[2].Y - Points[1].Y),
-						BottomMaxDiffInterpolant = (Interpolants[2] - Interpolants[0]) / (Points[2].Y - Points[0].Y);
+		float			BottomMinDiffX = (Points[2].X - Points[1].X) / PointDiffY_2_1,
+						BottomMaxDiffX = (Points[2].X - Points[0].X) / PointDiffY_2_0;
+		InterpolantType	BottomMinDiffInterpolant = (Interpolants[2] - Interpolants[1]) / PointDiffY_2_1,
+						BottomMaxDiffInterpolant = (Interpolants[2] - Interpolants[0]) / PointDiffY_2_0;
 
 		DrawTriangleTrapezoid(
 			Interpolants[0],
@@ -81,11 +90,11 @@ public:
 		DrawTriangleTrapezoid(
 			Interpolants[1],
 			BottomMinDiffInterpolant,
-			Interpolants[0] + TopMaxDiffInterpolant * (Points[1].Y - Points[0].Y),
+			Interpolants[0] + TopMaxDiffInterpolant * PointDiffY_1_0,
 			BottomMaxDiffInterpolant,
 			Points[1].X,
 			BottomMinDiffX,
-			Points[0].X + TopMaxDiffX * (Points[1].Y - Points[0].Y),
+			Points[0].X + TopMaxDiffX * PointDiffY_1_0,
 			BottomMaxDiffX,
 			Points[1].Y,
 			Points[2].Y,
