Resolving robomerge conflict - adding okforgithub tags

Move a console gamepad check in IsMouseAttached from GameViewportClient to an inherited OcnGameViewportClient.

diff --git a/Engine/Source/Runtime/Engine/Classes/Engine/GameViewportClient.h b/Engine/Source/Runtime/Engine/Classes/Engine/GameViewportClient.h
--- a/Engine/Source/Runtime/Engine/Classes/Engine/GameViewportClient.h
+++ b/Engine/Source/Runtime/Engine/Classes/Engine/GameViewportClient.h
@@ -337,7 +337,7 @@ public:
 	bool IsExclusiveFullscreenViewport() const;
 
 	/** @return mouse position in game viewport coordinates (does not account for splitscreen) */
-	bool GetMousePosition(FVector2D& MousePosition) const;
+	virtual bool GetMousePosition(FVector2D& MousePosition) const;
 
 	/**
 	 * Determine whether a fullscreen viewport should be used in cases where there are multiple players.
diff --git a/Engine/Source/Runtime/Engine/Private/GameViewportClient.cpp b/Engine/Source/Runtime/Engine/Private/GameViewportClient.cpp
--- a/Engine/Source/Runtime/Engine/Private/GameViewportClient.cpp
+++ b/Engine/Source/Runtime/Engine/Private/GameViewportClient.cpp
@@ -821,11 +821,7 @@ bool UGameViewportClient::GetMousePosition(FVector2D& MousePosition) const
 {
 	bool bGotMousePosition = false;
 
-#ifdef USE_CONSOLE_CONTROLLER
 	if (Viewport && FSlateApplication::Get().IsMouseAttached())
-#else
-	if (Viewport)
-#endif
 	{
 		FIntPoint MousePos;
 		Viewport->GetMousePos(MousePos);
