Add cvars to control the priority of async IO requests from the virtual texture system.
r.VT.IOPriority_NormalPagePri
r.VT.IOPriority_HighPagePri

These default to AIOP_Normal and AIOP_High respectively (no change from current behaviour)

diff --git a/Engine/Source/Runtime/Engine/Private/VT/VirtualTextureChunkManager.cpp b/Engine/Source/Runtime/Engine/Private/VT/VirtualTextureChunkManager.cpp
--- a/Engine/Source/Runtime/Engine/Private/VT/VirtualTextureChunkManager.cpp
+++ b/Engine/Source/Runtime/Engine/Private/VT/VirtualTextureChunkManager.cpp
@@ -23,6 +23,21 @@ static FAutoConsoleVariableRef CVarNumTranscodeRequests(
 	TEXT("Number of transcode request that can be in flight. default 32\n"),
 	ECVF_Default);
 
+int32 GVirtualTextureIOPriority_NormalPagePri = (int32)AIOP_Normal;
+static FAutoConsoleVariableRef CVarVirtualTextureIOPriority_NormalPagePri(
+	TEXT("r.VT.IOPriority_NormalPagePri"),
+	GVirtualTextureIOPriority_NormalPagePri,
+	TEXT("Priority of default priority I/O requests"),
+	ECVF_Default
+);
+
+int32 GVirtualTextureIOPriority_HighPagePri = (int32)AIOP_High;
+static FAutoConsoleVariableRef CVarVirtualTextureIOPriority_HighPagePri(
+	TEXT("r.VT.IOPriority_HighPagePri"),
+	GVirtualTextureIOPriority_HighPagePri,
+	TEXT("Priority of high priority I/O requests"),
+	ECVF_Default
+);
 
 /** 
  * FVirtualTextureChunkStreamingManager is a treated as a singleton referenced by multiple FUploadingVirtualTexture objects.
@@ -109,9 +124,9 @@ static EAsyncIOPriorityAndFlags GetAsyncIOPriority(EVTRequestPagePriority Priori
 {
 	switch (Priority)
 	{
-	case EVTRequestPagePriority::High: return AIOP_High;
-	case EVTRequestPagePriority::Normal: return AIOP_Normal;
-	default: check(false); return AIOP_Normal;
+	case EVTRequestPagePriority::High: return (EAsyncIOPriorityAndFlags)GVirtualTextureIOPriority_HighPagePri;
+	case EVTRequestPagePriority::Normal: return (EAsyncIOPriorityAndFlags)GVirtualTextureIOPriority_NormalPagePri;
+	default: check(false); return (EAsyncIOPriorityAndFlags)GVirtualTextureIOPriority_NormalPagePri;
 	}
 }
 
