Chaos - RBAN fix LOD change enable/disable logic

diff --git a/Engine/Source/Runtime/Engine/Private/PhysicsEngine/ImmediatePhysics/ImmediatePhysicsChaos/ImmediatePhysicsSimulation_Chaos.cpp b/Engine/Source/Runtime/Engine/Private/PhysicsEngine/ImmediatePhysics/ImmediatePhysicsChaos/ImmediatePhysicsSimulation_Chaos.cpp
--- a/Engine/Source/Runtime/Engine/Private/PhysicsEngine/ImmediatePhysics/ImmediatePhysicsChaos/ImmediatePhysicsSimulation_Chaos.cpp
+++ b/Engine/Source/Runtime/Engine/Private/PhysicsEngine/ImmediatePhysics/ImmediatePhysicsChaos/ImmediatePhysicsSimulation_Chaos.cpp
@@ -348,9 +348,15 @@ namespace ImmediatePhysics_Chaos
 				continue;
 			}
 
-			bool bShouldBeActive = NumActiveDynamicActorHandles < InNumActiveActorHandles;
-			Handle->SetEnabled(bShouldBeActive);
-			++NumActiveDynamicActorHandles;
+			if (NumActiveDynamicActorHandles < InNumActiveActorHandles)
+			{
+				Handle->SetEnabled(true);
+				++NumActiveDynamicActorHandles;
+			}
+			else
+			{
+				Handle->SetEnabled(false);
+			}
 		}
 	}
 
