#UE4 Null check ActorHandle

diff --git a/Engine/Source/Runtime/Engine/Private/PhysicsEngine/BodyInstance.cpp b/Engine/Source/Runtime/Engine/Private/PhysicsEngine/BodyInstance.cpp
--- a/Engine/Source/Runtime/Engine/Private/PhysicsEngine/BodyInstance.cpp
+++ b/Engine/Source/Runtime/Engine/Private/PhysicsEngine/BodyInstance.cpp
@@ -1350,7 +1350,10 @@ void FBodyInstance::TermBody(bool bNeverDeferRelease)
 {
 	SCOPE_CYCLE_COUNTER(STAT_TermBody);
 
-	FPhysicsInterface::ReleaseActor(ActorHandle, GetPhysicsScene(), bNeverDeferRelease);
+	if (ActorHandle)
+	{
+		FPhysicsInterface::ReleaseActor(ActorHandle, GetPhysicsScene(), bNeverDeferRelease);
+	}
 
 	// @TODO UE4: Release spring body here
 
