#UE4 Fix cis error

diff --git a/Engine/Source/Runtime/Engine/Private/PhysicsEngine/BodyInstance.cpp b/Engine/Source/Runtime/Engine/Private/PhysicsEngine/BodyInstance.cpp
--- a/Engine/Source/Runtime/Engine/Private/PhysicsEngine/BodyInstance.cpp
+++ b/Engine/Source/Runtime/Engine/Private/PhysicsEngine/BodyInstance.cpp
@@ -1350,7 +1350,7 @@ void FBodyInstance::TermBody(bool bNeverDeferRelease)
 {
 	SCOPE_CYCLE_COUNTER(STAT_TermBody);
 
-	if (ActorHandle)
+	if (IsValidBodyInstance())
 	{
 		FPhysicsInterface::ReleaseActor(ActorHandle, GetPhysicsScene(), bNeverDeferRelease);
 	}
