Back out changelist 10311864.

diff --git a/Engine/Source/Runtime/Engine/Private/HierarchicalInstancedStaticMesh.cpp b/Engine/Source/Runtime/Engine/Private/HierarchicalInstancedStaticMesh.cpp
--- a/Engine/Source/Runtime/Engine/Private/HierarchicalInstancedStaticMesh.cpp
+++ b/Engine/Source/Runtime/Engine/Private/HierarchicalInstancedStaticMesh.cpp
@@ -2983,7 +2983,8 @@ FPrimitiveSceneProxy* UHierarchicalInstancedStaticMeshComponent::CreateSceneProx
 		check(InstancingRandomSeed != 0);
 
 		// if instance data was modified, update GPU copy
-		if (InstanceUpdateCmdBuffer.NumInlineCommands() > 0)
+		// generally happens only in editor 
+		if (GIsEditor && InstanceUpdateCmdBuffer.NumInlineCommands() > 0)
 		{
 			PerInstanceRenderData->UpdateFromCommandBuffer(InstanceUpdateCmdBuffer);
 		}
