Fixed HISM Cluster Tree display not taking into account actor's transform (reported in a UDN)

diff --git a/Engine/Source/Runtime/Engine/Private/HierarchicalInstancedStaticMesh.cpp b/Engine/Source/Runtime/Engine/Private/HierarchicalInstancedStaticMesh.cpp
--- a/Engine/Source/Runtime/Engine/Private/HierarchicalInstancedStaticMesh.cpp
+++ b/Engine/Source/Runtime/Engine/Private/HierarchicalInstancedStaticMesh.cpp
@@ -1879,7 +1879,7 @@ void FHierarchicalStaticMeshSceneProxy::GetDynamicMeshElements(const TArray<cons
 
 				for (const FClusterNode& CulsterNode : ClusterTree)
 				{
-					DrawWireBox(Collector.GetPDI(ViewIndex), FBox(CulsterNode.BoundMin, CulsterNode.BoundMax), StartingColor, View->Family->EngineShowFlags.Game ? SDPG_World : SDPG_Foreground);
+					DrawWireBox(Collector.GetPDI(ViewIndex), GetLocalToWorld(), FBox(CulsterNode.BoundMin, CulsterNode.BoundMax), StartingColor, View->Family->EngineShowFlags.Game ? SDPG_World : SDPG_Foreground);
 					StartingColor.R += 5;
 					StartingColor.G += 5;
 					StartingColor.B += 5;
