NavLink: adding constant for InvalidUserId

#rnx

diff --git a/Engine/Source/Runtime/Engine/Classes/AI/Navigation/NavLinkDefinition.h b/Engine/Source/Runtime/Engine/Classes/AI/Navigation/NavLinkDefinition.h
--- a/Engine/Source/Runtime/Engine/Classes/AI/Navigation/NavLinkDefinition.h
+++ b/Engine/Source/Runtime/Engine/Classes/AI/Navigation/NavLinkDefinition.h
@@ -36,6 +36,9 @@ struct ENGINE_API FNavigationLinkBase
 	UPROPERTY(EditAnywhere, Category=Default, meta=(ClampMin = "0.0", DisplayName="Right Project Height"))
 	float MaxFallDownLength;
 
+	/** Needs to be 0 for recast data generator */
+	static constexpr int32 InvalidUserId = 0;
+
 	/** ID passed to navigation data generator */
 	int32 UserId;
 
diff --git a/Engine/Source/Runtime/Engine/Private/AI/Navigation/NavigationModifier.cpp b/Engine/Source/Runtime/Engine/Private/AI/Navigation/NavigationModifier.cpp
--- a/Engine/Source/Runtime/Engine/Private/AI/Navigation/NavigationModifier.cpp
+++ b/Engine/Source/Runtime/Engine/Private/AI/Navigation/NavigationModifier.cpp
@@ -18,7 +18,7 @@ static const float CONVEX_HULL_POINTS_MIN_DISTANCE_SQ = 4.0f * 4.0f;
 // FNavigationLinkBase
 //----------------------------------------------------------------------//
 FNavigationLinkBase::FNavigationLinkBase() 
-	: LeftProjectHeight(0.0f), MaxFallDownLength(1000.0f), UserId(0), SnapRadius(30.f), SnapHeight(50.0f),
+	: LeftProjectHeight(0.0f), MaxFallDownLength(1000.0f), UserId(InvalidUserId), SnapRadius(30.f), SnapHeight(50.0f),
 	  Direction(ENavLinkDirection::BothWays), bUseSnapHeight(false), bSnapToCheapestArea(true),
 	  bCustomFlag0(false), bCustomFlag1(false), bCustomFlag2(false), bCustomFlag3(false), bCustomFlag4(false),
 	  bCustomFlag5(false), bCustomFlag6(false), bCustomFlag7(false)
