Fixed crash when trying to import a DDS file that uses the legacy (pre-DX10) header format.

diff --git a/Engine/Source/Runtime/Engine/Private/DDSLoader.cpp b/Engine/Source/Runtime/Engine/Private/DDSLoader.cpp
--- a/Engine/Source/Runtime/Engine/Private/DDSLoader.cpp
+++ b/Engine/Source/Runtime/Engine/Private/DDSLoader.cpp
@@ -150,12 +150,12 @@ ETextureSourceFormat FDDSLoadHelper::ComputeSourceFormat() const
 
 bool FDDSLoadHelper::IsValidCubemapTexture() const
 {
-	if(IsValid() && (DDSHeader->dwCaps2 & DDSC_CubeMap) != 0 && (DDSHeader->dwCaps2 & DDSC_CubeMap_AllFaces) != 0)
+	if(DDSHeader != nullptr && (DDSHeader->dwCaps2 & DDSC_CubeMap) != 0 && (DDSHeader->dwCaps2 & DDSC_CubeMap_AllFaces) != 0)
 	{
 		return true;
 	}
 
-	if (IsValid() && DDS10Header->resourceType == 3 && (DDS10Header->miscFlag & 4))
+	if (DDS10Header != nullptr && DDS10Header->resourceType == 3 && (DDS10Header->miscFlag & 4))
 	{
 		return true;
 	}
