Improvements for headless client automation from UDN
#rnx

diff --git a/Engine/Source/Runtime/Engine/Private/GameViewportClient.cpp b/Engine/Source/Runtime/Engine/Private/GameViewportClient.cpp
--- a/Engine/Source/Runtime/Engine/Private/GameViewportClient.cpp
+++ b/Engine/Source/Runtime/Engine/Private/GameViewportClient.cpp
@@ -420,6 +420,14 @@ void UGameViewportClient::Init(struct FWorldContext& WorldContext, UGameInstance
 		}
 	}
 	MouseLockMode = GetDefault<UInputSettings>()->DefaultViewportMouseLockMode;
+
+	// Don't capture mouse when headless
+	if(!FApp::CanEverRender())
+	{
+		MouseCaptureMode = EMouseCaptureMode::NoCapture;
+		MouseLockMode = EMouseLockMode::DoNotLock;
+	}
+
 	// In off-screen rendering mode don't lock mouse to the viewport, as we don't want mouse to lock to an invisible window
 	if (FSlateApplication::Get().IsRenderingOffScreen()) {
 		MouseLockMode = EMouseLockMode::DoNotLock;
@@ -2620,7 +2628,8 @@ void UGameViewportClient::VerifyPathRenderingComponents()
 
 bool UGameViewportClient::CaptureMouseOnLaunch()
 {
-	return GetDefault<UInputSettings>()->bCaptureMouseOnLaunch;
+	// Capture mouse unless headless
+	return !FApp::CanEverRender() ? false : GetDefault<UInputSettings>()->bCaptureMouseOnLaunch;
 }
 
 bool UGameViewportClient::Exec( UWorld* InWorld, const TCHAR* Cmd,FOutputDevice& Ar)
diff --git a/Engine/Source/Runtime/Engine/Private/Slate/SceneViewport.cpp b/Engine/Source/Runtime/Engine/Private/Slate/SceneViewport.cpp
--- a/Engine/Source/Runtime/Engine/Private/Slate/SceneViewport.cpp
+++ b/Engine/Source/Runtime/Engine/Private/Slate/SceneViewport.cpp
@@ -1229,7 +1229,7 @@ void FSceneViewport::PaintDebugCanvas(const FGeometry& AllottedGeometry, FSlateW
 void FSceneViewport::ResizeFrame(uint32 NewWindowSizeX, uint32 NewWindowSizeY, EWindowMode::Type NewWindowMode)
 {
 	// Resizing the window directly is only supported in the game
-	if( FApp::IsGame() && NewWindowSizeX > 0 && NewWindowSizeY > 0 )
+	if( FApp::IsGame() && FApp::CanEverRender() && NewWindowSizeX > 0 && NewWindowSizeY > 0 )
 	{		
 		TSharedPtr<SWindow> WindowToResize = FSlateApplication::Get().FindWidgetWindow( ViewportWidget.Pin().ToSharedRef());
 
diff --git a/Engine/Source/Runtime/Engine/Private/UnrealEngine.cpp b/Engine/Source/Runtime/Engine/Private/UnrealEngine.cpp
--- a/Engine/Source/Runtime/Engine/Private/UnrealEngine.cpp
+++ b/Engine/Source/Runtime/Engine/Private/UnrealEngine.cpp
@@ -4748,6 +4748,11 @@ bool UEngine::HandleShaderComplexityCommand( const TCHAR* Cmd, FOutputDevice& Ar
 
 bool UEngine::HandleProfileGPUHitchesCommand( const TCHAR* Cmd, FOutputDevice& Ar )
 {
+	if (!FApp::CanEverRender())
+	{
+		return true;
+	}
+
 	GTriggerGPUHitchProfile = !GTriggerGPUHitchProfile;
 	if (GTriggerGPUHitchProfile)
 	{
@@ -4853,6 +4858,11 @@ bool UEngine::HandleProfileCommand( const TCHAR* Cmd, FOutputDevice& Ar )
 {
 	if ( FParse::Command(&Cmd,TEXT("GPU")) )
 	{
+		if (!FApp::CanEverRender())
+		{
+			return true;
+		}
+
 		if (!GTriggerGPUHitchProfile)
 		{
 			GTriggerGPUProfile = true;
@@ -4872,6 +4882,11 @@ bool UEngine::HandleProfileCommand( const TCHAR* Cmd, FOutputDevice& Ar )
 #if WITH_PROFILEGPU
 bool UEngine::HandleProfileGPUCommand( const TCHAR* Cmd, FOutputDevice& Ar )
 {
+	if (!FApp::CanEverRender())
+	{
+		return true;
+	}
+
 	if (FParse::Command(&Cmd, TEXT("TRACE")))
 	{
 		FString Filename = CreateProfileDirectoryAndFilename(TEXT(""), TEXT(".rtt"));
diff --git a/Engine/Source/Runtime/NullDrv/Private/NullRHI.cpp b/Engine/Source/Runtime/NullDrv/Private/NullRHI.cpp
--- a/Engine/Source/Runtime/NullDrv/Private/NullRHI.cpp
+++ b/Engine/Source/Runtime/NullDrv/Private/NullRHI.cpp
@@ -35,9 +35,9 @@ void FNullDynamicRHI::Init()
 
 	check(!GIsRHIInitialized);
 
-	// do not do this at least on dedicated server; clients with -NullRHI may need additional consideration
-#if !WITH_EDITOR
-	if (!IsRunningDedicatedServer())
+	// Create render resources if app supports rendering
+#if !WITH_EDITOR	
+	if (FApp::CanEverRender())
 #endif
 	{
 		// Notify all initialized FRenderResources that there's a valid RHI device to create their RHI resources for now.
