Push Model Replication Core.
Provides the basic interface, machinery, and codegen for allowing games to notify the networking system when properties have been changed.

diff --git a/Engine/Source/Editor/BlueprintGraph/Classes/K2Node_VariableSet.h b/Engine/Source/Editor/BlueprintGraph/Classes/K2Node_VariableSet.h
--- a/Engine/Source/Editor/BlueprintGraph/Classes/K2Node_VariableSet.h
+++ b/Engine/Source/Editor/BlueprintGraph/Classes/K2Node_VariableSet.h
@@ -40,6 +40,7 @@ class UK2Node_VariableSet : public UK2Node_Variable
 	BLUEPRINTGRAPH_API bool HasLocalRepNotify() const;
 	BLUEPRINTGRAPH_API FName GetRepNotifyName() const;
 	BLUEPRINTGRAPH_API bool ShouldFlushDormancyOnSet() const;
+	BLUEPRINTGRAPH_API bool IsNetProperty() const;
 
 	static FText GetPropertyTooltip(FProperty const* VariableProperty);
 	static FText GetBlueprintVarTooltip(FBPVariableDescription const& VarDesc);
diff --git a/Engine/Source/Editor/BlueprintGraph/Private/CallFunctionHandler.cpp b/Engine/Source/Editor/BlueprintGraph/Private/CallFunctionHandler.cpp
--- a/Engine/Source/Editor/BlueprintGraph/Private/CallFunctionHandler.cpp
+++ b/Engine/Source/Editor/BlueprintGraph/Private/CallFunctionHandler.cpp
@@ -6,11 +6,19 @@
 #include "K2Node_Event.h"
 #include "K2Node_CallParentFunction.h"
 #include "K2Node_ExecutionSequence.h"
+#include "K2Node_Self.h"
+#include "K2Node_VariableGet.h"
 #include "Kismet2/BlueprintEditorUtils.h"
+#include "K2Node_CallFunction.h"
 
 #include "EdGraphUtilities.h"
 #include "Engine/BlueprintGeneratedClass.h"
 #include "KismetCompiler.h"
+#include "Net/Core/PushModel/PushModelMacros.h"
+
+#if WITH_PUSH_MODEL
+#include "Net/NetPushModelHelpers.h"
+#endif
 
 #define LOCTEXT_NAMESPACE "CallFunctionHandler"
 
@@ -600,62 +608,229 @@ UFunction* FKCHandler_CallFunction::FindFunction(FKismetFunctionContext& Context
 void FKCHandler_CallFunction::Transform(FKismetFunctionContext& Context, UEdGraphNode* Node)
 {
 	// Add an object reference pin for this call
-
-	if (IsCalledFunctionPure(Node))
+	if (UK2Node_CallFunction * CallFuncNode = CastChecked<UK2Node_CallFunction>(Node))
 	{
-		// Flag for removal if pure and there are no consumers of the outputs
-		//@TODO: This isn't recursive (and shouldn't be here), it'll just catch the last node in a line of pure junk
-		bool bAnyOutputsUsed = false;
-		for (int32 PinIndex = 0; PinIndex < Node->Pins.Num(); ++PinIndex)
+		bool bIsPureAndNoUsedOutputs = false;
+		if (CallFuncNode->bIsPureFunc)
 		{
-			UEdGraphPin* Pin = Node->Pins[PinIndex];
-			if ((Pin->Direction == EGPD_Output) && (Pin->LinkedTo.Num() > 0))
+			// Flag for removal if pure and there are no consumers of the outputs
+			//@TODO: This isn't recursive (and shouldn't be here), it'll just catch the last node in a line of pure junk
+			bool bAnyOutputsUsed = false;
+			for (int32 PinIndex = 0; PinIndex < Node->Pins.Num(); ++PinIndex)
 			{
-				bAnyOutputsUsed = true;
-				break;
+				UEdGraphPin* Pin = Node->Pins[PinIndex];
+				if ((Pin->Direction == EGPD_Output) && (Pin->LinkedTo.Num() > 0))
+				{
+					bAnyOutputsUsed = true;
+					break;
+				}
+			}
+
+			if (!bAnyOutputsUsed)
+			{
+				//@TODO: Remove this node, not just warn about it
+				bIsPureAndNoUsedOutputs = true;
 			}
 		}
 
-		if (!bAnyOutputsUsed)
+		const UEdGraphSchema_K2* K2Schema = GetDefault<UEdGraphSchema_K2>();
+
+		// Find the function, starting at the parent class
+		if (UFunction * Function = FindFunction(Context, Node))
 		{
-			//@TODO: Remove this node, not just warn about it
+			if (Function->HasMetaData(FBlueprintMetadata::MD_Latent))
+			{
+				UEdGraphPin* OldOutPin = K2Schema->FindExecutionPin(*CallFuncNode, EGPD_Output);
 
-		}
-	}
+				if ((OldOutPin != NULL) && (OldOutPin->LinkedTo.Num() > 0))
+				{
+					// Create a dummy execution sequence that will be the target of the return call from the latent action
+					UK2Node_ExecutionSequence* DummyNode = CompilerContext.SpawnIntermediateNode<UK2Node_ExecutionSequence>(CallFuncNode);
+					DummyNode->AllocateDefaultPins();
 
-	const UEdGraphSchema_K2* K2Schema = GetDefault<UEdGraphSchema_K2>();
+					// Wire in the dummy node
+					UEdGraphPin* NewInPin = K2Schema->FindExecutionPin(*DummyNode, EGPD_Input);
+					UEdGraphPin* NewOutPin = K2Schema->FindExecutionPin(*DummyNode, EGPD_Output);
 
-	// Find the function, starting at the parent class
-	UFunction* Function = FindFunction(Context, Node);
-	if ((Function != NULL) && (Function->HasMetaData(FBlueprintMetadata::MD_Latent)))
-	{
-		UK2Node_CallFunction* CallFuncNode = CastChecked<UK2Node_CallFunction>(Node);
-		UEdGraphPin* OldOutPin = K2Schema->FindExecutionPin(*CallFuncNode, EGPD_Output);
+					if ((NewInPin != NULL) && (NewOutPin != NULL))
+					{
+						CompilerContext.MessageLog.NotifyIntermediatePinCreation(NewOutPin, OldOutPin);
 
-		if ((OldOutPin != NULL) && (OldOutPin->LinkedTo.Num() > 0))
-		{
-			// Create a dummy execution sequence that will be the target of the return call from the latent action
-			UK2Node_ExecutionSequence* DummyNode = CompilerContext.SpawnIntermediateNode<UK2Node_ExecutionSequence>(CallFuncNode);
-			DummyNode->AllocateDefaultPins();
+						while (OldOutPin->LinkedTo.Num() > 0)
+						{
+							UEdGraphPin* LinkedPin = OldOutPin->LinkedTo[0];
+
+							LinkedPin->BreakLinkTo(OldOutPin);
+							LinkedPin->MakeLinkTo(NewOutPin);
+						}
 
-			// Wire in the dummy node
-			UEdGraphPin* NewInPin = K2Schema->FindExecutionPin(*DummyNode, EGPD_Input);
-			UEdGraphPin* NewOutPin = K2Schema->FindExecutionPin(*DummyNode, EGPD_Output);
+						OldOutPin->MakeLinkTo(NewInPin);
+					}
+				}
+			}
 
-			if ((NewInPin != NULL) && (NewOutPin != NULL))
+#if WITH_PUSH_MODEL
+
+			/**
+			 * Warning!! Similar code exists in VariableSetHandler.cpp
+			 *
+			 * This code is for property dirty tracking.
+			 * It works by injecting in extra nodes while compiling that will call UNetPushModelHelpers::MarkPropertyDirtyFromRepIndex.
+			 *
+			 * That function will be called with the Owner of the property, and the RepIndex of the property.
+			 * One of these nodes needs to be added for ever Net Property that is passed by (out) reference to the function.
+			 *
+			 *
+			 * Note, this assumes that there's no way that a native class can add or remove replicated properties without also
+			 * recompiling the blueprint. The only scenario that seems possible is cooked games with custom built binaries,
+			 * but that still seems unsafe.
+			 *
+			 * If that can happen, we can instead switch to using the property name and resorting to a FindField call at runtime,
+			 * but that will be more expensive.
+			 */
+
+			// If the function is pure but won't actually be evaluated, if there are no out params,
+			// or there are no input pins, then we don't need to worry about any extra generation
+			// because there will either be no way to reference a NetProperty, or the node won't
+			// have any effect.
+			if (!bIsPureAndNoUsedOutputs && Function->NumParms > 0 && Function->HasAllFunctionFlags(FUNC_HasOutParms))
 			{
-				CompilerContext.MessageLog.NotifyIntermediatePinCreation(NewOutPin, OldOutPin);
+				static const FName MarkPropertyDirtyFuncName(TEXT("MarkPropertyDirtyFromRepIndex"));
+				static const FName ObjectPinName(TEXT("Object"));
+				static const FName RepIndexPinName(TEXT("RepIndex"));
+				static const FName PropertyNamePinName(TEXT("PropertyName"));
+
+				TArray<UEdGraphPin*> RemainingPins(Node->Pins);
+				UEdGraphPin* OldThenPin = CallFuncNode->GetThenPin();
+
+				// Note: This feels like it's going to be a very hot path during compilation as it will be hit for every
+				// CallFunction node. Any optimizations that can be made here are probably worth it to not terribly regress
+				// BP Compile time for cooked games.
 
-				while (OldOutPin->LinkedTo.Num() > 0)
+				// Iterate the properties looking for Out Params that are tied to Net Properties.
+				// This is similar to the loop in CreateCallFunction
+				for (TFieldIterator<UProperty> It(Function); It && (It->PropertyFlags & CPF_Parm); ++It)
 				{
-					UEdGraphPin* LinkedPin = OldOutPin->LinkedTo[0];
+					UProperty* Property = *It;
+					if (Property->HasAllPropertyFlags(CPF_OutParm | CPF_ReferenceParm))
+					{
+						for (int32 i = 0; i < RemainingPins.Num(); ++i)
+						{
+							if (Property->GetFName() == RemainingPins[i]->PinName)
+							{
+								UEdGraphPin* ParamPin = RemainingPins[i];
+								if (UEdGraphPin * PinToTry = FEdGraphUtilities::GetNetFromPin(ParamPin))
+								{
+									if (UK2Node_VariableGet* GetPropertyNode = Cast<UK2Node_VariableGet>(PinToTry->GetOwningNode()))
+									{
+										UProperty* ToCheck = GetPropertyNode->GetPropertyForVariable();
 
-					LinkedPin->BreakLinkTo(OldOutPin);
-					LinkedPin->MakeLinkTo(NewOutPin);
-				}
+										// We only need to do this work if we actually found a net property.
+										if (UNLIKELY(ToCheck->HasAnyPropertyFlags(CPF_Net)))
+										{
+											if (OldThenPin == nullptr)
+											{
+												// TODO: Unfortunately, we don't have an efficient way to determine whether or not a property
+												//			uses push model at this point. That is only known from a call to "GetLifetimeRepProperties"
+												//			on a given object, and even then that is a user defined function that could do any number
+												//			of things.
+												//			We *might* be able to use something like Object Annotations to cache off the results of these
+												//			calls and clear the cache post compile, but that seems hacky and prone to error.
+												CompilerContext.MessageLog.Warning(TEXT("Passing net properties as Out Parameters with Push Model Enabled is not supported. Function=%s, Param=%s, Property=%s"),
+													*Function->GetName(), *Property->GetName(), *ToCheck->GetName());
+											}
+											else if (UClass * OwningClass = ToCheck->GetTypedOuter<UClass>())
+											{
+												// TODO: As of CL-10376545 in //UE4/Main, we no longer sort Native Replicated Properties by GC Offset.
+												//			This this chunk of code can probably be removed.
+
+												// We need to make sure this class already has its property offsets setup, otherwise
+												// the order of our replicated properties won't match, meaning the RepIndex will be
+												// invalid.
+												if (ToCheck->GetOffset_ForGC() == 0)
+												{
+													if (UBlueprint * Blueprint = Cast<UBlueprint>(OwningClass->ClassGeneratedBy))
+													{
+														if (UClass * UseClass = Blueprint->GeneratedClass)
+														{
+															OwningClass = UseClass;
+															ToCheck = FindFieldChecked<UProperty>(OwningClass, ToCheck->GetFName());
+														}
+													}
+												}
+
+												ensureAlwaysMsgf(ToCheck->GetOffset_ForGC() != 0,
+													TEXT("Class does not have Property Offsets setup. This will cause issues with Push Model. Blueprint=%s, Class=%s, Property=%s"),
+													*Context.Blueprint->GetPathName(), *OwningClass->GetPathName(), *ToCheck->GetName());
+
+												if (!OwningClass->HasAnyClassFlags(CLASS_ReplicationDataIsSetUp))
+												{
+													OwningClass->SetUpRuntimeReplicationData();
+												}
+
+												// Actually insert the nodes.
+												// TODO: I don't expect this to be a very frequent occurence, but we could be more clever about this.
+												//			Instead of inserting a new node per Out Ref Net Property, we might be able to create
+												//			A separate node that takes in lists of Objects, RepIndices, and Names.
+												//			Ideally, this would work like the String Append node where we could specify the values
+												//			directly on the node.
+												{
+													// Create the node that will call MarkPropertyDirty.
+													UK2Node_CallFunction* MarkPropertyDirtyNode = Node->GetGraph()->CreateIntermediateNode<UK2Node_CallFunction>();
+													MarkPropertyDirtyNode->FunctionReference.SetExternalMember(MarkPropertyDirtyFuncName, UNetPushModelHelpers::StaticClass());
+													MarkPropertyDirtyNode->AllocateDefaultPins();
+
+													// Create the Pins for RepIndex, PropertyName, and Object.
+													UEdGraphPin* RepIndexPin = MarkPropertyDirtyNode->FindPinChecked(RepIndexPinName);
+													RepIndexPin->DefaultValue = FString::FromInt(Property->RepIndex);
+
+													UEdGraphPin* PropertyNamePin = MarkPropertyDirtyNode->FindPinChecked(PropertyNamePinName);
+													PropertyNamePin->DefaultValue = Property->GetFName().ToString();
+
+													UEdGraphPin* ObjectPin = MarkPropertyDirtyNode->FindPinChecked(ObjectPinName);
+													UEdGraphPin* PropertyOwnerPin = GetPropertyNode->FindPinChecked(UEdGraphSchema_K2::PN_Self);
+
+													// If the property is linked to some other object, go ahead and grab that.
+													// Otherwise, create an intermediate self Pin and use that.
+													if (PropertyOwnerPin->LinkedTo.Num() > 0)
+													{
+														PropertyOwnerPin = PropertyOwnerPin->LinkedTo[0];
+													}
+													else
+													{
+														UK2Node_Self* SelfNode = Node->GetGraph()->CreateIntermediateNode<UK2Node_Self>();
+														SelfNode->AllocateDefaultPins();
+														PropertyOwnerPin = SelfNode->FindPinChecked(UEdGraphSchema_K2::PN_Self);
+													}
+
+													ObjectPin->MakeLinkTo(PropertyOwnerPin);
+
+													UEdGraphPin* NewThenPin = MarkPropertyDirtyNode->GetThenPin();
+													if (ensure(NewThenPin))
+													{
+														NewThenPin->CopyPersistentDataFromOldPin(*OldThenPin);
+														OldThenPin->BreakAllPinLinks();
+														OldThenPin->MakeLinkTo(MarkPropertyDirtyNode->GetExecPin());
+
+														// Now that the connection is established, go ahead and set our OldThenPin to
+														// the NewThenPin. If we find another NetProperty later, this will save us from
+														// doing a lookup.
+														OldThenPin = NewThenPin;
+													}
+												}
+											}
+										}
+									}
+								}
 
-				OldOutPin->MakeLinkTo(NewInPin);
+								RemainingPins.RemoveAtSwap(i);
+								break;
+							}
+						}
+					}
+				}
 			}
+#endif
 		}
 	}
 }
diff --git a/Engine/Source/Editor/BlueprintGraph/Private/K2Node_VariableSet.cpp b/Engine/Source/Editor/BlueprintGraph/Private/K2Node_VariableSet.cpp
--- a/Engine/Source/Editor/BlueprintGraph/Private/K2Node_VariableSet.cpp
+++ b/Engine/Source/Editor/BlueprintGraph/Private/K2Node_VariableSet.cpp
@@ -306,6 +306,12 @@ bool UK2Node_VariableSet::ShouldFlushDormancyOnSet() const
 	return (Property != NULL && (Property->PropertyFlags & CPF_Net));
 }
 
+bool UK2Node_VariableSet::IsNetProperty() const
+{
+	UProperty* Property = GetPropertyForVariable();
+	return Property && (Property->PropertyFlags & CPF_Net);
+}
+
 FName UK2Node_VariableSet::GetRepNotifyName() const
 {
 	FProperty * Property = GetPropertyForVariable();
diff --git a/Engine/Source/Editor/BlueprintGraph/Private/VariableSetHandler.cpp b/Engine/Source/Editor/BlueprintGraph/Private/VariableSetHandler.cpp
--- a/Engine/Source/Editor/BlueprintGraph/Private/VariableSetHandler.cpp
+++ b/Engine/Source/Editor/BlueprintGraph/Private/VariableSetHandler.cpp
@@ -5,7 +5,9 @@
 #include "K2Node_CallFunction.h"
 #include "K2Node_Variable.h"
 #include "K2Node_VariableSet.h"
+#include "K2Node_Self.h"
 #include "Kismet2/BlueprintEditorUtils.h"
+#include "Net/NetPushModelHelpers.h"
 
 #include "EdGraphUtilities.h"
 #include "KismetCompiler.h"
@@ -196,6 +198,108 @@ void FKCHandler_VariableSet::Transform(FKismetFunctionContext& Context, UEdGraph
 				OldThenPin->MakeLinkTo(CallFuncNode->GetExecPin());
 			}
 		}
+
+		if (SetNotify->IsNetProperty())
+		{
+			/**
+			 * Warning!! Similar code exists in CallFunctionHandler.cpp
+			 *
+			 * This code is for property dirty tracking.
+			 * It works by injecting in extra nodes while compiling that will call UNetPushModelHelpers::MarkPropertyDirtyFromRepIndex.
+			 *
+			 * That function will be called with the Owner of the property (either Self or whatever is connected to the Target pin
+			 * of the BP Node), and the RepIndex of the property.
+			 *
+			 * Note, this assumes that there's no way that a native class can add or remove replicated properties without also
+			 * recompiling the blueprint. The only scenario that seems possible is cooked games with custom built binaries,
+			 * but that still seems unsafe.
+			 *
+			 * If that can happen, we can instead switch to using the property name and resorting to a FindField call at runtime,
+			 * but that will be more expensive.
+			 */
+
+			static const FName MarkPropertyDirtyFuncName(TEXT("MarkPropertyDirtyFromRepIndex"));
+			static const FName ObjectPinName(TEXT("Object"));
+			static const FName RepIndexPinName(TEXT("RepIndex"));
+			static const FName PropertyNamePinName(TEXT("PropertyName"));
+
+			UProperty* Property = SetNotify->GetPropertyForVariable();
+			UClass* Class = Property->GetTypedOuter<UClass>();
+
+			if (Property && Class)
+			{
+				// We need to make sure this class already has its property offsets setup, otherwise
+				// the order of our replicated properties won't match, meaning the RepIndex will be
+				// invalid.
+				if (Property->GetOffset_ForGC() == 0)
+				{
+					// Make sure that we're using the correct class and that it has replication data set up.
+					if (Class->ClassGeneratedBy == Context.Blueprint && Context.NewClass && Context.NewClass != Class)
+					{
+						Class = Context.NewClass;
+						Property = FindFieldChecked<UProperty>(Class, Property->GetFName());
+					}
+					if (Property->GetOffset_ForGC() == 0)
+					{
+						if (UBlueprint * Blueprint = Cast<UBlueprint>(Class->ClassGeneratedBy))
+						{
+							if (UClass * UseClass = Blueprint->GeneratedClass)
+							{
+								Class = UseClass;
+								Property = FindFieldChecked<UProperty>(Class, Property->GetFName());
+							}
+						}
+					}
+				}
+
+				ensureAlwaysMsgf(Property->GetOffset_ForGC() != 0,
+					TEXT("Class does not have Property Offsets setup. This will cause issues with Push Model. Blueprint=%s, Class=%s, Property=%s"),
+					*Context.Blueprint->GetPathName(), *Class->GetPathName(), *Property->GetName());
+
+				if (!Class->HasAnyClassFlags(CLASS_ReplicationDataIsSetUp))
+				{
+					Class->SetUpRuntimeReplicationData();
+				}
+
+				UK2Node_CallFunction* CallFuncNode = Node->GetGraph()->CreateIntermediateNode<UK2Node_CallFunction>();
+				CallFuncNode->FunctionReference.SetExternalMember(MarkPropertyDirtyFuncName, UNetPushModelHelpers::StaticClass());
+				CallFuncNode->AllocateDefaultPins();
+
+				// Take our old Self (Target) pin and hook it up to the Object pin for UNetPushModelHelpers::MarkPropertyDirty.
+				// If our Self pin isn't hooked up to anything, then create an intermediate Self node and use that.
+				UEdGraphPin* SelfPin = Node->FindPinChecked(UEdGraphSchema_K2::PN_Self);
+				UEdGraphPin* ObjectPin = CallFuncNode->FindPinChecked(ObjectPinName);
+
+				if (SelfPin->LinkedTo.Num() > 0)
+				{
+					SelfPin = SelfPin->LinkedTo[0];
+				}
+				else
+				{
+					UK2Node_Self* SelfNode = Node->GetGraph()->CreateIntermediateNode<UK2Node_Self>();
+					SelfNode->AllocateDefaultPins();
+					SelfPin = SelfNode->FindPinChecked(UEdGraphSchema_K2::PN_Self);
+				}
+
+				ObjectPin->MakeLinkTo(SelfPin);
+
+				UEdGraphPin* RepIndexPin = CallFuncNode->FindPinChecked(RepIndexPinName);
+				RepIndexPin->DefaultValue = FString::FromInt(Property->RepIndex);
+
+				UEdGraphPin* PropertyNamePin = CallFuncNode->FindPinChecked(PropertyNamePinName);
+				PropertyNamePin->DefaultValue = Property->GetFName().ToString();
+
+				// Hook up our exec pins.
+				UEdGraphPin* OldThenPin = Node->FindPinChecked(UEdGraphSchema_K2::PN_Then);
+				UEdGraphPin* NewThenPin = CallFuncNode->GetThenPin();
+				if (ensure(NewThenPin))
+				{
+					NewThenPin->CopyPersistentDataFromOldPin(*OldThenPin);
+					OldThenPin->BreakAllPinLinks();
+					OldThenPin->MakeLinkTo(CallFuncNode->GetExecPin());
+				}
+			}
+		}
 	}
 }
 
diff --git a/Engine/Source/Programs/UnrealBuildTool/Configuration/TargetRules.cs b/Engine/Source/Programs/UnrealBuildTool/Configuration/TargetRules.cs
--- a/Engine/Source/Programs/UnrealBuildTool/Configuration/TargetRules.cs
+++ b/Engine/Source/Programs/UnrealBuildTool/Configuration/TargetRules.cs
@@ -584,6 +584,12 @@ namespace UnrealBuildTool
 		}
 		private bool? bWithServerCodeOverride;
 
+		/// <summary>
+		/// When enabled, Push Model Networking will be used on the server. This can help reduce CPU overhead of networking, at the cost of more memory.
+		/// </summary>
+		[RequiresUniqueBuildEnvironment]
+		public bool bWithPushModel = false;
+
 		/// <summary>
 		/// Whether to include stats support even without the engine.
 		/// </summary>
@@ -2053,6 +2059,11 @@ namespace UnrealBuildTool
 			get { return Inner.bWithServerCode; }
 		}
 
+		public bool bWithPushModel
+		{
+			get { return Inner.bWithPushModel; }
+		}
+
 		public bool bCompileWithStatsWithoutEngine
 		{
 			get { return Inner.bCompileWithStatsWithoutEngine; }
diff --git a/Engine/Source/Programs/UnrealBuildTool/Configuration/UEBuildTarget.cs b/Engine/Source/Programs/UnrealBuildTool/Configuration/UEBuildTarget.cs
--- a/Engine/Source/Programs/UnrealBuildTool/Configuration/UEBuildTarget.cs
+++ b/Engine/Source/Programs/UnrealBuildTool/Configuration/UEBuildTarget.cs
@@ -3594,6 +3594,16 @@ namespace UnrealBuildTool
 				GlobalCompileEnvironment.Definitions.Add("WITH_SERVER_CODE=0");
 			}
 
+			// Set the defines for Push Model
+			if (Rules.bWithPushModel)
+			{
+				GlobalCompileEnvironment.Definitions.Add("WITH_PUSH_MODEL=1");
+			}
+			else
+			{
+				GlobalCompileEnvironment.Definitions.Add("WITH_PUSH_MODEL=0");
+			}
+
 			// Set the define for whether we're compiling with CEF3
 			if (Rules.bCompileCEF3 && (Platform == UnrealTargetPlatform.Win32 || Platform == UnrealTargetPlatform.Win64 || Platform == UnrealTargetPlatform.Mac || Platform == UnrealTargetPlatform.Linux))
 			{
diff --git a/Engine/Source/Programs/UnrealHeaderTool/Private/CodeGenerator.cpp b/Engine/Source/Programs/UnrealHeaderTool/Private/CodeGenerator.cpp
--- a/Engine/Source/Programs/UnrealHeaderTool/Private/CodeGenerator.cpp
+++ b/Engine/Source/Programs/UnrealHeaderTool/Private/CodeGenerator.cpp
@@ -257,7 +257,8 @@ namespace
 		FClass* Class,
 		FClass* SuperClass,
 		FOutputDevice& Writer,
-		const FUnrealSourceFile& SourceFile)
+		const FUnrealSourceFile& SourceFile,
+		FNativeClassHeaderGenerator::EExportClassOutFlags& OutFlags)
 	{
 		const bool bHasGetLifetimeReplicatedProps = HasIdentifierExactMatch(ClassRange.Start, ClassRange.End, GetLifetimeReplicatedPropsStr);
 
@@ -278,6 +279,18 @@ namespace
 		Class->SetUpUhtReplicationData(&ParentRepClass);
 
 		ExportNetData(Writer, Class, ParentRepClass, API);
+		
+		// If we don't have any Super Classes that have replicated properties,
+		// then we're the base most replicated class, and we should add in our interface stuff.
+		if (!ParentRepClass)
+		{
+			OutFlags |= FNativeClassHeaderGenerator::EExportClassOutFlags::NeedsPushModelHeaders;
+			Writer.Logf(TEXT(
+				"private:\r\n"
+				"\tREPLICATED_BASE_CLASS(%s)\r\n"
+				"public:\r\n"
+			), *Class->GetName());
+		}
 	}
 }
 
@@ -2949,7 +2962,8 @@ void FNativeClassHeaderGenerator::ExportClassFromSourceFileInner(
 	FOutputDevice&           OutDeclarations,
 	FReferenceGatherers&     OutReferenceGatherers,
 	FClass*                  Class,
-	const FUnrealSourceFile& SourceFile
+	const FUnrealSourceFile& SourceFile,
+	EExportClassOutFlags&    OutFlags
 ) const
 {
 	FUHTStringBuilder StandardUObjectConstructorsMacroCall;
@@ -3031,8 +3045,6 @@ void FNativeClassHeaderGenerator::ExportClassFromSourceFileInner(
 
 	FString PPOMacroName;
 
-	// Replication, add in the declaration for GetLifetimeReplicatedProps() automatically if there are any net flagged properties
-
 	ClassDefinitionRange ClassRange;
 	if (ClassDefinitionRange* FoundRange = ClassDefinitionRanges.Find(Class))
 	{
@@ -3190,7 +3202,7 @@ void FNativeClassHeaderGenerator::ExportClassFromSourceFileInner(
 
 			if (ClassHasReplicatedProperties(Class))
 			{
-				WriteReplicatedMacroData(ClassRange, *ClassCPPName, *APIArg, Class, SuperClass, InterfaceBoilerplate, SourceFile);
+				WriteReplicatedMacroData(ClassRange, *ClassCPPName, *APIArg, Class, SuperClass, InterfaceBoilerplate, SourceFile, OutFlags);
 			}
 
 			FString NoPureDeclsMacroName = SourceFile.GetGeneratedMacroName(ClassData, TEXT("_INCLASS_IINTERFACE_NO_PURE_DECLS"));
@@ -3217,7 +3229,7 @@ void FNativeClassHeaderGenerator::ExportClassFromSourceFileInner(
 
 			if (ClassHasReplicatedProperties(Class))
 			{
-				WriteReplicatedMacroData(ClassRange, *ClassCPPName, *APIArg, Class, SuperClass, Boilerplate, SourceFile);
+				WriteReplicatedMacroData(ClassRange, *ClassCPPName, *APIArg, Class, SuperClass, Boilerplate, SourceFile, OutFlags);
 			}
 
 			{
@@ -3475,12 +3487,18 @@ void FNativeClassHeaderGenerator::ExportConstructorsMacros(FOutputDevice& OutGen
 	EnhancedUObjectConstructorsMacroCall.Logf(TEXT("\t%s\r\n"), *EnhMacroName);
 }
 
-bool FNativeClassHeaderGenerator::WriteHeader(const TCHAR* Path, const FString& InBodyText, FReferenceGatherers& InOutReferenceGatherers) const
+bool FNativeClassHeaderGenerator::WriteHeader(const TCHAR* Path, const FString& InBodyText, const TSet<FString>& InAdditionalHeaders, FReferenceGatherers& InOutReferenceGatherers) const
 {
 	FUHTStringBuilder GeneratedHeaderTextWithCopyright;
 	GeneratedHeaderTextWithCopyright.Logf(TEXT("%s"), HeaderCopyright);
 	GeneratedHeaderTextWithCopyright.Log(TEXT("#include \"UObject/ObjectMacros.h\"\r\n"));
 	GeneratedHeaderTextWithCopyright.Log(TEXT("#include \"UObject/ScriptMacros.h\"\r\n"));
+
+	for (const FString& AdditionalHeader : InAdditionalHeaders)
+	{
+		GeneratedHeaderTextWithCopyright.Logf(TEXT("#include \"%s\"\r\n"), *AdditionalHeader);
+	}
+
 	GeneratedHeaderTextWithCopyright.Log(LINE_TERMINATOR);
 	GeneratedHeaderTextWithCopyright.Log(TEXT("PRAGMA_DISABLE_DEPRECATION_WARNINGS") LINE_TERMINATOR);
 
@@ -5744,16 +5762,23 @@ FNativeClassHeaderGenerator::FNativeClassHeaderGenerator(
 			ConstThis->ExportDelegateDefinition(GeneratedHeaderText, ReferenceGatherers, *SourceFile, Func);
 		}
 
+		EExportClassOutFlags ExportFlags = EExportClassOutFlags::None;
+		TSet<FString> AdditionalHeaders;
 		TArray<UClass*>	DefinedClasses;
 		SourceFile->AppendDefinedClasses(DefinedClasses);
 		for (UClass* Class : DefinedClasses)
 		{
 			if (!(Class->ClassFlags & CLASS_Intrinsic))
 			{
-				ConstThis->ExportClassFromSourceFileInner(GeneratedHeaderText, OutText, GeneratedFunctionDeclarations, ReferenceGatherers, (FClass*)Class, *SourceFile);
+				ConstThis->ExportClassFromSourceFileInner(GeneratedHeaderText, OutText, GeneratedFunctionDeclarations, ReferenceGatherers, (FClass*)Class, *SourceFile, ExportFlags);
 			}
 		}
 
+		if (EnumHasAnyFlags(ExportFlags, EExportClassOutFlags::NeedsPushModelHeaders))
+		{
+			AdditionalHeaders.Add(FString(TEXT("Net/Core/PushModel/PushModelMacros.h")));
+		}
+
 		GeneratedHeaderText.Log(TEXT("#undef CURRENT_FILE_ID\r\n"));
 		GeneratedHeaderText.Logf(TEXT("#define CURRENT_FILE_ID %s\r\n\r\n\r\n"), *SourceFile->GetFileId());
 
@@ -5763,7 +5788,7 @@ FNativeClassHeaderGenerator::FNativeClassHeaderGenerator(
 		}
 
 		FString HeaderPath = BaseSourceFilename + TEXT(".generated.h");
-		bool bHasChanged = ConstThis->WriteHeader(*HeaderPath, GeneratedHeaderText, ReferenceGatherers);
+		bool bHasChanged = ConstThis->WriteHeader(*HeaderPath, GeneratedHeaderText, AdditionalHeaders, ReferenceGatherers);
 
 		SourceFile->SetGeneratedFilename(MoveTemp(HeaderPath));
 		SourceFile->SetHasChanged(bHasChanged);
diff --git a/Engine/Source/Programs/UnrealHeaderTool/Private/NativeClassExporter.h b/Engine/Source/Programs/UnrealHeaderTool/Private/NativeClassExporter.h
--- a/Engine/Source/Programs/UnrealHeaderTool/Private/NativeClassExporter.h
+++ b/Engine/Source/Programs/UnrealHeaderTool/Private/NativeClassExporter.h
@@ -184,6 +184,16 @@ private:
 	 */
 	static void ExportUInterfaceBoilerplate(FUHTStringBuilder& UInterfaceBoilerplate, FClass* Class, const FString& FriendText);
 
+public:
+
+	enum class EExportClassOutFlags
+	{
+		None = 0x0,
+		NeedsPushModelHeaders = 0x1 << 0,
+	};
+
+private:
+
 	/**
 	 * After all of the dependency checking, and setup for isolating the generated code, actually export the class
 	 *
@@ -197,7 +207,8 @@ private:
 		FOutputDevice&           OutDeclarations,
 		FReferenceGatherers&     OutReferenceGatherers, 
 		FClass*                  Class,
-		const FUnrealSourceFile& SourceFile
+		const FUnrealSourceFile& SourceFile,
+		EExportClassOutFlags&    OutFlags
 	) const;
 
 	/**
@@ -205,7 +216,7 @@ private:
 	 *
 	 * @param	SourceFile			Source file to export.
 	 */
-	bool WriteHeader(const TCHAR* Path, const FString& InBodyText, FReferenceGatherers& InOutReferenceGatherers) const;
+	bool WriteHeader(const TCHAR* Path, const FString& InBodyText, const TSet<FString>& InAdditionalHeaders, FReferenceGatherers& InOutReferenceGatherers) const;
 
 	/**
 	 * Returns a string in the format CLASS_Something|CLASS_Something which represents all class flags that are set for the specified
@@ -552,3 +563,5 @@ public:
 	 */
 	static bool IsMissingVirtualSpecifier(const FString& SourceFile, int32 FunctionNamePosition);
 };
+
+ENUM_CLASS_FLAGS(FNativeClassHeaderGenerator::EExportClassOutFlags);
\ No newline at end of file
diff --git a/Engine/Source/Runtime/CoreUObject/Public/UObject/CoreNet.h b/Engine/Source/Runtime/CoreUObject/Public/UObject/CoreNet.h
--- a/Engine/Source/Runtime/CoreUObject/Public/UObject/CoreNet.h
+++ b/Engine/Source/Runtime/CoreUObject/Public/UObject/CoreNet.h
@@ -247,20 +247,46 @@ struct FPropertyRetirement
 class FLifetimeProperty
 {
 public:
-	uint16				RepIndex;
-	ELifetimeCondition	Condition;
+
+	uint16 RepIndex;
+	ELifetimeCondition Condition;
 	ELifetimeRepNotifyCondition RepNotifyCondition;
+	bool bIsPushBased;
 
-	FLifetimeProperty() : RepIndex( 0 ), Condition( COND_None ), RepNotifyCondition(REPNOTIFY_OnChanged) {}
-	FLifetimeProperty( int32 InRepIndex ) : RepIndex( InRepIndex ), Condition( COND_None ), RepNotifyCondition(REPNOTIFY_OnChanged) { check( InRepIndex <= 65535 ); }
-	FLifetimeProperty(int32 InRepIndex, ELifetimeCondition InCondition, ELifetimeRepNotifyCondition InRepNotifyCondition=REPNOTIFY_OnChanged) : RepIndex(InRepIndex), Condition(InCondition), RepNotifyCondition(InRepNotifyCondition) { check(InRepIndex <= 65535); }
+	FLifetimeProperty()
+		: RepIndex(0)
+		, Condition(COND_None)
+		, RepNotifyCondition(REPNOTIFY_OnChanged)
+		, bIsPushBased(false)
+	{
+	}
+
+	FLifetimeProperty(int32 InRepIndex)
+		: RepIndex(InRepIndex)
+		, Condition(COND_None)
+		, RepNotifyCondition(REPNOTIFY_OnChanged)
+		, bIsPushBased(false)
+	{
+		check(InRepIndex <= 65535);
+	}
+
+	FLifetimeProperty(int32 InRepIndex, ELifetimeCondition InCondition, ELifetimeRepNotifyCondition InRepNotifyCondition=REPNOTIFY_OnChanged, bool bInIsPushBased=false)
+		: RepIndex(InRepIndex)
+		, Condition(InCondition)
+		, RepNotifyCondition(InRepNotifyCondition)
+		, bIsPushBased(bInIsPushBased)
+	{
+		check(InRepIndex <= 65535);
+	}
 
-	inline bool operator==( const FLifetimeProperty& Other ) const
+	inline bool operator==(const FLifetimeProperty& Other) const
 	{
-		if ( RepIndex == Other.RepIndex )
+		if (RepIndex == Other.RepIndex)
 		{
-			check( Condition == Other.Condition );		// Can't have different conditions if the RepIndex matches, doesn't make sense
-			check( RepNotifyCondition == Other.RepNotifyCondition);
+			// Can't have different conditions if the RepIndex matches, doesn't make sense
+			check(Condition == Other.Condition);
+			check(RepNotifyCondition == Other.RepNotifyCondition);
+			check(bIsPushBased == Other.bIsPushBased);
 			return true;
 		}
 
@@ -432,9 +458,12 @@ class IRepChangedPropertyTracker
 public:
 	IRepChangedPropertyTracker() { }
 
-	virtual void SetCustomIsActiveOverride( const uint16 RepIndex, const bool bIsActive ) = 0;
+	virtual void SetCustomIsActiveOverride(
+		UObject* OwningObject,
+		const uint16 RepIndex,
+		const bool bIsActive) = 0;
 
-	virtual void SetExternalData( const uint8* Src, const int32 NumBits ) = 0;
+	virtual void SetExternalData(const uint8* Src, const int32 NumBits) = 0;
 
 	virtual bool IsReplay() const = 0;
 
diff --git a/Engine/Source/Runtime/CoreUObject/Public/UObject/Object.h b/Engine/Source/Runtime/CoreUObject/Public/UObject/Object.h
--- a/Engine/Source/Runtime/CoreUObject/Public/UObject/Object.h
+++ b/Engine/Source/Runtime/CoreUObject/Public/UObject/Object.h
@@ -1457,7 +1457,40 @@ private:
 
 public:
 
-	virtual void ValidateGeneratedRepEnums(const TArray<struct FRepRecord>& ClassReps) const {}
+	virtual void ValidateGeneratedRepEnums(const TArray<struct FRepRecord>& ClassReps) const {}	
+
+private:
+	
+	friend struct FObjectNetPushIdHelper;
+	virtual void SetNetPushIdDynamic(const int32 NewNetPushId)
+	{
+		// This method should only be called on Objects that are networked, and those should
+		// always have this implemented (by UHT).
+		check(false);
+	}
+	
+	virtual int32 GetNetPushIdDynamic() const
+	{
+		return INDEX_NONE;
+	}
+};
+
+struct FObjectNetPushIdHelper
+{
+private:
+	friend struct FNetPrivatePushIdHelper;
+	friend class UNetPushModelHelpers;
+	friend class UKismetArrayLibrary;
+
+	static void SetNetPushIdDynamic(UObject* Object, const int32 NewNetPushId)
+	{
+		Object->SetNetPushIdDynamic(NewNetPushId);
+	}
+	
+	static int32 GetNetPushIdDynamic(const UObject* const Object)
+	{
+		return Object->GetNetPushIdDynamic();
+	}
 };
 
 /**
diff --git a/Engine/Source/Runtime/Engine/Classes/Kismet/KismetArrayLibrary.h b/Engine/Source/Runtime/Engine/Classes/Kismet/KismetArrayLibrary.h
--- a/Engine/Source/Runtime/Engine/Classes/Kismet/KismetArrayLibrary.h
+++ b/Engine/Source/Runtime/Engine/Classes/Kismet/KismetArrayLibrary.h
@@ -10,6 +10,7 @@
 #include "UObject/UnrealType.h"
 #include "UObject/ScriptMacros.h"
 #include "Kismet/BlueprintFunctionLibrary.h"
+#include "Net/Core/PushModel/PushModel.h"
 #include "KismetArrayLibrary.generated.h"
 
 class AActor;
@@ -263,6 +264,7 @@ public:
  
  		P_FINISH;
 		P_NATIVE_BEGIN;
+		MARK_PROPERTY_DIRTY(Stack.Object, ArrayProperty);
 		*(int32*)RESULT_PARAM = GenericArray_Add(ArrayAddr, ArrayProperty, StorageSpace);
 		P_NATIVE_END;
 		InnerProp->DestroyValue(StorageSpace);
@@ -291,6 +293,7 @@ public:
 
 		P_FINISH;
 		P_NATIVE_BEGIN;
+		MARK_PROPERTY_DIRTY(Stack.Object, ArrayProperty);
 		*(int32*)RESULT_PARAM = GenericArray_AddUnique(ArrayAddr, ArrayProperty, StorageSpace);
 		P_NATIVE_END;
 		InnerProp->DestroyValue(StorageSpace);
@@ -310,6 +313,7 @@ public:
 
 		P_FINISH;
 		P_NATIVE_BEGIN;
+		MARK_PROPERTY_DIRTY(Stack.Object, ArrayProperty);
 		GenericArray_Shuffle(ArrayAddr, ArrayProperty);
 		P_NATIVE_END;
 	}
@@ -368,6 +372,7 @@ public:
 
 		P_FINISH;
 		P_NATIVE_BEGIN;
+		MARK_PROPERTY_DIRTY(Stack.Object, TargetArrayProperty);
 		GenericArray_Append(TargetArrayAddr, TargetArrayProperty, SourceArrayAddr, SourceArrayProperty);
 		P_NATIVE_END;
 	}
@@ -396,6 +401,7 @@ public:
 		P_GET_PROPERTY(FIntProperty, Index);
 		P_FINISH;
 		P_NATIVE_BEGIN;
+		MARK_PROPERTY_DIRTY(Stack.Object, ArrayProperty);
 		GenericArray_Insert(ArrayAddr, ArrayProperty, StorageSpace, Index);
 		P_NATIVE_END;
 		InnerProp->DestroyValue(StorageSpace);
@@ -416,6 +422,7 @@ public:
 		P_GET_PROPERTY(FIntProperty, Index);
 		P_FINISH;
 		P_NATIVE_BEGIN;
+		MARK_PROPERTY_DIRTY(Stack.Object, ArrayProperty);
 		GenericArray_Remove(ArrayAddr, ArrayProperty, Index);
 		P_NATIVE_END;
 	}
@@ -446,6 +453,7 @@ public:
 		// Bools need to be processed internally by the property so that C++ bool value is properly set.
 		GenericArray_HandleBool(InnerProp, ItemPtr);
 		P_NATIVE_BEGIN;
+		MARK_PROPERTY_DIRTY(Stack.Object, ArrayProperty);
 		*(bool*)RESULT_PARAM = GenericArray_RemoveItem(ArrayAddr, ArrayProperty, ItemPtr);
 		P_NATIVE_END;
 
@@ -465,6 +473,7 @@ public:
 		}
 		P_FINISH;
 		P_NATIVE_BEGIN;
+		MARK_PROPERTY_DIRTY(Stack.Object, ArrayProperty);
 		GenericArray_Clear(ArrayAddr, ArrayProperty);
 		P_NATIVE_END;
 	}
@@ -483,6 +492,7 @@ public:
 		P_GET_PROPERTY(FIntProperty, Size);
 		P_FINISH;
 		P_NATIVE_BEGIN;
+		MARK_PROPERTY_DIRTY(Stack.Object, ArrayProperty);
 		GenericArray_Resize(ArrayAddr, ArrayProperty, Size);
 		P_NATIVE_END;
 	}
@@ -578,6 +588,7 @@ public:
 		P_FINISH;
 
 		P_NATIVE_BEGIN;
+		MARK_PROPERTY_DIRTY(Stack.Object, ArrayProperty);
 		GenericArray_Set(ArrayAddr, ArrayProperty, Index, StorageSpace, bSizeToFit);
 		P_NATIVE_END;
 		InnerProp->DestroyValue(StorageSpace);
@@ -601,6 +612,7 @@ public:
 
 		P_FINISH;
 		P_NATIVE_BEGIN;
+		MARK_PROPERTY_DIRTY(Stack.Object, ArrayProperty);
 		GenericArray_Swap(ArrayAddr, ArrayProperty, First, Second);
 		P_NATIVE_END;
 	}
@@ -684,6 +696,7 @@ public:
 		P_FINISH;
 
 		P_NATIVE_BEGIN;
+		MARK_PROPERTY_DIRTY(Stack.Object, Stack.MostRecentProperty);
 		GenericArray_SetArrayPropertyByName(OwnerObject, ArrayPropertyName, SrcArrayAddr);
 		P_NATIVE_END;
 	}
diff --git a/Engine/Source/Runtime/Engine/Engine.Build.cs b/Engine/Source/Runtime/Engine/Engine.Build.cs
--- a/Engine/Source/Runtime/Engine/Engine.Build.cs
+++ b/Engine/Source/Runtime/Engine/Engine.Build.cs
@@ -20,6 +20,7 @@ public class Engine : ModuleRules
 				"Developer/DerivedDataCache/Public",
 				"Runtime/SynthBenchmark/Public",
 				"Runtime/Engine/Private",
+				"Runtime/Net/Core/Private/Net/Core/PushModel/Types"
 			}
 		);
 
diff --git a/Engine/Source/Runtime/Engine/Private/BlueprintGeneratedClass.cpp b/Engine/Source/Runtime/Engine/Private/BlueprintGeneratedClass.cpp
--- a/Engine/Source/Runtime/Engine/Private/BlueprintGeneratedClass.cpp
+++ b/Engine/Source/Runtime/Engine/Private/BlueprintGeneratedClass.cpp
@@ -20,6 +20,7 @@
 #include "Engine/SCS_Node.h"
 #include "Engine/InheritableComponentHandler.h"
 #include "Misc/ScopeLock.h"
+#include "Net/Core/PushModel/PushModel.h"
 #include "UObject/CoreObjectVersion.h"
 
 #if WITH_EDITOR
@@ -1680,7 +1681,7 @@ void UBlueprintGeneratedClass::GetLifetimeBlueprintReplicationList(TArray<FLifet
 		{
 			PropertiesLeft--;
 			
-			OutLifetimeProps.AddUnique(FLifetimeProperty(Prop->RepIndex, Prop->GetBlueprintReplicationCondition()));
+			OutLifetimeProps.AddUnique(FLifetimeProperty(Prop->RepIndex, Prop->GetBlueprintReplicationCondition(), REPNOTIFY_OnChanged, PUSH_MAKE_BP_PROPERTIES_PUSH_MODEL()));
 		}
 	}
 
diff --git a/Engine/Source/Runtime/Engine/Private/Net/NetPushModelHelpers.cpp b/Engine/Source/Runtime/Engine/Private/Net/NetPushModelHelpers.cpp
new file mode 100644
--- /dev/null
+++ b/Engine/Source/Runtime/Engine/Private/Net/NetPushModelHelpers.cpp
@@ -0,0 +1,76 @@
+// Copyright 1998-2019 Epic Games, Inc. All Rights Reserved.
+
+#include "Net/NetPushModelHelpers.h"
+#include "Net/Core/PushModel/PushModel.h"
+
+void UNetPushModelHelpers::MarkPropertyDirty(UObject* Object, FName PropertyName)
+{
+#if WITH_PUSH_MODEL
+	if (Object && IS_PUSH_MODEL_ENABLED())
+	{
+		UClass* Class = Object->GetClass();
+		if (Class->HasAnyClassFlags(CLASS_ReplicationDataIsSetUp))
+		{
+			// If this is too slow to be practical, we might be able build a lookup from ClassName+PropertyName -> RepIndex.
+			// It would be safest to invalidate the lookup on a hot reload / recompile (if possible), blueprint recompile,
+			// and map change.
+			UProperty* Property = FindField<UProperty>(Class, PropertyName);
+			if (Property == nullptr)
+			{
+				UE_LOG(LogNet, Warning, TEXT("UNetPushModelHelpers::MarkPropertyDirty: Unable to find Property %s in Class %s"), *PropertyName.ToString(), *Class->GetPathName());
+			}
+			else if ((Property->PropertyFlags & CPF_Net) == 0)
+			{
+				UE_LOG(LogNet, Warning, TEXT("UNetPushModelHelpers::MarkPropertyDirty: Property %s in class %s is not Replicated."), *PropertyName.ToString(), *Class->GetPathName());
+			}
+			else
+			{
+				MARK_PROPERTY_DIRTY_UNSAFE(Object, Property->RepIndex);
+			}
+		}
+	}
+#endif
+}
+
+#if WITH_PUSH_MODEL
+namespace UE4PushModelPrivate
+{
+#if WITH_PUSH_VALIDATION_SUPPORT
+	static bool bCheckPushBPRepIndexAgainstName = false;
+	FAutoConsoleVariableRef CVarCheckPushBPRepIndexAgainstName(
+		TEXT("Net.CheckPushBPRepIndexAgainstName"),
+		bCheckPushBPRepIndexAgainstName,
+		TEXT("When enabled, validates that BP generated values passed to MarkPropertyDirtyFromRepIndex match the actual property data")
+	);
+#else
+	static constexpr bool bCheckPushBPRepIndexAgainstName = false;
+#endif
+}
+#endif
+
+void UNetPushModelHelpers::MarkPropertyDirtyFromRepIndex(UObject* Object, int32 RepIndex, FName PropertyName)
+{
+#if WITH_PUSH_MODEL
+	if (Object && IS_PUSH_MODEL_ENABLED())
+	{
+		UClass* Class = Object->GetClass();
+		if (Class->HasAnyClassFlags(CLASS_ReplicationDataIsSetUp))
+		{
+			if (RepIndex < INDEX_NONE || RepIndex >= Class->ClassReps.Num())
+			{
+				UE_LOG(LogNet, Warning, TEXT("UNetPushModelHelpers::MarkPropertyDirtyFromRepIndex: Invalid Rep Index. Class %s RepIndex %d"), *Class->GetPathName(), RepIndex);
+			}
+			else
+			{
+#if WITH_PUSH_VALIDATION_SUPPORT
+				checkf(!UE4PushModelPrivate::bCheckPushBPRepIndexAgainstName || Class->ClassReps[RepIndex].Property->GetFName() == PropertyName,
+					TEXT("Property and RepIndex don't match! Object=%s, RepIndex=%d, InPropertyName=%s, FoundPropertyName=%s"),
+						*Object->GetPathName(), RepIndex, *PropertyName.ToString(), *(Class->ClassReps[RepIndex].Property->GetName()));
+#endif
+	
+				MARK_PROPERTY_DIRTY_UNSAFE(Object, RepIndex);
+			}
+		}
+	}
+#endif
+}
\ No newline at end of file
diff --git a/Engine/Source/Runtime/Engine/Private/NetworkDriver.cpp b/Engine/Source/Runtime/Engine/Private/NetworkDriver.cpp
--- a/Engine/Source/Runtime/Engine/Private/NetworkDriver.cpp
+++ b/Engine/Source/Runtime/Engine/Private/NetworkDriver.cpp
@@ -3251,6 +3251,8 @@ void UNetDriver::PostGarbageCollect()
 			It.RemoveCurrent();
 		}
 	}
+	
+	CALL_PUSH_MODEL_POSTGARBAGECOLLECT();
 }
 
 void UNetDriver::AddReferencedObjects(UObject* InThis, FReferenceCollector& Collector)
diff --git a/Engine/Source/Runtime/Engine/Private/RepLayout.cpp b/Engine/Source/Runtime/Engine/Private/RepLayout.cpp
--- a/Engine/Source/Runtime/Engine/Private/RepLayout.cpp
+++ b/Engine/Source/Runtime/Engine/Private/RepLayout.cpp
@@ -21,6 +21,7 @@
 #include "Net/NetworkGranularMemoryLogging.h"
 #include "Serialization/ArchiveCountMem.h"
 #include "Templates/AndOrNot.h"
+#include "PushModelPerNetDriverState.h"
 
 DECLARE_CYCLE_STAT(TEXT("RepLayout AddPropertyCmd"), STAT_RepLayout_AddPropertyCmd, STATGROUP_Game);
 DECLARE_CYCLE_STAT(TEXT("RepLayout InitFromObjectClass"), STAT_RepLayout_InitFromObjectClass, STATGROUP_Game);
@@ -60,6 +61,17 @@ static FAutoConsoleVariableRef CVarShareShadowState(TEXT("net.ShareShadowState")
 int32 GShareInitialCompareState = 0;
 static FAutoConsoleVariableRef CVarShareInitialCompareState(TEXT("net.ShareInitialCompareState"), GShareInitialCompareState, TEXT("If true and net.ShareShadowState is enabled, attempt to also share initial replication compares across connections."));
 
+#if WITH_PUSH_VALIDATION_SUPPORT
+
+static bool GbPushModelValidateProperties = false;
+static FAutoConsoleVariableRef CVarPushModelValidateProperties(TEXT("net.PushModelValidateProperties"), GbPushModelValidateProperties, TEXT("When true, we will compare all push model properties and warn if they haven't been marked dirty properly."));
+
+#else
+
+constexpr bool GbPushModelValidateProperties = false;
+
+#endif
+
 int32 MaxRepArraySize = UNetworkSettings::DefaultMaxRepArraySize;
 int32 MaxRepArrayMemory = UNetworkSettings::DefaultMaxRepArrayMemory;
 
@@ -212,6 +224,35 @@ namespace UE4_RepLayout_Private
 
 		NETWORK_PROFILER(GNetworkProfiler.TrackWritePropertyHeader(Property, HeaderBits, nullptr));
 	}
+
+#if WITH_PUSH_MODEL
+	const UE4PushModelPrivate::FPushModelPerNetDriverHandle ConditionallyAddPushModelObject(const UE4PushModelPrivate::FNetPushObjectId ObjectId, const TSharedRef<const FRepLayout>& InRepLayout)
+	{
+		const int32 NumReplicatedProperties = InRepLayout->GetNumParents();
+		if (UE4PushModelPrivate::IsPushModelEnabled() &&
+			NumReplicatedProperties > 0 &&
+			EnumHasAnyFlags(InRepLayout->GetFlags(), ERepLayoutFlags::FullPushSupport | ERepLayoutFlags::PartialPushSupport))
+		{
+			return UE4PushModelPrivate::AddPushModelObject(ObjectId, NumReplicatedProperties);
+		}
+
+		return UE4PushModelPrivate::FPushModelPerNetDriverHandle::MakeInvalidHandle();
+	}
+
+	void ConditionallyRemovePushModelObject(const UE4PushModelPrivate::FPushModelPerNetDriverHandle Handle)
+	{
+		if (Handle.IsValid())
+		{
+			UE4PushModelPrivate::RemovePushModelObject(Handle);
+		}
+	}
+
+	UE4PushModelPrivate::FPushModelPerNetDriverState* GetPerNetDriverState(const UE4PushModelPrivate::FPushModelPerNetDriverHandle Handle)
+	{
+		return Handle.IsValid() ? UE4PushModelPrivate::GetPerNetDriverState(Handle) : nullptr;
+	}
+
+#endif
 }
 
 //~ TODO: Consider moving the FastArray members into their own sub-struct to save memory for non fast array
@@ -744,6 +785,26 @@ static uint32 GetRepLayoutCmdCompatibleChecksum(
 	return CompatibleChecksum;
 }
 
+#if (WITH_PUSH_MODEL)
+struct FNetPrivatePushIdHelper
+{
+	static void SetNetPushID(UObject* InObject, const UE4PushModelPrivate::FNetPushObjectId ObjectId)
+	{
+		const UE4PushModelPrivate::FNetPushObjectId CurrentId = FObjectNetPushIdHelper::GetNetPushIdDynamic(InObject);
+		if (CurrentId != ObjectId)
+		{
+			check(CurrentId == INDEX_NONE);
+			FObjectNetPushIdHelper::SetNetPushIdDynamic(InObject, ObjectId);
+		}
+	}
+
+	static void MarkPropertyDirty(UObject* InObject, const int32 RepIndex)
+	{
+		MARK_PROPERTY_DIRTY_UNSAFE(InObject, RepIndex);
+	}
+};
+#endif // (WITH_PUSH_MODEL)
+
 FRepChangedPropertyTracker::FRepChangedPropertyTracker(const bool InbIsReplay, const bool InbIsClientReplayRecording):
 		bIsReplay(InbIsReplay),
 		bIsClientReplayRecording(InbIsClientReplayRecording),
@@ -752,14 +813,22 @@ FRepChangedPropertyTracker::FRepChangedPropertyTracker(const bool InbIsReplay, c
 
 FRepChangedPropertyTracker::~FRepChangedPropertyTracker() {}
 
-void FRepChangedPropertyTracker::SetCustomIsActiveOverride(const uint16 RepIndex, const bool bIsActive)
+void FRepChangedPropertyTracker::SetCustomIsActiveOverride(UObject* OwningObject, const uint16 RepIndex, const bool bIsActive)
 {
-	FRepChangedParent & Parent = Parents[RepIndex];
+	FRepChangedParent& Parent = Parents[RepIndex];
 
-	Parent.Active = (bIsActive || bIsClientReplayRecording) ? 1 : 0;
 	Parent.OldActive = Parent.Active;
+	Parent.Active = (bIsActive || bIsClientReplayRecording) ? 1 : 0;
+
+#if WITH_PUSH_MODEL
+	if (!Parent.OldActive && Parent.Active)
+	{
+		FNetPrivatePushIdHelper::MarkPropertyDirty(OwningObject, RepIndex);
+	}
+#endif	
 }
 
+
 void FRepChangedPropertyTracker::SetExternalData(const uint8* Src, const int32 NumBits)
 {
 	ExternalDataNumBits = NumBits;
@@ -898,15 +967,35 @@ void FRepStateStaticBuffer::CountBytes(FArchive& Ar) const
 
 FRepChangelistState::FRepChangelistState(
 	const TSharedRef<const FRepLayout>& InRepLayout,
-	const uint8* Source,
-	FCustomDeltaChangelistState* InDeltaChangelistState):
-
-	CustomDeltaChangelistState(InDeltaChangelistState),
-	HistoryStart(0),
-	HistoryEnd(0),
-	CompareIndex(0),
-	StaticBuffer(InRepLayout->CreateShadowBuffer(Source))
-{}
+	const uint8* InSource,
+	const UObject* InRepresenting,
+	FCustomDeltaChangelistState* InDeltaChangelistState)
+
+	: CustomDeltaChangelistState(InDeltaChangelistState)
+	, HistoryStart(0)
+	, HistoryEnd(0)
+	, CompareIndex(0)
+	, StaticBuffer(InRepLayout->CreateShadowBuffer(InSource))
+
+#if WITH_PUSH_MODEL
+	, PushModelObjectHandle(UE4_RepLayout_Private::ConditionallyAddPushModelObject(InRepresenting->GetUniqueID(), InRepLayout))
+#endif // WITH_PUSH_MODEL
+
+{
+#if (WITH_PUSH_MODEL)
+	if (PushModelObjectHandle.IsValid())
+	{
+		FNetPrivatePushIdHelper::SetNetPushID(const_cast<UObject*>(InRepresenting), PushModelObjectHandle.ObjectId);
+	}
+#endif // (WITH_PUSH_MODEL)
+}
+
+FRepChangelistState::~FRepChangelistState()
+{
+#if WITH_PUSH_MODEL
+	UE4_RepLayout_Private::ConditionallyRemovePushModelObject(PushModelObjectHandle);
+#endif // WITH_PUSH_MODEL
+}
 
 void FRepChangelistState::CountBytes(FArchive& Ar) const
 {
@@ -925,23 +1014,33 @@ void FRepChangelistState::CountBytes(FArchive& Ar) const
 
 FReplicationChangelistMgr::FReplicationChangelistMgr(
 	const TSharedRef<const FRepLayout>& InRepLayout,
-	const uint8* Source,
-	FCustomDeltaChangelistState* DeltaChangelistState):
+	const uint8* InSource,
+	const UObject* InRepresenting,
+	FCustomDeltaChangelistState* DeltaChangelistState)
+
+	: LastReplicationFrame(0)
+	, LastInitialReplicationFrame(0)
+	, RepChangelistState(InRepLayout, InSource, InRepresenting, DeltaChangelistState)
+{
+}
 
-	LastReplicationFrame(0),
-	LastInitialReplicationFrame(0),
-	RepChangelistState(InRepLayout, Source, DeltaChangelistState)
+FReplicationChangelistMgr::~FReplicationChangelistMgr()
 {
 }
 
-FReceivingRepState::FReceivingRepState(FRepStateStaticBuffer&& InStaticBuffer) :
-	StaticBuffer(MoveTemp(InStaticBuffer))
+void FReplicationChangelistMgr::CountBytes(FArchive& Ar) const
+{
+	RepChangelistState.CountBytes(Ar);
+}
+
+FReceivingRepState::FReceivingRepState(FRepStateStaticBuffer&& InStaticBuffer) 
+	: StaticBuffer(MoveTemp(InStaticBuffer))
 {
 }
 
-FRepLayout::FRepLayout() :
-	Flags(ERepLayoutFlags::None),
-	Owner(nullptr)
+FRepLayout::FRepLayout()
+	: Flags(ERepLayoutFlags::None)
+	, Owner(nullptr)
 {}
 
 FRepLayout::~FRepLayout()
@@ -1001,11 +1100,6 @@ void FRepLayout::UpdateChangelistMgr(
 	}
 }
 
-void FReplicationChangelistMgr::CountBytes(FArchive& Ar) const
-{
-	RepChangelistState.CountBytes(Ar);
-}
-
 struct FComparePropertiesSharedParams
 {
 	const bool bIsInitial;
@@ -1013,117 +1107,180 @@ struct FComparePropertiesSharedParams
 	const ERepLayoutFlags Flags;
 	const TArray<FRepParentCmd>& Parents;
 	const TArray<FRepLayoutCmd>& Cmds;
+	FSendingRepState* const RepState;
+	FRepChangelistState* const RepChangelistState;
+	FRepChangedPropertyTracker* const RepChangedPropertyTracker;
+};
+
+struct FComparePropertiesStackParams
+{
+	const FConstRepObjectDataBuffer Data;
+	FRepShadowDataBuffer ShadowData;
+	TArray<uint16>& Changed;
 };
 
 static uint16 CompareProperties_r(
 	const FComparePropertiesSharedParams& SharedParams,
+	FComparePropertiesStackParams& StackParams,
 	const uint16 CmdStart,
 	const uint16 CmdEnd,
-	FRepShadowDataBuffer ShadowData,
-	const FConstRepObjectDataBuffer Data,
-	TArray<uint16>& Changed,
 	uint16 Handle);
 
 static void CompareProperties_Array_r(
 	const FComparePropertiesSharedParams& SharedParams,
-	FRepShadowDataBuffer ShadowData,
-	const FConstRepObjectDataBuffer Data,
-	TArray<uint16>& Changed,
+	FComparePropertiesStackParams& StackParams,
 	const uint16 CmdIndex,
 	const uint16 Handle);
 
 static void CompareRoleProperty(
 	const FComparePropertiesSharedParams& SharedParams,
-	const FConstRepObjectDataBuffer Data,
+	FComparePropertiesStackParams& StackParams,
 	const uint16 RoleOrRemoteRoleIndex,
-	TEnumAsByte<ENetRole>& SavedRoleOrRemoteRole,
-	TArray<uint16>& Changed)
+	TEnumAsByte<ENetRole>& SavedRoleOrRemoteRole)
 {
 	const FRepParentCmd& RoleOrRemoteRoleParent = SharedParams.Parents[RoleOrRemoteRoleIndex];
 	const FRepLayoutCmd& RoleOrRemoteRoleCmd = SharedParams.Cmds[RoleOrRemoteRoleParent.CmdStart];
 	const uint16 Handle = RoleOrRemoteRoleCmd.RelativeHandle;
-	const TEnumAsByte<ENetRole> ActorRoleOrRemoteRole = *(const TEnumAsByte<ENetRole>*)(Data + RoleOrRemoteRoleParent).Data;
+	const TEnumAsByte<ENetRole> ActorRoleOrRemoteRole = *(const TEnumAsByte<ENetRole>*)(StackParams.Data + RoleOrRemoteRoleParent).Data;
 	if (SharedParams.bForceFail || SavedRoleOrRemoteRole != ActorRoleOrRemoteRole)
 	{
 		SavedRoleOrRemoteRole = ActorRoleOrRemoteRole;
-		Changed.Add(Handle);
-		}
+		StackParams.Changed.Add(Handle);
+	}
 }
 
 static void CompareRoleProperties(
 	const FComparePropertiesSharedParams& SharedParams,
-	FSendingRepState* RESTRICT RepState,
-	const FConstRepObjectDataBuffer Data,
-	TArray<uint16>& Changed)
+	FComparePropertiesStackParams& StackParams)
 {
-	if (RepState && EnumHasAnyFlags(SharedParams.Flags, ERepLayoutFlags::IsActor))
-		{
-		CompareRoleProperty(SharedParams, Data, (int32)AActor::ENetFields_Private::RemoteRole, RepState->SavedRemoteRole, Changed);
-		CompareRoleProperty(SharedParams, Data, (int32)AActor::ENetFields_Private::Role, RepState->SavedRole, Changed);
+	if (SharedParams.RepState && EnumHasAnyFlags(SharedParams.Flags, ERepLayoutFlags::IsActor))
+	{
+		CompareRoleProperty(SharedParams, StackParams, (int32)AActor::ENetFields_Private::RemoteRole, SharedParams.RepState->SavedRemoteRole);
+		CompareRoleProperty(SharedParams, StackParams, (int32)AActor::ENetFields_Private::Role, SharedParams.RepState->SavedRole);
 	}
 }
 
-static void CompareParentProperties(
+static void CompareParentProperty(
+	const int32 ParentIndex,
 	const FComparePropertiesSharedParams& SharedParams,
-	FSendingRepState* RESTRICT RepState,
-	FRepChangelistState* RESTRICT RepChangelistState,
-	const FConstRepObjectDataBuffer Data,
-	TArray<uint16>& Changed)
+	FComparePropertiesStackParams& StackParams)
 {
-	FRepShadowDataBuffer ShadowData = RepChangelistState->StaticBuffer.GetData();
-	check(ShadowData);
+	const FRepParentCmd& Parent = SharedParams.Parents[ParentIndex];
+	const bool bIsLifetime = EnumHasAnyFlags(Parent.Flags, ERepParentFlags::IsLifetime);
 
-	FRepChangedPropertyTracker* RepChangedPropertyTracker = (RepState ? RepState->RepChangedPropertyTracker.Get() : nullptr);
+	// Active state of a property applies to *all* connections.
+	// If the property is inactive, we can skip comparing it because we know it won't be sent.
+	// Further, this will keep the last active state of the property in the shadow buffer,
+	// meaning the next time the property becomes active it will be sent to all connections.
+	const bool bIsActive = !SharedParams.RepChangedPropertyTracker || SharedParams.RepChangedPropertyTracker->Parents[ParentIndex].Active;
 
-	const bool bCheckForRole = EnumHasAnyFlags(SharedParams.Flags, ERepLayoutFlags::IsActor) && RepState != nullptr;
+	const bool bShouldSkip = !bIsLifetime || !bIsActive || (Parent.Condition == COND_InitialOnly && !SharedParams.bIsInitial);
 
-	for (uint16 ParentIndex = 0; ParentIndex < SharedParams.Parents.Num(); ++ParentIndex)
+	if (bShouldSkip)
 	{
-		const FRepParentCmd& Parent = SharedParams.Parents[ParentIndex];
+		return;
+	}
+
+	const FRepLayoutCmd& Cmd = SharedParams.Cmds[Parent.CmdStart];
+
+	if (EnumHasAnyFlags(SharedParams.Flags, ERepLayoutFlags::IsActor))
+	{
+		if (UNLIKELY(ParentIndex == (int32)AActor::ENetFields_Private::Role))
+		{
+			CompareRoleProperty(SharedParams, StackParams, (int32)AActor::ENetFields_Private::Role, SharedParams.RepState->SavedRole);
+			return;
+		}
+		if (UNLIKELY(ParentIndex == (int32)AActor::ENetFields_Private::RemoteRole))
+		{
+			CompareRoleProperty(SharedParams, StackParams, (int32)AActor::ENetFields_Private::RemoteRole, SharedParams.RepState->SavedRemoteRole);
+			return;
+		}
+	}
 
-		const bool bIsLifetime = EnumHasAnyFlags(Parent.Flags, ERepParentFlags::IsLifetime);
+	// Note, Handle - 1 to account for CompareProperties_r incrementing handles.
+	CompareProperties_r(SharedParams, StackParams, Parent.CmdStart, Parent.CmdEnd, Cmd.RelativeHandle - 1);
+}
 
-		// Active state of a property applies to *all* connections.
-		// If the property is inactive, we can skip comparing it because we know it won't be sent.
-		// Further, this will keep the last active state of the property in the shadow buffer,
-		// meaning the next time the property becomes active it will be sent to all connections.
-		const bool bIsActive = !RepChangedPropertyTracker || RepChangedPropertyTracker->Parents[ParentIndex].Active;
+static void CompareParentProperties(
+	const FComparePropertiesSharedParams& SharedParams,
+	FComparePropertiesStackParams& StackParams)
+{
+	check(StackParams.ShadowData);
 
-		const bool bShouldSkip = !bIsLifetime || !bIsActive || (Parent.Condition == COND_InitialOnly && !SharedParams.bIsInitial);
+#if WITH_PUSH_MODEL
 
-		if (bShouldSkip)
+	if (UE4PushModelPrivate::FPushModelPerNetDriverState* PushModelObject = UE4_RepLayout_Private::GetPerNetDriverState(SharedParams.RepChangelistState->GetPushModelObjectHandle()))
+	{
+		// Typically, on an initial compare all properties will be dirty anyway.
+		// However, if an Actor is awakened from Dormancy, it's SendingRepState will have been recreated, invalidating
+		// its saved Role and RemoteRole (regardless of whether or not they were changed).
+		// This really shouldn't cause any problems, but it will cause Validation to consistently fail (if GbPushModelValidateProperties
+		// is enabled).
+		// To get around this, and just to be sure our code paths are consistent, we forcibly mark Role and RemoteRole dirty here.
+		if (EnumHasAnyFlags(SharedParams.Flags, ERepLayoutFlags::IsActor) && SharedParams.bIsInitial)
 		{
-			continue;
+			PushModelObject->MarkPropertyDirty((int32)AActor::ENetFields_Private::RemoteRole);
+			PushModelObject->MarkPropertyDirty((int32)AActor::ENetFields_Private::Role);
 		}
 
-		if (bCheckForRole)
+		// If we're forcibly comparing all properties, then don't bother checking dirty state.
+		if (SharedParams.bForceFail)
 		{
-			if (UNLIKELY(ParentIndex == (int32)AActor::ENetFields_Private::Role))
+			for (int32 ParentIndex = 0; ParentIndex < SharedParams.Parents.Num(); ++ParentIndex)
 			{
-				CompareRoleProperty(SharedParams, Data, (int32)AActor::ENetFields_Private::Role, RepState->SavedRole, Changed);
-				continue;
+				CompareParentProperty(ParentIndex, SharedParams, StackParams);
 			}
-			else if (UNLIKELY(ParentIndex == (int32)AActor::ENetFields_Private::RemoteRole))
+		}
+
+		// If we have full push model support, then we only need to check properties that are actually dirty.
+		else if (!GbPushModelValidateProperties && EnumHasAnyFlags(SharedParams.Flags, ERepLayoutFlags::FullPushSupport))
 		{
-				CompareRoleProperty(SharedParams, Data, (int32)AActor::ENetFields_Private::RemoteRole, RepState->SavedRemoteRole, Changed);
-				continue;
+			for (TConstSetBitIterator<> It = PushModelObject->GetDirtyProperties(); It; ++It)
+			{
+				CompareParentProperty(It.GetIndex(), SharedParams, StackParams);
 			}
 		}
-		
-			// Note, Handle - 1 to account for CompareProperties_r incrementing handles.
-		const FRepLayoutCmd& Cmd = SharedParams.Cmds[Parent.CmdStart];
-		const uint16 Handle = Cmd.RelativeHandle;
-			CompareProperties_r(SharedParams, Parent.CmdStart, Parent.CmdEnd, ShadowData, Data, Changed, Handle - 1);
+		else
+		{
+#if WITH_PUSH_VALIDATION_SUPPORT
+			int32 StartChangedNum = 0;
+#endif
+
+			for (int32 ParentIndex = 0; ParentIndex < SharedParams.Parents.Num(); ++ParentIndex)
+			{
+				const bool bIsPropertyDirty = !EnumHasAnyFlags(SharedParams.Parents[ParentIndex].Flags, ERepParentFlags::UsePushModel) || PushModelObject->IsPropertyDirty(ParentIndex);
+				if (GbPushModelValidateProperties || bIsPropertyDirty)
+				{
+					CompareParentProperty(ParentIndex, SharedParams, StackParams);
+
+#if WITH_PUSH_VALIDATION_SUPPORT
+					const bool bWasFoundDirtyButNotMarkedDirty = GbPushModelValidateProperties && !bIsPropertyDirty && StackParams.Changed.Num() > StartChangedNum;
+
+					ensureAlwaysMsgf(!bWasFoundDirtyButNotMarkedDirty, TEXT("Push Model Property changed value, but was not marked dirty! Property=%s"), *SharedParams.Parents[ParentIndex].Property->GetPathName());
+
+					StartChangedNum = StackParams.Changed.Num();
+#endif
+				}
+			}
 		}
+
+		PushModelObject->ResetDirtyStates();
+		return;
+	}
+#endif
+
+	for (int32 ParentIndex = 0; ParentIndex < SharedParams.Parents.Num(); ++ParentIndex)
+	{
+		CompareParentProperty(ParentIndex, SharedParams, StackParams);
+	}
 }
 
+
 static uint16 CompareProperties_r(
 	const FComparePropertiesSharedParams& SharedParams,
+	FComparePropertiesStackParams& StackParams,
 	const uint16 CmdStart,
 	const uint16 CmdEnd,
-	FRepShadowDataBuffer ShadowData,
-	const FConstRepObjectDataBuffer Data,
-	TArray<uint16>& Changed,
 	uint16 Handle)
 {
 	for (int32 CmdIndex = CmdStart; CmdIndex < CmdEnd; ++CmdIndex)
@@ -1134,17 +1291,26 @@ static uint16 CompareProperties_r(
 
 		++Handle;
 
+		const FConstRepObjectDataBuffer Data = StackParams.Data + Cmd;
+		FRepShadowDataBuffer ShadowData = StackParams.ShadowData + Cmd;
+
 		if (Cmd.Type == ERepLayoutCmdType::DynamicArray)
 		{
+			FComparePropertiesStackParams NewStackParams{
+				Data,
+				ShadowData,
+				StackParams.Changed
+			};
+
 			// Once we hit an array, start using a stack based approach
-			CompareProperties_Array_r(SharedParams, ShadowData + Cmd, Data + Cmd, Changed, CmdIndex, Handle);
+			CompareProperties_Array_r(SharedParams, NewStackParams, CmdIndex, Handle);
 			CmdIndex = Cmd.EndCmd - 1;		// The -1 to handle the ++ in the for loop
 			continue;
 		}
-		else if (SharedParams.bForceFail || !PropertiesAreIdentical(Cmd, (ShadowData + Cmd).Data, (Data + Cmd).Data))
+		else if (SharedParams.bForceFail || !PropertiesAreIdentical(Cmd, ShadowData.Data, Data.Data))
 		{
-			StoreProperty(Cmd, (ShadowData + Cmd).Data, (Data + Cmd.Offset).Data);
-			Changed.Add(Handle);
+			StoreProperty(Cmd, ShadowData.Data, Data.Data);
+			StackParams.Changed.Add(Handle);
 		}
 	}
 
@@ -1153,16 +1319,14 @@ static uint16 CompareProperties_r(
 
 static void CompareProperties_Array_r(
 	const FComparePropertiesSharedParams& SharedParams,
-	FRepShadowDataBuffer ShadowData,
-	const FConstRepObjectDataBuffer Data,
-	TArray<uint16>& Changed,
+	FComparePropertiesStackParams& StackParams,
 	const uint16 CmdIndex,
 	const uint16 Handle)
 {
 	const FRepLayoutCmd& Cmd = SharedParams.Cmds[CmdIndex];
 
-	FScriptArray* ShadowArray = (FScriptArray*)ShadowData.Data;
-	FScriptArray* Array = (FScriptArray*)Data.Data;
+	FScriptArray* ShadowArray = (FScriptArray*)StackParams.ShadowData.Data;
+	FScriptArray* Array = (FScriptArray*)StackParams.Data.Data;
 
 	const uint16 ArrayNum = Array->Num();
 	const uint16 ShadowArrayNum = ShadowArray->Num();
@@ -1187,16 +1351,23 @@ static void CompareProperties_Array_r(
 		{
 			const int32 ArrayElementOffset = i * Cmd.ElementSize;
 			bForceFail = bOldForceFail || i >= ShadowArrayNum;
-			LocalHandle = CompareProperties_r(SharedParams, CmdIndex + 1, Cmd.EndCmd - 1, ShadowArrayData + ArrayElementOffset, ArrayData + ArrayElementOffset, ChangedLocal, LocalHandle);
+
+			FComparePropertiesStackParams NewStackParams{
+				ArrayData + ArrayElementOffset,
+				ShadowArrayData + ArrayElementOffset,
+				ChangedLocal
+			};
+
+			LocalHandle = CompareProperties_r(SharedParams, NewStackParams, CmdIndex + 1, Cmd.EndCmd - 1, LocalHandle);
 		}
 	}
 
 	if (ChangedLocal.Num())
 	{
-		Changed.Add(Handle);
-		Changed.Add(ChangedLocal.Num());		// This is so we can jump over the array if we need to
-		Changed.Append(ChangedLocal);
-		Changed.Add(0);
+		StackParams.Changed.Add(Handle);
+		StackParams.Changed.Add(ChangedLocal.Num());		// This is so we can jump over the array if we need to
+		StackParams.Changed.Append(ChangedLocal);
+		StackParams.Changed.Add(0);
 	}
 	else if (ArrayNum != ShadowArrayNum)
 	{
@@ -1204,9 +1375,9 @@ static void CompareProperties_Array_r(
 		check(ArrayNum < ShadowArrayNum || SharedParams.Cmds[CmdIndex + 1].Type == ERepLayoutCmdType::DynamicArray);
 
 		// Array got smaller, send the array handle to force array size change
-		Changed.Add(Handle);
-		Changed.Add(0);
-		Changed.Add(0);
+		StackParams.Changed.Add(Handle);
+		StackParams.Changed.Add(0);
+		StackParams.Changed.Add(0);
 	}
 }
 
@@ -1238,16 +1409,25 @@ bool FRepLayout::CompareProperties(
 		/*bForceFail=*/ false,
 		Flags,
 		Parents,
-		Cmds
+		Cmds,
+		RepState,
+		RepChangelistState,
+		(RepState ? RepState->RepChangedPropertyTracker.Get() : nullptr)
+	};
+
+	FComparePropertiesStackParams StackParams{
+		Data,
+		RepChangelistState->StaticBuffer.GetData(),
+		Changed
 	};
 
 	if (RepFlags.bRolesOnly)
 	{
-		CompareRoleProperties(SharedParams, RepState, Data, Changed);
+		CompareRoleProperties(SharedParams, StackParams);
 	}
 	else
 	{
-		CompareParentProperties(SharedParams, RepState, RepChangelistState, Data, Changed);
+		CompareParentProperties(SharedParams, StackParams);
 	}
 
 	if (Changed.Num() == 0)
@@ -2759,7 +2939,7 @@ static FGuidReferencesMap* PrepReceivedArray(
 	return NewGuidReferencesArray ? NewGuidReferencesArray->Array : nullptr;
 }
 
-/** Struct containing parameters that don't change throughout recursion of ReceiveProperties_r */
+/** Struct containing parameters that don't change or can be shared throughout recursion of ReceiveProperties_r */
 struct FReceivePropertiesSharedParams
 {
 	const bool bDoChecksum;
@@ -2770,10 +2950,15 @@ struct FReceivePropertiesSharedParams
 	bool& bOutGuidsChanged;
 	const TArray<FRepParentCmd>& Parents;
 	const TArray<FRepLayoutCmd>& Cmds;
+	UObject* OwningObject;
 	uint16 ReadHandle = 0;
+	
+#if WITH_PUSH_MODEL	
+	int32 LastReceivedParent = INDEX_NONE;
+#endif	
 };
 
-/** Struct containing parameters that do change as we recurse into ReceiveProperties_r */
+/** Struct containing parameters that do change or can't be shared as we recurse into ReceiveProperties_r */
 struct FReceivePropertiesStackParams
 {
 	FRepObjectDataBuffer ObjectData;
@@ -2831,11 +3016,22 @@ static bool ReceiveProperties_r(FReceivePropertiesSharedParams& Params, FReceive
 		else
 		{
 			UE_LOG(LogRepProperties, VeryVerbose, TEXT("ReceiveProperties_r: Parent=%d, Cmd=%d"), Cmd.ParentIndex, CmdIndex);
-				
-			if (ERepLayoutCmdType::DynamicArray == Cmd.Type)
+
+			const FRepParentCmd& Parent = Params.Parents[Cmd.ParentIndex];
+
+#if WITH_PUSH_MODEL
+			if (Cmd.ParentIndex != Params.LastReceivedParent)
 			{
-				const FRepParentCmd& Parent = Params.Parents[Cmd.ParentIndex];
+				// Make sure we mark received properties dirty.
+				// This will make sure that anything we receive can be picked up by other Net Drivers
+				// (like the DemoNetDriver) on clients.
+				Params.LastReceivedParent = Cmd.ParentIndex;
+				FNetPrivatePushIdHelper::MarkPropertyDirty(Params.OwningObject, Params.LastReceivedParent);
+			}
+#endif
 
+			if (ERepLayoutCmdType::DynamicArray == Cmd.Type)
+			{
 				// Don't worry about checking the ShadowData for nullptr here.
 				// We're either:
 				//	1. At the top level and it's valid
@@ -2937,7 +3133,6 @@ static bool ReceiveProperties_r(FReceivePropertiesSharedParams& Params, FReceive
 			// TODO: Might be worth doing this before and after ReadNextHandle, or having ReadNextHandle check for errors?
 			if (Params.Bunch.IsError())
 			{
-				const FRepParentCmd& Parent = Params.Parents[Cmd.ParentIndex];
 				UE_LOG(LogRep, Error, TEXT("ReceiveProperties_r: Failed to receive property, BunchIsError - Property=%s, Parent=%d, Cmd=%d, ReadHandle=%d"), *Parent.CachedPropertyName.ToString(), Cmd.ParentIndex, CmdIndex, Params.ReadHandle);
 				return false;
 			}
@@ -2955,7 +3150,7 @@ bool FRepLayout::ReceiveProperties(
 	UActorChannel* OwningChannel,
 	UClass* InObjectClass,
 	FReceivingRepState* RESTRICT RepState,
-	FRepObjectDataBuffer Data,
+	UObject* Object,
 	FNetBitReader& InBunch,
 	bool& bOutHasUnmapped,
 	bool& bOutGuidsChanged,
@@ -2963,6 +3158,7 @@ bool FRepLayout::ReceiveProperties(
 {
 	check(InObjectClass == Owner);
 
+	FRepObjectDataBuffer Data(Object);
 	const bool bEnableRepNotifies = EnumHasAnyFlags(ReceiveFlags, EReceivePropertiesFlags::RepNotifies);
 
 	if (OwningChannel->Connection->InternalAck)
@@ -2989,12 +3185,13 @@ bool FRepLayout::ReceiveProperties(
 		bDoChecksum,
 		// We can skip swapping roles if we're not an Actor layout, or if we've been explicitly told we can skip.
 		EnumHasAnyFlags(ReceiveFlags, EReceivePropertiesFlags::SkipRoleSwap) || !EnumHasAnyFlags(Flags, ERepLayoutFlags::IsActor),
-		EnumHasAnyFlags(ReceiveFlags, EReceivePropertiesFlags::RepNotifies) ? &RepState->RepNotifies : nullptr,
+		bEnableRepNotifies ? &RepState->RepNotifies : nullptr,
 		InBunch,
 		bOutHasUnmapped,
 		bOutGuidsChanged,
 		Parents,
-		Cmds
+		Cmds,
+		Object
 	};
 
 	FReceivePropertiesStackParams StackParams{
@@ -4775,7 +4972,10 @@ static int32 InitFromProperty_r(
 	return RelativeHandle;
 }
 
-static FORCEINLINE uint16 AddParentProperty(TArray<FRepParentCmd>& Parents, FProperty* Property, int32 ArrayIndex)
+static FORCEINLINE uint16 AddParentProperty(
+	TArray<FRepParentCmd>& Parents,
+	FProperty* Property,
+	int32 ArrayIndex)
 {
 	return Parents.Emplace(Property, ArrayIndex);
 }
@@ -4905,7 +5105,11 @@ static void BuildShadowOffsets_r(TArray<FRepLayoutCmd>::TIterator& CmdIt, int32&
 }
 
 template<ERepBuildType ShadowType>
-static void BuildShadowOffsets(UStruct* Owner, TArray<FRepParentCmd>& Parents, TArray<FRepLayoutCmd>& Cmds, int32& ShadowOffset)
+static void BuildShadowOffsets(
+	UStruct* Owner,
+	TArray<FRepParentCmd>& Parents,
+	TArray<FRepLayoutCmd>& Cmds,
+	int32& ShadowOffset)
 {
 	SCOPE_CYCLE_COUNTER(STAT_RepLayout_BuildShadowOffsets);
 
@@ -5119,6 +5323,11 @@ void FRepLayout::InitFromClass(
 
 	Object->GetLifetimeReplicatedProps(LifetimeProps);
 
+	// Tracks the number of (non-delta) lifetime properties so we can check that against our
+	// Push Model Enabled properties.
+	int32 NumberOfLifetimeProperties = 0;
+	int32 NumberOfPushModelProperties = 0;
+
 	// Setup lifetime replicated properties
 	for (int32 i = 0; i < LifetimeProps.Num(); i++)
 	{
@@ -5146,7 +5355,22 @@ void FRepLayout::InitFromClass(
 			Parents[ParentIndex].RepNotifyNumParams = RepNotifyFunc->NumParms;
 		}
 
-		if (EnumHasAnyFlags(Parents[ParentIndex].Flags, ERepParentFlags::IsCustomDelta))
+		if (!EnumHasAnyFlags(Parents[ParentIndex].Flags, ERepParentFlags::IsCustomDelta))
+		{
+			Parents[ParentIndex].Flags |= ERepParentFlags::IsLifetime;
+			if (LifetimeProps[i].Condition == COND_None)
+			{
+				Parents[ParentIndex].Flags &= ~ERepParentFlags::IsConditional;
+			}
+
+			++NumberOfLifetimeProperties;
+			if (IS_PUSH_MODEL_ENABLED() && LifetimeProps[i].bIsPushBased)
+			{
+				++NumberOfPushModelProperties;
+				Parents[ParentIndex].Flags |= ERepParentFlags::UsePushModel;
+			}
+		}
+		else
 		{
 			// We'll track all Custom Lifetime Properties here, and we'll handle Fast Array Serialization
 			// specially.
@@ -5229,16 +5453,6 @@ void FRepLayout::InitFromClass(
 			{
 				LifetimeCustomPropertyState->Add(FLifetimeCustomDeltaProperty(ParentIndex));
 			}
-
-			continue;
-		}
-
-		Parents[ParentIndex].Flags |= ERepParentFlags::IsLifetime;
-
-		if (LifetimeProps[i].Condition == COND_None)
-		{
-			Parents[ParentIndex].Flags &= ~ERepParentFlags::IsConditional;
-		}
 	}
 
 	if (bIsObjectActor)
@@ -5256,6 +5470,13 @@ void FRepLayout::InitFromClass(
 			Parents[(int32)AActor::ENetFields_Private::RemoteRole].Flags |= ERepParentFlags::IsConditional;
 			Parents[(int32)AActor::ENetFields_Private::RemoteRole].Condition = COND_None;
 		}
+	}	
+
+	if (IS_PUSH_MODEL_ENABLED() && NumberOfPushModelProperties > 0)
+	{
+		Flags |= (NumberOfLifetimeProperties == NumberOfPushModelProperties) ?
+			ERepLayoutFlags::FullPushSupport :
+			ERepLayoutFlags::PartialPushSupport;
 	}
 
 	if (!ServerConnection || EnumHasAnyFlags(CreateFlags, ECreateRepLayoutFlags::MaySendProperties))
@@ -6025,7 +6246,7 @@ TSharedPtr<FReplicationChangelistMgr> FRepLayout::CreateReplicationChangelistMgr
 		ShadowStateSource = (const uint8*)InObject;
 	}
 
-	return MakeShareable(new FReplicationChangelistMgr(AsShared(), ShadowStateSource, DeltaChangelistState));
+	return MakeShareable(new FReplicationChangelistMgr(AsShared(), ShadowStateSource, InObject, DeltaChangelistState));
 }
 
 TUniquePtr<FRepState> FRepLayout::CreateRepState(
@@ -6596,6 +6817,9 @@ bool FRepLayout::DeltaSerializeFastArrayProperty(FFastArrayDeltaSerializeParams&
 						TArray<uint16> NewChangelist;
 						for (auto& IDIndexPair : ChangedElements)
 						{
+							NewChangelist.Empty(1);
+							const int32 ArrayElementOffset = ElementSize * IDIndexPair.Idx;
+
 							// Go ahead and do a property compare here, regardless of what we'll actually use below.
 							// This is to prevent issues where someone with an initial / outdated connection doesn't properly
 							// update the changelists in our history, but does update the shadow state inadvertently.
@@ -6604,16 +6828,19 @@ bool FRepLayout::DeltaSerializeFastArrayProperty(FFastArrayDeltaSerializeParams&
 								/*bForceFail=*/ bIsInitial || ShadowArrayItemIsNew[IDIndexPair.Idx],
 								Flags,
 								Parents,
-								Cmds
+								Cmds,
+								nullptr,
+								nullptr,
+								nullptr
 							};
 
-							const int32 ArrayElementOffset = ElementSize * IDIndexPair.Idx;
-
-							FConstRepObjectDataBuffer ElementData(ObjectArrayData + ArrayElementOffset);
-							FRepShadowDataBuffer ElementShadowData(ShadowArrayData + ArrayElementOffset);
-							NewChangelist.Empty(1);
+							FComparePropertiesStackParams StackParams{
+								FConstRepObjectDataBuffer(ObjectArrayData + ArrayElementOffset),
+								FRepShadowDataBuffer(ShadowArrayData + ArrayElementOffset),
+								NewChangelist
+							};
 
-							CompareProperties_r(SharedParams, ItemLayoutStart, ItemLayoutEnd, ElementShadowData, ElementData, NewChangelist, 0);
+							CompareProperties_r(SharedParams, StackParams, ItemLayoutStart, ItemLayoutEnd, 0);
 
 							if (NewChangelist.Num())
 							{
@@ -6624,7 +6851,7 @@ bool FRepLayout::DeltaSerializeFastArrayProperty(FFastArrayDeltaSerializeParams&
 								// to copying them into the shadow state (see FFastArraySerializerItem::operator=).
 								// In that case, we need to make make sure we reset our ID so they can be found the next
 								// time we try to replicate them.
-								int32& ShadowReplicationID = CustomDeltaProperty.GetFastArrayItemReplicationIDMutable(ElementShadowData.Data);
+								int32& ShadowReplicationID = CustomDeltaProperty.GetFastArrayItemReplicationIDMutable(StackParams.ShadowData.Data);
 								ShadowReplicationID = IDIndexPair.ID;
 							}
 						}
@@ -6861,7 +7088,8 @@ bool FRepLayout::DeltaSerializeFastArrayProperty(FFastArrayDeltaSerializeParams&
 					bOutHasUnmapped,
 					bOutGuidsChanged,
 					Parents,
-					Cmds
+					Cmds,
+					Object
 				};
 
 				FReceivePropertiesStackParams StackParams{
diff --git a/Engine/Source/Runtime/Engine/Private/UnrealNetwork.cpp b/Engine/Source/Runtime/Engine/Private/UnrealNetwork.cpp
--- a/Engine/Source/Runtime/Engine/Private/UnrealNetwork.cpp
+++ b/Engine/Source/Runtime/Engine/Private/UnrealNetwork.cpp
@@ -34,7 +34,7 @@ void RegisterReplicatedLifetimeProperty(
 		const uint16 RepIndex = PropertyDescriptor.RepIndex + i;
 		FLifetimeProperty* RegisteredPropertyPtr = OutLifetimeProps.FindByPredicate([&RepIndex](const FLifetimeProperty& Var) { return Var.RepIndex == RepIndex; });
 
-		FLifetimeProperty LifetimeProp(RepIndex, Params.Condition, Params.RepNotifyCondition);
+		FLifetimeProperty LifetimeProp(RepIndex, Params.Condition, Params.RepNotifyCondition, Params.bIsPushBased);
 
 		if (RegisteredPropertyPtr)
 		{
diff --git a/Engine/Source/Runtime/Engine/Public/Net/NetPushModelHelpers.h b/Engine/Source/Runtime/Engine/Public/Net/NetPushModelHelpers.h
new file mode 100644
--- /dev/null
+++ b/Engine/Source/Runtime/Engine/Public/Net/NetPushModelHelpers.h
@@ -0,0 +1,22 @@
+// Copyright 1998-2019 Epic Games, Inc. All Rights Reserved.
+
+#pragma once
+
+#include "Kismet/BlueprintFunctionLibrary.h"
+#include "NetPushModelHelpers.generated.h"
+
+UCLASS()
+class ENGINE_API UNetPushModelHelpers : public UBlueprintFunctionLibrary
+{
+	GENERATED_BODY()
+
+public:
+
+	UFUNCTION(BlueprintCallable, Category="Net")
+	static void MarkPropertyDirty(UObject* Object, FName PropertyName);
+
+private:
+
+	UFUNCTION(BlueprintCallable, Category = "Net", Meta=(BlueprintInternalUseOnly = "true", HidePin = "Object|RepIndex|PropertyName"))
+	static void MarkPropertyDirtyFromRepIndex(UObject* Object, int32 RepIndex, FName PropertyName);
+};
\ No newline at end of file
diff --git a/Engine/Source/Runtime/Engine/Public/Net/RepLayout.h b/Engine/Source/Runtime/Engine/Public/Net/RepLayout.h
--- a/Engine/Source/Runtime/Engine/Public/Net/RepLayout.h
+++ b/Engine/Source/Runtime/Engine/Public/Net/RepLayout.h
@@ -21,6 +21,7 @@
 #include "UObject/GCObject.h"
 #include "Containers/StaticBitArray.h"
 #include "Net/GuidReferences.h"
+#include "Net/Core/PushModel/PushModel.h"
 #include "Templates/CopyQualifiersFromTo.h"
 
 class FGuidReferences;
@@ -158,10 +159,14 @@ public:
 	 *
 	 * @see DOREPLIFETIME_ACTIVE_OVERRIDE
 	 *
-	 * @param RepIndex	Replication index for the Property.
-	 * @param bIsActive	The new Active state.
+	 * @param OwningObject	The object that we're tracking.
+	 * @param RepIndex		Replication index for the Property.
+	 * @param bIsActive		The new Active state.
 	 */
-	virtual void SetCustomIsActiveOverride(const uint16 RepIndex, const bool bIsActive) override;
+	virtual void SetCustomIsActiveOverride(
+		UObject* OwningObject,
+		const uint16 RepIndex,
+		const bool bIsActive) override;
 
 	/**
 	 * Sets (or resets) the External Data.
@@ -400,10 +405,13 @@ private:
 	FRepChangelistState(
 		const TSharedRef<const FRepLayout>& InRepLayout,
 		const uint8* Source,
-		struct FCustomDeltaChangelistState* CustomDeltaChangelistState);
+		const UObject* InRepresenting,
+		struct FCustomDeltaChangelistState* InDeltaChangelistState);
 
 public:
 
+	~FRepChangelistState();
+
 	/** The maximum number of individual changelists allowed.*/
 	static const int32 MAX_CHANGE_HISTORY = 64;
 
@@ -429,6 +437,17 @@ public:
 	FRepSerializationSharedInfo SharedSerialization;
 
 	void CountBytes(FArchive& Ar) const;
+
+#if WITH_PUSH_MODEL
+
+	const UE4PushModelPrivate::FPushModelPerNetDriverHandle& GetPushModelObjectHandle() const
+	{
+		return PushModelObjectHandle;
+	}
+
+private:
+	const UE4PushModelPrivate::FPushModelPerNetDriverHandle PushModelObjectHandle;
+#endif
 };
 
 /**
@@ -447,10 +466,13 @@ private:
 	FReplicationChangelistMgr(
 		const TSharedRef<const FRepLayout>& InRepLayout,
 		const uint8* Source,
+		const UObject* InRepresenting,
 		struct FCustomDeltaChangelistState* CustomDeltaChangelistState);
 
 public:
 
+	~FReplicationChangelistMgr();
+
 	FRepChangelistState* GetRepChangelistState() const
 	{
 		return const_cast<FRepChangelistState*>(&RepChangelistState);
@@ -694,6 +716,7 @@ enum class ERepParentFlags : uint32
 	IsZeroConstructible	= (1 << 6),	//! This property is ZeroConstructible.
 	IsFastArray			= (1 << 7), //! This property is a FastArraySerializer. This can't be a ERepLayoutCmdType, because
 									//! these Custom Delta structs will have their inner properties tracked.
+	UsePushModel		= (1 << 8),	//! This object relies on PushModel, and can be skipped when not dirty.
 };
 
 ENUM_CLASS_FLAGS(ERepParentFlags)
@@ -712,8 +735,8 @@ public:
 		CachedPropertyName(InProperty ? InProperty->GetFName() : NAME_None),
 		ArrayIndex(InArrayIndex),
 		ShadowOffset(0),
-		CmdStart(0), 
-		CmdEnd(0), 
+		CmdStart(0),
+		CmdEnd(0),
 		Condition(COND_None),
 		RepNotifyCondition(REPNOTIFY_OnChanged),
 		RepNotifyNumParams(INDEX_NONE),
@@ -985,8 +1008,10 @@ ENUM_CLASS_FLAGS(ECreateRepLayoutFlags);
 
 enum class ERepLayoutFlags : uint8
 {
-	None = 0,
-	IsActor = 1 << 1,	//! This RepLayout is for AActor or a subclass of AActor.
+	None				= 0,
+	IsActor 			= 1 << 0,	//! This RepLayout is for AActor or a subclass of AActor.
+	PartialPushSupport	= 1 << 1,	//! This RepLayout has some properties that use Push Model and some that don't.
+	FullPushSupport		= 1 << 2,	//! All properties in this RepLayout use Push Model.
 };
 ENUM_CLASS_FLAGS(ERepLayoutFlags);
 
@@ -1199,11 +1224,11 @@ public:
 		UActorChannel* OwningChannel,
 		UClass* InObjectClass,
 		FReceivingRepState* RESTRICT RepState,
-		FRepObjectDataBuffer Data,
+		UObject* Object,
 		FNetBitReader& InBunch,
 		bool& bOutHasUnmapped,
 		bool& bOutGuidsChanged,
-		const EReceivePropertiesFlags Flags) const;
+		const EReceivePropertiesFlags InFlags) const;
 
 	/**
 	 * Finds any properties in the Shadow Buffer of the given Rep State that are currently valid
@@ -1436,6 +1461,11 @@ PRAGMA_ENABLE_DEPRECATION_WARNINGS
 		return 0 == Parents.Num();
 	}
 
+	const int32 GetNumParents() const
+	{
+		return Parents.Num();
+	}
+
 	void CountBytes(FArchive& Ar) const;
 
 private:
diff --git a/Engine/Source/Runtime/Engine/Public/Net/UnrealNetwork.h b/Engine/Source/Runtime/Engine/Public/Net/UnrealNetwork.h
--- a/Engine/Source/Runtime/Engine/Public/Net/UnrealNetwork.h
+++ b/Engine/Source/Runtime/Engine/Public/Net/UnrealNetwork.h
@@ -73,6 +73,9 @@ struct ENGINE_API FDoRepLifetimeParams
 	 * properly set up to handle RepNotifies.
 	 */
 	ELifetimeRepNotifyCondition RepNotifyCondition = REPNOTIFY_OnChanged;
+	
+	/** Whether or not this property uses Push Model. See PushModel.h */
+	bool bIsPushBased = false;
 };
 
 namespace NetworkingPrivate
@@ -171,7 +174,15 @@ static FProperty* GetReplicatedProperty(const UClass* CallingClass, const UClass
 { \
 	static const bool bIsValid_##c_##v = ValidateReplicatedClassInheritance(StaticClass(), c::StaticClass(), TEXT(#v)); \
 	const TCHAR* DoRepPropertyName_##c_##v(TEXT(#v)); \
-	const NetworkingPrivate::FRepPropertyDescriptor PropertyDescriptor_##c_##v(DoRepPropertyName_##c_##v, NetworkingPrivate::Net_##c::NETFIELD_##v, NetworkingPrivate::Net_##c::ARRAYDIM_##v); \
+	const NetworkingPrivate::FRepPropertyDescriptor PropertyDescriptor_##c_##v(DoRepPropertyName_##c_##v, (int32)c::ENetFields_Private::v, 1); \
+	RegisterReplicatedLifetimeProperty(PropertyDescriptor_##c_##v, OutLifetimeProps, params); \
+}
+
+#define DOREPLIFETIME_WITH_PARAMS_FAST_STATIC_ARRAY(c,v,params) \
+{ \
+	static const bool bIsValid_##c_##v = ValidateReplicatedClassInheritance(StaticClass(), c::StaticClass(), TEXT(#v)); \
+	const TCHAR* DoRepPropertyName_##c_##v(TEXT(#v)); \
+	const NetworkingPrivate::FRepPropertyDescriptor PropertyDescriptor_##c_##v(DoRepPropertyName_##c_##v, (int32)c::ENetFields_Private::v##_STATIC_ARRAY, (int32)c::EArrayDims_Private::v); \
 	RegisterReplicatedLifetimeProperty(PropertyDescriptor_##c_##v, OutLifetimeProps, params); \
 }
 
@@ -219,19 +230,25 @@ static FProperty* GetReplicatedProperty(const UClass* CallingClass, const UClass
 #define DOREPLIFETIME_ACTIVE_OVERRIDE_FAST(c,v,active) \
 { \
 	static const bool bIsValid_##c_##v = ValidateReplicatedClassInheritance(StaticClass(), c::StaticClass(), TEXT(#v)); \
-	for (int32 i = 0; i < NetworkingPrivate::Net_##c::ARRAYDIM_##v; i++) \
+	ChangedPropertyTracker.SetCustomIsActiveOverride((int32)c::ENetFields_Private::v, active); \
+}
+
+#define DOREPLIFETIME_ACTIVE_OVERRIDE_FAST_STATIC_ARRAY(c,v,active) \
+{ \
+	static const bool bIsValid_##c_##v = ValidateReplicatedClassInheritance(StaticClass(), c::StaticClass(), TEXT(#v)); \
+	for (int32 i = 0; i < (int32)c::EArrayDims_Private::v; ++i) \
 	{ \
-		ChangedPropertyTracker.SetCustomIsActiveOverride(NetworkingPrivate::Net_##c::NETFIELD_##v + i, active); \
+		ChangedPropertyTracker.SetCustomIsActiveOverride((int32)c::ENetFields_Private::v##_STATIC_ARRAY + i, active); \
 	} \
 }
 
-#define DOREPLIFETIME_ACTIVE_OVERRIDE(c,v,active)	\
-{													\
+#define DOREPLIFETIME_ACTIVE_OVERRIDE(c,v,active) \
+{ \
 	static FProperty* sp##v = GetReplicatedProperty(StaticClass(), c::StaticClass(),GET_MEMBER_NAME_CHECKED(c,v)); \
-	for ( int32 i = 0; i < sp##v->ArrayDim; i++ )											\
-	{																						\
-		ChangedPropertyTracker.SetCustomIsActiveOverride( sp##v->RepIndex + i, active );	\
-	}																						\
+	for (int32 i = 0; i < sp##v->ArrayDim; i++) \
+	{ \
+		ChangedPropertyTracker.SetCustomIsActiveOverride(this, sp##v->RepIndex + i, active); \
+	} \
 }
 
 UE_DEPRECATED(4.24, "Please use the RESET_REPLIFETIME_CONDITION macro")
@@ -277,7 +294,15 @@ DisableReplicatedLifetimeProperty(StaticClass(), c::StaticClass(), GET_MEMBER_NA
 { \
 	static const bool bIsValid_##c_##v = ValidateReplicatedClassInheritance(StaticClass(), c::StaticClass(), TEXT(#v)); \
 	const TCHAR* DoRepPropertyName_##c_##v(TEXT(#v)); \
-	const NetworkingPrivate::FRepPropertyDescriptor PropertyDescriptor_##c_##v(DoRepPropertyName_##c_##v, NetworkingPrivate::Net_##c::NETFIELD_##v, NetworkingPrivate::Net_##c::ARRAYDIM_##v); \
+	const NetworkingPrivate::FRepPropertyDescriptor PropertyDescriptor_##c_##v(DoRepPropertyName_##c_##v, (int32)c::ENetFields_Private::v, 1); \
+	DisableReplicatedLifetimeProperty(PropertyDescriptor_##c_##v, OutLifetimeProps); \
+}
+
+#define DISABLE_REPLICATED_PROPERTY_FAST_STATIC_ARRAY(c,v) \
+{ \
+	static const bool bIsValid_##c_##v = ValidateReplicatedClassInheritance(StaticClass(), c::StaticClass(), TEXT(#v)); \
+	const TCHAR* DoRepPropertyName_##c_##v(TEXT(#v)); \
+	const NetworkingPrivate::FRepPropertyDescriptor PropertyDescriptor_##c_##v(DoRepPropertyName_##c_##v, (int32)c::ENetFields_Private::v##_STATIC_ARRAY, (int32)c::EArrayDims_Private::v); \
 	DisableReplicatedLifetimeProperty(PropertyDescriptor_##c_##v, OutLifetimeProps); \
 }
 
@@ -291,7 +316,7 @@ DisableAllReplicatedPropertiesOfClass(StaticClass(), c::StaticClass(), SuperClas
 { \
 	static const bool bIsValid_##c_##v = ValidateReplicatedClassInheritance(StaticClass(), c::StaticClass(), TEXT("DISABLE_ALL_CLASS_REPLICATED_PROPERTIES")); \
 	const TCHAR* DoRepPropertyName_##c(TEXT(#c)); \
-	const NetworkingPrivate::FRepClassDescriptor ClassDescriptor_##c(DoRepPropertyName_##c, NetworkingPrivate::Net_##c::NETFIELD_REP_START, NetworkingPrivate::Net_##c::NETFIELD_REP_END); \
+	const NetworkingPrivate::FRepClassDescriptor ClassDescriptor_##c(DoRepPropertyName_##c, (int32)c::ENetFields_Private::NETFIELD_REP_START, (int32)c::ENetFields_Private::NETFIELD_REP_END); \
 	DisableAllReplicatedPropertiesOfClass(StaticClass(), c::StaticClass(), SuperClassBehavior, OutLifetimeProps); \
 }	
 
@@ -314,11 +339,20 @@ ENGINE_API void DisableAllReplicatedPropertiesOfClass(const NetworkingPrivate::F
 { \
 	static const bool bIsValid_##c_##v = ValidateReplicatedClassInheritance(StaticClass(), c::StaticClass(), TEXT(#v)); \
 	const TCHAR* DoRepPropertyName_##c_##v(TEXT(#v)); \
-	const NetworkingPrivate::FRepPropertyDescriptor PropertyDescriptor_##c_##v(DoRepPropertyName_##c_##v, NetworkingPrivate::Net_##c::NETFIELD_##v, NetworkingPrivate::Net_##c::ARRAYDIM_##v); \
+	const NetworkingPrivate::FRepPropertyDescriptor PropertyDescriptor_##c_##v(DoRepPropertyName_##c_##v, (int32)c::ENetFields_Private::v, 1); \
+	ResetReplicatedLifetimeProperty(StaticClass(), c::StaticClass(), GET_MEMBER_NAME_CHECKED(c,v), cond, OutLifetimeProps); \
+}
+
+#define RESET_REPLIFETIME_CONDITION_FAST_STATIC_ARRAY(c,v,cond) \
+{ \
+	static const bool bIsValid_##c_##v = ValidateReplicatedClassInheritance(StaticClass(), c::StaticClass(), TEXT(#v)); \
+	const TCHAR* DoRepPropertyName_##c_##v(TEXT(#v)); \
+	const NetworkingPrivate::FRepPropertyDescriptor PropertyDescriptor_##c_##v(DoRepPropertyName_##c_##v, (int32)c::ENetFields_Private::v##_STATIC_ARRAY, (int32)c::EArrayDims_Private::v); \
 	ResetReplicatedLifetimeProperty(StaticClass(), c::StaticClass(), GET_MEMBER_NAME_CHECKED(c,v), cond, OutLifetimeProps); \
 }
 
 #define RESET_REPLIFETIME_FAST(c,v) RESET_REPLIFETIME_CONDITION_FAST(c, v, COND_None)
+#define RESET_REPLIFETIME_FAST_STATIC_ARRAY(c,v) RESET_REPLIFETIME_FAST_STATIC_ARRAY(c, v, COND_None)
 
 ENGINE_API void ResetReplicatedLifetimeProperty(
 	const UClass* ThisClass,
diff --git a/Engine/Source/Runtime/Net/Core/Private/Net/Core/PushModel/PushModel.cpp b/Engine/Source/Runtime/Net/Core/Private/Net/Core/PushModel/PushModel.cpp
new file mode 100644
--- /dev/null
+++ b/Engine/Source/Runtime/Net/Core/Private/Net/Core/PushModel/PushModel.cpp
@@ -0,0 +1,196 @@
+// Copyright 1998-2019 Epic Games, Inc. All Rights Reserved.
+
+#include "Net/Core/PushModel/PushModel.h"
+
+#if WITH_PUSH_MODEL
+
+#include "HAL/IConsoleManager.h"
+#include "Algo/Sort.h"
+
+#include "Types/PushModelUtils.h"
+#include "Types/PushModelPerObjectState.h"
+#include "Types/PushModelPerNetDriverState.h"
+
+
+namespace UE4PushModelPrivate
+{
+	class FPushModelObjectManager_CustomId
+	{
+	private:
+
+		using ThisClass = FPushModelObjectManager_CustomId;
+
+	public:
+
+		FPushModelObjectManager_CustomId()
+		{
+		}
+
+		~FPushModelObjectManager_CustomId()
+		{
+		}
+
+		void MarkPropertyDirty(const FNetPushObjectId ObjectId, const int32 RepIndex)
+		{
+			// The macros will take care of filtering out invalid objects, so we don't need to check here.
+			PerObjectStates[ObjectId].MarkPropertyDirty(RepIndex);
+		}
+
+		void MarkPropertyDirty(const FNetPushObjectId ObjectId, const int32 StartRepIndex, const int32 EndRepIndex)
+		{
+			FPushModelPerObjectState& ObjectState = PerObjectStates[ObjectId];
+			for (int RepIndex = StartRepIndex; RepIndex <= EndRepIndex; ++RepIndex)
+			{
+				ObjectState.MarkPropertyDirty(RepIndex);
+			}
+		}
+
+		const FPushModelPerNetDriverHandle AddNetworkObject(const FNetPushObjectId ObjectId, const uint16 NumReplicatedProperties)
+		{
+			FNetPushObjectId& InternalPushId = ObjectIdToInternalId.FindOrAdd(ObjectId, INDEX_NONE);
+			if (INDEX_NONE == InternalPushId)
+			{
+				FSparseArrayAllocationInfo AllocationInfo = PerObjectStates.AddUninitializedAtLowestFreeIndex(NewObjectLookupPosition);
+				new (AllocationInfo.Pointer) FPushModelPerObjectState(ObjectId, NumReplicatedProperties);
+				InternalPushId = AllocationInfo.Index;
+			}
+
+			FPushModelPerObjectState& PerObjectState = PerObjectStates[InternalPushId];
+			check(PerObjectState.GetNumberOfProperties() == NumReplicatedProperties);
+			check(PerObjectState.GetObjectId() == ObjectId);
+
+			const FNetPushPerNetDriverId NetDriverId = PerObjectState.AddPerNetDriverState();
+			return FPushModelPerNetDriverHandle(NetDriverId, InternalPushId);
+		}
+
+		void RemoveNetworkObject(const FPushModelPerNetDriverHandle Handle)
+		{
+			if (PerObjectStates.IsValidIndex(Handle.ObjectId))
+			{
+				FPushModelPerObjectState& PerObjectState = PerObjectStates[Handle.ObjectId];
+				PerObjectState.RemovePerNetDriverState(Handle.NetDriverId);
+				if (!PerObjectState.HasAnyNetDriverStates())
+				{
+					ObjectIdToInternalId.Remove(PerObjectState.GetObjectId());
+					PerObjectStates.RemoveAt(Handle.ObjectId);
+
+					if (NewObjectLookupPosition > Handle.ObjectId)
+					{
+						NewObjectLookupPosition = Handle.ObjectId;
+					}
+				}
+			}
+		}
+
+		void PreReplication()
+		{
+			for (FPushModelPerObjectState& PerObjectState : PerObjectStates)
+			{
+				PerObjectState.PushDirtyStateToNetDrivers();
+			}
+		}
+
+		void PostGarbageCollect()
+		{
+			// We can't compact PerObjectStates because we need ObjectIDs to be stable.
+			// But we can shrink it.
+
+			PerObjectStates.Shrink();
+			ObjectIdToInternalId.Compact();
+			NewObjectLookupPosition = 0;
+		}
+
+		FPushModelPerNetDriverState* GetPerNetDriverState(const FPushModelPerNetDriverHandle Handle)
+		{
+			if (PerObjectStates.IsValidIndex(Handle.ObjectId))
+			{
+				FPushModelPerObjectState& ObjectState = PerObjectStates[Handle.ObjectId];
+				ObjectState.PushDirtyStateToNetDrivers();
+				return &ObjectState.GetPerNetDriverState(Handle.NetDriverId);
+			}
+
+			return nullptr;
+		}
+
+	private:
+
+		int32 NewObjectLookupPosition = 0;
+		TMap<FNetPushObjectId, FNetPushObjectId> ObjectIdToInternalId;
+		TSparseArray<FPushModelPerObjectState> PerObjectStates;
+	};
+
+	static FPushModelObjectManager_CustomId PushObjectManager;
+
+	bool bIsPushModelEnabled = false;
+	FAutoConsoleVariableRef CVarIsPushModelEnabled(
+		TEXT("Net.IsPushModelEnabled"),
+		bIsPushModelEnabled,
+		TEXT("Whether or not Push Model is enabled. This networking mode allows game code to notify the networking system of changes, rather than scraping.")
+	);
+	
+	bool bMakeBpPropertiesPushModel = true;
+	FAutoConsoleVariableRef CVarMakeBpPropertiesPushModel(
+		TEXT("Net.MakeBpPropertiesPushModel"),
+		bMakeBpPropertiesPushModel,
+		TEXT("Whether or not Blueprint Properties will be forced to used Push Model")
+	);
+
+	void MarkPropertyDirty(const FNetPushObjectId ObjectId, const int32 RepIndex)
+	{
+		PushObjectManager.MarkPropertyDirty(ObjectId, RepIndex);
+	}
+
+	void MarkPropertyDirty(const FNetPushObjectId ObjectId, const int32 StartRepIndex, const int32 EndRepIndex)
+	{
+		PushObjectManager.MarkPropertyDirty(ObjectId, StartRepIndex, EndRepIndex);
+	}
+
+	void PostGarbageCollect()
+	{
+		PushObjectManager.PostGarbageCollect();
+	}
+
+	/**
+	 * Called by a given NetDriver to notify us that it's seen a given Object for the first time (or the first time
+	 * since it was removed).
+	 *
+	 * This may be called multiple times for a given Object if there are multiple NetDrivers, but it's expected
+	 * that each NetDriver only calls this once per object before RemoteNetworkObject is called.
+	 *
+	 * @param ObjectId						The UniqueId for the object.
+	 * @param NumberOfReplicatedProperties	The number of replicated properties for this object.
+	 *
+	 * @return A Handle that can be used in other calls to uniquely identify this object per NetDriver.
+	 */
+	const FPushModelPerNetDriverHandle AddPushModelObject(const FNetPushObjectId ObjectId, const uint16 NumberOfReplicatedProperties)
+	{
+		return PushObjectManager.AddNetworkObject(ObjectId, NumberOfReplicatedProperties);
+	}
+
+	/**
+	 * Called by a given NetDriver to notify us that a given Object is no longer valid for Networking.
+	 *
+	 * This may be called multiple times for a given Object if there are multiple NetDrivers, but it's expected
+	 * that each NetDriver only calls this once per object after AddNetworkObject is called, and never before
+	 * AddNetworkObject is called.
+	 *
+	 * @param Handle	The Push Model Object handle (returned by AddPushModelObject).
+	 */
+	void RemovePushModelObject(const FPushModelPerNetDriverHandle Handle)
+	{
+		PushObjectManager.RemoveNetworkObject(Handle);
+	}
+
+	/**
+	 * Gets the NetDriver specific state for a given Push Model Object.
+	 * Note, calling this will flush dirty state to all NetDriver states for the Object.
+	 *
+	 * @param Handle	The Push Model Object handle (returned by AddPushModelObject).
+	 */
+	FPushModelPerNetDriverState* GetPerNetDriverState(const FPushModelPerNetDriverHandle Handle)
+	{
+		return PushObjectManager.GetPerNetDriverState(Handle);
+	}
+}
+
+#endif // WITH_PUSH_MODEL
\ No newline at end of file
diff --git a/Engine/Source/Runtime/Net/Core/Private/Net/Core/PushModel/Types/PushModelPerNetDriverState.h b/Engine/Source/Runtime/Net/Core/Private/Net/Core/PushModel/Types/PushModelPerNetDriverState.h
new file mode 100644
--- /dev/null
+++ b/Engine/Source/Runtime/Net/Core/Private/Net/Core/PushModel/Types/PushModelPerNetDriverState.h
@@ -0,0 +1,72 @@
+// Copyright 1998-2019 Epic Games, Inc. All Rights Reserved.
+
+#pragma once
+
+#include "Net/Core/PushModel/PushModelMacros.h"
+
+#if WITH_PUSH_MODEL
+
+#include "CoreMinimal.h"
+#include "Containers/BitArray.h"
+#include "PushModelUtils.h"
+
+namespace UE4PushModelPrivate
+{
+	class FPushModelPerNetDriverState
+	{
+	public:
+
+		FPushModelPerNetDriverState(const uint16 InNumberOfProperties)
+			: PropertyDirtyStates(true, InNumberOfProperties)
+		{
+		}
+		
+		FPushModelPerNetDriverState(FPushModelPerNetDriverState&& Other)
+			: PropertyDirtyStates(MoveTemp(Other.PropertyDirtyStates))
+		{
+		}
+		
+		FPushModelPerNetDriverState(const FPushModelPerNetDriverState& Other) = delete;
+		FPushModelPerNetDriverState& operator=(const FPushModelPerNetDriverState& Other) = delete;
+
+		void ResetDirtyStates()
+		{
+			ResetBitArray(PropertyDirtyStates);
+		}
+
+		void CountBytes(FArchive& Ar) const
+		{
+			PropertyDirtyStates.CountBytes(Ar);
+		}
+
+		const bool IsPropertyDirty(const uint16 RepIndex)
+		{
+			return PropertyDirtyStates[RepIndex];
+		}
+
+		TConstSetBitIterator<> GetDirtyProperties() const
+		{
+			return TConstSetBitIterator<>(PropertyDirtyStates);
+		}
+
+		void MarkPropertiesDirty(const TBitArray<>& OtherBitArray)
+		{
+			BitwiseOrBitArrays(OtherBitArray, PropertyDirtyStates);
+		}
+
+		void MarkPropertyDirty(const uint16 RepIndex)
+		{
+			PropertyDirtyStates[RepIndex] = true;
+		}
+
+	private:
+
+		/**
+		 * Current state of our push model properties.
+		 * Note, bits will be allocated for all replicated properties, not just push model properties.
+		 */
+		TBitArray<> PropertyDirtyStates;
+	};
+}
+
+#endif
\ No newline at end of file
diff --git a/Engine/Source/Runtime/Net/Core/Private/Net/Core/PushModel/Types/PushModelPerObjectState.h b/Engine/Source/Runtime/Net/Core/Private/Net/Core/PushModel/Types/PushModelPerObjectState.h
new file mode 100644
--- /dev/null
+++ b/Engine/Source/Runtime/Net/Core/Private/Net/Core/PushModel/Types/PushModelPerObjectState.h
@@ -0,0 +1,127 @@
+// Copyright 1998-2019 Epic Games, Inc. All Rights Reserved.
+
+#pragma once
+
+#include "Net/Core/PushModel/PushModelMacros.h"
+
+#if WITH_PUSH_MODEL
+
+#include "CoreMinimal.h"
+#include "PushModelPerNetDriverState.h"
+#include "PushModelUtils.h"
+#include "Net/Core/PushModel/PushModel.h"
+#include "Containers/SparseArray.h"
+
+namespace UE4PushModelPrivate
+{
+	/**
+	 * This is a "state" for a given Object that is being tracked by a Push Model Object Manager.
+	 * This state is shared across all NetDrivers, and so has a 1:1 mapping with actual UObjects.
+	 */
+	class FPushModelPerObjectState
+	{
+	public:
+
+		/**
+		 * Creates a new FPushModelPerObjectState.
+		 *
+		 * @param InObjectId			The ID that we'll use to refer to this object. Should be unique across all Push Model Objects (but doesn't necessarily have to be UObject::GetUniqueId()).
+		 * @param InNumberOfProperties	The total number of replicated properties this object has.
+		 */
+		FPushModelPerObjectState(const FNetPushObjectId InObjectId, const uint16 InNumberOfProperties)
+			: ObjectId(InObjectId)
+			, DirtiedThisFrame(true, InNumberOfProperties)
+		{
+		}
+
+		FPushModelPerObjectState(FPushModelPerObjectState&& Other)
+			: ObjectId(Other.ObjectId)
+			, DirtiedThisFrame(MoveTemp(Other.DirtiedThisFrame))
+			, PerNetDriverStates(MoveTemp(Other.PerNetDriverStates))
+		{
+		}
+
+		FPushModelPerObjectState(const FPushModelPerObjectState& Other) = delete;
+		FPushModelPerObjectState& operator=(const FPushModelPerObjectState& Other) = delete;
+
+		void MarkPropertyDirty(const uint16 RepIndex)
+		{
+			DirtiedThisFrame[RepIndex] = true;
+		}
+		
+		/**
+		 * Pushes the current dirty state of the Push Model Object to each of the Net Driver States.
+		 * and then reset the dirty state.
+		 */
+		void PushDirtyStateToNetDrivers()
+		{
+			if (AreAnyBitsSet(DirtiedThisFrame))
+			{
+				for (FPushModelPerNetDriverState& NetDriverObject : PerNetDriverStates)
+				{
+					NetDriverObject.MarkPropertiesDirty(DirtiedThisFrame);
+				}
+
+				ResetBitArray(DirtiedThisFrame);
+			}
+		}
+
+		FPushModelPerNetDriverState& GetPerNetDriverState(FNetPushPerNetDriverId DriverId)
+		{
+			return PerNetDriverStates[DriverId];
+		}
+		
+		FNetPushPerNetDriverId AddPerNetDriverState()
+		{
+			FSparseArrayAllocationInfo AllocationInfo = PerNetDriverStates.AddUninitialized();
+			new (AllocationInfo.Pointer) FPushModelPerNetDriverState(DirtiedThisFrame.Num());
+			
+			return AllocationInfo.Index;
+		}
+		
+		void RemovePerNetDriverState(const FNetPushPerNetDriverId DriverId)
+		{
+			PerNetDriverStates.RemoveAt(DriverId);
+		}
+		
+		const bool HasAnyNetDriverStates()
+		{
+			return PerNetDriverStates.Num() != 0;
+		}
+		
+		void CountBytes(FArchive& Ar) const
+		{
+			DirtiedThisFrame.CountBytes(Ar);
+			PerNetDriverStates.CountBytes(Ar);
+			for (const FPushModelPerNetDriverState& PerNetDriverState : PerNetDriverStates)
+			{
+				PerNetDriverState.CountBytes(Ar);
+			}
+		}
+		
+		const int32 GetNumberOfProperties() const
+		{
+			return DirtiedThisFrame.Num();
+		}
+
+		const FNetPushObjectId GetObjectId() const
+		{
+			return ObjectId;
+		}
+
+	private:
+	
+		//! A unique ID for the object.
+		const FNetPushObjectId ObjectId;
+
+		//! Bitfield tracking which properties we've dirtied since the last time
+		//! our state was pushed to NetDrivers.
+		//! Note, bits will be allocated for all replicated properties, not just push model properties.
+		TBitArray<> DirtiedThisFrame;
+
+		//! Set of NetDriver states that have been requested and are currently tracking the object.
+		TSparseArray<FPushModelPerNetDriverState> PerNetDriverStates;
+	};
+}
+
+#endif
\ No newline at end of file
diff --git a/Engine/Source/Runtime/Net/Core/Private/Net/Core/PushModel/Types/PushModelUtils.h b/Engine/Source/Runtime/Net/Core/Private/Net/Core/PushModel/Types/PushModelUtils.h
new file mode 100644
--- /dev/null
+++ b/Engine/Source/Runtime/Net/Core/Private/Net/Core/PushModel/Types/PushModelUtils.h
@@ -0,0 +1,60 @@
+// Copyright 1998-2019 Epic Games, Inc. All Rights Reserved.
+
+#pragma once
+
+#include "Net/Core/PushModel/PushModelMacros.h"
+
+#if WITH_PUSH_MODEL
+
+#include "CoreMinimal.h"
+#include "Containers/BitArray.h"
+
+namespace UE4PushModelPrivate
+{
+	static const int32 GetNumWordsInBitArray(const TBitArray<>& BitArray)
+	{
+		return (BitArray.Num() + 31) >> 5;
+	}
+	
+	static void ResetBitArray(TBitArray<>& ToReset)
+	{
+		FMemory::Memzero(ToReset.GetData(), GetNumWordsInBitArray(ToReset) * 4);
+	}
+	
+	static void SetBitArray(TBitArray<>& ToSet)
+	{
+		FMemory::Memset(ToSet.GetData(), 0xFF, GetNumWordsInBitArray(ToSet) * 4);
+	}
+	
+	static void BitwiseOrBitArrays(const TBitArray<>& MaskBitArray, TBitArray<>& ResultBitArray)
+	{
+		check(MaskBitArray.Num() == ResultBitArray.Num());
+
+		const uint32* InWords = MaskBitArray.GetData();
+		uint32* OutWords = ResultBitArray.GetData();
+		const int32 NumWords = GetNumWordsInBitArray(MaskBitArray);
+
+		for (int32 i = 0; i < NumWords; ++i)
+		{
+			OutWords[i] |= InWords[i];
+		}
+	}
+
+	static bool AreAnyBitsSet(const TBitArray<>& BitArray)
+	{
+		const int32 NumWords = GetNumWordsInBitArray(BitArray);
+		const uint32* const Words = BitArray.GetData();
+
+		for (int32 i = 0; i < NumWords; ++i)
+		{
+			if (Words[i] != 0)
+			{
+				return true;
+			}
+		}
+
+		return false;
+	}
+}
+
+#endif
\ No newline at end of file
diff --git a/Engine/Source/Runtime/Net/Core/Public/Net/Core/PushModel/PushModel.h b/Engine/Source/Runtime/Net/Core/Public/Net/Core/PushModel/PushModel.h
new file mode 100644
--- /dev/null
+++ b/Engine/Source/Runtime/Net/Core/Public/Net/Core/PushModel/PushModel.h
@@ -0,0 +1,418 @@
+// Copyright 1998-2019 Epic Games, Inc. All Rights Reserved.
+
+#pragma once
+
+#include "PushModelMacros.h"
+
+#if WITH_PUSH_MODEL
+
+#include "CoreMinimal.h"
+
+/**
+ * Push Model Support for networking.
+ *
+ ****************************************************************
+ * Rationale
+ ****************************************************************
+ *
+ * The intention of Push Model is to provide a fundamental trade off in the way UE4 property replication works, without changing
+ * the underlying replication machinery.
+ *
+ * Traditionally, UE4 Networking works by having coders and designers mark properties as replicated, and then having the network
+ * system periodically compare the values of those properties to their last sent state. When we detect a change in the state,
+ * then we will send that data to clients.
+ *
+ * This system works great from the perspective that if there's a change, clients are going to receive it with minimal effort from
+ * designers and coders.
+ *
+ * Where the system starts to fall down is in terms of performance.
+ *
+ * Because the engine has no way to know if a given Object is "dirty", that means there's no choice but to compare every property
+ * on every object. This sort of overhead isn't a big deal while games are relatively small, but as they start to scale up, this
+ * can become a bottleneck.
+ *
+ * To abate this, there's a bunch of machinery UE4 provides, things like:
+ *
+ *	Net Frequency
+ *		Gives devs the ability to tune how often actors are considered for replication.
+ *
+ *	Net Dormancy, Replication Keys, Replication Pausing, and Net Relevancy
+ *		Gives devs the ability to completely stop replication of Actors and Objects to all connections, or subsets of connections.
+ *
+ *	Property Conditions, Fast Array Serialization, and Rep Changed Trackers
+ *		Gives devs the ability to completely stop replication of specific Properties to all connections, or subsets of connections.
+ *
+ *	Rep Graph
+ *		Gives devs the ability to fine tune when and how Actors are replicated that is much more specific to a game or mode.
+ *
+ * The main focus of all of these systems is to reduce the number of Actors, Objects, and Properties we consider in an efficient way.
+ * The tradeoff all of these things have in common is that it forces devs to be more cognizant of *how* and *when* things need
+ * to replicate.
+ *
+ * Fast Arrays, Replication Keys, and Dormancy are great examples of this, and they have a very similar mechanic to Push Model.
+ *
+ * With Fast Arrays (see NetSerialization.h), no changes will be replicated unless devs explicitly mark a given Array Property
+ * as being dirty. Beyond that, individual elements won't be sent unless they are marked dirty. This means that the replication
+ * system only needs to check a few flags / keys, and if they are different, it doesn't need to do anything.
+ *
+ * With Replication Keys (see ActorChannel.h, KeyNeedsToReplicate / ReplicateSubobject), we go up a step and handle changes
+ * at the Subobject level. Devs can assign an arbitrary ID and Key to a specific subobject. As they make changes to the
+ * Subobject, it's their responsibility to increment or otherwise update the Key. Then, when they go to replicate the
+ * subobject (via a call to ReplicateSubobject), they can check the Key for the object to see if it's changed.
+ *
+ * With Dormancy (see ActorChannel.h, NetDriver.h, NetConnection.h), we again go up a step and handle changes at the Actor
+ * level. Devs can decide whether or not a given Actor supports Dormancy, and at what dormant level an actor starts at
+ * (i.e., whether it's considered initially "active" in the game, or "inactive" in the game). When an Actor becomes Active
+ * (woken from Dormancy), it will replicate one more time to all connections. After that happens, it will return to its
+ * dormant state. Once Dormant, it is completely ignored by the networking system until Devs explicitly wake it via a call
+ * to something like AActor::ForceNetUpdate.
+ *
+ * The other benefit to PushModel is that it works at the Object Level, meaning any dirty state can be pushed to / shared
+ * across all Net Drivers. This means that we're potentially eliminating multiple property comparisons in cases of
+ * things like Server Side Replays.
+ *
+ ****************************************************************
+ * How does it work?
+ ****************************************************************
+ *
+ * With Push Model, we provide a very limited and straightforward interface for allowing Devs to notify the Networking
+ * System when they change the value of a property. There are some hooks to help do this automatically for Blueprints,
+ * but if they are using things like Blueprint Setters, passing values around by reference, or setting values in Native
+ * C++, then it is completely up to them to call the necessary methods.
+ *
+ * Devs don't need to care *if* a given object is actually networked or is actively replicating, Push Model handles
+ * that sort of tracking internally. As such, we don't expect users to ever need to know or have a reference to
+ * any information other than A) The Object that's being dirtied and B) The Property that's being dirtied.
+ * Whenever a Property's value is changed, a call to one of the MARK_PROPERTY_DIRTY macros should be added.
+ * This will forward the request to Push Model, and it will handle the rest.
+ *
+ * The preferred way of achieving this is by putting any Replicated Properties behind Getters and Setters.
+ * That way, these changes and be made in a single or very limited subset of implementation code, and callers
+ * never have to know what's going on.
+ *
+ * Behind the scenes, Push Model does very simple book keeping, and will forward these dirty notifications on to the
+ * networking systems as necessary. The next time an object is replicated, we can simply check a few flags to determine
+ * whether or not a given property is dirty. If it is, then we can compare and replicate it normally, otherwise we
+ * can quickly skip it.
+ *
+ * One important thing to note is that at this time, Push Model only supports Top Level properties of Objects.
+ * So, if changes are made to Structs, Containers, or nested properties therein, the owning Struct or Container
+ * property must be marked dirty.
+ *
+ * This also means that we *must* compare a property when it's dirty to make sure we know exactly what's changed.
+ *
+ ****************************************************************
+ * Examples
+ ****************************************************************
+ *
+ * NOTE: Struct and Object markup excluded to prevent issues from UHT.
+ */
+
+#if 0
+struct FMyAwesomeStruct
+{
+	UPROPERTY()
+	int32 SomeReplicatedProperty;
+
+	UPROPERTY()
+	bool bSomeOtherReplicatedProperty;
+}
+
+class AMyAwesomeActor : public AActor
+{
+	GENERATED_BODY()
+
+	/** Convenience methods to make sure we always capture changes to property state. */
+	void SetMyReplicatedBool(const bool bNewValue)
+	{
+		bMyReplicatedBool = bNewValue;
+		MARK_PROPERTY_DIRTY_FROM_NAME(AMyAwesomeActor, bMyReplicatedBool, this);
+	}
+
+	/**
+	 * Calls to Set this from Blueprints will be automatically flagged, but we still need
+	 * to manually set this if we change it natively.
+	 */
+	void SetMyBlueprintProperty(const int32 NewValue)
+	{
+		MyBlueprintProperty = NewValue;
+		MARK_PROPERTY_DIRTY_FROM_NAME(AMyAwesomeActor, MyBlueprintProperty, this);
+	}
+
+	/**
+	 * Since we have an explicit Blueprint Setter, this won't automatically be handled
+	 * anymore, and we need to mark it dirty.
+	 */
+	void MySetter(const FString& NewValue)
+	{
+		MyBlueprintSetterProperty = NewValue;
+		MARK_PROPERTY_DIRTY_FROM_NAME(AMyAwesomeActor, MyBlueprintSetterProperty, this);
+	}
+
+	/** For static arrays, each element will have its own RepIndex, and so we need to pass that along. */
+	void SetStaticArrayValue(const int32 NewValue, const int32 Index)
+	{
+		MyStaticArray[Index] = NewValue;
+		MARK_PROPERTY_DIRTY_FROM_NAME_STATIC_ARRAY_INDEX(AMyAwesomeActor, MyStaticArray, Index, this);
+	}
+
+	void SetEntireStaticArray(const int32 NewArray[4])
+	{
+		MyStaticArray = NewArray;
+		MARK_PROPERTY_DIRTY_FROM_NAME_STATIC_ARRAY(AMyAwesomeActor, MyStaticArray, this);
+	}
+
+	/**
+	 * Currently, there is no support for marking individual properties within structs or dynamic array indices as dirty.
+	 * So, we just need to mark the top level Struct or Array as dirty.
+	 * This will cause us to compare the entire struct / array next replication.
+	 */
+	void UpdateMyStruct(int32 NewValue)
+	{
+		MyStruct.SomeReplicatedProperty = NewValue;
+		MARK_PROPERTY_DIRTY_FROM_NAME(AMyAwesomeActor, MyStruct, this);
+	}
+
+	void UpdateMyStruct(bool bNewValue)
+	{
+		MyStruct.bSomeOtherReplicatedProperty = bNewValue;
+		MARK_PROPERTY_DIRTY_FROM_NAME(AMyAwesomeActor, MyStruct, this);
+	}
+
+	/**
+	 *!!!!!!!!!!!!!!!!!!!!!!!!!!!!!!!!!!!!!!!!!
+	 *!!!!!!!!!!!!!! Warning !!!!!!!!!!!!!!!!!!
+	 *!!!!!!!!!!!!!!!!!!!!!!!!!!!!!!!!!!!!!!!!!
+	 *
+	 * Always avoid returning references to Push Model Properties.
+	 * When absolutely necessary, the property should be marked dirty before returning it
+	 * (just assume values will be changed if you need a reference).
+	 *
+	 * Absolutely **never** hold onto references to Push Model Properties outside of short scopes.
+	 * Holding onto references and changing values *without* marking the property dirty will cause
+	 * those changes to be skipped until it is marked dirty again.
+	 *
+	 * If there are properties that need to have references held onto by other managers (this is not
+	 * recommended in any case), then it's safest to not make them push model based.
+	 */
+	FMyAwesomeStruct& GetMyStruct_Mutable()
+	{
+		MARK_PROPERTY_DIRTY_FROM_NAME(AMyAwesomeActor, MyStruct, this);
+		return MyStruct;
+	}
+
+	virtual void GetLifetimeReplicatedProps(TArray<FLifetimeProperty>& OutLifetimeProps) const override
+	{
+		Super::GetLifetimeReplicatedProps(OutLifetimeProps);
+
+		FDoRepLifetimeParams Params;
+
+		Params.bIsPushBased = true;
+		DOREPLIFETIME_WITH_PARAMS_FAST(AMyAwesomeActor, bMyReplicatedBool, Params);
+		DOREPLIFETIME_WITH_PARAMS_FAST(AMyAwesomeActor, MyBlueprintProperty, Params);
+		DOREPLIFETIME_WITH_PARAMS_FAST(AMyAwesomeActor, MyBlueprintSetterProperty, Params);
+		DOREPLIFETIME_WITH_PARAMS_FAST(AMyAwesomeActor, MyStruct, Params);
+		DOREPLIFETIME_WITH_PARAMS_FAST(AMyAwesomeActor, MyStaticArray, Params);
+
+		Params.bIsPushBased = false;
+		DOREPLIFETIME_WITH_PARAMS_FAST(AMyAwesomeActor, NonPushModelProperty, Params);
+	}
+
+private:
+
+	/** This is a standard property that can only be set natively */
+	bool bMyReplicatedBool;
+
+	/**
+	 * Properties marked as BlueprintReadWrite will automatically be marked dirty when set is called.
+	 * This won't work if the value is passed by reference.
+	 */
+	int32 MyBlueprintProperty;
+
+	/**
+	 * Properties with custom BlueprintSetters won't automatically be marked dirty when set is called.
+	 * This means our setter will need to call it manually.
+	 */
+	FString MyBlueprintSetterProperty;
+
+	FMyAwesomeStruct MyStruct;
+
+	/**
+	 * Static Arrays are broken into separate replicated properties for each index, so we can flag individual
+	 * indices for comparison, or the entire array.
+	 */
+	int32 MyStaticArray[4];
+
+	int32 NonPushModelProperty;
+};
+#endif
+
+// DO NOT USE METHODS IN THIS NAMESPACE DIRECTLY
+// Use the Macros instead, as they respect conditional compilation.
+namespace UE4PushModelPrivate
+{
+	using FNetPushObjectId = int32;
+	using FNetPushPerNetDriverId = int32;
+
+	/** A handle that can be used to refer to a Push Model Object State owned by a specific Net Driver. */
+	struct FPushModelPerNetDriverHandle
+	{
+		FPushModelPerNetDriverHandle(const FNetPushPerNetDriverId InNetDriverId, const FNetPushObjectId InObjectId)
+			: NetDriverId(InNetDriverId)
+			, ObjectId(InObjectId)
+		{
+			// They should either both be valid, or both invalid.
+			check((INDEX_NONE == ObjectId) == (INDEX_NONE == NetDriverId));
+		}
+
+		//! Used by Push Model code to find a specific Net Driver state.
+		//! Note, this **does not** have a global mapping to a given NetDriver.
+		//! That is, two FPushModelPerNetDriver handles that identify Push Model Object's that are both tracked
+		//! by the same Net Driver may report different NetDriver IDs.
+		const FNetPushPerNetDriverId NetDriverId;
+		
+		//! The "globally unique" ID that is used to refer to the Push Model Object.
+		const FNetPushObjectId ObjectId;
+
+		const bool IsValid() const
+		{
+			return INDEX_NONE != ObjectId;
+		}
+
+		static FPushModelPerNetDriverHandle MakeInvalidHandle()
+		{
+			return FPushModelPerNetDriverHandle(INDEX_NONE, INDEX_NONE);
+		}
+	};
+
+	extern NETCORE_API bool bIsPushModelEnabled;
+	extern NETCORE_API bool bMakeBpPropertiesPushModel;
+
+	/** @return Whether or not PushModel is currently enabled. See "net.IsPushModelEnabled." */
+	static const bool IsPushModelEnabled()
+	{
+		return bIsPushModelEnabled;
+	}
+
+	/** @return Whether or not Blueprint Properties will use Push Model. See "net.MakeBpPropertiesPushModel." */
+	static const bool MakeBpPropertiesPushModel()
+	{
+		return bMakeBpPropertiesPushModel;
+	}
+
+	/**
+	 * Marks a specified property as dirty, causing it to be compared the next time the owner is considered
+	 * for Replication.
+	 *
+	 * Note, C-Style Array Properties (Static Array is the name used everywhere else here), will
+	 * have a RepIndex equal to UProperty::RepIndex + ArrayIndex.
+	 *
+	 * MARK_PROPERTY_DIRTY_FROM_NAME and MARK_PROPERTY_DIRTY_FROM_NAME_STATIC_ARRAY are preferred over
+	 * MARK_PROPERTY_DIRTY and MARK_PROPERTY_DIRTY_STATIC_ARRAY, because they offer compile time warnings
+	 * for things like invalid Class, Property, Static Array Index, etc., and they'll use compile time
+	 * constants instead of relying on UProperty::RepIndex.
+	 *
+	 * @param Object	The object we're marking dirty.
+	 * @param RepIndex	The index of the property we're marking dirty. UProperty::RepIndex.
+	 */
+	NETCORE_API void MarkPropertyDirty(const FNetPushObjectId ObjectId, const int32 RepIndex);
+
+	/**
+	 * Marks a range of properties as dirty, causing it to be compared the next time the owner is considered
+	 * for Replication.
+	 *
+	 * This is primarily used to mark entire Static Arrays 
+	 *
+	 * MARK_PROPERTY_DIRTY_FROM_NAME and MARK_PROPERTY_DIRTY_FROM_NAME_STATIC_ARRAY are preferred over
+	 * MARK_PROPERTY_DIRTY and MARK_PROPERTY_DIRTY_STATIC_ARRAY, because they offer compile time warnings
+	 * for things like invalid Class, Property, Static Array Index, etc., and they'll use compile time
+	 * constants instead of relying on UProperty::RepIndex.
+	 *
+	 * @param Object	The object we're marking dirty.
+	 * @param RepIndex	The index of the property we're marking dirty. UProperty::RepIndex.
+	 */
+	NETCORE_API void MarkPropertyDirty(const FNetPushObjectId ObjectId, const int32 StartRepIndex, const int32 EndRepIndex);
+
+	/**
+	 * This method is called by all NetDriver's before they replicated properties.
+	 * It's required in order to give PushModel a chance to do any setup work.
+	 */
+	NETCORE_API void PreReplication();
+	
+	/**
+	 * This method is called by all NetDriver's PostGarbageCollection.
+	 * It's required in order to give PushModel a chance to do any cleanup work.
+	 */
+	NETCORE_API void PostGarbageCollect();
+	
+	NETCORE_API const FPushModelPerNetDriverHandle AddPushModelObject(const FNetPushObjectId ObjectId, const uint16 NumberOfReplicatedProperties);
+	NETCORE_API void RemovePushModelObject(const FPushModelPerNetDriverHandle Handle);
+
+	NETCORE_API class FPushModelPerNetDriverState* GetPerNetDriverState(const FPushModelPerNetDriverHandle Handle);
+}
+
+
+#define CONDITIONAL_ON_PUSH_MODEL(Work) if (UE4PushModelPrivate::IsPushModelEnabled()) { Work; }
+#define CALL_PUSH_MODEL_PREREPLICATION() CONDITIONAL_ON_PUSH_MODEL(UE4PushModelPrivate::PreReplication();)
+#define CALL_PUSH_MODEL_POSTGARBAGECOLLECT() CONDITIONAL_ON_PUSH_MODEL(UE4PushModelPrivate::PostGarbageCollect();)
+#define IS_PUSH_MODEL_ENABLED() UE4PushModelPrivate::IsPushModelEnabled()
+#define PUSH_MAKE_BP_PROPERTIES_PUSH_MODEL() (UE4PushModelPrivate::IsPushModelEnabled() && UE4PushModelPrivate::MakeBpPropertiesPushModel())
+
+#define GET_PROPERTY_REP_INDEX(ClassName, PropertyName) (int32)ClassName::ENetFields_Private::PropertyName
+#define GET_PROPERTY_REP_INDEX_STATIC_ARRAY_START(ClassName, PropertyName) ((int32)GET_PROPERTY_REP_INDEX(ClassName, PropertyName)_STATIC_ARRAY)
+#define GET_PROPERTY_REP_INDEX_STATIC_ARRAY_END(ClassName, PropertyName) ((int32)GET_PROPERTY_REP_INDEX(ClassName, PropertyName)_STATIC_ARRAY_END)
+#define GET_PROPERTY_REP_INDEX_STATIC_ARRAY_INDEX(ClassName, PropertyName, ArrayIndex) ((int32)GET_PROPERTY_REP_INDEX_STATIC_ARRAY_START(ClassName, PropertyName) + ArrayIndex)
+
+#define IS_PROPERTY_REPLICATED(Property) (0 != (EPropertyFlags::CPF_Net & Property->PropertyFlags))
+
+#define CONDITIONAL_ON_OBJECT_NET_ID(Object, Work) { const UE4PushModelPrivate::FNetPushObjectId PrivatePushId = Object->GetNetPushId(); if (INDEX_NONE != PrivatePushId) { Work; } }
+#define CONDITIONAL_ON_OBJECT_NET_ID_DYNAMIC(Object, Work) { const UE4PushModelPrivate::FNetPushObjectId PrivatePushId = FObjectNetPushIdHelper::GetNetPushIdDynamic(Object); if (INDEX_NONE != PrivatePushId) { Work; } }
+#define CONDITIONAL_ON_REP_INDEX_AND_OBJECT_NET_ID(Object, Property, Work) if (IS_PROPERTY_REPLICATED(Property)) { const UE4PushModelPrivate::FNetPushObjectId PrivatePushId = FObjectNetPushIdHelper::GetNetPushIdDynamic(Object); if (PrivatePushId != INDEX_NONE) { Work; } }
+
+//~ For these macros, we won't bother checking if Push Model is enabled. Instead, we'll just check to see whether or not the Custom ID is valid.
+
+// Marks a property dirty by RepIndex without doing additional rep index validation.
+#define MARK_PROPERTY_DIRTY_UNSAFE(Object, RepIndex) CONDITIONAL_ON_OBJECT_NET_ID_DYNAMIC(Object, UE4PushModelPrivate::MarkPropertyDirty(PrivatePushId, RepIndex))
+
+// Marks a property dirty by UProperty*, validating that it's actually a replicated property.
+#define MARK_PROPERTY_DIRTY(Object, Property) CONDITIONAL_ON_REP_INDEX_AND_OBJECT_NET_ID(Object, Property, UE4PushModelPrivate::MarkPropertyDirty(PrivatePushId, Property->RepIndex))
+
+
+// Marks a static array property dirty given, the Object, UProperty*, and Index.
+#define MARK_PROPERTY_DIRTY_STATIC_ARRAY_INDEX(Object, Property, ArrayIndex) CONDITIONAL_ON_REP_INDEX_AND_OBJECT_NET_ID(Object, Property, UE4PushModelPrivate::MarkPropertyDirty(PrivatePushId, Property->RepIndex + ArrayIndex))
+
+// Marks all elements of a static array property dirty, given the Object and UProperty*
+#define MARK_PROPERTY_DIRTY_STATIC_ARRAY(Object, Property) CONDITIONAL_ON_REP_INDEX_AND_OBJECT_NET_ID(Object, Property, UE4PushModelPrivate::MarkPropertyDirty(PrivatePushId, Property->RepIndex, Property->RepIndex + Property->ArrayDim - 1))
+
+
+// Marks a property dirty, given the Class Name, Property Name, and Object. This will fail to compile if the Property or Class aren't valid.
+#define MARK_PROPERTY_DIRTY_FROM_NAME(ClassName, PropertyName, Object) CONDITIONAL_ON_OBJECT_NET_ID(Object, UE4PushModelPrivate::MarkPropertyDirty(PrivatePushId, GET_PROPERTY_REP_INDEX(ClassName, PropertyName)))
+
+// Marks a static array property dirty, given the Class Name, Property Name, Index, and Object. This will fail to compile if the Property and Class aren't valid. Callers are responsible for validating the index.
+#define MARK_PROPERTY_DIRTY_FROM_NAME_STATIC_ARRAY_INDEX(ClassName, PropertyName, ArrayIndex, Object) CONDITIONAL_ON_OBJECT_NET_ID(Object, UE4PushModelPrivate::MarkPropertyDirty(PrivatePushId, GET_PROPERTY_REP_INDEX_STATIC_ARRAY_INDEX(ClassName, PropertyName, ArrayIndex)))
+
+// Marks an entire static array property dirty, given the Class Name, Property Name, and Object. This will fail to compile if the Property or Class aren't valid.
+#define MARK_PROPERTY_DIRTY_FROM_NAME_STATIC_ARRAY(ClassName, PropertyName, Object) CONDITIONAL_ON_OBJECT_NET_ID(Object, UE4PushModelPrivate::MarkPropertyDirty(PrivatePushId, GET_PROPERTY_REP_INDEX_STATIC_ARRAY_START(ClassName, PropertyName), GET_PROPERTY_REP_INDEX_STATIC_ARRAY_END(ClassName, PropertyName))
+
+
+#else // WITH_PUSH_MODEL
+
+#define MARK_PROPERTY_DIRTY(Object, Property) 
+#define MARK_PROPERTY_DIRTY_STATIC_ARRAY_INDEX(Object, RepIndex, ArrayIndex) 
+#define MARK_PROPERTY_DIRTY_STATIC_ARRAY(Object, RepIndex, ArrayIndex) 
+
+#define MARK_PROPERTY_DIRTY_FROM_NAME(ClassName, PropertyName, Object) 
+#define MARK_PROPERTY_DIRTY_FROM_NAME_STATIC_ARRAY_INDEX(ClassName, PropertyName, ArrayIndex, Object) 
+#define MARK_PROPERTY_DIRTY_FROM_NAME_STATIC_ARRAY(ClassName, PropertyName, ArrayIndex, Object) 
+
+
+#define GET_PROPERTY_REP_INDEX(ClassName, PropertyName) INDEX_NONE
+#define GET_PROPERTY_REP_INDEX_STATIC_ARRAY(ClassName, PropertyName, ArrayIndex) INDEX_NONE
+
+#define IS_PUSH_MODEL_ENABLED() false
+#define PUSH_MAKE_BP_PROPERTIES_PUSH_MODEL() false
+
+#define CALL_PUSH_MODEL_PREREPLICATION() 
+#define CALL_PUSH_MODEL_POSTGARBAGECOLLECT() 
+
+#endif
\ No newline at end of file
diff --git a/Engine/Source/Runtime/Net/Core/Public/Net/Core/PushModel/PushModelMacros.h b/Engine/Source/Runtime/Net/Core/Public/Net/Core/PushModel/PushModelMacros.h
new file mode 100644
--- /dev/null
+++ b/Engine/Source/Runtime/Net/Core/Public/Net/Core/PushModel/PushModelMacros.h
@@ -0,0 +1,28 @@
+// Copyright 1998-2019 Epic Games, Inc. All Rights Reserved.
+
+#pragma once
+
+#ifndef WITH_PUSH_MODEL
+#define WITH_PUSH_MODEL 0
+#endif
+
+#define DISABLE_PUSH_MODEL_VALIDATION 0
+#define WITH_PUSH_VALIDATION_SUPPORT (WITH_PUSH_MODEL && !(UE_BUILD_SHIPPING || UE_BUILD_TEST || DISABLE_PUSH_MODEL_VALIDATION))
+
+#if WITH_PUSH_MODEL
+
+// This macro could be moved somewhere else and be made more generic if we want.
+
+#define REPLICATED_BASE_CLASS(ClassName) \
+private: \
+int32 NetPushId_Internal = INDEX_NONE; \
+virtual void SetNetPushIdDynamic(const int32 InNetPushId) override { NetPushId_Internal = InNetPushId; } \
+virtual int32 GetNetPushIdDynamic() const override { return GetNetPushId(); } \
+public: \
+int32 GetNetPushId() const { return NetPushId_Internal; }
+
+#else
+
+#define REPLICATED_BASE_CLASS(ClassName)
+
+#endif // WITH_PUSH_MODEL
\ No newline at end of file
