Fix crash when SkipToLevelIndex is invalid

diff --git a/Engine/Source/Runtime/Engine/Private/DemoNetDriver.cpp b/Engine/Source/Runtime/Engine/Private/DemoNetDriver.cpp
--- a/Engine/Source/Runtime/Engine/Private/DemoNetDriver.cpp
+++ b/Engine/Source/Runtime/Engine/Private/DemoNetDriver.cpp
@@ -4333,10 +4333,14 @@ void UDemoNetDriver::ReplayStreamingReady(const FStartStreamingResult& Result)
 		if (SkipToLevelIndexOption)
 		{
 			int32 Index = FCString::Atoi(SkipToLevelIndexOption);
-			if (Index < LevelNamesAndTimes.Num())
+			if (LevelNamesAndTimes.IsValidIndex(Index))
 			{
 				AddReplayTask(new FGotoTimeInSecondsTask(this, (float)LevelNamesAndTimes[Index].LevelChangeTimeInMS / 1000.0f));
 			}
+			else
+			{
+				UE_LOG(LogDemo, Warning, TEXT("ReplayStreamingReady: SkipToLevelIndex was invalid: %d"), Index);
+			}
 		}
 
 		if (CVarDemoJumpToEndOfLiveReplay.GetValueOnGameThread() != 0)
