Back out changelist 10931598. Likely cause for crash at startup

diff --git a/Engine/Source/Editor/BlueprintGraph/Classes/K2Node_VariableSet.h b/Engine/Source/Editor/BlueprintGraph/Classes/K2Node_VariableSet.h
--- a/Engine/Source/Editor/BlueprintGraph/Classes/K2Node_VariableSet.h
+++ b/Engine/Source/Editor/BlueprintGraph/Classes/K2Node_VariableSet.h
@@ -40,7 +40,6 @@ class UK2Node_VariableSet : public UK2Node_Variable
 	BLUEPRINTGRAPH_API bool HasLocalRepNotify() const;
 	BLUEPRINTGRAPH_API FName GetRepNotifyName() const;
 	BLUEPRINTGRAPH_API bool ShouldFlushDormancyOnSet() const;
-	BLUEPRINTGRAPH_API bool IsNetProperty() const;
 
 	static FText GetPropertyTooltip(FProperty const* VariableProperty);
 	static FText GetBlueprintVarTooltip(FBPVariableDescription const& VarDesc);
diff --git a/Engine/Source/Editor/BlueprintGraph/Private/CallFunctionHandler.cpp b/Engine/Source/Editor/BlueprintGraph/Private/CallFunctionHandler.cpp
--- a/Engine/Source/Editor/BlueprintGraph/Private/CallFunctionHandler.cpp
+++ b/Engine/Source/Editor/BlueprintGraph/Private/CallFunctionHandler.cpp
@@ -6,19 +6,11 @@
 #include "K2Node_Event.h"
 #include "K2Node_CallParentFunction.h"
 #include "K2Node_ExecutionSequence.h"
-#include "K2Node_Self.h"
-#include "K2Node_VariableGet.h"
 #include "Kismet2/BlueprintEditorUtils.h"
-#include "K2Node_CallFunction.h"
 
 #include "EdGraphUtilities.h"
 #include "Engine/BlueprintGeneratedClass.h"
 #include "KismetCompiler.h"
-#include "Net/Core/PushModel/PushModelMacros.h"
-
-#if WITH_PUSH_MODEL
-#include "Net/NetPushModelHelpers.h"
-#endif
 
 #define LOCTEXT_NAMESPACE "CallFunctionHandler"
 
@@ -608,229 +600,62 @@ UFunction* FKCHandler_CallFunction::FindFunction(FKismetFunctionContext& Context
 void FKCHandler_CallFunction::Transform(FKismetFunctionContext& Context, UEdGraphNode* Node)
 {
 	// Add an object reference pin for this call
-	if (UK2Node_CallFunction * CallFuncNode = CastChecked<UK2Node_CallFunction>(Node))
+
+	if (IsCalledFunctionPure(Node))
 	{
-		bool bIsPureAndNoUsedOutputs = false;
-		if (CallFuncNode->bIsPureFunc)
+		// Flag for removal if pure and there are no consumers of the outputs
+		//@TODO: This isn't recursive (and shouldn't be here), it'll just catch the last node in a line of pure junk
+		bool bAnyOutputsUsed = false;
+		for (int32 PinIndex = 0; PinIndex < Node->Pins.Num(); ++PinIndex)
 		{
-			// Flag for removal if pure and there are no consumers of the outputs
-			//@TODO: This isn't recursive (and shouldn't be here), it'll just catch the last node in a line of pure junk
-			bool bAnyOutputsUsed = false;
-			for (int32 PinIndex = 0; PinIndex < Node->Pins.Num(); ++PinIndex)
+			UEdGraphPin* Pin = Node->Pins[PinIndex];
+			if ((Pin->Direction == EGPD_Output) && (Pin->LinkedTo.Num() > 0))
 			{
-				UEdGraphPin* Pin = Node->Pins[PinIndex];
-				if ((Pin->Direction == EGPD_Output) && (Pin->LinkedTo.Num() > 0))
-				{
-					bAnyOutputsUsed = true;
-					break;
-				}
-			}
-
-			if (!bAnyOutputsUsed)
-			{
-				//@TODO: Remove this node, not just warn about it
-				bIsPureAndNoUsedOutputs = true;
+				bAnyOutputsUsed = true;
+				break;
 			}
 		}
 
-		const UEdGraphSchema_K2* K2Schema = GetDefault<UEdGraphSchema_K2>();
-
-		// Find the function, starting at the parent class
-		if (UFunction * Function = FindFunction(Context, Node))
+		if (!bAnyOutputsUsed)
 		{
-			if (Function->HasMetaData(FBlueprintMetadata::MD_Latent))
-			{
-				UEdGraphPin* OldOutPin = K2Schema->FindExecutionPin(*CallFuncNode, EGPD_Output);
-
-				if ((OldOutPin != NULL) && (OldOutPin->LinkedTo.Num() > 0))
-				{
-					// Create a dummy execution sequence that will be the target of the return call from the latent action
-					UK2Node_ExecutionSequence* DummyNode = CompilerContext.SpawnIntermediateNode<UK2Node_ExecutionSequence>(CallFuncNode);
-					DummyNode->AllocateDefaultPins();
+			//@TODO: Remove this node, not just warn about it
 
-					// Wire in the dummy node
-					UEdGraphPin* NewInPin = K2Schema->FindExecutionPin(*DummyNode, EGPD_Input);
-					UEdGraphPin* NewOutPin = K2Schema->FindExecutionPin(*DummyNode, EGPD_Output);
+		}
+	}
 
-					if ((NewInPin != NULL) && (NewOutPin != NULL))
-					{
-						CompilerContext.MessageLog.NotifyIntermediatePinCreation(NewOutPin, OldOutPin);
+	const UEdGraphSchema_K2* K2Schema = GetDefault<UEdGraphSchema_K2>();
 
-						while (OldOutPin->LinkedTo.Num() > 0)
-						{
-							UEdGraphPin* LinkedPin = OldOutPin->LinkedTo[0];
+	// Find the function, starting at the parent class
+	UFunction* Function = FindFunction(Context, Node);
+	if ((Function != NULL) && (Function->HasMetaData(FBlueprintMetadata::MD_Latent)))
+	{
+		UK2Node_CallFunction* CallFuncNode = CastChecked<UK2Node_CallFunction>(Node);
+		UEdGraphPin* OldOutPin = K2Schema->FindExecutionPin(*CallFuncNode, EGPD_Output);
 
-							LinkedPin->BreakLinkTo(OldOutPin);
-							LinkedPin->MakeLinkTo(NewOutPin);
-						}
+		if ((OldOutPin != NULL) && (OldOutPin->LinkedTo.Num() > 0))
+		{
+			// Create a dummy execution sequence that will be the target of the return call from the latent action
+			UK2Node_ExecutionSequence* DummyNode = CompilerContext.SpawnIntermediateNode<UK2Node_ExecutionSequence>(CallFuncNode);
+			DummyNode->AllocateDefaultPins();
 
-						OldOutPin->MakeLinkTo(NewInPin);
-					}
-				}
-			}
+			// Wire in the dummy node
+			UEdGraphPin* NewInPin = K2Schema->FindExecutionPin(*DummyNode, EGPD_Input);
+			UEdGraphPin* NewOutPin = K2Schema->FindExecutionPin(*DummyNode, EGPD_Output);
 
-#if WITH_PUSH_MODEL
-
-			/**
-			 * Warning!! Similar code exists in VariableSetHandler.cpp
-			 *
-			 * This code is for property dirty tracking.
-			 * It works by injecting in extra nodes while compiling that will call UNetPushModelHelpers::MarkPropertyDirtyFromRepIndex.
-			 *
-			 * That function will be called with the Owner of the property, and the RepIndex of the property.
-			 * One of these nodes needs to be added for ever Net Property that is passed by (out) reference to the function.
-			 *
-			 *
-			 * Note, this assumes that there's no way that a native class can add or remove replicated properties without also
-			 * recompiling the blueprint. The only scenario that seems possible is cooked games with custom built binaries,
-			 * but that still seems unsafe.
-			 *
-			 * If that can happen, we can instead switch to using the property name and resorting to a FindField call at runtime,
-			 * but that will be more expensive.
-			 */
-
-			// If the function is pure but won't actually be evaluated, if there are no out params,
-			// or there are no input pins, then we don't need to worry about any extra generation
-			// because there will either be no way to reference a NetProperty, or the node won't
-			// have any effect.
-			if (!bIsPureAndNoUsedOutputs && Function->NumParms > 0 && Function->HasAllFunctionFlags(FUNC_HasOutParms))
+			if ((NewInPin != NULL) && (NewOutPin != NULL))
 			{
-				static const FName MarkPropertyDirtyFuncName(TEXT("MarkPropertyDirtyFromRepIndex"));
-				static const FName ObjectPinName(TEXT("Object"));
-				static const FName RepIndexPinName(TEXT("RepIndex"));
-				static const FName PropertyNamePinName(TEXT("PropertyName"));
-
-				TArray<UEdGraphPin*> RemainingPins(Node->Pins);
-				UEdGraphPin* OldThenPin = CallFuncNode->GetThenPin();
-
-				// Note: This feels like it's going to be a very hot path during compilation as it will be hit for every
-				// CallFunction node. Any optimizations that can be made here are probably worth it to not terribly regress
-				// BP Compile time for cooked games.
+				CompilerContext.MessageLog.NotifyIntermediatePinCreation(NewOutPin, OldOutPin);
 
-				// Iterate the properties looking for Out Params that are tied to Net Properties.
-				// This is similar to the loop in CreateCallFunction
-				for (TFieldIterator<UProperty> It(Function); It && (It->PropertyFlags & CPF_Parm); ++It)
+				while (OldOutPin->LinkedTo.Num() > 0)
 				{
-					UProperty* Property = *It;
-					if (Property->HasAllPropertyFlags(CPF_OutParm | CPF_ReferenceParm))
-					{
-						for (int32 i = 0; i < RemainingPins.Num(); ++i)
-						{
-							if (Property->GetFName() == RemainingPins[i]->PinName)
-							{
-								UEdGraphPin* ParamPin = RemainingPins[i];
-								if (UEdGraphPin * PinToTry = FEdGraphUtilities::GetNetFromPin(ParamPin))
-								{
-									if (UK2Node_VariableGet* GetPropertyNode = Cast<UK2Node_VariableGet>(PinToTry->GetOwningNode()))
-									{
-										UProperty* ToCheck = GetPropertyNode->GetPropertyForVariable();
-
-										// We only need to do this work if we actually found a net property.
-										if (UNLIKELY(ToCheck->HasAnyPropertyFlags(CPF_Net)))
-										{
-											if (OldThenPin == nullptr)
-											{
-												// TODO: Unfortunately, we don't have an efficient way to determine whether or not a property
-												//			uses push model at this point. That is only known from a call to "GetLifetimeRepProperties"
-												//			on a given object, and even then that is a user defined function that could do any number
-												//			of things.
-												//			We *might* be able to use something like Object Annotations to cache off the results of these
-												//			calls and clear the cache post compile, but that seems hacky and prone to error.
-												CompilerContext.MessageLog.Warning(TEXT("Passing net properties as Out Parameters with Push Model Enabled is not supported. Function=%s, Param=%s, Property=%s"),
-													*Function->GetName(), *Property->GetName(), *ToCheck->GetName());
-											}
-											else if (UClass * OwningClass = ToCheck->GetTypedOuter<UClass>())
-											{
-												// TODO: As of CL-10376545 in //UE4/Main, we no longer sort Native Replicated Properties by GC Offset.
-												//			This this chunk of code can probably be removed.
-
-												// We need to make sure this class already has its property offsets setup, otherwise
-												// the order of our replicated properties won't match, meaning the RepIndex will be
-												// invalid.
-												if (ToCheck->GetOffset_ForGC() == 0)
-												{
-													if (UBlueprint * Blueprint = Cast<UBlueprint>(OwningClass->ClassGeneratedBy))
-													{
-														if (UClass * UseClass = Blueprint->GeneratedClass)
-														{
-															OwningClass = UseClass;
-															ToCheck = FindFieldChecked<UProperty>(OwningClass, ToCheck->GetFName());
-														}
-													}
-												}
-
-												ensureAlwaysMsgf(ToCheck->GetOffset_ForGC() != 0,
-													TEXT("Class does not have Property Offsets setup. This will cause issues with Push Model. Blueprint=%s, Class=%s, Property=%s"),
-													*Context.Blueprint->GetPathName(), *OwningClass->GetPathName(), *ToCheck->GetName());
-
-												if (!OwningClass->HasAnyClassFlags(CLASS_ReplicationDataIsSetUp))
-												{
-													OwningClass->SetUpRuntimeReplicationData();
-												}
-
-												// Actually insert the nodes.
-												// TODO: I don't expect this to be a very frequent occurence, but we could be more clever about this.
-												//			Instead of inserting a new node per Out Ref Net Property, we might be able to create
-												//			A separate node that takes in lists of Objects, RepIndices, and Names.
-												//			Ideally, this would work like the String Append node where we could specify the values
-												//			directly on the node.
-												{
-													// Create the node that will call MarkPropertyDirty.
-													UK2Node_CallFunction* MarkPropertyDirtyNode = Node->GetGraph()->CreateIntermediateNode<UK2Node_CallFunction>();
-													MarkPropertyDirtyNode->FunctionReference.SetExternalMember(MarkPropertyDirtyFuncName, UNetPushModelHelpers::StaticClass());
-													MarkPropertyDirtyNode->AllocateDefaultPins();
-
-													// Create the Pins for RepIndex, PropertyName, and Object.
-													UEdGraphPin* RepIndexPin = MarkPropertyDirtyNode->FindPinChecked(RepIndexPinName);
-													RepIndexPin->DefaultValue = FString::FromInt(Property->RepIndex);
-
-													UEdGraphPin* PropertyNamePin = MarkPropertyDirtyNode->FindPinChecked(PropertyNamePinName);
-													PropertyNamePin->DefaultValue = Property->GetFName().ToString();
-
-													UEdGraphPin* ObjectPin = MarkPropertyDirtyNode->FindPinChecked(ObjectPinName);
-													UEdGraphPin* PropertyOwnerPin = GetPropertyNode->FindPinChecked(UEdGraphSchema_K2::PN_Self);
-
-													// If the property is linked to some other object, go ahead and grab that.
-													// Otherwise, create an intermediate self Pin and use that.
-													if (PropertyOwnerPin->LinkedTo.Num() > 0)
-													{
-														PropertyOwnerPin = PropertyOwnerPin->LinkedTo[0];
-													}
-													else
-													{
-														UK2Node_Self* SelfNode = Node->GetGraph()->CreateIntermediateNode<UK2Node_Self>();
-														SelfNode->AllocateDefaultPins();
-														PropertyOwnerPin = SelfNode->FindPinChecked(UEdGraphSchema_K2::PN_Self);
-													}
-
-													ObjectPin->MakeLinkTo(PropertyOwnerPin);
-
-													UEdGraphPin* NewThenPin = MarkPropertyDirtyNode->GetThenPin();
-													if (ensure(NewThenPin))
-													{
-														NewThenPin->CopyPersistentDataFromOldPin(*OldThenPin);
-														OldThenPin->BreakAllPinLinks();
-														OldThenPin->MakeLinkTo(MarkPropertyDirtyNode->GetExecPin());
-
-														// Now that the connection is established, go ahead and set our OldThenPin to
-														// the NewThenPin. If we find another NetProperty later, this will save us from
-														// doing a lookup.
-														OldThenPin = NewThenPin;
-													}
-												}
-											}
-										}
-									}
-								}
+					UEdGraphPin* LinkedPin = OldOutPin->LinkedTo[0];
 
-								RemainingPins.RemoveAtSwap(i);
-								break;
-							}
-						}
-					}
+					LinkedPin->BreakLinkTo(OldOutPin);
+					LinkedPin->MakeLinkTo(NewOutPin);
 				}
+
+				OldOutPin->MakeLinkTo(NewInPin);
 			}
-#endif
 		}
 	}
 }
diff --git a/Engine/Source/Editor/BlueprintGraph/Private/K2Node_VariableSet.cpp b/Engine/Source/Editor/BlueprintGraph/Private/K2Node_VariableSet.cpp
--- a/Engine/Source/Editor/BlueprintGraph/Private/K2Node_VariableSet.cpp
+++ b/Engine/Source/Editor/BlueprintGraph/Private/K2Node_VariableSet.cpp
@@ -306,12 +306,6 @@ bool UK2Node_VariableSet::ShouldFlushDormancyOnSet() const
 	return (Property != NULL && (Property->PropertyFlags & CPF_Net));
 }
 
-bool UK2Node_VariableSet::IsNetProperty() const
-{
-	UProperty* Property = GetPropertyForVariable();
-	return Property && (Property->PropertyFlags & CPF_Net);
-}
-
 FName UK2Node_VariableSet::GetRepNotifyName() const
 {
 	FProperty * Property = GetPropertyForVariable();
diff --git a/Engine/Source/Editor/BlueprintGraph/Private/VariableSetHandler.cpp b/Engine/Source/Editor/BlueprintGraph/Private/VariableSetHandler.cpp
--- a/Engine/Source/Editor/BlueprintGraph/Private/VariableSetHandler.cpp
+++ b/Engine/Source/Editor/BlueprintGraph/Private/VariableSetHandler.cpp
@@ -5,9 +5,7 @@
 #include "K2Node_CallFunction.h"
 #include "K2Node_Variable.h"
 #include "K2Node_VariableSet.h"
-#include "K2Node_Self.h"
 #include "Kismet2/BlueprintEditorUtils.h"
-#include "Net/NetPushModelHelpers.h"
 
 #include "EdGraphUtilities.h"
 #include "KismetCompiler.h"
@@ -198,108 +196,6 @@ void FKCHandler_VariableSet::Transform(FKismetFunctionContext& Context, UEdGraph
 				OldThenPin->MakeLinkTo(CallFuncNode->GetExecPin());
 			}
 		}
-
-		if (SetNotify->IsNetProperty())
-		{
-			/**
-			 * Warning!! Similar code exists in CallFunctionHandler.cpp
-			 *
-			 * This code is for property dirty tracking.
-			 * It works by injecting in extra nodes while compiling that will call UNetPushModelHelpers::MarkPropertyDirtyFromRepIndex.
-			 *
-			 * That function will be called with the Owner of the property (either Self or whatever is connected to the Target pin
-			 * of the BP Node), and the RepIndex of the property.
-			 *
-			 * Note, this assumes that there's no way that a native class can add or remove replicated properties without also
-			 * recompiling the blueprint. The only scenario that seems possible is cooked games with custom built binaries,
-			 * but that still seems unsafe.
-			 *
-			 * If that can happen, we can instead switch to using the property name and resorting to a FindField call at runtime,
-			 * but that will be more expensive.
-			 */
-
-			static const FName MarkPropertyDirtyFuncName(TEXT("MarkPropertyDirtyFromRepIndex"));
-			static const FName ObjectPinName(TEXT("Object"));
-			static const FName RepIndexPinName(TEXT("RepIndex"));
-			static const FName PropertyNamePinName(TEXT("PropertyName"));
-
-			UProperty* Property = SetNotify->GetPropertyForVariable();
-			UClass* Class = Property->GetTypedOuter<UClass>();
-
-			if (Property && Class)
-			{
-				// We need to make sure this class already has its property offsets setup, otherwise
-				// the order of our replicated properties won't match, meaning the RepIndex will be
-				// invalid.
-				if (Property->GetOffset_ForGC() == 0)
-				{
-					// Make sure that we're using the correct class and that it has replication data set up.
-					if (Class->ClassGeneratedBy == Context.Blueprint && Context.NewClass && Context.NewClass != Class)
-					{
-						Class = Context.NewClass;
-						Property = FindFieldChecked<UProperty>(Class, Property->GetFName());
-					}
-					if (Property->GetOffset_ForGC() == 0)
-					{
-						if (UBlueprint * Blueprint = Cast<UBlueprint>(Class->ClassGeneratedBy))
-						{
-							if (UClass * UseClass = Blueprint->GeneratedClass)
-							{
-								Class = UseClass;
-								Property = FindFieldChecked<UProperty>(Class, Property->GetFName());
-							}
-						}
-					}
-				}
-
-				ensureAlwaysMsgf(Property->GetOffset_ForGC() != 0,
-					TEXT("Class does not have Property Offsets setup. This will cause issues with Push Model. Blueprint=%s, Class=%s, Property=%s"),
-					*Context.Blueprint->GetPathName(), *Class->GetPathName(), *Property->GetName());
-
-				if (!Class->HasAnyClassFlags(CLASS_ReplicationDataIsSetUp))
-				{
-					Class->SetUpRuntimeReplicationData();
-				}
-
-				UK2Node_CallFunction* CallFuncNode = Node->GetGraph()->CreateIntermediateNode<UK2Node_CallFunction>();
-				CallFuncNode->FunctionReference.SetExternalMember(MarkPropertyDirtyFuncName, UNetPushModelHelpers::StaticClass());
-				CallFuncNode->AllocateDefaultPins();
-
-				// Take our old Self (Target) pin and hook it up to the Object pin for UNetPushModelHelpers::MarkPropertyDirty.
-				// If our Self pin isn't hooked up to anything, then create an intermediate Self node and use that.
-				UEdGraphPin* SelfPin = Node->FindPinChecked(UEdGraphSchema_K2::PN_Self);
-				UEdGraphPin* ObjectPin = CallFuncNode->FindPinChecked(ObjectPinName);
-
-				if (SelfPin->LinkedTo.Num() > 0)
-				{
-					SelfPin = SelfPin->LinkedTo[0];
-				}
-				else
-				{
-					UK2Node_Self* SelfNode = Node->GetGraph()->CreateIntermediateNode<UK2Node_Self>();
-					SelfNode->AllocateDefaultPins();
-					SelfPin = SelfNode->FindPinChecked(UEdGraphSchema_K2::PN_Self);
-				}
-
-				ObjectPin->MakeLinkTo(SelfPin);
-
-				UEdGraphPin* RepIndexPin = CallFuncNode->FindPinChecked(RepIndexPinName);
-				RepIndexPin->DefaultValue = FString::FromInt(Property->RepIndex);
-
-				UEdGraphPin* PropertyNamePin = CallFuncNode->FindPinChecked(PropertyNamePinName);
-				PropertyNamePin->DefaultValue = Property->GetFName().ToString();
-
-				// Hook up our exec pins.
-				UEdGraphPin* OldThenPin = Node->FindPinChecked(UEdGraphSchema_K2::PN_Then);
-				UEdGraphPin* NewThenPin = CallFuncNode->GetThenPin();
-				if (ensure(NewThenPin))
-				{
-					NewThenPin->CopyPersistentDataFromOldPin(*OldThenPin);
-					OldThenPin->BreakAllPinLinks();
-					OldThenPin->MakeLinkTo(CallFuncNode->GetExecPin());
-				}
-			}
-		}
 	}
 }
 
diff --git a/Engine/Source/Programs/UnrealBuildTool/Configuration/TargetRules.cs b/Engine/Source/Programs/UnrealBuildTool/Configuration/TargetRules.cs
--- a/Engine/Source/Programs/UnrealBuildTool/Configuration/TargetRules.cs
+++ b/Engine/Source/Programs/UnrealBuildTool/Configuration/TargetRules.cs
@@ -584,12 +584,6 @@ namespace UnrealBuildTool
 		}
 		private bool? bWithServerCodeOverride;
 
-		/// <summary>
-		/// When enabled, Push Model Networking will be used on the server. This can help reduce CPU overhead of networking, at the cost of more memory.
-		/// </summary>
-		[RequiresUniqueBuildEnvironment]
-		public bool bWithPushModel = false;
-
 		/// <summary>
 		/// Whether to include stats support even without the engine.
 		/// </summary>
@@ -2059,11 +2053,6 @@ namespace UnrealBuildTool
 			get { return Inner.bWithServerCode; }
 		}
 
-		public bool bWithPushModel
-		{
-			get { return Inner.bWithPushModel; }
-		}
-
 		public bool bCompileWithStatsWithoutEngine
 		{
 			get { return Inner.bCompileWithStatsWithoutEngine; }
diff --git a/Engine/Source/Programs/UnrealBuildTool/Configuration/UEBuildTarget.cs b/Engine/Source/Programs/UnrealBuildTool/Configuration/UEBuildTarget.cs
--- a/Engine/Source/Programs/UnrealBuildTool/Configuration/UEBuildTarget.cs
+++ b/Engine/Source/Programs/UnrealBuildTool/Configuration/UEBuildTarget.cs
@@ -3594,16 +3594,6 @@ namespace UnrealBuildTool
 				GlobalCompileEnvironment.Definitions.Add("WITH_SERVER_CODE=0");
 			}
 
-			// Set the defines for Push Model
-			if (Rules.bWithPushModel)
-			{
-				GlobalCompileEnvironment.Definitions.Add("WITH_PUSH_MODEL=1");
-			}
-			else
-			{
-				GlobalCompileEnvironment.Definitions.Add("WITH_PUSH_MODEL=0");
-			}
-
 			// Set the define for whether we're compiling with CEF3
 			if (Rules.bCompileCEF3 && (Platform == UnrealTargetPlatform.Win32 || Platform == UnrealTargetPlatform.Win64 || Platform == UnrealTargetPlatform.Mac || Platform == UnrealTargetPlatform.Linux))
 			{
diff --git a/Engine/Source/Programs/UnrealHeaderTool/Private/CodeGenerator.cpp b/Engine/Source/Programs/UnrealHeaderTool/Private/CodeGenerator.cpp
--- a/Engine/Source/Programs/UnrealHeaderTool/Private/CodeGenerator.cpp
+++ b/Engine/Source/Programs/UnrealHeaderTool/Private/CodeGenerator.cpp
@@ -257,8 +257,7 @@ namespace
 		FClass* Class,
 		FClass* SuperClass,
 		FOutputDevice& Writer,
-		const FUnrealSourceFile& SourceFile,
-		FNativeClassHeaderGenerator::EExportClassOutFlags& OutFlags)
+		const FUnrealSourceFile& SourceFile)
 	{
 		const bool bHasGetLifetimeReplicatedProps = HasIdentifierExactMatch(ClassRange.Start, ClassRange.End, GetLifetimeReplicatedPropsStr);
 
@@ -279,18 +278,6 @@ namespace
 		Class->SetUpUhtReplicationData(&ParentRepClass);
 
 		ExportNetData(Writer, Class, ParentRepClass, API);
-		
-		// If we don't have any Super Classes that have replicated properties,
-		// then we're the base most replicated class, and we should add in our interface stuff.
-		if (!ParentRepClass)
-		{
-			OutFlags |= FNativeClassHeaderGenerator::EExportClassOutFlags::NeedsPushModelHeaders;
-			Writer.Logf(TEXT(
-				"private:\r\n"
-				"\tREPLICATED_BASE_CLASS(%s)\r\n"
-				"public:\r\n"
-			), *Class->GetName());
-		}
 	}
 }
 
@@ -2962,8 +2949,7 @@ void FNativeClassHeaderGenerator::ExportClassFromSourceFileInner(
 	FOutputDevice&           OutDeclarations,
 	FReferenceGatherers&     OutReferenceGatherers,
 	FClass*                  Class,
-	const FUnrealSourceFile& SourceFile,
-	EExportClassOutFlags&    OutFlags
+	const FUnrealSourceFile& SourceFile
 ) const
 {
 	FUHTStringBuilder StandardUObjectConstructorsMacroCall;
@@ -3045,6 +3031,8 @@ void FNativeClassHeaderGenerator::ExportClassFromSourceFileInner(
 
 	FString PPOMacroName;
 
+	// Replication, add in the declaration for GetLifetimeReplicatedProps() automatically if there are any net flagged properties
+
 	ClassDefinitionRange ClassRange;
 	if (ClassDefinitionRange* FoundRange = ClassDefinitionRanges.Find(Class))
 	{
@@ -3202,7 +3190,7 @@ void FNativeClassHeaderGenerator::ExportClassFromSourceFileInner(
 
 			if (ClassHasReplicatedProperties(Class))
 			{
-				WriteReplicatedMacroData(ClassRange, *ClassCPPName, *APIArg, Class, SuperClass, InterfaceBoilerplate, SourceFile, OutFlags);
+				WriteReplicatedMacroData(ClassRange, *ClassCPPName, *APIArg, Class, SuperClass, InterfaceBoilerplate, SourceFile);
 			}
 
 			FString NoPureDeclsMacroName = SourceFile.GetGeneratedMacroName(ClassData, TEXT("_INCLASS_IINTERFACE_NO_PURE_DECLS"));
@@ -3229,7 +3217,7 @@ void FNativeClassHeaderGenerator::ExportClassFromSourceFileInner(
 
 			if (ClassHasReplicatedProperties(Class))
 			{
-				WriteReplicatedMacroData(ClassRange, *ClassCPPName, *APIArg, Class, SuperClass, Boilerplate, SourceFile, OutFlags);
+				WriteReplicatedMacroData(ClassRange, *ClassCPPName, *APIArg, Class, SuperClass, Boilerplate, SourceFile);
 			}
 
 			{
@@ -3487,18 +3475,12 @@ void FNativeClassHeaderGenerator::ExportConstructorsMacros(FOutputDevice& OutGen
 	EnhancedUObjectConstructorsMacroCall.Logf(TEXT("\t%s\r\n"), *EnhMacroName);
 }
 
-bool FNativeClassHeaderGenerator::WriteHeader(const TCHAR* Path, const FString& InBodyText, const TSet<FString>& InAdditionalHeaders, FReferenceGatherers& InOutReferenceGatherers) const
+bool FNativeClassHeaderGenerator::WriteHeader(const TCHAR* Path, const FString& InBodyText, FReferenceGatherers& InOutReferenceGatherers) const
 {
 	FUHTStringBuilder GeneratedHeaderTextWithCopyright;
 	GeneratedHeaderTextWithCopyright.Logf(TEXT("%s"), HeaderCopyright);
 	GeneratedHeaderTextWithCopyright.Log(TEXT("#include \"UObject/ObjectMacros.h\"\r\n"));
 	GeneratedHeaderTextWithCopyright.Log(TEXT("#include \"UObject/ScriptMacros.h\"\r\n"));
-
-	for (const FString& AdditionalHeader : InAdditionalHeaders)
-	{
-		GeneratedHeaderTextWithCopyright.Logf(TEXT("#include \"%s\"\r\n"), *AdditionalHeader);
-	}
-
 	GeneratedHeaderTextWithCopyright.Log(LINE_TERMINATOR);
 	GeneratedHeaderTextWithCopyright.Log(TEXT("PRAGMA_DISABLE_DEPRECATION_WARNINGS") LINE_TERMINATOR);
 
@@ -5762,23 +5744,16 @@ FNativeClassHeaderGenerator::FNativeClassHeaderGenerator(
 			ConstThis->ExportDelegateDefinition(GeneratedHeaderText, ReferenceGatherers, *SourceFile, Func);
 		}
 
-		EExportClassOutFlags ExportFlags = EExportClassOutFlags::None;
-		TSet<FString> AdditionalHeaders;
 		TArray<UClass*>	DefinedClasses;
 		SourceFile->AppendDefinedClasses(DefinedClasses);
 		for (UClass* Class : DefinedClasses)
 		{
 			if (!(Class->ClassFlags & CLASS_Intrinsic))
 			{
-				ConstThis->ExportClassFromSourceFileInner(GeneratedHeaderText, OutText, GeneratedFunctionDeclarations, ReferenceGatherers, (FClass*)Class, *SourceFile, ExportFlags);
+				ConstThis->ExportClassFromSourceFileInner(GeneratedHeaderText, OutText, GeneratedFunctionDeclarations, ReferenceGatherers, (FClass*)Class, *SourceFile);
 			}
 		}
 
-		if (EnumHasAnyFlags(ExportFlags, EExportClassOutFlags::NeedsPushModelHeaders))
-		{
-			AdditionalHeaders.Add(FString(TEXT("Net/Core/PushModel/PushModelMacros.h")));
-		}
-
 		GeneratedHeaderText.Log(TEXT("#undef CURRENT_FILE_ID\r\n"));
 		GeneratedHeaderText.Logf(TEXT("#define CURRENT_FILE_ID %s\r\n\r\n\r\n"), *SourceFile->GetFileId());
 
@@ -5788,7 +5763,7 @@ FNativeClassHeaderGenerator::FNativeClassHeaderGenerator(
 		}
 
 		FString HeaderPath = BaseSourceFilename + TEXT(".generated.h");
-		bool bHasChanged = ConstThis->WriteHeader(*HeaderPath, GeneratedHeaderText, AdditionalHeaders, ReferenceGatherers);
+		bool bHasChanged = ConstThis->WriteHeader(*HeaderPath, GeneratedHeaderText, ReferenceGatherers);
 
 		SourceFile->SetGeneratedFilename(MoveTemp(HeaderPath));
 		SourceFile->SetHasChanged(bHasChanged);
diff --git a/Engine/Source/Programs/UnrealHeaderTool/Private/NativeClassExporter.h b/Engine/Source/Programs/UnrealHeaderTool/Private/NativeClassExporter.h
--- a/Engine/Source/Programs/UnrealHeaderTool/Private/NativeClassExporter.h
+++ b/Engine/Source/Programs/UnrealHeaderTool/Private/NativeClassExporter.h
@@ -184,16 +184,6 @@ private:
 	 */
 	static void ExportUInterfaceBoilerplate(FUHTStringBuilder& UInterfaceBoilerplate, FClass* Class, const FString& FriendText);
 
-public:
-
-	enum class EExportClassOutFlags
-	{
-		None = 0x0,
-		NeedsPushModelHeaders = 0x1 << 0,
-	};
-
-private:
-
 	/**
 	 * After all of the dependency checking, and setup for isolating the generated code, actually export the class
 	 *
@@ -207,8 +197,7 @@ private:
 		FOutputDevice&           OutDeclarations,
 		FReferenceGatherers&     OutReferenceGatherers, 
 		FClass*                  Class,
-		const FUnrealSourceFile& SourceFile,
-		EExportClassOutFlags&    OutFlags
+		const FUnrealSourceFile& SourceFile
 	) const;
 
 	/**
@@ -216,7 +205,7 @@ private:
 	 *
 	 * @param	SourceFile			Source file to export.
 	 */
-	bool WriteHeader(const TCHAR* Path, const FString& InBodyText, const TSet<FString>& InAdditionalHeaders, FReferenceGatherers& InOutReferenceGatherers) const;
+	bool WriteHeader(const TCHAR* Path, const FString& InBodyText, FReferenceGatherers& InOutReferenceGatherers) const;
 
 	/**
 	 * Returns a string in the format CLASS_Something|CLASS_Something which represents all class flags that are set for the specified
@@ -563,5 +552,3 @@ public:
 	 */
 	static bool IsMissingVirtualSpecifier(const FString& SourceFile, int32 FunctionNamePosition);
 };
-
-ENUM_CLASS_FLAGS(FNativeClassHeaderGenerator::EExportClassOutFlags);
\ No newline at end of file
diff --git a/Engine/Source/Runtime/CoreUObject/Public/UObject/CoreNet.h b/Engine/Source/Runtime/CoreUObject/Public/UObject/CoreNet.h
--- a/Engine/Source/Runtime/CoreUObject/Public/UObject/CoreNet.h
+++ b/Engine/Source/Runtime/CoreUObject/Public/UObject/CoreNet.h
@@ -247,46 +247,20 @@ struct FPropertyRetirement
 class FLifetimeProperty
 {
 public:
-
-	uint16 RepIndex;
-	ELifetimeCondition Condition;
+	uint16				RepIndex;
+	ELifetimeCondition	Condition;
 	ELifetimeRepNotifyCondition RepNotifyCondition;
-	bool bIsPushBased;
 
-	FLifetimeProperty()
-		: RepIndex(0)
-		, Condition(COND_None)
-		, RepNotifyCondition(REPNOTIFY_OnChanged)
-		, bIsPushBased(false)
-	{
-	}
-
-	FLifetimeProperty(int32 InRepIndex)
-		: RepIndex(InRepIndex)
-		, Condition(COND_None)
-		, RepNotifyCondition(REPNOTIFY_OnChanged)
-		, bIsPushBased(false)
-	{
-		check(InRepIndex <= 65535);
-	}
-
-	FLifetimeProperty(int32 InRepIndex, ELifetimeCondition InCondition, ELifetimeRepNotifyCondition InRepNotifyCondition=REPNOTIFY_OnChanged, bool bInIsPushBased=false)
-		: RepIndex(InRepIndex)
-		, Condition(InCondition)
-		, RepNotifyCondition(InRepNotifyCondition)
-		, bIsPushBased(bInIsPushBased)
-	{
-		check(InRepIndex <= 65535);
-	}
+	FLifetimeProperty() : RepIndex( 0 ), Condition( COND_None ), RepNotifyCondition(REPNOTIFY_OnChanged) {}
+	FLifetimeProperty( int32 InRepIndex ) : RepIndex( InRepIndex ), Condition( COND_None ), RepNotifyCondition(REPNOTIFY_OnChanged) { check( InRepIndex <= 65535 ); }
+	FLifetimeProperty(int32 InRepIndex, ELifetimeCondition InCondition, ELifetimeRepNotifyCondition InRepNotifyCondition=REPNOTIFY_OnChanged) : RepIndex(InRepIndex), Condition(InCondition), RepNotifyCondition(InRepNotifyCondition) { check(InRepIndex <= 65535); }
 
-	inline bool operator==(const FLifetimeProperty& Other) const
+	inline bool operator==( const FLifetimeProperty& Other ) const
 	{
-		if (RepIndex == Other.RepIndex)
+		if ( RepIndex == Other.RepIndex )
 		{
-			// Can't have different conditions if the RepIndex matches, doesn't make sense
-			check(Condition == Other.Condition);
-			check(RepNotifyCondition == Other.RepNotifyCondition);
-			check(bIsPushBased == Other.bIsPushBased);
+			check( Condition == Other.Condition );		// Can't have different conditions if the RepIndex matches, doesn't make sense
+			check( RepNotifyCondition == Other.RepNotifyCondition);
 			return true;
 		}
 
@@ -458,12 +432,9 @@ class IRepChangedPropertyTracker
 public:
 	IRepChangedPropertyTracker() { }
 
-	virtual void SetCustomIsActiveOverride(
-		UObject* OwningObject,
-		const uint16 RepIndex,
-		const bool bIsActive) = 0;
+	virtual void SetCustomIsActiveOverride( const uint16 RepIndex, const bool bIsActive ) = 0;
 
-	virtual void SetExternalData(const uint8* Src, const int32 NumBits) = 0;
+	virtual void SetExternalData( const uint8* Src, const int32 NumBits ) = 0;
 
 	virtual bool IsReplay() const = 0;
 
diff --git a/Engine/Source/Runtime/CoreUObject/Public/UObject/Object.h b/Engine/Source/Runtime/CoreUObject/Public/UObject/Object.h
--- a/Engine/Source/Runtime/CoreUObject/Public/UObject/Object.h
+++ b/Engine/Source/Runtime/CoreUObject/Public/UObject/Object.h
@@ -1457,40 +1457,7 @@ private:
 
 public:
 
-	virtual void ValidateGeneratedRepEnums(const TArray<struct FRepRecord>& ClassReps) const {}	
-
-private:
-	
-	friend struct FObjectNetPushIdHelper;
-	virtual void SetNetPushIdDynamic(const int32 NewNetPushId)
-	{
-		// This method should only be called on Objects that are networked, and those should
-		// always have this implemented (by UHT).
-		check(false);
-	}
-	
-	virtual int32 GetNetPushIdDynamic() const
-	{
-		return INDEX_NONE;
-	}
-};
-
-struct FObjectNetPushIdHelper
-{
-private:
-	friend struct FNetPrivatePushIdHelper;
-	friend class UNetPushModelHelpers;
-	friend class UKismetArrayLibrary;
-
-	static void SetNetPushIdDynamic(UObject* Object, const int32 NewNetPushId)
-	{
-		Object->SetNetPushIdDynamic(NewNetPushId);
-	}
-	
-	static int32 GetNetPushIdDynamic(const UObject* const Object)
-	{
-		return Object->GetNetPushIdDynamic();
-	}
+	virtual void ValidateGeneratedRepEnums(const TArray<struct FRepRecord>& ClassReps) const {}
 };
 
 /**
diff --git a/Engine/Source/Runtime/Engine/Classes/Kismet/KismetArrayLibrary.h b/Engine/Source/Runtime/Engine/Classes/Kismet/KismetArrayLibrary.h
--- a/Engine/Source/Runtime/Engine/Classes/Kismet/KismetArrayLibrary.h
+++ b/Engine/Source/Runtime/Engine/Classes/Kismet/KismetArrayLibrary.h
@@ -10,7 +10,6 @@
 #include "UObject/UnrealType.h"
 #include "UObject/ScriptMacros.h"
 #include "Kismet/BlueprintFunctionLibrary.h"
-#include "Net/Core/PushModel/PushModel.h"
 #include "KismetArrayLibrary.generated.h"
 
 class AActor;
@@ -264,7 +263,6 @@ public:
  
  		P_FINISH;
 		P_NATIVE_BEGIN;
-		MARK_PROPERTY_DIRTY(Stack.Object, ArrayProperty);
 		*(int32*)RESULT_PARAM = GenericArray_Add(ArrayAddr, ArrayProperty, StorageSpace);
 		P_NATIVE_END;
 		InnerProp->DestroyValue(StorageSpace);
@@ -293,7 +291,6 @@ public:
 
 		P_FINISH;
 		P_NATIVE_BEGIN;
-		MARK_PROPERTY_DIRTY(Stack.Object, ArrayProperty);
 		*(int32*)RESULT_PARAM = GenericArray_AddUnique(ArrayAddr, ArrayProperty, StorageSpace);
 		P_NATIVE_END;
 		InnerProp->DestroyValue(StorageSpace);
@@ -313,7 +310,6 @@ public:
 
 		P_FINISH;
 		P_NATIVE_BEGIN;
-		MARK_PROPERTY_DIRTY(Stack.Object, ArrayProperty);
 		GenericArray_Shuffle(ArrayAddr, ArrayProperty);
 		P_NATIVE_END;
 	}
@@ -372,7 +368,6 @@ public:
 
 		P_FINISH;
 		P_NATIVE_BEGIN;
-		MARK_PROPERTY_DIRTY(Stack.Object, TargetArrayProperty);
 		GenericArray_Append(TargetArrayAddr, TargetArrayProperty, SourceArrayAddr, SourceArrayProperty);
 		P_NATIVE_END;
 	}
@@ -401,7 +396,6 @@ public:
 		P_GET_PROPERTY(FIntProperty, Index);
 		P_FINISH;
 		P_NATIVE_BEGIN;
-		MARK_PROPERTY_DIRTY(Stack.Object, ArrayProperty);
 		GenericArray_Insert(ArrayAddr, ArrayProperty, StorageSpace, Index);
 		P_NATIVE_END;
 		InnerProp->DestroyValue(StorageSpace);
@@ -422,7 +416,6 @@ public:
 		P_GET_PROPERTY(FIntProperty, Index);
 		P_FINISH;
 		P_NATIVE_BEGIN;
-		MARK_PROPERTY_DIRTY(Stack.Object, ArrayProperty);
 		GenericArray_Remove(ArrayAddr, ArrayProperty, Index);
 		P_NATIVE_END;
 	}
@@ -453,7 +446,6 @@ public:
 		// Bools need to be processed internally by the property so that C++ bool value is properly set.
 		GenericArray_HandleBool(InnerProp, ItemPtr);
 		P_NATIVE_BEGIN;
-		MARK_PROPERTY_DIRTY(Stack.Object, ArrayProperty);
 		*(bool*)RESULT_PARAM = GenericArray_RemoveItem(ArrayAddr, ArrayProperty, ItemPtr);
 		P_NATIVE_END;
 
@@ -473,7 +465,6 @@ public:
 		}
 		P_FINISH;
 		P_NATIVE_BEGIN;
-		MARK_PROPERTY_DIRTY(Stack.Object, ArrayProperty);
 		GenericArray_Clear(ArrayAddr, ArrayProperty);
 		P_NATIVE_END;
 	}
@@ -492,7 +483,6 @@ public:
 		P_GET_PROPERTY(FIntProperty, Size);
 		P_FINISH;
 		P_NATIVE_BEGIN;
-		MARK_PROPERTY_DIRTY(Stack.Object, ArrayProperty);
 		GenericArray_Resize(ArrayAddr, ArrayProperty, Size);
 		P_NATIVE_END;
 	}
@@ -588,7 +578,6 @@ public:
 		P_FINISH;
 
 		P_NATIVE_BEGIN;
-		MARK_PROPERTY_DIRTY(Stack.Object, ArrayProperty);
 		GenericArray_Set(ArrayAddr, ArrayProperty, Index, StorageSpace, bSizeToFit);
 		P_NATIVE_END;
 		InnerProp->DestroyValue(StorageSpace);
@@ -612,7 +601,6 @@ public:
 
 		P_FINISH;
 		P_NATIVE_BEGIN;
-		MARK_PROPERTY_DIRTY(Stack.Object, ArrayProperty);
 		GenericArray_Swap(ArrayAddr, ArrayProperty, First, Second);
 		P_NATIVE_END;
 	}
@@ -696,7 +684,6 @@ public:
 		P_FINISH;
 
 		P_NATIVE_BEGIN;
-		MARK_PROPERTY_DIRTY(Stack.Object, Stack.MostRecentProperty);
 		GenericArray_SetArrayPropertyByName(OwnerObject, ArrayPropertyName, SrcArrayAddr);
 		P_NATIVE_END;
 	}
diff --git a/Engine/Source/Runtime/Engine/Engine.Build.cs b/Engine/Source/Runtime/Engine/Engine.Build.cs
--- a/Engine/Source/Runtime/Engine/Engine.Build.cs
+++ b/Engine/Source/Runtime/Engine/Engine.Build.cs
@@ -20,7 +20,6 @@ public class Engine : ModuleRules
 				"Developer/DerivedDataCache/Public",
 				"Runtime/SynthBenchmark/Public",
 				"Runtime/Engine/Private",
-				"Runtime/Net/Core/Private/Net/Core/PushModel/Types"
 			}
 		);
 
diff --git a/Engine/Source/Runtime/Engine/Private/BlueprintGeneratedClass.cpp b/Engine/Source/Runtime/Engine/Private/BlueprintGeneratedClass.cpp
--- a/Engine/Source/Runtime/Engine/Private/BlueprintGeneratedClass.cpp
+++ b/Engine/Source/Runtime/Engine/Private/BlueprintGeneratedClass.cpp
@@ -20,7 +20,6 @@
 #include "Engine/SCS_Node.h"
 #include "Engine/InheritableComponentHandler.h"
 #include "Misc/ScopeLock.h"
-#include "Net/Core/PushModel/PushModel.h"
 #include "UObject/CoreObjectVersion.h"
 
 #if WITH_EDITOR
@@ -1681,7 +1680,7 @@ void UBlueprintGeneratedClass::GetLifetimeBlueprintReplicationList(TArray<FLifet
 		{
 			PropertiesLeft--;
 			
-			OutLifetimeProps.AddUnique(FLifetimeProperty(Prop->RepIndex, Prop->GetBlueprintReplicationCondition(), REPNOTIFY_OnChanged, PUSH_MAKE_BP_PROPERTIES_PUSH_MODEL()));
+			OutLifetimeProps.AddUnique(FLifetimeProperty(Prop->RepIndex, Prop->GetBlueprintReplicationCondition()));
 		}
 	}
 
diff --git a/Engine/Source/Runtime/Engine/Private/Net/NetPushModelHelpers.cpp b/Engine/Source/Runtime/Engine/Private/Net/NetPushModelHelpers.cpp
deleted file mode 100644
diff --git a/Engine/Source/Runtime/Engine/Private/NetworkDriver.cpp b/Engine/Source/Runtime/Engine/Private/NetworkDriver.cpp
--- a/Engine/Source/Runtime/Engine/Private/NetworkDriver.cpp
+++ b/Engine/Source/Runtime/Engine/Private/NetworkDriver.cpp
@@ -3251,8 +3251,6 @@ void UNetDriver::PostGarbageCollect()
 			It.RemoveCurrent();
 		}
 	}
-	
-	CALL_PUSH_MODEL_POSTGARBAGECOLLECT();
 }
 
 void UNetDriver::AddReferencedObjects(UObject* InThis, FReferenceCollector& Collector)
diff --git a/Engine/Source/Runtime/Engine/Private/RepLayout.cpp b/Engine/Source/Runtime/Engine/Private/RepLayout.cpp
--- a/Engine/Source/Runtime/Engine/Private/RepLayout.cpp
+++ b/Engine/Source/Runtime/Engine/Private/RepLayout.cpp
@@ -21,7 +21,6 @@
 #include "Net/NetworkGranularMemoryLogging.h"
 #include "Serialization/ArchiveCountMem.h"
 #include "Templates/AndOrNot.h"
-#include "PushModelPerNetDriverState.h"
 
 DECLARE_CYCLE_STAT(TEXT("RepLayout AddPropertyCmd"), STAT_RepLayout_AddPropertyCmd, STATGROUP_Game);
 DECLARE_CYCLE_STAT(TEXT("RepLayout InitFromObjectClass"), STAT_RepLayout_InitFromObjectClass, STATGROUP_Game);
@@ -61,17 +60,6 @@ static FAutoConsoleVariableRef CVarShareShadowState(TEXT("net.ShareShadowState")
 int32 GShareInitialCompareState = 0;
 static FAutoConsoleVariableRef CVarShareInitialCompareState(TEXT("net.ShareInitialCompareState"), GShareInitialCompareState, TEXT("If true and net.ShareShadowState is enabled, attempt to also share initial replication compares across connections."));
 
-#if WITH_PUSH_VALIDATION_SUPPORT
-
-static bool GbPushModelValidateProperties = false;
-static FAutoConsoleVariableRef CVarPushModelValidateProperties(TEXT("net.PushModelValidateProperties"), GbPushModelValidateProperties, TEXT("When true, we will compare all push model properties and warn if they haven't been marked dirty properly."));
-
-#else
-
-constexpr bool GbPushModelValidateProperties = false;
-
-#endif
-
 int32 MaxRepArraySize = UNetworkSettings::DefaultMaxRepArraySize;
 int32 MaxRepArrayMemory = UNetworkSettings::DefaultMaxRepArrayMemory;
 
@@ -224,35 +212,6 @@ namespace UE4_RepLayout_Private
 
 		NETWORK_PROFILER(GNetworkProfiler.TrackWritePropertyHeader(Property, HeaderBits, nullptr));
 	}
-
-#if WITH_PUSH_MODEL
-	const UE4PushModelPrivate::FPushModelPerNetDriverHandle ConditionallyAddPushModelObject(const UE4PushModelPrivate::FNetPushObjectId ObjectId, const TSharedRef<const FRepLayout>& InRepLayout)
-	{
-		const int32 NumReplicatedProperties = InRepLayout->GetNumParents();
-		if (UE4PushModelPrivate::IsPushModelEnabled() &&
-			NumReplicatedProperties > 0 &&
-			EnumHasAnyFlags(InRepLayout->GetFlags(), ERepLayoutFlags::FullPushSupport | ERepLayoutFlags::PartialPushSupport))
-		{
-			return UE4PushModelPrivate::AddPushModelObject(ObjectId, NumReplicatedProperties);
-		}
-
-		return UE4PushModelPrivate::FPushModelPerNetDriverHandle::MakeInvalidHandle();
-	}
-
-	void ConditionallyRemovePushModelObject(const UE4PushModelPrivate::FPushModelPerNetDriverHandle Handle)
-	{
-		if (Handle.IsValid())
-		{
-			UE4PushModelPrivate::RemovePushModelObject(Handle);
-		}
-	}
-
-	UE4PushModelPrivate::FPushModelPerNetDriverState* GetPerNetDriverState(const UE4PushModelPrivate::FPushModelPerNetDriverHandle Handle)
-	{
-		return Handle.IsValid() ? UE4PushModelPrivate::GetPerNetDriverState(Handle) : nullptr;
-	}
-
-#endif
 }
 
 //~ TODO: Consider moving the FastArray members into their own sub-struct to save memory for non fast array
@@ -785,26 +744,6 @@ static uint32 GetRepLayoutCmdCompatibleChecksum(
 	return CompatibleChecksum;
 }
 
-#if (WITH_PUSH_MODEL)
-struct FNetPrivatePushIdHelper
-{
-	static void SetNetPushID(UObject* InObject, const UE4PushModelPrivate::FNetPushObjectId ObjectId)
-	{
-		const UE4PushModelPrivate::FNetPushObjectId CurrentId = FObjectNetPushIdHelper::GetNetPushIdDynamic(InObject);
-		if (CurrentId != ObjectId)
-		{
-			check(CurrentId == INDEX_NONE);
-			FObjectNetPushIdHelper::SetNetPushIdDynamic(InObject, ObjectId);
-		}
-	}
-
-	static void MarkPropertyDirty(UObject* InObject, const int32 RepIndex)
-	{
-		MARK_PROPERTY_DIRTY_UNSAFE(InObject, RepIndex);
-	}
-};
-#endif // (WITH_PUSH_MODEL)
-
 FRepChangedPropertyTracker::FRepChangedPropertyTracker(const bool InbIsReplay, const bool InbIsClientReplayRecording):
 		bIsReplay(InbIsReplay),
 		bIsClientReplayRecording(InbIsClientReplayRecording),
@@ -813,22 +752,14 @@ FRepChangedPropertyTracker::FRepChangedPropertyTracker(const bool InbIsReplay, c
 
 FRepChangedPropertyTracker::~FRepChangedPropertyTracker() {}
 
-void FRepChangedPropertyTracker::SetCustomIsActiveOverride(UObject* OwningObject, const uint16 RepIndex, const bool bIsActive)
+void FRepChangedPropertyTracker::SetCustomIsActiveOverride(const uint16 RepIndex, const bool bIsActive)
 {
-	FRepChangedParent& Parent = Parents[RepIndex];
+	FRepChangedParent & Parent = Parents[RepIndex];
 
-	Parent.OldActive = Parent.Active;
 	Parent.Active = (bIsActive || bIsClientReplayRecording) ? 1 : 0;
-
-#if WITH_PUSH_MODEL
-	if (!Parent.OldActive && Parent.Active)
-	{
-		FNetPrivatePushIdHelper::MarkPropertyDirty(OwningObject, RepIndex);
-	}
-#endif	
+	Parent.OldActive = Parent.Active;
 }
 
-
 void FRepChangedPropertyTracker::SetExternalData(const uint8* Src, const int32 NumBits)
 {
 	ExternalDataNumBits = NumBits;
@@ -967,35 +898,15 @@ void FRepStateStaticBuffer::CountBytes(FArchive& Ar) const
 
 FRepChangelistState::FRepChangelistState(
 	const TSharedRef<const FRepLayout>& InRepLayout,
-	const uint8* InSource,
-	const UObject* InRepresenting,
-	FCustomDeltaChangelistState* InDeltaChangelistState)
-
-	: CustomDeltaChangelistState(InDeltaChangelistState)
-	, HistoryStart(0)
-	, HistoryEnd(0)
-	, CompareIndex(0)
-	, StaticBuffer(InRepLayout->CreateShadowBuffer(InSource))
-
-#if WITH_PUSH_MODEL
-	, PushModelObjectHandle(UE4_RepLayout_Private::ConditionallyAddPushModelObject(InRepresenting->GetUniqueID(), InRepLayout))
-#endif // WITH_PUSH_MODEL
-
-{
-#if (WITH_PUSH_MODEL)
-	if (PushModelObjectHandle.IsValid())
-	{
-		FNetPrivatePushIdHelper::SetNetPushID(const_cast<UObject*>(InRepresenting), PushModelObjectHandle.ObjectId);
-	}
-#endif // (WITH_PUSH_MODEL)
-}
-
-FRepChangelistState::~FRepChangelistState()
-{
-#if WITH_PUSH_MODEL
-	UE4_RepLayout_Private::ConditionallyRemovePushModelObject(PushModelObjectHandle);
-#endif // WITH_PUSH_MODEL
-}
+	const uint8* Source,
+	FCustomDeltaChangelistState* InDeltaChangelistState):
+
+	CustomDeltaChangelistState(InDeltaChangelistState),
+	HistoryStart(0),
+	HistoryEnd(0),
+	CompareIndex(0),
+	StaticBuffer(InRepLayout->CreateShadowBuffer(Source))
+{}
 
 void FRepChangelistState::CountBytes(FArchive& Ar) const
 {
@@ -1014,33 +925,23 @@ void FRepChangelistState::CountBytes(FArchive& Ar) const
 
 FReplicationChangelistMgr::FReplicationChangelistMgr(
 	const TSharedRef<const FRepLayout>& InRepLayout,
-	const uint8* InSource,
-	const UObject* InRepresenting,
-	FCustomDeltaChangelistState* DeltaChangelistState)
-
-	: LastReplicationFrame(0)
-	, LastInitialReplicationFrame(0)
-	, RepChangelistState(InRepLayout, InSource, InRepresenting, DeltaChangelistState)
-{
-}
-
-FReplicationChangelistMgr::~FReplicationChangelistMgr()
-{
-}
+	const uint8* Source,
+	FCustomDeltaChangelistState* DeltaChangelistState):
 
-void FReplicationChangelistMgr::CountBytes(FArchive& Ar) const
+	LastReplicationFrame(0),
+	LastInitialReplicationFrame(0),
+	RepChangelistState(InRepLayout, Source, DeltaChangelistState)
 {
-	RepChangelistState.CountBytes(Ar);
 }
 
-FReceivingRepState::FReceivingRepState(FRepStateStaticBuffer&& InStaticBuffer) 
-	: StaticBuffer(MoveTemp(InStaticBuffer))
+FReceivingRepState::FReceivingRepState(FRepStateStaticBuffer&& InStaticBuffer) :
+	StaticBuffer(MoveTemp(InStaticBuffer))
 {
 }
 
-FRepLayout::FRepLayout()
-	: Flags(ERepLayoutFlags::None)
-	, Owner(nullptr)
+FRepLayout::FRepLayout() :
+	Flags(ERepLayoutFlags::None),
+	Owner(nullptr)
 {}
 
 FRepLayout::~FRepLayout()
@@ -1100,6 +1001,11 @@ void FRepLayout::UpdateChangelistMgr(
 	}
 }
 
+void FReplicationChangelistMgr::CountBytes(FArchive& Ar) const
+{
+	RepChangelistState.CountBytes(Ar);
+}
+
 struct FComparePropertiesSharedParams
 {
 	const bool bIsInitial;
@@ -1107,180 +1013,117 @@ struct FComparePropertiesSharedParams
 	const ERepLayoutFlags Flags;
 	const TArray<FRepParentCmd>& Parents;
 	const TArray<FRepLayoutCmd>& Cmds;
-	FSendingRepState* const RepState;
-	FRepChangelistState* const RepChangelistState;
-	FRepChangedPropertyTracker* const RepChangedPropertyTracker;
-};
-
-struct FComparePropertiesStackParams
-{
-	const FConstRepObjectDataBuffer Data;
-	FRepShadowDataBuffer ShadowData;
-	TArray<uint16>& Changed;
 };
 
 static uint16 CompareProperties_r(
 	const FComparePropertiesSharedParams& SharedParams,
-	FComparePropertiesStackParams& StackParams,
 	const uint16 CmdStart,
 	const uint16 CmdEnd,
+	FRepShadowDataBuffer ShadowData,
+	const FConstRepObjectDataBuffer Data,
+	TArray<uint16>& Changed,
 	uint16 Handle);
 
 static void CompareProperties_Array_r(
 	const FComparePropertiesSharedParams& SharedParams,
-	FComparePropertiesStackParams& StackParams,
+	FRepShadowDataBuffer ShadowData,
+	const FConstRepObjectDataBuffer Data,
+	TArray<uint16>& Changed,
 	const uint16 CmdIndex,
 	const uint16 Handle);
 
 static void CompareRoleProperty(
 	const FComparePropertiesSharedParams& SharedParams,
-	FComparePropertiesStackParams& StackParams,
+	const FConstRepObjectDataBuffer Data,
 	const uint16 RoleOrRemoteRoleIndex,
-	TEnumAsByte<ENetRole>& SavedRoleOrRemoteRole)
+	TEnumAsByte<ENetRole>& SavedRoleOrRemoteRole,
+	TArray<uint16>& Changed)
 {
 	const FRepParentCmd& RoleOrRemoteRoleParent = SharedParams.Parents[RoleOrRemoteRoleIndex];
 	const FRepLayoutCmd& RoleOrRemoteRoleCmd = SharedParams.Cmds[RoleOrRemoteRoleParent.CmdStart];
 	const uint16 Handle = RoleOrRemoteRoleCmd.RelativeHandle;
-	const TEnumAsByte<ENetRole> ActorRoleOrRemoteRole = *(const TEnumAsByte<ENetRole>*)(StackParams.Data + RoleOrRemoteRoleParent).Data;
+	const TEnumAsByte<ENetRole> ActorRoleOrRemoteRole = *(const TEnumAsByte<ENetRole>*)(Data + RoleOrRemoteRoleParent).Data;
 	if (SharedParams.bForceFail || SavedRoleOrRemoteRole != ActorRoleOrRemoteRole)
 	{
 		SavedRoleOrRemoteRole = ActorRoleOrRemoteRole;
-		StackParams.Changed.Add(Handle);
+		Changed.Add(Handle);
 	}
 }
 
 static void CompareRoleProperties(
 	const FComparePropertiesSharedParams& SharedParams,
-	FComparePropertiesStackParams& StackParams)
+	FSendingRepState* RESTRICT RepState,
+	const FConstRepObjectDataBuffer Data,
+	TArray<uint16>& Changed)
 {
-	if (SharedParams.RepState && EnumHasAnyFlags(SharedParams.Flags, ERepLayoutFlags::IsActor))
+	if (RepState && EnumHasAnyFlags(SharedParams.Flags, ERepLayoutFlags::IsActor))
 	{
-		CompareRoleProperty(SharedParams, StackParams, (int32)AActor::ENetFields_Private::RemoteRole, SharedParams.RepState->SavedRemoteRole);
-		CompareRoleProperty(SharedParams, StackParams, (int32)AActor::ENetFields_Private::Role, SharedParams.RepState->SavedRole);
+		CompareRoleProperty(SharedParams, Data, (int32)AActor::ENetFields_Private::RemoteRole, RepState->SavedRemoteRole, Changed);
+		CompareRoleProperty(SharedParams, Data, (int32)AActor::ENetFields_Private::Role, RepState->SavedRole, Changed);
 	}
 }
 
-static void CompareParentProperty(
-	const int32 ParentIndex,
+static void CompareParentProperties(
 	const FComparePropertiesSharedParams& SharedParams,
-	FComparePropertiesStackParams& StackParams)
+	FSendingRepState* RESTRICT RepState,
+	FRepChangelistState* RESTRICT RepChangelistState,
+	const FConstRepObjectDataBuffer Data,
+	TArray<uint16>& Changed)
 {
-	const FRepParentCmd& Parent = SharedParams.Parents[ParentIndex];
-	const bool bIsLifetime = EnumHasAnyFlags(Parent.Flags, ERepParentFlags::IsLifetime);
+	FRepShadowDataBuffer ShadowData = RepChangelistState->StaticBuffer.GetData();
+	check(ShadowData);
 
-	// Active state of a property applies to *all* connections.
-	// If the property is inactive, we can skip comparing it because we know it won't be sent.
-	// Further, this will keep the last active state of the property in the shadow buffer,
-	// meaning the next time the property becomes active it will be sent to all connections.
-	const bool bIsActive = !SharedParams.RepChangedPropertyTracker || SharedParams.RepChangedPropertyTracker->Parents[ParentIndex].Active;
+	FRepChangedPropertyTracker* RepChangedPropertyTracker = (RepState ? RepState->RepChangedPropertyTracker.Get() : nullptr);
 
-	const bool bShouldSkip = !bIsLifetime || !bIsActive || (Parent.Condition == COND_InitialOnly && !SharedParams.bIsInitial);
+	const bool bCheckForRole = EnumHasAnyFlags(SharedParams.Flags, ERepLayoutFlags::IsActor) && RepState != nullptr;
 
-	if (bShouldSkip)
+	for (uint16 ParentIndex = 0; ParentIndex < SharedParams.Parents.Num(); ++ParentIndex)
 	{
-		return;
-	}
-
-	const FRepLayoutCmd& Cmd = SharedParams.Cmds[Parent.CmdStart];
-
-	if (EnumHasAnyFlags(SharedParams.Flags, ERepLayoutFlags::IsActor))
-	{
-		if (UNLIKELY(ParentIndex == (int32)AActor::ENetFields_Private::Role))
-		{
-			CompareRoleProperty(SharedParams, StackParams, (int32)AActor::ENetFields_Private::Role, SharedParams.RepState->SavedRole);
-			return;
-		}
-		if (UNLIKELY(ParentIndex == (int32)AActor::ENetFields_Private::RemoteRole))
-		{
-			CompareRoleProperty(SharedParams, StackParams, (int32)AActor::ENetFields_Private::RemoteRole, SharedParams.RepState->SavedRemoteRole);
-			return;
-		}
-	}
-
-	// Note, Handle - 1 to account for CompareProperties_r incrementing handles.
-	CompareProperties_r(SharedParams, StackParams, Parent.CmdStart, Parent.CmdEnd, Cmd.RelativeHandle - 1);
-}
+		const FRepParentCmd& Parent = SharedParams.Parents[ParentIndex];
 
-static void CompareParentProperties(
-	const FComparePropertiesSharedParams& SharedParams,
-	FComparePropertiesStackParams& StackParams)
-{
-	check(StackParams.ShadowData);
+		const bool bIsLifetime = EnumHasAnyFlags(Parent.Flags, ERepParentFlags::IsLifetime);
 
-#if WITH_PUSH_MODEL
+		// Active state of a property applies to *all* connections.
+		// If the property is inactive, we can skip comparing it because we know it won't be sent.
+		// Further, this will keep the last active state of the property in the shadow buffer,
+		// meaning the next time the property becomes active it will be sent to all connections.
+		const bool bIsActive = !RepChangedPropertyTracker || RepChangedPropertyTracker->Parents[ParentIndex].Active;
 
-	if (UE4PushModelPrivate::FPushModelPerNetDriverState* PushModelObject = UE4_RepLayout_Private::GetPerNetDriverState(SharedParams.RepChangelistState->GetPushModelObjectHandle()))
-	{
-		// Typically, on an initial compare all properties will be dirty anyway.
-		// However, if an Actor is awakened from Dormancy, it's SendingRepState will have been recreated, invalidating
-		// its saved Role and RemoteRole (regardless of whether or not they were changed).
-		// This really shouldn't cause any problems, but it will cause Validation to consistently fail (if GbPushModelValidateProperties
-		// is enabled).
-		// To get around this, and just to be sure our code paths are consistent, we forcibly mark Role and RemoteRole dirty here.
-		if (EnumHasAnyFlags(SharedParams.Flags, ERepLayoutFlags::IsActor) && SharedParams.bIsInitial)
-		{
-			PushModelObject->MarkPropertyDirty((int32)AActor::ENetFields_Private::RemoteRole);
-			PushModelObject->MarkPropertyDirty((int32)AActor::ENetFields_Private::Role);
-		}
+		const bool bShouldSkip = !bIsLifetime || !bIsActive || (Parent.Condition == COND_InitialOnly && !SharedParams.bIsInitial);
 
-		// If we're forcibly comparing all properties, then don't bother checking dirty state.
-		if (SharedParams.bForceFail)
+		if (bShouldSkip)
 		{
-			for (int32 ParentIndex = 0; ParentIndex < SharedParams.Parents.Num(); ++ParentIndex)
-			{
-				CompareParentProperty(ParentIndex, SharedParams, StackParams);
-			}
+			continue;
 		}
 
-		// If we have full push model support, then we only need to check properties that are actually dirty.
-		else if (!GbPushModelValidateProperties && EnumHasAnyFlags(SharedParams.Flags, ERepLayoutFlags::FullPushSupport))
+		if (bCheckForRole)
 		{
-			for (TConstSetBitIterator<> It = PushModelObject->GetDirtyProperties(); It; ++It)
+			if (UNLIKELY(ParentIndex == (int32)AActor::ENetFields_Private::Role))
 			{
-				CompareParentProperty(It.GetIndex(), SharedParams, StackParams);
+				CompareRoleProperty(SharedParams, Data, (int32)AActor::ENetFields_Private::Role, RepState->SavedRole, Changed);
+				continue;
 			}
-		}
-		else
-		{
-#if WITH_PUSH_VALIDATION_SUPPORT
-			int32 StartChangedNum = 0;
-#endif
-
-			for (int32 ParentIndex = 0; ParentIndex < SharedParams.Parents.Num(); ++ParentIndex)
+			else if (UNLIKELY(ParentIndex == (int32)AActor::ENetFields_Private::RemoteRole))
 			{
-				const bool bIsPropertyDirty = !EnumHasAnyFlags(SharedParams.Parents[ParentIndex].Flags, ERepParentFlags::UsePushModel) || PushModelObject->IsPropertyDirty(ParentIndex);
-				if (GbPushModelValidateProperties || bIsPropertyDirty)
-				{
-					CompareParentProperty(ParentIndex, SharedParams, StackParams);
-
-#if WITH_PUSH_VALIDATION_SUPPORT
-					const bool bWasFoundDirtyButNotMarkedDirty = GbPushModelValidateProperties && !bIsPropertyDirty && StackParams.Changed.Num() > StartChangedNum;
-
-					ensureAlwaysMsgf(!bWasFoundDirtyButNotMarkedDirty, TEXT("Push Model Property changed value, but was not marked dirty! Property=%s"), *SharedParams.Parents[ParentIndex].Property->GetPathName());
-
-					StartChangedNum = StackParams.Changed.Num();
-#endif
-				}
+				CompareRoleProperty(SharedParams, Data, (int32)AActor::ENetFields_Private::RemoteRole, RepState->SavedRemoteRole, Changed);
+				continue;
 			}
 		}
-
-		PushModelObject->ResetDirtyStates();
-		return;
-	}
-#endif
-
-	for (int32 ParentIndex = 0; ParentIndex < SharedParams.Parents.Num(); ++ParentIndex)
-	{
-		CompareParentProperty(ParentIndex, SharedParams, StackParams);
+		
+		// Note, Handle - 1 to account for CompareProperties_r incrementing handles.
+		const FRepLayoutCmd& Cmd = SharedParams.Cmds[Parent.CmdStart];
+		const uint16 Handle = Cmd.RelativeHandle;
+		CompareProperties_r(SharedParams, Parent.CmdStart, Parent.CmdEnd, ShadowData, Data, Changed, Handle - 1);
 	}
 }
 
-
 static uint16 CompareProperties_r(
 	const FComparePropertiesSharedParams& SharedParams,
-	FComparePropertiesStackParams& StackParams,
 	const uint16 CmdStart,
 	const uint16 CmdEnd,
+	FRepShadowDataBuffer ShadowData,
+	const FConstRepObjectDataBuffer Data,
+	TArray<uint16>& Changed,
 	uint16 Handle)
 {
 	for (int32 CmdIndex = CmdStart; CmdIndex < CmdEnd; ++CmdIndex)
@@ -1291,26 +1134,17 @@ static uint16 CompareProperties_r(
 
 		++Handle;
 
-		const FConstRepObjectDataBuffer Data = StackParams.Data + Cmd;
-		FRepShadowDataBuffer ShadowData = StackParams.ShadowData + Cmd;
-
 		if (Cmd.Type == ERepLayoutCmdType::DynamicArray)
 		{
-			FComparePropertiesStackParams NewStackParams{
-				Data,
-				ShadowData,
-				StackParams.Changed
-			};
-
 			// Once we hit an array, start using a stack based approach
-			CompareProperties_Array_r(SharedParams, NewStackParams, CmdIndex, Handle);
+			CompareProperties_Array_r(SharedParams, ShadowData + Cmd, Data + Cmd, Changed, CmdIndex, Handle);
 			CmdIndex = Cmd.EndCmd - 1;		// The -1 to handle the ++ in the for loop
 			continue;
 		}
-		else if (SharedParams.bForceFail || !PropertiesAreIdentical(Cmd, ShadowData.Data, Data.Data))
+		else if (SharedParams.bForceFail || !PropertiesAreIdentical(Cmd, (ShadowData + Cmd).Data, (Data + Cmd).Data))
 		{
-			StoreProperty(Cmd, ShadowData.Data, Data.Data);
-			StackParams.Changed.Add(Handle);
+			StoreProperty(Cmd, (ShadowData + Cmd).Data, (Data + Cmd.Offset).Data);
+			Changed.Add(Handle);
 		}
 	}
 
@@ -1319,14 +1153,16 @@ static uint16 CompareProperties_r(
 
 static void CompareProperties_Array_r(
 	const FComparePropertiesSharedParams& SharedParams,
-	FComparePropertiesStackParams& StackParams,
+	FRepShadowDataBuffer ShadowData,
+	const FConstRepObjectDataBuffer Data,
+	TArray<uint16>& Changed,
 	const uint16 CmdIndex,
 	const uint16 Handle)
 {
 	const FRepLayoutCmd& Cmd = SharedParams.Cmds[CmdIndex];
 
-	FScriptArray* ShadowArray = (FScriptArray*)StackParams.ShadowData.Data;
-	FScriptArray* Array = (FScriptArray*)StackParams.Data.Data;
+	FScriptArray* ShadowArray = (FScriptArray*)ShadowData.Data;
+	FScriptArray* Array = (FScriptArray*)Data.Data;
 
 	const uint16 ArrayNum = Array->Num();
 	const uint16 ShadowArrayNum = ShadowArray->Num();
@@ -1351,23 +1187,16 @@ static void CompareProperties_Array_r(
 		{
 			const int32 ArrayElementOffset = i * Cmd.ElementSize;
 			bForceFail = bOldForceFail || i >= ShadowArrayNum;
-
-			FComparePropertiesStackParams NewStackParams{
-				ArrayData + ArrayElementOffset,
-				ShadowArrayData + ArrayElementOffset,
-				ChangedLocal
-			};
-
-			LocalHandle = CompareProperties_r(SharedParams, NewStackParams, CmdIndex + 1, Cmd.EndCmd - 1, LocalHandle);
+			LocalHandle = CompareProperties_r(SharedParams, CmdIndex + 1, Cmd.EndCmd - 1, ShadowArrayData + ArrayElementOffset, ArrayData + ArrayElementOffset, ChangedLocal, LocalHandle);
 		}
 	}
 
 	if (ChangedLocal.Num())
 	{
-		StackParams.Changed.Add(Handle);
-		StackParams.Changed.Add(ChangedLocal.Num());		// This is so we can jump over the array if we need to
-		StackParams.Changed.Append(ChangedLocal);
-		StackParams.Changed.Add(0);
+		Changed.Add(Handle);
+		Changed.Add(ChangedLocal.Num());		// This is so we can jump over the array if we need to
+		Changed.Append(ChangedLocal);
+		Changed.Add(0);
 	}
 	else if (ArrayNum != ShadowArrayNum)
 	{
@@ -1375,9 +1204,9 @@ static void CompareProperties_Array_r(
 		check(ArrayNum < ShadowArrayNum || SharedParams.Cmds[CmdIndex + 1].Type == ERepLayoutCmdType::DynamicArray);
 
 		// Array got smaller, send the array handle to force array size change
-		StackParams.Changed.Add(Handle);
-		StackParams.Changed.Add(0);
-		StackParams.Changed.Add(0);
+		Changed.Add(Handle);
+		Changed.Add(0);
+		Changed.Add(0);
 	}
 }
 
@@ -1409,25 +1238,16 @@ bool FRepLayout::CompareProperties(
 		/*bForceFail=*/ false,
 		Flags,
 		Parents,
-		Cmds,
-		RepState,
-		RepChangelistState,
-		(RepState ? RepState->RepChangedPropertyTracker.Get() : nullptr)
-	};
-
-	FComparePropertiesStackParams StackParams{
-		Data,
-		RepChangelistState->StaticBuffer.GetData(),
-		Changed
+		Cmds
 	};
 
 	if (RepFlags.bRolesOnly)
 	{
-		CompareRoleProperties(SharedParams, StackParams);
+		CompareRoleProperties(SharedParams, RepState, Data, Changed);
 	}
 	else
 	{
-		CompareParentProperties(SharedParams, StackParams);
+		CompareParentProperties(SharedParams, RepState, RepChangelistState, Data, Changed);
 	}
 
 	if (Changed.Num() == 0)
@@ -2939,7 +2759,7 @@ static FGuidReferencesMap* PrepReceivedArray(
 	return NewGuidReferencesArray ? NewGuidReferencesArray->Array : nullptr;
 }
 
-/** Struct containing parameters that don't change or can be shared throughout recursion of ReceiveProperties_r */
+/** Struct containing parameters that don't change throughout recursion of ReceiveProperties_r */
 struct FReceivePropertiesSharedParams
 {
 	const bool bDoChecksum;
@@ -2950,15 +2770,10 @@ struct FReceivePropertiesSharedParams
 	bool& bOutGuidsChanged;
 	const TArray<FRepParentCmd>& Parents;
 	const TArray<FRepLayoutCmd>& Cmds;
-	UObject* OwningObject;
 	uint16 ReadHandle = 0;
-	
-#if WITH_PUSH_MODEL	
-	int32 LastReceivedParent = INDEX_NONE;
-#endif	
 };
 
-/** Struct containing parameters that do change or can't be shared as we recurse into ReceiveProperties_r */
+/** Struct containing parameters that do change as we recurse into ReceiveProperties_r */
 struct FReceivePropertiesStackParams
 {
 	FRepObjectDataBuffer ObjectData;
@@ -3016,22 +2831,11 @@ static bool ReceiveProperties_r(FReceivePropertiesSharedParams& Params, FReceive
 		else
 		{
 			UE_LOG(LogRepProperties, VeryVerbose, TEXT("ReceiveProperties_r: Parent=%d, Cmd=%d"), Cmd.ParentIndex, CmdIndex);
-
-			const FRepParentCmd& Parent = Params.Parents[Cmd.ParentIndex];
-
-#if WITH_PUSH_MODEL
-			if (Cmd.ParentIndex != Params.LastReceivedParent)
-			{
-				// Make sure we mark received properties dirty.
-				// This will make sure that anything we receive can be picked up by other Net Drivers
-				// (like the DemoNetDriver) on clients.
-				Params.LastReceivedParent = Cmd.ParentIndex;
-				FNetPrivatePushIdHelper::MarkPropertyDirty(Params.OwningObject, Params.LastReceivedParent);
-			}
-#endif
-
+				
 			if (ERepLayoutCmdType::DynamicArray == Cmd.Type)
 			{
+				const FRepParentCmd& Parent = Params.Parents[Cmd.ParentIndex];
+
 				// Don't worry about checking the ShadowData for nullptr here.
 				// We're either:
 				//	1. At the top level and it's valid
@@ -3133,6 +2937,7 @@ static bool ReceiveProperties_r(FReceivePropertiesSharedParams& Params, FReceive
 			// TODO: Might be worth doing this before and after ReadNextHandle, or having ReadNextHandle check for errors?
 			if (Params.Bunch.IsError())
 			{
+				const FRepParentCmd& Parent = Params.Parents[Cmd.ParentIndex];
 				UE_LOG(LogRep, Error, TEXT("ReceiveProperties_r: Failed to receive property, BunchIsError - Property=%s, Parent=%d, Cmd=%d, ReadHandle=%d"), *Parent.CachedPropertyName.ToString(), Cmd.ParentIndex, CmdIndex, Params.ReadHandle);
 				return false;
 			}
@@ -3150,7 +2955,7 @@ bool FRepLayout::ReceiveProperties(
 	UActorChannel* OwningChannel,
 	UClass* InObjectClass,
 	FReceivingRepState* RESTRICT RepState,
-	UObject* Object,
+	FRepObjectDataBuffer Data,
 	FNetBitReader& InBunch,
 	bool& bOutHasUnmapped,
 	bool& bOutGuidsChanged,
@@ -3158,7 +2963,6 @@ bool FRepLayout::ReceiveProperties(
 {
 	check(InObjectClass == Owner);
 
-	FRepObjectDataBuffer Data(Object);
 	const bool bEnableRepNotifies = EnumHasAnyFlags(ReceiveFlags, EReceivePropertiesFlags::RepNotifies);
 
 	if (OwningChannel->Connection->InternalAck)
@@ -3185,13 +2989,12 @@ bool FRepLayout::ReceiveProperties(
 		bDoChecksum,
 		// We can skip swapping roles if we're not an Actor layout, or if we've been explicitly told we can skip.
 		EnumHasAnyFlags(ReceiveFlags, EReceivePropertiesFlags::SkipRoleSwap) || !EnumHasAnyFlags(Flags, ERepLayoutFlags::IsActor),
-		bEnableRepNotifies ? &RepState->RepNotifies : nullptr,
+		EnumHasAnyFlags(ReceiveFlags, EReceivePropertiesFlags::RepNotifies) ? &RepState->RepNotifies : nullptr,
 		InBunch,
 		bOutHasUnmapped,
 		bOutGuidsChanged,
 		Parents,
-		Cmds,
-		Object
+		Cmds
 	};
 
 	FReceivePropertiesStackParams StackParams{
@@ -4972,10 +4775,7 @@ static int32 InitFromProperty_r(
 	return RelativeHandle;
 }
 
-static FORCEINLINE uint16 AddParentProperty(
-	TArray<FRepParentCmd>& Parents,
-	FProperty* Property,
-	int32 ArrayIndex)
+static FORCEINLINE uint16 AddParentProperty(TArray<FRepParentCmd>& Parents, FProperty* Property, int32 ArrayIndex)
 {
 	return Parents.Emplace(Property, ArrayIndex);
 }
@@ -5105,11 +4905,7 @@ static void BuildShadowOffsets_r(TArray<FRepLayoutCmd>::TIterator& CmdIt, int32&
 }
 
 template<ERepBuildType ShadowType>
-static void BuildShadowOffsets(
-	UStruct* Owner,
-	TArray<FRepParentCmd>& Parents,
-	TArray<FRepLayoutCmd>& Cmds,
-	int32& ShadowOffset)
+static void BuildShadowOffsets(UStruct* Owner, TArray<FRepParentCmd>& Parents, TArray<FRepLayoutCmd>& Cmds, int32& ShadowOffset)
 {
 	SCOPE_CYCLE_COUNTER(STAT_RepLayout_BuildShadowOffsets);
 
@@ -5323,11 +5119,6 @@ void FRepLayout::InitFromClass(
 
 	Object->GetLifetimeReplicatedProps(LifetimeProps);
 
-	// Tracks the number of (non-delta) lifetime properties so we can check that against our
-	// Push Model Enabled properties.
-	int32 NumberOfLifetimeProperties = 0;
-	int32 NumberOfPushModelProperties = 0;
-
 	// Setup lifetime replicated properties
 	for (int32 i = 0; i < LifetimeProps.Num(); i++)
 	{
@@ -5355,22 +5146,7 @@ void FRepLayout::InitFromClass(
 			Parents[ParentIndex].RepNotifyNumParams = RepNotifyFunc->NumParms;
 		}
 
-		if (!EnumHasAnyFlags(Parents[ParentIndex].Flags, ERepParentFlags::IsCustomDelta))
-		{
-			Parents[ParentIndex].Flags |= ERepParentFlags::IsLifetime;
-			if (LifetimeProps[i].Condition == COND_None)
-			{
-				Parents[ParentIndex].Flags &= ~ERepParentFlags::IsConditional;
-			}
-
-			++NumberOfLifetimeProperties;
-			if (IS_PUSH_MODEL_ENABLED() && LifetimeProps[i].bIsPushBased)
-			{
-				++NumberOfPushModelProperties;
-				Parents[ParentIndex].Flags |= ERepParentFlags::UsePushModel;
-			}
-		}
-		else
+		if (EnumHasAnyFlags(Parents[ParentIndex].Flags, ERepParentFlags::IsCustomDelta))
 		{
 			// We'll track all Custom Lifetime Properties here, and we'll handle Fast Array Serialization
 			// specially.
@@ -5453,6 +5229,16 @@ void FRepLayout::InitFromClass(
 			{
 				LifetimeCustomPropertyState->Add(FLifetimeCustomDeltaProperty(ParentIndex));
 			}
+
+			continue;
+		}
+
+		Parents[ParentIndex].Flags |= ERepParentFlags::IsLifetime;
+
+		if (LifetimeProps[i].Condition == COND_None)
+		{
+			Parents[ParentIndex].Flags &= ~ERepParentFlags::IsConditional;
+		}
 	}
 
 	if (bIsObjectActor)
@@ -5470,13 +5256,6 @@ void FRepLayout::InitFromClass(
 			Parents[(int32)AActor::ENetFields_Private::RemoteRole].Flags |= ERepParentFlags::IsConditional;
 			Parents[(int32)AActor::ENetFields_Private::RemoteRole].Condition = COND_None;
 		}
-	}	
-
-	if (IS_PUSH_MODEL_ENABLED() && NumberOfPushModelProperties > 0)
-	{
-		Flags |= (NumberOfLifetimeProperties == NumberOfPushModelProperties) ?
-			ERepLayoutFlags::FullPushSupport :
-			ERepLayoutFlags::PartialPushSupport;
 	}
 
 	if (!ServerConnection || EnumHasAnyFlags(CreateFlags, ECreateRepLayoutFlags::MaySendProperties))
@@ -6246,7 +6025,7 @@ TSharedPtr<FReplicationChangelistMgr> FRepLayout::CreateReplicationChangelistMgr
 		ShadowStateSource = (const uint8*)InObject;
 	}
 
-	return MakeShareable(new FReplicationChangelistMgr(AsShared(), ShadowStateSource, InObject, DeltaChangelistState));
+	return MakeShareable(new FReplicationChangelistMgr(AsShared(), ShadowStateSource, DeltaChangelistState));
 }
 
 TUniquePtr<FRepState> FRepLayout::CreateRepState(
@@ -6817,9 +6596,6 @@ bool FRepLayout::DeltaSerializeFastArrayProperty(FFastArrayDeltaSerializeParams&
 						TArray<uint16> NewChangelist;
 						for (auto& IDIndexPair : ChangedElements)
 						{
-							NewChangelist.Empty(1);
-							const int32 ArrayElementOffset = ElementSize * IDIndexPair.Idx;
-
 							// Go ahead and do a property compare here, regardless of what we'll actually use below.
 							// This is to prevent issues where someone with an initial / outdated connection doesn't properly
 							// update the changelists in our history, but does update the shadow state inadvertently.
@@ -6828,19 +6604,16 @@ bool FRepLayout::DeltaSerializeFastArrayProperty(FFastArrayDeltaSerializeParams&
 								/*bForceFail=*/ bIsInitial || ShadowArrayItemIsNew[IDIndexPair.Idx],
 								Flags,
 								Parents,
-								Cmds,
-								nullptr,
-								nullptr,
-								nullptr
+								Cmds
 							};
 
-							FComparePropertiesStackParams StackParams{
-								FConstRepObjectDataBuffer(ObjectArrayData + ArrayElementOffset),
-								FRepShadowDataBuffer(ShadowArrayData + ArrayElementOffset),
-								NewChangelist
-							};
+							const int32 ArrayElementOffset = ElementSize * IDIndexPair.Idx;
+
+							FConstRepObjectDataBuffer ElementData(ObjectArrayData + ArrayElementOffset);
+							FRepShadowDataBuffer ElementShadowData(ShadowArrayData + ArrayElementOffset);
+							NewChangelist.Empty(1);
 
-							CompareProperties_r(SharedParams, StackParams, ItemLayoutStart, ItemLayoutEnd, 0);
+							CompareProperties_r(SharedParams, ItemLayoutStart, ItemLayoutEnd, ElementShadowData, ElementData, NewChangelist, 0);
 
 							if (NewChangelist.Num())
 							{
@@ -6851,7 +6624,7 @@ bool FRepLayout::DeltaSerializeFastArrayProperty(FFastArrayDeltaSerializeParams&
 								// to copying them into the shadow state (see FFastArraySerializerItem::operator=).
 								// In that case, we need to make make sure we reset our ID so they can be found the next
 								// time we try to replicate them.
-								int32& ShadowReplicationID = CustomDeltaProperty.GetFastArrayItemReplicationIDMutable(StackParams.ShadowData.Data);
+								int32& ShadowReplicationID = CustomDeltaProperty.GetFastArrayItemReplicationIDMutable(ElementShadowData.Data);
 								ShadowReplicationID = IDIndexPair.ID;
 							}
 						}
@@ -7088,8 +6861,7 @@ bool FRepLayout::DeltaSerializeFastArrayProperty(FFastArrayDeltaSerializeParams&
 					bOutHasUnmapped,
 					bOutGuidsChanged,
 					Parents,
-					Cmds,
-					Object
+					Cmds
 				};
 
 				FReceivePropertiesStackParams StackParams{
diff --git a/Engine/Source/Runtime/Engine/Private/UnrealNetwork.cpp b/Engine/Source/Runtime/Engine/Private/UnrealNetwork.cpp
--- a/Engine/Source/Runtime/Engine/Private/UnrealNetwork.cpp
+++ b/Engine/Source/Runtime/Engine/Private/UnrealNetwork.cpp
@@ -34,7 +34,7 @@ void RegisterReplicatedLifetimeProperty(
 		const uint16 RepIndex = PropertyDescriptor.RepIndex + i;
 		FLifetimeProperty* RegisteredPropertyPtr = OutLifetimeProps.FindByPredicate([&RepIndex](const FLifetimeProperty& Var) { return Var.RepIndex == RepIndex; });
 
-		FLifetimeProperty LifetimeProp(RepIndex, Params.Condition, Params.RepNotifyCondition, Params.bIsPushBased);
+		FLifetimeProperty LifetimeProp(RepIndex, Params.Condition, Params.RepNotifyCondition);
 
 		if (RegisteredPropertyPtr)
 		{
diff --git a/Engine/Source/Runtime/Engine/Public/Net/NetPushModelHelpers.h b/Engine/Source/Runtime/Engine/Public/Net/NetPushModelHelpers.h
deleted file mode 100644
diff --git a/Engine/Source/Runtime/Engine/Public/Net/RepLayout.h b/Engine/Source/Runtime/Engine/Public/Net/RepLayout.h
--- a/Engine/Source/Runtime/Engine/Public/Net/RepLayout.h
+++ b/Engine/Source/Runtime/Engine/Public/Net/RepLayout.h
@@ -21,7 +21,6 @@
 #include "UObject/GCObject.h"
 #include "Containers/StaticBitArray.h"
 #include "Net/GuidReferences.h"
-#include "Net/Core/PushModel/PushModel.h"
 #include "Templates/CopyQualifiersFromTo.h"
 
 class FGuidReferences;
@@ -159,14 +158,10 @@ public:
 	 *
 	 * @see DOREPLIFETIME_ACTIVE_OVERRIDE
 	 *
-	 * @param OwningObject	The object that we're tracking.
-	 * @param RepIndex		Replication index for the Property.
-	 * @param bIsActive		The new Active state.
+	 * @param RepIndex	Replication index for the Property.
+	 * @param bIsActive	The new Active state.
 	 */
-	virtual void SetCustomIsActiveOverride(
-		UObject* OwningObject,
-		const uint16 RepIndex,
-		const bool bIsActive) override;
+	virtual void SetCustomIsActiveOverride(const uint16 RepIndex, const bool bIsActive) override;
 
 	/**
 	 * Sets (or resets) the External Data.
@@ -405,13 +400,10 @@ private:
 	FRepChangelistState(
 		const TSharedRef<const FRepLayout>& InRepLayout,
 		const uint8* Source,
-		const UObject* InRepresenting,
-		struct FCustomDeltaChangelistState* InDeltaChangelistState);
+		struct FCustomDeltaChangelistState* CustomDeltaChangelistState);
 
 public:
 
-	~FRepChangelistState();
-
 	/** The maximum number of individual changelists allowed.*/
 	static const int32 MAX_CHANGE_HISTORY = 64;
 
@@ -437,17 +429,6 @@ public:
 	FRepSerializationSharedInfo SharedSerialization;
 
 	void CountBytes(FArchive& Ar) const;
-
-#if WITH_PUSH_MODEL
-
-	const UE4PushModelPrivate::FPushModelPerNetDriverHandle& GetPushModelObjectHandle() const
-	{
-		return PushModelObjectHandle;
-	}
-
-private:
-	const UE4PushModelPrivate::FPushModelPerNetDriverHandle PushModelObjectHandle;
-#endif
 };
 
 /**
@@ -466,13 +447,10 @@ private:
 	FReplicationChangelistMgr(
 		const TSharedRef<const FRepLayout>& InRepLayout,
 		const uint8* Source,
-		const UObject* InRepresenting,
 		struct FCustomDeltaChangelistState* CustomDeltaChangelistState);
 
 public:
 
-	~FReplicationChangelistMgr();
-
 	FRepChangelistState* GetRepChangelistState() const
 	{
 		return const_cast<FRepChangelistState*>(&RepChangelistState);
@@ -716,7 +694,6 @@ enum class ERepParentFlags : uint32
 	IsZeroConstructible	= (1 << 6),	//! This property is ZeroConstructible.
 	IsFastArray			= (1 << 7), //! This property is a FastArraySerializer. This can't be a ERepLayoutCmdType, because
 									//! these Custom Delta structs will have their inner properties tracked.
-	UsePushModel		= (1 << 8),	//! This object relies on PushModel, and can be skipped when not dirty.
 };
 
 ENUM_CLASS_FLAGS(ERepParentFlags)
@@ -735,8 +712,8 @@ public:
 		CachedPropertyName(InProperty ? InProperty->GetFName() : NAME_None),
 		ArrayIndex(InArrayIndex),
 		ShadowOffset(0),
-		CmdStart(0),
-		CmdEnd(0),
+		CmdStart(0), 
+		CmdEnd(0), 
 		Condition(COND_None),
 		RepNotifyCondition(REPNOTIFY_OnChanged),
 		RepNotifyNumParams(INDEX_NONE),
@@ -1008,10 +985,8 @@ ENUM_CLASS_FLAGS(ECreateRepLayoutFlags);
 
 enum class ERepLayoutFlags : uint8
 {
-	None				= 0,
-	IsActor 			= 1 << 0,	//! This RepLayout is for AActor or a subclass of AActor.
-	PartialPushSupport	= 1 << 1,	//! This RepLayout has some properties that use Push Model and some that don't.
-	FullPushSupport		= 1 << 2,	//! All properties in this RepLayout use Push Model.
+	None = 0,
+	IsActor = 1 << 1,	//! This RepLayout is for AActor or a subclass of AActor.
 };
 ENUM_CLASS_FLAGS(ERepLayoutFlags);
 
@@ -1224,11 +1199,11 @@ public:
 		UActorChannel* OwningChannel,
 		UClass* InObjectClass,
 		FReceivingRepState* RESTRICT RepState,
-		UObject* Object,
+		FRepObjectDataBuffer Data,
 		FNetBitReader& InBunch,
 		bool& bOutHasUnmapped,
 		bool& bOutGuidsChanged,
-		const EReceivePropertiesFlags InFlags) const;
+		const EReceivePropertiesFlags Flags) const;
 
 	/**
 	 * Finds any properties in the Shadow Buffer of the given Rep State that are currently valid
@@ -1461,11 +1436,6 @@ PRAGMA_ENABLE_DEPRECATION_WARNINGS
 		return 0 == Parents.Num();
 	}
 
-	const int32 GetNumParents() const
-	{
-		return Parents.Num();
-	}
-
 	void CountBytes(FArchive& Ar) const;
 
 private:
diff --git a/Engine/Source/Runtime/Engine/Public/Net/UnrealNetwork.h b/Engine/Source/Runtime/Engine/Public/Net/UnrealNetwork.h
--- a/Engine/Source/Runtime/Engine/Public/Net/UnrealNetwork.h
+++ b/Engine/Source/Runtime/Engine/Public/Net/UnrealNetwork.h
@@ -73,9 +73,6 @@ struct ENGINE_API FDoRepLifetimeParams
 	 * properly set up to handle RepNotifies.
 	 */
 	ELifetimeRepNotifyCondition RepNotifyCondition = REPNOTIFY_OnChanged;
-	
-	/** Whether or not this property uses Push Model. See PushModel.h */
-	bool bIsPushBased = false;
 };
 
 namespace NetworkingPrivate
@@ -174,15 +171,7 @@ static FProperty* GetReplicatedProperty(const UClass* CallingClass, const UClass
 { \
 	static const bool bIsValid_##c_##v = ValidateReplicatedClassInheritance(StaticClass(), c::StaticClass(), TEXT(#v)); \
 	const TCHAR* DoRepPropertyName_##c_##v(TEXT(#v)); \
-	const NetworkingPrivate::FRepPropertyDescriptor PropertyDescriptor_##c_##v(DoRepPropertyName_##c_##v, (int32)c::ENetFields_Private::v, 1); \
-	RegisterReplicatedLifetimeProperty(PropertyDescriptor_##c_##v, OutLifetimeProps, params); \
-}
-
-#define DOREPLIFETIME_WITH_PARAMS_FAST_STATIC_ARRAY(c,v,params) \
-{ \
-	static const bool bIsValid_##c_##v = ValidateReplicatedClassInheritance(StaticClass(), c::StaticClass(), TEXT(#v)); \
-	const TCHAR* DoRepPropertyName_##c_##v(TEXT(#v)); \
-	const NetworkingPrivate::FRepPropertyDescriptor PropertyDescriptor_##c_##v(DoRepPropertyName_##c_##v, (int32)c::ENetFields_Private::v##_STATIC_ARRAY, (int32)c::EArrayDims_Private::v); \
+	const NetworkingPrivate::FRepPropertyDescriptor PropertyDescriptor_##c_##v(DoRepPropertyName_##c_##v, NetworkingPrivate::Net_##c::NETFIELD_##v, NetworkingPrivate::Net_##c::ARRAYDIM_##v); \
 	RegisterReplicatedLifetimeProperty(PropertyDescriptor_##c_##v, OutLifetimeProps, params); \
 }
 
@@ -230,25 +219,19 @@ static FProperty* GetReplicatedProperty(const UClass* CallingClass, const UClass
 #define DOREPLIFETIME_ACTIVE_OVERRIDE_FAST(c,v,active) \
 { \
 	static const bool bIsValid_##c_##v = ValidateReplicatedClassInheritance(StaticClass(), c::StaticClass(), TEXT(#v)); \
-	ChangedPropertyTracker.SetCustomIsActiveOverride((int32)c::ENetFields_Private::v, active); \
-}
-
-#define DOREPLIFETIME_ACTIVE_OVERRIDE_FAST_STATIC_ARRAY(c,v,active) \
-{ \
-	static const bool bIsValid_##c_##v = ValidateReplicatedClassInheritance(StaticClass(), c::StaticClass(), TEXT(#v)); \
-	for (int32 i = 0; i < (int32)c::EArrayDims_Private::v; ++i) \
+	for (int32 i = 0; i < NetworkingPrivate::Net_##c::ARRAYDIM_##v; i++) \
 	{ \
-		ChangedPropertyTracker.SetCustomIsActiveOverride((int32)c::ENetFields_Private::v##_STATIC_ARRAY + i, active); \
+		ChangedPropertyTracker.SetCustomIsActiveOverride(NetworkingPrivate::Net_##c::NETFIELD_##v + i, active); \
 	} \
 }
 
-#define DOREPLIFETIME_ACTIVE_OVERRIDE(c,v,active) \
-{ \
+#define DOREPLIFETIME_ACTIVE_OVERRIDE(c,v,active)	\
+{													\
 	static FProperty* sp##v = GetReplicatedProperty(StaticClass(), c::StaticClass(),GET_MEMBER_NAME_CHECKED(c,v)); \
-	for (int32 i = 0; i < sp##v->ArrayDim; i++) \
-	{ \
-		ChangedPropertyTracker.SetCustomIsActiveOverride(this, sp##v->RepIndex + i, active); \
-	} \
+	for ( int32 i = 0; i < sp##v->ArrayDim; i++ )											\
+	{																						\
+		ChangedPropertyTracker.SetCustomIsActiveOverride( sp##v->RepIndex + i, active );	\
+	}																						\
 }
 
 UE_DEPRECATED(4.24, "Please use the RESET_REPLIFETIME_CONDITION macro")
@@ -294,15 +277,7 @@ DisableReplicatedLifetimeProperty(StaticClass(), c::StaticClass(), GET_MEMBER_NA
 { \
 	static const bool bIsValid_##c_##v = ValidateReplicatedClassInheritance(StaticClass(), c::StaticClass(), TEXT(#v)); \
 	const TCHAR* DoRepPropertyName_##c_##v(TEXT(#v)); \
-	const NetworkingPrivate::FRepPropertyDescriptor PropertyDescriptor_##c_##v(DoRepPropertyName_##c_##v, (int32)c::ENetFields_Private::v, 1); \
-	DisableReplicatedLifetimeProperty(PropertyDescriptor_##c_##v, OutLifetimeProps); \
-}
-
-#define DISABLE_REPLICATED_PROPERTY_FAST_STATIC_ARRAY(c,v) \
-{ \
-	static const bool bIsValid_##c_##v = ValidateReplicatedClassInheritance(StaticClass(), c::StaticClass(), TEXT(#v)); \
-	const TCHAR* DoRepPropertyName_##c_##v(TEXT(#v)); \
-	const NetworkingPrivate::FRepPropertyDescriptor PropertyDescriptor_##c_##v(DoRepPropertyName_##c_##v, (int32)c::ENetFields_Private::v##_STATIC_ARRAY, (int32)c::EArrayDims_Private::v); \
+	const NetworkingPrivate::FRepPropertyDescriptor PropertyDescriptor_##c_##v(DoRepPropertyName_##c_##v, NetworkingPrivate::Net_##c::NETFIELD_##v, NetworkingPrivate::Net_##c::ARRAYDIM_##v); \
 	DisableReplicatedLifetimeProperty(PropertyDescriptor_##c_##v, OutLifetimeProps); \
 }
 
@@ -316,7 +291,7 @@ DisableAllReplicatedPropertiesOfClass(StaticClass(), c::StaticClass(), SuperClas
 { \
 	static const bool bIsValid_##c_##v = ValidateReplicatedClassInheritance(StaticClass(), c::StaticClass(), TEXT("DISABLE_ALL_CLASS_REPLICATED_PROPERTIES")); \
 	const TCHAR* DoRepPropertyName_##c(TEXT(#c)); \
-	const NetworkingPrivate::FRepClassDescriptor ClassDescriptor_##c(DoRepPropertyName_##c, (int32)c::ENetFields_Private::NETFIELD_REP_START, (int32)c::ENetFields_Private::NETFIELD_REP_END); \
+	const NetworkingPrivate::FRepClassDescriptor ClassDescriptor_##c(DoRepPropertyName_##c, NetworkingPrivate::Net_##c::NETFIELD_REP_START, NetworkingPrivate::Net_##c::NETFIELD_REP_END); \
 	DisableAllReplicatedPropertiesOfClass(StaticClass(), c::StaticClass(), SuperClassBehavior, OutLifetimeProps); \
 }	
 
@@ -339,20 +314,11 @@ ENGINE_API void DisableAllReplicatedPropertiesOfClass(const NetworkingPrivate::F
 { \
 	static const bool bIsValid_##c_##v = ValidateReplicatedClassInheritance(StaticClass(), c::StaticClass(), TEXT(#v)); \
 	const TCHAR* DoRepPropertyName_##c_##v(TEXT(#v)); \
-	const NetworkingPrivate::FRepPropertyDescriptor PropertyDescriptor_##c_##v(DoRepPropertyName_##c_##v, (int32)c::ENetFields_Private::v, 1); \
-	ResetReplicatedLifetimeProperty(StaticClass(), c::StaticClass(), GET_MEMBER_NAME_CHECKED(c,v), cond, OutLifetimeProps); \
-}
-
-#define RESET_REPLIFETIME_CONDITION_FAST_STATIC_ARRAY(c,v,cond) \
-{ \
-	static const bool bIsValid_##c_##v = ValidateReplicatedClassInheritance(StaticClass(), c::StaticClass(), TEXT(#v)); \
-	const TCHAR* DoRepPropertyName_##c_##v(TEXT(#v)); \
-	const NetworkingPrivate::FRepPropertyDescriptor PropertyDescriptor_##c_##v(DoRepPropertyName_##c_##v, (int32)c::ENetFields_Private::v##_STATIC_ARRAY, (int32)c::EArrayDims_Private::v); \
+	const NetworkingPrivate::FRepPropertyDescriptor PropertyDescriptor_##c_##v(DoRepPropertyName_##c_##v, NetworkingPrivate::Net_##c::NETFIELD_##v, NetworkingPrivate::Net_##c::ARRAYDIM_##v); \
 	ResetReplicatedLifetimeProperty(StaticClass(), c::StaticClass(), GET_MEMBER_NAME_CHECKED(c,v), cond, OutLifetimeProps); \
 }
 
 #define RESET_REPLIFETIME_FAST(c,v) RESET_REPLIFETIME_CONDITION_FAST(c, v, COND_None)
-#define RESET_REPLIFETIME_FAST_STATIC_ARRAY(c,v) RESET_REPLIFETIME_FAST_STATIC_ARRAY(c, v, COND_None)
 
 ENGINE_API void ResetReplicatedLifetimeProperty(
 	const UClass* ThisClass,
diff --git a/Engine/Source/Runtime/Net/Core/Private/Net/Core/PushModel/PushModel.cpp b/Engine/Source/Runtime/Net/Core/Private/Net/Core/PushModel/PushModel.cpp
deleted file mode 100644
diff --git a/Engine/Source/Runtime/Net/Core/Private/Net/Core/PushModel/Types/PushModelPerNetDriverState.h b/Engine/Source/Runtime/Net/Core/Private/Net/Core/PushModel/Types/PushModelPerNetDriverState.h
deleted file mode 100644
diff --git a/Engine/Source/Runtime/Net/Core/Private/Net/Core/PushModel/Types/PushModelPerObjectState.h b/Engine/Source/Runtime/Net/Core/Private/Net/Core/PushModel/Types/PushModelPerObjectState.h
deleted file mode 100644
diff --git a/Engine/Source/Runtime/Net/Core/Private/Net/Core/PushModel/Types/PushModelUtils.h b/Engine/Source/Runtime/Net/Core/Private/Net/Core/PushModel/Types/PushModelUtils.h
deleted file mode 100644
diff --git a/Engine/Source/Runtime/Net/Core/Public/Net/Core/PushModel/PushModel.h b/Engine/Source/Runtime/Net/Core/Public/Net/Core/PushModel/PushModel.h
deleted file mode 100644
diff --git a/Engine/Source/Runtime/Net/Core/Public/Net/Core/PushModel/PushModelMacros.h b/Engine/Source/Runtime/Net/Core/Public/Net/Core/PushModel/PushModelMacros.h
deleted file mode 100644
