Track all of physx allocations. Added some unused tags for higher level allocations which the allocator currently ignores, will use a solution later to track this, but checking in anyway to avoid having to find it again

diff --git a/Engine/Source/Runtime/Core/Private/HAL/LowLevelMemTracker.cpp b/Engine/Source/Runtime/Core/Private/HAL/LowLevelMemTracker.cpp
--- a/Engine/Source/Runtime/Core/Private/HAL/LowLevelMemTracker.cpp
+++ b/Engine/Source/Runtime/Core/Private/HAL/LowLevelMemTracker.cpp
@@ -87,6 +87,11 @@ DECLARE_LLM_MEMORY_STAT(TEXT("GC"), STAT_GCLLM, STATGROUP_LLMFULL);
 DECLARE_LLM_MEMORY_STAT(TEXT("UI"), STAT_UILLM, STATGROUP_LLMFULL);
 DECLARE_LLM_MEMORY_STAT(TEXT("Physics"), STAT_PhysicsLLM, STATGROUP_LLMFULL);
 DECLARE_LLM_MEMORY_STAT(TEXT("PhysX"), STAT_PhysXLLM, STATGROUP_LLMFULL);
+DECLARE_LLM_MEMORY_STAT(TEXT("PhysXGeometry"), STAT_PhysXGeometryLLM, STATGROUP_LLMFULL);
+DECLARE_LLM_MEMORY_STAT(TEXT("PhysXLandscape"), STAT_PhysXLandscapeLLM, STATGROUP_LLMFULL);
+DECLARE_LLM_MEMORY_STAT(TEXT("PhysXTrimesh"), STAT_PhysXTrimeshLLM, STATGROUP_LLMFULL);
+DECLARE_LLM_MEMORY_STAT(TEXT("PhysXConvex"), STAT_PhysXConvexLLM, STATGROUP_LLMFULL);
+DECLARE_LLM_MEMORY_STAT(TEXT("PhysXAllocator"), STAT_PhysXAllocatorLLM, STATGROUP_LLMFULL);
 DECLARE_LLM_MEMORY_STAT(TEXT("Chaos"), STAT_ChaosLLM, STATGROUP_LLMFULL);
 DECLARE_LLM_MEMORY_STAT(TEXT("ChaosGeometry"), STAT_ChaosGeometryLLM, STATGROUP_LLMFULL);
 DECLARE_LLM_MEMORY_STAT(TEXT("ChaosAcceleration"), STAT_ChaosAccelerationLLM, STATGROUP_LLMFULL);
diff --git a/Engine/Source/Runtime/Core/Public/HAL/LowLevelMemTracker.h b/Engine/Source/Runtime/Core/Public/HAL/LowLevelMemTracker.h
--- a/Engine/Source/Runtime/Core/Public/HAL/LowLevelMemTracker.h
+++ b/Engine/Source/Runtime/Core/Public/HAL/LowLevelMemTracker.h
@@ -186,6 +186,11 @@ enum class ELLMTagSet : uint8
 	macro(UI,									"UI",							GET_STATFNAME(STAT_UILLM),									GET_STATFNAME(STAT_UISummaryLLM),				-1)\
 	macro(Physics,								"Physics",						GET_STATFNAME(STAT_PhysicsLLM),								GET_STATFNAME(STAT_PhysicsSummaryLLM),			-1)\
 	macro(PhysX,								"PhysX",						GET_STATFNAME(STAT_PhysXLLM),								GET_STATFNAME(STAT_PhysXSummaryLLM),			ELLMTag::Physics)\
+	macro(PhysXGeometry,						"PhysXGeometry",				GET_STATFNAME(STAT_PhysXGeometryLLM),						GET_STATFNAME(STAT_PhysXSummaryLLM),			ELLMTag::Physics)\
+	macro(PhysXTrimesh,							"PhysXTrimesh",					GET_STATFNAME(STAT_PhysXTrimeshLLM),						GET_STATFNAME(STAT_PhysXSummaryLLM),			ELLMTag::Physics)\
+	macro(PhysXConvex,							"PhysXConvex",					GET_STATFNAME(STAT_PhysXConvexLLM),							GET_STATFNAME(STAT_PhysXSummaryLLM),			ELLMTag::Physics)\
+	macro(PhysXAllocator,						"PhysXAllocator",				GET_STATFNAME(STAT_PhysXAllocatorLLM),						GET_STATFNAME(STAT_PhysXSummaryLLM),			ELLMTag::Physics)\
+	macro(PhysXLandscape,						"PhysXLandscape",				GET_STATFNAME(STAT_PhysXLandscapeLLM),						GET_STATFNAME(STAT_PhysXSummaryLLM),			ELLMTag::Physics)\
 	macro(Chaos,								"Chaos",						GET_STATFNAME(STAT_ChaosLLM),								GET_STATFNAME(STAT_ChaosSummaryLLM),			ELLMTag::Physics)\
 	macro(ChaosGeometry,						"ChaosGeometry",				GET_STATFNAME(STAT_ChaosGeometryLLM),						GET_STATFNAME(STAT_ChaosSummaryLLM),			ELLMTag::Physics)\
 	macro(ChaosAcceleration,					"ChaosAcceleration",			GET_STATFNAME(STAT_ChaosAccelerationLLM),					GET_STATFNAME(STAT_ChaosSummaryLLM),			ELLMTag::Physics)\
diff --git a/Engine/Source/Runtime/Engine/Private/PhysicsEngine/PhysScene_PhysX.cpp b/Engine/Source/Runtime/Engine/Private/PhysicsEngine/PhysScene_PhysX.cpp
--- a/Engine/Source/Runtime/Engine/Private/PhysicsEngine/PhysScene_PhysX.cpp
+++ b/Engine/Source/Runtime/Engine/Private/PhysicsEngine/PhysScene_PhysX.cpp
@@ -435,6 +435,7 @@ FPhysScene_PhysX::FPhysScene_PhysX(const AWorldSettings* Settings)
 #endif
 #endif 
 {
+	LLM_SCOPE(ELLMTag::PhysX);
 	LineBatcher = NULL;
 	OwningWorld = NULL;
 #if WITH_PHYSX
diff --git a/Engine/Source/Runtime/Engine/Private/PhysicsEngine/PhysXSupport.cpp b/Engine/Source/Runtime/Engine/Private/PhysicsEngine/PhysXSupport.cpp
--- a/Engine/Source/Runtime/Engine/Private/PhysicsEngine/PhysXSupport.cpp
+++ b/Engine/Source/Runtime/Engine/Private/PhysicsEngine/PhysXSupport.cpp
@@ -505,27 +505,47 @@ FPhysXCookingDataReader::FPhysXCookingDataReader( FByteBulkData& InBulkData, FBo
 	ConvexMeshes.Empty(NumConvexElementsCooked);
 	for( int32 ElementIndex = 0; ElementIndex < NumConvexElementsCooked; ElementIndex++ )
 	{
-		PxConvexMesh* ConvexMesh = ReadConvexMesh( Ar, DataPtr, InBulkData.GetBulkDataSize() );
+		//NOTE: these tags currently get ignored because of low level allocator
+		PxConvexMesh* ConvexMesh;
+		{
+			LLM_SCOPE(ELLMTag::PhysXConvex);
+			ConvexMesh = ReadConvexMesh(Ar,DataPtr,InBulkData.GetBulkDataSize());
+		}
+		
 		ConvexMeshes.Add( ConvexMesh );
 	}
 
 	ConvexMeshesNegX.Empty(NumMirroredElementsCooked);
 	for( int32 ElementIndex = 0; ElementIndex < NumMirroredElementsCooked; ElementIndex++ )
 	{
-		PxConvexMesh* ConvexMeshNegX = ReadConvexMesh( Ar, DataPtr, InBulkData.GetBulkDataSize() );
+		PxConvexMesh* ConvexMeshNegX;
+		{
+			LLM_SCOPE(ELLMTag::PhysXConvex);
+			ConvexMeshNegX = ReadConvexMesh( Ar, DataPtr, InBulkData.GetBulkDataSize() );
+		}
+		
 		ConvexMeshesNegX.Add( ConvexMeshNegX );
 	}
 
 	TriMeshes.Empty(NumTriMeshesCooked);
 	for(int32 ElementIndex = 0; ElementIndex < NumTriMeshesCooked; ++ElementIndex)
 	{
-		PxTriangleMesh* TriMesh = ReadTriMesh( Ar, DataPtr, InBulkData.GetBulkDataSize() );
+		PxTriangleMesh* TriMesh;
+		{
+			LLM_SCOPE(ELLMTag::PhysXTrimesh);
+			TriMesh = ReadTriMesh( Ar, DataPtr, InBulkData.GetBulkDataSize() );
+		}
+		
 		TriMeshes.Add(TriMesh);
 	}
 
 	// Init UVInfo pointer
 	check(UVInfo);
-	Ar << *UVInfo;
+	{
+		LLM_SCOPE(ELLMTag::PhysXTrimesh);
+		Ar << *UVInfo;
+	}
+	
 
 	InBulkData.Unlock();
 }
diff --git a/Engine/Source/Runtime/Landscape/Private/LandscapeCollision.cpp b/Engine/Source/Runtime/Landscape/Private/LandscapeCollision.cpp
--- a/Engine/Source/Runtime/Landscape/Private/LandscapeCollision.cpp
+++ b/Engine/Source/Runtime/Landscape/Private/LandscapeCollision.cpp
@@ -550,6 +550,9 @@ void ULandscapeHeightfieldCollisionComponent::CreateCollisionObject()
 {
 #if WITH_CHAOS
 	LLM_SCOPE(ELLMTag::ChaosLandscape);
+#else
+	//NOTE: this currently gets ignored because of low level allocator
+	LLM_SCOPE(ELLMTag::PhysXLandscape);
 #endif
 	// If we have not created a heightfield yet - do it now.
 	if (!IsValidRef(HeightfieldRef))
diff --git a/Engine/Source/Runtime/PhysicsCore/Public/PhysXSupportCore.h b/Engine/Source/Runtime/PhysicsCore/Public/PhysXSupportCore.h
--- a/Engine/Source/Runtime/PhysicsCore/Public/PhysXSupportCore.h
+++ b/Engine/Source/Runtime/PhysicsCore/Public/PhysXSupportCore.h
@@ -158,6 +158,7 @@ public:
 
 		return (uint8*)AllocationHeader + sizeof(FPhysXAllocationHeader);
 #else
+		LLM_SCOPE(ELLMTag::PhysXAllocator);
 		void* ptr = FMemory::Malloc(size, 16);
 #if PHYSX_MEMORY_STAT_ONLY
 		INC_DWORD_STAT_BY(STAT_MemoryPhysXTotalAllocationSize, FMemory::GetAllocSize(ptr));
