Added smoothing to replicated server world time delta

This ensures that spikes in the percieved server world time caused by delayed replication updates are not directly and immediately reflected in the client calls to AGameStateBase::GetServerWorldTimeSeconds() which would be problematic for systems that want to work relative to server world time.

diff --git a/Engine/Source/Runtime/Engine/Classes/GameFramework/GameStateBase.h b/Engine/Source/Runtime/Engine/Classes/GameFramework/GameStateBase.h
--- a/Engine/Source/Runtime/Engine/Classes/GameFramework/GameStateBase.h
+++ b/Engine/Source/Runtime/Engine/Classes/GameFramework/GameStateBase.h
@@ -157,6 +157,11 @@ protected:
 	/** Handle for efficient management of the UpdateServerTimeSeconds timer */
 	FTimerHandle TimerHandle_UpdateServerTimeSeconds;
 
+	/** Cumulative sum of computed server world time deltas for smoothed-averaging */
+	double SumServerWorldTimeSecondsDelta;
+	/** The number of server world time deltas accumulated in SumServerWorldTimeSecondsDelta - used for computing the mean */
+	uint32 NumServerWorldTimeSecondsDeltas;
+
 private:
 	// Hidden functions that don't make sense to use on this class.
 	HIDE_ACTOR_TRANSFORM_FUNCTIONS();
diff --git a/Engine/Source/Runtime/Engine/Private/GameStateBase.cpp b/Engine/Source/Runtime/Engine/Private/GameStateBase.cpp
--- a/Engine/Source/Runtime/Engine/Private/GameStateBase.cpp
+++ b/Engine/Source/Runtime/Engine/Private/GameStateBase.cpp
@@ -31,6 +31,9 @@ AGameStateBase::AGameStateBase(const FObjectInitializer& ObjectInitializer)
 
 	// Default to every few seconds.
 	ServerWorldTimeSecondsUpdateFrequency = 0.1f;
+
+	SumServerWorldTimeSecondsDelta = 0.0;
+	NumServerWorldTimeSecondsDeltas = 0;
 }
 
 const AGameModeBase* AGameStateBase::GetDefaultGameMode() const
@@ -158,7 +161,30 @@ void AGameStateBase::OnRep_ReplicatedWorldTimeSeconds()
 	UWorld* World = GetWorld();
 	if (World)
 	{
-		ServerWorldTimeSecondsDelta = ReplicatedWorldTimeSeconds - World->GetTimeSeconds();
+		const float ServerWorldTimeDelta = ReplicatedWorldTimeSeconds - World->GetTimeSeconds();
+
+		// Accumulate the computed server world delta
+		SumServerWorldTimeSecondsDelta += ServerWorldTimeDelta;
+		NumServerWorldTimeSecondsDeltas += 1.0;
+
+		// Reset the accumulated values to ensure that we remain representative of the current delta
+		if (NumServerWorldTimeSecondsDeltas > 250)
+		{
+			SumServerWorldTimeSecondsDelta /= NumServerWorldTimeSecondsDeltas;
+			NumServerWorldTimeSecondsDeltas = 1;
+		}
+
+		double TargetWorldTimeSecondsDelta = SumServerWorldTimeSecondsDelta / NumServerWorldTimeSecondsDeltas;
+
+		// Smoothly interpolate towards the new delta if we've already got one to avoid significant spikes
+		if (ServerWorldTimeSecondsDelta == 0.0)
+		{
+			ServerWorldTimeSecondsDelta = TargetWorldTimeSecondsDelta;
+		}
+		else
+		{
+			ServerWorldTimeSecondsDelta += (TargetWorldTimeSecondsDelta - ServerWorldTimeSecondsDelta) * 0.5;
+		}
 	}
 }
 
