Fixes for CI failures
+ Added an include for UTextProperty
+ Added an LexToString for EEdGraphPinDirection

diff --git a/Engine/Plugins/FX/Niagara/Source/NiagaraEditor/Private/NiagaraEditorUtilities.cpp b/Engine/Plugins/FX/Niagara/Source/NiagaraEditor/Private/NiagaraEditorUtilities.cpp
--- a/Engine/Plugins/FX/Niagara/Source/NiagaraEditor/Private/NiagaraEditorUtilities.cpp
+++ b/Engine/Plugins/FX/Niagara/Source/NiagaraEditor/Private/NiagaraEditorUtilities.cpp
@@ -51,6 +51,7 @@
 #include "AssetToolsModule.h"
 #include "Subsystems/AssetEditorSubsystem.h"
 #include "Framework/Commands/GenericCommands.h"
+#include "UObject/TextProperty.h"
 
 #define LOCTEXT_NAMESPACE "FNiagaraEditorUtilities"
 
diff --git a/Engine/Source/Runtime/Engine/Classes/EdGraph/EdGraphNode.h b/Engine/Source/Runtime/Engine/Classes/EdGraph/EdGraphNode.h
--- a/Engine/Source/Runtime/Engine/Classes/EdGraph/EdGraphNode.h
+++ b/Engine/Source/Runtime/Engine/Classes/EdGraph/EdGraphNode.h
@@ -81,6 +81,22 @@ enum EEdGraphPinDirection
 	EGPD_MAX,
 };
 
+inline const TCHAR* const LexToString(const EEdGraphPinDirection State)
+{
+	switch (State)
+	{
+	case EEdGraphPinDirection::EGPD_Input:
+		return TEXT("Input");
+	case EEdGraphPinDirection::EGPD_Output:
+		return TEXT("Output");
+	default:
+		break;
+	}
+
+	checkf(false, TEXT("Missing EEdGraphPinDirection Type: %d"), static_cast<const int32>(State));
+	return TEXT("");
+}
+
 /** Enum used to define what container type a pin represents. */
 UENUM()
 enum class EPinContainerType : uint8
