Add more context to checks in TrackNetFieldExports.

This will make it easier to find troublesome assets or classes just from logging, and hopefully make it easier to create repros if these issues show up.

diff --git a/Engine/Source/Runtime/Engine/Private/PackageMapClient.cpp b/Engine/Source/Runtime/Engine/Private/PackageMapClient.cpp
--- a/Engine/Source/Runtime/Engine/Private/PackageMapClient.cpp
+++ b/Engine/Source/Runtime/Engine/Private/PackageMapClient.cpp
@@ -1327,8 +1327,17 @@ void UPackageMapClient::AddNetFieldExportGroup( const FString& PathName, TShared
 void UPackageMapClient::TrackNetFieldExport( FNetFieldExportGroup* NetFieldExportGroup, const int32 NetFieldExportHandle )
 {
 	check(Connection->InternalAck);
-	check( NetFieldExportHandle >= 0 );
-	check( NetFieldExportGroup->NetFieldExports[NetFieldExportHandle].Handle == NetFieldExportHandle );
+	check(NetFieldExportGroup);
+
+	checkf(NetFieldExportGroup->NetFieldExports.IsValidIndex(NetFieldExportHandle),
+		TEXT("Invalid NetFieldExportHandle. GroupPath = %s, NumExports = %d, ExportHandle = %d"),
+		*(NetFieldExportGroup->PathName), NetFieldExportGroup->NetFieldExports.Num(), NetFieldExportHandle);
+
+	checkf(NetFieldExportGroup->NetFieldExports[NetFieldExportHandle].Handle == NetFieldExportHandle,
+		TEXT("NetFieldExportHandle Mismatch. GroupPath = %s, NumExports = %d, ExportHandle = %d, Expected Handle = %d"),
+		*(NetFieldExportGroup->PathName), NetFieldExportGroup->NetFieldExports.Num(), NetFieldExportHandle, NetFieldExportGroup->NetFieldExports[NetFieldExportHandle].Handle);
+
+
 	NetFieldExportGroup->NetFieldExports[NetFieldExportHandle].bExported = true;
 
 	const uint64 CmdHandle = ( ( uint64 )NetFieldExportGroup->PathNameIndex ) << 32 | ( uint64 )NetFieldExportHandle;
