Integrating CL# 10943569 from Dev-Enterprise
Fixing static mesh lightmap coordinate index being stuck at 0 with when there are generated LODs

diff --git a/Engine/Source/Runtime/Engine/Private/StaticMesh.cpp b/Engine/Source/Runtime/Engine/Private/StaticMesh.cpp
--- a/Engine/Source/Runtime/Engine/Private/StaticMesh.cpp
+++ b/Engine/Source/Runtime/Engine/Private/StaticMesh.cpp
@@ -5968,7 +5968,7 @@ void UStaticMesh::EnforceLightmapRestrictions(bool bUseRenderData)
 	// Legacy content may contain a lightmap resolution of 0, which was valid when vertex lightmaps were supported, but not anymore with only texture lightmaps
 	LightMapResolution = FMath::Max(LightMapResolution, 4);
 
-	// Lightmass only supports 4 UVs
+	// Lightmass only supports 4 UVs from Lightmass::MAX_TEXCOORDS
 	int32 NumUVs = 4;
 
 #if !WITH_EDITORONLY_DATA
@@ -6003,28 +6003,35 @@ void UStaticMesh::EnforceLightmapRestrictions(bool bUseRenderData)
 	{
 		for (int32 LODIndex = 0; LODIndex < GetNumSourceModels(); ++LODIndex)
 		{
-			if (const FMeshDescription* MeshDescription = GetMeshDescription(LODIndex))
-			{
-				const TVertexInstanceAttributesConstRef<FVector2D> UVChannels = MeshDescription->VertexInstanceAttributes().GetAttributesRef<FVector2D>(MeshAttribute::VertexInstance::TextureCoordinate);
+			//If the LOD is generated we validate the BaseLODModel instead as the generated LOD is not available before the build and has the same UV properties as its base LOD.
+			const bool bIsGeneratedLOD = !IsMeshDescriptionValid(LODIndex) || IsReductionActive(LODIndex);
+			const int32 SourceLOD = bIsGeneratedLOD ? GetSourceModel(LODIndex).ReductionSettings.BaseLODModel : LODIndex;
 
-				// skip empty LODs
-				if (UVChannels.GetNumElements() > 0)
+			if (!bIsGeneratedLOD || LODIndex == SourceLOD)
+			{
+				if (const FMeshDescription* MeshDescription = GetMeshDescription(SourceLOD))
 				{
-					int NumChannelsInLOD = UVChannels.GetNumIndices();
-					const FStaticMeshSourceModel& SourceModel = GetSourceModel(LODIndex);
+					const TVertexInstanceAttributesConstRef<FVector2D> UVChannels = MeshDescription->VertexInstanceAttributes().GetAttributesRef<FVector2D>(MeshAttribute::VertexInstance::TextureCoordinate);
 
-					if (SourceModel.BuildSettings.bGenerateLightmapUVs)
+					// skip empty/stripped LODs
+					if (UVChannels.GetNumElements() > 0)
 					{
-						NumChannelsInLOD = FMath::Max(NumChannelsInLOD, SourceModel.BuildSettings.DstLightmapIndex + 1);
-					}
+						int NumChannelsInLOD = UVChannels.GetNumIndices();
+						const FStaticMeshSourceModel& SourceModel = GetSourceModel(SourceLOD);
+
+						if (SourceModel.BuildSettings.bGenerateLightmapUVs)
+						{
+							NumChannelsInLOD = FMath::Max(NumChannelsInLOD, SourceModel.BuildSettings.DstLightmapIndex + 1);
+						}
 
-					NumUVs = FMath::Min(NumChannelsInLOD, NumUVs);
+						NumUVs = FMath::Min(NumChannelsInLOD, NumUVs);
+					}
+				}
+				else
+				{
+					NumUVs = 1;
+					break;
 				}
-			}
-			else
-			{
-				NumUVs = 1;
-				break;
 			}
 		}
 
