Fixed issue where FStaticMeshRenderResources::bInitialized was not being set to false when releasing resources. This causes static mesh validation to remain silent during scene proxy creation and instead causes a crash much farther downstream.

diff --git a/Engine/Source/Runtime/Engine/Private/StaticMesh.cpp b/Engine/Source/Runtime/Engine/Private/StaticMesh.cpp
--- a/Engine/Source/Runtime/Engine/Private/StaticMesh.cpp
+++ b/Engine/Source/Runtime/Engine/Private/StaticMesh.cpp
@@ -1510,6 +1510,8 @@ void FStaticMeshRenderData::InitResources(ERHIFeatureLevel::Type InFeatureLevel,
 
 void FStaticMeshRenderData::ReleaseResources()
 {
+	bIsInitialized = false;
+
 	for (int32 LODIndex = 0; LODIndex < LODResources.Num(); ++LODIndex)
 	{
 		if (LODResources[LODIndex].VertexBuffers.StaticMeshVertexBuffer.GetNumVertices() > 0)
