Add another NetFieldExport check. This time in RepLayout code so we can try to spot differences earlier.

diff --git a/Engine/Source/Runtime/Engine/Private/RepLayout.cpp b/Engine/Source/Runtime/Engine/Private/RepLayout.cpp
--- a/Engine/Source/Runtime/Engine/Private/RepLayout.cpp
+++ b/Engine/Source/Runtime/Engine/Private/RepLayout.cpp
@@ -2532,6 +2532,10 @@ void FRepLayout::SendProperties_BackwardsCompatible(
 		PackageMapClient->AddNetFieldExportGroup(OwnerPathName, NetFieldExportGroup);
 	}
 
+	checkf(NetFieldExportGroup->NetFieldExports.Num() == Cmds.Num(),
+		TEXT("NetFieldExports.Num() does not match number of commands! PathName = %s, NetFieldExportGroup.PathName = %s, Cmds.Num() = %d, NetFieldExports.Num() = %d"),
+		*OwnerPathName, *(NetFieldExportGroup->PathName), Cmds.Num(), NetFieldExportGroup->NetFieldExports.Num());
+
 	const int32 NumBits = Writer.GetNumBits();
 
 	if (Changed.Num() == 0)
@@ -6371,6 +6375,10 @@ bool FRepLayout::DeltaSerializeFastArrayProperty(FFastArrayDeltaSerializeParams&
 			PackageMap->AddNetFieldExportGroup(OwnerPathName, LocalNetFieldExportGroup);
 		}
 
+		checkf(LocalNetFieldExportGroup->NetFieldExports.Num() == Cmds.Num(),
+			TEXT("NetFieldExports.Num() does not match number of commands! PathName = %s, NetFieldExportGroup.PathName = %s, Cmds.Num() = %d, NetFieldExports.Num() = %d"),
+			*OwnerPathName, *(LocalNetFieldExportGroup->PathName), Cmds.Num(), LocalNetFieldExportGroup->NetFieldExports.Num());
+
 		NetFieldExportGroup = LocalNetFieldExportGroup.Get();
 	}
 
