Fix attempting to load softobject path from audio thread in standalone game (links are already loaded as submixes are loaded on audio engine init)

diff --git a/Engine/Source/Runtime/Engine/Private/SoundWave.cpp b/Engine/Source/Runtime/Engine/Private/SoundWave.cpp
--- a/Engine/Source/Runtime/Engine/Private/SoundWave.cpp
+++ b/Engine/Source/Runtime/Engine/Private/SoundWave.cpp
@@ -1757,8 +1757,8 @@ void USoundWave::Parse(FAudioDevice* AudioDevice, const UPTRINT NodeWaveInstance
 
 		if (SoundClassProperties->bApplyEffects)
 		{
-			const UAudioSettings* Settings = GetDefault<UAudioSettings>();
-			WaveInstance->SoundSubmix = Cast<USoundSubmix>(Settings->EQSubmix.TryLoad());
+			UAudioSettings* Settings = GetMutableDefault<UAudioSettings>();
+			WaveInstance->SoundSubmix = Cast<USoundSubmix>(FSoftObjectPtr(Settings->EQSubmix).Get());
 		}
 		else if (SoundClassProperties->DefaultSubmix)
 		{
@@ -1793,8 +1793,8 @@ void USoundWave::Parse(FAudioDevice* AudioDevice, const UPTRINT NodeWaveInstance
 	WaveInstance->bIsAmbisonics = bIsAmbisonics;
 	if (bIsAmbisonics)
 	{
-		const UAudioSettings* Settings = GetDefault<UAudioSettings>();
-		WaveInstance->SoundSubmix = Cast<USoundSubmix>(Settings->AmbisonicSubmix.TryLoad());
+		UAudioSettings* Settings = GetMutableDefault<UAudioSettings>();
+		WaveInstance->SoundSubmix = Cast<USoundSubmix>(FSoftObjectPtr(Settings->AmbisonicSubmix).Get());
 	}
 	else if (ParseParams.SoundSubmix)
 	{
