Subtitle Manager - Some light cleanup in the subtitle manager.  It needs some more, it looks like the combination of movie subtitles and audio subtitles needs some love.

diff --git a/Engine/Source/Runtime/Engine/Private/SubtitleManager.cpp b/Engine/Source/Runtime/Engine/Private/SubtitleManager.cpp
--- a/Engine/Source/Runtime/Engine/Private/SubtitleManager.cpp
+++ b/Engine/Source/Runtime/Engine/Private/SubtitleManager.cpp
@@ -17,19 +17,16 @@ DEFINE_LOG_CATEGORY(LogSubtitle);
 /** The default offset of the outline box. */
 FIntRect DrawStringOutlineBoxOffset(2, 2, 4, 4);
 
-
 void FSubtitleManager::KillAllSubtitles()
 {
 	ActiveSubtitles.Empty();
 }
 
-
 void FSubtitleManager::KillSubtitles(PTRINT SubtitleID)
 {
 	ActiveSubtitles.Remove(SubtitleID);
 }
 
-
 void FSubtitleManager::QueueSubtitles(const FQueueSubtitleParams& QueueSubtitleParams)
 {
 	check(IsInAudioThread());
@@ -62,7 +59,6 @@ void FSubtitleManager::QueueSubtitles(const FQueueSubtitleParams& QueueSubtitleP
 	}, GET_STATID(STAT_AudioQueueSubtitles));
 }
 
-
 void FSubtitleManager::QueueSubtitles(PTRINT SubtitleID, float Priority, bool bManualWordWrap, bool bSingleLine, float SoundDuration, const TArray<FSubtitleCue>& Subtitles, float InStartTime, float InCurrentTime)
 {
 	check(GEngine);
@@ -123,7 +119,6 @@ void FSubtitleManager::QueueSubtitles(PTRINT SubtitleID, float Priority, bool bM
 	}
 }
 
-
 void FSubtitleManager::DisplaySubtitle(FCanvas* Canvas, FActiveSubtitle* Subtitle, FIntRect& Parms, const FLinearColor& Color)
 {
 	if (OnSetSubtitleTextDelegate.IsBound())
@@ -475,18 +470,10 @@ void FSubtitleManager::DisplaySubtitles(FCanvas* InCanvas, FIntRect& InSubtitleR
 
 FSubtitleManager* FSubtitleManager::GetSubtitleManager()
 {
-	static FSubtitleManager* SSubtitleManager = nullptr;
-
-	if (!SSubtitleManager)
-	{
-		SSubtitleManager = new FSubtitleManager();
-		SSubtitleManager->CurrentSubtitleHeight = 0.0f;
-	}
-
-	return(SSubtitleManager);
+	static FSubtitleManager* SSubtitleManager = new FSubtitleManager();
+	return SSubtitleManager;
 }
 
-
 void FSubtitleManager::SetMovieSubtitle(UObject* SubtitleOwner, const TArray<FString>& Subtitles)
 {
 	if (SubtitleOwner != nullptr)
@@ -512,7 +499,7 @@ void FSubtitleManager::SetMovieSubtitle(UObject* SubtitleOwner, const TArray<FSt
 			const bool bSplit = true;
 			const bool bSingleLine = false;
 
-			MovieSubtitle = MakeShareable(new FActiveSubtitle(Index, Priority, bSplit, bSingleLine, Cues));
+			MovieSubtitle = MakeShared<FActiveSubtitle>(Index, Priority, bSplit, bSingleLine, Cues);
 		}
 	}
 }
diff --git a/Engine/Source/Runtime/Engine/Public/SubtitleManager.h b/Engine/Source/Runtime/Engine/Public/SubtitleManager.h
--- a/Engine/Source/Runtime/Engine/Public/SubtitleManager.h
+++ b/Engine/Source/Runtime/Engine/Public/SubtitleManager.h
@@ -4,6 +4,7 @@
 
 #include "CoreMinimal.h"
 #include "Engine/EngineTypes.h"
+#include "UObject/ObjectKey.h"
 
 #define SUBTITLE_SCREEN_DEPTH_FOR_3D 0.1f
 
@@ -118,7 +119,7 @@ public:
 	/**
 	 * Draws a subtitle at the specified pixel location.
 	 */
-	void DisplaySubtitle(FCanvas *Canvas, FActiveSubtitle* Subtitle, FIntRect & Parms, const FLinearColor& Color);
+	void DisplaySubtitle(FCanvas* Canvas, FActiveSubtitle* Subtitle, FIntRect & Parms, const FLinearColor& Color);
 
 	/**
 	 * Display the currently queued subtitles and cleanup after they have finished rendering.
@@ -126,7 +127,7 @@ public:
 	 * @param Canvas Where to render the subtitles.
 	 * @param CurrentTime Current world time.
 	 */
-	ENGINE_API void DisplaySubtitles(FCanvas *InCanvas, FIntRect & SubtitleRegion, float InAudioTimeSeconds);
+	ENGINE_API void DisplaySubtitles(FCanvas* InCanvas, FIntRect & SubtitleRegion, float InAudioTimeSeconds);
 
 	/**
 	 * Whether there are any active subtitles.
@@ -183,8 +184,8 @@ private:
 	TMap<PTRINT, FActiveSubtitle> ActiveSubtitles;
 
 	/** The current height of the subtitles. */
-	float CurrentSubtitleHeight;
+	float CurrentSubtitleHeight = 0;
 
 	/** HACK: The currently active movie subtitle for given object owners. Each owner can have one movie subtitle at a time */
-	TMap<UObject*, TSharedPtr<FActiveSubtitle>> ActiveMovieSubtitles;
+	TMap<FObjectKey, TSharedPtr<FActiveSubtitle>> ActiveMovieSubtitles;
 };
