Chaos - soft joint parity

diff --git a/Engine/Source/Runtime/Engine/Private/PhysicsEngine/ImmediatePhysics/ImmediatePhysicsChaos/ImmediatePhysicsActorHandle_Chaos.cpp b/Engine/Source/Runtime/Engine/Private/PhysicsEngine/ImmediatePhysics/ImmediatePhysicsChaos/ImmediatePhysicsActorHandle_Chaos.cpp
--- a/Engine/Source/Runtime/Engine/Private/PhysicsEngine/ImmediatePhysics/ImmediatePhysicsChaos/ImmediatePhysicsActorHandle_Chaos.cpp
+++ b/Engine/Source/Runtime/Engine/Private/PhysicsEngine/ImmediatePhysics/ImmediatePhysicsChaos/ImmediatePhysicsActorHandle_Chaos.cpp
@@ -11,6 +11,7 @@
 #include "Chaos/PBDRigidParticles.h"
 #include "Chaos/PBDRigidsEvolutionGBF.h"
 #include "Chaos/Sphere.h"
+#include "Chaos/Utilities.h"
 
 #include "Physics/Experimental/ChaosInterfaceUtils.h"
 
@@ -156,18 +157,20 @@ namespace ImmediatePhysics_Chaos
 
 	bool CreateGeometry(FBodyInstance* BodyInstance, const FVector& Scale, float& OutMass, Chaos::TVector<float, 3>& OutInertia, Chaos::TRigidTransform<float, 3>& OutCoMTransform, TUniquePtr<Chaos::FImplicitObject>& OutGeom, TArray<TUniquePtr<Chaos::TPerShapeData<float, 3>>>& OutShapes)
 	{
+		using namespace Chaos;
+
 		UBodySetup* BodySetup = BodyInstance->BodySetup.Get();
 
 #if WITH_CHAOS && !PHYSICS_INTERFACE_PHYSX
-		Chaos::TMassProperties<float, 3> MassProperties;
+		TMassProperties<float, 3> MassProperties;
 		CalculateMassProperties(Scale, FTransform::Identity, BodySetup->AggGeom, MassProperties);
 		float Density = 1.e-3f;	// 1g/cm3	@todo(ccaulfield): should come from material
 		if (BodyInstance->bOverrideMass)
 		{
 			Density = BodyInstance->GetMassOverride() / MassProperties.Volume;
 		}
-		OutMass = Density * (BodyInstance->MassScale * MassProperties.Volume);
-		OutInertia = Density * (BodyInstance->InertiaTensorScale * Chaos::TVector<float, 3>(MassProperties.InertiaTensor.M[0][0], MassProperties.InertiaTensor.M[1][1], MassProperties.InertiaTensor.M[2][2]));
+		OutMass = Density * BodyInstance->MassScale * MassProperties.Volume;
+		OutInertia = Utilities::ScaleInertia(Density * TVector<float, 3>(MassProperties.InertiaTensor.M[0][0], MassProperties.InertiaTensor.M[1][1], MassProperties.InertiaTensor.M[2][2]), BodyInstance->InertiaTensorScale);
 		OutCoMTransform = FTransform(MassProperties.RotationOfMass, MassProperties.CenterOfMass + BodyInstance->COMNudge);
 #else
 		OutMass = BodyInstance->GetBodyMass();
@@ -190,7 +193,7 @@ namespace ImmediatePhysics_Chaos
 #if CHAOS_PARTICLE_ACTORTRANSFORM
 		AddParams.LocalTransform = FTransform::Identity;
 #else
-		AddParams.LocalTransform = Chaos::TRigidTransform<float, 3>(OutCoMTransform.GetRotation().Inverse() * -OutCoMTransform.GetTranslation(), OutCoMTransform.GetRotation().Inverse());
+		AddParams.LocalTransform = TRigidTransform<float, 3>(OutCoMTransform.GetRotation().Inverse() * -OutCoMTransform.GetTranslation(), OutCoMTransform.GetRotation().Inverse());
 #endif
 		AddParams.WorldTransform = BodyInstance->GetUnrealWorldTransform();
 		AddParams.Geometry = &BodySetup->AggGeom;
@@ -201,8 +204,8 @@ namespace ImmediatePhysics_Chaos
 		AddParams.ChaosTriMeshes = MakeArrayView(BodySetup->ChaosTriMeshes);
 #endif
 
-		TArray<TUniquePtr<Chaos::FImplicitObject>> Geoms;
-		TArray<TUniquePtr<Chaos::TPerShapeData<float, 3>>, TInlineAllocator<1>> Shapes;
+		TArray<TUniquePtr<FImplicitObject>> Geoms;
+		TArray<TUniquePtr<TPerShapeData<float, 3>>, TInlineAllocator<1>> Shapes;
 		ChaosInterface::CreateGeometry(AddParams, Geoms, Shapes);
 
 		if (Geoms.Num() == 0)
@@ -216,7 +219,7 @@ namespace ImmediatePhysics_Chaos
 		}
 		else
 		{
-			OutGeom = MakeUnique<Chaos::FImplicitObjectUnion>(MoveTemp(Geoms));
+			OutGeom = MakeUnique<FImplicitObjectUnion>(MoveTemp(Geoms));
 		}
 
 		for (auto& Shape : Shapes)
diff --git a/Engine/Source/Runtime/Engine/Private/PhysicsEngine/ImmediatePhysics/ImmediatePhysicsChaos/ImmediatePhysicsJointHandle_Chaos.cpp b/Engine/Source/Runtime/Engine/Private/PhysicsEngine/ImmediatePhysics/ImmediatePhysicsChaos/ImmediatePhysicsJointHandle_Chaos.cpp
--- a/Engine/Source/Runtime/Engine/Private/PhysicsEngine/ImmediatePhysics/ImmediatePhysicsChaos/ImmediatePhysicsJointHandle_Chaos.cpp
+++ b/Engine/Source/Runtime/Engine/Private/PhysicsEngine/ImmediatePhysics/ImmediatePhysicsChaos/ImmediatePhysicsJointHandle_Chaos.cpp
@@ -21,37 +21,34 @@ static_assert((int32)Chaos::EJointAngularConstraintIndex::Swing1 == 2, "Angular
 namespace ImmediatePhysics_Chaos
 {
 	float ChaosImmediate_JointStiffness = 1.0f;
-	FAutoConsoleVariableRef CVarJointStiffness(TEXT("p.Chaos.ImmPhys.JointStiffness"), ChaosImmediate_JointStiffness, TEXT("Hard-joint solver stiffness."));
 
 	float ChaosImmediate_DriveStiffnessScale = 0.85f;
-	float ChaosImmediate_DriveDampingScale = 1.0f;
+	float ChaosImmediate_DriveDampingScale = 0.85f;
+
+	int ChaosImmediate_SoftLinearForceMode = (int)Chaos::EJointForceMode::Acceleration;
+	float ChaosImmediate_SoftLinearStiffnessScale = 1.2f;
+	float ChaosImmediate_SoftLinearDampingScale = 1.2f;
+
+	int ChaosImmediate_SoftAngularForceMode = (int)Chaos::EJointForceMode::Acceleration;
+	float ChaosImmediate_SoftAngularStiffnessScale = 500;
+	float ChaosImmediate_SoftAngularDampingScale = 5;
+
+	float ChaosImmediate_JointMinLinearProjection = 0.0f;
+	float ChaosImmediate_JointMaxLinearProjection = 0.8f;
+	float ChaosImmediate_JointMinAngularProjection = 0.0f;
+	float ChaosImmediate_JointMaxAngularProjection = 0.1f;
+
+	FAutoConsoleVariableRef CVarJointStiffness(TEXT("p.Chaos.ImmPhys.JointStiffness"), ChaosImmediate_JointStiffness, TEXT("Hard-joint solver stiffness."));
 	FAutoConsoleVariableRef CVarDriveStiffnessScale(TEXT("p.Chaos.ImmPhys.DriveStiffnessScale"), ChaosImmediate_DriveStiffnessScale, TEXT("Conversion factor for drive stiffness."));
 	FAutoConsoleVariableRef CVarDriveDampingScale(TEXT("p.Chaos.ImmPhys.DriveDampingScale"), ChaosImmediate_DriveDampingScale, TEXT("Conversion factor for drive damping."));
-
-	float ChaosImmediate_SoftLinearStiffnessScale = 1.75f;
-	float ChaosImmediate_SoftLinearDampingScale = 1.75f;
 	FAutoConsoleVariableRef CVarSoftLinearStiffnessScale(TEXT("p.Chaos.ImmPhys.SoftLinearStiffnessScale"), ChaosImmediate_SoftLinearStiffnessScale, TEXT("Conversion factor for soft-joint stiffness."));
 	FAutoConsoleVariableRef CVarSoftLinearDampingScale(TEXT("p.Chaos.ImmPhys.SoftLinearDampingScale"), ChaosImmediate_SoftLinearDampingScale, TEXT("Conversion factor for soft-joint damping."));
-
-	int ChaosImmediate_SoftLinearForceMode = 0;
-	int ChaosImmediate_SoftAngularForceMode = 0;
 	FAutoConsoleVariableRef CVarSoftLinearForceMode(TEXT("p.Chaos.ImmPhys.SoftLinearForceMode"), ChaosImmediate_SoftLinearForceMode, TEXT("Soft Linear constraint force mode (0: Acceleration; 1: Force"));
 	FAutoConsoleVariableRef CVarSoftAngularForceMode(TEXT("p.Chaos.ImmPhys.SoftAngularForceMode"), ChaosImmediate_SoftAngularForceMode, TEXT("Soft Angular constraint force mode (0: Acceleration; 1: Force"));
-
-	// This value US passed to PhysX was 100000 UE
-	// (In PhysX, Soft constraints were implicitly integrated, and essentially equivalent to XPBD, which for large stiffness is just PBD).
-	float ChaosImmediate_SoftAngularStiffnessScale = 100000.0f;
-	float ChaosImmediate_SoftAngularDampingScale = 100000.0f;
 	FAutoConsoleVariableRef CVarSoftAngularStiffnessScale(TEXT("p.Chaos.ImmPhys.SoftAngularStiffnessScale"), ChaosImmediate_SoftAngularStiffnessScale, TEXT("Conversion factor for soft-joint stiffness."));
 	FAutoConsoleVariableRef CVarSoftAngularDampingScale(TEXT("p.Chaos.ImmPhys.SoftAngularDampingScale"), ChaosImmediate_SoftAngularDampingScale, TEXT("Conversion factor for soft-joint damping."));
-
-	float ChaosImmediate_JointMinLinearProjection = 0.0f;
-	float ChaosImmediate_JointMaxLinearProjection = 0.8f;
 	FAutoConsoleVariableRef CVarJointMinLinearProjection(TEXT("p.Chaos.ImmPhys.JointMinLinearProjection"), ChaosImmediate_JointMinLinearProjection, TEXT("Joint min projection (for joints with projection disabled)."));
 	FAutoConsoleVariableRef CVarJointMaxLinearProjection(TEXT("p.Chaos.ImmPhys.JointMaxLinearProjection"), ChaosImmediate_JointMaxLinearProjection, TEXT("Joint max projection (for joints with projection enabled)."));
-
-	float ChaosImmediate_JointMinAngularProjection = 0.0f;
-	float ChaosImmediate_JointMaxAngularProjection = 0.1f;
 	FAutoConsoleVariableRef CVarJointMinAngularProjection(TEXT("p.Chaos.ImmPhys.JointMinAngularProjection"), ChaosImmediate_JointMinAngularProjection, TEXT("Joint min projection (for joints with projection disabled)."));
 	FAutoConsoleVariableRef CVarJointMaxAngularProjection(TEXT("p.Chaos.ImmPhys.JointMaxAngularProjection"), ChaosImmediate_JointMaxAngularProjection, TEXT("Joint max projection (for joints with projection enabled)."));
 
diff --git a/Engine/Source/Runtime/Engine/Private/PhysicsEngine/ImmediatePhysics/ImmediatePhysicsChaos/ImmediatePhysicsSimulation_Chaos.cpp b/Engine/Source/Runtime/Engine/Private/PhysicsEngine/ImmediatePhysics/ImmediatePhysicsChaos/ImmediatePhysicsSimulation_Chaos.cpp
--- a/Engine/Source/Runtime/Engine/Private/PhysicsEngine/ImmediatePhysics/ImmediatePhysicsChaos/ImmediatePhysicsSimulation_Chaos.cpp
+++ b/Engine/Source/Runtime/Engine/Private/PhysicsEngine/ImmediatePhysics/ImmediatePhysicsChaos/ImmediatePhysicsSimulation_Chaos.cpp
@@ -47,7 +47,7 @@ FAutoConsoleVariableRef CVarChaosImmPhysCollisionPriority(TEXT("p.Chaos.ImmPhys.
 FAutoConsoleVariableRef CVarChaosImmPhysThickness(TEXT("p.Chaos.ImmPhys.Collision.Thickness"), ChaosImmediate_Collision_Thickness, TEXT("ChaosImmediateThickness"));
 
 int32 ChaosImmediate_Joint_PairIterations = 2;
-int32 ChaosImmediate_Joint_PushOutPairIterations = 1;
+int32 ChaosImmediate_Joint_PushOutPairIterations = 0;// 1; Until I fix projection for soft joints
 float ChaosImmediate_Joint_SwingTwistAngleTolerance = 1.0e-6f;
 int32 ChaosImmediate_Joint_EnableTwistLimits = 1;
 int32 ChaosImmediate_Joint_EnableSwingLimits = 1;
diff --git a/Engine/Source/Runtime/Experimental/Chaos/Private/Chaos/Joint/PBDJointSolverGaussSeidel.cpp b/Engine/Source/Runtime/Experimental/Chaos/Private/Chaos/Joint/PBDJointSolverGaussSeidel.cpp
--- a/Engine/Source/Runtime/Experimental/Chaos/Private/Chaos/Joint/PBDJointSolverGaussSeidel.cpp
+++ b/Engine/Source/Runtime/Experimental/Chaos/Private/Chaos/Joint/PBDJointSolverGaussSeidel.cpp
@@ -14,9 +14,9 @@
 #if !INTEL_ISPC
 const bool bChaos_Joint_ISPC_Enabled = false;
 #elif UE_BUILD_SHIPPING
-const bool bChaos_Joint_ISPC_Enabled = true;
+const bool bChaos_Joint_ISPC_Enabled = false;
 #else
-bool bChaos_Joint_ISPC_Enabled = true;
+bool bChaos_Joint_ISPC_Enabled = false;
 FAutoConsoleVariableRef CVarChaosJointISPCEnabled(TEXT("p.Chaos.Joint.ISPC"), bChaos_Joint_ISPC_Enabled, TEXT("Whether to use ISPC optimizations in the Joint Solver"));
 #endif
 
@@ -355,7 +355,7 @@ namespace Chaos
 		const FPBDJointSettings& JointSettings)
 	{
 		ApplyPositionConstraints(Dt, SolverSettings, JointSettings);
-
+	
 		ApplyRotationConstraints(Dt, SolverSettings, JointSettings);
 	}
 
@@ -365,9 +365,9 @@ namespace Chaos
 		const FPBDJointSolverSettings& SolverSettings,
 		const FPBDJointSettings& JointSettings)
 	{
-		ApplyRotationDrives(Dt, SolverSettings, JointSettings);
-
 		ApplyPositionDrives(Dt, SolverSettings, JointSettings);
+
+		ApplyRotationDrives(Dt, SolverSettings, JointSettings);
 	}
 
 
@@ -893,8 +893,34 @@ namespace Chaos
 			const FVec3 IA1 = Utilities::Multiply(InvI1, Axis1);
 
 			// Joint-space inverse mass
-			const FReal II0 = FVec3::DotProduct(Axis0, IA0);
-			const FReal II1 = FVec3::DotProduct(Axis1, IA1);
+			FReal II0 = FVec3::DotProduct(Axis0, IA0);
+			FReal II1 = FVec3::DotProduct(Axis1, IA1);
+
+			// If we are correcting the position, we need to adjust the constraint effective mass using parallel axis theorem
+			// @todo(ccaulfield): the IMs here are constant per constraint...pre-build and cache it
+			if (AngularPositionCorrection > 0)
+			{
+				if (InvMs[0] > 0)
+				{
+					const FVec3 LinearAxis0 = FVec3::CrossProduct(Xs[0] - Ps[0], Axis0);
+					const FReal LinearAxisLen0Sq = LinearAxis0.SizeSquared();
+					if (LinearAxisLen0Sq > KINDA_SMALL_NUMBER)
+					{
+						const FReal IM0 = InvMs[0] / LinearAxisLen0Sq;
+						II0 = II0 * IM0 / (II0 + IM0);
+					}
+				}
+				if (InvMs[1] > 0)
+				{
+					const FVec3 LinearAxis1 = FVec3::CrossProduct(Xs[1] - Ps[1], Axis1);
+					const FReal LinearAxisLen1Sq = LinearAxis1.SizeSquared();
+					if (LinearAxisLen1Sq > KINDA_SMALL_NUMBER)
+					{
+						const FReal IM1 = InvMs[1] / LinearAxisLen1Sq;
+						II1 = II1 * IM1 / (II1 + IM1);
+					}
+				}
+			}
 			const FReal II = (II0 + II1);
 
 			// Damping angular velocity
