Hot fix for a crash where we ended up creating more mip map level then we should had.
This issue was happening after using LeaveExistingMips for the MipGenSettings and setting a maximun texture size.
Added a validation when getting that data from the ddc as it could be invalid.

diff --git a/Engine/Source/Developer/TextureCompressor/Private/TextureCompressorModule.cpp b/Engine/Source/Developer/TextureCompressor/Private/TextureCompressorModule.cpp
--- a/Engine/Source/Developer/TextureCompressor/Private/TextureCompressorModule.cpp
+++ b/Engine/Source/Developer/TextureCompressor/Private/TextureCompressorModule.cpp
@@ -2079,15 +2079,30 @@ private:
 
 		// Determine the number of mips required by BuildSettings.
 		int32 NumOutputMips = (BuildSettings.MipGenSettings == TMGS_NoMipmaps) ? 1 : MaxSourceMipCount;
-		NumOutputMips = FMath::Min(NumOutputMips, MaxDestMipCount);
 
 		int32 NumSourceMips = InSourceMips.Num();
 
+		// See if the smallest provided mip image is still too large for the current compressor.
+		int32 LevelsToUsableSource = FMath::Max(0, MaxSourceMipCount - MaxDestMipCount);
+		int32 StartMip = FMath::Max(0, LevelsToUsableSource);
+
 		if (BuildSettings.MipGenSettings == TMGS_LeaveExistingMips)
 		{
-			NumOutputMips = InSourceMips.Num();
+			NumOutputMips = InSourceMips.Num() - StartMip;
+			if (NumOutputMips <= 0)
+			{
+				// We can't generate 0 mip maps
+				UE_LOG(LogTextureCompressor, Warning,
+					TEXT("The source image has %d mips while the first mip would be %d. Please verify the maximun texture size or change the mips gen settings."),
+					NumSourceMips,
+					StartMip);
+				return false;
+			}
 		}
 
+		NumOutputMips = FMath::Min(NumOutputMips, MaxDestMipCount);
+
+
 		if (BuildSettings.MipGenSettings != TMGS_LeaveExistingMips || bLongLatCubemap)
 		{
 			NumSourceMips = 1;
@@ -2186,9 +2201,6 @@ private:
 
 		const TArray<FImage>& PostOptionalUpscaleSourceMips = (PaddedSourceMips.Num() > 0) ? PaddedSourceMips : InSourceMips;
 
-		// See if the smallest provided mip image is still too large for the current compressor.
-		int32 LevelsToUsableSource = FMath::Max(0, MaxSourceMipCount - MaxDestMipCount);
-		int32 StartMip = FMath::Max(0, LevelsToUsableSource);
 		bool bBuildSourceImage = StartMip > (NumSourceMips - 1);
 
 		TArray<FImage> GeneratedSourceMips;
diff --git a/Engine/Source/Runtime/Engine/Private/TextureDerivedData.cpp b/Engine/Source/Runtime/Engine/Private/TextureDerivedData.cpp
--- a/Engine/Source/Runtime/Engine/Private/TextureDerivedData.cpp
+++ b/Engine/Source/Runtime/Engine/Private/TextureDerivedData.cpp
@@ -620,7 +620,7 @@ uint32 PutDerivedDataInCache(FTexturePlatformData* DerivedData, const FString& D
 	FMemoryWriter Ar(RawDerivedData, /*bIsPersistent=*/ true);
 	DerivedData->Serialize(Ar, NULL);
 	TotalBytesPut += RawDerivedData.Num();
-	GetDerivedDataCacheRef().Put(*DerivedDataKey, RawDerivedData);
+	GetDerivedDataCacheRef().Put(*DerivedDataKey, RawDerivedData, true);
 	UE_LOG(LogTexture,Verbose,TEXT("%s  Derived Data: %d bytes"),*LogString,RawDerivedData.Num());
 	return TotalBytesPut;
 }
diff --git a/Engine/Source/Runtime/Engine/Private/TextureDerivedDataTask.cpp b/Engine/Source/Runtime/Engine/Private/TextureDerivedDataTask.cpp
--- a/Engine/Source/Runtime/Engine/Private/TextureDerivedDataTask.cpp
+++ b/Engine/Source/Runtime/Engine/Private/TextureDerivedDataTask.cpp
@@ -24,10 +24,13 @@
 #if WITH_EDITOR
 
 #include "DerivedDataCacheInterface.h"
+#include "Engine/TextureCube.h"
+#include "GenericPlatform/GenericPlatformMath.h"
 #include "Interfaces/ITargetPlatform.h"
 #include "Interfaces/ITargetPlatformManagerModule.h"
 #include "Interfaces/ITextureFormat.h"
-#include "Engine/TextureCube.h"
+#include "Math/NumericLimits.h"
+#include "Math/UnrealMathUtility.h"
 #include "ProfilingDebugging/CookStats.h"
 #include "VT/VirtualTextureDataBuilder.h"
 
@@ -435,6 +438,28 @@ void FTextureCacheDerivedDataWorker::DoWork()
 			else
 			{
 				bSucceeded = DerivedData->AreDerivedMipsAvailable();
+
+				if (bSucceeded && BuildSettingsPerLayer.Num() > 0)
+				{
+					// Code inspired by the texture compressor module as a hot fix for the bad data that might have been push into the ddc in 4.23 or 4.24 
+					const bool bLongLatCubemap = DerivedData->IsCubemap() && DerivedData->GetNumSlices() == 1;
+					int32 MaximunNumberOfMipMap = TNumericLimits<int32>::Max();
+					if (bLongLatCubemap)
+					{
+						MaximunNumberOfMipMap = FMath::CeilLogTwo(FMath::Clamp<uint32>(uint32(1 << FMath::FloorLog2(DerivedData->SizeX / 2)), uint32(32), BuildSettingsPerLayer[0].MaxTextureResolution)) + 1;
+					}
+					else
+					{
+						MaximunNumberOfMipMap = FMath::CeilLogTwo(FMath::Max3(DerivedData->SizeX, DerivedData->SizeY, BuildSettingsPerLayer[0].bVolume ? DerivedData->GetNumSlices() : 1)) + 1;
+					}
+
+					bSucceeded = DerivedData->Mips.Num() <= MaximunNumberOfMipMap;
+				}
+
+				if (!bSucceeded)
+				{
+					UE_LOG(LogTexture, Warning, TEXT("The data retrieved from the derived data cache for the texture %s was invalid. The texture will be rebuild."), *Texture.GetFullName())
+				}
 			}
 		}
 		bLoadedFromDDC = true;
