AActor::SetLODParent() minor optimisations
* Avoid dirtying render state if draw distance is unchanged
* Avoid array allocation

diff --git a/Engine/Source/Runtime/Engine/Private/Actor.cpp b/Engine/Source/Runtime/Engine/Private/Actor.cpp
--- a/Engine/Source/Runtime/Engine/Private/Actor.cpp
+++ b/Engine/Source/Runtime/Engine/Private/Actor.cpp
@@ -4964,19 +4964,19 @@ void AActor::PostRename(UObject* OldOuter, const FName OldName)
 
 void AActor::SetLODParent(UPrimitiveComponent* InLODParent, float InParentDrawDistance)
 {
-	if (InLODParent)
+	if (InLODParent && InLODParent->MinDrawDistance != InParentDrawDistance)
 	{
 		InLODParent->MinDrawDistance = InParentDrawDistance;
 		InLODParent->MarkRenderStateDirty();
 	}
 
-	TArray<UPrimitiveComponent*> ComponentsToBeReplaced;
-	GetComponents(ComponentsToBeReplaced);
-
-	for (UPrimitiveComponent* Component : ComponentsToBeReplaced)
+	for (UActorComponent* Component : GetComponents())
 	{
-		// parent primitive will be null if no LOD parent is selected
-		Component->SetLODParentPrimitive(InLODParent);
+		if (UPrimitiveComponent* PrimitiveComponent = Cast<UPrimitiveComponent>(Component))
+		{
+			// parent primitive will be null if no LOD parent is selected
+			PrimitiveComponent->SetLODParentPrimitive(InLODParent);
+		}
 	}
 }
 
