Change ensure in UPhysicalAnimationComponent::UpdatePhysicsEngineImp() from ensure to Chaos ensure.

diff --git a/Engine/Source/Runtime/Engine/Private/PhysicsEngine/PhysicalAnimationComponent.cpp b/Engine/Source/Runtime/Engine/Private/PhysicsEngine/PhysicalAnimationComponent.cpp
--- a/Engine/Source/Runtime/Engine/Private/PhysicsEngine/PhysicalAnimationComponent.cpp
+++ b/Engine/Source/Runtime/Engine/Private/PhysicsEngine/PhysicalAnimationComponent.cpp
@@ -6,6 +6,7 @@
 #include "PhysicsEngine/PhysicsAsset.h"
 #include "PhysXPublic.h"
 #include "Physics/PhysicsInterfaceCore.h"
+#include "ChaosCheck.h"
 
 
 const FConstraintProfileProperties UPhysicalAnimationComponent::PhysicalAnimationProfile = []()
@@ -376,7 +377,7 @@ void UPhysicalAnimationComponent::UpdatePhysicsEngineImp()
 					if (FBodyInstance* ChildBody = (ChildBodyIdx == INDEX_NONE ? nullptr : SkeletalMeshComponent->Bodies[ChildBodyIdx]))
 					{
 #if WITH_CHAOS || WITH_IMMEDIATE_PHYSX
-                        ensure(false);
+                        CHAOS_ENSURE(false);
 #else
 						if (PxRigidActor* PRigidActor = FPhysicsInterface_PhysX::GetPxRigidActor_AssumesLocked(ChildBody->ActorHandle))
 						{
