* Remove duplicate virtual texture information from shader key

diff --git a/Engine/Source/Runtime/Engine/Private/ShaderCompiler/ShaderCompiler.cpp b/Engine/Source/Runtime/Engine/Private/ShaderCompiler/ShaderCompiler.cpp
--- a/Engine/Source/Runtime/Engine/Private/ShaderCompiler/ShaderCompiler.cpp
+++ b/Engine/Source/Runtime/Engine/Private/ShaderCompiler/ShaderCompiler.cpp
@@ -3667,13 +3667,6 @@ void GlobalBeginCompileShader(
 		Input.Environment.SetDefine(TEXT("PROJECT_MOBILE_USE_LEGACY_SHADING"), CVar ? (CVar->GetInt() != 0) : 0);
 	}
 
-	{
-		const ERHIFeatureLevel::Type FeatureLevel = GetMaxSupportedFeatureLevel((EShaderPlatform)Target.Platform);
-		const ITargetPlatform* TargetPlatform = GetTargetPlatformManager()->FindTargetPlatform(ShaderPlatformToPlatformName(EShaderPlatform(Target.Platform)).ToString());
-		const bool bVirtualTextureEnabled = UseVirtualTexturing(FeatureLevel, TargetPlatform);
-		Input.Environment.SetDefine(TEXT("PROJECT_SUPPORT_VIRTUAL_TEXTURE"), bVirtualTextureEnabled ? 1 : 0);
-	}
-
 	{
 		static IConsoleVariable* CVar = IConsoleManager::Get().FindConsoleVariable(TEXT("r.Mobile.ForceFullPrecisionInPS"));
 		if (CVar && CVar->GetInt() != 0)
diff --git a/Engine/Source/Runtime/RenderCore/Private/Shader.cpp b/Engine/Source/Runtime/RenderCore/Private/Shader.cpp
--- a/Engine/Source/Runtime/RenderCore/Private/Shader.cpp
+++ b/Engine/Source/Runtime/RenderCore/Private/Shader.cpp
@@ -2457,13 +2457,6 @@ void ShaderMapAppendKeyString(EShaderPlatform Platform, FString& KeyString)
 		}
 	}
 
-	{
-		if (UseVirtualTexturing(GetMaxSupportedFeatureLevel(Platform), TargetPlatform))
-		{
-			KeyString += TEXT("_VT");
-		}
-	}
-
 	{
 		static const auto CVar = IConsoleManager::Get().FindTConsoleVariableDataInt(TEXT("r.PostProcessing.PropagateAlpha"));
 		if (CVar && CVar->GetValueOnAnyThread() > 0)
