Prevent NoCollision bodies from creating physics actor and geometry.

diff --git a/Engine/Source/Runtime/Engine/Private/PhysicsEngine/BodyInstance.cpp b/Engine/Source/Runtime/Engine/Private/PhysicsEngine/BodyInstance.cpp
--- a/Engine/Source/Runtime/Engine/Private/PhysicsEngine/BodyInstance.cpp
+++ b/Engine/Source/Runtime/Engine/Private/PhysicsEngine/BodyInstance.cpp
@@ -1118,6 +1118,20 @@ struct FInitBodiesHelper
 				}
 			}
 
+#if WITH_CHAOS
+			if (Instance->GetCollisionEnabled() == ECollisionEnabled::NoCollision)
+			{
+				// Object should not be used in physics engine at all. Skip it.
+				// This is probably OK for PhysX too, as it appears to create NoCollision actors and shapes.
+				
+				Instance->OwnerComponent = nullptr;
+				Instance->BodySetup      = nullptr;
+
+				continue;
+			}
+#endif
+
+
 			// Don't process if we've already got a body
 			// Just ask actorref
 			if (FPhysicsInterface::IsValid(Instance->GetPhysicsActorHandle()))
@@ -1211,7 +1225,7 @@ struct FInitBodiesHelper
 #if WITH_CHAOS
 							// If this shape shouldn't collide in the sim we disable it here until we support
 							// a separation of unions for these shapes
-							if(BI->GetCollisionEnabled() == ECollisionEnabled::QueryOnly || BI->GetCollisionEnabled() == ECollisionEnabled::NoCollision)
+							if(BI->GetCollisionEnabled() == ECollisionEnabled::QueryOnly)
 							{
 								const int32 NumShapes = FPhysicsInterface::GetNumShapes(ActorHandle);
 								for(int32 ShapeIndex = 0; ShapeIndex < NumShapes; ++ShapeIndex)
@@ -1243,7 +1257,10 @@ struct FInitBodiesHelper
 					for (int32 BodyIdx = 0, NumBodies = Bodies.Num(); BodyIdx < NumBodies; ++BodyIdx)
 					{
 						FBodyInstance* Instance = Bodies[BodyIdx];
-						Instance->InitDynamicProperties_AssumesLocked();
+						if (Instance->ActorHandle)
+						{
+							Instance->InitDynamicProperties_AssumesLocked();
+						}
 					}
 				}
 			});
