Fix non deterministic skeleton cooking

diff --git a/Engine/Source/Runtime/Engine/Classes/Animation/SmartName.h b/Engine/Source/Runtime/Engine/Classes/Animation/SmartName.h
--- a/Engine/Source/Runtime/Engine/Classes/Animation/SmartName.h
+++ b/Engine/Source/Runtime/Engine/Classes/Animation/SmartName.h
@@ -170,6 +170,11 @@ private:
 #endif
 
 	TMap<FName, FCurveMetaData> CurveMetaDataMap;
+
+#if WITH_EDITORONLY_DATA
+	// Editor copy of the data we loaded, used to preserve determinism during cooking
+	TMap<FName, FCurveMetaData> LoadDataCurveMetaDataMap;
+#endif
 };
 
 USTRUCT()
diff --git a/Engine/Source/Runtime/Engine/Private/Animation/SmartName.cpp b/Engine/Source/Runtime/Engine/Private/Animation/SmartName.cpp
--- a/Engine/Source/Runtime/Engine/Private/Animation/SmartName.cpp
+++ b/Engine/Source/Runtime/Engine/Private/Animation/SmartName.cpp
@@ -111,7 +111,16 @@ void FSmartNameMapping::Serialize(FArchive& Ar)
 
 	if (Ar.CustomVer(FFrameworkObjectVersion::GUID) >= FFrameworkObjectVersion::MoveCurveTypesToSkeleton)
 	{
-		Ar << CurveMetaDataMap;
+#if WITH_EDITORONLY_DATA
+		if (Ar.IsCooking())
+		{
+			Ar << LoadDataCurveMetaDataMap; // Cook out loaded data for determinism
+		}
+		else
+#endif
+		{
+			Ar << CurveMetaDataMap;
+		}
 	}
 
 	if (Ar.IsLoading())
@@ -120,6 +129,10 @@ void FSmartNameMapping::Serialize(FArchive& Ar)
 #if !WITH_EDITOR
 		CurveMetaDataMap.GenerateValueArray(CurveMetaDataList);
 #endif
+
+#if WITH_EDITORONLY_DATA
+		LoadDataCurveMetaDataMap = CurveMetaDataMap;
+#endif
 	}
 }
 
