const correctness for quadtree accessor

diff --git a/Engine/Source/Runtime/Engine/Public/GenericQuadTree.h b/Engine/Source/Runtime/Engine/Public/GenericQuadTree.h
--- a/Engine/Source/Runtime/Engine/Public/GenericQuadTree.h
+++ b/Engine/Source/Runtime/Engine/Public/GenericQuadTree.h
@@ -18,7 +18,7 @@ public:
 	TQuadTree(const FBox2D& InBox, float InMinimumQuadSize = 100.f);
 
 	/** Gets the TreeBox so systems can test insertions before trying to do so with invalid regions */
-	const FBox2D& GetTreeBox() { return TreeBox; }
+	const FBox2D& GetTreeBox() const { return TreeBox; }
 
 	/** Inserts an object of type ElementType with an associated 2D box of size Box (log n). Pass in a DebugContext so when an issue occurs the log can report what requested this insert. */
 	void Insert(const ElementType& Element, const FBox2D& Box, const TCHAR* DebugContext = nullptr);
