Fix BodyInstance CIS error (PhysX)

diff --git a/Engine/Source/Runtime/Engine/Private/PhysicsEngine/BodyInstance.cpp b/Engine/Source/Runtime/Engine/Private/PhysicsEngine/BodyInstance.cpp
--- a/Engine/Source/Runtime/Engine/Private/PhysicsEngine/BodyInstance.cpp
+++ b/Engine/Source/Runtime/Engine/Private/PhysicsEngine/BodyInstance.cpp
@@ -1257,7 +1257,7 @@ struct FInitBodiesHelper
 					for (int32 BodyIdx = 0, NumBodies = Bodies.Num(); BodyIdx < NumBodies; ++BodyIdx)
 					{
 						FBodyInstance* Instance = Bodies[BodyIdx];
-						if (Instance->ActorHandle)
+						if (Instance->GetCollisionEnabled() != ECollisionEnabled::NoCollision)
 						{
 							Instance->InitDynamicProperties_AssumesLocked();
 						}
