Avoid using unions when there is only 1 shape. Saves memory

diff --git a/Engine/Source/Runtime/Engine/Private/PhysicsEngine/Experimental/PhysInterface_Chaos.cpp b/Engine/Source/Runtime/Engine/Private/PhysicsEngine/Experimental/PhysInterface_Chaos.cpp
--- a/Engine/Source/Runtime/Engine/Private/PhysicsEngine/Experimental/PhysInterface_Chaos.cpp
+++ b/Engine/Source/Runtime/Engine/Private/PhysicsEngine/Experimental/PhysInterface_Chaos.cpp
@@ -1348,7 +1348,14 @@ void FPhysInterface_Chaos::AddGeometry(FPhysicsActorHandle& InActor, const FGeom
 		}
 
 		//todo: we should not be creating unique geometry per actor
-		InActor->SetGeometry(MakeUnique<Chaos::FImplicitObjectUnion>(MoveTemp(Geoms)));
+		if(Geoms.Num() > 1)
+		{
+			InActor->SetGeometry(MakeUnique<Chaos::FImplicitObjectUnion>(MoveTemp(Geoms)));
+		}
+		else
+		{
+			InActor->SetGeometry(MoveTemp(Geoms[0]));
+		}
 		InActor->SetShapesArray(MoveTemp(Shapes));
 	}
 #endif
