When the last cull distance volume is deleted, still update the cached cull distance of primitives

diff --git a/Engine/Source/Runtime/Engine/Private/World.cpp b/Engine/Source/Runtime/Engine/Private/World.cpp
--- a/Engine/Source/Runtime/Engine/Private/World.cpp
+++ b/Engine/Source/Runtime/Engine/Private/World.cpp
@@ -1899,25 +1899,33 @@ bool UWorld::UpdateCullDistanceVolumes(AActor* ActorToUpdate, UPrimitiveComponen
 
 		TArray<ACullDistanceVolume*> CullDistanceVolumes;
 
+		auto EvaluateComponent = [&bUpdatedDrawDistances, &CompToNewMaxDrawMap](UPrimitiveComponent* PrimitiveComponent)
+		{
+			if (ACullDistanceVolume::CanBeAffectedByVolumes(PrimitiveComponent))
+			{
+				CompToNewMaxDrawMap.Add(PrimitiveComponent, PrimitiveComponent->LDMaxDrawDistance);
+				if (!bUpdatedDrawDistances && !FMath::IsNearlyEqual(PrimitiveComponent->LDMaxDrawDistance, PrimitiveComponent->CachedMaxDrawDistance))
+				{
+					bUpdatedDrawDistances = true;
+				}
+			}
+		};
+
 		// Establish base line of LD specified cull distances.
 		if (ActorToUpdate || ComponentToUpdate)
 		{
 			if (ComponentToUpdate)
 			{
 				check((ActorToUpdate == nullptr) || (ActorToUpdate == ComponentToUpdate->GetOwner()));
-				if (ACullDistanceVolume::CanBeAffectedByVolumes(ComponentToUpdate))
-				{
-					CompToNewMaxDrawMap.Add(ComponentToUpdate, ComponentToUpdate->LDMaxDrawDistance);
-				}
+				EvaluateComponent(ComponentToUpdate);
 			}
 			else
 			{
-				TInlineComponentArray<UPrimitiveComponent*> PrimitiveComponents(ActorToUpdate);
-				for (UPrimitiveComponent* PrimitiveComponent : PrimitiveComponents)
+				for (UActorComponent* ActorComponent : ActorToUpdate->GetComponents())
 				{
-					if (ACullDistanceVolume::CanBeAffectedByVolumes(PrimitiveComponent))
+					if (UPrimitiveComponent* PrimitiveComponent = Cast<UPrimitiveComponent>(ActorComponent))
 					{
-						CompToNewMaxDrawMap.Add(PrimitiveComponent, PrimitiveComponent->LDMaxDrawDistance);
+						EvaluateComponent(PrimitiveComponent);
 					}
 				}
 			}
@@ -1934,12 +1942,11 @@ bool UWorld::UpdateCullDistanceVolumes(AActor* ActorToUpdate, UPrimitiveComponen
 		{
 			for( AActor* Actor : FActorRange(this) )
 			{
-				TInlineComponentArray<UPrimitiveComponent*> PrimitiveComponents(Actor);
-				for (UPrimitiveComponent* PrimitiveComponent : PrimitiveComponents)
+				for (UActorComponent* ActorComponent : Actor->GetComponents())
 				{
-					if (ACullDistanceVolume::CanBeAffectedByVolumes(PrimitiveComponent))
+					if (UPrimitiveComponent* PrimitiveComponent = Cast<UPrimitiveComponent>(ActorComponent))
 					{
-						CompToNewMaxDrawMap.Add(PrimitiveComponent, PrimitiveComponent->LDMaxDrawDistance);
+						EvaluateComponent(PrimitiveComponent);
 					}
 				}
 
@@ -1951,22 +1958,23 @@ bool UWorld::UpdateCullDistanceVolumes(AActor* ActorToUpdate, UPrimitiveComponen
 			}
 		}
 
-		// Only perform the update if we actually have cull distance volumes
-		if (CullDistanceVolumes.Num() > 0)
+		if (CullDistanceVolumes.Num() > 0 && CompToNewMaxDrawMap.Num() > 0)
 		{
 			// Iterate over all cull distance volumes and get new cull distances.
-			if (CompToNewMaxDrawMap.Num() > 0)
+			for (ACullDistanceVolume* CullDistanceVolume : CullDistanceVolumes)
 			{
-				for (ACullDistanceVolume* CullDistanceVolume : CullDistanceVolumes)
-				{
-					CullDistanceVolume->GetPrimitiveMaxDrawDistances(CompToNewMaxDrawMap);
-				}
-
-				bUpdatedDrawDistances = true;
+				CullDistanceVolume->GetPrimitiveMaxDrawDistances(CompToNewMaxDrawMap);
 			}
 
+			bUpdatedDrawDistances = true;
+		}
+
+		// Only perform the update if we actually have cull distance volumes 
+		// or we found a component that had a cached value different than the LD set value
+		if (bUpdatedDrawDistances)
+		{
 			// Finally, go over all primitives, and see if they need to change.
-			// Only if they do do we reregister them, as thats slow.
+			// Only if they do do we reregister them, as that's slow.
 			for (TMap<UPrimitiveComponent*, float>::TIterator It(CompToNewMaxDrawMap); It; ++It)
 			{
 				UPrimitiveComponent* PrimComp = It.Key();
@@ -1985,7 +1993,6 @@ bool UWorld::UpdateCullDistanceVolumes(AActor* ActorToUpdate, UPrimitiveComponen
 	return bUpdatedDrawDistances;
 }
 
-
 void UWorld::ModifyLevel(ULevel* Level) const
 {
 	if( Level && Level->HasAnyFlags(RF_Transactional))
