Specify the name to create for a subobject when duplicating for actor replacement so that subobjects created in the constructor will be reused, rather than being in the way and resulting a new name being used.

diff --git a/Engine/Source/Runtime/Engine/Private/InstancedReferenceSubobjectHelper.cpp b/Engine/Source/Runtime/Engine/Private/InstancedReferenceSubobjectHelper.cpp
--- a/Engine/Source/Runtime/Engine/Private/InstancedReferenceSubobjectHelper.cpp
+++ b/Engine/Source/Runtime/Engine/Private/InstancedReferenceSubobjectHelper.cpp
@@ -340,7 +340,7 @@ void FFindInstancedReferenceSubobjectHelper::Duplicate(UObject* OldObject, UObje
 					const bool bNotHandledYet = !ReferenceReplacementMap.Contains(Obj);
 					if (bKeptByOld && bNotHandledYet)
 					{
-						UObject* NewEditInlineSubobject = StaticDuplicateObject(Obj, NewObject);
+						UObject* NewEditInlineSubobject = StaticDuplicateObject(Obj, NewObject, Obj->GetFName());
 						ReferenceReplacementMap.Add(Obj, NewEditInlineSubobject);
 
 						// NOTE: we cannot patch OldObject's linker table here, since we don't 
