Fix Static Analysis issue with IS_PUSH_MODEL_ENABLED when push model is compiled out.

diff --git a/Engine/Source/Runtime/Engine/Private/RepLayout.cpp b/Engine/Source/Runtime/Engine/Private/RepLayout.cpp
--- a/Engine/Source/Runtime/Engine/Private/RepLayout.cpp
+++ b/Engine/Source/Runtime/Engine/Private/RepLayout.cpp
@@ -5236,6 +5236,7 @@ void FRepLayout::InitFromClass(
 	SCOPE_CYCLE_COUNTER(STAT_RepLayout_InitFromObjectClass);
 	SCOPE_CYCLE_UOBJECT(ObjectClass, InObjectClass);
 
+	const bool bIsPushModelEnabled = IS_PUSH_MODEL_ENABLED();
 	const bool bIsObjectActor = InObjectClass->IsChildOf(AActor::StaticClass());
 
 	if (bIsObjectActor)
@@ -5368,7 +5369,7 @@ void FRepLayout::InitFromClass(
 			}
 
 			++NumberOfLifetimeProperties;
-			if (IS_PUSH_MODEL_ENABLED() && LifetimeProps[i].bIsPushBased)
+			if (bIsPushModelEnabled && LifetimeProps[i].bIsPushBased)
 			{
 				++NumberOfPushModelProperties;
 				Parents[ParentIndex].Flags |= ERepParentFlags::UsePushModel;
@@ -5477,7 +5478,7 @@ void FRepLayout::InitFromClass(
 		}
 	}	
 
-	if (IS_PUSH_MODEL_ENABLED() && NumberOfPushModelProperties > 0)
+	if (bIsPushModelEnabled && NumberOfPushModelProperties > 0)
 	{
 		Flags |= (NumberOfLifetimeProperties == NumberOfPushModelProperties) ?
 			ERepLayoutFlags::FullPushSupport :
