Removing a few common small scope counters to reduce stats overhead.  Trivial scopes were deleted while some non-trivial scopes were moved to verbose groups.

Made Verbose:  FStaticComponentRenderAssetManager_IncrementalUpdate, STAT_Collision_PreFilter, STAT_AddSubjectPrimitive, STAT_GnmCreateUniformBufferTime
Deleted: STAT_ParticleMemTime, STAT_SpriteSpawnTime, STAT_PARTALLOC, STAT_FMeshBatchAndRelevance, STAT_DequeueCooldowns, STAT_GatherTicksForParallel
Moved: STAT_SpriteSpawnTime

#rnx

diff --git a/Engine/Source/Runtime/Core/Public/Stats/Stats2.h b/Engine/Source/Runtime/Core/Public/Stats/Stats2.h
--- a/Engine/Source/Runtime/Core/Public/Stats/Stats2.h
+++ b/Engine/Source/Runtime/Core/Public/Stats/Stats2.h
@@ -2229,9 +2229,11 @@ DECLARE_STATS_GROUP(TEXT("MapBuildData"),STATGROUP_MapBuildData, STATCAT_Advance
 DECLARE_STATS_GROUP(TEXT("Shader Compiling"),STATGROUP_ShaderCompiling, STATCAT_Advanced);
 DECLARE_STATS_GROUP(TEXT("Shader Compression"),STATGROUP_Shaders, STATCAT_Advanced);
 DECLARE_STATS_GROUP(TEXT("Shadow Rendering"),STATGROUP_ShadowRendering, STATCAT_Advanced);
+DECLARE_STATS_GROUP_VERBOSE(TEXT("Shadow Rendering Verbose"), STATGROUP_ShadowRenderingVerbose, STATCAT_Advanced);
 DECLARE_STATS_GROUP(TEXT("Stat System"),STATGROUP_StatSystem, STATCAT_Advanced);
 DECLARE_STATS_GROUP_SORTBYNAME(TEXT("Streaming Overview"),STATGROUP_StreamingOverview, STATCAT_Advanced);
 DECLARE_STATS_GROUP_SORTBYNAME(TEXT("Streaming Details"),STATGROUP_StreamingDetails, STATCAT_Advanced);
+DECLARE_STATS_GROUP_VERBOSE(TEXT("Streaming Details Verbose"), STATGROUP_StreamingDetailsVerbose, STATCAT_Advanced);
 DECLARE_STATS_GROUP_SORTBYNAME(TEXT("Streaming"),STATGROUP_Streaming, STATCAT_Advanced);
 DECLARE_STATS_GROUP(TEXT("Target Platform"),STATGROUP_TargetPlatform, STATCAT_Advanced);
 DECLARE_STATS_GROUP(TEXT("Text"),STATGROUP_Text, STATCAT_Advanced);
diff --git a/Engine/Source/Runtime/Engine/Private/Particles/ParticleBeam2EmitterInstance.cpp b/Engine/Source/Runtime/Engine/Private/Particles/ParticleBeam2EmitterInstance.cpp
--- a/Engine/Source/Runtime/Engine/Private/Particles/ParticleBeam2EmitterInstance.cpp
+++ b/Engine/Source/Runtime/Engine/Private/Particles/ParticleBeam2EmitterInstance.cpp
@@ -1194,7 +1194,6 @@ FDynamicEmitterDataBase* FParticleBeam2EmitterInstance::GetDynamicData(bool bSel
 	// Allocate the dynamic data
 	FDynamicBeam2EmitterData* NewEmitterData = new FDynamicBeam2EmitterData(LODLevel->RequiredModule);
 	{
-		SCOPE_CYCLE_COUNTER(STAT_ParticleMemTime);
 		INC_DWORD_STAT(STAT_DynamicEmitterCount);
 		INC_DWORD_STAT(STAT_DynamicBeamCount);
 		INC_DWORD_STAT_BY(STAT_DynamicEmitterMem, sizeof(FDynamicBeam2EmitterData));
diff --git a/Engine/Source/Runtime/Engine/Private/Particles/ParticleEmitterInstances.cpp b/Engine/Source/Runtime/Engine/Private/Particles/ParticleEmitterInstances.cpp
--- a/Engine/Source/Runtime/Engine/Private/Particles/ParticleEmitterInstances.cpp
+++ b/Engine/Source/Runtime/Engine/Private/Particles/ParticleEmitterInstances.cpp
@@ -45,7 +45,6 @@ DEFINE_STAT(STAT_SortingTime);
 DEFINE_STAT(STAT_SpriteRenderingTime);
 DEFINE_STAT(STAT_SpriteTickTime);
 DEFINE_STAT(STAT_SpriteSpawnTime);
-DEFINE_STAT(STAT_SpriteUpdateTime);
 DEFINE_STAT(STAT_PSysCompTickTime);
 DEFINE_STAT(STAT_ParticlePoolTime);
 DEFINE_STAT(STAT_ParticleComputeTickTime);
@@ -94,7 +93,6 @@ DEFINE_STAT(STAT_GPUParticlesSimulationCommands);
 
 /** Particle memory stats */
 
-DEFINE_STAT(STAT_ParticleMemTime);
 DEFINE_STAT(STAT_GTParticleData);
 DEFINE_STAT(STAT_DynamicSpriteGTMem);
 DEFINE_STAT(STAT_DynamicSubUVGTMem);
@@ -226,9 +224,6 @@ FORCEINLINE static void* FastParticleSmallBlockAlloc(size_t AllocSize)
 
 FORCEINLINE static void FastParticleSmallBlockFree(void *RawMemory, size_t AllocSize)
 {
-	QUICK_SCOPE_CYCLE_COUNTER(STAT_PARTALLOC);
-
-
 #if FASTPARTICLEALLOC_CHECKSIZE
 	{
 		FScopeLock S(&GFastPoolSizeMapCriticalSection);
@@ -289,13 +284,11 @@ FORCEINLINE static void FastParticleSmallBlockFree(void *RawMemory, size_t Alloc
 #else
 FORCEINLINE static void* FastParticleSmallBlockAlloc(size_t AllocSize)
 {
-	QUICK_SCOPE_CYCLE_COUNTER(STAT_PARTALLOC);
 	return FMemory::Malloc(AllocSize);
 }
 
 FORCEINLINE static void FastParticleSmallBlockFree(void *RawMemory, size_t AllocSize)
 {
-	QUICK_SCOPE_CYCLE_COUNTER(STAT_PARTALLOC);
 	FMemory::Free(RawMemory);
 }
 
@@ -760,8 +753,6 @@ bool FParticleEmitterInstance::Resize(int32 NewMaxActiveParticles, bool bSetMaxA
 #endif
 
 		{
-			SCOPE_CYCLE_COUNTER(STAT_ParticleMemTime);
-
 			ParticleData = (uint8*) FMemory::Realloc(ParticleData, ParticleStride * NewMaxActiveParticles);
 			check(ParticleData);
 
@@ -838,7 +829,6 @@ void FParticleEmitterInstance::Tick(float DeltaTime, bool bSuppressSpawning)
 		ResetParticleParameters(DeltaTime);
 
 		// Update the particles
-		SCOPE_CYCLE_COUNTER(STAT_SpriteUpdateTime);
 		CurrentMaterial = LODLevel->RequiredModule->Material;
 		Tick_ModuleUpdate(DeltaTime, LODLevel);
 
@@ -1047,7 +1037,6 @@ float FParticleEmitterInstance::Tick_SpawnParticles(float DeltaTime, UParticleLO
 {
 	if (!bHaltSpawning && !bHaltSpawningExternal && !bSuppressSpawning && (EmitterTime >= 0.0f))
 	{
-		SCOPE_CYCLE_COUNTER(STAT_SpriteSpawnTime);
 		// If emitter is not done - spawn at current rate.
 		// If EmitterLoops is 0, then we loop forever, so always spawn.
 		if ((InCurrentLODLevel->RequiredModule->EmitterLoops == 0) ||
@@ -1984,6 +1973,7 @@ void FParticleEmitterInstance::CheckSpawnCount(int32 InNewCount, int32 InMaxCoun
  */
 float FParticleEmitterInstance::Spawn(float DeltaTime)
 {
+	SCOPE_CYCLE_COUNTER(STAT_SpriteSpawnTime);
 	UParticleLODLevel* LODLevel = GetCurrentLODLevelChecked();
 
 	// For beams, we probably want to ignore the SpawnRate distribution,
@@ -3057,7 +3047,6 @@ FDynamicEmitterDataBase* FParticleSpriteEmitterInstance::GetDynamicData(bool bSe
 	// Allocate the dynamic data
 	FDynamicSpriteEmitterData* NewEmitterData = new FDynamicSpriteEmitterData(LODLevel->RequiredModule);
 	{
-		SCOPE_CYCLE_COUNTER(STAT_ParticleMemTime);
 		INC_DWORD_STAT(STAT_DynamicEmitterCount);
 		INC_DWORD_STAT(STAT_DynamicSpriteCount);
 		INC_DWORD_STAT_BY(STAT_DynamicEmitterMem, sizeof(FDynamicSpriteEmitterData));
@@ -3703,7 +3692,6 @@ FDynamicEmitterDataBase* FParticleMeshEmitterInstance::GetDynamicData(bool bSele
 	// Allocate the dynamic data
 	FDynamicMeshEmitterData* NewEmitterData = new FDynamicMeshEmitterData(LODLevel->RequiredModule);
 	{
-		SCOPE_CYCLE_COUNTER(STAT_ParticleMemTime);
 		INC_DWORD_STAT(STAT_DynamicEmitterCount);
 		INC_DWORD_STAT(STAT_DynamicMeshCount);
 		INC_DWORD_STAT_BY(STAT_DynamicEmitterMem, sizeof(FDynamicMeshEmitterData));
diff --git a/Engine/Source/Runtime/Engine/Private/Particles/ParticleTrail2EmitterInstance.cpp b/Engine/Source/Runtime/Engine/Private/Particles/ParticleTrail2EmitterInstance.cpp
--- a/Engine/Source/Runtime/Engine/Private/Particles/ParticleTrail2EmitterInstance.cpp
+++ b/Engine/Source/Runtime/Engine/Private/Particles/ParticleTrail2EmitterInstance.cpp
@@ -2663,7 +2663,6 @@ FDynamicEmitterDataBase* FParticleRibbonEmitterInstance::GetDynamicData(bool bSe
 	// Allocate the dynamic data
 	FDynamicRibbonEmitterData* NewEmitterData = new FDynamicRibbonEmitterData(LODLevel->RequiredModule);
 	{
-		SCOPE_CYCLE_COUNTER(STAT_ParticleMemTime);
 		INC_DWORD_STAT(STAT_DynamicEmitterCount);
 		INC_DWORD_STAT(STAT_DynamicRibbonCount);
 		INC_DWORD_STAT_BY(STAT_DynamicEmitterMem, sizeof(FDynamicRibbonEmitterData));
@@ -3966,7 +3965,6 @@ FDynamicEmitterDataBase* FParticleAnimTrailEmitterInstance::GetDynamicData(bool
 	// Allocate the dynamic data
 	FDynamicAnimTrailEmitterData* NewEmitterData = new FDynamicAnimTrailEmitterData(LODLevel->RequiredModule);
 	{
-		SCOPE_CYCLE_COUNTER(STAT_ParticleMemTime);
 		INC_DWORD_STAT(STAT_DynamicEmitterCount);
 		INC_DWORD_STAT(STAT_DynamicAnimTrailCount);
 		INC_DWORD_STAT_BY(STAT_DynamicEmitterMem, sizeof(FDynamicAnimTrailEmitterData));
diff --git a/Engine/Source/Runtime/Engine/Private/SceneManagement.cpp b/Engine/Source/Runtime/Engine/Private/SceneManagement.cpp
--- a/Engine/Source/Runtime/Engine/Private/SceneManagement.cpp
+++ b/Engine/Source/Runtime/Engine/Private/SceneManagement.cpp
@@ -194,7 +194,6 @@ FMeshBatchAndRelevance::FMeshBatchAndRelevance(const FMeshBatch& InMesh, const F
 	Mesh(&InMesh),
 	PrimitiveSceneProxy(InPrimitiveSceneProxy)
 {
-	QUICK_SCOPE_CYCLE_COUNTER(STAT_FMeshBatchAndRelevance);
 	const FMaterial* Material = InMesh.MaterialRenderProxy->GetMaterial(FeatureLevel);
 	EBlendMode BlendMode = Material->GetBlendMode();
 	bHasOpaqueMaterial = (BlendMode == BLEND_Opaque);
diff --git a/Engine/Source/Runtime/Engine/Private/Streaming/LevelTextureManager.cpp b/Engine/Source/Runtime/Engine/Private/Streaming/LevelTextureManager.cpp
--- a/Engine/Source/Runtime/Engine/Private/Streaming/LevelTextureManager.cpp
+++ b/Engine/Source/Runtime/Engine/Private/Streaming/LevelTextureManager.cpp
@@ -9,6 +9,8 @@
 #include "Engine/Texture2D.h"
 #include "Engine/World.h"
 
+DECLARE_CYCLE_STAT(TEXT("Render Asset Incremental Update"), FStaticComponentRenderAssetManager_IncrementalUpdate, STATGROUP_StreamingDetailsVerbose);
+
 FLevelRenderAssetManager::FLevelRenderAssetManager(ULevel* InLevel, RenderAssetInstanceTask::FDoWorkTask& AsyncTask)
 	: Level(InLevel)
 	, bIsInitialized(false)
@@ -270,7 +272,7 @@ void FLevelRenderAssetManager::IncrementalUpdate(
 	float Percentage, 
 	bool bUseDynamicStreaming) 
 {
-	QUICK_SCOPE_CYCLE_COUNTER(FStaticComponentRenderAssetManager_IncrementalUpdate);
+	SCOPE_CYCLE_COUNTER(FStaticComponentRenderAssetManager_IncrementalUpdate);
 
 	check(Level);
 
diff --git a/Engine/Source/Runtime/Engine/Private/TickTaskManager.cpp b/Engine/Source/Runtime/Engine/Private/TickTaskManager.cpp
--- a/Engine/Source/Runtime/Engine/Private/TickTaskManager.cpp
+++ b/Engine/Source/Runtime/Engine/Private/TickTaskManager.cpp
@@ -27,8 +27,6 @@ DECLARE_CYCLE_STAT(TEXT("Queue Ticks"),STAT_QueueTicks,STATGROUP_Game);
 DECLARE_CYCLE_STAT(TEXT("Queue Ticks Wait"),STAT_QueueTicksWait,STATGROUP_Game);
 DECLARE_CYCLE_STAT(TEXT("Queue Tick Task"),STAT_QueueTickTask,STATGROUP_Game);
 DECLARE_CYCLE_STAT(TEXT("Post Queue Tick Task"),STAT_PostTickTask,STATGROUP_Game);
-DECLARE_CYCLE_STAT(TEXT("Cooldown Dequeuing"),STAT_DequeueCooldowns,STATGROUP_Game);
-DECLARE_CYCLE_STAT(TEXT("Gather Ticks for Parallel"),STAT_GatherTicksForParallel,STATGROUP_Game);
 DECLARE_CYCLE_STAT(TEXT("Finalize Parallel Queue"),STAT_FinalizeParallelQueue,STATGROUP_Game);
 DECLARE_CYCLE_STAT(TEXT("Schedule cooldowns"),STAT_ScheduleCooldowns,STATGROUP_Game);
 DECLARE_DWORD_COUNTER_STAT(TEXT("Ticks Queued"),STAT_TicksQueued,STATGROUP_Game);
@@ -743,7 +741,6 @@ public:
 
 		int32 CooldownTicksEnabled = 0;
 		{
-			SCOPE_CYCLE_COUNTER(STAT_DequeueCooldowns);
 			// Determine which cooled down ticks will be enabled this frame
 			float CumulativeCooldown = 0.f;
 			FTickFunction* TickFunction = AllCoolingDownTickFunctions.Head;
@@ -780,7 +777,6 @@ public:
 		bTickNewlySpawned = true;
 
 		{
-			SCOPE_CYCLE_COUNTER(STAT_GatherTicksForParallel);
 			for (TSet<FTickFunction*>::TIterator It(AllEnabledTickFunctions); It; ++It)
 			{
 				FTickFunction* TickFunction = *It;
@@ -788,7 +784,6 @@ public:
 			}
 		}
 		{
-			SCOPE_CYCLE_COUNTER(STAT_DequeueCooldowns);
 			// Determine which cooled down ticks will be enabled this frame
 			float CumulativeCooldown = 0.f;
 			FTickFunction* TickFunction = AllCoolingDownTickFunctions.Head;
diff --git a/Engine/Source/Runtime/Engine/Public/Collision.h b/Engine/Source/Runtime/Engine/Public/Collision.h
--- a/Engine/Source/Runtime/Engine/Public/Collision.h
+++ b/Engine/Source/Runtime/Engine/Public/Collision.h
@@ -27,7 +27,7 @@ DECLARE_CYCLE_STAT_EXTERN(TEXT("GeomOverlapAny"), STAT_Collision_GeomOverlapAny,
 DECLARE_CYCLE_STAT_EXTERN(TEXT("BodyInstanceOverlapMulti"), STAT_Collision_FBodyInstance_OverlapMulti, STATGROUP_Collision, );
 DECLARE_CYCLE_STAT_EXTERN(TEXT("BodyInstanceOverlapTest"), STAT_Collision_FBodyInstance_OverlapTest, STATGROUP_Collision, );
 DECLARE_CYCLE_STAT_EXTERN(TEXT("BodyInstanceLineTrace"), STAT_Collision_FBodyInstance_LineTrace, STATGROUP_Collision, );
-DECLARE_CYCLE_STAT_EXTERN(TEXT("PreFilter"), STAT_Collision_PreFilter, STATGROUP_Collision, );
+DECLARE_CYCLE_STAT_EXTERN(TEXT("PreFilter"), STAT_Collision_PreFilter, STATGROUP_CollisionVerbose, );
 
 /** Enable collision analyzer support */
 #if (1 && !(UE_BUILD_SHIPPING || UE_BUILD_TEST) && WITH_EDITOR && WITH_UNREAL_DEVELOPER_TOOLS)
diff --git a/Engine/Source/Runtime/Engine/Public/ParticleHelper.h b/Engine/Source/Runtime/Engine/Public/ParticleHelper.h
--- a/Engine/Source/Runtime/Engine/Public/ParticleHelper.h
+++ b/Engine/Source/Runtime/Engine/Public/ParticleHelper.h
@@ -259,7 +259,6 @@ DECLARE_CYCLE_STAT_EXTERN(TEXT("Sort Time RT"),STAT_SortingTime,STATGROUP_Partic
 DECLARE_CYCLE_STAT_EXTERN(TEXT("Sprite Render Time RT"),STAT_SpriteRenderingTime,STATGROUP_Particles, );
 DECLARE_CYCLE_STAT_EXTERN(TEXT("Sprite Tick Time GT"),STAT_SpriteTickTime,STATGROUP_Particles, );
 DECLARE_CYCLE_STAT_EXTERN(TEXT("Sprite Spawn Time GT"),STAT_SpriteSpawnTime,STATGROUP_Particles, );
-DECLARE_CYCLE_STAT_EXTERN(TEXT("Sprite Update Time GT"),STAT_SpriteUpdateTime,STATGROUP_Particles, );
 DECLARE_CYCLE_STAT_EXTERN(TEXT("PSys Comp Tick Time GT"),STAT_PSysCompTickTime,STATGROUP_Particles, );
 DECLARE_CYCLE_STAT_EXTERN(TEXT("Particle Collision Time GT"),STAT_ParticleCollisionTime,STATGROUP_Particles, );
 DECLARE_CYCLE_STAT_EXTERN(TEXT("Particle SkelMeshSurf Time GT"),STAT_ParticleSkelMeshSurfTime,STATGROUP_Particles, );
@@ -276,7 +275,6 @@ DECLARE_CYCLE_STAT_EXTERN(TEXT("Async Work Time"),STAT_ParticleAsyncTime,STATGRO
 DECLARE_CYCLE_STAT_EXTERN(TEXT("Wait For ASync Time"),STAT_ParticleAsyncWaitTime,STATGROUP_Particles, );   // can be either performed on this thread or a true wait
 DECLARE_CYCLE_STAT_EXTERN(TEXT("Update Bounds Time GT"),STAT_ParticleUpdateBounds,STATGROUP_Particles, );
 
-DECLARE_CYCLE_STAT_EXTERN(TEXT("Particle Memory Time"),STAT_ParticleMemTime,STATGROUP_ParticleMem, );
 DECLARE_MEMORY_STAT_EXTERN(TEXT("Ptcls Data GT Mem"),STAT_GTParticleData,STATGROUP_ParticleMem, );
 DECLARE_DWORD_COUNTER_STAT_EXTERN(TEXT("Ptcls Data GT Mem MAX"),STAT_GTParticleData_MAX,STATGROUP_ParticleMem, );
 DECLARE_DWORD_COUNTER_STAT_EXTERN(TEXT("Ptcls Data RT Mem"),STAT_RTParticleData,STATGROUP_ParticleMem, );
diff --git a/Engine/Source/Runtime/RenderCore/Private/RenderCore.cpp b/Engine/Source/Runtime/RenderCore/Private/RenderCore.cpp
--- a/Engine/Source/Runtime/RenderCore/Private/RenderCore.cpp
+++ b/Engine/Source/Runtime/RenderCore/Private/RenderCore.cpp
@@ -147,6 +147,7 @@ DEFINE_STAT(STAT_PerObjectShadows);
 DEFINE_STAT(STAT_PreShadows);
 DEFINE_STAT(STAT_CachedPreShadows);
 DEFINE_STAT(STAT_ShadowDynamicPathDrawCalls);
+DEFINE_STAT(STAT_AddSubjectPrimitive);
 
 DEFINE_STAT(STAT_TranslucentInjectTime);
 DEFINE_STAT(STAT_DirectLightRenderingTime);
diff --git a/Engine/Source/Runtime/RenderCore/Public/RenderCore.h b/Engine/Source/Runtime/RenderCore/Public/RenderCore.h
--- a/Engine/Source/Runtime/RenderCore/Public/RenderCore.h
+++ b/Engine/Source/Runtime/RenderCore/Public/RenderCore.h
@@ -90,6 +90,7 @@ DECLARE_CYCLE_STAT_EXTERN(TEXT("WholeScene Shadow Projections"),STAT_RenderWhole
 DECLARE_CYCLE_STAT_EXTERN(TEXT("WholeScene Shadow Depths"),STAT_RenderWholeSceneShadowDepthsTime,STATGROUP_ShadowRendering, RENDERCORE_API);
 DECLARE_CYCLE_STAT_EXTERN(TEXT("PerObject Shadow Projections"),STAT_RenderPerObjectShadowProjectionsTime,STATGROUP_ShadowRendering, RENDERCORE_API);
 DECLARE_CYCLE_STAT_EXTERN(TEXT("PerObject Shadow Depths"),STAT_RenderPerObjectShadowDepthsTime,STATGROUP_ShadowRendering, RENDERCORE_API);
+DECLARE_CYCLE_STAT_EXTERN(TEXT("Add Subject Primitive"), STAT_AddSubjectPrimitive, STATGROUP_ShadowRenderingVerbose, RENDERCORE_API);
 DECLARE_MEMORY_STAT_EXTERN(TEXT("Per-Frame Shadowmap Atlases"),STAT_ShadowmapAtlasMemory,STATGROUP_ShadowRendering, RENDERCORE_API); 
 DECLARE_MEMORY_STAT_EXTERN(TEXT("Cached Shadowmaps"),STAT_CachedShadowmapMemory,STATGROUP_ShadowRendering, RENDERCORE_API); 
 DECLARE_DWORD_COUNTER_STAT_EXTERN(TEXT("Whole Scene shadows"),STAT_WholeSceneShadows,STATGROUP_ShadowRendering, RENDERCORE_API);
diff --git a/Engine/Source/Runtime/Renderer/Private/ShadowSetup.cpp b/Engine/Source/Runtime/Renderer/Private/ShadowSetup.cpp
--- a/Engine/Source/Runtime/Renderer/Private/ShadowSetup.cpp
+++ b/Engine/Source/Runtime/Renderer/Private/ShadowSetup.cpp
@@ -1017,7 +1017,7 @@ bool FProjectedShadowInfo::ShouldDrawStaticMeshes(FViewInfo& InCurrentView, bool
 
 void FProjectedShadowInfo::AddSubjectPrimitive(FPrimitiveSceneInfo* PrimitiveSceneInfo, TArray<FViewInfo>* ViewArray, ERHIFeatureLevel::Type FeatureLevel, bool bRecordShadowSubjectsForMobileShading)
 {
-	QUICK_SCOPE_CYCLE_COUNTER(STAT_AddSubjectPrimitive);
+	SCOPE_CYCLE_COUNTER(STAT_AddSubjectPrimitive);
 
 	// Ray traced shadows use the GPU managed distance field object buffers, no CPU culling should be used
 	check(!bRayTracedDistanceField);
