Convex mesh collision elements now correctly render their transform in the static mesh editor.

#rnx

diff --git a/Engine/Source/Runtime/Engine/Private/PhysicsEngine/PhysDrawing.cpp b/Engine/Source/Runtime/Engine/Private/PhysicsEngine/PhysDrawing.cpp
--- a/Engine/Source/Runtime/Engine/Private/PhysicsEngine/PhysDrawing.cpp
+++ b/Engine/Source/Runtime/Engine/Private/PhysicsEngine/PhysDrawing.cpp
@@ -695,11 +695,10 @@ void FKConvexElem::DrawElemWire(FPrimitiveDrawInterface* PDI, const FTransform&
 		const int32 NumVerts = VertexData.Num();
 		TArray<FVector> TransformedVerts;
 		TransformedVerts.Reserve(NumVerts);
-		FTransform FinalTransform = Transform.GetRelativeTransformReverse(ElemTM);
 
 		for(const FVector& Vert : VertexData)
 		{
-			TransformedVerts.Add(FinalTransform.TransformPosition(Vert));
+			TransformedVerts.Add(Transform.TransformPosition(Vert));
 		}
 
 		for(int32 Base = 0; Base < NumIndices; Base += 3)
