Fix CIS (PhysicsAssetEditor missing pragma once; ImmediatePhysicsJointHandl Incorrect Copyright notice)

diff --git a/Engine/Source/Editor/PhysicsAssetEditor/Private/PhysicsAssetEditorPhysicsHandleComponent.h b/Engine/Source/Editor/PhysicsAssetEditor/Private/PhysicsAssetEditorPhysicsHandleComponent.h
--- a/Engine/Source/Editor/PhysicsAssetEditor/Private/PhysicsAssetEditorPhysicsHandleComponent.h
+++ b/Engine/Source/Editor/PhysicsAssetEditor/Private/PhysicsAssetEditorPhysicsHandleComponent.h
@@ -1,4 +1,5 @@
 // Copyright Epic Games, Inc. All Rights Reserved.
+#pragma once
 
 #include "PhysicsEngine/PhysicsHandleComponent.h"
 
diff --git a/Engine/Source/Runtime/Engine/Public/Physics/ImmediatePhysics/ImmediatePhysicsJointHandle.h b/Engine/Source/Runtime/Engine/Public/Physics/ImmediatePhysics/ImmediatePhysicsJointHandle.h
--- a/Engine/Source/Runtime/Engine/Public/Physics/ImmediatePhysics/ImmediatePhysicsJointHandle.h
+++ b/Engine/Source/Runtime/Engine/Public/Physics/ImmediatePhysics/ImmediatePhysicsJointHandle.h
@@ -1,4 +1,4 @@
-// Copyright 1998-2019 Epic Games, Inc. All Rights Reserved.
+// Copyright Epic Games, Inc. All Rights Reserved.
 
 #pragma once
 
