Fixed runtime scale update not working on physics geometry (convex and trimesh) that was spawned initially at scale (1,1,1).

diff --git a/Engine/Source/Runtime/Engine/Private/PhysicsEngine/BodyInstance.cpp b/Engine/Source/Runtime/Engine/Private/PhysicsEngine/BodyInstance.cpp
--- a/Engine/Source/Runtime/Engine/Private/PhysicsEngine/BodyInstance.cpp
+++ b/Engine/Source/Runtime/Engine/Private/PhysicsEngine/BodyInstance.cpp
@@ -1808,9 +1808,9 @@ bool FBodyInstance::UpdateBodyScale(const FVector& InScale3D, bool bForceUpdate)
 				case ImplicitObjectType::Convex:
 				{
 
-					if (!CHAOS_ENSURE(IsScaled(ImplicitType)))
+					if (!CHAOS_ENSURE(IsScaled(ImplicitType) || IsInstanced(ImplicitType)))
 					{
-						// Currently assuming all Convex are scaled implicits.
+						// Currently assuming all convexes are scaled or instanced (if scale == 1).
 						break;
 					}
 
@@ -1834,9 +1834,9 @@ bool FBodyInstance::UpdateBodyScale(const FVector& InScale3D, bool bForceUpdate)
 				}
 				case ImplicitObjectType::TriangleMesh:
 				{
-					if (!CHAOS_ENSURE(IsScaled(ImplicitType)))
+					if(!CHAOS_ENSURE(IsScaled(ImplicitType) || IsInstanced(ImplicitType)))
 					{
-						// Currently assuming all TriangleMesh are scaled implicits.
+						// Currently assuming all triangle meshes are scaled or instanced (if scale == 1).
 						break;
 					}
 
