FMaterial::GetStaticParameterSet should not be called when FPlatformProperties::RequiresCookedData() is true. It can create otherwise spikes up to 200ms++

#rnx

diff --git a/Engine/Source/Developer/MaterialBaking/Private/ExportMaterialProxy.h b/Engine/Source/Developer/MaterialBaking/Private/ExportMaterialProxy.h
--- a/Engine/Source/Developer/MaterialBaking/Private/ExportMaterialProxy.h
+++ b/Engine/Source/Developer/MaterialBaking/Private/ExportMaterialProxy.h
@@ -273,7 +273,7 @@ public:
 			break;
 		};
 
-		CacheShaders(ResourceId, StaticParamSet, GMaxRHIShaderPlatform);
+		CacheShaders(ResourceId, &StaticParamSet, GMaxRHIShaderPlatform);
 	}
 
 	/** This override is required otherwise the shaders aren't ready for use when the surface is rendered resulting in a blank image */
diff --git a/Engine/Source/Developer/MaterialUtilities/Private/MaterialUtilities.cpp b/Engine/Source/Developer/MaterialUtilities/Private/MaterialUtilities.cpp
--- a/Engine/Source/Developer/MaterialUtilities/Private/MaterialUtilities.cpp
+++ b/Engine/Source/Developer/MaterialUtilities/Private/MaterialUtilities.cpp
@@ -451,7 +451,7 @@ public:
 			break;
 		};
 		
-		CacheShaders(ResourceId, StaticParamSet, GMaxRHIShaderPlatform);
+		CacheShaders(ResourceId, &StaticParamSet, GMaxRHIShaderPlatform);
 	}
 
 	virtual bool IsUsedWithStaticLighting() const { return true; }
diff --git a/Engine/Source/Editor/UnrealEd/Private/Lightmass/LightmassRender.cpp b/Engine/Source/Editor/UnrealEd/Private/Lightmass/LightmassRender.cpp
--- a/Engine/Source/Editor/UnrealEd/Private/Lightmass/LightmassRender.cpp
+++ b/Engine/Source/Editor/UnrealEd/Private/Lightmass/LightmassRender.cpp
@@ -263,7 +263,7 @@ public:
 
 				// Override with a special usage so we won't re-use the shader map used by the material for rendering
 				ResourceId.Usage = GetShaderMapUsage();
-				CacheShaders(ResourceId, StaticParamSet, GMaxRHIShaderPlatform);
+				CacheShaders(ResourceId, &StaticParamSet, GMaxRHIShaderPlatform);
 			}
 		}
 	}
diff --git a/Engine/Source/Runtime/Engine/Private/DebugViewModeMaterialProxy.cpp b/Engine/Source/Runtime/Engine/Private/DebugViewModeMaterialProxy.cpp
--- a/Engine/Source/Runtime/Engine/Private/DebugViewModeMaterialProxy.cpp
+++ b/Engine/Source/Runtime/Engine/Private/DebugViewModeMaterialProxy.cpp
@@ -84,7 +84,7 @@ FDebugViewModeMaterialProxy::FDebugViewModeMaterialProxy(
 
 		ResourceId.Usage = Usage;
 
-		CacheShaders(ResourceId, StaticParamSet, ShaderPlatform);
+		CacheShaders(ResourceId, &StaticParamSet, ShaderPlatform);
 	}
 	else
 	{
diff --git a/Engine/Source/Runtime/Engine/Private/Materials/MaterialShared.cpp b/Engine/Source/Runtime/Engine/Private/Materials/MaterialShared.cpp
--- a/Engine/Source/Runtime/Engine/Private/Materials/MaterialShared.cpp
+++ b/Engine/Source/Runtime/Engine/Private/Materials/MaterialShared.cpp
@@ -1787,16 +1787,22 @@ bool FMaterial::CacheShaders(EShaderPlatform Platform, const ITargetPlatform* Ta
 	FAllowCachingStaticParameterValues AllowCachingStaticParameterValues(*this);
 	FMaterialShaderMapId NoStaticParametersId;
 	GetShaderMapId(Platform, NoStaticParametersId);
+	
+	if (FPlatformProperties::RequiresCookedData())
+	{ 
+		return CacheShaders(NoStaticParametersId, nullptr, Platform, TargetPlatform);
+	}
+
 	FStaticParameterSet StaticParameterSet;
 	GetStaticParameterSet(Platform, StaticParameterSet);
-	return CacheShaders(NoStaticParametersId, StaticParameterSet, Platform, TargetPlatform);
+	return CacheShaders(NoStaticParametersId, &StaticParameterSet, Platform, TargetPlatform);
 }
 
 /**
  * Caches the material shaders for the given static parameter set and platform.
  * This is used by material resources of UMaterialInstances.
  */
-bool FMaterial::CacheShaders(const FMaterialShaderMapId& ShaderMapId, const FStaticParameterSet &StaticParameterSet, EShaderPlatform Platform, const ITargetPlatform* TargetPlatform)
+bool FMaterial::CacheShaders(const FMaterialShaderMapId& ShaderMapId, const FStaticParameterSet *StaticParameterSet, EShaderPlatform Platform, const ITargetPlatform* TargetPlatform)
 {
 	bool bSucceeded = false;
 	UE_CLOG(!ShaderMapId.IsValid(), LogMaterial, Warning, TEXT("Invalid shader map ID caching shaders for '%s', will use default material."), *GetFriendlyName());
@@ -1921,7 +1927,8 @@ bool FMaterial::CacheShaders(const FMaterialShaderMapId& ShaderMapId, const FSta
 
 			// If there's no cached shader map for this material, compile a new one.
 			// This is just kicking off the async compile, GameThreadShaderMap will not be complete yet
-			bSucceeded = BeginCompileShaderMap(ShaderMapId, StaticParameterSet, Platform, ShaderMap, TargetPlatform);
+			check(StaticParameterSet != nullptr);
+			bSucceeded = BeginCompileShaderMap(ShaderMapId, *StaticParameterSet, Platform, ShaderMap, TargetPlatform);
 
 			if (!bSucceeded)
 			{
diff --git a/Engine/Source/Runtime/Engine/Public/MaterialShared.h b/Engine/Source/Runtime/Engine/Public/MaterialShared.h
--- a/Engine/Source/Runtime/Engine/Public/MaterialShared.h
+++ b/Engine/Source/Runtime/Engine/Public/MaterialShared.h
@@ -1401,7 +1401,7 @@ public:
 	 * Caches the material shaders for the given static parameter set and platform.
 	 * This is used by material resources of UMaterialInstances.
 	 */
-	ENGINE_API bool CacheShaders(const FMaterialShaderMapId& ShaderMapId, const FStaticParameterSet &StaticParameterSet, EShaderPlatform Platform, const ITargetPlatform* TargetPlatform = nullptr);
+	ENGINE_API bool CacheShaders(const FMaterialShaderMapId& ShaderMapId, const FStaticParameterSet *StaticParameterSet, EShaderPlatform Platform, const ITargetPlatform* TargetPlatform = nullptr);
 
 	/**
 	 * Should the shader for this material with the given platform, shader type and vertex 
