Remove return statement in void function

diff --git a/Engine/Source/Runtime/Engine/Private/PhysicsEngine/Experimental/PhysInterface_Chaos.cpp b/Engine/Source/Runtime/Engine/Private/PhysicsEngine/Experimental/PhysInterface_Chaos.cpp
--- a/Engine/Source/Runtime/Engine/Private/PhysicsEngine/Experimental/PhysInterface_Chaos.cpp
+++ b/Engine/Source/Runtime/Engine/Private/PhysicsEngine/Experimental/PhysInterface_Chaos.cpp
@@ -537,7 +537,7 @@ void FPhysInterface_Chaos::SetLinearVelocity_AssumesLocked(const FPhysicsActorHa
 		Chaos::TKinematicGeometryParticle<float, 3>* Kinematic = InActorReference->CastToKinematicParticle();
 		if (ensure(Kinematic))
 		{
-			return Kinematic->SetV(InNewVelocity);
+			Kinematic->SetV(InNewVelocity);
 		}
 	}
 }
