integrate fix 11007881 for build error

diff --git a/Engine/Source/Runtime/Engine/Private/PhysicsEngine/Experimental/ChaosDerivedData.cpp b/Engine/Source/Runtime/Engine/Private/PhysicsEngine/Experimental/ChaosDerivedData.cpp
--- a/Engine/Source/Runtime/Engine/Private/PhysicsEngine/Experimental/ChaosDerivedData.cpp
+++ b/Engine/Source/Runtime/Engine/Private/PhysicsEngine/Experimental/ChaosDerivedData.cpp
@@ -191,7 +191,8 @@ void FChaosDerivedDataCooker::BuildTriangleMeshes(TArray<TUniquePtr<Chaos::FTria
 			}
 		}
 
-		OutTriangleMeshes.Emplace(new Chaos::FTriangleMeshImplicitObject(MoveTemp(TriMeshParticles), MoveTemp(Triangles), MoveTemp(MaterialIndices), MoveTemp(MakeUnique<TArray<int32>>(OutFaceRemap))));
+		TUniquePtr<TArray<int32>> OutFaceRemapPtr = MakeUnique<TArray<int32>>(OutFaceRemap);
+		OutTriangleMeshes.Emplace(new Chaos::FTriangleMeshImplicitObject(MoveTemp(TriMeshParticles), MoveTemp(Triangles), MoveTemp(MaterialIndices), MoveTemp(OutFaceRemapPtr);
 	};
 
 	if(FinalVerts.Num() < TNumericLimits<uint16>::Max())
