#UE4 CIS fix

diff --git a/Engine/Source/Runtime/Engine/Private/PhysicsEngine/Experimental/ChaosDerivedData.cpp b/Engine/Source/Runtime/Engine/Private/PhysicsEngine/Experimental/ChaosDerivedData.cpp
--- a/Engine/Source/Runtime/Engine/Private/PhysicsEngine/Experimental/ChaosDerivedData.cpp
+++ b/Engine/Source/Runtime/Engine/Private/PhysicsEngine/Experimental/ChaosDerivedData.cpp
@@ -192,7 +192,7 @@ void FChaosDerivedDataCooker::BuildTriangleMeshes(TArray<TUniquePtr<Chaos::FTria
 		}
 
 		TUniquePtr<TArray<int32>> OutFaceRemapPtr = MakeUnique<TArray<int32>>(OutFaceRemap);
-		OutTriangleMeshes.Emplace(new Chaos::FTriangleMeshImplicitObject(MoveTemp(TriMeshParticles), MoveTemp(Triangles), MoveTemp(MaterialIndices), MoveTemp(OutFaceRemapPtr);
+		OutTriangleMeshes.Emplace(new Chaos::FTriangleMeshImplicitObject(MoveTemp(TriMeshParticles), MoveTemp(Triangles), MoveTemp(MaterialIndices), MoveTemp(OutFaceRemapPtr)));
 	};
 
 	if(FinalVerts.Num() < TNumericLimits<uint16>::Max())
