Adjust default streaming priorities for virtual textures.
This now matches the standard texture streaming priorities.
Kept the CVar priority overrides but moved where these are resolved so as to keep the Async DCC logic consistent.

diff --git a/Engine/Source/Runtime/Engine/Private/VT/UploadingVirtualTexture.cpp b/Engine/Source/Runtime/Engine/Private/VT/UploadingVirtualTexture.cpp
--- a/Engine/Source/Runtime/Engine/Private/VT/UploadingVirtualTexture.cpp
+++ b/Engine/Source/Runtime/Engine/Private/VT/UploadingVirtualTexture.cpp
@@ -27,6 +27,32 @@ static TAutoConsoleVariable<int32> CVarVTCodecNumThreshold(
 	TEXT("Once number of VT codecs exceeds this number, attempt to retire codecs that haven't been recently used"),
 	ECVF_RenderThreadSafe);
 
+int32 GVirtualTextureIOPriority_NormalPagePri = (int32)AIOP_Low;
+static FAutoConsoleVariableRef CVarVirtualTextureIOPriority_NormalPagePri(
+	TEXT("r.VT.IOPriority_NormalPagePri"),
+	GVirtualTextureIOPriority_NormalPagePri,
+	TEXT("Priority of default priority VT I/O requests"),
+	ECVF_Default
+);
+
+int32 GVirtualTextureIOPriority_HighPagePri = (int32)AIOP_BelowNormal;
+static FAutoConsoleVariableRef CVarVirtualTextureIOPriority_HighPagePri(
+	TEXT("r.VT.IOPriority_HighPagePri"),
+	GVirtualTextureIOPriority_HighPagePri,
+	TEXT("Priority of high priority VT I/O requests"),
+	ECVF_Default
+);
+
+static EAsyncIOPriorityAndFlags GetAsyncIOPriority(EVTRequestPagePriority Priority)
+{
+	switch (Priority)
+	{
+	case EVTRequestPagePriority::High: return (EAsyncIOPriorityAndFlags)GVirtualTextureIOPriority_HighPagePri;
+	case EVTRequestPagePriority::Normal: return (EAsyncIOPriorityAndFlags)GVirtualTextureIOPriority_NormalPagePri;
+	default: check(false); return (EAsyncIOPriorityAndFlags)GVirtualTextureIOPriority_NormalPagePri;
+	}
+}
+
 FVirtualTextureCodec* FVirtualTextureCodec::ListHead = nullptr;
 FVirtualTextureCodec FVirtualTextureCodec::ListTail;
 uint32 FVirtualTextureCodec::NumCodecs = 0u;
@@ -264,7 +290,7 @@ struct FCreateCodecTask
 	}
 };
 
-FVTCodecAndStatus FUploadingVirtualTexture::GetCodecForChunk(FGraphEventArray& OutCompletionEvents, uint32 ChunkIndex, EAsyncIOPriorityAndFlags Priority)
+FVTCodecAndStatus FUploadingVirtualTexture::GetCodecForChunk(FGraphEventArray& OutCompletionEvents, uint32 ChunkIndex, EVTRequestPagePriority Priority)
 {
 	const FVirtualTextureDataChunk& Chunk = Data->Chunks[ChunkIndex];
 	if (Chunk.CodecPayloadSize == 0u)
@@ -308,7 +334,7 @@ FVTCodecAndStatus FUploadingVirtualTexture::GetCodecForChunk(FGraphEventArray& O
 	return FVTCodecAndStatus(EVTRequestPageStatus::Pending, Codec.Get());
 }
 
-FVTDataAndStatus FUploadingVirtualTexture::ReadData(FGraphEventArray& OutCompletionEvents, uint32 ChunkIndex, size_t Offset, size_t Size, EAsyncIOPriorityAndFlags Priority)
+FVTDataAndStatus FUploadingVirtualTexture::ReadData(FGraphEventArray& OutCompletionEvents, uint32 ChunkIndex, size_t Offset, size_t Size, EVTRequestPagePriority Priority)
 {
 	FVirtualTextureDataChunk& Chunk = Data->Chunks[ChunkIndex];
 	FByteBulkData& BulkData = Chunk.BulkData;
@@ -345,7 +371,7 @@ FVTDataAndStatus FUploadingVirtualTexture::ReadData(FGraphEventArray& OutComplet
 		// This way we can service these high priority tasks immediately
 		// Would be better to have DCC cache return a task event handle, which could be used to chain a subsequent read operation,
 		// but that would be more complicated, and this should generally not be a critical runtime path
-		const bool bAsyncDCC = (Priority & AIOP_PRIORITY_MASK) < AIOP_High;
+		const bool bAsyncDCC = (Priority == EVTRequestPagePriority::Normal);
 
 		const bool Available = GetVirtualTextureChunkDDCCache()->MakeChunkAvailable(&Data->Chunks[ChunkIndex], ChunkFileNameDCC, bAsyncDCC);
 		if (!Available)
@@ -385,7 +411,7 @@ FVTDataAndStatus FUploadingVirtualTexture::ReadData(FGraphEventArray& OutComplet
 		SET_MEMORY_STAT(STAT_FileCacheSize, IFileCacheHandle::GetFileCacheSize());
 	}
 
-	IMemoryReadStreamRef ReadData = Handle->ReadData(OutCompletionEvents, ChunkOffsetInFile + Offset, Size, Priority);
+	IMemoryReadStreamRef ReadData = Handle->ReadData(OutCompletionEvents, ChunkOffsetInFile + Offset, Size, GetAsyncIOPriority(Priority));
 	if (!ReadData)
 	{
 		return EVTRequestPageStatus::Saturated;
diff --git a/Engine/Source/Runtime/Engine/Private/VT/UploadingVirtualTexture.h b/Engine/Source/Runtime/Engine/Private/VT/UploadingVirtualTexture.h
--- a/Engine/Source/Runtime/Engine/Private/VT/UploadingVirtualTexture.h
+++ b/Engine/Source/Runtime/Engine/Private/VT/UploadingVirtualTexture.h
@@ -73,10 +73,10 @@ public:
 	inline const FVirtualTextureBuiltData* GetVTData() const { return Data; }
 
 	// gets the codec for the given chunk, data is not valid until returned OutCompletionEvents are complete
-	FVTCodecAndStatus GetCodecForChunk(FGraphEventArray& OutCompletionEvents, uint32 ChunkIndex, EAsyncIOPriorityAndFlags Priority);
+	FVTCodecAndStatus GetCodecForChunk(FGraphEventArray& OutCompletionEvents, uint32 ChunkIndex, EVTRequestPagePriority Priority);
 
 	// read a portion of a chunk
-	FVTDataAndStatus ReadData(FGraphEventArray& OutCompletionEvents, uint32 ChunkIndex, size_t Offset, size_t Size, EAsyncIOPriorityAndFlags Priority);
+	FVTDataAndStatus ReadData(FGraphEventArray& OutCompletionEvents, uint32 ChunkIndex, size_t Offset, size_t Size, EVTRequestPagePriority Priority);
 
 private:
 	friend class FVirtualTextureCodec;
diff --git a/Engine/Source/Runtime/Engine/Private/VT/VirtualTextureChunkManager.cpp b/Engine/Source/Runtime/Engine/Private/VT/VirtualTextureChunkManager.cpp
--- a/Engine/Source/Runtime/Engine/Private/VT/VirtualTextureChunkManager.cpp
+++ b/Engine/Source/Runtime/Engine/Private/VT/VirtualTextureChunkManager.cpp
@@ -23,22 +23,6 @@ static FAutoConsoleVariableRef CVarNumTranscodeRequests(
 	TEXT("Number of transcode request that can be in flight. default 32\n"),
 	ECVF_Default);
 
-int32 GVirtualTextureIOPriority_NormalPagePri = (int32)AIOP_Normal;
-static FAutoConsoleVariableRef CVarVirtualTextureIOPriority_NormalPagePri(
-	TEXT("r.VT.IOPriority_NormalPagePri"),
-	GVirtualTextureIOPriority_NormalPagePri,
-	TEXT("Priority of default priority I/O requests"),
-	ECVF_Default
-);
-
-int32 GVirtualTextureIOPriority_HighPagePri = (int32)AIOP_High;
-static FAutoConsoleVariableRef CVarVirtualTextureIOPriority_HighPagePri(
-	TEXT("r.VT.IOPriority_HighPagePri"),
-	GVirtualTextureIOPriority_HighPagePri,
-	TEXT("Priority of high priority I/O requests"),
-	ECVF_Default
-);
-
 /** 
  * FVirtualTextureChunkStreamingManager is a treated as a singleton referenced by multiple FUploadingVirtualTexture objects.
  * Use RefCount to ensure it is unregistered and destroyed on shutdown.
@@ -120,16 +104,6 @@ void FVirtualTextureChunkStreamingManager::CancelForcedResources()
 
 }
 
-static EAsyncIOPriorityAndFlags GetAsyncIOPriority(EVTRequestPagePriority Priority)
-{
-	switch (Priority)
-	{
-	case EVTRequestPagePriority::High: return (EAsyncIOPriorityAndFlags)GVirtualTextureIOPriority_HighPagePri;
-	case EVTRequestPagePriority::Normal: return (EAsyncIOPriorityAndFlags)GVirtualTextureIOPriority_NormalPagePri;
-	default: check(false); return (EAsyncIOPriorityAndFlags)GVirtualTextureIOPriority_NormalPagePri;
-	}
-}
-
 FVTRequestPageResult FVirtualTextureChunkStreamingManager::RequestTile(FUploadingVirtualTexture* VTexture, const FVirtualTextureProducerHandle& ProducerHandle, uint8 LayerMask, uint8 vLevel, uint32 vAddress, EVTRequestPagePriority Priority)
 {
 	SCOPE_CYCLE_COUNTER(STAT_VTP_RequestTile);
@@ -164,9 +138,8 @@ FVTRequestPageResult FVirtualTextureChunkStreamingManager::RequestTile(FUploadin
 		return EVTRequestPageStatus::Saturated;
 	}
 
-	const EAsyncIOPriorityAndFlags AsyncIOPriority = GetAsyncIOPriority(Priority);
 	FGraphEventArray GraphCompletionEvents;
-	const FVTCodecAndStatus CodecResult = VTexture->GetCodecForChunk(GraphCompletionEvents, ChunkIndex, AsyncIOPriority);
+	const FVTCodecAndStatus CodecResult = VTexture->GetCodecForChunk(GraphCompletionEvents, ChunkIndex, Priority);
 	if (!VTRequestPageStatus_HasData(CodecResult.Status))
 	{
 		// May fail to get codec if the file cache is saturated
@@ -189,7 +162,7 @@ FVTRequestPageResult FVirtualTextureChunkStreamingManager::RequestTile(FUploadin
 	const uint32 OffsetEnd = VTData->GetTileOffset(ChunkIndex, TileIndex + MaxLayerIndex + 1u);
 	const uint32 RequestSize = OffsetEnd - OffsetStart;
 
-	const FVTDataAndStatus TileDataResult = VTexture->ReadData(GraphCompletionEvents, ChunkIndex, OffsetStart, RequestSize, AsyncIOPriority);
+	const FVTDataAndStatus TileDataResult = VTexture->ReadData(GraphCompletionEvents, ChunkIndex, OffsetStart, RequestSize, Priority);
 	if (!VTRequestPageStatus_HasData(TileDataResult.Status))
 	{
 		return TileDataResult.Status;
