Reduce size of generated PPID thumbnails by removing mips.

diff --git a/Engine/Source/Runtime/Engine/Private/ImageUtils.cpp b/Engine/Source/Runtime/Engine/Private/ImageUtils.cpp
--- a/Engine/Source/Runtime/Engine/Private/ImageUtils.cpp
+++ b/Engine/Source/Runtime/Engine/Private/ImageUtils.cpp
@@ -221,12 +221,17 @@ UTexture2D* FImageUtils::CreateTexture2D(int32 SrcWidth, int32 SrcHeight, const
 
 	// Set compression options.
 	Tex2D->SRGB = InParams.bSRGB;
-	Tex2D->CompressionSettings	= InParams.CompressionSettings;
+	Tex2D->CompressionSettings = InParams.CompressionSettings;
+	Tex2D->MipGenSettings = InParams.MipGenSettings;
 	if( !InParams.bUseAlpha )
 	{
 		Tex2D->CompressionNoAlpha = true;
 	}
 	Tex2D->DeferCompression	= InParams.bDeferCompression;
+	if (InParams.TextureGroup != TEXTUREGROUP_MAX)
+	{
+		Tex2D->LODGroup = InParams.TextureGroup;
+	}
 
 	Tex2D->PostEditChange();
 	return Tex2D;
diff --git a/Engine/Source/Runtime/Engine/Public/ImageUtils.h b/Engine/Source/Runtime/Engine/Public/ImageUtils.h
--- a/Engine/Source/Runtime/Engine/Public/ImageUtils.h
+++ b/Engine/Source/Runtime/Engine/Public/ImageUtils.h
@@ -32,6 +32,12 @@ struct FCreateTexture2DParameters
 	/** If texture should be set as SRGB */
 	bool						bSRGB;
 
+	/** Mip-map generation settings */
+	TextureMipGenSettings		MipGenSettings;
+
+	/** Group this texture belongs to */
+	TextureGroup				TextureGroup;
+
 	/* The Guid hash to use part of the texture source's DDC key */
 	FGuid						SourceGuidHash;
 
@@ -39,7 +45,9 @@ struct FCreateTexture2DParameters
 		:	bUseAlpha(false),
 			CompressionSettings(TC_Default),
 			bDeferCompression(false),
-			bSRGB(true)
+			bSRGB(true),
+			MipGenSettings(TMGS_FromTextureGroup),
+			TextureGroup(TEXTUREGROUP_MAX)
 	{
 	}
 };
