Fix using RequireCPUAccess instead of just removing GIsEditor so that HISM can update at runtime if ComponentRequestsCPUAccess returns true

diff --git a/Engine/Source/Runtime/Engine/Private/HierarchicalInstancedStaticMesh.cpp b/Engine/Source/Runtime/Engine/Private/HierarchicalInstancedStaticMesh.cpp
--- a/Engine/Source/Runtime/Engine/Private/HierarchicalInstancedStaticMesh.cpp
+++ b/Engine/Source/Runtime/Engine/Private/HierarchicalInstancedStaticMesh.cpp
@@ -2502,9 +2502,10 @@ void UHierarchicalInstancedStaticMeshComponent::BuildTree()
 	QUICK_SCOPE_CYCLE_COUNTER(STAT_UHierarchicalInstancedStaticMeshComponent_BuildTree);
 
 	// upload instance edits to GPU, before validing if the mesh is valid, as it's possible that PerInstanceSMData.Num() == 0, so we have to hide everything before doing the build
-	if (GIsEditor && InstanceUpdateCmdBuffer.NumInlineCommands() > 0 && PerInstanceRenderData.IsValid())
+	if (InstanceUpdateCmdBuffer.NumInlineCommands() > 0 && PerInstanceRenderData.IsValid() && PerInstanceRenderData->InstanceBuffer.RequireCPUAccess)
 	{
-		// this is allowed only in editor, at runtime upload will happen when buffer is built from component data
+		// if instance data was modified, update GPU copy
+		// if InstanceBuffer was initialized with RequireCPUAccess (always true in editor))
 		PerInstanceRenderData->UpdateFromCommandBuffer(InstanceUpdateCmdBuffer);
 		MarkRenderStateDirty();
 	}
@@ -2813,9 +2814,10 @@ void UHierarchicalInstancedStaticMeshComponent::BuildTreeAsync()
 	check(BuildTreeAsyncTasks.Num() == 0);
 
 	// upload instance edits to GPU, before validing if the mesh is valid, as it's possible that PerInstanceSMData.Num() == 0, so we have to hide everything before doing the build
-	if (GIsEditor && InstanceUpdateCmdBuffer.NumInlineCommands() > 0 && PerInstanceRenderData.IsValid())
+	if (InstanceUpdateCmdBuffer.NumInlineCommands() > 0 && PerInstanceRenderData.IsValid() && PerInstanceRenderData->InstanceBuffer.RequireCPUAccess)
 	{
-		// this is allowed only in editor, at runtime upload will happen when buffer is built from component data
+		// if instance data was modified, update GPU copy
+		// if InstanceBuffer was initialized with RequireCPUAccess (always true in editor))
 		PerInstanceRenderData->UpdateFromCommandBuffer(InstanceUpdateCmdBuffer);
 		MarkRenderStateDirty();
 	}
@@ -2995,8 +2997,8 @@ FPrimitiveSceneProxy* UHierarchicalInstancedStaticMeshComponent::CreateSceneProx
 		check(InstancingRandomSeed != 0);
 
 		// if instance data was modified, update GPU copy
-		// generally happens only in editor 
-		if (GIsEditor && InstanceUpdateCmdBuffer.NumInlineCommands() > 0)
+		// if InstanceBuffer was initialized with RequireCPUAccess (always true in editor))
+		if (InstanceUpdateCmdBuffer.NumInlineCommands() > 0 && PerInstanceRenderData->InstanceBuffer.RequireCPUAccess)
 		{
 			PerInstanceRenderData->UpdateFromCommandBuffer(InstanceUpdateCmdBuffer);
 		}
