Quick fix to allow the SoundNodeModulatorContinuous class to be used in other modules.

#rnx

diff --git a/Engine/Source/Runtime/Engine/Classes/Sound/SoundNodeModulatorContinuous.h b/Engine/Source/Runtime/Engine/Classes/Sound/SoundNodeModulatorContinuous.h
--- a/Engine/Source/Runtime/Engine/Classes/Sound/SoundNodeModulatorContinuous.h
+++ b/Engine/Source/Runtime/Engine/Classes/Sound/SoundNodeModulatorContinuous.h
@@ -71,7 +71,7 @@ struct FModulatorContinuousParams
 /**
  * Allows named parameter based manipulation of pitch and volume
  */
-UCLASS(hidecategories=Object, editinlinenew, meta=( DisplayName="Continuous Modulator" ))
+UCLASS(hidecategories=Object, editinlinenew, MinimalAPI, meta=( DisplayName="Continuous Modulator" ))
 class USoundNodeModulatorContinuous : public USoundNode
 {
 	GENERATED_UCLASS_BODY()
