Add emergency failsafe for physics.

diff --git a/Engine/Source/Runtime/Engine/Private/PhysicsEngine/Experimental/PhysScene_Chaos.cpp b/Engine/Source/Runtime/Engine/Private/PhysicsEngine/Experimental/PhysScene_Chaos.cpp
--- a/Engine/Source/Runtime/Engine/Private/PhysicsEngine/Experimental/PhysScene_Chaos.cpp
+++ b/Engine/Source/Runtime/Engine/Private/PhysicsEngine/Experimental/PhysScene_Chaos.cpp
@@ -41,6 +41,7 @@
 #include "Engine/World.h"
 #include "DrawDebugHelpers.h"
 
+TAutoConsoleVariable<int32> CVar_ChaosSimulationEnable(TEXT("P.Chaos.Simulation.Enable"), 1, TEXT("Enable / disable chaos simulation. If disabled, physics will not tick."));
 TAutoConsoleVariable<int32> CVar_ChaosDrawHierarchyEnable(TEXT("P.Chaos.DrawHierarchy.Enable"), 0, TEXT("Enable / disable drawing of the physics hierarchy"));
 TAutoConsoleVariable<int32> CVar_ChaosDrawHierarchyCells(TEXT("P.Chaos.DrawHierarchy.Cells"), 0, TEXT("Enable / disable drawing of the physics hierarchy cells"));
 TAutoConsoleVariable<int32> CVar_ChaosDrawHierarchyBounds(TEXT("P.Chaos.DrawHierarchy.Bounds"), 1, TEXT("Enable / disable drawing of the physics hierarchy bounds"));
@@ -1385,6 +1386,11 @@ void FPhysScene_ChaosInterface::StartFrame()
 
 	SCOPE_CYCLE_COUNTER(STAT_Scene_StartFrame);
 
+	if (CVar_ChaosSimulationEnable.GetValueOnGameThread() == 0)
+	{
+		return;
+	}
+
 	FChaosSolversModule* SolverModule = FChaosSolversModule::GetModule();
 	checkSlow(SolverModule);
 
@@ -1403,11 +1409,6 @@ void FPhysScene_ChaosInterface::StartFrame()
 	}
 #endif
 
-	if (FPhysicsReplication* PhysicsReplication = Scene.GetPhysicsReplication())
-	{
-		PhysicsReplication->Tick(Dt);
-	}
-
 	if (Chaos::IDispatcher* Dispatcher = SolverModule->GetDispatcher())
 	{
 		if (FPhysicsSolver* Solver = GetSolver())
@@ -1462,6 +1463,11 @@ void FPhysScene_ChaosInterface::StartFrame()
 			break;
 		}
 	}
+
+	if (FPhysicsReplication* PhysicsReplication = Scene.GetPhysicsReplication())
+	{
+		PhysicsReplication->Tick(Dt);
+	}
 }
 
 void FPhysScene_ChaosInterface::EndFrame(ULineBatchComponent* InLineBatcher)
@@ -1470,6 +1476,11 @@ void FPhysScene_ChaosInterface::EndFrame(ULineBatchComponent* InLineBatcher)
 
 	SCOPE_CYCLE_COUNTER(STAT_Scene_EndFrame);
 
+	if (CVar_ChaosSimulationEnable.GetValueOnGameThread() == 0)
+	{
+		return;
+	}
+
 	FChaosSolversModule* SolverModule = FChaosSolversModule::GetModule();
 	checkSlow(SolverModule);
 
