Added STAT_ChaosTick and csv profiler scopes for Chaos

#rnx

diff --git a/Engine/Source/Runtime/Engine/Private/PhysicsEngine/Experimental/PhysScene_Chaos.cpp b/Engine/Source/Runtime/Engine/Private/PhysicsEngine/Experimental/PhysScene_Chaos.cpp
--- a/Engine/Source/Runtime/Engine/Private/PhysicsEngine/Experimental/PhysScene_Chaos.cpp
+++ b/Engine/Source/Runtime/Engine/Private/PhysicsEngine/Experimental/PhysScene_Chaos.cpp
@@ -466,6 +466,8 @@ bool FPhysScene_Chaos::IsTickable() const
 
 void FPhysScene_Chaos::Tick(float DeltaTime)
 {
+	SCOPE_CYCLE_COUNTER(STAT_ChaosTick);
+	CSV_SCOPED_TIMING_STAT_EXCLUSIVE(Physics);
 	LLM_SCOPE(ELLMTag::Chaos);
 
 #if WITH_EDITOR
diff --git a/Engine/Source/Runtime/Engine/Private/PhysicsEngine/PhysLevel.cpp b/Engine/Source/Runtime/Engine/Private/PhysicsEngine/PhysLevel.cpp
--- a/Engine/Source/Runtime/Engine/Private/PhysicsEngine/PhysLevel.cpp
+++ b/Engine/Source/Runtime/Engine/Private/PhysicsEngine/PhysLevel.cpp
@@ -187,6 +187,7 @@ void UWorld::FinishPhysicsSim()
 void FStartPhysicsTickFunction::ExecuteTick(float DeltaTime, enum ELevelTick TickType, ENamedThreads::Type CurrentThread, const FGraphEventRef& MyCompletionGraphEvent)
 {
 	QUICK_SCOPE_CYCLE_COUNTER(FStartPhysicsTickFunction_ExecuteTick);
+	CSV_SCOPED_TIMING_STAT_EXCLUSIVE(Physics);
 	check(Target);
 	Target->StartPhysicsSim();
 }
@@ -204,6 +205,7 @@ FName FStartPhysicsTickFunction::DiagnosticContext(bool bDetailed)
 void FEndPhysicsTickFunction::ExecuteTick(float DeltaTime, enum ELevelTick TickType, ENamedThreads::Type CurrentThread, const FGraphEventRef& MyCompletionGraphEvent)
 {
 	QUICK_SCOPE_CYCLE_COUNTER(FEndPhysicsTickFunction_ExecuteTick);
+	CSV_SCOPED_TIMING_STAT_EXCLUSIVE(Physics);
 
 	check(Target);
 	FPhysScene* PhysScene = Target->GetPhysicsScene();
diff --git a/Engine/Source/Runtime/Experimental/Chaos/Private/ChaosStats.cpp b/Engine/Source/Runtime/Experimental/Chaos/Private/ChaosStats.cpp
--- a/Engine/Source/Runtime/Experimental/Chaos/Private/ChaosStats.cpp
+++ b/Engine/Source/Runtime/Experimental/Chaos/Private/ChaosStats.cpp
@@ -2,6 +2,7 @@
 
 #include "ChaosStats.h"
 
+DEFINE_STAT(STAT_ChaosTick);
 DEFINE_STAT(STAT_PhysicsAdvance);
 DEFINE_STAT(STAT_SolverAdvance);
 DEFINE_STAT(STAT_HandleSolverCommands);
diff --git a/Engine/Source/Runtime/Experimental/Chaos/Public/ChaosStats.h b/Engine/Source/Runtime/Experimental/Chaos/Public/ChaosStats.h
--- a/Engine/Source/Runtime/Experimental/Chaos/Public/ChaosStats.h
+++ b/Engine/Source/Runtime/Experimental/Chaos/Public/ChaosStats.h
@@ -13,6 +13,7 @@ DECLARE_STATS_GROUP(TEXT("ChaosEngine"), STATGROUP_ChaosEngine, STATCAT_Advanced
 DECLARE_STATS_GROUP(TEXT("ChaosCollision"), STATGROUP_ChaosCollision, STATCAT_Advanced);
 DECLARE_STATS_GROUP(TEXT("ChaosJoint"), STATGROUP_ChaosJoint, STATCAT_Advanced);
 
+DECLARE_CYCLE_STAT_EXTERN(TEXT("Physics Tick"), STAT_ChaosTick, STATGROUP_Chaos, CHAOS_API);
 DECLARE_CYCLE_STAT_EXTERN(TEXT("Physics Advance"), STAT_PhysicsAdvance, STATGROUP_Chaos, CHAOS_API);
 DECLARE_CYCLE_STAT_EXTERN(TEXT("Solver Advance"), STAT_SolverAdvance, STATGROUP_Chaos, CHAOS_API);
 DECLARE_CYCLE_STAT_EXTERN(TEXT("Handle Solver Commands"), STAT_HandleSolverCommands, STATGROUP_Chaos, CHAOS_API);
diff --git a/Engine/Source/Runtime/Experimental/ChaosSolvers/Private/Framework/PhysicsTickTask.cpp b/Engine/Source/Runtime/Experimental/ChaosSolvers/Private/Framework/PhysicsTickTask.cpp
--- a/Engine/Source/Runtime/Experimental/ChaosSolvers/Private/Framework/PhysicsTickTask.cpp
+++ b/Engine/Source/Runtime/Experimental/ChaosSolvers/Private/Framework/PhysicsTickTask.cpp
@@ -7,6 +7,7 @@
 #include "ChaosStats.h"
 #include "PhysicsSolver.h"
 #include "PhysicsCoreTypes.h"
+#include "ProfilingDebugging/CsvProfiler.h"
 
 FAutoConsoleTaskPriority CPrio_FPhysicsTickTask(
 	TEXT("TaskGraph.TaskPriorities.PhysicsTickTask"),
@@ -163,6 +164,9 @@ ESubsequentsMode::Type FPhysicsSolverAdvanceTask::GetSubsequentsMode()
 
 void FPhysicsSolverAdvanceTask::DoTask(ENamedThreads::Type CurrentThread, const FGraphEventRef& MyCompletionGraphEvent)
 {
+	SCOPE_CYCLE_COUNTER(STAT_ChaosTick);
+	CSV_SCOPED_TIMING_STAT_EXCLUSIVE(Physics);
+
 	StepSolver(Solver, Dt);
 }
 
