Rendering - bCreateRenderStateForHiddenComponents renamed to bCreateRenderStateForHiddenComponentsWithCollsion, because that's really what it means, and adding another condition to check if IsCollisionEnabled() before we force a scene proxy to be created for a scene component.  This fixes it so that we're not forcing all lights (even those that are not visible) to start rendering lights.  Would make way more sense if things with collision just added a new proxy to represent their collision geometry, rather relying on creating a proxy for a hidden element just to draw the dynamic mesh data for the collision.  Perhaps every component should be able to add a Debug drawing scene proxy?

Kevin.Ortegren, Daniel.Wright

diff --git a/Engine/Source/Editor/UnrealEd/Private/EditorViewportClient.cpp b/Engine/Source/Editor/UnrealEd/Private/EditorViewportClient.cpp
--- a/Engine/Source/Editor/UnrealEd/Private/EditorViewportClient.cpp
+++ b/Engine/Source/Editor/UnrealEd/Private/EditorViewportClient.cpp
@@ -5134,9 +5134,9 @@ void FEditorViewportClient::UpdateHiddenCollisionDrawing()
 			bool bCollisionMode = EngineShowFlags.Collision || EngineShowFlags.CollisionVisibility || EngineShowFlags.CollisionPawn;
 
 			// Tell engine to create proxies for hidden components, so we can still draw collision
-			if (World->bCreateRenderStateForHiddenComponents != bCollisionMode)
+			if (World->bCreateRenderStateForHiddenComponentsWithCollsion != bCollisionMode)
 			{
-				World->bCreateRenderStateForHiddenComponents = bCollisionMode;
+				World->bCreateRenderStateForHiddenComponentsWithCollsion = bCollisionMode;
 
 				// Need to recreate scene proxies when this flag changes.
 				FGlobalComponentRecreateRenderStateContext Recreate;
diff --git a/Engine/Source/Runtime/Engine/Classes/Engine/World.h b/Engine/Source/Runtime/Engine/Classes/Engine/World.h
--- a/Engine/Source/Runtime/Engine/Classes/Engine/World.h
+++ b/Engine/Source/Runtime/Engine/Classes/Engine/World.h
@@ -1064,7 +1064,7 @@ public:
 
 #if !UE_BUILD_SHIPPING
 	/** If TRUE, 'hidden' components will still create render proxy, so can draw info (see USceneComponent::ShouldRender) */
-	uint8 bCreateRenderStateForHiddenComponents:1;
+	uint8 bCreateRenderStateForHiddenComponentsWithCollsion:1;
 #endif // !UE_BUILD_SHIPPING
 
 #if WITH_EDITOR
diff --git a/Engine/Source/Runtime/Engine/Private/Components/SceneComponent.cpp b/Engine/Source/Runtime/Engine/Private/Components/SceneComponent.cpp
--- a/Engine/Source/Runtime/Engine/Private/Components/SceneComponent.cpp
+++ b/Engine/Source/Runtime/Engine/Private/Components/SceneComponent.cpp
@@ -2920,7 +2920,7 @@ bool USceneComponent::ShouldRender() const
 
 #if !UE_BUILD_SHIPPING
 	// If we want to create render state even for hidden components, return true here
-	if (World && World->bCreateRenderStateForHiddenComponents)
+	if (World && World->bCreateRenderStateForHiddenComponentsWithCollsion && IsCollisionEnabled())
 	{
 		return true;
 	}
diff --git a/Engine/Source/Runtime/Engine/Private/GameViewportClient.cpp b/Engine/Source/Runtime/Engine/Private/GameViewportClient.cpp
--- a/Engine/Source/Runtime/Engine/Private/GameViewportClient.cpp
+++ b/Engine/Source/Runtime/Engine/Private/GameViewportClient.cpp
@@ -2958,7 +2958,7 @@ void UGameViewportClient::ToggleShowCollision()
 	if (World != nullptr)
 	{
 		// Tell engine to create proxies for hidden components, so we can still draw collision
-		World->bCreateRenderStateForHiddenComponents = bIsShowingCollision;
+		World->bCreateRenderStateForHiddenComponentsWithCollsion = bIsShowingCollision;
 
 		// Need to recreate scene proxies when this flag changes.
 		FGlobalComponentRecreateRenderStateContext Recreate;
diff --git a/Engine/Source/Runtime/Engine/Private/UnrealEngine.cpp b/Engine/Source/Runtime/Engine/Private/UnrealEngine.cpp
--- a/Engine/Source/Runtime/Engine/Private/UnrealEngine.cpp
+++ b/Engine/Source/Runtime/Engine/Private/UnrealEngine.cpp
@@ -12361,7 +12361,7 @@ bool UEngine::LoadMap( FWorldContext& WorldContext, FURL URL, class UPendingNetG
 	GInitRunaway();
 
 #if !UE_BUILD_SHIPPING
-	const bool bOldWorldWasShowingCollisionForHiddenComponents = WorldContext.World() && WorldContext.World()->bCreateRenderStateForHiddenComponents;
+	const bool bOldWorldWasShowingCollisionForHiddenComponents = WorldContext.World() && WorldContext.World()->bCreateRenderStateForHiddenComponentsWithCollsion;
 #endif
 
 	// Unload the current world
@@ -12675,7 +12675,7 @@ bool UEngine::LoadMap( FWorldContext& WorldContext, FURL URL, class UPendingNetG
 	WorldContext.World()->WorldType = WorldContext.WorldType;
 
 #if !UE_BUILD_SHIPPING
-	GWorld->bCreateRenderStateForHiddenComponents = bOldWorldWasShowingCollisionForHiddenComponents;
+	GWorld->bCreateRenderStateForHiddenComponentsWithCollsion = bOldWorldWasShowingCollisionForHiddenComponents;
 #endif
 
 	// Fixme: hacky but we need to set PackageFlags here if we are in a PIE Context.
diff --git a/Engine/Source/Runtime/Renderer/Private/RendererScene.cpp b/Engine/Source/Runtime/Renderer/Private/RendererScene.cpp
--- a/Engine/Source/Runtime/Renderer/Private/RendererScene.cpp
+++ b/Engine/Source/Runtime/Renderer/Private/RendererScene.cpp
@@ -1148,7 +1148,7 @@ void FScene::AddPrimitive(UPrimitiveComponent* Primitive)
 	FPrimitiveSceneInfo* PrimitiveSceneInfo = new FPrimitiveSceneInfo(Primitive, this);
 	PrimitiveSceneProxy->PrimitiveSceneInfo = PrimitiveSceneInfo;
 
-	// Cache the primitive's initial transform.
+	// Cache the primitives initial transform.
 	FMatrix RenderMatrix = Primitive->GetRenderMatrix();
 	FVector AttachmentRootPosition(0);
 
@@ -1628,7 +1628,7 @@ void FScene::AddLight(ULightComponent* Light)
 		Light->SceneProxy = Proxy;
 
 		// Update the light's transform and position.
-		Proxy->SetTransform(Light->GetComponentTransform().ToMatrixNoScale(),Light->GetLightPosition());
+		Proxy->SetTransform(Light->GetComponentTransform().ToMatrixNoScale(), Light->GetLightPosition());
 
 		// Create the light scene info.
 		Proxy->LightSceneInfo = new FLightSceneInfo(Proxy, true);
