Fix debug physics drawing for convex hulls

diff --git a/Engine/Source/Runtime/Engine/Private/PhysicsEngine/PhysDrawing.cpp b/Engine/Source/Runtime/Engine/Private/PhysicsEngine/PhysDrawing.cpp
--- a/Engine/Source/Runtime/Engine/Private/PhysicsEngine/PhysDrawing.cpp
+++ b/Engine/Source/Runtime/Engine/Private/PhysicsEngine/PhysDrawing.cpp
@@ -692,13 +692,15 @@ void FKConvexElem::DrawElemWire(FPrimitiveDrawInterface* PDI, const FTransform&
 	const int32 NumIndices = IndexData.Num();
 	if(NumIndices > 0 && ensure(NumIndices % 3 == 0))
 	{
+		// NOTE: With chaos, instead of using a mesh with transformed verts, we use the 
+		// VertexData directly, so we don't need to remove the body->elem transform like
+		// we did with physx.
 		const int32 NumVerts = VertexData.Num();
 		TArray<FVector> TransformedVerts;
 		TransformedVerts.Reserve(NumVerts);
-
 		for(const FVector& Vert : VertexData)
 		{
-			TransformedVerts.Add(Transform.TransformPosition(Vert));
+			TransformedVerts.Add(ElemTM.TransformPosition(Vert));
 		}
 
 		for(int32 Base = 0; Base < NumIndices; Base += 3)
