Fix ComplexAsSimple and SimpleAsComplex geometry querying against both original simple/complex and the 'new' simple/complex.

diff --git a/Engine/Source/Runtime/Engine/Private/PhysicsEngine/Experimental/ChaosInterfaceUtils.cpp b/Engine/Source/Runtime/Engine/Private/PhysicsEngine/Experimental/ChaosInterfaceUtils.cpp
--- a/Engine/Source/Runtime/Engine/Private/PhysicsEngine/Experimental/ChaosInterfaceUtils.cpp
+++ b/Engine/Source/Runtime/Engine/Private/PhysicsEngine/Experimental/ChaosInterfaceUtils.cpp
@@ -21,6 +21,7 @@
 
 #include "Physics/PhysicsInterfaceTypes.h"
 #include "PhysicsEngine/AggregateGeom.h"
+#include "PhysicsEngine/PhysicsSettings.h"
 
 #if WITH_PHYSX
 #include "PhysXIncludes.h"
@@ -141,6 +142,21 @@ namespace ChaosInterface
 		TArray<TUniquePtr<Chaos::FImplicitObject>>& Geoms = OutGeoms;
 		Chaos::TShapesArray<float, 3>& Shapes = OutShapes;
 
+		ECollisionTraceFlag CollisionTraceType = InParams.CollisionTraceType;
+		if (CollisionTraceType == CTF_UseDefault)
+		{
+			CollisionTraceType = UPhysicsSettings::Get()->DefaultShapeComplexity;
+		}
+
+		// Complex as simple should not create simple geometry, unless there is no complex geometry.  Otherwise both get queried against.
+		bool bMakeSimpleGeometry = (CollisionTraceType != CTF_UseComplexAsSimple) || (InParams.ChaosTriMeshes.Num() == 0);
+
+		// The reverse is true for Simple as Complex.
+		const int32 SimpleShapeCount = InParams.Geometry->SphereElems.Num() + InParams.Geometry->BoxElems.Num() + InParams.Geometry->ConvexElems.Num() + InParams.Geometry->SphylElems.Num();
+		bool bMakeComplexGeometery = (CollisionTraceType != CTF_UseSimpleAsComplex) || (SimpleShapeCount == 0);
+
+		ensure(bMakeComplexGeometery || bMakeSimpleGeometry);
+
 		auto NewShapeHelper = [&InParams](Chaos::TSerializablePtr<Chaos::FImplicitObject> InGeom, void* UserData, bool bComplexShape = false)
 		{
 			auto NewShape = Chaos::TPerShapeData<float, 3>::CreatePerShapeData();
@@ -153,146 +169,153 @@ namespace ChaosInterface
 			return NewShape;
 		};
 
-		for (uint32 i = 0; i < static_cast<uint32>(InParams.Geometry->SphereElems.Num()); ++i)
-		{
-			const FKSphereElem& SphereElem = InParams.Geometry->SphereElems[i];
-			const FKSphereElem ScaledSphereElem = SphereElem.GetFinalScaled(Scale, InParams.LocalTransform);
-			const float UseRadius = FMath::Max(ScaledSphereElem.Radius, KINDA_SMALL_NUMBER);
-			auto ImplicitSphere = MakeUnique<Chaos::TSphere<float, 3>>(ScaledSphereElem.Center, UseRadius);
-			auto NewShape = NewShapeHelper(MakeSerializable(ImplicitSphere), (void*)SphereElem.GetUserData());
-			Shapes.Emplace(MoveTemp(NewShape));
-			Geoms.Add(MoveTemp(ImplicitSphere));
-		}
-
-		for (uint32 i = 0; i < static_cast<uint32>(InParams.Geometry->BoxElems.Num()); ++i)
+		if (bMakeSimpleGeometry)
 		{
-			const FKBoxElem& BoxElem = InParams.Geometry->BoxElems[i];
-			const FKBoxElem ScaledBoxElem = BoxElem.GetFinalScaled(Scale, InParams.LocalTransform);
-			const FTransform& BoxTransform = ScaledBoxElem.GetTransform();
-			Chaos::TVector<float, 3> HalfExtents = Chaos::TVector<float, 3>(ScaledBoxElem.X * 0.5f, ScaledBoxElem.Y * 0.5f, ScaledBoxElem.Z * 0.5f);
-
-			HalfExtents.X = FMath::Max(HalfExtents.X, KINDA_SMALL_NUMBER);
-			HalfExtents.Y = FMath::Max(HalfExtents.Y, KINDA_SMALL_NUMBER);
-			HalfExtents.Z = FMath::Max(HalfExtents.Z, KINDA_SMALL_NUMBER);
-
-			// TAABB can handle translations internally but if we have a rotation we need to wrap it in a transform
-			TUniquePtr<Chaos::FImplicitObject> Implicit;
-			if (!BoxTransform.GetRotation().IsIdentity())
-			{
-				auto ImplicitBox = MakeUnique<Chaos::TBox<float, 3>>(-HalfExtents, HalfExtents);
-				Implicit = TUniquePtr<Chaos::FImplicitObject>(new Chaos::TImplicitObjectTransformed<float, 3>(MoveTemp(ImplicitBox), BoxTransform));
-			}
-			else
+			for (uint32 i = 0; i < static_cast<uint32>(InParams.Geometry->SphereElems.Num()); ++i)
 			{
-				Implicit = MakeUnique<Chaos::TBox<float, 3>>(BoxTransform.GetTranslation() - HalfExtents, BoxTransform.GetTranslation() + HalfExtents);
-			}
-
-			auto NewShape = NewShapeHelper(MakeSerializable(Implicit), (void*)BoxElem.GetUserData());
-			Shapes.Emplace(MoveTemp(NewShape));
-			Geoms.Add(MoveTemp(Implicit));
-		}
-		for (uint32 i = 0; i < static_cast<uint32>(InParams.Geometry->SphylElems.Num()); ++i)
-		{
-			const FKSphylElem& UnscaledSphyl = InParams.Geometry->SphylElems[i];
-			const FKSphylElem ScaledSphylElem = UnscaledSphyl.GetFinalScaled(Scale, InParams.LocalTransform);
-			float HalfHeight = FMath::Max(ScaledSphylElem.Length * 0.5f, KINDA_SMALL_NUMBER);
-			const float Radius = FMath::Max(ScaledSphylElem.Radius, KINDA_SMALL_NUMBER);
-
-			if (HalfHeight < KINDA_SMALL_NUMBER)
-			{
-				//not a capsule just use a sphere
-				auto ImplicitSphere = MakeUnique<Chaos::TSphere<float, 3>>(ScaledSphylElem.Center, Radius);
-				auto NewShape = NewShapeHelper(MakeSerializable(ImplicitSphere), (void*)UnscaledSphyl.GetUserData());
+				const FKSphereElem& SphereElem = InParams.Geometry->SphereElems[i];
+				const FKSphereElem ScaledSphereElem = SphereElem.GetFinalScaled(Scale, InParams.LocalTransform);
+				const float UseRadius = FMath::Max(ScaledSphereElem.Radius, KINDA_SMALL_NUMBER);
+				auto ImplicitSphere = MakeUnique<Chaos::TSphere<float, 3>>(ScaledSphereElem.Center, UseRadius);
+				auto NewShape = NewShapeHelper(MakeSerializable(ImplicitSphere), (void*)SphereElem.GetUserData());
 				Shapes.Emplace(MoveTemp(NewShape));
 				Geoms.Add(MoveTemp(ImplicitSphere));
-
 			}
-			else
+
+			for (uint32 i = 0; i < static_cast<uint32>(InParams.Geometry->BoxElems.Num()); ++i)
 			{
-				Chaos::TVector<float, 3> HalfExtents = ScaledSphylElem.Rotation.RotateVector(Chaos::TVector<float, 3>(0, 0, HalfHeight));
+				const FKBoxElem& BoxElem = InParams.Geometry->BoxElems[i];
+				const FKBoxElem ScaledBoxElem = BoxElem.GetFinalScaled(Scale, InParams.LocalTransform);
+				const FTransform& BoxTransform = ScaledBoxElem.GetTransform();
+				Chaos::TVector<float, 3> HalfExtents = Chaos::TVector<float, 3>(ScaledBoxElem.X * 0.5f, ScaledBoxElem.Y * 0.5f, ScaledBoxElem.Z * 0.5f);
+
+				HalfExtents.X = FMath::Max(HalfExtents.X, KINDA_SMALL_NUMBER);
+				HalfExtents.Y = FMath::Max(HalfExtents.Y, KINDA_SMALL_NUMBER);
+				HalfExtents.Z = FMath::Max(HalfExtents.Z, KINDA_SMALL_NUMBER);
+
+				// TAABB can handle translations internally but if we have a rotation we need to wrap it in a transform
+				TUniquePtr<Chaos::FImplicitObject> Implicit;
+				if (!BoxTransform.GetRotation().IsIdentity())
+				{
+					auto ImplicitBox = MakeUnique<Chaos::TBox<float, 3>>(-HalfExtents, HalfExtents);
+					Implicit = TUniquePtr<Chaos::FImplicitObject>(new Chaos::TImplicitObjectTransformed<float, 3>(MoveTemp(ImplicitBox), BoxTransform));
+				}
+				else
+				{
+					Implicit = MakeUnique<Chaos::TBox<float, 3>>(BoxTransform.GetTranslation() - HalfExtents, BoxTransform.GetTranslation() + HalfExtents);
+				}
 
-				auto ImplicitCapsule = MakeUnique<Chaos::TCapsule<float>>(ScaledSphylElem.Center - HalfExtents, ScaledSphylElem.Center + HalfExtents, Radius);
-				auto NewShape = NewShapeHelper(MakeSerializable(ImplicitCapsule), (void*)UnscaledSphyl.GetUserData());
+				auto NewShape = NewShapeHelper(MakeSerializable(Implicit), (void*)BoxElem.GetUserData());
 				Shapes.Emplace(MoveTemp(NewShape));
-				Geoms.Add(MoveTemp(ImplicitCapsule));
+				Geoms.Add(MoveTemp(Implicit));
 			}
-		}
-#if 0
-		for (uint32 i = 0; i < static_cast<uint32>(InParams.Geometry->TaperedCapsuleElems.Num()); ++i)
-		{
-			ensure(FMath::IsNearlyEqual(Scale[0], Scale[1]) && FMath::IsNearlyEqual(Scale[1], Scale[2]));
-			const auto& TCapsule = InParams.Geometry->TaperedCapsuleElems[i];
-			if (TCapsule.Length == 0)
+			for (uint32 i = 0; i < static_cast<uint32>(InParams.Geometry->SphylElems.Num()); ++i)
 			{
-				Chaos::TSphere<float, 3> * ImplicitSphere = new Chaos::TSphere<float, 3>(-half_extents, TCapsule.Radius * Scale[0]);
-				if (PhysicsProxy) PhysicsProxy->ImplicitObjects_GameThread.Add(ImplicitSphere);
-				else if (OutOptShapes) OutOptShapes->Add({ ImplicitSphere,true,true,InActor });
+				const FKSphylElem& UnscaledSphyl = InParams.Geometry->SphylElems[i];
+				const FKSphylElem ScaledSphylElem = UnscaledSphyl.GetFinalScaled(Scale, InParams.LocalTransform);
+				float HalfHeight = FMath::Max(ScaledSphylElem.Length * 0.5f, KINDA_SMALL_NUMBER);
+				const float Radius = FMath::Max(ScaledSphylElem.Radius, KINDA_SMALL_NUMBER);
+
+				if (HalfHeight < KINDA_SMALL_NUMBER)
+				{
+					//not a capsule just use a sphere
+					auto ImplicitSphere = MakeUnique<Chaos::TSphere<float, 3>>(ScaledSphylElem.Center, Radius);
+					auto NewShape = NewShapeHelper(MakeSerializable(ImplicitSphere), (void*)UnscaledSphyl.GetUserData());
+					Shapes.Emplace(MoveTemp(NewShape));
+					Geoms.Add(MoveTemp(ImplicitSphere));
+
+				}
+				else
+				{
+					Chaos::TVector<float, 3> HalfExtents = ScaledSphylElem.Rotation.RotateVector(Chaos::TVector<float, 3>(0, 0, HalfHeight));
+
+					auto ImplicitCapsule = MakeUnique<Chaos::TCapsule<float>>(ScaledSphylElem.Center - HalfExtents, ScaledSphylElem.Center + HalfExtents, Radius);
+					auto NewShape = NewShapeHelper(MakeSerializable(ImplicitCapsule), (void*)UnscaledSphyl.GetUserData());
+					Shapes.Emplace(MoveTemp(NewShape));
+					Geoms.Add(MoveTemp(ImplicitCapsule));
+				}
 			}
-			else
+#if 0
+			for (uint32 i = 0; i < static_cast<uint32>(InParams.Geometry->TaperedCapsuleElems.Num()); ++i)
 			{
-				Chaos::TVector<float, 3> half_extents(0, 0, TCapsule.Length / 2 * Scale[0]);
-				auto ImplicitCylinder = MakeUnique<Chaos::TCylinder<float>>(-half_extents, half_extents, TCapsule.Radius * Scale[0]);
-				if (PhysicsProxy) PhysicsProxy->ImplicitObjects_GameThread.Add(MoveTemp(ImplicitSphere));
-				else if (OutOptShapes) OutOptShapes->Add({ ImplicitSphere,true,true,InActor });
-
-				auto ImplicitSphereA = MakeUnique<Chaos::TSphere<float, 3>>(-half_extents, TCapsule.Radius * Scale[0]);
-				if (PhysicsProxy) PhysicsProxy->ImplicitObjects_GameThread.Add(MoveTemp(ImplicitSphereA));
-				else if (OutOptShapes) OutOptShapes->Add({ ImplicitSphereA,true,true,InActor });
-
-				auto ImplicitSphereB = MakeUnique<Chaos::TSphere<float, 3>>(half_extents, TCapsule.Radius * Scale[0]);
-				if (PhysicsProxy) PhysicsProxy->ImplicitObjects_GameThread.Add(MoveTemp(ImplicitSphereB));
-				else if (OutOptShapes) OutOptShapes->Add({ ImplicitSphereB,true,true,InActor });
+				ensure(FMath::IsNearlyEqual(Scale[0], Scale[1]) && FMath::IsNearlyEqual(Scale[1], Scale[2]));
+				const auto& TCapsule = InParams.Geometry->TaperedCapsuleElems[i];
+				if (TCapsule.Length == 0)
+				{
+					Chaos::TSphere<float, 3>* ImplicitSphere = new Chaos::TSphere<float, 3>(-half_extents, TCapsule.Radius * Scale[0]);
+					if (PhysicsProxy) PhysicsProxy->ImplicitObjects_GameThread.Add(ImplicitSphere);
+					else if (OutOptShapes) OutOptShapes->Add({ ImplicitSphere,true,true,InActor });
+				}
+				else
+				{
+					Chaos::TVector<float, 3> half_extents(0, 0, TCapsule.Length / 2 * Scale[0]);
+					auto ImplicitCylinder = MakeUnique<Chaos::TCylinder<float>>(-half_extents, half_extents, TCapsule.Radius * Scale[0]);
+					if (PhysicsProxy) PhysicsProxy->ImplicitObjects_GameThread.Add(MoveTemp(ImplicitSphere));
+					else if (OutOptShapes) OutOptShapes->Add({ ImplicitSphere,true,true,InActor });
+
+					auto ImplicitSphereA = MakeUnique<Chaos::TSphere<float, 3>>(-half_extents, TCapsule.Radius * Scale[0]);
+					if (PhysicsProxy) PhysicsProxy->ImplicitObjects_GameThread.Add(MoveTemp(ImplicitSphereA));
+					else if (OutOptShapes) OutOptShapes->Add({ ImplicitSphereA,true,true,InActor });
+
+					auto ImplicitSphereB = MakeUnique<Chaos::TSphere<float, 3>>(half_extents, TCapsule.Radius * Scale[0]);
+					if (PhysicsProxy) PhysicsProxy->ImplicitObjects_GameThread.Add(MoveTemp(ImplicitSphereB));
+					else if (OutOptShapes) OutOptShapes->Add({ ImplicitSphereB,true,true,InActor });
+				}
 			}
-		}
 #endif
 #if WITH_CHAOS && !PHYSICS_INTERFACE_PHYSX
-		for (uint32 i = 0; i < static_cast<uint32>(InParams.Geometry->ConvexElems.Num()); ++i)
-		{
-			const FKConvexElem& CollisionBody = InParams.Geometry->ConvexElems[i];
-			const FTransform& ConvexTransform = CollisionBody.GetTransform();
-			if (const auto& ConvexImplicit = CollisionBody.GetChaosConvexMesh())
+			for (uint32 i = 0; i < static_cast<uint32>(InParams.Geometry->ConvexElems.Num()); ++i)
 			{
-				//if (!ConvexTransform.GetTranslation().IsNearlyZero() || !ConvexTransform.GetRotation().IsIdentity())
-				//{
-				//	TUniquePtr<Chaos::FImplicitObject> TransformImplicit = TUniquePtr<Chaos::FImplicitObject>(new Chaos::TImplicitObjectTransformed<float, 3>(MakeSerializable(ConvexImplicit), ConvexTransform));
-				//	TUniquePtr<Chaos::TImplicitObjectScaled<float, 3, false>> Implicit = MakeUnique<Chaos::TImplicitObjectScaled<float, 3, false>>(MoveTemp(TransformImplicit), Scale);
-				//	auto NewShape = NewShapeHelper(MakeSerializable(Implicit));
-				//	Shapes.Emplace(MoveTemp(NewShape));
-				//	Geoms.Add(MoveTemp(Implicit));
-				//}
-				//else
+				const FKConvexElem& CollisionBody = InParams.Geometry->ConvexElems[i];
+				const FTransform& ConvexTransform = CollisionBody.GetTransform();
+				if (const auto& ConvexImplicit = CollisionBody.GetChaosConvexMesh())
 				{
-					TUniquePtr<Chaos::FImplicitObject> Implicit;
-					if(Scale == FVector(1))
-					{
-						Implicit = TUniquePtr<Chaos::FImplicitObject>(new Chaos::TImplicitObjectInstanced<Chaos::FConvex>(ConvexImplicit));
-					}
-					else
+					//if (!ConvexTransform.GetTranslation().IsNearlyZero() || !ConvexTransform.GetRotation().IsIdentity())
+					//{
+					//	TUniquePtr<Chaos::FImplicitObject> TransformImplicit = TUniquePtr<Chaos::FImplicitObject>(new Chaos::TImplicitObjectTransformed<float, 3>(MakeSerializable(ConvexImplicit), ConvexTransform));
+					//	TUniquePtr<Chaos::TImplicitObjectScaled<float, 3, false>> Implicit = MakeUnique<Chaos::TImplicitObjectScaled<float, 3, false>>(MoveTemp(TransformImplicit), Scale);
+					//	auto NewShape = NewShapeHelper(MakeSerializable(Implicit));
+					//	Shapes.Emplace(MoveTemp(NewShape));
+					//	Geoms.Add(MoveTemp(Implicit));
+					//}
+					//else
 					{
-						Implicit = TUniquePtr<Chaos::FImplicitObject>(new Chaos::TImplicitObjectScaled<Chaos::FConvex>(ConvexImplicit,Scale));
+						TUniquePtr<Chaos::FImplicitObject> Implicit;
+						if (Scale == FVector(1))
+						{
+							Implicit = TUniquePtr<Chaos::FImplicitObject>(new Chaos::TImplicitObjectInstanced<Chaos::FConvex>(ConvexImplicit));
+						}
+						else
+						{
+							Implicit = TUniquePtr<Chaos::FImplicitObject>(new Chaos::TImplicitObjectScaled<Chaos::FConvex>(ConvexImplicit, Scale));
+						}
+
+						auto NewShape = NewShapeHelper(MakeSerializable(Implicit), (void*)CollisionBody.GetUserData());
+						Shapes.Emplace(MoveTemp(NewShape));
+						Geoms.Add(MoveTemp(Implicit));
 					}
-
-					auto NewShape = NewShapeHelper(MakeSerializable(Implicit), (void*)CollisionBody.GetUserData());
-					Shapes.Emplace(MoveTemp(NewShape));
-					Geoms.Add(MoveTemp(Implicit));
 				}
 			}
 		}
 
-		for (const auto& ChaosTriMesh : InParams.ChaosTriMeshes)
+		if (bMakeComplexGeometery)
 		{
-			TUniquePtr<Chaos::FImplicitObject> Implicit;
-			if(Scale == FVector(1))
-			{
-				Implicit = TUniquePtr<Chaos::FImplicitObject>(new Chaos::TImplicitObjectInstanced<Chaos::FTriangleMeshImplicitObject>(ChaosTriMesh));
-			} else
+			for (const auto& ChaosTriMesh : InParams.ChaosTriMeshes)
 			{
-				Implicit = TUniquePtr<Chaos::FImplicitObject>(new Chaos::TImplicitObjectScaled<Chaos::FTriangleMeshImplicitObject>(ChaosTriMesh,Scale));
+				TUniquePtr<Chaos::FImplicitObject> Implicit;
+				if (Scale == FVector(1))
+				{
+					Implicit = TUniquePtr<Chaos::FImplicitObject>(new Chaos::TImplicitObjectInstanced<Chaos::FTriangleMeshImplicitObject>(ChaosTriMesh));
+				}
+				else
+				{
+					Implicit = TUniquePtr<Chaos::FImplicitObject>(new Chaos::TImplicitObjectScaled<Chaos::FTriangleMeshImplicitObject>(ChaosTriMesh, Scale));
+				}
+
+				auto NewShape = NewShapeHelper(MakeSerializable(Implicit), nullptr, true);
+				Shapes.Emplace(MoveTemp(NewShape));
+				Geoms.Add(MoveTemp(Implicit));
 			}
-			
-			auto NewShape = NewShapeHelper(MakeSerializable(Implicit), nullptr, true);
-			Shapes.Emplace(MoveTemp(NewShape));
-			Geoms.Add(MoveTemp(Implicit));
 		}
 #endif
 #if WITH_PHYSX && PHYSICS_INTERFACE_PHYSX
