[NO NON-PII COMMIT MESSAGE]

diff --git a/Engine/Source/Runtime/Engine/Private/AudioDevice.cpp b/Engine/Source/Runtime/Engine/Private/AudioDevice.cpp
--- a/Engine/Source/Runtime/Engine/Private/AudioDevice.cpp
+++ b/Engine/Source/Runtime/Engine/Private/AudioDevice.cpp
@@ -1667,6 +1667,13 @@ bool FAudioDevice::HandleAudioMemoryInfo(const TCHAR* Cmd, FOutputDevice& Ar)
 		ReportAr->Log(TEXT("Sound Wave Memory Report"));
 		ReportAr->Log(TEXT(""));
 
+		FString StreamingMemoryReport = IStreamingManager::Get().GetAudioStreamingManager().GenerateMemoryReport();
+
+		ReportAr->Log(TEXT("\n/*******************/\n"));
+		ReportAr->Log(TEXT("Streaming Audio Info:"));
+		ReportAr->Log(*StreamingMemoryReport);
+		ReportAr->Log(TEXT("\n/*******************/\n"));
+
 		if (SoundWaveObjects.Num())
 		{
 			// Alpha sort the sound wave objects
diff --git a/Engine/Source/Runtime/Engine/Private/AudioStreaming.cpp b/Engine/Source/Runtime/Engine/Private/AudioStreaming.cpp
--- a/Engine/Source/Runtime/Engine/Private/AudioStreaming.cpp
+++ b/Engine/Source/Runtime/Engine/Private/AudioStreaming.cpp
@@ -943,6 +943,11 @@ FWaveRequest& FLegacyAudioStreamingManager::GetWaveRequest(USoundWave* SoundWave
 	return *WaveRequest;
 }
 
+FString FLegacyAudioStreamingManager::GenerateMemoryReport()
+{
+	return FString(TEXT("Only supported when stream caching is enabled.\n"));
+}
+
 uint64 FLegacyAudioStreamingManager::TrimMemory(uint64 NumBytesToFree)
 {
 	UE_LOG(LogAudio, Warning, TEXT("IAudioStreamingManager::TrimMemory is only supported when Stream Caching is enabled."));
diff --git a/Engine/Source/Runtime/Engine/Private/AudioStreamingCache.cpp b/Engine/Source/Runtime/Engine/Private/AudioStreamingCache.cpp
--- a/Engine/Source/Runtime/Engine/Private/AudioStreamingCache.cpp
+++ b/Engine/Source/Runtime/Engine/Private/AudioStreamingCache.cpp
@@ -112,6 +112,30 @@ static FAutoConsoleCommand GResizeAudioCacheCommand(
 })
 );
 
+static FAutoConsoleCommand GEnableProfilingAudioCacheCommand(
+	TEXT("au.streamcaching.StartProfiling"),
+	TEXT("This will start a performance-intensive profiling mode for this streaming manager. Profile stats can be output with audiomemreport."),
+	FConsoleCommandDelegate::CreateStatic(
+		[]()
+{
+	IStreamingManager::Get().GetAudioStreamingManager().SetProfilingMode(true);
+
+	UE_LOG(LogAudio, Display, TEXT("Enabled profiling mode on the audio stream cache."));
+})
+);
+
+static FAutoConsoleCommand GDisableProfilingAudioCacheCommand(
+	TEXT("au.streamcaching.StopProfiling"),
+	TEXT("This will start a performance-intensive profiling mode for this streaming manager. Profile stats can be output with audiomemreport."),
+	FConsoleCommandDelegate::CreateStatic(
+		[]()
+{
+	IStreamingManager::Get().GetAudioStreamingManager().SetProfilingMode(false);
+
+	UE_LOG(LogAudio, Display, TEXT("Disabled profiling mode on the audio stream cache."));
+})
+);
+
 FCachedAudioStreamingManager::FCachedAudioStreamingManager(const FCachedAudioStreamingManagerParams& InitParams)
 {
 	check(FPlatformCompressionUtilities::IsCurrentPlatformUsingStreamCaching());
@@ -697,18 +721,76 @@ void FAudioChunkCache::StopLoggingCacheMisses()
 FString FAudioChunkCache::FlushCacheMissLog()
 {
 	FString ConcatenatedCacheMisses;
-	ConcatenatedCacheMisses.Append(TEXT("SoundWave:\tChunkIndex\n"));
+	ConcatenatedCacheMisses.Append(TEXT("All Cache Misses:\nSoundWave:\t, ChunkIndex\n"));
+
+	struct FMissedChunk 
+	{
+		FName SoundWaveName;
+		int32 ChunkIndex;
+		int32 MissCount;
+	};
+
+	struct FCacheMissSortPredicate
+	{
+		FORCEINLINE bool operator()(const FMissedChunk& A, const FMissedChunk& B) const
+		{
+			return A.MissCount < B.MissCount;
+		}
+	};
+
+	TMap<FChunkKey, int32> CacheMissCount;
 
 	FCacheMissInfo CacheMissInfo;
 	while (CacheMissQueue.Dequeue(CacheMissInfo))
 	{
 		ConcatenatedCacheMisses.Append(CacheMissInfo.SoundWaveName.ToString());
-		ConcatenatedCacheMisses.Append(TEXT("\t"));
+		ConcatenatedCacheMisses.Append(TEXT("\t, "));
 		ConcatenatedCacheMisses.AppendInt(CacheMissInfo.ChunkIndex);
 		ConcatenatedCacheMisses.Append(TEXT("\n"));
+
+		FChunkKey Chunk = 
+		{ 
+			nullptr, 
+			CacheMissInfo.SoundWaveName, 
+			CacheMissInfo.ChunkIndex,
+#if WITH_EDITOR
+			INDEX_NONE
+#endif
+		};
+
+		int32& MissCount = CacheMissCount.FindOrAdd(Chunk);
+		MissCount++;
+	}
+
+	// Sort our cache miss count map:
+	TArray<FMissedChunk> ChunkMissArray;
+	for (auto& CacheMiss : CacheMissCount)
+	{
+		FMissedChunk MissedChunk =
+		{
+			CacheMiss.Key.SoundWaveName,
+			CacheMiss.Key.ChunkIndex,
+			CacheMiss.Value
+		};
+
+		ChunkMissArray.Add(MissedChunk);
+	}
+
+	ChunkMissArray.Sort(FCacheMissSortPredicate());
+
+	FString TopChunkMissesLog = TEXT("Most Missed Chunks:\n");
+	TopChunkMissesLog += TEXT("Name:\t, Index:\t, Miss Count:\n");
+	for (FMissedChunk& MissedChunk : ChunkMissArray)
+	{
+		TopChunkMissesLog.Append(MissedChunk.SoundWaveName.ToString());
+		TopChunkMissesLog.Append(TEXT("\t, "));
+		TopChunkMissesLog.AppendInt(MissedChunk.ChunkIndex);
+		TopChunkMissesLog.Append(TEXT("\t, "));
+		TopChunkMissesLog.AppendInt(MissedChunk.MissCount);
+		TopChunkMissesLog.Append(TEXT("\n"));
 	}
 
-	return ConcatenatedCacheMisses;
+	return TopChunkMissesLog + TEXT("\n") + ConcatenatedCacheMisses;
 }
 
 FAudioChunkCache::FCacheElement* FAudioChunkCache::FindElementForKey(const FChunkKey& InKey)
@@ -1109,6 +1191,35 @@ int32 FCachedAudioStreamingManager::RenderStatAudioStreaming(UWorld* World, FVie
 	return Y + Height;
 }
 
+FString FCachedAudioStreamingManager::GenerateMemoryReport()
+{
+	FString OutputString;
+	for (FAudioChunkCache& Cache : CacheArray)
+	{
+		OutputString += Cache.DebugPrint();
+	}
+
+	return OutputString;
+}
+
+void FCachedAudioStreamingManager::SetProfilingMode(bool bEnabled)
+{
+	if (bEnabled)
+	{
+		for (FAudioChunkCache& Cache : CacheArray)
+		{
+			Cache.BeginLoggingCacheMisses();
+		}
+	}
+	else
+	{
+		for (FAudioChunkCache& Cache : CacheArray)
+		{
+			Cache.StopLoggingCacheMisses();
+		}
+	}
+}
+
 uint64 FCachedAudioStreamingManager::TrimMemory(uint64 NumBytesToFree)
 {
 	uint64 NumBytesLeftToFree = NumBytesToFree;
@@ -1296,4 +1407,110 @@ TPair<int, int> FAudioChunkCache::DebugDisplay(UWorld* World, FViewport* Viewpor
 	UEngine::GetSmallFont()->GetStringHeightAndWidth(*CacheMemoryUsage, CacheMemoryTextOffsetX, CacheMemoryTextOffsetY);
 
 	return TPair<int, int>(X + CacheTitleOffsetX + CacheMemoryTextOffsetX - InitialX, Y - InitialY);
-}
\ No newline at end of file
+}
+
+FString FAudioChunkCache::DebugPrint()
+{
+	FScopeLock ScopeLock(const_cast<FCriticalSection*>(&CacheMutationCriticalSection));
+
+	FString OutputString;
+
+	FString NumElementsDetail = *FString::Printf(TEXT("Number of chunks loaded: %d of %d"), ChunksInUse, CachePool.Num());
+
+	OutputString += NumElementsDetail + TEXT("\n");
+
+	// First pass: We run through and get a snap shot of the amount of memory currently in use.
+	FCacheElement* CurrentElement = MostRecentElement;
+	uint32 NumBytesCounter = 0;
+
+	uint32 NumBytesRetained = 0;
+
+	while (CurrentElement != nullptr)
+	{
+		// Note: this is potentially a stale value if we're in the middle of FCacheElement::KickOffAsyncLoad.
+		NumBytesCounter += CurrentElement->ChunkData.Num();
+
+		if (CurrentElement->IsInUse())
+		{
+			NumBytesRetained += CurrentElement->ChunkData.Num();
+		}
+
+		CurrentElement = CurrentElement->LessRecentElement;
+	}
+
+	// Convert to megabytes and print the total size:
+	const double NumMegabytesInUse = (double)NumBytesCounter / (1024 * 1024);
+	const double NumMegabytesRetained = (double)NumBytesRetained / (1024 * 1024);
+	
+	const double MaxCacheSizeMB = ((double)MemoryLimitBytes) / (1024 * 1024);
+	const double PercentageOfCacheRetained = NumMegabytesRetained / MaxCacheSizeMB;
+
+	FString CacheMemoryHeader = *FString::Printf(TEXT("Retaining:\t, Loaded:\t, Max Potential Usage:\t, \n"));
+	FString CacheMemoryUsage = *FString::Printf(TEXT("%.4f Megabytes (%.3f of total capacity)\t,  %.4f Megabytes (%lu bytes)\t, %.4f Megabytes\t, \n"), NumMegabytesRetained, PercentageOfCacheRetained, NumMegabytesInUse, MemoryCounterBytes.Load(), MaxCacheSizeMB);
+
+	OutputString += CacheMemoryHeader + CacheMemoryUsage + TEXT("\n");
+
+	// Second Pass: We're going to list the actual chunks in the cache.
+	CurrentElement = MostRecentElement;
+	int32 Index = 0;
+
+	OutputString += TEXT("Index:\t, Size (KB):\t, Chunk:\t, Request Count:\t, Average Index:\t, Number of Handles Retaining Chunk:\t, Chunk Load Time:\t, Name: \t, Notes:\t, \n");
+
+	// More detailed info about individual chunks here:
+	while (CurrentElement != nullptr)
+	{
+		// We use a CVar to clamp the max amount of chunks we display.
+		if (Index > DebugMaxElementsDisplayCVar)
+		{
+			break;
+		}
+
+		int32 NumTotalChunks = -1;
+		int32 NumTimesTouched = -1;
+		double TimeToLoad = -1.0;
+		float AveragePlaceInCache = -1.0f;
+		bool bWasCacheMiss = false;
+		bool bIsStaleChunk = false;
+
+#if DEBUG_STREAM_CACHE
+		NumTotalChunks = CurrentElement->DebugInfo.NumTotalChunks;
+		NumTimesTouched = CurrentElement->DebugInfo.NumTimesTouched;
+		TimeToLoad = CurrentElement->DebugInfo.TimeToLoad;
+		AveragePlaceInCache = CurrentElement->DebugInfo.AverageLocationInCacheWhenNeeded;
+		bWasCacheMiss = CurrentElement->DebugInfo.bWasCacheMiss;
+#endif
+
+#if WITH_EDITOR
+		// TODO: Worry about whether the sound wave is alive here. In most editor cases this is ok because the soundwave will always be loaded, but this may not be the case in the future.
+		bIsStaleChunk = (CurrentElement->Key.SoundWave == nullptr) || (CurrentElement->Key.SoundWave->CurrentChunkRevision.GetValue() != CurrentElement->Key.ChunkRevision);
+#endif
+
+		const bool bWasTrimmed = CurrentElement->ChunkData.Num() == 0;
+
+		FString ElementInfo = *FString::Printf(TEXT("%4i.\t, %6.2f KB\t, %d of %d\t, %d\t, %6.2f\t, %d\t,  %6.4fms\t, %s\t, %s %s %s"),
+			Index,
+			CurrentElement->ChunkData.Num() / 1024.0f,
+			CurrentElement->Key.ChunkIndex,
+			NumTotalChunks,
+			NumTimesTouched,
+			AveragePlaceInCache,
+			CurrentElement->NumConsumers.GetValue(),
+			TimeToLoad,
+			bWasTrimmed ? TEXT("TRIMMED CHUNK") : *CurrentElement->Key.SoundWaveName.ToString(),
+			bWasCacheMiss ? TEXT("(Cache Miss!)") : TEXT(""),
+			bIsStaleChunk ? TEXT("(Stale Chunk)") : TEXT(""),
+			CurrentElement->IsLoadInProgress() ? TEXT("(Loading In Progress)") : TEXT("")
+		);
+
+		
+		OutputString += ElementInfo + TEXT("\n");
+
+		CurrentElement = CurrentElement->LessRecentElement;
+		Index++;
+	}
+
+	OutputString += TEXT("Cache Miss Log:\n");
+	OutputString += FlushCacheMissLog();
+
+	return OutputString;
+}
diff --git a/Engine/Source/Runtime/Engine/Public/AudioStreaming.h b/Engine/Source/Runtime/Engine/Public/AudioStreaming.h
--- a/Engine/Source/Runtime/Engine/Public/AudioStreaming.h
+++ b/Engine/Source/Runtime/Engine/Public/AudioStreaming.h
@@ -277,6 +277,8 @@ struct FLegacyAudioStreamingManager : public IAudioStreamingManager
 	virtual FAudioChunkHandle GetLoadedChunk(const USoundWave* SoundWave, uint32 ChunkIndex, bool bBlockForLoad = false) const override;
 	virtual uint64 TrimMemory(uint64 NumBytesToFree) override;
 	virtual int32 RenderStatAudioStreaming(UWorld* World, FViewport* Viewport, FCanvas* Canvas, int32 X, int32 Y, const FVector* ViewLocation, const FRotator* ViewRotation) override;
+	virtual FString GenerateMemoryReport() override;
+	virtual void SetProfilingMode(bool bEnabled) override {};
 	// End IAudioStreamingManager interface
 
 	/** Called when an async callback is made on an async loading audio chunk request. */
@@ -317,4 +319,5 @@ protected:
 
 	/** Compressed audio info objects which are used to avoid deleting chunks with in-fligth decodes. */
 	TArray<ICompressedAudioInfo*> CompressedAudioInfos;
+
 };
diff --git a/Engine/Source/Runtime/Engine/Public/AudioStreamingCache.h b/Engine/Source/Runtime/Engine/Public/AudioStreamingCache.h
--- a/Engine/Source/Runtime/Engine/Public/AudioStreamingCache.h
+++ b/Engine/Source/Runtime/Engine/Public/AudioStreamingCache.h
@@ -43,6 +43,8 @@ public:
 			return (SoundWaveName == Other.SoundWaveName) && (ChunkIndex == Other.ChunkIndex);
 #endif
 		}
+
+		
 	};
 
 	FAudioChunkCache(uint32 InMaxChunkSize, uint32 NumChunks, uint64 InMemoryLimitInBytes);
@@ -96,6 +98,8 @@ public:
 	// This is for debugging purposes only. Prints the elements in the cache from most recently used to least.
 	// Returns the dimensions of this debug log so that multiple caches can be tiled across the screen.
 	TPair<int, int> DebugDisplay(UWorld* World, FViewport* Viewport, FCanvas* Canvas, int32 X, int32 Y, const FVector* ViewLocation, const FRotator* ViewRotation) const;
+	// Generate a formatted text file for this cache.
+	FString DebugPrint();
 
 private:
 
@@ -354,6 +358,8 @@ public:
 	virtual FAudioChunkHandle GetLoadedChunk(const USoundWave* SoundWave, uint32 ChunkIndex, bool bBlockForLoad = false) const override;
 	virtual uint64 TrimMemory(uint64 NumBytesToFree) override;
 	virtual int32 RenderStatAudioStreaming(UWorld* World, FViewport* Viewport, FCanvas* Canvas, int32 X, int32 Y, const FVector* ViewLocation, const FRotator* ViewRotation) override;
+	virtual FString GenerateMemoryReport() override;
+	virtual void SetProfilingMode(bool bEnabled) override;
 	// End IAudioStreamingManager interface
 
 protected:
@@ -381,3 +387,8 @@ protected:
 	
 
 };
+
+inline int32 GetTypeHash(const FAudioChunkCache::FChunkKey& InKey)
+{
+	return HashCombine(InKey.SoundWaveName.GetNumber(), InKey.ChunkIndex);
+}
diff --git a/Engine/Source/Runtime/Engine/Public/ContentStreaming.h b/Engine/Source/Runtime/Engine/Public/ContentStreaming.h
--- a/Engine/Source/Runtime/Engine/Public/ContentStreaming.h
+++ b/Engine/Source/Runtime/Engine/Public/ContentStreaming.h
@@ -537,6 +537,16 @@ struct IAudioStreamingManager : public IStreamingManager
 	 */
 	virtual int32 RenderStatAudioStreaming(UWorld* World, FViewport* Viewport, FCanvas* Canvas, int32 X, int32 Y, const FVector* ViewLocation, const FRotator* ViewRotation) = 0;
 
+	/**
+	 * Generate a memory report as a formatted string for this streaming manager.
+	 */
+	virtual FString GenerateMemoryReport() = 0;
+
+	/**
+	 * Whether to toggle a performance intensive profiling mode the streaming manager.
+	 */
+	virtual void SetProfilingMode(bool bEnabled) = 0;
+
 protected:
 	friend FAudioChunkHandle;
 
