Implement missing impulse functions

diff --git a/Engine/Source/Runtime/Engine/Private/PhysicsEngine/Experimental/PhysInterface_Chaos.cpp b/Engine/Source/Runtime/Engine/Private/PhysicsEngine/Experimental/PhysInterface_Chaos.cpp
--- a/Engine/Source/Runtime/Engine/Private/PhysicsEngine/Experimental/PhysInterface_Chaos.cpp
+++ b/Engine/Source/Runtime/Engine/Private/PhysicsEngine/Experimental/PhysInterface_Chaos.cpp
@@ -676,36 +676,66 @@ void FPhysInterface_Chaos::SetAngularDamping_AssumesLocked(const FPhysicsActorHa
 
 void FPhysInterface_Chaos::AddImpulse_AssumesLocked(const FPhysicsActorHandle& InActorReference, const FVector& InForce)
 {
-	// #todo : Implement
-    //InActorReference.GetScene()->AddForce(InForce, InActorReference.GetId());
-	CHAOS_ENSURE(false);
+	if (ensure(FPhysicsInterface::IsValid(InActorReference)))
+	{
+		Chaos::TPBDRigidParticle<float, 3>* Rigid = InActorReference->CastToRigidParticle();
+		if (ensure(Rigid))
+		{
+			AddVelocity_AssumesLocked(InActorReference, Rigid->InvM() * InForce);
+		}
+	}
 }
 
 void FPhysInterface_Chaos::AddAngularImpulseInRadians_AssumesLocked(const FPhysicsActorHandle& InActorReference, const FVector& InTorque)
 {
-	// #todo : Implement
-    //InActorReference.GetScene()->AddTorque(InTorque, InActorReference.GetId());
-	CHAOS_ENSURE(false);
+	if (ensure(FPhysicsInterface::IsValid(InActorReference)))
+	{
+		Chaos::TPBDRigidParticle<float, 3>* Rigid = InActorReference->CastToRigidParticle();
+		if (ensure(Rigid))
+		{
+			AddAngularVelocityInRadians_AssumesLocked(InActorReference, Chaos::Utilities::ComputeWorldSpaceInertia(Rigid->R(), Rigid->InvI()) * InTorque);
+		}
+	}
 }
 
-void FPhysInterface_Chaos::AddVelocity_AssumesLocked(const FPhysicsActorHandle& InActorReference, const FVector& InForce)
-{	
-	// #todo : Implement
-    //InActorReference.GetScene()->AddForce(InForce * InActorReference.GetScene()->Scene.GetSolver()->GetRigidParticles().M(InActorReference.GetScene()->GetIndexFromId(InActorReference.GetId())), InActorReference.GetId());
-	CHAOS_ENSURE(false);
+void FPhysInterface_Chaos::AddVelocity_AssumesLocked(const FPhysicsActorHandle& InActorReference, const FVector& InVelocityDelta)
+{
+	if (ensure(FPhysicsInterface::IsValid(InActorReference)))
+	{
+		Chaos::TKinematicGeometryParticle<float, 3>* Kinematic = InActorReference->CastToKinematicParticle();
+		if (ensure(Kinematic))
+		{
+			Kinematic->SetV(Kinematic->V() + InVelocityDelta);
+		}
+	}
 }
 
-void FPhysInterface_Chaos::AddAngularVelocityInRadians_AssumesLocked(const FPhysicsActorHandle& InActorReference, const FVector& InTorque)
+void FPhysInterface_Chaos::AddAngularVelocityInRadians_AssumesLocked(const FPhysicsActorHandle& InActorReference, const FVector& InAngularVelocityDeltaRad)
 {
-	// #todo : Implement
-    //InActorReference.GetScene()->AddTorque(InActorReference.GetScene()->Scene.GetSolver()->GetRigidParticles().I(InActorReference.GetScene()->GetIndexFromId(InActorReference.GetId())) * Chaos::TVector<float, 3>(InTorque), InActorReference.GetId());
-	CHAOS_ENSURE(false);
+	if (ensure(FPhysicsInterface::IsValid(InActorReference)))
+	{
+		Chaos::TKinematicGeometryParticle<float, 3>* Kinematic = InActorReference->CastToKinematicParticle();
+		if (ensure(Kinematic))
+		{
+			// NOTE: Need to convert W to rad/sec?
+			Kinematic->SetW(Kinematic->W() + InAngularVelocityDeltaRad);
+		}
+	}
 }
 
 void FPhysInterface_Chaos::AddImpulseAtLocation_AssumesLocked(const FPhysicsActorHandle& InActorReference, const FVector& InImpulse, const FVector& InLocation)
 {
-    // @todo(mlentine): We don't currently have a way to apply an instantaneous force. Do we need this?
-	CHAOS_ENSURE(false);
+	if (ensure(FPhysicsInterface::IsValid(InActorReference)))
+	{
+		Chaos::TPBDRigidParticle<float, 3>* Rigid = InActorReference->CastToRigidParticle();
+		if (ensure(Rigid))
+		{
+			const Chaos::FVec3 WorldCOM = Chaos::FParticleUtilitiesGT::GetCoMWorldPosition(Rigid);
+			const Chaos::FVec3 AngularImpulse = Chaos::FVec3::CrossProduct(WorldCOM - InLocation, InImpulse);
+			AddImpulse_AssumesLocked(InActorReference, InImpulse);
+			AddAngularImpulseInRadians_AssumesLocked(InActorReference, AngularImpulse);
+		}
+	}
 }
 
 void FPhysInterface_Chaos::AddRadialImpulse_AssumesLocked(const FPhysicsActorHandle& InActorReference, const FVector& InOrigin, float InRadius, float InStrength, ERadialImpulseFalloff InFalloff, bool bInVelChange)
