Fix CIS error for PhysX build.

diff --git a/Engine/Source/Runtime/Engine/Private/PhysicsEngine/Experimental/ChaosInterfaceUtils.cpp b/Engine/Source/Runtime/Engine/Private/PhysicsEngine/Experimental/ChaosInterfaceUtils.cpp
--- a/Engine/Source/Runtime/Engine/Private/PhysicsEngine/Experimental/ChaosInterfaceUtils.cpp
+++ b/Engine/Source/Runtime/Engine/Private/PhysicsEngine/Experimental/ChaosInterfaceUtils.cpp
@@ -148,14 +148,19 @@ namespace ChaosInterface
 			CollisionTraceType = UPhysicsSettings::Get()->DefaultShapeComplexity;
 		}
 
+		bool bMakeSimpleGeometry = true;
+		bool bMakeComplexGeometry = true;
+
+#if WITH_CHAOS
 		// Complex as simple should not create simple geometry, unless there is no complex geometry.  Otherwise both get queried against.
 		bool bMakeSimpleGeometry = (CollisionTraceType != CTF_UseComplexAsSimple) || (InParams.ChaosTriMeshes.Num() == 0);
 
 		// The reverse is true for Simple as Complex.
 		const int32 SimpleShapeCount = InParams.Geometry->SphereElems.Num() + InParams.Geometry->BoxElems.Num() + InParams.Geometry->ConvexElems.Num() + InParams.Geometry->SphylElems.Num();
-		bool bMakeComplexGeometery = (CollisionTraceType != CTF_UseSimpleAsComplex) || (SimpleShapeCount == 0);
+		bool bMakeComplexGeometry = (CollisionTraceType != CTF_UseSimpleAsComplex) || (SimpleShapeCount == 0);
+#endif
 
-		ensure(bMakeComplexGeometery || bMakeSimpleGeometry);
+		ensure(bMakeComplexGeometry || bMakeSimpleGeometry);
 
 		auto NewShapeHelper = [&InParams](Chaos::TSerializablePtr<Chaos::FImplicitObject> InGeom, void* UserData, bool bComplexShape = false)
 		{
@@ -298,7 +303,7 @@ namespace ChaosInterface
 			}
 		}
 
-		if (bMakeComplexGeometery)
+		if (bMakeComplexGeometry)
 		{
 			for (const auto& ChaosTriMesh : InParams.ChaosTriMeshes)
 			{
@@ -316,8 +321,8 @@ namespace ChaosInterface
 				Shapes.Emplace(MoveTemp(NewShape));
 				Geoms.Add(MoveTemp(Implicit));
 			}
-		}
 #endif
+		}
 #if WITH_PHYSX && PHYSICS_INTERFACE_PHYSX
 		for (const auto& PhysXMesh : InParams.TriMeshes)
 		{
