Merging CL#11058544

Fix Build

diff --git a/Engine/Source/Runtime/Engine/Private/PhysicsEngine/Experimental/ChaosInterfaceUtils.cpp b/Engine/Source/Runtime/Engine/Private/PhysicsEngine/Experimental/ChaosInterfaceUtils.cpp
--- a/Engine/Source/Runtime/Engine/Private/PhysicsEngine/Experimental/ChaosInterfaceUtils.cpp
+++ b/Engine/Source/Runtime/Engine/Private/PhysicsEngine/Experimental/ChaosInterfaceUtils.cpp
@@ -148,9 +148,6 @@ namespace ChaosInterface
 			CollisionTraceType = UPhysicsSettings::Get()->DefaultShapeComplexity;
 		}
 
-		bool bMakeSimpleGeometry = true;
-		bool bMakeComplexGeometry = true;
-
 #if WITH_CHAOS
 		// Complex as simple should not create simple geometry, unless there is no complex geometry.  Otherwise both get queried against.
 		bool bMakeSimpleGeometry = (CollisionTraceType != CTF_UseComplexAsSimple) || (InParams.ChaosTriMeshes.Num() == 0);
@@ -158,6 +155,9 @@ namespace ChaosInterface
 		// The reverse is true for Simple as Complex.
 		const int32 SimpleShapeCount = InParams.Geometry->SphereElems.Num() + InParams.Geometry->BoxElems.Num() + InParams.Geometry->ConvexElems.Num() + InParams.Geometry->SphylElems.Num();
 		bool bMakeComplexGeometry = (CollisionTraceType != CTF_UseSimpleAsComplex) || (SimpleShapeCount == 0);
+#else
+		bool bMakeSimpleGeometry = true;
+		bool bMakeComplexGeometry = true;
 #endif
 
 		ensure(bMakeComplexGeometry || bMakeSimpleGeometry);
