Merging CL#11058859

#UE4 Fix uint32 init error

diff --git a/Engine/Source/Runtime/Engine/Private/AudioStreamingCache.cpp b/Engine/Source/Runtime/Engine/Private/AudioStreamingCache.cpp
--- a/Engine/Source/Runtime/Engine/Private/AudioStreamingCache.cpp
+++ b/Engine/Source/Runtime/Engine/Private/AudioStreamingCache.cpp
@@ -726,7 +726,7 @@ FString FAudioChunkCache::FlushCacheMissLog()
 	struct FMissedChunk 
 	{
 		FName SoundWaveName;
-		int32 ChunkIndex;
+		uint32 ChunkIndex;
 		int32 MissCount;
 	};
 
