Correct impulse at location math

diff --git a/Engine/Source/Runtime/Engine/Private/PhysicsEngine/Experimental/PhysInterface_Chaos.cpp b/Engine/Source/Runtime/Engine/Private/PhysicsEngine/Experimental/PhysInterface_Chaos.cpp
--- a/Engine/Source/Runtime/Engine/Private/PhysicsEngine/Experimental/PhysInterface_Chaos.cpp
+++ b/Engine/Source/Runtime/Engine/Private/PhysicsEngine/Experimental/PhysInterface_Chaos.cpp
@@ -731,7 +731,7 @@ void FPhysInterface_Chaos::AddImpulseAtLocation_AssumesLocked(const FPhysicsActo
 		if (ensure(Rigid))
 		{
 			const Chaos::FVec3 WorldCOM = Chaos::FParticleUtilitiesGT::GetCoMWorldPosition(Rigid);
-			const Chaos::FVec3 AngularImpulse = Chaos::FVec3::CrossProduct(WorldCOM - InLocation, InImpulse);
+			const Chaos::FVec3 AngularImpulse = Chaos::FVec3::CrossProduct(InLocation - WorldCOM, InImpulse);
 			AddImpulse_AssumesLocked(InActorReference, InImpulse);
 			AddAngularImpulseInRadians_AssumesLocked(InActorReference, AngularImpulse);
 		}
