RTHF shadow fixes:
- More flexible UV space allocation strategy for the heightfield atlas
- Fixed a bug where heightfield objects can be falsely culled
- Fixed a bug where heightfield object data is not updated when the atlas realloc

diff --git a/Engine/Shaders/Private/DistanceFieldLightingShared.ush b/Engine/Shaders/Private/DistanceFieldLightingShared.ush
--- a/Engine/Shaders/Private/DistanceFieldLightingShared.ush
+++ b/Engine/Shaders/Private/DistanceFieldLightingShared.ush
@@ -299,7 +299,7 @@ float4x4 LoadHeightFieldObjectWorldToLocalFromBuffer(uint ObjectIndex, Structure
 	return transpose(float4x4(M0, M1, M2, float4(0, 0, 0, 1)));
 }
 
-float4 LoadHeightFieldObjectRectFromBuffer(uint ObjectIndex, StructuredBuffer<float4> InBuffer)
+float4 LoadHeightFieldObjectSizeScaleFromBuffer(uint ObjectIndex, StructuredBuffer<float4> InBuffer)
 {
 	return InBuffer[ObjectIndex * CULLED_OBJECT_DATA_STRIDE + 3];
 }
@@ -384,9 +384,9 @@ float4 LoadHeightFieldObjectAtlasUVScaleBias(uint ObjectIndex)
 	return LoadHeightFieldObjectAtlasUVScaleBiasFromBuffer(ObjectIndex, CulledObjectData);
 }
 
-float4 LoadHeightFieldObjectRect(uint ObjectIndex)
+float4 LoadHeightFieldObjectSizeScale(uint ObjectIndex)
 {
-	return LoadHeightFieldObjectRectFromBuffer(ObjectIndex, CulledObjectData);
+	return LoadHeightFieldObjectSizeScaleFromBuffer(ObjectIndex, CulledObjectData);
 }
 
 // x = Offset in global buffer, y = NumLOD0, z = NumSurfels (all LODs), W = instance index
diff --git a/Engine/Shaders/Private/DistanceFieldShadowing.usf b/Engine/Shaders/Private/DistanceFieldShadowing.usf
--- a/Engine/Shaders/Private/DistanceFieldShadowing.usf
+++ b/Engine/Shaders/Private/DistanceFieldShadowing.usf
@@ -9,6 +9,14 @@
 #include "DistanceFieldLightingShared.ush"
 #include "DistanceFieldShadowingShared.ush"
 
+#if (DISTANCEFIELD_PRIMITIVE_TYPE == DFPT_SignedDistanceField)
+#define MAX_TRACE_SPHERE_RADIUS 100
+#define SELF_SHADOW_VERTICAL_BIAS 0
+#else
+#define MAX_TRACE_SPHERE_RADIUS 500
+#define SELF_SHADOW_VERTICAL_BIAS 100
+#endif
+
 uint ObjectBoundingGeometryIndexCount;
 
 void CopyCulledObjectData(uint DestIndex, uint SourceIndex)
@@ -18,6 +26,7 @@ void CopyCulledObjectData(uint DestIndex, uint SourceIndex)
 #else
 	float4 BoxOriginAndSphereRadius = LoadFloat4(SceneObjectBounds, SourceIndex * OBJECT_BOUNDS_STRIDE);
 	float4 BoxExtent = LoadFloat4(SceneObjectBounds, SourceIndex * OBJECT_BOUNDS_STRIDE + 1);
+	BoxExtent.z += MAX_TRACE_SPHERE_RADIUS;
 
 	RWCulledObjectBounds[DestIndex * OBJECT_BOUNDS_STRIDE] = BoxOriginAndSphereRadius;
 	RWCulledObjectBounds[DestIndex * OBJECT_BOUNDS_STRIDE + 1] = BoxExtent;
@@ -613,12 +622,12 @@ void DistanceFieldShadowingCS(
 		// Keeps result from going all the way sharp
 		float MinSphereRadius = .4f;
 		// Maintain reasonable culling bounds
-		float MaxSphereRadius = 100;
+		float MaxSphereRadius = MAX_TRACE_SPHERE_RADIUS;
 
 		// Reduce shadowing when we are close the ray origin (only used for heightfield)
 		float SelfShadowFadeDistance = 100;
 		// Mitigate self-shadow caused by discontinuity on heightfield borders (only used for heightfield)
-		float SelfShadowVerticalBias = 100;
+		float SelfShadowVerticalBias = SELF_SHADOW_VERTICAL_BIAS;
 
 #if HAS_PREVIOUS_OUTPUT
 	#if PLATFORM_SUPPORTS_TYPED_UAV_LOAD
@@ -628,13 +637,8 @@ void DistanceFieldShadowingCS(
 		float PrevResult = Texture2DSampleLevel(ShadowFactorsTexture, ShadowFactorsSampler, PrevResultUV, 0).x;
 	#endif
 #endif
-		bool bIsHeightField = DISTANCEFIELD_PRIMITIVE_TYPE == DFPT_HeightField;
 
-		if (bIsHeightField)
-		{
-			MaxSphereRadius = 1000;
-			OpaqueWorldPosition.z += SelfShadowVerticalBias;
-		}
+		OpaqueWorldPosition.z += SelfShadowVerticalBias;
 
 		// World space offset along the start of the ray to avoid incorrect self-shadowing
 		float RayStartOffset = 2 + RayStartOffsetDepthScale * SceneDepth;
@@ -671,6 +675,8 @@ void DistanceFieldShadowingCS(
 			bool bUseCulling = false;
 		#endif
 
+		bool bIsHeightField = DISTANCEFIELD_PRIMITIVE_TYPE == DFPT_HeightField;
+
 		if (bIsHeightField)
 		{
 			Result = ShadowRayTraceThroughCulledHeightFieldObjects(
diff --git a/Engine/Shaders/Private/DistanceFieldShadowingShared.ush b/Engine/Shaders/Private/DistanceFieldShadowingShared.ush
--- a/Engine/Shaders/Private/DistanceFieldShadowingShared.ush
+++ b/Engine/Shaders/Private/DistanceFieldShadowingShared.ush
@@ -277,7 +277,9 @@ float ShadowRayTraceThroughCulledHeightFieldObjects(
 		}
 
 		float4x4 WorldToLocal = LoadHeightFieldObjectWorldToLocal(ObjectIndex);
-		float2 HeightFieldSize = LoadHeightFieldObjectRect(ObjectIndex).zw;
+		float4 SizeScale = LoadHeightFieldObjectSizeScale(ObjectIndex);
+		float2 HeightFieldSize = SizeScale.xy;
+		float WorldToLocalScale = SizeScale.z;
 
 		float3 LocalRayStart = mul(float4(WorldRayStart, 1.0), WorldToLocal).xyz;
 		float3 LocalRayEnd = mul(float4(WorldRayEnd, 1.0), WorldToLocal).xyz;
@@ -298,7 +300,6 @@ float ShadowRayTraceThroughCulledHeightFieldObjects(
 			float StartRayTime = LocalRayLength * IntersectionNearFar.x;
 			float EndRayTime = LocalRayLength * IntersectionNearFar.y;
 
-			float WorldToLocalScale = length(WorldToLocal[0].xyz);
 			float LocalSelfShadowScale = 1.0 / (SelfShadowFadeDistance * WorldToLocalScale);
 			float LocalMaxSphereRadius = MaxSphereRadius * WorldToLocalScale;
 
diff --git a/Engine/Source/Runtime/Engine/Private/DistanceFieldAtlas.cpp b/Engine/Source/Runtime/Engine/Private/DistanceFieldAtlas.cpp
--- a/Engine/Source/Runtime/Engine/Private/DistanceFieldAtlas.cpp
+++ b/Engine/Source/Runtime/Engine/Private/DistanceFieldAtlas.cpp
@@ -133,14 +133,21 @@ static int32 GHeightFieldAtlasTileSize = 64;
 static FAutoConsoleVariableRef CVarHeightFieldAtlasTileSize(
 	TEXT("r.HeightFields.AtlasTileSize"),
 	GHeightFieldAtlasTileSize,
-	TEXT(""),
+	TEXT("Suballocation granularity"),
 	ECVF_RenderThreadSafe);
 
 static int32 GHeightFieldAtlasDimInTiles = 16;
 static FAutoConsoleVariableRef CVarHeightFieldAtlasDimInTiles(
 	TEXT("r.HeightFields.AtlasDimInTiles"),
 	GHeightFieldAtlasDimInTiles,
-	TEXT(""),
+	TEXT("Number of tiles the atlas has in one dimension"),
+	ECVF_RenderThreadSafe);
+
+static int32 GHeightFieldAtlasDownSampleLevel = 2;
+static FAutoConsoleVariableRef CVarHeightFieldAtlasDownSampleLevel(
+	TEXT("r.HeightFields.AtlasDownSampleLevel"),
+	GHeightFieldAtlasDownSampleLevel,
+	TEXT("Max number of times a suballocation can be down-sampled"),
 	ECVF_RenderThreadSafe);
 
 TGlobalResource<FDistanceFieldVolumeTextureAtlas> GDistanceFieldVolumeTextureAtlas = TGlobalResource<FDistanceFieldVolumeTextureAtlas>();
@@ -1216,12 +1223,21 @@ void FHeightFieldTextureAtlas::InitializeIfNeeded()
 {
 	const uint32 LocalTileSize = (uint32)GHeightFieldAtlasTileSize;
 	const uint32 LocalDimInTiles = (uint32)GHeightFieldAtlasDimInTiles;
+	const uint32 LocalDownSampleLevel = (uint32)GHeightFieldAtlasDownSampleLevel;
 
 	if (!AtlasTextureRHI
-		|| TileAllocator.TileSize != LocalTileSize
-		|| TileAllocator.DimInTiles != LocalDimInTiles)
+		|| AddrSpaceAllocator.TileSize != LocalTileSize
+		|| AddrSpaceAllocator.DimInTiles != LocalDimInTiles
+		|| MaxDownSampleLevel != LocalDownSampleLevel)
 	{
-		TileAllocator.Init(LocalTileSize, 1, LocalDimInTiles);
+		AddrSpaceAllocator.Init(LocalTileSize, 1, LocalDimInTiles);
+
+		for (int32 Idx = 0; Idx < PendingStreamingTextures.Num(); ++Idx)
+		{
+			UTexture2D* Texture = PendingStreamingTextures[Idx];
+			Texture->bForceMiplevelsToBeResident = false;
+		}
+		PendingStreamingTextures.Empty();
 
 		for (TSet<FAllocation>::TConstIterator It(CurrentAllocations); It; ++It)
 		{
@@ -1231,19 +1247,20 @@ void FHeightFieldTextureAtlas::InitializeIfNeeded()
 
 		CurrentAllocations.Reset();
 
-		const uint32 SizeX = TileAllocator.DimInTexels;
-		const uint32 SizeY = TileAllocator.DimInTexels;
+		const uint32 SizeX = AddrSpaceAllocator.DimInTexels;
+		const uint32 SizeY = AddrSpaceAllocator.DimInTexels;
 		const uint32 Flags = TexCreate_ShaderResource | TexCreate_UAV;
 		FRHIResourceCreateInfo CreateInfo(TEXT("HeightFieldAtlas"));
 
 		AtlasTextureRHI = RHICreateTexture2D(SizeX, SizeY, PF_R8G8, 1, 1, Flags, CreateInfo);
 		AtlasUAVRHI = RHICreateUnorderedAccessView(AtlasTextureRHI, 0);
 
+		MaxDownSampleLevel = LocalDownSampleLevel;
 		++Generation;
 	}
 }
 
-void FHeightFieldTextureAtlas::AddAllocation(const UTexture2D* Texture)
+void FHeightFieldTextureAtlas::AddAllocation(UTexture2D* Texture)
 {
 	check(Texture);
 
@@ -1272,7 +1289,7 @@ void FHeightFieldTextureAtlas::AddAllocation(const UTexture2D* Texture)
 	}
 }
 
-void FHeightFieldTextureAtlas::RemoveAllocation(const UTexture2D* Texture)
+void FHeightFieldTextureAtlas::RemoveAllocation(UTexture2D* Texture)
 {
 	FSetElementId Id = PendingAllocations.FindId(Texture);
 	if (Id.IsValidId())
@@ -1280,6 +1297,7 @@ void FHeightFieldTextureAtlas::RemoveAllocation(const UTexture2D* Texture)
 		check(PendingAllocations[Id].RefCount);
 		if (!--PendingAllocations[Id].RefCount)
 		{
+			check(!PendingStreamingTextures.Contains(Texture));
 			PendingAllocations.Remove(Id);
 		}
 		return;
@@ -1291,6 +1309,7 @@ void FHeightFieldTextureAtlas::RemoveAllocation(const UTexture2D* Texture)
 		check(FailedAllocations[Id].RefCount);
 		if (!--FailedAllocations[Id].RefCount)
 		{
+			check(!PendingStreamingTextures.Contains(Texture));
 			FailedAllocations.Remove(Id);
 		}
 		return;
@@ -1300,10 +1319,11 @@ void FHeightFieldTextureAtlas::RemoveAllocation(const UTexture2D* Texture)
 	if (Id.IsValidId())
 	{
 		FAllocation& Allocation = CurrentAllocations[Id];
-		check(Allocation.RefCount && Allocation.TileIdx != INDEX_NONE);
+		check(Allocation.RefCount && Allocation.Handle != INDEX_NONE);
 		if (!--Allocation.RefCount)
 		{
-			TileAllocator.FreeTile(Allocation.TileIdx);
+			AddrSpaceAllocator.Free(Allocation.Handle);
+			PendingStreamingTextures.Remove(Texture);
 			CurrentAllocations.Remove(Id);
 		}
 	}
@@ -1343,101 +1363,187 @@ class FUploadHeightFieldToAtlasCS : public FGlobalShader
 
 IMPLEMENT_GLOBAL_SHADER(FUploadHeightFieldToAtlasCS, "/Engine/Private/HeightFieldAtlasManagement.usf", "UploadHeightFieldToAtlasCS", SF_Compute);
 
+uint32 FHeightFieldTextureAtlas::CalculateDownSampleLevel(uint32 SizeX, uint32 SizeY) const
+{
+	const uint32 TileSize = AddrSpaceAllocator.TileSize;
+
+	for (uint32 CurLevel = 0; CurLevel <= MaxDownSampleLevel; ++CurLevel)
+	{
+		const uint32 DownSampledSizeX = SizeX >> CurLevel;
+		const uint32 DownSampledSizeY = SizeY >> CurLevel;
+
+		if (DownSampledSizeX <= TileSize && DownSampledSizeY <= TileSize)
+		{
+			return CurLevel;
+		}
+	}
+
+	return MaxDownSampleLevel;
+}
+
 void FHeightFieldTextureAtlas::UpdateAllocations(FRHICommandListImmediate& RHICmdList, ERHIFeatureLevel::Type InFeatureLevel)
 {
 	InitializeIfNeeded();
 
-	TArray<uint32, TInlineAllocator<8>> UploadTileIndices;
+	TArray<uint32, TInlineAllocator<8>> UploadHandles;
+	TArray<FIntVector, TInlineAllocator<8>> UploadSizesAndMipBias;
 	TArray<FRHITexture*, TInlineAllocator<8>> UploadSourceTextures;
 
+	auto AllocSortPred = [](const FAllocation& A, const FAllocation& B)
+	{
+		const int32 SizeA = FMath::Max(A.SourceTexture->GetSizeX(), A.SourceTexture->GetSizeY());
+		const int32 SizeB = FMath::Max(B.SourceTexture->GetSizeX(), B.SourceTexture->GetSizeY());
+		return SizeA < SizeB;
+	};
+
+	for (int32 Idx = 0; Idx < PendingStreamingTextures.Num(); ++Idx)
+	{
+		UTexture2D* SourceTexture = PendingStreamingTextures[Idx];
+		const uint32 SizeX = SourceTexture->GetSizeX();
+		const uint32 SizeY = SourceTexture->GetSizeY();
+		const uint32 DownSampleLevel = CalculateDownSampleLevel(SizeX, SizeY);
+		const uint32 NumMissingMips = SourceTexture->GetNumMips() - SourceTexture->GetCachedNumResidentLODs();
+
+		if (NumMissingMips <= DownSampleLevel)
+		{
+			SourceTexture->bForceMiplevelsToBeResident = false;
+			const uint32 SourceMipBias = DownSampleLevel - NumMissingMips;
+			const FAllocation* Allocation = CurrentAllocations.Find(SourceTexture);
+			check(Allocation && Allocation->Handle != INDEX_NONE);
+			UploadHandles.Add(Allocation->Handle);
+			UploadSizesAndMipBias.Add(FIntVector(SizeX >> DownSampleLevel, SizeY >> DownSampleLevel, SourceMipBias));
+			UploadSourceTextures.Add(SourceTexture->Resource->TextureRHI);
+			PendingStreamingTextures.RemoveAtSwap(Idx--);
+		}
+	}
+
 	if (PendingAllocations.Num())
 	{
+		PendingAllocations.Sort(AllocSortPred);
 		bool bAllocFailed = false;
 
 		for (TSet<FAllocation>::TConstIterator It(PendingAllocations); It; ++It)
 		{
+			FAllocation TmpAllocation = *It;
+
 			if (!bAllocFailed)
 			{
-				const uint32 TileIdx = TileAllocator.AllocTile();
-				bAllocFailed = TileIdx == INDEX_NONE;
-				if (!bAllocFailed)
+				UTexture2D* SourceTexture = TmpAllocation.SourceTexture;
+				const uint32 SizeX = SourceTexture->GetSizeX();
+				const uint32 SizeY = SourceTexture->GetSizeY();
+				const uint32 DownSampleLevel = CalculateDownSampleLevel(SizeX, SizeY);
+				const uint32 DownSampledSizeX = SizeX >> DownSampleLevel;
+				const uint32 DownSampledSizeY = SizeY >> DownSampleLevel;
+				const uint32 Handle = AddrSpaceAllocator.Alloc(DownSampledSizeX, DownSampledSizeY);
+
+				if (Handle == INDEX_NONE)
+				{
+					FailedAllocations.Add(TmpAllocation);
+					bAllocFailed = true;
+					continue;
+				}
+
+				const uint32 NumMissingMips = SourceTexture->GetNumMips() - SourceTexture->GetCachedNumResidentLODs();
+				const uint32 SourceMipBias = NumMissingMips > DownSampleLevel ? 0 : DownSampleLevel - NumMissingMips;
+
+				if (NumMissingMips > DownSampleLevel)
 				{
-					FAllocation TmpAllocation = *It;
-					TmpAllocation.TileIdx = TileIdx;
-					CurrentAllocations.Add(TmpAllocation);
-					UploadTileIndices.Add(TileIdx);
-					UploadSourceTextures.Add(TmpAllocation.SourceTexture->Resource->TextureRHI);
+					SourceTexture->bForceMiplevelsToBeResident = true;
+					check(!PendingStreamingTextures.Contains(SourceTexture));
+					PendingStreamingTextures.Add(SourceTexture);
 				}
+
+				TmpAllocation.Handle = Handle;
+				CurrentAllocations.Add(TmpAllocation);
+				UploadHandles.Add(Handle);
+				UploadSizesAndMipBias.Add(FIntVector(DownSampledSizeX, DownSampledSizeY, SourceMipBias));
+				UploadSourceTextures.Add(SourceTexture->Resource->TextureRHI);
 			}
 			else
 			{
-				FailedAllocations.Add(*It);
+				FailedAllocations.Add(TmpAllocation);
 			}
 		}
 
+		if (bAllocFailed)
+		{
+			FailedAllocations.Sort(AllocSortPred);
+		}
 		PendingAllocations.Empty();
 	}
 
-	if (FailedAllocations.Num() && TileAllocator.CanAlloc())
+	if (FailedAllocations.Num())
 	{
 		for (TSet<FAllocation>::TIterator It(FailedAllocations); It; ++It)
 		{
-			const uint32 TileIdx = TileAllocator.AllocTile();
-			if (TileIdx != INDEX_NONE)
+			FAllocation TmpAllocation = *It;
+			UTexture2D* SourceTexture = TmpAllocation.SourceTexture;
+			const uint32 SizeX = SourceTexture->GetSizeX();
+			const uint32 SizeY = SourceTexture->GetSizeY();
+			const uint32 DownSampleLevel = CalculateDownSampleLevel(SizeX, SizeY);
+			const uint32 DownSampledSizeX = SizeX >> DownSampleLevel;
+			const uint32 DownSampledSizeY = SizeY >> DownSampleLevel;
+			const uint32 Handle = AddrSpaceAllocator.Alloc(DownSampledSizeX, DownSampledSizeY);
+
+			if (Handle == INDEX_NONE)
 			{
-				FAllocation TmpAllocation = *It;
-				TmpAllocation.TileIdx = TileIdx;
-				CurrentAllocations.Add(TmpAllocation);
-				UploadTileIndices.Add(TileIdx);
-				UploadSourceTextures.Add(TmpAllocation.SourceTexture->Resource->TextureRHI);
-				It.RemoveCurrent();
+				break;
 			}
-			else
+
+			const uint32 NumMissingMips = SourceTexture->GetNumMips() - SourceTexture->GetCachedNumResidentLODs();
+			const uint32 SourceMipBias = NumMissingMips > DownSampleLevel ? 0 : DownSampleLevel - NumMissingMips;
+
+			if (NumMissingMips > DownSampleLevel)
 			{
-				break;
+				SourceTexture->bForceMiplevelsToBeResident = true;
+				check(!PendingStreamingTextures.Contains(SourceTexture));
+				PendingStreamingTextures.Add(SourceTexture);
 			}
+
+			TmpAllocation.Handle = Handle;
+			CurrentAllocations.Add(TmpAllocation);
+			UploadHandles.Add(Handle);
+			UploadSizesAndMipBias.Add(FIntVector(DownSampledSizeX, DownSampledSizeY, SourceMipBias));
+			UploadSourceTextures.Add(SourceTexture->Resource->TextureRHI);
+			It.RemoveCurrent();
 		}
 	}
 
-	if (UploadTileIndices.Num())
+	if (UploadHandles.Num())
 	{
 		FRDGBuilder GraphBuilder(RHICmdList);
 
-		const uint32 TileSize = TileAllocator.TileSize;
-		const uint32 BorderSize = TileAllocator.BorderSize;
-		const uint32 TileSizeWithBorder = TileSize + BorderSize * 2;
-
-		const float SourceScale = 1.f / TileSize;
-		const float SourceBias = (.5f - BorderSize) * SourceScale;
-
-		check(TileSize && !(TileSize & (TileSize - 1)));
-		const uint32 TileLevels = FMath::CountBits(TileSize - 1) + 1;
-
-		const uint32 NumGroupsX = FMath::DivideAndRoundUp(TileSizeWithBorder, FUploadHeightFieldToAtlasCS::ThreadGroupSizeX);
-		const uint32 NumGroupsY = FMath::DivideAndRoundUp(TileSizeWithBorder, FUploadHeightFieldToAtlasCS::ThreadGroupSizeY);
-
 		TShaderMapRef<FUploadHeightFieldToAtlasCS> ComputeShader(GetGlobalShaderMap(InFeatureLevel));
 		FUploadHeightFieldToAtlasCS* ComputeShaderPtr = *ComputeShader;
 
 		RHICmdList.TransitionResource(EResourceTransitionAccess::ERWBarrier, EResourceTransitionPipeline::EComputeToCompute, AtlasUAVRHI);
 
-		for (int32 Idx = 0; Idx < UploadTileIndices.Num(); ++Idx)
+		for (int32 Idx = 0; Idx < UploadHandles.Num(); ++Idx)
 		{
-			const uint32 TileIdx = UploadTileIndices[Idx];
-			const FIntPoint TileCoord = TileAllocator.GetTileCoord(TileIdx);
-
+			const uint32 Handle = UploadHandles[Idx];
 			FRHITexture2D* SourceTexture = (FRHITexture2D*)UploadSourceTextures[Idx];
-			const uint32 SourceLevels = SourceTexture->GetNumMips();
-			const uint32 SourceMipBias = SourceLevels > TileLevels ? SourceLevels - TileLevels : 0;
+			const FIntVector& SizesAndLevel = UploadSizesAndMipBias[Idx];
+			const uint32 DownSampledSizeX = SizesAndLevel.X;
+			const uint32 DownSampledSizeY = SizesAndLevel.Y;
+			const uint32 SourceMipBias = SizesAndLevel.Z;
+			const float InvDownSampledSizeX = 1.f / DownSampledSizeX;
+			const float InvDownSampledSizeY = 1.f / DownSampledSizeY;
+			const uint32 BorderSize = AddrSpaceAllocator.BorderSize;
+			const uint32 UpdateRegionSizeX = DownSampledSizeX + 2 * BorderSize;
+			const uint32 UpdateRegionSizeY = DownSampledSizeY + 2 * BorderSize;
+			const FIntPoint StartOffset = AddrSpaceAllocator.GetStartOffset(Handle);
 
 			FUploadHeightFieldToAtlasCS::FParameters* Parameters = GraphBuilder.AllocParameters<FUploadHeightFieldToAtlasCS::FParameters>();
-			Parameters->UpdateRegionOffsetAndSize = FUintVector4(TileCoord.X * TileSizeWithBorder, TileCoord.Y * TileSizeWithBorder, TileSizeWithBorder, TileSizeWithBorder);
-			Parameters->SourceScaleBias = FVector4(SourceScale, SourceScale, SourceBias, SourceBias);
+			Parameters->UpdateRegionOffsetAndSize = FUintVector4(StartOffset.X, StartOffset.Y, UpdateRegionSizeX, UpdateRegionSizeY);
+			Parameters->SourceScaleBias = FVector4(InvDownSampledSizeX, InvDownSampledSizeY, (.5f - BorderSize) * InvDownSampledSizeX, (.5f - BorderSize) * InvDownSampledSizeY);
 			Parameters->SourceMipBias = SourceMipBias;
 			Parameters->SourceTexture = SourceTexture;
 			Parameters->SourceTextureSampler = TStaticSamplerState<SF_Bilinear>::GetRHI();
 			Parameters->RWHeightFieldAtlas = AtlasUAVRHI;
 
+			const uint32 NumGroupsX = FMath::DivideAndRoundUp(UpdateRegionSizeX, FUploadHeightFieldToAtlasCS::ThreadGroupSizeX);
+			const uint32 NumGroupsY = FMath::DivideAndRoundUp(UpdateRegionSizeY, FUploadHeightFieldToAtlasCS::ThreadGroupSizeY);
+
 			GraphBuilder.AddPass(
 				RDG_EVENT_NAME("UploadHeightFieldToAtlas"),
 				Parameters,
@@ -1453,18 +1559,15 @@ void FHeightFieldTextureAtlas::UpdateAllocations(FRHICommandListImmediate& RHICm
 	}
 }
 
-uint32 FHeightFieldTextureAtlas::GetAllocationHandle(const UTexture2D* Texture) const
+uint32 FHeightFieldTextureAtlas::GetAllocationHandle(UTexture2D* Texture) const
 {
 	const FAllocation* Allocation = CurrentAllocations.Find(Texture);
-	return Allocation ? Allocation->TileIdx : INDEX_NONE;
+	return Allocation ? Allocation->Handle : INDEX_NONE;
 }
 
-FVector4 FHeightFieldTextureAtlas::GetTileScaleBias(uint32 TileHandle) const
+FVector4 FHeightFieldTextureAtlas::GetAllocationScaleBias(uint32 Handle) const
 {
-	const FIntPoint TileCoord = TileAllocator.GetTileCoord(TileHandle);
-	const float BiasX = (TileCoord.X * TileAllocator.TileSizeWithBorder + TileAllocator.BorderSize) * TileAllocator.TexelSize;
-	const float BiasY = (TileCoord.Y * TileAllocator.TileSizeWithBorder + TileAllocator.BorderSize) * TileAllocator.TexelSize;
-	return FVector4(TileAllocator.TileScale, TileAllocator.TileScale, BiasX, BiasY);
+	return AddrSpaceAllocator.GetScaleBias(Handle);
 }
 
 void FHeightFieldTextureAtlas::FSubAllocator::Init(uint32 InTileSize, uint32 InBorderSize, uint32 InDimInTiles)
@@ -1479,50 +1582,136 @@ void FHeightFieldTextureAtlas::FSubAllocator::Init(uint32 InTileSize, uint32 InB
 	DimInTilesMask = InDimInTiles - 1;
 	DimInTexels = InDimInTiles * TileSizeWithBorder;
 	MaxNumTiles = InDimInTiles * InDimInTiles;
+	
 	TexelSize = 1.f / DimInTexels;
 	TileScale = TileSize * TexelSize;
-	NextFreeTileIdx = 0;
-	FreeTileIndices.Empty();
+
+	LevelOffsets.Empty();
+	MarkerQuadTree.Empty();
+	SubAllocInfos.Empty();
+
+	uint32 NumBits = 0;
+	for (uint32 Level = 1; Level <= DimInTiles; Level <<= 1)
+	{
+		const uint32 NumQuadsInLevel = Level * Level;
+		LevelOffsets.Add(NumBits);
+		NumBits += NumQuadsInLevel;
+	}
+	MarkerQuadTree.Add(false, NumBits);
 }
 
-uint32 FHeightFieldTextureAtlas::FSubAllocator::AllocTile()
+uint32 FHeightFieldTextureAtlas::FSubAllocator::Alloc(uint32 SizeX, uint32 SizeY)
 {
-	uint32 Ret = INDEX_NONE;
-	if (FreeTileIndices.Num())
+	const uint32 NumTiles1D = FMath::DivideAndRoundUp(FMath::Max(SizeX, SizeY), TileSize);
+	check(NumTiles1D <= DimInTiles);
+	const uint32 NumLevels = LevelOffsets.Num();
+	const uint32 Level = NumLevels - FMath::CeilLogTwo(NumTiles1D) - 1;
+	const uint32 LevelOffset = LevelOffsets[Level];
+	const uint32 QuadsInLevel1D = 1u << Level;
+	const uint32 SearchEnd = LevelOffset + QuadsInLevel1D * QuadsInLevel1D;
+
+	uint32 QuadIdx = LevelOffset;
+	for (; QuadIdx < SearchEnd; ++QuadIdx)
 	{
-		Ret = FreeTileIndices.Pop();
+		if (!MarkerQuadTree[QuadIdx])
+		{
+			break;
+		}
 	}
-	else if (NextFreeTileIdx < MaxNumTiles)
+
+	if (QuadIdx != SearchEnd)
 	{
-		Ret = NextFreeTileIdx++;
+		const uint32 QuadIdxInLevel = QuadIdx - LevelOffset;
+		
+		uint32 ParentLevel = Level;
+		uint32 ParentQuadIdxInLevel = QuadIdxInLevel;
+		for (; ParentLevel != (uint32)-1; --ParentLevel)
+		{
+			const uint32 ParentLevelOffset = LevelOffsets[ParentLevel];
+			const uint32 ParentQuadIdx = ParentLevelOffset + ParentQuadIdxInLevel;
+			FBitReference Marker = MarkerQuadTree[ParentQuadIdx];
+			if (Marker)
+			{
+				break;
+			}
+			Marker = true;
+			ParentQuadIdxInLevel >>= 2;
+		}
+		
+		const uint32 QuadX = FMath::ReverseMortonCode2(QuadIdxInLevel);
+		const uint32 QuadY = FMath::ReverseMortonCode2(QuadIdxInLevel >> 1);
+		const uint32 QuadSizeInTiles1D = DimInTiles >> Level;
+		const uint32 TileX = QuadX * QuadSizeInTiles1D;
+		const uint32 TileY = QuadY * QuadSizeInTiles1D;
+
+		FSubAllocInfo SubAllocInfo;
+		SubAllocInfo.Level = Level;
+		SubAllocInfo.QuadIdx = QuadIdx;
+		SubAllocInfo.UVScaleBias.X = SizeX * TexelSize;
+		SubAllocInfo.UVScaleBias.Y = SizeY * TexelSize;
+		SubAllocInfo.UVScaleBias.Z = TileX / (float)DimInTiles + BorderSize * TexelSize;
+		SubAllocInfo.UVScaleBias.W = TileY / (float)DimInTiles + BorderSize * TexelSize;
+
+		return SubAllocInfos.Add(SubAllocInfo);
 	}
-	return Ret;
+
+	return INDEX_NONE;
 }
 
-void FHeightFieldTextureAtlas::FSubAllocator::FreeTile(uint32 TileIdx)
+void FHeightFieldTextureAtlas::FSubAllocator::Free(uint32 Handle)
 {
-	check(TileIdx < NextFreeTileIdx);
-	FreeTileIndices.Add(TileIdx);
+	check(SubAllocInfos.IsValidIndex(Handle));
+
+	const FSubAllocInfo SubAllocInfo = SubAllocInfos[Handle];
+	SubAllocInfos.RemoveAt(Handle);
+
+	const uint32 Level = SubAllocInfo.Level;
+	const uint32 QuadIdx = SubAllocInfo.QuadIdx;
+	check(MarkerQuadTree[QuadIdx]);
+	MarkerQuadTree[QuadIdx] = false;
+
+	uint32 TestIdxInLevel = (QuadIdx - LevelOffsets[Level]) & ~3u;
+	uint32 ParentLevel = Level - 1;
+	for (; ParentLevel != (uint32)-1; --ParentLevel)
+	{
+		const uint32 TestIdx = TestIdxInLevel + LevelOffsets[ParentLevel + 1];
+		const bool bParentFree = !MarkerQuadTree[TestIdx] && !MarkerQuadTree[TestIdx + 1] && !MarkerQuadTree[TestIdx + 2] && !MarkerQuadTree[TestIdx + 3];
+		if (!bParentFree)
+		{
+			break;
+		}
+		const uint32 ParentIdxInLevel = TestIdxInLevel >> 2;
+		const uint32 ParentIdx = ParentIdxInLevel + LevelOffsets[ParentLevel];
+		MarkerQuadTree[ParentIdx] = false;
+		TestIdxInLevel = ParentIdxInLevel & ~3u;
+	}
 }
 
-bool FHeightFieldTextureAtlas::FSubAllocator::CanAlloc() const
+FVector4 FHeightFieldTextureAtlas::FSubAllocator::GetScaleBias(uint32 Handle) const
 {
-	return FreeTileIndices.Num() || NextFreeTileIdx < MaxNumTiles;
+	check(SubAllocInfos.IsValidIndex(Handle));
+	return SubAllocInfos[Handle].UVScaleBias;
 }
 
-FIntPoint FHeightFieldTextureAtlas::FSubAllocator::GetTileCoord(uint32 TileIdx) const
+FIntPoint FHeightFieldTextureAtlas::FSubAllocator::GetStartOffset(uint32 Handle) const
 {
-	return FIntPoint(TileIdx & DimInTilesMask, TileIdx >> DimInTilesShift);
+	check(SubAllocInfos.IsValidIndex(Handle));
+	const FSubAllocInfo& Info = SubAllocInfos[Handle];
+	const uint32 QuadIdxInLevel = Info.QuadIdx - LevelOffsets[Info.Level];
+	const uint32 QuadX = FMath::ReverseMortonCode2(QuadIdxInLevel);
+	const uint32 QuadY = FMath::ReverseMortonCode2(QuadIdxInLevel >> 1);
+	const uint32 QuadSizeInTexels1D = (DimInTiles >> Info.Level) * TileSizeWithBorder;	
+	return FIntPoint(QuadX * QuadSizeInTexels1D, QuadY * QuadSizeInTexels1D);
 }
 
 FHeightFieldTextureAtlas::FAllocation::FAllocation()
 	: SourceTexture(nullptr)
 	, RefCount(0)
-	, TileIdx(INDEX_NONE)
+	, Handle(INDEX_NONE)
 {}
 
-FHeightFieldTextureAtlas::FAllocation::FAllocation(const UTexture2D* InTexture)
+FHeightFieldTextureAtlas::FAllocation::FAllocation(UTexture2D* InTexture)
 	: SourceTexture(InTexture)
 	, RefCount(1)
-	, TileIdx(INDEX_NONE)
+	, Handle(INDEX_NONE)
 {}
diff --git a/Engine/Source/Runtime/Engine/Public/DistanceFieldAtlas.h b/Engine/Source/Runtime/Engine/Public/DistanceFieldAtlas.h
--- a/Engine/Source/Runtime/Engine/Public/DistanceFieldAtlas.h
+++ b/Engine/Source/Runtime/Engine/Public/DistanceFieldAtlas.h
@@ -159,15 +159,15 @@ class ENGINE_API FHeightFieldTextureAtlas : public FRenderResource
 public:
 	void InitializeIfNeeded();
 
-	void AddAllocation(const UTexture2D* Texture);
+	void AddAllocation(UTexture2D* Texture);
 
-	void RemoveAllocation(const UTexture2D* Texture);
+	void RemoveAllocation(UTexture2D* Texture);
 
 	void UpdateAllocations(FRHICommandListImmediate& RHICmdList, ERHIFeatureLevel::Type InFeatureLevel);
 
-	uint32 GetAllocationHandle(const UTexture2D* Texture) const;
+	uint32 GetAllocationHandle(UTexture2D* Texture) const;
 
-	FVector4 GetTileScaleBias(uint32 TileHandle) const;
+	FVector4 GetAllocationScaleBias(uint32 Handle) const;
 
 	FRHITexture2D* GetAtlasTexture() const
 	{
@@ -176,12 +176,12 @@ public:
 
 	uint32 GetSizeX() const
 	{
-		return TileAllocator.DimInTexels;
+		return AddrSpaceAllocator.DimInTexels;
 	}
 
 	uint32 GetSizeY() const
 	{
-		return TileAllocator.DimInTexels;
+		return AddrSpaceAllocator.DimInTexels;
 	}
 
 	uint32 GetGeneration() const
@@ -190,20 +190,29 @@ public:
 	}
 
 private:
+	uint32 CalculateDownSampleLevel(uint32 SizeX, uint32 SizeY) const;
+
 	class FSubAllocator
 	{
 	public:
 		void Init(uint32 InTileSize, uint32 InBorderSize, uint32 InDimInTiles);
 
-		uint32 AllocTile();
+		uint32 Alloc(uint32 SizeX, uint32 SizeY);
 
-		void FreeTile(uint32 TileIdx);
+		void Free(uint32 Handle);
 
-		bool CanAlloc() const;
+		FVector4 GetScaleBias(uint32 Handle) const;
 
-		FIntPoint GetTileCoord(uint32 TileIdx) const;
+		FIntPoint GetStartOffset(uint32 Handle) const;
 
 	private:
+		struct FSubAllocInfo
+		{
+			uint32 Level;
+			uint32 QuadIdx;
+			FVector4 UVScaleBias;
+		};
+		
 		uint32 TileSize;
 		uint32 BorderSize;
 		uint32 TileSizeWithBorder;
@@ -216,21 +225,24 @@ private:
 		float TexelSize;
 		float TileScale;
 
-		uint32 NextFreeTileIdx;
-		TArray<uint32> FreeTileIndices;
+		// 0: Free, 1: Allocated
+		TBitArray<> MarkerQuadTree;
+		TArray<uint32, TInlineAllocator<8>> LevelOffsets;
+
+		TSparseArray<FSubAllocInfo> SubAllocInfos;
 
 		friend class FHeightFieldTextureAtlas;
 	};
 
 	struct FAllocation
 	{
-		const UTexture2D* SourceTexture;
+		UTexture2D* SourceTexture;
 		uint32 RefCount;
-		uint32 TileIdx;
+		uint32 Handle;
 
 		FAllocation();
 
-		FAllocation(const UTexture2D* InTexture);
+		FAllocation(UTexture2D* InTexture);
 
 		bool operator==(const FAllocation& Other) const
 		{
@@ -243,15 +255,17 @@ private:
 		}
 	};
 
-	FSubAllocator TileAllocator;
+	FSubAllocator AddrSpaceAllocator;
 
 	TSet<FAllocation> PendingAllocations;
 	TSet<FAllocation> FailedAllocations;
 	TSet<FAllocation> CurrentAllocations;
+	TArray<UTexture2D*> PendingStreamingTextures;
 
 	FTexture2DRHIRef AtlasTextureRHI;
 	FUnorderedAccessViewRHIRef AtlasUAVRHI;
 
+	uint32 MaxDownSampleLevel;
 	uint32 Generation;
 };
 
diff --git a/Engine/Source/Runtime/Renderer/Private/DistanceFieldObjectManagement.cpp b/Engine/Source/Runtime/Renderer/Private/DistanceFieldObjectManagement.cpp
--- a/Engine/Source/Runtime/Renderer/Private/DistanceFieldObjectManagement.cpp
+++ b/Engine/Source/Runtime/Renderer/Private/DistanceFieldObjectManagement.cpp
@@ -1235,6 +1235,12 @@ bool ProcessHeightFieldPrimitiveUpdate(
 	FHeightfieldComponentDescription HeightFieldCompDesc(PrimitiveSceneInfo->Proxy->GetLocalToWorld());
 	PrimitiveSceneInfo->Proxy->GetHeightfieldRepresentation(HeightNormalTexture, DiffuseColorTexture, VisibilityTexture, HeightFieldCompDesc);
 
+	const uint32 Handle = GHeightFieldTextureAtlas.GetAllocationHandle(HeightNormalTexture);
+	if (Handle == INDEX_NONE)
+	{
+		return false;
+	}
+
 	uint32 UploadIdx;
 	if (bIsAddOperation)
 	{
@@ -1269,18 +1275,20 @@ bool ProcessHeightFieldPrimitiveUpdate(
 	UploadObjectData.Add(*(const FVector4*)&WorldToLocalT.M[2]);
 
 	const FIntRect& HeightFieldRect = HeightFieldCompDesc.HeightfieldRect;
-	UploadObjectData.Add(FVector4(HeightFieldRect.Min.X, HeightFieldRect.Min.Y, HeightFieldRect.Width(), HeightFieldRect.Height()));
+	const float WorldToLocalScale = FMath::Min3(
+		WorldToLocalT.GetColumn(0).Size(),
+		WorldToLocalT.GetColumn(1).Size(),
+		WorldToLocalT.GetColumn(2).Size());
+	UploadObjectData.Add(FVector4(HeightFieldRect.Width(), HeightFieldRect.Height(), WorldToLocalScale, 0.f));
 
 	const FVector4& HeightFieldScaleBias = HeightFieldCompDesc.HeightfieldScaleBias;
 	check(HeightFieldScaleBias.Y >= 0.f && HeightFieldScaleBias.Z >= 0.f && HeightFieldScaleBias.W >= 0.f);
-	const uint32 TileHandle = GHeightFieldTextureAtlas.GetAllocationHandle(HeightNormalTexture);
-	check(TileHandle != INDEX_NONE);
-	const FVector4 TileScaleBias = GHeightFieldTextureAtlas.GetTileScaleBias(TileHandle);
+	const FVector4 AllocationScaleBias = GHeightFieldTextureAtlas.GetAllocationScaleBias(Handle);
 	UploadObjectData.Add(FVector4(
-		FMath::Abs(HeightFieldScaleBias.X) * TileScaleBias.X,
-		HeightFieldScaleBias.Y * TileScaleBias.Y,
-		HeightFieldScaleBias.Z * TileScaleBias.X + TileScaleBias.Z,
-		HeightFieldScaleBias.W * TileScaleBias.Y + TileScaleBias.W));
+		FMath::Abs(HeightFieldScaleBias.X) * AllocationScaleBias.X,
+		HeightFieldScaleBias.Y * AllocationScaleBias.Y,
+		HeightFieldScaleBias.Z * AllocationScaleBias.X + AllocationScaleBias.Z,
+		HeightFieldScaleBias.W * AllocationScaleBias.Y + AllocationScaleBias.W));
 
 	check(!(UploadObjectData.Num() % UploadHeightFieldObjectDataStride));
 
@@ -1627,7 +1635,7 @@ void FDeferredShadingSceneRenderer::UpdateGlobalHeightFieldObjectBuffers(FRHICom
 					&& !DistanceFieldSceneData.PendingHeightFieldAddOps.Contains(Primitive)
 					&& !DistanceFieldSceneData.PendingHeightFieldUpdateOps.Contains(Primitive))
 				{
-					DistanceFieldSceneData.PendingUpdateOperations.Add((FPrimitiveSceneInfo*)Primitive);
+					DistanceFieldSceneData.PendingHeightFieldUpdateOps.Add((FPrimitiveSceneInfo*)Primitive);
 				}
 			}
 		}
diff --git a/Engine/Source/Runtime/Renderer/Private/RendererScene.cpp b/Engine/Source/Runtime/Renderer/Private/RendererScene.cpp
--- a/Engine/Source/Runtime/Renderer/Private/RendererScene.cpp
+++ b/Engine/Source/Runtime/Renderer/Private/RendererScene.cpp
@@ -405,8 +405,6 @@ void FDistanceFieldSceneData::AddPrimitive(FPrimitiveSceneInfo* InPrimitive)
 	{
 		if (Proxy->SupportsHeightfieldRepresentation())
 		{
-			// Theoretically, we may need to wait for high-res mips to be streamed in.
-			// In practice, however, those mips are usually resident at this point.
 			UTexture2D* HeightAndNormal;
 			UTexture2D* Unused0, *Unused1;
 			FHeightfieldComponentDescription Unused2(FMatrix::Identity);
