Chaos
- expose RBAN iteration counts in PhAt, RBAN and blueprint
- hook up physics materials
- adjust default line thickness

diff --git a/Engine/Source/Runtime/AnimGraphRuntime/Private/BoneControllers/AnimNode_RigidBody.cpp b/Engine/Source/Runtime/AnimGraphRuntime/Private/BoneControllers/AnimNode_RigidBody.cpp
--- a/Engine/Source/Runtime/AnimGraphRuntime/Private/BoneControllers/AnimNode_RigidBody.cpp
+++ b/Engine/Source/Runtime/AnimGraphRuntime/Private/BoneControllers/AnimNode_RigidBody.cpp
@@ -61,6 +61,13 @@ FAnimNode_RigidBody::FAnimNode_RigidBody():
 	bFreezeIncomingPoseOnStart = false;
 	bClampLinearTranslationLimitToRefPose = false;
 
+	OverrideSolverIterations.SolverIterations = -1;
+	OverrideSolverIterations.JointIterations = -1;
+	OverrideSolverIterations.CollisionIterations = -1;
+	OverrideSolverIterations.SolverPushOutIterations = -1;
+	OverrideSolverIterations.JointPushOutIterations = -1;
+	OverrideSolverIterations.CollisionPushOutIterations = -1;
+
 	PreviousTransform = CurrentTransform = FTransform::Identity;
 	PreviousComponentLinearVelocity = FVector::ZeroVector;	
 
@@ -496,6 +503,22 @@ void FAnimNode_RigidBody::EvaluateSkeletalControl_AnyThread(FComponentSpacePoseC
 				PhysicsSimulation->Simulate_AssumesLocked(StepDeltaTime, SimSpaceGravity);
 			}
 #else
+			PhysicsSimulation->SetSolverIterations(
+				SolverIterations.SolverIterations,
+				SolverIterations.JointIterations,
+				SolverIterations.CollisionIterations,
+				SolverIterations.SolverPushOutIterations,
+				SolverIterations.JointPushOutIterations,
+				SolverIterations.CollisionPushOutIterations
+			);
+			PhysicsSimulation->SetSolverIterations(
+				OverrideSolverIterations.SolverIterations, 
+				OverrideSolverIterations.JointIterations, 
+				OverrideSolverIterations.CollisionIterations, 
+				OverrideSolverIterations.SolverPushOutIterations, 
+				OverrideSolverIterations.JointPushOutIterations, 
+				OverrideSolverIterations.CollisionPushOutIterations);
+
 			PhysicsSimulation->Simulate_AssumesLocked(DeltaSeconds, MaxDeltaSeconds, MaxSteps, SimSpaceGravity);
 #endif
 		}
@@ -840,6 +863,18 @@ void FAnimNode_RigidBody::InitPhysics(const UAnimInstance* InAnimInstance)
 
 		PhysicsSimulation->SetIgnoreCollisionPairTable(IgnorePairs);
 		PhysicsSimulation->SetIgnoreCollisionActors(IgnoreCollisionActors);
+
+#if WITH_CHAOS
+		SolverIterations = UsePhysicsAsset->SolverIterations;
+		PhysicsSimulation->SetSolverIterations(
+			SolverIterations.SolverIterations,
+			SolverIterations.JointIterations,
+			SolverIterations.CollisionIterations,
+			SolverIterations.SolverPushOutIterations,
+			SolverIterations.JointPushOutIterations,
+			SolverIterations.CollisionPushOutIterations
+		);
+#endif
 	}
 }
 
diff --git a/Engine/Source/Runtime/AnimGraphRuntime/Public/BoneControllers/AnimNode_RigidBody.h b/Engine/Source/Runtime/AnimGraphRuntime/Public/BoneControllers/AnimNode_RigidBody.h
--- a/Engine/Source/Runtime/AnimGraphRuntime/Public/BoneControllers/AnimNode_RigidBody.h
+++ b/Engine/Source/Runtime/AnimGraphRuntime/Public/BoneControllers/AnimNode_RigidBody.h
@@ -149,6 +149,15 @@ public:
 	UPROPERTY(EditAnywhere, Category = Settings)
 	uint8 bClampLinearTranslationLimitToRefPose : 1;
 
+	/**
+	 * Solver iteration settings overrides (defaults are set in the Physics Asset).
+	 * These can be varied in the runtime and set through blueprint (e.g., to increase
+	 * iterations during difficult movements).
+	 * Set to -1 to leave an individual iteration value at its Physics Asset value.
+	 */
+	UPROPERTY(EditAnywhere, Category = Settings, meta = (PinHiddenByDefault))
+	FSolverIterations OverrideSolverIterations;
+
 private:
 	uint8 bEnabled : 1;
 	uint8 bSimulationStarted : 1;
@@ -183,6 +192,7 @@ private:
 	TWeakObjectPtr<USkeletalMeshComponent> SkelMeshCompWeakPtr;
 
 	ImmediatePhysics::FSimulation* PhysicsSimulation;
+	FSolverIterations SolverIterations;
 
 	struct FOutputBoneData
 	{
diff --git a/Engine/Source/Runtime/Engine/Classes/PhysicsEngine/PhysicsAsset.h b/Engine/Source/Runtime/Engine/Classes/PhysicsEngine/PhysicsAsset.h
--- a/Engine/Source/Runtime/Engine/Classes/PhysicsEngine/PhysicsAsset.h
+++ b/Engine/Source/Runtime/Engine/Classes/PhysicsEngine/PhysicsAsset.h
@@ -15,6 +15,60 @@
 class FMeshElementCollector;
 class USkeletalBodySetup;
 
+/**
+ * [Chaos Only]
+ */
+USTRUCT(BlueprintType)
+struct ENGINE_API FSolverIterations
+{
+	GENERATED_USTRUCT_BODY()
+
+		FSolverIterations();
+
+	/**
+	 * [Chaos Only]
+	 * The recommended number of solver iterations. Increase this if collision and joints are fighting, or joint chains are stretching.
+	 */
+	UPROPERTY(EditAnywhere, BlueprintReadWrite, Category = SolverSettings, meta = (ClampMax = 50))
+		int32 SolverIterations;
+
+	/**
+	 * [Chaos Only]
+	 * The recommended number of joint sub-iterations. Increasing this can help with chains of long-thin bodies.
+	 */
+	UPROPERTY(EditAnywhere, BlueprintReadWrite, Category = SolverSettings, meta = (ClampMax = 50))
+		int32 JointIterations;
+
+	/**
+	 * [Chaos Only]
+	 * The recommended number of collision sub-iterations. Increasing this can help with collision jitter.
+	 */
+	UPROPERTY(EditAnywhere, BlueprintReadWrite, Category = SolverSettings, meta = (ClampMax = 50))
+		int32 CollisionIterations;
+
+	/**
+	 * [Chaos Only]
+	 * The recommended number of solver push-out iterations. Increasing this can help with collision penetration problems.
+	 */
+	 /** Increase this if bodies remain penetrating */
+	UPROPERTY(EditAnywhere, BlueprintReadWrite, Category = SolverSettings, meta = (ClampMax = 50))
+		int32 SolverPushOutIterations;
+
+	/**
+	 * [Chaos Only]
+	 * The recommended number of joint sub-push-out iterations.
+	 */
+	UPROPERTY(EditAnywhere, BlueprintReadWrite, Category = SolverSettings, meta = (ClampMax = 50))
+		int32 JointPushOutIterations;
+
+	/**
+	 * [Chaos Only]
+	 * The recommended number of joint sub-push-out iterations. Increasing this can help with collision penetration problems.
+	 */
+	UPROPERTY(EditAnywhere, BlueprintReadWrite, Category = SolverSettings, meta = (ClampMax = 50))
+		int32 CollisionPushOutIterations;
+};
+
 /**
  * PhysicsAsset contains a set of rigid bodies and constraints that make up a single ragdoll.
  * The asset is not limited to human ragdolls, and can be used for any physical simulation using bodies and constraints.
@@ -75,6 +129,10 @@ class UPhysicsAsset : public UObject, public IInterface_PreviewMeshProvider
 
 public:
 
+	/** Recommended solver iterations (may be overridden depending on how the Physics Asset is used). E.g., Rigid Body Anim Nodes support overrides. */
+	UPROPERTY(EditAnywhere, AdvancedDisplay, Category = SolverSettings)
+	FSolverIterations SolverIterations;
+
 	/** If true, we skip instancing bodies for this PhysicsAsset on dedicated servers */
 	UPROPERTY(EditAnywhere, AdvancedDisplay, Category = Physics)
 	uint8 bNotForDedicatedServer:1;
diff --git a/Engine/Source/Runtime/Engine/Private/PhysicsEngine/ImmediatePhysics/ImmediatePhysicsChaos/ImmediatePhysicsSimulation_Chaos.cpp b/Engine/Source/Runtime/Engine/Private/PhysicsEngine/ImmediatePhysics/ImmediatePhysicsChaos/ImmediatePhysicsSimulation_Chaos.cpp
--- a/Engine/Source/Runtime/Engine/Private/PhysicsEngine/ImmediatePhysics/ImmediatePhysicsChaos/ImmediatePhysicsSimulation_Chaos.cpp
+++ b/Engine/Source/Runtime/Engine/Private/PhysicsEngine/ImmediatePhysics/ImmediatePhysicsChaos/ImmediatePhysicsSimulation_Chaos.cpp
@@ -14,11 +14,14 @@
 #include "Chaos/PBDRigidsEvolutionGBF.h"
 #include "ChaosLog.h"
 
+#include "PhysicalMaterials/PhysicalMaterial.h"
 #include "PhysicsEngine/BodyInstance.h"
 #include "PhysicsEngine/ConstraintInstance.h"
 
 //#pragma optimize("", off)
 
+DECLARE_CYCLE_STAT(TEXT("FSimulation::Simulate_Chaos"), STAT_ImmediateSimulate_Chaos, STATGROUP_ImmediatePhysics);
+
 //////////////////////////////////////////////////////////////////////////
 // @todo(ccaulfield): remove when finished
 //
@@ -26,28 +29,28 @@ float ChaosImmediate_Evolution_StepTime = 0.0f;
 int32 ChaosImmediate_Evolution_NumSteps = 0;
 float ChaosImmediate_Evolution_InitialStepTime = 0.033f;
 int32 ChaosImmediate_Evolution_DeltaTimeCount = 100;
-int32 ChaosImmediate_Evolution_Iterations = 2;
-int32 ChaosImmediate_Evolution_PushOutIterations = 1;
+int32 ChaosImmediate_Evolution_Iterations = -1;
+int32 ChaosImmediate_Evolution_PushOutIterations = -1;
 FAutoConsoleVariableRef CVarChaosImmPhysStepTime(TEXT("p.Chaos.ImmPhys.StepTime"), ChaosImmediate_Evolution_StepTime, TEXT("Override step time (if not zero)"));
 FAutoConsoleVariableRef CVarChaosImmPhysNumSteps(TEXT("p.Chaos.ImmPhys.NumSteps"), ChaosImmediate_Evolution_NumSteps, TEXT("Override num steps (if not zero)"));
 FAutoConsoleVariableRef CVarChaosImmPhysInitialStepTime(TEXT("p.Chaos.ImmPhys.InitialStepTime"), ChaosImmediate_Evolution_InitialStepTime, TEXT("Initial step time (then calculated from rolling average)"));
 FAutoConsoleVariableRef CVarChaosImmPhysDeltaTimeCount(TEXT("p.Chaos.ImmPhys.DeltaTimeCount"), ChaosImmediate_Evolution_DeltaTimeCount, TEXT("The number of ticks over which the moving average is calculated"));
-FAutoConsoleVariableRef CVarChaosImmPhysIterations(TEXT("p.Chaos.ImmPhys.Iterations"), ChaosImmediate_Evolution_Iterations, TEXT("Number of constraint solver loops in immediate physics"));
-FAutoConsoleVariableRef CVarChaosImmPhysPushOutIterations(TEXT("p.Chaos.ImmPhys.PushOutIterations"), ChaosImmediate_Evolution_PushOutIterations, TEXT("Set the ApplyPushOut() (position correction) iteration count"));
+FAutoConsoleVariableRef CVarChaosImmPhysIterations(TEXT("p.Chaos.ImmPhys.Iterations"), ChaosImmediate_Evolution_Iterations, TEXT("Override number of constraint solver loops in immediate physics (if >= 0)"));
+FAutoConsoleVariableRef CVarChaosImmPhysPushOutIterations(TEXT("p.Chaos.ImmPhys.PushOutIterations"), ChaosImmediate_Evolution_PushOutIterations, TEXT("Override number of solver push-out loops (if >= 0)"));
 
 int32 ChaosImmediate_Collision_Enabled = 1;
-int32 ChaosImmediate_Collision_PairIterations = 1;
-int32 ChaosImmediate_Collision_PushOutPairIterations = 2;
+int32 ChaosImmediate_Collision_PairIterations = -1;
+int32 ChaosImmediate_Collision_PushOutPairIterations = -1;
 int32 ChaosImmediate_Collision_Priority = 1;
 float ChaosImmediate_Collision_Thickness = 0;
-FAutoConsoleVariableRef CVarChaosImmPhysDisableCollisions(TEXT("p.Chaos.ImmPhys.Collision.Enabled"), ChaosImmediate_Collision_Enabled, TEXT("Enable/Disable collisions in Immediate Physics."));
-FAutoConsoleVariableRef CVarChaosImmPhysCollisionPairIterations(TEXT("p.Chaos.ImmPhys.Collision.PairIterations"), ChaosImmediate_Collision_PairIterations, TEXT("Set the Apply() internal pair iteration count (0 to disable)"));
-FAutoConsoleVariableRef CVarChaosImmPhysCollisionPushOutPairIterations(TEXT("p.Chaos.ImmPhys.Collision.PushOutPairIterations"), ChaosImmediate_Collision_PushOutPairIterations, TEXT("Set the ApplyPushOut() internal pair iteration count (0 to disable)"));
+FAutoConsoleVariableRef CVarChaosImmPhysCollisionDisable(TEXT("p.Chaos.ImmPhys.Collision.Enabled"), ChaosImmediate_Collision_Enabled, TEXT("Enable/Disable collisions in Immediate Physics."));
+FAutoConsoleVariableRef CVarChaosImmPhysCollisionPairIterations(TEXT("p.Chaos.ImmPhys.Collision.PairIterations"), ChaosImmediate_Collision_PairIterations, TEXT("Override collision pair iterations (if >= 0)"));
+FAutoConsoleVariableRef CVarChaosImmPhysCollisionPushOutPairIterations(TEXT("p.Chaos.ImmPhys.Collision.PushOutPairIterations"), ChaosImmediate_Collision_PushOutPairIterations, TEXT("Override collision push-out pair iterations (if >= 0)"));
 FAutoConsoleVariableRef CVarChaosImmPhysCollisionPriority(TEXT("p.Chaos.ImmPhys.Collision.Priority"), ChaosImmediate_Collision_Priority, TEXT("Set the Collision constraint sort order (Joints have priority 0)"));
-FAutoConsoleVariableRef CVarChaosImmPhysThickness(TEXT("p.Chaos.ImmPhys.Collision.Thickness"), ChaosImmediate_Collision_Thickness, TEXT("ChaosImmediateThickness"));
+FAutoConsoleVariableRef CVarChaosImmPhysCollisionThickness(TEXT("p.Chaos.ImmPhys.Collision.Thickness"), ChaosImmediate_Collision_Thickness, TEXT("ChaosImmediateThickness"));
 
-int32 ChaosImmediate_Joint_PairIterations = 1;
-int32 ChaosImmediate_Joint_PushOutPairIterations = 0;
+int32 ChaosImmediate_Joint_PairIterations = -1;
+int32 ChaosImmediate_Joint_PushOutPairIterations = -1;
 float ChaosImmediate_Joint_SwingTwistAngleTolerance = 1.0e-6f;
 int32 ChaosImmediate_Joint_EnableTwistLimits = 1;
 int32 ChaosImmediate_Joint_EnableSwingLimits = 1;
@@ -67,30 +70,27 @@ float ChaosImmediate_Joint_AngularDriveDamping = 0.0f;
 float ChaosImmediate_Joint_MinParentMassRatio = 0.2f;
 float ChaosImmediate_Joint_MaxInertiaRatio = 5.0f;
 float ChaosImmediate_Joint_AngularPositionCorrection = 1.0f;
-FAutoConsoleVariableRef CVarChaosImmPhysPairIterations(TEXT("p.Chaos.ImmPhys.Joint.PairIterations"), ChaosImmediate_Joint_PairIterations, TEXT("PairIterations."));
-FAutoConsoleVariableRef CVarChaosImmPhysPushOutPairIterations(TEXT("p.Chaos.ImmPhys.Joint.PushOutPairIterations"), ChaosImmediate_Joint_PushOutPairIterations, TEXT("PushOutPairIterations."));
-FAutoConsoleVariableRef CVarChaosImmPhysSwingTwistAngleTolerance(TEXT("p.Chaos.ImmPhys.Joint.SwingTwistAngleTolerance"), ChaosImmediate_Joint_SwingTwistAngleTolerance, TEXT("SwingTwistAngleTolerance."));
-FAutoConsoleVariableRef CVarChaosImmPhysEnableTwistLimits(TEXT("p.Chaos.ImmPhys.Joint.EnableTwistLimits"), ChaosImmediate_Joint_EnableTwistLimits, TEXT("EnableTwistLimits."));
-FAutoConsoleVariableRef CVarChaosImmPhysEnableSwingLimits(TEXT("p.Chaos.ImmPhys.Joint.EnableSwingLimits"), ChaosImmediate_Joint_EnableSwingLimits, TEXT("EnableSwingLimits."));
-FAutoConsoleVariableRef CVarChaosImmPhysEnableDrives(TEXT("p.Chaos.ImmPhys.Joint.EnableDrives"), ChaosImmediate_Joint_EnableDrives, TEXT("EnableDrives."));
-FAutoConsoleVariableRef CVarChaosImmPhysLinearProjection(TEXT("p.Chaos.ImmPhys.Joint.LinearProjection"), ChaosImmediate_Joint_LinearProjection, TEXT("6Dof joint projection amount override (if > 0)."));
-FAutoConsoleVariableRef CVarChaosImmPhysAngularProjection(TEXT("p.Chaos.ImmPhys.Joint.AngularProjection"), ChaosImmediate_Joint_AngularProjection, TEXT("6Dof joint projection amount override (if > 0)."));
-FAutoConsoleVariableRef CVarChaosImmPhysStiffness(TEXT("p.Chaos.ImmPhys.Joint.Stiffness"), ChaosImmediate_Joint_Stiffness, TEXT("6Dof joint stiffness override (if > 0)."));
-FAutoConsoleVariableRef CVarChaosImmPhysSoftLinearStiffness(TEXT("p.Chaos.ImmPhys.Joint.SoftLinearStiffness"), ChaosImmediate_Joint_SoftLinearStiffness, TEXT("6Dof joint soft linear stiffness override (if > 0)."));
-FAutoConsoleVariableRef CVarChaosImmPhysSoftTwistStiffness(TEXT("p.Chaos.ImmPhys.Joint.SoftTwistStiffness"), ChaosImmediate_Joint_SoftTwistStiffness, TEXT("6Dof joint SoftTwist stiffness override (if > 0)."));
-FAutoConsoleVariableRef CVarChaosImmPhysSoftTwistDamping(TEXT("p.Chaos.ImmPhys.Joint.SoftTwistDamping"), ChaosImmediate_Joint_SoftTwistDamping, TEXT("6Dof joint SoftTwist damping override (if > 0)."));
-FAutoConsoleVariableRef CVarChaosImmPhysSoftSwingStiffness(TEXT("p.Chaos.ImmPhys.Joint.SoftSwingStiffness"), ChaosImmediate_Joint_SoftSwingStiffness, TEXT("6Dof joint SoftSwing stiffness override (if > 0)."));
-FAutoConsoleVariableRef CVarChaosImmPhysSoftSwingDamping(TEXT("p.Chaos.ImmPhys.Joint.SoftSwingDamping"), ChaosImmediate_Joint_SoftSwingDamping, TEXT("6Dof joint SoftSwing damping override (if > 0)."));
-FAutoConsoleVariableRef CVarChaosImmPhysLinearDriveStiffness(TEXT("p.Chaos.ImmPhys.Joint.LinearDriveStiffness"), ChaosImmediate_Joint_LinearDriveStiffness, TEXT("6Dof joint drive stiffness override (if > 0)."));
-FAutoConsoleVariableRef CVarChaosImmPhysLinearDriveDamping(TEXT("p.Chaos.ImmPhys.Joint.LinearDriveDamping"), ChaosImmediate_Joint_LinearDriveDamping, TEXT("6Dof joint drive damping override (if > 0)."));
-FAutoConsoleVariableRef CVarChaosImmPhysAngularDriveStiffness(TEXT("p.Chaos.ImmPhys.Joint.AngularDriveStiffness"), ChaosImmediate_Joint_AngularDriveStiffness, TEXT("6Dof joint drive stiffness override (if > 0)."));
-FAutoConsoleVariableRef CVarChaosImmPhysAngularDriveDamping(TEXT("p.Chaos.ImmPhys.Joint.AngularDriveDamping"), ChaosImmediate_Joint_AngularDriveDamping, TEXT("6Dof joint drive damping override (if > 0)."));
-FAutoConsoleVariableRef CVarChaosImmPhysMinParentMassRatio(TEXT("p.Chaos.ImmPhys.Joint.MinParentMassRatio"), ChaosImmediate_Joint_MinParentMassRatio, TEXT("6Dof joint MinParentMassRatio (if > 0)"));
-FAutoConsoleVariableRef CVarChaosImmPhysMaxInertiaRatio(TEXT("p.Chaos.ImmPhys.Joint.MaxInertiaRatio"), ChaosImmediate_Joint_MaxInertiaRatio, TEXT("6Dof joint MaxInertiaRatio (if > 0)"));
-FAutoConsoleVariableRef CVarChaosImmPhysAngularPositionCorrection(TEXT("p.Chaos.ImmPhys.Joint.AngularPositionCorrection"), ChaosImmediate_Joint_AngularPositionCorrection, TEXT("6Dof joint post-rotation constraint position correction amount [0-1]"));
-
-float ChaosImmediate_Material_Resitution = 0.0f;
-float ChaosImmediate_Material_Friction = 0.5f;
+FAutoConsoleVariableRef CVarChaosImmPhysJointPairIterations(TEXT("p.Chaos.ImmPhys.Joint.PairIterations"), ChaosImmediate_Joint_PairIterations, TEXT("Override joint pair iterations (if >= 0)"));
+FAutoConsoleVariableRef CVarChaosImmPhysJointPushOutPairIterations(TEXT("p.Chaos.ImmPhys.Joint.PushOutPairIterations"), ChaosImmediate_Joint_PushOutPairIterations, TEXT("Override joint push-out pair iterations (if >= 0)"));
+FAutoConsoleVariableRef CVarChaosImmPhysJointSwingTwistAngleTolerance(TEXT("p.Chaos.ImmPhys.Joint.SwingTwistAngleTolerance"), ChaosImmediate_Joint_SwingTwistAngleTolerance, TEXT("SwingTwistAngleTolerance."));
+FAutoConsoleVariableRef CVarChaosImmPhysJointEnableTwistLimits(TEXT("p.Chaos.ImmPhys.Joint.EnableTwistLimits"), ChaosImmediate_Joint_EnableTwistLimits, TEXT("EnableTwistLimits."));
+FAutoConsoleVariableRef CVarChaosImmPhysJointEnableSwingLimits(TEXT("p.Chaos.ImmPhys.Joint.EnableSwingLimits"), ChaosImmediate_Joint_EnableSwingLimits, TEXT("EnableSwingLimits."));
+FAutoConsoleVariableRef CVarChaosImmPhysJointEnableDrives(TEXT("p.Chaos.ImmPhys.Joint.EnableDrives"), ChaosImmediate_Joint_EnableDrives, TEXT("EnableDrives."));
+FAutoConsoleVariableRef CVarChaosImmPhysJointLinearProjection(TEXT("p.Chaos.ImmPhys.Joint.LinearProjection"), ChaosImmediate_Joint_LinearProjection, TEXT("6Dof joint projection amount override (if > 0)."));
+FAutoConsoleVariableRef CVarChaosImmPhysJointAngularProjection(TEXT("p.Chaos.ImmPhys.Joint.AngularProjection"), ChaosImmediate_Joint_AngularProjection, TEXT("6Dof joint projection amount override (if > 0)."));
+FAutoConsoleVariableRef CVarChaosImmPhysJointStiffness(TEXT("p.Chaos.ImmPhys.Joint.Stiffness"), ChaosImmediate_Joint_Stiffness, TEXT("6Dof joint stiffness override (if > 0)."));
+FAutoConsoleVariableRef CVarChaosImmPhysJointSoftLinearStiffness(TEXT("p.Chaos.ImmPhys.Joint.SoftLinearStiffness"), ChaosImmediate_Joint_SoftLinearStiffness, TEXT("6Dof joint soft linear stiffness override (if > 0)."));
+FAutoConsoleVariableRef CVarChaosImmPhysJointSoftTwistStiffness(TEXT("p.Chaos.ImmPhys.Joint.SoftTwistStiffness"), ChaosImmediate_Joint_SoftTwistStiffness, TEXT("6Dof joint SoftTwist stiffness override (if > 0)."));
+FAutoConsoleVariableRef CVarChaosImmPhysJointSoftTwistDamping(TEXT("p.Chaos.ImmPhys.Joint.SoftTwistDamping"), ChaosImmediate_Joint_SoftTwistDamping, TEXT("6Dof joint SoftTwist damping override (if > 0)."));
+FAutoConsoleVariableRef CVarChaosImmPhysJointSoftSwingStiffness(TEXT("p.Chaos.ImmPhys.Joint.SoftSwingStiffness"), ChaosImmediate_Joint_SoftSwingStiffness, TEXT("6Dof joint SoftSwing stiffness override (if > 0)."));
+FAutoConsoleVariableRef CVarChaosImmPhysJointSoftSwingDamping(TEXT("p.Chaos.ImmPhys.Joint.SoftSwingDamping"), ChaosImmediate_Joint_SoftSwingDamping, TEXT("6Dof joint SoftSwing damping override (if > 0)."));
+FAutoConsoleVariableRef CVarChaosImmPhysJointLinearDriveStiffness(TEXT("p.Chaos.ImmPhys.Joint.LinearDriveStiffness"), ChaosImmediate_Joint_LinearDriveStiffness, TEXT("6Dof joint drive stiffness override (if > 0)."));
+FAutoConsoleVariableRef CVarChaosImmPhysJointLinearDriveDamping(TEXT("p.Chaos.ImmPhys.Joint.LinearDriveDamping"), ChaosImmediate_Joint_LinearDriveDamping, TEXT("6Dof joint drive damping override (if > 0)."));
+FAutoConsoleVariableRef CVarChaosImmPhysJointAngularDriveStiffness(TEXT("p.Chaos.ImmPhys.Joint.AngularDriveStiffness"), ChaosImmediate_Joint_AngularDriveStiffness, TEXT("6Dof joint drive stiffness override (if > 0)."));
+FAutoConsoleVariableRef CVarChaosImmPhysJointAngularDriveDamping(TEXT("p.Chaos.ImmPhys.Joint.AngularDriveDamping"), ChaosImmediate_Joint_AngularDriveDamping, TEXT("6Dof joint drive damping override (if > 0)."));
+FAutoConsoleVariableRef CVarChaosImmPhysJointMinParentMassRatio(TEXT("p.Chaos.ImmPhys.Joint.MinParentMassRatio"), ChaosImmediate_Joint_MinParentMassRatio, TEXT("6Dof joint MinParentMassRatio (if > 0)"));
+FAutoConsoleVariableRef CVarChaosImmPhysJointMaxInertiaRatio(TEXT("p.Chaos.ImmPhys.Joint.MaxInertiaRatio"), ChaosImmediate_Joint_MaxInertiaRatio, TEXT("6Dof joint MaxInertiaRatio (if > 0)"));
+FAutoConsoleVariableRef CVarChaosImmPhysJointAngularPositionCorrection(TEXT("p.Chaos.ImmPhys.Joint.AngularPositionCorrection"), ChaosImmediate_Joint_AngularPositionCorrection, TEXT("6Dof joint post-rotation constraint position correction amount [0-1]"));
 
 //
 // end remove when finished
@@ -264,6 +264,7 @@ namespace ImmediatePhysics_Chaos
 		JointHandles.Empty();
 	}
 
+
 	FActorHandle* FSimulation::CreateStaticActor(FBodyInstance* BodyInstance)
 	{
 		return CreateActor(EActorType::StaticActor, BodyInstance, BodyInstance->GetUnrealWorldTransform());
@@ -286,18 +287,18 @@ namespace ImmediatePhysics_Chaos
 		FActorHandle* ActorHandle = new FActorHandle(Particles, ActorType, BodyInstance, Transform);
 		int ActorIndex = ActorHandles.Add(ActorHandle);
 
-		// @todo(ccaulfield): add materials
+		// @todo(ccaulfield): shared materials
 		TUniquePtr<FChaosPhysicsMaterial> Material = MakeUnique<FChaosPhysicsMaterial>();
-		if (ChaosImmediate_Material_Resitution >= 0.0f)
-		{
-			Material->Restitution = ChaosImmediate_Material_Resitution;
-		}
-		if (ChaosImmediate_Material_Friction >= 0.0f)
+		if (BodyInstance != nullptr)
 		{
-			Material->Friction = ChaosImmediate_Material_Friction;
+			UPhysicalMaterial* SimplePhysMat = BodyInstance->GetSimplePhysicalMaterial();
+			if (SimplePhysMat != nullptr)
+			{
+				Material->Friction = SimplePhysMat->Friction;
+				Material->Restitution = SimplePhysMat->Restitution;
+			}
 		}
 
-
 		ParticleMaterials.Add(MakeSerializable(Material));
 		PerParticleMaterials.Add(MoveTemp(Material));
 		CollidedParticles.Add(false);
@@ -486,8 +487,6 @@ namespace ImmediatePhysics_Chaos
 		return RollingAverageStepTime;
 	}
 
-	DECLARE_CYCLE_STAT(TEXT("FSimulation::Simulate_Chaos"), STAT_ImmediateSimulate_Chaos, STATGROUP_ImmediatePhysics);
-
 	void FSimulation::Simulate(float InDeltaTime, float MaxStepTime, int32 MaxSubSteps, const FVector& InGravity)
 	{
 		SCOPE_CYCLE_COUNTER(STAT_ImmediateSimulate_Chaos);
@@ -515,8 +514,6 @@ namespace ImmediatePhysics_Chaos
 			}
 
 			FPBDJointSolverSettings JointsSettings = Joints.GetSettings();
-			JointsSettings.ApplyPairIterations = ChaosImmediate_Joint_PairIterations;
-			JointsSettings.ApplyPushOutPairIterations = ChaosImmediate_Joint_PushOutPairIterations;
 			JointsSettings.SwingTwistAngleTolerance = ChaosImmediate_Joint_SwingTwistAngleTolerance;
 			JointsSettings.MinParentMassRatio = ChaosImmediate_Joint_MinParentMassRatio;
 			JointsSettings.MaxInertiaRatio = ChaosImmediate_Joint_MaxInertiaRatio;
@@ -536,19 +533,19 @@ namespace ImmediatePhysics_Chaos
 			JointsSettings.LinearDriveDamping = ChaosImmediate_Joint_LinearDriveDamping;
 			JointsSettings.AngularDriveStiffness = ChaosImmediate_Joint_AngularDriveStiffness;
 			JointsSettings.AngularDriveDamping = ChaosImmediate_Joint_AngularDriveDamping;
-
 			Joints.SetSettings(JointsSettings);
 
-			Evolution.SetNumIterations(ChaosImmediate_Evolution_Iterations);
-			Evolution.SetNumPushOutIterations(ChaosImmediate_Evolution_PushOutIterations);
-
 			Collisions.SetThickness(ChaosImmediate_Collision_Thickness);
-			Collisions.SetPairIterations(ChaosImmediate_Collision_PairIterations);
-			Collisions.SetPushOutPairIterations(ChaosImmediate_Collision_PushOutPairIterations);
-			CollisionsRule.SetPriority(ChaosImmediate_Collision_Priority);
-		
-			// TEMP until we can remove constraints again, or I add broad-phase filtering. (FilterCollisionConstraints will crash since the persistent collision changes)
 			Collisions.SetCollisionsEnabled(ChaosImmediate_Collision_Enabled != 0);
+			CollisionsRule.SetPriority(ChaosImmediate_Collision_Priority);
+
+			SetSolverIterations(
+				ChaosImmediate_Evolution_Iterations,
+				ChaosImmediate_Joint_PairIterations,
+				ChaosImmediate_Collision_PairIterations,
+				ChaosImmediate_Evolution_PushOutIterations,
+				ChaosImmediate_Joint_PushOutPairIterations,
+				ChaosImmediate_Collision_PushOutPairIterations);
 		}
 		UE_LOG(LogChaosJoint, Verbose, TEXT("Simulate Dt = %f Steps %d x %f"), DeltaTime, NumSteps, StepTime);
 
@@ -567,6 +564,36 @@ namespace ImmediatePhysics_Chaos
 		DebugDrawConstraints(1, 2, 1.0f);
 	}
 
+	void FSimulation::SetSolverIterations(const int32 SolverIts, const int32 JointIts, const int32 CollisionIts, const int32 SolverPushOutIts, const int32 JointPushOutIts, const int32 CollisionPushOutIts)
+	{
+		if (SolverIts >= 0)
+		{
+			Evolution.SetNumIterations(SolverIts);
+		}
+		if (SolverPushOutIts >= 0)
+		{
+			Evolution.SetNumPushOutIterations(SolverPushOutIts);
+		}
+
+		if (JointIts >= 0)
+		{
+			Joints.SetNumPairIterations(JointIts);
+		}
+		if (JointPushOutIts >= 0)
+		{
+			Joints.SetNumPushOutPairIterations(JointPushOutIts);
+		}
+
+		if (CollisionIts >= 0)
+		{
+			Collisions.SetPairIterations(CollisionIts);
+		}
+		if (CollisionPushOutIts >= 0)
+		{
+			Collisions.SetPushOutPairIterations(CollisionPushOutIts);
+		}
+	}
+
 	void FSimulation::DebugDrawKinematicParticles(const int32 MinDebugLevel, const int32 MaxDebugLevel, const FColor& Color)
 	{
 #if CHAOS_DEBUG_DRAW
diff --git a/Engine/Source/Runtime/Engine/Private/PhysicsEngine/PhysicsAsset.cpp b/Engine/Source/Runtime/Engine/Private/PhysicsEngine/PhysicsAsset.cpp
--- a/Engine/Source/Runtime/Engine/Private/PhysicsEngine/PhysicsAsset.cpp
+++ b/Engine/Source/Runtime/Engine/Private/PhysicsEngine/PhysicsAsset.cpp
@@ -21,6 +21,16 @@
 
 #define LOCTEXT_NAMESPACE "PhysicsAsset"
 
+FSolverIterations::FSolverIterations()
+	: SolverIterations(2)
+	, JointIterations(1)
+	, CollisionIterations(1)
+	, SolverPushOutIterations(1)
+	, JointPushOutIterations(0)
+	, CollisionPushOutIterations(2)
+{
+}
+
 ///////////////////////////////////////	
 //////////// UPhysicsAsset ////////////
 ///////////////////////////////////////
diff --git a/Engine/Source/Runtime/Engine/Public/Physics/ImmediatePhysics/ImmediatePhysicsChaos/ImmediatePhysicsSimulation_Chaos.h b/Engine/Source/Runtime/Engine/Public/Physics/ImmediatePhysics/ImmediatePhysicsChaos/ImmediatePhysicsSimulation_Chaos.h
--- a/Engine/Source/Runtime/Engine/Public/Physics/ImmediatePhysics/ImmediatePhysicsChaos/ImmediatePhysicsSimulation_Chaos.h
+++ b/Engine/Source/Runtime/Engine/Public/Physics/ImmediatePhysics/ImmediatePhysicsChaos/ImmediatePhysicsSimulation_Chaos.h
@@ -69,6 +69,9 @@ namespace ImmediatePhysics_Chaos
 
 		void SetSimulationSpaceTransform(const FTransform& Transform) { SimulationSpaceTransform = Transform; }
 
+		/** Set new iteration counts. A negative value with leave that iteration count unchanged */
+		void SetSolverIterations(const int32 SolverIts, const int32 JointIts, const int32 CollisionIts, const int32 SolverPushOutIts, const int32 JointPushOutIts, const int32 CollisionPushOutIts);
+
 	private:
 		void ConditionConstraints();
 		FReal UpdateStepTime(const FReal DeltaTime, const FReal MaxStepTime);
diff --git a/Engine/Source/Runtime/Experimental/Chaos/Private/Chaos/ChaosDebugDraw.cpp b/Engine/Source/Runtime/Experimental/Chaos/Private/Chaos/ChaosDebugDraw.cpp
--- a/Engine/Source/Runtime/Experimental/Chaos/Private/Chaos/ChaosDebugDraw.cpp
+++ b/Engine/Source/Runtime/Experimental/Chaos/Private/Chaos/ChaosDebugDraw.cpp
@@ -30,7 +30,7 @@ namespace Chaos
 		float DrawScale = 1.0f;
 		float FontHeight = 10.0f;
 		float FontScale = 1.5f;
-		float ShapeThicknesScale = 2.0f;
+		float ShapeThicknesScale = 0.2f;
 		float PointSize = 2.0f;
 		int DrawPriority = 10.0f;
 
@@ -42,6 +42,7 @@ namespace Chaos
 		FAutoConsoleVariableRef CVarContactOwnerWidth(TEXT("p.Chaos.DebugDrawContactOwnerWidth"), ContactOwnerWidth, TEXT("ContactOwnerWidth."));
 		FAutoConsoleVariableRef CVarConstraintAxisLen(TEXT("p.Chaos.DebugDrawConstraintAxisLen"), ConstraintAxisLen, TEXT("ConstraintAxisLen."));
 		FAutoConsoleVariableRef CVarLineThickness(TEXT("p.Chaos.DebugDrawLineThickness"), LineThickness, TEXT("LineThickness."));
+		FAutoConsoleVariableRef CVarLineShapeThickness(TEXT("p.Chaos.DebugDrawShapeLineThicknessScale"), ShapeThicknesScale, TEXT("Shape lineThickness multiplier."));
 		FAutoConsoleVariableRef CVarScale(TEXT("p.Chaos.DebugDrawScale"), DrawScale, TEXT("Scale applied to all Chaos Debug Draw line lengths etc."));
 
 		//
diff --git a/Engine/Source/Runtime/Experimental/Chaos/Public/Chaos/PBDJointConstraints.h b/Engine/Source/Runtime/Experimental/Chaos/Public/Chaos/PBDJointConstraints.h
--- a/Engine/Source/Runtime/Experimental/Chaos/Public/Chaos/PBDJointConstraints.h
+++ b/Engine/Source/Runtime/Experimental/Chaos/Public/Chaos/PBDJointConstraints.h
@@ -70,6 +70,9 @@ namespace Chaos
 		const FPBDJointSolverSettings& GetSettings() const;
 		void SetSettings(const FPBDJointSolverSettings& InSettings);
 
+		void SetNumPairIterations(const int32 NumPairIterationss) { Settings.ApplyPairIterations = NumPairIterationss; }
+		void SetNumPushOutPairIterations(const int32 NumPairIterationss) { Settings.ApplyPushOutPairIterations = NumPairIterationss; }
+
 		//
 		// Constraint Container API
 		//
