Fix for narrowing INDEX_NONE to an integer in a list initializer.

diff --git a/Engine/Source/Runtime/Engine/Private/AudioStreamingCache.cpp b/Engine/Source/Runtime/Engine/Private/AudioStreamingCache.cpp
--- a/Engine/Source/Runtime/Engine/Private/AudioStreamingCache.cpp
+++ b/Engine/Source/Runtime/Engine/Private/AudioStreamingCache.cpp
@@ -754,7 +754,7 @@ FString FAudioChunkCache::FlushCacheMissLog()
 			CacheMissInfo.SoundWaveName, 
 			CacheMissInfo.ChunkIndex,
 #if WITH_EDITOR
-			INDEX_NONE
+			0
 #endif
 		};
 
