FURL: Don't crash in constructor for url strings containing invalid package name characters while still mapping to an existing file on disk.

Verified fixed on ShooterGame located on c drive with the following url string: "unreal://..\..\..\..\..\..\Windows\system.ini"

diff --git a/Engine/Source/Runtime/Engine/Private/URL.cpp b/Engine/Source/Runtime/Engine/Private/URL.cpp
--- a/Engine/Source/Runtime/Engine/Private/URL.cpp
+++ b/Engine/Source/Runtime/Engine/Private/URL.cpp
@@ -353,9 +353,9 @@ FURL::FURL( FURL* Base, const TCHAR* TextURL, ETravelType Type )
 			FString MapFullName;
 			FText MapNameError;
 			bool bFoundMap = false;
-			if (FPaths::FileExists(URL))
+			if (FPaths::FileExists(URL) && FPackageName::TryConvertFilenameToLongPackageName(URL, MapFullName))
 			{
-				Map = FPackageName::FilenameToLongPackageName(URL);
+				Map = MapFullName;
 				bFoundMap = true;
 			}
 			else if (!FPackageName::DoesPackageNameContainInvalidCharacters(URL, &MapNameError))
