Make budget exceeded logging for the Demo Net Driver configurable, and set the default to 60 seconds.

diff --git a/Engine/Source/Runtime/Engine/Classes/Engine/DemoNetDriver.h b/Engine/Source/Runtime/Engine/Classes/Engine/DemoNetDriver.h
--- a/Engine/Source/Runtime/Engine/Classes/Engine/DemoNetDriver.h
+++ b/Engine/Source/Runtime/Engine/Classes/Engine/DemoNetDriver.h
@@ -387,7 +387,13 @@ ENUM_CLASS_FLAGS(EWriteDemoFrameFlags);
 UCLASS(transient, config=Engine)
 class ENGINE_API UDemoNetDriver : public UNetDriver
 {
-	GENERATED_UCLASS_BODY()
+	GENERATED_BODY()
+
+public:
+
+	UDemoNetDriver(const FObjectInitializer& ObjectInitializer);
+	UDemoNetDriver(FVTableHelper& Helper);
+	~UDemoNetDriver();
 
 	/** Current record/playback frame number */
 	int32 DemoFrameNum;
@@ -448,6 +454,8 @@ class ENGINE_API UDemoNetDriver : public UNetDriver
 
 private:
 
+	void InitDefaults();
+
 	bool LoadCheckpoint(const FGotoResult& GotoResult);
 
 	struct FDeltaCheckpointData
@@ -1251,6 +1259,8 @@ private:
 		return ECreateReplicationChangelistMgrFlags::SkipDeltaCustomState;
 	}
 
+	TUniquePtr<struct FDemoBudgetLogHelper> BudgetLogHelper;
+
 //////////////////////////////////////////////////////////////////////////
 // Replay frame fidelity
 public:
diff --git a/Engine/Source/Runtime/Engine/Private/DemoNetDriver.cpp b/Engine/Source/Runtime/Engine/Private/DemoNetDriver.cpp
--- a/Engine/Source/Runtime/Engine/Private/DemoNetDriver.cpp
+++ b/Engine/Source/Runtime/Engine/Private/DemoNetDriver.cpp
@@ -142,23 +142,122 @@ static bool ShouldActorGoDormantForDemo(const AActor* Actor, const UActorChannel
 
 namespace DemoNetDriverRecordingPrivate
 {
-	static constexpr float WarningTimeInterval = 1.f;
-	static double LastWarningTime = 0.f;
+	static float WarningTimeInterval = 60.f;
+	static FAutoConsoleVariableRef CVarExceededBudgetWarningInterval(
+		TEXT("Demo.ExceededBudgetWarningInterval"),
+		WarningTimeInterval,
+		TEXT("When > 0, we will wait this many seconds between logging warnings for demo recording exceeding time budgets.")
+	);	
+}
 
-	template<size_t N, typename... T>
-	FORCEINLINE static void LogDemoRecordTimeElapsed(TCHAR const (&Format)[N], T... Args)
+struct FDemoBudgetLogHelper
+{
+	enum EBudgetCategory
+	{
+		Prioritization,
+		Replication,
+		CATEGORY_COUNT
+	};
+
+	FDemoBudgetLogHelper(FString&& Identifier)
+		: Identifier(MoveTemp(Identifier))
+	{
+		ResetCounters();
+	}
+
+	void NewFrame()
 	{
-		if (UE_LOG_ACTIVE(LogDemo, Log))
+		if (FirstWarningTime != 0.f)
 		{
+			++NumFrames;
+			bOverBudgetThisFrame = false;
+
 			const double Time = FPlatformTime::Seconds();
-			if ((Time - LastWarningTime) > WarningTimeInterval)
+			if (Time - FirstWarningTime > DemoNetDriverRecordingPrivate::WarningTimeInterval)
 			{
-				UE_LOG(LogDemo, Log, Format, Args...);
-				LastWarningTime = Time;
+				if (UE_LOG_ACTIVE(LogDemo, Log))
+				{
+					TArray<FString> LogLines;
+					LogLines.Reserve((CATEGORY_COUNT * 2) + 1);
+
+					LogLines.Emplace(FString::Printf(TEXT("%s: Recorded Frames: %d, Frames Over Budget: %d"), *Identifier, NumFrames, NumFramesOverBudget));
+
+					for (int32 i = 0; i < CATEGORY_COUNT; ++i)
+					{
+						LogLines.Emplace(FString::Printf(TEXT("Total number of over budget frames in category %d: %d"), i, NumFramesOverBudgetByCategory[i]));
+
+						if (NumFramesOverBudgetByCategory[i] > 0)
+						{
+							LogLines.Emplace(MoveTemp(LogSamplesByBudget[i]));
+						}
+					}
+
+					UE_LOG(LogDemo, Log, TEXT("%s"), *FString::Join(LogLines, TEXT("\n")));
+				}
+
+				ResetCounters();
 			}
 		}
 	}
-}
+
+	template<size_t N, typename... T>
+	void MarkFrameOverBudget(EBudgetCategory Category, TCHAR const (&Format)[N], T... Args)
+	{
+		if (!UE_LOG_ACTIVE(LogDemo, Log))
+		{
+			return;
+		}
+
+		if (DemoNetDriverRecordingPrivate::WarningTimeInterval == 0.f)
+		{
+			UE_LOG(LogDemo, Log, Format, Args...);
+			return;
+		}
+
+		if (!bOverBudgetThisFrame)
+		{
+			bOverBudgetThisFrame = true;
+			++NumFramesOverBudget;
+
+			if (FirstWarningTime == 0.f)
+			{
+				FirstWarningTime = FPlatformTime::Seconds();
+			}
+		}
+
+		++NumFramesOverBudgetByCategory[Category];
+		if (LogSamplesByBudget[Category].IsEmpty())
+		{
+			LogSamplesByBudget[Category] = FString::Printf(Format, Args...);
+		}
+	}
+
+	void ResetCounters()
+	{
+		NumFrames = 0;
+		NumFramesOverBudget = 0;
+		FirstWarningTime = 0.f;
+		for (int32 i = 0; i < CATEGORY_COUNT; ++i)
+		{
+			NumFramesOverBudgetByCategory[i] = 0;
+			LogSamplesByBudget[i] = FString();
+		}
+	}
+
+private:
+
+	bool bOverBudgetThisFrame = false;
+
+	int32 NumFrames;
+	int32 NumFramesOverBudget;
+
+	double FirstWarningTime = 0.f;
+
+	int32 NumFramesOverBudgetByCategory[CATEGORY_COUNT];
+	FString LogSamplesByBudget[CATEGORY_COUNT];
+
+	FString Identifier;
+};
 
 // Helps manage packets, and any associations with streaming levels or exported GUIDs / fields.
 class FScopedPacketManager : public FNoncopyable
@@ -635,10 +734,9 @@ private:
 	UDemoNetDriver.
 -----------------------------------------------------------------------------*/
 
-UDemoNetDriver::UDemoNetDriver(const FObjectInitializer& ObjectInitializer)
-	: Super(ObjectInitializer)
-	, DemoSessionID(FGuid::NewGuid().ToString().ToLower())
+void UDemoNetDriver::InitDefaults()
 {
+	DemoSessionID = FGuid::NewGuid().ToString().ToLower();
 	CurrentLevelIndex = 0;
 	bRecordMapChanges = false;
 	bIsWaitingForHeaderDownload = false;
@@ -654,6 +752,22 @@ UDemoNetDriver::UDemoNetDriver(const FObjectInitializer& ObjectInitializer)
 	RecordBuildConsiderAndPrioritizeTimeSlice = CVarDemoMaximumRepPrioritizeTime.GetValueOnGameThread();
 }
 
+UDemoNetDriver::UDemoNetDriver(const FObjectInitializer& ObjectInitializer)
+	: Super(ObjectInitializer)
+{
+	InitDefaults();
+}
+
+UDemoNetDriver::UDemoNetDriver(FVTableHelper& Helper)
+	: Super(Helper)
+{
+	InitDefaults();
+}
+
+UDemoNetDriver::~UDemoNetDriver()
+{
+}
+
 FString UDemoNetDriver::GetLevelPackageName(const ULevel& InLevel)
 {
 	FString PathName = InLevel.GetOutermost()->GetFName().ToString();
@@ -1097,6 +1211,10 @@ bool UDemoNetDriver::InitListen(FNetworkNotify* InNotify, FURL& ListenURL, bool
 		}
 	}
 
+	// Technically, NetDriver's can be renamed so this could become stale.
+	// However, it's only used for logging and DemoNetDriver's are typically given a special name.
+	BudgetLogHelper = MakeUnique<FDemoBudgetLogHelper>(NetDriverName.ToString());
+
 	bIsWaitingForStream = true;
 
 	ActiveReplayName = DemoURL.Map;
@@ -2555,6 +2673,7 @@ void UDemoNetDriver::TickDemoRecordFrame(float DeltaSeconds)
 	// Save out a frame
 	DemoFrameNum++;
 	ReplicationFrame++;
+	BudgetLogHelper->NewFrame();
 
 	UDemoNetConnection* ClientConnection = CastChecked<UDemoNetConnection>(ClientConnections[0]);
 
@@ -2801,8 +2920,10 @@ void UDemoNetDriver::TickDemoRecordFrame(float DeltaSeconds)
 	// See ReplicatePriorizeActor.
 	if (RecordTimeLimit > 0.0f && TotalPrioritizeActorsTime > RecordTimeLimit)
 	{
-		DemoNetDriverRecordingPrivate::LogDemoRecordTimeElapsed(TEXT("Exceeded maximum desired recording time (during Prioritization).  Max: %.3fms, TimeSpent: %.3fms, Active Actors: %d, Prioritized Actors: %d"),
-			MaxDesiredRecordTimeMS, TotalPrioritizeActorsTimeMS, NumActiveObjects, NumPrioritizedActors);
+		BudgetLogHelper->MarkFrameOverBudget(
+				FDemoBudgetLogHelper::Prioritization,
+				TEXT("Exceeded maximum desired recording time (during Prioritization).  Max: %.3fms, TimeSpent: %.3fms, Active Actors: %d, Prioritized Actors: %d"),
+				MaxDesiredRecordTimeMS, TotalPrioritizeActorsTimeMS, NumActiveObjects, NumPrioritizedActors);
 	}
 
 	float MinRecordHz = CVarDemoMinRecordHz.GetValueOnAnyThread();
@@ -2969,7 +3090,11 @@ bool UDemoNetDriver::ReplicatePrioritizedActor(const FActorPriority& ActorPriori
 
 		if (TotalRecordTimeSeconds > Params.TimeLimitSeconds)
 		{
-			DemoNetDriverRecordingPrivate::LogDemoRecordTimeElapsed(TEXT("Exceeded maximum desired recording time (during Actor Replication).  Max: %.3fms."), MaxDesiredRecordTimeMS);
+			BudgetLogHelper->MarkFrameOverBudget(
+				FDemoBudgetLogHelper::Replication,
+				TEXT("Exceeded maximum desired recording time (during Actor Replication).  Max: %.3fms."),
+				MaxDesiredRecordTimeMS);
+
 			return false;
 		}
 	}
