- Fix crash when user sets first effect on legacy default EQ/Reverb submixes to unsupported effect types
- Move closer to Audio Volumes generically controlling effect parameter interpolation as opposed to hardcoded support for Reverb/EQ/Radio

diff --git a/Engine/Source/Runtime/AudioMixer/Classes/SubmixEffects/AudioMixerSubmixEffectEQ.h b/Engine/Source/Runtime/AudioMixer/Classes/SubmixEffects/AudioMixerSubmixEffectEQ.h
--- a/Engine/Source/Runtime/AudioMixer/Classes/SubmixEffects/AudioMixerSubmixEffectEQ.h
+++ b/Engine/Source/Runtime/AudioMixer/Classes/SubmixEffects/AudioMixerSubmixEffectEQ.h
@@ -57,7 +57,12 @@ public:
 	virtual void OnProcessAudio(const FSoundEffectSubmixInputData& InData, FSoundEffectSubmixOutputData& OutData) override;
 
 	// Sets the effect parameters using the old audio engine preset setting object
-	void SetEffectParameters(const FAudioEQEffect& InEQEffectParameters);
+	virtual bool SetParameters(const FAudioEffectParameters& InParameters) override;
+
+	virtual bool SupportsDefaultEQ() const override
+	{
+		return true;
+	}
 
 	// Called when an audio effect preset is changed
 	virtual void OnPresetChanged() override;
diff --git a/Engine/Source/Runtime/AudioMixer/Classes/SubmixEffects/AudioMixerSubmixEffectReverb.h b/Engine/Source/Runtime/AudioMixer/Classes/SubmixEffects/AudioMixerSubmixEffectReverb.h
--- a/Engine/Source/Runtime/AudioMixer/Classes/SubmixEffects/AudioMixerSubmixEffectReverb.h
+++ b/Engine/Source/Runtime/AudioMixer/Classes/SubmixEffects/AudioMixerSubmixEffectReverb.h
@@ -93,8 +93,11 @@ public:
 	// Called when an audio effect preset is changed
 	virtual void OnPresetChanged() override;
 
-	// We want to receive downmixed submix audio to stereo input for the reverb effect
-	virtual uint32 GetDesiredInputChannelCountOverride() const override { return 2; }
+	// Forces receiving downmixed submix audio to stereo input for the reverb effect
+	virtual uint32 GetDesiredInputChannelCountOverride() const override
+	{
+		return 2;
+	}
 
 	// Returns the drylevel of the effect
 	virtual float GetDryLevel() const override { return DryLevel; }
@@ -103,7 +106,13 @@ public:
 	virtual void OnProcessAudio(const FSoundEffectSubmixInputData& InData, FSoundEffectSubmixOutputData& OutData) override;
 
 	// Sets the reverb effect parameters based from audio thread code
-	void SetEffectParameters(const FAudioReverbEffect& InReverbEffectParameters);
+	virtual bool SetParameters(const FAudioEffectParameters& InEffectParameters) override;
+
+	// Whether this effect supports the default reverb system
+	virtual bool SupportsDefaultReverb() const override
+	{
+		return true;
+	}
 
 private:
 	void UpdateParameters();
diff --git a/Engine/Source/Runtime/AudioMixer/Classes/SubmixEffects/AudioMixerSubmixEffectReverbFast.h b/Engine/Source/Runtime/AudioMixer/Classes/SubmixEffects/AudioMixerSubmixEffectReverbFast.h
--- a/Engine/Source/Runtime/AudioMixer/Classes/SubmixEffects/AudioMixerSubmixEffectReverbFast.h
+++ b/Engine/Source/Runtime/AudioMixer/Classes/SubmixEffects/AudioMixerSubmixEffectReverbFast.h
@@ -9,6 +9,8 @@
 
 #include "AudioMixerSubmixEffectReverbFast.generated.h"
 
+struct FAudioEffectParameters;
+
 USTRUCT(BlueprintType)
 struct AUDIOMIXER_API FSubmixEffectReverbFastSettings
 {
@@ -107,7 +109,13 @@ public:
 	virtual void OnProcessAudio(const FSoundEffectSubmixInputData& InData, FSoundEffectSubmixOutputData& OutData) override;
 
 	// Sets the reverb effect parameters based from audio thread code
-	void SetEffectParameters(const FAudioReverbEffect& InReverbEffectParameters);
+	virtual bool SetParameters(const FAudioEffectParameters& InParameters) override;
+
+	// Whether this effect supports the default reverb system
+	virtual bool SupportsDefaultReverb() const override
+	{
+		return true;
+	}
 
 private:
 
diff --git a/Engine/Source/Runtime/AudioMixer/Private/AudioMixerEffectsManager.cpp b/Engine/Source/Runtime/AudioMixer/Private/AudioMixerEffectsManager.cpp
--- a/Engine/Source/Runtime/AudioMixer/Private/AudioMixerEffectsManager.cpp
+++ b/Engine/Source/Runtime/AudioMixer/Private/AudioMixerEffectsManager.cpp
@@ -9,89 +9,6 @@
 
 namespace Audio
 {
-#define ENABLE_REVERB_SETTINGS_PRINTING 0
-#define ENABLE_EQ_SETTINGS_PRINTING 0
-
-	static void PrintReverbSettings(const FAudioReverbEffect& Settings)
-	{
-#if ENABLE_REVERB_SETTINGS_PRINTING
-
-		const char* FmtText =
-			"\nVolume: %.4f\n"
-			"Density: %.4f\n"
-			"Diffusion: %.4f\n"
-			"Gain: %.4f\n"
-			"GainHF: %.4f\n"
-			"DecayTime: %.4f\n"
-			"DecayHFRatio: %.4f\n"
-			"ReflectionsGain: %.4f\n"
-			"ReflectionsDelay: %.4f\n"
-			"LateGain: %.4f\n"
-			"LateDelay: %.4f\n"
-			"AirAbsorptionGainHF: %.4f\n"
-			"RoomRolloffFactor: %.4f\n";
-
-		FString FmtString(FmtText);
-
-		FString Params = FString::Printf(
-			*FmtString,
-			Settings.Volume,
-			Settings.Density,
-			Settings.Diffusion,
-			Settings.Gain,
-			Settings.GainHF,
-			Settings.DecayTime,
-			Settings.DecayHFRatio,
-			Settings.ReflectionsGain,
-			Settings.ReflectionsDelay,
-			Settings.LateGain,
-			Settings.LateDelay,
-			Settings.AirAbsorptionGainHF,
-			Settings.RoomRolloffFactor
-			);
-
-		UE_LOG(LogTemp, Log, TEXT("%s"), *Params);
-#endif
-	}
-
-	static void PrintEQSettings(const FAudioEQEffect& Settings)
-	{
-#if ENABLE_EQ_SETTINGS_PRINTING
-		const char* FmtText =
-			"\nFrequencyCenter0: %.4f\n"
-			"Gain0: %.4f\n"
-			"Bandwidth0: %.4f\n"
-			"FrequencyCenter1: %.4f\n"
-			"Gain1: %.4f\n"
-			"Bandwidth1: %.4f\n"
-			"FrequencyCenter2: %.4f\n"
-			"Gain2: %.4f\n"
-			"Bandwidth2: %.4f\n"
-			"FrequencyCenter3: %.4f\n"
-			"Gain3: %.4f\n"
-			"Bandwidth3: %.4f\n";
-
-		FString FmtString(FmtText);
-
-		FString Params = FString::Printf(
-			*FmtString,
-			Settings.FrequencyCenter0,
-			Settings.Gain0,
-			Settings.Bandwidth0,
-			Settings.FrequencyCenter1,
-			Settings.Gain1,
-			Settings.Bandwidth1,
-			Settings.FrequencyCenter2,
-			Settings.Gain2,
-			Settings.Bandwidth2,
-			Settings.FrequencyCenter3,
-			Settings.Gain3,
-			Settings.Bandwidth3
-		);
-
-		UE_LOG(LogTemp, Log, TEXT("%s"), *Params);
-#endif
-	}
 
 	FAudioMixerEffectsManager::FAudioMixerEffectsManager(FAudioDevice* InDevice)
 		: FAudioEffectsManager(InDevice)
@@ -101,7 +18,7 @@ namespace Audio
 	FAudioMixerEffectsManager::~FAudioMixerEffectsManager()
 	{}
 
-	void FAudioMixerEffectsManager::SetReverbEffectParameters(const FAudioReverbEffect& ReverbEffectParameters)
+	void FAudioMixerEffectsManager::SetReverbEffectParameters(const FAudioEffectParameters& InEffectParameters)
 	{
 		FMixerDevice* MixerDevice = (FMixerDevice*)AudioDevice;
 
@@ -110,18 +27,36 @@ namespace Audio
 		
 		if (MasterReverbSubmixPtr.IsValid())
 		{
+			bool bReportFailure = false;
+
 			FSoundEffectSubmixPtr SoundEffectSubmix = MasterReverbSubmixPtr->GetSubmixEffect(0);
 			if (SoundEffectSubmix.IsValid())
 			{
-				TSharedPtr<FSubmixEffectReverbFast, ESPMode::ThreadSafe> SoundEffectReverb = StaticCastSharedPtr<FSubmixEffectReverbFast, FSoundEffectSubmix, ESPMode::ThreadSafe>(SoundEffectSubmix);
-				SoundEffectReverb->SetEffectParameters(ReverbEffectParameters);
+				InEffectParameters.PrintSettings();
+
+				if (!SoundEffectSubmix->SupportsDefaultReverb() || !SoundEffectSubmix->SetParameters(InEffectParameters))
+				{
+					bReportFailure = true;
+				}
+				else
+				{
+					InvalidReverbEffect.Reset();
+				}
+			}
+			else
+			{
+				bReportFailure = true;
+			}
 
-				PrintReverbSettings(ReverbEffectParameters);
+			if (bReportFailure && InvalidReverbEffect != SoundEffectSubmix)
+			{
+				UE_LOG(LogAudioMixer, Error, TEXT("Failed to update reverb parameters on Default Reverb Submix. Ensure first submix effect is supported type"));
+				InvalidReverbEffect = SoundEffectSubmix;
 			}
 		}
 	}
 
-	void FAudioMixerEffectsManager::SetEQEffectParameters(const FAudioEQEffect& InEQEffectParameters)
+	void FAudioMixerEffectsManager::SetEQEffectParameters(const FAudioEffectParameters& InEffectParameters)
 	{
 		FMixerDevice* MixerDevice = (FMixerDevice*)AudioDevice;
 
@@ -130,19 +65,36 @@ namespace Audio
 
 		if (MasterEQSubmixPtr.IsValid())
 		{
+			bool bReportFailure = false;
 			FSoundEffectSubmixPtr SoundEffectSubmix = MasterEQSubmixPtr->GetSubmixEffect(0);
 			if (SoundEffectSubmix.IsValid())
 			{
-				TSharedPtr<FSubmixEffectSubmixEQ, ESPMode::ThreadSafe> SoundEffectEQ = StaticCastSharedPtr<FSubmixEffectSubmixEQ, FSoundEffectSubmix, ESPMode::ThreadSafe>(SoundEffectSubmix);
-				SoundEffectEQ->SetEffectParameters(InEQEffectParameters);
-				PrintEQSettings(InEQEffectParameters);
+				InEffectParameters.PrintSettings();
+
+				if (!SoundEffectSubmix->SupportsDefaultEQ() || !SoundEffectSubmix->SetParameters(InEffectParameters))
+				{
+					bReportFailure = true;
+				}
+				else
+				{
+					InvalidEQEffect.Reset();
+				}
+			}
+			else
+			{
+				bReportFailure = true;
+			}
+
+			if (bReportFailure && InvalidEQEffect != SoundEffectSubmix)
+			{
+				UE_LOG(LogAudioMixer, Error, TEXT("Failed to update EQ parameters on legacy Default EQ Submix. Ensure first submix effect is supported type"));
+				InvalidEQEffect = SoundEffectSubmix;
 			}
 		}
 	}
 
-	void FAudioMixerEffectsManager::SetRadioEffectParameters(const FAudioRadioEffect& ReverbEffectParameters)
+	void FAudioMixerEffectsManager::SetRadioEffectParameters(const FAudioEffectParameters& InEffectParameters)
 	{
-
+		// Effect system deprecated
 	}
-
 }
diff --git a/Engine/Source/Runtime/AudioMixer/Private/AudioMixerEffectsManager.h b/Engine/Source/Runtime/AudioMixer/Private/AudioMixerEffectsManager.h
--- a/Engine/Source/Runtime/AudioMixer/Private/AudioMixerEffectsManager.h
+++ b/Engine/Source/Runtime/AudioMixer/Private/AudioMixerEffectsManager.h
@@ -8,6 +8,7 @@ for the multiplatform audio mixer
 #pragma once
 
 #include "AudioEffect.h"
+#include "Sound/SoundEffectSubmix.h"
 
 namespace Audio
 {
@@ -20,14 +21,16 @@ namespace Audio
 		~FAudioMixerEffectsManager() override;
 
 		//~ Begin FAudioEffectsManager
-		virtual void SetReverbEffectParameters(const FAudioReverbEffect& ReverbEffectParameters) override;
-		virtual void SetEQEffectParameters(const FAudioEQEffect& ReverbEffectParameters) override;
-		virtual void SetRadioEffectParameters(const FAudioRadioEffect& ReverbEffectParameters) override;
+		virtual void SetReverbEffectParameters(const FAudioEffectParameters& InEffectParameters) override;
+		virtual void SetEQEffectParameters(const FAudioEffectParameters& InEffectParameters) override;
+		virtual void SetRadioEffectParameters(const FAudioEffectParameters& InEffectParameters) override;
 		//~ End FAudioEffectsManager
 
 	protected:
-
 		FRuntimeFloatCurve MasterReverbWetLevelCurve;
+
+	private:
+		FSoundEffectSubmixPtr InvalidReverbEffect;
+		FSoundEffectSubmixPtr InvalidEQEffect;
 	};
 }
-
diff --git a/Engine/Source/Runtime/AudioMixer/Private/Effects/AudioMixerSubmixEffectEQ.cpp b/Engine/Source/Runtime/AudioMixer/Private/Effects/AudioMixerSubmixEffectEQ.cpp
--- a/Engine/Source/Runtime/AudioMixer/Private/Effects/AudioMixerSubmixEffectEQ.cpp
+++ b/Engine/Source/Runtime/AudioMixer/Private/Effects/AudioMixerSubmixEffectEQ.cpp
@@ -165,35 +165,37 @@ static float GetClampedFrequency(const float InFrequency)
 	return FMath::Clamp(InFrequency, 20.0f, 20000.0f);
 }
 
-void FSubmixEffectSubmixEQ::SetEffectParameters(const FAudioEQEffect& InEQEffectParameters)
+bool FSubmixEffectSubmixEQ::SetParameters(const FAudioEffectParameters& InParameters)
 {
+	const FAudioEQEffect& EQEffectParameters = static_cast<const FAudioEQEffect&>(InParameters);
+
 	// This function maps the old audio engine eq effect params to the new eq effect.
 	// Note that this is always a 4-band EQ and not flexible w/ respect.
 	FSubmixEffectSubmixEQSettings NewSettings;
 	FSubmixEffectEQBand Band;
 
 	Band.bEnabled = true;
-	Band.Frequency = GetClampedFrequency(InEQEffectParameters.FrequencyCenter0);
-	Band.Bandwidth = GetClampedBandwidth(InEQEffectParameters.Bandwidth0);
-	Band.GainDb = Audio::ConvertToDecibels(GetClampedGain(InEQEffectParameters.Gain0));
+	Band.Frequency = GetClampedFrequency(EQEffectParameters.FrequencyCenter0);
+	Band.Bandwidth = GetClampedBandwidth(EQEffectParameters.Bandwidth0);
+	Band.GainDb = Audio::ConvertToDecibels(GetClampedGain(EQEffectParameters.Gain0));
 	NewSettings.EQBands.Add(Band);
 
 	Band.bEnabled = true;
-	Band.Frequency = GetClampedFrequency(InEQEffectParameters.FrequencyCenter1);
-	Band.Bandwidth = GetClampedBandwidth(InEQEffectParameters.Bandwidth1);
-	Band.GainDb = Audio::ConvertToDecibels(GetClampedGain(InEQEffectParameters.Gain1));
+	Band.Frequency = GetClampedFrequency(EQEffectParameters.FrequencyCenter1);
+	Band.Bandwidth = GetClampedBandwidth(EQEffectParameters.Bandwidth1);
+	Band.GainDb = Audio::ConvertToDecibels(GetClampedGain(EQEffectParameters.Gain1));
 	NewSettings.EQBands.Add(Band);
 
 	Band.bEnabled = true;
-	Band.Frequency = GetClampedFrequency(InEQEffectParameters.FrequencyCenter2);
-	Band.Bandwidth = GetClampedBandwidth(InEQEffectParameters.Bandwidth2);
-	Band.GainDb = Audio::ConvertToDecibels(GetClampedGain(InEQEffectParameters.Gain2));
+	Band.Frequency = GetClampedFrequency(EQEffectParameters.FrequencyCenter2);
+	Band.Bandwidth = GetClampedBandwidth(EQEffectParameters.Bandwidth2);
+	Band.GainDb = Audio::ConvertToDecibels(GetClampedGain(EQEffectParameters.Gain2));
 	NewSettings.EQBands.Add(Band);
 
 	Band.bEnabled = true;
-	Band.Frequency = GetClampedFrequency(InEQEffectParameters.FrequencyCenter3);
-	Band.Bandwidth = GetClampedBandwidth(InEQEffectParameters.Bandwidth3);
-	Band.GainDb = Audio::ConvertToDecibels(GetClampedGain(InEQEffectParameters.Gain3));
+	Band.Frequency = GetClampedFrequency(EQEffectParameters.FrequencyCenter3);
+	Band.Bandwidth = GetClampedBandwidth(EQEffectParameters.Bandwidth3);
+	Band.GainDb = Audio::ConvertToDecibels(GetClampedGain(EQEffectParameters.Gain3));
 	NewSettings.EQBands.Add(Band);
 
 	// Don't make any changes if this is the exact same parameters
@@ -203,6 +205,7 @@ void FSubmixEffectSubmixEQ::SetEffectParameters(const FAudioEQEffect& InEQEffect
 		PendingSettings.SetParams(GameThreadEQSettings);
 	}
 
+	return true;
 }
 
 void FSubmixEffectSubmixEQ::UpdateParameters(const int32 InNumOutputChannels)
diff --git a/Engine/Source/Runtime/AudioMixer/Private/Effects/AudioMixerSubmixEffectReverb.cpp b/Engine/Source/Runtime/AudioMixer/Private/Effects/AudioMixerSubmixEffectReverb.cpp
--- a/Engine/Source/Runtime/AudioMixer/Private/Effects/AudioMixerSubmixEffectReverb.cpp
+++ b/Engine/Source/Runtime/AudioMixer/Private/Effects/AudioMixerSubmixEffectReverb.cpp
@@ -98,7 +98,7 @@ void FSubmixEffectReverb::OnPresetChanged()
 
 	DryLevel = Settings.DryLevel;
 
-	SetEffectParameters(ReverbEffect);
+	SetParameters(ReverbEffect);
 }
 
 void FSubmixEffectReverb::OnProcessAudio(const FSoundEffectSubmixInputData& InData, FSoundEffectSubmixOutputData& OutData)
@@ -171,26 +171,28 @@ void FSubmixEffectReverb::OnProcessAudio(const FSoundEffectSubmixInputData& InDa
 	}
 }
 
-void FSubmixEffectReverb::SetEffectParameters(const FAudioReverbEffect& InParams)
+bool FSubmixEffectReverb::SetParameters(const FAudioEffectParameters& InEffectParameters)
 {
 	LLM_SCOPE(ELLMTag::AudioMixer);
 
+	const FAudioReverbEffect& ReverbParams = static_cast<const FAudioReverbEffect&>(InEffectParameters);
+
 	Audio::FPlateReverbSettings NewSettings;
 
-	NewSettings.EarlyReflections.Gain = FMath::GetMappedRangeValueClamped({ 0.0f, 3.16f }, { 0.0f, 1.0f }, InParams.ReflectionsGain);
-	NewSettings.EarlyReflections.PreDelayMsec = FMath::GetMappedRangeValueClamped({ 0.0f, 0.3f }, { 0.0f, 300.0f }, InParams.ReflectionsDelay);
-	NewSettings.EarlyReflections.Bandwidth = FMath::GetMappedRangeValueClamped({ 0.0f, 1.0f }, { 0.0f, 1.0f }, InParams.GainHF);
+	NewSettings.EarlyReflections.Gain = FMath::GetMappedRangeValueClamped({ 0.0f, 3.16f }, { 0.0f, 1.0f }, ReverbParams.ReflectionsGain);
+	NewSettings.EarlyReflections.PreDelayMsec = FMath::GetMappedRangeValueClamped({ 0.0f, 0.3f }, { 0.0f, 300.0f }, ReverbParams.ReflectionsDelay);
+	NewSettings.EarlyReflections.Bandwidth = FMath::GetMappedRangeValueClamped({ 0.0f, 1.0f }, { 0.0f, 1.0f }, ReverbParams.GainHF);
 
-	NewSettings.LateDelayMsec = FMath::GetMappedRangeValueClamped({ 0.0f, 0.1f }, { 0.0f, 100.0f }, InParams.LateDelay);
-	NewSettings.LateGain = FMath::GetMappedRangeValueClamped({ 0.0f, 1.0f }, { 0.0f, 1.0f }, InParams.Gain);
-	NewSettings.Bandwidth = FMath::GetMappedRangeValueClamped({ 0.0f, 1.0f }, { 0.2f, 0.999f }, InParams.AirAbsorptionGainHF);
-	NewSettings.Diffusion = FMath::GetMappedRangeValueClamped({ 0.0f, 1.0f }, { 0.0f, 1.0f }, InParams.Diffusion);
-	NewSettings.Dampening = FMath::GetMappedRangeValueClamped({ 0.1f, 2.0f }, { 0.0f, 0.999f }, InParams.DecayHFRatio);
-	NewSettings.Density = FMath::GetMappedRangeValueClamped({ 0.0f, 1.0f }, { 0.01f, 1.0f }, InParams.Density);
-	NewSettings.Wetness = FMath::GetMappedRangeValueClamped({ 0.0f, 10.0f }, { 0.0f, 10.0f }, InParams.Volume);
+	NewSettings.LateDelayMsec = FMath::GetMappedRangeValueClamped({ 0.0f, 0.1f }, { 0.0f, 100.0f }, ReverbParams.LateDelay);
+	NewSettings.LateGain = FMath::GetMappedRangeValueClamped({ 0.0f, 1.0f }, { 0.0f, 1.0f }, ReverbParams.Gain);
+	NewSettings.Bandwidth = FMath::GetMappedRangeValueClamped({ 0.0f, 1.0f }, { 0.2f, 0.999f }, ReverbParams.AirAbsorptionGainHF);
+	NewSettings.Diffusion = FMath::GetMappedRangeValueClamped({ 0.0f, 1.0f }, { 0.0f, 1.0f }, ReverbParams.Diffusion);
+	NewSettings.Dampening = FMath::GetMappedRangeValueClamped({ 0.1f, 2.0f }, { 0.0f, 0.999f }, ReverbParams.DecayHFRatio);
+	NewSettings.Density = FMath::GetMappedRangeValueClamped({ 0.0f, 1.0f }, { 0.01f, 1.0f }, ReverbParams.Density);
+	NewSettings.Wetness = FMath::GetMappedRangeValueClamped({ 0.0f, 10.0f }, { 0.0f, 10.0f }, ReverbParams.Volume);
 
 	// Use mapping function to get decay time in seconds to internal linear decay scale value
-	const float DecayValue = DecayCurve.Eval(InParams.DecayTime);
+	const float DecayValue = DecayCurve.Eval(ReverbParams.DecayTime);
 	NewSettings.Decay = DecayValue;
 
 	// Convert to db
@@ -200,6 +202,8 @@ void FSubmixEffectReverb::SetEffectParameters(const FAudioReverbEffect& InParams
 	Params.SetParams(NewSettings);
 
 	bIsEnabled = true;
+
+	return true;
 }
 
 void FSubmixEffectReverb::UpdateParameters()
diff --git a/Engine/Source/Runtime/AudioMixer/Private/Effects/AudioMixerSubmixEffectReverbFast.cpp b/Engine/Source/Runtime/AudioMixer/Private/Effects/AudioMixerSubmixEffectReverbFast.cpp
--- a/Engine/Source/Runtime/AudioMixer/Private/Effects/AudioMixerSubmixEffectReverbFast.cpp
+++ b/Engine/Source/Runtime/AudioMixer/Private/Effects/AudioMixerSubmixEffectReverbFast.cpp
@@ -108,7 +108,7 @@ void FSubmixEffectReverbFast::OnPresetChanged()
 	ReverbEffect.RoomRolloffFactor = 0.0f; // not used
 	ReverbEffect.Volume = Settings.bBypass ? 0.0f : Settings.WetLevel;
 
-	SetEffectParameters(ReverbEffect);
+	SetParameters(ReverbEffect);
 
 	Audio::FWetDry NewWetDryParams;
 	NewWetDryParams.DryLevel = Settings.bBypass ? 1.0f : Settings.DryLevel;
@@ -149,10 +149,12 @@ void FSubmixEffectReverbFast::OnProcessAudio(const FSoundEffectSubmixInputData&
 	Audio::MixInBufferFast(*InData.AudioBuffer, *OutData.AudioBuffer, LastDry, CurrentWetDry.DryLevel);
 }
 
-void FSubmixEffectReverbFast::SetEffectParameters(const FAudioReverbEffect& InParams)
+bool FSubmixEffectReverbFast::SetParameters(const FAudioEffectParameters& InParams)
 {
 	LLM_SCOPE(ELLMTag::AudioMixer);
 
+	const FAudioReverbEffect& ReverbEffectParams = static_cast<const FAudioReverbEffect&>(InParams);
+
 	/* `FPlateReverbFast` produces slightly different quality effect than `FPlateReverb`. Comparing the 
 	 * settings between FSubmixEffectReverb and FSubmixEffectReverbFast slight differences will arise.
 	 *
@@ -165,20 +167,20 @@ void FSubmixEffectReverbFast::SetEffectParameters(const FAudioReverbEffect& InPa
 	Audio::FPlateReverbFastSettings NewSettings;
 
 	// Early Reflections
-	NewSettings.EarlyReflections.Gain = FMath::GetMappedRangeValueClamped({ 0.0f, 3.16f }, { 0.0f, 1.0f }, InParams.ReflectionsGain);
-	NewSettings.EarlyReflections.PreDelayMsec = FMath::GetMappedRangeValueClamped({ 0.0f, 0.3f }, { 0.0f, 300.0f }, InParams.ReflectionsDelay);
-	NewSettings.EarlyReflections.Bandwidth = FMath::GetMappedRangeValueClamped({ 0.0f, 1.0f }, { 0.0f, 1.0f }, InParams.GainHF);
+	NewSettings.EarlyReflections.Gain = FMath::GetMappedRangeValueClamped({ 0.0f, 3.16f }, { 0.0f, 1.0f }, ReverbEffectParams.ReflectionsGain);
+	NewSettings.EarlyReflections.PreDelayMsec = FMath::GetMappedRangeValueClamped({ 0.0f, 0.3f }, { 0.0f, 300.0f }, ReverbEffectParams.ReflectionsDelay);
+	NewSettings.EarlyReflections.Bandwidth = FMath::GetMappedRangeValueClamped({ 0.0f, 1.0f }, { 0.0f, 1.0f }, ReverbEffectParams.GainHF);
 
 	// LateReflections
-	NewSettings.LateReflections.LateDelayMsec = FMath::GetMappedRangeValueClamped({ 0.0f, 0.1f }, { 0.0f, 100.0f }, InParams.LateDelay);
-	NewSettings.LateReflections.LateGainDB = FMath::GetMappedRangeValueClamped({ 0.0f, 1.0f }, { 0.0f, 1.0f }, InParams.Gain);
-	NewSettings.LateReflections.Bandwidth = FMath::GetMappedRangeValueClamped({ 0.0f, 1.0f }, { 0.1f, 0.6f }, InParams.AirAbsorptionGainHF);
-	NewSettings.LateReflections.Diffusion = FMath::GetMappedRangeValueClamped({ 0.05f, 1.0f }, { 0.0f, 0.95f }, InParams.Diffusion);
-	NewSettings.LateReflections.Dampening = FMath::GetMappedRangeValueClamped({ 0.05f, 1.95f }, { 0.0f, 0.999f }, InParams.DecayHFRatio);
-	NewSettings.LateReflections.Density = FMath::GetMappedRangeValueClamped({ 0.0f, 0.95f }, { 0.06f, 1.0f }, InParams.Density);
+	NewSettings.LateReflections.LateDelayMsec = FMath::GetMappedRangeValueClamped({ 0.0f, 0.1f }, { 0.0f, 100.0f }, ReverbEffectParams.LateDelay);
+	NewSettings.LateReflections.LateGainDB = FMath::GetMappedRangeValueClamped({ 0.0f, 1.0f }, { 0.0f, 1.0f }, ReverbEffectParams.Gain);
+	NewSettings.LateReflections.Bandwidth = FMath::GetMappedRangeValueClamped({ 0.0f, 1.0f }, { 0.1f, 0.6f }, ReverbEffectParams.AirAbsorptionGainHF);
+	NewSettings.LateReflections.Diffusion = FMath::GetMappedRangeValueClamped({ 0.05f, 1.0f }, { 0.0f, 0.95f }, ReverbEffectParams.Diffusion);
+	NewSettings.LateReflections.Dampening = FMath::GetMappedRangeValueClamped({ 0.05f, 1.95f }, { 0.0f, 0.999f }, ReverbEffectParams.DecayHFRatio);
+	NewSettings.LateReflections.Density = FMath::GetMappedRangeValueClamped({ 0.0f, 0.95f }, { 0.06f, 1.0f }, ReverbEffectParams.Density);
 
 	// Use mapping function to get decay time in seconds to internal linear decay scale value
-	const float DecayValue = DecayCurve.Eval(InParams.DecayTime);
+	const float DecayValue = DecayCurve.Eval(ReverbEffectParams.DecayTime);
 	NewSettings.LateReflections.Decay = DecayValue;
 
 	// Convert to db
@@ -186,6 +188,8 @@ void FSubmixEffectReverbFast::SetEffectParameters(const FAudioReverbEffect& InPa
 
 	// Apply the settings the thread safe settings object
 	ReverbParams.SetParams(NewSettings);
+
+	return true;
 }
 
 void FSubmixEffectReverbFast::UpdateParameters()
diff --git a/Engine/Source/Runtime/Engine/Classes/Sound/AudioVolume.h b/Engine/Source/Runtime/Engine/Classes/Sound/AudioVolume.h
--- a/Engine/Source/Runtime/Engine/Classes/Sound/AudioVolume.h
+++ b/Engine/Source/Runtime/Engine/Classes/Sound/AudioVolume.h
@@ -8,113 +8,17 @@
 #pragma once
 
 #include "CoreMinimal.h"
-#include "UObject/ObjectMacros.h"
-#include "UObject/Class.h"
+
 #include "GameFramework/Volume.h"
-#include "Audio.h"
+#include "ReverbSettings.h"
+#include "UObject/Class.h"
+#include "UObject/ObjectMacros.h"
+
 #include "AudioVolume.generated.h"
 
 class AAudioVolume;
 struct FBodyInstance;
 
-/**
- * DEPRECATED: Exists for backwards compatibility
- * Indicates a reverb preset to use.
- *
- */
-UENUM()
-enum ReverbPreset
-{
-	REVERB_Default,
-	REVERB_Bathroom,
-	REVERB_StoneRoom,
-	REVERB_Auditorium,
-	REVERB_ConcertHall,
-	REVERB_Cave,
-	REVERB_Hallway,
-	REVERB_StoneCorridor,
-	REVERB_Alley,
-	REVERB_Forest,
-	REVERB_City,
-	REVERB_Mountains,
-	REVERB_Quarry,
-	REVERB_Plain,
-	REVERB_ParkingLot,
-	REVERB_SewerPipe,
-	REVERB_Underwater,
-	REVERB_SmallRoom,
-	REVERB_MediumRoom,
-	REVERB_LargeRoom,
-	REVERB_MediumHall,
-	REVERB_LargeHall,
-	REVERB_Plate,
-	REVERB_MAX,
-};
-
-/** Struct encapsulating settings for reverb effects. */
-USTRUCT(BlueprintType)
-struct FReverbSettings
-{
-	GENERATED_USTRUCT_BODY()
-
-	/* Whether to apply the reverb settings below. */
-	UPROPERTY(EditAnywhere, BlueprintReadWrite, Category=ReverbSettings )
-	bool bApplyReverb;
-
-#if WITH_EDITORONLY_DATA
-	/** The reverb preset to employ. */
-	UPROPERTY()
-	TEnumAsByte<enum ReverbPreset> ReverbType_DEPRECATED;
-#endif
-
-	/** The reverb asset to employ. */
-	UPROPERTY(EditAnywhere, BlueprintReadWrite, Category=ReverbSettings)
-	class UReverbEffect* ReverbEffect;
-	
-
-	/** This is used to apply plugin-specific settings when a Reverb Plugin is being used. */
-	UPROPERTY(EditAnywhere, BlueprintReadWrite, Category = ReverbSettings)
-	class USoundEffectSubmixPreset* ReverbPluginEffect;
-
-	/** Volume level of the reverb affect. */
-	UPROPERTY(EditAnywhere, BlueprintReadWrite, Category = ReverbSettings)
-	float Volume;
-
-	/** Time to fade from the current reverb settings into this setting, in seconds. */
-	UPROPERTY(EditAnywhere, BlueprintReadWrite, Category=ReverbSettings)
-	float FadeTime;
-
-	FReverbSettings()
-		: bApplyReverb(true)
-#if WITH_EDITORONLY_DATA
-		, ReverbType_DEPRECATED(REVERB_Default)
-#endif
-		, ReverbEffect(nullptr)
-		, ReverbPluginEffect(nullptr)
-		, Volume(0.5f)
-		, FadeTime(2.0f)
-	{
-	}
-
-	bool operator==(const FReverbSettings& Other) const;
-	bool operator!=(const FReverbSettings& Other) const { return !(*this == Other); }
-
-#if WITH_EDITORONLY_DATA
-	void PostSerialize(const FArchive& Ar);
-#endif
-};
-
-#if WITH_EDITORONLY_DATA
-template<>
-struct TStructOpsTypeTraits<FReverbSettings> : public TStructOpsTypeTraitsBase2<FReverbSettings>
-{
-	enum 
-	{
-		WithPostSerialize = true,
-	};
-};
-#endif
-
 /** Struct encapsulating settings for interior areas. */
 USTRUCT(BlueprintType)
 struct FInteriorSettings
@@ -270,6 +174,3 @@ public:
 	virtual void PostRegisterAllComponents() override;
 	//~ End AActor Interface
 };
-
-
-
diff --git a/Engine/Source/Runtime/Engine/Classes/Sound/ReverbSettings.h b/Engine/Source/Runtime/Engine/Classes/Sound/ReverbSettings.h
new file mode 100644
--- /dev/null
+++ b/Engine/Source/Runtime/Engine/Classes/Sound/ReverbSettings.h
@@ -0,0 +1,108 @@
+// Copyright Epic Games, Inc. All Rights Reserved.
+#pragma once
+
+#include "CoreMinimal.h"
+
+#include "ReverbSettings.generated.h"
+
+class UReverbEffect;
+class USoundEffectSubmixPreset;
+
+
+/**
+ * DEPRECATED: Exists for backwards compatibility
+ * Indicates a reverb preset to use.
+ *
+ */
+UENUM()
+enum ReverbPreset
+{
+	REVERB_Default,
+	REVERB_Bathroom,
+	REVERB_StoneRoom,
+	REVERB_Auditorium,
+	REVERB_ConcertHall,
+	REVERB_Cave,
+	REVERB_Hallway,
+	REVERB_StoneCorridor,
+	REVERB_Alley,
+	REVERB_Forest,
+	REVERB_City,
+	REVERB_Mountains,
+	REVERB_Quarry,
+	REVERB_Plain,
+	REVERB_ParkingLot,
+	REVERB_SewerPipe,
+	REVERB_Underwater,
+	REVERB_SmallRoom,
+	REVERB_MediumRoom,
+	REVERB_LargeRoom,
+	REVERB_MediumHall,
+	REVERB_LargeHall,
+	REVERB_Plate,
+	REVERB_MAX,
+};
+
+/** Struct encapsulating settings for reverb effects. */
+USTRUCT(BlueprintType)
+struct FReverbSettings
+{
+	GENERATED_USTRUCT_BODY()
+
+	/* Whether to apply the reverb settings below. */
+	UPROPERTY(EditAnywhere, BlueprintReadWrite, Category=ReverbSettings )
+	bool bApplyReverb;
+
+#if WITH_EDITORONLY_DATA
+	/** The reverb preset to employ. */
+	UPROPERTY()
+	TEnumAsByte<ReverbPreset> ReverbType_DEPRECATED;
+#endif
+
+	/** The reverb asset to employ. */
+	UPROPERTY(EditAnywhere, BlueprintReadWrite, Category=ReverbSettings)
+	UReverbEffect* ReverbEffect;
+
+	/** This is used to apply plugin-specific settings when a Reverb Plugin is being used. */
+	UPROPERTY(EditAnywhere, BlueprintReadWrite, Category = ReverbSettings)
+	USoundEffectSubmixPreset* ReverbPluginEffect;
+
+	/** Volume level of the reverb affect. */
+	UPROPERTY(EditAnywhere, BlueprintReadWrite, Category = ReverbSettings)
+	float Volume;
+
+	/** Time to fade from the current reverb settings into this setting, in seconds. */
+	UPROPERTY(EditAnywhere, BlueprintReadWrite, Category=ReverbSettings)
+	float FadeTime;
+
+	FReverbSettings()
+		: bApplyReverb(true)
+#if WITH_EDITORONLY_DATA
+		, ReverbType_DEPRECATED(REVERB_Default)
+#endif
+		, ReverbEffect(nullptr)
+		, ReverbPluginEffect(nullptr)
+		, Volume(0.5f)
+		, FadeTime(2.0f)
+	{
+	}
+
+	bool operator==(const FReverbSettings& Other) const;
+	bool operator!=(const FReverbSettings& Other) const { return !(*this == Other); }
+
+#if WITH_EDITORONLY_DATA
+	void PostSerialize(const FArchive& Ar);
+#endif // WITH_EDITORONLY_DATA
+};
+
+#if WITH_EDITORONLY_DATA
+template<>
+struct TStructOpsTypeTraits<FReverbSettings> : public TStructOpsTypeTraitsBase2<FReverbSettings>
+{
+	enum 
+	{
+		WithPostSerialize = true,
+	};
+};
+#endif // WITH_EDITORONLY_DATA
+
diff --git a/Engine/Source/Runtime/Engine/Classes/Sound/SoundEffectSubmix.h b/Engine/Source/Runtime/Engine/Classes/Sound/SoundEffectSubmix.h
--- a/Engine/Source/Runtime/Engine/Classes/Sound/SoundEffectSubmix.h
+++ b/Engine/Source/Runtime/Engine/Classes/Sound/SoundEffectSubmix.h
@@ -2,13 +2,15 @@
 #pragma once
 
 #include "CoreMinimal.h"
-#include "UObject/ObjectMacros.h"
+
 #include "Sound/SoundEffectPreset.h"
 #include "Sound/SoundEffectBase.h"
+#include "UObject/ObjectMacros.h"
+
 #include "SoundEffectSubmix.generated.h"
 
 class FSoundEffectSubmix;
-
+struct FAudioEffectParameters;
 
 /** Preset of a submix effect that can be shared between sounds. */
 UCLASS(config = Engine, hidecategories = Object, abstract, editinlinenew, BlueprintType)
@@ -73,14 +75,41 @@ struct FSoundEffectSubmixOutputData
 class ENGINE_API FSoundEffectSubmix : public FSoundEffectBase
 {
 public:
-	FSoundEffectSubmix() {}
-	virtual ~FSoundEffectSubmix() {}
+	FSoundEffectSubmix()
+	{
+	}
+
+	virtual ~FSoundEffectSubmix()
+	{
+	}
 
 	/** Called on an audio effect at initialization on main thread before audio processing begins. */
-	virtual void Init(const FSoundEffectSubmixInitData& InSampleRate) {};
+	virtual void Init(const FSoundEffectSubmixInitData& InSampleRate)
+	{
+	}
+
+	/**  Provided for interpolating parameters from audio volume system, enabling transition between various settings */
+	virtual bool SetParameters(const FAudioEffectParameters& InParameters)
+	{
+		return false;
+	}
+
+	// Whether or not effect supports default reverb system
+	virtual bool SupportsDefaultReverb() const
+	{
+		return false;
+	}
+
+	// Whether or not effect supports default EQ system
+	virtual bool SupportsDefaultEQ() const
+	{
+		return false;
+	}
 
 	/** Called on game thread to allow submix effect to query game data if needed. */
-	virtual void Tick() {}
+	virtual void Tick()
+	{
+	}
 
 	/** Override to down mix input audio to a desired channel count. */
 	virtual uint32 GetDesiredInputChannelCountOverride() const
@@ -89,7 +118,9 @@ public:
 	}
 
 	/** Process the input block of audio. Called on audio thread. */
-	virtual void OnProcessAudio(const FSoundEffectSubmixInputData& InData, FSoundEffectSubmixOutputData& OutData) {};
+	virtual void OnProcessAudio(const FSoundEffectSubmixInputData& InData, FSoundEffectSubmixOutputData& OutData)
+	{
+	}
 
 	/** Allow effects to supply a drylevel. */
 	virtual float GetDryLevel() const { return 0.0f; }
diff --git a/Engine/Source/Runtime/Engine/Classes/Sound/SoundMix.h b/Engine/Source/Runtime/Engine/Classes/Sound/SoundMix.h
--- a/Engine/Source/Runtime/Engine/Classes/Sound/SoundMix.h
+++ b/Engine/Source/Runtime/Engine/Classes/Sound/SoundMix.h
@@ -3,16 +3,44 @@
 #pragma once
 
 #include "CoreMinimal.h"
-#include "UObject/ObjectMacros.h"
-#include "UObject/Object.h"
+
 #include "AudioDefines.h"
+#include "UObject/Object.h"
+#include "UObject/ObjectMacros.h"
+
 #include "SoundMix.generated.h"
 
 class USoundClass;
 struct FPropertyChangedEvent;
 
 USTRUCT()
-struct FAudioEQEffect
+struct ENGINE_API FAudioEffectParameters
+{
+	GENERATED_USTRUCT_BODY()
+
+public:
+	FAudioEffectParameters()
+	{
+	}
+
+	virtual ~FAudioEffectParameters()
+	{
+	}
+
+	// Interpolates between one set of parameters and another and stores result in local copy
+	virtual bool Interpolate(const FAudioEffectParameters& InStart, const FAudioEffectParameters& InEnd)
+	{
+		return false;
+	}
+
+	// Prints effect parameters
+	virtual void PrintSettings() const
+	{
+	}
+};
+
+USTRUCT()
+struct FAudioEQEffect : public FAudioEffectParameters
 {
 	GENERATED_USTRUCT_BODY()
 
@@ -84,15 +112,16 @@ struct FAudioEQEffect
 	{}
 
 	/** 
-	* Interpolate EQ settings based on time
-	*/
-	void Interpolate( float InterpValue, const FAudioEQEffect& Start, const FAudioEQEffect& End );
+	 * Interpolates between Start and End EQ effect settings, storing results locally and returning if interpolation is complete
+	 */
+	bool Interpolate(const FAudioEffectParameters& InStart, const FAudioEffectParameters& InEnd) override;
 		
 	/** 
 	* Clamp all settings in range
 	*/
 	void ClampValues();
 
+	virtual void PrintSettings() const override;
 };
 
 /**
diff --git a/Engine/Source/Runtime/Engine/Private/AudioEffect.cpp b/Engine/Source/Runtime/Engine/Private/AudioEffect.cpp
--- a/Engine/Source/Runtime/Engine/Private/AudioEffect.cpp
+++ b/Engine/Source/Runtime/Engine/Private/AudioEffect.cpp
@@ -12,43 +12,47 @@
 #include "AudioDeviceManager.h"
 #include "AudioDevice.h"
 
+#define ENABLE_REVERB_SETTINGS_PRINTING 0
+
 /** 
  * Default settings for a null reverb effect
  */
-FAudioReverbEffect::FAudioReverbEffect( void )
+FAudioReverbEffect::FAudioReverbEffect()
 {
 	Time = 0.0;
 	Volume = 0.0f;
 
-	Density = 1.0f;					
-	Diffusion = 1.0f;				
-	Gain = 0.32f;					
-	GainHF = 0.89f;					
-	DecayTime = 1.49f;				
-	DecayHFRatio = 0.83f;			
-	ReflectionsGain = 0.05f;		
-	ReflectionsDelay = 0.007f;		
-	LateGain = 1.26f;				
-	LateDelay = 0.011f;				
-	AirAbsorptionGainHF = 0.994f;	
-	RoomRolloffFactor = 0.0f;		
+	Density = 1.0f;
+	Diffusion = 1.0f;
+	Gain = 0.32f;
+	GainHF = 0.89f;
+	DecayTime = 1.49f;
+	DecayHFRatio = 0.83f;
+	ReflectionsGain = 0.05f;
+	ReflectionsDelay = 0.007f;
+	LateGain = 1.26f;
+	LateDelay = 0.011f;
+	AirAbsorptionGainHF = 0.994f;
+	RoomRolloffFactor = 0.0f;
 }
 
 /** 
  * Construct generic reverb settings based in the I3DL2 standards
  */
-FAudioReverbEffect::FAudioReverbEffect( float InRoom, 
-								 		float InRoomHF, 
-								 		float InRoomRolloffFactor,	
-								 		float InDecayTime,			
-								 		float InDecayHFRatio,		
-								 		float InReflections,		
-								 		float InReflectionsDelay,	
-								 		float InReverb,				
-								 		float InReverbDelay,		
-								 		float InDiffusion,			
-								 		float InDensity,			
-								 		float InAirAbsorption )
+FAudioReverbEffect::FAudioReverbEffect(
+	float InRoom,
+	float InRoomHF,
+	float InRoomRolloffFactor,
+	float InDecayTime,
+	float InDecayHFRatio,
+	float InReflections,
+	float InReflectionsDelay,
+	float InReverb,
+	float InReverbDelay,
+	float InDiffusion,
+	float InDensity,
+	float InAirAbsorption
+)
 {
 	Time = 0.0;
 	Volume = 0.0f;
@@ -88,26 +92,90 @@ FAudioReverbEffect& FAudioReverbEffect::operator=(class UReverbEffect* InReverbE
 	return *this;
 }
 
+void FAudioReverbEffect::PrintSettings() const
+{
+#if ENABLE_REVERB_SETTINGS_PRINTING
+	const char* FmtText =
+		"\nVolume: %.4f\n"
+		"Density: %.4f\n"
+		"Diffusion: %.4f\n"
+		"Gain: %.4f\n"
+		"GainHF: %.4f\n"
+		"DecayTime: %.4f\n"
+		"DecayHFRatio: %.4f\n"
+		"ReflectionsGain: %.4f\n"
+		"ReflectionsDelay: %.4f\n"
+		"LateGain: %.4f\n"
+		"LateDelay: %.4f\n"
+		"AirAbsorptionGainHF: %.4f\n"
+		"RoomRolloffFactor: %.4f\n";
+
+	FString FmtString(FmtText);
+
+	FString Params = FString::Printf(
+		*FmtString,
+		Settings.Volume,
+		Settings.Density,
+		Settings.Diffusion,
+		Settings.Gain,
+		Settings.GainHF,
+		Settings.DecayTime,
+		Settings.DecayHFRatio,
+		Settings.ReflectionsGain,
+		Settings.ReflectionsDelay,
+		Settings.LateGain,
+		Settings.LateDelay,
+		Settings.AirAbsorptionGainHF,
+		Settings.RoomRolloffFactor
+	);
+
+	UE_LOG(LogTemp, Log, TEXT("%s"), *Params);
+#endif // ENABLE_REVERB_SETTINGS_PRINTING
+}
+
 /** 
  * Get interpolated reverb parameters
  */
-void FAudioReverbEffect::Interpolate( float InterpValue, const FAudioReverbEffect& Start, const FAudioReverbEffect& End )
+bool FAudioReverbEffect::Interpolate(const FAudioEffectParameters& InStart, const FAudioEffectParameters& InEnd)
 {
+	const FAudioReverbEffect& Start = static_cast<const FAudioReverbEffect&>(InStart);
+	const FAudioReverbEffect& End = static_cast<const FAudioReverbEffect&>(InEnd);
+
+	float InterpValue = 1.0f;
+	if (End.Time - Start.Time > 0.0)
+	{
+		InterpValue = (float)((FApp::GetCurrentTime() - Start.Time) / (End.Time - Start.Time));
+	}
+
+	if (InterpValue >= 1.0f)
+	{
+		*this = End;
+		return true;
+	}
+
+	if (InterpValue <= 0.0f)
+	{
+		*this = Start;
+		return false;
+	}
+
 	float InvInterpValue = 1.0f - InterpValue;
 	Time = FApp::GetCurrentTime();
-	Volume = ( Start.Volume * InvInterpValue ) + ( End.Volume * InterpValue );
-	Density = ( Start.Density * InvInterpValue ) + ( End.Density * InterpValue );				
-	Diffusion = ( Start.Diffusion * InvInterpValue ) + ( End.Diffusion * InterpValue );				
-	Gain = ( Start.Gain * InvInterpValue ) + ( End.Gain * InterpValue );					
-	GainHF = ( Start.GainHF * InvInterpValue ) + ( End.GainHF * InterpValue );					
-	DecayTime = ( Start.DecayTime * InvInterpValue ) + ( End.DecayTime * InterpValue );				
-	DecayHFRatio = ( Start.DecayHFRatio * InvInterpValue ) + ( End.DecayHFRatio * InterpValue );			
-	ReflectionsGain = ( Start.ReflectionsGain * InvInterpValue ) + ( End.ReflectionsGain * InterpValue );		
-	ReflectionsDelay = ( Start.ReflectionsDelay * InvInterpValue ) + ( End.ReflectionsDelay * InterpValue );		
-	LateGain = ( Start.LateGain * InvInterpValue ) + ( End.LateGain * InterpValue );				
-	LateDelay = ( Start.LateDelay * InvInterpValue ) + ( End.LateDelay * InterpValue );				
-	AirAbsorptionGainHF = ( Start.AirAbsorptionGainHF * InvInterpValue ) + ( End.AirAbsorptionGainHF * InterpValue );	
-	RoomRolloffFactor = ( Start.RoomRolloffFactor * InvInterpValue ) + ( End.RoomRolloffFactor * InterpValue );		
+	Volume = (Start.Volume * InvInterpValue) + (End.Volume * InterpValue);
+	Density = (Start.Density * InvInterpValue) + (End.Density * InterpValue);
+	Diffusion = (Start.Diffusion * InvInterpValue) + (End.Diffusion * InterpValue);
+	Gain = (Start.Gain * InvInterpValue) + (End.Gain * InterpValue);
+	GainHF = (Start.GainHF * InvInterpValue) + (End.GainHF * InterpValue);
+	DecayTime = (Start.DecayTime * InvInterpValue) + (End.DecayTime * InterpValue);
+	DecayHFRatio = (Start.DecayHFRatio * InvInterpValue) + (End.DecayHFRatio * InterpValue);
+	ReflectionsGain = (Start.ReflectionsGain * InvInterpValue) + (End.ReflectionsGain * InterpValue);
+	ReflectionsDelay = (Start.ReflectionsDelay * InvInterpValue) + (End.ReflectionsDelay * InterpValue);
+	LateGain = (Start.LateGain * InvInterpValue) + (End.LateGain * InterpValue);
+	LateDelay = (Start.LateDelay * InvInterpValue) + (End.LateDelay * InterpValue);
+	AirAbsorptionGainHF = (Start.AirAbsorptionGainHF * InvInterpValue) + (End.AirAbsorptionGainHF * InterpValue);
+	RoomRolloffFactor = (Start.RoomRolloffFactor * InvInterpValue) + (End.RoomRolloffFactor * InterpValue);
+
+	return false;
 }
 
 /** 
@@ -132,11 +200,73 @@ void FAudioEQEffect::ClampValues( void )
 
 }
 
+#define ENABLE_EQ_SETTINGS_PRINTING 0
+
+void FAudioEQEffect::PrintSettings() const
+{
+#if ENABLE_EQ_SETTINGS_PRINTING
+	const char* FmtText =
+		"\nFrequencyCenter0: %.4f\n"
+		"Gain0: %.4f\n"
+		"Bandwidth0: %.4f\n"
+		"FrequencyCenter1: %.4f\n"
+		"Gain1: %.4f\n"
+		"Bandwidth1: %.4f\n"
+		"FrequencyCenter2: %.4f\n"
+		"Gain2: %.4f\n"
+		"Bandwidth2: %.4f\n"
+		"FrequencyCenter3: %.4f\n"
+		"Gain3: %.4f\n"
+		"Bandwidth3: %.4f\n";
+
+	FString FmtString(FmtText);
+
+	FString Params = FString::Printf(
+		*FmtString,
+		Settings.FrequencyCenter0,
+		Settings.Gain0,
+		Settings.Bandwidth0,
+		Settings.FrequencyCenter1,
+		Settings.Gain1,
+		Settings.Bandwidth1,
+		Settings.FrequencyCenter2,
+		Settings.Gain2,
+		Settings.Bandwidth2,
+		Settings.FrequencyCenter3,
+		Settings.Gain3,
+		Settings.Bandwidth3
+	);
+
+	UE_LOG(LogTemp, Log, TEXT("%s"), *Params);
+#endif
+}
+
 /** 
  * Interpolate EQ settings based on time
  */
-void FAudioEQEffect::Interpolate(float InterpValue, const FAudioEQEffect& Start, const FAudioEQEffect& End )
+bool FAudioEQEffect::Interpolate(const FAudioEffectParameters& InStart, const FAudioEffectParameters& InEnd)
 {
+	const FAudioEQEffect& Start = static_cast<const FAudioEQEffect&>(InStart);
+	const FAudioEQEffect& End = static_cast<const FAudioEQEffect&>(InEnd);
+
+	float InterpValue = 1.0f;
+	if (End.RootTime - Start.RootTime > 0.0)
+	{
+		InterpValue = (float)((FApp::GetCurrentTime() - Start.RootTime) / (End.RootTime - Start.RootTime));
+	}
+
+	if (InterpValue >= 1.0f)
+	{
+		*this = End;
+		return true;
+	}
+
+	if (InterpValue <= 0.0f)
+	{
+		*this = Start;
+		return false;
+	}
+
 	RootTime = FApp::GetCurrentTime();
 
 	FrequencyCenter0 = FMath::Lerp(Start.FrequencyCenter0, End.FrequencyCenter0, InterpValue);
@@ -153,6 +283,8 @@ void FAudioEQEffect::Interpolate(float InterpValue, const FAudioEQEffect& Start,
 	Bandwidth1 = FMath::Lerp(Start.Bandwidth1, End.Bandwidth1, InterpValue);
 	Bandwidth2 = FMath::Lerp(Start.Bandwidth2, End.Bandwidth2, InterpValue);
 	Bandwidth3 = FMath::Lerp(Start.Bandwidth3, End.Bandwidth3, InterpValue);
+
+	return false;
 }
 
 /**
@@ -186,60 +318,6 @@ int64 FAudioEffectsManager::VolumeToMilliBels( float Volume, int32 MaxMilliBels
 	return( MilliBels );
 }
 
-/** 
- * Gets the parameters for reverb based on settings and time
- */
-bool FAudioEffectsManager::Interpolate( FAudioReverbEffect& Current, const FAudioReverbEffect& Start, const FAudioReverbEffect& End )
-{
-	float InterpValue = 1.0f;
-	if( End.Time - Start.Time > 0.0 )
-	{
-		InterpValue = ( float )( ( FApp::GetCurrentTime() - Start.Time ) / ( End.Time - Start.Time ) );
-	}
-
-	if( InterpValue >= 1.0f )
-	{
-		Current = End;
-		return true;
-	}
-
-	if( InterpValue <= 0.0f )
-	{
-		Current = Start;
-		return false;
-	}
-
-	Current.Interpolate( InterpValue, Start, End );
-	return false;
-}
-
-/** 
- * Gets the parameters for EQ based on settings and time
- */
-bool FAudioEffectsManager::Interpolate( FAudioEQEffect& Current, const FAudioEQEffect& Start, const FAudioEQEffect& End )
-{
-	float InterpValue = 1.0f;
-	if( End.RootTime - Start.RootTime > 0.0 )
-	{
-		InterpValue = ( float )( ( FApp::GetCurrentTime() - Start.RootTime ) / ( End.RootTime - Start.RootTime ) );
-	}
-
-	if( InterpValue >= 1.0f )
-	{
-		Current = End;
-		return true;
-	}
-
-	if( InterpValue <= 0.0f )
-	{
-		Current = Start;
-		return false;
-	}
-
-	Current.Interpolate( InterpValue, Start, End );
-	return false;
-}
-
 /** 
  * Clear out any reverb and EQ settings
  */
@@ -447,7 +525,7 @@ void FAudioEffectsManager::Update()
 
 #endif
 
-	const bool bIsReverbDone = Interpolate(CurrentReverbEffect, SourceReverbEffect, DestinationReverbEffect);
+	const bool bIsReverbDone = CurrentReverbEffect.Interpolate(SourceReverbEffect, DestinationReverbEffect);
 	if (!bIsReverbDone || bReverbActive || bReverbChanged)
 	{
 		bReverbChanged = false;
@@ -456,7 +534,7 @@ void FAudioEffectsManager::Update()
 		SetReverbEffectParameters(CurrentReverbEffect);
 	}
 
-	const bool bIsEQDone = Interpolate(CurrentEQEffect, SourceEQEffect, DestinationEQEffect);
+	const bool bIsEQDone = CurrentEQEffect.Interpolate(SourceEQEffect, DestinationEQEffect);
 	if (!bIsEQDone || bEQActive || bEQChanged)
 	{
 		bEQChanged = false;
diff --git a/Engine/Source/Runtime/Engine/Private/AudioVolume.cpp b/Engine/Source/Runtime/Engine/Private/AudioVolume.cpp
--- a/Engine/Source/Runtime/Engine/Private/AudioVolume.cpp
+++ b/Engine/Source/Runtime/Engine/Private/AudioVolume.cpp
@@ -45,127 +45,6 @@ void FInteriorSettings::PostSerialize(const FArchive& Ar)
 }
 #endif
 
-bool FReverbSettings::operator==(const FReverbSettings& Other) const
-{
-	return (bApplyReverb == Other.bApplyReverb
-			&& ReverbEffect == Other.ReverbEffect
-			&& ReverbPluginEffect == Other.ReverbPluginEffect
-			&& Volume == Other.Volume
-			&& FadeTime == Other.FadeTime);
-}
-
-#if WITH_EDITORONLY_DATA
-void FReverbSettings::PostSerialize(const FArchive& Ar)
-{
-	if( Ar.UE4Ver() < VER_UE4_REVERB_EFFECT_ASSET_TYPE )
-	{
-		FString ReverbAssetName;
-		switch(ReverbType_DEPRECATED)
-		{
-			case REVERB_Default:
-				// No replacement asset for default reverb type
-				return;
-
-			case REVERB_Bathroom:
-				ReverbAssetName = TEXT("/Engine/EngineSounds/ReverbSettings/Bathroom.Bathroom");
-				break;
-
-			case REVERB_StoneRoom:
-				ReverbAssetName = TEXT("/Engine/EngineSounds/ReverbSettings/StoneRoom.StoneRoom");
-				break;
-
-			case REVERB_Auditorium:
-				ReverbAssetName = TEXT("/Engine/EngineSounds/ReverbSettings/Auditorium.Auditorium");
-				break;
-
-			case REVERB_ConcertHall:
-				ReverbAssetName = TEXT("/Engine/EngineSounds/ReverbSettings/ConcertHall.ConcertHall");
-				break;
-
-			case REVERB_Cave:
-				ReverbAssetName = TEXT("/Engine/EngineSounds/ReverbSettings/Cave.Cave");
-				break;
-
-			case REVERB_Hallway:
-				ReverbAssetName = TEXT("/Engine/EngineSounds/ReverbSettings/Hallway.Hallway");
-				break;
-
-			case REVERB_StoneCorridor:
-				ReverbAssetName = TEXT("/Engine/EngineSounds/ReverbSettings/StoneCorridor.StoneCorridor");
-				break;
-
-			case REVERB_Alley:
-				ReverbAssetName = TEXT("/Engine/EngineSounds/ReverbSettings/Alley.Alley");
-				break;
-
-			case REVERB_Forest:
-				ReverbAssetName = TEXT("/Engine/EngineSounds/ReverbSettings/Forest.Forest");
-				break;
-
-			case REVERB_City:
-				ReverbAssetName = TEXT("/Engine/EngineSounds/ReverbSettings/City.City");
-				break;
-
-			case REVERB_Mountains:
-				ReverbAssetName = TEXT("/Engine/EngineSounds/ReverbSettings/Mountains.Mountains");
-				break;
-
-			case REVERB_Quarry:
-				ReverbAssetName = TEXT("/Engine/EngineSounds/ReverbSettings/Quarry.Quarry");
-				break;
-
-			case REVERB_Plain:
-				ReverbAssetName = TEXT("/Engine/EngineSounds/ReverbSettings/Plain.Plain");
-				break;
-
-			case REVERB_ParkingLot:
-				ReverbAssetName = TEXT("/Engine/EngineSounds/ReverbSettings/ParkingLot.ParkingLot");
-				break;
-
-			case REVERB_SewerPipe:
-				ReverbAssetName = TEXT("/Engine/EngineSounds/ReverbSettings/SewerPipe.SewerPipe");
-				break;
-
-			case REVERB_Underwater:
-				ReverbAssetName = TEXT("/Engine/EngineSounds/ReverbSettings/Underwater.Underwater");
-				break;
-
-			case REVERB_SmallRoom:
-				ReverbAssetName = TEXT("/Engine/EngineSounds/ReverbSettings/SmallRoom.SmallRoom");
-				break;
-
-			case REVERB_MediumRoom:
-				ReverbAssetName = TEXT("/Engine/EngineSounds/ReverbSettings/MediumRoom.MediumRoom");
-				break;
-
-			case REVERB_LargeRoom:
-				ReverbAssetName = TEXT("/Engine/EngineSounds/ReverbSettings/LargeRoom.LargeRoom");
-				break;
-
-			case REVERB_MediumHall:
-				ReverbAssetName = TEXT("/Engine/EngineSounds/ReverbSettings/MediumHall.MediumHall");
-				break;
-
-			case REVERB_LargeHall:
-				ReverbAssetName = TEXT("/Engine/EngineSounds/ReverbSettings/LargeHall.LargeHall");
-				break;
-
-			case REVERB_Plate:
-				ReverbAssetName = TEXT("/Engine/EngineSounds/ReverbSettings/Plate.Plate");
-				break;
-
-			default:
-				// This should cover every type of reverb preset
-				checkNoEntry();
-				break;
-		}
-
-		ReverbEffect = LoadObject<UReverbEffect>(NULL, *ReverbAssetName);
-		check( ReverbEffect );
-	}
-}
-#endif
-
 bool FInteriorSettings::operator==(const FInteriorSettings& Other) const
 {
 	return (Other.bIsWorldSettings == bIsWorldSettings)
diff --git a/Engine/Source/Runtime/Engine/Private/ReverbSettings.cpp b/Engine/Source/Runtime/Engine/Private/ReverbSettings.cpp
new file mode 100644
--- /dev/null
+++ b/Engine/Source/Runtime/Engine/Private/ReverbSettings.cpp
@@ -0,0 +1,128 @@
+// Copyright Epic Games, Inc. All Rights Reserved.
+
+#include "Sound/ReverbSettings.h"
+
+#include "Sound/ReverbEffect.h"
+#include "Sound/SoundEffectSubmix.h"
+
+
+bool FReverbSettings::operator==(const FReverbSettings& Other) const
+{
+	return (bApplyReverb == Other.bApplyReverb
+		&& ReverbEffect == Other.ReverbEffect
+		&& ReverbPluginEffect == Other.ReverbPluginEffect
+		&& Volume == Other.Volume
+		&& FadeTime == Other.FadeTime);
+}
+
+#if WITH_EDITORONLY_DATA
+void FReverbSettings::PostSerialize(const FArchive& Ar)
+{
+	if (Ar.UE4Ver() < VER_UE4_REVERB_EFFECT_ASSET_TYPE)
+	{
+		FString ReverbAssetName;
+		switch (ReverbType_DEPRECATED)
+		{
+		case REVERB_Default:
+			// No replacement asset for default reverb type
+			return;
+
+		case REVERB_Bathroom:
+			ReverbAssetName = TEXT("/Engine/EngineSounds/ReverbSettings/Bathroom.Bathroom");
+			break;
+
+		case REVERB_StoneRoom:
+			ReverbAssetName = TEXT("/Engine/EngineSounds/ReverbSettings/StoneRoom.StoneRoom");
+			break;
+
+		case REVERB_Auditorium:
+			ReverbAssetName = TEXT("/Engine/EngineSounds/ReverbSettings/Auditorium.Auditorium");
+			break;
+
+		case REVERB_ConcertHall:
+			ReverbAssetName = TEXT("/Engine/EngineSounds/ReverbSettings/ConcertHall.ConcertHall");
+			break;
+
+		case REVERB_Cave:
+			ReverbAssetName = TEXT("/Engine/EngineSounds/ReverbSettings/Cave.Cave");
+			break;
+
+		case REVERB_Hallway:
+			ReverbAssetName = TEXT("/Engine/EngineSounds/ReverbSettings/Hallway.Hallway");
+			break;
+
+		case REVERB_StoneCorridor:
+			ReverbAssetName = TEXT("/Engine/EngineSounds/ReverbSettings/StoneCorridor.StoneCorridor");
+			break;
+
+		case REVERB_Alley:
+			ReverbAssetName = TEXT("/Engine/EngineSounds/ReverbSettings/Alley.Alley");
+			break;
+
+		case REVERB_Forest:
+			ReverbAssetName = TEXT("/Engine/EngineSounds/ReverbSettings/Forest.Forest");
+			break;
+
+		case REVERB_City:
+			ReverbAssetName = TEXT("/Engine/EngineSounds/ReverbSettings/City.City");
+			break;
+
+		case REVERB_Mountains:
+			ReverbAssetName = TEXT("/Engine/EngineSounds/ReverbSettings/Mountains.Mountains");
+			break;
+
+		case REVERB_Quarry:
+			ReverbAssetName = TEXT("/Engine/EngineSounds/ReverbSettings/Quarry.Quarry");
+			break;
+
+		case REVERB_Plain:
+			ReverbAssetName = TEXT("/Engine/EngineSounds/ReverbSettings/Plain.Plain");
+			break;
+
+		case REVERB_ParkingLot:
+			ReverbAssetName = TEXT("/Engine/EngineSounds/ReverbSettings/ParkingLot.ParkingLot");
+			break;
+
+		case REVERB_SewerPipe:
+			ReverbAssetName = TEXT("/Engine/EngineSounds/ReverbSettings/SewerPipe.SewerPipe");
+			break;
+
+		case REVERB_Underwater:
+			ReverbAssetName = TEXT("/Engine/EngineSounds/ReverbSettings/Underwater.Underwater");
+			break;
+
+		case REVERB_SmallRoom:
+			ReverbAssetName = TEXT("/Engine/EngineSounds/ReverbSettings/SmallRoom.SmallRoom");
+			break;
+
+		case REVERB_MediumRoom:
+			ReverbAssetName = TEXT("/Engine/EngineSounds/ReverbSettings/MediumRoom.MediumRoom");
+			break;
+
+		case REVERB_LargeRoom:
+			ReverbAssetName = TEXT("/Engine/EngineSounds/ReverbSettings/LargeRoom.LargeRoom");
+			break;
+
+		case REVERB_MediumHall:
+			ReverbAssetName = TEXT("/Engine/EngineSounds/ReverbSettings/MediumHall.MediumHall");
+			break;
+
+		case REVERB_LargeHall:
+			ReverbAssetName = TEXT("/Engine/EngineSounds/ReverbSettings/LargeHall.LargeHall");
+			break;
+
+		case REVERB_Plate:
+			ReverbAssetName = TEXT("/Engine/EngineSounds/ReverbSettings/Plate.Plate");
+			break;
+
+		default:
+			// This should cover every type of reverb preset
+			checkNoEntry();
+			break;
+		}
+
+		ReverbEffect = LoadObject<UReverbEffect>(NULL, *ReverbAssetName);
+		check(ReverbEffect);
+	}
+}
+#endif // WITH_EDITORONLY_DATA
diff --git a/Engine/Source/Runtime/Engine/Private/SoundEffectSubmix.cpp b/Engine/Source/Runtime/Engine/Private/SoundEffectSubmix.cpp
--- a/Engine/Source/Runtime/Engine/Private/SoundEffectSubmix.cpp
+++ b/Engine/Source/Runtime/Engine/Private/SoundEffectSubmix.cpp
@@ -1,7 +1,7 @@
 // Copyright Epic Games, Inc. All Rights Reserved.
 
-
 #include "Sound/SoundEffectSubmix.h"
+#include "Sound/SoundMix.h"
 
 USoundEffectSubmixPreset::USoundEffectSubmixPreset(const FObjectInitializer& ObjectInitializer)
 	: Super(ObjectInitializer)
diff --git a/Engine/Source/Runtime/Engine/Public/AudioEffect.h b/Engine/Source/Runtime/Engine/Public/AudioEffect.h
--- a/Engine/Source/Runtime/Engine/Public/AudioEffect.h
+++ b/Engine/Source/Runtime/Engine/Public/AudioEffect.h
@@ -6,38 +6,53 @@
 
 #pragma once
 
+//#include "AudioDevice.h"
 #include "CoreMinimal.h"
+#include "Sound/AudioVolume.h"
+#include "Sound/ReverbSettings.h"
 #include "Sound/SoundMix.h"
-#include "AudioDevice.h"
+#include "UObject/Object.h"
+#include "UObject/ObjectMacros.h"
+
+#include "AudioEffect.generated.h"
 
 class FAudioDevice;
 class FSoundSource;
 class UReverbEffect;
-struct FReverbSettings;
 
-class ENGINE_API FAudioReverbEffect
+
+USTRUCT()
+struct ENGINE_API FAudioReverbEffect : public FAudioEffectParameters
 {
+	GENERATED_BODY()
+
 public:
 	/** Sets the default values for a reverb effect */
-	FAudioReverbEffect( void );
+	FAudioReverbEffect();
+
 	/** Sets the platform agnostic parameters */
-	FAudioReverbEffect( float InRoom, 
-						float InRoomHF, 
-						float InRoomRolloffFactor, 
-						float InDecayTime, 
-						float InDecayHFRatio, 
-						float InReflections, 
-						float InReflectionsDelay, 
-						float InReverb, 
-						float InReverbDelay, 
-						float InDiffusion, 
-						float InDensity, 
-						float InAirAbsorption );
+	FAudioReverbEffect
+	(
+		float InRoom,
+		float InRoomHF,
+		float InRoomRolloffFactor,
+		float InDecayTime,
+		float InDecayHFRatio,
+		float InReflections,
+		float InReflectionsDelay,
+		float InReverb,
+		float InReverbDelay,
+		float InDiffusion,
+		float InDensity,
+		float InAirAbsorption
+	);
 
 	FAudioReverbEffect& operator=(class UReverbEffect* InReverbEffect);
 
-	/** Interpolates between Start and End reverb effect settings */
-	void Interpolate( float InterpValue, const FAudioReverbEffect& Start, const FAudioReverbEffect& End );
+	/** Interpolates between Start and End reverb effect settings, storing results locally and returning if interpolation is complete */
+	bool Interpolate(const FAudioEffectParameters& InStart, const FAudioEffectParameters& InEnd);
+
+	void PrintSettings() const override;
 
 	/** Time when this reverb was initiated or completed faded in */
 	double		Time;
@@ -63,13 +78,13 @@ public:
 /**
  * Used to store and manipulate parameters related to a radio effect. 
  */
-class FAudioRadioEffect
+class FAudioRadioEffect : public FAudioEffectParameters
 {
 public:
 	/**
 	 * Constructor (default).
 	 */
-	FAudioRadioEffect( void )
+	FAudioRadioEffect(void)
 	{
 	}
 };
@@ -114,23 +129,23 @@ public:
 	/** 
 	 * Calls the platform specific code to set the parameters that define reverb
 	 */
-	virtual void SetReverbEffectParameters( const FAudioReverbEffect& ReverbEffectParameters )
+	virtual void SetReverbEffectParameters(const FAudioEffectParameters& InEffectParameters)
 	{
 	}
 
 	/** 
 	 * Calls the platform specific code to set the parameters that define EQ
 	 */
-	virtual void SetEQEffectParameters( const FAudioEQEffect& EQEffectParameters ) 
+	virtual void SetEQEffectParameters(const FAudioEffectParameters& InEffectParameters)
 	{
 	}
 
 	/** 
 	 * Calls the platform-specific code to set the parameters that define a radio effect.
 	 * 
-	 * @param	RadioEffectParameters	The new parameters for the radio distortion effect. 
+	 * @param	InEffectParameters	The new parameters for the radio distortion effect. 
 	 */
-	virtual void SetRadioEffectParameters( const FAudioRadioEffect& RadioEffectParameters ) 
+	virtual void SetRadioEffectParameters(const FAudioEffectParameters& InEffectParameters)
 	{
 	}
 
@@ -192,21 +207,15 @@ protected:
 	 */
 	void InitAudioEffects( void );
 
-	/** 
-	 * Helper function to interpolate between different effects
-	 */
-	bool Interpolate( FAudioReverbEffect& Current, const FAudioReverbEffect& Start, const FAudioReverbEffect& End );
-	bool Interpolate( FAudioEQEffect& Current, const FAudioEQEffect& Start, const FAudioEQEffect& End );
-
 	FAudioDevice*			AudioDevice;
 	bool					bEffectsInitialised;
 
 	UReverbEffect*			CurrentReverbAsset;
 
-	FAudioReverbEffect		SourceReverbEffect;
-	FAudioReverbEffect		CurrentReverbEffect;
-	FAudioReverbEffect		PrevReverbEffect;
-	FAudioReverbEffect		DestinationReverbEffect;
+	FAudioReverbEffect	SourceReverbEffect;
+	FAudioReverbEffect	CurrentReverbEffect;
+	FAudioReverbEffect	PrevReverbEffect;
+	FAudioReverbEffect	DestinationReverbEffect;
 
 	FReverbSettings 		CurrentReverbSettings;
 
diff --git a/Engine/Source/Runtime/Mac/CoreAudio/Private/CoreAudioEffects.cpp b/Engine/Source/Runtime/Mac/CoreAudio/Private/CoreAudioEffects.cpp
--- a/Engine/Source/Runtime/Mac/CoreAudio/Private/CoreAudioEffects.cpp
+++ b/Engine/Source/Runtime/Mac/CoreAudio/Private/CoreAudioEffects.cpp
@@ -104,9 +104,11 @@ FCoreAudioEffectsManager::~FCoreAudioEffectsManager()
 /** 
  * Calls the platform specific code to set the parameters that define reverb
  */
-void FCoreAudioEffectsManager::SetReverbEffectParameters( const FAudioReverbEffect& ReverbEffectParameters )
+void FCoreAudioEffectsManager::SetReverbEffectParameters(const FAudioEffectParameters& InEffectParameters)
 {
 #if CORE_AUDIO_REVERB_ENABLED
+	const FAudioReverbEffect& ReverbEffectParameters = static_cast<FAudioReverbEffect&>(InEffectParameters);
+
 	float DryWetMix = FMath::Sin(ReverbEffectParameters.Volume*M_PI_2) * 100.0f;										// 0.0-100.0, 100.0
 	float SmallLargeMix = ReverbEffectParameters.GainHF * 100.0f;	// 0.0-100.0, 50.0
 	float PreDelay = ReverbEffectParameters.ReflectionsDelay;										// 0.001->0.03, 0.025
@@ -159,8 +161,10 @@ void FCoreAudioEffectsManager::SetReverbEffectParameters( const FAudioReverbEffe
 /** 
  * Calls the platform specific code to set the parameters that define EQ
  */
-void FCoreAudioEffectsManager::SetEQEffectParameters( const FAudioEQEffect& Params )
+void FCoreAudioEffectsManager::SetEQEffectParameters(const FAudioEffectParameters& InEffectParameters)
 {
+	const FAudioEQEffect& Params = static_cast<const FAudioEQEffect&>(InEffectParameters);
+
 	float Gain0 = VolumeToDeciBels(Params.Gain0);
 	float Gain1 = VolumeToDeciBels(Params.Gain1);
 	float Gain2 = VolumeToDeciBels(Params.Gain2);
@@ -206,7 +210,7 @@ void FCoreAudioEffectsManager::SetEQEffectParameters( const FAudioEQEffect& Para
  * 
  * @param	RadioEffectParameters	The new parameters for the radio distortion effect. 
  */
-void FCoreAudioEffectsManager::SetRadioEffectParameters( const FAudioRadioEffect& RadioEffectParameters )
+void FCoreAudioEffectsManager::SetRadioEffectParameters(const FAudioEffectParameters& InEffectParameters)
 {
 	enum ERadioEffectParams
 	{
diff --git a/Engine/Source/Runtime/Mac/CoreAudio/Public/CoreAudioEffects.h b/Engine/Source/Runtime/Mac/CoreAudio/Public/CoreAudioEffects.h
--- a/Engine/Source/Runtime/Mac/CoreAudio/Public/CoreAudioEffects.h
+++ b/Engine/Source/Runtime/Mac/CoreAudio/Public/CoreAudioEffects.h
@@ -26,19 +26,19 @@ public:
 	/** 
 	 * Calls the platform specific code to set the parameters that define reverb
 	 */
-	virtual void SetReverbEffectParameters( const FAudioReverbEffect& ReverbEffectParameters );
+	virtual void SetReverbEffectParameters(const FAudioEffectParameters& InEffectParameters);
 
 	/** 
 	 * Calls the platform specific code to set the parameters that define EQ
 	 */
-	virtual void SetEQEffectParameters( const FAudioEQEffect& EQEffectParameters );
+	virtual void SetEQEffectParameters(const FAudioEffectParameters& InEffectParameters);
 
 	/** 
 	 * Calls the platform specific code to set the parameters that define a radio effect.
 	 * 
 	 * @param	RadioEffectParameters	The new parameters for the radio distortion effect. 
 	 */
-	virtual void SetRadioEffectParameters( const FAudioRadioEffect& RadioEffectParameters );
+	virtual void SetRadioEffectParameters(const FAudioEffectParameters& InEffectParameters);
 
 private:
 
diff --git a/Engine/Source/Runtime/Windows/XAudio2/Private/XAudio2Effects.cpp b/Engine/Source/Runtime/Windows/XAudio2/Private/XAudio2Effects.cpp
--- a/Engine/Source/Runtime/Windows/XAudio2/Private/XAudio2Effects.cpp
+++ b/Engine/Source/Runtime/Windows/XAudio2/Private/XAudio2Effects.cpp
@@ -561,10 +561,12 @@ FXAudio2EffectsManager::~FXAudio2EffectsManager( void )
 /**
  * Applies the generic reverb parameters to the XAudio2 hardware
  */
-void FXAudio2EffectsManager::SetReverbEffectParameters( const FAudioReverbEffect& ReverbEffectParameters )
+void FXAudio2EffectsManager::SetReverbEffectParameters(const FAudioEffectParameters& InEffectParameters)
 {
 	if( ReverbEffectVoice != NULL )
 	{
+		const FAudioReverbEffect& ReverbEffectParameters = static_cast<const FAudioReverbEffect&>(InEffectParameters);
+
 		XAUDIO2FX_REVERB_I3DL2_PARAMETERS ReverbParameters;
 		XAUDIO2FX_REVERB_PARAMETERS NativeParameters;
 
@@ -592,10 +594,12 @@ void FXAudio2EffectsManager::SetReverbEffectParameters( const FAudioReverbEffect
 /**
  * Applies the generic EQ parameters to the XAudio2 hardware
  */
-void FXAudio2EffectsManager::SetEQEffectParameters( const FAudioEQEffect& EQEffectParameters )
+void FXAudio2EffectsManager::SetEQEffectParameters(const FAudioEffectParameters& InEffectParameters)
 {
 	if (EQPremasterVoice != NULL)
 	{
+		const FAudioEQEffect& EQEffectParameters = static_cast<const FAudioEQEffect&>(InEffectParameters);
+
 		FXEQ_PARAMETERS NativeParameters;
 
 		NativeParameters.FrequencyCenter0 = EQEffectParameters.FrequencyCenter0;
@@ -624,8 +628,10 @@ void FXAudio2EffectsManager::SetEQEffectParameters( const FAudioEQEffect& EQEffe
  * 
  * @param	RadioEffectParameters	The new parameters for the radio distortion effect. 
  */
-void FXAudio2EffectsManager::SetRadioEffectParameters( const FAudioRadioEffect& RadioEffectParameters ) 
+void FXAudio2EffectsManager::SetRadioEffectParameters(const FAudioEffectParameters& InEffectParameters)
 {
+	const FAudioRadioEffect& RadioEffectParameters = static_cast<const FAudioRadioEffect&>(InEffectParameters);
+
 	XAudio2Device->ValidateAPICall( TEXT( "SetEffectParameters (Radio)" ), 
 		RadioEffectVoice->SetEffectParameters(0, &RadioEffectParameters, sizeof(RadioEffectParameters)));
 }
diff --git a/Engine/Source/Runtime/Windows/XAudio2/Public/XAudio2Effects.h b/Engine/Source/Runtime/Windows/XAudio2/Public/XAudio2Effects.h
--- a/Engine/Source/Runtime/Windows/XAudio2/Public/XAudio2Effects.h
+++ b/Engine/Source/Runtime/Windows/XAudio2/Public/XAudio2Effects.h
@@ -38,19 +38,19 @@ public:
 	/** 
 	 * Calls the platform specific code to set the parameters that define reverb
 	 */
-	virtual void SetReverbEffectParameters( const FAudioReverbEffect& ReverbEffectParameters );
+	virtual void SetReverbEffectParameters(const FAudioEffectParameters& InEffectParameters);
 
 	/** 
 	 * Calls the platform specific code to set the parameters that define EQ
 	 */
-	virtual void SetEQEffectParameters( const FAudioEQEffect& ReverbEffectParameters );
+	virtual void SetEQEffectParameters(const FAudioEffectParameters& InEffectParameters);
 
 	/** 
 	 * Calls the platform specific code to set the parameters that define a radio effect.
 	 * 
 	 * @param	RadioEffectParameters	The new parameters for the radio distortion effect. 
 	 */
-	virtual void SetRadioEffectParameters( const FAudioRadioEffect& RadioEffectParameters );
+	virtual void SetRadioEffectParameters(const FAudioEffectParameters& InEffectParameters);
 
 private:
 	/** 
