Fix crash recording replay in PIE by reverting invalid change made for UE-77057

diff --git a/Engine/Source/Runtime/Engine/Private/PackageMapClient.cpp b/Engine/Source/Runtime/Engine/Private/PackageMapClient.cpp
--- a/Engine/Source/Runtime/Engine/Private/PackageMapClient.cpp
+++ b/Engine/Source/Runtime/Engine/Private/PackageMapClient.cpp
@@ -1307,7 +1307,7 @@ void UPackageMapClient::ReceiveNetGUIDBunch( FInBunch &InBunch )
 
 TSharedPtr<FNetFieldExportGroup> UPackageMapClient::GetNetFieldExportGroup(const FString& PathName)
 {
-	return GuidCache->NetFieldExportGroupMap.FindRef(UWorld::RemovePIEPrefix(PathName));
+	return GuidCache->NetFieldExportGroupMap.FindRef(PathName);
 }
 
 void UPackageMapClient::AddNetFieldExportGroup( const FString& PathName, TSharedPtr< FNetFieldExportGroup > NewNetFieldExportGroup )
