Chaos: Fix SetCollisionEnabled not updating shape disable flag when collision is enabled/disabled.

diff --git a/Engine/Source/Runtime/Engine/Private/PhysicsEngine/BodyInstance.cpp b/Engine/Source/Runtime/Engine/Private/PhysicsEngine/BodyInstance.cpp
--- a/Engine/Source/Runtime/Engine/Private/PhysicsEngine/BodyInstance.cpp
+++ b/Engine/Source/Runtime/Engine/Private/PhysicsEngine/BodyInstance.cpp
@@ -877,6 +877,11 @@ void FBodyInstance::UpdatePhysicsFilterData()
 					bUpdateMassProperties = true;
 		}
 
+#if WITH_CHAOS
+				// If this shape shouldn't collide in the sim we disable it here until we have more support.
+				Actor->SetShapeCollisionDisable(ShapeIndex, !CollisionEnabledHasPhysics(GetCollisionEnabled()));
+#endif
+
 				// Apply new collision settings to this shape
 				FPhysicsInterface::SetSimulationFilter(InnerShape, FilterData.SimFilter);
 				FPhysicsInterface::SetQueryFilter(InnerShape, bIsTrimesh ? FilterData.QueryComplexFilter : FilterData.QuerySimpleFilter);
