fix issue where au.streamcaching.MaxCachesToDisplay would truncate the amount of chunks we print to the audiomemreport.

diff --git a/Engine/Source/Runtime/Engine/Private/AudioStreamingCache.cpp b/Engine/Source/Runtime/Engine/Private/AudioStreamingCache.cpp
--- a/Engine/Source/Runtime/Engine/Private/AudioStreamingCache.cpp
+++ b/Engine/Source/Runtime/Engine/Private/AudioStreamingCache.cpp
@@ -1459,12 +1459,6 @@ FString FAudioChunkCache::DebugPrint()
 	// More detailed info about individual chunks here:
 	while (CurrentElement != nullptr)
 	{
-		// We use a CVar to clamp the max amount of chunks we display.
-		if (Index > DebugMaxElementsDisplayCVar)
-		{
-			break;
-		}
-
 		int32 NumTotalChunks = -1;
 		int32 NumTimesTouched = -1;
 		double TimeToLoad = -1.0;
