[NO NON-PII COMMIT MESSAGE]

diff --git a/Engine/Source/Runtime/Engine/Private/AudioStreamingCache.cpp b/Engine/Source/Runtime/Engine/Private/AudioStreamingCache.cpp
--- a/Engine/Source/Runtime/Engine/Private/AudioStreamingCache.cpp
+++ b/Engine/Source/Runtime/Engine/Private/AudioStreamingCache.cpp
@@ -43,6 +43,14 @@ FAutoConsoleVariableRef CVarTrimCacheWhenOverBudget(
 	TEXT("when set to a nonzero value, TrimMemory will be called in AddOrTouchChunk to ensure we never go over budget.\n")
 	TEXT("n: Number of elements to display on screen."),
 	ECVF_Default);
+	
+static int32 AlwaysLogCacheMissesCVar = 0;
+FAutoConsoleVariableRef CVarAlwaysLogCacheMisses(
+	TEXT("au.streamcaching.AlwaysLogCacheMisses"),
+	AlwaysLogCacheMissesCVar,
+	TEXT("when set to a nonzero value, all cache misses will be added to the audiomemreport.\n")
+	TEXT("0: Don't log cache misses until au.streamcaching.StartProfiling is called. 1: Always log cache misses."),
+	ECVF_Default);
 
 static int32 ReadRequestPriorityCVar = 2;
 FAutoConsoleVariableRef CVarReadRequestPriority(
@@ -567,7 +575,7 @@ TArrayView<uint8> FAudioChunkCache::GetChunk(const FChunkKey& InKey, bool bBlock
 			FoundElement->NumConsumers.Increment();
 			return TArrayView<uint8>(FoundElement->ChunkData.GetData(), FoundElement->ChunkDataSize);
 		}
-		else if (bLogCacheMisses)
+		else if (bLogCacheMisses || AlwaysLogCacheMissesCVar)
 		{
 			// Chunks missing. Log this as a miss.
 			const uint32 TotalNumChunksInWave = InKey.SoundWave->GetNumChunks();
