Chaos: fix build

diff --git a/Engine/Source/Runtime/Engine/Private/PhysicsEngine/Experimental/PhysScene_Chaos.cpp b/Engine/Source/Runtime/Engine/Private/PhysicsEngine/Experimental/PhysScene_Chaos.cpp
--- a/Engine/Source/Runtime/Engine/Private/PhysicsEngine/Experimental/PhysScene_Chaos.cpp
+++ b/Engine/Source/Runtime/Engine/Private/PhysicsEngine/Experimental/PhysScene_Chaos.cpp
@@ -1737,7 +1737,7 @@ void FPhysScene_ChaosInterface::CompleteSceneSimulation(ENamedThreads::Type Curr
 			{
 				// We need to process the Command Q here to prevent corruption in BufferPhysicsResults & FlipBuffers
 				// needs further investigation.
-				TQueue<TFunction<void(FPhysicsSolver*)>, EQueueMode::Mpsc>& Queue = Solver->CommandQueue;
+				TQueue<TFunction<void(Chaos::FPhysicsSolver*)>, EQueueMode::Mpsc>& Queue = Solver->GetCommandQueue();
 				TFunction<void(FPhysicsSolver*)> Command;
 				while (Queue.Dequeue(Command))
 				{
