1) Fix for legit chunk loads being reported as cache misses
2) Fix for subsequent calls to AudioMemReport clearing the cache miss log. If you want subsequent calls to AudioMemReport to only log cache misses since the last call to AudioMemReport, set au.streamcaching.KeepCacheMissBufferOnFlush to 0.
3) Fixed Top Cache Misses going from least misses to most, rather than the other way around.

diff --git a/Engine/Source/Runtime/Engine/Classes/Sound/SoundWave.h b/Engine/Source/Runtime/Engine/Classes/Sound/SoundWave.h
--- a/Engine/Source/Runtime/Engine/Classes/Sound/SoundWave.h
+++ b/Engine/Source/Runtime/Engine/Classes/Sound/SoundWave.h
@@ -896,7 +896,7 @@ public:
 	/**
 	 * This returns the initial chunk of compressed data for streaming data sources.
 	 */
-	TArrayView<const uint8> GetZerothChunk();
+	TArrayView<const uint8> GetZerothChunk(bool bForImmediatePlayback = false);
 
 	/**
 	 * Attempts to update the cached platform data after any changes that might affect it
diff --git a/Engine/Source/Runtime/Engine/Private/ADPCMAudioInfo.cpp b/Engine/Source/Runtime/Engine/Private/ADPCMAudioInfo.cpp
--- a/Engine/Source/Runtime/Engine/Private/ADPCMAudioInfo.cpp
+++ b/Engine/Source/Runtime/Engine/Private/ADPCMAudioInfo.cpp
@@ -822,13 +822,13 @@ const uint8* FADPCMAudioInfo::GetLoadedChunk(USoundWave* InSoundWave, uint32 Chu
 	}
 	else if (ChunkIndex == 0)
 	{
-		TArrayView<const uint8> ZerothChunk = InSoundWave->GetZerothChunk();
+		TArrayView<const uint8> ZerothChunk = InSoundWave->GetZerothChunk(true);
 		OutChunkSize = ZerothChunk.Num();
 		return ZerothChunk.GetData();
 	}
 	else
 	{
-		CurCompressedChunkHandle = IStreamingManager::Get().GetAudioStreamingManager().GetLoadedChunk(InSoundWave, ChunkIndex);
+		CurCompressedChunkHandle = IStreamingManager::Get().GetAudioStreamingManager().GetLoadedChunk(InSoundWave, ChunkIndex, false, true);
 		OutChunkSize = CurCompressedChunkHandle.Num();
 		return CurCompressedChunkHandle.GetData();
 	}
diff --git a/Engine/Source/Runtime/Engine/Private/AudioDecompress.cpp b/Engine/Source/Runtime/Engine/Private/AudioDecompress.cpp
--- a/Engine/Source/Runtime/Engine/Private/AudioDecompress.cpp
+++ b/Engine/Source/Runtime/Engine/Private/AudioDecompress.cpp
@@ -371,13 +371,13 @@ const uint8* IStreamedCompressedInfo::GetLoadedChunk(USoundWave* InSoundWave, ui
 	}
 	else if (ChunkIndex == 0)
 	{
-		TArrayView<const uint8> ZerothChunk = InSoundWave->GetZerothChunk();
+		TArrayView<const uint8> ZerothChunk = InSoundWave->GetZerothChunk(true);
 		OutChunkSize = ZerothChunk.Num();
 		return ZerothChunk.GetData();
 	}
 	else
 	{
-		CurCompressedChunkHandle = IStreamingManager::Get().GetAudioStreamingManager().GetLoadedChunk(InSoundWave, ChunkIndex);
+		CurCompressedChunkHandle = IStreamingManager::Get().GetAudioStreamingManager().GetLoadedChunk(InSoundWave, ChunkIndex, false, true);
 		OutChunkSize = CurCompressedChunkHandle.Num();
 		return CurCompressedChunkHandle.GetData();
 	}
diff --git a/Engine/Source/Runtime/Engine/Private/AudioStreaming.cpp b/Engine/Source/Runtime/Engine/Private/AudioStreaming.cpp
--- a/Engine/Source/Runtime/Engine/Private/AudioStreaming.cpp
+++ b/Engine/Source/Runtime/Engine/Private/AudioStreaming.cpp
@@ -890,13 +890,13 @@ bool FLegacyAudioStreamingManager::IsManagedStreamingSoundSource(const FSoundSou
 	return StreamingSoundSources.FindByKey(SoundSource) != NULL;
 }
 
-bool FLegacyAudioStreamingManager::RequestChunk(USoundWave* SoundWave, uint32 ChunkIndex, TFunction<void(EAudioChunkLoadResult)> OnLoadCompleted, ENamedThreads::Type ThreadToCallOnLoadCompletedOn)
+bool FLegacyAudioStreamingManager::RequestChunk(USoundWave* SoundWave, uint32 ChunkIndex, TFunction<void(EAudioChunkLoadResult)> OnLoadCompleted, ENamedThreads::Type ThreadToCallOnLoadCompletedOn, bool bForImmediatePlayback)
 {
 	UE_LOG(LogAudio, Warning, TEXT("RequestChunk is only supported in Stream Caching."));
 	return false;
 }
 
-FAudioChunkHandle FLegacyAudioStreamingManager::GetLoadedChunk(const USoundWave* SoundWave, uint32 ChunkIndex, bool bBlockForLoad) const
+FAudioChunkHandle FLegacyAudioStreamingManager::GetLoadedChunk(const USoundWave* SoundWave, uint32 ChunkIndex, bool bBlockForLoad, bool bForImmediatePlayback) const
 {
 	// Check for the spoof of failing to load a stream chunk
 	if (SpoofFailedStreamChunkLoad > 0)
diff --git a/Engine/Source/Runtime/Engine/Private/AudioStreamingCache.cpp b/Engine/Source/Runtime/Engine/Private/AudioStreamingCache.cpp
--- a/Engine/Source/Runtime/Engine/Private/AudioStreamingCache.cpp
+++ b/Engine/Source/Runtime/Engine/Private/AudioStreamingCache.cpp
@@ -28,6 +28,14 @@ FAutoConsoleVariableRef CVarDebugDisplayCaches(
 	TEXT("n: Number of elements to display on screen."),
 	ECVF_Default);
 
+static int32 KeepCacheMissBufferOnFlushCVar = 1;
+FAutoConsoleVariableRef CVarKeepCacheMissBufferOnFlush(
+	TEXT("au.streamcaching.KeepCacheMissBufferOnFlush"),
+	KeepCacheMissBufferOnFlushCVar,
+	TEXT("IF set to 1, this will maintain the buffer of recorded cache misses after calling AudioMemReport. Otherwise, calling audiomemreport will flush all previous recorded cache misses.\n")
+	TEXT("1: All cache misses from the  whole session will show up in audiomemreport. 0: Only cache misses since the previous call to audiomemreport will show up in the current audiomemreport."),
+	ECVF_Default);
+
 static int32 ForceBlockForLoadCVar = 0;
 FAutoConsoleVariableRef CVarForceBlockForLoad(
 	TEXT("au.streamcaching.ForceBlockForLoad"),
@@ -60,6 +68,14 @@ FAutoConsoleVariableRef CVarReadRequestPriority(
 	TEXT("0: High, 1: Normal, 2: Below Normal, 3: Low, 4: Min"),
 	ECVF_Default);
 
+static int32 PlaybackRequestPriorityCVar = 0;
+FAutoConsoleVariableRef CVarPlaybackRequestPriority(
+	TEXT("au.streamcaching.PlaybackRequestPriority"),
+	PlaybackRequestPriorityCVar,
+	TEXT("This cvar sets the default request priority for audio chunks that are about to play back but aren't in the cache.\n")
+	TEXT("0: High, 1: Normal, 2: Below Normal, 3: Low, 4: Min"),
+	ECVF_Default);
+
 static float StreamCacheSizeOverrideMBCVar = 0.0f;
 FAutoConsoleVariableRef CVarStreamCacheSizeOverrideMB(
 	TEXT("au.streamcaching.StreamCacheSizeOverrideMB"),
@@ -273,7 +289,7 @@ bool FCachedAudioStreamingManager::IsManagedStreamingSoundSource(const FSoundSou
 	return true;
 }
 
-FAudioChunkHandle FCachedAudioStreamingManager::GetLoadedChunk(const USoundWave* SoundWave, uint32 ChunkIndex, bool bBlockForLoad) const
+FAudioChunkHandle FCachedAudioStreamingManager::GetLoadedChunk(const USoundWave* SoundWave, uint32 ChunkIndex, bool bBlockForLoad, bool bForImmediatePlayback) const
 {
 	bBlockForLoad |= (ForceBlockForLoadCVar != 0);
 
@@ -303,7 +319,7 @@ FAudioChunkHandle FCachedAudioStreamingManager::GetLoadedChunk(const USoundWave*
 		}
 
 		// The function call below increments the reference count to the internal chunk.
-		TArrayView<uint8> LoadedChunk = Cache->GetChunk(ChunkKey, bBlockForLoad);
+		TArrayView<uint8> LoadedChunk = Cache->GetChunk(ChunkKey, bBlockForLoad, bForImmediatePlayback);
 
 		// Finally, if there's a chunk after this in the sound, request that it is in the cache.
 		const int32 NextChunk = GetNextChunkIndex(SoundWave, ChunkIndex);
@@ -320,7 +336,7 @@ FAudioChunkHandle FCachedAudioStreamingManager::GetLoadedChunk(const USoundWave*
 #endif
 			};
 
-			bool bIsValidChunk = Cache->AddOrTouchChunk(NextChunkKey, [](EAudioChunkLoadResult) {}, ENamedThreads::AnyThread);
+			bool bIsValidChunk = Cache->AddOrTouchChunk(NextChunkKey, [](EAudioChunkLoadResult) {}, ENamedThreads::AnyThread, false);
 			if (!bIsValidChunk)
 			{
 				UE_LOG(LogAudio, Warning, TEXT("Cache overflow!!! couldn't load chunk %d for sound %s!"), ChunkIndex, *SoundWave->GetName());
@@ -432,13 +448,13 @@ void FCachedAudioStreamingManager::RemoveReferenceToChunk(const FAudioChunkHandl
 	Cache->RemoveReferenceToChunk(ChunkKey);
 }
 
-bool FCachedAudioStreamingManager::RequestChunk(USoundWave* SoundWave, uint32 ChunkIndex, TFunction<void(EAudioChunkLoadResult)> OnLoadCompleted, ENamedThreads::Type ThreadToCallOnLoadCompletedOn)
+bool FCachedAudioStreamingManager::RequestChunk(USoundWave* SoundWave, uint32 ChunkIndex, TFunction<void(EAudioChunkLoadResult)> OnLoadCompleted, ENamedThreads::Type ThreadToCallOnLoadCompletedOn, bool bForImmediatePlayback)
 {
 	FAudioChunkCache* Cache = GetCacheForWave(SoundWave);
 	if (Cache)
 	{
 		FAudioChunkCache::FChunkKey ChunkKey = { SoundWave, SoundWave->GetFName(), ChunkIndex };
-		return Cache->AddOrTouchChunk(ChunkKey, OnLoadCompleted, ThreadToCallOnLoadCompletedOn);
+		return Cache->AddOrTouchChunk(ChunkKey, OnLoadCompleted, ThreadToCallOnLoadCompletedOn, bForImmediatePlayback);
 	}
 	else
 	{
@@ -471,7 +487,7 @@ FAudioChunkCache::~FAudioChunkCache()
 	check(NumberOfLoadsInFlight.GetValue() == 0);
 }
 
-bool FAudioChunkCache::AddOrTouchChunk(const FChunkKey& InKey, TFunction<void(EAudioChunkLoadResult)> OnLoadCompleted, ENamedThreads::Type CallbackThread)
+bool FAudioChunkCache::AddOrTouchChunk(const FChunkKey& InKey, TFunction<void(EAudioChunkLoadResult)> OnLoadCompleted, ENamedThreads::Type CallbackThread, bool bNeededForPlayback)
 {
 	// Update cache limit if needed.
 	if (!FMath::IsNearlyZero(StreamCacheSizeOverrideMBCVar) && StreamCacheSizeOverrideMBCVar > 0.0f)
@@ -515,10 +531,10 @@ bool FAudioChunkCache::AddOrTouchChunk(const FChunkKey& InKey, TFunction<void(EA
 		}
 
 #if DEBUG_STREAM_CACHE
-		CacheElement->DebugInfo.bWasCacheMiss = true;
+		CacheElement->DebugInfo.bWasCacheMiss = bNeededForPlayback;
 #endif
 
-		KickOffAsyncLoad(CacheElement, InKey, OnLoadCompleted, CallbackThread);
+		KickOffAsyncLoad(CacheElement, InKey, OnLoadCompleted, CallbackThread, bNeededForPlayback);
 
 		if (TrimCacheWhenOverBudgetCVar != 0 && MemoryCounterBytes > MemoryLimitBytes)
 		{
@@ -529,7 +545,7 @@ bool FAudioChunkCache::AddOrTouchChunk(const FChunkKey& InKey, TFunction<void(EA
 	}
 }
 
-TArrayView<uint8> FAudioChunkCache::GetChunk(const FChunkKey& InKey, bool bBlockForLoadCompletion)
+TArrayView<uint8> FAudioChunkCache::GetChunk(const FChunkKey& InKey, bool bBlockForLoadCompletion, bool bNeededForPlayback)
 {
 	FScopeLock ScopeLock(&CacheMutationCriticalSection);
 
@@ -565,7 +581,7 @@ TArrayView<uint8> FAudioChunkCache::GetChunk(const FChunkKey& InKey, bool bBlock
 			return TArrayView<uint8>();
 		}
 
-		KickOffAsyncLoad(FoundElement, InKey, [](EAudioChunkLoadResult InResult) {}, ENamedThreads::AnyThread);
+		KickOffAsyncLoad(FoundElement, InKey, [](EAudioChunkLoadResult InResult) {}, ENamedThreads::AnyThread, bNeededForPlayback);
 
 		if (bBlockForLoadCompletion)
 		{
@@ -575,9 +591,9 @@ TArrayView<uint8> FAudioChunkCache::GetChunk(const FChunkKey& InKey, bool bBlock
 			FoundElement->NumConsumers.Increment();
 			return TArrayView<uint8>(FoundElement->ChunkData.GetData(), FoundElement->ChunkDataSize);
 		}
-		else if (bLogCacheMisses || AlwaysLogCacheMissesCVar)
+		else if (bNeededForPlayback && (bLogCacheMisses || AlwaysLogCacheMissesCVar))
 		{
-			// Chunks missing. Log this as a miss.
+			// We missed 
 			const uint32 TotalNumChunksInWave = InKey.SoundWave->GetNumChunks();
 
 			FCacheMissInfo CacheMissInfo = { InKey.SoundWaveName, InKey.ChunkIndex, TotalNumChunksInWave, false };
@@ -742,12 +758,15 @@ FString FAudioChunkCache::FlushCacheMissLog()
 	{
 		FORCEINLINE bool operator()(const FMissedChunk& A, const FMissedChunk& B) const
 		{
-			return A.MissCount < B.MissCount;
+			// Sort from highest miss count to lowest.
+			return A.MissCount > B.MissCount;
 		}
 	};
 
 	TMap<FChunkKey, int32> CacheMissCount;
 
+	TQueue<FCacheMissInfo> BackupQueue;
+
 	FCacheMissInfo CacheMissInfo;
 	while (CacheMissQueue.Dequeue(CacheMissInfo))
 	{
@@ -768,6 +787,11 @@ FString FAudioChunkCache::FlushCacheMissLog()
 
 		int32& MissCount = CacheMissCount.FindOrAdd(Chunk);
 		MissCount++;
+
+		if (KeepCacheMissBufferOnFlushCVar)
+		{
+			BackupQueue.Enqueue(CacheMissInfo);
+		}
 	}
 
 	// Sort our cache miss count map:
@@ -798,6 +822,16 @@ FString FAudioChunkCache::FlushCacheMissLog()
 		TopChunkMissesLog.Append(TEXT("\n"));
 	}
 
+	// If we are keeping the full cache miss buffer around, re-enqueue every cache miss we dequeued.
+	// Note: This could be done more gracefully if TQueue had a move constructor.
+	if (KeepCacheMissBufferOnFlushCVar)
+	{
+		while (BackupQueue.Dequeue(CacheMissInfo))
+		{
+			CacheMissQueue.Enqueue(CacheMissInfo);
+		}
+	}
+
 	return TopChunkMissesLog + TEXT("\n") + ConcatenatedCacheMisses;
 }
 
@@ -1010,7 +1044,7 @@ FAudioChunkCache::FCacheElement* FAudioChunkCache::EvictLeastRecentChunk()
 	return CacheElement;
 }
 
-void FAudioChunkCache::KickOffAsyncLoad(FCacheElement* CacheElement, const FChunkKey& InKey, TFunction<void(EAudioChunkLoadResult)> OnLoadCompleted, ENamedThreads::Type CallbackThread)
+void FAudioChunkCache::KickOffAsyncLoad(FCacheElement* CacheElement, const FChunkKey& InKey, TFunction<void(EAudioChunkLoadResult)> OnLoadCompleted, ENamedThreads::Type CallbackThread, bool bNeededForPlayback)
 {
 	LLM_SCOPE(ELLMTag::Audio);
 	check(CacheElement);
@@ -1018,7 +1052,7 @@ void FAudioChunkCache::KickOffAsyncLoad(FCacheElement* CacheElement, const FChun
 	const FStreamedAudioChunk& Chunk = InKey.SoundWave->RunningPlatformData->Chunks[InKey.ChunkIndex];
 	int32 ChunkDataSize = Chunk.AudioDataSize;
 
-	EAsyncIOPriorityAndFlags AsyncIOPriority = GetAsyncPriorityForChunk(InKey);
+	EAsyncIOPriorityAndFlags AsyncIOPriority = GetAsyncPriorityForChunk(InKey, bNeededForPlayback);
 
 	MemoryCounterBytes -= CacheElement->ChunkData.Num();
 	// Reallocate our chunk data This allows us to shrink if possible.
@@ -1122,35 +1156,66 @@ void FAudioChunkCache::KickOffAsyncLoad(FCacheElement* CacheElement, const FChun
 	}
 }
 
-EAsyncIOPriorityAndFlags FAudioChunkCache::GetAsyncPriorityForChunk(const FChunkKey& InKey)
+EAsyncIOPriorityAndFlags FAudioChunkCache::GetAsyncPriorityForChunk(const FChunkKey& InKey, bool bNeededForPlayback)
 {
 
 	// TODO: In the future we can add an enum to USoundWaves to tweak load priority of individual assets.
 
-	switch (ReadRequestPriorityCVar)
+	if (bNeededForPlayback)
 	{
-		case 4:
-		{
-			return AIOP_MIN;
-		}
-		case 3:
-		{
-			return AIOP_Low;
-		}
-		case 2:
+		switch (PlaybackRequestPriorityCVar)
 		{
-			return AIOP_BelowNormal;
-		}
-		case 1:
-		{
-			return AIOP_Normal;
+			case 4:
+			{
+				return AIOP_MIN;
+			}
+			case 3:
+			{
+				return AIOP_Low;
+			}
+			case 2:
+			{
+				return AIOP_BelowNormal;
+			}
+			case 1:
+			{
+				return AIOP_Normal;
+			}
+			case 0:
+			default:
+			{
+				return AIOP_High;
+			}
 		}
-		case 0:
-		default:
+	}
+	else
+	{
+		switch (ReadRequestPriorityCVar)
 		{
-			return AIOP_High;
+			case 4:
+			{
+				return AIOP_MIN;
+			}
+			case 3:
+			{
+				return AIOP_Low;
+			}
+			case 2:
+			{
+				return AIOP_BelowNormal;
+			}
+			case 1:
+			{
+				return AIOP_Normal;
+			}
+			case 0:
+			default:
+			{
+				return AIOP_High;
+			}
 		}
 	}
+	
 }
 
 void FAudioChunkCache::ExecuteOnLoadCompleteCallback(EAudioChunkLoadResult Result, const TFunction<void(EAudioChunkLoadResult)>& OnLoadCompleted, const ENamedThreads::Type& CallbackThread)
@@ -1504,8 +1569,10 @@ FString FAudioChunkCache::DebugPrint()
 			CurrentElement->IsLoadInProgress() ? TEXT("(Loading In Progress)") : TEXT("")
 		);
 
-		
-		OutputString += ElementInfo + TEXT("\n");
+		if (!bWasTrimmed)
+		{
+			OutputString += ElementInfo + TEXT("\n");
+		}
 
 		CurrentElement = CurrentElement->LessRecentElement;
 		Index++;
diff --git a/Engine/Source/Runtime/Engine/Private/SoundWave.cpp b/Engine/Source/Runtime/Engine/Private/SoundWave.cpp
--- a/Engine/Source/Runtime/Engine/Private/SoundWave.cpp
+++ b/Engine/Source/Runtime/Engine/Private/SoundWave.cpp
@@ -2002,7 +2002,7 @@ bool USoundWave::ShouldUseStreamCaching() const
 	return  FPlatformCompressionUtilities::IsCurrentPlatformUsingStreamCaching() && IsStreaming();
 }
 
-TArrayView<const uint8> USoundWave::GetZerothChunk()
+TArrayView<const uint8> USoundWave::GetZerothChunk(bool bForImmediatePlayback)
 {
 	if (ShouldUseStreamCaching())
 	{
@@ -2017,7 +2017,7 @@ TArrayView<const uint8> USoundWave::GetZerothChunk()
 		if (GetNumChunks() > 1)
 		{
 			// Prime first chunk for playback.
-			IStreamingManager::Get().GetAudioStreamingManager().RequestChunk(this, 1, [](EAudioChunkLoadResult InResult) {});
+			IStreamingManager::Get().GetAudioStreamingManager().RequestChunk(this, 1, [](EAudioChunkLoadResult InResult) {}, ENamedThreads::AnyThread, bForImmediatePlayback);
 		}
 
 		return ZerothChunkData.GetView();
diff --git a/Engine/Source/Runtime/Engine/Private/VorbisAudioInfo.cpp b/Engine/Source/Runtime/Engine/Private/VorbisAudioInfo.cpp
--- a/Engine/Source/Runtime/Engine/Private/VorbisAudioInfo.cpp
+++ b/Engine/Source/Runtime/Engine/Private/VorbisAudioInfo.cpp
@@ -586,13 +586,13 @@ const uint8* FVorbisAudioInfo::GetLoadedChunk(USoundWave* InSoundWave, uint32 Ch
 	}
 	else if (ChunkIndex == 0)
 	{
-		TArrayView<const uint8> ZerothChunk = InSoundWave->GetZerothChunk();
+		TArrayView<const uint8> ZerothChunk = InSoundWave->GetZerothChunk(true);
 		OutChunkSize = ZerothChunk.Num();
 		return ZerothChunk.GetData();
 	}
 	else
 	{
-		CurCompressedChunkHandle = IStreamingManager::Get().GetAudioStreamingManager().GetLoadedChunk(InSoundWave, ChunkIndex);
+		CurCompressedChunkHandle = IStreamingManager::Get().GetAudioStreamingManager().GetLoadedChunk(InSoundWave, ChunkIndex, false, true);
 		OutChunkSize = CurCompressedChunkHandle.Num();
 		return CurCompressedChunkHandle.GetData();
 	}
diff --git a/Engine/Source/Runtime/Engine/Public/AudioStreaming.h b/Engine/Source/Runtime/Engine/Public/AudioStreaming.h
--- a/Engine/Source/Runtime/Engine/Public/AudioStreaming.h
+++ b/Engine/Source/Runtime/Engine/Public/AudioStreaming.h
@@ -273,8 +273,8 @@ struct FLegacyAudioStreamingManager : public IAudioStreamingManager
 	virtual void AddStreamingSoundSource(FSoundSource* SoundSource) override;
 	virtual void RemoveStreamingSoundSource(FSoundSource* SoundSource) override;
 	virtual bool IsManagedStreamingSoundSource(const FSoundSource* SoundSource) const override;
-	virtual bool RequestChunk(USoundWave* SoundWave, uint32 ChunkIndex, TFunction<void(EAudioChunkLoadResult)> OnLoadCompleted, ENamedThreads::Type ThreadToCallOnLoadCompletedOn) override;
-	virtual FAudioChunkHandle GetLoadedChunk(const USoundWave* SoundWave, uint32 ChunkIndex, bool bBlockForLoad = false) const override;
+	virtual bool RequestChunk(USoundWave* SoundWave, uint32 ChunkIndex, TFunction<void(EAudioChunkLoadResult)> OnLoadCompleted, ENamedThreads::Type ThreadToCallOnLoadCompletedOn, bool bForImmediatePlayback = false) override;
+	virtual FAudioChunkHandle GetLoadedChunk(const USoundWave* SoundWave, uint32 ChunkIndex, bool bBlockForLoad = false, bool bForImmediatePlayback = false) const override;
 	virtual uint64 TrimMemory(uint64 NumBytesToFree) override;
 	virtual int32 RenderStatAudioStreaming(UWorld* World, FViewport* Viewport, FCanvas* Canvas, int32 X, int32 Y, const FVector* ViewLocation, const FRotator* ViewRotation) override;
 	virtual FString GenerateMemoryReport() override;
diff --git a/Engine/Source/Runtime/Engine/Public/AudioStreamingCache.h b/Engine/Source/Runtime/Engine/Public/AudioStreamingCache.h
--- a/Engine/Source/Runtime/Engine/Public/AudioStreamingCache.h
+++ b/Engine/Source/Runtime/Engine/Public/AudioStreamingCache.h
@@ -52,10 +52,10 @@ public:
 	~FAudioChunkCache();
 
 	// Places chunk in cache, or puts this chunk back at the top of the cache if it's already loaded. Returns false on failure.
-	bool AddOrTouchChunk(const FChunkKey& InKey, TFunction<void(EAudioChunkLoadResult) > OnLoadCompleted, ENamedThreads::Type CallbackThread);
+	bool AddOrTouchChunk(const FChunkKey& InKey, TFunction<void(EAudioChunkLoadResult) > OnLoadCompleted, ENamedThreads::Type CallbackThread, bool bNeededForPlayback);
 
 	// Returns the chunk asked for, or an empty TArrayView if that chunk is not loaded.
-	TArrayView<uint8> GetChunk(const FChunkKey& InKey, bool bBlockForLoadCompletion);
+	TArrayView<uint8> GetChunk(const FChunkKey& InKey, bool bBlockForLoadCompletion, bool bNeededForPlayback);
 
 	// add an additional reference for a chunk.
 	void AddNewReferenceToChunk(const FChunkKey& InKey);
@@ -281,8 +281,8 @@ private:
 	// Returns the least recent chunk and fixes up the linked list accordingly.
 	FCacheElement* EvictLeastRecentChunk();
 
-	void KickOffAsyncLoad(FCacheElement* CacheElement, const FChunkKey& InKey, TFunction<void(EAudioChunkLoadResult)> OnLoadCompleted, ENamedThreads::Type CallbackThread);
-	EAsyncIOPriorityAndFlags GetAsyncPriorityForChunk(const FChunkKey& InKey);
+	void KickOffAsyncLoad(FCacheElement* CacheElement, const FChunkKey& InKey, TFunction<void(EAudioChunkLoadResult)> OnLoadCompleted, ENamedThreads::Type CallbackThread, bool bNeededForPlayback);
+	EAsyncIOPriorityAndFlags GetAsyncPriorityForChunk(const FChunkKey& InKey, bool bNeededForPlayback);
 
 	// Calls OnLoadCompleted on current thread if CallbackThread == ENamedThreads::AnyThread, and dispatchs an async task on a named thread otherwise.
 	static void ExecuteOnLoadCompleteCallback(EAudioChunkLoadResult Result, const TFunction<void(EAudioChunkLoadResult)>& OnLoadCompleted, const ENamedThreads::Type& CallbackThread);
@@ -354,8 +354,8 @@ public:
 	virtual void AddStreamingSoundSource(FSoundSource* SoundSource) override;
 	virtual void RemoveStreamingSoundSource(FSoundSource* SoundSource) override;
 	virtual bool IsManagedStreamingSoundSource(const FSoundSource* SoundSource) const override;
-	virtual bool RequestChunk(USoundWave* SoundWave, uint32 ChunkIndex, TFunction<void(EAudioChunkLoadResult)> OnLoadCompleted, ENamedThreads::Type ThreadToCallOnLoadCompletedOn) override;
-	virtual FAudioChunkHandle GetLoadedChunk(const USoundWave* SoundWave, uint32 ChunkIndex, bool bBlockForLoad = false) const override;
+	virtual bool RequestChunk(USoundWave* SoundWave, uint32 ChunkIndex, TFunction<void(EAudioChunkLoadResult)> OnLoadCompleted, ENamedThreads::Type ThreadToCallOnLoadCompletedOn, bool bForImmediatePlayback = false) override;
+	virtual FAudioChunkHandle GetLoadedChunk(const USoundWave* SoundWave, uint32 ChunkIndex, bool bBlockForLoad = false, bool bForImmediatePlayback = false) const override;
 	virtual uint64 TrimMemory(uint64 NumBytesToFree) override;
 	virtual int32 RenderStatAudioStreaming(UWorld* World, FViewport* Viewport, FCanvas* Canvas, int32 X, int32 Y, const FVector* ViewLocation, const FRotator* ViewRotation) override;
 	virtual FString GenerateMemoryReport() override;
diff --git a/Engine/Source/Runtime/Engine/Public/ContentStreaming.h b/Engine/Source/Runtime/Engine/Public/ContentStreaming.h
--- a/Engine/Source/Runtime/Engine/Public/ContentStreaming.h
+++ b/Engine/Source/Runtime/Engine/Public/ContentStreaming.h
@@ -512,17 +512,20 @@ struct IAudioStreamingManager : public IStreamingManager
 	 * @param ChunkIndex the index of that soundwave we'd like to request a chunk of.
 	 * @param OnLoadCompleted optional callback when the load completes.
 	 * @param ThreadToCallOnLoadCompleteOn. Optional specifier for which thread OnLoadCompleted should be called on.
+	 * @param bForImmediatePlaybac if true, this will optionally reprioritize this chunk's load request.
 	 */
-	virtual bool RequestChunk(USoundWave* SoundWave, uint32 ChunkIndex, TFunction<void(EAudioChunkLoadResult)> OnLoadCompleted = [](EAudioChunkLoadResult) {}, ENamedThreads::Type ThreadToCallOnLoadCompletedOn = ENamedThreads::AnyThread) = 0;
+	virtual bool RequestChunk(USoundWave* SoundWave, uint32 ChunkIndex, TFunction<void(EAudioChunkLoadResult)> OnLoadCompleted = [](EAudioChunkLoadResult) {}, ENamedThreads::Type ThreadToCallOnLoadCompletedOn = ENamedThreads::AnyThread, bool bForImmediatePlayback = false) = 0;
 
 	/**
 	 * Gets a pointer to a chunk of audio data
 	 *
 	 * @param SoundWave		SoundWave we want a chunk from
 	 * @param ChunkIndex	Index of the chunk we want
+	 * @param bBlockForLoad if true, will block this thread until we finish loading this chunk.
+	 * @param bForImmediatePlayback if true, will optionally reprioritize this chunk's load request. See au.streamcaching.PlaybackRequestPriority.
 	 * @return a handle to the loaded chunk. Can return a default constructed FAudioChunkHandle if the chunk is not loaded yet.
 	 */
-	virtual FAudioChunkHandle GetLoadedChunk(const USoundWave* SoundWave, uint32 ChunkIndex,  bool bBlockForLoad = false) const = 0;
+	virtual FAudioChunkHandle GetLoadedChunk(const USoundWave* SoundWave, uint32 ChunkIndex,  bool bBlockForLoad = false, bool bForImmediatePlayback = false) const = 0;
 
 	/**
 	 * This will start evicting elements from the cache until either hit our target of bytes or run out of chunks we can free.
