Expose USimpleConstructionScript::GenerateNewComponentName for use outside of the class.

diff --git a/Engine/Source/Runtime/Engine/Classes/Engine/SimpleConstructionScript.h b/Engine/Source/Runtime/Engine/Classes/Engine/SimpleConstructionScript.h
--- a/Engine/Source/Runtime/Engine/Classes/Engine/SimpleConstructionScript.h
+++ b/Engine/Source/Runtime/Engine/Classes/Engine/SimpleConstructionScript.h
@@ -149,13 +149,13 @@ private:
 	/** Helper function for generating list of currently used names */
 	void GenerateListOfExistingNames(TSet<FName>& CurrentNames) const;
 
-	/** Helper function for generating a name for a new node, DesiredName can be used to optionally request a name, ComponentClass is mandatory */
-	FName GenerateNewComponentName(const UClass* ComponentClass, FName DesiredName = NAME_None) const;
-
 	/** Helper function to perform actions that all node creation functions have in common */
 	USCS_Node* CreateNodeImpl(UActorComponent* NewComponentTemplate, FName ComponentVariableName);
 
 public:
+	/** Helper function for generating a name for a new node, DesiredName can be used to optionally request a name, ComponentClass is mandatory */
+	ENGINE_API FName GenerateNewComponentName(const UClass* ComponentClass, FName DesiredName = NAME_None) const;
+
 	/** Creates a new SCS node using the given class to create the component template */
 	ENGINE_API USCS_Node* CreateNode(class UClass* NewComponentClass, FName NewComponentVariableName = NAME_None);
 
