Add copy-out phase to single threaded PhysScene_Chaos StartFrame

diff --git a/Engine/Source/Runtime/Engine/Private/PhysicsEngine/Experimental/PhysScene_Chaos.cpp b/Engine/Source/Runtime/Engine/Private/PhysicsEngine/Experimental/PhysScene_Chaos.cpp
--- a/Engine/Source/Runtime/Engine/Private/PhysicsEngine/Experimental/PhysScene_Chaos.cpp
+++ b/Engine/Source/Runtime/Engine/Private/PhysicsEngine/Experimental/PhysScene_Chaos.cpp
@@ -1434,6 +1434,14 @@ void FPhysScene_ChaosInterface::StartFrame()
 			// Here we can directly tick the scene. Single threaded mode doesn't buffer any commands
 			// that would require pumping here - everything is done on demand.
 			Scene.Tick(Dt);
+
+			// Copy out solver data
+			if (Chaos::FPhysicsSolver* Solver = GetSolver())
+			{
+				Solver->GetActiveParticlesBuffer()->CaptureSolverData(Solver);
+				Solver->BufferPhysicsResults();
+				Solver->FlipBuffers();
+			}
 		}
 		break;
 		case EChaosThreadingMode::TaskGraph:
