RTHF shadow - added support that allows light cast inside terrains through holes. Due to its cost, it is only enabled in the highest quality level. Currently, boundaries are jaggy for shadows cast through holes. Users need to hide the aritifact with DF shadow by placing static meshes around holes.
Fixed a crash that can happen on exit.
Further reduced self shadow artifact when using RTHF shadow.

diff --git a/Engine/Shaders/Private/DistanceFieldLightingShared.ush b/Engine/Shaders/Private/DistanceFieldLightingShared.ush
--- a/Engine/Shaders/Private/DistanceFieldLightingShared.ush
+++ b/Engine/Shaders/Private/DistanceFieldLightingShared.ush
@@ -40,6 +40,7 @@ float3 DistanceFieldVolumePositionToUV(float3 VolumePosition, float3 UVScale, fl
 	return VolumeUV;
 }
 
+Texture2D HFVisibilityTexture;
 SamplerState DistanceFieldSampler;
 
 #if (DISTANCEFIELD_PRIMITIVE_TYPE == DFPT_SignedDistanceField)
@@ -58,6 +59,11 @@ float SampleHeightFieldAtlas(float2 UV)
 {
 	return 0.0;
 }
+
+float SampleHFVisibilityTexture(float2 UV)
+{
+	return 0.0;
+}
 #else
 Texture2D DistanceFieldTexture;
 
@@ -71,6 +77,11 @@ float SampleHeightFieldAtlas(float2 UV)
 	float4 SampleValue = Texture2DSampleLevel(DistanceFieldTexture, DistanceFieldSampler, UV, 0);
 	return DecodePackedHeight(SampleValue.xy);
 }
+
+float SampleHFVisibilityTexture(float2 UV)
+{
+	return Texture2DSampleLevel(HFVisibilityTexture, DistanceFieldSampler, UV, 0).r;
+}
 #endif
 
 float3 DistanceFieldAtlasTexelSize;
@@ -90,7 +101,7 @@ uint GetCulledNumObjects()
 #define OBJECT_DATA_STRIDE 17
 #elif (DISTANCEFIELD_PRIMITIVE_TYPE == DFPT_HeightField)
 #define OBJECT_BOUNDS_STRIDE 2
-#define OBJECT_DATA_STRIDE 5
+#define OBJECT_DATA_STRIDE 6
 #else
 #error Unknown distance field primitive type
 #endif
@@ -215,7 +226,7 @@ float4x4 LoadGlobalObjectOBBToWorld(uint ObjectIndex)
 #define CULLED_OBJECT_DATA_STRIDE 17
 #define CULLED_OBJECT_BOX_BOUNDS_STRIDE 5
 #elif (DISTANCEFIELD_PRIMITIVE_TYPE == DFPT_HeightField)
-#define CULLED_OBJECT_DATA_STRIDE 5
+#define CULLED_OBJECT_DATA_STRIDE 6
 #define CULLED_OBJECT_BOX_BOUNDS_STRIDE 5
 #else
 #error Unknown distance field primitive type
@@ -309,6 +320,11 @@ float4 LoadHeightFieldObjectAtlasUVScaleBiasFromBuffer(uint ObjectIndex, Structu
 	return InBuffer[ObjectIndex * CULLED_OBJECT_DATA_STRIDE + 4];
 }
 
+float4 LoadHeightFieldObjectVisibilityAtlasUVScaleBiasFromBuffer(uint ObjectIndex, StructuredBuffer<float4> InBuffer)
+{
+	return InBuffer[ObjectIndex * CULLED_OBJECT_DATA_STRIDE + 5];
+}
+
 // These are structured buffers so they can be scalar memory loads on GCN
 StructuredBuffer<float4> CulledObjectBounds;
 StructuredBuffer<float4> CulledObjectData;
@@ -389,6 +405,11 @@ float4 LoadHeightFieldObjectSizeScale(uint ObjectIndex)
 	return LoadHeightFieldObjectSizeScaleFromBuffer(ObjectIndex, CulledObjectData);
 }
 
+float4 LoadHeightFieldObjectVisibilityAtlasUVScaleBias(uint ObjectIndex)
+{
+	return LoadHeightFieldObjectVisibilityAtlasUVScaleBiasFromBuffer(ObjectIndex, CulledObjectData);
+}
+
 // x = Offset in global buffer, y = NumLOD0, z = NumSurfels (all LODs), W = instance index
 uint4 LoadObjectSurfelCoordinate(uint ObjectIndex)
 {
diff --git a/Engine/Shaders/Private/DistanceFieldShadowing.usf b/Engine/Shaders/Private/DistanceFieldShadowing.usf
--- a/Engine/Shaders/Private/DistanceFieldShadowing.usf
+++ b/Engine/Shaders/Private/DistanceFieldShadowing.usf
@@ -12,9 +12,11 @@
 #if (DISTANCEFIELD_PRIMITIVE_TYPE == DFPT_SignedDistanceField)
 #define MAX_TRACE_SPHERE_RADIUS 100
 #define SELF_SHADOW_VERTICAL_BIAS 0
+#define SELF_SHADOW_VIEW_BIAS 0
 #else
 #define MAX_TRACE_SPHERE_RADIUS 500
 #define SELF_SHADOW_VERTICAL_BIAS 100
+#define SELF_SHADOW_VIEW_BIAS 50
 #endif
 
 uint ObjectBoundingGeometryIndexCount;
@@ -628,6 +630,8 @@ void DistanceFieldShadowingCS(
 		float SelfShadowFadeDistance = 100;
 		// Mitigate self-shadow caused by discontinuity on heightfield borders (only used for heightfield)
 		float SelfShadowVerticalBias = SELF_SHADOW_VERTICAL_BIAS;
+		// Mitigate self-shadow on steep terrain surfces
+		float SelfShadowViewBias = SELF_SHADOW_VIEW_BIAS;
 
 #if HAS_PREVIOUS_OUTPUT
 	#if PLATFORM_SUPPORTS_TYPED_UAV_LOAD
@@ -639,6 +643,7 @@ void DistanceFieldShadowingCS(
 #endif
 
 		OpaqueWorldPosition.z += SelfShadowVerticalBias;
+		OpaqueWorldPosition -= SelfShadowViewBias * View.ViewForward;
 
 		// World space offset along the start of the ray to avoid incorrect self-shadowing
 		float RayStartOffset = 2 + RayStartOffsetDepthScale * SceneDepth;
diff --git a/Engine/Shaders/Private/DistanceFieldShadowingShared.ush b/Engine/Shaders/Private/DistanceFieldShadowingShared.ush
--- a/Engine/Shaders/Private/DistanceFieldShadowingShared.ush
+++ b/Engine/Shaders/Private/DistanceFieldShadowingShared.ush
@@ -201,6 +201,7 @@ float RayTraceHeightfieldLocal(
 	float StartRayTime,
 	float EndRayTime,
 	float4 AtlasUVScaleBias,
+	float4 VisUVScaleBias,
 	float TanLightAngle,
 	float LocalSelfShadowScale,
 	float LocalMaxSphereRadius)
@@ -216,6 +217,7 @@ float RayTraceHeightfieldLocal(
 	float StepRayTime = (EndRayTime - StartRayTime) / MaxSteps;
 	float SampleRayTime = StartRayTime;
 	float MinConeVisibility = 1.0;
+	float MinConeVisibility2 = VisUVScaleBias.x > 0.0 ? 1.0 : 0.0;
 
 	for (float StepIndex = 0; StepIndex < MaxSteps; ++StepIndex)
 	{
@@ -224,14 +226,45 @@ float RayTraceHeightfieldLocal(
 		float StepHeight = SampleHeightFieldAtlas(StepUV);
 
 		float SelfShadowVisibility = 1 - saturate(SampleRayTime * LocalSelfShadowScale);
-
 		float SphereRadius = clamp(TanLightAngle * SampleRayTime, 0.0001, LocalMaxSphereRadius);
+
+		SampleRayTime += StepRayTime;
+
 		float Distance = StepPosition.z - StepHeight;
 		float StepVisibility = max(saturate(Distance / SphereRadius), SelfShadowVisibility);
 
 		MinConeVisibility = min(MinConeVisibility, StepVisibility);
-		SampleRayTime += StepRayTime;
 
+#if DF_SHADOW_QUALITY == 3
+		// Compute ray visibility assuming it starts below heightfield. Only needed when the heightfield has hole
+		if (VisUVScaleBias.x > 0.0)
+		{
+			float2 StepVisUV = StepPosition.xy * VisUVScaleBias.xy + VisUVScaleBias.zw;
+			float StepVis = SampleHFVisibilityTexture(StepVisUV);
+			bool bIsHole = StepVis > 0.0;
+			
+			// Assuming no more occlusion within this heightfield if the ray goes through a hole
+			if (bIsHole && Distance > 0.0)
+			{
+				break;
+			}
+
+			// Ignore this sample if it is beneath a hole
+			if (!bIsHole)
+			{
+				float StepVisibility2 = max(saturate(-Distance / SphereRadius), SelfShadowVisibility);
+				MinConeVisibility2 = min(MinConeVisibility2, StepVisibility2);
+			}
+		}
+
+		if (max(MinConeVisibility, MinConeVisibility2) < 0.01)
+		{
+			break;
+		}
+	}
+
+	return max(MinConeVisibility, MinConeVisibility2);
+#else
 		if (MinConeVisibility < 0.01)
 		{
 			break;
@@ -239,6 +272,7 @@ float RayTraceHeightfieldLocal(
 	}
 
 	return MinConeVisibility;
+#endif
 }
 
 float ShadowRayTraceThroughCulledHeightFieldObjects(
@@ -293,7 +327,11 @@ float ShadowRayTraceThroughCulledHeightFieldObjects(
 		if (bValidIntersection)
 		{
 			float4 AtlasUVScaleBias = LoadHeightFieldObjectAtlasUVScaleBias(ObjectIndex);
-
+#if DF_SHADOW_QUALITY == 3
+			float4 VisUVScaleBias = LoadHeightFieldObjectVisibilityAtlasUVScaleBias(ObjectIndex);
+#else
+			float4 VisUVScaleBias = float4(0, 0, 0, 0);
+#endif
 			float3 LocalRayDirection = LocalRayEnd - LocalRayStart;
 			float LocalRayLength = length(LocalRayDirection);
 			float3 LocalRayUnitDirection = LocalRayDirection / LocalRayLength;
@@ -309,6 +347,7 @@ float ShadowRayTraceThroughCulledHeightFieldObjects(
 				StartRayTime,
 				EndRayTime,
 				AtlasUVScaleBias,
+				VisUVScaleBias,
 				TanLightAngle,
 				LocalSelfShadowScale,
 				LocalMaxSphereRadius);
diff --git a/Engine/Shaders/Private/HeightFieldAtlasManagement.usf b/Engine/Shaders/Private/HeightFieldAtlasManagement.usf
--- a/Engine/Shaders/Private/HeightFieldAtlasManagement.usf
+++ b/Engine/Shaders/Private/HeightFieldAtlasManagement.usf
@@ -26,3 +26,18 @@ void UploadHeightFieldToAtlasCS(uint3 DispatchThreadId : SV_DispatchThreadID)
 		RWHeightFieldAtlas[WriteCoord] = Texture2DSampleLevel(SourceTexture, SourceTextureSampler, SourceUV, SourceMipBias).xy;
 	}
 }
+
+float4 VisibilityChannelMask;
+RWTexture2D<float> RWVisibilityAtlas;
+
+[numthreads(THREADGROUP_SIZEX, THREADGROUP_SIZEY, 1)]
+void UploadVisibilityToAtlasCS(uint3 DispatchThreadId : SV_DispatchThreadID)
+{
+	if (all(DispatchThreadId.xy < UpdateRegionOffsetAndSize.zw))
+	{
+		float2 SourceUV = DispatchThreadId.xy * SourceScaleBias.xy + SourceScaleBias.zw;
+		uint2 WriteCoord = DispatchThreadId.xy + UpdateRegionOffsetAndSize.xy;
+
+		RWVisibilityAtlas[WriteCoord] = dot(Texture2DSampleLevel(SourceTexture, SourceTextureSampler, SourceUV, SourceMipBias), VisibilityChannelMask);
+	}
+}
diff --git a/Engine/Source/Runtime/Engine/Private/DistanceFieldAtlas.cpp b/Engine/Source/Runtime/Engine/Private/DistanceFieldAtlas.cpp
--- a/Engine/Source/Runtime/Engine/Private/DistanceFieldAtlas.cpp
+++ b/Engine/Source/Runtime/Engine/Private/DistanceFieldAtlas.cpp
@@ -150,8 +150,30 @@ static FAutoConsoleVariableRef CVarHeightFieldAtlasDownSampleLevel(
 	TEXT("Max number of times a suballocation can be down-sampled"),
 	ECVF_RenderThreadSafe);
 
+static int32 GHFVisibilityAtlasTileSize = 64;
+static FAutoConsoleVariableRef CVarHFVisibilityAtlasTileSize(
+	TEXT("r.HeightFields.VisibilityAtlasTileSize"),
+	GHFVisibilityAtlasTileSize,
+	TEXT("Suballocation granularity"),
+	ECVF_RenderThreadSafe);
+
+static int32 GHFVisibilityAtlasDimInTiles = 8;
+static FAutoConsoleVariableRef CVarHFVisibilityAtlasDimInTiles(
+	TEXT("r.HeightFields.VisibilityAtlasDimInTiles"),
+	GHFVisibilityAtlasDimInTiles,
+	TEXT("Number of tiles the atlas has in one dimension"),
+	ECVF_RenderThreadSafe);
+
+static int32 GHFVisibilityAtlasDownSampleLevel = 2;
+static FAutoConsoleVariableRef CVarHFVisibilityAtlasDownSampleLevel(
+	TEXT("r.HeightFields.VisibilityAtlasDownSampleLevel"),
+	GHFVisibilityAtlasDownSampleLevel,
+	TEXT("Max number of times a suballocation can be down-sampled"),
+	ECVF_RenderThreadSafe);
+
 TGlobalResource<FDistanceFieldVolumeTextureAtlas> GDistanceFieldVolumeTextureAtlas = TGlobalResource<FDistanceFieldVolumeTextureAtlas>();
-TGlobalResource<FHeightFieldTextureAtlas> GHeightFieldTextureAtlas;
+TGlobalResource<FLandscapeTextureAtlas> GHeightFieldTextureAtlas(FLandscapeTextureAtlas::SAT_Height);
+TGlobalResource<FLandscapeTextureAtlas> GHFVisibilityTextureAtlas(FLandscapeTextureAtlas::SAT_Visibility);
 
 FDistanceFieldVolumeTextureAtlas::FDistanceFieldVolumeTextureAtlas() :
 	BlockAllocator(0, 0, 0, 0, 0, 0, false, false),
@@ -1219,11 +1241,18 @@ void FDistanceFieldAsyncQueue::Shutdown()
 	}
 }
 
-void FHeightFieldTextureAtlas::InitializeIfNeeded()
+FLandscapeTextureAtlas::FLandscapeTextureAtlas(ESubAllocType InSubAllocType)
+	: MaxDownSampleLevel(0)
+	, Generation(0)
+	, SubAllocType(InSubAllocType)
+{}
+
+void FLandscapeTextureAtlas::InitializeIfNeeded()
 {
-	const uint32 LocalTileSize = (uint32)GHeightFieldAtlasTileSize;
-	const uint32 LocalDimInTiles = (uint32)GHeightFieldAtlasDimInTiles;
-	const uint32 LocalDownSampleLevel = (uint32)GHeightFieldAtlasDownSampleLevel;
+	const bool bHeight = SubAllocType == SAT_Height;
+	const uint32 LocalTileSize = bHeight ? GHeightFieldAtlasTileSize : GHFVisibilityAtlasTileSize;
+	const uint32 LocalDimInTiles = bHeight ? GHeightFieldAtlasDimInTiles : GHFVisibilityAtlasDimInTiles;
+	const uint32 LocalDownSampleLevel = bHeight ? GHeightFieldAtlasDownSampleLevel : GHFVisibilityAtlasDownSampleLevel;
 
 	if (!AtlasTextureRHI
 		|| AddrSpaceAllocator.TileSize != LocalTileSize
@@ -1250,9 +1279,10 @@ void FHeightFieldTextureAtlas::InitializeIfNeeded()
 		const uint32 SizeX = AddrSpaceAllocator.DimInTexels;
 		const uint32 SizeY = AddrSpaceAllocator.DimInTexels;
 		const uint32 Flags = TexCreate_ShaderResource | TexCreate_UAV;
-		FRHIResourceCreateInfo CreateInfo(TEXT("HeightFieldAtlas"));
+		const EPixelFormat Format = bHeight ? PF_R8G8 : PF_G8;
+		FRHIResourceCreateInfo CreateInfo(bHeight ? TEXT("HeightFieldAtlas") : TEXT("VisibilityAtlas"));
 
-		AtlasTextureRHI = RHICreateTexture2D(SizeX, SizeY, PF_R8G8, 1, 1, Flags, CreateInfo);
+		AtlasTextureRHI = RHICreateTexture2D(SizeX, SizeY, Format, 1, 1, Flags, CreateInfo);
 		AtlasUAVRHI = RHICreateUnorderedAccessView(AtlasTextureRHI, 0);
 
 		MaxDownSampleLevel = LocalDownSampleLevel;
@@ -1260,7 +1290,7 @@ void FHeightFieldTextureAtlas::InitializeIfNeeded()
 	}
 }
 
-void FHeightFieldTextureAtlas::AddAllocation(UTexture2D* Texture)
+void FLandscapeTextureAtlas::AddAllocation(UTexture2D* Texture, uint32 VisibilityChannel)
 {
 	check(Texture);
 
@@ -1285,11 +1315,11 @@ void FHeightFieldTextureAtlas::AddAllocation(UTexture2D* Texture)
 	}
 	else
 	{
-		PendingAllocations.Add(Texture);
+		PendingAllocations.Add(FAllocation(Texture, VisibilityChannel));
 	}
 }
 
-void FHeightFieldTextureAtlas::RemoveAllocation(UTexture2D* Texture)
+void FLandscapeTextureAtlas::RemoveAllocation(UTexture2D* Texture)
 {
 	FSetElementId Id = PendingAllocations.FindId(Texture);
 	if (Id.IsValidId())
@@ -1329,22 +1359,23 @@ void FHeightFieldTextureAtlas::RemoveAllocation(UTexture2D* Texture)
 	}
 }
 
-class FUploadHeightFieldToAtlasCS : public FGlobalShader
+class FUploadLandscapeTextureToAtlasCS : public FGlobalShader
 {
-	DECLARE_GLOBAL_SHADER(FUploadHeightFieldToAtlasCS);
-	SHADER_USE_PARAMETER_STRUCT(FUploadHeightFieldToAtlasCS, FGlobalShader);
-
-	using FPermutationDomain = FShaderPermutationNone;
-
-	BEGIN_SHADER_PARAMETER_STRUCT(FParameters,)
+public:
+	BEGIN_SHADER_PARAMETER_STRUCT(FSharedParameters, )
 		SHADER_PARAMETER(FUintVector4, UpdateRegionOffsetAndSize)
 		SHADER_PARAMETER(FVector4, SourceScaleBias)
 		SHADER_PARAMETER(uint32, SourceMipBias)
 		SHADER_PARAMETER_TEXTURE(Texture2D, SourceTexture)
 		SHADER_PARAMETER_SAMPLER(SamplerState, SourceTextureSampler)
-		SHADER_PARAMETER_UAV(RWTexture2D<float2>, RWHeightFieldAtlas)
 	END_SHADER_PARAMETER_STRUCT()
 
+	FUploadLandscapeTextureToAtlasCS() = default;
+
+	FUploadLandscapeTextureToAtlasCS(const ShaderMetaType::CompiledShaderInitializerType& Initializer)
+		: FGlobalShader(Initializer)
+	{}
+
 	static bool ShouldCompilePermutation(const FGlobalShaderPermutationParameters& Parameters)
 	{
 		return IsFeatureLevelSupported(Parameters.Platform, ERHIFeatureLevel::SM5) && DoesPlatformSupportDistanceFieldShadowing(Parameters.Platform);
@@ -1361,9 +1392,38 @@ class FUploadHeightFieldToAtlasCS : public FGlobalShader
 	static constexpr uint32 ThreadGroupSizeY = 8;
 };
 
+class FUploadHeightFieldToAtlasCS : public FUploadLandscapeTextureToAtlasCS
+{
+	DECLARE_GLOBAL_SHADER(FUploadHeightFieldToAtlasCS);
+	SHADER_USE_PARAMETER_STRUCT(FUploadHeightFieldToAtlasCS, FUploadLandscapeTextureToAtlasCS);
+
+	using FPermutationDomain = FShaderPermutationNone;
+
+	BEGIN_SHADER_PARAMETER_STRUCT(FParameters,)
+		SHADER_PARAMETER_STRUCT_INCLUDE(FUploadLandscapeTextureToAtlasCS::FSharedParameters, SharedParams)
+		SHADER_PARAMETER_UAV(RWTexture2D<float2>, RWHeightFieldAtlas)
+	END_SHADER_PARAMETER_STRUCT()
+};
+
 IMPLEMENT_GLOBAL_SHADER(FUploadHeightFieldToAtlasCS, "/Engine/Private/HeightFieldAtlasManagement.usf", "UploadHeightFieldToAtlasCS", SF_Compute);
 
-uint32 FHeightFieldTextureAtlas::CalculateDownSampleLevel(uint32 SizeX, uint32 SizeY) const
+class FUploadVisibilityToAtlasCS : public FUploadLandscapeTextureToAtlasCS
+{
+	DECLARE_GLOBAL_SHADER(FUploadVisibilityToAtlasCS);
+	SHADER_USE_PARAMETER_STRUCT(FUploadVisibilityToAtlasCS, FUploadLandscapeTextureToAtlasCS);
+
+	using FPermutationDomain = FShaderPermutationNone;
+
+	BEGIN_SHADER_PARAMETER_STRUCT(FParameters, )
+		SHADER_PARAMETER_STRUCT_INCLUDE(FUploadLandscapeTextureToAtlasCS::FSharedParameters, SharedParams)
+		SHADER_PARAMETER(FVector4, VisibilityChannelMask)
+		SHADER_PARAMETER_UAV(RWTexture2D<float>, RWVisibilityAtlas)
+	END_SHADER_PARAMETER_STRUCT()
+};
+
+IMPLEMENT_GLOBAL_SHADER(FUploadVisibilityToAtlasCS, "/Engine/Private/HeightFieldAtlasManagement.usf", "UploadVisibilityToAtlasCS", SF_Compute);
+
+uint32 FLandscapeTextureAtlas::CalculateDownSampleLevel(uint32 SizeX, uint32 SizeY) const
 {
 	const uint32 TileSize = AddrSpaceAllocator.TileSize;
 
@@ -1381,13 +1441,11 @@ uint32 FHeightFieldTextureAtlas::CalculateDownSampleLevel(uint32 SizeX, uint32 S
 	return MaxDownSampleLevel;
 }
 
-void FHeightFieldTextureAtlas::UpdateAllocations(FRHICommandListImmediate& RHICmdList, ERHIFeatureLevel::Type InFeatureLevel)
+void FLandscapeTextureAtlas::UpdateAllocations(FRHICommandListImmediate& RHICmdList, ERHIFeatureLevel::Type InFeatureLevel)
 {
 	InitializeIfNeeded();
 
-	TArray<uint32, TInlineAllocator<8>> UploadHandles;
-	TArray<FIntVector, TInlineAllocator<8>> UploadSizesAndMipBias;
-	TArray<FRHITexture*, TInlineAllocator<8>> UploadSourceTextures;
+	TArray<FPendingUpload, TInlineAllocator<8>> PendingUploads;
 
 	auto AllocSortPred = [](const FAllocation& A, const FAllocation& B)
 	{
@@ -1410,9 +1468,8 @@ void FHeightFieldTextureAtlas::UpdateAllocations(FRHICommandListImmediate& RHICm
 			const uint32 SourceMipBias = DownSampleLevel - NumMissingMips;
 			const FAllocation* Allocation = CurrentAllocations.Find(SourceTexture);
 			check(Allocation && Allocation->Handle != INDEX_NONE);
-			UploadHandles.Add(Allocation->Handle);
-			UploadSizesAndMipBias.Add(FIntVector(SizeX >> DownSampleLevel, SizeY >> DownSampleLevel, SourceMipBias));
-			UploadSourceTextures.Add(SourceTexture->Resource->TextureRHI);
+			const uint32 VisibilityChannel = Allocation->VisibilityChannel;
+			PendingUploads.Emplace(SourceTexture, SizeX >> DownSampleLevel, SizeY >> DownSampleLevel, SourceMipBias, Allocation->Handle, VisibilityChannel);
 			PendingStreamingTextures.RemoveAtSwap(Idx--);
 		}
 	}
@@ -1435,6 +1492,7 @@ void FHeightFieldTextureAtlas::UpdateAllocations(FRHICommandListImmediate& RHICm
 				const uint32 DownSampledSizeX = SizeX >> DownSampleLevel;
 				const uint32 DownSampledSizeY = SizeY >> DownSampleLevel;
 				const uint32 Handle = AddrSpaceAllocator.Alloc(DownSampledSizeX, DownSampledSizeY);
+				const uint32 VisibilityChannel = TmpAllocation.VisibilityChannel;
 
 				if (Handle == INDEX_NONE)
 				{
@@ -1455,9 +1513,7 @@ void FHeightFieldTextureAtlas::UpdateAllocations(FRHICommandListImmediate& RHICm
 
 				TmpAllocation.Handle = Handle;
 				CurrentAllocations.Add(TmpAllocation);
-				UploadHandles.Add(Handle);
-				UploadSizesAndMipBias.Add(FIntVector(DownSampledSizeX, DownSampledSizeY, SourceMipBias));
-				UploadSourceTextures.Add(SourceTexture->Resource->TextureRHI);
+				PendingUploads.Emplace(SourceTexture, DownSampledSizeX, DownSampledSizeY, SourceMipBias, Handle, VisibilityChannel);
 			}
 			else
 			{
@@ -1484,6 +1540,7 @@ void FHeightFieldTextureAtlas::UpdateAllocations(FRHICommandListImmediate& RHICm
 			const uint32 DownSampledSizeX = SizeX >> DownSampleLevel;
 			const uint32 DownSampledSizeY = SizeY >> DownSampleLevel;
 			const uint32 Handle = AddrSpaceAllocator.Alloc(DownSampledSizeX, DownSampledSizeY);
+			const uint32 VisibilityChannel = TmpAllocation.VisibilityChannel;
 
 			if (Handle == INDEX_NONE)
 			{
@@ -1502,56 +1559,54 @@ void FHeightFieldTextureAtlas::UpdateAllocations(FRHICommandListImmediate& RHICm
 
 			TmpAllocation.Handle = Handle;
 			CurrentAllocations.Add(TmpAllocation);
-			UploadHandles.Add(Handle);
-			UploadSizesAndMipBias.Add(FIntVector(DownSampledSizeX, DownSampledSizeY, SourceMipBias));
-			UploadSourceTextures.Add(SourceTexture->Resource->TextureRHI);
+			PendingUploads.Emplace(SourceTexture, DownSampledSizeX, DownSampledSizeY, SourceMipBias, Handle, VisibilityChannel);
 			It.RemoveCurrent();
 		}
 	}
 
-	if (UploadHandles.Num())
+	if (PendingUploads.Num())
 	{
 		FRDGBuilder GraphBuilder(RHICmdList);
 
-		TShaderMapRef<FUploadHeightFieldToAtlasCS> ComputeShader(GetGlobalShaderMap(InFeatureLevel));
-		FUploadHeightFieldToAtlasCS* ComputeShaderPtr = *ComputeShader;
-
 		RHICmdList.TransitionResource(EResourceTransitionAccess::ERWBarrier, EResourceTransitionPipeline::EComputeToCompute, AtlasUAVRHI);
 
-		for (int32 Idx = 0; Idx < UploadHandles.Num(); ++Idx)
+		if (SubAllocType == SAT_Height)
 		{
-			const uint32 Handle = UploadHandles[Idx];
-			FRHITexture2D* SourceTexture = (FRHITexture2D*)UploadSourceTextures[Idx];
-			const FIntVector& SizesAndLevel = UploadSizesAndMipBias[Idx];
-			const uint32 DownSampledSizeX = SizesAndLevel.X;
-			const uint32 DownSampledSizeY = SizesAndLevel.Y;
-			const uint32 SourceMipBias = SizesAndLevel.Z;
-			const float InvDownSampledSizeX = 1.f / DownSampledSizeX;
-			const float InvDownSampledSizeY = 1.f / DownSampledSizeY;
-			const uint32 BorderSize = AddrSpaceAllocator.BorderSize;
-			const uint32 UpdateRegionSizeX = DownSampledSizeX + 2 * BorderSize;
-			const uint32 UpdateRegionSizeY = DownSampledSizeY + 2 * BorderSize;
-			const FIntPoint StartOffset = AddrSpaceAllocator.GetStartOffset(Handle);
-
-			FUploadHeightFieldToAtlasCS::FParameters* Parameters = GraphBuilder.AllocParameters<FUploadHeightFieldToAtlasCS::FParameters>();
-			Parameters->UpdateRegionOffsetAndSize = FUintVector4(StartOffset.X, StartOffset.Y, UpdateRegionSizeX, UpdateRegionSizeY);
-			Parameters->SourceScaleBias = FVector4(InvDownSampledSizeX, InvDownSampledSizeY, (.5f - BorderSize) * InvDownSampledSizeX, (.5f - BorderSize) * InvDownSampledSizeY);
-			Parameters->SourceMipBias = SourceMipBias;
-			Parameters->SourceTexture = SourceTexture;
-			Parameters->SourceTextureSampler = TStaticSamplerState<SF_Bilinear>::GetRHI();
-			Parameters->RWHeightFieldAtlas = AtlasUAVRHI;
-
-			const uint32 NumGroupsX = FMath::DivideAndRoundUp(UpdateRegionSizeX, FUploadHeightFieldToAtlasCS::ThreadGroupSizeX);
-			const uint32 NumGroupsY = FMath::DivideAndRoundUp(UpdateRegionSizeY, FUploadHeightFieldToAtlasCS::ThreadGroupSizeY);
-
-			GraphBuilder.AddPass(
-				RDG_EVENT_NAME("UploadHeightFieldToAtlas"),
-				Parameters,
-				ERDGPassFlags::Compute,
-				[Parameters, ComputeShaderPtr, NumGroupsX, NumGroupsY](FRHICommandList& CmdList)
+			TShaderMapRef<FUploadHeightFieldToAtlasCS> ComputeShader(GetGlobalShaderMap(InFeatureLevel));
+			FUploadHeightFieldToAtlasCS* ComputeShaderPtr = *ComputeShader;
+
+			for (int32 Idx = 0; Idx < PendingUploads.Num(); ++Idx)
 			{
-				FComputeShaderUtils::Dispatch(CmdList, ComputeShaderPtr, *Parameters, FIntVector(NumGroupsX, NumGroupsY, 1));
-			});
+				typename FUploadHeightFieldToAtlasCS::FParameters* Parameters = GraphBuilder.AllocParameters<typename FUploadHeightFieldToAtlasCS::FParameters>();
+				const FIntPoint UpdateRegion = PendingUploads[Idx].SetShaderParameters(Parameters, *this);
+				GraphBuilder.AddPass(
+					RDG_EVENT_NAME("UploadHeightFieldToAtlas"),
+					Parameters,
+					ERDGPassFlags::Compute,
+					[Parameters, ComputeShaderPtr, UpdateRegion](FRHICommandList& CmdList)
+				{
+					FComputeShaderUtils::Dispatch(CmdList, ComputeShaderPtr, *Parameters, FIntVector(UpdateRegion.X, UpdateRegion.Y, 1));
+				});
+			}
+		}
+		else
+		{
+			TShaderMapRef<FUploadVisibilityToAtlasCS> ComputeShader(GetGlobalShaderMap(InFeatureLevel));
+			FUploadVisibilityToAtlasCS* ComputeShaderPtr = *ComputeShader;
+
+			for (int32 Idx = 0; Idx < PendingUploads.Num(); ++Idx)
+			{
+				typename FUploadVisibilityToAtlasCS::FParameters* Parameters = GraphBuilder.AllocParameters<typename FUploadVisibilityToAtlasCS::FParameters>();
+				const FIntPoint UpdateRegion = PendingUploads[Idx].SetShaderParameters(Parameters, *this);
+				GraphBuilder.AddPass(
+					RDG_EVENT_NAME("UploadVisibilityToAtlas"),
+					Parameters,
+					ERDGPassFlags::Compute,
+					[Parameters, ComputeShaderPtr, UpdateRegion](FRHICommandList& CmdList)
+				{
+					FComputeShaderUtils::Dispatch(CmdList, ComputeShaderPtr, *Parameters, FIntVector(UpdateRegion.X, UpdateRegion.Y, 1));
+				});
+			}
 		}
 
 		GraphBuilder.Execute();
@@ -1559,18 +1614,18 @@ void FHeightFieldTextureAtlas::UpdateAllocations(FRHICommandListImmediate& RHICm
 	}
 }
 
-uint32 FHeightFieldTextureAtlas::GetAllocationHandle(UTexture2D* Texture) const
+uint32 FLandscapeTextureAtlas::GetAllocationHandle(UTexture2D* Texture) const
 {
 	const FAllocation* Allocation = CurrentAllocations.Find(Texture);
 	return Allocation ? Allocation->Handle : INDEX_NONE;
 }
 
-FVector4 FHeightFieldTextureAtlas::GetAllocationScaleBias(uint32 Handle) const
+FVector4 FLandscapeTextureAtlas::GetAllocationScaleBias(uint32 Handle) const
 {
 	return AddrSpaceAllocator.GetScaleBias(Handle);
 }
 
-void FHeightFieldTextureAtlas::FSubAllocator::Init(uint32 InTileSize, uint32 InBorderSize, uint32 InDimInTiles)
+void FLandscapeTextureAtlas::FSubAllocator::Init(uint32 InTileSize, uint32 InBorderSize, uint32 InDimInTiles)
 {
 	check(InDimInTiles && !(InDimInTiles & (InDimInTiles - 1)));
 
@@ -1600,7 +1655,7 @@ void FHeightFieldTextureAtlas::FSubAllocator::Init(uint32 InTileSize, uint32 InB
 	MarkerQuadTree.Add(false, NumBits);
 }
 
-uint32 FHeightFieldTextureAtlas::FSubAllocator::Alloc(uint32 SizeX, uint32 SizeY)
+uint32 FLandscapeTextureAtlas::FSubAllocator::Alloc(uint32 SizeX, uint32 SizeY)
 {
 	const uint32 NumTiles1D = FMath::DivideAndRoundUp(FMath::Max(SizeX, SizeY), TileSize);
 	check(NumTiles1D <= DimInTiles);
@@ -1637,6 +1692,21 @@ uint32 FHeightFieldTextureAtlas::FSubAllocator::Alloc(uint32 SizeX, uint32 SizeY
 			Marker = true;
 			ParentQuadIdxInLevel >>= 2;
 		}
+
+		uint32 ChildLevel = Level + 1;
+		uint32 ChildQuadIdxInLevel = QuadIdxInLevel << 2;
+		uint32 NumChildren = 4;
+		for (; ChildLevel < NumLevels; ++ChildLevel)
+		{
+			const uint32 ChildQuadIdx = ChildQuadIdxInLevel + LevelOffsets[ChildLevel];
+			for (uint32 Idx = 0; Idx < NumChildren; ++Idx)
+			{
+				check(!MarkerQuadTree[ChildQuadIdx + Idx]);
+				MarkerQuadTree[ChildQuadIdx + Idx] = true;
+			}
+			ChildQuadIdxInLevel <<= 2;
+			NumChildren <<= 2;
+		}
 		
 		const uint32 QuadX = FMath::ReverseMortonCode2(QuadIdxInLevel);
 		const uint32 QuadY = FMath::ReverseMortonCode2(QuadIdxInLevel >> 1);
@@ -1658,7 +1728,7 @@ uint32 FHeightFieldTextureAtlas::FSubAllocator::Alloc(uint32 SizeX, uint32 SizeY
 	return INDEX_NONE;
 }
 
-void FHeightFieldTextureAtlas::FSubAllocator::Free(uint32 Handle)
+void FLandscapeTextureAtlas::FSubAllocator::Free(uint32 Handle)
 {
 	check(SubAllocInfos.IsValidIndex(Handle));
 
@@ -1667,8 +1737,22 @@ void FHeightFieldTextureAtlas::FSubAllocator::Free(uint32 Handle)
 
 	const uint32 Level = SubAllocInfo.Level;
 	const uint32 QuadIdx = SubAllocInfo.QuadIdx;
-	check(MarkerQuadTree[QuadIdx]);
-	MarkerQuadTree[QuadIdx] = false;
+	
+	uint32 ChildLevel = Level;
+	uint32 ChildIdxInLevel = QuadIdx - LevelOffsets[Level];
+	uint32 NumChildren = 1;
+	const uint32 NumLevels = LevelOffsets.Num();
+	for (; ChildLevel < NumLevels; ++ChildLevel)
+	{
+		const uint32 ChildIdx = ChildIdxInLevel + LevelOffsets[ChildLevel];
+		for (uint32 Idx = 0; Idx < NumChildren; ++Idx)
+		{
+			check(MarkerQuadTree[ChildIdx + Idx]);
+			MarkerQuadTree[ChildIdx + Idx] = false;
+		}
+		ChildIdxInLevel <<= 2;
+		NumChildren <<= 2;
+	}
 
 	uint32 TestIdxInLevel = (QuadIdx - LevelOffsets[Level]) & ~3u;
 	uint32 ParentLevel = Level - 1;
@@ -1687,13 +1771,13 @@ void FHeightFieldTextureAtlas::FSubAllocator::Free(uint32 Handle)
 	}
 }
 
-FVector4 FHeightFieldTextureAtlas::FSubAllocator::GetScaleBias(uint32 Handle) const
+FVector4 FLandscapeTextureAtlas::FSubAllocator::GetScaleBias(uint32 Handle) const
 {
 	check(SubAllocInfos.IsValidIndex(Handle));
 	return SubAllocInfos[Handle].UVScaleBias;
 }
 
-FIntPoint FHeightFieldTextureAtlas::FSubAllocator::GetStartOffset(uint32 Handle) const
+FIntPoint FLandscapeTextureAtlas::FSubAllocator::GetStartOffset(uint32 Handle) const
 {
 	check(SubAllocInfos.IsValidIndex(Handle));
 	const FSubAllocInfo& Info = SubAllocInfos[Handle];
@@ -1704,14 +1788,66 @@ FIntPoint FHeightFieldTextureAtlas::FSubAllocator::GetStartOffset(uint32 Handle)
 	return FIntPoint(QuadX * QuadSizeInTexels1D, QuadY * QuadSizeInTexels1D);
 }
 
-FHeightFieldTextureAtlas::FAllocation::FAllocation()
+FLandscapeTextureAtlas::FAllocation::FAllocation()
 	: SourceTexture(nullptr)
-	, RefCount(0)
 	, Handle(INDEX_NONE)
+	, VisibilityChannel(0)
+	, RefCount(0)
 {}
 
-FHeightFieldTextureAtlas::FAllocation::FAllocation(UTexture2D* InTexture)
+FLandscapeTextureAtlas::FAllocation::FAllocation(UTexture2D* InTexture, uint32 InVisibilityChannel)
 	: SourceTexture(InTexture)
-	, RefCount(1)
 	, Handle(INDEX_NONE)
+	, VisibilityChannel(InVisibilityChannel)
+	, RefCount(1)
+{}
+
+FLandscapeTextureAtlas::FPendingUpload::FPendingUpload(UTexture2D* Texture, uint32 SizeX, uint32 SizeY, uint32 MipBias, uint32 InHandle, uint32 Channel)
+	: SourceTexture(Texture->Resource->TextureRHI)
+	, SizesAndMipBias(FIntVector(SizeX, SizeY, MipBias))
+	, VisibilityChannel(Channel)
+	, Handle(InHandle)
 {}
+
+FIntPoint FLandscapeTextureAtlas::FPendingUpload::SetShaderParameters(void* ParamsPtr, const FLandscapeTextureAtlas& Atlas) const
+{
+	if (Atlas.SubAllocType == SAT_Height)
+	{
+		typename FUploadHeightFieldToAtlasCS::FParameters* Params = (typename FUploadHeightFieldToAtlasCS::FParameters*)ParamsPtr;
+		Params->RWHeightFieldAtlas = Atlas.AtlasUAVRHI;
+		return SetCommonShaderParameters(&Params->SharedParams, Atlas);
+	}
+	else
+	{
+		typename FUploadVisibilityToAtlasCS::FParameters* Params = (typename FUploadVisibilityToAtlasCS::FParameters*)ParamsPtr;
+		FVector4 ChannelMask(ForceInitToZero);
+		ChannelMask[VisibilityChannel] = 1.f;
+		Params->VisibilityChannelMask = ChannelMask;
+		Params->RWVisibilityAtlas = Atlas.AtlasUAVRHI;
+		return SetCommonShaderParameters(&Params->SharedParams, Atlas);
+	}
+}
+
+FIntPoint FLandscapeTextureAtlas::FPendingUpload::SetCommonShaderParameters(void* ParamsPtr, const FLandscapeTextureAtlas& Atlas) const
+{
+	const uint32 DownSampledSizeX = SizesAndMipBias.X;
+	const uint32 DownSampledSizeY = SizesAndMipBias.Y;
+	const uint32 SourceMipBias = SizesAndMipBias.Z;
+	const float InvDownSampledSizeX = 1.f / DownSampledSizeX;
+	const float InvDownSampledSizeY = 1.f / DownSampledSizeY;
+	const uint32 BorderSize = Atlas.AddrSpaceAllocator.BorderSize;
+	const uint32 UpdateRegionSizeX = DownSampledSizeX + 2 * BorderSize;
+	const uint32 UpdateRegionSizeY = DownSampledSizeY + 2 * BorderSize;
+	const FIntPoint StartOffset = Atlas.AddrSpaceAllocator.GetStartOffset(Handle);
+
+	typename FUploadLandscapeTextureToAtlasCS::FSharedParameters* CommonParams = (typename FUploadLandscapeTextureToAtlasCS::FSharedParameters*)ParamsPtr;
+	CommonParams->UpdateRegionOffsetAndSize = FUintVector4(StartOffset.X, StartOffset.Y, UpdateRegionSizeX, UpdateRegionSizeY);
+	CommonParams->SourceScaleBias = FVector4(InvDownSampledSizeX, InvDownSampledSizeY, (.5f - BorderSize) * InvDownSampledSizeX, (.5f - BorderSize) * InvDownSampledSizeY);
+	CommonParams->SourceMipBias = SourceMipBias;
+	CommonParams->SourceTexture = SourceTexture;
+	CommonParams->SourceTextureSampler = TStaticSamplerState<SF_Bilinear>::GetRHI();
+
+	const uint32 NumGroupsX = FMath::DivideAndRoundUp(UpdateRegionSizeX, FUploadLandscapeTextureToAtlasCS::ThreadGroupSizeX);
+	const uint32 NumGroupsY = FMath::DivideAndRoundUp(UpdateRegionSizeY, FUploadLandscapeTextureToAtlasCS::ThreadGroupSizeY);
+	return FIntPoint(NumGroupsX, NumGroupsY);
+}
diff --git a/Engine/Source/Runtime/Engine/Public/DistanceFieldAtlas.h b/Engine/Source/Runtime/Engine/Public/DistanceFieldAtlas.h
--- a/Engine/Source/Runtime/Engine/Public/DistanceFieldAtlas.h
+++ b/Engine/Source/Runtime/Engine/Public/DistanceFieldAtlas.h
@@ -154,12 +154,21 @@ private:
 
 extern ENGINE_API TGlobalResource<FDistanceFieldVolumeTextureAtlas> GDistanceFieldVolumeTextureAtlas;
 
-class ENGINE_API FHeightFieldTextureAtlas : public FRenderResource
+class ENGINE_API FLandscapeTextureAtlas : public FRenderResource
 {
 public:
+	enum ESubAllocType
+	{
+		SAT_Height,
+		SAT_Visibility,
+		SAT_Num
+	};
+
+	FLandscapeTextureAtlas(ESubAllocType InSubAllocType);
+
 	void InitializeIfNeeded();
 
-	void AddAllocation(UTexture2D* Texture);
+	void AddAllocation(UTexture2D* Texture, uint32 VisibilityChannel = 0);
 
 	void RemoveAllocation(UTexture2D* Texture);
 
@@ -231,18 +240,19 @@ private:
 
 		TSparseArray<FSubAllocInfo> SubAllocInfos;
 
-		friend class FHeightFieldTextureAtlas;
+		friend class FLandscapeTextureAtlas;
 	};
 
 	struct FAllocation
 	{
 		UTexture2D* SourceTexture;
-		uint32 RefCount;
 		uint32 Handle;
+		uint32 VisibilityChannel : 2;
+		uint32 RefCount : 30;
 
 		FAllocation();
 
-		FAllocation(UTexture2D* InTexture);
+		FAllocation(UTexture2D* InTexture, uint32 InVisibilityChannel = 0);
 
 		bool operator==(const FAllocation& Other) const
 		{
@@ -255,6 +265,21 @@ private:
 		}
 	};
 
+	struct FPendingUpload
+	{
+		FRHITexture* SourceTexture;
+		FIntVector SizesAndMipBias;
+		uint32 VisibilityChannel : 2;
+		uint32 Handle : 30;
+
+		FPendingUpload(UTexture2D* Texture, uint32 SizeX, uint32 SizeY, uint32 MipBias, uint32 InHandle, uint32 Channel);
+
+		FIntPoint SetShaderParameters(void* ParamsPtr, const FLandscapeTextureAtlas& Atlas) const;
+
+	private:
+		FIntPoint SetCommonShaderParameters(void* ParamsPtr, const FLandscapeTextureAtlas& Atlas) const;
+	};
+
 	FSubAllocator AddrSpaceAllocator;
 
 	TSet<FAllocation> PendingAllocations;
@@ -267,9 +292,12 @@ private:
 
 	uint32 MaxDownSampleLevel;
 	uint32 Generation;
+
+	const ESubAllocType SubAllocType;
 };
 
-extern ENGINE_API TGlobalResource<FHeightFieldTextureAtlas> GHeightFieldTextureAtlas;
+extern ENGINE_API TGlobalResource<FLandscapeTextureAtlas> GHeightFieldTextureAtlas;
+extern ENGINE_API TGlobalResource<FLandscapeTextureAtlas> GHFVisibilityTextureAtlas;
 
 /** Distance field data payload and output of the mesh build process. */
 class ENGINE_API FDistanceFieldVolumeData : public FDeferredCleanupInterface
diff --git a/Engine/Source/Runtime/Renderer/Private/DeferredShadingRenderer.cpp b/Engine/Source/Runtime/Renderer/Private/DeferredShadingRenderer.cpp
--- a/Engine/Source/Runtime/Renderer/Private/DeferredShadingRenderer.cpp
+++ b/Engine/Source/Runtime/Renderer/Private/DeferredShadingRenderer.cpp
@@ -569,6 +569,11 @@ void FDeferredShadingSceneRenderer::PrepareDistanceFieldScene(FRHICommandListImm
 
 	if (ShouldPrepareHeightFieldScene())
 	{
+		extern int32 GHFShadowQuality;
+		if (GHFShadowQuality > 2)
+		{
+			GHFVisibilityTextureAtlas.UpdateAllocations(RHICmdList, FeatureLevel);
+		}
 		GHeightFieldTextureAtlas.UpdateAllocations(RHICmdList, FeatureLevel);
 		UpdateGlobalHeightFieldObjectBuffers(RHICmdList);
 	}
diff --git a/Engine/Source/Runtime/Renderer/Private/DistanceFieldGlobalIllumination.cpp b/Engine/Source/Runtime/Renderer/Private/DistanceFieldGlobalIllumination.cpp
--- a/Engine/Source/Runtime/Renderer/Private/DistanceFieldGlobalIllumination.cpp
+++ b/Engine/Source/Runtime/Renderer/Private/DistanceFieldGlobalIllumination.cpp
@@ -153,7 +153,7 @@ public:
 		AtlasSizeY = GDistanceFieldVolumeTextureAtlas.GetSizeY();
 		AtlasSizeZ = GDistanceFieldVolumeTextureAtlas.GetSizeZ();
 
-		ObjectParameters.Set(RHICmdList, ShaderRHI, GShadowCulledObjectBuffers.Buffers, TextureAtlas, AtlasSizeX, AtlasSizeY, AtlasSizeZ);
+		ObjectParameters.Set(RHICmdList, ShaderRHI, GShadowCulledObjectBuffers.Buffers, TextureAtlas, FIntVector(AtlasSizeX, AtlasSizeY, AtlasSizeZ));
 
 		SetShaderValue(RHICmdList, ShaderRHI, InvPlacementGridSize, InvPlacementGridSizeValue);
 		SetShaderValue(RHICmdList, ShaderRHI, WorldToShadow, WorldToShadowValue);
@@ -603,7 +603,7 @@ public:
 		AtlasSizeY = GDistanceFieldVolumeTextureAtlas.GetSizeY();
 		AtlasSizeZ = GDistanceFieldVolumeTextureAtlas.GetSizeZ();
 
-		ObjectParameters.Set(RHICmdList, ShaderRHI, GAOCulledObjectBuffers.Buffers, TextureAtlas, AtlasSizeX, AtlasSizeY, AtlasSizeZ);
+		ObjectParameters.Set(RHICmdList, ShaderRHI, GAOCulledObjectBuffers.Buffers, TextureAtlas, FIntVector(AtlasSizeX, AtlasSizeY, AtlasSizeZ));
 
 		SetShaderValue(RHICmdList, ShaderRHI, ObjectProcessStride, GVPLSpreadUpdateOver);
 	}
@@ -700,7 +700,7 @@ public:
 		AtlasSizeY = GDistanceFieldVolumeTextureAtlas.GetSizeY();
 		AtlasSizeZ = GDistanceFieldVolumeTextureAtlas.GetSizeZ();
 
-		ObjectParameters.Set(RHICmdList, ShaderRHI, GAOCulledObjectBuffers.Buffers, TextureAtlas, AtlasSizeX, AtlasSizeY, AtlasSizeZ);
+		ObjectParameters.Set(RHICmdList, ShaderRHI, GAOCulledObjectBuffers.Buffers, TextureAtlas, FIntVector(AtlasSizeX, AtlasSizeY, AtlasSizeZ));
 		SurfelParameters.Set(RHICmdList, ShaderRHI, *Scene->DistanceFieldSceneData.SurfelBuffers, *Scene->DistanceFieldSceneData.InstancedSurfelBuffers);
 
 		FLightShaderParameters LightParameters;
@@ -1063,7 +1063,7 @@ public:
 		AtlasSizeY = GDistanceFieldVolumeTextureAtlas.GetSizeY();
 		AtlasSizeZ = GDistanceFieldVolumeTextureAtlas.GetSizeZ();
 
-		ObjectParameters.Set(RHICmdList, ShaderRHI, GAOCulledObjectBuffers.Buffers, TextureAtlas, AtlasSizeX, AtlasSizeY, AtlasSizeZ);
+		ObjectParameters.Set(RHICmdList, ShaderRHI, GAOCulledObjectBuffers.Buffers, TextureAtlas, FIntVector(AtlasSizeX, AtlasSizeY, AtlasSizeZ));
 
 		AOParameters.Set(RHICmdList, ShaderRHI, Parameters);
 		ScreenGridParameters.Set(RHICmdList, ShaderRHI, View, DistanceFieldNormal);
diff --git a/Engine/Source/Runtime/Renderer/Private/DistanceFieldLightingShared.h b/Engine/Source/Runtime/Renderer/Private/DistanceFieldLightingShared.h
--- a/Engine/Source/Runtime/Renderer/Private/DistanceFieldLightingShared.h
+++ b/Engine/Source/Runtime/Renderer/Private/DistanceFieldLightingShared.h
@@ -409,6 +409,7 @@ public:
 		CulledObjectBounds.Bind(ParameterMap, TEXT("CulledObjectBounds"));
 		CulledObjectData.Bind(ParameterMap, TEXT("CulledObjectData"));
 		CulledObjectBoxBounds.Bind(ParameterMap, TEXT("CulledObjectBoxBounds"));
+		HFVisibilityTexture.Bind(ParameterMap, TEXT("HFVisibilityTexture"));
 		DistanceFieldTexture.Bind(ParameterMap, TEXT("DistanceFieldTexture"));
 		DistanceFieldSampler.Bind(ParameterMap, TEXT("DistanceFieldSampler"));
 		DistanceFieldAtlasTexelSize.Bind(ParameterMap, TEXT("DistanceFieldAtlasTexelSize"));
@@ -420,9 +421,8 @@ public:
 		TShaderRHI* ShaderRHI,
 		const TDistanceFieldCulledObjectBuffers<PrimitiveType>& ObjectBuffers,
 		FRHITexture* TextureAtlas,
-		int32 AtlasSizeX,
-		int32 AtlasSizeY,
-		int32 AtlasSizeZ)
+		const FIntVector& AtlasSizes,
+		FRHITexture* HFVisibilityAtlas = nullptr)
 	{
 		ObjectIndirectArguments.SetBuffer(RHICmdList, ShaderRHI, ObjectBuffers.ObjectIndirectArguments);
 		CulledObjectBounds.SetBuffer(RHICmdList, ShaderRHI, ObjectBuffers.Bounds);
@@ -443,9 +443,15 @@ public:
 			TextureAtlas
 			);
 
-		const int32 NumTexelsOneDimX = AtlasSizeX;
-		const int32 NumTexelsOneDimY = AtlasSizeY;
-		const int32 NumTexelsOneDimZ = AtlasSizeZ;
+		if (HFVisibilityTexture.IsBound())
+		{
+			check(HFVisibilityAtlas);
+			SetTextureParameter(RHICmdList, ShaderRHI, HFVisibilityTexture, HFVisibilityAtlas);
+		}
+
+		const int32 NumTexelsOneDimX = AtlasSizes.X;
+		const int32 NumTexelsOneDimY = AtlasSizes.Y;
+		const int32 NumTexelsOneDimZ = AtlasSizes.Z;
 		const FVector InvTextureDim(1.0f / NumTexelsOneDimX, 1.0f / NumTexelsOneDimY, 1.0f / NumTexelsOneDimZ);
 		SetShaderValue(RHICmdList, ShaderRHI, DistanceFieldAtlasTexelSize, InvTextureDim);
 	}
@@ -494,7 +500,14 @@ public:
 
 	friend FArchive& operator<<(FArchive& Ar, TDistanceFieldCulledObjectBufferParameters<PrimitiveType>& P)
 	{
-		Ar << P.ObjectIndirectArguments << P.CulledObjectBounds << P.CulledObjectData << P.CulledObjectBoxBounds << P.DistanceFieldTexture << P.DistanceFieldSampler << P.DistanceFieldAtlasTexelSize;
+		Ar << P.ObjectIndirectArguments;
+		Ar << P.CulledObjectBounds;
+		Ar << P.CulledObjectData;
+		Ar << P.CulledObjectBoxBounds;
+		Ar << P.HFVisibilityTexture;
+		Ar << P.DistanceFieldTexture;
+		Ar << P.DistanceFieldSampler;
+		Ar << P.DistanceFieldAtlasTexelSize;
 		return Ar;
 	}
 
@@ -503,6 +516,7 @@ private:
 	FRWShaderParameter CulledObjectBounds;
 	FRWShaderParameter CulledObjectData;
 	FRWShaderParameter CulledObjectBoxBounds;
+	FShaderResourceParameter HFVisibilityTexture;
 	FShaderResourceParameter DistanceFieldTexture;
 	FShaderResourceParameter DistanceFieldSampler;
 	FShaderParameter DistanceFieldAtlasTexelSize;
diff --git a/Engine/Source/Runtime/Renderer/Private/DistanceFieldObjectCulling.cpp b/Engine/Source/Runtime/Renderer/Private/DistanceFieldObjectCulling.cpp
--- a/Engine/Source/Runtime/Renderer/Private/DistanceFieldObjectCulling.cpp
+++ b/Engine/Source/Runtime/Renderer/Private/DistanceFieldObjectCulling.cpp
@@ -159,7 +159,7 @@ public:
 		AtlasSizeZ = GDistanceFieldVolumeTextureAtlas.GetSizeZ();
 
 		ObjectBufferParameters.Set(RHICmdList, ShaderRHI, *(Scene->DistanceFieldSceneData.GetCurrentObjectBuffers()), Scene->DistanceFieldSceneData.NumObjectsInBuffer, TextureAtlas, AtlasSizeX, AtlasSizeY, AtlasSizeZ);
-		CulledObjectParameters.Set(RHICmdList, ShaderRHI, GAOCulledObjectBuffers.Buffers, TextureAtlas, AtlasSizeX, AtlasSizeY, AtlasSizeZ);
+		CulledObjectParameters.Set(RHICmdList, ShaderRHI, GAOCulledObjectBuffers.Buffers, TextureAtlas, FIntVector(AtlasSizeX, AtlasSizeY, AtlasSizeZ));
 
 		AOParameters.Set(RHICmdList, ShaderRHI, Parameters);
 
@@ -382,7 +382,7 @@ public:
 		AtlasSizeY = GDistanceFieldVolumeTextureAtlas.GetSizeY();
 		AtlasSizeZ = GDistanceFieldVolumeTextureAtlas.GetSizeZ();
 
-		ObjectParameters.Set(RHICmdList, ShaderRHI, GAOCulledObjectBuffers.Buffers, TextureAtlas, AtlasSizeX, AtlasSizeY, AtlasSizeZ);
+		ObjectParameters.Set(RHICmdList, ShaderRHI, GAOCulledObjectBuffers.Buffers, TextureAtlas, FIntVector(AtlasSizeX, AtlasSizeY, AtlasSizeZ));
 		AOParameters.Set(RHICmdList, ShaderRHI, Parameters);
 
 		const int32 NumRings = StencilingGeometry::GLowPolyStencilSphereVertexBuffer.GetNumRings();
@@ -461,7 +461,7 @@ public:
 		AtlasSizeY = GDistanceFieldVolumeTextureAtlas.GetSizeY();
 		AtlasSizeZ = GDistanceFieldVolumeTextureAtlas.GetSizeZ();
 
-		ObjectParameters.Set(RHICmdList, ShaderRHI, GAOCulledObjectBuffers.Buffers, TextureAtlas, AtlasSizeX, AtlasSizeY, AtlasSizeZ);
+		ObjectParameters.Set(RHICmdList, ShaderRHI, GAOCulledObjectBuffers.Buffers, TextureAtlas, FIntVector(AtlasSizeX, AtlasSizeY, AtlasSizeZ));
 		AOParameters.Set(RHICmdList, ShaderRHI, Parameters);
 
 		FTileIntersectionResources* TileIntersectionResources = ((FSceneViewState*)View.State)->AOTileIntersectionResources;
@@ -552,7 +552,7 @@ public:
 		AtlasSizeY = GDistanceFieldVolumeTextureAtlas.GetSizeY();
 		AtlasSizeZ = GDistanceFieldVolumeTextureAtlas.GetSizeZ();
 
-		ObjectParameters.Set(RHICmdList, ShaderRHI, GAOCulledObjectBuffers.Buffers, TextureAtlas, AtlasSizeX, AtlasSizeY, AtlasSizeZ);
+		ObjectParameters.Set(RHICmdList, ShaderRHI, GAOCulledObjectBuffers.Buffers, TextureAtlas, FIntVector(AtlasSizeX, AtlasSizeY, AtlasSizeZ));
 
 		FTileIntersectionResources* TileIntersectionResources = ((FSceneViewState*)View.State)->AOTileIntersectionResources;
 		
diff --git a/Engine/Source/Runtime/Renderer/Private/DistanceFieldObjectManagement.cpp b/Engine/Source/Runtime/Renderer/Private/DistanceFieldObjectManagement.cpp
--- a/Engine/Source/Runtime/Renderer/Private/DistanceFieldObjectManagement.cpp
+++ b/Engine/Source/Runtime/Renderer/Private/DistanceFieldObjectManagement.cpp
@@ -37,10 +37,10 @@ FAutoConsoleVariableRef CVarAOLogObjectBufferReallocation(
 
 // Must match equivalent shader defines
 template<> int32 TDistanceFieldObjectBuffers<DFPT_SignedDistanceField>::ObjectDataStride = 17;
-template<> int32 TDistanceFieldObjectBuffers<DFPT_HeightField>::ObjectDataStride = 5;
+template<> int32 TDistanceFieldObjectBuffers<DFPT_HeightField>::ObjectDataStride = 6;
 template<> int32 TDistanceFieldCulledObjectBuffers<DFPT_SignedDistanceField>::ObjectDataStride = 17;
 template<> int32 TDistanceFieldCulledObjectBuffers<DFPT_SignedDistanceField>::ObjectBoxBoundsStride = 5;
-template<> int32 TDistanceFieldCulledObjectBuffers<DFPT_HeightField>::ObjectDataStride = 5;
+template<> int32 TDistanceFieldCulledObjectBuffers<DFPT_HeightField>::ObjectDataStride = 6;
 template<> int32 TDistanceFieldCulledObjectBuffers<DFPT_HeightField>::ObjectBoxBoundsStride = 5;
 
 // In float4's.  Must match corresponding usf definition
@@ -1290,6 +1290,18 @@ bool ProcessHeightFieldPrimitiveUpdate(
 		HeightFieldScaleBias.Z * AllocationScaleBias.X + AllocationScaleBias.Z,
 		HeightFieldScaleBias.W * AllocationScaleBias.Y + AllocationScaleBias.W));
 
+	FVector4 VisUVScaleBias(ForceInitToZero);
+	if (VisibilityTexture)
+	{
+		const uint32 VisHandle = GHFVisibilityTextureAtlas.GetAllocationHandle(VisibilityTexture);
+		if (VisHandle != INDEX_NONE)
+		{
+			const FVector4 ScaleBias = GHFVisibilityTextureAtlas.GetAllocationScaleBias(VisHandle);
+			VisUVScaleBias = FVector4(1.f / HeightFieldRect.Width() * ScaleBias.X, 1.f / HeightFieldRect.Height() * ScaleBias.Y, ScaleBias.Z, ScaleBias.W);
+		}
+	}
+	UploadObjectData.Add(VisUVScaleBias);
+
 	check(!(UploadObjectData.Num() % UploadHeightFieldObjectDataStride));
 
 	return true;
@@ -1613,7 +1625,8 @@ void FDeferredShadingSceneRenderer::UpdateGlobalHeightFieldObjectBuffers(FRHICom
 
 	if (GHeightFieldTextureAtlas.GetAtlasTexture()
 		&& (DistanceFieldSceneData.HasPendingHeightFieldOperations()
-			|| DistanceFieldSceneData.HeightFieldAtlasGeneration != GHeightFieldTextureAtlas.GetGeneration()))
+			|| DistanceFieldSceneData.HeightFieldAtlasGeneration != GHeightFieldTextureAtlas.GetGeneration()
+			|| DistanceFieldSceneData.HFVisibilityAtlasGenerattion != GHFVisibilityTextureAtlas.GetGeneration()))
 	{
 		QUICK_SCOPE_CYCLE_COUNTER(STAT_UpdateHeightFieldSceneObjectData);
 		SCOPED_DRAW_EVENT(RHICmdList, UpdateHeightFieldSceneObjectData);
@@ -1623,19 +1636,21 @@ void FDeferredShadingSceneRenderer::UpdateGlobalHeightFieldObjectBuffers(FRHICom
 			DistanceFieldSceneData.HeightFieldObjectBuffers = new FHeightFieldObjectBuffers;
 		}
 
-		if (DistanceFieldSceneData.HeightFieldAtlasGeneration != GHeightFieldTextureAtlas.GetGeneration())
+		if (DistanceFieldSceneData.HeightFieldAtlasGeneration != GHeightFieldTextureAtlas.GetGeneration()
+			|| DistanceFieldSceneData.HFVisibilityAtlasGenerattion != GHFVisibilityTextureAtlas.GetGeneration())
 		{
 			DistanceFieldSceneData.HeightFieldAtlasGeneration = GHeightFieldTextureAtlas.GetGeneration();
+			DistanceFieldSceneData.HFVisibilityAtlasGenerattion = GHFVisibilityTextureAtlas.GetGeneration();
 
 			for (int32 Idx = 0; Idx < DistanceFieldSceneData.HeightfieldPrimitives.Num(); ++Idx)
 			{
-				const FPrimitiveSceneInfo* Primitive = DistanceFieldSceneData.HeightfieldPrimitives[Idx];
+				FPrimitiveSceneInfo* Primitive = DistanceFieldSceneData.HeightfieldPrimitives[Idx];
 
 				if (!DistanceFieldSceneData.HasPendingRemoveHeightFieldPrimitive(Primitive)
 					&& !DistanceFieldSceneData.PendingHeightFieldAddOps.Contains(Primitive)
 					&& !DistanceFieldSceneData.PendingHeightFieldUpdateOps.Contains(Primitive))
 				{
-					DistanceFieldSceneData.PendingHeightFieldUpdateOps.Add((FPrimitiveSceneInfo*)Primitive);
+					DistanceFieldSceneData.PendingHeightFieldUpdateOps.Add(Primitive);
 				}
 			}
 		}
diff --git a/Engine/Source/Runtime/Renderer/Private/DistanceFieldScreenGridLighting.cpp b/Engine/Source/Runtime/Renderer/Private/DistanceFieldScreenGridLighting.cpp
--- a/Engine/Source/Runtime/Renderer/Private/DistanceFieldScreenGridLighting.cpp
+++ b/Engine/Source/Runtime/Renderer/Private/DistanceFieldScreenGridLighting.cpp
@@ -142,7 +142,7 @@ public:
 		AtlasSizeY = GDistanceFieldVolumeTextureAtlas.GetSizeY();
 		AtlasSizeZ = GDistanceFieldVolumeTextureAtlas.GetSizeZ();
 
-		ObjectParameters.Set(RHICmdList, ShaderRHI, GAOCulledObjectBuffers.Buffers, TextureAtlas, AtlasSizeX, AtlasSizeY, AtlasSizeZ);
+		ObjectParameters.Set(RHICmdList, ShaderRHI, GAOCulledObjectBuffers.Buffers, TextureAtlas, FIntVector(AtlasSizeX, AtlasSizeY, AtlasSizeZ));
 
 		AOParameters.Set(RHICmdList, ShaderRHI, Parameters);
 		ScreenGridParameters.Set(RHICmdList, ShaderRHI, View, DistanceFieldNormal);
@@ -305,7 +305,7 @@ public:
 		AtlasSizeY = GDistanceFieldVolumeTextureAtlas.GetSizeY();
 		AtlasSizeZ = GDistanceFieldVolumeTextureAtlas.GetSizeZ();
 
-		ObjectParameters.Set(RHICmdList, ShaderRHI, GAOCulledObjectBuffers.Buffers, TextureAtlas, AtlasSizeX, AtlasSizeY, AtlasSizeZ);
+		ObjectParameters.Set(RHICmdList, ShaderRHI, GAOCulledObjectBuffers.Buffers, TextureAtlas, FIntVector(AtlasSizeX, AtlasSizeY, AtlasSizeZ));
 
 		AOParameters.Set(RHICmdList, ShaderRHI, Parameters);
 		ScreenGridParameters.Set(RHICmdList, ShaderRHI, View, DistanceFieldNormal);
diff --git a/Engine/Source/Runtime/Renderer/Private/DistanceFieldShadowing.cpp b/Engine/Source/Runtime/Renderer/Private/DistanceFieldShadowing.cpp
--- a/Engine/Source/Runtime/Renderer/Private/DistanceFieldShadowing.cpp
+++ b/Engine/Source/Runtime/Renderer/Private/DistanceFieldShadowing.cpp
@@ -92,12 +92,12 @@ FAutoConsoleVariableRef CVarHeightFieldShadowing(
 	TEXT("Whether the height field shadowing feature is allowed."),
 	ECVF_Scalability | ECVF_RenderThreadSafe);
 
-static int32 GHFShadowQuality = 3;
+int32 GHFShadowQuality = 2;
 FAutoConsoleVariableRef CVarHFShadowQuality(
 	TEXT("r.HFShadowQuality"),
 	GHFShadowQuality,
 	TEXT("Defines the height field shadow method which allows to adjust for quality or performance.\n")
-	TEXT(" 0:off, 1:low (8 steps), 2:medium (16 steps), 3:high (32 steps, default)"),
+	TEXT(" 0:off, 1:low (8 steps), 2:medium (16 steps, default), 3:high (32 steps, hole aware)"),
 	ECVF_Scalability | ECVF_RenderThreadSafe);
 
 static float GMinDirectionalLightAngleForRTHF = 27.f;
@@ -180,6 +180,7 @@ public:
 		
 		int32 NumObjectsInBuffer;
 		FRHITexture* TextureAtlas;
+		FRHITexture* VisibilityAtlas;
 		int32 AtlasSizeX;
 		int32 AtlasSizeY;
 		int32 AtlasSizeZ;
@@ -188,6 +189,7 @@ public:
 		{
 			NumObjectsInBuffer = SceneData.NumHeightFieldObjectsInBuffer;
 			TextureAtlas = GHeightFieldTextureAtlas.GetAtlasTexture();
+			VisibilityAtlas = GHFVisibilityTextureAtlas.GetAtlasTexture();
 			AtlasSizeX = GHeightFieldTextureAtlas.GetSizeX();
 			AtlasSizeY = GHeightFieldTextureAtlas.GetSizeY();
 			AtlasSizeZ = 1;
@@ -196,6 +198,7 @@ public:
 		{
 			NumObjectsInBuffer = SceneData.NumObjectsInBuffer;
 			TextureAtlas = GDistanceFieldVolumeTextureAtlas.VolumeTextureRHI;
+			VisibilityAtlas = nullptr;
 			AtlasSizeX = GDistanceFieldVolumeTextureAtlas.GetSizeX();
 			AtlasSizeY = GDistanceFieldVolumeTextureAtlas.GetSizeY();
 			AtlasSizeZ = GDistanceFieldVolumeTextureAtlas.GetSizeZ();
@@ -212,7 +215,7 @@ public:
 		OutUAVs[3] = CulledObjectBuffers.Buffers.BoxBounds.UAV;
 		RHICmdList.TransitionResources(EResourceTransitionAccess::ERWBarrier, EResourceTransitionPipeline::EComputeToCompute, OutUAVs, UE_ARRAY_COUNT(OutUAVs));
 
-		CulledObjectParameters.Set(RHICmdList, ShaderRHI, CulledObjectBuffers.Buffers, TextureAtlas, AtlasSizeX, AtlasSizeY, AtlasSizeZ);
+		CulledObjectParameters.Set(RHICmdList, ShaderRHI, CulledObjectBuffers.Buffers, TextureAtlas, FIntVector(AtlasSizeX, AtlasSizeY, AtlasSizeZ), VisibilityAtlas);
 
 		SetShaderValue(RHICmdList, ShaderRHI, ObjectBoundingGeometryIndexCount, UE_ARRAY_COUNT(GCubeIndices));
 		SetShaderValue(RHICmdList, ShaderRHI, WorldToShadow, WorldToShadowValue);
@@ -311,6 +314,7 @@ public:
 		auto& CulledObjectBuffers = TSelector<bIsHeightField>()(GShadowCulledHeightFieldObjectBuffers, GShadowCulledObjectBuffers);
 
 		FRHITexture* TextureAtlas;
+		FRHITexture* VisibilityAtlas;
 		int32 AtlasSizeX;
 		int32 AtlasSizeY;
 		int32 AtlasSizeZ;
@@ -318,6 +322,7 @@ public:
 		if (bIsHeightField)
 		{
 			TextureAtlas = GHeightFieldTextureAtlas.GetAtlasTexture();
+			VisibilityAtlas = GHFVisibilityTextureAtlas.GetAtlasTexture();
 			AtlasSizeX = GHeightFieldTextureAtlas.GetSizeX();
 			AtlasSizeY = GHeightFieldTextureAtlas.GetSizeY();
 			AtlasSizeZ = 1;
@@ -325,12 +330,13 @@ public:
 		else
 		{
 			TextureAtlas = GDistanceFieldVolumeTextureAtlas.VolumeTextureRHI;
+			VisibilityAtlas = nullptr;
 			AtlasSizeX = GDistanceFieldVolumeTextureAtlas.GetSizeX();
 			AtlasSizeY = GDistanceFieldVolumeTextureAtlas.GetSizeY();
 			AtlasSizeZ = GDistanceFieldVolumeTextureAtlas.GetSizeZ();
 		}
 
-		ObjectParameters.Set(RHICmdList, ShaderRHI, CulledObjectBuffers.Buffers, TextureAtlas, AtlasSizeX, AtlasSizeY, AtlasSizeZ);
+		ObjectParameters.Set(RHICmdList, ShaderRHI, CulledObjectBuffers.Buffers, TextureAtlas, FIntVector(AtlasSizeX, AtlasSizeY, AtlasSizeZ), VisibilityAtlas);
 
 		SetShaderValue(RHICmdList, ShaderRHI, WorldToShadow, WorldToShadowMatrixValue);
 
@@ -398,6 +404,7 @@ public:
 		auto& CulledObjectBuffers = TSelector<bIsHeightField>()(GShadowCulledHeightFieldObjectBuffers, GShadowCulledObjectBuffers);
 
 		FRHITexture* TextureAtlas;
+		FRHITexture* VisibilityAtlas;
 		int32 AtlasSizeX;
 		int32 AtlasSizeY;
 		int32 AtlasSizeZ;
@@ -405,6 +412,7 @@ public:
 		if (bIsHeightField)
 		{
 			TextureAtlas = GHeightFieldTextureAtlas.GetAtlasTexture();
+			VisibilityAtlas = GHFVisibilityTextureAtlas.GetAtlasTexture();
 			AtlasSizeX = GHeightFieldTextureAtlas.GetSizeX();
 			AtlasSizeY = GHeightFieldTextureAtlas.GetSizeY();
 			AtlasSizeZ = 1;
@@ -412,12 +420,13 @@ public:
 		else
 		{
 			TextureAtlas = GDistanceFieldVolumeTextureAtlas.VolumeTextureRHI;
+			VisibilityAtlas = nullptr;
 			AtlasSizeX = GDistanceFieldVolumeTextureAtlas.GetSizeX();
 			AtlasSizeY = GDistanceFieldVolumeTextureAtlas.GetSizeY();
 			AtlasSizeZ = GDistanceFieldVolumeTextureAtlas.GetSizeZ();
 		}
 
-		ObjectParameters.Set(RHICmdList, ShaderRHI, CulledObjectBuffers.Buffers, TextureAtlas, AtlasSizeX, AtlasSizeY, AtlasSizeZ);
+		ObjectParameters.Set(RHICmdList, ShaderRHI, CulledObjectBuffers.Buffers, TextureAtlas, FIntVector(AtlasSizeX, AtlasSizeY, AtlasSizeZ), VisibilityAtlas);
 
 		LightTileIntersectionParameters.Set(RHICmdList, ShaderRHI, *TileIntersectionResources);
 	}
@@ -533,6 +542,7 @@ public:
 		auto& CulledObjectBuffers = TSelector<bIsHeightField>()(GShadowCulledHeightFieldObjectBuffers, GShadowCulledObjectBuffers);
 
 		FRHITexture* TextureAtlas;
+		FRHITexture* VisibilityAtlas;
 		int32 AtlasSizeX;
 		int32 AtlasSizeY;
 		int32 AtlasSizeZ;
@@ -540,6 +550,7 @@ public:
 		if (bIsHeightField)
 		{
 			TextureAtlas = GHeightFieldTextureAtlas.GetAtlasTexture();
+			VisibilityAtlas = GHFVisibilityTextureAtlas.GetAtlasTexture();
 			AtlasSizeX = GHeightFieldTextureAtlas.GetSizeX();
 			AtlasSizeY = GHeightFieldTextureAtlas.GetSizeY();
 			AtlasSizeZ = 1;
@@ -547,12 +558,13 @@ public:
 		else
 		{
 			TextureAtlas = GDistanceFieldVolumeTextureAtlas.VolumeTextureRHI;
+			VisibilityAtlas = nullptr;
 			AtlasSizeX = GDistanceFieldVolumeTextureAtlas.GetSizeX();
 			AtlasSizeY = GDistanceFieldVolumeTextureAtlas.GetSizeY();
 			AtlasSizeZ = GDistanceFieldVolumeTextureAtlas.GetSizeZ();
 		}
 
-		ObjectParameters.Set(RHICmdList, ShaderRHI, CulledObjectBuffers.Buffers, TextureAtlas, AtlasSizeX, AtlasSizeY, AtlasSizeZ);
+		ObjectParameters.Set(RHICmdList, ShaderRHI, CulledObjectBuffers.Buffers, TextureAtlas, FIntVector(AtlasSizeX, AtlasSizeY, AtlasSizeZ), VisibilityAtlas);
 
 		SceneTextureParameters.Set(RHICmdList, ShaderRHI, View.FeatureLevel, ESceneTextureSetupMode::All);
 
diff --git a/Engine/Source/Runtime/Renderer/Private/DistanceFieldVisualization.cpp b/Engine/Source/Runtime/Renderer/Private/DistanceFieldVisualization.cpp
--- a/Engine/Source/Runtime/Renderer/Private/DistanceFieldVisualization.cpp
+++ b/Engine/Source/Runtime/Renderer/Private/DistanceFieldVisualization.cpp
@@ -76,7 +76,7 @@ public:
 		AtlasSizeY = GDistanceFieldVolumeTextureAtlas.GetSizeY();
 		AtlasSizeZ = GDistanceFieldVolumeTextureAtlas.GetSizeZ();
 
-		ObjectParameters.Set(RHICmdList, ShaderRHI, GAOCulledObjectBuffers.Buffers, TextureAtlas, AtlasSizeX, AtlasSizeY, AtlasSizeZ);
+		ObjectParameters.Set(RHICmdList, ShaderRHI, GAOCulledObjectBuffers.Buffers, TextureAtlas, FIntVector(AtlasSizeX, AtlasSizeY, AtlasSizeZ));
 
 		AOParameters.Set(RHICmdList, ShaderRHI, Parameters);
 		SceneTextureParameters.Set(RHICmdList, ShaderRHI, View.FeatureLevel, ESceneTextureSetupMode::All);
diff --git a/Engine/Source/Runtime/Renderer/Private/GlobalDistanceField.cpp b/Engine/Source/Runtime/Renderer/Private/GlobalDistanceField.cpp
--- a/Engine/Source/Runtime/Renderer/Private/GlobalDistanceField.cpp
+++ b/Engine/Source/Runtime/Renderer/Private/GlobalDistanceField.cpp
@@ -220,7 +220,7 @@ public:
 		OutUAVs[3] = GGlobalDistanceFieldCulledObjectBuffers.Buffers.BoxBounds.UAV;
 		RHICmdList.TransitionResources(EResourceTransitionAccess::ERWBarrier, EResourceTransitionPipeline::EComputeToCompute, OutUAVs, UE_ARRAY_COUNT(OutUAVs));
 
-		CulledObjectParameters.Set(RHICmdList, ShaderRHI, GGlobalDistanceFieldCulledObjectBuffers.Buffers, TextureAtlas, AtlasSizeX, AtlasSizeY, AtlasSizeZ);
+		CulledObjectParameters.Set(RHICmdList, ShaderRHI, GGlobalDistanceFieldCulledObjectBuffers.Buffers, TextureAtlas, FIntVector(AtlasSizeX, AtlasSizeY, AtlasSizeZ));
 
 		extern float GAOConeHalfAngle;
 		const float GlobalMaxSphereQueryRadius = MaxOcclusionDistance / (1 + FMath::Tan(GAOConeHalfAngle));
@@ -398,7 +398,7 @@ public:
 		AtlasSizeY = GDistanceFieldVolumeTextureAtlas.GetSizeY();
 		AtlasSizeZ = GDistanceFieldVolumeTextureAtlas.GetSizeZ();
 
-		CulledObjectBufferParameters.Set(RHICmdList, ShaderRHI, GGlobalDistanceFieldCulledObjectBuffers.Buffers, TextureAtlas, AtlasSizeX, AtlasSizeY, AtlasSizeZ);
+		CulledObjectBufferParameters.Set(RHICmdList, ShaderRHI, GGlobalDistanceFieldCulledObjectBuffers.Buffers, TextureAtlas, FIntVector(AtlasSizeX, AtlasSizeY, AtlasSizeZ));
 		GlobalDistanceFieldParameters.Set(RHICmdList, ShaderRHI, GlobalDistanceFieldInfo.ParameterData);
 
 		RHICmdList.TransitionResource(EResourceTransitionAccess::ERWBarrier, EResourceTransitionPipeline::EComputeToCompute, GObjectGridBuffers.CullGridObjectNum.UAV);
@@ -544,7 +544,7 @@ public:
 		AtlasSizeY = GDistanceFieldVolumeTextureAtlas.GetSizeY();
 		AtlasSizeZ = GDistanceFieldVolumeTextureAtlas.GetSizeZ();
 
-		CulledObjectBufferParameters.Set(RHICmdList, ShaderRHI, GGlobalDistanceFieldCulledObjectBuffers.Buffers, TextureAtlas, AtlasSizeX, AtlasSizeY, AtlasSizeZ);
+		CulledObjectBufferParameters.Set(RHICmdList, ShaderRHI, GGlobalDistanceFieldCulledObjectBuffers.Buffers, TextureAtlas, FIntVector(AtlasSizeX, AtlasSizeY, AtlasSizeZ));
 		GlobalDistanceFieldParameters.Set(RHICmdList, ShaderRHI, ParameterData);
 
 		const FSceneRenderTargetItem& ClipMapRTI = Clipmap.RenderTarget->GetRenderTargetItem();
diff --git a/Engine/Source/Runtime/Renderer/Private/HeightfieldLighting.cpp b/Engine/Source/Runtime/Renderer/Private/HeightfieldLighting.cpp
--- a/Engine/Source/Runtime/Renderer/Private/HeightfieldLighting.cpp
+++ b/Engine/Source/Runtime/Renderer/Private/HeightfieldLighting.cpp
@@ -956,7 +956,7 @@ public:
 		AtlasSizeY = GDistanceFieldVolumeTextureAtlas.GetSizeY();
 		AtlasSizeZ = GDistanceFieldVolumeTextureAtlas.GetSizeZ();
 
-		ObjectParameters.Set(RHICmdList, ShaderRHI, CulledObjectBuffers.Buffers, TextureAtlas, AtlasSizeX, AtlasSizeY, AtlasSizeZ);
+		ObjectParameters.Set(RHICmdList, ShaderRHI, CulledObjectBuffers.Buffers, TextureAtlas, FIntVector(AtlasSizeX, AtlasSizeY, AtlasSizeZ));
 
 		SetShaderValue(RHICmdList, ShaderRHI, LightDirection, ProjectedShadowInfo->GetLightSceneInfo().Proxy->GetDirection());
 
diff --git a/Engine/Source/Runtime/Renderer/Private/RendererScene.cpp b/Engine/Source/Runtime/Renderer/Private/RendererScene.cpp
--- a/Engine/Source/Runtime/Renderer/Private/RendererScene.cpp
+++ b/Engine/Source/Runtime/Renderer/Private/RendererScene.cpp
@@ -376,6 +376,7 @@ FDistanceFieldSceneData::FDistanceFieldSceneData(EShaderPlatform ShaderPlatform)
 	, InstancedSurfelBuffers(NULL)
 	, AtlasGeneration(0)
 	, HeightFieldAtlasGeneration(0)
+	, HFVisibilityAtlasGenerattion(0)
 {
 	ObjectBuffers[0] = nullptr;
 	ObjectBuffers[1] = nullptr;
@@ -406,11 +407,18 @@ void FDistanceFieldSceneData::AddPrimitive(FPrimitiveSceneInfo* InPrimitive)
 		if (Proxy->SupportsHeightfieldRepresentation())
 		{
 			UTexture2D* HeightAndNormal;
-			UTexture2D* Unused0, *Unused1;
-			FHeightfieldComponentDescription Unused2(FMatrix::Identity);
-			Proxy->GetHeightfieldRepresentation(HeightAndNormal, Unused0, Unused1, Unused2);
+			UTexture2D* DiffuseColor;
+			UTexture2D* Visibility;
+			FHeightfieldComponentDescription Desc(FMatrix::Identity);
+			Proxy->GetHeightfieldRepresentation(HeightAndNormal, DiffuseColor, Visibility, Desc);
 			GHeightFieldTextureAtlas.AddAllocation(HeightAndNormal);
 
+			if (Visibility)
+			{
+				check(Desc.VisibilityChannel >= 0 && Desc.VisibilityChannel < 4);
+				GHFVisibilityTextureAtlas.AddAllocation(Visibility, Desc.VisibilityChannel);
+			}
+
 			checkSlow(!PendingHeightFieldAddOps.Contains(InPrimitive));
 			checkSlow(!PendingHeightFieldUpdateOps.Contains(InPrimitive));
 			PendingHeightFieldAddOps.Add(InPrimitive);
@@ -467,11 +475,17 @@ void FDistanceFieldSceneData::RemovePrimitive(FPrimitiveSceneInfo* InPrimitive)
 		if (Proxy->SupportsHeightfieldRepresentation())
 		{
 			UTexture2D* HeightAndNormal;
-			UTexture2D* Unused0, *Unused1;
-			FHeightfieldComponentDescription Unused2(FMatrix::Identity);
-			Proxy->GetHeightfieldRepresentation(HeightAndNormal, Unused0, Unused1, Unused2);
+			UTexture2D* DiffuseColor;
+			UTexture2D* Visibility;
+			FHeightfieldComponentDescription Desc(FMatrix::Identity);
+			Proxy->GetHeightfieldRepresentation(HeightAndNormal, DiffuseColor, Visibility, Desc);
 			GHeightFieldTextureAtlas.RemoveAllocation(HeightAndNormal);
 
+			if (Visibility)
+			{
+				GHFVisibilityTextureAtlas.RemoveAllocation(Visibility);
+			}
+
 			PendingHeightFieldAddOps.Remove(InPrimitive);
 			PendingHeightFieldUpdateOps.Remove(InPrimitive);
 
diff --git a/Engine/Source/Runtime/Renderer/Private/ScenePrivate.h b/Engine/Source/Runtime/Renderer/Private/ScenePrivate.h
--- a/Engine/Source/Runtime/Renderer/Private/ScenePrivate.h
+++ b/Engine/Source/Runtime/Renderer/Private/ScenePrivate.h
@@ -1881,6 +1881,7 @@ public:
 	/** Used to detect atlas reallocations, since objects store UVs into the atlas and need to be updated when it changes. */
 	int32 AtlasGeneration;
 	int32 HeightFieldAtlasGeneration;
+	int32 HFVisibilityAtlasGenerattion;
 
 	bool bTrackAllPrimitives;
 	bool bCanUse16BitObjectIndices;
