UChildActorComponent::SetChildClass can now specify the template to use rather than a default version of the desired class.

diff --git a/Engine/Source/Runtime/Engine/Classes/Components/ChildActorComponent.h b/Engine/Source/Runtime/Engine/Classes/Components/ChildActorComponent.h
--- a/Engine/Source/Runtime/Engine/Classes/Components/ChildActorComponent.h
+++ b/Engine/Source/Runtime/Engine/Classes/Components/ChildActorComponent.h
@@ -57,8 +57,29 @@ class ENGINE_API UChildActorComponent : public USceneComponent
 {
 	GENERATED_UCLASS_BODY()
 
+	/**
+	 * Sets the class to use for the child actor. 
+	 * If called on a template component (owned by a CDO), the properties of any existing child actor template will be copied as best possible to the template. 
+	 * If called on a component instance in a world (and the class is changing), the created ChildActor will use the class defaults as template.
+	 * @param InClass The Actor subclass to spawn as a child actor
+	 */
 	UFUNCTION(BlueprintCallable, Category=ChildActorComponent)
-	void SetChildActorClass(TSubclassOf<AActor> InClass);
+	void SetChildActorClass(TSubclassOf<AActor> InClass)
+	{
+		SetChildActorClass(InClass, nullptr);
+	}
+
+	/**
+	 * Sets then class to use for the child actor providing an optional Actor to use as the template.
+	 * If called on a template component (owned by a CDO) and NewChildActorTemplate is not null, the new child actor template will be created using the supplied Actor as template.
+	 * If called on a template component and NewChildActorTemplate is null, the properties of any existing child actor template will be copied as best possible to the template.
+	 * If called on a component instance in a world with NewChildActorTemplate not null, then if registered a new child actor will be created using the supplied Actor as template, 
+	 *    otherwise if not registered it will ensure. If the class also changed, then future ChildActors created by this component the class defaults will be used.
+	 * If called on a component instance in a world with NewChildActorTemplate null and the class is changing, the created ChildActor will use the class defaults as template.
+	 * @param InClass                 The Actor subclass to spawn as a child actor
+	 * @param NewChildActorTemplate   An Actor to use as the template when spawning a child actor using this component (per the rules listed above)
+	 */
+	void SetChildActorClass(TSubclassOf<AActor> InClass, AActor* NewChildActorTemplate);
 
 	TSubclassOf<AActor> GetChildActorClass() const { return ChildActorClass; }
 
diff --git a/Engine/Source/Runtime/Engine/Private/Components/ChildActorComponent.cpp b/Engine/Source/Runtime/Engine/Private/Components/ChildActorComponent.cpp
--- a/Engine/Source/Runtime/Engine/Private/Components/ChildActorComponent.cpp
+++ b/Engine/Source/Runtime/Engine/Private/Components/ChildActorComponent.cpp
@@ -423,28 +423,32 @@ void UChildActorComponent::ApplyComponentInstanceData(FChildActorComponentInstan
 	}
 }
 
-void UChildActorComponent::SetChildActorClass(TSubclassOf<AActor> Class)
+void UChildActorComponent::SetChildActorClass(TSubclassOf<AActor> Class, AActor* ActorTemplate)
 {
 	ChildActorClass = Class;
 	if (IsTemplate())
 	{
 		if (ChildActorClass)
 		{
-			if (ChildActorTemplate == nullptr || (ChildActorTemplate->GetClass() != ChildActorClass))
+			if (ChildActorTemplate == nullptr || ActorTemplate || (ChildActorTemplate->GetClass() != ChildActorClass))
 			{
 				Modify();
 
-				AActor* NewChildActorTemplate = NewObject<AActor>(GetTransientPackage(), ChildActorClass, NAME_None, RF_ArchetypeObject | RF_Transactional | RF_Public);
+				AActor* NewChildActorTemplate = NewObject<AActor>(GetTransientPackage(), ChildActorClass, NAME_None, RF_ArchetypeObject | RF_Transactional | RF_Public, ActorTemplate);
 
 				if (ChildActorTemplate)
 				{
-					UEngine::CopyPropertiesForUnrelatedObjects(ChildActorTemplate, NewChildActorTemplate);
-#if WITH_EDITOR
-					NewChildActorTemplate->ClearActorLabel();
-#endif
+					if (ActorTemplate == nullptr)
+					{
+						UEngine::CopyPropertiesForUnrelatedObjects(ChildActorTemplate, NewChildActorTemplate);
+					}
 					ChildActorTemplate->Rename(nullptr, GetTransientPackage(), REN_DontCreateRedirectors);
 				}
 
+#if WITH_EDITOR
+				NewChildActorTemplate->ClearActorLabel();
+#endif
+
 				ChildActorTemplate = NewChildActorTemplate;
 
 				// Record initial object state in case we're in a transaction context.
@@ -476,8 +480,16 @@ void UChildActorComponent::SetChildActorClass(TSubclassOf<AActor> Class)
 		{
 			ChildActorName = NAME_None;
 			DestroyChildActor();
+
+			// If an actor template was supplied, temporarily set ChildActorTemplate to create the new Actor with ActorTemplate used as the template
+			TGuardValue<AActor*> ChildActorTemplateGuard(ChildActorTemplate, (ActorTemplate ? ActorTemplate : ChildActorTemplate));
+
 			CreateChildActor();
 		}
+		else if (ActorTemplate)
+		{
+			UE_LOG(LogChildActorComponent, Warning, TEXT("Call to SetChildActorClass on '%s' supplied ActorTemplate '%s', but it will not be used due to the component not being registered."), *GetPathName(), *ActorTemplate->GetPathName());
+		}
 	}
 }
 
