UE4.24 - Show error is out of disk space when compiling shaders
#rnx

diff --git a/Engine/Source/Runtime/Engine/Private/ShaderCompiler/ShaderCompiler.cpp b/Engine/Source/Runtime/Engine/Private/ShaderCompiler/ShaderCompiler.cpp
--- a/Engine/Source/Runtime/Engine/Private/ShaderCompiler/ShaderCompiler.cpp
+++ b/Engine/Source/Runtime/Engine/Private/ShaderCompiler/ShaderCompiler.cpp
@@ -364,6 +364,11 @@ namespace SCWErrorCode
 		ModalErrorOrLog(FString::Printf(TEXT("ShaderCompileWorker failed:\n%s\n"), Data));
 	}
 
+	void HandleOutputFileEmpty(const TCHAR* Filename)
+	{
+		ModalErrorOrLog(FString::Printf(TEXT("Output file %s size is 0. Are you out of disk space?"), Filename));
+	}
+
 	void HandleOutputFileCorrupted(const TCHAR* Filename, int64 ExpectedSize, int64 ActualSize)
 	{
 		ModalErrorOrLog(FString::Printf(TEXT("Output file corrupted (expected %I64d bytes, but only got %I64d): %s"), ExpectedSize, ActualSize, Filename));
@@ -770,6 +775,11 @@ PRAGMA_ENABLE_OPTIMIZATION;
 // Process results from Worker Process
 void FShaderCompileUtilities::DoReadTaskResults(const TArray<FShaderCommonCompileJob*>& QueuedJobs, FArchive& OutputFile)
 {
+	if (OutputFile.TotalSize() == 0)
+	{
+		SCWErrorCode::HandleOutputFileEmpty(*OutputFile.GetArchiveName());
+	}
+
 	int32 OutputVersion = ShaderCompileWorkerOutputVersion;
 	OutputFile << OutputVersion;
 
