Remove RigVM dependency from Engine build

diff --git a/Engine/Source/Runtime/Engine/Engine.Build.cs b/Engine/Source/Runtime/Engine/Engine.Build.cs
--- a/Engine/Source/Runtime/Engine/Engine.Build.cs
+++ b/Engine/Source/Runtime/Engine/Engine.Build.cs
@@ -86,8 +86,7 @@ public class Engine : ModuleRules
                 "NetworkReplayStreaming",
 				"PhysicsCore",
 				"PhysicsSQ",
-				"ChaosSolvers",
-				"RigVM"
+				"ChaosSolvers"
 			}
 		);
 
