COPY from //Dev-Anim to //Dev-Main

diff --git a/Engine/Source/Runtime/Core/Private/Misc/FrameRate.cpp b/Engine/Source/Runtime/Core/Private/Misc/FrameRate.cpp
--- a/Engine/Source/Runtime/Core/Private/Misc/FrameRate.cpp
+++ b/Engine/Source/Runtime/Core/Private/Misc/FrameRate.cpp
@@ -257,7 +257,7 @@ bool FFrameRate::ComputeGridSpacing(float PixelsPerSecond, double& OutMajorInter
 		OutMajorInterval  = MajorIntervalFrames * AsInterval();
 
 		// Find the lowest number of divisions we can show that's larger than the minimum tick size
-		OutMinorDivisions = 0;
+		OutMinorDivisions = MajorIntervalFrames;
 		for (int32 DivIndex = 0; DivIndex < BaseIndex; ++DivIndex)
 		{
 			if (Base % CommonBases[DivIndex] == 0)
diff --git a/Engine/Source/Runtime/Core/Public/Containers/DynamicRHIResourceArray.h b/Engine/Source/Runtime/Core/Public/Containers/DynamicRHIResourceArray.h
--- a/Engine/Source/Runtime/Core/Public/Containers/DynamicRHIResourceArray.h
+++ b/Engine/Source/Runtime/Core/Public/Containers/DynamicRHIResourceArray.h
@@ -118,9 +118,9 @@ public:
 	*
 	* @param Ar	FArchive to bulk serialize this TArray to/from
 	*/
-	void BulkSerialize(FArchive& Ar)
+	void BulkSerialize(FArchive& Ar, bool bForcePerElementSerialization = false)
 	{
-		Super::BulkSerialize(Ar);
+		Super::BulkSerialize(Ar, bForcePerElementSerialization);
 	}
 
 	/**
diff --git a/Engine/Source/Runtime/Core/Public/GenericPlatform/GenericPlatformMath.h b/Engine/Source/Runtime/Core/Public/GenericPlatform/GenericPlatformMath.h
--- a/Engine/Source/Runtime/Core/Public/GenericPlatform/GenericPlatformMath.h
+++ b/Engine/Source/Runtime/Core/Public/GenericPlatform/GenericPlatformMath.h
@@ -224,19 +224,29 @@ struct FGenericPlatformMath
 	/** Return true if value is NaN (not a number). */
 	static FORCEINLINE bool IsNaN( float A ) 
 	{
-		return ((*(uint32*)&A) & 0x7FFFFFFF) > 0x7F800000;
+		return ((*(uint32*)&A) & 0x7FFFFFFFU) > 0x7F800000U;
 	}
+	static FORCEINLINE bool IsNaN(double A)
+	{
+		return ((*(uint64*)&A) & 0x7FFFFFFFFFFFFFFFULL) > 0x7FF0000000000000ULL;
+	}
+
 	/** Return true if value is finite (not NaN and not Infinity). */
 	static FORCEINLINE bool IsFinite( float A )
 	{
-		return ((*(uint32*)&A) & 0x7F800000) != 0x7F800000;
+		return ((*(uint32*)&A) & 0x7F800000U) != 0x7F800000U;
 	}
-	static FORCEINLINE bool IsNegativeFloat(const float& A)
+	static FORCEINLINE bool IsFinite(double A)
+	{
+		return ((*(uint64*)&A) & 0x7FF0000000000000ULL) != 0x7FF0000000000000ULL;
+	}
+
+	static FORCEINLINE bool IsNegativeFloat(float A)
 	{
 		return ( (*(uint32*)&A) >= (uint32)0x80000000 ); // Detects sign bit.
 	}
 
-	static FORCEINLINE bool IsNegativeDouble(const double& A)
+	static FORCEINLINE bool IsNegativeDouble(double A)
 	{
 		return ( (*(uint64*)&A) >= (uint64)0x8000000000000000 ); // Detects sign bit.
 	}
diff --git a/Engine/Source/Runtime/Core/Public/Mac/MacPlatformMath.h b/Engine/Source/Runtime/Core/Public/Mac/MacPlatformMath.h
--- a/Engine/Source/Runtime/Core/Public/Mac/MacPlatformMath.h
+++ b/Engine/Source/Runtime/Core/Public/Mac/MacPlatformMath.h
@@ -85,7 +85,9 @@ struct FMacPlatformMath : public FClangPlatformMath
 #endif
 
 	static FORCEINLINE bool IsNaN( float A ) { return isnan(A) != 0; }
+	static FORCEINLINE bool IsNaN(double A) { return isnan(A) != 0; }
 	static FORCEINLINE bool IsFinite( float A ) { return isfinite(A); }
+	static FORCEINLINE bool IsFinite(double A) { return isfinite(A); }
 
 #if PLATFORM_ENABLE_VECTORINTRINSICS
 	static FORCEINLINE float InvSqrt(float F)
diff --git a/Engine/Source/Runtime/Core/Public/Math/RandomStream.h b/Engine/Source/Runtime/Core/Public/Math/RandomStream.h
--- a/Engine/Source/Runtime/Core/Public/Math/RandomStream.h
+++ b/Engine/Source/Runtime/Core/Public/Math/RandomStream.h
@@ -62,7 +62,7 @@ public:
 	void Initialize( int32 InSeed )
 	{
 		InitialSeed = InSeed;
-		Seed = InSeed;
+		Seed = uint32(InSeed);
 	}
 
 	/**
@@ -82,7 +82,7 @@ public:
 			InitialSeed = FPlatformTime::Cycles();
 		}
 
-		Seed = InitialSeed;
+		Seed = uint32(InitialSeed);
 	}
 
 	/**
@@ -90,7 +90,7 @@ public:
 	 */
 	void Reset() const
 	{
-		Seed = InitialSeed;
+		Seed = uint32(InitialSeed);
 	}
 
 	int32 GetInitialSeed() const
@@ -107,7 +107,7 @@ public:
 	}
 
 	/**
-	 * Returns a random number between 0 and 1.
+	 * Returns a random float number in the range [0, 1).
 	 *
 	 * @return Random number.
 	 */
@@ -115,12 +115,11 @@ public:
 	{
 		MutateSeed();
 
-		const float SRandTemp = 1.0f;
 		float Result;
 
-		*(int32*)&Result = (*(int32*)&SRandTemp & 0xff800000) | (Seed & 0x007fffff);
+		*(uint32*)&Result = 0x3F800000U | (Seed >> 9);
 
-		return FMath::Fractional(Result); 
+		return Result - 1.0f; 
 	}
 
 	/**
@@ -132,7 +131,7 @@ public:
 	{
 		MutateSeed();
 
-		return *(uint32*)&Seed;
+		return Seed;
 	}
 
 	/**
@@ -165,7 +164,7 @@ public:
 	 */
 	int32 GetCurrentSeed() const
 	{
-		return Seed;
+		return int32(Seed);
 	}
 
 	/**
@@ -185,8 +184,8 @@ public:
 	 */
 	FORCEINLINE int32 RandHelper( int32 A ) const
 	{
-		// Can't just multiply GetFraction by A, as GetFraction could be == 1.0f
-		return ((A > 0) ? FMath::TruncToInt(GetFraction() * ((float)A - DELTA)) : 0);
+		// GetFraction guarantees a result in the [0,1) range.
+		return ((A > 0) ? FMath::TruncToInt(GetFraction() * float(A)) : 0);
 	}
 
 	/**
@@ -335,7 +334,7 @@ public:
 	 */
 	FString ToString() const
 	{
-		return FString::Printf(TEXT("FRandomStream(InitialSeed=%i, Seed=%i)"), InitialSeed, Seed);
+		return FString::Printf(TEXT("FRandomStream(InitialSeed=%i, Seed=%u)"), InitialSeed, Seed);
 	}
 
 protected:
@@ -345,7 +344,7 @@ protected:
 	 */
 	void MutateSeed() const
 	{
-		Seed = (Seed * 196314165) + 907633515; 
+		Seed = (Seed * 196314165U) + 907633515U; 
 	}
 
 private:
@@ -353,6 +352,7 @@ private:
 	// Holds the initial seed.
 	int32 InitialSeed;
 
-	// Holds the current seed.
-	mutable int32 Seed;
+	// Holds the current seed. This should be an uint32 so that any shift to obtain top bits
+	// is a logical shift, rather than an arithmetic shift (which smears down the negative bit).
+	mutable uint32 Seed;
 };
diff --git a/Engine/Source/Runtime/Core/Public/Math/UnrealMathUtility.h b/Engine/Source/Runtime/Core/Public/Math/UnrealMathUtility.h
--- a/Engine/Source/Runtime/Core/Public/Math/UnrealMathUtility.h
+++ b/Engine/Source/Runtime/Core/Public/Math/UnrealMathUtility.h
@@ -247,6 +247,101 @@ struct FMath : public FPlatformMath
 		return Abs<double>( Value ) <= ErrorTolerance;
 	}
 
+private:
+	template<typename FloatType, typename IntegralType, IntegralType SignedBit>
+	static inline bool TIsNearlyEqualByULP(FloatType A, FloatType B, int32 MaxUlps)
+	{
+		// Any comparison with NaN always fails.
+		if (FMath::IsNaN(A) || FMath::IsNaN(B))
+		{
+			return false;
+		}
+
+		// If either number is infinite, then ignore ULP and do a simple equality test. 
+		// The rationale being that two infinities, of the same sign, should compare the same 
+		// no matter the ULP, but FLT_MAX and Inf should not, even if they're neighbors in
+		// their bit representation.
+		if (!FMath::IsFinite(A) || !FMath::IsFinite(B))
+		{
+			return A == B;
+		}
+
+		// Convert the integer representation of the float from sign + magnitude to
+		// a signed number representation where 0 is 1 << 31. This allows us to compare
+		// ULP differences around zero values.
+		auto FloatToSignedNumber = [](IntegralType V) {
+			if (V & SignedBit)
+			{
+				return ~V + 1;
+			}
+			else
+			{
+				return SignedBit | V;
+			}
+		};
+
+		union FFloatToInt { FloatType F; IntegralType I; };
+		FFloatToInt FloatA;
+		FFloatToInt FloatB;
+
+		FloatA.F = A;
+		FloatB.F = B;
+
+		IntegralType SNA = FloatToSignedNumber(FloatA.I);
+		IntegralType SNB = FloatToSignedNumber(FloatB.I);
+		IntegralType Distance = (SNA >= SNB) ? (SNA - SNB) : (SNB - SNA);
+		return Distance <= IntegralType(MaxUlps);
+	}
+
+public:
+
+	/**
+	 *	Check if two floating point numbers are nearly equal to within specific number of 
+	 *	units of last place (ULP). A single ULP difference between two floating point numbers
+	 *	means that they have an adjacent representation and that no other floating point number
+	 *	can be constructed to fit between them. This enables making consistent comparisons 
+	 *	based on representational distance between floating point numbers, regardless of 
+	 *	their magnitude. 
+	 *
+	 *	Use when the two numbers vary greatly in range. Otherwise, if absolute tolerance is
+	 *	required, use IsNearlyEqual instead.
+	 *  
+	 *	Note: Since IEEE 754 floating point operations are guaranteed to be exact to 0.5 ULP,
+	 *	a value of 4 ought to be sufficient for all but the most complex float operations.
+	 * 
+	 *	@param A				First number to compare
+	 *	@param B				Second number to compare
+	 *	@param MaxUlps          The maximum ULP distance by which neighboring floating point 
+	 *	                        numbers are allowed to differ.
+	 *	@return					true if the two values are nearly equal.
+	 */
+	static FORCEINLINE bool IsNearlyEqualByULP(float A, float B, int32 MaxUlps = 4)
+	{
+		return TIsNearlyEqualByULP<float, uint32, uint32(1U << 31)>(A, B, MaxUlps);
+	}
+
+	/**
+	 *	Check if two floating point numbers are nearly equal to within specific number of
+	 *	units of last place (ULP). A single ULP difference between two floating point numbers
+	 *	means that they have an adjacent representation and that no other floating point number
+	 *	can be constructed to fit between them. This enables making consistent comparisons
+	 *	based on representational distance between floating point numbers, regardless of
+	 *	their magnitude.
+	 *
+	 *	Note: Since IEEE 754 floating point operations are guaranteed to be exact to 0.5 ULP,
+	 *	a value of 4 ought to be sufficient for all but the most complex float operations.
+	 *
+	 *	@param A				First number to compare
+	 *	@param B				Second number to compare
+	 *	@param MaxUlps          The maximum ULP distance by which neighboring floating point
+	 *	                        numbers are allowed to differ.
+	 *	@return					true if the two values are nearly equal.
+	 */
+	static FORCEINLINE bool IsNearlyEqualByULP(double A, double B, int32 MaxUlps = 4)
+	{
+		return TIsNearlyEqualByULP<double, uint64, uint64(1ULL << 63)>(A, B, MaxUlps);
+	}
+
 	/**
 	 *	Checks whether a number is a power of two.
 	 *	@param Value	Number to check
diff --git a/Engine/Source/Runtime/Core/Public/UObject/AnimObjectVersion.h b/Engine/Source/Runtime/Core/Public/UObject/AnimObjectVersion.h
--- a/Engine/Source/Runtime/Core/Public/UObject/AnimObjectVersion.h
+++ b/Engine/Source/Runtime/Core/Public/UObject/AnimObjectVersion.h
@@ -18,6 +18,20 @@ struct CORE_API FAnimObjectVersion
 		// Cached marker sync names on skeleton for editor
 		StoreMarkerNamesOnSkeleton,
 
+		// Serialized register array state for RigVM
+		SerializeRigVMRegisterArrayState,
+
+		// Increase number of bones per chunk from uint8 to uint16
+		IncreaseBoneIndexLimitPerChunk,
+
+		UnlimitedBoneInfluences,
+
+		// Anim sequences have colors for their curves
+		AnimSequenceCurveColors,
+
+		// Notifies and sync markers now have Guids
+		NotifyAndSyncMarkerGuids,
+
 		// -----<new versions can be added above this line>-------------------------------------------------
 		VersionPlusOne,
 		LatestVersion = VersionPlusOne - 1
diff --git a/Engine/Source/Runtime/Core/Public/Unix/UnixPlatformMath.h b/Engine/Source/Runtime/Core/Public/Unix/UnixPlatformMath.h
--- a/Engine/Source/Runtime/Core/Public/Unix/UnixPlatformMath.h
+++ b/Engine/Source/Runtime/Core/Public/Unix/UnixPlatformMath.h
@@ -88,7 +88,9 @@ struct FLinuxPlatformMath : public FClangPlatformMath
 #endif
 
 	static FORCEINLINE bool IsNaN( float A ) { return isnan(A) != 0; }
+	static FORCEINLINE bool IsNaN(double A) { return isnan(A) != 0; }
 	static FORCEINLINE bool IsFinite( float A ) { return isfinite(A); }
+	static FORCEINLINE bool IsFinite(double A) { return isfinite(A); }
 
 	static FORCEINLINE float InvSqrt( float F )
 	{
diff --git a/Engine/Source/Runtime/CoreUObject/Public/UObject/ObjectMacros.h b/Engine/Source/Runtime/CoreUObject/Public/UObject/ObjectMacros.h
--- a/Engine/Source/Runtime/CoreUObject/Public/UObject/ObjectMacros.h
+++ b/Engine/Source/Runtime/CoreUObject/Public/UObject/ObjectMacros.h
@@ -1157,6 +1157,9 @@ namespace UM
 		/// [PropertyMetadata] Deprecated.
 		FixedIncrement,
 
+		/// [PropertyMetadata] Used by asset properties. Indicates that the asset pickers should always show engine content
+		ForceShowEngineContent,
+
 		/// [PropertyMetadata] Used for FColor and FLinearColor properties. Indicates that the Alpha property should be hidden when displaying the property widget in the details.
 		HideAlphaChannel,
 
diff --git a/Engine/Source/Runtime/Engine/Classes/Animation/AnimCompress.h b/Engine/Source/Runtime/Engine/Classes/Animation/AnimCompress.h
--- a/Engine/Source/Runtime/Engine/Classes/Animation/AnimCompress.h
+++ b/Engine/Source/Runtime/Engine/Classes/Animation/AnimCompress.h
@@ -12,6 +12,7 @@
 #include "UObject/Object.h"
 #include "Templates/Function.h"
 #include "Animation/AnimSequence.h"
+#include "Animation/AnimBoneCompressionCodec.h"
 #include "AnimationUtils.h"
 #include "AnimEnums.h"
 #include "AnimationCompression.h"
@@ -74,6 +75,11 @@ public:
 		return Items[0];
 	}
 
+	bool IsValid() const
+	{
+		return Items.Num() > 0;
+	}
+
 private:
 	//Storage of tracked items
 	TArray<DataType> Items;
@@ -185,15 +191,16 @@ class ENGINE_API FCompressionMemorySummary
 public:
 	FCompressionMemorySummary(bool bInEnabled);
 
-	void GatherPreCompressionStats(const FCompressibleAnimData& CompressibleAnimData, int32 PreCompressedSize, int32 ProgressNumerator, int32 ProgressDenominator);
+	void GatherPreCompressionStats(const FString& Name, int32 RawSize, int32 PreviousCompressionSize, int32 ProgressNumerator, int32 ProgressDenominator);
 
-	void GatherPostCompressionStats(const FCompressibleAnimData& CompressibleAnimData, FCompressibleAnimDataResult& CompressedData, double CompressionTime);
+	void GatherPostCompressionStats(const FCompressedAnimSequence& CompressedData, const TArray<FBoneData>& BoneData, const FName AnimFName, double CompressionTime, bool bInPerformedCompression);
 
 	~FCompressionMemorySummary();
 
 private:
 	bool bEnabled;
 	bool bUsed;
+	bool bPerformedCompression;
 	int32 TotalRaw;
 	int32 TotalBeforeCompressed;
 	int32 TotalAfterCompressed;
@@ -222,9 +229,9 @@ struct ENGINE_API FAnimCompressContext
 private:
 	FCompressionMemorySummary	CompressionSummary;
 
-	void GatherPreCompressionStats(const FCompressibleAnimData& CompressibleAnimData, int32 PreviousCompressionSize);
+	void GatherPreCompressionStats(const FString& Name, int32 RawSize, int32 PreviousCompressionSize);
 
-	void GatherPostCompressionStats(const FCompressibleAnimData& CompressibleAnimData, FCompressibleAnimDataResult& CompressedData, double CompressionTime);
+	void GatherPostCompressionStats(const FCompressedAnimSequence& CompressedData, const TArray<FBoneData>& BoneData, const FName AnimFName, double CompressionTime, bool bInPerformedCompression);
 
 
 public:
@@ -252,86 +259,48 @@ public:
 
 	friend class FAnimationUtils;
 	friend class FDerivedDataAnimationCompression;
+	friend class UAnimSequence;
 };
 
-#if USE_SEGMENTING_CONTEXT
-//////////////////////////////////////////////////////////////////////////
-// FAnimSegmentContext - This holds the relevant intermediate information
-// when compressing animation sequence segments.
-struct FAnimSegmentContext
-{
-	int32 StartFrame;
-	int32 NumFrames;
-
-	TArray<struct FTranslationTrack> TranslationData;
-	TArray<struct FRotationTrack> RotationData;
-	TArray<struct FScaleTrack> ScaleData;
-
-	AnimationCompressionFormat TranslationCompressionFormat;
-	AnimationCompressionFormat RotationCompressionFormat;
-	AnimationCompressionFormat ScaleCompressionFormat;
-
-	TArray<int32> CompressedTrackOffsets;
-	FCompressedOffsetData CompressedScaleOffsets;
-	TArray<uint8> CompressedByteStream;
-	TArray<uint8> CompressedTrivialTracksByteStream;
-};
-#endif
-
 UCLASS(abstract, hidecategories=Object, MinimalAPI, EditInlineNew)
-class UAnimCompress : public UObject
+class UAnimCompress : public UAnimBoneCompressionCodec
 {
 	GENERATED_UCLASS_BODY()
 
-	/** Name of Compression Scheme used for this asset*/
-	UPROPERTY(Category=Compression, VisibleAnywhere)
-	FString Description;
-
 	/** Compression algorithms requiring a skeleton should set this value to true. */
 	UPROPERTY()
 	uint32 bNeedsSkeleton:1;
 
-	/** Whether to enable segmenting or not. Needed for USE_SEGMENTING_CONTEXT (currently turned off) */
-	//UPROPERTY(Category = Compression, EditAnywhere)
-	UPROPERTY()
-	uint32 bEnableSegmenting:1;
-
-	/** When splitting the sequence into segments, we will try to approach this value as much as possible. */
-	UPROPERTY(Category = Compression, EditAnywhere, meta = (ClampMin = "16"))
-	uint32 IdealNumFramesPerSegment;
-
-	/** When splitting the sequence into segments, we will never allow a segment with more than the maximum number of frames. */
-	UPROPERTY(Category = Compression, EditAnywhere)
-	uint32 MaxNumFramesPerSegment;
-
 	/** Format for bitwise compression of translation data. */
-	UPROPERTY()
+	UPROPERTY(Category = Compression, EditAnywhere)
 	TEnumAsByte<AnimationCompressionFormat> TranslationCompressionFormat;
 
 	/** Format for bitwise compression of rotation data. */
-	UPROPERTY()
+	UPROPERTY(Category = Compression, EditAnywhere)
 	TEnumAsByte<AnimationCompressionFormat> RotationCompressionFormat;
 
 	/** Format for bitwise compression of scale data. */
-	UPROPERTY()
+	UPROPERTY(Category = Compression, EditAnywhere)
 	TEnumAsByte<AnimationCompressionFormat> ScaleCompressionFormat;
 
 	/** Max error for compression of curves using remove redundant keys */
 	UPROPERTY(Category = Compression, EditAnywhere)
 	float MaxCurveError;
 
-#if WITH_EDITOR
 public:
-	/**
-	 * Reduce the number of keyframes and bitwise compress all sequences in the specified array.
-	 *
-	 * @param	AnimSequences	The animations to compress.
-	 * @param	bOutput			If false don't generate output or compute memory savings.
-	 * @return					false if a skeleton was needed by the algorithm but not provided.
-	 */
-	ENGINE_API bool Reduce(const FCompressibleAnimData& CompressibleAnimData, FCompressibleAnimDataResult& OutResult);
+#if WITH_EDITORONLY_DATA
+	/** UAnimBoneCompressionCodec implementation */
+	virtual bool Compress(const FCompressibleAnimData& CompressibleAnimData, FCompressibleAnimDataResult& OutResult) override;
+#endif
+
+	virtual TUniquePtr<ICompressedAnimData> AllocateAnimData() const override;
+
+	virtual void ByteSwapIn(ICompressedAnimData& AnimData, TArrayView<uint8> CompressedData, FMemoryReader& MemoryStream) const override;
+	virtual void ByteSwapOut(ICompressedAnimData& AnimData, TArrayView<uint8> CompressedData, FMemoryWriter& MemoryStream) const override;
+
+	virtual void DecompressPose(FAnimSequenceDecompressionContext& DecompContext, const BoneTrackArray& RotationPairs, const BoneTrackArray& TranslationPairs, const BoneTrackArray& ScalePairs, TArrayView<FTransform>& OutAtoms) const override;
+	virtual void DecompressBone(FAnimSequenceDecompressionContext& DecompContext, int32 TrackIndex, FTransform& OutAtom) const override;
 
-#endif // WITH_EDITOR
 protected:
 #if WITH_EDITOR
 	/**
@@ -340,8 +309,9 @@ protected:
 	 *
 	 * @return		true if the keyframe reduction was successful.
 	 */
-	virtual void DoReduction(const FCompressibleAnimData& CompressibleAnimData, FCompressibleAnimDataResult& OutResult) PURE_VIRTUAL(UAnimCompress::DoReduction,);
+	virtual bool DoReduction(const FCompressibleAnimData& CompressibleAnimData, FCompressibleAnimDataResult& OutResult) PURE_VIRTUAL(UAnimCompress::DoReduction,return false;);
 #endif // WITH_EDITOR
+
 	/**
 	 * Common compression utility to remove 'redundant' position keys based on the provided delta threshold
 	 *
@@ -421,22 +391,6 @@ protected:
 		float MaxRotDelta, 
 		float MaxScaleDelta);
 
-	/**
-	 * Same as above function but it will execute for every segment.
-	 *
-	 * @param	RawSegments		Array of segments being compressed
-	 * @param	MaxPosDelta		Maximum local-space threshold for stationary motion
-	 * @param	MaxRotDelta		Maximum angle threshold to consider stationary motion
-	 * @param	MaxScaleDelta	Maximum scale threshold to consider stationary motion
-	 */
-#if USE_SEGMENTING_CONTEXT
-	static void FilterTrivialKeys(
-		TArray<FAnimSegmentContext>& RawSegments,
-		float MaxPosDelta,
-		float MaxRotDelta,
-		float MaxScaleDelta);
-#endif
-
 	/**
 	 * Common compression utility to retain only intermittent position keys. For example,
 	 * calling with an Interval of 3 would keep every third key in the set and discard the rest
@@ -521,27 +475,6 @@ protected:
 		TArray<struct FRotationTrack>& OutRotationData, 
 		TArray<struct FScaleTrack>& OutScaleData);
 
-	/**
-	 * This will populate segments from the raw animation data.
-	 *
-	 * @param	AnimSeq						AnimSequence being compressed
-	 * @param	TranslationData				Translation tracks
-	 * @param	RotationData				Rotation tracks
-	 * @param	ScaleData					Scale tracks
-	 * @param	IdealNumFramesPerSegment	The ideal number of frames within a segment
-	 * @param	MaxNumFramesPerSegment		The maximum number of frames within a segment
-	 * @param	OutRawSegments				The array of segments to populate
-	 */
-#if USE_SEGMENTING_CONTEXT
-	static void SeparateRawDataIntoTracks(
-		const UAnimSequence& AnimSeq,
-		const TArray<struct FTranslationTrack>& TranslationData,
-		const TArray<struct FRotationTrack>& RotationData,
-		const TArray<struct FScaleTrack>& ScaleData,
-		int32 IdealNumFramesPerSegment,
-		int32 MaxNumFramesPerSegment,
-		TArray<FAnimSegmentContext>& OutRawSegments);
-#endif
 	/**
 	 * Common compression utility to walk an array of rotation tracks and enforce
 	 * that all adjacent rotation keys are represented by shortest-arc quaternion pairs.
@@ -575,63 +508,6 @@ public:
 		const TArray<FScaleTrack>& ScaleData,
 		bool IncludeKeyTable = false);
 
-	/**
-	 * Encodes individual key arrays into an AnimSequence using the desired bit packing formats.
-	 * This action is performed for every segment supplied.
-	 *
-	 * @param	AnimSeq						Animation Sequence which will contain the bit-packed data.
-	 * @param	TargetTranslationFormat		The format to use when encoding translation keys.
-	 * @param	TargetRotationFormat		The format to use when encoding rotation keys.
-	 * @param	TargetScaleFormat			The format to use when encoding scale keys.
-	 * @param	RawSegments					The array of segments to compress
-	 * @param	bIsSorted					For variable interpolation, is the compressed data sorted or not?
-	 */
-#if USE_SEGMENTING_CONTEXT
-	static void BitwiseCompressAnimationTracks(
-		UAnimSequence& AnimSeq,
-		AnimationCompressionFormat TargetTranslationFormat,
-		AnimationCompressionFormat TargetRotationFormat,
-		AnimationCompressionFormat TargetScaleFormat,
-		TArray<FAnimSegmentContext>& RawSegments,
-		bool bIsSorted = false);
-
-	/**
-	 * Encodes individual key arrays into an AnimSequence using the desired bit packing formats.
-	 * This action is performed for a single segment.
-	 *
-	 * @param	AnimSeq						Animation Sequence which will contain the bit-packed data.
-	 * @param	TargetTranslationFormat		The format to use when encoding translation keys.
-	 * @param	TargetRotationFormat		The format to use when encoding rotation keys.
-	 * @param	TargetScaleFormat			The format to use when encoding scale keys.
-	 * @param	RawSegment					The segment to compress
-	 * @param	bIsSorted					For variable interpolation, is the compressed data sorted or not?
-	 */
-	static void BitwiseCompressAnimationTracks(
-		const UAnimSequence& AnimSeq,
-		AnimationCompressionFormat TargetTranslationFormat,
-		AnimationCompressionFormat TargetRotationFormat,
-		AnimationCompressionFormat TargetScaleFormat,
-		FAnimSegmentContext& RawSegment,
-		bool bIsSorted = false);
-
-	/**
-	 * Encodes the trivial tracks within a segment.
-	 *
-	 * @param	AnimSeq						Animation Sequence which will contain the bit-packed data.
-	 * @param	RawSegment					The segment to compress
-	 */
-	static void BitwiseCompressTrivialAnimationTracks(const UAnimSequence& AnimSeq, FAnimSegmentContext& RawSegment);
-
-	/**
-	 * Coalesces the compressed data from every segment into a single contiguous array stored within the sequence.
-	 *
-	 * @param	AnimSeq						Animation Sequence which will contain the bit-packed data.
-	 * @param	RawSegments					The array of segments to coalesce
-	 * @param	bIsSorted					For variable interpolation, is the compressed data sorted or not?
-	 */
-	static void CoalesceCompressedSegments(UAnimSequence& AnimSeq, const TArray<FAnimSegmentContext>& RawSegments, bool bIsSorted = false);
-#endif
-
 #if WITH_EDITOR
 	void PopulateDDCKeyArchive(FArchive& Ar) { PopulateDDCKey(Ar); }
 
@@ -639,19 +515,6 @@ protected:
 	virtual void PopulateDDCKey(FArchive& Ar);
 #endif // WITH_EDITOR
 
-	/**
-	 * Structure that holds the range min/extent for a track.
-	 */
-	struct FAnimTrackRange
-	{
-		FVector RotMin;
-		FVector RotExtent;
-		FVector TransMin;
-		FVector TransExtent;
-		FVector ScaleMin;
-		FVector ScaleExtent;
-	};
-
 	/**
 	 * Utility function to append data to a byte stream.
 	 *
@@ -703,206 +566,6 @@ protected:
 		const float* Mins,
 		const float* Ranges);
 
-#if USE_SEGMENTING_CONTEXT
-	/**
-	 * Utility function that performs minimal sanity checks.
-	 *
-	 * @param	AnimSeq						The anim sequence to check
-	 * @param	Segment						The segment to check
-	 */
-	static void SanityCheckTrackData(const UAnimSequence& AnimSeq, const FAnimSegmentContext& Segment);
-#endif
-
-	/**
-	 * Calculates the translation track range.
-	 *
-	 * @param	TranslationData				The translation track data
-	 * @param	Format						Compression format
-	 * @param	OutMin						The output range minimum
-	 * @param	OutExtent					The output range extent
-	 */
-	static void CalculateTrackRange(const FTranslationTrack& TranslationData, AnimationCompressionFormat Format, FVector& OutMin, FVector& OutExtent);
-
-	/**
-	 * Calculates the rotation track range.
-	 *
-	 * @param	TranslationData				The rotation track data
-	 * @param	Format						Compression format
-	 * @param	OutMin						The output range minimum
-	 * @param	OutExtent					The output range extent
-	 */
-	static void CalculateTrackRange(const FRotationTrack& RotationData, AnimationCompressionFormat Format, FVector& OutMin, FVector& OutExtent);
-
-	/**
-	 * Calculates the sacle track range.
-	 *
-	 * @param	TranslationData				The scale track data
-	 * @param	Format						Compression format
-	 * @param	OutMin						The output range minimum
-	 * @param	OutExtent					The output range extent
-	 */
-	static void CalculateTrackRange(const FScaleTrack& ScaleData, AnimationCompressionFormat Format, FVector& OutMin, FVector& OutExtent);
-
-#if USE_SEGMENTING_CONTEXT
-	/**
-	 * Calculates the track ranges within a segment.
-	 *
-	 * @param	TargetTranslationFormat		Compression format for translations
-	 * @param	TargetRotationFormat		Compression format for rotations
-	 * @param	TargetScaleFormat			Compression format for scales
-	 * @param	Segment						The segment that contains our tracks to process
-	 * @param	TrackRanges					The calculated track ranges
-	 */
-	static void CalculateTrackRanges(
-		AnimationCompressionFormat TargetTranslationFormat,
-		AnimationCompressionFormat TargetRotationFormat,
-		AnimationCompressionFormat TargetScaleFormat,
-		const FAnimSegmentContext& Segment,
-		TArray<FAnimTrackRange>& TrackRanges);
-#endif
-
-	/**
-	 * Structure to wrap and represent track key flags.
-	 */
-	struct FTrackKeyFlags
-	{
-		constexpr FTrackKeyFlags() : Flags(0) {}
-		constexpr explicit FTrackKeyFlags(uint8 InFlags) : Flags(InFlags) {}
-
-		constexpr bool IsComponentNeededX() const { return (Flags & 0x1) != 0; }
-		constexpr bool IsComponentNeededY() const { return (Flags & 0x2) != 0; }
-		constexpr bool IsComponentNeededZ() const { return (Flags & 0x4) != 0; }
-
-		constexpr bool IsValid() const { return (Flags & ~0x7) == 0; }
-
-		uint8 Flags;
-	};
-
-#if USE_SEGMENTING_CONTEXT
-	/**
-	 * Writes the necessary track ranges to a byte stream.
-	 *
-	 * @param	ByteStream					Byte stream to append to
-	 * @param	GetTranslationFormatFun		Function that returns the translation format for a track index
-	 * @param	GetRotationFormatFun		Function that returns the rotation format for a track index
-	 * @param	GetScaleFormatFun			Function that returns the scale format for a track index
-	 * @param	GetTranslationFlagsFun		Function that returns the translation flags for a track index
-	 * @param	GetRotationFlagsFun			Function that returns the rotation flags for a track index
-	 * @param	GetScaleFlagsFun			Function that returns the scale flags for a track index
-	 * @param	Segment						Segment to process
-	 * @param	TrackRanges					Track ranges to pack
-	 * @param	bInterleaveValues			Whether to interleave range values or not
-	 */
-	static void WriteTrackRanges(
-		TArray<uint8>& ByteStream,
-		TFunction<AnimationCompressionFormat(int32 TrackIndex)> GetTranslationFormatFun,
-		TFunction<AnimationCompressionFormat(int32 TrackIndex)> GetRotationFormatFun,
-		TFunction<AnimationCompressionFormat(int32 TrackIndex)> GetScaleFormatFun,
-		TFunction<FTrackKeyFlags(int32 TrackIndex)> GetTranslationFlagsFun,
-		TFunction<FTrackKeyFlags(int32 TrackIndex)> GetRotationFlagsFun,
-		TFunction<FTrackKeyFlags(int32 TrackIndex)> GetScaleFlagsFun,
-		const FAnimSegmentContext& Segment,
-		const TArray<FAnimTrackRange>& TrackRanges,
-		bool bInterleaveValues);
-
-	/**
-	 * Writes a segment's uniform track data to a byte stream.
-	 * A track's data is uniform if no keys are removed.
-	 *
-	 * @param	ByteStream					Byte stream to append to
-	 * @param	GetTranslationFormatFun		Function that returns the translation format for a track index
-	 * @param	GetRotationFormatFun		Function that returns the rotation format for a track index
-	 * @param	GetScaleFormatFun			Function that returns the scale format for a track index
-	 * @param	IsTranslationUniformFun		Function that returns if a track translation's data is uniform or not for a track index
-	 * @param	IsRotationUniformFun		Function that returns if a track rotation's data is uniform or not for a track index
-	 * @param	IsScaleUniformFun			Function that returns if a track scale's data is uniform or not for a track index
-	 * @param	PackTranslationKeyFun		Function that packs a translation key for a track index
-	 * @param	PackRotationKeyFun			Function that packs a rotation key for a track index
-	 * @param	PackScaleKeyFun				Function that packs a scale key for a track index
-	 * @param	Segment						Segment to process
-	 * @param	TrackRanges					Track ranges for the segment to process
-	 */
-	static void WriteUniformTrackData(
-		TArray<uint8>& ByteStream,
-		TFunction<AnimationCompressionFormat(int32 TrackIndex)> GetTranslationFormatFun,
-		TFunction<AnimationCompressionFormat(int32 TrackIndex)> GetRotationFormatFun,
-		TFunction<AnimationCompressionFormat(int32 TrackIndex)> GetScaleFormatFun,
-		TFunction<bool(int32 TrackIndex)> IsTranslationUniformFun,
-		TFunction<bool(int32 TrackIndex)> IsRotationUniformFun,
-		TFunction<bool(int32 TrackIndex)> IsScaleUniformFun,
-		TFunction<void(TArray<uint8>& ByteStream, AnimationCompressionFormat Format, const FVector& Key, const float* Mins, const float* Ranges, int32 TrackIndex)> PackTranslationKeyFun,
-		TFunction<void(TArray<uint8>& ByteStream, AnimationCompressionFormat Format, const FQuat& Key, const float* Mins, const float* Ranges, int32 TrackIndex)> PackRotationKeyFun,
-		TFunction<void(TArray<uint8>& ByteStream, AnimationCompressionFormat Format, const FVector& Key, const float* Mins, const float* Ranges, int32 TrackIndex)> PackScaleKeyFun,
-		const FAnimSegmentContext& Segment,
-		const TArray<FAnimTrackRange>& TrackRanges);
-
-	/**
-	 * Writes a segment's variable track data to a byte stream with a sorted ordering.
-	 * A track's data is variable if some keys are removed.
-	 *
-	 * @param	ByteStream					Byte stream to append to
-	 * @param	AnimSeq						The anim sequence to process
-	 * @param	GetTranslationFormatFun		Function that returns the translation format for a track index
-	 * @param	GetRotationFormatFun		Function that returns the rotation format for a track index
-	 * @param	GetScaleFormatFun			Function that returns the scale format for a track index
-	 * @param	IsTranslationUniformFun		Function that returns if a track translation's data is uniform or not for a track index
-	 * @param	IsRotationUniformFun		Function that returns if a track rotation's data is uniform or not for a track index
-	 * @param	IsScaleUniformFun			Function that returns if a track scale's data is uniform or not for a track index
-	 * @param	PackTranslationKeyFun		Function that packs a translation key for a track index
-	 * @param	PackRotationKeyFun			Function that packs a rotation key for a track index
-	 * @param	PackScaleKeyFun				Function that packs a scale key for a track index
-	 * @param	Segment						Segment to process
-	 * @param	TrackRanges					Track ranges for the segment to process
-	 */
-	static void WriteSortedVariableTrackData(
-		TArray<uint8>& ByteStream,
-		const UAnimSequence& AnimSeq,
-		TFunction<AnimationCompressionFormat(int32 TrackIndex)> GetTranslationFormatFun,
-		TFunction<AnimationCompressionFormat(int32 TrackIndex)> GetRotationFormatFun,
-		TFunction<AnimationCompressionFormat(int32 TrackIndex)> GetScaleFormatFun,
-		TFunction<bool(int32 TrackIndex)> IsTranslationVariableFun,
-		TFunction<bool(int32 TrackIndex)> IsRotationVariableFun,
-		TFunction<bool(int32 TrackIndex)> IsScaleVariableFun,
-		TFunction<void(TArray<uint8>& ByteStream, AnimationCompressionFormat Format, const FVector& Key, const float* Mins, const float* Ranges, int32 TrackIndex)> PackTranslationKeyFun,
-		TFunction<void(TArray<uint8>& ByteStream, AnimationCompressionFormat Format, const FQuat& Key, const float* Mins, const float* Ranges, int32 TrackIndex)> PackRotationKeyFun,
-		TFunction<void(TArray<uint8>& ByteStream, AnimationCompressionFormat Format, const FVector& Key, const float* Mins, const float* Ranges, int32 TrackIndex)> PackScaleKeyFun,
-		FAnimSegmentContext& Segment,
-		const TArray<FAnimTrackRange>& TrackRanges);
-
-	/**
-	 * Writes a segment's variable track data to a byte stream with a linear ordering.
-	 * A track's data is variable if some keys are removed.
-	 *
-	 * @param	ByteStream					Byte stream to append to
-	 * @param	AnimSeq						The anim sequence to process
-	 * @param	GetTranslationFormatFun		Function that returns the translation format for a track index
-	 * @param	GetRotationFormatFun		Function that returns the rotation format for a track index
-	 * @param	GetScaleFormatFun			Function that returns the scale format for a track index
-	 * @param	IsTranslationUniformFun		Function that returns if a track translation's data is uniform or not for a track index
-	 * @param	IsRotationUniformFun		Function that returns if a track rotation's data is uniform or not for a track index
-	 * @param	IsScaleUniformFun			Function that returns if a track scale's data is uniform or not for a track index
-	 * @param	PackTranslationKeyFun		Function that packs a translation key for a track index
-	 * @param	PackRotationKeyFun			Function that packs a rotation key for a track index
-	 * @param	PackScaleKeyFun				Function that packs a scale key for a track index
-	 * @param	Segment						Segment to process
-	 * @param	TrackRanges					Track ranges for the segment to process
-	 */
-	static void WriteLinearVariableTrackData(
-		TArray<uint8>& ByteStream,
-		const UAnimSequence& AnimSeq,
-		TFunction<AnimationCompressionFormat(int32 TrackIndex)> GetTranslationFormatFun,
-		TFunction<AnimationCompressionFormat(int32 TrackIndex)> GetRotationFormatFun,
-		TFunction<AnimationCompressionFormat(int32 TrackIndex)> GetScaleFormatFun,
-		TFunction<bool(int32 TrackIndex)> IsTranslationVariableFun,
-		TFunction<bool(int32 TrackIndex)> IsRotationVariableFun,
-		TFunction<bool(int32 TrackIndex)> IsScaleVariableFun,
-		TFunction<void(TArray<uint8>& ByteStream, AnimationCompressionFormat Format, const FVector& Key, const float* Mins, const float* Ranges, int32 TrackIndex)> PackTranslationKeyFun,
-		TFunction<void(TArray<uint8>& ByteStream, AnimationCompressionFormat Format, const FQuat& Key, const float* Mins, const float* Ranges, int32 TrackIndex)> PackRotationKeyFun,
-		TFunction<void(TArray<uint8>& ByteStream, AnimationCompressionFormat Format, const FVector& Key, const float* Mins, const float* Ranges, int32 TrackIndex)> PackScaleKeyFun,
-		const FAnimSegmentContext& Segment,
-		const TArray<FAnimTrackRange>& TrackRanges);
-#endif
-
 	/**
 	 * Pads a byte stream to force a particular alignment for the data to follow.
 	 *
diff --git a/Engine/Source/Runtime/Engine/Classes/Animation/AnimCompress_BitwiseCompressOnly.h b/Engine/Source/Runtime/Engine/Classes/Animation/AnimCompress_BitwiseCompressOnly.h
--- a/Engine/Source/Runtime/Engine/Classes/Animation/AnimCompress_BitwiseCompressOnly.h
+++ b/Engine/Source/Runtime/Engine/Classes/Animation/AnimCompress_BitwiseCompressOnly.h
@@ -21,7 +21,7 @@ class UAnimCompress_BitwiseCompressOnly : public UAnimCompress
 protected:
 	//~ Begin UAnimCompress Interface
 #if WITH_EDITOR
-	virtual void DoReduction(const FCompressibleAnimData& CompressibleAnimData, FCompressibleAnimDataResult& OutResult) override;
+	virtual bool DoReduction(const FCompressibleAnimData& CompressibleAnimData, FCompressibleAnimDataResult& OutResult) override;
 #endif // WITH_EDITOR
 	//~ Begin UAnimCompress Interface
 };
diff --git a/Engine/Source/Runtime/Engine/Classes/Animation/AnimCompress_LeastDestructive.h b/Engine/Source/Runtime/Engine/Classes/Animation/AnimCompress_LeastDestructive.h
--- a/Engine/Source/Runtime/Engine/Classes/Animation/AnimCompress_LeastDestructive.h
+++ b/Engine/Source/Runtime/Engine/Classes/Animation/AnimCompress_LeastDestructive.h
@@ -10,21 +10,11 @@
 #include "CoreMinimal.h"
 #include "UObject/ObjectMacros.h"
 #include "Animation/AnimCompress.h"
+#include "Animation/AnimCompress_BitwiseCompressOnly.h"
 #include "AnimCompress_LeastDestructive.generated.h"
 
 UCLASS()
-class UAnimCompress_LeastDestructive : public UAnimCompress
+class UAnimCompress_LeastDestructive : public UAnimCompress_BitwiseCompressOnly
 {
 	GENERATED_UCLASS_BODY()
-
-
-protected:
-	//~ Begin UAnimCompress Interface
-#if WITH_EDITOR
-	virtual void DoReduction(const FCompressibleAnimData& CompressibleAnimData, FCompressibleAnimDataResult& OutResult) override;
-#endif // WITH_EDITOR
-	//~ Begin UAnimCompress Interface
 };
-
-
-
diff --git a/Engine/Source/Runtime/Engine/Classes/Animation/AnimCompress_PerTrackCompression.h b/Engine/Source/Runtime/Engine/Classes/Animation/AnimCompress_PerTrackCompression.h
--- a/Engine/Source/Runtime/Engine/Classes/Animation/AnimCompress_PerTrackCompression.h
+++ b/Engine/Source/Runtime/Engine/Classes/Animation/AnimCompress_PerTrackCompression.h
@@ -130,16 +130,13 @@ class UAnimCompress_PerTrackCompression : public UAnimCompress_RemoveLinearKeys
 
 
 public:
-	/**
-	 * Cached metastructures used within DoReduction, tied to a particular sequence and mesh
-	 */
-	struct FPerTrackCachedInfo* PerReductionCachedData;
+	virtual void DecompressBone(FAnimSequenceDecompressionContext& DecompContext, int32 TrackIndex, FTransform& OutAtom) const override;
 
 protected:
 	//~ Begin UAnimCompress Interface
 #if WITH_EDITOR
 	virtual void PostEditChangeProperty(struct FPropertyChangedEvent& PropertyChangedEvent) override;
-	virtual void DoReduction(const FCompressibleAnimData& CompressibleAnimData, FCompressibleAnimDataResult& OutResult) override;
+	virtual bool DoReduction(const FCompressibleAnimData& CompressibleAnimData, FCompressibleAnimDataResult& OutResult) override;
 	virtual void PopulateDDCKey(FArchive& Ar) override;
 #endif // WITH_EDITOR
 	//~ Begin UAnimCompress Interface
@@ -154,50 +151,12 @@ protected:
 		const TArray<FScaleTrack>& ScaleData,
 		const bool bFinalPass) override;
 
-#if USE_SEGMENTING_CONTEXT
-	virtual void CompressUsingUnderlyingCompressor(
-		UAnimSequence& AnimSeq,
-		const TArray<FBoneData>& BoneData,
-		TArray<FAnimSegmentContext>& RawSegments,
-		const bool bFinalPass) override;
-#endif
-
-	/**
-	 * Performs the per track compression optimization for a single segment.
-	 * This can be called from multiple threads concurrently.
-	 */
-	void OptimizeSegmentTracks(struct FOptimizeSegmentTracksContext& Context) const;
-
-	virtual void FilterBeforeMainKeyRemoval(
+	virtual void* FilterBeforeMainKeyRemoval(
 		const FCompressibleAnimData& CompressibleAnimData,
 		TArray<FTranslationTrack>& TranslationData,
 		TArray<FRotationTrack>& RotationData,
 		TArray<FScaleTrack>& ScaleData) override;
 	//~ End UAnimCompress_RemoveLinearKeys Interface
-
-	/**
-	* Structure to hold the track format information that we calculate is most optimal.
-	*/
-	struct FPerTrackFormat
-	{
-		AnimationCompressionFormat RotationFormat;
-		AnimationCompressionFormat TranslationFormat;
-		AnimationCompressionFormat ScaleFormat;
-
-		bool bHasRotationTimeMarkers;
-		bool bHasTranslationTimeMarkers;
-		bool bHasScaleTimeMarkers;
-
-		FTrackKeyFlags RotationKeyFlags;
-		FTrackKeyFlags TranslationKeyFlags;
-		FTrackKeyFlags ScaleKeyFlags;
-	};
-
-	static void PackTranslationKey(TArray<uint8>& ByteStream, AnimationCompressionFormat Format, const FVector& Key, const float* Mins, const float* Ranges, const struct FPerTrackFormat& TrackFormat);
-	static void PackRotationKey(TArray<uint8>& ByteStream, AnimationCompressionFormat Format, const FQuat& Key, const float* Mins, const float* Ranges, const struct FPerTrackFormat& TrackFormat);
-	static void PackScaleKey(TArray<uint8>& ByteStream, AnimationCompressionFormat Format, const FVector& Key, const float* Mins, const float* Ranges, const struct FPerTrackFormat& TrackFormat);
-
-	friend struct FAsyncOptimizeSegmentTracksTaskGroupContext;
 #endif // WITH_EDITOR
 };
 
diff --git a/Engine/Source/Runtime/Engine/Classes/Animation/AnimCompress_RemoveEverySecondKey.h b/Engine/Source/Runtime/Engine/Classes/Animation/AnimCompress_RemoveEverySecondKey.h
--- a/Engine/Source/Runtime/Engine/Classes/Animation/AnimCompress_RemoveEverySecondKey.h
+++ b/Engine/Source/Runtime/Engine/Classes/Animation/AnimCompress_RemoveEverySecondKey.h
@@ -32,7 +32,7 @@ class UAnimCompress_RemoveEverySecondKey : public UAnimCompress
 protected:
 	//~ Begin UAnimCompress Interface
 #if WITH_EDITOR
-	virtual void DoReduction(const FCompressibleAnimData& CompressibleAnimData, FCompressibleAnimDataResult& OutResult) override;
+	virtual bool DoReduction(const FCompressibleAnimData& CompressibleAnimData, FCompressibleAnimDataResult& OutResult) override;
 	virtual void PopulateDDCKey(FArchive& Ar) override;
 #endif // WITH_EDITOR
 	//~ Begin UAnimCompress Interface
diff --git a/Engine/Source/Runtime/Engine/Classes/Animation/AnimCompress_RemoveLinearKeys.h b/Engine/Source/Runtime/Engine/Classes/Animation/AnimCompress_RemoveLinearKeys.h
--- a/Engine/Source/Runtime/Engine/Classes/Animation/AnimCompress_RemoveLinearKeys.h
+++ b/Engine/Source/Runtime/Engine/Classes/Animation/AnimCompress_RemoveLinearKeys.h
@@ -14,7 +14,6 @@
 #include "Animation/AnimCompress.h"
 #include "AnimCompress_RemoveLinearKeys.generated.h"
 
-struct FProcessAnimationTracksContext;
 
 UCLASS(MinimalAPI)
 class UAnimCompress_RemoveLinearKeys : public UAnimCompress
@@ -77,36 +76,10 @@ class UAnimCompress_RemoveLinearKeys : public UAnimCompress
 	UPROPERTY(EditAnywhere, Category=LinearKeyRemoval)
 	uint32 bActuallyFilterLinearKeys:1;
 
-	// USE_SEGMENTING_CONTEXT - Hide property 
-	/**
-	 * Whether or not to optimize the memory layout for forward playback. Note that this will degrade the performance
-	 * if the sequence is played backwards quite dramatically depending on the circumstances.
-	 */
-	//UPROPERTY(EditAnywhere, Category = LinearKeyRemoval, meta = (EditCondition="bEnableSegmenting"))
-	UPROPERTY()
-	uint32 bOptimizeForForwardPlayback : 1;
-
-	/**
-	 * Whether to use the legacy behavior that compressed and decompressed the whole sequence whenever a key was removed.
-	 * This is left in for now in order to debug the code if issues arise with the new stateless version.
-	 * The legacy behavior to compress the whole sequence constantly makes it impossible to compress segments
-	 * in parallel safely. The new code samples the sequence without the compression step in order to make it
-	 * thread safe.
-	 * Once the stateless code is robust, this can be safely removed. It should never be used.
-	 */
-	UPROPERTY(EditAnywhere, Category = LinearKeyRemoval)
-	uint32 bUseDecompression : 1;
-
-	/**
-	 * Whether or not to use multiple threads during compression.
-	 */
-	UPROPERTY(EditAnywhere, Category = LinearKeyRemoval)
-	uint32 bUseMultithreading : 1;
-
 protected:
 	//~ Begin UAnimCompress Interface
 #if WITH_EDITOR
-	virtual void DoReduction(const FCompressibleAnimData& CompressibleAnimData, FCompressibleAnimDataResult& OutResult) override;
+	virtual bool DoReduction(const FCompressibleAnimData& CompressibleAnimData, FCompressibleAnimDataResult& OutResult) override;
 	virtual void PopulateDDCKey(FArchive& Ar) override;
 #endif // WITH_EDITOR
 	//~ Begin UAnimCompress Interface
@@ -115,7 +88,7 @@ protected:
 	/**
 	 * Pre-filters the tracks before running the main key removal algorithm
 	 */
-	virtual void FilterBeforeMainKeyRemoval(
+	virtual void* FilterBeforeMainKeyRemoval(
 		const FCompressibleAnimData& CompressibleAnimData,
 		TArray<FTranslationTrack>& TranslationData,
 		TArray<FRotationTrack>& RotationData, 
@@ -132,17 +105,6 @@ protected:
 		const TArray<FScaleTrack>& ScaleData,
 		const bool bFinalPass);
 
-#if USE_SEGMENTING_CONTEXT
-	/**
-	 * Compresses the tracks passed in using the underlying compressor for this key removal codec for a single segment
-	 */
-	virtual void CompressUsingUnderlyingCompressor(
-		UAnimSequence& AnimSeq,
-		const TArray<FBoneData>& BoneData,
-		TArray<FAnimSegmentContext>& RawSegments,
-		const bool bFinalPass);
-#endif
-
 	/**
 	 * Updates the world bone transforms for a range of bone indices
 	 */
@@ -183,35 +145,6 @@ protected:
 		float TimePerFrame,
 		TArray<FTransform>& BoneAtoms);
 
-#if USE_SEGMENTING_CONTEXT
-	/**
-	 * Updates the world bone transforms for a range of bone indices for a single segment
-	 */
-	void UpdateWorldBoneTransformRange(
-		FProcessAnimationTracksContext& Context,
-		int32 StartingBoneIndex,
-		int32 EndingBoneIndex);
-
-	/**
-	 * To guide the key removal process, we need to maintain a table of world transforms
-	 * for the bones we are investigating. This helper function fills a row of the
-	 * table for a specified bone. For a single segment.
-	 */
-	void UpdateWorldBoneTransformTable(
-		const FProcessAnimationTracksContext& Context,
-		int32 BoneIndex,
-		bool UseRaw,
-		TArray<FTransform>& OutputWorldBones);
-
-	/**
-	 * Creates a list of the bone atom result for every frame of a given track for a single segment
-	 */
-	void UpdateBoneAtomList(
-		const FProcessAnimationTracksContext& Context,
-		int32 TrackIndex,
-		TArray<FTransform>& BoneAtoms) const;
-#endif
-
 	/**
 	 * If the passed in animation sequence is additive, converts it to absolute (using the frame 0 pose) and returns true
 	 * (indicating it should be converted back to relative later with ConvertToRelativeSpace)
@@ -271,47 +204,6 @@ protected:
 		TArray<FRotationTrack>& RotationTracks, 
 		TArray<FScaleTrack>& ScaleTracks);
 
-#if USE_SEGMENTING_CONTEXT
-	/**
-	 * Locates spans of keys within the position and rotation tracks provided which can be estimated
-	 * through linear interpolation of the surrounding keys. The remaining key values are bit packed into
-	 * the animation sequence provided
-	 *
-	 * @param	AnimSeq				The animation sequence being compressed
-	 * @param	BoneData			BoneData array describing the hierarchy of the animated skeleton
-	 * @param	RawSegments			Segment list to process
-	 */
-	void ProcessAnimationTracks(
-		UAnimSequence& AnimSeq,
-		const TArray<FBoneData>& BoneData,
-		TArray<FAnimSegmentContext>& RawSegments);
-
-	/**
-	 * Performs retargetting for a single segment.
-	 */
-	void PerformRetargeting(
-		FProcessAnimationTracksContext& Context,
-		int32 BoneIndex,
-		int32 HighestTargetBoneIndex,
-		int32 FurthestTargetBoneIndex);
-
-	/**
-	 * Samples a track within a segment at a specified time with the current track format.
-	 */
-	static FTransform SampleSegment(const FProcessAnimationTracksContext& Context, int32 TrackIndex, float Time);
-#endif
-
-private:
-#if USE_SEGMENTING_CONTEXT
-	/**
-	 * Performs the linear key reduction and retargetting for a single segment.
-	 * This can be called from multiple threads concurrently.
-	 */
-	void ProcessAnimationTracks(FProcessAnimationTracksContext& Context);
-
-	friend struct FAsyncProcessAnimationTracksTaskGroupContext;
-#endif
-
 #endif // WITH_EDITOR
 };
 
diff --git a/Engine/Source/Runtime/Engine/Classes/Animation/AnimCurveCompressionSettings.h b/Engine/Source/Runtime/Engine/Classes/Animation/AnimCurveCompressionSettings.h
--- a/Engine/Source/Runtime/Engine/Classes/Animation/AnimCurveCompressionSettings.h
+++ b/Engine/Source/Runtime/Engine/Classes/Animation/AnimCurveCompressionSettings.h
@@ -28,6 +28,9 @@ class ENGINE_API UAnimCurveCompressionSettings : public UObject
 	UAnimCurveCompressionCodec* GetCodec(const FString& Path);
 
 #if WITH_EDITORONLY_DATA
+	// UObject overrides
+	virtual void GetPreloadDependencies(TArray<UObject*>& OutDeps) override;
+
 	/** Returns whether or not we can use these settings to compress. */
 	bool AreSettingsValid() const;
 
@@ -38,6 +41,6 @@ class ENGINE_API UAnimCurveCompressionSettings : public UObject
 	bool Compress(const FCompressibleAnimData& AnimSeq, FCompressedAnimSequence& OutCompressedData) const;
 
 	/** Generates a DDC key that takes into account the current settings and selected codec. */
-	void PopulateDDCKey(FArchive& Ar) const;
+	void PopulateDDCKey(FArchive& Ar);
 #endif
 };
diff --git a/Engine/Source/Runtime/Engine/Classes/Animation/AnimInstance.h b/Engine/Source/Runtime/Engine/Classes/Animation/AnimInstance.h
--- a/Engine/Source/Runtime/Engine/Classes/Animation/AnimInstance.h
+++ b/Engine/Source/Runtime/Engine/Classes/Animation/AnimInstance.h
@@ -1330,22 +1330,36 @@ protected:
 	/** Override point for derived classes to destroy their own proxy objects (allows custom allocation) */
 	virtual void DestroyAnimInstanceProxy(FAnimInstanceProxy* InProxy);
 
-	/** Access the proxy but block if a task is currently in progress as it wouldn't be safe to access it */
+	/** Access the proxy but block if a task is currently in progress as it wouldn't be safe to access it 
+	 *	This is protected static member for allowing derived to access
+	 */
 	template <typename T /*= FAnimInstanceProxy*/>	// @TODO: Cant default parameters to this function on Xbox One until we move off the VS2012 compiler
-	FORCEINLINE T& GetProxyOnGameThread()
+	FORCEINLINE static T* GetProxyOnGameThreadStatic(UAnimInstance* InAnimInstance)
 	{
-		check(IsInGameThread());
-		if(GetOuter() && GetOuter()->IsA<USkeletalMeshComponent>())
+		if (InAnimInstance)
 		{
-			bool bBlockOnTask = true;
-			bool bPerformPostAnimEvaluation = true;
-			GetSkelMeshComponent()->HandleExistingParallelEvaluationTask(bBlockOnTask, bPerformPostAnimEvaluation);
-		}
-		if(AnimInstanceProxy == nullptr)
-		{
-			AnimInstanceProxy = CreateAnimInstanceProxy();
+			check(IsInGameThread());
+			UObject* OuterObj = InAnimInstance->GetOuter();
+			if (OuterObj && OuterObj->IsA<USkeletalMeshComponent>())
+			{
+				bool bBlockOnTask = true;
+				bool bPerformPostAnimEvaluation = true;
+				InAnimInstance->GetSkelMeshComponent()->HandleExistingParallelEvaluationTask(bBlockOnTask, bPerformPostAnimEvaluation);
+			}
+			if (InAnimInstance->AnimInstanceProxy == nullptr)
+			{
+				InAnimInstance->AnimInstanceProxy = InAnimInstance->CreateAnimInstanceProxy();
+			}
+			return static_cast<T*>(InAnimInstance->AnimInstanceProxy);
 		}
-		return *static_cast<T*>(AnimInstanceProxy);
+
+		return nullptr;
+	}
+	/** Access the proxy but block if a task is currently in progress as it wouldn't be safe to access it */
+	template <typename T /*= FAnimInstanceProxy*/>	// @TODO: Cant default parameters to this function on Xbox One until we move off the VS2012 compiler
+	FORCEINLINE T& GetProxyOnGameThread()
+	{
+		return *GetProxyOnGameThreadStatic<T>(this);
 	}
 
 	/** Access the proxy but block if a task is currently in progress as it wouldn't be safe to access it */
diff --git a/Engine/Source/Runtime/Engine/Classes/Animation/AnimMontage.h b/Engine/Source/Runtime/Engine/Classes/Animation/AnimMontage.h
--- a/Engine/Source/Runtime/Engine/Classes/Animation/AnimMontage.h
+++ b/Engine/Source/Runtime/Engine/Classes/Animation/AnimMontage.h
@@ -43,7 +43,7 @@ struct FCompositeSection : public FAnimLinkableElement
 	float StartTime_DEPRECATED;
 
 	/** Should this animation loop. */
-	UPROPERTY()
+	UPROPERTY(VisibleAnywhere, Category=Section)
 	FName NextSectionName;
 
 	/** Meta data that can be saved with the asset 
@@ -525,11 +525,17 @@ public:
 	/** static functions that are used by matinee functionality */
 	static UAnimMontage* SetMatineeAnimPositionInner(FName SlotName, USkeletalMeshComponent* SkeletalMeshComponent, UAnimSequenceBase* InAnimSequence, float InPosition, bool bLooping);
 	static UAnimMontage* PreviewMatineeSetAnimPositionInner(FName SlotName, USkeletalMeshComponent* SkeletalMeshComponent, UAnimSequenceBase* InAnimSequence, float InPosition, bool bLooping, bool bFireNotifies, float DeltaTime);
+	
+	static UAnimMontage* SetMatineeAnimPositionInner(FName SlotName, UAnimInstance* AnimInstance, UAnimSequenceBase* InAnimSequence, float InPosition, bool bLooping);
+	static UAnimMontage* PreviewMatineeSetAnimPositionInner(FName SlotName, USkeletalMeshComponent* SkeletalMeshComponent, UAnimInstance* AnimInstance, UAnimSequenceBase* InAnimSequence, float InPosition, bool bLooping, bool bFireNotifies, float DeltaTime);
+	
 	/** static functions that are used by sequencer montage support*/
 	static UAnimMontage* SetSequencerMontagePosition(FName SlotName, USkeletalMeshComponent* SkeletalMeshComponent, int32& InOutInstanceId, UAnimSequenceBase* InAnimSequence, float InFromPosition, float InToPosition, float Weight, bool bLooping, bool bPlaying);
 	static UAnimMontage* PreviewSequencerMontagePosition(FName SlotName, USkeletalMeshComponent* SkeletalMeshComponent, int32& InOutInstanceId, UAnimSequenceBase* InAnimSequence, float InFromPosition, float InToPosition, float Weight, bool bLooping, bool bFireNotifies, bool bPlaying);
+	static UAnimMontage* SetSequencerMontagePosition(FName SlotName, UAnimInstance* AnimInstance, int32& InOutInstanceId, UAnimSequenceBase* InAnimSequence, float InFromPosition, float InToPosition, float Weight, bool bLooping, bool bPlaying);
+	static UAnimMontage* PreviewSequencerMontagePosition(FName SlotName, USkeletalMeshComponent* SkeletalMeshComponent, UAnimInstance* AnimInstance, int32& InOutInstanceId, UAnimSequenceBase* InAnimSequence, float InFromPosition, float InToPosition, float Weight, bool bLooping, bool bFireNotifies, bool bPlaying);
 private:
-	static UAnimMontage* InitializeMatineeControl(FName SlotName, USkeletalMeshComponent* SkeletalMeshComponent, UAnimSequenceBase* InAnimSequence, bool bLooping);
+	static UAnimMontage* InitializeMatineeControl(FName SlotName, UAnimInstance* AnimInstance, UAnimSequenceBase* InAnimSequence, bool bLooping);
 };
 
 /* 
diff --git a/Engine/Source/Runtime/Engine/Classes/Animation/AnimNode_ApplyMeshSpaceAdditive.h b/Engine/Source/Runtime/Engine/Classes/Animation/AnimNode_ApplyMeshSpaceAdditive.h
--- a/Engine/Source/Runtime/Engine/Classes/Animation/AnimNode_ApplyMeshSpaceAdditive.h
+++ b/Engine/Source/Runtime/Engine/Classes/Animation/AnimNode_ApplyMeshSpaceAdditive.h
@@ -19,12 +19,32 @@ struct ENGINE_API FAnimNode_ApplyMeshSpaceAdditive : public FAnimNode_Base
 	UPROPERTY(EditAnywhere, BlueprintReadWrite, Category = Links)
 	FPoseLink Additive;
 
+	/** The data type used to control the alpha blending of the additive pose.
+		Note: Changing this value will disconnect alpha input pins.
+	*/
+	UPROPERTY(EditAnywhere, BlueprintReadWrite, Category = Settings)
+	EAnimAlphaInputType AlphaInputType;
+
+	/** The float value that controls the alpha blending when the alpha input type is set to 'Float' */
 	UPROPERTY(EditAnywhere, BlueprintReadWrite, Category = Settings, meta = (PinShownByDefault))
 	float Alpha;
 
+	/** The boolean value that controls the alpha blending when the alpha input type is set to 'Bool' */
+	UPROPERTY(EditAnywhere, BlueprintReadWrite, Category = Settings, meta = (PinShownByDefault, DisplayName = "bEnabled"))
+	uint8 bAlphaBoolEnabled : 1;
+
+	UPROPERTY(EditAnywhere, BlueprintReadWrite, Category = Settings, meta = (DisplayName = "Blend Settings"))
+	FInputAlphaBoolBlend AlphaBoolBlend;
+
+	UPROPERTY(EditAnywhere, BlueprintReadWrite, Category = Settings, meta = (PinShownByDefault))
+	FName AlphaCurveName;
+
 	UPROPERTY(EditAnywhere, BlueprintReadWrite, Category = Settings)
 	FInputScaleBias AlphaScaleBias;
 
+	UPROPERTY(EditAnywhere, BlueprintReadWrite, Category = Settings)
+	FInputScaleBiasClamp AlphaScaleBiasClamp;
+
 	/*
 	* Max LOD that this node is allowed to run
 	* For example if you have LODThreadhold to be 2, it will run until LOD 2 (based on 0 index)
diff --git a/Engine/Source/Runtime/Engine/Classes/Animation/AnimNode_CustomProperty.h b/Engine/Source/Runtime/Engine/Classes/Animation/AnimNode_CustomProperty.h
--- a/Engine/Source/Runtime/Engine/Classes/Animation/AnimNode_CustomProperty.h
+++ b/Engine/Source/Runtime/Engine/Classes/Animation/AnimNode_CustomProperty.h
@@ -81,10 +81,10 @@ protected:
 
 	/* Initialize property links from the source instance, in this case AnimInstance 
 	 * Compiler creates those properties during compile time */
-	void InitializeProperties(const UObject* InSourceInstance, UClass* InTargetClass);
+	virtual void InitializeProperties(const UObject* InSourceInstance, UClass* InTargetClass);
 
 	/* Propagate the Source Instances' properties to Target Instance*/
-	void PropagateInputProperties(const UObject* InSourceInstance);
+	virtual void PropagateInputProperties(const UObject* InSourceInstance);
 
 	/** Get Target Class */
 	virtual UClass* GetTargetClass() const PURE_VIRTUAL(FAnimNode_CustomProperty::GetTargetClass, return nullptr;);
diff --git a/Engine/Source/Runtime/Engine/Classes/Animation/AnimNode_SequencePlayer.h b/Engine/Source/Runtime/Engine/Classes/Animation/AnimNode_SequencePlayer.h
--- a/Engine/Source/Runtime/Engine/Classes/Animation/AnimNode_SequencePlayer.h
+++ b/Engine/Source/Runtime/Engine/Classes/Animation/AnimNode_SequencePlayer.h
@@ -6,7 +6,6 @@
 #include "UObject/ObjectMacros.h"
 #include "Animation/AnimNode_AssetPlayerBase.h"
 #include "Animation/InputScaleBias.h"
-#include "Animation/AnimSequenceDecompressionContext.h"
 #include "Animation/AnimSequenceBase.h"
 #include "AnimNode_SequencePlayer.generated.h"
 
diff --git a/Engine/Source/Runtime/Engine/Classes/Animation/AnimSequence.h b/Engine/Source/Runtime/Engine/Classes/Animation/AnimSequence.h
--- a/Engine/Source/Runtime/Engine/Classes/Animation/AnimSequence.h
+++ b/Engine/Source/Runtime/Engine/Classes/Animation/AnimSequence.h
@@ -19,7 +19,6 @@
 
 #include "AnimSequence.generated.h"
 
-#define USE_SEGMENTING_CONTEXT 0 // Uses segmenting in anim compression + context in decompression
 
 typedef TArray<FTransform> FTransformArrayA2;
 
@@ -194,8 +193,6 @@ struct ENGINE_API FRequestAnimCompressionParams
 	void InitFrameStrippingFromPlatform(const class ITargetPlatform* TargetPlatform);
 };
 
-FArchive& operator<<(FArchive& Ar, FCompressedOffsetData& D);
-
 UCLASS(config=Engine, hidecategories=(UObject, Length), BlueprintType)
 class ENGINE_API UAnimSequence : public UAnimSequenceBase
 {
@@ -250,12 +247,6 @@ protected:
 public:
 
 #if WITH_EDITORONLY_DATA
-	/**
-	 * The compression scheme that was most recently used to compress this animation.
-	 */
-	UPROPERTY(Category=Compression, VisibleAnywhere)
-	class UAnimCompress* CompressionScheme;
-
 	/**
 	 * Allow frame stripping to be performed on this animation if the platform requests it
 	 * Can be disabled if animation has high frequency movements that are being lost.
@@ -270,11 +261,14 @@ public:
 	 */
 	UPROPERTY(Category = Compression, EditAnywhere)
 	float CompressionErrorThresholdScale;
-
 #endif
 
+	/** The bone compression settings used to compress bones in this sequence. */
+	UPROPERTY(Category = Compression, EditAnywhere, meta = (ForceShowEngineContent))
+	class UAnimBoneCompressionSettings* BoneCompressionSettings;
+
 	/** The curve compression settings used to compress curves in this sequence. */
-	UPROPERTY(Category=Compression, EditAnywhere)
+	UPROPERTY(Category = Compression, EditAnywhere, meta = (ForceShowEngineContent))
 	class UAnimCurveCompressionSettings* CurveCompressionSettings;
 
 	FCompressedAnimSequence CompressedData;
@@ -459,9 +453,6 @@ public:
 	
 	// Adds a new track (if no track of the supplied name is found) to the raw animation data, optionally setting it to TrackData.
 	int32 AddNewRawTrack(FName TrackName, FRawAnimSequenceTrack* TrackData = nullptr);
-
-	// Get the Alternate compression error threshold 
-	float GetAltCompressionErrorThreshold() const;
 #endif
 
 	const TArray<FTrackToSkeletonMap>& GetRawTrackToSkeletonMapTable() const { return TrackToSkeletonMapTable; }
@@ -585,11 +576,26 @@ public:
 	bool CompressRawAnimData();
 
 	// Get compressed data for this UAnimSequence. May be built directly or pulled from DDC
+#if WITH_EDITOR
+	bool ShouldPerformStripping(const bool bPerformFrameStripping, const bool bPerformStrippingOnOddFramedAnims) const;
+	FString GetDDCCacheKeySuffix(const bool bPerformStripping) const;
+	void ApplyCompressedData(const FString& DataCacheKeySuffix, const bool bPerformFrameStripping, const TArray<uint8>& Data);
+#endif
+	void WaitOnExistingCompression(const bool bCancelIfNotStarted=false);
 	void RequestAnimCompression(FRequestAnimCompressionParams Params);
 	void RequestSyncAnimRecompression(bool bOutput = false) { RequestAnimCompression(FRequestAnimCompressionParams(false, false, bOutput)); }
+	void RequestAsyncAnimRecompression(bool bOutput = false) { RequestAnimCompression(FRequestAnimCompressionParams(true, false, bOutput)); }
+
+protected:
+	void ApplyCompressedData(const TArray<uint8>& Data);
+
+public:
 	bool IsCompressedDataValid() const;
 	bool IsCurveCompressedDataValid() const;
 
+	void ClearCompressedBoneData();
+	void ClearCompressedCurveData();
+
 	// Write the compressed data to the supplied FArchive
 	void SerializeCompressedData(FArchive& Ar, bool bDDCData);
 
@@ -813,7 +819,7 @@ private:
 
 	/** Retargeting functions */
 	bool ConvertAnimationDataToRiggingData(FAnimSequenceTrackContainer & RiggingAnimationData);
-	bool ConvertRiggingDataToAnimationData(FAnimSequenceTrackContainer & RiggingAnimationData);
+	bool ConvertRiggingDataToAnimationData(FAnimSequenceTrackContainer & RiggingAnimationData, bool bPerformPostProcess=true);
 	int32 GetSpaceBasedAnimationData(TArray< TArray<FTransform> > & AnimationDataInComponentSpace, FAnimSequenceTrackContainer * RiggingAnimationData) const;
 
 	/** Verify Track Map is valid, if not, fix up */
@@ -850,13 +856,17 @@ private:
 	// Should we be always using our raw data (i.e is our compressed data stale)
 	bool bUseRawDataOnly;
 
-public:
+#if WITH_EDITOR
 	// Are we currently compressing this animation
 	bool bCompressionInProgress;
+#endif
+
+public:
 
 	friend class UAnimationAsset;
 	friend struct FScopedAnimSequenceRawDataCache;
 	friend class UAnimationBlueprintLibrary;
+	friend class UAnimBoneCompressionSettings;
 };
 
 
diff --git a/Engine/Source/Runtime/Engine/Classes/Animation/AnimStreamable.h b/Engine/Source/Runtime/Engine/Classes/Animation/AnimStreamable.h
--- a/Engine/Source/Runtime/Engine/Classes/Animation/AnimStreamable.h
+++ b/Engine/Source/Runtime/Engine/Classes/Animation/AnimStreamable.h
@@ -109,12 +109,6 @@ public:
 	UPROPERTY()
 	const UAnimSequence* SourceSequence;
 
-	/**
-	 * The compression scheme that was most recently used to compress this animation.
-	 */
-	UPROPERTY(Category = Compression, VisibleAnywhere)
-	class UAnimCompress* CompressionScheme;
-
 	UPROPERTY()
 	FGuid RawDataGuid;
 
@@ -176,6 +170,10 @@ public:
 #endif
 	}
 
+	/** The bone compression settings used to compress bones in this sequence. */
+	UPROPERTY(Category = Compression, EditAnywhere)
+	class UAnimBoneCompressionSettings* BoneCompressionSettings;
+
 	/** The curve compression settings used to compress curves in this sequence. */
 	UPROPERTY(Category = Compression, EditAnywhere)
 	class UAnimCurveCompressionSettings* CurveCompressionSettings;
@@ -221,13 +219,11 @@ public:
 	private:
 
 #if WITH_EDITOR
-	float GetAltCompressionErrorThreshold() const;
-
 	ENGINE_API void RequestCompressedData(const  ITargetPlatform* Platform=nullptr);
 
 	void UpdateRawData();
 
-	FString GetBaseDDCKey(uint32 NumChunks, float AltCompressionErrorThreshold) const;
+	FString GetBaseDDCKey(uint32 NumChunks) const;
 
 	void RequestCompressedDataForChunk(const FString& ChunkDDCKey, FAnimStreamableChunk& Chunk, const int32 ChunkIndex, const uint32 FrameStart, const uint32 FrameEnd, TSharedRef<FAnimCompressContext> CompressContext);
 #endif
diff --git a/Engine/Source/Runtime/Engine/Classes/Animation/AnimationSettings.h b/Engine/Source/Runtime/Engine/Classes/Animation/AnimationSettings.h
--- a/Engine/Source/Runtime/Engine/Classes/Animation/AnimationSettings.h
+++ b/Engine/Source/Runtime/Engine/Classes/Animation/AnimationSettings.h
@@ -41,16 +41,6 @@ class ENGINE_API UAnimationSettings : public UDeveloperSettings
 	UPROPERTY(config, EditAnywhere, Category = Compression)
 	float MaxCurveError;
 
-	/** The alternate error threshold (0.0 means don't try anything other than the current / default scheme) 
-	* 
-	* Determines the current setting for world-space error tolerance in the animation compressor.
-	* When requested, animation being compressed will also consider an alternative compression
-	* method if the end result of that method produces less error than the AlternativeCompressionThreshold.
-	* Also known as "Max End Effector Error"
-	*/
-	UPROPERTY(config, EditAnywhere, Category = Compression, meta=(ClampMin = "0", UIMin = "0"))
-	float AlternativeCompressionThreshold;
-
 	UPROPERTY(config, EditAnywhere, Category = Compression)
 	bool ForceRecompression;
 
diff --git a/Engine/Source/Runtime/Engine/Classes/Animation/Skeleton.h b/Engine/Source/Runtime/Engine/Classes/Animation/Skeleton.h
--- a/Engine/Source/Runtime/Engine/Classes/Animation/Skeleton.h
+++ b/Engine/Source/Runtime/Engine/Classes/Animation/Skeleton.h
@@ -231,7 +231,7 @@ public:
 
 namespace VirtualBoneNameHelpers
 {
-	const FString VirtualBonePrefix(TEXT("VB "));
+	extern ENGINE_API const FString VirtualBonePrefix;
 
 	ENGINE_API FString AddVirtualBonePrefix(const FString& InName);
 	ENGINE_API FName RemoveVirtualBonePrefix(const FString& InName);
@@ -434,7 +434,10 @@ public:
 	ENGINE_API const FAnimSlotGroup* FindAnimSlotGroup(const FName& InGroupName) const;
 	ENGINE_API const TArray<FAnimSlotGroup>& GetSlotGroups() const;
 	ENGINE_API bool ContainsSlotName(const FName& InSlotName) const;
-	ENGINE_API void RegisterSlotNode(const FName& InSlotName);
+
+	/** Register a slot name. Return true if a slot was registered, false if it was already registered. */
+	ENGINE_API bool RegisterSlotNode(const FName& InSlotName);
+
 	ENGINE_API void SetSlotGroupName(const FName& InSlotName, const FName& InGroupName);
 	/** Returns true if Group is added, false if it already exists */
 	ENGINE_API bool AddSlotGroupName(const FName& InNewGroupName);
diff --git a/Engine/Source/Runtime/Engine/Classes/Components/SkeletalMeshComponent.h b/Engine/Source/Runtime/Engine/Classes/Components/SkeletalMeshComponent.h
--- a/Engine/Source/Runtime/Engine/Classes/Components/SkeletalMeshComponent.h
+++ b/Engine/Source/Runtime/Engine/Classes/Components/SkeletalMeshComponent.h
@@ -298,6 +298,7 @@ class ENGINE_API USkeletalMeshComponent : public USkinnedMeshComponent, public I
 	friend class UAnimInstance;
 	friend struct FAnimNode_LinkedAnimGraph;
 	friend struct FAnimNode_LinkedAnimLayer;
+	friend struct FLinkedInstancesAdapter;
 
 	/**
 	 * Animation 
@@ -2251,3 +2252,30 @@ public:
 	// Are we currently within PostAnimEvaluation
 	bool IsPostEvaluatingAnimation() const { return bPostEvaluatingAnimation; }
 };
+
+struct FLinkedInstancesAdapter
+{
+	static void AddLinkedInstance(USkeletalMeshComponent* InComponent, UAnimInstance* InAnimInstance)
+	{
+		if (InComponent && InAnimInstance)
+		{
+			InComponent->LinkedInstances.Add(InAnimInstance);
+		}
+	}
+
+	static void RemoveLinkedInstance(USkeletalMeshComponent* InComponent, UAnimInstance* InAnimInstance)
+	{
+		if (InComponent && InAnimInstance)
+		{
+			InComponent->LinkedInstances.Remove(InAnimInstance);
+		}
+	}
+
+	static void ResetLinkedInstance(USkeletalMeshComponent* InComponent)
+	{
+		if (InComponent)
+		{
+			InComponent->LinkedInstances.Reset();
+		}
+	}
+};
\ No newline at end of file
diff --git a/Engine/Source/Runtime/Engine/Classes/Components/SkinnedMeshComponent.h b/Engine/Source/Runtime/Engine/Classes/Components/SkinnedMeshComponent.h
--- a/Engine/Source/Runtime/Engine/Classes/Components/SkinnedMeshComponent.h
+++ b/Engine/Source/Runtime/Engine/Classes/Components/SkinnedMeshComponent.h
@@ -4,6 +4,7 @@
 #pragma once
 
 #include "CoreMinimal.h"
+#include "GPUSkinPublicDefs.h"
 #include "UObject/ObjectMacros.h"
 #include "UObject/Object.h"
 #include "Engine/EngineTypes.h"
@@ -149,10 +150,10 @@ struct FSkelMeshSkinWeightInfo
 
 	/** Index of bones that influence this vertex */
 	UPROPERTY()
-	int32	Bones[8];
+	int32	Bones[MAX_TOTAL_INFLUENCES];
 	/** Influence of each bone on this vertex */
 	UPROPERTY()
-	uint8	Weights[8];
+	uint8	Weights[MAX_TOTAL_INFLUENCES];
 };
 
 /** LOD specific setup for the skeletal mesh component. */
@@ -1564,7 +1565,7 @@ public:
 
 
 /** Simple, CPU evaluation of a vertex's skinned position helper function */
-template <bool bExtraBoneInfluencesT, bool bCachedMatrices>
+template <bool bCachedMatrices>
 FVector GetTypedSkinnedVertexPosition(
 	const USkinnedMeshComponent* SkinnedComp,
 	const FSkelMeshRenderSection& Section,
diff --git a/Engine/Source/Runtime/Engine/Classes/Curves/RichCurve.h b/Engine/Source/Runtime/Engine/Classes/Curves/RichCurve.h
--- a/Engine/Source/Runtime/Engine/Classes/Curves/RichCurve.h
+++ b/Engine/Source/Runtime/Engine/Classes/Curves/RichCurve.h
@@ -205,6 +205,7 @@ public:
 	/** Functions for getting keys based on handles */
 	FRichCurveKey& GetKey(FKeyHandle KeyHandle);
 	FRichCurveKey GetKey(FKeyHandle KeyHandle) const;
+	const FRichCurveKey& GetKeyRef(FKeyHandle KeyHandle) const;
 	
 	/** Quick accessors for the first and last keys */
 	FRichCurveKey GetFirstKey() const;
diff --git a/Engine/Source/Runtime/Engine/Classes/EdGraph/EdGraphNode.h b/Engine/Source/Runtime/Engine/Classes/EdGraph/EdGraphNode.h
--- a/Engine/Source/Runtime/Engine/Classes/EdGraph/EdGraphNode.h
+++ b/Engine/Source/Runtime/Engine/Classes/EdGraph/EdGraphNode.h
@@ -880,6 +880,9 @@ public:
 	// called to replace this nodes comment text
 	virtual void OnUpdateCommentText( const FString& NewComment );
 
+	// returns true if this node supports comment bubbles
+	virtual bool SupportsCommentBubble() const { return true; }
+
 	// called when the node's comment bubble is toggled
 	virtual void OnCommentBubbleToggled( bool bInCommentBubbleVisible ) {}
 
diff --git a/Engine/Source/Runtime/Engine/Classes/EdGraph/EdGraphSchema.h b/Engine/Source/Runtime/Engine/Classes/EdGraph/EdGraphSchema.h
--- a/Engine/Source/Runtime/Engine/Classes/EdGraph/EdGraphSchema.h
+++ b/Engine/Source/Runtime/Engine/Classes/EdGraph/EdGraphSchema.h
@@ -1037,4 +1037,14 @@ class ENGINE_API UEdGraphSchema : public UObject
 	 * @return	true if the pin types and directions are compatible.
 	 */
 	virtual bool ArePinsCompatible(const UEdGraphPin* PinA, const UEdGraphPin* PinB, const UClass* CallingContext = NULL, bool bIgnoreArray = false) const { return true; }
+
+	/**
+	 * Returns true if the schema wants to overdrive the behaviour of dirtying the blueprint on new node creation.
+	 *
+	 * @param   InBlueprint    The blueprint to dirty or not
+	 * @param   InEdGraphNode  The node that was just added and caused the request
+	 * 
+	 * @return  true if the blueprint marking has been taken care off.
+	 */
+	virtual bool MarkBlueprintDirtyFromNewNode(UBlueprint* InBlueprint, UEdGraphNode* InEdGraphNode) const { return false; }
 };
diff --git a/Engine/Source/Runtime/Engine/Classes/Engine/Blueprint.h b/Engine/Source/Runtime/Engine/Classes/Engine/Blueprint.h
--- a/Engine/Source/Runtime/Engine/Classes/Engine/Blueprint.h
+++ b/Engine/Source/Runtime/Engine/Classes/Engine/Blueprint.h
@@ -1017,12 +1017,9 @@ public:
 	virtual void GetInstanceActions(FBlueprintActionDatabaseRegistrar& ActionRegistrar) const {}
 
 	/**
-	 * Allows child classes to override the undo behaviour on undo / redo.
-	 * Normally blueprints are marked as dirty / structurally modified on undo,
-	 * which might not be required for specialized classes.
-	 * @return True if the blueprint should be marked as structurally modified on undo redo.
+	 * Returns true if this blueprint should be marked dirty upon a transaction
 	 */
-	virtual bool RequiresMarkAsStructurallyModifiedOnUndo() const { return true; }
+	virtual bool ShouldBeMarkedDirtyUponTransaction() const { return true; }
 
 #if WITH_EDITOR
 private:
diff --git a/Engine/Source/Runtime/Engine/Classes/Engine/RendererSettings.h b/Engine/Source/Runtime/Engine/Classes/Engine/RendererSettings.h
--- a/Engine/Source/Runtime/Engine/Classes/Engine/RendererSettings.h
+++ b/Engine/Source/Runtime/Engine/Classes/Engine/RendererSettings.h
@@ -750,6 +750,12 @@ class ENGINE_API URendererSettings : public UDeveloperSettings
 		ToolTip = "Maximum amount of memory (in MB) per world/scene allowed for the Compute Skincache to generate output vertex data and recompute tangents."))
 		float SkinCacheSceneMemoryLimitInMB;
 
+	UPROPERTY(config, EditAnywhere, Category = Optimizations, meta = (
+		ConsoleVariable = "r.GPUSkin.Support16BitBoneIndex", DisplayName = "Support 16-bit Bone Index",
+		ToolTip = "If enabled, a new mesh imported will use 8 bit (if <=256 bones) or 16 bit (if > 256 bones) bone indices for rendering.",
+		ConfigRestartRequired = true))
+		uint32 bSupport16BitBoneIndex : 1;
+
 	UPROPERTY(config, EditAnywhere, Category = Optimizations, meta = (
 		ConsoleVariable = "r.GPUSkin.Limit2BoneInfluences", DisplayName = "Limit GPU skinning to 2 bones influence",
 		ToolTip = "Whether to use 2 bone influences instead of the default of 4 for GPU skinning. This does not change skeletal mesh assets but reduces the number of instructions required by the GPU skin vertex shaders. Changing this setting requires restarting the editor.",
diff --git a/Engine/Source/Runtime/Engine/Classes/Particles/Location/ParticleModuleLocationSkelVertSurface.h b/Engine/Source/Runtime/Engine/Classes/Particles/Location/ParticleModuleLocationSkelVertSurface.h
--- a/Engine/Source/Runtime/Engine/Classes/Particles/Location/ParticleModuleLocationSkelVertSurface.h
+++ b/Engine/Source/Runtime/Engine/Classes/Particles/Location/ParticleModuleLocationSkelVertSurface.h
@@ -229,6 +229,5 @@ class ENGINE_API UParticleModuleLocationSkelVertSurface : public UParticleModule
 
 private:
 	/** Helper function for concrete types. */
-	template<bool bExtraBoneInfluencesT>
 	bool VertInfluencedByActiveBoneTyped(FSkeletalMeshLODRenderData& LODData, int32 LODIndex, const FSkelMeshRenderSection& Section, int32 VertIndex, USkeletalMeshComponent* InSkelMeshComponent, FModuleLocationVertSurfaceInstancePayload* InstancePayload, int32* OutBoneIndex);
 };
diff --git a/Engine/Source/Runtime/Engine/Private/Animation/AnimBlueprintGeneratedClass.cpp b/Engine/Source/Runtime/Engine/Private/Animation/AnimBlueprintGeneratedClass.cpp
--- a/Engine/Source/Runtime/Engine/Private/Animation/AnimBlueprintGeneratedClass.cpp
+++ b/Engine/Source/Runtime/Engine/Private/Animation/AnimBlueprintGeneratedClass.cpp
@@ -308,7 +308,7 @@ void UAnimBlueprintGeneratedClass::Link(FArchive& Ar, bool bRelinkExistingProper
 				{
 					LinkedAnimLayerNodeProperties.Add(StructProp);
 				}
-				else if(StructProp->Struct == FAnimNode_StateMachine::StaticStruct())
+				else if(StructProp->Struct->IsChildOf(FAnimNode_StateMachine::StaticStruct()))
 				{
 					StateMachineNodeProperties.Add(StructProp);
 				}
