COPY from //Dev-Anim to //Dev-Main

diff --git a/Engine/Source/Runtime/Engine/Private/Animation/AnimCompress.cpp b/Engine/Source/Runtime/Engine/Private/Animation/AnimCompress.cpp
--- a/Engine/Source/Runtime/Engine/Private/Animation/AnimCompress.cpp
+++ b/Engine/Source/Runtime/Engine/Private/Animation/AnimCompress.cpp
@@ -7,6 +7,7 @@
 #include "Animation/AnimCompress.h"
 #include "Misc/MessageDialog.h"
 #include "Serialization/MemoryWriter.h"
+#include "Serialization/MemoryReader.h"
 #include "Misc/FeedbackContext.h"
 #include "AnimationCompression.h"
 #include "AnimEncoding.h"
@@ -17,48 +18,45 @@ DEFINE_LOG_CATEGORY(LogAnimationCompression);
 UAnimCompress::UAnimCompress(const FObjectInitializer& ObjectInitializer)
 	: Super(ObjectInitializer)
 {
-	Description = TEXT("None");
 	TranslationCompressionFormat = ACF_None;
 	RotationCompressionFormat = ACF_Float96NoW;
 
 	UAnimationSettings* AnimationSettings = UAnimationSettings::Get();
 	MaxCurveError = AnimationSettings->MaxCurveError;
-	bEnableSegmenting = AnimationSettings->bEnableSegmenting;
-	IdealNumFramesPerSegment = 64;
-	MaxNumFramesPerSegment = (IdealNumFramesPerSegment * 2) - 1;
 }
 
-void FCompressionMemorySummary::GatherPostCompressionStats(const FCompressibleAnimData& CompressibleAnimData, FCompressibleAnimDataResult& CompressedData, double CompressionTime)
+void FCompressionMemorySummary::GatherPostCompressionStats(const FCompressedAnimSequence& CompressedData, const TArray<FBoneData>& BoneData, const FName AnimFName, double CompressionTime, bool bInPerformedCompression)
 {
-#if WITH_EDITOR
-	if (bEnabled)
+#if WITH_EDITORONLY_DATA
+	if (bEnabled && CompressedData.CompressedDataStructure)
 	{
-		TotalAfterCompressed += CompressedData.GetApproxBoneCompressedSize();
+		bPerformedCompression |= bInPerformedCompression;
+
+		TotalAfterCompressed +=  CompressedData.CompressedDataStructure->GetApproxCompressedSize();
 		TotalCompressionExecutionTime += CompressionTime;
 
-		// determine the error added by the compression
-		AnimationErrorStats ErrorStats;
-		FAnimationUtils::ComputeCompressionError(CompressibleAnimData, CompressedData, ErrorStats);
+		const AnimationErrorStats& ErrorStats = CompressedData.CompressedDataStructure->BoneCompressionErrorStats;
 
 		ErrorTotal += ErrorStats.AverageError;
 		ErrorCount += 1.0f;
 		AverageError = ErrorTotal / ErrorCount;
 
-		WorstBoneError.StoreErrorStat(ErrorStats.MaxError, ErrorStats.MaxError, ErrorStats.MaxErrorTime, ErrorStats.MaxErrorBone, CompressibleAnimData.BoneData[ErrorStats.MaxErrorBone].Name, CompressibleAnimData.AnimFName);
+		WorstBoneError.StoreErrorStat(ErrorStats.MaxError, ErrorStats.MaxError, ErrorStats.MaxErrorTime, ErrorStats.MaxErrorBone, BoneData[ErrorStats.MaxErrorBone].Name, AnimFName);
 
-		WorstAnimationError.StoreErrorStat(ErrorStats.AverageError, ErrorStats.AverageError, CompressibleAnimData.AnimFName);
+		WorstAnimationError.StoreErrorStat(ErrorStats.AverageError, ErrorStats.AverageError, AnimFName);
 	}
 #endif
 }
 
-void FAnimCompressContext::GatherPostCompressionStats(const FCompressibleAnimData& CompressibleAnimData, FCompressibleAnimDataResult& CompressedData, double CompressionTime)
+void FAnimCompressContext::GatherPostCompressionStats(const FCompressedAnimSequence& CompressedData, const TArray<FBoneData>& BoneData, const FName AnimFName, double CompressionTime, bool bInPerformedCompression)
 {
-	CompressionSummary.GatherPostCompressionStats(CompressibleAnimData, CompressedData, CompressionTime);
+	CompressionSummary.GatherPostCompressionStats(CompressedData, BoneData, AnimFName, CompressionTime, bInPerformedCompression);
 }
 
 FCompressionMemorySummary::FCompressionMemorySummary(bool bInEnabled)
 	: bEnabled(bInEnabled)
 	, bUsed(false)
+	, bPerformedCompression(false)
 	, TotalRaw(0)
 	, TotalBeforeCompressed(0)
 	, TotalAfterCompressed(0)
@@ -72,48 +70,53 @@ FCompressionMemorySummary::FCompressionMemorySummary(bool bInEnabled)
 
 FCompressionMemorySummary::~FCompressionMemorySummary()
 {
-	if (bEnabled)
+	if (bEnabled && bUsed && WorstBoneError.IsValid())
 	{
-		if (bUsed)
+		const int32 TotalBeforeSaving = TotalRaw - TotalBeforeCompressed;
+		const int32 TotalAfterSaving = TotalRaw - TotalAfterCompressed;
+		const float OldCompressionRatio = (TotalBeforeCompressed > 0.f) ? (static_cast<float>(TotalRaw) / TotalBeforeCompressed) : 0.f;
+		const float NewCompressionRatio = (TotalAfterCompressed > 0.f) ? (static_cast<float>(TotalRaw) / TotalAfterCompressed) : 0.f;
+
+		FNumberFormattingOptions Options;
+		Options.MinimumIntegralDigits = 7;
+		Options.MinimumFractionalDigits = 2;
+
+		FFormatNamedArguments Args;
+		Args.Add(TEXT("TotalRaw"), FText::AsMemory(TotalRaw, &Options));
+		Args.Add(TEXT("TotalBeforeCompressed"), FText::AsMemory(TotalBeforeCompressed, &Options));
+		Args.Add(TEXT("TotalBeforeSaving"), FText::AsMemory(TotalBeforeSaving, &Options));
+		Args.Add(TEXT("NumberOfAnimations"), FText::AsNumber(NumberOfAnimations));
+		Args.Add(TEXT("OldCompressionRatio"), OldCompressionRatio);
+
+		Args.Add(TEXT("TotalAfterCompressed"), FText::AsMemory(TotalAfterCompressed, &Options));
+		Args.Add(TEXT("TotalAfterSaving"), FText::AsMemory(TotalAfterSaving, &Options));
+		Args.Add(TEXT("NewCompressionRatio"), NewCompressionRatio);
+		Args.Add(TEXT("TotalTimeSpentCompressingPretty"), FText::FromString(FPlatformTime::PrettyTime(TotalCompressionExecutionTime)));
+		Args.Add(TEXT("TotalTimeSpentCompressingRawSeconds"), FText::AsNumber((float)TotalCompressionExecutionTime, &Options));
+
+		const FErrorTrackerWorstBone WorstBone = WorstBoneError.GetMaxErrorItem();
+		const FErrorTrackerWorstAnimation WorstAnimation = WorstAnimationError.GetMaxErrorItem();
+
+		Args.Add(TEXT("AverageError"), FText::AsNumber(AverageError, &Options));
+
+		Args.Add(TEXT("WorstBoneError"), WorstBone.ToText());
+		Args.Add(TEXT("WorstAnimationError"), WorstAnimation.ToText());
+
+		FText Message;
+		if (!bPerformedCompression)
+		{
+			Message = FText::Format(NSLOCTEXT("Engine", "CompressionMemorySummaryNoChange", "Compressed {NumberOfAnimations} Animation(s)\n\nRaw: {TotalRaw} - Compressed: {TotalAfterCompressed}\nSaving: {TotalAfterSaving} ({NewCompressionRatio})\n\nEnd Effector Translation Added By Compression:\n Average: {AverageError} Max:\n{WorstBoneError}\n\nMax Average Animation Error:\n{WorstAnimationError}"), Args);
+		}
+		else
 		{
-			const int32 TotalBeforeSaving = TotalRaw - TotalBeforeCompressed;
-			const int32 TotalAfterSaving = TotalRaw - TotalAfterCompressed;
-			const float OldCompressionRatio = (TotalBeforeCompressed > 0.f) ? (static_cast<float>(TotalRaw) / TotalBeforeCompressed) : 0.f;
-			const float NewCompressionRatio = (TotalAfterCompressed > 0.f) ? (static_cast<float>(TotalRaw) / TotalAfterCompressed) : 0.f;
-
-			FNumberFormattingOptions Options;
-			Options.MinimumIntegralDigits = 7;
-			Options.MinimumFractionalDigits = 2;
-
-			FFormatNamedArguments Args;
-			Args.Add(TEXT("TotalRaw"), FText::AsMemory(TotalRaw, &Options));
-			Args.Add(TEXT("TotalBeforeCompressed"), FText::AsMemory(TotalBeforeCompressed, &Options));
-			Args.Add(TEXT("TotalBeforeSaving"), FText::AsMemory(TotalBeforeSaving, &Options));
-			Args.Add(TEXT("NumberOfAnimations"), FText::AsNumber(NumberOfAnimations));
-			Args.Add(TEXT("OldCompressionRatio"), OldCompressionRatio);
-
-			Args.Add(TEXT("TotalAfterCompressed"), FText::AsMemory(TotalAfterCompressed, &Options));
-			Args.Add(TEXT("TotalAfterSaving"), FText::AsMemory(TotalAfterSaving, &Options));
-			Args.Add(TEXT("NewCompressionRatio"), NewCompressionRatio);
-			Args.Add(TEXT("TotalTimeSpentCompressingPretty"), FText::FromString(FPlatformTime::PrettyTime(TotalCompressionExecutionTime)));
-			Args.Add(TEXT("TotalTimeSpentCompressingRawSeconds"), FText::AsNumber((float)TotalCompressionExecutionTime, &Options));
-
-			const FErrorTrackerWorstBone WorstBone = WorstBoneError.GetMaxErrorItem();
-			const FErrorTrackerWorstAnimation WorstAnimation = WorstAnimationError.GetMaxErrorItem();
-
-			Args.Add(TEXT("AverageError"), FText::AsNumber(AverageError, &Options));
-
-			Args.Add(TEXT("WorstBoneError"), WorstBone.ToText());
-			Args.Add(TEXT("WorstAnimationError"), WorstAnimation.ToText());
-
-			const FText Message = FText::Format(NSLOCTEXT("Engine", "CompressionMemorySummary", "Compressed {NumberOfAnimations} Animation(s)\n\nPre Compression:\n\nRaw: {TotalRaw} - Compressed: {TotalBeforeCompressed}\nSaving: {TotalBeforeSaving} ({OldCompressionRatio})\n\nPost Compression:\n\nRaw: {TotalRaw} - Compressed: {TotalAfterCompressed}\nSaving: {TotalAfterSaving} ({NewCompressionRatio})\n\nTotal Compression Time: {TotalTimeSpentCompressingPretty} (Seconds: {TotalTimeSpentCompressingRawSeconds})\n\nEnd Effector Translation Added By Compression:\n Average: {AverageError} Max:\n{WorstBoneError}\n\nMax Average Animation Error:\n{WorstAnimationError}"), Args);
-
-			UE_LOG(LogAnimationCompression, Display, TEXT("Top 10 Worst Bone Errors:"));
-			WorstBoneError.LogErrorStat();
-			UE_LOG(LogAnimationCompression, Display, TEXT("Top 10 Worst Average Animation Errors:"));
-			WorstAnimationError.LogErrorStat();
-			FMessageDialog::Open(EAppMsgType::Ok, Message);
+			Message = FText::Format(NSLOCTEXT("Engine", "CompressionMemorySummary", "Compressed {NumberOfAnimations} Animation(s)\n\nPre Compression:\n\nRaw: {TotalRaw} - Compressed: {TotalBeforeCompressed}\nSaving: {TotalBeforeSaving} ({OldCompressionRatio})\n\nPost Compression:\n\nRaw: {TotalRaw} - Compressed: {TotalAfterCompressed}\nSaving: {TotalAfterSaving} ({NewCompressionRatio})\n\nTotal Compression Time: {TotalTimeSpentCompressingPretty} (Seconds: {TotalTimeSpentCompressingRawSeconds})\n\nEnd Effector Translation Added By Compression:\n Average: {AverageError} Max:\n{WorstBoneError}\n\nMax Average Animation Error:\n{WorstAnimationError}"), Args);
 		}
+
+		UE_LOG(LogAnimationCompression, Display, TEXT("Top 10 Worst Bone Errors:"));
+		WorstBoneError.LogErrorStat();
+		UE_LOG(LogAnimationCompression, Display, TEXT("Top 10 Worst Average Animation Errors:"));
+		WorstAnimationError.LogErrorStat();
+		FMessageDialog::Open(EAppMsgType::Ok, Message);
 	}
 }
 
@@ -202,14 +205,13 @@ uint8 MakeBitForFlag(uint32 Item, uint32 Position)
 //////////////////////////////////////////////////////////////////////////////////////
 // FCompressionMemorySummary
 
-
-void FCompressionMemorySummary::GatherPreCompressionStats(const FCompressibleAnimData& CompressibleAnimData, int32 PreCompressedSize, int32 ProgressNumerator, int32 ProgressDenominator)
+void FCompressionMemorySummary::GatherPreCompressionStats(const FString& Name, int32 RawSize, int32 PreCompressedSize, int32 ProgressNumerator, int32 ProgressDenominator)
 {
 	if (bEnabled)
 	{
 		bUsed = true;
 		FFormatNamedArguments Args;
-		Args.Add(TEXT("AnimSequenceName"), FText::FromString(CompressibleAnimData.Name));
+		Args.Add(TEXT("AnimSequenceName"), FText::FromString(Name));
 		Args.Add(TEXT("ProgressNumerator"), ProgressNumerator);
 		Args.Add(TEXT("ProgressDenominator"), ProgressDenominator);
 
@@ -217,7 +219,7 @@ void FCompressionMemorySummary::GatherPreCompressionStats(const FCompressibleAni
 			ProgressDenominator,
 			FText::Format(NSLOCTEXT("CompressionMemorySummary", "CompressingTaskStatusMessageFormat", "Compressing {AnimSequenceName} ({ProgressNumerator}/{ProgressDenominator})"), Args));
 
-		TotalRaw += CompressibleAnimData.GetApproxRawBoneSize();
+		TotalRaw += RawSize;
 		TotalBeforeCompressed += PreCompressedSize;
 		++NumberOfAnimations;
 	}
@@ -225,9 +227,10 @@ void FCompressionMemorySummary::GatherPreCompressionStats(const FCompressibleAni
 
 
 //////////////////////////////////////////////////////////////////////////////////////
-void FAnimCompressContext::GatherPreCompressionStats(const FCompressibleAnimData& CompressibleAnimData, int32 PreviousCompressionSize)
+
+void FAnimCompressContext::GatherPreCompressionStats(const FString& Name, int32 RawSize, int32 PreviousCompressionSize)
 {
-	CompressionSummary.GatherPreCompressionStats(CompressibleAnimData, PreviousCompressionSize, AnimIndex, MaxAnimations);
+	CompressionSummary.GatherPreCompressionStats(Name, RawSize, PreviousCompressionSize, AnimIndex, MaxAnimations);
 }
 
 
@@ -279,6 +282,8 @@ void UAnimCompress::BitwiseCompressAnimationTracks(
 	const TArray<FScaleTrack>& ScaleData,
 	bool IncludeKeyTable)
 {
+	FUECompressedAnimDataMutable& AnimData = static_cast<FUECompressedAnimDataMutable&>(*OutCompressedData.AnimData);
+
 	// Ensure supported compression formats.
 	bool bInvalidCompressionFormat = false;
 	if (!(TargetTranslationFormat == ACF_None) && !(TargetTranslationFormat == ACF_IntervalFixed32NoW) && !(TargetTranslationFormat == ACF_Float96NoW))
@@ -298,18 +303,18 @@ void UAnimCompress::BitwiseCompressAnimationTracks(
 	}
 	if (bInvalidCompressionFormat)
 	{
-		OutCompressedData.TranslationCompressionFormat = ACF_None;
-		OutCompressedData.RotationCompressionFormat = ACF_None;
-		OutCompressedData.ScaleCompressionFormat = ACF_None;
-		OutCompressedData.CompressedTrackOffsets.Empty();
-		OutCompressedData.CompressedScaleOffsets.Empty();
-		OutCompressedData.CompressedByteStream.Empty();
+		AnimData.TranslationCompressionFormat = ACF_None;
+		AnimData.RotationCompressionFormat = ACF_None;
+		AnimData.ScaleCompressionFormat = ACF_None;
+		AnimData.CompressedTrackOffsets.Empty();
+		AnimData.CompressedScaleOffsets.Empty();
+		AnimData.CompressedByteStream.Empty();
 	}
 	else
 	{
-		OutCompressedData.RotationCompressionFormat = TargetRotationFormat;
-		OutCompressedData.TranslationCompressionFormat = TargetTranslationFormat;
-		OutCompressedData.ScaleCompressionFormat = TargetScaleFormat;
+		AnimData.RotationCompressionFormat = TargetRotationFormat;
+		AnimData.TranslationCompressionFormat = TargetTranslationFormat;
+		AnimData.ScaleCompressionFormat = TargetScaleFormat;
 
 		check(TranslationData.Num() == RotationData.Num());
 		const int32 NumTracks = RotationData.Num();
@@ -320,26 +325,26 @@ void UAnimCompress::BitwiseCompressAnimationTracks(
 			UE_LOG(LogAnimationCompression, Warning, TEXT("When compressing %s: no key-reduced data"), *CompressibleAnimData.Name);
 		}
 
-		OutCompressedData.CompressedTrackOffsets.Empty(NumTracks * 4);
-		OutCompressedData.CompressedTrackOffsets.AddUninitialized(NumTracks * 4);
+		AnimData.CompressedTrackOffsets.Empty(NumTracks * 4);
+		AnimData.CompressedTrackOffsets.AddUninitialized(NumTracks * 4);
 
 		// just empty it since there is chance this can be 0
-		OutCompressedData.CompressedScaleOffsets.Empty();
+		AnimData.CompressedScaleOffsets.Empty();
 		// only do this if Scale exists;
 		if (bHasScale)
 		{
-			OutCompressedData.CompressedScaleOffsets.SetStripSize(2);
-			OutCompressedData.CompressedScaleOffsets.AddUninitialized(NumTracks);
+			AnimData.CompressedScaleOffsets.SetStripSize(2);
+			AnimData.CompressedScaleOffsets.AddUninitialized(NumTracks);
 		}
 
-		OutCompressedData.CompressedByteStream.Empty();
+		AnimData.CompressedByteStream.Empty();
 
 		for (int32 TrackIndex = 0; TrackIndex < NumTracks; ++TrackIndex)
 		{
 			// Translation data.
 			const FTranslationTrack& SrcTrans = TranslationData[TrackIndex];
 
-			const int32 OffsetTrans = OutCompressedData.CompressedByteStream.Num();
+			const int32 OffsetTrans = AnimData.CompressedByteStream.Num();
 			const int32 NumKeysTrans = SrcTrans.PosKeys.Num();
 
 			// Warn on empty data.
@@ -349,8 +354,8 @@ void UAnimCompress::BitwiseCompressAnimationTracks(
 			}
 
 			checkf((OffsetTrans % 4) == 0, TEXT("CompressedByteStream not aligned to four bytes"));
-			OutCompressedData.CompressedTrackOffsets[TrackIndex * 4] = OffsetTrans;
-			OutCompressedData.CompressedTrackOffsets[TrackIndex * 4 + 1] = NumKeysTrans;
+			AnimData.CompressedTrackOffsets[TrackIndex * 4] = OffsetTrans;
+			AnimData.CompressedTrackOffsets[TrackIndex * 4 + 1] = NumKeysTrans;
 
 			// Calculate the bounding box of the translation keys
 			FBox PositionBounds(SrcTrans.PosKeys);
@@ -366,21 +371,21 @@ void UAnimCompress::BitwiseCompressAnimationTracks(
 				// Write the mins and ranges if they'll be used on the other side
 				if (TargetTranslationFormat == ACF_IntervalFixed32NoW)
 				{
-					UnalignedWriteToStream(OutCompressedData.CompressedByteStream, TransMins, sizeof(float) * 3);
-					UnalignedWriteToStream(OutCompressedData.CompressedByteStream, TransRanges, sizeof(float) * 3);
+					UnalignedWriteToStream(AnimData.CompressedByteStream, TransMins, sizeof(float) * 3);
+					UnalignedWriteToStream(AnimData.CompressedByteStream, TransRanges, sizeof(float) * 3);
 				}
 
 				// Pack the positions into the stream
 				for (int32 KeyIndex = 0; KeyIndex < NumKeysTrans; ++KeyIndex)
 				{
 					const FVector& Vec = SrcTrans.PosKeys[KeyIndex];
-					PackVectorToStream(OutCompressedData.CompressedByteStream, TargetTranslationFormat, Vec, TransMins, TransRanges);
+					PackVectorToStream(AnimData.CompressedByteStream, TargetTranslationFormat, Vec, TransMins, TransRanges);
 				}
 
 				if (IncludeKeyTable)
 				{
 					// Align to four bytes.
-					PadByteStream(OutCompressedData.CompressedByteStream, 4, AnimationPadSentinel);
+					PadByteStream(AnimData.CompressedByteStream, 4, AnimationPadSentinel);
 
 					// write the key table
 					const int32 NumFrames = CompressibleAnimData.NumFrames;
@@ -390,7 +395,7 @@ void UAnimCompress::BitwiseCompressAnimationTracks(
 
 					const int32 TableSize = NumKeysTrans*FrameSize;
 					const int32 TableDwords = (TableSize + 3) >> 2;
-					const int32 StartingOffset = OutCompressedData.CompressedByteStream.Num();
+					const int32 StartingOffset = AnimData.CompressedByteStream.Num();
 
 					for (int32 KeyIndex = 0; KeyIndex < NumKeysTrans; ++KeyIndex)
 					{
@@ -398,20 +403,20 @@ void UAnimCompress::BitwiseCompressAnimationTracks(
 						float KeyTime = SrcTrans.Times[KeyIndex];
 						float FrameTime = KeyTime * FrameRate;
 						int32 FrameIndex = FMath::Clamp(FMath::TruncToInt(FrameTime + 0.5f), 0, LastFrame);
-						UnalignedWriteToStream(OutCompressedData.CompressedByteStream, &FrameIndex, FrameSize);
+						UnalignedWriteToStream(AnimData.CompressedByteStream, &FrameIndex, FrameSize);
 					}
 
 					// Align to four bytes. Padding with 0's to round out the key table
-					PadByteStream(OutCompressedData.CompressedByteStream, 4, 0);
+					PadByteStream(AnimData.CompressedByteStream, 4, 0);
 
-					const int32 EndingOffset = OutCompressedData.CompressedByteStream.Num();
+					const int32 EndingOffset = AnimData.CompressedByteStream.Num();
 					check((EndingOffset - StartingOffset) == (TableDwords * 4));
 				}
 			}
 			else if (NumKeysTrans == 1)
 			{
 				// A single translation key gets written out a single uncompressed float[3].
-				UnalignedWriteToStream(OutCompressedData.CompressedByteStream, &(SrcTrans.PosKeys[0]), sizeof(FVector));
+				UnalignedWriteToStream(AnimData.CompressedByteStream, &(SrcTrans.PosKeys[0]), sizeof(FVector));
 			}
 			else
 			{
@@ -419,16 +424,16 @@ void UAnimCompress::BitwiseCompressAnimationTracks(
 			}
 
 			// Align to four bytes.
-			PadByteStream(OutCompressedData.CompressedByteStream, 4, AnimationPadSentinel);
+			PadByteStream(AnimData.CompressedByteStream, 4, AnimationPadSentinel);
 
 			// Compress rotation data.
 			const FRotationTrack& SrcRot = RotationData[TrackIndex];
-			const int32 OffsetRot = OutCompressedData.CompressedByteStream.Num();
+			const int32 OffsetRot = AnimData.CompressedByteStream.Num();
 			const int32 NumKeysRot = SrcRot.RotKeys.Num();
 
 			checkf((OffsetRot % 4) == 0, TEXT("CompressedByteStream not aligned to four bytes"));
-			OutCompressedData.CompressedTrackOffsets[TrackIndex * 4 + 2] = OffsetRot;
-			OutCompressedData.CompressedTrackOffsets[TrackIndex * 4 + 3] = NumKeysRot;
+			AnimData.CompressedTrackOffsets[TrackIndex * 4 + 2] = OffsetRot;
+			AnimData.CompressedTrackOffsets[TrackIndex * 4 + 3] = NumKeysRot;
 
 			if (NumKeysRot > 1)
 			{
@@ -467,22 +472,22 @@ void UAnimCompress::BitwiseCompressAnimationTracks(
 				// Write the mins and ranges if they'll be used on the other side
 				if (TargetRotationFormat == ACF_IntervalFixed32NoW)
 				{
-					UnalignedWriteToStream(OutCompressedData.CompressedByteStream, Mins, sizeof(float) * 3);
-					UnalignedWriteToStream(OutCompressedData.CompressedByteStream, Ranges, sizeof(float) * 3);
+					UnalignedWriteToStream(AnimData.CompressedByteStream, Mins, sizeof(float) * 3);
+					UnalignedWriteToStream(AnimData.CompressedByteStream, Ranges, sizeof(float) * 3);
 				}
 
 				// n elements of the compressed type.
 				for (int32 KeyIndex = 0; KeyIndex < SrcRot.RotKeys.Num(); ++KeyIndex)
 				{
 					const FQuat& Quat = SrcRot.RotKeys[KeyIndex];
-					PackQuaternionToStream(OutCompressedData.CompressedByteStream, TargetRotationFormat, Quat, Mins, Ranges);
+					PackQuaternionToStream(AnimData.CompressedByteStream, TargetRotationFormat, Quat, Mins, Ranges);
 				}
 
 				// n elements of frame indices
 				if (IncludeKeyTable)
 				{
 					// Align to four bytes.
-					PadByteStream(OutCompressedData.CompressedByteStream, 4, AnimationPadSentinel);
+					PadByteStream(AnimData.CompressedByteStream, 4, AnimationPadSentinel);
 
 					// write the key table
 					const int32 NumFrames = CompressibleAnimData.NumFrames;
@@ -492,7 +497,7 @@ void UAnimCompress::BitwiseCompressAnimationTracks(
 
 					const int32 TableSize = NumKeysRot*FrameSize;
 					const int32 TableDwords = (TableSize + 3) >> 2;
-					const int32 StartingOffset = OutCompressedData.CompressedByteStream.Num();
+					const int32 StartingOffset = AnimData.CompressedByteStream.Num();
 
 					for (int32 KeyIndex = 0; KeyIndex < NumKeysRot; ++KeyIndex)
 					{
@@ -500,13 +505,13 @@ void UAnimCompress::BitwiseCompressAnimationTracks(
 						float KeyTime = SrcRot.Times[KeyIndex];
 						float FrameTime = KeyTime * FrameRate;
 						int32 FrameIndex = FMath::Clamp(FMath::TruncToInt(FrameTime + 0.5f), 0, LastFrame);
-						UnalignedWriteToStream(OutCompressedData.CompressedByteStream, &FrameIndex, FrameSize);
+						UnalignedWriteToStream(AnimData.CompressedByteStream, &FrameIndex, FrameSize);
 					}
 
 					// Align to four bytes. Padding with 0's to round out the key table
-					PadByteStream(OutCompressedData.CompressedByteStream, 4, 0);
+					PadByteStream(AnimData.CompressedByteStream, 4, 0);
 
-					const int32 EndingOffset = OutCompressedData.CompressedByteStream.Num();
+					const int32 EndingOffset = AnimData.CompressedByteStream.Num();
 					check((EndingOffset - StartingOffset) == (TableDwords * 4));
 
 				}
@@ -516,7 +521,7 @@ void UAnimCompress::BitwiseCompressAnimationTracks(
 				// For a rotation track of n=1 keys, the single key is packed as an FQuatFloat96NoW.
 				const FQuat& Quat = SrcRot.RotKeys[0];
 				const FQuatFloat96NoW QuatFloat96NoW(Quat);
-				UnalignedWriteToStream(OutCompressedData.CompressedByteStream, &QuatFloat96NoW, sizeof(FQuatFloat96NoW));
+				UnalignedWriteToStream(AnimData.CompressedByteStream, &QuatFloat96NoW, sizeof(FQuatFloat96NoW));
 			}
 			else
 			{
@@ -525,19 +530,19 @@ void UAnimCompress::BitwiseCompressAnimationTracks(
 
 
 			// Align to four bytes.
-			PadByteStream(OutCompressedData.CompressedByteStream, 4, AnimationPadSentinel);
+			PadByteStream(AnimData.CompressedByteStream, 4, AnimationPadSentinel);
 
 			// we also should do this only when scale exists. 
 			if (bHasScale)
 			{
 				const FScaleTrack& SrcScale = ScaleData[TrackIndex];
 
-				const int32 OffsetScale = OutCompressedData.CompressedByteStream.Num();
+				const int32 OffsetScale = AnimData.CompressedByteStream.Num();
 				const int32 NumKeysScale = SrcScale.ScaleKeys.Num();
 
 				checkf((OffsetScale % 4) == 0, TEXT("CompressedByteStream not aligned to four bytes"));
-				OutCompressedData.CompressedScaleOffsets.SetOffsetData(TrackIndex, 0, OffsetScale);
-				OutCompressedData.CompressedScaleOffsets.SetOffsetData(TrackIndex, 1, NumKeysScale);
+				AnimData.CompressedScaleOffsets.SetOffsetData(TrackIndex, 0, OffsetScale);
+				AnimData.CompressedScaleOffsets.SetOffsetData(TrackIndex, 1, NumKeysScale);
 
 				// Calculate the bounding box of the Scalelation keys
 				FBox ScaleBoundsBounds(SrcScale.ScaleKeys);
@@ -554,21 +559,21 @@ void UAnimCompress::BitwiseCompressAnimationTracks(
 					// Write the mins and ranges if they'll be used on the other side
 					if (TargetScaleFormat == ACF_IntervalFixed32NoW)
 					{
-						UnalignedWriteToStream(OutCompressedData.CompressedByteStream, ScaleMins, sizeof(float) * 3);
-						UnalignedWriteToStream(OutCompressedData.CompressedByteStream, ScaleRanges, sizeof(float) * 3);
+						UnalignedWriteToStream(AnimData.CompressedByteStream, ScaleMins, sizeof(float) * 3);
+						UnalignedWriteToStream(AnimData.CompressedByteStream, ScaleRanges, sizeof(float) * 3);
 					}
 
 					// Pack the positions into the stream
 					for (int32 KeyIndex = 0; KeyIndex < NumKeysScale; ++KeyIndex)
 					{
 						const FVector& Vec = SrcScale.ScaleKeys[KeyIndex];
-						PackVectorToStream(OutCompressedData.CompressedByteStream, TargetScaleFormat, Vec, ScaleMins, ScaleRanges);
+						PackVectorToStream(AnimData.CompressedByteStream, TargetScaleFormat, Vec, ScaleMins, ScaleRanges);
 					}
 
 					if (IncludeKeyTable)
 					{
 						// Align to four bytes.
-						PadByteStream(OutCompressedData.CompressedByteStream, 4, AnimationPadSentinel);
+						PadByteStream(AnimData.CompressedByteStream, 4, AnimationPadSentinel);
 
 						// write the key table
 						const int32 NumFrames = CompressibleAnimData.NumFrames;
@@ -578,7 +583,7 @@ void UAnimCompress::BitwiseCompressAnimationTracks(
 
 						const int32 TableSize = NumKeysScale*FrameSize;
 						const int32 TableDwords = (TableSize + 3) >> 2;
-						const int32 StartingOffset = OutCompressedData.CompressedByteStream.Num();
+						const int32 StartingOffset = AnimData.CompressedByteStream.Num();
 
 						for (int32 KeyIndex = 0; KeyIndex < NumKeysScale; ++KeyIndex)
 						{
@@ -586,20 +591,20 @@ void UAnimCompress::BitwiseCompressAnimationTracks(
 							float KeyTime = SrcScale.Times[KeyIndex];
 							float FrameTime = KeyTime * FrameRate;
 							int32 FrameIndex = FMath::Clamp(FMath::TruncToInt(FrameTime + 0.5f), 0, LastFrame);
-							UnalignedWriteToStream(OutCompressedData.CompressedByteStream, &FrameIndex, FrameSize);
+							UnalignedWriteToStream(AnimData.CompressedByteStream, &FrameIndex, FrameSize);
 						}
 
 						// Align to four bytes. Padding with 0's to round out the key table
-						PadByteStream(OutCompressedData.CompressedByteStream, 4, 0);
+						PadByteStream(AnimData.CompressedByteStream, 4, 0);
 
-						const int32 EndingOffset = OutCompressedData.CompressedByteStream.Num();
+						const int32 EndingOffset = AnimData.CompressedByteStream.Num();
 						check((EndingOffset - StartingOffset) == (TableDwords * 4));
 					}
 				}
 				else if (NumKeysScale == 1)
 				{
 					// A single Scalelation key gets written out a single uncompressed float[3].
-					UnalignedWriteToStream(OutCompressedData.CompressedByteStream, &(SrcScale.ScaleKeys[0]), sizeof(FVector));
+					UnalignedWriteToStream(AnimData.CompressedByteStream, &(SrcScale.ScaleKeys[0]), sizeof(FVector));
 				}
 				else
 				{
@@ -607,982 +612,105 @@ void UAnimCompress::BitwiseCompressAnimationTracks(
 				}
 
 				// Align to four bytes.
-				PadByteStream(OutCompressedData.CompressedByteStream, 4, AnimationPadSentinel);
+				PadByteStream(AnimData.CompressedByteStream, 4, AnimationPadSentinel);
 			}
 		}
 
 		// Trim unused memory.
-		OutCompressedData.CompressedByteStream.Shrink();
+		AnimData.CompressedByteStream.Shrink();
 	}
 }
 
-#if USE_SEGMENTING_CONTEXT
-void UAnimCompress::BitwiseCompressAnimationTracks(
-	class UAnimSequence& AnimSeq,
-	AnimationCompressionFormat TargetTranslationFormat,
-	AnimationCompressionFormat TargetRotationFormat,
-	AnimationCompressionFormat TargetScaleFormat,
-	TArray<FAnimSegmentContext>& RawSegments,
-	bool bIsSorted)
-{
-	// Ensure supported compression formats.
-	bool bInvalidCompressionFormat = false;
-	if (!(TargetTranslationFormat == ACF_None) && !(TargetTranslationFormat == ACF_IntervalFixed32NoW) && !(TargetTranslationFormat == ACF_Float96NoW))
-	{
-		FMessageDialog::Open(EAppMsgType::Ok, FText::Format(NSLOCTEXT("Engine", "UnknownTranslationCompressionFormat", "Unknown or unsupported translation compression format ({0})"), FText::AsNumber((int32)TargetTranslationFormat)));
-		bInvalidCompressionFormat = true;
-	}
-	if (!(TargetRotationFormat >= ACF_None && TargetRotationFormat < ACF_MAX))
-	{
-		FMessageDialog::Open(EAppMsgType::Ok, FText::Format(NSLOCTEXT("Engine", "UnknownRotationCompressionFormat", "Unknown or unsupported rotation compression format ({0})"), FText::AsNumber((int32)TargetRotationFormat)));
-		bInvalidCompressionFormat = true;
-	}
-	if (!(TargetScaleFormat == ACF_None) && !(TargetScaleFormat == ACF_IntervalFixed32NoW) && !(TargetScaleFormat == ACF_Float96NoW))
-	{
-		FMessageDialog::Open(EAppMsgType::Ok, FText::Format(NSLOCTEXT("Engine", "UnknownScaleCompressionFormat", "Unknown or unsupported Scale compression format ({0})"), FText::AsNumber((int32)TargetScaleFormat)));
-		bInvalidCompressionFormat = true;
-	}
-
-	if (!bInvalidCompressionFormat && RawSegments.Num() > 0)
-	{
-		// First segment holds the compressed trivial tracks
-		BitwiseCompressTrivialAnimationTracks(AnimSeq, RawSegments[0]);
-	}
-
-	for (FAnimSegmentContext& Segment : RawSegments)
-	{
-		if (bInvalidCompressionFormat)
-		{
-			Segment.TranslationCompressionFormat = ACF_None;
-			Segment.RotationCompressionFormat = ACF_None;
-			Segment.ScaleCompressionFormat = ACF_None;
-			Segment.CompressedTrackOffsets.Empty();
-			Segment.CompressedScaleOffsets.Empty();
-			Segment.CompressedByteStream.Empty();
-			Segment.CompressedTrivialTracksByteStream.Empty();
-		}
-		else
-		{
-			BitwiseCompressAnimationTracks(AnimSeq, TargetTranslationFormat, TargetRotationFormat, TargetScaleFormat, Segment, bIsSorted);
-		}
-	}
-}
-
-void UAnimCompress::SanityCheckTrackData(const UAnimSequence& AnimSeq, const FAnimSegmentContext& Segment)
-{
-	check(Segment.TranslationData.Num() == Segment.RotationData.Num());
-
-	const int32 NumTracks = Segment.RotationData.Num();
-	const bool bHasScale = Segment.ScaleData.Num() > 0;
-
-	if (NumTracks == 0)
-	{
-		UE_LOG(LogAnimationCompression, Warning, TEXT("When compressing %s: no key-reduced data"), *AnimSeq.GetFullName());
-	}
-
-	for (int32 TrackIndex = 0; TrackIndex < NumTracks; ++TrackIndex)
-	{
-		const FTranslationTrack& SrcTrans = Segment.TranslationData[TrackIndex];
-		if (SrcTrans.PosKeys.Num() == 0)
-		{
-			UE_LOG(LogAnimationCompression, Warning, TEXT("When compressing %s track %i: no translation keys"), *AnimSeq.GetFullName(), TrackIndex);
-		}
-
-		const FRotationTrack& SrcRot = Segment.RotationData[TrackIndex];
-		if (SrcRot.RotKeys.Num() == 0)
-		{
-			UE_LOG(LogAnimationCompression, Warning, TEXT("When compressing %s track %i: no rotation keys"), *AnimSeq.GetFullName(), TrackIndex);
-		}
-
-		if (bHasScale)
-		{
-			const FScaleTrack& SrcScale = Segment.ScaleData[TrackIndex];
-			if (SrcScale.ScaleKeys.Num() == 0)
-			{
-				UE_LOG(LogAnimationCompression, Warning, TEXT("When compressing %s track %i: no Scale keys"), *AnimSeq.GetFullName(), TrackIndex);
-			}
-		}
-	}
-}
-#endif
-
-void UAnimCompress::CalculateTrackRange(const FTranslationTrack& TranslationData, AnimationCompressionFormat Format, FVector& OutMin, FVector& OutExtent)
-{
-	FVector Mins;
-	FVector Ranges;
-
-	const int32 NumKeysTrans = TranslationData.PosKeys.Num();
-	if (NumKeysTrans > 1 && Format == ACF_IntervalFixed32NoW)
-	{
-		// Calculate the range of the XYZ components of the translation
-		FBox PositionBounds(TranslationData.PosKeys);
-
-		Mins = PositionBounds.Min;
-		Ranges = PositionBounds.Max - PositionBounds.Min;
-
-		if (Ranges.X == 0.f) { Ranges.X = 1.f; }
-		if (Ranges.Y == 0.f) { Ranges.Y = 1.f; }
-		if (Ranges.Z == 0.f) { Ranges.Z = 1.f; }
-	}
-	else
-	{
-		Mins = FVector::ZeroVector;
-		Ranges = FVector::ZeroVector;
-	}
-
-	OutMin = Mins;
-	OutExtent = Ranges;
-}
-
-void UAnimCompress::CalculateTrackRange(const FRotationTrack& RotationData, AnimationCompressionFormat Format, FVector& OutMin, FVector& OutExtent)
-{
-	FVector Mins;
-	FVector Ranges;
-
-	const int32 NumKeysRot = RotationData.RotKeys.Num();
-	if (NumKeysRot > 1 && Format == ACF_IntervalFixed32NoW)
-	{
-		// Calculate the range of the XYZ components of the quaternion
-		Mins = FVector(1.f, 1.f, 1.f);
-		FVector Maxs(-1.f, -1.f, -1.f);
-
-		for (int32 KeyIndex = 0; KeyIndex < NumKeysRot; ++KeyIndex)
-		{
-			FQuat Quat(RotationData.RotKeys[KeyIndex]);
-			if (Quat.W < 0.f)
-			{
-				Quat.X = -Quat.X;
-				Quat.Y = -Quat.Y;
-				Quat.Z = -Quat.Z;
-				Quat.W = -Quat.W;
-			}
-			Quat.Normalize();
-
-			FVector QuatV(Quat.X, Quat.Y, Quat.Z);
-			Mins = Mins.ComponentMin(QuatV);
-			Maxs = Maxs.ComponentMax(QuatV);
-		}
-
-		Ranges = Maxs - Mins;
-		if (Ranges.X == 0.f) { Ranges.X = 1.f; }
-		if (Ranges.Y == 0.f) { Ranges.Y = 1.f; }
-		if (Ranges.Z == 0.f) { Ranges.Z = 1.f; }
-	}
-	else
-	{
-		Mins = FVector::ZeroVector;
-		Ranges = FVector::ZeroVector;
-	}
-
-	OutMin = Mins;
-	OutExtent = Ranges;
-}
-
-void UAnimCompress::CalculateTrackRange(const FScaleTrack& ScaleData, AnimationCompressionFormat Format, FVector& OutMin, FVector& OutExtent)
-{
-	FVector Mins;
-	FVector Ranges;
-
-	const int32 NumKeysScale = ScaleData.ScaleKeys.Num();
-	if (NumKeysScale > 1 && Format == ACF_IntervalFixed32NoW)
-	{
-		// Calculate the range of the XYZ components of the translation
-		FBox Bounds(ScaleData.ScaleKeys);
-
-		Mins = Bounds.Min;
-		Ranges = Bounds.Max - Bounds.Min;
-
-		if (Ranges.X == 0.f) { Ranges.X = 1.f; }
-		if (Ranges.Y == 0.f) { Ranges.Y = 1.f; }
-		if (Ranges.Z == 0.f) { Ranges.Z = 1.f; }
-	}
-	else
-	{
-		Mins = FVector::ZeroVector;
-		Ranges = FVector::ZeroVector;
-	}
-
-	OutMin = Mins;
-	OutExtent = Ranges;
-}
+#endif // WITH_EDITOR
 
-#if USE_SEGMENTING_CONTEXT
-void UAnimCompress::CalculateTrackRanges(
-	AnimationCompressionFormat TargetTranslationFormat,
-	AnimationCompressionFormat TargetRotationFormat,
-	AnimationCompressionFormat TargetScaleFormat,
-	const FAnimSegmentContext& Segment,
-	TArray<FAnimTrackRange>& TrackRanges)
+#if WITH_EDITORONLY_DATA
+bool UAnimCompress::Compress(const FCompressibleAnimData& CompressibleAnimData, FCompressibleAnimDataResult& OutResult)
 {
-	const int32 NumTracks = Segment.RotationData.Num();
-	const bool bHasScale = Segment.ScaleData.Num() > 0;
-
-	TrackRanges.Empty(NumTracks);
-	TrackRanges.AddZeroed(NumTracks);
+	FUECompressedAnimDataMutable AnimDataMutable;
+	AnimDataMutable.CompressedNumberOfFrames = CompressibleAnimData.NumFrames;
 
-	for (int32 TrackIndex = 0; TrackIndex < NumTracks; ++TrackIndex)
-	{
-		FAnimTrackRange& TrackRange = TrackRanges[TrackIndex];
+	OutResult.Codec = this;
+	OutResult.AnimData = TUniquePtr<ICompressedAnimData>(&AnimDataMutable);
 
-		CalculateTrackRange(Segment.TranslationData[TrackIndex], TargetTranslationFormat, TrackRange.TransMin, TrackRange.TransExtent);
-		CalculateTrackRange(Segment.RotationData[TrackIndex], TargetRotationFormat, TrackRange.RotMin, TrackRange.RotExtent);
+	const bool bSuccess = DoReduction(CompressibleAnimData, OutResult);
 
-		if (bHasScale)
-		{
-			CalculateTrackRange(Segment.ScaleData[TrackIndex], TargetScaleFormat, TrackRange.ScaleMin, TrackRange.ScaleExtent);
-		}
-	}
-}
+	// Clear without free since we were on the stack
+	OutResult.AnimData.Release();
 
-void UAnimCompress::WriteTrackRanges(
-	TArray<uint8>& ByteStream,
-	TFunction<AnimationCompressionFormat(int32 TrackIndex)> GetTranslationFormatFun,
-	TFunction<AnimationCompressionFormat(int32 TrackIndex)> GetRotationFormatFun,
-	TFunction<AnimationCompressionFormat(int32 TrackIndex)> GetScaleFormatFun,
-	TFunction<FTrackKeyFlags(int32 TrackIndex)> GetTranslationFlagsFun,
-	TFunction<FTrackKeyFlags(int32 TrackIndex)> GetRotationFlagsFun,
-	TFunction<FTrackKeyFlags(int32 TrackIndex)> GetScaleFlagsFun,
-	const FAnimSegmentContext& Segment,
-	const TArray<FAnimTrackRange>& TrackRanges,
-	bool bInterleaveValues)
-{
-	const int32 NumTracks = Segment.RotationData.Num();
-	const bool bHasScale = Segment.ScaleData.Num() > 0;
-
-	for (int32 TrackIndex = 0; TrackIndex < NumTracks; ++TrackIndex)
+	if (bSuccess)
 	{
-		const FAnimTrackRange& TrackRange = TrackRanges[TrackIndex];
-
-		const FTranslationTrack& SrcTrans = Segment.TranslationData[TrackIndex];
-		if (SrcTrans.PosKeys.Num() > 1 && GetTranslationFormatFun(TrackIndex) == ACF_IntervalFixed32NoW)
-		{
-			if (bInterleaveValues)
-			{
-				const FTrackKeyFlags Flags = GetTranslationFlagsFun(TrackIndex);
-				if (Flags.IsComponentNeededX())
-				{
-					UnalignedWriteToStream(ByteStream, &TrackRange.TransMin.X, sizeof(float));
-					UnalignedWriteToStream(ByteStream, &TrackRange.TransExtent.X, sizeof(float));
-				}
-
-				if (Flags.IsComponentNeededY())
-				{
-					UnalignedWriteToStream(ByteStream, &TrackRange.TransMin.Y, sizeof(float));
-					UnalignedWriteToStream(ByteStream, &TrackRange.TransExtent.Y, sizeof(float));
-				}
-
-				if (Flags.IsComponentNeededZ())
-				{
-					UnalignedWriteToStream(ByteStream, &TrackRange.TransMin.Z, sizeof(float));
-					UnalignedWriteToStream(ByteStream, &TrackRange.TransExtent.Z, sizeof(float));
-				}
-			}
-			else
-			{
-				UnalignedWriteToStream(ByteStream, &TrackRange.TransMin, sizeof(FVector));
-				UnalignedWriteToStream(ByteStream, &TrackRange.TransExtent, sizeof(FVector));
-			}
-		}
-
-		const FRotationTrack& SrcRot = Segment.RotationData[TrackIndex];
-		if (SrcRot.RotKeys.Num() > 1 && GetRotationFormatFun(TrackIndex) == ACF_IntervalFixed32NoW)
-		{
-			if (bInterleaveValues)
-			{
-				const FTrackKeyFlags Flags = GetRotationFlagsFun(TrackIndex);
-				if (Flags.IsComponentNeededX())
-				{
-					UnalignedWriteToStream(ByteStream, &TrackRange.RotMin.X, sizeof(float));
-					UnalignedWriteToStream(ByteStream, &TrackRange.RotExtent.X, sizeof(float));
-				}
-
-				if (Flags.IsComponentNeededY())
-				{
-					UnalignedWriteToStream(ByteStream, &TrackRange.RotMin.Y, sizeof(float));
-					UnalignedWriteToStream(ByteStream, &TrackRange.RotExtent.Y, sizeof(float));
-				}
+		// Finished compressing, coalesce the buffers
+		AnimDataMutable.BuildFinalBuffer(OutResult.CompressedByteStream);
 
-				if (Flags.IsComponentNeededZ())
-				{
-					UnalignedWriteToStream(ByteStream, &TrackRange.RotMin.Z, sizeof(float));
-					UnalignedWriteToStream(ByteStream, &TrackRange.RotExtent.Z, sizeof(float));
-				}
-			}
-			else
-			{
-				UnalignedWriteToStream(ByteStream, &TrackRange.RotMin, sizeof(FVector));
-				UnalignedWriteToStream(ByteStream, &TrackRange.RotExtent, sizeof(FVector));
-			}
-		}
-
-		if (bHasScale)
-		{
-			const FScaleTrack& SrcScale = Segment.ScaleData[TrackIndex];
-			if (SrcScale.ScaleKeys.Num() > 1 && GetScaleFormatFun(TrackIndex) == ACF_IntervalFixed32NoW)
-			{
-				if (bInterleaveValues)
-				{
-					const FTrackKeyFlags Flags = GetScaleFlagsFun(TrackIndex);
-					if (Flags.IsComponentNeededX())
-					{
-						UnalignedWriteToStream(ByteStream, &TrackRange.ScaleMin.X, sizeof(float));
-						UnalignedWriteToStream(ByteStream, &TrackRange.ScaleExtent.X, sizeof(float));
-					}
+		// Build our read-only version from the mutable source
+		TUniquePtr<FUECompressedAnimData> AnimData = MakeUnique<FUECompressedAnimData>(AnimDataMutable);
 
-					if (Flags.IsComponentNeededY())
-					{
-						UnalignedWriteToStream(ByteStream, &TrackRange.ScaleMin.Y, sizeof(float));
-						UnalignedWriteToStream(ByteStream, &TrackRange.ScaleExtent.Y, sizeof(float));
-					}
+		// The buffers will point to the mutable data, bind to the byte stream instead
+		AnimData->InitViewsFromBuffer(OutResult.CompressedByteStream);
 
-					if (Flags.IsComponentNeededZ())
-					{
-						UnalignedWriteToStream(ByteStream, &TrackRange.ScaleMin.Z, sizeof(float));
-						UnalignedWriteToStream(ByteStream, &TrackRange.ScaleExtent.Z, sizeof(float));
-					}
-				}
-				else
-				{
-					UnalignedWriteToStream(ByteStream, &TrackRange.ScaleMin, sizeof(FVector));
-					UnalignedWriteToStream(ByteStream, &TrackRange.ScaleExtent, sizeof(FVector));
-				}
-			}
-		}
+		OutResult.AnimData = MoveTemp(AnimData);
 	}
-}
 
-void UAnimCompress::WriteUniformTrackData(
-	TArray<uint8>& ByteStream,
-	TFunction<AnimationCompressionFormat(int32 TrackIndex)> GetTranslationFormatFun,
-	TFunction<AnimationCompressionFormat(int32 TrackIndex)> GetRotationFormatFun,
-	TFunction<AnimationCompressionFormat(int32 TrackIndex)> GetScaleFormatFun,
-	TFunction<bool(int32 TrackIndex)> IsTranslationUniformFun,
-	TFunction<bool(int32 TrackIndex)> IsRotationUniformFun,
-	TFunction<bool(int32 TrackIndex)> IsScaleUniformFun,
-	TFunction<void(TArray<uint8>& ByteStream, AnimationCompressionFormat Format, const FVector& Key, const float* Mins, const float* Ranges, int32 TrackIndex)> PackTranslationKeyFun,
-	TFunction<void(TArray<uint8>& ByteStream, AnimationCompressionFormat Format, const FQuat& Key, const float* Mins, const float* Ranges, int32 TrackIndex)> PackRotationKeyFun,
-	TFunction<void(TArray<uint8>& ByteStream, AnimationCompressionFormat Format, const FVector& Key, const float* Mins, const float* Ranges, int32 TrackIndex)> PackScaleKeyFun,
-	const FAnimSegmentContext& Segment,
-	const TArray<FAnimTrackRange>& TrackRanges)
-{
-	const int32 NumFrames = Segment.NumFrames;
-	const int32 NumTracks = Segment.RotationData.Num();
-	const bool bHasScale = Segment.ScaleData.Num() > 0;
-
-	// Samples are sorted by time first, track second to ensure that when we sample a time T, data for all tracks at
-	// that time is contiguous in memory
-	for (int32 KeyIndex = 0; KeyIndex < NumFrames; ++KeyIndex)
-	{
-		for (int32 TrackIndex = 0; TrackIndex < NumTracks; ++TrackIndex)
-		{
-			const FAnimTrackRange& TrackRange = TrackRanges[TrackIndex];
-
-			const FTranslationTrack& SrcTrans = Segment.TranslationData[TrackIndex];
-			if (SrcTrans.PosKeys.Num() > 1 && IsTranslationUniformFun(TrackIndex))
-			{
-				const AnimationCompressionFormat Format = GetTranslationFormatFun(TrackIndex);
-				const FVector& Vec = SrcTrans.PosKeys[KeyIndex];
-				PackTranslationKeyFun(ByteStream, Format, Vec, &TrackRange.TransMin.X, &TrackRange.TransExtent.X, TrackIndex);
-			}
-
-			const FRotationTrack& SrcRot = Segment.RotationData[TrackIndex];
-			if (SrcRot.RotKeys.Num() > 1 && IsRotationUniformFun(TrackIndex))
-			{
-				const AnimationCompressionFormat Format = GetRotationFormatFun(TrackIndex);
-				const FQuat& Quat = SrcRot.RotKeys[KeyIndex];
-				PackRotationKeyFun(ByteStream, Format, Quat, &TrackRange.RotMin.X, &TrackRange.RotExtent.X, TrackIndex);
-			}
-
-			if (bHasScale)
-			{
-				const FScaleTrack& SrcScale = Segment.ScaleData[TrackIndex];
-				if (SrcScale.ScaleKeys.Num() > 1 && IsScaleUniformFun(TrackIndex))
-				{
-					const AnimationCompressionFormat Format = GetScaleFormatFun(TrackIndex);
-					const FVector& Vec = SrcScale.ScaleKeys[KeyIndex];
-					PackScaleKeyFun(ByteStream, Format, Vec, &TrackRange.ScaleMin.X, &TrackRange.ScaleExtent.X, TrackIndex);
-				}
-			}
-		}
-	}
-}
-
-void UAnimCompress::WriteSortedVariableTrackData(
-	TArray<uint8>& ByteStream,
-	const UAnimSequence& AnimSeq,
-	TFunction<AnimationCompressionFormat(int32 TrackIndex)> GetTranslationFormatFun,
-	TFunction<AnimationCompressionFormat(int32 TrackIndex)> GetRotationFormatFun,
-	TFunction<AnimationCompressionFormat(int32 TrackIndex)> GetScaleFormatFun,
-	TFunction<bool(int32 TrackIndex)> IsTranslationVariableFun,
-	TFunction<bool(int32 TrackIndex)> IsRotationVariableFun,
-	TFunction<bool(int32 TrackIndex)> IsScaleVariableFun,
-	TFunction<void(TArray<uint8>& ByteStream, AnimationCompressionFormat Format, const FVector& Key, const float* Mins, const float* Ranges, int32 TrackIndex)> PackTranslationKeyFun,
-	TFunction<void(TArray<uint8>& ByteStream, AnimationCompressionFormat Format, const FQuat& Key, const float* Mins, const float* Ranges, int32 TrackIndex)> PackRotationKeyFun,
-	TFunction<void(TArray<uint8>& ByteStream, AnimationCompressionFormat Format, const FVector& Key, const float* Mins, const float* Ranges, int32 TrackIndex)> PackScaleKeyFun,
-	FAnimSegmentContext& Segment,
-	const TArray<FAnimTrackRange>& TrackRanges)
-{
-	const int32 NumFrames = Segment.NumFrames;
-	const int32 NumTracks = Segment.RotationData.Num();
-	const bool bHasScale = Segment.ScaleData.Num() > 0;
-
-	const int32 LastFrame = AnimSeq.GetRawNumberOfFrames() - 1;
-	const float FrameRate = float(AnimSeq.GetRawNumberOfFrames() - 1) / AnimSeq.SequenceLength;
-
-	// Samples are sorted by time needed first, track second to ensure that when we sample a time T, data for all tracks at
-	// that time is contiguous in memory
-	// Imagine the following rotation keys:
-	// rot0 @ 0.0										rot0 @ 0.75
-	// rot1 @ 0.0					rot1 @ 0.5			rot1 @ 0.75
-	// rot2 @ 0.0	rot2 @ 0.25		rot2 @ 0.5			rot2 @ 0.75
-	//
-	// After being sorted, we have:
-	// rot0 @ 0.0, rot1 @ 0.0, rot2 @ 0.0, rot2 @ 0.25, rot1 @ 0.5, rot0 @ 0.75, rot2 @ 0.5, rot1 @ 0.75, rot2 @ 0.75
-	//
-	// The sorting is done by the time of the previous key, which is when the current key is expected in the stream.
-	// The first key at T = 0.0 is special and is needed first, as such for the purpose of the sort, it has a negative time T.
-	// rot1 @ 0.5 is needed at T = 0.0, rot1 @ 0.75 is needed at T = 0.5, etc.
-	//
-	// When we will read from the stream, the time delta can be positive or negative.
-	// When we wish to sample a time T, we must read from the stream until:
-	//    - The time value of the packed sample is greater than T
-	//    - AND the previously read time value for the packed track sample is greater or equal to T
-	//    - OR we reach the end of the stream (we have an EOS marker)
-	//
-	// Sampling at 0.4 with the above example, we'll have:
-	// [rot0 @ 0.0, rot1 @ 0.0, rot2 @ 0.0, rot2 @ 0.25]: The first 4 values have a sample time lower or equal to T,
-	// we consume them and cache the last 2 samples for every track type.
-	// rot1 @ 0.5 is greater than T, but the last sample read for rot1 is 0.0, we consume it and cache it
-	// rot0 @ 0.75 is greater than T, but the last sample read for rot0 is 0.0, we consume it and cache it
-	// rot2 @ 0.5 is greater than T, but the last sample read for rot2 is 0.25, we consume it and cache it
-	// rot1 @ 0.75 is greater than T, but the last sample read for rot1 is 0.5, we don't need this sample and we can stop reading
-	//
-	// Our cached samples are now:
-	// [rot0 @ 0.0, rot0 @ 0.75], [rot1 @ 0.0, rot1 @ 0.5], [rot2 @ 0.25, rot2 @ 0.5]
-	//
-	// Each sample packed has a small header.
-	// 16 bits are used for the track index
-	// The MSB of the next byte tells us if the rest of the header is 8 bits (small) or 16 bits (large) wide.
-	// The following 2 bits tell us if the sample is a rotation (0), translation (1), or scale (2).
-	// If the header is small, the remaining 5 bits represent the time delta from the last packed sample (previous packed time + packed time delta).
-	// If the header is large, the remaining 13 bits represent the time delta from the last packed sample.
-	// Time does not always increment as such the time delta is can be positive or negative.
-	// As such each sample has a header of 24 bits or 32 bits. In practice, we will have a sample at nearly every
-	// key frame as such the time delta should be very small and 24 bits should end up being used for virtually every sample.
-	//
-	// [ track index : 16 ] [ is large header? : 1 ] [ sample type : 2 ] [ time delta : 5 ]
-	// [ track index : 16 ] [ is large header? : 1 ] [ sample type : 2 ] [ time delta : 13 ]
-
-	struct SampleRef
-	{
-		int32 TrackIndex;
-		int32 SampleType;
-		int32 FrameIndex;
-		int32 KeyIndex;
-		float NeededAtTime;
-	};
-
-	TArray<SampleRef> SamplesToPack;
-	SamplesToPack.Empty(NumTracks * NumFrames);
-	for (int32 TrackIndex = 0; TrackIndex < NumTracks; ++TrackIndex)
-	{
-		const FAnimTrackRange& TrackRange = TrackRanges[TrackIndex];
-
-		const FTranslationTrack& SrcTrans = Segment.TranslationData[TrackIndex];
-		const int32 NumKeysTrans = SrcTrans.PosKeys.Num();
-		if (NumKeysTrans > 1 && IsTranslationVariableFun(TrackIndex))
-		{
-			for (int32 KeyIndex = 0; KeyIndex < NumKeysTrans; ++KeyIndex)
-			{
-				float KeyTime = SrcTrans.Times[KeyIndex];
-				float FrameTime = KeyTime * FrameRate;
-				int32 FrameIndex = FMath::Clamp(FMath::TruncToInt(FrameTime + 0.5f), 0, LastFrame);
-				int32 SegmentFrameIndex = FrameIndex - Segment.StartFrame;
-				float NeededAtTime = KeyIndex == 0 ? -1.0f : SrcTrans.Times[KeyIndex - 1];
-				SamplesToPack.Add(SampleRef{ TrackIndex, 0, SegmentFrameIndex, KeyIndex, NeededAtTime });
-			}
-		}
-
-		const FRotationTrack& SrcRot = Segment.RotationData[TrackIndex];
-		const int32 NumKeysRot = SrcRot.RotKeys.Num();
-		if (NumKeysRot > 1 && IsRotationVariableFun(TrackIndex))
-		{
-			for (int32 KeyIndex = 0; KeyIndex < NumKeysRot; ++KeyIndex)
-			{
-				float KeyTime = SrcRot.Times[KeyIndex];
-				float FrameTime = KeyTime * FrameRate;
-				int32 FrameIndex = FMath::Clamp(FMath::TruncToInt(FrameTime + 0.5f), 0, LastFrame);
-				int32 SegmentFrameIndex = FrameIndex - Segment.StartFrame;
-				float NeededAtTime = KeyIndex == 0 ? -1.0f : SrcRot.Times[KeyIndex - 1];
-				SamplesToPack.Add(SampleRef{ TrackIndex, 1, SegmentFrameIndex, KeyIndex, NeededAtTime });
-			}
-		}
-
-		if (bHasScale)
-		{
-			const FScaleTrack& SrcScale = Segment.ScaleData[TrackIndex];
-			const int32 NumKeysScale = SrcScale.ScaleKeys.Num();
-			if (NumKeysScale > 1 && IsScaleVariableFun(TrackIndex))
-			{
-				for (int32 KeyIndex = 0; KeyIndex < NumKeysScale; ++KeyIndex)
-				{
-					// write the frame values for each key
-					float KeyTime = SrcScale.Times[KeyIndex];
-					float FrameTime = KeyTime * FrameRate;
-					int32 FrameIndex = FMath::Clamp(FMath::TruncToInt(FrameTime + 0.5f), 0, LastFrame);
-					int32 SegmentFrameIndex = FrameIndex - Segment.StartFrame;
-					float NeededAtTime = KeyIndex == 0 ? -1.0f : SrcScale.Times[KeyIndex - 1];
-					SamplesToPack.Add(SampleRef{ TrackIndex, 2, SegmentFrameIndex, KeyIndex, NeededAtTime });
-				}
-			}
-		}
-	}
-
-	// Sort by key index, followed by time as explained above
-	SamplesToPack.Sort([](const SampleRef& lhs, const SampleRef& rhs)
-	{
-		// Lowest needed at frame index first
-		if (lhs.NeededAtTime < rhs.NeededAtTime)
-		{
-			return true;
-		}
-
-		if (lhs.NeededAtTime > rhs.NeededAtTime)
-		{
-			return false;
-		}
-
-		// Lowest sample type first
-		return lhs.SampleType < rhs.SampleType;
-	});
-
-	int32 PreviousPackedFrameIndex = 0;
-	int32 SampleIndex = 0;
-	for (const SampleRef& Sample : SamplesToPack)
-	{
-		const FAnimTrackRange& TrackRange = TrackRanges[Sample.TrackIndex];
-
-		const int32 TimeDelta = Sample.FrameIndex - PreviousPackedFrameIndex;
-		const FSortedKeyHeader KeyHeader(Sample.TrackIndex, Sample.SampleType, TimeDelta);
-		UnalignedWriteToStream(ByteStream, &KeyHeader, KeyHeader.GetSize());
-
-		PreviousPackedFrameIndex = Sample.FrameIndex;
-
-		if (Sample.SampleType == 0)
-		{
-			const FTranslationTrack& SrcTrans = Segment.TranslationData[Sample.TrackIndex];
-			const FVector& Vec = SrcTrans.PosKeys[Sample.KeyIndex];
-			const AnimationCompressionFormat Format = GetTranslationFormatFun(Sample.TrackIndex);
-			PackTranslationKeyFun(ByteStream, Format, Vec, &TrackRange.TransMin.X, &TrackRange.TransExtent.X, Sample.TrackIndex);
-		}
-		else if (Sample.SampleType == 1)
-		{
-			const FRotationTrack& SrcRot = Segment.RotationData[Sample.TrackIndex];
-			const FQuat& Quat = SrcRot.RotKeys[Sample.KeyIndex];
-			const AnimationCompressionFormat Format = GetRotationFormatFun(Sample.TrackIndex);
-			PackRotationKeyFun(ByteStream, Format, Quat, &TrackRange.RotMin.X, &TrackRange.RotExtent.X, Sample.TrackIndex);
-		}
-		else
-		{
-			checkSlow(Sample.SampleType == 2);
-			const FScaleTrack& SrcScale = Segment.ScaleData[Sample.TrackIndex];
-			const FVector& Vec = SrcScale.ScaleKeys[Sample.KeyIndex];
-			const AnimationCompressionFormat Format = GetScaleFormatFun(Sample.TrackIndex);
-			PackScaleKeyFun(ByteStream, Format, Vec, &TrackRange.ScaleMin.X, &TrackRange.ScaleExtent.X, Sample.TrackIndex);
-		}
-
-		SampleIndex++;
-	}
-
-	// End the stream with a terminator
-	const FSortedKeyHeader EndOfStreamKeyHeader;
-	UnalignedWriteToStream(ByteStream, &EndOfStreamKeyHeader, EndOfStreamKeyHeader.GetSize());
+	return bSuccess;
 }
+#endif
 
-static int32 GetNumAnimatedTrackStreams(
-	TFunction<bool(int32 TrackIndex)> IsTranslationVariableFun,
-	TFunction<bool(int32 TrackIndex)> IsRotationVariableFun,
-	TFunction<bool(int32 TrackIndex)> IsScaleVariableFun,
-	const FAnimSegmentContext& RawSegment)
+TUniquePtr<ICompressedAnimData> UAnimCompress::AllocateAnimData() const
 {
-	const int32 NumTracks = RawSegment.RotationData.Num();
-	const bool bHasScale = RawSegment.ScaleData.Num() > 0;
-
-	int32 NumAnimatedTrackTypes = 0;
-	for (int32 TrackIndex = 0; TrackIndex < NumTracks; ++TrackIndex)
-	{
-		NumAnimatedTrackTypes += RawSegment.TranslationData[TrackIndex].PosKeys.Num() > 1 && IsTranslationVariableFun(TrackIndex) ? 1 : 0;
-		NumAnimatedTrackTypes += RawSegment.RotationData[TrackIndex].RotKeys.Num() > 1 && IsRotationVariableFun(TrackIndex) ? 1 : 0;
-		NumAnimatedTrackTypes += bHasScale && RawSegment.ScaleData[TrackIndex].ScaleKeys.Num() > 1 && IsScaleVariableFun(TrackIndex) ? 1 : 0;
-	}
-	return NumAnimatedTrackTypes;
+	return MakeUnique<FUECompressedAnimData>();
 }
 
-void UAnimCompress::WriteLinearVariableTrackData(
-	TArray<uint8>& ByteStream,
-	const UAnimSequence& AnimSeq,
-	TFunction<AnimationCompressionFormat(int32 TrackIndex)> GetTranslationFormatFun,
-	TFunction<AnimationCompressionFormat(int32 TrackIndex)> GetRotationFormatFun,
-	TFunction<AnimationCompressionFormat(int32 TrackIndex)> GetScaleFormatFun,
-	TFunction<bool(int32 TrackIndex)> IsTranslationVariableFun,
-	TFunction<bool(int32 TrackIndex)> IsRotationVariableFun,
-	TFunction<bool(int32 TrackIndex)> IsScaleVariableFun,
-	TFunction<void(TArray<uint8>& ByteStream, AnimationCompressionFormat Format, const FVector& Key, const float* Mins, const float* Ranges, int32 TrackIndex)> PackTranslationKeyFun,
-	TFunction<void(TArray<uint8>& ByteStream, AnimationCompressionFormat Format, const FQuat& Key, const float* Mins, const float* Ranges, int32 TrackIndex)> PackRotationKeyFun,
-	TFunction<void(TArray<uint8>& ByteStream, AnimationCompressionFormat Format, const FVector& Key, const float* Mins, const float* Ranges, int32 TrackIndex)> PackScaleKeyFun,
-	const FAnimSegmentContext& Segment,
-	const TArray<FAnimTrackRange>& TrackRanges)
+void UAnimCompress::ByteSwapIn(ICompressedAnimData& AnimData, TArrayView<uint8> CompressedData, FMemoryReader& MemoryStream) const
 {
-	const int32 NumFrames = Segment.NumFrames;
-	const int32 NumTracks = Segment.RotationData.Num();
-	const bool bHasScale = Segment.ScaleData.Num() > 0;
-
-	const int32 LastFrame = AnimSeq.GetRawNumberOfFrames() - 1;
-	const float FrameRate = float(AnimSeq.GetRawNumberOfFrames() - 1) / AnimSeq.SequenceLength;
-
-	// The linear packing format is more or less the same as the legacy format.
-	// First we have a list of pairs for each track type (rot, trans, scale): offset in stream, number of keys in stream
-	// Offsets are 32 bits, number of keys are 16 bits
-	// This is followed by the packed track data and their time markers.
-	//
-	// [rot0 offset, num rot0 keys] [rot1 offset, num rot1 keys] [rot0 t0, rot0 t1, rot0, t2, rot0 key0, rot0 key1, rot0 key2] [rot1 t0, rot1 t1, rot1 key0, rot1 key1]
-
-	const int32 OffsetNumKeysPairSize = sizeof(uint32) + sizeof(uint16);
-	const int32 NumAnimatedTrackStreams = GetNumAnimatedTrackStreams(IsTranslationVariableFun, IsRotationVariableFun, IsScaleVariableFun, Segment);
-
-	int32 OffsetNumKeysPairStreamOffset = ByteStream.Num();
-	ByteStream.AddUninitialized(OffsetNumKeysPairSize * NumAnimatedTrackStreams);
-
-	// If we don't have too many frames in our segment, use uint8 instead of uint16 for the time markers
-	const int32 TimeMarkerSize = NumFrames < 256 ? sizeof(uint8) : sizeof(uint16);
-
-	// TODO: There is 4 byte alignment for the packed keys due to the decompression code not handling unaligned reads properly
-	// Either properly handle unaligned reads of move the time markers after the packed keys
-
-	for (int32 TrackIndex = 0; TrackIndex < NumTracks; ++TrackIndex)
-	{
-		const FAnimTrackRange& TrackRange = TrackRanges[TrackIndex];
-
-		const FTranslationTrack& SrcTrans = Segment.TranslationData[TrackIndex];
-		const int32 NumKeysTrans = SrcTrans.PosKeys.Num();
-		if (NumKeysTrans > 1 && IsTranslationVariableFun(TrackIndex))
-		{
-			if (TimeMarkerSize == sizeof(uint16))
-				PadByteStream(ByteStream, 2, AnimationPadSentinel);
-
-			const uint32 TrackDataOffset = ByteStream.Num();
-			const uint16 NumTrackKeys = static_cast<uint16>(NumKeysTrans);
-			UnalignedWriteToStream(ByteStream, OffsetNumKeysPairStreamOffset, &TrackDataOffset, sizeof(uint32));
-			UnalignedWriteToStream(ByteStream, OffsetNumKeysPairStreamOffset, &NumTrackKeys, sizeof(uint16));
-
-			for (int32 KeyIndex = 0; KeyIndex < NumKeysTrans; ++KeyIndex)
-			{
-				const float KeyTime = SrcTrans.Times[KeyIndex];
-				const float FrameTime = KeyTime * FrameRate;
-				const int32 FrameIndex = FMath::Clamp(FMath::TruncToInt(FrameTime + 0.5f), 0, LastFrame);
-				const int32 SegmentFrameIndex = FrameIndex - Segment.StartFrame;
-				UnalignedWriteToStream(ByteStream, &SegmentFrameIndex, TimeMarkerSize);
-			}
-
-			PadByteStream(ByteStream, 4, AnimationPadSentinel);
-
-			const AnimationCompressionFormat Format = GetTranslationFormatFun(TrackIndex);
-			for (int32 KeyIndex = 0; KeyIndex < NumKeysTrans; ++KeyIndex)
-			{
-				const FVector& Key = SrcTrans.PosKeys[KeyIndex];
-				PackTranslationKeyFun(ByteStream, Format, Key, &TrackRange.TransMin.X, &TrackRange.TransExtent.X, TrackIndex);
-			}
-		}
-
-		const FRotationTrack& SrcRot = Segment.RotationData[TrackIndex];
-		const int32 NumKeysRot = SrcRot.RotKeys.Num();
-		if (NumKeysRot > 1 && IsRotationVariableFun(TrackIndex))
-		{
-			if (TimeMarkerSize == sizeof(uint16))
-				PadByteStream(ByteStream, 2, AnimationPadSentinel);
-
-			const uint32 TrackDataOffset = ByteStream.Num();
-			const uint16 NumTrackKeys = static_cast<uint16>(NumKeysRot);
-			UnalignedWriteToStream(ByteStream, OffsetNumKeysPairStreamOffset, &TrackDataOffset, sizeof(uint32));
-			UnalignedWriteToStream(ByteStream, OffsetNumKeysPairStreamOffset, &NumTrackKeys, sizeof(uint16));
-
-			for (int32 KeyIndex = 0; KeyIndex < NumKeysRot; ++KeyIndex)
-			{
-				const float KeyTime = SrcRot.Times[KeyIndex];
-				const float FrameTime = KeyTime * FrameRate;
-				const int32 FrameIndex = FMath::Clamp(FMath::TruncToInt(FrameTime + 0.5f), 0, LastFrame);
-				const int32 SegmentFrameIndex = FrameIndex - Segment.StartFrame;
-				UnalignedWriteToStream(ByteStream, &SegmentFrameIndex, TimeMarkerSize);
-			}
-
-			PadByteStream(ByteStream, 4, AnimationPadSentinel);
-
-			const AnimationCompressionFormat Format = GetRotationFormatFun(TrackIndex);
-			for (int32 KeyIndex = 0; KeyIndex < NumKeysRot; ++KeyIndex)
-			{
-				const FQuat& Key = SrcRot.RotKeys[KeyIndex];
-				PackRotationKeyFun(ByteStream, Format, Key, &TrackRange.RotMin.X, &TrackRange.RotExtent.X, TrackIndex);
-			}
-		}
-
-		if (bHasScale)
-		{
-			const FScaleTrack& SrcScale = Segment.ScaleData[TrackIndex];
-			const int32 NumKeysScale = SrcScale.ScaleKeys.Num();
-			if (NumKeysScale > 1 && IsScaleVariableFun(TrackIndex))
-			{
-				if (TimeMarkerSize == sizeof(uint16))
-					PadByteStream(ByteStream, 2, AnimationPadSentinel);
-
-				const uint32 TrackDataOffset = ByteStream.Num();
-				const uint16 NumTrackKeys = static_cast<uint16>(NumKeysScale);
-				UnalignedWriteToStream(ByteStream, OffsetNumKeysPairStreamOffset, &TrackDataOffset, sizeof(uint32));
-				UnalignedWriteToStream(ByteStream, OffsetNumKeysPairStreamOffset, &NumTrackKeys, sizeof(uint16));
-
-				for (int32 KeyIndex = 0; KeyIndex < NumKeysScale; ++KeyIndex)
-				{
-					const float KeyTime = SrcScale.Times[KeyIndex];
-					const float FrameTime = KeyTime * FrameRate;
-					const int32 FrameIndex = FMath::Clamp(FMath::TruncToInt(FrameTime + 0.5f), 0, LastFrame);
-					const int32 SegmentFrameIndex = FrameIndex - Segment.StartFrame;
-					UnalignedWriteToStream(ByteStream, &SegmentFrameIndex, TimeMarkerSize);
-				}
-
-				PadByteStream(ByteStream, 4, AnimationPadSentinel);
-
-				const AnimationCompressionFormat Format = GetScaleFormatFun(TrackIndex);
-				for (int32 KeyIndex = 0; KeyIndex < NumKeysScale; ++KeyIndex)
-				{
-					const FVector& Key = SrcScale.ScaleKeys[KeyIndex];
-					PackScaleKeyFun(ByteStream, Format, Key, &TrackRange.ScaleMin.X, &TrackRange.ScaleExtent.X, TrackIndex);
-				}
-			}
-		}
-	}
+	static_cast<FUECompressedAnimData&>(AnimData).ByteSwapIn(CompressedData, MemoryStream);
 }
 
-void UAnimCompress::BitwiseCompressAnimationTracks(
-	const UAnimSequence& AnimSeq,
-	AnimationCompressionFormat TargetTranslationFormat,
-	AnimationCompressionFormat TargetRotationFormat,
-	AnimationCompressionFormat TargetScaleFormat,
-	FAnimSegmentContext& RawSegment,
-	bool bIsSorted)
+void UAnimCompress::ByteSwapOut(ICompressedAnimData& AnimData, TArrayView<uint8> CompressedData, FMemoryWriter& MemoryStream) const
 {
-	RawSegment.RotationCompressionFormat = TargetRotationFormat;
-	RawSegment.TranslationCompressionFormat = TargetTranslationFormat;
-	RawSegment.ScaleCompressionFormat = TargetScaleFormat;
-
-	SanityCheckTrackData(AnimSeq, RawSegment);
-
-	RawSegment.CompressedByteStream.Empty(64 * 1024);
-
-	TArray<FAnimTrackRange> TrackRanges;
-	CalculateTrackRanges(TargetTranslationFormat, TargetRotationFormat, TargetScaleFormat, RawSegment, TrackRanges);
-
-	// Write track ranges
-	WriteTrackRanges(
-		RawSegment.CompressedByteStream,
-		[&](int32 TrackIndex) { return TargetTranslationFormat; },
-		[&](int32 TrackIndex) { return TargetRotationFormat; },
-		[&](int32 TrackIndex) { return TargetScaleFormat; },
-		[](int32 TrackIndex) { return FTrackKeyFlags(); },
-		[](int32 TrackIndex) { return FTrackKeyFlags(); },
-		[](int32 TrackIndex) { return FTrackKeyFlags(); },
-		RawSegment, TrackRanges, false);
-
-	checkf((RawSegment.CompressedByteStream.Num() % 4) == 0, TEXT("CompressedByteStream not aligned to four bytes"));
-
-	if (AnimSeq.KeyEncodingFormat == AKF_ConstantKeyLerp)
-	{
-		WriteUniformTrackData(
-			RawSegment.CompressedByteStream,
-			[&](int32 TrackIndex) { return TargetTranslationFormat; },
-			[&](int32 TrackIndex) { return TargetRotationFormat; },
-			[&](int32 TrackIndex) { return TargetScaleFormat; },
-			[](int32 TrackIndex) { return true; },
-			[](int32 TrackIndex) { return true; },
-			[](int32 TrackIndex) { return true; },
-			[](TArray<uint8>& ByteStream, AnimationCompressionFormat Format, const FVector& Key, const float* Mins, const float* Ranges, int32 TrackIndex) { PackVectorToStream(ByteStream, Format, Key, Mins, Ranges); },
-			[](TArray<uint8>& ByteStream, AnimationCompressionFormat Format, const FQuat& Key, const float* Mins, const float* Ranges, int32 TrackIndex) { PackQuaternionToStream(ByteStream, Format, Key, Mins, Ranges); },
-			[](TArray<uint8>& ByteStream, AnimationCompressionFormat Format, const FVector& Key, const float* Mins, const float* Ranges, int32 TrackIndex) { PackVectorToStream(ByteStream, Format, Key, Mins, Ranges); },
-			RawSegment, TrackRanges);
-	}
-	else if (AnimSeq.KeyEncodingFormat == AKF_VariableKeyLerp)
-	{
-		if (bIsSorted)
-		{
-			WriteSortedVariableTrackData(
-				RawSegment.CompressedByteStream,
-				AnimSeq,
-				[&](int32 TrackIndex) { return TargetTranslationFormat; },
-				[&](int32 TrackIndex) { return TargetRotationFormat; },
-				[&](int32 TrackIndex) { return TargetScaleFormat; },
-				[](int32 TrackIndex) { return true; },
-				[](int32 TrackIndex) { return true; },
-				[](int32 TrackIndex) { return true; },
-				[](TArray<uint8>& ByteStream, AnimationCompressionFormat Format, const FVector& Key, const float* Mins, const float* Ranges, int32 TrackIndex) { PackVectorToStream(ByteStream, Format, Key, Mins, Ranges); },
-				[](TArray<uint8>& ByteStream, AnimationCompressionFormat Format, const FQuat& Key, const float* Mins, const float* Ranges, int32 TrackIndex) { PackQuaternionToStream(ByteStream, Format, Key, Mins, Ranges); },
-				[](TArray<uint8>& ByteStream, AnimationCompressionFormat Format, const FVector& Key, const float* Mins, const float* Ranges, int32 TrackIndex) { PackVectorToStream(ByteStream, Format, Key, Mins, Ranges); },
-				RawSegment, TrackRanges);
-		}
-		else
-		{
-			WriteLinearVariableTrackData(
-				RawSegment.CompressedByteStream,
-				AnimSeq,
-				[&](int32 TrackIndex) { return TargetTranslationFormat; },
-				[&](int32 TrackIndex) { return TargetRotationFormat; },
-				[&](int32 TrackIndex) { return TargetScaleFormat; },
-				[](int32 TrackIndex) { return true; },
-				[](int32 TrackIndex) { return true; },
-				[](int32 TrackIndex) { return true; },
-				[](TArray<uint8>& ByteStream, AnimationCompressionFormat Format, const FVector& Key, const float* Mins, const float* Ranges, int32 TrackIndex) { PackVectorToStream(ByteStream, Format, Key, Mins, Ranges); },
-				[](TArray<uint8>& ByteStream, AnimationCompressionFormat Format, const FQuat& Key, const float* Mins, const float* Ranges, int32 TrackIndex) { PackQuaternionToStream(ByteStream, Format, Key, Mins, Ranges); },
-				[](TArray<uint8>& ByteStream, AnimationCompressionFormat Format, const FVector& Key, const float* Mins, const float* Ranges, int32 TrackIndex) { PackVectorToStream(ByteStream, Format, Key, Mins, Ranges); },
-				RawSegment, TrackRanges);
-		}
-	}
-
-	// Make sure we have a safe alignment
-	PadByteStream(RawSegment.CompressedByteStream, 4, AnimationPadSentinel);
-
-	// Trim unused memory.
-	RawSegment.CompressedByteStream.Shrink();
+	static_cast<FUECompressedAnimData&>(AnimData).ByteSwapOut(CompressedData, MemoryStream);
 }
 
-void UAnimCompress::BitwiseCompressTrivialAnimationTracks(const UAnimSequence& AnimSeq, FAnimSegmentContext& RawSegment)
+void UAnimCompress::DecompressPose(FAnimSequenceDecompressionContext& DecompContext, const BoneTrackArray& RotationPairs, const BoneTrackArray& TranslationPairs, const BoneTrackArray& ScalePairs, TArrayView<FTransform>& OutAtoms) const
 {
-	RawSegment.CompressedTrivialTracksByteStream.Empty();
+	const FUECompressedAnimData& AnimData = static_cast<const FUECompressedAnimData&>(DecompContext.CompressedAnimData);
 
-	check(RawSegment.TranslationData.Num() == RawSegment.RotationData.Num());
-	const int32 NumTracks = RawSegment.RotationData.Num();
-	const bool bHasScale = RawSegment.ScaleData.Num() > 0;
-
-	const int32 NumFrames = RawSegment.NumFrames;
-	const int32 LastFrame = RawSegment.NumFrames - 1;
-	const size_t FrameSize = RawSegment.NumFrames > 0xff ? sizeof(uint16) : sizeof(uint8);
-	const float FrameRate = AnimSeq.GetRawNumberOfFrames() / AnimSeq.SequenceLength;
-
-	SanityCheckTrackData(AnimSeq, RawSegment);
-
-	for (int32 TrackIndex = 0; TrackIndex < NumTracks; ++TrackIndex)
-	{
-		const FTranslationTrack& SrcTrans = RawSegment.TranslationData[TrackIndex];
-		if (SrcTrans.PosKeys.Num() == 1)
-		{
-			// A single translation key gets written out a single uncompressed float[3].
-			UnalignedWriteToStream(RawSegment.CompressedTrivialTracksByteStream, &SrcTrans.PosKeys[0], sizeof(FVector));
-		}
-
-		const FRotationTrack& SrcRot = RawSegment.RotationData[TrackIndex];
-		if (SrcRot.RotKeys.Num() == 1)
-		{
-			// For a rotation track of n=1 keys, the single key is packed as an FQuatFloat96NoW.
-			const FQuat& Quat = SrcRot.RotKeys[0];
-			const FQuatFloat96NoW QuatFloat96NoW(Quat);
-			UnalignedWriteToStream(RawSegment.CompressedTrivialTracksByteStream, &QuatFloat96NoW, sizeof(FQuatFloat96NoW));
-		}
-
-		if (bHasScale)
-		{
-			const FScaleTrack& SrcScale = RawSegment.ScaleData[TrackIndex];
-			if (SrcScale.ScaleKeys.Num() == 1)
-			{
-				// A single Scale key gets written out a single uncompressed float[3].
-				UnalignedWriteToStream(RawSegment.CompressedTrivialTracksByteStream, &SrcScale.ScaleKeys[0], sizeof(FVector));
-			}
-		}
-	}
+	AnimData.TranslationCodec->GetPoseTranslations(OutAtoms, TranslationPairs, DecompContext);
 
-	// Trim unused memory.
-	RawSegment.CompressedTrivialTracksByteStream.Shrink();
-}
+	AnimData.RotationCodec->GetPoseRotations(OutAtoms, RotationPairs, DecompContext);
 
-static int32 GetNumFrames(const TArray<FAnimSegmentContext>& RawSegments)
-{
-	int32 NumFrames = 0;
-	for (const FAnimSegmentContext& Segment : RawSegments)
+	if (AnimData.CompressedScaleOffsets.IsValid())
 	{
-		NumFrames += Segment.NumFrames;
+		AnimData.ScaleCodec->GetPoseScales(OutAtoms, ScalePairs, DecompContext);
 	}
-	return NumFrames;
 }
 
-void UAnimCompress::CoalesceCompressedSegments(UAnimSequence& AnimSeq, const TArray<FAnimSegmentContext>& RawSegments, bool bIsSorted)
+void UAnimCompress::DecompressBone(FAnimSequenceDecompressionContext& DecompContext, int32 TrackIndex, FTransform& OutAtom) const
 {
-	AnimSeq.CompressedByteStream.Empty();
-
-	const FAnimSegmentContext& FirstRawSegment = RawSegments[0];
-	const int32 NumTracks = FirstRawSegment.RotationData.Num();
-	const int32 NumFrames = GetNumFrames(RawSegments);
-	const bool bHasScale = FirstRawSegment.ScaleData.Num() != 0;
-
-	FAnimSequenceCompressionHeader Header;
-	Header.NumTracks = NumTracks;
-	Header.NumFrames = NumFrames;
-	Header.SequenceCRC = 0;
-	Header.bHasScale = bHasScale;
-	Header.bIsSorted = bIsSorted;
-
-	AnimSeq.CompressedByteStream.AddUninitialized(sizeof(FAnimSequenceCompressionHeader));
-	memcpy(AnimSeq.CompressedByteStream.GetData(), &Header, sizeof(FAnimSequenceCompressionHeader));
-
-	// TODO: Convert to bitset!
-	int32 TrackInfoOffset = AnimSeq.CompressedByteStream.Num();
-	AnimSeq.CompressedByteStream.AddUninitialized(NumTracks);
-	for (int32 TrackIndex = 0; TrackIndex < NumTracks; ++TrackIndex)
-	{
-		bool IsTranslationTrivial = FirstRawSegment.TranslationData[TrackIndex].PosKeys.Num() <= 1;
-		bool IsRotationTrivial = FirstRawSegment.RotationData[TrackIndex].RotKeys.Num() <= 1;
-		bool IsScaleTrivial = !bHasScale || FirstRawSegment.ScaleData[TrackIndex].ScaleKeys.Num() <= 1;
+	// Initialize to identity to set the scale and in case of a missing rotation or translation codec
+	OutAtom.SetIdentity();
 
-		// Bottom 3 bits indicate whether or not trans/rot/scale is trivial
-		uint8 TrackFlags = (IsTranslationTrivial ? 0x4 : 0x0) | (IsRotationTrivial ? 0x2 : 0x0) | (IsScaleTrivial ? 0x1 : 0x0);
-		AnimSeq.CompressedByteStream[TrackInfoOffset + TrackIndex] = TrackFlags;
-	}
+	const FUECompressedAnimData& AnimData = static_cast<const FUECompressedAnimData&>(DecompContext.CompressedAnimData);
 
-	PadByteStream(AnimSeq.CompressedByteStream, 4, AnimationPadSentinel);
+	// decompress the translation component using the proper method
+	((AnimEncodingLegacyBase*)AnimData.TranslationCodec)->GetBoneAtomTranslation(OutAtom, DecompContext, TrackIndex);
 
-	const int32 NumSegments = RawSegments.Num();
+	// decompress the rotation component using the proper method
+	((AnimEncodingLegacyBase*)AnimData.RotationCodec)->GetBoneAtomRotation(OutAtom, DecompContext, TrackIndex);
 
-	AnimSeq.CompressedSegments.Empty(NumSegments);
-	AnimSeq.CompressedSegments.AddZeroed(NumSegments);
-	for (int32 SegmentIndex = 0; SegmentIndex < NumSegments; ++SegmentIndex)
+	// we assume scale keys can be empty, so only extract if we have valid keys
+	if (AnimData.CompressedScaleOffsets.IsValid())
 	{
-		const FAnimSegmentContext& RawSegment = RawSegments[SegmentIndex];
-		FCompressedSegment& CompressedSegment = AnimSeq.CompressedSegments[SegmentIndex];
-
-		if (SegmentIndex == 0)
-		{
-			int32 TrivialKeysOffset = AnimSeq.CompressedByteStream.Num();
-			AnimSeq.CompressedByteStream.AddUninitialized(RawSegment.CompressedTrivialTracksByteStream.Num());
-			memcpy(AnimSeq.CompressedByteStream.GetData() + TrivialKeysOffset, RawSegment.CompressedTrivialTracksByteStream.GetData(), RawSegment.CompressedTrivialTracksByteStream.Num());
-		}
-
-		CompressedSegment.StartFrame = RawSegment.StartFrame;
-		CompressedSegment.NumFrames = RawSegment.NumFrames;
-		CompressedSegment.ByteStreamOffset = AnimSeq.CompressedByteStream.Num();
-		CompressedSegment.TranslationCompressionFormat = RawSegment.TranslationCompressionFormat;
-		CompressedSegment.RotationCompressionFormat = RawSegment.RotationCompressionFormat;
-		CompressedSegment.ScaleCompressionFormat = RawSegment.ScaleCompressionFormat;
-
-		AnimSeq.CompressedByteStream.AddUninitialized(RawSegment.CompressedByteStream.Num());
-		memcpy(AnimSeq.CompressedByteStream.GetData() + CompressedSegment.ByteStreamOffset, RawSegment.CompressedByteStream.GetData(), RawSegment.CompressedByteStream.Num());
+		// decompress the rotation component using the proper method
+		((AnimEncodingLegacyBase*)AnimData.ScaleCodec)->GetBoneAtomScale(OutAtom, DecompContext, TrackIndex);
 	}
-
-	AnimSeq.CompressedByteStream.Shrink();
-
-	// Calculate CRC and update it
-	Header.SequenceCRC = FCrc::MemCrc32(AnimSeq.CompressedByteStream.GetData(), AnimSeq.CompressedByteStream.Num());
-	memcpy(AnimSeq.CompressedByteStream.GetData(), &Header, sizeof(FAnimSequenceCompressionHeader));
 }
-#endif
 
 #if WITH_EDITOR
 
 void UAnimCompress::PopulateDDCKey(FArchive& Ar)
 {
+	Super::PopulateDDCKey(Ar);
+
 	uint8 TCF, RCF, SCF;
 	TCF = (uint8)TranslationCompressionFormat.GetValue();
 	RCF = (uint8)RotationCompressionFormat.GetValue();
@@ -1591,36 +719,8 @@ void UAnimCompress::PopulateDDCKey(FArchive& Ar)
 	Ar << TCF << RCF << SCF;
 
 	Ar << MaxCurveError;
-
-	bool ES = bEnableSegmenting != 0 && USE_SEGMENTING_CONTEXT;
-	Ar << ES << IdealNumFramesPerSegment << MaxNumFramesPerSegment;
 }
 
-/**
-
- * Tracks
- */
-
-bool UAnimCompress::Reduce(const FCompressibleAnimData& CompressibleAnimData, FCompressibleAnimDataResult& OutResult)
-{
-#if WITH_EDITORONLY_DATA
-	const bool bSkeletonExistsIfNeeded = ( CompressibleAnimData.Skeleton || !bNeedsSkeleton);
-	if ( bSkeletonExistsIfNeeded )
-	{
-		OutResult.CompressedNumberOfFrames = CompressibleAnimData.NumFrames;
-
-		DoReduction(CompressibleAnimData, OutResult);
-		
-		OutResult.CompressionScheme = this;
-		return true;
-	}
-#endif // WITH_EDITORONLY_DATA
-
-	return false;
-}
-
-#endif // WITH_EDITOR
-
 void UAnimCompress::FilterTrivialPositionKeys(
 	FTranslationTrack& Track, 
 	float MaxPosDelta)
@@ -1781,21 +881,6 @@ void UAnimCompress::FilterTrivialKeys(
 	FilterTrivialScaleKeys(ScaleTracks, MaxScaleDelta);
 }
 
-#if USE_SEGMENTING_CONTEXT
-void UAnimCompress::FilterTrivialKeys(
-	TArray<FAnimSegmentContext>& RawSegments,
-	float MaxPosDelta,
-	float MaxRotDelta,
-	float MaxScaleDelta)
-{
-	for (FAnimSegmentContext& Segment : RawSegments)
-	{
-		FilterTrivialRotationKeys(Segment.RotationData, MaxRotDelta);
-		FilterTrivialPositionKeys(Segment.TranslationData, MaxPosDelta);
-		FilterTrivialScaleKeys(Segment.ScaleData, MaxScaleDelta);
-	}
-}
-#endif
 
 void UAnimCompress::FilterIntermittentPositionKeys(
 	FTranslationTrack& Track, 
@@ -2019,235 +1104,4 @@ void UAnimCompress::SeparateRawDataIntoTracks(
 		OutScaleData.Empty();
 	}
 }
-
-#if USE_SEGMENTING_CONTEXT
-static void GenerateAnimSequenceSegments(const UAnimSequence& AnimSeq, int32 NumFrames, int32 IdealNumFramesPerSegment, int32 MaxNumFramesPerSegment, TArray<FAnimSegmentContext>& OutRawSegments)
-{
-	int32 NumSegments;
-	TArray<int32> NumFramesPerSegment;
-
-	IdealNumFramesPerSegment = FMath::Max(IdealNumFramesPerSegment, 0);
-	MaxNumFramesPerSegment = FMath::Max(IdealNumFramesPerSegment, MaxNumFramesPerSegment);
-
-	if (NumFrames > MaxNumFramesPerSegment && IdealNumFramesPerSegment > 0)
-	{
-		NumSegments = (NumFrames + IdealNumFramesPerSegment - 1) / IdealNumFramesPerSegment;
-		int32 PaddedNumFrames = NumSegments * IdealNumFramesPerSegment;
-
-		NumFramesPerSegment.Init(IdealNumFramesPerSegment, NumSegments);
-
-		int32 NumLeftoverFrames = IdealNumFramesPerSegment - (PaddedNumFrames - NumFrames);
-		if (NumLeftoverFrames != 0)
-		{
-			NumFramesPerSegment[NumSegments - 1] = NumLeftoverFrames;
-		}
-
-		int32 Slack = MaxNumFramesPerSegment - IdealNumFramesPerSegment;
-		if ((NumSegments - 1) * Slack >= NumLeftoverFrames)
-		{
-			// Enough segments to distribute the leftover samples of the last segment
-			while (NumFramesPerSegment[NumSegments - 1] != 0)
-			{
-				for (int32 SegmentIndex = 0; SegmentIndex < NumSegments - 1 && NumFramesPerSegment[NumSegments - 1] != 0; ++SegmentIndex)
-				{
-					NumFramesPerSegment[SegmentIndex]++;
-					NumFramesPerSegment[NumSegments - 1]--;
-				}
-			}
-
-			NumSegments--;
-		}
-
-		checkf(NumSegments != 1, TEXT("Expected a number of segments greater than 1."));
-	}
-	else
-	{
-		// Everything fits in a single segment
-		NumSegments = 1;
-		NumFramesPerSegment.Init(NumFrames, NumSegments);
-	}
-
-	OutRawSegments.Empty(NumSegments);
-	OutRawSegments.AddZeroed(NumSegments);
-	int32 NumPreviousFrames = 0;
-	for (int32 SegmentIndex = 0; SegmentIndex < NumSegments; ++SegmentIndex)
-	{
-		int32 NumFramesInSegment = NumFramesPerSegment[SegmentIndex];
-
-		FAnimSegmentContext& Segment = OutRawSegments[SegmentIndex];
-		Segment.StartFrame = NumPreviousFrames;
-		Segment.NumFrames = NumFramesInSegment;
-
-		NumPreviousFrames += NumFramesInSegment;
-	}
-}
-#endif
-
-static int32 GetNumFrames(const UAnimSequence& AnimSeq, const TArray<FTranslationTrack>& TranslationData, const TArray<FRotationTrack>& RotationData, const TArray<FScaleTrack>& ScaleData)
-{
-	for (const FTranslationTrack& track : TranslationData)
-	{
-		if (track.PosKeys.Num() > 1)
-		{
-			return track.PosKeys.Num();
-		}
-	}
-
-	for (const FRotationTrack& track : RotationData)
-	{
-		if (track.RotKeys.Num() > 1)
-		{
-			return track.RotKeys.Num();
-		}
-	}
-
-	for (const FScaleTrack& track : ScaleData)
-	{
-		if (track.ScaleKeys.Num() > 1)
-		{
-			return track.ScaleKeys.Num();
-		}
-	}
-
-	// Fallback
-	return AnimSeq.GetRawNumberOfFrames();
-}
-
-#if USE_SEGMENTING_CONTEXT
-void UAnimCompress::SeparateRawDataIntoTracks(
-	const UAnimSequence& AnimSeq,
-	const TArray<FTranslationTrack>& TranslationData,
-	const TArray<FRotationTrack>& RotationData,
-	const TArray<FScaleTrack>& ScaleData,
-	int32 IdealNumFramesPerSegment, int32 MaxNumFramesPerSegment,
-	TArray<FAnimSegmentContext>& OutRawSegments)
-{
-	int32 NumFrames = GetNumFrames(AnimSeq, TranslationData, RotationData, ScaleData);
-	GenerateAnimSequenceSegments(AnimSeq, NumFrames, IdealNumFramesPerSegment, MaxNumFramesPerSegment, OutRawSegments);
-
-	const int32 NumTracks = RotationData.Num();
-
-	for (FAnimSegmentContext& Segment : OutRawSegments)
-	{
-		Segment.TranslationData.Empty(NumTracks);
-		Segment.RotationData.Empty(NumTracks);
-		Segment.ScaleData.Empty(NumTracks);
-		Segment.TranslationData.AddZeroed(NumTracks);
-		Segment.RotationData.AddZeroed(NumTracks);
-		Segment.ScaleData.AddZeroed(NumTracks);
-
-		// only compress scale if it has valid scale keys
-		bool bCompressScaleKeys = false;
-
-		for (int32 TrackIndex = 0; TrackIndex < NumTracks; ++TrackIndex)
-		{
-			const FTranslationTrack&	SeqTranslationTrack = TranslationData[TrackIndex];
-			const FRotationTrack&		SeqRotationTrack = RotationData[TrackIndex];
-
-			FTranslationTrack&	TranslationTrack = Segment.TranslationData[TrackIndex];
-			FRotationTrack&		RotationTrack = Segment.RotationData[TrackIndex];
-
-			const int32 NumPosKeys = SeqTranslationTrack.PosKeys.Num();
-			const int32 NumRotKeys = SeqRotationTrack.RotKeys.Num();
-
-			// Do nothing if the data for this track is empty.
-			if (NumPosKeys == 0 || NumRotKeys == 0)
-			{
-				continue;
-			}
-
-			const int32 NumPosSegmentKeys = FMath::Min(NumPosKeys, Segment.NumFrames);
-			const int32 NumRotSegmentKeys = FMath::Min(NumRotKeys, Segment.NumFrames);
-
-			TranslationTrack.PosKeys.Empty(NumPosSegmentKeys);
-			TranslationTrack.PosKeys.AddUninitialized(NumPosSegmentKeys);
-			TranslationTrack.Times.Empty(NumPosSegmentKeys);
-			TranslationTrack.Times.AddUninitialized(NumPosSegmentKeys);
-
-			if (NumPosKeys == 1)
-			{
-				TranslationTrack.PosKeys[0] = SeqTranslationTrack.PosKeys[0];
-				TranslationTrack.Times[0] = 0.f;
-			}
-			else
-			{
-				// Copy over position keys.
-				memcpy(TranslationTrack.PosKeys.GetData(), SeqTranslationTrack.PosKeys.GetData() + Segment.StartFrame, NumPosSegmentKeys * sizeof(FVector));
-
-				// Set times for the translation track.
-				const float PosFrameInterval = AnimSeq.SequenceLength / static_cast<float>(NumPosKeys - 1);
-				for (int32 SeqPosIndex = Segment.StartFrame, SegmentPosIndex = 0; SegmentPosIndex < NumPosSegmentKeys; ++SeqPosIndex, ++SegmentPosIndex)
-				{
-					TranslationTrack.Times[SegmentPosIndex] = SeqPosIndex * PosFrameInterval;
-				}
-			}
-
-			RotationTrack.RotKeys.Empty(NumRotSegmentKeys);
-			RotationTrack.RotKeys.AddUninitialized(NumRotSegmentKeys);
-			RotationTrack.Times.Empty(NumRotSegmentKeys);
-			RotationTrack.Times.AddUninitialized(NumRotSegmentKeys);
-
-			if (NumRotKeys == 1)
-			{
-				RotationTrack.RotKeys[0] = SeqRotationTrack.RotKeys[0];
-				RotationTrack.Times[0] = 0.f;
-			}
-			else
-			{
-				// Copy over rotation keys.
-				memcpy(RotationTrack.RotKeys.GetData(), SeqRotationTrack.RotKeys.GetData() + Segment.StartFrame, NumRotSegmentKeys * sizeof(FQuat));
-
-				// Set times for the rotation track.
-				const float RotFrameInterval = AnimSeq.SequenceLength / static_cast<float>(NumRotKeys - 1);
-				for (int32 SeqRotIndex = Segment.StartFrame, SegmentRotIndex = 0; SegmentRotIndex < NumRotSegmentKeys; ++SeqRotIndex, ++SegmentRotIndex)
-				{
-					RotationTrack.Times[SegmentRotIndex] = SeqRotIndex * RotFrameInterval;
-				}
-			}
-
-			if (ScaleData.Num() != 0)
-			{
-				bCompressScaleKeys = true;
-
-				const FScaleTrack&	SeqScaleTrack = ScaleData[TrackIndex];
-				FScaleTrack&		ScaleTrack = Segment.ScaleData[TrackIndex];
-
-				const int32 NumScaleKeys = SeqScaleTrack.ScaleKeys.Num();
-				const int32 NumScaleSegmentKeys = FMath::Min(NumScaleKeys, Segment.NumFrames);
-
-				ScaleTrack.ScaleKeys.Empty(NumScaleSegmentKeys);
-				ScaleTrack.ScaleKeys.AddUninitialized(NumScaleSegmentKeys);
-				ScaleTrack.Times.Empty(NumScaleSegmentKeys);
-				ScaleTrack.Times.AddUninitialized(NumScaleSegmentKeys);
-
-				if (NumScaleKeys == 1)
-				{
-					ScaleTrack.ScaleKeys[0] = SeqScaleTrack.ScaleKeys[0];
-					ScaleTrack.Times[0] = 0.f;
-				}
-				else
-				{
-					// Copy over scale keys.
-					memcpy(ScaleTrack.ScaleKeys.GetData(), SeqScaleTrack.ScaleKeys.GetData() + Segment.StartFrame, NumScaleSegmentKeys * sizeof(FVector));
-
-					// Set times for the scale track.
-					const float ScaleFrameInterval = AnimSeq.SequenceLength / static_cast<float>(NumScaleKeys - 1);
-					for (int32 SeqScaleIndex = Segment.StartFrame, SegmentScaleIndex = 0; SegmentScaleIndex < NumScaleSegmentKeys; ++SeqScaleIndex, ++SegmentScaleIndex)
-					{
-						ScaleTrack.Times[SegmentScaleIndex] = SeqScaleIndex * ScaleFrameInterval;
-					}
-				}
-			}
-		}
-
-		// if nothing to compress, empty the ScaleData
-		// that way we don't have to worry about compressing scale data. 
-		if (!bCompressScaleKeys)
-		{
-			Segment.ScaleData.Empty();
-		}
-	}
-}
-#endif
-
 #endif
\ No newline at end of file
diff --git a/Engine/Source/Runtime/Engine/Private/Animation/AnimCompress_BitwiseCompressOnly.cpp b/Engine/Source/Runtime/Engine/Private/Animation/AnimCompress_BitwiseCompressOnly.cpp
--- a/Engine/Source/Runtime/Engine/Private/Animation/AnimCompress_BitwiseCompressOnly.cpp
+++ b/Engine/Source/Runtime/Engine/Private/Animation/AnimCompress_BitwiseCompressOnly.cpp
@@ -15,7 +15,7 @@ UAnimCompress_BitwiseCompressOnly::UAnimCompress_BitwiseCompressOnly(const FObje
 }
 
 #if WITH_EDITOR
-void UAnimCompress_BitwiseCompressOnly::DoReduction(const FCompressibleAnimData& CompressibleAnimData, FCompressibleAnimDataResult& OutResult)
+bool UAnimCompress_BitwiseCompressOnly::DoReduction(const FCompressibleAnimData& CompressibleAnimData, FCompressibleAnimDataResult& OutResult)
 {
 #if WITH_EDITORONLY_DATA
 	// split the raw data into tracks
@@ -28,49 +28,28 @@ void UAnimCompress_BitwiseCompressOnly::DoReduction(const FCompressibleAnimData&
 	FilterTrivialKeys(TranslationData, RotationData, ScaleData, TRANSLATION_ZEROING_THRESHOLD, QUATERNION_ZEROING_THRESHOLD, SCALE_ZEROING_THRESHOLD);
 
 	// record the proper runtime decompressor to use
-	OutResult.KeyEncodingFormat = AKF_ConstantKeyLerp;
-	OutResult.RotationCompressionFormat = RotationCompressionFormat;
-	OutResult.TranslationCompressionFormat = TranslationCompressionFormat;
-	OutResult.ScaleCompressionFormat = ScaleCompressionFormat;
-	AnimationFormat_SetInterfaceLinks(OutResult);
-
-#if USE_SEGMENTING_CONTEXT
-	if (bEnableSegmenting)
-	{
-		TArray<FAnimSegmentContext> RawSegments;
-		SeparateRawDataIntoTracks(*AnimSeq, TranslationData, RotationData, ScaleData, IdealNumFramesPerSegment, MaxNumFramesPerSegment, RawSegments);
-
-		BitwiseCompressAnimationTracks(
-			*AnimSeq,
-			static_cast<AnimationCompressionFormat>(TranslationCompressionFormat),
-			static_cast<AnimationCompressionFormat>(RotationCompressionFormat),
-			static_cast<AnimationCompressionFormat>(ScaleCompressionFormat),
-			RawSegments);
-
-		CoalesceCompressedSegments(*AnimSeq, RawSegments);
-
-		AnimSeq->TranslationCompressionFormat = TranslationCompressionFormat;
-		AnimSeq->RotationCompressionFormat = RotationCompressionFormat;
-		AnimSeq->ScaleCompressionFormat = ScaleCompressionFormat;
-	}
-	else
-#endif
-	{
-		// bitwise compress the tracks into the anim sequence buffers
-		BitwiseCompressAnimationTracks(
-			CompressibleAnimData,
-			OutResult,
-			static_cast<AnimationCompressionFormat>(TranslationCompressionFormat),
-			static_cast<AnimationCompressionFormat>(RotationCompressionFormat),
-			static_cast<AnimationCompressionFormat>(ScaleCompressionFormat),
-			TranslationData,
-			RotationData,
-			ScaleData);
-	}
+	FUECompressedAnimDataMutable& AnimData = static_cast<FUECompressedAnimDataMutable&>(*OutResult.AnimData);
+	AnimData.KeyEncodingFormat = AKF_ConstantKeyLerp;
+	AnimData.RotationCompressionFormat = RotationCompressionFormat;
+	AnimData.TranslationCompressionFormat = TranslationCompressionFormat;
+	AnimData.ScaleCompressionFormat = ScaleCompressionFormat;
+	AnimationFormat_SetInterfaceLinks(AnimData);
+
+	// bitwise compress the tracks into the anim sequence buffers
+	BitwiseCompressAnimationTracks(
+		CompressibleAnimData,
+		OutResult,
+		static_cast<AnimationCompressionFormat>(TranslationCompressionFormat),
+		static_cast<AnimationCompressionFormat>(RotationCompressionFormat),
+		static_cast<AnimationCompressionFormat>(ScaleCompressionFormat),
+		TranslationData,
+		RotationData,
+		ScaleData);
 
 	// We could be invalid, set the links again
-	AnimationFormat_SetInterfaceLinks(OutResult);
+	AnimationFormat_SetInterfaceLinks(AnimData);
 #endif // WITH_EDITORONLY_DATA
+	return true;
 }
 
 #endif // WITH_EDITOR
diff --git a/Engine/Source/Runtime/Engine/Private/Animation/AnimCompress_LeastDestructive.cpp b/Engine/Source/Runtime/Engine/Private/Animation/AnimCompress_LeastDestructive.cpp
--- a/Engine/Source/Runtime/Engine/Private/Animation/AnimCompress_LeastDestructive.cpp
+++ b/Engine/Source/Runtime/Engine/Private/Animation/AnimCompress_LeastDestructive.cpp
@@ -12,16 +12,5 @@ UAnimCompress_LeastDestructive::UAnimCompress_LeastDestructive(const FObjectInit
 {
 	Description = TEXT("Least Destructive");
 	TranslationCompressionFormat = ACF_None;
-	RotationCompressionFormat = ACF_None;
+	RotationCompressionFormat = ACF_Float96NoW;
 }
-
-
-#if WITH_EDITOR
-void UAnimCompress_LeastDestructive::DoReduction(const FCompressibleAnimData& CompressibleAnimData, FCompressibleAnimDataResult& OutResult)
-{
-	UAnimCompress* BitwiseCompressor = NewObject<UAnimCompress_BitwiseCompressOnly>();
-	BitwiseCompressor->RotationCompressionFormat = ACF_Float96NoW;
-	BitwiseCompressor->TranslationCompressionFormat = ACF_None;
-	BitwiseCompressor->Reduce(CompressibleAnimData, OutResult);
-}
-#endif // WITH_EDITOR
