COPY from //Dev-Anim to //Dev-Main

diff --git a/Engine/Source/Runtime/Engine/Private/Animation/AnimCompress_RemoveLinearKeys.cpp b/Engine/Source/Runtime/Engine/Private/Animation/AnimCompress_RemoveLinearKeys.cpp
--- a/Engine/Source/Runtime/Engine/Private/Animation/AnimCompress_RemoveLinearKeys.cpp
+++ b/Engine/Source/Runtime/Engine/Private/Animation/AnimCompress_RemoveLinearKeys.cpp
@@ -62,54 +62,9 @@ UAnimCompress_RemoveLinearKeys::UAnimCompress_RemoveLinearKeys(const FObjectInit
 	ParentKeyScale = 2.0f;
 	bRetarget = true;
 	bActuallyFilterLinearKeys = true;
-	bOptimizeForForwardPlayback = false;
-	bUseDecompression = false;
-	bUseMultithreading = true;
 }
 
 #if WITH_EDITOR
-
-#if USE_SEGMENTING_CONTEXT
-/**
- * Structure that holds the necessary information for performing linear key reduction.
- * Each segment will have its own instance. Each instance is independent.
- * This makes it safe to compress multiple segments in parallel.
- */
-struct FProcessAnimationTracksContext
-{
-	UAnimSequence& AnimSeq;
-	const TArray<FBoneData>& BoneData;
-	const TArray<FTransform>& RefPose;
-	const TArray<int32>& EndEffectors;
-	const TArray<int32>& BoneIndexToTrackIndex;
-	const TArray<FTransform>& TrivialTrackValues;
-
-	FAnimSegmentContext& Segment;
-	TArray<FAnimSegmentContext>& SegmentList;
-
-	TArray<FTransform> RawWorldBones;
-	TArray<FTransform> NewWorldBones;
-
-	FProcessAnimationTracksContext(
-		UAnimSequence& AnimSeq_,
-		const TArray<FBoneData>& BoneData_,
-		const TArray<int32>& EndEffectors_,
-		const TArray<int32>& BoneIndexToTrackIndex_,
-		const TArray<FTransform>& TrivialTrackValues_,
-		FAnimSegmentContext& Segment_,
-		TArray<FAnimSegmentContext>& SegmentList_)
-		: AnimSeq(AnimSeq_)
-		, BoneData(BoneData_)
-		, RefPose(AnimSeq_.GetSkeleton()->GetRefLocalPoses())
-		, EndEffectors(EndEffectors_)
-		, BoneIndexToTrackIndex(BoneIndexToTrackIndex_)
-		, TrivialTrackValues(TrivialTrackValues_)
-		, Segment(Segment_)
-		, SegmentList(SegmentList_)
-	{}
-};
-#endif
-
 struct RotationAdapter
 {
 	typedef FQuat KeyType;
@@ -344,213 +299,6 @@ void FilterLinearKeysTemplate(
 	}
 }
 
-#if USE_SEGMENTING_CONTEXT
-template <typename AdapterType>
-void FilterLinearKeysTemplate(
-	TArray<typename AdapterType::KeyType>& Keys,
-	TArray<float>& Times,
-	TArray<FTransform>& BoneAtoms,
-	const TArray<float>* ParentTimes,
-	const FProcessAnimationTracksContext& Context,
-	const TArray<int32>& TargetBoneIndices,
-	int32 BoneIndex,
-	int32 ParentBoneIndex,
-	float ParentScale,
-	float MaxDelta,
-	float MaxTargetDelta,
-	float EffectorDiffSocket)
-{
-	typedef typename AdapterType::KeyType KeyType;
-
-	const int32 KeyCount = Keys.Num();
-	check(Keys.Num() == Times.Num());
-	check(KeyCount >= 1);
-
-	// generate new arrays we will fill with the final keys
-	TArray<KeyType> NewKeys;
-	TArray<float> NewTimes;
-	NewKeys.Empty(KeyCount);
-	NewTimes.Empty(KeyCount);
-
-	// Only bother doing anything if we have some keys!
-	if (KeyCount > 0)
-	{
-		int32 LowKey = 0;
-		int32 HighKey = KeyCount - 1;
-
-		TArray<uint32> KnownParentTimes;
-		KnownParentTimes.Empty(KeyCount);
-		const int32 ParentKeyCount = ParentTimes ? ParentTimes->Num() : 0;
-		for (int32 TimeIndex = 0, ParentTimeIndex = 0; TimeIndex < KeyCount; TimeIndex++)
-		{
-			while ((ParentTimeIndex < ParentKeyCount) && (Times[TimeIndex] > (*ParentTimes)[ParentTimeIndex]))
-			{
-				ParentTimeIndex++;
-			}
-
-			KnownParentTimes.Add((ParentTimeIndex < ParentKeyCount) && (Times[TimeIndex] == (*ParentTimes)[ParentTimeIndex]));
-		}
-
-		TArray<FTransform> CachedInvRawBases;
-		CachedInvRawBases.Empty(KeyCount);
-		for (int32 FrameIndex = 0; FrameIndex < KeyCount; ++FrameIndex)
-		{
-			const FTransform& RawBase = Context.RawWorldBones[(BoneIndex * Context.Segment.NumFrames) + FrameIndex];
-			CachedInvRawBases.Add(RawBase.Inverse());
-		}
-
-		// copy the low key (this one is a given)
-		NewTimes.Add(Times[0]);
-		NewKeys.Add(Keys[0]);
-
-		const FTransform EndEffectorDummyBoneSocket(FQuat::Identity, FVector(END_EFFECTOR_DUMMY_BONE_LENGTH_SOCKET));
-		const FTransform EndEffectorDummyBone(FQuat::Identity, FVector(END_EFFECTOR_DUMMY_BONE_LENGTH));
-
-		const float DeltaThreshold = (Context.BoneData[BoneIndex].IsEndEffector() && (Context.BoneData[BoneIndex].bHasSocket || Context.BoneData[BoneIndex].bKeyEndEffector)) ? EffectorDiffSocket : MaxTargetDelta;
-
-		// We will test within a sliding window between LowKey and HighKey.
-		// Therefore, we are done when the LowKey exceeds the range
-		while (LowKey + 1 < KeyCount)
-		{
-			int32 GoodHighKey = LowKey + 1;
-			int32 BadHighKey = KeyCount;
-
-			// bisect until we find the lowest acceptable high key
-			while (BadHighKey - GoodHighKey >= 2)
-			{
-				HighKey = GoodHighKey + (BadHighKey - GoodHighKey) / 2;
-
-				// get the parameters of the window we are testing
-				const float LowTime = Times[LowKey];
-				const float HighTime = Times[HighKey];
-				const KeyType& LowValue = Keys[LowKey];
-				const KeyType& HighValue = Keys[HighKey];
-				const float Range = HighTime - LowTime;
-				const float InvRange = 1.0f / Range;
-
-				// iterate through all interpolated members of the window to
-				// compute the error when compared to the original raw values
-				float MaxLerpError = 0.0f;
-				float MaxTargetError = 0.0f;
-				for (int32 TestKey = LowKey + 1; TestKey < HighKey; ++TestKey)
-				{
-					// get the parameters of the member being tested
-					float TestTime = Times[TestKey];
-					const KeyType& TestValue = Keys[TestKey];
-
-					// compute the proposed, interpolated value for the key
-					const float Alpha = (TestTime - LowTime) * InvRange;
-					const KeyType LerpValue = AnimationCompressionUtils::Interpolate(LowValue, HighValue, Alpha);
-
-					// compute the error between our interpolated value and the desired value
-					float LerpError = CalcDelta(TestValue, LerpValue);
-
-					// if the local-space lerp error is within our tolerances, we will also check the
-					// effect this interpolated key will have on our target end effectors
-					float TargetError = -1.0f;
-					if (LerpError <= MaxDelta)
-					{
-						// get the raw world transform for this bone (the original world-space position)
-						const int32 FrameIndex = TestKey;
-						const FTransform& InvRawBase = CachedInvRawBases[FrameIndex];
-
-						// generate the proposed local bone atom and transform (local space)
-						FTransform ProposedTM = AdapterType::UpdateBoneAtom(BoneAtoms[FrameIndex], LerpValue);
-
-						// convert the proposed local transform to world space using this bone's parent transform
-						const FTransform& CurrentParent = ParentBoneIndex != INDEX_NONE ? Context.NewWorldBones[(ParentBoneIndex * Context.Segment.NumFrames) + FrameIndex] : FTransform::Identity;
-						FTransform ProposedBase = ProposedTM * CurrentParent;
-
-						// for each target end effector, compute the error we would introduce with our proposed key
-						for (int32 TargetIndex = 0; TargetIndex < TargetBoneIndices.Num(); ++TargetIndex)
-						{
-							// find the offset transform from the raw base to the end effector
-							const int32 TargetBoneIndex = TargetBoneIndices[TargetIndex];
-							FTransform RawTarget = Context.RawWorldBones[(TargetBoneIndex * Context.Segment.NumFrames) + FrameIndex];
-							const FTransform RelTM = RawTarget * InvRawBase;
-
-							// forecast where the new end effector would be using our proposed key
-							FTransform ProposedTarget = RelTM * ProposedBase;
-
-							// If this is an EndEffector, add a dummy bone to measure the effect of compressing the rotation.
-							// Sockets and Key EndEffectors have a longer dummy bone to maintain higher precision.
-							if (Context.BoneData[TargetIndex].bHasSocket || Context.BoneData[TargetIndex].bKeyEndEffector)
-							{
-								ProposedTarget = EndEffectorDummyBoneSocket * ProposedTarget;
-								RawTarget = EndEffectorDummyBoneSocket * RawTarget;
-							}
-							else
-							{
-								ProposedTarget = EndEffectorDummyBone * ProposedTarget;
-								RawTarget = EndEffectorDummyBone * RawTarget;
-							}
-
-							// determine the extend of error at the target end effector
-							const float ThisError = (ProposedTarget.GetTranslation() - RawTarget.GetTranslation()).Size();
-							TargetError = FMath::Max(TargetError, ThisError);
-
-							// exit early when we encounter a large delta
-							const float TargetDeltaThreshold = Context.BoneData[TargetIndex].bHasSocket ? EffectorDiffSocket : DeltaThreshold;
-							if (TargetError > TargetDeltaThreshold)
-							{
-								break;
-							}
-						}
-					}
-
-					// If the parent has a key at this time, we'll scale our error values as requested.
-					// This increases the odds that we will choose keys on the same frames as our parent bone,
-					// making the skeleton more uniform in key distribution.
-					if (ParentTimes)
-					{
-						if (KnownParentTimes[TestKey])
-						{
-							// our parent has a key at this time, 
-							// inflate our perceived error to increase our sensitivity
-							// for also retaining a key at this time
-							LerpError *= ParentScale;
-							TargetError *= ParentScale;
-						}
-					}
-
-					// keep track of the worst errors encountered for both 
-					// the local-space 'lerp' error and the end effector drift we will cause
-					MaxLerpError = FMath::Max(MaxLerpError, LerpError);
-					MaxTargetError = FMath::Max(MaxTargetError, TargetError);
-
-					// exit early if we have failed in this span
-					if (MaxLerpError > MaxDelta ||
-						MaxTargetError > DeltaThreshold)
-					{
-						break;
-					}
-				}
-
-				// determine if the span succeeded. That is, the worst errors found are within tolerances
-				if ((MaxLerpError <= MaxDelta) && (MaxTargetError <= DeltaThreshold))
-				{
-					GoodHighKey = HighKey;
-				}
-				else
-				{
-					BadHighKey = HighKey;
-				}
-			}
-
-			//UE_LOG(LogAnimation, Warning, TEXT("%u: Keeping key %u"), BoneIndex, GoodHighKey);
-
-			NewTimes.Add(Times[GoodHighKey]);
-			NewKeys.Add(Keys[GoodHighKey]);
-
-			LowKey = GoodHighKey;
-		}
-
-		// return the new key set to the caller
-		Times = NewTimes;
-		Keys = NewKeys;
-	}
-}
-#endif
 
 void UAnimCompress_RemoveLinearKeys::UpdateWorldBoneTransformTable(
 	const FCompressibleAnimData& CompressibleAnimData,
@@ -615,483 +363,67 @@ void UAnimCompress_RemoveLinearKeys::UpdateWorldBoneTransformTable(
 	}
 }
 
-static void PackVector(
-	const FVector& Key,
-	AnimationCompressionFormat Format,
-	const float* Mins,
-	const float* Ranges,
-	uint8* Buffer)
+
+void* UAnimCompress_RemoveLinearKeys::FilterBeforeMainKeyRemoval(
+	const FCompressibleAnimData& CompressibleAnimData,
+	TArray<FTranslationTrack>& TranslationData,
+	TArray<FRotationTrack>& RotationData, 
+	TArray<FScaleTrack>& ScaleData)
 {
-	if (Format == ACF_None || Format == ACF_Float96NoW)
-	{
-		memcpy(Buffer, &Key, sizeof(FVector));
-	}
-	else if (Format == ACF_IntervalFixed32NoW)
-	{
-		const FVectorIntervalFixed32NoW CompressedKey(Key, Mins, Ranges);
-		memcpy(Buffer, &CompressedKey, sizeof(FVectorIntervalFixed32NoW));
-	}
+	// remove obviously redundant keys from the source data
+	FilterTrivialKeys(TranslationData, RotationData, ScaleData, TRANSLATION_ZEROING_THRESHOLD, QUATERNION_ZEROING_THRESHOLD, SCALE_ZEROING_THRESHOLD);
+	return nullptr;
 }
 
-static FVector UnpackVector(
-	const uint8* Buffer,
-	AnimationCompressionFormat Format,
-	const float* Mins,
-	const float* Ranges)
+
+
+void UAnimCompress_RemoveLinearKeys::UpdateWorldBoneTransformRange(
+	const FCompressibleAnimData& CompressibleAnimData,
+	FCompressibleAnimDataResult& OutCompressedData,
+	const TArray<FTransform>& RefPose,
+	const TArray<FTranslationTrack>& PositionTracks,
+	const TArray<FRotationTrack>& RotationTracks,
+	const TArray<FScaleTrack>& ScaleTracks,
+	int32 StartingBoneIndex, // this bone index should be of skeleton, not mesh
+	int32 EndingBoneIndex,// this bone index should be of skeleton, not mesh
+	bool UseRaw,
+	TArray<FTransform>& OutputWorldBones)
 {
-	FVector Result;
+	// bitwise compress the tracks into the anim sequence buffers
+	// to make sure the data we've compressed so far is ready for solving
+	CompressUsingUnderlyingCompressor(
+		CompressibleAnimData,
+		OutCompressedData,
+		PositionTracks,
+		RotationTracks,
+		ScaleTracks,
+		false);
 
-	if (Format == ACF_None || Format == ACF_Float96NoW)
-	{
-		memcpy(&Result, Buffer, sizeof(FVector));
-	}
-	else if (Format == ACF_IntervalFixed32NoW)
+	// build all world-space transforms from this bone to the target end effector we are monitoring
+	// all parent transforms have been built already
+	for ( int32 Index = StartingBoneIndex; Index <= EndingBoneIndex; ++Index )
 	{
-		FVectorIntervalFixed32NoW CompressedKey;
-		memcpy(&CompressedKey, Buffer, sizeof(FVectorIntervalFixed32NoW));
-		CompressedKey.ToVector(Result, Mins, Ranges);
+		UpdateWorldBoneTransformTable(
+			CompressibleAnimData,
+			OutCompressedData,
+			RefPose,
+			Index,
+			UseRaw,
+			OutputWorldBones);
 	}
-
-	return Result;
 }
 
-static void PackQuat(
-	const FQuat& Key,
-	AnimationCompressionFormat Format,
-	const float* Mins,
-	const float* Ranges,
-	uint8* Buffer)
+void UAnimCompress_RemoveLinearKeys::UpdateBoneAtomList(
+	const FCompressibleAnimData& CompressibleAnimData,
+	FCompressibleAnimDataResult& OutCompressedData,
+	int32 BoneIndex,
+	int32 TrackIndex,
+	int32 NumFrames,
+	float TimePerFrame,
+	TArray<FTransform>& BoneAtoms)
 {
-	if (Format == ACF_None)
-	{
-		memcpy(Buffer, &Key, sizeof(FQuat));
-	}
-	else if (Format == ACF_Float96NoW)
-	{
-		const FQuatFloat96NoW QuatFloat96NoW(Key);
-		memcpy(Buffer, &QuatFloat96NoW, sizeof(FQuatFloat96NoW));
-	}
-	else if (Format == ACF_Fixed32NoW)
-	{
-		const FQuatFixed32NoW QuatFixed32NoW(Key);
-		memcpy(Buffer, &QuatFixed32NoW, sizeof(FQuatFixed32NoW));
-	}
-	else if (Format == ACF_Fixed48NoW)
-	{
-		const FQuatFixed48NoW QuatFixed48NoW(Key);
-		memcpy(Buffer, &QuatFixed48NoW, sizeof(FQuatFixed48NoW));
-	}
-	else if (Format == ACF_IntervalFixed32NoW)
-	{
-		const FQuatIntervalFixed32NoW QuatIntervalFixed32NoW(Key, Mins, Ranges);
-		memcpy(Buffer, &QuatIntervalFixed32NoW, sizeof(FQuatIntervalFixed32NoW));
-	}
-	else if (Format == ACF_Float32NoW)
-	{
-		const FQuatFloat32NoW QuatFloat32NoW(Key);
-		memcpy(Buffer, &QuatFloat32NoW, sizeof(FQuatFloat32NoW));
-	}
-}
-
-static FQuat UnpackQuat(
-	const uint8* Buffer,
-	AnimationCompressionFormat Format,
-	const float* Mins,
-	const float* Ranges)
-{
-	FQuat Result;
-
-	if (Format == ACF_None)
-	{
-		memcpy(&Result, Buffer, sizeof(FQuat));
-	}
-	else if (Format == ACF_Float96NoW)
-	{
-		FQuatFloat96NoW QuatFloat96NoW;
-		memcpy(&QuatFloat96NoW, Buffer, sizeof(FQuatFloat96NoW));
-		QuatFloat96NoW.ToQuat(Result);
-	}
-	else if (Format == ACF_Fixed32NoW)
-	{
-		FQuatFixed32NoW QuatFixed32NoW;
-		memcpy(&QuatFixed32NoW, Buffer, sizeof(FQuatFixed32NoW));
-		QuatFixed32NoW.ToQuat(Result);
-	}
-	else if (Format == ACF_Fixed48NoW)
-	{
-		FQuatFixed48NoW QuatFixed48NoW;
-		memcpy(&QuatFixed48NoW, Buffer, sizeof(FQuatFixed48NoW));
-		QuatFixed48NoW.ToQuat(Result);
-	}
-	else if (Format == ACF_IntervalFixed32NoW)
-	{
-		FQuatIntervalFixed32NoW QuatIntervalFixed32NoW;
-		memcpy(&QuatIntervalFixed32NoW, Buffer, sizeof(FQuatIntervalFixed32NoW));
-		QuatIntervalFixed32NoW.ToQuat(Result, Mins, Ranges);
-	}
-	else if (Format == ACF_Float32NoW)
-	{
-		FQuatFloat32NoW QuatFloat32NoW;
-		memcpy(&QuatFloat32NoW, Buffer, sizeof(FQuatFloat32NoW));
-		QuatFloat32NoW.ToQuat(Result);
-	}
-
-	return Result;
-}
-
-static float GetFrameRate(const UAnimSequence& AnimSeq)
-{
-	// We cannot round since some sequences have a non integral frame rate for legacy reasons
-	return float(AnimSeq.GetRawNumberOfFrames() - 1) / AnimSeq.SequenceLength;
-}
-
-static float FindKeyInterpolationData(const UAnimSequence& AnimSeq, const TArray<float>& KeyTimes, float Time, int32& OutKeyIndex0, int32& OutKeyIndex1)
-{
-	// Use the same logic as the decompression in order to ensure we match 100%
-	const int32 LastFrame = AnimSeq.GetRawNumberOfFrames() - 1;
-	const float FrameRate = GetFrameRate(AnimSeq);
-	const float RelativePos = Time / AnimSeq.SequenceLength;
-	const float FramePos = RelativePos * float(LastFrame);
-	const int32 FramePosFloor = FMath::Clamp(FMath::TruncToInt(FramePos), 0, LastFrame);
-
-	const int32 NumKeys = KeyTimes.Num();
-	int32 KeyIndex0 = NumKeys - 1;	// Default
-	for (int32 KeyIndex = 1; KeyIndex < NumKeys; ++KeyIndex)
-	{
-		const float KeyTime = KeyTimes[KeyIndex];
-		const float FrameTime = KeyTime * FrameRate;
-		const int32 FrameIndex = FMath::Clamp(FMath::TruncToInt(FrameTime + 0.5f), 0, LastFrame);
-		if (FrameIndex > FramePosFloor)
-		{
-			KeyIndex0 = KeyIndex - 1;
-			break;
-		}
-	}
-
-	const int32 KeyIndex1 = FMath::Min(KeyIndex0 + 1, NumKeys - 1);
-	const float KeyTime0 = KeyTimes[KeyIndex0];
-	const float KeyTime1 = KeyTimes[KeyIndex1];
-
-	const float FrameTime0 = KeyTime0 * FrameRate;
-	const float FrameTime1 = KeyTime1 * FrameRate;
-	const int32 FrameIndex0 = FMath::Clamp(FMath::TruncToInt(FrameTime0 + 0.5f), 0, LastFrame);
-	const int32 FrameIndex1 = FMath::Clamp(FMath::TruncToInt(FrameTime1 + 0.5f), 0, LastFrame);
-
-	const int32 Delta = FMath::Max(FrameIndex1 - FrameIndex0, 1);
-	const float Remainder = FramePos - (float)FrameIndex0;
-	const float InterpolationAlpha = AnimSeq.Interpolation == EAnimInterpolationType::Step ? 0.0f : FMath::Min(Remainder / (float)Delta, 1.0f);
-
-	OutKeyIndex0 = KeyIndex0;
-	OutKeyIndex1 = KeyIndex1;
-	return InterpolationAlpha;
-}
-
-#if USE_SEGMENTING_CONTEXT
-FTransform UAnimCompress_RemoveLinearKeys::SampleSegment(const FProcessAnimationTracksContext& Context, int32 TrackIndex, float Time)
-{
-	FTransform Result;
-
-	{
-		const FTranslationTrack& TranslationData = Context.Segment.TranslationData[TrackIndex];
-		const int32 NumTransKeys = TranslationData.Times.Num();
-		if (NumTransKeys > 1)
-		{
-			int32 KeyIndex0;
-			int32 KeyIndex1;
-			const float InterpolationAlpha = FindKeyInterpolationData(Context.AnimSeq, TranslationData.Times, Time, KeyIndex0, KeyIndex1);
-
-			FVector RangeMin;
-			FVector RangeExtent;
-			CalculateTrackRange(TranslationData, Context.Segment.TranslationCompressionFormat, RangeMin, RangeExtent);
-
-			const FVector& Translation0 = TranslationData.PosKeys[KeyIndex0];
-			uint8 Buffer0[32];
-			PackVector(Translation0, Context.Segment.TranslationCompressionFormat, &RangeMin.X, &RangeExtent.X, Buffer0);
-			FVector LossyTranslation = UnpackVector(Buffer0, Context.Segment.TranslationCompressionFormat, &RangeMin.X, &RangeExtent.X);
-
-			if (KeyIndex0 != KeyIndex1)
-			{
-				const FVector& Translation1 = TranslationData.PosKeys[KeyIndex1];
-				uint8 Buffer1[32];
-				PackVector(Translation1, Context.Segment.TranslationCompressionFormat, &RangeMin.X, &RangeExtent.X, Buffer1);
-				const FVector LossyTranslation1 = UnpackVector(Buffer1, Context.Segment.TranslationCompressionFormat, &RangeMin.X, &RangeExtent.X);
-
-				LossyTranslation = FMath::Lerp(LossyTranslation, LossyTranslation1, InterpolationAlpha);
-			}
-
-			Result.SetTranslation(LossyTranslation);
-		}
-		else
-		{
-			Result.SetTranslation(Context.TrivialTrackValues[TrackIndex].GetTranslation());
-		}
-	}
-
-	{
-		const FRotationTrack& RotationData = Context.Segment.RotationData[TrackIndex];
-		const int32 NumRotKeys = RotationData.Times.Num();
-		if (NumRotKeys > 1)
-		{
-			int32 KeyIndex0;
-			int32 KeyIndex1;
-			const float InterpolationAlpha = FindKeyInterpolationData(Context.AnimSeq, RotationData.Times, Time, KeyIndex0, KeyIndex1);
-
-			FVector RangeMin;
-			FVector RangeExtent;
-			CalculateTrackRange(RotationData, Context.Segment.RotationCompressionFormat, RangeMin, RangeExtent);
-
-			const FQuat& Rotation0 = RotationData.RotKeys[KeyIndex0];
-			uint8 Buffer0[32];
-			PackQuat(Rotation0, Context.Segment.RotationCompressionFormat, &RangeMin.X, &RangeExtent.X, Buffer0);
-			FQuat LossyRotation = UnpackQuat(Buffer0, Context.Segment.RotationCompressionFormat, &RangeMin.X, &RangeExtent.X);
-
-			if (KeyIndex0 != KeyIndex1)
-			{
-				const FQuat& Rotation1 = RotationData.RotKeys[KeyIndex1];
-				uint8 Buffer1[32];
-				PackQuat(Rotation1, Context.Segment.RotationCompressionFormat, &RangeMin.X, &RangeExtent.X, Buffer1);
-				const FQuat LossyRotation1 = UnpackQuat(Buffer1, Context.Segment.RotationCompressionFormat, &RangeMin.X, &RangeExtent.X);
-
-				LossyRotation = FQuat::FastLerp(LossyRotation, LossyRotation1, InterpolationAlpha);
-				LossyRotation.Normalize();
-			}
-
-			Result.SetRotation(LossyRotation);
-		}
-		else
-		{
-			const FQuat Rotation = Context.TrivialTrackValues[TrackIndex].GetRotation();
-			uint8 Buffer[32];
-			PackQuat(Rotation, ACF_Float96NoW, nullptr, nullptr, Buffer);
-			const FQuat LossyRotation = UnpackQuat(Buffer, ACF_Float96NoW, nullptr, nullptr);
-
-			Result.SetRotation(LossyRotation);
-		}
-	}
-
-	if (Context.Segment.ScaleData.Num() != 0)
-	{
-		const FScaleTrack& ScaleData = Context.Segment.ScaleData[TrackIndex];
-		const int32 NumScaleKeys = ScaleData.Times.Num();
-		if (NumScaleKeys > 1)
-		{
-			int32 KeyIndex0;
-			int32 KeyIndex1;
-			const float InterpolationAlpha = FindKeyInterpolationData(Context.AnimSeq, ScaleData.Times, Time, KeyIndex0, KeyIndex1);
-
-			FVector RangeMin;
-			FVector RangeExtent;
-			CalculateTrackRange(ScaleData, Context.Segment.ScaleCompressionFormat, RangeMin, RangeExtent);
-
-			const FVector& Scale0 = ScaleData.ScaleKeys[KeyIndex0];
-			uint8 Buffer0[32];
-			PackVector(Scale0, Context.Segment.ScaleCompressionFormat, &RangeMin.X, &RangeExtent.X, Buffer0);
-			FVector LossyScale = UnpackVector(Buffer0, Context.Segment.ScaleCompressionFormat, &RangeMin.X, &RangeExtent.X);
-
-			if (KeyIndex0 != KeyIndex1)
-			{
-				const FVector& Scale1 = ScaleData.ScaleKeys[KeyIndex1];
-				uint8 Buffer1[32];
-				PackVector(Scale1, Context.Segment.ScaleCompressionFormat, &RangeMin.X, &RangeExtent.X, Buffer1);
-				const FVector LossyScale1 = UnpackVector(Buffer1, Context.Segment.ScaleCompressionFormat, &RangeMin.X, &RangeExtent.X);
-
-				LossyScale = FMath::Lerp(LossyScale, LossyScale1, InterpolationAlpha);
-			}
-
-			Result.SetScale3D(LossyScale);
-		}
-		else
-		{
-			Result.SetScale3D(Context.TrivialTrackValues[TrackIndex].GetScale3D());
-		}
-	}
-
-	return Result;
-}
-
-void UAnimCompress_RemoveLinearKeys::UpdateWorldBoneTransformRange(
-	FProcessAnimationTracksContext& Context,
-	int32 StartingBoneIndex,	// this bone index should be of skeleton, not mesh
-	int32 EndingBoneIndex)		// this bone index should be of skeleton, not mesh
-{
-	if (bUseDecompression)
-	{
-		// bitwise compress the tracks into the anim sequence buffers
-		// to make sure the data we've compressed so far is ready for solving
-		CompressUsingUnderlyingCompressor(Context.AnimSeq, Context.BoneData, Context.SegmentList, false);
-	}
-
-	// build all world-space transforms from this bone to the target end effector we are monitoring
-	// all parent transforms have been built already
-	for (int32 Index = StartingBoneIndex; Index <= EndingBoneIndex; ++Index)
-	{
-		UpdateWorldBoneTransformTable(Context, Index, false, Context.NewWorldBones);
-	}
-}
-
-void UAnimCompress_RemoveLinearKeys::UpdateBoneAtomList(
-	const FProcessAnimationTracksContext& Context,
-	int32 TrackIndex,
-	TArray<FTransform>& BoneAtoms) const
-{
-	const float FrameRate = GetFrameRate(Context.AnimSeq);
-	FAnimSequenceDecompressionContext DecompContext(&Context.AnimSeq);
-
-	BoneAtoms.Reset(Context.Segment.NumFrames);
-	for (int32 FrameIndex = 0; FrameIndex < Context.Segment.NumFrames; ++FrameIndex)
-	{
-		const float Time = (float)(FrameIndex + Context.Segment.StartFrame) / FrameRate;
-
-		FTransform LocalAtom;
-		if (bUseDecompression)
-		{
-			DecompContext.Seek(Time);
-			Context.AnimSeq.GetBoneTransform(LocalAtom, TrackIndex, DecompContext, false);
-		}
-		else
-		{
-			LocalAtom = SampleSegment(Context, TrackIndex, Time);
-		}
-
-		FQuat Rot = LocalAtom.GetRotation();
-		LocalAtom.SetRotation(EnforceShortestArc(FQuat::Identity, Rot));
-		BoneAtoms.Add(LocalAtom);
-	}
-}
-
-void UAnimCompress_RemoveLinearKeys::UpdateWorldBoneTransformTable(
-	const FProcessAnimationTracksContext& Context,
-	int32 BoneIndex,
-	bool UseRaw,
-	TArray<FTransform>& OutputWorldBones)
-{
-	const int32 TrackIndex = Context.BoneIndexToTrackIndex[BoneIndex];
-
-	check(OutputWorldBones.Num() >= ((BoneIndex * Context.Segment.NumFrames) + Context.Segment.NumFrames));
-
-	if (TrackIndex != INDEX_NONE)
-	{
-		const float FrameRate = GetFrameRate(Context.AnimSeq);
-		FAnimSequenceDecompressionContext DecompContext(&Context.AnimSeq);
-
-		// get the local-space bone transforms using the animation solver
-		for (int32 FrameIndex = 0; FrameIndex < Context.Segment.NumFrames; ++FrameIndex)
-		{
-			const float Time = (float)(Context.Segment.StartFrame + FrameIndex) / FrameRate;
-
-			FTransform LocalAtom;
-			if (UseRaw)
-			{
-				DecompContext.Seek(Time);
-				Context.AnimSeq.GetBoneTransform(LocalAtom, TrackIndex, DecompContext, UseRaw);
-			}
-			else
-			{
-				if (bUseDecompression)
-				{
-					DecompContext.Seek(Time);
-					Context.AnimSeq.GetBoneTransform(LocalAtom, TrackIndex, DecompContext, UseRaw);
-				}
-				else
-				{
-					LocalAtom = SampleSegment(Context, TrackIndex, Time);
-				}
-			}
-
-			FQuat Rot = LocalAtom.GetRotation();
-			LocalAtom.SetRotation(EnforceShortestArc(FQuat::Identity, Rot));
-			// Saw some crashes happening with it, so normalize here.
-			LocalAtom.NormalizeRotation();
-
-			OutputWorldBones[(BoneIndex * Context.Segment.NumFrames) + FrameIndex] = LocalAtom;
-		}
-	}
-	else
-	{
-		// get the default rotation and translation from the reference skeleton
-		FTransform DefaultTransform = Context.RefPose[BoneIndex];
-		DefaultTransform.SetRotation(EnforceShortestArc(FQuat::Identity, DefaultTransform.GetRotation()));
-
-		// copy the default transformation into the world bone table
-		for (int32 FrameIndex = 0; FrameIndex < Context.Segment.NumFrames; ++FrameIndex)
-		{
-			OutputWorldBones[(BoneIndex * Context.Segment.NumFrames) + FrameIndex] = DefaultTransform;
-		}
-	}
-
-	// apply parent transforms to bake into world space. We assume the parent transforms were previously set using this function
-	const FBoneData& Bone = Context.BoneData[BoneIndex];
-	const int32 ParentIndex = Bone.GetParent();
-	if (ParentIndex != INDEX_NONE)
-	{
-		check(ParentIndex < BoneIndex);
-		for (int32 FrameIndex = 0; FrameIndex < Context.Segment.NumFrames; ++FrameIndex)
-		{
-			OutputWorldBones[(BoneIndex * Context.Segment.NumFrames) + FrameIndex] = OutputWorldBones[(BoneIndex * Context.Segment.NumFrames) + FrameIndex] * OutputWorldBones[(ParentIndex * Context.Segment.NumFrames) + FrameIndex];
-		}
-	}
-}
-#endif
-
-void UAnimCompress_RemoveLinearKeys::FilterBeforeMainKeyRemoval(
-	const FCompressibleAnimData& CompressibleAnimData,
-	TArray<FTranslationTrack>& TranslationData,
-	TArray<FRotationTrack>& RotationData, 
-	TArray<FScaleTrack>& ScaleData)
-{
-	// remove obviously redundant keys from the source data
-	FilterTrivialKeys(TranslationData, RotationData, ScaleData, TRANSLATION_ZEROING_THRESHOLD, QUATERNION_ZEROING_THRESHOLD, SCALE_ZEROING_THRESHOLD);
-}
-
-
-
-void UAnimCompress_RemoveLinearKeys::UpdateWorldBoneTransformRange(
-	const FCompressibleAnimData& CompressibleAnimData,
-	FCompressibleAnimDataResult& OutCompressedData,
-	const TArray<FTransform>& RefPose,
-	const TArray<FTranslationTrack>& PositionTracks,
-	const TArray<FRotationTrack>& RotationTracks,
-	const TArray<FScaleTrack>& ScaleTracks,
-	int32 StartingBoneIndex, // this bone index should be of skeleton, not mesh
-	int32 EndingBoneIndex,// this bone index should be of skeleton, not mesh
-	bool UseRaw,
-	TArray<FTransform>& OutputWorldBones)
-{
-	// bitwise compress the tracks into the anim sequence buffers
-	// to make sure the data we've compressed so far is ready for solving
-	CompressUsingUnderlyingCompressor(
-		CompressibleAnimData,
-		OutCompressedData,
-		PositionTracks,
-		RotationTracks,
-		ScaleTracks,
-		false);
-
-	// build all world-space transforms from this bone to the target end effector we are monitoring
-	// all parent transforms have been built already
-	for ( int32 Index = StartingBoneIndex; Index <= EndingBoneIndex; ++Index )
-	{
-		UpdateWorldBoneTransformTable(
-			CompressibleAnimData,
-			OutCompressedData,
-			RefPose,
-			Index,
-			UseRaw,
-			OutputWorldBones);
-	}
-}
-
-void UAnimCompress_RemoveLinearKeys::UpdateBoneAtomList(
-	const FCompressibleAnimData& CompressibleAnimData,
-	FCompressibleAnimDataResult& OutCompressedData,
-	int32 BoneIndex,
-	int32 TrackIndex,
-	int32 NumFrames,
-	float TimePerFrame,
-	TArray<FTransform>& BoneAtoms)
-{
-	BoneAtoms.Reset(NumFrames);
-	for ( int32 FrameIndex = 0; FrameIndex < NumFrames; ++FrameIndex )
+	BoneAtoms.Reset(NumFrames);
+	for ( int32 FrameIndex = 0; FrameIndex < NumFrames; ++FrameIndex )
 	{
 		float Time = (float)FrameIndex * TimePerFrame;
 		FTransform LocalAtom;
@@ -1670,645 +1002,40 @@ void UAnimCompress_RemoveLinearKeys::ProcessAnimationTracks(
 	}
 };
 
-#if USE_SEGMENTING_CONTEXT
-static void CalculateBoneChainInformation(const FProcessAnimationTracksContext& Context, int32 BoneIndex, int32& OutHighestTargetBoneIndex, int32& OutFurthestTargetBoneIndex, int32& OutShortestChain, TArray<int32>& OutTargetBoneIndices)
+void UAnimCompress_RemoveLinearKeys::CompressUsingUnderlyingCompressor(
+	const FCompressibleAnimData& CompressibleAnimData,
+	FCompressibleAnimDataResult& OutCompressedData,
+	const TArray<FTranslationTrack>& TranslationData,
+	const TArray<FRotationTrack>& RotationData,
+	const TArray<FScaleTrack>& ScaleData,
+	const bool bFinalPass)
 {
-	const int32 NumBones = Context.BoneData.Num();
+	BitwiseCompressAnimationTracks(
+		CompressibleAnimData,
+		OutCompressedData,
+		static_cast<AnimationCompressionFormat>(TranslationCompressionFormat),
+		static_cast<AnimationCompressionFormat>(RotationCompressionFormat),
+		static_cast<AnimationCompressionFormat>(ScaleCompressionFormat),
+		TranslationData,
+		RotationData,
+		ScaleData,
+		true);
 
-	OutTargetBoneIndices.Reset(NumBones);
+	// record the proper runtime decompressor to use
+	FUECompressedAnimDataMutable& AnimData = static_cast<FUECompressedAnimDataMutable&>(*OutCompressedData.AnimData);
+	AnimData.KeyEncodingFormat = AKF_VariableKeyLerp;
+	AnimationFormat_SetInterfaceLinks(AnimData);
+}
 
-	int32 HighestTargetBoneIndex = BoneIndex;
-	int32 FurthestTargetBoneIndex = BoneIndex;
-	int32 ShortestChain = 0;
-	float OffsetLength = -1.0f;
-	for (int32 EffectorIndex = 0; EffectorIndex < Context.EndEffectors.Num(); ++EffectorIndex)
-	{
-		const int32 EffectorBoneIndex = Context.EndEffectors[EffectorIndex];
-		const FBoneData& EffectorBoneData = Context.BoneData[EffectorBoneIndex];
+bool UAnimCompress_RemoveLinearKeys::DoReduction(const FCompressibleAnimData& CompressibleAnimData, FCompressibleAnimDataResult& OutResult)
+{
+#if WITH_EDITORONLY_DATA
+	// Only need to do the heavy lifting if it will have some impact
+	// One of these will always be true for the base class, but derived classes may choose to turn both off (e.g., in PerTrackCompression)
+	const bool bRunningProcessor = bRetarget || bActuallyFilterLinearKeys;
 
-		int32 RootIndex = EffectorBoneData.BonesToRoot.Find(BoneIndex);
-		if (RootIndex != INDEX_NONE)
-		{
-			if (ShortestChain == 0 || (RootIndex + 1) < ShortestChain)
-			{
-				ShortestChain = (RootIndex + 1);
-			}
-
-			OutTargetBoneIndices.Add(EffectorBoneIndex);
-			HighestTargetBoneIndex = FMath::Max(HighestTargetBoneIndex, EffectorBoneIndex);
-			float ChainLength = 0.0f;
-			for (long FamilyIndex = 0; FamilyIndex < RootIndex; ++FamilyIndex)
-			{
-				const int32 NextParentBoneIndex = EffectorBoneData.BonesToRoot[FamilyIndex];
-				ChainLength += Context.RefPose[NextParentBoneIndex].GetTranslation().Size();
-			}
-
-			if (ChainLength > OffsetLength)
-			{
-				FurthestTargetBoneIndex = EffectorBoneIndex;
-				OffsetLength = ChainLength;
-			}
-		}
-	}
-
-	OutHighestTargetBoneIndex = HighestTargetBoneIndex;
-	OutFurthestTargetBoneIndex = FurthestTargetBoneIndex;
-	OutShortestChain = ShortestChain;
-}
-
-void UAnimCompress_RemoveLinearKeys::PerformRetargeting(
-	FProcessAnimationTracksContext& Context,
-	int32 BoneIndex,
-	int32 HighestTargetBoneIndex,
-	int32 FurthestTargetBoneIndex)
-{
-	const int32 TrackIndex = Context.BoneIndexToTrackIndex[BoneIndex];
-	const FBoneData& Bone = Context.BoneData[BoneIndex];
-	const int32 ParentBoneIndex = Bone.GetParent();
-
-	const int32 NumScaleKeys = Context.Segment.ScaleData.Num() > 0 ? Context.Segment.ScaleData[TrackIndex].ScaleKeys.Num() : 0;
-	if (NumScaleKeys > 1 && ParentBoneIndex != INDEX_NONE)
-	{
-		// update our bone table from the current bone through the last end effector we need to test
-		UpdateWorldBoneTransformRange(Context, BoneIndex, HighestTargetBoneIndex);
-
-		FScaleTrack& ScaleTrack = Context.Segment.ScaleData[TrackIndex];
-
-		// adjust all translation keys to align better with the destination
-		for (int32 KeyIndex = 0; KeyIndex < NumScaleKeys; ++KeyIndex)
-		{
-			FVector& Key = ScaleTrack.ScaleKeys[KeyIndex];
-
-			const int32 FrameIndex = FMath::Clamp(KeyIndex, 0, Context.Segment.NumFrames);
-			const FTransform& NewWorldParent = Context.NewWorldBones[(ParentBoneIndex * Context.Segment.NumFrames) + FrameIndex];
-			const FTransform& RawWorldChild = Context.RawWorldBones[(BoneIndex * Context.Segment.NumFrames) + FrameIndex];
-			const FTransform RelTM = RawWorldChild.GetRelativeTransform(NewWorldParent);
-
-			Key = RelTM.GetScale3D();
-		}
-	}
-
-	FRotationTrack& RotTrack = Context.Segment.RotationData[TrackIndex];
-	const int32 NumRotKeys = RotTrack.RotKeys.Num();
-	if (NumRotKeys > 1 && ParentBoneIndex != INDEX_NONE)
-	{
-		// update our bone table from the current bone through the last end effector we need to test
-		UpdateWorldBoneTransformRange(Context, BoneIndex, HighestTargetBoneIndex);
-
-		if (HighestTargetBoneIndex == BoneIndex)
-		{
-			for (int32 KeyIndex = 0; KeyIndex < NumRotKeys; ++KeyIndex)
-			{
-				FQuat& Key = RotTrack.RotKeys[KeyIndex];
-
-				const int32 FrameIndex = FMath::Clamp(KeyIndex, 0, Context.Segment.NumFrames);
-				FTransform NewWorldParent = Context.NewWorldBones[(ParentBoneIndex * Context.Segment.NumFrames) + FrameIndex];
-				FTransform RawWorldChild = Context.RawWorldBones[(BoneIndex * Context.Segment.NumFrames) + FrameIndex];
-				const FTransform RelTM = RawWorldChild.GetRelativeTransform(NewWorldParent);
-				FQuat Rot = RelTM.GetRotation();
-
-				const FQuat AlignedKey = EnforceShortestArc(Key, Rot);
-				Key = AlignedKey;
-			}
-		}
-		else
-		{
-			// adjust all rotation keys towards the end effector target
-			for (int32 KeyIndex = 0; KeyIndex < NumRotKeys; ++KeyIndex)
-			{
-				FQuat& Key = RotTrack.RotKeys[KeyIndex];
-
-				const int32 FrameIndex = FMath::Clamp(KeyIndex, 0, Context.Segment.NumFrames);
-
-				const FTransform& NewWorldTransform = Context.NewWorldBones[(BoneIndex * Context.Segment.NumFrames) + FrameIndex];
-
-				const FTransform DesiredChildTransform = Context.RawWorldBones[(FurthestTargetBoneIndex * Context.Segment.NumFrames) + FrameIndex].GetRelativeTransform(NewWorldTransform);
-				const FTransform CurrentChildTransform = Context.NewWorldBones[(FurthestTargetBoneIndex * Context.Segment.NumFrames) + FrameIndex].GetRelativeTransform(NewWorldTransform);
-
-				// find the two vectors which represent the angular error we are trying to correct
-				const FVector CurrentHeading = CurrentChildTransform.GetTranslation();
-				const FVector DesiredHeading = DesiredChildTransform.GetTranslation();
-
-				// if these are valid, we can continue
-				if (!CurrentHeading.IsNearlyZero() && !DesiredHeading.IsNearlyZero())
-				{
-					const float DotResult = CurrentHeading.GetSafeNormal() | DesiredHeading.GetSafeNormal();
-
-					// limit the range we will retarget to something reasonable (~60 degrees)
-					if (DotResult < 1.0f && DotResult > 0.5f)
-					{
-						FQuat Adjustment = FQuat::FindBetweenVectors(CurrentHeading, DesiredHeading);
-						Adjustment = EnforceShortestArc(FQuat::Identity, Adjustment);
-
-						const FVector Test = Adjustment.RotateVector(CurrentHeading);
-						const float DeltaSqr = (Test - DesiredHeading).SizeSquared();
-						if (DeltaSqr < FMath::Square(0.001f))
-						{
-							FQuat NewKey = Adjustment * Key;
-							NewKey.Normalize();
-
-							const FQuat AlignedKey = EnforceShortestArc(Key, NewKey);
-							Key = AlignedKey;
-						}
-					}
-				}
-			}
-		}
-	}
-
-	FTranslationTrack& TransTrack = Context.Segment.TranslationData[TrackIndex];
-	const int32 NumPosKeys = TransTrack.PosKeys.Num();
-	if (NumPosKeys > 1 && ParentBoneIndex != INDEX_NONE)
-	{
-		// update our bone table from the current bone through the last end effector we need to test
-		UpdateWorldBoneTransformRange(Context, BoneIndex, HighestTargetBoneIndex);
-
-		// adjust all translation keys to align better with the destination
-		for (int32 KeyIndex = 0; KeyIndex < NumPosKeys; ++KeyIndex)
-		{
-			FVector& Key = TransTrack.PosKeys[KeyIndex];
-
-			const int32 FrameIndex = FMath::Clamp(KeyIndex, 0, Context.Segment.NumFrames);
-			const FTransform& NewWorldParent = Context.NewWorldBones[(ParentBoneIndex * Context.Segment.NumFrames) + FrameIndex];
-			const FTransform& RawWorldChild = Context.RawWorldBones[(BoneIndex * Context.Segment.NumFrames) + FrameIndex];
-			const FTransform RelTM = RawWorldChild.GetRelativeTransform(NewWorldParent);
-
-			Key = RelTM.GetTranslation();
-		}
-	}
-}
-#endif
-
-static int32 FindGuideTrackIndex(const FBoneData& Bone, const TArray<int32>& BoneIndexToTrackIndex, float ParentKeyScale)
-{
-	int32 GuideTrackIndex = INDEX_NONE;
-	if (ParentKeyScale > 1.0f)
-	{
-		for (int32 FamilyIndex = 0; FamilyIndex < Bone.BonesToRoot.Num() && GuideTrackIndex == INDEX_NONE; ++FamilyIndex)
-		{
-			const int32 NextParentBoneIndex = Bone.BonesToRoot[FamilyIndex];
-
-			GuideTrackIndex = BoneIndexToTrackIndex[NextParentBoneIndex];
-		}
-	}
-
-	return GuideTrackIndex;
-}
-#if USE_SEGMENTING_CONTEXT
-void UAnimCompress_RemoveLinearKeys::ProcessAnimationTracks(FProcessAnimationTracksContext& Context)
-{
-	const int32 NumBones = Context.BoneData.Num();
-
-	Context.RawWorldBones.Empty(NumBones * Context.Segment.NumFrames);
-	Context.NewWorldBones.Empty(NumBones * Context.Segment.NumFrames);
-	Context.RawWorldBones.AddZeroed(NumBones * Context.Segment.NumFrames);
-	Context.NewWorldBones.AddZeroed(NumBones * Context.Segment.NumFrames);
-
-	TArray<int32> TargetBoneIndices;
-	TargetBoneIndices.Empty(NumBones);
-
-	// Create an array of FTransform to use as a workspace
-	TArray<FTransform> BoneAtoms;
-
-	// setup the raw bone transformation and find all end effectors
-	for (int32 BoneIndex = 0; BoneIndex < NumBones; ++BoneIndex)
-	{
-		// get the raw world-atoms for this bone
-		UpdateWorldBoneTransformTable(Context, BoneIndex, true, Context.RawWorldBones);
-	}
-
-	const bool bHasScale = Context.Segment.ScaleData.Num() > 0;
-
-	// for each bone...
-	for (int32 BoneIndex = 0; BoneIndex < NumBones; ++BoneIndex)
-	{
-		const FBoneData& Bone = Context.BoneData[BoneIndex];
-		const int32 ParentBoneIndex = Bone.GetParent();
-
-		const int32 TrackIndex = Context.BoneIndexToTrackIndex[BoneIndex];
-
-		if (TrackIndex != INDEX_NONE)
-		{
-			// get the tracks we will be editing for this bone
-			FRotationTrack& RotTrack = Context.Segment.RotationData[TrackIndex];
-			FTranslationTrack& TransTrack = Context.Segment.TranslationData[TrackIndex];
-			const int32 NumRotKeys = RotTrack.RotKeys.Num();
-			const int32 NumPosKeys = TransTrack.PosKeys.Num();
-			const int32 NumScaleKeys = bHasScale ? Context.Segment.ScaleData[TrackIndex].ScaleKeys.Num() : 0;
-
-			check((NumPosKeys == 1) || (NumRotKeys == 1) || (NumPosKeys == NumRotKeys));
-
-			// build an array of end effectors we need to monitor
-			int32 HighestTargetBoneIndex;
-			int32 FurthestTargetBoneIndex;
-			int32 ShortestChain;
-			CalculateBoneChainInformation(Context, BoneIndex, HighestTargetBoneIndex, FurthestTargetBoneIndex, ShortestChain, TargetBoneIndices);
-
-			// if requested, retarget the FBoneAtoms towards the target end effectors
-			if (bRetarget)
-			{
-				PerformRetargeting(Context, BoneIndex, HighestTargetBoneIndex, FurthestTargetBoneIndex);
-			}
-
-			// look for a parent track to reference as a guide
-			const int32 GuideTrackIndex = FindGuideTrackIndex(Bone, Context.BoneIndexToTrackIndex, ParentKeyScale);
-
-			// update our bone table from the current bone through the last end effector we need to test
-			UpdateWorldBoneTransformRange(Context, BoneIndex, HighestTargetBoneIndex);
-
-			// rebuild the BoneAtoms table using the current set of keys
-			UpdateBoneAtomList(Context, TrackIndex, BoneAtoms);
-
-			// determine the EndEffectorTolerance. 
-			// We use the Maximum value by default, and the Minimum value
-			// as we approach the end effectors
-			const float EndEffectorTolerance = ShortestChain <= 1 ? MinEffectorDiff : MaxEffectorDiff;
-
-			// Determine if a guidance track should be used to aid in choosing keys to retain
-			const TArray<float>* GuidanceTrack = nullptr;
-			float GuidanceScale = 1.0f;
-			if (GuideTrackIndex != INDEX_NONE)
-			{
-				const FTranslationTrack& GuideTransTrack = Context.Segment.TranslationData[GuideTrackIndex];
-				GuidanceTrack = &GuideTransTrack.Times;
-				GuidanceScale = ParentKeyScale;
-			}
-
-			// if the TargetBoneIndices array is empty, then this bone is an end effector.
-			// so we add it to the list to maintain our tolerance checks
-			if (TargetBoneIndices.Num() == 0)
-			{
-				TargetBoneIndices.Add(BoneIndex);
-			}
-
-			if (bActuallyFilterLinearKeys)
-			{
-				if (bHasScale)
-				{
-					FScaleTrack& ScaleTrack = Context.Segment.ScaleData[TrackIndex];
-					// filter out translations we can approximate through interpolation
-					FilterLinearKeysTemplate<ScaleAdapter>(
-						ScaleTrack.ScaleKeys,
-						ScaleTrack.Times,
-						BoneAtoms,
-						GuidanceTrack,
-						Context,
-						TargetBoneIndices,
-						BoneIndex,
-						ParentBoneIndex,
-						GuidanceScale,
-						MaxScaleDiff,
-						EndEffectorTolerance,
-						EffectorDiffSocket);
-
-					// update our bone table from the current bone through the last end effector we need to test
-					UpdateWorldBoneTransformRange(Context, BoneIndex, HighestTargetBoneIndex);
-
-					// rebuild the BoneAtoms table using the current set of keys
-					UpdateBoneAtomList(Context, TrackIndex, BoneAtoms);
-				}
-
-				// filter out translations we can approximate through interpolation
-				FilterLinearKeysTemplate<TranslationAdapter>(
-					TransTrack.PosKeys,
-					TransTrack.Times,
-					BoneAtoms,
-					GuidanceTrack,
-					Context,
-					TargetBoneIndices,
-					BoneIndex,
-					ParentBoneIndex,
-					GuidanceScale,
-					MaxPosDiff,
-					EndEffectorTolerance,
-					EffectorDiffSocket);
-
-				// update our bone table from the current bone through the last end effector we need to test
-				UpdateWorldBoneTransformRange(Context, BoneIndex, HighestTargetBoneIndex);
-
-				// rebuild the BoneAtoms table using the current set of keys
-				UpdateBoneAtomList(Context, TrackIndex, BoneAtoms);
-
-				// filter out rotations we can approximate through interpolation
-				FilterLinearKeysTemplate<RotationAdapter>(
-					RotTrack.RotKeys,
-					RotTrack.Times,
-					BoneAtoms,
-					GuidanceTrack,
-					Context,
-					TargetBoneIndices,
-					BoneIndex,
-					ParentBoneIndex,
-					GuidanceScale,
-					MaxAngleDiff,
-					EndEffectorTolerance,
-					EffectorDiffSocket);
-			}
-		}
-
-		// make sure the final compressed keys are represented in our NewWorldBones table
-		UpdateWorldBoneTransformRange(Context, BoneIndex, BoneIndex);
-	}
-}
-
-/**
- * Struct that holds the relevant information to process animation tracks in parallel.
- * Instances of this structure are live as long as parallel task instances are live.
- */
-struct FAsyncProcessAnimationTracksTaskGroupContext
-{
-	TArray<FProcessAnimationTracksContext*> TaskContexes;
-	volatile int32 AtomicTaskIndexCounter;
-	volatile int32 AtomicNumExecutedTasks;
-
-	FAsyncProcessAnimationTracksTaskGroupContext()
-		: AtomicTaskIndexCounter(0)
-		, AtomicNumExecutedTasks(0)
-	{}
-
-	~FAsyncProcessAnimationTracksTaskGroupContext()
-	{
-		for (FProcessAnimationTracksContext* Context : TaskContexes)
-		{
-			delete Context;
-		}
-	}
-
-	void ExecuteTasks(UAnimCompress_RemoveLinearKeys* Compressor)
-	{
-		while (true)
-		{
-			const int32 TaskIndex = FPlatformAtomics::InterlockedIncrement(&AtomicTaskIndexCounter) - 1;
-			if (TaskIndex >= TaskContexes.Num())
-			{
-				break;
-			}
-
-			FProcessAnimationTracksContext* JobContext = TaskContexes[TaskIndex];
-			Compressor->ProcessAnimationTracks(*JobContext);
-
-			FPlatformAtomics::InterlockedIncrement(&AtomicNumExecutedTasks);
-		}
-	}
-
-	void WaitForAllTasks()
-	{
-		// We just spin wait until everything is done
-		// This is a decent option because segments are already sorted largest to smallest and so
-		// they should all take about the same amount of time. We should never end up waiting here for too long.
-		while (AtomicNumExecutedTasks != TaskContexes.Num());
-	}
-};
-
-/**
- * Class that represents a parallel task to perform linear key reduction.
- */
-class FAsyncProcessAnimationTracksTask
-{
-public:
-	FAsyncProcessAnimationTracksTask(UAnimCompress_RemoveLinearKeys* Compressor_, FAsyncProcessAnimationTracksTaskGroupContext* TaskGroupContext_)
-		: Compressor(Compressor_)
-		, TaskGroupContext(TaskGroupContext_)
-	{}
-
-	/** return the name of the task **/
-	static const TCHAR* GetTaskName() { return TEXT("FAsyncProcessAnimationTracksTask"); }
-	FORCEINLINE static TStatId GetStatId() { RETURN_QUICK_DECLARE_CYCLE_STAT(FAsyncProcessAnimationTracksTask, STATGROUP_TaskGraphTasks); }
-
-	static ENamedThreads::Type GetDesiredThread() { return ENamedThreads::AnyThread; }
-	static ESubsequentsMode::Type GetSubsequentsMode() { return ESubsequentsMode::TrackSubsequents; }
-
-	void DoTask(ENamedThreads::Type CurrentThread, const FGraphEventRef& MyCompletionGraphEvent)
-	{
-		TaskGroupContext->ExecuteTasks(Compressor);
-	}
-
-	UAnimCompress_RemoveLinearKeys* Compressor;
-	FAsyncProcessAnimationTracksTaskGroupContext* TaskGroupContext;
-};
-
-/**
- * Class that represents a clean up task once every parallel task has executed.
- * This is executed last.
- */
-class FAsyncCleanUpProcessAnimationTracksTask
-{
-public:
-	FAsyncCleanUpProcessAnimationTracksTask(FAsyncProcessAnimationTracksTaskGroupContext* TaskGroupContext_)
-		: TaskGroupContext(TaskGroupContext_)
-	{}
-
-	/** return the name of the task **/
-	static const TCHAR* GetTaskName() { return TEXT("FAsyncCleanUpProcessAnimationTracksTask"); }
-	FORCEINLINE static TStatId GetStatId() { RETURN_QUICK_DECLARE_CYCLE_STAT(FAsyncCleanUpProcessAnimationTracksTask, STATGROUP_TaskGraphTasks); }
-
-	static ENamedThreads::Type GetDesiredThread() { return ENamedThreads::AnyThread; }
-	static ESubsequentsMode::Type GetSubsequentsMode() { return ESubsequentsMode::TrackSubsequents; }
-
-	void DoTask(ENamedThreads::Type CurrentThread, const FGraphEventRef& MyCompletionGraphEvent)
-	{
-		delete TaskGroupContext;
-	}
-
-	FAsyncProcessAnimationTracksTaskGroupContext* TaskGroupContext;
-};
-
-void UAnimCompress_RemoveLinearKeys::ProcessAnimationTracks(
-	UAnimSequence& AnimSeq,
-	const TArray<FBoneData>& BoneData,
-	TArray<FAnimSegmentContext>& RawSegments)
-{
-	// extract all the data we'll need about the skeleton and animation sequence
-	const int32 NumBones = BoneData.Num();
-
-	// make sure the parent key scale is properly bound to 1.0 or more
-	ParentKeyScale = FMath::Max(ParentKeyScale, 1.0f);
-
-	TArray<int32> BoneIndexToTrackIndex;
-	BoneIndexToTrackIndex.Empty(NumBones);
-	BoneIndexToTrackIndex.AddUninitialized(NumBones);
-	for (int32 BoneIndex = 0; BoneIndex < NumBones; ++BoneIndex)
-	{
-		BoneIndexToTrackIndex[BoneIndex] = AnimSeq.GetSkeleton()->GetRawAnimationTrackIndex(BoneIndex, &AnimSeq);
-	}
-
-	// generate an array to hold the indices of our end effectors
-	TArray<int32> EndEffectors;
-	for (int32 BoneIndex = 0; BoneIndex < NumBones; ++BoneIndex)
-	{
-		const FBoneData& Bone = BoneData[BoneIndex];
-		if (Bone.IsEndEffector())
-		{
-			EndEffectors.Add(BoneIndex);
-		}
-	}
-
-	const int32 NumTracks = RawSegments[0].RotationData.Num();
-	TArray<FTransform> TrivialTrackValues;
-	TrivialTrackValues.Empty(NumTracks);
-	TrivialTrackValues.AddUninitialized(NumTracks);
-	for (int32 TrackIndex = 0; TrackIndex < NumTracks; ++TrackIndex)
-	{
-		const FQuat& Rotation = RawSegments[0].RotationData[TrackIndex].RotKeys[0];
-		const FVector& Translation = RawSegments[0].TranslationData[TrackIndex].PosKeys[0];
-		FVector Scale = FVector::OneVector;
-		if (RawSegments[0].ScaleData.Num() > 0)
-		{
-			Scale = RawSegments[0].ScaleData[TrackIndex].ScaleKeys[0];
-		}
-
-		TrivialTrackValues[TrackIndex] = FTransform(Rotation, Translation, Scale);
-	}
-
-	if (bUseDecompression || !bUseMultithreading || RawSegments.Num() <= 1)
-	{
-		for (FAnimSegmentContext& Segment : RawSegments)
-		{
-			FProcessAnimationTracksContext Context(AnimSeq, BoneData, EndEffectors, BoneIndexToTrackIndex, TrivialTrackValues, Segment, RawSegments);
-			ProcessAnimationTracks(Context);
-		}
-	}
-	else
-	{
-		// Created the context objects.
-		FAsyncProcessAnimationTracksTaskGroupContext* TaskGroupContext = new FAsyncProcessAnimationTracksTaskGroupContext();
-		for (FAnimSegmentContext& Segment : RawSegments)
-		{
-			FProcessAnimationTracksContext* Context = new FProcessAnimationTracksContext(AnimSeq, BoneData, EndEffectors, BoneIndexToTrackIndex, TrivialTrackValues, Segment, RawSegments);
-			TaskGroupContext->TaskContexes.Add(Context);
-		}
-
-		// Dispatch 1 task per job thread.
-		FGraphEventArray AsyncTaskCompletionEvents;
-		const int32 NumTaskThreads = FTaskGraphInterface::Get().GetNumWorkerThreads();
-		for (int32 TaskIndex = 0; TaskIndex < NumTaskThreads; ++TaskIndex)
-		{
-			AsyncTaskCompletionEvents.Add(TGraphTask<FAsyncProcessAnimationTracksTask>::CreateTask(NULL, ENamedThreads::AnyThread).ConstructAndDispatchWhenReady(this, TaskGroupContext));
-		}
-
-		// Execute the contexts concurrently.
-		TaskGroupContext->ExecuteTasks(this);
-
-		// Wait for all concurrent tasks to be done, we only wait for ones that were executing.
-		TaskGroupContext->WaitForAllTasks();
-
-		// Dispatch a cleanup job that will execute once all the tasks are done.
-		// The current thread/task no longer needs the TaskGroupContext beyond this point nor anything it references.
-		// All tasks that were executing actual compression work are done or queued up and pending
-		// and will do nothing since no work remains (they only touch the TaskGroupContext to find more work).
-		// Once all tasks are actually done, the clean up task will execute, freeing memory.
-		TGraphTask<FAsyncCleanUpProcessAnimationTracksTask>::CreateTask(&AsyncTaskCompletionEvents, ENamedThreads::AnyThread).ConstructAndDispatchWhenReady(TaskGroupContext);
-	}
-}
-#endif
-
-void UAnimCompress_RemoveLinearKeys::CompressUsingUnderlyingCompressor(
-	const FCompressibleAnimData& CompressibleAnimData,
-	FCompressibleAnimDataResult& OutCompressedData,
-	const TArray<FTranslationTrack>& TranslationData,
-	const TArray<FRotationTrack>& RotationData,
-	const TArray<FScaleTrack>& ScaleData,
-	const bool bFinalPass)
-{
-	BitwiseCompressAnimationTracks(
-		CompressibleAnimData,
-		OutCompressedData,
-		static_cast<AnimationCompressionFormat>(TranslationCompressionFormat),
-		static_cast<AnimationCompressionFormat>(RotationCompressionFormat),
-		static_cast<AnimationCompressionFormat>(ScaleCompressionFormat),
-		TranslationData,
-		RotationData,
-		ScaleData,
-		true);
-
-	// record the proper runtime decompressor to use
-	OutCompressedData.KeyEncodingFormat = AKF_VariableKeyLerp;
-	AnimationFormat_SetInterfaceLinks(OutCompressedData);
-}
-
-#if USE_SEGMENTING_CONTEXT
-void UAnimCompress_RemoveLinearKeys::CompressUsingUnderlyingCompressor(
-	UAnimSequence& AnimSeq,
-	const TArray<FBoneData>& BoneData,
-	TArray<FAnimSegmentContext>& RawSegments,
-	const bool bFinalPass)
-{
-	// record the proper runtime decompressor to use
-	AnimSeq.KeyEncodingFormat = AKF_VariableKeyLerp;
-	AnimSeq.TranslationCompressionFormat = TranslationCompressionFormat;
-	AnimSeq.RotationCompressionFormat = RotationCompressionFormat;
-	AnimSeq.ScaleCompressionFormat = ScaleCompressionFormat;
-	AnimationFormat_SetInterfaceLinks(AnimSeq);
-
-	const bool bIsSorted = bOptimizeForForwardPlayback;
-
-	BitwiseCompressAnimationTracks(
-		AnimSeq,
-		static_cast<AnimationCompressionFormat>(TranslationCompressionFormat),
-		static_cast<AnimationCompressionFormat>(RotationCompressionFormat),
-		static_cast<AnimationCompressionFormat>(ScaleCompressionFormat),
-		RawSegments,
-		bIsSorted);
-
-	CoalesceCompressedSegments(AnimSeq, RawSegments, bIsSorted);
-
-	// We could be invalid, set the links again
-	AnimationFormat_SetInterfaceLinks(AnimSeq);
-}
-#endif
-
-#define DO_RAW_DATA_DEBUG_LOGGING 0
-
-#if DO_RAW_DATA_DEBUG_LOGGING
-TArray<FVector>& GetKeysArray(const FTranslationTrack& Track)
-{
-	return Track.PosKeys;
-}
-
-TArray<FQuat>& GetKeysArray(const FRotationTrack& Track)
-{
-	return Track.RotKeys;
-}
-
-TArray<FVector>& GetKeysArray(const FScaleTrack& Track)
-{
-	return Track.ScaleKeys;
-}
-
-template<typename ArrayType>
-void DebugLogTrack(const TArray<ArrayType>& Data, int32 Track, const TCHAR* TrackName)
-{
-	if (Data.IsValidIndex(Track))
-	{
-		const ArrayType& Item = Data[Track];
-		FPlatformMisc::LowLevelOutputDebugStringf(TEXT("%s : %i\n"), TrackName, Track);
-		DebugLogArray(GetKeysArray(Item));
-	}
-}
-
-void DebugLogTrackData(const TArray<FTranslationTrack>& TranslationData, const TArray<FRotationTrack>& RotationData, const TArray<FScaleTrack>& ScaleData, const TCHAR* Section)
-{
-	FPlatformMisc::LowLevelOutputDebugStringf(TEXT("Post Section: %s\n"), Section);
-	const int32 ArrayMax = FMath::Max(FMath::Max(TranslationData.Num(), RotationData.Num()), ScaleData.Num());
-
-	for (int i = 0; i < ArrayMax; ++i)
-	{
-		DebugLogTrack(TranslationData, i, TEXT("Translation"));
-		DebugLogTrack(RotationData, i, TEXT("Rotation"));
-		DebugLogTrack(ScaleData, i, TEXT("Scale"));
-	}
-}
-#endif
-
-void UAnimCompress_RemoveLinearKeys::DoReduction(const FCompressibleAnimData& CompressibleAnimData, FCompressibleAnimDataResult& OutResult)
-{
-#if WITH_EDITORONLY_DATA
-	// Only need to do the heavy lifting if it will have some impact
-	// One of these will always be true for the base class, but derived classes may choose to turn both off (e.g., in PerTrackCompression)
-	const bool bRunningProcessor = bRetarget || bActuallyFilterLinearKeys;
-
-	// If the processor is to be run, then additive animations need to be converted from relative to absolute
-	const bool bNeedToConvertBackToAdditive = bRunningProcessor ? CompressibleAnimData.bIsValidAdditive : false;
+	// If the processor is to be run, then additive animations need to be converted from relative to absolute
+	const bool bNeedToConvertBackToAdditive = bRunningProcessor ? CompressibleAnimData.bIsValidAdditive : false;
 
 	FCompressibleAnimData TempConvertedAnimData;
 
@@ -2326,73 +1053,9 @@ void UAnimCompress_RemoveLinearKeys::DoReduction(const FCompressibleAnimData& Co
 	TArray<FScaleTrack> ScaleData;
 	SeparateRawDataIntoTracks(CompressibleDataToOperateOn.RawAnimationData, CompressibleDataToOperateOn.SequenceLength, TranslationData, RotationData, ScaleData);
 
-	FilterBeforeMainKeyRemoval(CompressibleDataToOperateOn, TranslationData, RotationData, ScaleData);
+	check(OutResult.CompressionUserData == nullptr);
+	OutResult.CompressionUserData = FilterBeforeMainKeyRemoval(CompressibleDataToOperateOn, TranslationData, RotationData, ScaleData);
 
-#if USE_SEGMENTING_CONTEXT
-	if (bEnableSegmenting)
-	{
-#if TIME_LINEAR_KEY_REMOVAL
-		double TimeStart = FPlatformTime::Seconds();
-#endif
-
-		TArray<FAnimSegmentContext> RawSegments;
-		SeparateRawDataIntoTracks(*AnimSeq, TranslationData, RotationData, ScaleData, IdealNumFramesPerSegment, MaxNumFramesPerSegment, RawSegments);
-
-		if (bRunningProcessor)
-		{
-			if (bUseDecompression)
-			{
-				// compress this animation without any key-reduction to prime the codec
-				CompressUsingUnderlyingCompressor(
-					*AnimSeq,
-					BoneData,
-					RawSegments,
-					false);
-			}
-			else
-			{
-				// record the proper runtime decompressor to use
-				AnimSeq->KeyEncodingFormat = AKF_VariableKeyLerp;
-				AnimSeq->TranslationCompressionFormat = TranslationCompressionFormat;
-				AnimSeq->RotationCompressionFormat = RotationCompressionFormat;
-				AnimSeq->ScaleCompressionFormat = ScaleCompressionFormat;
-				AnimationFormat_SetInterfaceLinks(*AnimSeq);
-
-				AnimSeq->CompressedByteStream.Empty();
-				AnimSeq->CompressedScaleOffsets.Empty();
-				AnimSeq->CompressedTrackOffsets.Empty();
-				AnimSeq->CompressedSegments.Empty();
-
-				for (FAnimSegmentContext& Segment : RawSegments)
-				{
-					Segment.TranslationCompressionFormat = TranslationCompressionFormat;
-					Segment.RotationCompressionFormat = RotationCompressionFormat;
-					Segment.ScaleCompressionFormat = ScaleCompressionFormat;
-				}
-			}
-
-			// now remove the keys which can be approximated with linear interpolation
-			ProcessAnimationTracks(
-				*AnimSeq,
-				BoneData,
-				RawSegments);
-
-#if TIME_LINEAR_KEY_REMOVAL
-			double ElapsedTime = FPlatformTime::Seconds() - TimeStart;
-			UE_LOG(LogAnimationCompression, Log, TEXT("ProcessAnimationTracks time is (%f) seconds"), ElapsedTime);
-#endif
-		}
-
-		// compress the final (possibly key-reduced) tracks into the anim sequence buffers
-		CompressUsingUnderlyingCompressor(
-			*AnimSeq,
-			BoneData,
-			RawSegments,
-			true);
-
-		return;
-	}
-#endif
 	if (bRunningProcessor)
 	{
 #if TIME_LINEAR_KEY_REMOVAL
@@ -2436,6 +1099,7 @@ void UAnimCompress_RemoveLinearKeys::DoReduction(const FCompressibleAnimData& Co
 		true);
 
 #endif // WITH_EDITORONLY_DATA
+	return true;
 }
 
 void UAnimCompress_RemoveLinearKeys::PopulateDDCKey(FArchive& Ar)
@@ -2449,9 +1113,8 @@ void UAnimCompress_RemoveLinearKeys::PopulateDDCKey(FArchive& Ar)
 	Ar << EffectorDiffSocket;
 	Ar << ParentKeyScale;
 	uint8 Flags =	MakeBitForFlag(bRetarget, 0) +
-					MakeBitForFlag(bActuallyFilterLinearKeys, 1) +
-					MakeBitForFlag(bOptimizeForForwardPlayback, 2);
+					MakeBitForFlag(bActuallyFilterLinearKeys, 1);
 	Ar << Flags;
 }
 
-#endif // WITH_EIDOT
+#endif // WITH_EDITOR
diff --git a/Engine/Source/Runtime/Engine/Private/Animation/AnimCompress_RemoveTrivialKeys.cpp b/Engine/Source/Runtime/Engine/Private/Animation/AnimCompress_RemoveTrivialKeys.cpp
--- a/Engine/Source/Runtime/Engine/Private/Animation/AnimCompress_RemoveTrivialKeys.cpp
+++ b/Engine/Source/Runtime/Engine/Private/Animation/AnimCompress_RemoveTrivialKeys.cpp
@@ -18,7 +18,7 @@ UAnimCompress_RemoveTrivialKeys::UAnimCompress_RemoveTrivialKeys(const FObjectIn
 }
 
 #if WITH_EDITOR
-void UAnimCompress_RemoveTrivialKeys::DoReduction(const FCompressibleAnimData& CompressibleAnimData, FCompressibleAnimDataResult& OutResult)
+bool UAnimCompress_RemoveTrivialKeys::DoReduction(const FCompressibleAnimData& CompressibleAnimData, FCompressibleAnimDataResult& OutResult)
 {
 #if WITH_EDITORONLY_DATA
 	// split the filtered data into tracks
@@ -31,46 +31,26 @@ void UAnimCompress_RemoveTrivialKeys::DoReduction(const FCompressibleAnimData& C
 	FilterTrivialKeys(TranslationData, RotationData, ScaleData, TRANSLATION_ZEROING_THRESHOLD, QUATERNION_ZEROING_THRESHOLD, SCALE_ZEROING_THRESHOLD);
 
 	// record the proper runtime decompressor to use
-	OutResult.KeyEncodingFormat = AKF_ConstantKeyLerp;
-	AnimationFormat_SetInterfaceLinks(OutResult);
+	FUECompressedAnimDataMutable& AnimData = static_cast<FUECompressedAnimDataMutable&>(*OutResult.AnimData);
+	AnimData.KeyEncodingFormat = AKF_ConstantKeyLerp;
+	AnimationFormat_SetInterfaceLinks(AnimData);
 
-#if USE_SEGMENTING_CONTEXT
-	if (bEnableSegmenting)
-	{
-		TArray<FAnimSegmentContext> RawSegments;
-		SeparateRawDataIntoTracks(*AnimSeq, TranslationData, RotationData, ScaleData, IdealNumFramesPerSegment, MaxNumFramesPerSegment, RawSegments);
-
-		BitwiseCompressAnimationTracks(
-			*AnimSeq,
-			static_cast<AnimationCompressionFormat>(TranslationCompressionFormat),
-			static_cast<AnimationCompressionFormat>(RotationCompressionFormat),
-			static_cast<AnimationCompressionFormat>(ScaleCompressionFormat),
-			RawSegments);
-
-		CoalesceCompressedSegments(*AnimSeq, RawSegments);
-
-		AnimSeq->TranslationCompressionFormat = TranslationCompressionFormat;
-		AnimSeq->RotationCompressionFormat = RotationCompressionFormat;
-		AnimSeq->ScaleCompressionFormat = ScaleCompressionFormat;
-	}
-	else
-#endif
-	{
-		// bitwise compress the tracks into the anim sequence buffers
-		BitwiseCompressAnimationTracks(
-			CompressibleAnimData,
-			OutResult,
-			static_cast<AnimationCompressionFormat>(TranslationCompressionFormat),
-			static_cast<AnimationCompressionFormat>(RotationCompressionFormat),
-			static_cast<AnimationCompressionFormat>(ScaleCompressionFormat),
-			TranslationData,
-			RotationData,
-			ScaleData);
-	}
+	// bitwise compress the tracks into the anim sequence buffers
+	BitwiseCompressAnimationTracks(
+		CompressibleAnimData,
+		OutResult,
+		static_cast<AnimationCompressionFormat>(TranslationCompressionFormat),
+		static_cast<AnimationCompressionFormat>(RotationCompressionFormat),
+		static_cast<AnimationCompressionFormat>(ScaleCompressionFormat),
+		TranslationData,
+		RotationData,
+		ScaleData);
 
 	// We could be invalid, set the links again
-	AnimationFormat_SetInterfaceLinks(OutResult);
+	AnimationFormat_SetInterfaceLinks(AnimData);
 #endif // WITH_EDITORONLY_DATA
+
+	return true;
 }
 
 void UAnimCompress_RemoveTrivialKeys::PopulateDDCKey(FArchive& Ar)
diff --git a/Engine/Source/Runtime/Engine/Private/Animation/AnimCompressionDerivedData.cpp b/Engine/Source/Runtime/Engine/Private/Animation/AnimCompressionDerivedData.cpp
--- a/Engine/Source/Runtime/Engine/Private/Animation/AnimCompressionDerivedData.cpp
+++ b/Engine/Source/Runtime/Engine/Private/Animation/AnimCompressionDerivedData.cpp
@@ -1,11 +1,15 @@
 // Copyright Epic Games, Inc. All Rights Reserved.
 #include "Animation/AnimCompressionDerivedData.h"
+#include "Animation/AnimCompressionDerivedDataPublic.h"
+#include "CoreMinimal.h"
+#include "DerivedDataCacheInterface.h"
 #include "Stats/Stats.h"
 #include "Animation/AnimSequence.h"
 #include "Serialization/MemoryWriter.h"
 #include "AnimationUtils.h"
 #include "AnimEncoding.h"
 #include "Animation/AnimCompress.h"
+#include "Animation/AnimBoneCompressionSettings.h"
 #include "Animation/AnimCurveCompressionSettings.h"
 #include "AnimationCompression.h"
 #include "UObject/Package.h"
@@ -14,11 +18,10 @@
 
 DECLARE_CYCLE_STAT(TEXT("Anim Compression (Derived Data)"), STAT_AnimCompressionDerivedData, STATGROUP_Anim);
 
-FDerivedDataAnimationCompression::FDerivedDataAnimationCompression(const TCHAR* InTypeName, const FString& InAssetDDCKey, TSharedPtr<FAnimCompressContext> InCompressContext, int32 InPreviousCompressedSize)
+FDerivedDataAnimationCompression::FDerivedDataAnimationCompression(const TCHAR* InTypeName, const FString& InAssetDDCKey, TSharedPtr<FAnimCompressContext> InCompressContext)
 	: TypeName(InTypeName)
 	, AssetDDCKey(InAssetDDCKey)
 	, CompressContext(InCompressContext)
-	, PreviousCompressedSize(InPreviousCompressedSize)
 {
 
 }
@@ -32,11 +35,13 @@ const TCHAR* FDerivedDataAnimationCompression::GetVersionString() const
 	// This is a version string that mimics the old versioning scheme. If you
 	// want to bump this version, generate a new guid using VS->Tools->Create GUID and
 	// return it here. Ex.
-	return TEXT("2E79BF10172A48FDACA76883B8951538");
+	return TEXT("A20FA6C92A4443CDAE897F0B63C32814");
 }
 
 bool FDerivedDataAnimationCompression::Build( TArray<uint8>& OutDataArray )
 {
+	const double CompressionStartTime = FPlatformTime::Seconds();
+
 	check(DataToCompressPtr.IsValid());
 	FCompressibleAnimData& DataToCompress = *DataToCompressPtr.Get();
 	FCompressedAnimSequence OutData;
@@ -50,20 +55,16 @@ bool FDerivedDataAnimationCompression::Build( TArray<uint8>& OutDataArray )
 	{
 		DataToCompress.Update(OutData);
 
-		bool bCurveCompressionSuccess = FAnimationUtils::CompressAnimCurves(DataToCompress, OutData);
+		const bool bBoneCompressionOk = FAnimationUtils::CompressAnimBones(DataToCompress, CompressionResult);
+		const bool bCurveCompressionOk = FAnimationUtils::CompressAnimCurves(DataToCompress, OutData);
 
 #if DO_CHECK
-		FString CompressionName = DataToCompress.RequestedCompressionScheme->GetFullName();
+		FString CompressionName = DataToCompress.BoneCompressionSettings->GetFullName();
 		const TCHAR* AAC = CompressContext.Get()->bAllowAlternateCompressor ? TEXT("true") : TEXT("false");
 		const TCHAR* OutputStr = CompressContext.Get()->bOutput ? TEXT("true") : TEXT("false");
 #endif
 
-		CompressionResult.CompressedNumberOfFrames = DataToCompress.NumFrames; //Do this before compression so compress code can read the correct value
-
-		CompressContext->GatherPreCompressionStats(DataToCompress, PreviousCompressedSize);
-
-		FAnimationUtils::CompressAnimSequence(DataToCompress, CompressionResult, *CompressContext.Get());
-		bCompressionSuccessful = (CompressionResult.IsCompressedDataValid() || DataToCompress.RawAnimationData.Num() == 0) && bCurveCompressionSuccess;
+		bCompressionSuccessful = bBoneCompressionOk && bCurveCompressionOk;
 
 		ensureMsgf(bCompressionSuccessful, TEXT("Anim Compression failed for Sequence '%s' with compression scheme '%s': compressed data empty\n\tAnimIndex: %i\n\tMaxAnim:%i\n\tAllowAltCompressor:%s\n\tOutput:%s"), 
 											*DataToCompress.FullName,
@@ -72,23 +73,189 @@ bool FDerivedDataAnimationCompression::Build( TArray<uint8>& OutDataArray )
 											CompressContext.Get()->MaxAnimations,
 											AAC,
 											OutputStr);
+	}
+
+	if (bCompressionSuccessful)
+	{
+		const double CompressionEndTime = FPlatformTime::Seconds();
+		const double CompressionTime = CompressionEndTime - CompressionStartTime;
+
+		CompressContext->GatherPostCompressionStats(OutData, DataToCompress.BoneData, DataToCompress.AnimFName, CompressionTime, true);
+
+		OutData.CompressedByteStream = MoveTemp(CompressionResult.CompressedByteStream);
+		OutData.CompressedDataStructure = MoveTemp(CompressionResult.AnimData);
+		OutData.BoneCompressionCodec = CompressionResult.Codec;
+
+		FMemoryWriter Ar(OutDataArray, true);
+		OutData.SerializeCompressedData(Ar, true, nullptr, DataToCompress.Skeleton, DataToCompress.BoneCompressionSettings, DataToCompress.CurveCompressionSettings); //Save out compressed
+	}
+
+	return bCompressionSuccessful;
+}
+
+const uint64 GigaBytes = 1024 * 1024 * 1024;
+const uint64 MAX_ASYNC_COMPRESSION_MEM_USAGE = 2 * GigaBytes;
+const int32 MAX_ACTIVE_COMPRESSIONS = 2;
+
+FAsyncCompressedAnimationsManagement* GAsyncCompressedAnimationsTracker = nullptr;
+
+FAsyncCompressedAnimationsManagement& FAsyncCompressedAnimationsManagement::Get()
+{
+	static FAsyncCompressedAnimationsManagement SingletonInstance;
+	GAsyncCompressedAnimationsTracker = &SingletonInstance;
+	return SingletonInstance;
+}
+
+void FAsyncCompressedAnimationsManagement::OnActiveCompressionFinished(int32 ActiveAnimIndex)
+{
+	FDerivedDataCacheInterface& DerivedDataCache = GetDerivedDataCacheRef();
+
+	FActiveAsyncCompressionTask& Task = ActiveAsyncCompressionTasks[ActiveAnimIndex];
+
+	TArray<uint8> OutData;
+	bool bBuiltLocally = false;
+	if (DerivedDataCache.GetAsynchronousResults(Task.AsyncHandle, OutData, &bBuiltLocally))
+	{
+		Task.Sequence->ApplyCompressedData(Task.CacheKey, Task.bPerformFrameStripping, OutData);
+	}
+	else
+	{
+		UE_LOG(LogAnimationCompression, Fatal, TEXT("Failed to get async compressed animation data for anim '%s'"), *Task.Sequence->GetName());
+		Task.Sequence->ApplyCompressedData(FString(), false, OutData); // Clear active compression on Sequence
+	}
+	ActiveMemoryUsage -= Task.TaskSize;
+	ActiveAsyncCompressionTasks.RemoveAtSwap(ActiveAnimIndex, 1, false);
+}
+
+void FAsyncCompressedAnimationsManagement::Tick(float DeltaTime)
+{
+	const double MaxProcessingTime = 0.1; // try not to hang the editor too much
+	const double EndTime = FPlatformTime::Seconds() + MaxProcessingTime;
+	const double StartTime = FPlatformTime::Seconds();
+
+	FDerivedDataCacheInterface& DerivedDataCache = GetDerivedDataCacheRef();
 
-		if (CompressionResult.IsCompressedDataValid())
+	for (int32 ActiveAnim = ActiveAsyncCompressionTasks.Num() - 1; ActiveAnim >= 0; --ActiveAnim)
+	{
+		const FActiveAsyncCompressionTask& Task = ActiveAsyncCompressionTasks[ActiveAnim];
+
+		if (DerivedDataCache.PollAsynchronousCompletion(Task.AsyncHandle))
 		{
-			CompressionResult.BuildFinalBuffer(OutData.CompressedByteStream); // Build final compressed data buffer
+			OnActiveCompressionFinished(ActiveAnim);
+		}
 
-			OutData.CompressedDataStructure.CopyFrom(CompressionResult); //Copy header info
-			OutData.CompressedDataStructure.InitViewsFromBuffer(OutData.CompressedByteStream); //Init views to CompressedByteStream
-			
+		if (FPlatformTime::Seconds() > EndTime)
+		{
+			return; // Finish for this tick
 		}
 	}
 
-	if (bCompressionSuccessful)
+	const bool bHasQueuedTasks = QueuedAsyncCompressionWork.Num() > 0;
+
+	while (ActiveAsyncCompressionTasks.Num() < MAX_ACTIVE_COMPRESSIONS)
 	{
-		FMemoryWriter Ar(OutDataArray, true);
-		OutData.SerializeCompressedData(Ar, true, nullptr, DataToCompress.Skeleton, DataToCompress.CurveCompressionSettings); //Save out compressed
+		if (QueuedAsyncCompressionWork.Num() == 0)
+		{
+			break;
+		}
+		FQueuedAsyncCompressionWork NewTask = QueuedAsyncCompressionWork.Pop(false);
+		StartAsyncWork(NewTask.Compressor, NewTask.Anim, NewTask.Compressor.GetMemoryUsage(), NewTask.bPerformFrameStripping);
 	}
 
-	return bCompressionSuccessful;
+	if (bHasQueuedTasks && QueuedAsyncCompressionWork.Num() == 0)
+	{
+		QueuedAsyncCompressionWork.Empty(); //free memory
+	}
+}
+
+TStatId FAsyncCompressedAnimationsManagement::GetStatId() const
+{
+	RETURN_QUICK_DECLARE_CYCLE_STAT(FAsyncCompressedAnimationsTracker, STATGROUP_Tickables);
 }
+
+bool FAsyncCompressedAnimationsManagement::RequestAsyncCompression(FDerivedDataAnimationCompression& Compressor, UAnimSequence* Anim, const bool bPerformFrameStripping, TArray<uint8>& OutData)
+{
+	const uint64 NewTaskSize = Compressor.GetMemoryUsage();
+
+	bool bWasAsync = true;
+
+	if (ActiveMemoryUsage + NewTaskSize >= MAX_ASYNC_COMPRESSION_MEM_USAGE)
+	{
+		Tick(0.f); // Try to free up some memory
+	}
+
+	const bool bCanRunASync = (ActiveMemoryUsage + NewTaskSize < MAX_ASYNC_COMPRESSION_MEM_USAGE);
+	const bool bForceSync = false; //debugging override
+
+	if (bCanRunASync && !bForceSync)
+	{
+		//Queue Async
+		ActiveMemoryUsage += NewTaskSize;
+
+		if(ActiveAsyncCompressionTasks.Num() < MAX_ACTIVE_COMPRESSIONS)
+		{
+			StartAsyncWork(Compressor, Anim, NewTaskSize, bPerformFrameStripping);
+		}
+		else
+		{
+			QueuedAsyncCompressionWork.Emplace(Compressor, Anim, bPerformFrameStripping);
+		}
+	}
+	else
+	{
+		//Do in place
+		GetDerivedDataCacheRef().GetSynchronous(&Compressor, OutData);
+		bWasAsync = false;
+	}
+	return bWasAsync;
+}
+
+void FAsyncCompressedAnimationsManagement::StartAsyncWork(FDerivedDataAnimationCompression& Compressor, UAnimSequence* Anim, const uint64 NewTaskSize, const bool bPerformFrameStripping)
+{
+	const FString CacheKey = Compressor.GetPluginSpecificCacheKeySuffix();
+	uint32 AsyncHandle = GetDerivedDataCacheRef().GetAsynchronous(&Compressor);
+	ActiveAsyncCompressionTasks.Emplace(Anim, CacheKey, NewTaskSize, AsyncHandle, bPerformFrameStripping);
+}
+
+bool FAsyncCompressedAnimationsManagement::WaitOnActiveCompression(UAnimSequence* Anim)
+{
+	for (int32 ActiveIndex = 0; ActiveIndex < ActiveAsyncCompressionTasks.Num(); ++ActiveIndex)
+	{
+		FActiveAsyncCompressionTask& Task = ActiveAsyncCompressionTasks[ActiveIndex];
+		if (Task.Sequence == Anim)
+		{
+			GetDerivedDataCacheRef().WaitAsynchronousCompletion(Task.AsyncHandle);
+			OnActiveCompressionFinished(ActiveIndex);
+			return true; // Done
+		}
+	}
+	return false;
+}
+
+bool FAsyncCompressedAnimationsManagement::WaitOnExistingCompression(UAnimSequence* Anim, const bool bCancelIfNotStarted)
+{
+	if(!WaitOnActiveCompression(Anim))
+	{
+		//Check if we have a queued task
+		for (int32 QueuedIndex = 0; QueuedIndex < QueuedAsyncCompressionWork.Num(); ++QueuedIndex)
+		{
+			FQueuedAsyncCompressionWork& Task = QueuedAsyncCompressionWork[QueuedIndex];
+			if (Task.Anim == Anim)
+			{
+				if (!bCancelIfNotStarted)
+				{
+					StartAsyncWork(Task.Compressor, Task.Anim, Task.Compressor.GetMemoryUsage(), Task.bPerformFrameStripping);
+				}
+				else
+				{
+					delete &Task.Compressor;
+				}
+				QueuedAsyncCompressionWork.RemoveAtSwap(QueuedIndex, 1, false);
+			}
+		}
+		return bCancelIfNotStarted ? false : WaitOnActiveCompression(Anim);
+	}
+	return true;
+}
+
 #endif	//WITH_EDITOR
diff --git a/Engine/Source/Runtime/Engine/Private/Animation/AnimCompressionDerivedData.h b/Engine/Source/Runtime/Engine/Private/Animation/AnimCompressionDerivedData.h
--- a/Engine/Source/Runtime/Engine/Private/Animation/AnimCompressionDerivedData.h
+++ b/Engine/Source/Runtime/Engine/Private/Animation/AnimCompressionDerivedData.h
@@ -31,17 +31,8 @@ private:
 	// FAnimCompressContext to use during compression if we don't pull from the DDC
 	TSharedPtr<FAnimCompressContext> CompressContext;
 
-	// Size of the previous compressed data (for stat tracking)
-	int32 PreviousCompressedSize;
-
-	// Whether we should frame strip (remove every other frame from even frames animations)
-	bool bPerformStripping;
-
-	// Track if it is an even framed animation (when stripping odd framed animations will need to be resampled)
-	bool bIsEvenFramed;
-
 public:
-	FDerivedDataAnimationCompression(const TCHAR* InTypeName, const FString& InAssetDDCKey, TSharedPtr<FAnimCompressContext> InCompressContext, int32 InPreviousCompressedSize);
+	FDerivedDataAnimationCompression(const TCHAR* InTypeName, const FString& InAssetDDCKey, TSharedPtr<FAnimCompressContext> InCompressContext);
 	virtual ~FDerivedDataAnimationCompression();
 
 	void SetCompressibleData(TSharedRef<FCompressibleAnimData> InCompressibleAnimData)
@@ -50,6 +41,11 @@ public:
 		check(DataToCompressPtr->Skeleton != nullptr);
 	}
 
+	uint64 GetMemoryUsage() const
+	{
+		return DataToCompressPtr.IsValid() ? DataToCompressPtr->GetApproxRawSize() : 0;
+	}
+
 	virtual const TCHAR* GetPluginName() const override
 	{
 		return TypeName;
@@ -65,7 +61,7 @@ public:
 
 	virtual bool IsBuildThreadsafe() const override
 	{
-		return false;
+		return true;
 	}
 
 	virtual bool Build( TArray<uint8>& OutDataArray) override;
diff --git a/Engine/Source/Runtime/Engine/Private/Animation/AnimCompressionTypes.cpp b/Engine/Source/Runtime/Engine/Private/Animation/AnimCompressionTypes.cpp
--- a/Engine/Source/Runtime/Engine/Private/Animation/AnimCompressionTypes.cpp
+++ b/Engine/Source/Runtime/Engine/Private/Animation/AnimCompressionTypes.cpp
@@ -5,6 +5,8 @@
 #include "AnimEncoding.h"
 #include "Misc/SecureHash.h"
 #include "Interfaces/ITargetPlatform.h"
+#include "Animation/AnimBoneCompressionCodec.h"
+#include "Animation/AnimBoneCompressionSettings.h"
 #include "Animation/AnimCurveCompressionCodec.h"
 #include "Animation/AnimCurveCompressionSettings.h"
 #include "AnimationRuntime.h"
@@ -123,18 +125,20 @@ void StripFramesOdd(TArray<ArrayValue>& Keys, const int32 NumFrames)
 	}
 }
 
-FCompressibleAnimData::FCompressibleAnimData(class UAnimSequence* InSeq, const bool bPerformStripping, const float InAltCompressionErrorThreshold) :
-#if WITH_EDITOR
-	RequestedCompressionScheme(InSeq->CompressionScheme) ,
-#endif
-	CurveCompressionSettings(InSeq->CurveCompressionSettings)
+FCompressibleAnimData::FCompressibleAnimData(class UAnimSequence* InSeq, const bool bPerformStripping)
+	: CurveCompressionSettings(InSeq->CurveCompressionSettings)
+	, BoneCompressionSettings(InSeq->BoneCompressionSettings)
 	, Skeleton(InSeq->GetSkeleton())
 	, TrackToSkeletonMapTable(InSeq->GetRawTrackToSkeletonMapTable())
 	, Interpolation(InSeq->Interpolation)
 	, SequenceLength(InSeq->SequenceLength)
 	, NumFrames(InSeq->GetRawNumberOfFrames())
 	, bIsValidAdditive(InSeq->IsValidAdditive())
-	, AltCompressionErrorThreshold(InAltCompressionErrorThreshold)
+#if WITH_EDITORONLY_DATA
+	, ErrorThresholdScale(InSeq->CompressionErrorThresholdScale)
+#else
+	, ErrorThresholdScale(1.f)
+#endif
 	, Name(InSeq->GetName())
 	, FullName(InSeq->GetFullName())
 	, AnimFName(InSeq->GetFName())
@@ -198,33 +202,31 @@ FCompressibleAnimData::FCompressibleAnimData(class UAnimSequence* InSeq, const b
 #endif
 }
 
-FCompressibleAnimData::FCompressibleAnimData(UAnimCompress* InRequestedCompressionScheme, UAnimCurveCompressionSettings* InCurveCompressionSettings, USkeleton* InSkeleton, EAnimInterpolationType InInterpolation, float InSequenceLength, int32 InNumFrames, const float InAltCompressionErrorThreshold) :
-#if WITH_EDITOR
-	RequestedCompressionScheme(InRequestedCompressionScheme) ,
-#endif
-	CurveCompressionSettings(InCurveCompressionSettings)
+FCompressibleAnimData::FCompressibleAnimData(UAnimBoneCompressionSettings* InBoneCompressionSettings, UAnimCurveCompressionSettings* InCurveCompressionSettings, USkeleton* InSkeleton, EAnimInterpolationType InInterpolation, float InSequenceLength, int32 InNumFrames)
+	: CurveCompressionSettings(InCurveCompressionSettings)
+	, BoneCompressionSettings(InBoneCompressionSettings)
 	, Skeleton(InSkeleton)
 	, Interpolation(InInterpolation)
 	, SequenceLength(InSequenceLength)
 	, NumFrames(InNumFrames)
 	, bIsValidAdditive(false)
-	, AltCompressionErrorThreshold(InAltCompressionErrorThreshold)
+	, ErrorThresholdScale(1.f)
 {
 #if WITH_EDITOR
 	FAnimationUtils::BuildSkeletonMetaData(Skeleton, BoneData);
 #endif
 }
 
-FCompressibleAnimData::FCompressibleAnimData() : RequestedCompressionScheme(nullptr)
-, CurveCompressionSettings(nullptr)
+FCompressibleAnimData::FCompressibleAnimData()
+: CurveCompressionSettings(nullptr)
+, BoneCompressionSettings(nullptr)
 , Skeleton(nullptr)
 , Interpolation((EAnimInterpolationType)0)
 , SequenceLength(0.f)
 , NumFrames(0)
 , bIsValidAdditive(false)
-, AltCompressionErrorThreshold(FAnimationUtils::GetAlternativeCompressionThreshold())
+, ErrorThresholdScale(1.f)
 {
-
 }
 
 void FCompressibleAnimData::Update(FCompressedAnimSequence& InOutCompressedData) const
@@ -249,16 +251,6 @@ void WriteArray(FMemoryWriter& MemoryWriter, TArray<T>& Array)
 	MemoryWriter.Serialize(Array.GetData(), NumBytes);
 }
 
-void FCompressibleAnimDataResult::BuildFinalBuffer(TArray<uint8>& OutBuffer)
-{
-	OutBuffer.Reset();
-	FMemoryWriter MemoryWriter(OutBuffer);
-
-	WriteArray(MemoryWriter, CompressedTrackOffsets);
-	WriteArray(MemoryWriter, CompressedScaleOffsets.OffsetData);
-	WriteArray(MemoryWriter, CompressedByteStream);
-}
-
 template<typename T>
 void InitArrayView(TArrayView<T>& View, uint8*& DataPtr)
 {
@@ -266,22 +258,6 @@ void InitArrayView(TArrayView<T>& View, uint8*& DataPtr)
 	DataPtr += (View.Num() * View.GetTypeSize());
 }
 
-template<typename T>
-void ResetArrayView(TArrayView<T>& ArrayView)
-{
-	ArrayView = TArrayView<T>();
-}
-
-void FUECompressedAnimData::Reset()
-{
-	ResetArrayView(CompressedTrackOffsets);
-	ResetArrayView(CompressedScaleOffsets.OffsetData);
-	ResetArrayView(CompressedByteStream);
-
-	TranslationCompressionFormat = RotationCompressionFormat = ScaleCompressionFormat = ACF_None;
-	TranslationCodec = RotationCodec = ScaleCodec = nullptr;
-}
-
 void FUECompressedAnimData::InitViewsFromBuffer(const TArrayView<uint8> BulkData)
 {
 	check(BulkData.Num() > 0);
@@ -301,19 +277,6 @@ void InitArrayViewSize(TArrayView<T>& Dest, const TArray<T>& Src)
 	Dest = TArrayView<T>((T*)nullptr, Src.Num());
 }
 
-#if WITH_EDITOR
-void FUECompressedAnimData::CopyFrom(const FCompressibleAnimDataResult& Other)
-{
-	InitArrayViewSize(CompressedTrackOffsets, Other.CompressedTrackOffsets);
-	InitArrayViewSize(CompressedScaleOffsets.OffsetData, Other.CompressedScaleOffsets.OffsetData);
-	InitArrayViewSize(CompressedByteStream, Other.CompressedByteStream);
-
-	CompressedScaleOffsets.StripSize = Other.CompressedScaleOffsets.StripSize;
-
-	CopyFromSettings(Other);
-}
-#endif
-
 template<typename T>
 void SerializeView(class FArchive& Ar, TArrayView<T>& View)
 {
@@ -356,23 +319,31 @@ FArchive& operator<<(FArchive& Ar, AnimationKeyFormat& Fmt)
 	return Ar;
 }
 
-void FUECompressedAnimData::SerializeCompressedData(class FArchive& Ar)
+void FUECompressedAnimData::SerializeCompressedData(FArchive& Ar)
 {
+	ICompressedAnimData::SerializeCompressedData(Ar);
+
 	Ar << KeyEncodingFormat;
 	Ar << TranslationCompressionFormat;
 	Ar << RotationCompressionFormat;
 	Ar << ScaleCompressionFormat;
 
-	Ar << CompressedNumberOfFrames;
-
+	SerializeView(Ar, CompressedByteStream);
 	SerializeView(Ar, CompressedTrackOffsets);
 	SerializeView(Ar, CompressedScaleOffsets.OffsetData);
 	Ar << CompressedScaleOffsets.StripSize;
-	SerializeView(Ar, CompressedByteStream);
 
 	AnimationFormat_SetInterfaceLinks(*this);
 }
 
+FString FUECompressedAnimData::GetDebugString() const
+{
+	FString TranslationFormat = FAnimationUtils::GetAnimationCompressionFormatString(TranslationCompressionFormat);
+	FString RotationFormat = FAnimationUtils::GetAnimationCompressionFormatString(RotationCompressionFormat);
+	FString ScaleFormat = FAnimationUtils::GetAnimationCompressionFormatString(ScaleCompressionFormat);
+	return FString::Printf(TEXT("[%s, %s, %s]"), *TranslationFormat, *RotationFormat, *ScaleFormat);
+}
+
 template<typename TArchive, typename T>
 void ByteSwapArray(TArchive& MemoryStream, uint8*& StartOfArray, TArrayView<T>& ArrayView)
 {
@@ -429,14 +400,35 @@ void ValidateUObjectLoaded(UObject* Obj, UObject* Source)
 	checkf(!Obj->HasAnyFlags(RF_NeedLoad), TEXT("Failed to load %s in %s"), *Obj->GetFullName(), *Source->GetFullName()); // in non editor should have been preloaded by GetPreloadDependencies
 }
 
-void FCompressedAnimSequence::SerializeCompressedData(FArchive& Ar, bool bDDCData, UObject* DataOwner, USkeleton* Skeleton, UAnimCurveCompressionSettings* CurveCompressionSettings, bool bCanUseBulkData)
+void FUECompressedAnimDataMutable::BuildFinalBuffer(TArray<uint8>& OutCompressedByteStream)
+{
+	OutCompressedByteStream.Reset();
+
+	FMemoryWriter MemoryWriter(OutCompressedByteStream);
+
+	WriteArray(MemoryWriter, CompressedTrackOffsets);
+	WriteArray(MemoryWriter, CompressedScaleOffsets.OffsetData);
+	WriteArray(MemoryWriter, CompressedByteStream);
+}
+
+void ICompressedAnimData::SerializeCompressedData(class FArchive& Ar)
+{
+	Ar << CompressedNumberOfFrames;
+
+#if WITH_EDITORONLY_DATA
+	if (!Ar.IsFilterEditorOnly())
+	{
+		Ar << BoneCompressionErrorStats;
+	}
+#endif
+}
+
+void FCompressedAnimSequence::SerializeCompressedData(FArchive& Ar, bool bDDCData, UObject* DataOwner, USkeleton* Skeleton, UAnimBoneCompressionSettings* BoneCompressionSettings, UAnimCurveCompressionSettings* CurveCompressionSettings, bool bCanUseBulkData)
 {
 	Ar << CompressedRawDataSize;
 	Ar << CompressedTrackToSkeletonMapTable;
 	Ar << CompressedCurveNames;
 
-	CompressedDataStructure.SerializeCompressedData(Ar);
-
 	// Serialize the compressed byte stream from the archive to the buffer.
 	int32 NumBytes = CompressedByteStream.Num();
 	Ar << NumBytes;
@@ -448,6 +440,8 @@ void FCompressedAnimSequence::SerializeCompressedData(FArchive& Ar, bool bDDCDat
 		{
 			Ar << bUseBulkDataForLoad;
 		}
+
+		TArray<uint8> SerializedData;
 		if (bUseBulkDataForLoad)
 		{
 #if !WITH_EDITOR
@@ -477,43 +471,29 @@ void FCompressedAnimSequence::SerializeCompressedData(FArchive& Ar, bool bDDCDat
 			CompressedByteStream.AcceptOwnedBulkDataPtr(OwnedPtr, Size);
 #endif
 			delete OwnedPtr;
-
-			CompressedDataStructure.InitViewsFromBuffer(CompressedByteStream);
 		}
 		else
 		{
 			CompressedByteStream.Empty(NumBytes);
 			CompressedByteStream.AddUninitialized(NumBytes);
 
-			if (CompressedByteStream.Num() > 0)
-			{
-				CompressedDataStructure.InitViewsFromBuffer(CompressedByteStream);
-			}
-
 			if (FPlatformProperties::RequiresCookedData())
 			{
 				Ar.Serialize(CompressedByteStream.GetData(), NumBytes);
 			}
 			else
 			{
-				TArray<uint8> SerializedData;
 				SerializedData.Empty(NumBytes);
 				SerializedData.AddUninitialized(NumBytes);
 				Ar.Serialize(SerializedData.GetData(), NumBytes);
-
-				// Swap the buffer into the byte stream.
-				FMemoryReader MemoryReader(SerializedData, true);
-				MemoryReader.SetByteSwapping(Ar.ForceByteSwapping());
-
-				CompressedDataStructure.ByteSwapIn(CompressedByteStream, MemoryReader);
 			}
 		}
 
+		FString BoneCodecDDCHandle;
 		FString CurveCodecPath;
-		Ar << CurveCodecPath;
 
-		ValidateUObjectLoaded(CurveCompressionSettings, DataOwner);
-		CurveCompressionCodec = CurveCompressionSettings->GetCodec(CurveCodecPath);
+		Ar << BoneCodecDDCHandle;
+		Ar << CurveCodecPath;
 
 		int32 NumCurveBytes;
 		Ar << NumCurveBytes;
@@ -521,6 +501,29 @@ void FCompressedAnimSequence::SerializeCompressedData(FArchive& Ar, bool bDDCDat
 		CompressedCurveByteStream.Empty(NumCurveBytes);
 		CompressedCurveByteStream.AddUninitialized(NumCurveBytes);
 		Ar.Serialize(CompressedCurveByteStream.GetData(), NumCurveBytes);
+
+		// Lookup our codecs in our settings assets
+		ValidateUObjectLoaded(BoneCompressionSettings, DataOwner);
+		ValidateUObjectLoaded(CurveCompressionSettings, DataOwner);
+		BoneCompressionCodec = BoneCompressionSettings->GetCodec(BoneCodecDDCHandle);
+		CurveCompressionCodec = CurveCompressionSettings->GetCodec(CurveCodecPath);
+
+		if (BoneCompressionCodec != nullptr)
+		{
+			CompressedDataStructure = BoneCompressionCodec->AllocateAnimData();
+			CompressedDataStructure->SerializeCompressedData(Ar);
+			CompressedDataStructure->Bind(CompressedByteStream);
+
+			// The codec can be null if we are a default object, a sequence with no raw bone data (just curves),
+			// or if we are duplicating the sequence during compression (new settings are assigned)
+			if (SerializedData.Num() != 0)
+			{
+				// Swap the buffer into the byte stream.
+				FMemoryReader MemoryReader(SerializedData, true);
+				MemoryReader.SetByteSwapping(Ar.ForceByteSwapping());
+				BoneCompressionCodec->ByteSwapIn(*CompressedDataStructure, CompressedByteStream, MemoryReader);
+			}
+		}
 	}
 	else if (Ar.IsSaving() || Ar.IsCountingMemory())
 	{
@@ -529,10 +532,13 @@ void FCompressedAnimSequence::SerializeCompressedData(FArchive& Ar, bool bDDCDat
 
 		const bool bIsCooking = !bDDCData && Ar.IsCooking();
 
-		// and then use the codecs to byte swap
-		FMemoryWriter MemoryWriter(SerializedData, true);
-		MemoryWriter.SetByteSwapping(Ar.ForceByteSwapping());
-		CompressedDataStructure.ByteSwapOut(CompressedByteStream, MemoryWriter);
+		// The codec can be null if we are a default object or a sequence with no raw data, just curves
+		if (BoneCompressionCodec != nullptr)
+		{
+			FMemoryWriter MemoryWriter(SerializedData, true);
+			MemoryWriter.SetByteSwapping(Ar.ForceByteSwapping());
+			BoneCompressionCodec->ByteSwapOut(*CompressedDataStructure, CompressedByteStream, MemoryWriter);
+		}
 
 		// Make sure the entire byte stream was serialized.
 		check(NumBytes == SerializedData.Num());
@@ -608,12 +614,20 @@ void FCompressedAnimSequence::SerializeCompressedData(FArchive& Ar, bool bDDCDat
 			Ar.Serialize(SerializedData.GetData(), SerializedData.Num());
 		}
 
+		FString BoneCodecDDCHandle = BoneCompressionCodec != nullptr ? BoneCompressionCodec->GetCodecDDCHandle() : TEXT("");
+		Ar << BoneCodecDDCHandle;
+
 		FString CurveCodecPath = CurveCompressionCodec->GetPathName();
 		Ar << CurveCodecPath;
 
 		int32 NumCurveBytes = CompressedCurveByteStream.Num();
 		Ar << NumCurveBytes;
 		Ar.Serialize(CompressedCurveByteStream.GetData(), NumCurveBytes);
+
+		if (BoneCompressionCodec != nullptr)
+		{
+			CompressedDataStructure->SerializeCompressedData(Ar);
+		}
 	}
 
 #if WITH_EDITOR
@@ -636,10 +650,23 @@ SIZE_T FCompressedAnimSequence::GetMemorySize() const
 	return	  CompressedTrackToSkeletonMapTable.GetAllocatedSize()
 			+ CompressedCurveNames.GetAllocatedSize()
 			+ CompressedCurveByteStream.GetAllocatedSize()
-			+ CompressedDataStructure.GetApproxBoneCompressedSize()
+			+ CompressedDataStructure->GetApproxCompressedSize()
 			+ sizeof(FCompressedAnimSequence);
 }
 
+void FCompressedAnimSequence::ClearCompressedBoneData()
+{
+	CompressedByteStream.Empty(0);
+	CompressedDataStructure.Reset();
+	BoneCompressionCodec = nullptr;
+}
+
+void FCompressedAnimSequence::ClearCompressedCurveData()
+{
+	CompressedCurveByteStream.Empty(0);
+	CurveCompressionCodec = nullptr;
+}
+
 struct FGetBonePoseScratchArea : public TThreadSingleton<FGetBonePoseScratchArea>
 {
 	BoneTrackArray RotationScalePairs;
@@ -731,7 +758,7 @@ void DecompressPose(FCompactPose& OutPose, const FCompressedAnimSequence& Compre
 		CSV_SCOPED_TIMING_STAT(Animation, ExtractPoseFromAnimData);
 		CSV_CUSTOM_STAT(Animation, NumberOfExtractedAnimations, 1, ECsvCustomStatOp::Accumulate);
 
-		FAnimSequenceDecompressionContext EvalDecompContext(SequenceLength, Interpolation, SourceName, CompressedData.CompressedDataStructure);
+		FAnimSequenceDecompressionContext EvalDecompContext(SequenceLength, Interpolation, SourceName, *CompressedData.CompressedDataStructure);
 		EvalDecompContext.Seek(ExtractionContext.CurrentTime);
 
 		// Handle Root Bone separately
@@ -741,10 +768,7 @@ void DecompressPose(FCompactPose& OutPose, const FCompressedAnimSequence& Compre
 			FCompactPoseBoneIndex RootBone(0);
 			FTransform& RootAtom = OutPose[RootBone];
 
-			AnimationFormat_GetBoneAtom(
-				RootAtom,
-				EvalDecompContext,
-				TrackIndex);
+			CompressedData.BoneCompressionCodec->DecompressBone(EvalDecompContext, TrackIndex, RootAtom);
 
 			// @laurent - we should look into splitting rotation and translation tracks, so we don't have to process translation twice.
 			FAnimationRuntime::RetargetBoneTransform(Skeleton, RetargetSource, RootAtom, 0, RootBone, RequiredBones, bIsBakedAdditive);
@@ -753,17 +777,8 @@ void DecompressPose(FCompactPose& OutPose, const FCompressedAnimSequence& Compre
 		if (RotationScalePairs.Num() > 0)
 		{
 			// get the remaining bone atoms
-			FTransformArray LocalBones;
-			OutPose.MoveBonesTo(LocalBones);
-
-			AnimationFormat_GetAnimationPose(
-				LocalBones,
-				RotationScalePairs,
-				TranslationPairs,
-				RotationScalePairs,
-				EvalDecompContext);
-
-			OutPose.MoveBonesFrom(MoveTemp(LocalBones));
+			TArrayView<FTransform> OutPoseBones = OutPose.GetMutableBones();
+			CompressedData.BoneCompressionCodec->DecompressPose(EvalDecompContext, RotationScalePairs, TranslationPairs, RotationScalePairs, OutPoseBones);
 		}
 	}
 
@@ -849,3 +864,17 @@ void DecompressPose(FCompactPose& OutPose, const FCompressedAnimSequence& Compre
 	}
 }
 
+FArchive& operator<<(FArchive& Ar, FCompressedOffsetData& D)
+{
+	Ar << D.OffsetData << D.StripSize;
+	return Ar;
+}
+
+FArchive& operator<<(FArchive& Ar, FAnimationErrorStats& ErrorStats)
+{
+	Ar << ErrorStats.AverageError;
+	Ar << ErrorStats.MaxError;
+	Ar << ErrorStats.MaxErrorTime;
+	Ar << ErrorStats.MaxErrorBone;
+	return Ar;
+}
diff --git a/Engine/Source/Runtime/Engine/Private/Animation/AnimCurveCompressionSettings.cpp b/Engine/Source/Runtime/Engine/Private/Animation/AnimCurveCompressionSettings.cpp
--- a/Engine/Source/Runtime/Engine/Private/Animation/AnimCurveCompressionSettings.cpp
+++ b/Engine/Source/Runtime/Engine/Private/Animation/AnimCurveCompressionSettings.cpp
@@ -18,6 +18,16 @@ UAnimCurveCompressionCodec* UAnimCurveCompressionSettings::GetCodec(const FStrin
 }
 
 #if WITH_EDITORONLY_DATA
+void UAnimCurveCompressionSettings::GetPreloadDependencies(TArray<UObject*>& OutDeps)
+{
+	Super::GetPreloadDependencies(OutDeps);
+
+	if (Codec != nullptr)
+	{
+		OutDeps.Add(Codec);
+	}
+}
+
 bool UAnimCurveCompressionSettings::AreSettingsValid() const
 {
 	return Codec != nullptr && Codec->IsCodecValid();
@@ -41,7 +51,7 @@ bool UAnimCurveCompressionSettings::Compress(const FCompressibleAnimData& AnimSe
 	return Success;
 }
 
-void UAnimCurveCompressionSettings::PopulateDDCKey(FArchive& Ar) const
+void UAnimCurveCompressionSettings::PopulateDDCKey(FArchive& Ar)
 {
 	if (Codec)
 	{
diff --git a/Engine/Source/Runtime/Engine/Private/Animation/AnimCurveTypes.cpp b/Engine/Source/Runtime/Engine/Private/Animation/AnimCurveTypes.cpp
--- a/Engine/Source/Runtime/Engine/Private/Animation/AnimCurveTypes.cpp
+++ b/Engine/Source/Runtime/Engine/Private/Animation/AnimCurveTypes.cpp
@@ -2,6 +2,8 @@
 
 #include "Animation/AnimCurveTypes.h"
 #include "UObject/FrameworkObjectVersion.h"
+#include "UObject/AnimObjectVersion.h"
+#include "Math/RandomStream.h"
 
 DECLARE_CYCLE_STAT(TEXT("EvalRawCurveData"), STAT_EvalRawCurveData, STATGROUP_Anim);
 
@@ -12,8 +14,9 @@ void FAnimCurveBase::PostSerialize(FArchive& Ar)
 {
 	SmartName::UID_Type CurveUid = SmartName::MaxUID;
 	Ar.UsingCustomVersion(FFrameworkObjectVersion::GUID);
+	Ar.UsingCustomVersion(FAnimObjectVersion::GUID);
 
-	if ((Ar.IsLoading()))
+	if (Ar.IsLoading())
 	{
 		if (Ar.CustomVer(FFrameworkObjectVersion::GUID) < FFrameworkObjectVersion::SmartNameRefactor)
 		{
@@ -29,6 +32,13 @@ void FAnimCurveBase::PostSerialize(FArchive& Ar)
 				Name.DisplayName = LastObservedName_DEPRECATED;
 			}
 		}
+
+#if WITH_EDITORONLY_DATA
+		if(Ar.CustomVer(FAnimObjectVersion::GUID) < FAnimObjectVersion::AnimSequenceCurveColors)
+		{
+			Color = MakeColor();
+		}
+#endif
 	}
 }
 
@@ -66,6 +76,16 @@ int32 FAnimCurveBase::GetCurveTypeFlags() const
 	return CurveTypeFlags;
 }
 
+#if WITH_EDITORONLY_DATA
+FLinearColor FAnimCurveBase::MakeColor()
+{
+	// Create a color based on the hash of the name
+	FRandomStream Stream(GetTypeHash(Name.DisplayName));
+	const uint8 Hue = (uint8)(Stream.FRand() * 255.0f);
+	return FLinearColor::MakeFromHSV8(Hue, 196, 196);
+}
+#endif
+
 ////////////////////////////////////////////////////
 //  FFloatCurve
 
diff --git a/Engine/Source/Runtime/Engine/Private/Animation/AnimEncoding.cpp b/Engine/Source/Runtime/Engine/Private/Animation/AnimEncoding.cpp
--- a/Engine/Source/Runtime/Engine/Private/Animation/AnimEncoding.cpp
+++ b/Engine/Source/Runtime/Engine/Private/Animation/AnimEncoding.cpp
@@ -178,89 +178,6 @@ void PadMemoryReader(FMemoryReader* MemoryReader, uint8*& TrackData, const int32
 	TrackData += Pad;
 }
 
-/**
- * Extracts a single BoneAtom from an Animation Sequence.
- *
- * @param	OutAtom			The BoneAtom to fill with the extracted result.
- * @param	DecompContext	The decompression context to use.
- * @param	TrackIndex		The index of the track desired in the Animation Sequence.
- */
-void AnimationFormat_GetBoneAtom(
-	FTransform& OutAtom,
-	FAnimSequenceDecompressionContext& DecompContext,
-	int32 TrackIndex)
-{
-	checkSlow(DecompContext.GetRotationCodec() != nullptr);
-	DecompContext.GetRotationCodec()->GetBoneAtom(OutAtom, DecompContext, TrackIndex);
-}
-
-#if USE_ANIMATION_CODEC_BATCH_SOLVER
-
-/**
- * Extracts an array of BoneAtoms from an Animation Sequence representing an entire pose of the skeleton.
- *
- * @param	Atoms				The BoneAtoms to fill with the extracted result.
- * @param	RotationTracks		A BoneTrackArray element for each bone requesting rotation data.
- * @param	TranslationTracks	A BoneTrackArray element for each bone requesting translation data.
- * @param	ScaleTracks			A BoneTrackArray element for each bone requesting scale data.
- * @param	DecompContext		The decompression context to use.
- */
-void AnimationFormat_GetAnimationPose(
-	FTransformArray& Atoms,
-	const BoneTrackArray& RotationPairs,
-	const BoneTrackArray& TranslationPairs,
-	const BoneTrackArray& ScalePairs,
-	FAnimSequenceDecompressionContext& DecompContext)
-{
-	// decompress the translation component using the proper method
-	checkSlow(DecompContext.GetTranslationCodec() != nullptr);
-	if (TranslationPairs.Num() > 0)
-	{
-		DecompContext.GetTranslationCodec()->GetPoseTranslations(Atoms, TranslationPairs, DecompContext);
-	}
-
-	// decompress the rotation component using the proper method
-	checkSlow(DecompContext.GetRotationCodec() != nullptr);
-	DecompContext.GetRotationCodec()->GetPoseRotations(Atoms, RotationPairs, DecompContext);
-
-	checkSlow(DecompContext.GetScaleCodec() != nullptr);
-	// we allow scale key to be empty
-	if (DecompContext.bHasScale)
-	{
-		DecompContext.GetScaleCodec()->GetPoseScales(Atoms, ScalePairs, DecompContext);
-	}
-}
-#endif
-
-/**
- * Extracts a single BoneAtom from an Animation Sequence.
- *
- * @param	OutAtom			The BoneAtom to fill with the extracted result.
- * @param	DecompContext	The decompression context to use.
- * @param	TrackIndex		The index of the track desired in the Animation Sequence.
- */
-void AnimEncodingLegacyBase::GetBoneAtom(FTransform& OutAtom, FAnimSequenceDecompressionContext& DecompContext, int32 TrackIndex)
-{
-	// Initialize to identity to set the scale and in case of a missing rotation or translation codec
-	OutAtom.SetIdentity();
-
-	// decompress the translation component using the proper method
-	checkSlow(DecompContext.GetTranslationCodec() != nullptr);
-	((AnimEncodingLegacyBase*)DecompContext.GetTranslationCodec())->GetBoneAtomTranslation(OutAtom, DecompContext, TrackIndex);
-
-	// decompress the rotation component using the proper method
-	checkSlow(DecompContext.GetRotationCodec() != nullptr);
-	((AnimEncodingLegacyBase*)DecompContext.GetRotationCodec())->GetBoneAtomRotation(OutAtom, DecompContext, TrackIndex);
-
-	// we assume scale keys can be empty, so only extract if we have valid keys
-	if (DecompContext.bHasScale)
-	{
-		// decompress the rotation component using the proper method
-		checkSlow(DecompContext.GetScaleCodec() != nullptr);
-		((AnimEncodingLegacyBase*)DecompContext.GetScaleCodec())->GetBoneAtomScale(OutAtom, DecompContext, TrackIndex);
-	}
-}
-
 /**
  * Handles Byte-swapping incoming animation data from a MemoryReader
  *
@@ -860,4 +777,4 @@ void AnimationFormat_SetInterfaceLinks(CompressedDataType& CompressedData)
 }
 
 template void AnimationFormat_SetInterfaceLinks(FUECompressedAnimData& CompressedData);
-template void AnimationFormat_SetInterfaceLinks(FCompressibleAnimDataResult& CompressedData);
\ No newline at end of file
+template void AnimationFormat_SetInterfaceLinks(FUECompressedAnimDataMutable& CompressedData);
