COPY from //Dev-Anim to //Dev-Main

diff --git a/Engine/Source/Runtime/Engine/Private/Animation/AnimEncoding_ConstantKeyLerp.cpp b/Engine/Source/Runtime/Engine/Private/Animation/AnimEncoding_ConstantKeyLerp.cpp
--- a/Engine/Source/Runtime/Engine/Private/Animation/AnimEncoding_ConstantKeyLerp.cpp
+++ b/Engine/Source/Runtime/Engine/Private/Animation/AnimEncoding_ConstantKeyLerp.cpp
@@ -7,6 +7,9 @@
 #include "AnimEncoding_ConstantKeyLerp.h"
 #include "Serialization/MemoryReader.h"
 #include "Serialization/MemoryWriter.h"
+#if INTEL_ISPC
+#include "AnimEncoding_ConstantKeyLerp.ispc.generated.h"
+#endif
 
 /**
  * Handles the ByteSwap of compressed rotation data on import
@@ -235,59 +238,191 @@ void AEFConstantKeyLerpShared::ByteSwapScaleOut(
 	}
 }
 
-#if USE_SEGMENTING_CONTEXT
-void AEFConstantKeyLerpShared::CreateEncodingContext(FAnimSequenceDecompressionContext& DecompContext)
-{
-	checkSlow(DecompContext.EncodingContext == nullptr);
-	DecompContext.EncodingContext = new FAEConstantKeyLerpContext(DecompContext);
-}
+#if USE_ANIMATION_CODEC_BATCH_SOLVER
 
-void AEFConstantKeyLerpShared::ReleaseEncodingContext(FAnimSequenceDecompressionContext& DecompContext)
+/**
+ * Decompress all requested rotation components from an Animation Sequence
+ *
+ * @param	Atoms			The FTransform array to fill in.
+ * @param	DesiredPairs	Array of requested bone information
+ * @param	DecompContext	The decompression context to use.
+ */
+template<int32 FORMAT>
+inline void AEFConstantKeyLerp<FORMAT>::GetPoseRotations(
+	TArrayView<FTransform>& Atoms,
+	const BoneTrackArray& DesiredPairs,
+	FAnimSequenceDecompressionContext& DecompContext)
 {
-	checkSlow(DecompContext.EncodingContext != nullptr);
-	delete DecompContext.EncodingContext;
-	DecompContext.EncodingContext = nullptr;
-}
+	const int32 PairCount = DesiredPairs.Num();
 
-FAEConstantKeyLerpContext::FAEConstantKeyLerpContext(const FAnimSequenceDecompressionContext& DecompContext)
-	: KeyFrameSize(-1)
-{
+	if (INTEL_ISPC)
+	{
+#if INTEL_ISPC
+		const FUECompressedAnimData& AnimData = static_cast<const FUECompressedAnimData&>(DecompContext.CompressedAnimData);
+
+		ispc::GetConstantKeyLerpPoseRotations(
+			(ispc::FTransform*)&Atoms[0],
+			(ispc::BoneTrackPair*)&DesiredPairs[0],
+			AnimData.CompressedTrackOffsets.GetData(),
+			AnimData.CompressedByteStream.GetData(),
+			DecompContext.SequenceLength,
+			DecompContext.RelativePos,
+			(uint8)DecompContext.Interpolation,
+			FORMAT,
+			PairCount);
+#endif
+	}
+	else
+	{
+		for (int32 PairIndex = 0; PairIndex < PairCount; ++PairIndex)
+		{
+			const BoneTrackPair& Pair = DesiredPairs[PairIndex];
+			const int32 TrackIndex = Pair.TrackIndex;
+			const int32 AtomIndex = Pair.AtomIndex;
+			FTransform& BoneAtom = Atoms[AtomIndex];
+
+			// call the decoder directly (not through the vtable)
+			AEFConstantKeyLerp<FORMAT>::GetBoneAtomRotation(BoneAtom, DecompContext, TrackIndex);
+		}
+	}
 }
 
-void FAEConstantKeyLerpContext::Seek(const FAnimSequenceDecompressionContext& DecompContext, float SampleAtTime)
+/**
+ * Decompress all requested translation components from an Animation Sequence
+ *
+ * @param	Atoms			The FTransform array to fill in.
+ * @param	DesiredPairs	Array of requested bone information
+ * @param	DecompContext	The decompression context to use.
+ */
+template<int32 FORMAT>
+inline void AEFConstantKeyLerp<FORMAT>::GetPoseTranslations(
+	TArrayView<FTransform>& Atoms,
+	const BoneTrackArray& DesiredPairs,
+	FAnimSequenceDecompressionContext& DecompContext)
 {
-	if (KeyFrameSize < 0)
+	const int32 PairCount = DesiredPairs.Num();
+
+	if (INTEL_ISPC)
 	{
-		// First update, cache some stuff
-		UniformKeyOffsets.Empty(DecompContext.NumTracks * DecompContext.NumStreamsPerTrack);
-		UniformKeyOffsets.AddUninitialized(DecompContext.NumTracks * DecompContext.NumStreamsPerTrack);
+#if INTEL_ISPC
+		const FUECompressedAnimData& AnimData = static_cast<const FUECompressedAnimData&>(DecompContext.CompressedAnimData);
 
-		const int32 PackedTranslationSize0 = CompressedTranslationStrides[DecompContext.Segment0->TranslationCompressionFormat] * CompressedTranslationNum[DecompContext.Segment0->TranslationCompressionFormat];
-		const int32 PackedRotationSize0 = CompressedRotationStrides[DecompContext.Segment0->RotationCompressionFormat] * CompressedRotationNum[DecompContext.Segment0->RotationCompressionFormat];
-		const int32 PackedScaleSize0 = DecompContext.bHasScale ? (CompressedScaleStrides[DecompContext.Segment0->ScaleCompressionFormat] * CompressedScaleNum[DecompContext.Segment0->ScaleCompressionFormat]) : 0;
+		ispc::GetConstantKeyLerpPoseTranslations(
+			(ispc::FTransform*)&Atoms[0],
+			(ispc::BoneTrackPair*)&DesiredPairs[0],
+			AnimData.CompressedTrackOffsets.GetData(),
+			AnimData.CompressedByteStream.GetData(),
+			DecompContext.SequenceLength,
+			DecompContext.RelativePos,
+			(uint8)DecompContext.Interpolation,
+			FORMAT,
+			PairCount);
+#endif
+	}
+	else
+	{
+		//@TODO: Verify that this prefetch is helping
+		// Prefetch the desired pairs array and 2 destination spots; the loop will prefetch one 2 out each iteration
+		FPlatformMisc::Prefetch(&(DesiredPairs[0]));
+		const int32 PrefetchCount = FMath::Min(PairCount, 1);
+		for (int32 PairIndex = 0; PairIndex < PairCount; ++PairIndex)
+		{
+			const BoneTrackPair& Pair = DesiredPairs[PairIndex];
+			FPlatformMisc::Prefetch(Atoms.GetData() + Pair.AtomIndex);
+		}
 
-		int32 KeyOffset = 0;
-		for (int32 TrackIndex = 0; TrackIndex < DecompContext.NumTracks; ++TrackIndex)
+		for (int32 PairIndex = 0; PairIndex < PairCount; ++PairIndex)
 		{
-			const FTrivialTrackFlags TrackFlags(DecompContext.TrackFlags[TrackIndex]);
+			int32 PrefetchIndex = PairIndex + PrefetchCount;
+			if (PrefetchIndex < PairCount)
+			{
+				FPlatformMisc::Prefetch(Atoms.GetData() + DesiredPairs[PrefetchIndex].AtomIndex);
+			}
+
+			const BoneTrackPair& Pair = DesiredPairs[PairIndex];
+			const int32 TrackIndex = Pair.TrackIndex;
+			const int32 AtomIndex = Pair.AtomIndex;
+			FTransform& BoneAtom = Atoms[AtomIndex];
+
+			// call the decoder directly (not through the vtable)
+			AEFConstantKeyLerp<FORMAT>::GetBoneAtomTranslation(BoneAtom, DecompContext, TrackIndex);
+		}
+	}
+}
+
+/**
+ * Decompress all requested Scale components from an Animation Sequence
+ *
+ * @param	Atoms			The FTransform array to fill in.
+ * @param	DesiredPairs	Array of requested bone information
+ * @param	DecompContext	The decompression context to use.
+ */
+template<int32 FORMAT>
+inline void AEFConstantKeyLerp<FORMAT>::GetPoseScales(
+	TArrayView<FTransform>& Atoms,
+	const BoneTrackArray& DesiredPairs,
+	FAnimSequenceDecompressionContext& DecompContext)
+{
+	const int32 PairCount= DesiredPairs.Num();
+
+	if (INTEL_ISPC)
+	{
+#if INTEL_ISPC
+		const FUECompressedAnimData& AnimData = static_cast<const FUECompressedAnimData&>(DecompContext.CompressedAnimData);
 
-			UniformKeyOffsets[DecompContext.GetTranslationValueOffset(TrackIndex)] = KeyOffset;
-			KeyOffset += TrackFlags.IsTranslationTrivial() ? 0 : PackedTranslationSize0;
+		const TArrayView<int32> ScaleOffsets = AnimData.CompressedScaleOffsets.OffsetData;
+		const int32 StripSize = AnimData.CompressedScaleOffsets.StripSize;
 
-			UniformKeyOffsets[DecompContext.GetRotationValueOffset(TrackIndex)] = KeyOffset;
-			KeyOffset += TrackFlags.IsRotationTrivial() ? 0 : PackedRotationSize0;
+		ispc::GetConstantKeyLerpPoseScales(
+			(ispc::FTransform*)&Atoms[0],
+			(ispc::BoneTrackPair*)&DesiredPairs[0],
+			ScaleOffsets.GetData(),
+			StripSize,
+			AnimData.CompressedByteStream.GetData(),
+			DecompContext.SequenceLength,
+			DecompContext.RelativePos,
+			(uint8)DecompContext.Interpolation,
+			FORMAT,
+			PairCount);
+#endif
+	}
+	else
+	{
+		//@TODO: Verify that this prefetch is helping
+		// Prefetch the desired pairs array and 2 destination spots; the loop will prefetch one 2 out each iteration
+		FPlatformMisc::Prefetch(&(DesiredPairs[0]));
+		const int32 PrefetchCount = FMath::Min(PairCount, 1);
+		for (int32 PairIndex = 0; PairIndex < PairCount; ++PairIndex)
+		{
+			const BoneTrackPair& Pair = DesiredPairs[PairIndex];
+			FPlatformMisc::Prefetch(Atoms.GetData() + Pair.AtomIndex);
+		}
 
-			if (DecompContext.bHasScale)
+		for (int32 PairIndex = 0; PairIndex < PairCount; ++PairIndex)
+		{
+			int32 PrefetchIndex = PairIndex + PrefetchCount;
+			if (PrefetchIndex < PairCount)
 			{
-				UniformKeyOffsets[DecompContext.GetScaleValueOffset(TrackIndex)] = KeyOffset;
-				KeyOffset += TrackFlags.IsScaleTrivial() ? 0 : PackedScaleSize0;
+				FPlatformMisc::Prefetch(Atoms.GetData() + DesiredPairs[PrefetchIndex].AtomIndex);
 			}
-		}
 
-		KeyFrameSize = KeyOffset;
-	}
+			const BoneTrackPair& Pair = DesiredPairs[PairIndex];
+			const int32 TrackIndex = Pair.TrackIndex;
+			const int32 AtomIndex = Pair.AtomIndex;
+			FTransform& BoneAtom = Atoms[AtomIndex];
 
-	FrameKeysOffset[0] = DecompContext.Segment0->ByteStreamOffset + DecompContext.RangeDataSize0 + (KeyFrameSize * DecompContext.SegmentKeyIndex0);
-	FrameKeysOffset[1] = DecompContext.Segment1->ByteStreamOffset + DecompContext.RangeDataSize0 + (KeyFrameSize * DecompContext.SegmentKeyIndex1);
+			// call the decoder directly (not through the vtable)
+			AEFConstantKeyLerp<FORMAT>::GetBoneAtomScale(BoneAtom, DecompContext, TrackIndex);
+		}
+	}
 }
+
+template class AEFConstantKeyLerp<ACF_None>;
+template class AEFConstantKeyLerp<ACF_Float96NoW>;
+template class AEFConstantKeyLerp<ACF_Fixed48NoW>;
+template class AEFConstantKeyLerp<ACF_IntervalFixed32NoW>;
+template class AEFConstantKeyLerp<ACF_Fixed32NoW>;
+template class AEFConstantKeyLerp<ACF_Float32NoW>;
+template class AEFConstantKeyLerp<ACF_Identity>;
+
 #endif
\ No newline at end of file
diff --git a/Engine/Source/Runtime/Engine/Private/Animation/AnimEncoding_PerTrackCompression.cpp b/Engine/Source/Runtime/Engine/Private/Animation/AnimEncoding_PerTrackCompression.cpp
--- a/Engine/Source/Runtime/Engine/Private/Animation/AnimEncoding_PerTrackCompression.cpp
+++ b/Engine/Source/Runtime/Engine/Private/Animation/AnimEncoding_PerTrackCompression.cpp
@@ -8,71 +8,10 @@
 #include "Serialization/MemoryWriter.h"
 #include "Serialization/MemoryReader.h"
 #include "Animation/AnimEncodingHeapAllocator.h"
-#include "Animation/AnimEncodingDecompressionContext.h"
 #include "AnimationCompression.h"
-
-
-class FAEPerTrackKeyLerpContext : public FAnimEncodingDecompressionContext
-{
-public:
-	static constexpr int32 OffsetNumKeysPairSize = sizeof(uint32) + sizeof(uint16);
-
-	FAEPerTrackKeyLerpContext(const FAnimSequenceDecompressionContext& DecompContext);
-	virtual void Seek(const FAnimSequenceDecompressionContext& DecompContext, float SampleAtTime) override;
-
-	inline void GetRotation(FTransform& OutAtom, const FAnimSequenceDecompressionContext& DecompContext, int32 TrackIndex) const;
-	inline void GetTranslation(FTransform& OutAtom, const FAnimSequenceDecompressionContext& DecompContext, int32 TrackIndex) const;
-	inline void GetScale(FTransform& OutAtom, const FAnimSequenceDecompressionContext& DecompContext, int32 TrackIndex) const;
-
-	// Common
-	TArray<int32, FAnimEncodingHeapAllocator> RangeOffsets[2];
-	int32 PerTrackStreamFlagOffsets[2];
-	int32 RangeDataSize[2];
-	uint16 PreviousSegmentIndex[2];
-
-	// Uniform
-	TArray<int32, FAnimEncodingHeapAllocator> UniformKeyOffsets[2];
-	int32 UniformKeyFrameSize[2];
-	int32 UniformDataOffsets[2];
-
-	// Variable common
-	TArray<uint8, FAnimEncodingHeapAllocator> TrackStreamKeySizes[2];
-
-	// Variable linear
-	float SegmentRelativePos0;
-	uint8 TimeMarkerSize[2];	// sizeof(uint8) or sizeof(uint16)
-	int32 OffsetNumKeysPairOffsets[2];
-	TArray<int32, FAnimEncodingHeapAllocator> NumAnimatedTrackStreams[2];
-
-	// Variable sorted
-	struct FCachedKey
-	{
-		// Linear interpolation only requires 2 keys
-		// Index 0 is always the oldest key, index 1 the newest
-		int32 RotOffsets[2];
-		int32 RotFrameIndices[2];
-		int32 TransOffsets[2];
-		int32 TransFrameIndices[2];
-		int32 ScaleOffsets[2];
-		int32 ScaleFrameIndices[2];
-	};
-
-	TArray<FCachedKey, FAnimEncodingHeapAllocator> CachedKeys;			// 1 Entry per track
-
-	int32 SegmentStartFrame[2];
-	float FramePos;
-	float PreviousSampleAtTime;
-
-	const uint8_t* PackedSampleData;		// The current pointer into our data stream
-	int32 PreviousFrameIndex;				// The previously read frame index
-	int32 CurrentFrameIndex;				// The current frame index
-	uint8 CurrentSegmentIndex;				// The current segment index of the sorted data, either 0 or 1
-
-private:
-	void CacheSegmentValues(const FAnimSequenceDecompressionContext& DecompContext, const FCompressedSegment& Segment, uint8 SegmentIndex);
-	void ResetSortedCache(const FAnimSequenceDecompressionContext& DecompContext);
-	void AdvanceSortedCachedKeys(const FAnimSequenceDecompressionContext& DecompContext);
-};
+#if INTEL_ISPC
+#include "AnimEncoding_PerTrackCompression.ispc.generated.h"
+#endif
 
 // This define controls whether scalar or vector code is used to decompress keys.  Note that not all key decompression code
 // is vectorized yet, so some (seldom used) formats will actually get slower (due to extra LHS stalls) when enabled.
@@ -509,61 +448,19 @@ void AEFPerTrackCompressionCodec::ByteSwapOut(
 }
 
 
-
-/**
- * Extracts a single BoneAtom from an Animation Sequence.
- *
- * @param	OutAtom			The BoneAtom to fill with the extracted result.
- * @param	DecompContext	The decompression context to use.
- * @param	TrackIndex		The index of the track desired in the Animation Sequence.
- */
-void AEFPerTrackCompressionCodec::GetBoneAtom(
-	FTransform& OutAtom,
-	FAnimSequenceDecompressionContext& DecompContext,
-	int32 TrackIndex)
-{
-	// Initialize to identity to set the scale and in case of a missing rotation or translation codec
-	OutAtom.SetIdentity();
-
-	GetBoneAtomTranslation(OutAtom, DecompContext, TrackIndex);
-	GetBoneAtomRotation(OutAtom, DecompContext, TrackIndex);
-
-	if (DecompContext.bHasScale)
-	{
-		GetBoneAtomScale(OutAtom, DecompContext, TrackIndex);
-	}
-}
-
 void AEFPerTrackCompressionCodec::GetBoneAtomRotation(
 	FTransform& OutAtom,
 	FAnimSequenceDecompressionContext& DecompContext,
 	int32 TrackIndex)
 {
-#if USE_SEGMENTING_CONTEXT
-	if (DecompContext.AnimSeq->CompressedSegments.Num() != 0)
-	{
-		const FTrivialAnimKeyHandle TrivialKeyHandle = DecompContext.GetTrivialRotationKeyHandle(TrackIndex);
-		if (TrivialKeyHandle.IsValid())
-		{
-			DecompContext.GetTrivialRotation(OutAtom, TrivialKeyHandle);
-		}
-		else
-		{
-			const FAEPerTrackKeyLerpContext& EncodingContext = *static_cast<const FAEPerTrackKeyLerpContext*>(DecompContext.EncodingContext);
-			EncodingContext.GetRotation(OutAtom, DecompContext, TrackIndex);
-		}
-
-		return;
-	}
-#endif
-
-	const int32* RESTRICT TrackOffsetData = DecompContext.GetCompressedTrackOffsets() + (TrackIndex * 2);
+	const FUECompressedAnimData& AnimData = static_cast<const FUECompressedAnimData&>(DecompContext.CompressedAnimData);
+	const int32* RESTRICT TrackOffsetData = AnimData.CompressedTrackOffsets.GetData() + (TrackIndex * 2);
 	const int32 RotKeysOffset = TrackOffsetData[1];
 
 	if (RotKeysOffset != INDEX_NONE)
 	{
-		const uint8* RESTRICT TrackData = DecompContext.GetCompressedByteStream() + RotKeysOffset + 4;
-		const int32 Header = *((int32*)(DecompContext.GetCompressedByteStream() + RotKeysOffset));
+		const uint8* RESTRICT TrackData = AnimData.CompressedByteStream.GetData() + RotKeysOffset + 4;
+		const int32 Header = *((int32*)(AnimData.CompressedByteStream.GetData() + RotKeysOffset));
 
 		int32 KeyFormat;
 		int32 NumKeys;
@@ -583,12 +480,12 @@ void AEFPerTrackCompressionCodec::GetBoneAtomRotation(
 		{
 			if ((FormatFlags & 0x8) == 0)
 			{
-				Alpha = TimeToIndex(DecompContext.GetSequenceLength(), DecompContext.RelativePos, NumKeys, DecompContext.GetInterpolation(), Index0, Index1);
+				Alpha = TimeToIndex(DecompContext.SequenceLength, DecompContext.RelativePos, NumKeys, DecompContext.Interpolation, Index0, Index1);
 			}
 			else
 			{
 				const uint8* RESTRICT FrameTable = Align(TrackData + FixedBytes + BytesPerKey * NumKeys, 4);
-				Alpha = TimeToIndex(DecompContext.GetInterpolation(), DecompContext.GetCompressedNumberOfFrames(), FrameTable, DecompContext.RelativePos, NumKeys, Index0, Index1);
+				Alpha = TimeToIndex(DecompContext.Interpolation, AnimData.CompressedNumberOfFrames, FrameTable, DecompContext.RelativePos, NumKeys, Index0, Index1);
 			}
 		}
 
@@ -644,31 +541,14 @@ void AEFPerTrackCompressionCodec::GetBoneAtomTranslation(
 	FAnimSequenceDecompressionContext& DecompContext,
 	int32 TrackIndex)
 {
-#if USE_SEGMENTING_CONTEXT
-	if (DecompContext.AnimSeq->CompressedSegments.Num() != 0)
-	{
-		const FTrivialAnimKeyHandle TrivialKeyHandle = DecompContext.GetTrivialTranslationKeyHandle(TrackIndex);
-		if (TrivialKeyHandle.IsValid())
-		{
-			DecompContext.GetTrivialTranslation(OutAtom, TrivialKeyHandle);
-		}
-		else
-		{
-			const FAEPerTrackKeyLerpContext& EncodingContext = *static_cast<const FAEPerTrackKeyLerpContext*>(DecompContext.EncodingContext);
-			EncodingContext.GetTranslation(OutAtom, DecompContext, TrackIndex);
-		}
-
-		return;
-	}
-#endif
-
-	const int32* RESTRICT TrackOffsetData = DecompContext.GetCompressedTrackOffsets() + (TrackIndex * 2);
+	const FUECompressedAnimData& AnimData = static_cast<const FUECompressedAnimData&>(DecompContext.CompressedAnimData);
+	const int32* RESTRICT TrackOffsetData = AnimData.CompressedTrackOffsets.GetData() + (TrackIndex * 2);
 	const int32 PosKeysOffset = TrackOffsetData[0];
 
 	if (PosKeysOffset != INDEX_NONE)
 	{
-		const uint8* RESTRICT TrackData = DecompContext.GetCompressedByteStream() + PosKeysOffset + 4;
-		const int32 Header = *((int32*)(DecompContext.GetCompressedByteStream() + PosKeysOffset));
+		const uint8* RESTRICT TrackData = AnimData.CompressedByteStream.GetData() + PosKeysOffset + 4;
+		const int32 Header = *((int32*)(AnimData.CompressedByteStream.GetData() + PosKeysOffset));
 
 		int32 KeyFormat;
 		int32 NumKeys;
@@ -677,9 +557,9 @@ void AEFPerTrackCompressionCodec::GetBoneAtomTranslation(
 		int32 FixedBytes;
 		FAnimationCompression_PerTrackUtils::DecomposeHeader(Header, /*OUT*/ KeyFormat, /*OUT*/ NumKeys, /*OUT*/ FormatFlags, /*OUT*/BytesPerKey, /*OUT*/ FixedBytes);
 
-		checkf(KeyFormat != ACF_None, TEXT("[%s] contians invalid keyformat. NumKeys (%d), FormatFlags (%d), BytesPerKeys (%d), FixedBytes (%d)"), *DecompContext.GetAnimFName().ToString(), NumKeys, FormatFlags, BytesPerKey, FixedBytes);
+		checkf(KeyFormat != ACF_None, TEXT("[%s] contians invalid keyformat. NumKeys (%d), FormatFlags (%d), BytesPerKeys (%d), FixedBytes (%d)"), *DecompContext.AnimName.ToString(), NumKeys, FormatFlags, BytesPerKey, FixedBytes);
 
-		checkf(KeyFormat < ACF_MAX, TEXT("[%s] contians invalid keyformat. NumKeys (%d), FormatFlags (%d), BytesPerKeys (%d), FixedBytes (%d), PosKeysOffset (%d), TrackIndex (%d)"), *DecompContext.GetAnimFName().ToString(), NumKeys, FormatFlags, BytesPerKey, FixedBytes, PosKeysOffset, TrackIndex)
+		checkf(KeyFormat < ACF_MAX, TEXT("[%s] contians invalid keyformat. NumKeys (%d), FormatFlags (%d), BytesPerKeys (%d), FixedBytes (%d), PosKeysOffset (%d), TrackIndex (%d)"), *DecompContext.AnimName.ToString(), NumKeys, FormatFlags, BytesPerKey, FixedBytes, PosKeysOffset, TrackIndex)
 		// Figure out the key indexes
 		int32 Index0 = 0;
 		int32 Index1 = 0;
@@ -691,12 +571,12 @@ void AEFPerTrackCompressionCodec::GetBoneAtomTranslation(
 		{
 			if ((FormatFlags & 0x8) == 0)
 			{
-				Alpha = TimeToIndex(DecompContext.GetSequenceLength(), DecompContext.RelativePos, NumKeys, DecompContext.GetInterpolation(), Index0, Index1);
+				Alpha = TimeToIndex(DecompContext.SequenceLength, DecompContext.RelativePos, NumKeys, DecompContext.Interpolation, Index0, Index1);
 			}
 			else
 			{
 				const uint8* RESTRICT FrameTable = Align(TrackData + FixedBytes + BytesPerKey * NumKeys, 4);
-				Alpha = TimeToIndex(DecompContext.GetInterpolation(), DecompContext.GetCompressedNumberOfFrames(), FrameTable, DecompContext.RelativePos, NumKeys, Index0, Index1);
+				Alpha = TimeToIndex(DecompContext.Interpolation, AnimData.CompressedNumberOfFrames, FrameTable, DecompContext.RelativePos, NumKeys, Index0, Index1);
 			}
 		}
 
@@ -751,31 +631,13 @@ void AEFPerTrackCompressionCodec::GetBoneAtomScale(
 	FAnimSequenceDecompressionContext& DecompContext,
 	int32 TrackIndex)
 {
-#if USE_SEGMENTING_CONTEXT
-	if (DecompContext.AnimSeq->CompressedSegments.Num() != 0)
-	{
-		const FTrivialAnimKeyHandle TrivialKeyHandle = DecompContext.GetTrivialScaleKeyHandle(TrackIndex);
-		if (TrivialKeyHandle.IsValid())
-		{
-			DecompContext.GetTrivialScale(OutAtom, TrivialKeyHandle);
-		}
-		else
-		{
-			const FAEPerTrackKeyLerpContext& EncodingContext = *static_cast<const FAEPerTrackKeyLerpContext*>(DecompContext.EncodingContext);
-			EncodingContext.GetScale(OutAtom, DecompContext, TrackIndex);
-		}
-
-		return;
-	}
-#endif
-
-
-	const int32 ScaleKeysOffset = DecompContext.GetCompressedScaleOffsets()->GetOffsetData( TrackIndex, 0 );
+	const FUECompressedAnimData& AnimData = static_cast<const FUECompressedAnimData&>(DecompContext.CompressedAnimData);
+	const int32 ScaleKeysOffset = AnimData.CompressedScaleOffsets.GetOffsetData( TrackIndex, 0 );
 
 	if (ScaleKeysOffset != INDEX_NONE)
 	{
-		const uint8* RESTRICT TrackData = DecompContext.GetCompressedByteStream() + ScaleKeysOffset + 4;
-		const int32 Header = *((int32*)(DecompContext.GetCompressedByteStream() + ScaleKeysOffset));
+		const uint8* RESTRICT TrackData = AnimData.CompressedByteStream.GetData() + ScaleKeysOffset + 4;
+		const int32 Header = *((int32*)(AnimData.CompressedByteStream.GetData() + ScaleKeysOffset));
 
 		int32 KeyFormat;
 		int32 NumKeys;
@@ -795,12 +657,12 @@ void AEFPerTrackCompressionCodec::GetBoneAtomScale(
 		{
 			if ((FormatFlags & 0x8) == 0)
 			{
-				Alpha = TimeToIndex(DecompContext.GetSequenceLength(), DecompContext.RelativePos, NumKeys, DecompContext.GetInterpolation(), Index0, Index1);
+				Alpha = TimeToIndex(DecompContext.SequenceLength, DecompContext.RelativePos, NumKeys, DecompContext.Interpolation, Index0, Index1);
 			}
 			else
 			{
 				const uint8* RESTRICT FrameTable = Align(TrackData + FixedBytes + BytesPerKey * NumKeys, 4);
-				Alpha = TimeToIndex(DecompContext.GetInterpolation(), DecompContext.GetCompressedNumberOfFrames(), FrameTable, DecompContext.RelativePos, NumKeys, Index0, Index1);
+				Alpha = TimeToIndex(DecompContext.Interpolation, AnimData.CompressedNumberOfFrames, FrameTable, DecompContext.RelativePos, NumKeys, Index0, Index1);
 			}
 		}
 
@@ -859,20 +721,40 @@ void AEFPerTrackCompressionCodec::GetBoneAtomScale(
  * @param	DecompContext	The decompression context to use.
  */
 void AEFPerTrackCompressionCodec::GetPoseRotations(
-	FTransformArray& Atoms,
+	TArrayView<FTransform>& Atoms,
 	const BoneTrackArray& DesiredPairs,
 	FAnimSequenceDecompressionContext& DecompContext)
 {
 	const int32 PairCount = DesiredPairs.Num();
 
-	for( int32 PairIndex = 0; PairIndex < PairCount; ++PairIndex )
+	if (INTEL_ISPC)
+	{
+#if INTEL_ISPC
+		const FUECompressedAnimData& AnimData = static_cast<const FUECompressedAnimData&>(DecompContext.CompressedAnimData);
+
+		ispc::GetPerTrackCompressionPoseRotations(
+			(ispc::FTransform*)&Atoms[0],
+			(ispc::BoneTrackPair*)&DesiredPairs[0],
+			AnimData.CompressedTrackOffsets.GetData(),
+			AnimData.CompressedByteStream.GetData(),
+			AnimData.CompressedNumberOfFrames,
+			DecompContext.SequenceLength,
+			DecompContext.RelativePos,
+			(uint8)DecompContext.Interpolation,
+			PairCount);
+#endif
+	}
+	else
 	{
-		const BoneTrackPair& Pair = DesiredPairs[PairIndex];
-		const int32 TrackIndex = Pair.TrackIndex;
-		const int32 AtomIndex = Pair.AtomIndex;
-		FTransform& BoneAtom = Atoms[AtomIndex];
-
-		GetBoneAtomRotation(BoneAtom, DecompContext, TrackIndex);
+		for( int32 PairIndex = 0; PairIndex < PairCount; ++PairIndex )
+		{
+			const BoneTrackPair& Pair = DesiredPairs[PairIndex];
+			const int32 TrackIndex = Pair.TrackIndex;
+			const int32 AtomIndex = Pair.AtomIndex;
+			FTransform& BoneAtom = Atoms[AtomIndex];
+	
+			GetBoneAtomRotation(BoneAtom, DecompContext, TrackIndex);
+		}
 	}
 }
 
@@ -884,20 +766,40 @@ void AEFPerTrackCompressionCodec::GetPoseRotations(
  * @param	DecompContext	The decompression context to use.
  */
 void AEFPerTrackCompressionCodec::GetPoseTranslations(
-	FTransformArray& Atoms,
+	TArrayView<FTransform>& Atoms,
 	const BoneTrackArray& DesiredPairs,
 	FAnimSequenceDecompressionContext& DecompContext)
 {
 	const int32 PairCount = DesiredPairs.Num();
 
-	for( int32 PairIndex = 0; PairIndex < PairCount; ++PairIndex )
+	if (INTEL_ISPC)
+	{
+#if INTEL_ISPC
+		const FUECompressedAnimData& AnimData = static_cast<const FUECompressedAnimData&>(DecompContext.CompressedAnimData);
+
+		ispc::GetPerTrackCompressionPoseTranslations(
+			(ispc::FTransform*)&Atoms[0],
+			(ispc::BoneTrackPair*)&DesiredPairs[0],
+			AnimData.CompressedTrackOffsets.GetData(),
+			AnimData.CompressedByteStream.GetData(),
+			AnimData.CompressedNumberOfFrames,
+			DecompContext.SequenceLength,
+			DecompContext.RelativePos,
+			(uint8)DecompContext.Interpolation,
+			PairCount);
+#endif
+	}
+	else
 	{
-		const BoneTrackPair& Pair = DesiredPairs[PairIndex];
-		const int32 TrackIndex = Pair.TrackIndex;
-		const int32 AtomIndex = Pair.AtomIndex;
-		FTransform& BoneAtom = Atoms[AtomIndex];
-
-		GetBoneAtomTranslation(BoneAtom, DecompContext, TrackIndex);
+		for( int32 PairIndex = 0; PairIndex < PairCount; ++PairIndex )
+		{
+			const BoneTrackPair& Pair = DesiredPairs[PairIndex];
+			const int32 TrackIndex = Pair.TrackIndex;
+			const int32 AtomIndex = Pair.AtomIndex;
+			FTransform& BoneAtom = Atoms[AtomIndex];
+	
+			GetBoneAtomTranslation(BoneAtom, DecompContext, TrackIndex);
+		}
 	}
 }
 
@@ -909,1005 +811,45 @@ void AEFPerTrackCompressionCodec::GetPoseTranslations(
  * @param	DecompContext	The decompression context to use.
  */
 void AEFPerTrackCompressionCodec::GetPoseScales(
-	FTransformArray& Atoms,
+	TArrayView<FTransform>& Atoms,
 	const BoneTrackArray& DesiredPairs,
 	FAnimSequenceDecompressionContext& DecompContext)
 {
-	checkSlow(DecompContext.bHasScale);
-
 	const int32 PairCount = DesiredPairs.Num();
 
-	for( int32 PairIndex = 0; PairIndex < PairCount; ++PairIndex )
-	{
-		const BoneTrackPair& Pair = DesiredPairs[PairIndex];
-		const int32 TrackIndex = Pair.TrackIndex;
-		const int32 AtomIndex = Pair.AtomIndex;
-		FTransform& BoneAtom = Atoms[AtomIndex];
-
-		GetBoneAtomScale(BoneAtom, DecompContext, TrackIndex);
-	}
-}
-
-#endif // USE_ANIMATION_CODEC_BATCH_SOLVER
-
-#if USE_SEGMENTING_CONTEXT
-void AEFPerTrackCompressionCodec::CreateEncodingContext(FAnimSequenceDecompressionContext& DecompContext)
-{
-	checkSlow(DecompContext.EncodingContext == nullptr);
-	DecompContext.EncodingContext = new FAEPerTrackKeyLerpContext(DecompContext);
-}
-
-void AEFPerTrackCompressionCodec::ReleaseEncodingContext(FAnimSequenceDecompressionContext& DecompContext)
-{
-	checkSlow(DecompContext.EncodingContext != nullptr);
-	delete DecompContext.EncodingContext;
-	DecompContext.EncodingContext = nullptr;
-}
-
-FAEPerTrackKeyLerpContext::FAEPerTrackKeyLerpContext(const FAnimSequenceDecompressionContext& DecompContext)
-	: PreviousSampleAtTime(FLT_MAX)	// Very large to trigger a reset on the first Seek()
-{
-	PreviousSegmentIndex[0] = -1;
-	PreviousSegmentIndex[1] = -1;
-
-	const int32 NumEntries = DecompContext.NumTracks * DecompContext.NumStreamsPerTrack;
-	for (int32 SegmentIndex = 0; SegmentIndex < 2; ++SegmentIndex)
-	{
-		UniformKeyOffsets[SegmentIndex].Empty(NumEntries);
-		UniformKeyOffsets[SegmentIndex].AddUninitialized(NumEntries);
-		RangeOffsets[SegmentIndex].Empty(NumEntries);
-		RangeOffsets[SegmentIndex].AddUninitialized(NumEntries);
-
-		if (!DecompContext.bIsSorted)
-		{
-			NumAnimatedTrackStreams[SegmentIndex].Empty(NumEntries);
-			NumAnimatedTrackStreams[SegmentIndex].AddUninitialized(NumEntries);
-			TrackStreamKeySizes[SegmentIndex].Empty(NumEntries);
-			TrackStreamKeySizes[SegmentIndex].AddUninitialized(NumEntries);
-		}
-	}
-}
-
-void FAEPerTrackKeyLerpContext::CacheSegmentValues(const FAnimSequenceDecompressionContext& DecompContext, const FCompressedSegment& Segment, uint8 SegmentIndex)
-{
-	PerTrackStreamFlagOffsets[SegmentIndex] = Segment.ByteStreamOffset;
-
-	const int32 PerTrackFlagsSize = Align(DecompContext.NumTracks * DecompContext.NumStreamsPerTrack * sizeof(uint8), 4);
-	const int32 RangeBaseOffset = Segment.ByteStreamOffset + PerTrackFlagsSize;
-	const uint8* PerTrackStreamFlags = DecompContext.CompressedByteStream + Segment.ByteStreamOffset;
-
-	int32 KeyOffset = 0;
-	int32 RangeOffset = RangeBaseOffset;
-	int32 TotalNumAnimatedTrackStreams = 0;
-	for (int32 TrackIndex = 0; TrackIndex < DecompContext.NumTracks; ++TrackIndex)
-	{
-		const FTrivialTrackFlags TrivialTrackFlags(DecompContext.TrackFlags[TrackIndex]);
-
-		const int32 TranslationValueOffset = DecompContext.GetTranslationValueOffset(TrackIndex);
-		UniformKeyOffsets[SegmentIndex][TranslationValueOffset] = KeyOffset;
-		RangeOffsets[SegmentIndex][TranslationValueOffset] = RangeOffset;
-		NumAnimatedTrackStreams[SegmentIndex][TranslationValueOffset] = TotalNumAnimatedTrackStreams;
-
-		int32 BytesPerKey = 0;
-		if (!TrivialTrackFlags.IsTranslationTrivial())
-		{
-			const FPerTrackFlags TranslationTrackFlags(PerTrackStreamFlags[TranslationValueOffset]);
-			const uint8 Format = TranslationTrackFlags.GetFormat();
-			const uint8 FormatFlags = TranslationTrackFlags.GetFormatFlags();
-
-			int32 BytesPerRange;
-			FAnimationCompression_PerTrackUtils::GetByteSizesFromFormat(Format, FormatFlags, BytesPerKey, BytesPerRange);
-
-			if (TranslationTrackFlags.IsUniform())
-			{
-				KeyOffset += BytesPerKey;
-			}
-			else
-			{
-				TotalNumAnimatedTrackStreams++;
-			}
-
-			if (Format == ACF_IntervalFixed32NoW)
-			{
-				RangeOffset += BytesPerRange;
-			}
-		}
-
-		TrackStreamKeySizes[SegmentIndex][TranslationValueOffset] = static_cast<uint8>(BytesPerKey);
-
-		const int32 RotationValueOffset = DecompContext.GetRotationValueOffset(TrackIndex);
-		UniformKeyOffsets[SegmentIndex][RotationValueOffset] = KeyOffset;
-		RangeOffsets[SegmentIndex][RotationValueOffset] = RangeOffset;
-		NumAnimatedTrackStreams[SegmentIndex][RotationValueOffset] = TotalNumAnimatedTrackStreams;
-
-		BytesPerKey = 0;
-		if (!TrivialTrackFlags.IsRotationTrivial())
-		{
-			const FPerTrackFlags RotationTrackFlags(PerTrackStreamFlags[RotationValueOffset]);
-			const uint8 Format = RotationTrackFlags.GetFormat();
-			const uint8 FormatFlags = RotationTrackFlags.GetFormatFlags();
-
-			int32 BytesPerRange;
-			FAnimationCompression_PerTrackUtils::GetByteSizesFromFormat(Format, FormatFlags, BytesPerKey, BytesPerRange);
-
-			if (RotationTrackFlags.IsUniform())
-			{
-				KeyOffset += BytesPerKey;
-			}
-			else
-			{
-				TotalNumAnimatedTrackStreams++;
-			}
-
-			if (Format == ACF_IntervalFixed32NoW)
-			{
-				RangeOffset += BytesPerRange;
-			}
-		}
-
-		TrackStreamKeySizes[SegmentIndex][RotationValueOffset] = static_cast<uint8>(BytesPerKey);
-
-		if (DecompContext.bHasScale)
-		{
-			const int32 ScaleValueOffset = DecompContext.GetScaleValueOffset(TrackIndex);
-			UniformKeyOffsets[SegmentIndex][ScaleValueOffset] = KeyOffset;
-			RangeOffsets[SegmentIndex][ScaleValueOffset] = RangeOffset;
-			NumAnimatedTrackStreams[SegmentIndex][ScaleValueOffset] = TotalNumAnimatedTrackStreams;
-
-			BytesPerKey = 0;
-			if (!TrivialTrackFlags.IsScaleTrivial())
-			{
-				const FPerTrackFlags ScaleTrackFlags(PerTrackStreamFlags[ScaleValueOffset]);
-				const uint8 Format = ScaleTrackFlags.GetFormat();
-				const uint8 FormatFlags = ScaleTrackFlags.GetFormatFlags();
-
-				int32 BytesPerRange;
-				FAnimationCompression_PerTrackUtils::GetByteSizesFromFormat(Format, FormatFlags, BytesPerKey, BytesPerRange);
-
-				if (ScaleTrackFlags.IsUniform())
-				{
-					KeyOffset += BytesPerKey;
-				}
-				else
-				{
-					TotalNumAnimatedTrackStreams++;
-				}
-
-				if (Format == ACF_IntervalFixed32NoW)
-				{
-					RangeOffset += BytesPerRange;
-				}
-			}
-
-			TrackStreamKeySizes[SegmentIndex][ScaleValueOffset] = static_cast<uint8>(BytesPerKey);
-		}
-	}
-
-	const int32 SegmentUniformKeyFrameSize = KeyOffset;
-	UniformKeyFrameSize[SegmentIndex] = SegmentUniformKeyFrameSize;
-	RangeDataSize[SegmentIndex] = RangeOffset - RangeBaseOffset;
-
-	UniformDataOffsets[SegmentIndex] = Segment.ByteStreamOffset + PerTrackFlagsSize + RangeDataSize[SegmentIndex];
-
-	if (!DecompContext.bIsSorted)
-	{
-		// Variable linear
-		TimeMarkerSize[SegmentIndex] = Segment.NumFrames < 256 ? sizeof(uint8) : sizeof(uint16);
-
-		const int32 UniformDataSize = Align(SegmentUniformKeyFrameSize * Segment.NumFrames, 4);
-
-		OffsetNumKeysPairOffsets[SegmentIndex] = Segment.ByteStreamOffset + PerTrackFlagsSize + RangeDataSize[SegmentIndex] + UniformDataSize;
-	}
-}
-
-void FAEPerTrackKeyLerpContext::ResetSortedCache(const FAnimSequenceDecompressionContext& DecompContext)
-{
-	CachedKeys.Reset(DecompContext.NumTracks);
-	CachedKeys.AddZeroed(DecompContext.NumTracks);
-
-	const int32 PerTrackFlagsSize = Align(DecompContext.NumTracks * DecompContext.NumStreamsPerTrack * sizeof(uint8), 4);
-	const int32 UniformDataSize = Align(UniformKeyFrameSize[0] * DecompContext.Segment0->NumFrames, 4);
-
-	PackedSampleData = DecompContext.CompressedByteStream + DecompContext.Segment0->ByteStreamOffset + PerTrackFlagsSize + RangeDataSize[0] + UniformDataSize;
-	PreviousFrameIndex = 0;
-	CurrentSegmentIndex = 0;
-}
-
-static constexpr uint8 AEPerTrackKeyLerpContextCachedKeyStructOffsets[3][2] =
-{
-	{ offsetof(FAEPerTrackKeyLerpContext::FCachedKey, TransOffsets), offsetof(FAEPerTrackKeyLerpContext::FCachedKey, TransFrameIndices) },
-	{ offsetof(FAEPerTrackKeyLerpContext::FCachedKey, RotOffsets), offsetof(FAEPerTrackKeyLerpContext::FCachedKey, RotFrameIndices) },
-	{ offsetof(FAEPerTrackKeyLerpContext::FCachedKey, ScaleOffsets), offsetof(FAEPerTrackKeyLerpContext::FCachedKey, ScaleFrameIndices) },
-};
-
-void FAEPerTrackKeyLerpContext::AdvanceSortedCachedKeys(const FAnimSequenceDecompressionContext& DecompContext)
-{
-	while (true)
-	{
-		const uint8* SampleData = PackedSampleData;
-		const FSortedKeyHeader KeyHeader(SampleData);
-		if (KeyHeader.IsEndOfStream())
-		{
-			break;	// Reached the end of the stream
-		}
-		checkSlow(KeyHeader.TrackIndex < DecompContext.NumTracks);
-
-		const uint8 SampleType = KeyHeader.GetKeyType();
-		checkSlow(SampleType <= 2);
-
-		const int32 TimeDelta = KeyHeader.GetTimeDelta();
-		const int32 FrameIndex = PreviousFrameIndex + TimeDelta;
-
-		// Swap and update
-		FCachedKey& CachedKey = CachedKeys[KeyHeader.TrackIndex];
-		int32* DataOffset = reinterpret_cast<int32*>(reinterpret_cast<uint8*>(&CachedKey) + AEPerTrackKeyLerpContextCachedKeyStructOffsets[SampleType][0]);
-		int32* FrameIndicesOffset = reinterpret_cast<int32*>(reinterpret_cast<uint8*>(&CachedKey) + AEPerTrackKeyLerpContextCachedKeyStructOffsets[SampleType][1]);
-		if (FrameIndex > CurrentFrameIndex && FrameIndicesOffset[1] >= CurrentFrameIndex)
-		{
-			break;		// Reached a sample we don't need yet, stop for now
-		}
-
-		SampleData += KeyHeader.GetSize();
-
-		DataOffset[0] = DataOffset[1];
-		DataOffset[1] = static_cast<int32>(SampleData - DecompContext.CompressedByteStream);
-		FrameIndicesOffset[0] = FrameIndicesOffset[1];
-		FrameIndicesOffset[1] = FrameIndex;
-
-		const uint8 BytesPerKey = TrackStreamKeySizes[CurrentSegmentIndex][(KeyHeader.TrackIndex * DecompContext.NumStreamsPerTrack) + SampleType];
-
-		PreviousFrameIndex = FrameIndex;
-		SampleData += BytesPerKey;
-		PackedSampleData = SampleData;	// Update the pointer since we consumed the sample
-	}
-}
-
-void FAEPerTrackKeyLerpContext::Seek(const FAnimSequenceDecompressionContext& DecompContext, float SampleAtTime)
-{
-	const bool IsSegmentCacheStale0 = PreviousSegmentIndex[0] != DecompContext.SegmentIndex0;
-	const bool IsSegmentCacheStale1 = PreviousSegmentIndex[1] != DecompContext.SegmentIndex1;
-	if (IsSegmentCacheStale0 || IsSegmentCacheStale1)
-	{
-		if (IsSegmentCacheStale0 && PreviousSegmentIndex[1] == DecompContext.SegmentIndex0)
-		{
-			// Forward playback, the new segment 0 is our old segment 1, swap the data and refresh the segment 1
-			PerTrackStreamFlagOffsets[0] = PerTrackStreamFlagOffsets[1];
-			UniformKeyOffsets[0] = UniformKeyOffsets[1];
-			UniformKeyFrameSize[0] = UniformKeyFrameSize[1];
-			RangeOffsets[0] = RangeOffsets[1];
-			RangeDataSize[0] = RangeDataSize[1];
-			UniformDataOffsets[0] = UniformDataOffsets[1];
-
-			// Variable linear
-			TimeMarkerSize[0] = TimeMarkerSize[1];
-			OffsetNumKeysPairOffsets[0] = OffsetNumKeysPairOffsets[1];
-			NumAnimatedTrackStreams[0] = NumAnimatedTrackStreams[1];
-			TrackStreamKeySizes[0] = TrackStreamKeySizes[1];
-
-			if (IsSegmentCacheStale1)
-			{
-				CacheSegmentValues(DecompContext, *DecompContext.Segment1, 1);
-			}
-		}
-		else if (IsSegmentCacheStale1 && PreviousSegmentIndex[0] == DecompContext.SegmentIndex1)
-		{
-			// Backward playback, the new segment 1 is our old segment 0, swap the data and refresh the segment 0
-			PerTrackStreamFlagOffsets[1] = PerTrackStreamFlagOffsets[0];
-			UniformKeyOffsets[1] = UniformKeyOffsets[0];
-			UniformKeyFrameSize[1] = UniformKeyFrameSize[0];
-			RangeOffsets[1] = RangeOffsets[0];
-			RangeDataSize[1] = RangeDataSize[0];
-			UniformDataOffsets[1] = UniformDataOffsets[0];
-
-			// Variable linear
-			TimeMarkerSize[1] = TimeMarkerSize[0];
-			OffsetNumKeysPairOffsets[1] = OffsetNumKeysPairOffsets[0];
-			NumAnimatedTrackStreams[1] = NumAnimatedTrackStreams[0];
-			TrackStreamKeySizes[1] = TrackStreamKeySizes[0];
-
-			if (IsSegmentCacheStale0)
-			{
-				CacheSegmentValues(DecompContext, *DecompContext.Segment0, 0);
-			}
-		}
-		else
-		{
-			if (IsSegmentCacheStale0)
-			{
-				CacheSegmentValues(DecompContext, *DecompContext.Segment0, 0);
-			}
-
-			if (IsSegmentCacheStale1)
-			{
-				CacheSegmentValues(DecompContext, *DecompContext.Segment1, 1);
-			}
-		}
-
-		PreviousSegmentIndex[0] = DecompContext.SegmentIndex0;
-		PreviousSegmentIndex[1] = DecompContext.SegmentIndex1;
-	}
-
-	FramePos = DecompContext.RelativePos * float(DecompContext.AnimSeq->NumFrames - 1);
-
-	if (DecompContext.bIsSorted)
-	{
-		if (SampleAtTime < PreviousSampleAtTime)
-		{
-			// Seeking backwards is terribly slow because we start over from the start
-			ResetSortedCache(DecompContext);
-		}
-		else if (IsSegmentCacheStale0)
-		{
-			// We are seeking forward into a new segment, start over
-			ResetSortedCache(DecompContext);
-		}
-
-		SegmentStartFrame[0] = DecompContext.Segment0->StartFrame;
-		SegmentStartFrame[1] = DecompContext.Segment1->StartFrame;
-
-		if (DecompContext.NeedsTwoSegments)
-		{
-			CurrentFrameIndex = CurrentSegmentIndex == 0 ? DecompContext.SegmentKeyIndex0 : DecompContext.SegmentKeyIndex1;
-		}
-		else
-		{
-			CurrentFrameIndex = FMath::Max(DecompContext.SegmentKeyIndex1, 1);
-		}
-
-		AdvanceSortedCachedKeys(DecompContext);
-
-		if (DecompContext.NeedsTwoSegments && CurrentSegmentIndex == 0)
-		{
-			// Switch to our segment 1
-			const int32 PerTrackFlagsSize = Align(DecompContext.NumTracks * DecompContext.NumStreamsPerTrack * sizeof(uint8), 4);
-			const int32 UniformDataSize = Align(UniformKeyFrameSize[1] * DecompContext.Segment1->NumFrames, 4);
-
-			PackedSampleData = DecompContext.CompressedByteStream + DecompContext.Segment1->ByteStreamOffset + PerTrackFlagsSize + RangeDataSize[1] + UniformDataSize;
-			PreviousFrameIndex = 0;
-			CurrentFrameIndex = DecompContext.SegmentKeyIndex1;
-			CurrentSegmentIndex = 1;
-
-			AdvanceSortedCachedKeys(DecompContext);
-
-			// Any track that is variable in segment 0 but not in segment 1 needs to be manually rotated in the cache
-			const uint8* PerTrackStreamFlags0 = DecompContext.CompressedByteStream + DecompContext.Segment0->ByteStreamOffset;
-			const uint8* PerTrackStreamFlags1 = DecompContext.CompressedByteStream + DecompContext.Segment1->ByteStreamOffset;
-
-			for (int32 TrackIndex = 0; TrackIndex < DecompContext.NumTracks; ++TrackIndex)
-			{
-				const FTrivialTrackFlags TrivialTrackFlags(DecompContext.TrackFlags[TrackIndex]);
-
-				if (!TrivialTrackFlags.IsTranslationTrivial())
-				{
-					const int32 TranslationValueOffset = DecompContext.GetTranslationValueOffset(TrackIndex);
-					const FPerTrackFlags TranslationTrackFlags0(PerTrackStreamFlags0[TranslationValueOffset]);
-					const FPerTrackFlags TranslationTrackFlags1(PerTrackStreamFlags1[TranslationValueOffset]);
-					if (!TranslationTrackFlags0.IsUniform() && TranslationTrackFlags1.IsUniform())
-					{
-						FCachedKey& CachedKey = CachedKeys[TrackIndex];
-						CachedKey.TransFrameIndices[0] = CachedKey.TransFrameIndices[1];
-						CachedKey.TransOffsets[0] = CachedKey.TransOffsets[1];
-					}
-				}
-
-				if (!TrivialTrackFlags.IsRotationTrivial())
-				{
-					const int32 RotationValueOffset = DecompContext.GetRotationValueOffset(TrackIndex);
-					const FPerTrackFlags RotationTrackFlags0(PerTrackStreamFlags0[RotationValueOffset]);
-					const FPerTrackFlags RotationTrackFlags1(PerTrackStreamFlags1[RotationValueOffset]);
-					if (!RotationTrackFlags0.IsUniform() && RotationTrackFlags1.IsUniform())
-					{
-						FCachedKey& CachedKey = CachedKeys[TrackIndex];
-						CachedKey.RotFrameIndices[0] = CachedKey.RotFrameIndices[1];
-						CachedKey.RotOffsets[0] = CachedKey.RotOffsets[1];
-					}
-				}
-
-				if (DecompContext.bHasScale && !TrivialTrackFlags.IsScaleTrivial())
-				{
-					const int32 ScaleValueOffset = DecompContext.GetScaleValueOffset(TrackIndex);
-					const FPerTrackFlags ScaleTrackFlags0(PerTrackStreamFlags0[ScaleValueOffset]);
-					const FPerTrackFlags ScaleTrackFlags1(PerTrackStreamFlags1[ScaleValueOffset]);
-					if (!ScaleTrackFlags0.IsUniform() && ScaleTrackFlags1.IsUniform())
-					{
-						FCachedKey& CachedKey = CachedKeys[TrackIndex];
-						CachedKey.ScaleFrameIndices[0] = CachedKey.ScaleFrameIndices[1];
-						CachedKey.ScaleOffsets[0] = CachedKey.ScaleOffsets[1];
-					}
-				}
-			}
-		}
-
-		PreviousSampleAtTime = SampleAtTime;
-	}
-	else
-	{
-		const float SegmentFramePos = FramePos - float(DecompContext.Segment0->StartFrame);
-		SegmentRelativePos0 = SegmentFramePos / float(DecompContext.Segment0->NumFrames - 1);
-	}
-}
-
-void FAEPerTrackKeyLerpContext::GetRotation(FTransform& OutAtom, const FAnimSequenceDecompressionContext& DecompContext, int32 TrackIndex) const
-{
-	const int32 RotationValueOffset = DecompContext.GetRotationValueOffset(TrackIndex);
-
-	const uint8* PerTrackStreamFlags0 = DecompContext.CompressedByteStream + PerTrackStreamFlagOffsets[0];
-	const FPerTrackFlags RotationFlags0(PerTrackStreamFlags0[RotationValueOffset]);
-
-	const uint8 KeyFormat0 = RotationFlags0.GetFormat();
-	const uint8 FormatFlags0 = RotationFlags0.GetFormatFlags();
-
-	const int32 RangeOffset0 = RangeOffsets[0][RotationValueOffset];
-	const uint8* RangeData0 = DecompContext.CompressedByteStream + RangeOffset0;
-
-	if (DecompContext.NeedsTwoSegments)
-	{
-		FQuat Rotation0;
-		if (RotationFlags0.IsUniform())
-		{
-			const int32 KeyOffset0 = UniformKeyOffsets[0][RotationValueOffset];
-			const int32 FrameStartOffset0 = UniformKeyFrameSize[0] * DecompContext.SegmentKeyIndex0;
-
-			const uint8* KeyData0 = DecompContext.CompressedByteStream + UniformDataOffsets[0] + FrameStartOffset0 + KeyOffset0;
-
-			FAnimationCompression_PerTrackUtils::DecompressRotation<false>(KeyFormat0, FormatFlags0, Rotation0, RangeData0, KeyData0);
-		}
-		else if (DecompContext.bIsSorted)
-		{
-			const FCachedKey& CachedKey = CachedKeys[TrackIndex];
-
-			const uint8* KeyData0 = DecompContext.CompressedByteStream + CachedKey.RotOffsets[0];
-
-			FAnimationCompression_PerTrackUtils::DecompressRotation<false>(KeyFormat0, FormatFlags0, Rotation0, RangeData0, KeyData0);
-		}
-		else
-		{
-			const int32 NumTrackStreams0 = NumAnimatedTrackStreams[0][RotationValueOffset];
-
-			const uint8* OffsetNumKeysPairs0 = DecompContext.CompressedByteStream + OffsetNumKeysPairOffsets[0];
-			const uint8* OffsetNumKeysPair0 = OffsetNumKeysPairs0 + (OffsetNumKeysPairSize * NumTrackStreams0);
-			const uint16 NumKeys0 = *reinterpret_cast<const uint16*>(OffsetNumKeysPair0 + sizeof(uint32));
-
-			const uint32 KeysOffset0 = AnimationCompressionUtils::UnalignedRead<uint32>(OffsetNumKeysPair0);
-			const int32 TimeMarkersOffset0 = DecompContext.Segment0->ByteStreamOffset + KeysOffset0;
-
-			const int32 TrackDataOffset0 = Align(TimeMarkersOffset0 + (NumKeys0 * TimeMarkerSize[0]), 4);
-			const uint8 BytesPerKey0 = TrackStreamKeySizes[0][RotationValueOffset];
-			const int32 KeyDataOffset0 = TrackDataOffset0 + ((NumKeys0 - 1) * BytesPerKey0);	// We need the last key of segment 0 and the first key of segment 1
-			const uint8* KeyData0 = DecompContext.CompressedByteStream + KeyDataOffset0;
-
-			FAnimationCompression_PerTrackUtils::DecompressRotation<false>(KeyFormat0, FormatFlags0, Rotation0, RangeData0, KeyData0);
-		}
-
-		const uint8* PerTrackStreamFlags1 = DecompContext.CompressedByteStream + PerTrackStreamFlagOffsets[1];
-		const FPerTrackFlags RotationFlags1(PerTrackStreamFlags1[RotationValueOffset]);
-
-		const uint8 KeyFormat1 = RotationFlags1.GetFormat();
-		const uint8 FormatFlags1 = RotationFlags1.GetFormatFlags();
-
-		const int32 RangeOffset1 = RangeOffsets[1][RotationValueOffset];
-		const uint8* RangeData1 = DecompContext.CompressedByteStream + RangeOffset1;
-
-		FQuat Rotation1;
-		if (RotationFlags1.IsUniform())
-		{
-			const int32 KeyOffset1 = UniformKeyOffsets[1][RotationValueOffset];
-			const int32 FrameStartOffset1 = UniformKeyFrameSize[1] * DecompContext.SegmentKeyIndex1;
-
-			const uint8* KeyData1 = DecompContext.CompressedByteStream + UniformDataOffsets[1] + FrameStartOffset1 + KeyOffset1;
-
-			FAnimationCompression_PerTrackUtils::DecompressRotation<false>(KeyFormat1, FormatFlags1, Rotation1, RangeData1, KeyData1);
-		}
-		else if (DecompContext.bIsSorted)
-		{
-			const FCachedKey& CachedKey = CachedKeys[TrackIndex];
-
-			const uint8* KeyData1 = DecompContext.CompressedByteStream + CachedKey.RotOffsets[1];
-
-			FAnimationCompression_PerTrackUtils::DecompressRotation<false>(KeyFormat1, FormatFlags1, Rotation1, RangeData1, KeyData1);
-		}
-		else
-		{
-			const int32 NumTrackStreams1 = NumAnimatedTrackStreams[1][RotationValueOffset];
-
-			const uint8* OffsetNumKeysPairs1 = DecompContext.CompressedByteStream + OffsetNumKeysPairOffsets[1];
-			const uint8* OffsetNumKeysPair1 = OffsetNumKeysPairs1 + (OffsetNumKeysPairSize * NumTrackStreams1);
-			const uint16 NumKeys1 = *reinterpret_cast<const uint16*>(OffsetNumKeysPair1 + sizeof(uint32));
-
-			const uint32 KeysOffset1 = AnimationCompressionUtils::UnalignedRead<uint32>(OffsetNumKeysPair1);
-			const int32 TimeMarkersOffset1 = DecompContext.Segment1->ByteStreamOffset + KeysOffset1;
-
-			const int32 TrackDataOffset1 = Align(TimeMarkersOffset1 + (NumKeys1 * TimeMarkerSize[1]), 4);
-			// We need the last key of segment 0 and the first key of segment 1
-			const int32 KeyDataOffset1 = TrackDataOffset1;	// First key!
-			const uint8* KeyData1 = DecompContext.CompressedByteStream + KeyDataOffset1;
-
-			FAnimationCompression_PerTrackUtils::DecompressRotation<false>(KeyFormat1, FormatFlags1, Rotation1, RangeData1, KeyData1);
-		}
-
-		// Fast linear quaternion interpolation.
-		FQuat BlendedQuat = FQuat::FastLerp(Rotation0, Rotation1, DecompContext.KeyAlpha);
-		BlendedQuat.Normalize();
-
-		OutAtom.SetRotation(BlendedQuat);
-	}
-	else
-	{
-		if (RotationFlags0.IsUniform())
-		{
-			const int32 KeyOffset0 = UniformKeyOffsets[0][RotationValueOffset];
-			const int32 UniformKeyFrameSize0 = UniformKeyFrameSize[0];
-			const int32 FrameStartOffset0 = UniformKeyFrameSize0 * DecompContext.SegmentKeyIndex0;
-
-			const uint8* KeyData0 = DecompContext.CompressedByteStream + UniformDataOffsets[0] + FrameStartOffset0 + KeyOffset0;
-
-			FQuat Rotation;
-			FAnimationCompression_PerTrackUtils::DecompressRotation<false>(KeyFormat0, FormatFlags0, Rotation, RangeData0, KeyData0);
-
-			if (DecompContext.NeedsInterpolation)
-			{
-				const uint8* KeyData1 = KeyData0 + UniformKeyFrameSize0;
-
-				FQuat Rotation1;
-				FAnimationCompression_PerTrackUtils::DecompressRotation<false>(KeyFormat0, FormatFlags0, Rotation1, RangeData0, KeyData1);
-
-				// Fast linear quaternion interpolation.
-				FQuat BlendedQuat = FQuat::FastLerp(Rotation, Rotation1, DecompContext.KeyAlpha);
-				BlendedQuat.Normalize();
-				Rotation = BlendedQuat;
-			}
-
-			OutAtom.SetRotation(Rotation);
-		}
-		else if (DecompContext.bIsSorted)
-		{
-			const FCachedKey& CachedKey = CachedKeys[TrackIndex];
-
-			// compute the blend parameters for the keys we have found
-			const int32 FrameIndex0 = SegmentStartFrame[0] + CachedKey.RotFrameIndices[0];
-			const int32 FrameIndex1 = SegmentStartFrame[1] + CachedKey.RotFrameIndices[1];
-
-			const int32 Delta = FMath::Max(FrameIndex1 - FrameIndex0, 1);
-			const float Remainder = FramePos - float(FrameIndex0);
-			const float Alpha = DecompContext.AnimSeq->Interpolation == EAnimInterpolationType::Step ? 0.0f : (Remainder / float(Delta));
-
-			const uint8* KeyData0 = DecompContext.CompressedByteStream + CachedKey.RotOffsets[0];
-
-			FQuat Rotation0;
-			FAnimationCompression_PerTrackUtils::DecompressRotation<false>(KeyFormat0, FormatFlags0, Rotation0, RangeData0, KeyData0);
-
-			const uint8* KeyData1 = DecompContext.CompressedByteStream + CachedKey.RotOffsets[1];
-
-			FQuat Rotation1;
-			FAnimationCompression_PerTrackUtils::DecompressRotation<false>(KeyFormat0, FormatFlags0, Rotation1, RangeData0, KeyData1);
-
-			// Fast linear quaternion interpolation.
-			FQuat BlendedQuat = FQuat::FastLerp(Rotation0, Rotation1, Alpha);
-			BlendedQuat.Normalize();
-
-			OutAtom.SetRotation(BlendedQuat);
-		}
-		else
-		{
-			const int32 NumTrackStreams0 = NumAnimatedTrackStreams[0][RotationValueOffset];
-
-			const uint8* OffsetNumKeysPairs0 = DecompContext.CompressedByteStream + OffsetNumKeysPairOffsets[0];
-			const uint8* OffsetNumKeysPair0 = OffsetNumKeysPairs0 + (OffsetNumKeysPairSize * NumTrackStreams0);
-			const uint16 NumKeys0 = *reinterpret_cast<const uint16*>(OffsetNumKeysPair0 + sizeof(uint32));
-
-			const uint32 KeysOffset0 = AnimationCompressionUtils::UnalignedRead<uint32>(OffsetNumKeysPair0);
-			const int32 TimeMarkersOffset0 = DecompContext.Segment0->ByteStreamOffset + KeysOffset0;
-			const uint8* TimeMarkers0 = DecompContext.CompressedByteStream + TimeMarkersOffset0;
-
-			int32 FrameIndex0;
-			int32 FrameIndex1;
-			float Alpha = AnimEncoding::TimeToIndex(DecompContext, TimeMarkers0, NumKeys0, DecompContext.Segment0->NumFrames, TimeMarkerSize[0], SegmentRelativePos0, FrameIndex0, FrameIndex1);
-
-			const int32 TrackDataOffset0 = Align(TimeMarkersOffset0 + (NumKeys0 * TimeMarkerSize[0]), 4);
-			const uint8 BytesPerKey0 = TrackStreamKeySizes[0][RotationValueOffset];
-			const int32 KeyDataOffset0 = TrackDataOffset0 + (FrameIndex0 * BytesPerKey0);
-			const uint8* KeyData0 = DecompContext.CompressedByteStream + KeyDataOffset0;
-
-			FQuat Rotation;
-			FAnimationCompression_PerTrackUtils::DecompressRotation<false>(KeyFormat0, FormatFlags0, Rotation, RangeData0, KeyData0);
-
-			if (DecompContext.NeedsInterpolation)
-			{
-				const uint8* KeyData1 = KeyData0 + BytesPerKey0;
-
-				FQuat Rotation1;
-				FAnimationCompression_PerTrackUtils::DecompressRotation<false>(KeyFormat0, FormatFlags0, Rotation1, RangeData0, KeyData1);
-
-				// Fast linear quaternion interpolation.
-				FQuat BlendedQuat = FQuat::FastLerp(Rotation, Rotation1, Alpha);
-				BlendedQuat.Normalize();
-				Rotation = BlendedQuat;
-			}
-
-			OutAtom.SetRotation(Rotation);
-		}
-	}
-}
-
-void FAEPerTrackKeyLerpContext::GetTranslation(FTransform& OutAtom, const FAnimSequenceDecompressionContext& DecompContext, int32 TrackIndex) const
-{
-	const int32 TranslationValueOffset = DecompContext.GetTranslationValueOffset(TrackIndex);
-
-	const uint8* PerTrackStreamFlags0 = DecompContext.CompressedByteStream + PerTrackStreamFlagOffsets[0];
-	const FPerTrackFlags TranslationFlags0(PerTrackStreamFlags0[TranslationValueOffset]);
-
-	const uint8 KeyFormat0 = TranslationFlags0.GetFormat();
-	const uint8 FormatFlags0 = TranslationFlags0.GetFormatFlags();
-
-	const int32 RangeOffset0 = RangeOffsets[0][TranslationValueOffset];
-	const uint8* RangeData0 = DecompContext.CompressedByteStream + RangeOffset0;
-
-	if (DecompContext.NeedsTwoSegments)
-	{
-		FVector Translation0;
-		if (TranslationFlags0.IsUniform())
-		{
-			const int32 KeyOffset0 = UniformKeyOffsets[0][TranslationValueOffset];
-			const int32 FrameStartOffset0 = UniformKeyFrameSize[0] * DecompContext.SegmentKeyIndex0;
-
-			const uint8* KeyData0 = DecompContext.CompressedByteStream + UniformDataOffsets[0] + FrameStartOffset0 + KeyOffset0;
-
-			FAnimationCompression_PerTrackUtils::DecompressTranslation<false>(KeyFormat0, FormatFlags0, Translation0, RangeData0, KeyData0);
-		}
-		else if (DecompContext.bIsSorted)
-		{
-			const FCachedKey& CachedKey = CachedKeys[TrackIndex];
-
-			const uint8* KeyData0 = DecompContext.CompressedByteStream + CachedKey.TransOffsets[0];
-
-			FAnimationCompression_PerTrackUtils::DecompressTranslation<false>(KeyFormat0, FormatFlags0, Translation0, RangeData0, KeyData0);
-		}
-		else
-		{
-			const int32 NumTrackStreams0 = NumAnimatedTrackStreams[0][TranslationValueOffset];
-
-			const uint8* OffsetNumKeysPairs0 = DecompContext.CompressedByteStream + OffsetNumKeysPairOffsets[0];
-			const uint8* OffsetNumKeysPair0 = OffsetNumKeysPairs0 + (OffsetNumKeysPairSize * NumTrackStreams0);
-			const uint16 NumKeys0 = *reinterpret_cast<const uint16*>(OffsetNumKeysPair0 + sizeof(uint32));
-
-			const uint32 KeysOffset0 = AnimationCompressionUtils::UnalignedRead<uint32>(OffsetNumKeysPair0);
-			const int32 TimeMarkersOffset0 = DecompContext.Segment0->ByteStreamOffset + KeysOffset0;
-
-			const int32 TrackDataOffset0 = Align(TimeMarkersOffset0 + (NumKeys0 * TimeMarkerSize[0]), 4);
-			const uint8 BytesPerKey0 = TrackStreamKeySizes[0][TranslationValueOffset];
-			const int32 KeyDataOffset0 = TrackDataOffset0 + ((NumKeys0 - 1) * BytesPerKey0);	// We need the last key of segment 0 and the first key of segment 1
-			const uint8* KeyData0 = DecompContext.CompressedByteStream + KeyDataOffset0;
-
-			FAnimationCompression_PerTrackUtils::DecompressTranslation<false>(KeyFormat0, FormatFlags0, Translation0, RangeData0, KeyData0);
-		}
-
-		const uint8* PerTrackStreamFlags1 = DecompContext.CompressedByteStream + PerTrackStreamFlagOffsets[1];
-		const FPerTrackFlags TranslationFlags1(PerTrackStreamFlags1[TranslationValueOffset]);
-
-		const uint8 KeyFormat1 = TranslationFlags1.GetFormat();
-		const uint8 FormatFlags1 = TranslationFlags1.GetFormatFlags();
-
-		const int32 RangeOffset1 = RangeOffsets[1][TranslationValueOffset];
-		const uint8* RangeData1 = DecompContext.CompressedByteStream + RangeOffset1;
-
-		FVector Translation1;
-		if (TranslationFlags1.IsUniform())
-		{
-			const int32 KeyOffset1 = UniformKeyOffsets[1][TranslationValueOffset];
-			const int32 FrameStartOffset1 = UniformKeyFrameSize[1] * DecompContext.SegmentKeyIndex1;
-
-			const uint8* KeyData1 = DecompContext.CompressedByteStream + UniformDataOffsets[1] + FrameStartOffset1 + KeyOffset1;
-
-			FAnimationCompression_PerTrackUtils::DecompressTranslation<false>(KeyFormat1, FormatFlags1, Translation1, RangeData1, KeyData1);
-		}
-		else if (DecompContext.bIsSorted)
-		{
-			const FCachedKey& CachedKey = CachedKeys[TrackIndex];
-
-			const uint8* KeyData1 = DecompContext.CompressedByteStream + CachedKey.TransOffsets[1];
-
-			FAnimationCompression_PerTrackUtils::DecompressTranslation<false>(KeyFormat1, FormatFlags1, Translation1, RangeData1, KeyData1);
-		}
-		else
-		{
-			const int32 NumTrackStreams1 = NumAnimatedTrackStreams[1][TranslationValueOffset];
-
-			const uint8* OffsetNumKeysPairs1 = DecompContext.CompressedByteStream + OffsetNumKeysPairOffsets[1];
-			const uint8* OffsetNumKeysPair1 = OffsetNumKeysPairs1 + (OffsetNumKeysPairSize * NumTrackStreams1);
-			const uint16 NumKeys1 = *reinterpret_cast<const uint16*>(OffsetNumKeysPair1 + sizeof(uint32));
-
-			const uint32 KeysOffset1 = AnimationCompressionUtils::UnalignedRead<uint32>(OffsetNumKeysPair1);
-			const int32 TimeMarkersOffset1 = DecompContext.Segment1->ByteStreamOffset + KeysOffset1;
-
-			const int32 TrackDataOffset1 = Align(TimeMarkersOffset1 + (NumKeys1 * TimeMarkerSize[1]), 4);
-			// We need the last key of segment 0 and the first key of segment 1
-			const int32 KeyDataOffset1 = TrackDataOffset1;	// First key!
-			const uint8* KeyData1 = DecompContext.CompressedByteStream + KeyDataOffset1;
-
-			FAnimationCompression_PerTrackUtils::DecompressTranslation<false>(KeyFormat1, FormatFlags1, Translation1, RangeData1, KeyData1);
-		}
-
-		const FVector Translation = FMath::Lerp(Translation0, Translation1, DecompContext.KeyAlpha);
-		OutAtom.SetTranslation(Translation);
+	if (INTEL_ISPC)
+	{
+#if INTEL_ISPC
+		const FUECompressedAnimData& AnimData = static_cast<const FUECompressedAnimData&>(DecompContext.CompressedAnimData);
+
+		const TArrayView<int32> ScaleOffsets = AnimData.CompressedScaleOffsets.OffsetData;
+		const int32 StripSize = AnimData.CompressedScaleOffsets.StripSize;
+
+		ispc::GetPerTrackCompressionPoseScales(
+			(ispc::FTransform*)&Atoms[0],
+			(ispc::BoneTrackPair*)&DesiredPairs[0],
+			ScaleOffsets.GetData(),
+			StripSize,
+			AnimData.CompressedByteStream.GetData(),
+			AnimData.CompressedNumberOfFrames,
+			DecompContext.SequenceLength,
+			DecompContext.RelativePos,
+			(uint8)DecompContext.Interpolation,
+			PairCount);
+#endif
 	}
 	else
 	{
-		if (TranslationFlags0.IsUniform())
+		for( int32 PairIndex = 0; PairIndex < PairCount; ++PairIndex )
 		{
-			const int32 KeyOffset0 = UniformKeyOffsets[0][TranslationValueOffset];
-			const int32 UniformKeyFrameSize0 = UniformKeyFrameSize[0];
-			const int32 FrameStartOffset0 = UniformKeyFrameSize0 * DecompContext.SegmentKeyIndex0;
-
-			const uint8* KeyData0 = DecompContext.CompressedByteStream + UniformDataOffsets[0] + FrameStartOffset0 + KeyOffset0;
-
-			FVector Translation;
-			FAnimationCompression_PerTrackUtils::DecompressTranslation<false>(KeyFormat0, FormatFlags0, Translation, RangeData0, KeyData0);
-
-			if (DecompContext.NeedsInterpolation)
-			{
-				const uint8* KeyData1 = KeyData0 + UniformKeyFrameSize0;
-
-				FVector Translation1;
-				FAnimationCompression_PerTrackUtils::DecompressTranslation<false>(KeyFormat0, FormatFlags0, Translation1, RangeData0, KeyData1);
-
-				Translation = FMath::Lerp(Translation, Translation1, DecompContext.KeyAlpha);
-			}
-
-			OutAtom.SetTranslation(Translation);
-		}
-		else if (DecompContext.bIsSorted)
-		{
-			const FCachedKey& CachedKey = CachedKeys[TrackIndex];
-
-			// compute the blend parameters for the keys we have found
-			const int32 FrameIndex0 = SegmentStartFrame[0] + CachedKey.TransFrameIndices[0];
-			const int32 FrameIndex1 = SegmentStartFrame[1] + CachedKey.TransFrameIndices[1];
-
-			const int32 Delta = FMath::Max(FrameIndex1 - FrameIndex0, 1);
-			const float Remainder = FramePos - float(FrameIndex0);
-			const float Alpha = DecompContext.AnimSeq->Interpolation == EAnimInterpolationType::Step ? 0.0f : (Remainder / float(Delta));
-
-			const uint8* KeyData0 = DecompContext.CompressedByteStream + CachedKey.TransOffsets[0];
-
-			FVector Translation0;
-			FAnimationCompression_PerTrackUtils::DecompressTranslation<false>(KeyFormat0, FormatFlags0, Translation0, RangeData0, KeyData0);
-
-			const uint8* KeyData1 = DecompContext.CompressedByteStream + CachedKey.TransOffsets[1];
-
-			FVector Translation1;
-			FAnimationCompression_PerTrackUtils::DecompressTranslation<false>(KeyFormat0, FormatFlags0, Translation1, RangeData0, KeyData1);
-
-			FVector Translation = FMath::Lerp(Translation0, Translation1, Alpha);
-
-			OutAtom.SetTranslation(Translation);
-		}
-		else
-		{
-			const int32 NumTrackStreams0 = NumAnimatedTrackStreams[0][TranslationValueOffset];
-
-			const uint8* OffsetNumKeysPairs0 = DecompContext.CompressedByteStream + OffsetNumKeysPairOffsets[0];
-			const uint8* OffsetNumKeysPair0 = OffsetNumKeysPairs0 + (OffsetNumKeysPairSize * NumTrackStreams0);
-			const uint16 NumKeys0 = *reinterpret_cast<const uint16*>(OffsetNumKeysPair0 + sizeof(uint32));
-
-			const uint32 KeysOffset0 = AnimationCompressionUtils::UnalignedRead<uint32>(OffsetNumKeysPair0);
-			const int32 TimeMarkersOffset0 = DecompContext.Segment0->ByteStreamOffset + KeysOffset0;
-			const uint8* TimeMarkers0 = DecompContext.CompressedByteStream + TimeMarkersOffset0;
-
-			int32 FrameIndex0;
-			int32 FrameIndex1;
-			float Alpha = AnimEncoding::TimeToIndex(DecompContext, TimeMarkers0, NumKeys0, DecompContext.Segment0->NumFrames, TimeMarkerSize[0], SegmentRelativePos0, FrameIndex0, FrameIndex1);
-
-			const int32 TrackDataOffset0 = Align(TimeMarkersOffset0 + (NumKeys0 * TimeMarkerSize[0]), 4);
-			const uint8 BytesPerKey0 = TrackStreamKeySizes[0][TranslationValueOffset];
-			const int32 KeyDataOffset0 = TrackDataOffset0 + (FrameIndex0 * BytesPerKey0);
-			const uint8* KeyData0 = DecompContext.CompressedByteStream + KeyDataOffset0;
-
-			FVector Translation;
-			FAnimationCompression_PerTrackUtils::DecompressTranslation<false>(KeyFormat0, FormatFlags0, Translation, RangeData0, KeyData0);
-
-			if (DecompContext.NeedsInterpolation)
-			{
-				const uint8* KeyData1 = KeyData0 + BytesPerKey0;
-
-				FVector Translation1;
-				FAnimationCompression_PerTrackUtils::DecompressTranslation<false>(KeyFormat0, FormatFlags0, Translation1, RangeData0, KeyData1);
-
-				Translation = FMath::Lerp(Translation, Translation1, Alpha);
-			}
-
-			OutAtom.SetTranslation(Translation);
+			const BoneTrackPair& Pair = DesiredPairs[PairIndex];
+			const int32 TrackIndex = Pair.TrackIndex;
+			const int32 AtomIndex = Pair.AtomIndex;
+			FTransform& BoneAtom = Atoms[AtomIndex];
+	
+			GetBoneAtomScale(BoneAtom, DecompContext, TrackIndex);
 		}
 	}
 }
 
-void FAEPerTrackKeyLerpContext::GetScale(FTransform& OutAtom, const FAnimSequenceDecompressionContext& DecompContext, int32 TrackIndex) const
-{
-	const int32 ScaleValueOffset = DecompContext.GetScaleValueOffset(TrackIndex);
-
-	const uint8* PerTrackStreamFlags0 = DecompContext.CompressedByteStream + PerTrackStreamFlagOffsets[0];
-	const FPerTrackFlags ScaleFlags0(PerTrackStreamFlags0[ScaleValueOffset]);
-
-	const uint8 KeyFormat0 = ScaleFlags0.GetFormat();
-	const uint8 FormatFlags0 = ScaleFlags0.GetFormatFlags();
-
-	const int32 RangeOffset0 = RangeOffsets[0][ScaleValueOffset];
-	const uint8* RangeData0 = DecompContext.CompressedByteStream + RangeOffset0;
-
-	if (DecompContext.NeedsTwoSegments)
-	{
-		FVector Scale0;
-		if (ScaleFlags0.IsUniform())
-		{
-			const int32 KeyOffset0 = UniformKeyOffsets[0][ScaleValueOffset];
-			const int32 FrameStartOffset0 = UniformKeyFrameSize[0] * DecompContext.SegmentKeyIndex0;
-
-			const uint8* KeyData0 = DecompContext.CompressedByteStream + UniformDataOffsets[0] + FrameStartOffset0 + KeyOffset0;
-
-			FAnimationCompression_PerTrackUtils::DecompressScale<false>(KeyFormat0, FormatFlags0, Scale0, RangeData0, KeyData0);
-		}
-		else if (DecompContext.bIsSorted)
-		{
-			const FCachedKey& CachedKey = CachedKeys[TrackIndex];
-
-			const uint8* KeyData0 = DecompContext.CompressedByteStream + CachedKey.ScaleOffsets[0];
-
-			FAnimationCompression_PerTrackUtils::DecompressScale<false>(KeyFormat0, FormatFlags0, Scale0, RangeData0, KeyData0);
-		}
-		else
-		{
-			const int32 NumTrackStreams0 = NumAnimatedTrackStreams[0][ScaleValueOffset];
-
-			const uint8* OffsetNumKeysPairs0 = DecompContext.CompressedByteStream + OffsetNumKeysPairOffsets[0];
-			const uint8* OffsetNumKeysPair0 = OffsetNumKeysPairs0 + (OffsetNumKeysPairSize * NumTrackStreams0);
-			const uint16 NumKeys0 = *reinterpret_cast<const uint16*>(OffsetNumKeysPair0 + sizeof(uint32));
-
-			const uint32 KeysOffset0 = AnimationCompressionUtils::UnalignedRead<uint32>(OffsetNumKeysPair0);
-			const int32 TimeMarkersOffset0 = DecompContext.Segment0->ByteStreamOffset + KeysOffset0;
-
-			const int32 TrackDataOffset0 = Align(TimeMarkersOffset0 + (NumKeys0 * TimeMarkerSize[0]), 4);
-			const uint8 BytesPerKey0 = TrackStreamKeySizes[0][ScaleValueOffset];
-			const int32 KeyDataOffset0 = TrackDataOffset0 + ((NumKeys0 - 1) * BytesPerKey0);	// We need the last key of segment 0 and the first key of segment 1
-			const uint8* KeyData0 = DecompContext.CompressedByteStream + KeyDataOffset0;
-
-			FAnimationCompression_PerTrackUtils::DecompressScale<false>(KeyFormat0, FormatFlags0, Scale0, RangeData0, KeyData0);
-		}
-
-		const uint8* PerTrackStreamFlags1 = DecompContext.CompressedByteStream + PerTrackStreamFlagOffsets[1];
-		const FPerTrackFlags ScaleFlags1(PerTrackStreamFlags1[ScaleValueOffset]);
-
-		const uint8 KeyFormat1 = ScaleFlags1.GetFormat();
-		const uint8 FormatFlags1 = ScaleFlags1.GetFormatFlags();
-
-		const int32 RangeOffset1 = RangeOffsets[1][ScaleValueOffset];
-		const uint8* RangeData1 = DecompContext.CompressedByteStream + RangeOffset1;
-
-		FVector Scale1;
-		if (ScaleFlags1.IsUniform())
-		{
-			const int32 KeyOffset1 = UniformKeyOffsets[1][ScaleValueOffset];
-			const int32 FrameStartOffset1 = UniformKeyFrameSize[1] * DecompContext.SegmentKeyIndex1;
-
-			const uint8* KeyData1 = DecompContext.CompressedByteStream + UniformDataOffsets[1] + FrameStartOffset1 + KeyOffset1;
-
-			FAnimationCompression_PerTrackUtils::DecompressScale<false>(KeyFormat1, FormatFlags1, Scale1, RangeData1, KeyData1);
-		}
-		else if (DecompContext.bIsSorted)
-		{
-			const FCachedKey& CachedKey = CachedKeys[TrackIndex];
-
-			const uint8* KeyData1 = DecompContext.CompressedByteStream + CachedKey.ScaleOffsets[1];
-
-			FAnimationCompression_PerTrackUtils::DecompressScale<false>(KeyFormat1, FormatFlags1, Scale1, RangeData1, KeyData1);
-		}
-		else
-		{
-			const int32 NumTrackStreams1 = NumAnimatedTrackStreams[1][ScaleValueOffset];
-
-			const uint8* OffsetNumKeysPairs1 = DecompContext.CompressedByteStream + OffsetNumKeysPairOffsets[1];
-			const uint8* OffsetNumKeysPair1 = OffsetNumKeysPairs1 + (OffsetNumKeysPairSize * NumTrackStreams1);
-			const uint16 NumKeys1 = *reinterpret_cast<const uint16*>(OffsetNumKeysPair1 + sizeof(uint32));
-
-			const uint32 KeysOffset1 = AnimationCompressionUtils::UnalignedRead<uint32>(OffsetNumKeysPair1);
-			const int32 TimeMarkersOffset1 = DecompContext.Segment1->ByteStreamOffset + KeysOffset1;
-
-			const int32 TrackDataOffset1 = Align(TimeMarkersOffset1 + (NumKeys1 * TimeMarkerSize[1]), 4);
-			// We need the last key of segment 0 and the first key of segment 1
-			const int32 KeyDataOffset1 = TrackDataOffset1;	// First key!
-			const uint8* KeyData1 = DecompContext.CompressedByteStream + KeyDataOffset1;
-
-			FAnimationCompression_PerTrackUtils::DecompressScale<false>(KeyFormat1, FormatFlags1, Scale1, RangeData1, KeyData1);
-		}
-
-		const FVector Scale = FMath::Lerp(Scale0, Scale1, DecompContext.KeyAlpha);
-		OutAtom.SetScale3D(Scale);
-	}
-	else
-	{
-		if (ScaleFlags0.IsUniform())
-		{
-			const int32 KeyOffset0 = UniformKeyOffsets[0][ScaleValueOffset];
-			const int32 UniformKeyFrameSize0 = UniformKeyFrameSize[0];
-			const int32 FrameStartOffset0 = UniformKeyFrameSize0 * DecompContext.SegmentKeyIndex0;
-
-			const uint8* KeyData0 = DecompContext.CompressedByteStream + UniformDataOffsets[0] + FrameStartOffset0 + KeyOffset0;
-
-			FVector Scale;
-			FAnimationCompression_PerTrackUtils::DecompressScale<false>(KeyFormat0, FormatFlags0, Scale, RangeData0, KeyData0);
-
-			if (DecompContext.NeedsInterpolation)
-			{
-				const uint8* KeyData1 = KeyData0 + UniformKeyFrameSize0;
-
-				FVector Scale1;
-				FAnimationCompression_PerTrackUtils::DecompressScale<false>(KeyFormat0, FormatFlags0, Scale1, RangeData0, KeyData1);
-
-				Scale = FMath::Lerp(Scale, Scale1, DecompContext.KeyAlpha);
-			}
-
-			OutAtom.SetScale3D(Scale);
-		}
-		else if (DecompContext.bIsSorted)
-		{
-			const FCachedKey& CachedKey = CachedKeys[TrackIndex];
-
-			// compute the blend parameters for the keys we have found
-			const int32 FrameIndex0 = SegmentStartFrame[0] + CachedKey.ScaleFrameIndices[0];
-			const int32 FrameIndex1 = SegmentStartFrame[1] + CachedKey.ScaleFrameIndices[1];
-
-			const int32 Delta = FMath::Max(FrameIndex1 - FrameIndex0, 1);
-			const float Remainder = FramePos - float(FrameIndex0);
-			const float Alpha = DecompContext.AnimSeq->Interpolation == EAnimInterpolationType::Step ? 0.0f : (Remainder / float(Delta));
-
-			const uint8* KeyData0 = DecompContext.CompressedByteStream + CachedKey.ScaleOffsets[0];
-
-			FVector Scale0;
-			FAnimationCompression_PerTrackUtils::DecompressScale<false>(KeyFormat0, FormatFlags0, Scale0, RangeData0, KeyData0);
-
-			const uint8* KeyData1 = DecompContext.CompressedByteStream + CachedKey.ScaleOffsets[1];
-
-			FVector Scale1;
-			FAnimationCompression_PerTrackUtils::DecompressScale<false>(KeyFormat0, FormatFlags0, Scale1, RangeData0, KeyData1);
-
-			FVector Scale = FMath::Lerp(Scale0, Scale1, Alpha);
-
-			OutAtom.SetScale3D(Scale);
-		}
-		else
-		{
-			const int32 NumTrackStreams0 = NumAnimatedTrackStreams[0][ScaleValueOffset];
-
-			const uint8* OffsetNumKeysPairs0 = DecompContext.CompressedByteStream + OffsetNumKeysPairOffsets[0];
-			const uint8* OffsetNumKeysPair0 = OffsetNumKeysPairs0 + (OffsetNumKeysPairSize * NumTrackStreams0);
-			const uint16 NumKeys0 = *reinterpret_cast<const uint16*>(OffsetNumKeysPair0 + sizeof(uint32));
-
-			const uint32 KeysOffset0 = AnimationCompressionUtils::UnalignedRead<uint32>(OffsetNumKeysPair0);
-			const int32 TimeMarkersOffset0 = DecompContext.Segment0->ByteStreamOffset + KeysOffset0;
-			const uint8* TimeMarkers0 = DecompContext.CompressedByteStream + TimeMarkersOffset0;
-
-			int32 FrameIndex0;
-			int32 FrameIndex1;
-			float Alpha = AnimEncoding::TimeToIndex(DecompContext, TimeMarkers0, NumKeys0, DecompContext.Segment0->NumFrames, TimeMarkerSize[0], SegmentRelativePos0, FrameIndex0, FrameIndex1);
-
-			const int32 TrackDataOffset0 = Align(TimeMarkersOffset0 + (NumKeys0 * TimeMarkerSize[0]), 4);
-			const uint8 BytesPerKey0 = TrackStreamKeySizes[0][ScaleValueOffset];
-			const int32 KeyDataOffset0 = TrackDataOffset0 + (FrameIndex0 * BytesPerKey0);
-			const uint8* KeyData0 = DecompContext.CompressedByteStream + KeyDataOffset0;
-
-			FVector Scale;
-			FAnimationCompression_PerTrackUtils::DecompressScale<false>(KeyFormat0, FormatFlags0, Scale, RangeData0, KeyData0);
-
-			if (DecompContext.NeedsInterpolation)
-			{
-				const uint8* KeyData1 = KeyData0 + BytesPerKey0;
-
-				FVector Scale1;
-				FAnimationCompression_PerTrackUtils::DecompressScale<false>(KeyFormat0, FormatFlags0, Scale1, RangeData0, KeyData1);
-
-				Scale = FMath::Lerp(Scale, Scale1, Alpha);
-			}
-
-			OutAtom.SetScale3D(Scale);
-		}
-	}
-}
-#endif
\ No newline at end of file
+#endif // USE_ANIMATION_CODEC_BATCH_SOLVER
diff --git a/Engine/Source/Runtime/Engine/Private/Animation/AnimEncoding_VariableKeyLerp.cpp b/Engine/Source/Runtime/Engine/Private/Animation/AnimEncoding_VariableKeyLerp.cpp
--- a/Engine/Source/Runtime/Engine/Private/Animation/AnimEncoding_VariableKeyLerp.cpp
+++ b/Engine/Source/Runtime/Engine/Private/Animation/AnimEncoding_VariableKeyLerp.cpp
@@ -7,6 +7,9 @@
 #include "AnimEncoding_VariableKeyLerp.h"
 #include "Serialization/MemoryReader.h"
 #include "Serialization/MemoryWriter.h"
+#if INTEL_ISPC
+#include "AnimEncoding_VariableKeyLerp.ispc.generated.h"
+#endif
 
 /**
  * Handles the ByteSwap of compressed rotation data on import
@@ -200,169 +203,159 @@ void AEFVariableKeyLerpShared::ByteSwapScaleOut(
 	}
 }
 
-#if USE_SEGMENTING_CONTEXT
-void AEFVariableKeyLerpShared::CreateEncodingContext(FAnimSequenceDecompressionContext& DecompContext)
+#if USE_ANIMATION_CODEC_BATCH_SOLVER
+
+/**
+ * Decompress all requested rotation components from an Animation Sequence
+ *
+ * @param	Atoms			The FTransform array to fill in.
+ * @param	DesiredPairs	Array of requested bone information
+ * @param	DecompContext	The decompression context to use.
+ */
+template<int32 FORMAT>
+void AEFVariableKeyLerp<FORMAT>::GetPoseRotations(
+	TArrayView<FTransform>& Atoms,
+	const BoneTrackArray& DesiredPairs,
+	FAnimSequenceDecompressionContext& DecompContext)
 {
-	checkSlow(DecompContext.EncodingContext == nullptr);
-	const FAnimSequenceCompressionHeader* Header = reinterpret_cast<const FAnimSequenceCompressionHeader*>(DecompContext.CompressedByteStream);
-	if (Header->bIsSorted)
+	const int32 PairCount = DesiredPairs.Num();
+
+	if (INTEL_ISPC)
 	{
-		DecompContext.EncodingContext = new FAEVariableKeyLerpSortedContext(DecompContext);
+#if INTEL_ISPC
+		const FUECompressedAnimData& AnimData = static_cast<const FUECompressedAnimData&>(DecompContext.CompressedAnimData);
+
+		ispc::GetVariableKeyLerpPoseRotations(
+			(ispc::FTransform*)&Atoms[0],
+			(ispc::BoneTrackPair*)&DesiredPairs[0],
+			AnimData.CompressedTrackOffsets.GetData(),
+			AnimData.CompressedByteStream.GetData(),
+			AnimData.CompressedNumberOfFrames,
+			DecompContext.RelativePos,
+			(uint8)DecompContext.Interpolation,
+			FORMAT,
+			PairCount);
+#endif
 	}
 	else
 	{
-		DecompContext.EncodingContext = new FAEVariableKeyLerpLinearContext(DecompContext);
-	}
-}
-
-void AEFVariableKeyLerpShared::ReleaseEncodingContext(FAnimSequenceDecompressionContext& DecompContext)
-{
-	checkSlow(DecompContext.EncodingContext != nullptr);
-	delete DecompContext.EncodingContext;
-	DecompContext.EncodingContext = nullptr;
-}
-
-static constexpr uint8 AEVariableKeyLerpSortedContextCachedKeyStructOffsets[3][2] =
-{
-	{ offsetof(FAEVariableKeyLerpSortedContext::FCachedKey, TransOffsets), offsetof(FAEVariableKeyLerpSortedContext::FCachedKey, TransFrameIndices) },
-	{ offsetof(FAEVariableKeyLerpSortedContext::FCachedKey, RotOffsets), offsetof(FAEVariableKeyLerpSortedContext::FCachedKey, RotFrameIndices) },
-	{ offsetof(FAEVariableKeyLerpSortedContext::FCachedKey, ScaleOffsets), offsetof(FAEVariableKeyLerpSortedContext::FCachedKey, ScaleFrameIndices) },
-};
-
-static void AdvanceCachedKeys(const FAnimSequenceDecompressionContext& DecompContext, FAEVariableKeyLerpSortedContext& EncodingContext)
-{
-	const int32 SampleSizes[3] = { DecompContext.PackedTranslationSize0, DecompContext.PackedRotationSize0, DecompContext.PackedScaleSize0 };
-
-	while (true)
-	{
-		const uint8* PackedSampleData = EncodingContext.PackedSampleData;
-		const FSortedKeyHeader KeyHeader(PackedSampleData);
-		if (KeyHeader.IsEndOfStream())
+		for (int32 PairIndex = 0; PairIndex < PairCount; ++PairIndex)
 		{
-			break;	// Reached the end of the stream
-		}
-		checkSlow(KeyHeader.TrackIndex < DecompContext.NumTracks);
-
-		const uint8 SampleType = KeyHeader.GetKeyType();
-		checkSlow(SampleType <= 2);
-
-		const int32 TimeDelta = KeyHeader.GetTimeDelta();
-		const int32 FrameIndex = EncodingContext.PreviousFrameIndex + TimeDelta;
+			const BoneTrackPair& Pair = DesiredPairs[PairIndex];
+			const int32 TrackIndex = Pair.TrackIndex;
+			const int32 AtomIndex = Pair.AtomIndex;
+			FTransform& BoneAtom = Atoms[AtomIndex];
 
-		// Swap and update
-		FAEVariableKeyLerpSortedContext::FCachedKey& CachedKey = EncodingContext.CachedKeys[KeyHeader.TrackIndex];
-		int32* DataOffset = reinterpret_cast<int32*>(reinterpret_cast<uint8*>(&CachedKey) + AEVariableKeyLerpSortedContextCachedKeyStructOffsets[SampleType][0]);
-		int32* FrameIndicesOffset = reinterpret_cast<int32*>(reinterpret_cast<uint8*>(&CachedKey) + AEVariableKeyLerpSortedContextCachedKeyStructOffsets[SampleType][1]);
-		if (FrameIndex > EncodingContext.CurrentFrameIndex && FrameIndicesOffset[1] >= EncodingContext.CurrentFrameIndex)
-		{
-			break;		// Reached a sample we don't need yet, stop for now
+			// call the decoder directly (not through the vtable)
+			AEFVariableKeyLerp<FORMAT>::GetBoneAtomRotation(BoneAtom, DecompContext, TrackIndex);
 		}
-
-		PackedSampleData += KeyHeader.GetSize();
-
-		DataOffset[0] = DataOffset[1];
-		DataOffset[1] = static_cast<int32>(PackedSampleData - DecompContext.CompressedByteStream);
-		FrameIndicesOffset[0] = FrameIndicesOffset[1];
-		FrameIndicesOffset[1] = FrameIndex;
-
-		EncodingContext.PreviousFrameIndex = FrameIndex;
-		PackedSampleData += SampleSizes[SampleType];
-		EncodingContext.PackedSampleData = PackedSampleData;	// Update the pointer since we consumed the sample
 	}
 }
 
-FAEVariableKeyLerpSortedContext::FAEVariableKeyLerpSortedContext(const FAnimSequenceDecompressionContext& DecompContext)
-	: PreviousSampleAtTime(FLT_MAX)	// Very large to trigger a reset on the first Seek()
-{
-	CachedKeys.Empty(DecompContext.NumTracks);
-}
-
-static void Reset(const FAnimSequenceDecompressionContext& DecompContext, FAEVariableKeyLerpSortedContext& EncodingContext)
-{
-	EncodingContext.CachedKeys.Empty(DecompContext.NumTracks);
-	EncodingContext.CachedKeys.AddZeroed(DecompContext.NumTracks);
-
-	EncodingContext.PackedSampleData = DecompContext.CompressedByteStream + DecompContext.Segment0->ByteStreamOffset + DecompContext.RangeDataSize0;
-	EncodingContext.PreviousFrameIndex = 0;
-	EncodingContext.PreviousSegmentIndex = DecompContext.SegmentIndex0;
-}
-
-void FAEVariableKeyLerpSortedContext::Seek(const FAnimSequenceDecompressionContext& DecompContext, float SampleAtTime)
+/**
+ * Decompress all requested translation components from an Animation Sequence
+ *
+ * @param	Atoms			The FTransform array to fill in.
+ * @param	DesiredPairs	Array of requested bone information
+ * @param	DecompContext	The decompression context to use.
+ */
+template<int32 FORMAT>
+void AEFVariableKeyLerp<FORMAT>::GetPoseTranslations(
+	TArrayView<FTransform>& Atoms,
+	const BoneTrackArray& DesiredPairs,
+	FAnimSequenceDecompressionContext& DecompContext)
 {
-	if (SampleAtTime < PreviousSampleAtTime)
-	{
-		// Seeking backwards is terribly slow because we start over from the start
-		Reset(DecompContext, *this);
-	}
-	else if (PreviousSegmentIndex != DecompContext.SegmentIndex0)
-	{
-		// We are seeking forward into a new segment, start over
-		Reset(DecompContext, *this);
-	}
-
-	SegmentStartFrame0 = DecompContext.Segment0->StartFrame;
-	SegmentStartFrame1 = DecompContext.Segment1->StartFrame;
+	const int32 PairCount = DesiredPairs.Num();
 
-	FramePos = DecompContext.RelativePos * float(DecompContext.AnimSeq->NumFrames - 1);
-
-	if (DecompContext.NeedsTwoSegments)
+	if (INTEL_ISPC)
 	{
-		CurrentFrameIndex = PreviousSegmentIndex == 0 ? DecompContext.SegmentKeyIndex0 : DecompContext.SegmentKeyIndex1;
+#if INTEL_ISPC
+		const FUECompressedAnimData& AnimData = static_cast<const FUECompressedAnimData&>(DecompContext.CompressedAnimData);
+
+		ispc::GetVariableKeyLerpPoseTranslations(
+			(ispc::FTransform*)&Atoms[0],
+			(ispc::BoneTrackPair*)&DesiredPairs[0],
+			AnimData.CompressedTrackOffsets.GetData(),
+			AnimData.CompressedByteStream.GetData(),
+			AnimData.CompressedNumberOfFrames,
+			DecompContext.RelativePos,
+			(uint8)DecompContext.Interpolation,
+			FORMAT,
+			PairCount);
+#endif
 	}
 	else
 	{
-		CurrentFrameIndex = FMath::Max(DecompContext.SegmentKeyIndex1, 1);
-	}
-
-	AdvanceCachedKeys(DecompContext, *this);
-
-	if (DecompContext.NeedsTwoSegments)
-	{
-		PackedSampleData = DecompContext.CompressedByteStream + DecompContext.Segment1->ByteStreamOffset + DecompContext.RangeDataSize0;
-		PreviousFrameIndex = 0;
-		CurrentFrameIndex = DecompContext.SegmentKeyIndex1;
-		PreviousSegmentIndex = DecompContext.SegmentIndex1;
-
-		AdvanceCachedKeys(DecompContext, *this);
+		for (int32 PairIndex = 0; PairIndex < PairCount; ++PairIndex)
+		{
+			const BoneTrackPair& Pair = DesiredPairs[PairIndex];
+			const int32 TrackIndex = Pair.TrackIndex;
+			const int32 AtomIndex = Pair.AtomIndex;
+			FTransform& BoneAtom = Atoms[AtomIndex];
+	
+			// call the decoder directly (not through the vtable)
+			AEFVariableKeyLerp<FORMAT>::GetBoneAtomTranslation(BoneAtom, DecompContext, TrackIndex);
+		}
 	}
-
-	PreviousSampleAtTime = SampleAtTime;
 }
 
-FAEVariableKeyLerpLinearContext::FAEVariableKeyLerpLinearContext(const FAnimSequenceDecompressionContext& DecompContext)
+/**
+ * Decompress all requested Scale components from an Animation Sequence
+ *
+ * @param	Atoms			The FTransform array to fill in.
+ * @param	DesiredPairs	Array of requested bone information
+ * @param	DecompContext	The decompression context to use.
+ */
+template<int32 FORMAT>
+void AEFVariableKeyLerp<FORMAT>::GetPoseScales(
+	TArrayView<FTransform>& Atoms,
+	const BoneTrackArray& DesiredPairs,
+	FAnimSequenceDecompressionContext& DecompContext)
 {
-	NumAnimatedTrackStreams.Empty(DecompContext.NumTracks * DecompContext.NumStreamsPerTrack);
-	NumAnimatedTrackStreams.AddUninitialized(DecompContext.NumTracks * DecompContext.NumStreamsPerTrack);
+	const int32 PairCount = DesiredPairs.Num();
 
-	int32 TotalNumAnimatedTrackStreams = 0;
-	for (int32 TrackIndex = 0; TrackIndex < DecompContext.NumTracks; ++TrackIndex)
+	if (INTEL_ISPC)
 	{
-		const FTrivialTrackFlags TrackFlags(DecompContext.TrackFlags[TrackIndex]);
-
-		NumAnimatedTrackStreams[DecompContext.GetTranslationValueOffset(TrackIndex)] = TotalNumAnimatedTrackStreams;
-		TotalNumAnimatedTrackStreams += TrackFlags.IsTranslationTrivial() ? 0 : 1;
-
-		NumAnimatedTrackStreams[DecompContext.GetRotationValueOffset(TrackIndex)] = TotalNumAnimatedTrackStreams;
-		TotalNumAnimatedTrackStreams += TrackFlags.IsRotationTrivial() ? 0 : 1;
-
-		if (DecompContext.bHasScale)
+#if INTEL_ISPC
+		const FUECompressedAnimData& AnimData = static_cast<const FUECompressedAnimData&>(DecompContext.CompressedAnimData);
+
+		const TArrayView<int32> ScaleOffsets = AnimData.CompressedScaleOffsets.OffsetData;
+		const int32 StripSize = AnimData.CompressedScaleOffsets.StripSize;
+
+		ispc::GetVariableKeyLerpPoseScales(
+			(ispc::FTransform*)&Atoms[0],
+			(ispc::BoneTrackPair*)&DesiredPairs[0],
+			ScaleOffsets.GetData(),
+			StripSize,
+			AnimData.CompressedByteStream.GetData(),
+			AnimData.CompressedNumberOfFrames,
+			DecompContext.RelativePos,
+			(uint8)DecompContext.Interpolation,
+			FORMAT,
+			PairCount);
+#endif
+	}
+	else
+	{
+		for (int32 PairIndex = 0; PairIndex < PairCount; ++PairIndex)
 		{
-			NumAnimatedTrackStreams[DecompContext.GetScaleValueOffset(TrackIndex)] = TotalNumAnimatedTrackStreams;
-			TotalNumAnimatedTrackStreams += TrackFlags.IsScaleTrivial() ? 0 : 1;
+			const BoneTrackPair& Pair = DesiredPairs[PairIndex];
+			const int32 TrackIndex = Pair.TrackIndex;
+			const int32 AtomIndex = Pair.AtomIndex;
+			FTransform& BoneAtom = Atoms[AtomIndex];
+	
+			// call the decoder directly (not through the vtable)
+			AEFVariableKeyLerp<FORMAT>::GetBoneAtomScale(BoneAtom, DecompContext, TrackIndex);
 		}
 	}
 }
 
-void FAEVariableKeyLerpLinearContext::Seek(const FAnimSequenceDecompressionContext& DecompContext, float SampleAtTime)
-{
-	float FramePos = DecompContext.RelativePos * float(DecompContext.AnimSeq->NumFrames - 1);
-	float SegmentFramePos = FramePos - float(DecompContext.Segment0->StartFrame);
-	SegmentRelativePos0 = SegmentFramePos / float(DecompContext.Segment0->NumFrames - 1);
-
-	TimeMarkerSize[0] = DecompContext.Segment0->NumFrames < 256 ? sizeof(uint8) : sizeof(uint16);
-	TimeMarkerSize[1] = DecompContext.Segment1->NumFrames < 256 ? sizeof(uint8) : sizeof(uint16);
+template class AEFVariableKeyLerp<ACF_None>;
+template class AEFVariableKeyLerp<ACF_Float96NoW>;
+template class AEFVariableKeyLerp<ACF_Fixed48NoW>;
+template class AEFVariableKeyLerp<ACF_IntervalFixed32NoW>;
+template class AEFVariableKeyLerp<ACF_Fixed32NoW>;
+template class AEFVariableKeyLerp<ACF_Float32NoW>;
+template class AEFVariableKeyLerp<ACF_Identity>;
 
-	OffsetNumKeysPairs[0] = DecompContext.CompressedByteStream + DecompContext.Segment0->ByteStreamOffset + DecompContext.RangeDataSize0;
-	OffsetNumKeysPairs[1] = DecompContext.CompressedByteStream + DecompContext.Segment1->ByteStreamOffset + DecompContext.RangeDataSize0;
-}
 #endif
\ No newline at end of file
diff --git a/Engine/Source/Runtime/Engine/Private/Animation/AnimNode_ApplyMeshSpaceAdditive.cpp b/Engine/Source/Runtime/Engine/Private/Animation/AnimNode_ApplyMeshSpaceAdditive.cpp
--- a/Engine/Source/Runtime/Engine/Private/Animation/AnimNode_ApplyMeshSpaceAdditive.cpp
+++ b/Engine/Source/Runtime/Engine/Private/Animation/AnimNode_ApplyMeshSpaceAdditive.cpp
@@ -1,6 +1,7 @@
 // Copyright Epic Games, Inc. All Rights Reserved.
 
 #include "Animation/AnimNode_ApplyMeshSpaceAdditive.h"
+#include "Animation/AnimInstanceProxy.h"
 #include "AnimationRuntime.h"
 #include "Animation/AnimTrace.h"
 
@@ -9,40 +10,69 @@
 
 void FAnimNode_ApplyMeshSpaceAdditive::Initialize_AnyThread(const FAnimationInitializeContext& Context)
 {
+	DECLARE_SCOPE_HIERARCHICAL_COUNTER_ANIMNODE(Initialize_AnyThread)
 	FAnimNode_Base::Initialize_AnyThread(Context);
 
 	Base.Initialize(Context);
 	Additive.Initialize(Context);
+
+	AlphaBoolBlend.Reinitialize();
+	AlphaScaleBiasClamp.Reinitialize();
 }
 
 void FAnimNode_ApplyMeshSpaceAdditive::CacheBones_AnyThread(const FAnimationCacheBonesContext& Context)
 {
+	DECLARE_SCOPE_HIERARCHICAL_COUNTER_ANIMNODE(CacheBones_AnyThread)
 	Base.CacheBones(Context);
 	Additive.CacheBones(Context);
 }
 
 void FAnimNode_ApplyMeshSpaceAdditive::Update_AnyThread(const FAnimationUpdateContext& Context)
 {
+	DECLARE_SCOPE_HIERARCHICAL_COUNTER_ANIMNODE(Update_AnyThread)
+	QUICK_SCOPE_CYCLE_COUNTER(STAT_FAnimNode_ApplyMeshSpaceAdditive_Update);
+
 	Base.Update(Context);
 
+	if (!IsLODEnabled(Context.AnimInstanceProxy))
+	{
+		// Avoid doing work if we're not even going to be used.
+		return;
+	}
+
+	GetEvaluateGraphExposedInputs().Execute(Context);
+
 	ActualAlpha = 0.f;
-	if (IsLODEnabled(Context.AnimInstanceProxy))
+	switch (AlphaInputType)
 	{
-		// @note: If you derive this class, and if you have input that you rely on for base
-		// this is not going to work	
-		GetEvaluateGraphExposedInputs().Execute(Context);
-		ActualAlpha = AlphaScaleBias.ApplyTo(Alpha);
-		if (FAnimWeight::IsRelevant(ActualAlpha))
+	case EAnimAlphaInputType::Float:
+		ActualAlpha = AlphaScaleBiasClamp.ApplyTo(AlphaScaleBias.ApplyTo(Alpha), Context.GetDeltaTime());
+		break;
+	case EAnimAlphaInputType::Bool:
+		ActualAlpha = AlphaBoolBlend.ApplyTo(bAlphaBoolEnabled, Context.GetDeltaTime());
+		break;
+	case EAnimAlphaInputType::Curve:
+		if (UAnimInstance* AnimInstance = Cast<UAnimInstance>(Context.AnimInstanceProxy->GetAnimInstanceObject()))
 		{
-			Additive.Update(Context.FractionalWeight(ActualAlpha));
+			ActualAlpha = AlphaScaleBiasClamp.ApplyTo(AnimInstance->GetCurveValue(AlphaCurveName), Context.GetDeltaTime());
 		}
+		break;
 	}
 
+	ActualAlpha = FMath::Clamp(ActualAlpha, 0.0f, 1.0f);
+
+	if (FAnimWeight::IsRelevant(ActualAlpha))
+	{
+		Additive.Update(Context.FractionalWeight(ActualAlpha));
+	}
+	
 	TRACE_ANIM_NODE_VALUE(Context, TEXT("Alpha"), ActualAlpha);
 }
 
 void FAnimNode_ApplyMeshSpaceAdditive::Evaluate_AnyThread(FPoseContext& Output)
 {
+	DECLARE_SCOPE_HIERARCHICAL_COUNTER_ANIMNODE(Evaluate_AnyThread)
+
 	//@TODO: Could evaluate Base into Output and save a copy
 	if (FAnimWeight::IsRelevant(ActualAlpha))
 	{
@@ -62,7 +92,10 @@ void FAnimNode_ApplyMeshSpaceAdditive::Evaluate_AnyThread(FPoseContext& Output)
 }
 
 FAnimNode_ApplyMeshSpaceAdditive::FAnimNode_ApplyMeshSpaceAdditive()
-	: Alpha(1.0f)
+	: AlphaInputType(EAnimAlphaInputType::Float)
+	, Alpha(1.0f)
+	, bAlphaBoolEnabled(true)
+	, AlphaCurveName(NAME_None)
 	, LODThreshold(INDEX_NONE)
 	, ActualAlpha(0.f)
 {
