COPY from //Dev-Anim to //Dev-Main

diff --git a/Engine/Source/Runtime/Engine/Private/Animation/AnimSequence.cpp b/Engine/Source/Runtime/Engine/Private/Animation/AnimSequence.cpp
--- a/Engine/Source/Runtime/Engine/Private/Animation/AnimSequence.cpp
+++ b/Engine/Source/Runtime/Engine/Private/Animation/AnimSequence.cpp
@@ -24,6 +24,8 @@
 #include "Animation/BlendSpaceBase.h"
 #include "Animation/Rig.h"
 #include "Animation/AnimationSettings.h"
+#include "Animation/AnimBoneCompressionCodec.h"
+#include "Animation/AnimBoneCompressionSettings.h"
 #include "Animation/AnimCurveCompressionCodec.h"
 #include "Animation/AnimCurveCompressionSettings.h"
 #include "EditorFramework/AssetImportData.h"
@@ -32,6 +34,7 @@
 #include "DerivedDataCacheInterface.h"
 #include "Interfaces/ITargetPlatform.h"
 #include "Animation/AnimCompressionDerivedData.h"
+#include "Animation/AnimCompressionDerivedDataPublic.h"
 #include "UObject/UObjectThreadContext.h"
 #include "Animation/AnimNotifies/AnimNotifyState.h"
 #include "Widgets/Notifications/SNotificationList.h"
@@ -204,7 +207,7 @@ void OnCVarsChanged()
 
 FAutoConsoleVariableSink AnimationCVarSink(FConsoleCommandDelegate::CreateStatic(&OnCVarsChanged));
 
-FString GetAnimSequenceSpecificCacheKeySuffix(const UAnimSequence& Seq, bool bPerformStripping, float AltErrorThreshold)
+FString GetAnimSequenceSpecificCacheKeySuffix(const UAnimSequence& Seq, bool bPerformStripping, float CompressionErrorThresholdScale)
 {
 	//Make up our content key consisting of:
 	//	* Global animation compression version
@@ -222,9 +225,9 @@ FString GetAnimSequenceSpecificCacheKeySuffix(const UAnimSequence& Seq, bool bPe
 
 	FArcToHexString ArcToHexString;
 
-	ArcToHexString.Ar << AltErrorThreshold;
+	ArcToHexString.Ar << CompressionErrorThresholdScale;
 	ArcToHexString.Ar << bPerformStripping;
-	Seq.CompressionScheme->PopulateDDCKeyArchive(ArcToHexString.Ar);
+	Seq.BoneCompressionSettings->PopulateDDCKey(ArcToHexString.Ar);
 	Seq.CurveCompressionSettings->PopulateDDCKey(ArcToHexString.Ar);
 
 	FString Ret = FString::Printf(TEXT("%i_%s%s%s_%c%c%i_%s_%s"),
@@ -547,6 +550,9 @@ UAnimSequence::UAnimSequence(const FObjectInitializer& ObjectInitializer)
 	, bUseNormalizedRootMotionScale(true)
 	, bRootMotionSettingsCopiedFromMontage(false)
 	, bUseRawDataOnly(!FPlatformProperties::RequiresCookedData())
+#if WITH_EDITOR
+	, bCompressionInProgress(false)
+#endif
 {
 	RateScale = 1.0;
 
@@ -592,6 +598,7 @@ void UAnimSequence::AddReferencedObjects(UObject* This, FReferenceCollector& Col
 	Super::AddReferencedObjects(This, Collector);
 
 	UAnimSequence* AnimSeq = CastChecked<UAnimSequence>(This);
+	Collector.AddReferencedObject(AnimSeq->CompressedData.BoneCompressionCodec);
 	Collector.AddReferencedObject(AnimSeq->CompressedData.CurveCompressionCodec);
 }
 
@@ -628,7 +635,7 @@ int32 UAnimSequence::GetApproxRawSize() const
 
 int32 UAnimSequence::GetApproxBoneCompressedSize() const
 {
-	return CompressedData.CompressedDataStructure.GetApproxBoneCompressedSize();
+	return CompressedData.CompressedDataStructure != nullptr ? CompressedData.CompressedDataStructure->GetApproxCompressedSize() : 0;
 }
 
 int32 UAnimSequence::GetApproxCompressedSize() const
@@ -764,11 +771,6 @@ void UAnimSequence::Serialize(FArchive& Ar)
 			}
 		}
 
-		if (bIsDuplicating)
-		{
-			Ar << bCompressionInProgress;
-		}
-
 		if (bSerializeCompressedData)
 		{
 			SerializeCompressedData(Ar,false);
@@ -826,12 +828,26 @@ void UAnimSequence::SortSyncMarkers()
 void UAnimSequence::GetPreloadDependencies(TArray<UObject*>& OutDeps)
 {
 	Super::GetPreloadDependencies(OutDeps);
-	OutDeps.Add(CurveCompressionSettings);
+
+	// We preload the compression settings because we need them loaded during Serialize to lookup the proper codec
+	// which is stored as a path/string.
+	if (CurveCompressionSettings != nullptr)
+	{
+		OutDeps.Add(CurveCompressionSettings);
+	}
+
+	if (BoneCompressionSettings != nullptr)
+	{
+		OutDeps.Add(BoneCompressionSettings);
+	}
 }
 
 void UAnimSequence::PreSave(const class ITargetPlatform* TargetPlatform)
 {
 #if WITH_EDITOR
+	// Could already be compressing
+	WaitOnExistingCompression();
+
 	// we have to bake it if it's not baked
 	if (DoesNeedRebake())
 	{
@@ -860,6 +876,7 @@ void UAnimSequence::PreSave(const class ITargetPlatform* TargetPlatform)
 		RequestAnimCompression(Params);
 	}
 
+	WaitOnExistingCompression(); // Wait on updated data
 #endif
 
 	Super::PreSave(TargetPlatform);
@@ -916,9 +933,9 @@ void UAnimSequence::PostLoad()
 		bUseRawDataOnly = true;
 	}
 
-	if (bUseRawDataOnly && !bCompressionInProgress)
+	if (bUseRawDataOnly)
 	{
-		RequestSyncAnimRecompression();
+		RequestAnimCompression(FRequestAnimCompressionParams(true, false, false));
 	}
 #endif
 
@@ -953,7 +970,11 @@ void UAnimSequence::PostLoad()
 		SequenceLength = MINIMUM_ANIMATION_LENGTH;
 	}
 	// Raw data exists, but missing compress animation data
-	else if( !bCompressionInProgress && GetSkeleton() && !IsCompressedDataValid())
+	else if( GetSkeleton() && !IsCompressedDataValid()
+#if WITH_EDITOR
+		&& !bCompressionInProgress
+#endif
+		)
 	{
 		UE_LOG(LogAnimation, Fatal, TEXT("No animation compression exists for sequence %s (%s)"), *GetName(), (GetOuter() ? *GetOuter()->GetFullName() : *GetFullName()) );
 	}
@@ -1108,10 +1129,13 @@ void UAnimSequence::VerifyTrackMap(USkeleton* MySkeleton)
 #endif // WITH_EDITOR
 void UAnimSequence::BeginDestroy()
 {
+	// Could already be compressing
+	WaitOnExistingCompression(false);
+
 	Super::BeginDestroy();
 
-	CompressedData.CompressedDataStructure.TranslationCodec = nullptr;
-	CompressedData.CompressedDataStructure.RotationCodec = nullptr;
+	ClearCompressedCurveData();
+	ClearCompressedBoneData();
 }
 
 #if WITH_EDITOR
@@ -1174,10 +1198,17 @@ void UAnimSequence::PostEditChangeProperty(FPropertyChangedEvent& PropertyChange
 		PostProcessSequence();
 	}
 
-	FProperty* Property = PropertyChangedEvent.Property;
-	if (Property != nullptr && Property->GetFName() == GET_MEMBER_NAME_CHECKED(UAnimSequence, CurveCompressionSettings))
+	if (PropertyChangedEvent.Property != nullptr)
 	{
-		RequestSyncAnimRecompression(false);
+		if (PropertyChangedEvent.Property->GetFName() == GET_MEMBER_NAME_CHECKED(UAnimSequence, CurveCompressionSettings))
+		{
+			RequestSyncAnimRecompression(false);
+		}
+
+		if (PropertyChangedEvent.Property->GetFName() == GET_MEMBER_NAME_CHECKED(UAnimSequence, BoneCompressionSettings))
+		{
+			RequestSyncAnimRecompression(false);
+		}
 	}
 }
 
@@ -1218,9 +1249,9 @@ void UAnimSequence::GetBoneTransform(FTransform& OutAtom, int32 TrackIndex, floa
 	// If the caller didn't request that raw animation data be used . . .
 	if ( !bUseRawData && IsCompressedDataValid() )
 	{
-		FAnimSequenceDecompressionContext DecompContext(SequenceLength, Interpolation, GetFName(), CompressedData.CompressedDataStructure);
+		FAnimSequenceDecompressionContext DecompContext(SequenceLength, Interpolation, GetFName(), *CompressedData.CompressedDataStructure);
 		DecompContext.Seek(Time);
-		AnimationFormat_GetBoneAtom( OutAtom, DecompContext, TrackIndex );
+		CompressedData.BoneCompressionCodec->DecompressBone(DecompContext, TrackIndex, OutAtom);
 		return;
 	}
 
@@ -1232,7 +1263,7 @@ void UAnimSequence::GetBoneTransform(FTransform& OutAtom, int32 TrackIndex, FAni
 	// If the caller didn't request that raw animation data be used . . .
 	if (!bUseRawData && IsCompressedDataValid())
 	{
-		AnimationFormat_GetBoneAtom(OutAtom, DecompContext, TrackIndex);
+		CompressedData.BoneCompressionCodec->DecompressBone(DecompContext, TrackIndex, OutAtom);
 		return;
 	}
 
@@ -1751,12 +1782,6 @@ void UAnimSequence::RetargetBoneTransform(FTransform& BoneTransform, const int32
 }
 
 #if WITH_EDITOR
-float UAnimSequence::GetAltCompressionErrorThreshold() const
-{
-	// CompressionErrorThresholdScale represents a world scale we will play back the animation at
-	return FAnimationUtils::GetAlternativeCompressionThreshold() / CompressionErrorThresholdScale;
-}
-
 /** Utility function to crop data from a RawAnimSequenceTrack */
 static int32 CropRawTrack(FRawAnimSequenceTrack& RawTrack, int32 StartKey, int32 NumKeys, int32 TotalNumOfFrames)
 {
@@ -2103,9 +2128,43 @@ void UAnimSequence::FlipRotationWForNonRoot(USkeletalMesh * SkelMesh)
 }
 #endif 
 
+#if WITH_EDITOR
+bool UAnimSequence::ShouldPerformStripping(const bool bPerformFrameStripping, const bool bPerformStrippingOnOddFramedAnims) const
+{
+	const bool bShouldPerformFrameStripping = bPerformFrameStripping && bAllowFrameStripping;
+
+	// Can only do stripping on animations that have an even number of frames once the end frame is removed)
+	const bool bIsEvenFramed = ((NumFrames - 1) % 2) == 0;
+	const bool bIsValidForStripping = bIsEvenFramed || bPerformStrippingOnOddFramedAnims;
+
+	const bool bStripCandidate = (NumFrames > 10) && bIsValidForStripping;
+
+	return bStripCandidate && bShouldPerformFrameStripping;
+}
+
+FString UAnimSequence::GetDDCCacheKeySuffix(const bool bPerformStripping) const
+{
+	return GetAnimSequenceSpecificCacheKeySuffix(*this, bPerformStripping, CompressionErrorThresholdScale);
+}
+#endif
+
+void UAnimSequence::WaitOnExistingCompression(const bool bCancelIfNotStarted)
+{
+#if WITH_EDITOR
+	check(IsInGameThread());
+	if (bCompressionInProgress)
+	{
+		double WaitTime = FPlatformTime::Seconds();
+		FAsyncCompressedAnimationsManagement::Get().WaitOnExistingCompression(this, bCancelIfNotStarted);
+		bCompressionInProgress = false;
+	}
+#endif
+}
+
 void UAnimSequence::RequestAnimCompression(FRequestAnimCompressionParams Params)
 {
 #if WITH_EDITOR
+	check(IsInGameThread());
 	USkeleton* CurrentSkeleton = GetSkeleton();
 	if (CurrentSkeleton == nullptr)
 	{
@@ -2113,20 +2172,16 @@ void UAnimSequence::RequestAnimCompression(FRequestAnimCompressionParams Params)
 		return;
 	}
 
-	if (GetOutermost() == GetTransientPackage())
-	{
-		bUseRawDataOnly = true;
-		return; // Skip transient animations, they are most likely the leftovers of previous compression attempts.
-	}
-
 	if (FPlatformProperties::RequiresCookedData())
 	{
 		return;
 	}
 
-	if (!CompressionScheme)
+	WaitOnExistingCompression(false);
+
+	if (BoneCompressionSettings == nullptr || !BoneCompressionSettings->AreSettingsValid())
 	{
-		CompressionScheme = FAnimationUtils::GetDefaultAnimationCompressionAlgorithm();
+		BoneCompressionSettings = FAnimationUtils::GetDefaultAnimationBoneCompressionSettings();
 	}
 
 	if (CurveCompressionSettings == nullptr || !CurveCompressionSettings->AreSettingsValid())
@@ -2134,97 +2189,129 @@ void UAnimSequence::RequestAnimCompression(FRequestAnimCompressionParams Params)
 		CurveCompressionSettings = FAnimationUtils::GetDefaultAnimationCurveCompressionSettings();
 	}
 
+	// Make sure all our required dependencies are loaded
+	FAnimationUtils::EnsureAnimSequenceLoaded(*this);
+
 	if (!RawDataGuid.IsValid())
 	{
 		RawDataGuid = GenerateGuidFromRawData();
 	}
 
-	Params.bAsyncCompression = false; //Just get sync working first
 	bUseRawDataOnly = true;
 
-	TGuardValue<bool> CompressGuard(bCompressionInProgress, true);
+	check(!bCompressionInProgress);
+	bCompressionInProgress = true;
 
 	// Need to make sure this is up to date.
 	VerifyCurveNames<FFloatCurve>(*CurrentSkeleton, USkeleton::AnimCurveMappingName, RawCurveData.FloatCurves);
 	VerifyTrackMap(CurrentSkeleton);
 
-	if (Params.bAsyncCompression)
-	{
-	}
-	else
-	{
-		const bool bPerformFrameStripping = Params.bPerformFrameStripping && bAllowFrameStripping;
-		const bool bPerformStrippingOnOddFramedAnims = Params.bPerformFrameStrippingOnOddNumberedFrames;
+	Params.CompressContext->GatherPreCompressionStats(GetName(), GetApproxRawSize(), GetApproxCompressedSize());
 
-		// Can only do stripping on animations that have an even number of frames once the end frame is removed)
-		const bool bIsEvenFramed = ((NumFrames - 1) % 2) == 0;
-		const bool bIsValidForStripping = bIsEvenFramed || bPerformStrippingOnOddFramedAnims;
+	const double CompressionStartTime = FPlatformTime::Seconds();
 
-		const bool bStripCandidate = (NumFrames > 10) && bIsValidForStripping;
+	const bool bPerformStripping = ShouldPerformStripping(Params.bPerformFrameStripping, Params.bPerformFrameStrippingOnOddNumberedFrames);
+	const FString AssetDDCKey = GetDDCCacheKeySuffix(bPerformStripping);
 
-		const bool bPerformStripping = bStripCandidate && bPerformFrameStripping;
+	bool bCompressedDataFromDDC = false;
 
-		const int32 PreviousCompressionSize = GetApproxCompressedSize();
+	TArray<uint8> OutData;
+	{
+		FDerivedDataAnimationCompression* AnimCompressor = new FDerivedDataAnimationCompression(TEXT("AnimSeq"), AssetDDCKey, Params.CompressContext);
 
-		const float AltCompressionErrorThreshold = GetAltCompressionErrorThreshold();
+		const FString FinalDDCKey = FDerivedDataCacheInterface::BuildCacheKey(AnimCompressor->GetPluginName(), AnimCompressor->GetVersionString(), *AnimCompressor->GetPluginSpecificCacheKeySuffix());
 
-		FString AssetDDCKey = GetAnimSequenceSpecificCacheKeySuffix(*this, bPerformStripping, AltCompressionErrorThreshold);
+		// For debugging DDC/Compression issues		
+		const bool bSkipDDC = false;
 
-		TArray<uint8> OutData;
+		if (bSkipDDC || !GetDerivedDataCacheRef().GetSynchronous(*FinalDDCKey, OutData))
 		{
-			bool bNeedToCleanUpAnimCompressor = true;
-			FDerivedDataAnimationCompression* AnimCompressor = new FDerivedDataAnimationCompression(TEXT("AnimSeq"), AssetDDCKey, Params.CompressContext, PreviousCompressionSize);
+			// Data does not exist, need to build it.
+			// Filter RAW data to get rid of mismatched tracks (translation/rotation data with a different number of keys than there are frames)
+			// No trivial key removal is done at this point (impossible error metrics of -1), since all of the techniques will perform it themselves
+			CompressRawAnimData(-1.0f, -1.0f);
 
-			const FString FinalDDCKey = FDerivedDataCacheInterface::BuildCacheKey(AnimCompressor->GetPluginName(), AnimCompressor->GetVersionString(), *AnimCompressor->GetPluginSpecificCacheKeySuffix());
+			TSharedRef<FCompressibleAnimData> CompressibleData = MakeShared<FCompressibleAnimData>(this, bPerformStripping);
+			AnimCompressor->SetCompressibleData(CompressibleData);
 
-			// For debugging DDC/Compression issues		
-			const bool bSkipDDC = false;
-
-			if (bSkipDDC || !GetDerivedDataCacheRef().GetSynchronous(*FinalDDCKey, OutData))
+			if (bSkipDDC || (CompressCommandletVersion == INDEX_NONE))
 			{
-				// Data does not exist, need to build it.
-				// Filter RAW data to get rid of mismatched tracks (translation/rotation data with a different number of keys than there are frames)
-				// No trivial key removal is done at this point (impossible error metrics of -1), since all of the techniques will perform it themselves
-				CompressRawAnimData(-1.0f, -1.0f);
-
-				TSharedRef<FCompressibleAnimData> CompressibleData = MakeShared<FCompressibleAnimData>(this, bPerformStripping, AltCompressionErrorThreshold);
-				AnimCompressor->SetCompressibleData(CompressibleData);
-
-				if (bSkipDDC || (CompressCommandletVersion == INDEX_NONE))
+				AnimCompressor->Build(OutData);
+			}
+			else if (AnimCompressor->CanBuild())
+			{
+				if (Params.bAsyncCompression)
 				{
-					AnimCompressor->Build(OutData);
+					FAsyncCompressedAnimationsManagement::Get().RequestAsyncCompression(*AnimCompressor, this, bPerformStripping, OutData);
 				}
-				else if (AnimCompressor->CanBuild())
+				else
 				{
-					bNeedToCleanUpAnimCompressor = false; // GetSynchronous will handle this
 					GetDerivedDataCacheRef().GetSynchronous(AnimCompressor, OutData);
 				}
-			}
-			
-			if(bNeedToCleanUpAnimCompressor)
-			{
-				delete AnimCompressor; // Would really like to do auto mem management but GetDerivedDataCacheRef().GetSynchronous expects a point it can delete, every
 				AnimCompressor = nullptr;
 			}
 		}
+		else
+		{
+			bCompressedDataFromDDC = true;
+		}
 
-		if (OutData.Num() > 0)
+		delete AnimCompressor; // Would really like to do auto mem management but GetDerivedDataCacheRef().GetSynchronous expects a pointer it can delete
+		AnimCompressor = nullptr;
+	}
+
+	if (OutData.Num() > 0) // Haven't async compressed
+	{
+		ApplyCompressedData(OutData);
+
+		if(bCompressedDataFromDDC)
 		{
-			FMemoryReader MemAr(OutData);
-			SerializeCompressedData(MemAr, true);
-			//This is only safe during sync anim compression
-			SetSkeletonVirtualBoneGuid(GetSkeleton()->GetVirtualBoneGuid());
-			bUseRawDataOnly = false;
+			const double CompressionEndTime = FPlatformTime::Seconds();
+			const double CompressionTime = CompressionEndTime - CompressionStartTime;
+
+			TArray<FBoneData> BoneData;
+			FAnimationUtils::BuildSkeletonMetaData(GetSkeleton(), BoneData);
+			Params.CompressContext->GatherPostCompressionStats(CompressedData, BoneData, GetFName(), CompressionTime, false);
 		}
+
+	}
+#endif
+}
+
+#if WITH_EDITOR
+void UAnimSequence::ApplyCompressedData(const FString& DataCacheKeySuffix, const bool bPerformFrameStripping, const TArray<uint8>& Data)
+{
+	if (GetDDCCacheKeySuffix(bPerformFrameStripping) == DataCacheKeySuffix)
+	{
+		ApplyCompressedData(Data);
 	}
+	else
+	{
+		bCompressionInProgress = false;
+	}
+}
 #endif
+
+void UAnimSequence::ApplyCompressedData(const TArray<uint8>& Data)
+{
+#if WITH_EDITOR
+	bCompressionInProgress = false;
+#endif
+	if(Data.Num() > 0)
+	{
+		FMemoryReader MemAr(Data);
+		SerializeCompressedData(MemAr, true);
+		//This is only safe during sync anim compression
+		SetSkeletonVirtualBoneGuid(GetSkeleton()->GetVirtualBoneGuid());
+		bUseRawDataOnly = false;
+	}
 }
 
 void UAnimSequence::SerializeCompressedData(FArchive& Ar, bool bDDCData)
 {
 	if (!HasAnyFlags(RF_ClassDefaultObject))
 	{
-		CompressedData.SerializeCompressedData(Ar, bDDCData, this, this->GetSkeleton(), CurveCompressionSettings);
+		CompressedData.SerializeCompressedData(Ar, bDDCData, this, this->GetSkeleton(), BoneCompressionSettings, CurveCompressionSettings);
 	}
 }
 
@@ -2603,17 +2690,16 @@ void UAnimSequence::RecycleAnimSequence()
 	RawDataGuid.Invalidate();
 	AnimationTrackNames.Empty();
 	TrackToSkeletonMapTable.Empty();
-	CompressedData.CompressedDataStructure.Reset();
 	SourceRawAnimationData.Empty(0);
 	RawCurveData.Empty();
-	CompressedData.CompressedCurveByteStream.Empty(0);
-	CompressedData.CompressedCurveNames.Empty(0);
+
+	ClearCompressedBoneData();
+	ClearCompressedCurveData();
+
 	AuthoredSyncMarkers.Empty();
 	UniqueMarkerNames.Empty();
 	Notifies.Empty();
 	AnimNotifyTracks.Empty();
-	CompressionScheme = nullptr;
-	CompressedData.CurveCompressionCodec = nullptr;
 #endif // WITH_EDITORONLY_DATA
 }
 
@@ -2625,7 +2711,8 @@ void UAnimSequence::CleanAnimSequenceForImport()
 	RawDataGuid.Invalidate();
 	AnimationTrackNames.Empty();
 	TrackToSkeletonMapTable.Empty();
-	CompressedData.CompressedDataStructure.Reset();
+	ClearCompressedBoneData();
+	ClearCompressedCurveData();
 	SourceRawAnimationData.Empty(0);
 }
 #endif // WITH_EDITOR
@@ -3104,7 +3191,8 @@ void UAnimSequence::RemapTracksToNewSkeleton( USkeleton* NewSkeleton, bool bConv
 			SetSkeleton(NewSkeleton);
 
 			// convert back to animated data with new skeleton
-			ConvertRiggingDataToAnimationData(RiggingAnimationData);
+			const bool bPerformPostProcess = false; //Don't do PostProcess during Remap as any animations we reference may not have been updated yet
+			ConvertRiggingDataToAnimationData(RiggingAnimationData, bPerformPostProcess);
 		}
 		// @todo end rig testing
 		// @IMPORTANT: now otherwise this will try to do bone to bone mapping
@@ -4032,7 +4120,7 @@ bool UAnimSequence::ConvertAnimationDataToRiggingData(FAnimSequenceTrackContaine
 	return false;
 }
 
-bool UAnimSequence::ConvertRiggingDataToAnimationData(FAnimSequenceTrackContainer& RiggingAnimationData)
+bool UAnimSequence::ConvertRiggingDataToAnimationData(FAnimSequenceTrackContainer& RiggingAnimationData, bool bPerformPostProcess)
 {
 	if (RiggingAnimationData.GetNum() > 0)
 	{
@@ -4126,7 +4214,11 @@ bool UAnimSequence::ConvertRiggingDataToAnimationData(FAnimSequenceTrackContaine
 		{
 			TrackToSkeletonMapTable[TrackIdx++].BoneTreeIndex = MySkeleton->GetReferenceSkeleton().FindBoneIndex(TrackName);
 		}
-		PostProcessSequence();
+
+		if (bPerformPostProcess)
+		{
+			PostProcessSequence();
+		}
 
 		return true;
 	}
@@ -4339,7 +4431,9 @@ void UAnimSequence::ResetAnimation()
 	SourceRawAnimationData.Empty();
 	AnimationTrackNames.Empty();
 	TrackToSkeletonMapTable.Empty();
-	CompressedData.CompressedDataStructure.Reset();
+
+	ClearCompressedBoneData();
+	ClearCompressedCurveData();
 
 	Notifies.Empty();
 	AuthoredSyncMarkers.Empty();
@@ -5314,10 +5408,10 @@ void UAnimSequence::EnableRootMotionSettingFromMontage(bool bInEnableRootMotion,
 #if WITH_EDITOR
 void UAnimSequence::OnRawDataChanged()
 {
-	CompressedData.CompressedDataStructure.Reset();
+	ClearCompressedBoneData();
 	bUseRawDataOnly = true;
 
-	RequestSyncAnimRecompression(false);
+	RequestAsyncAnimRecompression(false);
 	//MDW - Once we have async anim ddc requests we should do this too
 	//RequestDependentAnimRecompression();
 }
@@ -5325,16 +5419,21 @@ void UAnimSequence::OnRawDataChanged()
 
 bool UAnimSequence::IsCompressedDataValid() const
 {
-	return CompressedData.CompressedDataStructure.IsCompressedDataValid() || RawAnimationData.Num() == 0;
+	// For bone compressed data, we don't check if we have a codec. It is valid to have no compressed data
+	// if we have no raw data. This can happen with sequences that only has curves.
+
+	if (RawAnimationData.Num() == 0)
+	{
+		return true;
+	}
+
+	return CompressedData.CompressedDataStructure != nullptr;
 }
 
 bool UAnimSequence::IsCurveCompressedDataValid() const
 {
-	if (CompressedData.CurveCompressionCodec == nullptr)
-	{
-		// No codec
-		return false;
-	}
+	// For curve compressed data, we don't check if we have a codec. It is valid to have no compressed data
+	// if we have no raw data. This can happen with sequences that only has bones.
 
 	if (CompressedData.CompressedCurveByteStream.Num() == 0 && RawCurveData.FloatCurves.Num() != 0)
 	{
@@ -5351,6 +5450,16 @@ bool UAnimSequence::IsCurveCompressedDataValid() const
 	return true;
 }
 
+void UAnimSequence::ClearCompressedBoneData()
+{
+	CompressedData.ClearCompressedBoneData();
+}
+
+void UAnimSequence::ClearCompressedCurveData()
+{
+	CompressedData.ClearCompressedCurveData();
+}
+
 /*-----------------------------------------------------------------------------
 	AnimNotify& subclasses
 -----------------------------------------------------------------------------*/
@@ -5385,6 +5494,11 @@ void GatherAnimSequenceStats(FOutputDevice& Ar)
 	{
 		UAnimSequence* Seq = *It;
 
+		if (Seq->CompressedData.CompressedDataStructure == nullptr || !Seq->CompressedData.BoneCompressionCodec->IsA<UAnimCompress>())
+		{
+			continue;	// Custom codec we know nothing about, skip it
+		}
+
 		int32 NumTransTracks = 0;
 		int32 NumRotTracks = 0;
 		int32 NumScaleTracks = 0;
@@ -5399,8 +5513,10 @@ void GatherAnimSequenceStats(FOutputDevice& Ar)
 		int32 NumRotTracksWithOneKey = 0;
 		int32 NumScaleTracksWithOneKey = 0;
 
+		const FUECompressedAnimData& AnimData = static_cast<FUECompressedAnimData&>(*Seq->CompressedData.CompressedDataStructure);
+
 		AnimationFormat_GetStats(
-			Seq->CompressedData.CompressedDataStructure,
+			AnimData,
 			NumTransTracks,
 			NumRotTracks,
 			NumScaleTracks,
@@ -5435,7 +5551,7 @@ void GatherAnimSequenceStats(FOutputDevice& Ar)
 			NumTransTracks, NumRotTracks, NumScaleTracks,
 			NumTransTracksWithOneKey, NumRotTracksWithOneKey, NumScaleTracksWithOneKey,
 			TotalNumTransKeys, TotalNumRotKeys, TotalNumScaleKeys,
-			*FAnimationUtils::GetAnimationKeyFormatString(static_cast<AnimationKeyFormat>(Seq->CompressedData.CompressedDataStructure.KeyEncodingFormat)),
+			*FAnimationUtils::GetAnimationKeyFormatString(AnimData.KeyEncodingFormat),
 			(int32)Seq->GetResourceSizeBytes(EResourceSizeMode::EstimatedTotal) );
 	}
 	Ar.Logf( TEXT("======================================================================") );
@@ -5448,11 +5564,5 @@ void GatherAnimSequenceStats(FOutputDevice& Ar)
 
 #endif // !UE_BUILD_SHIPPING
 
-FArchive& operator<<(FArchive& Ar, FCompressedOffsetData& D)
-{
-	Ar << D.OffsetData << D.StripSize;
-	return Ar;
-}
-
 
 #undef LOCTEXT_NAMESPACE 
diff --git a/Engine/Source/Runtime/Engine/Private/Animation/AnimSequenceBase.cpp b/Engine/Source/Runtime/Engine/Private/Animation/AnimSequenceBase.cpp
--- a/Engine/Source/Runtime/Engine/Private/Animation/AnimSequenceBase.cpp
+++ b/Engine/Source/Runtime/Engine/Private/Animation/AnimSequenceBase.cpp
@@ -454,12 +454,14 @@ void UAnimSequenceBase::RefreshCacheData()
 
 		// Handle overlapping notifies
 		FAnimNotifyTrack* TrackToUse = nullptr;
+		int32 TrackIndexToUse = INDEX_NONE;
 		for (int32 TrackOffset = 0; TrackOffset < AnimNotifyTracks.Num(); ++TrackOffset)
 		{
 			const int32 TrackIndex = (Notify.TrackIndex + TrackOffset) % AnimNotifyTracks.Num();
 			if (CanNotifyUseTrack(AnimNotifyTracks[TrackIndex], Notify))
 			{
 				TrackToUse = &AnimNotifyTracks[TrackIndex];
+				TrackIndexToUse = TrackIndex;
 				break;
 			}
 		}
@@ -467,10 +469,13 @@ void UAnimSequenceBase::RefreshCacheData()
 		if (TrackToUse == nullptr)
 		{
 			TrackToUse = &AddNewTrack(AnimNotifyTracks);
+			TrackIndexToUse = AnimNotifyTracks.Num() - 1;
 		}
 
 		check(TrackToUse);
+		check(TrackIndexToUse != INDEX_NONE);
 
+		Notify.TrackIndex = TrackIndexToUse;
 		TrackToUse->Notifies.Add(&Notify);
 	}
 
@@ -520,7 +525,7 @@ int32 UAnimSequenceBase::GetNumberOfFrames() const
 {
 	static float DefaultSampleRateInterval = 1.f / DEFAULT_SAMPLERATE;
 	// because of float error, add small margin at the end, so it can clamp correctly
-	return (SequenceLength / DefaultSampleRateInterval + KINDA_SMALL_NUMBER);
+	return (int32)(SequenceLength / DefaultSampleRateInterval + KINDA_SMALL_NUMBER) + 1;
 }
 
 #if WITH_EDITOR
diff --git a/Engine/Source/Runtime/Engine/Private/Animation/AnimSingleNodeInstance.cpp b/Engine/Source/Runtime/Engine/Private/Animation/AnimSingleNodeInstance.cpp
--- a/Engine/Source/Runtime/Engine/Private/Animation/AnimSingleNodeInstance.cpp
+++ b/Engine/Source/Runtime/Engine/Private/Animation/AnimSingleNodeInstance.cpp
@@ -323,6 +323,11 @@ float UAnimSingleNodeInstance::GetCurrentTime() const
 void UAnimSingleNodeInstance::SetReverse(bool bInReverse)
 {
 	GetProxyOnGameThread<FAnimSingleNodeInstanceProxy>().SetReverse(bInReverse);
+
+	if (FAnimMontageInstance* CurMontageInstance = GetActiveMontageInstance())
+	{
+		CurMontageInstance->SetPlayRate(GetPlayRate());
+	}
 }
 
 void UAnimSingleNodeInstance::SetPositionWithPreviousTime(float InPosition, float InPreviousTime, bool bFireNotifies)
diff --git a/Engine/Source/Runtime/Engine/Private/Animation/AnimStreamable.cpp b/Engine/Source/Runtime/Engine/Private/Animation/AnimStreamable.cpp
--- a/Engine/Source/Runtime/Engine/Private/Animation/AnimStreamable.cpp
+++ b/Engine/Source/Runtime/Engine/Private/Animation/AnimStreamable.cpp
@@ -12,8 +12,9 @@
 #include "UObject/LinkerLoad.h"
 #include "Animation/AnimCompressionDerivedData.h"
 #include "DerivedDataCacheInterface.h"
+#include "Animation/AnimBoneCompressionSettings.h"
 #include "Animation/AnimCurveCompressionSettings.h"
-#include "Animation/AnimCompress.h"
+#include "Animation/AnimBoneCompressionCodec.h"
 #include "Animation/AnimCurveCompressionCodec.h"
 #include "BonePose.h"
 #include "ContentStreaming.h"
@@ -58,7 +59,7 @@ void FAnimStreamableChunk::Serialize(FArchive& Ar, UAnimStreamable* Owner, int32
 				check(!CompressedAnimSequence);
 				CompressedAnimSequence = new FCompressedAnimSequence();
 			}
-			CompressedAnimSequence->SerializeCompressedData(Ar, false, Owner, Owner->GetSkeleton(), Owner->CurveCompressionSettings, false);
+			CompressedAnimSequence->SerializeCompressedData(Ar, false, Owner, Owner->GetSkeleton(), Owner->BoneCompressionSettings, Owner->CurveCompressionSettings, false);
 		}
 		else
 		{
@@ -68,11 +69,11 @@ void FAnimStreamableChunk::Serialize(FArchive& Ar, UAnimStreamable* Owner, int32
 			{
 				//Need to pack compressed data into BulkData
 				TArray<uint8> TempBytes;
-				const int32 InitialSize = CompressedAnimSequence->CompressedDataStructure.GetApproxBoneCompressedSize();
+				const int32 InitialSize = CompressedAnimSequence->CompressedDataStructure->GetApproxCompressedSize();
 				TempBytes.Reset(InitialSize);
 
 				FMemoryWriter TempAr(TempBytes, true);
-				CompressedAnimSequence->SerializeCompressedData(TempAr, false, Owner, Owner->GetSkeleton(), Owner->CurveCompressionSettings, false);
+				CompressedAnimSequence->SerializeCompressedData(TempAr, false, Owner, Owner->GetSkeleton(), Owner->BoneCompressionSettings, Owner->CurveCompressionSettings, false);
 
 				BulkData.Lock(LOCK_READ_WRITE);
 				void* ChunkData = BulkData.Realloc(TempBytes.Num());
@@ -353,11 +354,6 @@ void UAnimStreamable::PostLoad()
 		NonConstSeq->ConditionalPostLoad();
 	}
 
-	if (SourceSequence)
-	{
-		CompressionScheme = DuplicateObject<UAnimCompress>(SourceSequence->CompressionScheme, this);
-	}
-
 	if (SourceSequence && (GenerateGuidFromRawAnimData(SourceSequence->GetRawAnimationData(), SourceSequence->RawCurveData) != RawDataGuid))
 	{
 		InitFrom(SourceSequence);
@@ -416,17 +412,14 @@ int32 UAnimStreamable::GetChunkIndexForTime(const TArray<FAnimStreamableChunk>&
 }
 
 #if WITH_EDITOR
-float UAnimStreamable::GetAltCompressionErrorThreshold() const
-{
-	return SourceSequence ? SourceSequence->GetAltCompressionErrorThreshold() : FAnimationUtils::GetAlternativeCompressionThreshold();
-}
-
 void UAnimStreamable::InitFrom(const UAnimSequence* InSourceSequence)
 {
 	Modify();
 	SetSkeleton(InSourceSequence->GetSkeleton());
 	SourceSequence = InSourceSequence;
-	CompressionScheme = DuplicateObject<UAnimCompress>(SourceSequence->CompressionScheme, this);
+
+	BoneCompressionSettings = InSourceSequence->BoneCompressionSettings;
+	CurveCompressionSettings = InSourceSequence->CurveCompressionSettings;
 
 	RawAnimationData = InSourceSequence->GetRawAnimationData();
 	RawCurveData = InSourceSequence->RawCurveData;
@@ -502,9 +495,9 @@ void UAnimStreamable::RequestCompressedData(const ITargetPlatform* Platform)
 		IStreamingManager::Get().GetAnimationStreamingManager().RemoveStreamingAnim(this);
 	}
 
-	if (!CompressionScheme)
+	if (BoneCompressionSettings == nullptr || !BoneCompressionSettings->AreSettingsValid())
 	{
-		CompressionScheme = FAnimationUtils::GetDefaultAnimationCompressionAlgorithm();
+		BoneCompressionSettings = FAnimationUtils::GetDefaultAnimationBoneCompressionSettings();
 	}
 
 	if (CurveCompressionSettings == nullptr || !CurveCompressionSettings->AreSettingsValid())
@@ -538,7 +531,7 @@ void UAnimStreamable::RequestCompressedData(const ITargetPlatform* Platform)
 
 	PlatformData.Chunks.AddDefaulted(NumChunks);
 
-	const FString BaseDDCKey = GetBaseDDCKey(NumChunks, GetAltCompressionErrorThreshold());
+	const FString BaseDDCKey = GetBaseDDCKey(NumChunks);
 
 	const bool bInAllowAlternateCompressor = false;
 	const bool bInOutput				   = false;
@@ -605,7 +598,7 @@ void UAnimStreamable::RequestCompressedDataForChunk(const FString& ChunkDDCKey,
 	TArray<uint8> OutData;
 	{
 		bool bNeedToCleanUpAnimCompressor = true;
-		FDerivedDataAnimationCompression* AnimCompressor = new FDerivedDataAnimationCompression(TEXT("StreamAnim"), ChunkDDCKey, CompressContext, 0);
+		FDerivedDataAnimationCompression* AnimCompressor = new FDerivedDataAnimationCompression(TEXT("StreamAnim"), ChunkDDCKey, CompressContext);
 
 		const FString FinalDDCKey = FDerivedDataCacheInterface::BuildCacheKey(AnimCompressor->GetPluginName(), AnimCompressor->GetVersionString(), *AnimCompressor->GetPluginSpecificCacheKeySuffix());
 
@@ -619,7 +612,7 @@ void UAnimStreamable::RequestCompressedDataForChunk(const FString& ChunkDDCKey,
 
 		if (bSkipDDC || !GetDerivedDataCacheRef().GetSynchronous(*FinalDDCKey, OutData))
 		{
-			TSharedRef<FCompressibleAnimData> CompressibleData = MakeShared<FCompressibleAnimData>(CompressionScheme, CurveCompressionSettings, GetSkeleton(), Interpolation, Chunk.SequenceLength, ChunkNumFrames+1, GetAltCompressionErrorThreshold());
+			TSharedRef<FCompressibleAnimData> CompressibleData = MakeShared<FCompressibleAnimData>(BoneCompressionSettings, CurveCompressionSettings, GetSkeleton(), Interpolation, Chunk.SequenceLength, ChunkNumFrames+1);
 
 			CompressibleData->RawAnimationData.AddDefaulted(RawAnimationData.Num());
 
@@ -684,7 +677,7 @@ void UAnimStreamable::RequestCompressedDataForChunk(const FString& ChunkDDCKey,
 		{
 			Chunk.CompressedAnimSequence = new FCompressedAnimSequence();
 		}
-		Chunk.CompressedAnimSequence->SerializeCompressedData(MemAr, true, this, this->GetSkeleton(), CurveCompressionSettings);
+		Chunk.CompressedAnimSequence->SerializeCompressedData(MemAr, true, this, this->GetSkeleton(), BoneCompressionSettings, CurveCompressionSettings);
 	}
 }
 
@@ -694,7 +687,7 @@ void UAnimStreamable::UpdateRawData()
 	RequestCompressedData();
 }
 
-FString UAnimStreamable::GetBaseDDCKey(uint32 NumChunks, float AltCompressionErrorThreshold) const
+FString UAnimStreamable::GetBaseDDCKey(uint32 NumChunks) const
 {
 	//Make up our content key consisting of:
 	//  * Streaming Anim Chunk logic version
@@ -707,8 +700,7 @@ FString UAnimStreamable::GetBaseDDCKey(uint32 NumChunks, float AltCompressionErr
 	FArcToHexString ArcToHexString;
 
 	ArcToHexString.Ar << NumChunks;
-	ArcToHexString.Ar << AltCompressionErrorThreshold;
-	CompressionScheme->PopulateDDCKeyArchive(ArcToHexString.Ar);
+	BoneCompressionSettings->PopulateDDCKey(ArcToHexString.Ar);
 	CurveCompressionSettings->PopulateDDCKey(ArcToHexString.Ar);
 
 	FString Ret = FString::Printf(TEXT("%s%s%s%s_%s"),
diff --git a/Engine/Source/Runtime/Engine/Private/Animation/AnimTypes.cpp b/Engine/Source/Runtime/Engine/Private/Animation/AnimTypes.cpp
--- a/Engine/Source/Runtime/Engine/Private/Animation/AnimTypes.cpp
+++ b/Engine/Source/Runtime/Engine/Private/Animation/AnimTypes.cpp
@@ -5,6 +5,7 @@
 #include "AnimationUtils.h"
 #include "Animation/AnimNotifies/AnimNotify.h"
 #include "Animation/AnimNotifies/AnimNotifyState.h"
+#include "UObject/AnimObjectVersion.h"
 
 #define NOTIFY_TRIGGER_OFFSET KINDA_SMALL_NUMBER;
 
@@ -39,6 +40,28 @@ float GetTriggerTimeOffsetForType(EAnimEventTriggerOffsets::Type OffsetType)
 /////////////////////////////////////////////////////
 // FAnimNotifyEvent
 
+#if WITH_EDITORONLY_DATA
+
+bool FAnimNotifyEvent::Serialize(FArchive& Ar)
+{
+	Ar.UsingCustomVersion(FAnimObjectVersion::GUID);
+	return false;
+}
+
+void FAnimNotifyEvent::PostSerialize(const FArchive& Ar)
+{
+	if (Ar.IsLoading())
+	{
+		if (!Guid.IsValid())
+		{
+			// Create a new Guid if this one is invalid
+			Guid = FGuid::NewGuid();
+		}
+	}
+}
+
+#endif
+
 void FAnimNotifyEvent::RefreshTriggerOffset(EAnimEventTriggerOffsets::Type PredictedOffsetType)
 {
 	if(PredictedOffsetType == EAnimEventTriggerOffsets::NoOffset || TriggerTimeOffset == 0.f)
@@ -110,6 +133,37 @@ FName FAnimNotifyEvent::GetNotifyEventName() const
 	return NAME_None;
 }
 
+////////////////////////////
+//
+// FAnimSyncMarker
+// 
+////////////////////////////
+
+#if WITH_EDITORONLY_DATA
+bool FAnimSyncMarker::Serialize(FArchive& Ar)
+{
+	Ar.UsingCustomVersion(FAnimObjectVersion::GUID);
+	
+	UScriptStruct* Struct = FAnimSyncMarker::StaticStruct();
+	
+	if (Ar.IsLoading() || Ar.IsSaving())
+	{
+		Struct->SerializeTaggedProperties(Ar, (uint8*)this, Struct, nullptr);
+	}
+
+	if (Ar.IsLoading())
+	{
+		if(!Guid.IsValid())
+		{
+			// Create a new Guid if this one is invalid
+			Guid = FGuid::NewGuid();
+		}
+	}
+
+	return true;
+}
+#endif
+
 ////////////////////////////
 //
 // FMarkerSyncData
diff --git a/Engine/Source/Runtime/Engine/Private/Animation/AnimationRuntime.cpp b/Engine/Source/Runtime/Engine/Private/Animation/AnimationRuntime.cpp
--- a/Engine/Source/Runtime/Engine/Private/Animation/AnimationRuntime.cpp
+++ b/Engine/Source/Runtime/Engine/Private/Animation/AnimationRuntime.cpp
@@ -9,6 +9,9 @@
 #include "Animation/BlendSpaceBase.h"
 #include "Animation/AnimInstance.h"
 #include "SkeletalRender.h"
+#if INTEL_ISPC
+#include "AnimationRuntime.ispc.generated.h"
+#endif
 
 DEFINE_LOG_CATEGORY(LogAnimation);
 DEFINE_LOG_CATEGORY(LogRootMotion);
@@ -70,6 +73,50 @@ FORCEINLINE void BlendPose(const FCompactPose& SourcePose, FCompactPose& ResultP
 	}
 }
 
+template <>
+FORCEINLINE void BlendPose<ETransformBlendMode::Overwrite>(const FCompactPose& SourcePose, FCompactPose& ResultPose, const float BlendWeight)
+{
+	if (INTEL_ISPC)
+	{
+#if INTEL_ISPC
+		ispc::BlendTransformOverwrite(
+			(ispc::FTransform*)&SourcePose.GetBones()[0],
+			(ispc::FTransform*)&ResultPose.GetBones()[0],
+			BlendWeight,
+			SourcePose.GetNumBones());
+#endif
+	}
+	else
+	{
+		for (FCompactPoseBoneIndex BoneIndex : SourcePose.ForEachBoneIndex())
+		{
+			BlendTransform<ETransformBlendMode::Overwrite>(SourcePose[BoneIndex], ResultPose[BoneIndex], BlendWeight);
+		}
+	}
+}
+
+template <>
+FORCEINLINE void BlendPose<ETransformBlendMode::Accumulate>(const FCompactPose& SourcePose, FCompactPose& ResultPose, const float BlendWeight)
+{
+	if (INTEL_ISPC)
+	{
+#if INTEL_ISPC
+		ispc::BlendTransformAccumulate(
+			(ispc::FTransform*)&SourcePose.GetBones()[0],
+			(ispc::FTransform*)&ResultPose.GetBones()[0],
+			BlendWeight,
+			SourcePose.GetNumBones());
+#endif
+	}
+	else
+	{
+		for (FCompactPoseBoneIndex BoneIndex : SourcePose.ForEachBoneIndex())
+		{
+			BlendTransform<ETransformBlendMode::Accumulate>(SourcePose[BoneIndex], ResultPose[BoneIndex], BlendWeight);
+		}
+	}
+}
+
 FORCEINLINE void BlendCurves(const TArrayView<const FBlendedCurve> SourceCurves, const TArrayView<const float> SourceWeights, const TArrayView<const int32> SourceWeightsIndices, FBlendedCurve& OutCurve)
 {
 	if (SourceCurves.Num() > 0)
@@ -128,6 +175,37 @@ FORCEINLINE void BlendCurves(const TArrayView<const FBlendedCurve* const> Source
 				BlendCurves(SourceCurves, SourceWeights, OutCurve);
 			}
 		}
+		else if (BlendOption == ECurveBlendOption::Type::UseMaxValue)
+		{
+			OutCurve.Override(*SourceCurves[0], SourceWeights[0]);
+
+			for (int32 CurveIndex = 1; CurveIndex < SourceCurves.Num(); ++CurveIndex)
+			{
+				OutCurve.UseMaxValue(*SourceCurves[CurveIndex]);
+			}
+		}
+		else if (BlendOption == ECurveBlendOption::Type::UseMinValue)
+		{
+			OutCurve.Override(*SourceCurves[0], SourceWeights[0]);
+
+			for (int32 CurveIndex = 1; CurveIndex < SourceCurves.Num(); ++CurveIndex)
+			{
+				OutCurve.UseMinValue(*SourceCurves[CurveIndex]);
+			}
+		}
+		else if (BlendOption == ECurveBlendOption::Type::UseBasePose)
+		{
+			OutCurve.Override(*SourceCurves[0], SourceWeights[0]);
+		}
+		else if (BlendOption == ECurveBlendOption::Type::DoNotOverride)
+		{
+			OutCurve.Override(*SourceCurves[0], SourceWeights[0]);
+
+			for (int32 CurveIndex = 1; CurveIndex < SourceCurves.Num(); ++CurveIndex)
+			{
+				OutCurve.CombinePreserved(*SourceCurves[CurveIndex]);
+			}
+		}
 		else
 		{
 			OutCurve.Override(*SourceCurves[0], SourceWeights[0]);
@@ -675,12 +753,24 @@ void FAnimationRuntime::ConvertPoseToMeshRotation(FCompactPose& LocalPose)
 
 	// Convert all rotations to mesh space
 	// only the root bone doesn't have a parent. So skip it to save a branch in the iteration.
-	for (FCompactPoseBoneIndex BoneIndex(1); BoneIndex < LocalPose.GetNumBones(); ++BoneIndex)
+	if (INTEL_ISPC)
 	{
-		const FCompactPoseBoneIndex ParentIndex = LocalPose.GetParentBoneIndex(BoneIndex);
-
-		const FQuat MeshSpaceRotation = LocalPose[ParentIndex].GetRotation() * LocalPose[BoneIndex].GetRotation();
-		LocalPose[BoneIndex].SetRotation(MeshSpaceRotation);
+#if INTEL_ISPC
+		ispc::ConvertPoseToMeshRotation(
+			(ispc::FTransform*)&LocalPose.GetBones()[0],
+			(int32*)&LocalPose.GetBoneContainer().GetCompactPoseParentBoneArray().GetData()[0],
+			LocalPose.GetNumBones());
+#endif
+	}
+	else
+	{
+		for (FCompactPoseBoneIndex BoneIndex(1); BoneIndex < LocalPose.GetNumBones(); ++BoneIndex)
+		{
+			const FCompactPoseBoneIndex ParentIndex = LocalPose.GetParentBoneIndex(BoneIndex);
+	
+			const FQuat MeshSpaceRotation = LocalPose[ParentIndex].GetRotation() * LocalPose[BoneIndex].GetRotation();
+			LocalPose[BoneIndex].SetRotation(MeshSpaceRotation);
+		}
 	}
 }
 
@@ -690,12 +780,24 @@ void FAnimationRuntime::ConvertMeshRotationPoseToLocalSpace(FCompactPose& Pose)
 
 	// Convert all rotations to mesh space
 	// only the root bone doesn't have a parent. So skip it to save a branch in the iteration.
-	for (FCompactPoseBoneIndex BoneIndex(Pose.GetNumBones()-1); BoneIndex > 0; --BoneIndex)
+	if (INTEL_ISPC)
+	{
+#if INTEL_ISPC
+		ispc::ConvertMeshRotationPoseToLocalSpace(
+			(ispc::FTransform*)&Pose.GetBones()[0],
+			(int32*)&Pose.GetBoneContainer().GetCompactPoseParentBoneArray().GetData()[0],
+			Pose.GetNumBones());
+#endif
+	}
+	else
 	{
-		const FCompactPoseBoneIndex ParentIndex = Pose.GetParentBoneIndex(BoneIndex);
+		for (FCompactPoseBoneIndex BoneIndex(Pose.GetNumBones() - 1); BoneIndex > 0; --BoneIndex)
+		{
+			const FCompactPoseBoneIndex ParentIndex = Pose.GetParentBoneIndex(BoneIndex);
 
-		FQuat LocalSpaceRotation = Pose[ParentIndex].GetRotation().Inverse() * Pose[BoneIndex].GetRotation();
-		Pose[BoneIndex].SetRotation(LocalSpaceRotation);
+			FQuat LocalSpaceRotation = Pose[ParentIndex].GetRotation().Inverse() * Pose[BoneIndex].GetRotation();
+			Pose[BoneIndex].SetRotation(LocalSpaceRotation);
+		}
 	}
 }
 
@@ -723,10 +825,23 @@ void FAnimationRuntime::AccumulateLocalSpaceAdditivePoseInternal(FCompactPose& B
 		const ScalarRegister VBlendWeight(Weight);
 		if (FAnimWeight::IsFullWeight(Weight))
 		{
-			// fast path, no need to weight additive.
-			for (FCompactPoseBoneIndex BoneIndex : BasePose.ForEachBoneIndex())
+			if (INTEL_ISPC)
+			{
+#if INTEL_ISPC
+				ispc::AccumulateWithAdditiveScale(
+					(ispc::FTransform*)&BasePose.GetBones()[0],
+					(ispc::FTransform*)&AdditivePose.GetBones()[0],
+					Weight,
+					BasePose.GetNumBones());
+#endif
+			}
+			else
 			{
-				BasePose[BoneIndex].AccumulateWithAdditiveScale(AdditivePose[BoneIndex], VBlendWeight);
+				// fast path, no need to weight additive.
+				for (FCompactPoseBoneIndex BoneIndex : BasePose.ForEachBoneIndex())
+				{
+					BasePose[BoneIndex].AccumulateWithAdditiveScale(AdditivePose[BoneIndex], VBlendWeight);
+				}
 			}
 		}
 		else
@@ -1053,6 +1168,14 @@ void FAnimationRuntime::BlendPosesPerBoneFilter(
 	return BlendPosesPerBoneFilter(BasePose, BlendPoses, BaseCurve, BlendedCurves, OutPose, OutCurve, BoneBlendWeights, BlendFlags, CurveBlendOption);
 }
 
+// Helper function to get FTransform from a PoseIndex and BoneIndex
+extern "C" const uint8* GetTransformFromArray(const uint8 *BlendPoseBase, const int32 PoseIndex, const int32 BoneIndex)
+{
+	const TArray<struct FCompactPose>& BlendPoses = *(const TArray<struct FCompactPose>*) BlendPoseBase;
+	const FTransform* BlendPose = &BlendPoses[PoseIndex][FCompactPoseBoneIndex(BoneIndex)];
+	return (const uint8*)BlendPose;
+}
+
 void FAnimationRuntime::BlendPosesPerBoneFilter(
 	struct FCompactPose& BasePose, 
 	const TArray<struct FCompactPose>& BlendPoses, 
@@ -1172,56 +1295,78 @@ void FAnimationRuntime::BlendPosesPerBoneFilter(
 	// blend poses with both mesh space rotation and scaling (we assume uniform scale)
 	if (bMeshSpaceRotationBlend && bMeshSpaceScaleBlend)
 	{
-		for (const FCompactPoseBoneIndex BoneIndex : BasePose.ForEachBoneIndex())
+		if (INTEL_ISPC)
 		{
-			const int32 PoseIndex = BoneBlendWeights[BoneIndex.GetInt()].SourceIndex;
-			const FCompactPoseBoneIndex ParentIndex = BoneContainer.GetParentBoneIndex(BoneIndex);
-
-			if (ParentIndex != INDEX_NONE)
-			{
-				AccumulateMeshSpaceRotation(PoseIndex, BoneIndex, SourceRotations[ParentIndex], TargetRotations[ParentIndex]);
-				AccumulateMeshSpaceScale(PoseIndex, BoneIndex, SourceScales[ParentIndex], TargetScales[ParentIndex]);
-			}
-			else
+#if INTEL_ISPC
+			ispc::BlendPosesPerBoneFilterScaleRotation(
+				(ispc::FTransform*)OutPose.GetBones().GetData(),
+				(ispc::FTransform*)BasePose.GetBones().GetData(),
+				(const uint8*)&BlendPoses,
+				(ispc::FVector4*)SourceRotations.GetData(),
+				(ispc::FVector*)SourceScales.GetData(),
+				(ispc::FVector4*)TargetRotations.GetData(),
+				(ispc::FVector*)TargetScales.GetData(),
+				(ispc::FVector4*)BlendRotations.GetData(),
+				(ispc::FVector*)BlendScales.GetData(),
+				MaxPoseWeights.GetData(),
+				(ispc::FPerBoneBlendWeight*)BoneBlendWeights.GetData(),
+				(int32*)BoneContainer.GetCompactPoseParentBoneArray().GetData(),
+				BasePose.GetNumBones());
+#endif
+		} 
+		else
+		{
+			for (const FCompactPoseBoneIndex BoneIndex : BasePose.ForEachBoneIndex())
 			{
-				AccumulateMeshSpaceRotation(PoseIndex, BoneIndex, FQuat::Identity, FQuat::Identity);
-				AccumulateMeshSpaceScale(PoseIndex, BoneIndex, FVector(1.0f), FVector(1.0f));
-			}
+				const int32 PoseIndex = BoneBlendWeights[BoneIndex.GetInt()].SourceIndex;
+				const FCompactPoseBoneIndex ParentIndex = BoneContainer.GetParentBoneIndex(BoneIndex);
 
-			FTransform BaseAtom = BasePose[BoneIndex];
-			FTransform TargetAtom = BlendPoses[PoseIndex][BoneIndex];
-			FTransform BlendAtom;
+				if (ParentIndex != INDEX_NONE)
+				{
+					AccumulateMeshSpaceRotation(PoseIndex, BoneIndex, SourceRotations[ParentIndex], TargetRotations[ParentIndex]);
+					AccumulateMeshSpaceScale(PoseIndex, BoneIndex, SourceScales[ParentIndex], TargetScales[ParentIndex]);
+				}
+				else
+				{
+					AccumulateMeshSpaceRotation(PoseIndex, BoneIndex, FQuat::Identity, FQuat::Identity);
+					AccumulateMeshSpaceScale(PoseIndex, BoneIndex, FVector(1.0f), FVector(1.0f));
+				}
 
-			const float BlendWeight = FMath::Clamp(BoneBlendWeights[BoneIndex.GetInt()].BlendWeight, 0.f, 1.f);
-			MaxPoseWeights[PoseIndex] = FMath::Max(MaxPoseWeights[PoseIndex], BlendWeight);
+				FTransform BaseAtom = BasePose[BoneIndex];
+				FTransform TargetAtom = BlendPoses[PoseIndex][BoneIndex];
+				FTransform BlendAtom;
 
-			if (!FAnimWeight::IsRelevant(BlendWeight))
-			{
-				BlendAtom = BaseAtom;
-				BlendRotations[BoneIndex] = SourceRotations[BoneIndex];
-				BlendScales[BoneIndex] = SourceScales[BoneIndex];
-			}
-			else if (FAnimWeight::IsFullWeight(BlendWeight))
-			{
-				BlendAtom = TargetAtom;
-				BlendRotations[BoneIndex] = TargetRotations[BoneIndex];
-				BlendScales[BoneIndex] = TargetScales[BoneIndex];
-			}
-			else
-			{
-				BlendAtom = BaseAtom;
-				BlendAtom.BlendWith(TargetAtom, BlendWeight);
-				LerpMeshSpaceRotation(BoneIndex, BlendWeight);
-				LerpMeshSpaceScale(BoneIndex, BlendWeight);
-			}
+				const float BlendWeight = FMath::Clamp(BoneBlendWeights[BoneIndex.GetInt()].BlendWeight, 0.f, 1.f);
+				MaxPoseWeights[PoseIndex] = FMath::Max(MaxPoseWeights[PoseIndex], BlendWeight);
 
-			if (ParentIndex != INDEX_NONE)
-			{
-				ConvertMeshToLocalSpaceRotation(BlendAtom, ParentIndex, BoneIndex);
-				ConvertMeshToLocalSpaceScale(BlendAtom, ParentIndex, BoneIndex);
-			}
+				if (!FAnimWeight::IsRelevant(BlendWeight))
+				{
+					BlendAtom = BaseAtom;
+					BlendRotations[BoneIndex] = SourceRotations[BoneIndex];
+					BlendScales[BoneIndex] = SourceScales[BoneIndex];
+				}
+				else if (FAnimWeight::IsFullWeight(BlendWeight))
+				{
+					BlendAtom = TargetAtom;
+					BlendRotations[BoneIndex] = TargetRotations[BoneIndex];
+					BlendScales[BoneIndex] = TargetScales[BoneIndex];
+				}
+				else
+				{
+					BlendAtom = BaseAtom;
+					BlendAtom.BlendWith(TargetAtom, BlendWeight);
+					LerpMeshSpaceRotation(BoneIndex, BlendWeight);
+					LerpMeshSpaceScale(BoneIndex, BlendWeight);
+				}
 
-			OutPose[BoneIndex] = BlendAtom;
+				if (ParentIndex != INDEX_NONE)
+				{
+					ConvertMeshToLocalSpaceRotation(BlendAtom, ParentIndex, BoneIndex);
+					ConvertMeshToLocalSpaceScale(BlendAtom, ParentIndex, BoneIndex);
+				}
+
+				OutPose[BoneIndex] = BlendAtom;
+			}
 		}
 	}
 
diff --git a/Engine/Source/Runtime/Engine/Private/Animation/AnimationSettings.cpp b/Engine/Source/Runtime/Engine/Private/Animation/AnimationSettings.cpp
--- a/Engine/Source/Runtime/Engine/Private/Animation/AnimationSettings.cpp
+++ b/Engine/Source/Runtime/Engine/Private/Animation/AnimationSettings.cpp
@@ -10,7 +10,6 @@ UAnimationSettings::UAnimationSettings(const FObjectInitializer& ObjectInitializ
 	, RotationCompressionFormat(ACF_Float96NoW)
 	, TranslationCompressionFormat(ACF_None)
 	, MaxCurveError(0.f)
-	, AlternativeCompressionThreshold(1.f)
 	, ForceRecompression(false)
 	, bOnlyCheckForMissingSkeletalMeshes(false)
 	, bForceBelowThreshold(false)
diff --git a/Engine/Source/Runtime/Engine/Private/Animation/AnimationStreaming.cpp b/Engine/Source/Runtime/Engine/Private/Animation/AnimationStreaming.cpp
--- a/Engine/Source/Runtime/Engine/Private/Animation/AnimationStreaming.cpp
+++ b/Engine/Source/Runtime/Engine/Private/Animation/AnimationStreaming.cpp
@@ -349,7 +349,7 @@ void FAnimationStreamingManager::OnAsyncFileCallback(FStreamingAnimationData* St
 		FMemoryReaderView Reader(MemView);
 
 		UAnimStreamable* Anim = StreamingAnimData->StreamableAnim;
-		NewCompressedData->SerializeCompressedData(Reader, false, Anim, Anim->GetSkeleton(), Anim->CurveCompressionSettings);
+		NewCompressedData->SerializeCompressedData(Reader, false, Anim, Anim->GetSkeleton(), Anim->BoneCompressionSettings, Anim->CurveCompressionSettings);
 
 		const float SequenceLength = StreamingAnimData->StreamableAnim->GetRunningPlatformData().Chunks[ChunkIndex].SequenceLength;
 		if (SequenceLength < RequestDuration)
