COPY from //Dev-Anim to //Dev-Main

diff --git a/Engine/Source/Runtime/Engine/Private/Animation/AnimationUtils.cpp b/Engine/Source/Runtime/Engine/Private/Animation/AnimationUtils.cpp
--- a/Engine/Source/Runtime/Engine/Private/Animation/AnimationUtils.cpp
+++ b/Engine/Source/Runtime/Engine/Private/Animation/AnimationUtils.cpp
@@ -8,15 +8,10 @@
 #include "Misc/ConfigCacheIni.h"
 #include "UObject/Package.h"
 #include "AnimationRuntime.h"
-#include "Animation/AnimCompress.h"
-#include "Animation/AnimCompress_BitwiseCompressOnly.h"
-#include "Animation/AnimCompress_RemoveLinearKeys.h"
-#include "Animation/AnimCompress_PerTrackCompression.h"
-#include "Animation/AnimCompress_LeastDestructive.h"
-#include "Animation/AnimCompress_RemoveEverySecondKey.h"
-#include "Animation/AnimCompress_Automatic.h"
 #include "Animation/AnimSet.h"
 #include "Animation/AnimationSettings.h"
+#include "Animation/AnimBoneCompressionCodec.h"
+#include "Animation/AnimBoneCompressionSettings.h"
 #include "Animation/AnimCurveCompressionSettings.h"
 #include "AnimationCompression.h"
 #include "Engine/SkeletalMeshSocket.h"
@@ -240,9 +235,9 @@ void FAnimationUtils::ComputeCompressionError(const FCompressibleAnimData& Compr
 	ErrorStats.MaxErrorTime = 0.0f;
 	int32 MaxErrorTrack = -1;
 
-	if (CompressedData.CompressedNumberOfFrames > 0)
+	if (CompressedData.AnimData != nullptr && CompressedData.AnimData->CompressedNumberOfFrames > 0)
 	{
-		const bool bCanUseCompressedData = (CompressedData.CompressedByteStream.Num() > 0);
+		const bool bCanUseCompressedData = CompressedData.AnimData->IsValid();
 		if (!bCanUseCompressedData)
 		{
 			// If we can't use CompressedData, there's not much point in being here.
@@ -290,15 +285,14 @@ void FAnimationUtils::ComputeCompressionError(const FCompressibleAnimData& Compr
 
 		const FTransform EndEffectorDummyBoneSocket(FQuat::Identity, FVector(END_EFFECTOR_DUMMY_BONE_LENGTH_SOCKET));
 		const FTransform EndEffectorDummyBone(FQuat::Identity, FVector(END_EFFECTOR_DUMMY_BONE_LENGTH));
-		const FAnimKeyHelper Helper(CompressibleAnimData.SequenceLength, CompressedData.CompressedNumberOfFrames);
+		const FAnimKeyHelper Helper(CompressibleAnimData.SequenceLength, CompressedData.AnimData->CompressedNumberOfFrames);
 		const float KeyLength = Helper.TimePerKey() + SMALL_NUMBER;
 
-		const FUECompressedAnimData CompressedDataWrapper(CompressedData);
-		FAnimSequenceDecompressionContext DecompContext(CompressibleAnimData, CompressedDataWrapper);
+		FAnimSequenceDecompressionContext DecompContext(CompressibleAnimData.SequenceLength, CompressibleAnimData.Interpolation, CompressibleAnimData.AnimFName, *CompressedData.AnimData);
 
 		const TArray<FBoneData>& BoneData = CompressibleAnimData.BoneData;
 
-		for (int32 FrameIndex = 0; FrameIndex< CompressedData.CompressedNumberOfFrames; FrameIndex++)
+		for (int32 FrameIndex = 0; FrameIndex< CompressedData.AnimData->CompressedNumberOfFrames; FrameIndex++)
 		{
 			const float Time = (float)FrameIndex * KeyLength;
 			DecompContext.Seek(Time);
@@ -329,8 +323,7 @@ void FAnimationUtils::ComputeCompressionError(const FCompressibleAnimData& Compr
 						FTransform AdditiveRawTransform;
 						FTransform AdditiveNewTransform;
 						FAnimationUtils::ExtractTransformFromTrack(Time, CompressibleAnimData.NumFrames, CompressibleAnimData.SequenceLength, CompressibleAnimData.RawAnimationData[BoneIndexData.TrackIndex], CompressibleAnimData.Interpolation, AdditiveRawTransform);
-
-						AnimationFormat_GetBoneAtom(AdditiveNewTransform, DecompContext, BoneIndexData.TrackIndex);
+						CompressedData.Codec->DecompressBone(DecompContext, BoneIndexData.TrackIndex, AdditiveNewTransform);
 
 						const ScalarRegister VBlendWeight(1.f);
 						RawTransforms[BoneIndex].AccumulateWithAdditiveScale(AdditiveRawTransform, VBlendWeight);
@@ -339,7 +332,7 @@ void FAnimationUtils::ComputeCompressionError(const FCompressibleAnimData& Compr
 					else
 					{
 						FAnimationUtils::ExtractTransformFromTrack(Time, CompressibleAnimData.NumFrames, CompressibleAnimData.SequenceLength, CompressibleAnimData.RawAnimationData[BoneIndexData.TrackIndex], CompressibleAnimData.Interpolation, RawTransforms[BoneIndex]);
-						AnimationFormat_GetBoneAtom(NewTransforms[BoneIndex], DecompContext, BoneIndexData.TrackIndex);
+						CompressedData.Codec->DecompressBone(DecompContext, BoneIndexData.TrackIndex, NewTransforms[BoneIndex]);
 					}
 				}
 
@@ -399,67 +392,6 @@ void FAnimationUtils::ComputeCompressionError(const FCompressibleAnimData& Compr
 }
 #endif
 
-///////////////////////////////////////////////////////////////////////////////////////////////////////////////////////
-//
-// Default animation compression algorithm.
-//
-///////////////////////////////////////////////////////////////////////////////////////////////////////////////////////
-
-namespace {
-
-/**
- * @return		A new instance of the default animation compression algorithm singleton, attached to the root set.
- */
-static inline UAnimCompress* ConstructDefaultCompressionAlgorithm()
-{
-	// Algorithm.
-	const UAnimationSettings* AnimSettings = UAnimationSettings::Get();
-
-	// Rotation compression format.
-	AnimationCompressionFormat RotationCompressionFormat = AnimSettings->RotationCompressionFormat;
-	// Translation compression format.
-	AnimationCompressionFormat TranslationCompressionFormat = AnimSettings->TranslationCompressionFormat;
-
-	UClass* CompressionAlgorithmClass = AnimSettings->DefaultCompressionAlgorithm;
-	if ( !CompressionAlgorithmClass )
-	{
-		// if can't find back out to bitwise
-		CompressionAlgorithmClass = UAnimCompress_BitwiseCompressOnly::StaticClass();
-		UE_LOG(LogAnimationCompression, Warning, TEXT("Couldn't find animation compression, default to AnimCompress_BitwiseCompressOnly") );
-	}
-
-	UAnimCompress* NewAlgorithm = NewObject<UAnimCompress>(GetTransientPackage(), CompressionAlgorithmClass);
-	NewAlgorithm->RotationCompressionFormat = RotationCompressionFormat;
-	NewAlgorithm->TranslationCompressionFormat = TranslationCompressionFormat;
-	NewAlgorithm->AddToRoot();
-	return NewAlgorithm;
-}
-
-} // namespace
-
-/**
- * @return		The default animation compression algorithm singleton, instantiating it if necessary.
- */
-UAnimCompress* FAnimationUtils::GetDefaultAnimationCompressionAlgorithm()
-{
-	static UAnimCompress* SAlgorithm = ConstructDefaultCompressionAlgorithm();
-	return SAlgorithm;
-}
-
-/**
- * Determines the current setting for world-space error tolerance in the animation compressor.
- * When requested, animation being compressed will also consider an alternative compression
- * method if the end result of that method produces less error than the AlternativeCompressionThreshold.
- * The default tolerance value is 0.0f (no alternatives allowed) but may be overridden using a field in the base engine INI file.
- *
- * @return				World-space error tolerance for considering an alternative compression method
- */
-float FAnimationUtils::GetAlternativeCompressionThreshold()
-{
-	// Allow the Engine INI file to provide a new override
-	return UAnimationSettings::Get()->AlternativeCompressionThreshold;
-}
-
 /**
  * Determines the current setting for recompressing all animations upon load. The default value 
  * is False, but may be overridden by an optional field in the base engine INI file. 
@@ -475,1227 +407,6 @@ bool FAnimationUtils::GetForcedRecompressionSetting()
 	return ForcedRecompressionSetting;
 }
 
-struct FAnimCompressionJobContext
-{
-	// Inputs
-	bool bForceBelowThreshold;
-	FAnimCompressContext* CompressContext;
-	UAnimCompress* CompressionAlgorithm;
-	
-	const FCompressibleAnimData* CompressibleAnimData;
-	FCompressibleAnimDataResult CompressionResult;
-
-	const TCHAR* CompressionName;
-	int32* WinningCompressor_Count;
-	float* WinningCompressor_Error;
-	int64* WinningCompressor_Margin;
-
-	// Outputs
-	AnimationErrorStats ErrorStats;
-	float PctSaving;
-
-	FAnimCompressionJobContext()
-	{
-		memset(this, 0, sizeof(FAnimCompressionJobContext));
-	}
-
-	void UpdatePctSaving(int64 OriginalSize)
-	{
-		PctSaving = (OriginalSize > 0) ? (100.f - (100.f * float(CompressionResult.GetApproxBoneCompressedSize()) / float(OriginalSize))) : 0.f;
-	}
-};
-
-static void TryCompressionInner(FAnimCompressionJobContext& JobContext, bool is_async);
-
-class FAsyncAnimCompressionTask
-{
-public:
-	FAsyncAnimCompressionTask(FAnimCompressionJobContext* JobContext_)
-		: JobContext(JobContext_)
-	{}
-
-	/** return the name of the task **/
-	static const TCHAR* GetTaskName() { return TEXT("FAsyncAnimCompressionTask"); }
-	FORCEINLINE static TStatId GetStatId() { RETURN_QUICK_DECLARE_CYCLE_STAT(FAsyncAnimCompressionTask, STATGROUP_TaskGraphTasks); }
-
-	static ENamedThreads::Type GetDesiredThread() { return ENamedThreads::AnyThread; }
-	static ESubsequentsMode::Type GetSubsequentsMode() { return ESubsequentsMode::TrackSubsequents; }
-
-	void DoTask(ENamedThreads::Type CurrentThread, const FGraphEventRef& MyCompletionGraphEvent)
-	{
-		TryCompressionInner(*JobContext, true);
-	}
-
-	FAnimCompressionJobContext* JobContext;
-};
-
-bool ShouldKeepNewCompressionMethod(const FAnimCompressionJobContext& JobContext, SIZE_T OriginalSize, SIZE_T CurrentSize, float WinningCompressorError, float MasterTolerance)
-{
-	const SIZE_T NewSize = JobContext.CompressionResult.GetApproxBoneCompressedSize();
-
-	/* compute the savings and compression error*/
-	const int64 MemorySavingsFromOriginal = OriginalSize - NewSize;
-	const int64 MemorySavingsFromPrevious = CurrentSize - NewSize;
-
-	/* figure out our new compression error*/
-	const AnimationErrorStats& ErrorStats = JobContext.ErrorStats;
-
-	const bool bLowersError = ErrorStats.MaxError < WinningCompressorError;
-	const bool bErrorUnderThreshold = ErrorStats.MaxError <= MasterTolerance;
-
-	/* keep it if it we want to force the error below the threshold and it reduces error */
-	bool bKeepNewCompressionMethod = false;
-	const bool bReducesErrorBelowThreshold = (bLowersError && (WinningCompressorError > MasterTolerance) && JobContext.bForceBelowThreshold);
-	bKeepNewCompressionMethod |= bReducesErrorBelowThreshold;
-	/* or if has an acceptable error and saves space  */
-	const bool bHasAcceptableErrorAndSavesSpace = bErrorUnderThreshold && (MemorySavingsFromPrevious > 0);
-	bKeepNewCompressionMethod |= bHasAcceptableErrorAndSavesSpace;
-	/* or if saves the same amount and an acceptable error that is lower than the previous best */
-	const bool bLowersErrorAndSavesSameOrBetter = bErrorUnderThreshold && bLowersError && (MemorySavingsFromPrevious >= 0);
-	bKeepNewCompressionMethod |= bLowersErrorAndSavesSameOrBetter;
-
-	UE_LOG(LogAnimationCompression, Verbose, TEXT("- %s - bytes saved(%i) (%.1f%%) from previous(%i) MaxError(%.2f) bLowersError(%d) %s"),
-		JobContext.CompressionName, MemorySavingsFromOriginal, JobContext.PctSaving, MemorySavingsFromPrevious, ErrorStats.MaxError, bLowersError, bKeepNewCompressionMethod ? TEXT("(**Best so far**)") : TEXT(""));
-
-	UE_LOG(LogAnimationCompression, Verbose, TEXT("    bReducesErrorBelowThreshold(%d) bHasAcceptableErrorAndSavesSpace(%d) bLowersErrorAndSavesSameOrBetter(%d)"),
-		bReducesErrorBelowThreshold, bHasAcceptableErrorAndSavesSpace, bLowersErrorAndSavesSameOrBetter);
-
-	UE_LOG(LogAnimationCompression, Verbose, TEXT("    WinningCompressorError(%f) MasterTolerance(%f) bForceBelowThreshold(%d) bErrorUnderThreshold(%d)"),
-		WinningCompressorError, MasterTolerance, JobContext.bForceBelowThreshold, bErrorUnderThreshold);
-	return bKeepNewCompressionMethod;
-}
-
-static void TryCompressionInner(FAnimCompressionJobContext& JobContext, bool bIsAsync)
-{
-#if WITH_EDITOR
-	/* try the alternative compressor	*/
-	JobContext.CompressionAlgorithm->Reduce(*JobContext.CompressibleAnimData, JobContext.CompressionResult);
-	FAnimationUtils::ComputeCompressionError(*JobContext.CompressibleAnimData, JobContext.CompressionResult, JobContext.ErrorStats);
-#endif	// WITH_EDITOR
-}
-
-#if WITH_EDITORONLY_DATA
-struct WinningCompressorStats
-{
-	SIZE_T CurrentSize;
-
-	int32* WinningCompressorCounter;
-	float* WinningCompressorErrorSum;
-	int64* WinningCompressorMarginalSavingsSum;
-	int64 WinningCompressorMarginalSavings;
-	FString WinningCompressorName;
-	int32 WinningCompressorSavings;
-	float PctSaving;
-	float WinningCompressorError;
-
-	WinningCompressorStats()
-	{
-		memset(this, 0, sizeof(WinningCompressorStats));
-	}
-};
-
-void HandlePostTryCompression(FAnimCompressionJobContext& JobContext, int64 OriginalSize, float MasterTolerance, WinningCompressorStats& OutCompressorStats, FCompressibleAnimDataResult& OutCompressedData, AnimationErrorStats& OutNewErrorStats)
-{
-	JobContext.UpdatePctSaving(OriginalSize);
-
-	bool bKeepNewCompressionMethod = ShouldKeepNewCompressionMethod(JobContext, OriginalSize, OutCompressorStats.CurrentSize, OutCompressorStats.WinningCompressorError, MasterTolerance);
-	if (bKeepNewCompressionMethod)
-	{
-		OutCompressedData = JobContext.CompressionResult;
-
-		OutNewErrorStats = JobContext.ErrorStats;
-
-		int64 NewSize = JobContext.CompressionResult.GetApproxBoneCompressedSize();
-		const int64 MemorySavingsFromOriginal = OriginalSize - NewSize;
-		const int64 MemorySavingsFromPrevious = OutCompressorStats.CurrentSize - NewSize;
-
-		OutCompressorStats.CurrentSize = NewSize;
-		OutCompressorStats.WinningCompressorMarginalSavings = MemorySavingsFromPrevious;
-		OutCompressorStats.WinningCompressorCounter = JobContext.WinningCompressor_Count; 
-		OutCompressorStats.WinningCompressorErrorSum = JobContext.WinningCompressor_Error; 
-		OutCompressorStats.WinningCompressorMarginalSavingsSum = JobContext.WinningCompressor_Margin; 
-		OutCompressorStats.WinningCompressorName = JobContext.CompressionName; 
-		OutCompressorStats.WinningCompressorSavings = MemorySavingsFromOriginal;
-		OutCompressorStats.WinningCompressorError = JobContext.ErrorStats.MaxError;
-		OutCompressorStats.PctSaving = JobContext.PctSaving;
-	}
-}
-#endif
-
-#define POPULATE_ANIM_COMPRESSION_JOB_CONTEXT(Name_, CompressionAlgorithm_, CompressContext_, JobContext_) \
-	(JobContext_).bForceBelowThreshold = bForceBelowThreshold; \
-	(JobContext_).CompressContext = &CompressContext_; \
-	(JobContext_).CompressibleAnimData = &CompressibleAnimData; \
-	(JobContext_).CompressionAlgorithm = DuplicateObject<UAnimCompress>((CompressionAlgorithm_), GetTransientPackage()); \
-	(JobContext_).CompressionAlgorithm->bEnableSegmenting = bEnableSegmenting; \
-	(JobContext_).CompressionAlgorithm->IdealNumFramesPerSegment = IdealNumFramesPerSegment; \
-	(JobContext_).CompressionAlgorithm->MaxNumFramesPerSegment = MaxNumFramesPerSegment; \
-	(JobContext_).CompressionName = TEXT(#Name_); \
-	(JobContext_).WinningCompressor_Count = &(Name_ ## CompressorWins); \
-	(JobContext_).WinningCompressor_Error = &(Name_ ## CompressorSumError); \
-	(JobContext_).WinningCompressor_Margin = &(Name_ ## CompressorWinMargin);
-
-#define TRYCOMPRESSION(Name, CompressionAlgorithm_) \
-	do \
-	{ \
-		FAnimCompressionJobContext JobContext; \
-		POPULATE_ANIM_COMPRESSION_JOB_CONTEXT(Name, CompressionAlgorithm_, CompressContext, JobContext); \
-		\
-		TryCompressionInner(JobContext, false); \
-		\
-		HandlePostTryCompression(JobContext, OriginalSize, MasterTolerance, CompressorStats, OutCompressedData, NewErrorStats); \
-	} \
-	while (false)
-
-#define TRYCOMPRESSION_SYNC(Name, CompressionAlgorithm_) TRYCOMPRESSION(Name, CompressionAlgorithm_)
-
-// TODO: Async compression is DISABLED for additive sequences because UAnimCompress_RemoveLinearKeys::ConvertFromRelativeSpace()
-// modifies the RAW data! This is bad... Even though we duplicate the anim sequence, some additive information isn't
-// copied over and it doesn't seem safe to generate it by calling UAnimSequence::BakeOutAdditiveIntoRawData()
-#define TRYCOMPRESSION_ASYNC(Name, CompressionAlgorithm_) \
-	do \
-	{ \
-		if (!CompressibleAnimData.bIsValidAdditive) \
-		{ \
-			FAnimCompressionJobContext* JobContext = new FAnimCompressionJobContext(); \
-			FAnimCompressContext CompressContextCopy(CompressContext); \
-			CompressContextCopy.CompressionSummary = FCompressionMemorySummary(false); \
-			POPULATE_ANIM_COMPRESSION_JOB_CONTEXT(Name, CompressionAlgorithm_, CompressContextCopy, *JobContext); \
-			AnimCompressionJobContexes.Add(JobContext); \
-			\
-			AnimCompressionTask_CompletionEvents.Add(TGraphTask<FAsyncAnimCompressionTask>::CreateTask(NULL, ENamedThreads::GameThread).ConstructAndDispatchWhenReady(JobContext)); \
-		} \
-		else \
-		{ \
-			TRYCOMPRESSION_SYNC(Name, CompressionAlgorithm_); \
-		} \
-	} \
-	while (false)
-
-#define LOG_COMPRESSION_STATUS(Name) \
-	UE_LOG(LogAnimationCompression, Log, TEXT("\t\tWins for '%32s': %4i\t\t%f\t%i bytes"), TEXT(#Name), Name ## CompressorWins, (Name ## CompressorWins > 0) ? Name ## CompressorSumError / Name ## CompressorWins : 0.0f, Name ## CompressorWinMargin)
-
-#define DECLARE_ANIM_COMP_ALGORITHM(Algorithm) \
-	static int32 Algorithm ## CompressorWins = 0; \
-	static float Algorithm ## CompressorSumError = 0.0f; \
-	static int64 Algorithm ## CompressorWinMargin = 0
-
-#if WITH_EDITORONLY_DATA
-
-static void WaitForAnimCompressionJobs(const FGraphEventArray& AnimCompressionTask_CompletionEvents)
-{
-	FTaskGraphInterface::Get().WaitUntilTasksComplete(AnimCompressionTask_CompletionEvents, ENamedThreads::GameThread);
-}
-
-static void ClearAnimCompressionJobs(FGraphEventArray& AnimCompressionTask_CompletionEvents, TArray<FAnimCompressionJobContext*>& AnimCompressionJobContexes)
-{
-	for (FAnimCompressionJobContext* Context : AnimCompressionJobContexes)
-	{
-		delete Context;
-	}
-	AnimCompressionJobContexes.Reset();
-	AnimCompressionTask_CompletionEvents.Reset();
-}
-
-static FAnimCompressionJobContext* FindBestAnimCompression(TArray<FAnimCompressionJobContext*>& AnimCompressionJobContexes, SIZE_T OriginalSize, SIZE_T CurrentSize, float WinningCompressorError, float MasterTolerance)
-{
-	FAnimCompressionJobContext* BestJobContext = nullptr;
-
-	for (FAnimCompressionJobContext* context : AnimCompressionJobContexes)
-	{
-		FAnimCompressionJobContext& JobContext = *context;
-		JobContext.UpdatePctSaving(OriginalSize);
-
-		bool bKeepNewCompressionMethod = ShouldKeepNewCompressionMethod(JobContext, OriginalSize, CurrentSize, WinningCompressorError, MasterTolerance);
-
-		if (bKeepNewCompressionMethod)
-		{
-			BestJobContext = context;
-			WinningCompressorError = JobContext.ErrorStats.MaxError;
-			CurrentSize = JobContext.CompressionResult.GetApproxBoneCompressedSize();
-		}
-	}
-
-	return BestJobContext;
-}
-
-static void UpdateAnimCompressionFromAsyncJobs(FCompressibleAnimDataResult& OutCompressedData, FGraphEventArray& AnimCompressionTask_CompletionEvents, TArray<FAnimCompressionJobContext*>& AnimCompressionJobContexes, SIZE_T OriginalSize, WinningCompressorStats& CompressorStats, float MasterTolerance)
-{
-	// Pick the best
-	const FAnimCompressionJobContext* BestJobContext = FindBestAnimCompression(AnimCompressionJobContexes, OriginalSize, CompressorStats.CurrentSize, CompressorStats.WinningCompressorError, MasterTolerance);
-
-	if (BestJobContext != nullptr)
-	{
-		const FAnimCompressionJobContext& JobContext = *BestJobContext;
-
-		/* Copy our data */
-		OutCompressedData = JobContext.CompressionResult;
-
-		const SIZE_T NewSize = OutCompressedData.GetApproxBoneCompressedSize();
-
-		const int64 MemorySavingsFromOriginal = OriginalSize - NewSize;
-		const SIZE_T MemorySavingsFromPrevious = CompressorStats.CurrentSize - NewSize;
-
-		CompressorStats.WinningCompressorMarginalSavings = MemorySavingsFromPrevious;
-		CompressorStats.WinningCompressorCounter = JobContext.WinningCompressor_Count;
-		CompressorStats.WinningCompressorErrorSum = JobContext.WinningCompressor_Error;
-		CompressorStats.WinningCompressorMarginalSavingsSum = JobContext.WinningCompressor_Margin;
-		CompressorStats.WinningCompressorName = JobContext.CompressionName;
-		CompressorStats.CurrentSize = NewSize;
-		CompressorStats.WinningCompressorSavings = MemorySavingsFromOriginal;
-		CompressorStats.WinningCompressorError = JobContext.ErrorStats.MaxError;
-	}
-
-	ClearAnimCompressionJobs(AnimCompressionTask_CompletionEvents, AnimCompressionJobContexes);
-}
-#endif	// WITH_EDITORONLY_DATA
-
-// Simple struct to calc duration of of scope
-struct FCompressionTimeElapsed
-{
-private:
-
-	// Time we started tracking
-	double StartTime;
-
-	// Where to store the elapsed time
-	double& Result;
-
-public:
-	FCompressionTimeElapsed(double& InResult)
-		: StartTime(FPlatformTime::Seconds())
-		, Result(InResult)
-	{}
-
-	~FCompressionTimeElapsed()
-	{
-		Result = FPlatformTime::Seconds() - StartTime;
-	}
-};
-
-/**
- * Utility function to compress an animation. If the animation is currently associated with a codec, it will be used to 
- * compress the animation. Otherwise, the default codec will be used. If AllowAlternateCompressor is true, an
- * alternative compression codec will also be tested. If the alternative codec produces better compression and 
- * the accuracy of the compressed animation remains within tolerances, the alternative codec will be used. 
- * See GetAlternativeCompressionThreshold for information on the tolerance value used.
- *
- * @param	AnimSet		The animset to compress.
- * @param	AllowAlternateCompressor	true if an alternative compressor is permitted.
- * @param	bOutput		If false don't generate output or compute memory savings.
- * @return				None.
- */
-void FAnimationUtils::CompressAnimSequence(const FCompressibleAnimData& CompressibleAnimData, FCompressibleAnimDataResult& OutCompressedData, FAnimCompressContext& CompressContext)
-{
-	if (FPlatformProperties::HasEditorOnlyData())
-	{
-		bool bOnlyCheckForMissingSkeletalMeshes = UAnimationSettings::Get()->bOnlyCheckForMissingSkeletalMeshes;
-		if (!bOnlyCheckForMissingSkeletalMeshes)
-		{
-			UAnimationSettings* AnimSetting = UAnimationSettings::Get();
-			bool bForceBelowThreshold = AnimSetting->bForceBelowThreshold;
-			bool bFirstRecompressUsingCurrentOrDefault = AnimSetting->bFirstRecompressUsingCurrentOrDefault;
-			bool bRaiseMaxErrorToExisting = AnimSetting->bRaiseMaxErrorToExisting;
-			// If we don't allow alternate compressors, and just want to recompress with default/existing, then make sure we do so.
-			if( !CompressContext.bAllowAlternateCompressor )
-			{
-				bFirstRecompressUsingCurrentOrDefault = true;
-			}
-
-			bool bTryExhaustiveSearch = AnimSetting->bTryExhaustiveSearch;
-			bool bEnableSegmenting = AnimSetting->bEnableSegmenting;
-			int32 IdealNumFramesPerSegment = 64;
-			int32 MaxNumFramesPerSegment = (IdealNumFramesPerSegment * 2) - 1;
-
-#if WITH_EDITORONLY_DATA
-			UAnimCompress_Automatic* AutoCompressionScheme = Cast<UAnimCompress_Automatic>(CompressibleAnimData.RequestedCompressionScheme);
-
-			if (AutoCompressionScheme != nullptr)
-			{
-				bTryExhaustiveSearch = AutoCompressionScheme->bTryExhaustiveSearch;
-				bEnableSegmenting = AutoCompressionScheme->bEnableSegmenting;
-				IdealNumFramesPerSegment = AutoCompressionScheme->IdealNumFramesPerSegment;
-				MaxNumFramesPerSegment = AutoCompressionScheme->MaxNumFramesPerSegment;
-			}
-			else if (CompressibleAnimData.RequestedCompressionScheme != nullptr)
-			{
-				bTryExhaustiveSearch = AnimSetting->bTryExhaustiveSearch;
-			}
-#endif
-
-			double CompressionTime = 0.0;
-			{
-				//Scoped timing of compression, make sure nothing else is added to this scope
-				FCompressionTimeElapsed TimeTracker(CompressionTime);
-				CompressAnimSequenceExplicit(
-					CompressibleAnimData,
-					OutCompressedData,
-					CompressContext,
-					CompressibleAnimData.AltCompressionErrorThreshold,
-					bFirstRecompressUsingCurrentOrDefault,
-					bForceBelowThreshold,
-					bRaiseMaxErrorToExisting,
-					bTryExhaustiveSearch,
-					bEnableSegmenting,
-					IdealNumFramesPerSegment,
-					MaxNumFramesPerSegment);
-			}
-
-			CompressContext.GatherPostCompressionStats(CompressibleAnimData, OutCompressedData, CompressionTime);
-		}
-	}
-}
-
-/**
- * Utility function to compress an animation. If the animation is currently associated with a codec, it will be used to 
- * compress the animation. Otherwise, the default codec will be used. If AllowAlternateCompressor is true, an
- * alternative compression codec will also be tested. If the alternative codec produces better compression and 
- * the accuracy of the compressed animation remains within tolerances, the alternative codec will be used. 
- * See GetAlternativeCompressionThreshold for information on the tolerance value used.
- *
- * @param	AnimSet		The animset to compress.
- * @param	MasterTolerance	The alternate error threshold (0.0 means don't try anything other than the current / default scheme)
- * @param	bOutput		If false don't generate output or compute memory savings.
- * @return				None.
- */
-void FAnimationUtils::CompressAnimSequenceExplicit(
-	const FCompressibleAnimData& CompressibleAnimData,
-	FCompressibleAnimDataResult& OutCompressedData,
-	FAnimCompressContext& CompressContext,
-	float MasterTolerance,
-	bool bFirstRecompressUsingCurrentOrDefault,
-	bool bForceBelowThreshold,
-	bool bRaiseMaxErrorToExisting,
-	bool bTryExhaustiveSearch,
-	bool bEnableSegmenting,
-	int32 IdealNumFramesPerSegment,
-	int32 MaxNumFramesPerSegment)
-{
-#if WITH_EDITORONLY_DATA
-	DECLARE_ANIM_COMP_ALGORITHM(BitwiseACF_Float96);
-	DECLARE_ANIM_COMP_ALGORITHM(BitwiseACF_Fixed48);
-	DECLARE_ANIM_COMP_ALGORITHM(BitwiseACF_IntervalFixed32);
-	DECLARE_ANIM_COMP_ALGORITHM(BitwiseACF_Fixed32);
-
-	DECLARE_ANIM_COMP_ALGORITHM(HalfOddACF_Float96);
-	DECLARE_ANIM_COMP_ALGORITHM(HalfOddACF_Fixed48);
-	DECLARE_ANIM_COMP_ALGORITHM(HalfOddACF_IntervalFixed32);
-	DECLARE_ANIM_COMP_ALGORITHM(HalfOddACF_Fixed32);
-
-	DECLARE_ANIM_COMP_ALGORITHM(HalfEvenACF_Float96);
-	DECLARE_ANIM_COMP_ALGORITHM(HalfEvenACF_Fixed48);
-	DECLARE_ANIM_COMP_ALGORITHM(HalfEvenACF_IntervalFixed32);
-	DECLARE_ANIM_COMP_ALGORITHM(HalfEvenACF_Fixed32);
-
-	DECLARE_ANIM_COMP_ALGORITHM(LinearACF_Float96);
-	DECLARE_ANIM_COMP_ALGORITHM(LinearACF_Fixed48);
-	DECLARE_ANIM_COMP_ALGORITHM(LinearACF_IntervalFixed32);
-	DECLARE_ANIM_COMP_ALGORITHM(LinearACF_Fixed32);
-
-	DECLARE_ANIM_COMP_ALGORITHM(Progressive_PerTrack);
-	DECLARE_ANIM_COMP_ALGORITHM(Bitwise_PerTrack);
-	DECLARE_ANIM_COMP_ALGORITHM(Linear_PerTrack);
-	DECLARE_ANIM_COMP_ALGORITHM(Adaptive1_LinPerTrack);
-	DECLARE_ANIM_COMP_ALGORITHM(Adaptive1_LinPerTrackNoRT);
-
-	DECLARE_ANIM_COMP_ALGORITHM(Adaptive2_PerTrack);
-	DECLARE_ANIM_COMP_ALGORITHM(Adaptive2_LinPerTrack);
-	DECLARE_ANIM_COMP_ALGORITHM(Adaptive2_LinPerTrackNoRT);
-
-	DECLARE_ANIM_COMP_ALGORITHM(Downsample20Hz_PerTrack);
-	DECLARE_ANIM_COMP_ALGORITHM(Downsample15Hz_PerTrack);
-	DECLARE_ANIM_COMP_ALGORITHM(Downsample10Hz_PerTrack);
-	DECLARE_ANIM_COMP_ALGORITHM(Downsample5Hz_PerTrack);
-
-	DECLARE_ANIM_COMP_ALGORITHM(Adaptive1_15Hz_LinPerTrack);
-	DECLARE_ANIM_COMP_ALGORITHM(Adaptive1_10Hz_LinPerTrack);
-	DECLARE_ANIM_COMP_ALGORITHM(Adaptive1_5Hz_LinPerTrack);
-
-	DECLARE_ANIM_COMP_ALGORITHM(Adaptive2_15Hz_LinPerTrack);
-	DECLARE_ANIM_COMP_ALGORITHM(Adaptive2_10Hz_LinPerTrack);
-
-	DECLARE_ANIM_COMP_ALGORITHM(Linear_PerTrackExp1);
-	DECLARE_ANIM_COMP_ALGORITHM(Linear_PerTrackExp2);
-
-	USkeleton* Skeleton = CompressibleAnimData.Skeleton;
-	check (Skeleton);
-	if (Skeleton->HasAnyFlags(RF_NeedLoad))
-	{
-		Skeleton->GetLinker()->Preload(Skeleton);
-	}
-
-	static int32 TotalRecompressions = 0;
-
-	static int32 TotalNoWinnerRounds = 0;
-
-	static int32 AlternativeCompressorLossesFromSize = 0;
-	static int32 AlternativeCompressorLossesFromError = 0;
-	static int32 AlternativeCompressorSavings = 0;
-	static int64 TotalSizeBefore = 0;
-	static int64 TotalSizeNow = 0;
-	static int64 TotalUncompressed = 0;
-
-	const int32 NumRawDataTracks = CompressibleAnimData.RawAnimationData.Num();
-
-	// we must have raw data to continue
-	if(NumRawDataTracks > 0 )
-	{
-		// See if we're trying alternate compressors
-		// If compression Scheme is automatic, then we definitely want to 'TryAlternateCompressor'.
-		const bool bIsAutomatic = CompressibleAnimData.RequestedCompressionScheme && CompressibleAnimData.RequestedCompressionScheme->IsA(UAnimCompress_Automatic::StaticClass());
-		const bool bTryAlternateCompressor = bIsAutomatic || CompressContext.bAllowAlternateCompressor;
-
-		// We shouldn't override as this value can come from automatic compression sequences
-		// but that was broken by CL 3489273
-		// Preserving override behaviour till issue can be properly addressed
-		MasterTolerance = CompressibleAnimData.AltCompressionErrorThreshold;
-
-		AnimationErrorStats TrueOriginalErrorStats;
-		FAnimationUtils::ComputeCompressionError(CompressibleAnimData, OutCompressedData, TrueOriginalErrorStats);
-
-		AnimationErrorStats OriginalErrorStats;
-
-		int32 AfterOriginalRecompression = 0;
-		if( (bFirstRecompressUsingCurrentOrDefault && !bTryAlternateCompressor) || (OutCompressedData.CompressedByteStream.Num() == 0))
-		{
-			UAnimCompress* OriginalCompressionAlgorithm = CompressibleAnimData.RequestedCompressionScheme ? CompressibleAnimData.RequestedCompressionScheme : FAnimationUtils::GetDefaultAnimationCompressionAlgorithm();
-
-			// Automatic compression brings us back here, so don't create an infinite loop and pick bitwise compress instead.
-			if (!OriginalCompressionAlgorithm || OriginalCompressionAlgorithm->IsA(UAnimCompress_Automatic::StaticClass()))
-			{
-				UAnimCompress* CompressionAlgorithm = NewObject<UAnimCompress_BitwiseCompressOnly>();
-				if (OriginalCompressionAlgorithm != nullptr)
-				{
-					// Keep the same segmenting settings
-					CompressionAlgorithm->bEnableSegmenting = OriginalCompressionAlgorithm->bEnableSegmenting;
-					CompressionAlgorithm->IdealNumFramesPerSegment = OriginalCompressionAlgorithm->IdealNumFramesPerSegment;
-					CompressionAlgorithm->MaxNumFramesPerSegment = OriginalCompressionAlgorithm->MaxNumFramesPerSegment;
-				}
-
-				OriginalCompressionAlgorithm = CompressionAlgorithm;
-			}
-
-			UE_LOG(LogAnimationCompression, Log, TEXT("Recompressing (%s) using current/default (%s) bFirstRecompressUsingCurrentOrDefault(%d) bTryAlternateCompressor(%d) IsCompressedDataValid(%d)"),
-				*CompressibleAnimData.FullName,
-				*OriginalCompressionAlgorithm->GetName(),
-				bFirstRecompressUsingCurrentOrDefault,
-				bTryAlternateCompressor,
-				OutCompressedData.IsCompressedDataValid());
-
-			OriginalCompressionAlgorithm->Reduce(CompressibleAnimData, OutCompressedData);
-			
-			AfterOriginalRecompression = OutCompressedData.GetApproxBoneCompressedSize();
-
-			// figure out our current compression error
-			FAnimationUtils::ComputeCompressionError(CompressibleAnimData, OutCompressedData, OriginalErrorStats);
-		}
-		else
-		{
-			AfterOriginalRecompression = OutCompressedData.GetApproxBoneCompressedSize();
-			OriginalErrorStats = TrueOriginalErrorStats;
-		}
-
-		//For logging
-		FString OriginalKeyEncodingFormat = GetAnimationKeyFormatString(static_cast<AnimationKeyFormat>(OutCompressedData.KeyEncodingFormat));
-		FString OriginalRotationFormat = FAnimationUtils::GetAnimationCompressionFormatString(static_cast<AnimationCompressionFormat>(OutCompressedData.TranslationCompressionFormat));
-		FString OriginalTranslationFormat = FAnimationUtils::GetAnimationCompressionFormatString(static_cast<AnimationCompressionFormat>(OutCompressedData.RotationCompressionFormat));
-
-		// Get the current size
-		SIZE_T OriginalSize = AfterOriginalRecompression;
-		TotalSizeBefore += OriginalSize;
-
-		// Estimate total uncompressed
-		TotalUncompressed += ((sizeof(FVector) + sizeof(FQuat) + sizeof(FVector)) * NumRawDataTracks * CompressibleAnimData.NumFrames);
-
-		// Check for global permission to try an alternative compressor
-		// And it's valid to manually recompress animations
-		if( bTryAlternateCompressor )
-		{
-			ensure(OutCompressedData.CompressedByteStream.Num() > 0);
-
-			AnimationErrorStats NewErrorStats = OriginalErrorStats;
-			if (bRaiseMaxErrorToExisting)
-			{
-				if (NewErrorStats.MaxError > MasterTolerance)
-				{
-					UE_LOG(LogAnimationCompression, Log, TEXT("  Boosting MasterTolerance to %f, as existing MaxDiff was higher than %f and bRaiseMaxErrorToExisting=true"), NewErrorStats.MaxError, MasterTolerance);
-					MasterTolerance = NewErrorStats.MaxError;
-				}
-			}
-
-			// count all attempts for debugging
-			++TotalRecompressions;
-
-			WinningCompressorStats CompressorStats;
-			CompressorStats.CurrentSize = OriginalSize;
-			CompressorStats.WinningCompressorError = OriginalErrorStats.MaxError;
-
-			FGraphEventArray AnimCompressionTask_CompletionEvents;
-			TArray<FAnimCompressionJobContext*> AnimCompressionJobContexes;
-
-			if (!bTryExhaustiveSearch)
-			{
-				// Dispatch our async compression
-				{
-					{
-						// Adaptive error through probing the effect of perturbations at each track
-						UAnimCompress_PerTrackCompression* NewPerTrackCompressor = NewObject<UAnimCompress_PerTrackCompression>();
-						NewPerTrackCompressor->bUseAdaptiveError2 = true;
-						NewPerTrackCompressor->MaxPosDiffBitwise = 0.05f;
-						NewPerTrackCompressor->MaxAngleDiffBitwise = 0.02f;
-						NewPerTrackCompressor->MaxScaleDiffBitwise = 0.00005f;
-
-						TRYCOMPRESSION_ASYNC(Adaptive2_PerTrack, NewPerTrackCompressor);
-
-						NewPerTrackCompressor->bActuallyFilterLinearKeys = true;
-						NewPerTrackCompressor->bRetarget = true;
-						TRYCOMPRESSION_ASYNC(Adaptive2_LinPerTrack, NewPerTrackCompressor);
-
-						NewPerTrackCompressor->bActuallyFilterLinearKeys = true;
-						NewPerTrackCompressor->bRetarget = false;
-						TRYCOMPRESSION_ASYNC(Adaptive2_LinPerTrackNoRT, NewPerTrackCompressor);
-					}
-
-					{
-						UAnimCompress_PerTrackCompression* PerTrackCompressor = NewObject<UAnimCompress_PerTrackCompression>();
-						PerTrackCompressor->bUseAdaptiveError = true;
-
-						if (CompressibleAnimData.NumFrames > 1)
-						{
-							PerTrackCompressor->bActuallyFilterLinearKeys = true;
-							PerTrackCompressor->bRetarget = true;
-
-							PerTrackCompressor->MaxPosDiff = 0.1f;
-							// 						PerTrackCompressor->MaxAngleDiff = 0.1;
-							PerTrackCompressor->MaxScaleDiff = 0.00001f;
-							PerTrackCompressor->ParentingDivisor = 2.0f;
-							PerTrackCompressor->ParentingDivisorExponent = 1.0f;
-							TRYCOMPRESSION_ASYNC(Linear_PerTrackExp1, PerTrackCompressor);
-						}
-					}
-
-					{
-						UAnimCompress_PerTrackCompression* PerTrackCompressor = NewObject<UAnimCompress_PerTrackCompression>();
-
-						// Straight PerTrackCompression, no key decimation and no linear key removal
-						TRYCOMPRESSION_ASYNC(Bitwise_PerTrack, PerTrackCompressor);
-						PerTrackCompressor->bUseAdaptiveError = true;
-
-						// Full blown linear
-						PerTrackCompressor->bActuallyFilterLinearKeys = true;
-						PerTrackCompressor->bRetarget = true;
-						TRYCOMPRESSION_ASYNC(Linear_PerTrack, PerTrackCompressor);
-
-						// Adaptive retargetting based on height within the skeleton
-						PerTrackCompressor->bActuallyFilterLinearKeys = true;
-						PerTrackCompressor->bRetarget = false;
-						PerTrackCompressor->ParentingDivisor = 2.0f;
-						PerTrackCompressor->ParentingDivisorExponent = 1.6f;
-						TRYCOMPRESSION_ASYNC(Adaptive1_LinPerTrackNoRT, PerTrackCompressor);
-					}
-
-					{
-						UAnimCompress_PerTrackCompression* PerTrackCompressor = NewObject<UAnimCompress_PerTrackCompression>();
-						PerTrackCompressor->bUseAdaptiveError = true;
-
-						// Try the decimation algorithms
-						if (CompressibleAnimData.NumFrames >= PerTrackCompressor->MinKeysForResampling)
-						{
-							PerTrackCompressor->bActuallyFilterLinearKeys = false;
-							PerTrackCompressor->bRetarget = false;
-							PerTrackCompressor->bUseAdaptiveError = false;
-							PerTrackCompressor->bResampleAnimation = true;
-
-							// Try PerTrackCompression, downsample to 5 Hz
-							PerTrackCompressor->ResampledFramerate = 5.0f;
-							TRYCOMPRESSION_ASYNC(Downsample5Hz_PerTrack, PerTrackCompressor);
-						}
-					}
-
-					if (CompressibleAnimData.NumFrames > 1)
-					{
-						UAnimCompress_RemoveLinearKeys* LinearKeyRemover = NewObject<UAnimCompress_RemoveLinearKeys>();
-						{
-							// Try ACF_Float96NoW
-							LinearKeyRemover->RotationCompressionFormat = ACF_Float96NoW;
-							LinearKeyRemover->TranslationCompressionFormat = ACF_None;
-							TRYCOMPRESSION_ASYNC(LinearACF_Float96, LinearKeyRemover);
-						}
-					}
-
-					{
-						UAnimCompress_BitwiseCompressOnly* BitwiseCompressor = NewObject<UAnimCompress_BitwiseCompressOnly>();
-
-						// Try ACF_Float96NoW
-						BitwiseCompressor->RotationCompressionFormat = ACF_Float96NoW;
-						BitwiseCompressor->TranslationCompressionFormat = ACF_None;
-						TRYCOMPRESSION_ASYNC(BitwiseACF_Float96, BitwiseCompressor);
-
-						// Try ACF_Fixed48NoW
-						BitwiseCompressor->RotationCompressionFormat = ACF_Fixed48NoW;
-						BitwiseCompressor->TranslationCompressionFormat = ACF_None;
-						TRYCOMPRESSION_ASYNC(BitwiseACF_Fixed48, BitwiseCompressor);
-					}
-				}
-
-				WaitForAnimCompressionJobs(AnimCompressionTask_CompletionEvents);
-				UpdateAnimCompressionFromAsyncJobs(OutCompressedData, AnimCompressionTask_CompletionEvents, AnimCompressionJobContexes, OriginalSize, CompressorStats, MasterTolerance);
-			}
-			else
-			{
-				// Prepare to compress
-				UE_LOG(LogAnimationCompression, Log, TEXT("Compressing %s (%s)\n\tSkeleton: %s\n\tOriginal Size: %i   MaxDiff: %f"),
-					*CompressibleAnimData.Name,
-					*CompressibleAnimData.FullName,
-					Skeleton ? *Skeleton->GetFName().ToString() : TEXT("NULL - Not all compression techniques can be used!"),
-					OriginalSize,
-					TrueOriginalErrorStats.MaxError);
-
-				UE_LOG(LogAnimationCompression, Log, TEXT("Original Key Encoding: %s\n\tOriginal Rotation Format: %s\n\tOriginal Translation Format: %s\n\tNumFrames: %i\n\tSequenceLength: %f (%2.1f fps)"),
-					*OriginalKeyEncodingFormat,
-					*OriginalRotationFormat,
-					*OriginalTranslationFormat,
-					OutCompressedData.CompressedNumberOfFrames,
-					CompressibleAnimData.SequenceLength,
-					(OutCompressedData.CompressedNumberOfFrames > 1) ? (OutCompressedData.CompressedNumberOfFrames -1) / CompressibleAnimData.SequenceLength : DEFAULT_SAMPLERATE);
-
-				if (bFirstRecompressUsingCurrentOrDefault)
-				{
-					UE_LOG(LogAnimationCompression, Log, TEXT("Recompressed using current/default\n\tRecompress Size: %i   MaxDiff: %f\n\tRecompress Scheme: %s"),
-						AfterOriginalRecompression,
-						OriginalErrorStats.MaxError,
-						OutCompressedData.CompressionScheme ? *OutCompressedData.CompressionScheme->GetClass()->GetName() : TEXT("NULL"));
-				}
-
-				// Progressive Algorithm
-				{
-					UAnimCompress_PerTrackCompression* PerTrackCompressor = NewObject<UAnimCompress_PerTrackCompression>();
-
-					// Start not too aggressive
-// 					PerTrackCompressor->MaxPosDiffBitwise /= 10.f;
-// 					PerTrackCompressor->MaxAngleDiffBitwise /= 10.f;
-// 					PerTrackCompressor->MaxScaleDiffBitwise /= 10.f;
-					PerTrackCompressor->bUseAdaptiveError2 = true;
-
-					// Try default compressor first
-					TRYCOMPRESSION(Progressive_PerTrack, PerTrackCompressor);
-
-					if( NewErrorStats.MaxError >= MasterTolerance )
-					{
-						UE_LOG(LogAnimationCompression, Log, TEXT("\tStandard bitwise compressor too aggressive, lower default settings."));
-					}
-					else
-					{
-						// First, start by finding most downsampling factor.
-						if (CompressibleAnimData.NumFrames >= PerTrackCompressor->MinKeysForResampling)
-						{
-							PerTrackCompressor->bResampleAnimation = true;
-				
-							// Try PerTrackCompression, down sample to 5 Hz
-							PerTrackCompressor->ResampledFramerate = 5.0f;
-							UE_LOG(LogAnimationCompression, Log, TEXT("\tResampledFramerate: %f"), PerTrackCompressor->ResampledFramerate);
-							TRYCOMPRESSION(Progressive_PerTrack, PerTrackCompressor);
-
-							// If too much error, try 6Hz
-							if( NewErrorStats.MaxError >= MasterTolerance )
-							{
-								PerTrackCompressor->ResampledFramerate = 6.0f;
-								UE_LOG(LogAnimationCompression, Log, TEXT("\tResampledFramerate: %f"), PerTrackCompressor->ResampledFramerate);
-								TRYCOMPRESSION(Progressive_PerTrack, PerTrackCompressor);
-
-								// if too much error go 10Hz, 15Hz, 20Hz.
-								if( NewErrorStats.MaxError >= MasterTolerance )
-								{
-									PerTrackCompressor->ResampledFramerate = 5.0f;
-									// Keep trying until we find something that works (or we just don't downsample)
-									while( PerTrackCompressor->ResampledFramerate < 20.f && NewErrorStats.MaxError >= MasterTolerance )
-									{
-										PerTrackCompressor->ResampledFramerate += 5.f;
-										UE_LOG(LogAnimationCompression, Log, TEXT("\tResampledFramerate: %f"), PerTrackCompressor->ResampledFramerate);
-										TRYCOMPRESSION(Progressive_PerTrack, PerTrackCompressor);
-									}
-								}
-							}
-							
-							// Give up downsampling if it didn't work.
-							if( NewErrorStats.MaxError >= MasterTolerance )
-							{
-								UE_LOG(LogAnimationCompression, Log, TEXT("\tDownsampling didn't work."));
-								PerTrackCompressor->bResampleAnimation = false;
-							}
-						}
-						
-						// Now do Linear Key Removal
-						if (CompressibleAnimData.NumFrames > 1)
-						{
-							PerTrackCompressor->bActuallyFilterLinearKeys = true;
-							PerTrackCompressor->bRetarget = true;
-							
-							int32 const TestSteps = 16;
-							float const MaxScale = 2^(TestSteps);
-
-							// Start with the least aggressive first. if that one doesn't succeed, don't bother going through all the steps.
-							PerTrackCompressor->MaxPosDiff /= MaxScale;
-							PerTrackCompressor->MaxAngleDiff /= MaxScale;
-							PerTrackCompressor->MaxScaleDiff /= MaxScale;
-							PerTrackCompressor->MaxEffectorDiff /= MaxScale;
-							PerTrackCompressor->MinEffectorDiff /= MaxScale;
-							PerTrackCompressor->EffectorDiffSocket /= MaxScale;
-							UE_LOG(LogAnimationCompression, Log, TEXT("\tLinearKeys. MaxPosDiff: %f, MaxAngleDiff: %f, MaxScaleDiff : %f"), PerTrackCompressor->MaxPosDiff, PerTrackCompressor->MaxAngleDiff, PerTrackCompressor->MaxScaleDiff);
-							TRYCOMPRESSION(Progressive_PerTrack, PerTrackCompressor);							
-							PerTrackCompressor->MaxPosDiff *= MaxScale;
-							PerTrackCompressor->MaxAngleDiff *= MaxScale;
-							PerTrackCompressor->MaxScaleDiff *= MaxScale;
-							PerTrackCompressor->MaxEffectorDiff *= MaxScale;
-							PerTrackCompressor->MinEffectorDiff *= MaxScale;
-							PerTrackCompressor->EffectorDiffSocket *= MaxScale;
-
-							if( NewErrorStats.MaxError < MasterTolerance )
-							{
-								// Start super aggressive, and go down until we find something that works.
-								UE_LOG(LogAnimationCompression, Log, TEXT("\tLinearKeys. MaxPosDiff: %f, MaxAngleDiff: %f, MaxScaleDiff : %f"), PerTrackCompressor->MaxPosDiff, PerTrackCompressor->MaxAngleDiff, PerTrackCompressor->MaxScaleDiff);
-								TRYCOMPRESSION(Progressive_PerTrack, PerTrackCompressor);
-
-								for(int32 Step=0; Step<TestSteps && (NewErrorStats.MaxError >= MasterTolerance); Step++)
-								{
-									PerTrackCompressor->MaxPosDiff /= 2.f;
-									PerTrackCompressor->MaxAngleDiff /= 2.f;
-									PerTrackCompressor->MaxScaleDiff /= 2.f;
-									PerTrackCompressor->MaxEffectorDiff /= 2.f;
-									PerTrackCompressor->MinEffectorDiff /= 2.f;
-									PerTrackCompressor->EffectorDiffSocket /= 2.f;
-									UE_LOG(LogAnimationCompression, Log, TEXT("\tLinearKeys. MaxPosDiff: %f, MaxAngleDiff: %f, MaxScaleDiff : %f"), PerTrackCompressor->MaxPosDiff, PerTrackCompressor->MaxAngleDiff, PerTrackCompressor->MaxScaleDiff);
-									TRYCOMPRESSION(Progressive_PerTrack, PerTrackCompressor);
-								}
-							}
-
-							// Give up Linear Key Compression if it didn't work
-							if( NewErrorStats.MaxError >= MasterTolerance )
-							{
-								PerTrackCompressor->bActuallyFilterLinearKeys = false;
-								PerTrackCompressor->bRetarget = false;
-							}
-						}
-
-						// Finally tighten up bitwise compression
-						PerTrackCompressor->MaxPosDiffBitwise *= 10.f;
-						PerTrackCompressor->MaxAngleDiffBitwise *= 10.f;
-						PerTrackCompressor->MaxScaleDiffBitwise *= 10.f;
-						{
-							int32 const TestSteps = 16;
-							float const MaxScale = 2^(TestSteps/2);
-
-							PerTrackCompressor->MaxPosDiffBitwise *= MaxScale;
-							PerTrackCompressor->MaxAngleDiffBitwise *= MaxScale;
-							PerTrackCompressor->MaxScaleDiffBitwise *= MaxScale;
-							UE_LOG(LogAnimationCompression, Log, TEXT("\tBitwise. MaxPosDiffBitwise: %f, MaxAngleDiffBitwise: %f, MaxScaleDiffBitwise: %f"), PerTrackCompressor->MaxPosDiffBitwise, PerTrackCompressor->MaxAngleDiffBitwise, PerTrackCompressor->MaxScaleDiffBitwise);
-							TRYCOMPRESSION(Progressive_PerTrack, PerTrackCompressor);
-							PerTrackCompressor->MaxPosDiffBitwise /= 2.f;
-							PerTrackCompressor->MaxAngleDiffBitwise /= 2.f;
-							PerTrackCompressor->MaxScaleDiffBitwise /= 2.f;
-							for(int32 Step=0; Step<TestSteps && (NewErrorStats.MaxError >= MasterTolerance) && (PerTrackCompressor->MaxPosDiffBitwise >= PerTrackCompressor->MaxZeroingThreshold); Step++)
-							{
-								UE_LOG(LogAnimationCompression, Log, TEXT("\tBitwise. MaxPosDiffBitwise: %f, MaxAngleDiffBitwise: %f, MaxScaleDiffBitwise: %f"), PerTrackCompressor->MaxPosDiffBitwise, PerTrackCompressor->MaxAngleDiffBitwise, PerTrackCompressor->MaxScaleDiffBitwise);
-								TRYCOMPRESSION(Progressive_PerTrack, PerTrackCompressor);
-								PerTrackCompressor->MaxPosDiffBitwise /= 2.f;
-								PerTrackCompressor->MaxAngleDiffBitwise /= 2.f;
-								PerTrackCompressor->MaxScaleDiffBitwise /= 2.f;
-							}
-						}
-					}
-				}
-
-				// Start with Bitwise Compress only
-				{
-					UAnimCompress_BitwiseCompressOnly* BitwiseCompressor = NewObject<UAnimCompress_BitwiseCompressOnly>();
-
-					// Try ACF_Float96NoW
-					BitwiseCompressor->RotationCompressionFormat = ACF_Float96NoW;
-					BitwiseCompressor->TranslationCompressionFormat = ACF_None;
-					TRYCOMPRESSION_ASYNC(BitwiseACF_Float96, BitwiseCompressor);
-
-					// Try ACF_Fixed48NoW
-					BitwiseCompressor->RotationCompressionFormat = ACF_Fixed48NoW;
-					BitwiseCompressor->TranslationCompressionFormat = ACF_None;
-					TRYCOMPRESSION_ASYNC(BitwiseACF_Fixed48,BitwiseCompressor);
-
-// 32bits currently unusable due to creating too much error
-// 					// Try ACF_IntervalFixed32NoW
-// 					BitwiseCompressor->RotationCompressionFormat = ACF_IntervalFixed32NoW;
-// 					BitwiseCompressor->TranslationCompressionFormat = ACF_None;
-// 					TRYCOMPRESSION(BitwiseACF_IntervalFixed32,BitwiseCompressor);
-// 
-// 					// Try ACF_Fixed32NoW
-// 					BitwiseCompressor->RotationCompressionFormat = ACF_Fixed32NoW;
-// 					BitwiseCompressor->TranslationCompressionFormat = ACF_None;
-// 					TRYCOMPRESSION(BitwiseACF_Fixed32,BitwiseCompressor);
-				}
-
-				// Start with Bitwise Compress only
-				// this compressor has a minimum number of frames requirement. So no need to go there if we don't meet that...
-				{
-					UAnimCompress_RemoveEverySecondKey* RemoveEveryOtherKeyCompressor = NewObject<UAnimCompress_RemoveEverySecondKey>();
-					if(CompressibleAnimData.NumFrames > RemoveEveryOtherKeyCompressor->MinKeys )
-					{
-						RemoveEveryOtherKeyCompressor->bStartAtSecondKey = false;
-						{
-							// Try ACF_Float96NoW
-							RemoveEveryOtherKeyCompressor->RotationCompressionFormat = ACF_Float96NoW;	
-							RemoveEveryOtherKeyCompressor->TranslationCompressionFormat = ACF_None;		
-							TRYCOMPRESSION_ASYNC(HalfOddACF_Float96, RemoveEveryOtherKeyCompressor);
-
-							// Try ACF_Fixed48NoW
-							RemoveEveryOtherKeyCompressor->RotationCompressionFormat = ACF_Fixed48NoW;	
-							RemoveEveryOtherKeyCompressor->TranslationCompressionFormat = ACF_None;		
-							TRYCOMPRESSION_ASYNC(HalfOddACF_Fixed48, RemoveEveryOtherKeyCompressor);
-
-// 32bits currently unusable due to creating too much error
-// 							// Try ACF_IntervalFixed32NoW
-// 							RemoveEveryOtherKeyCompressor->RotationCompressionFormat = ACF_IntervalFixed32NoW;	
-// 							RemoveEveryOtherKeyCompressor->TranslationCompressionFormat = ACF_None;		
-// 							TRYCOMPRESSION(HalfOddACF_IntervalFixed32, RemoveEveryOtherKeyCompressor);
-// 
-// 							// Try ACF_Fixed32NoW
-// 							RemoveEveryOtherKeyCompressor->RotationCompressionFormat = ACF_Fixed32NoW;	
-// 							RemoveEveryOtherKeyCompressor->TranslationCompressionFormat = ACF_None;		
-// 							TRYCOMPRESSION(HalfOddACF_Fixed32, RemoveEveryOtherKeyCompressor);
-						}
-						RemoveEveryOtherKeyCompressor->bStartAtSecondKey = true;
-						{
-							// Try ACF_Float96NoW
-							RemoveEveryOtherKeyCompressor->RotationCompressionFormat = ACF_Float96NoW;	
-							RemoveEveryOtherKeyCompressor->TranslationCompressionFormat = ACF_None;		
-							TRYCOMPRESSION_ASYNC(HalfEvenACF_Float96,RemoveEveryOtherKeyCompressor);
-
-							// Try ACF_Fixed48NoW
-							RemoveEveryOtherKeyCompressor->RotationCompressionFormat = ACF_Fixed48NoW;	
-							RemoveEveryOtherKeyCompressor->TranslationCompressionFormat = ACF_None;		
-							TRYCOMPRESSION_ASYNC(HalfEvenACF_Fixed48,RemoveEveryOtherKeyCompressor);
-
-// 32bits currently unusable due to creating too much error
-// 							// Try ACF_IntervalFixed32NoW
-// 							RemoveEveryOtherKeyCompressor->RotationCompressionFormat = ACF_IntervalFixed32NoW;	
-// 							RemoveEveryOtherKeyCompressor->TranslationCompressionFormat = ACF_None;		
-// 							TRYCOMPRESSION(HalfEvenACF_IntervalFixed32,RemoveEveryOtherKeyCompressor);
-// 
-// 							// Try ACF_Fixed32NoW
-// 							RemoveEveryOtherKeyCompressor->RotationCompressionFormat = ACF_Fixed32NoW;	
-// 							RemoveEveryOtherKeyCompressor->TranslationCompressionFormat = ACF_None;		
-// 							TRYCOMPRESSION(HalfEvenACF_Fixed32,RemoveEveryOtherKeyCompressor);
-						}
-					}
-				}
-
-				// construct the proposed compressor		
-				if(CompressibleAnimData.NumFrames > 1 )
-				{
-					UAnimCompress_RemoveLinearKeys* LinearKeyRemover = NewObject<UAnimCompress_RemoveLinearKeys>();
-					{
-						// Try ACF_Float96NoW
-						LinearKeyRemover->RotationCompressionFormat = ACF_Float96NoW;
-						LinearKeyRemover->TranslationCompressionFormat = ACF_None;	
-						TRYCOMPRESSION_ASYNC(LinearACF_Float96,LinearKeyRemover);
-
-						// Try ACF_Fixed48NoW
-						LinearKeyRemover->RotationCompressionFormat = ACF_Fixed48NoW;
-						LinearKeyRemover->TranslationCompressionFormat = ACF_None;	
-						TRYCOMPRESSION_ASYNC(LinearACF_Fixed48,LinearKeyRemover);
-
-// Error is too bad w/ 32bits
-// 						// Try ACF_IntervalFixed32NoW
-// 						LinearKeyRemover->RotationCompressionFormat = ACF_IntervalFixed32NoW;
-// 						LinearKeyRemover->TranslationCompressionFormat = ACF_None;
-// 						TRYCOMPRESSION(LinearACF_IntervalFixed32,LinearKeyRemover);
-// 
-// 						// Try ACF_Fixed32NoW
-// 						LinearKeyRemover->RotationCompressionFormat = ACF_Fixed32NoW;
-// 						LinearKeyRemover->TranslationCompressionFormat = ACF_None;
-// 						TRYCOMPRESSION(LinearACF_Fixed32,LinearKeyRemover);
-					}
-				}
-
-				{
-					UAnimCompress_PerTrackCompression* PerTrackCompressor = NewObject<UAnimCompress_PerTrackCompression>();
-
-					// Straight PerTrackCompression, no key decimation and no linear key removal
-					TRYCOMPRESSION_ASYNC(Bitwise_PerTrack, PerTrackCompressor);
-					PerTrackCompressor->bUseAdaptiveError = true;
-
-					// Full blown linear
-					PerTrackCompressor->bActuallyFilterLinearKeys = true;
-					PerTrackCompressor->bRetarget = true;
-					TRYCOMPRESSION_ASYNC(Linear_PerTrack, PerTrackCompressor);
-
-					// Adaptive retargetting based on height within the skeleton
-					PerTrackCompressor->bActuallyFilterLinearKeys = true;
-					PerTrackCompressor->bRetarget = false;
-					PerTrackCompressor->ParentingDivisor = 2.0f;
-					PerTrackCompressor->ParentingDivisorExponent = 1.6f;
-					TRYCOMPRESSION_ASYNC(Adaptive1_LinPerTrackNoRT, PerTrackCompressor);
-					PerTrackCompressor->ParentingDivisor = 1.0f;
-					PerTrackCompressor->ParentingDivisorExponent = 1.0f;
-
-					PerTrackCompressor->bActuallyFilterLinearKeys = true;
-					PerTrackCompressor->bRetarget = true;
-					PerTrackCompressor->ParentingDivisor = 2.0f;
-					PerTrackCompressor->ParentingDivisorExponent = 1.6f;
-					TRYCOMPRESSION_ASYNC(Adaptive1_LinPerTrack, PerTrackCompressor);
-					PerTrackCompressor->ParentingDivisor = 1.0f;
-					PerTrackCompressor->ParentingDivisorExponent = 1.0f;
-				}
-
-
-				{
-					UAnimCompress_PerTrackCompression* PerTrackCompressor = NewObject<UAnimCompress_PerTrackCompression>();
-					PerTrackCompressor->bUseAdaptiveError = true;
-
-					if (CompressibleAnimData.NumFrames > 1 )
-					{
-						PerTrackCompressor->bActuallyFilterLinearKeys = true;
-						PerTrackCompressor->bRetarget = true;
-
-						PerTrackCompressor->MaxPosDiff = 0.1f;
-// 						PerTrackCompressor->MaxAngleDiff = 0.1f;
-						PerTrackCompressor->MaxScaleDiff = 0.00001f;
-						PerTrackCompressor->ParentingDivisor = 2.0f;
-						PerTrackCompressor->ParentingDivisorExponent = 1.0f;
-						TRYCOMPRESSION_ASYNC(Linear_PerTrackExp1, PerTrackCompressor);
-
-						PerTrackCompressor->MaxPosDiff = 0.01f;
-// 						PerTrackCompressor->MaxAngleDiff = 0.025f;
-						PerTrackCompressor->MaxScaleDiff = 0.000001f;
-						PerTrackCompressor->ParentingDivisor = 2.0f;
-						PerTrackCompressor->ParentingDivisorExponent = 1.0f;
-						TRYCOMPRESSION_ASYNC(Linear_PerTrackExp2, PerTrackCompressor);
-
-						PerTrackCompressor->bRetarget = false;
-						PerTrackCompressor->MaxPosDiff = 0.1f;
-// 						PerTrackCompressor->MaxAngleDiff = 0.025f;
-						PerTrackCompressor->MaxScaleDiff = 0.00001f;
-						PerTrackCompressor->ParentingDivisor = 1.0f;
-						PerTrackCompressor->ParentingDivisorExponent = 1.0f;
-					}
-				}
-
-				{
-					UAnimCompress_PerTrackCompression* PerTrackCompressor = NewObject<UAnimCompress_PerTrackCompression>();
-					PerTrackCompressor->bUseAdaptiveError = true;
-
-					// Try the decimation algorithms
-					if (CompressibleAnimData.NumFrames >= PerTrackCompressor->MinKeysForResampling)
-					{
-						PerTrackCompressor->bActuallyFilterLinearKeys = false;
-						PerTrackCompressor->bRetarget = false;
-						PerTrackCompressor->bUseAdaptiveError = false;
-						PerTrackCompressor->bResampleAnimation = true;
-
-						// Try PerTrackCompression, downsample to 20 Hz
-						PerTrackCompressor->ResampledFramerate = 20.0f;
-						TRYCOMPRESSION_ASYNC(Downsample20Hz_PerTrack, PerTrackCompressor);
-
-						// Try PerTrackCompression, downsample to 15 Hz
-						PerTrackCompressor->ResampledFramerate = 15.0f;
-						TRYCOMPRESSION_ASYNC(Downsample15Hz_PerTrack, PerTrackCompressor);
-
-						// Try PerTrackCompression, downsample to 10 Hz
-						PerTrackCompressor->ResampledFramerate = 10.0f;
-						TRYCOMPRESSION_ASYNC(Downsample10Hz_PerTrack, PerTrackCompressor);
-
-						// Try PerTrackCompression, downsample to 5 Hz
-						PerTrackCompressor->ResampledFramerate = 5.0f;
-						TRYCOMPRESSION_ASYNC(Downsample5Hz_PerTrack, PerTrackCompressor);
-
-
-						// Downsampling with linear key removal and adaptive error metrics
-						PerTrackCompressor->bActuallyFilterLinearKeys = true;
-						PerTrackCompressor->bRetarget = false;
-						PerTrackCompressor->bUseAdaptiveError = true;
-						PerTrackCompressor->ParentingDivisor = 2.0f;
-						PerTrackCompressor->ParentingDivisorExponent = 1.6f;
-
-						PerTrackCompressor->ResampledFramerate = 15.0f;
-						TRYCOMPRESSION_ASYNC(Adaptive1_15Hz_LinPerTrack, PerTrackCompressor);
-
-						PerTrackCompressor->ResampledFramerate = 10.0f;
-						TRYCOMPRESSION_ASYNC(Adaptive1_10Hz_LinPerTrack, PerTrackCompressor);
-
-						PerTrackCompressor->ResampledFramerate = 5.0f;
-						TRYCOMPRESSION_ASYNC(Adaptive1_5Hz_LinPerTrack, PerTrackCompressor);
-					}
-				}
-
-
-				{
-					// Try the decimation algorithms
-					if (CompressibleAnimData.NumFrames >= 3)
-					{
-						UAnimCompress_PerTrackCompression* NewPerTrackCompressor = NewObject<UAnimCompress_PerTrackCompression>();
-
-						// Downsampling with linear key removal and adaptive error metrics v2
-						NewPerTrackCompressor->MinKeysForResampling = 3;
-						NewPerTrackCompressor->bUseAdaptiveError2 = true;
-						NewPerTrackCompressor->MaxPosDiffBitwise = 0.05f;
-						NewPerTrackCompressor->MaxAngleDiffBitwise = 0.02f;
-						NewPerTrackCompressor->MaxScaleDiffBitwise = 0.00005f;
-						NewPerTrackCompressor->bActuallyFilterLinearKeys = true;
-						NewPerTrackCompressor->bRetarget = true;
-
-						NewPerTrackCompressor->ResampledFramerate = 15.0f;
-						TRYCOMPRESSION_ASYNC(Adaptive2_15Hz_LinPerTrack, NewPerTrackCompressor);
-
-						NewPerTrackCompressor->ResampledFramerate = 10.0f;
-						TRYCOMPRESSION_ASYNC(Adaptive2_10Hz_LinPerTrack, NewPerTrackCompressor);
-					}
-				}
-
-
-				{
-					// Adaptive error through probing the effect of perturbations at each track
-					UAnimCompress_PerTrackCompression* NewPerTrackCompressor = NewObject<UAnimCompress_PerTrackCompression>();
-					NewPerTrackCompressor->bUseAdaptiveError2 = true;
-					NewPerTrackCompressor->MaxPosDiffBitwise = 0.05f;
-					NewPerTrackCompressor->MaxAngleDiffBitwise = 0.02f;
-					NewPerTrackCompressor->MaxScaleDiffBitwise = 0.00005f;
-
-					TRYCOMPRESSION_ASYNC(Adaptive2_PerTrack, NewPerTrackCompressor);
-
-					NewPerTrackCompressor->bActuallyFilterLinearKeys = true;
-					NewPerTrackCompressor->bRetarget = true;
-					TRYCOMPRESSION_ASYNC(Adaptive2_LinPerTrack, NewPerTrackCompressor);
-
-					NewPerTrackCompressor->bActuallyFilterLinearKeys = true;
-					NewPerTrackCompressor->bRetarget = false;
-					TRYCOMPRESSION_ASYNC(Adaptive2_LinPerTrackNoRT, NewPerTrackCompressor);
-				}
-
-				WaitForAnimCompressionJobs(AnimCompressionTask_CompletionEvents);
-				UpdateAnimCompressionFromAsyncJobs(OutCompressedData, AnimCompressionTask_CompletionEvents, AnimCompressionJobContexes, OriginalSize, CompressorStats, MasterTolerance);
-			}
-
-			// Increase winning compressor.
-			{
-				int32 SizeDecrease = OriginalSize - CompressorStats.CurrentSize;
-				if (CompressorStats.WinningCompressorCounter)
-				{
-					++(*CompressorStats.WinningCompressorCounter);
-					(*CompressorStats.WinningCompressorErrorSum) += CompressorStats.WinningCompressorError;
-					AlternativeCompressorSavings += CompressorStats.WinningCompressorSavings;
-					*CompressorStats.WinningCompressorMarginalSavingsSum += CompressorStats.WinningCompressorMarginalSavings;
-					check(CompressorStats.WinningCompressorSavings == SizeDecrease);
-
-					UE_LOG(LogAnimationCompression, Log, TEXT("  Recompressing(%s) with compressor('%s') saved %i bytes (%i -> %i -> %i) (max diff=%f)\n"),
-						*CompressibleAnimData.Name,
-						*CompressorStats.WinningCompressorName,
-						SizeDecrease,
-						OriginalSize, AfterOriginalRecompression,CompressorStats. CurrentSize,
-						CompressorStats.WinningCompressorError);
-				}
-				else
-				{
-					UE_LOG(LogAnimationCompression, Log, TEXT("  No compressor suitable! Recompressing(%s) with original/default compressor(%s) saved %i bytes (%i -> %i -> %i) (max diff=%f)\n"), 
-						*CompressibleAnimData.Name,
-						*OutCompressedData.CompressionScheme->GetName(),
-						SizeDecrease,
-						OriginalSize, AfterOriginalRecompression, CompressorStats.CurrentSize,
-						CompressorStats.WinningCompressorError);
-
-					UE_LOG(LogAnimationCompression, Log, TEXT("  CompressedTrackOffsets(%d) CompressedByteStream(%d) CompressedScaleOffsets(%d) CompressedSegments(%d)"),
-						OutCompressedData.CompressedTrackOffsets.Num(),
-						OutCompressedData.CompressedByteStream.Num(),
-						OutCompressedData.CompressedScaleOffsets.GetMemorySize(),
-						0);
-
-					TotalNoWinnerRounds++;
-				}
-			}
-
-			// Make sure we got that right.
-			check(CompressorStats.CurrentSize == OutCompressedData.GetApproxBoneCompressedSize());
-			TotalSizeNow += CompressorStats.CurrentSize;
-
-			CompressorStats.PctSaving = TotalSizeBefore > 0 ? 100.f - (100.f * float(TotalSizeNow) / float(TotalSizeBefore)) : 0.f;
-			UE_LOG(LogAnimationCompression, Log, TEXT("Compression Stats Summary [Recompressions(%i) Bytes saved(%i) before(%i) now(%i) savings(%3.1f%%) Uncompressed(%i) TotalRatio(%i:1)]"), 
-				TotalRecompressions,
-				AlternativeCompressorSavings,
-				TotalSizeBefore, 
-				TotalSizeNow,
-				CompressorStats.PctSaving,
-				TotalUncompressed,
-				(TotalUncompressed / TotalSizeNow));
-
-			UE_LOG(LogAnimationCompression, Log, TEXT("\t\tDefault compressor wins:                      %i"), TotalNoWinnerRounds);
-
-			{
-				LOG_COMPRESSION_STATUS(BitwiseACF_Float96);
-				LOG_COMPRESSION_STATUS(BitwiseACF_Fixed48);
-// 				LOG_COMPRESSION_STATUS(BitwiseACF_IntervalFixed32);
-// 				LOG_COMPRESSION_STATUS(BitwiseACF_Fixed32);
-			}
-
-			{
-				LOG_COMPRESSION_STATUS(HalfOddACF_Float96);
-				LOG_COMPRESSION_STATUS(HalfOddACF_Fixed48);
-// 				LOG_COMPRESSION_STATUS(HalfOddACF_IntervalFixed32);
-// 				LOG_COMPRESSION_STATUS(HalfOddACF_Fixed32);
-
-				LOG_COMPRESSION_STATUS(HalfEvenACF_Float96);
-				LOG_COMPRESSION_STATUS(HalfEvenACF_Fixed48);
-// 				LOG_COMPRESSION_STATUS(HalfEvenACF_IntervalFixed32);
-// 				LOG_COMPRESSION_STATUS(HalfEvenACF_Fixed32);
-			}
-
-			{
-				LOG_COMPRESSION_STATUS(LinearACF_Float96);
-				LOG_COMPRESSION_STATUS(LinearACF_Fixed48);
-// 				LOG_COMPRESSION_STATUS(LinearACF_IntervalFixed32);
-// 				LOG_COMPRESSION_STATUS(LinearACF_Fixed32);
-			}
-
-			{
-				LOG_COMPRESSION_STATUS(Progressive_PerTrack);
-				LOG_COMPRESSION_STATUS(Bitwise_PerTrack);
-				LOG_COMPRESSION_STATUS(Linear_PerTrack);
-				LOG_COMPRESSION_STATUS(Adaptive1_LinPerTrackNoRT);
-				LOG_COMPRESSION_STATUS(Adaptive1_LinPerTrack);
-					
-				LOG_COMPRESSION_STATUS(Linear_PerTrackExp1);
-				LOG_COMPRESSION_STATUS(Linear_PerTrackExp2);
-			}
-
-			{
-				LOG_COMPRESSION_STATUS(Downsample20Hz_PerTrack);
-				LOG_COMPRESSION_STATUS(Downsample15Hz_PerTrack);
-				LOG_COMPRESSION_STATUS(Downsample10Hz_PerTrack);
-				LOG_COMPRESSION_STATUS(Downsample5Hz_PerTrack);
-
-				LOG_COMPRESSION_STATUS(Adaptive1_15Hz_LinPerTrack);
-				LOG_COMPRESSION_STATUS(Adaptive1_10Hz_LinPerTrack);
-				LOG_COMPRESSION_STATUS(Adaptive1_5Hz_LinPerTrack);
-
-				LOG_COMPRESSION_STATUS(Adaptive2_15Hz_LinPerTrack);
-				LOG_COMPRESSION_STATUS(Adaptive2_10Hz_LinPerTrack);
-			}
-
-			{
-				LOG_COMPRESSION_STATUS(Adaptive2_PerTrack);
-				LOG_COMPRESSION_STATUS(Adaptive2_LinPerTrack);
-				LOG_COMPRESSION_STATUS(Adaptive2_LinPerTrackNoRT);
-			}
-		}
-		// Do not recompress - Still take into account size for stats.
-		else
-		{
-			TotalSizeNow += OutCompressedData.GetApproxBoneCompressedSize();
-		}
-	}
-	else
-	{
-		// this can happen if the animation only contains curve - i.e. blendshape curves
-		UE_LOG(LogAnimationCompression, Log, TEXT("Compression Requested for Empty Animation %s"), *CompressibleAnimData.Name );
-	}
-#endif // WITH_EDITORONLY_DATA
-}
-
 static void GetBindPoseAtom(FTransform &OutBoneAtom, int32 BoneIndex, USkeleton *Skeleton)
 {
 	OutBoneAtom = Skeleton->GetRefLocalPoses()[BoneIndex];
@@ -2021,70 +732,108 @@ void FAnimationUtils::TallyErrorsFromPerturbation(
 	}
 }
 
+static UAnimBoneCompressionSettings* DefaultBoneCompressionSettings = nullptr;
+static UAnimBoneCompressionSettings* DefaultRecorderBoneCompressionSettings = nullptr;
 static UAnimCurveCompressionSettings* DefaultCurveCompressionSettings = nullptr;
 
-UAnimCurveCompressionSettings* FAnimationUtils::GetDefaultAnimationCurveCompressionSettings()
+static void EnsureDependenciesAreLoaded(UObject* Object)
 {
-	if (DefaultCurveCompressionSettings == nullptr)
+	if (Object == nullptr)
 	{
-		FConfigSection* AnimDefaultObjectSettingsSection = GConfig->GetSectionPrivate(TEXT("Animation.DefaultObjectSettings"), false, true, GEngineIni);
-		const FConfigValue* Value = AnimDefaultObjectSettingsSection != nullptr ? AnimDefaultObjectSettingsSection->Find(TEXT("CurveCompressionSettings")) : nullptr;
+		return;	// Nothing to do
+	}
 
-		if (Value != nullptr)
-		{
-			const FString& CurveCompressionSettingsName = Value->GetValue();
+	if (Object->HasAnyFlags(RF_NeedLoad))
+	{
+		Object->GetLinker()->Preload(Object);
+	}
 
-			DefaultCurveCompressionSettings = LoadObject<UAnimCurveCompressionSettings>(nullptr, *CurveCompressionSettingsName);
+	if (Object->HasAnyFlags(RF_NeedPostLoad))
+	{
+		Object->ConditionalPostLoad();
+	}
 
-			if (DefaultCurveCompressionSettings == nullptr)
-			{
-				UE_LOG(LogAnimationCompression, Fatal, TEXT("Couldn't load default curve compression settings with path '%s'"), *CurveCompressionSettingsName);
-			}
-		}
-		else
-		{
-			UE_LOG(LogAnimationCompression, Fatal, TEXT("Couldn't find default curve compression setting under '[Animation.DefaultObjectSettings]'"));
-		}
+	TArray<UObject*> ObjectReferences;
+	FReferenceFinder(ObjectReferences, Object, false, true, true, true).FindReferences(Object);
 
-		// Force load the default settings and all its dependencies just in case it hasn't happened yet
-		bool bLoadDependencies = false;
-		if (DefaultCurveCompressionSettings->HasAnyFlags(RF_NeedLoad))
+	for (UObject* Dependency : ObjectReferences)
+	{
+		if (Dependency->HasAnyFlags(RF_NeedLoad))
 		{
-			DefaultCurveCompressionSettings->GetLinker()->Preload(DefaultCurveCompressionSettings);
-			bLoadDependencies = true;
+			Dependency->GetLinker()->Preload(Dependency);
 		}
 
-		if (DefaultCurveCompressionSettings->HasAnyFlags(RF_NeedPostLoad))
+		if (Dependency->HasAnyFlags(RF_NeedPostLoad))
 		{
-			DefaultCurveCompressionSettings->ConditionalPostLoad();
-			bLoadDependencies = true;
+			Dependency->ConditionalPostLoad();
 		}
+	}
+}
 
-		if (bLoadDependencies)
-		{
-			TArray<UObject*> ObjectReferences;
-			FReferenceFinder(ObjectReferences, nullptr, false, true, false, true).FindReferences(DefaultCurveCompressionSettings);
+UObject* GetDefaultAnimationCompressionSettings(const TCHAR* IniValueName)
+{
+	FConfigSection* AnimDefaultObjectSettingsSection = GConfig->GetSectionPrivate(TEXT("Animation.DefaultObjectSettings"), false, true, GEngineIni);
+	const FConfigValue* Value = AnimDefaultObjectSettingsSection != nullptr ? AnimDefaultObjectSettingsSection->Find(IniValueName) : nullptr;
 
-			for (UObject* Dependency : ObjectReferences)
-			{
-				if (Dependency->HasAnyFlags(RF_NeedLoad))
-				{
-					Dependency->GetLinker()->Preload(Dependency);
-				}
+	if (Value == nullptr)
+	{
+		UE_LOG(LogAnimationCompression, Fatal, TEXT("Couldn't find default compression setting for '%s' under '[Animation.DefaultObjectSettings]'"), IniValueName);
+	}
 
-				if (Dependency->HasAnyFlags(RF_NeedPostLoad))
-				{
-					Dependency->ConditionalPostLoad();
-				}
-			}
-		}
+	const FString& CompressionSettingsName = Value->GetValue();
+	UObject* DefaultCompressionSettings = LoadObject<UObject>(nullptr, *CompressionSettingsName);
+
+	if (DefaultCompressionSettings == nullptr)
+	{
+		UE_LOG(LogAnimationCompression, Fatal, TEXT("Couldn't load default compression settings asset with path '%s'"), *CompressionSettingsName);
+	}
+
+	// Force load the default settings and all its dependencies just in case it hasn't happened yet
+	EnsureDependenciesAreLoaded(DefaultCompressionSettings);
+
+	DefaultCompressionSettings->AddToRoot();
+
+	return DefaultCompressionSettings;
+}
 
-		DefaultCurveCompressionSettings->AddToRoot();
+UAnimBoneCompressionSettings* FAnimationUtils::GetDefaultAnimationBoneCompressionSettings()
+{
+	if (DefaultBoneCompressionSettings == nullptr)
+	{
+		DefaultBoneCompressionSettings = Cast<UAnimBoneCompressionSettings>(GetDefaultAnimationCompressionSettings(TEXT("BoneCompressionSettings")));
+	}
+
+	return DefaultBoneCompressionSettings;
+}
+
+UAnimBoneCompressionSettings* FAnimationUtils::GetDefaultAnimationRecorderBoneCompressionSettings()
+{
+	if (DefaultRecorderBoneCompressionSettings == nullptr)
+	{
+		DefaultRecorderBoneCompressionSettings = Cast<UAnimBoneCompressionSettings>(GetDefaultAnimationCompressionSettings(TEXT("AnimationRecorderBoneCompressionSettings")));
+	}
+
+	return DefaultRecorderBoneCompressionSettings;
+}
+
+UAnimCurveCompressionSettings* FAnimationUtils::GetDefaultAnimationCurveCompressionSettings()
+{
+	if (DefaultCurveCompressionSettings == nullptr)
+	{
+		DefaultCurveCompressionSettings = Cast<UAnimCurveCompressionSettings>(GetDefaultAnimationCompressionSettings(TEXT("CurveCompressionSettings")));
 	}
 
 	return DefaultCurveCompressionSettings;
 }
 
+void FAnimationUtils::EnsureAnimSequenceLoaded(UAnimSequence& AnimSeq)
+{
+	EnsureDependenciesAreLoaded(&AnimSeq);
+	EnsureDependenciesAreLoaded(AnimSeq.GetSkeleton());
+	EnsureDependenciesAreLoaded(AnimSeq.BoneCompressionSettings);
+	EnsureDependenciesAreLoaded(AnimSeq.CurveCompressionSettings);
+}
+
 void FAnimationUtils::ExtractTransformFromTrack(float Time, int32 NumFrames, float SequenceLength, const struct FRawAnimSequenceTrack& RawTrack, EAnimInterpolationType Interpolation, FTransform &OutAtom)
 {
 	// Bail out (with rather wacky data) if data is empty for some reason.
@@ -2171,34 +920,53 @@ void FAnimationUtils::ExtractTransformFromTrack(float Time, int32 NumFrames, flo
 }
 
 #if WITH_EDITOR
-
 void FAnimationUtils::ExtractTransformFromCompressionData(const FCompressibleAnimData& CompressibleAnimData, FCompressibleAnimDataResult& CompressedAnimData, float Time, int32 TrackIndex, bool bUseRawData, FTransform& OutBoneTransform)
 {
+	FUECompressedAnimDataMutable& AnimDataMutable = static_cast<FUECompressedAnimDataMutable&>(*CompressedAnimData.AnimData);
+
 	// If the caller didn't request that raw animation data be used . . .
-	if (!bUseRawData && CompressedAnimData.IsCompressedDataValid())
+	if (!bUseRawData && AnimDataMutable.IsValid())
 	{
-		const FUECompressedAnimData CompressedDataWrapper(CompressedAnimData);
-		FAnimSequenceDecompressionContext DecompContext(CompressibleAnimData, CompressedDataWrapper);
+		// Build our read-only version from the mutable source
+		FUECompressedAnimData AnimData(AnimDataMutable);
+
+		FAnimSequenceDecompressionContext DecompContext(CompressibleAnimData.SequenceLength, CompressibleAnimData.Interpolation, CompressibleAnimData.AnimFName, AnimData);
 		DecompContext.Seek(Time);
-		AnimationFormat_GetBoneAtom(OutBoneTransform, DecompContext, TrackIndex);
+		CompressedAnimData.Codec->DecompressBone(DecompContext, TrackIndex, OutBoneTransform);
 		return;
 	}
 
 	FAnimationUtils::ExtractTransformFromTrack(Time, CompressibleAnimData.NumFrames, CompressibleAnimData.SequenceLength, CompressibleAnimData.RawAnimationData[TrackIndex], CompressibleAnimData.Interpolation, OutBoneTransform);
 }
 
+bool FAnimationUtils::CompressAnimBones(FCompressibleAnimData& AnimSeq, FCompressibleAnimDataResult& Target)
+{
+	// Clear any previous data we might have even if we end up failing to compress
+	Target = FCompressibleAnimDataResult();
+
+	if (AnimSeq.Skeleton == nullptr)
+	{
+		return false;
+	}
+
+	if (AnimSeq.BoneCompressionSettings == nullptr || !AnimSeq.BoneCompressionSettings->AreSettingsValid())
+	{
+		return false;
+	}
+
+	return AnimSeq.BoneCompressionSettings->Compress(AnimSeq, Target);
+}
+
 bool FAnimationUtils::CompressAnimCurves(FCompressibleAnimData& AnimSeq, FCompressedAnimSequence& Target)
 {
 	// Clear any previous data we might have even if we end up failing to compress
-	Target.CompressedCurveByteStream.Empty();
-	Target.CurveCompressionCodec = nullptr;
+	Target.ClearCompressedCurveData();
 
 	if (AnimSeq.CurveCompressionSettings == nullptr || !AnimSeq.CurveCompressionSettings->AreSettingsValid())
 	{
 		return false;
 	}
 
-	check(AnimSeq.CurveCompressionSettings->AreSettingsValid());
 	return AnimSeq.CurveCompressionSettings->Compress(AnimSeq, Target);
 }
 #endif
diff --git a/Engine/Source/Runtime/Engine/Private/Animation/BonePose.cpp b/Engine/Source/Runtime/Engine/Private/Animation/BonePose.cpp
--- a/Engine/Source/Runtime/Engine/Private/Animation/BonePose.cpp
+++ b/Engine/Source/Runtime/Engine/Private/Animation/BonePose.cpp
@@ -4,6 +4,83 @@
 #include "AnimationRuntime.h"
 #include "AnimEncoding.h"
 #include "HAL/ThreadSingleton.h"
+#if INTEL_ISPC
+#include "BonePose.ispc.generated.h"
+#endif
+
+// Normalizes all rotations in this pose
+void FCompactPose::NormalizeRotations()
+{
+	if (INTEL_ISPC)
+	{
+#if INTEL_ISPC
+		ispc::NormalizeRotations((ispc::FTransform*)this->Bones.GetData(), this->Bones.Num());
+#endif
+	}
+	else
+	{
+		for (FTransform& Bone : this->Bones)
+		{
+			Bone.NormalizeRotation();
+		}
+	}
+}
+
+// Sets every bone transform to Identity
+void FCompactPose::ResetToAdditiveIdentity()
+{
+	if (INTEL_ISPC)
+	{
+#if INTEL_ISPC
+		ispc::ResetToAdditiveIdentity((ispc::FTransform*)this->Bones.GetData(), this->Bones.Num());
+#endif
+	}
+	else
+	{
+		for (FTransform& Bone : this->Bones)
+		{
+			Bone.SetIdentity();
+			Bone.SetScale3D(FVector::ZeroVector);
+		}
+	}
+}
+
+// Normalizes all rotations in this pose
+void FCompactHeapPose::NormalizeRotations()
+{
+	if (INTEL_ISPC)
+	{
+#if INTEL_ISPC
+		ispc::NormalizeRotations((ispc::FTransform*)this->Bones.GetData(), this->Bones.Num());
+#endif
+	}
+	else
+	{
+		for (FTransform& Bone : this->Bones)
+		{
+			Bone.NormalizeRotation();
+		}
+	}
+}
+
+// Sets every bone transform to Identity
+void FCompactHeapPose::ResetToAdditiveIdentity()
+{
+	if (INTEL_ISPC)
+	{
+#if INTEL_ISPC
+		ispc::ResetToAdditiveIdentity((ispc::FTransform*)this->Bones.GetData(), this->Bones.Num());
+#endif
+	}
+	else
+	{
+		for (FTransform& Bone : this->Bones)
+		{
+			Bone.SetIdentity();
+			Bone.SetScale3D(FVector::ZeroVector);
+		}
+	}
+}
 
 void FMeshPose::ResetToRefPose()
 {
diff --git a/Engine/Source/Runtime/Engine/Private/Animation/SkeletalMeshActor.cpp b/Engine/Source/Runtime/Engine/Private/Animation/SkeletalMeshActor.cpp
--- a/Engine/Source/Runtime/Engine/Private/Animation/SkeletalMeshActor.cpp
+++ b/Engine/Source/Runtime/Engine/Private/Animation/SkeletalMeshActor.cpp
@@ -78,7 +78,7 @@ void ASkeletalMeshActor::PreviewSetAnimPosition(FName SlotName, int32 ChannelInd
 	if (CanPlayAnimation(InAnimSequence))
 	{
 		TWeakObjectPtr<class UAnimMontage>& CurrentlyPlayingMontage = CurrentlyPlayingMontages.FindOrAdd(SlotName);
-		CurrentlyPlayingMontage = FAnimMontageInstance::PreviewMatineeSetAnimPositionInner(SlotName, SkeletalMeshComponent, InAnimSequence, InPosition, bLooping, bFireNotifies, DeltaTime);
+		CurrentlyPlayingMontage = FAnimMontageInstance::PreviewMatineeSetAnimPositionInner(SlotName, SkeletalMeshComponent, SkeletalMeshComponent->GetAnimInstance(),  InAnimSequence, InPosition, bLooping, bFireNotifies, DeltaTime);
 	}
 }
 
@@ -224,7 +224,7 @@ void ASkeletalMeshActor::SetAnimPosition(FName SlotName, int32 ChannelIndex, UAn
 	if (CanPlayAnimation(InAnimSequence))
 	{
 		TWeakObjectPtr<class UAnimMontage>& CurrentlyPlayingMontage = CurrentlyPlayingMontages.FindOrAdd(SlotName);
-		CurrentlyPlayingMontage = FAnimMontageInstance::SetMatineeAnimPositionInner(SlotName, SkeletalMeshComponent, InAnimSequence, InPosition, bLooping);
+		CurrentlyPlayingMontage = FAnimMontageInstance::SetMatineeAnimPositionInner(SlotName, SkeletalMeshComponent->GetAnimInstance(), InAnimSequence, InPosition, bLooping);
 	}
 }
 
diff --git a/Engine/Source/Runtime/Engine/Private/Animation/Skeleton.cpp b/Engine/Source/Runtime/Engine/Private/Animation/Skeleton.cpp
--- a/Engine/Source/Runtime/Engine/Private/Animation/Skeleton.cpp
+++ b/Engine/Source/Runtime/Engine/Private/Animation/Skeleton.cpp
@@ -89,14 +89,19 @@ const TCHAR* SkipPrefix(const FString& InName)
 	return &InName[PrefixLength];
 }
 
-FString VirtualBoneNameHelpers::AddVirtualBonePrefix(const FString& InName)
+namespace VirtualBoneNameHelpers
 {
-	return VirtualBoneNameHelpers::VirtualBonePrefix + InName;
-}
+	const FString VirtualBonePrefix(TEXT("VB "));
 
-FName VirtualBoneNameHelpers::RemoveVirtualBonePrefix(const FString& InName)
-{
-	return FName(SkipPrefix(InName));
+	FString AddVirtualBonePrefix(const FString& InName)
+	{
+		return VirtualBonePrefix + InName;
+	}
+
+	FName RemoveVirtualBonePrefix(const FString& InName)
+	{
+		return FName(SkipPrefix(InName));
+	}
 }
 
 USkeleton::USkeleton(const FObjectInitializer& ObjectInitializer)
@@ -1152,13 +1157,16 @@ bool USkeleton::ContainsSlotName(const FName& InSlotName) const
 	return SlotToGroupNameMap.Contains(InSlotName);
 }
 
-void USkeleton::RegisterSlotNode(const FName& InSlotName)
+bool USkeleton::RegisterSlotNode(const FName& InSlotName)
 {
 	// verify the slot name exists, if not create it in the default group.
 	if (!ContainsSlotName(InSlotName))
 	{
 		SetSlotGroupName(InSlotName, FAnimSlotGroup::DefaultGroupName);
+		return true;
 	}
+
+	return false;
 }
 
 void USkeleton::SetSlotGroupName(const FName& InSlotName, const FName& InGroupName)
diff --git a/Engine/Source/Runtime/Engine/Private/Animation/SkinWeightProfile.cpp b/Engine/Source/Runtime/Engine/Private/Animation/SkinWeightProfile.cpp
--- a/Engine/Source/Runtime/Engine/Private/Animation/SkinWeightProfile.cpp
+++ b/Engine/Source/Runtime/Engine/Private/Animation/SkinWeightProfile.cpp
@@ -4,7 +4,7 @@
 #include "Engine/SkeletalMesh.h"
 #include "UObject/UObjectIterator.h"
 #include "ContentStreaming.h"
-
+#include "UObject/AnimObjectVersion.h"
 #include "Components/SkinnedMeshComponent.h"
 #include "Rendering/SkeletalMeshRenderData.h"
 #include "Rendering/SkinWeightVertexBuffer.h"
@@ -90,8 +90,25 @@ static FAutoConsoleVariableRef CVarSkinWeightProfilesAllowedFromLOD(
 
 FArchive& operator<<(FArchive& Ar, FRuntimeSkinWeightProfileData& OverrideData)
 {
+	Ar.UsingCustomVersion(FAnimObjectVersion::GUID);
 	Ar << OverrideData.OverridesInfo;
-	Ar << OverrideData.Weights;	
+	if (Ar.IsLoading() && Ar.CustomVer(FAnimObjectVersion::GUID) < FAnimObjectVersion::IncreaseBoneIndexLimitPerChunk)
+	{
+		TArray<uint16> OldWeights;
+		Ar << OldWeights;
+		// repack index_8bit|weight_8bit into index_16bit|weight_16bit
+		for (int32 i = 0; i < OldWeights.Num(); ++i)
+		{
+			uint16 OldWeight = OldWeights[i];
+			const uint8 Index = OldWeight >> 8;
+			const uint8 Weight = OldWeight & 0xFF;
+			OverrideData.Weights[i] = (Index << 16) | Weight;
+		}
+	}
+	else
+	{
+		Ar << OverrideData.Weights;
+	}
 	Ar << OverrideData.VertexIndexOverrideIndex;
 
 	return Ar;
@@ -123,10 +140,39 @@ FArchive& operator<<(FArchive& Ar, FImportedSkinWeightProfileData& ProfileData)
 
 FArchive& operator<<(FArchive& Ar, FRawSkinWeight& OverrideEntry)
 {
-	for (int32 InfluenceIndex = 0; InfluenceIndex < MAX_TOTAL_INFLUENCES; ++InfluenceIndex)
+	Ar.UsingCustomVersion(FAnimObjectVersion::GUID);
+
+	if (Ar.IsLoading())
 	{
-		Ar << OverrideEntry.InfluenceBones[InfluenceIndex];
-		Ar << OverrideEntry.InfluenceWeights[InfluenceIndex];
+		FMemory::Memzero(OverrideEntry.InfluenceBones);
+		FMemory::Memzero(OverrideEntry.InfluenceWeights);
+	}
+
+	if (Ar.CustomVer(FAnimObjectVersion::GUID) < FAnimObjectVersion::UnlimitedBoneInfluences)
+	{
+		for (int32 InfluenceIndex = 0; InfluenceIndex < EXTRA_BONE_INFLUENCES; ++InfluenceIndex)
+		{
+			if (Ar.IsLoading() && Ar.CustomVer(FAnimObjectVersion::GUID) < FAnimObjectVersion::IncreaseBoneIndexLimitPerChunk)
+			{
+				uint8 BoneIndex = 0;
+				Ar << BoneIndex;
+				OverrideEntry.InfluenceBones[InfluenceIndex] = BoneIndex;
+			}
+			else
+			{
+				Ar << OverrideEntry.InfluenceBones[InfluenceIndex];
+			}
+
+			Ar << OverrideEntry.InfluenceWeights[InfluenceIndex];
+		}
+	}
+	else
+	{
+		for (int32 InfluenceIndex = 0; InfluenceIndex < MAX_TOTAL_INFLUENCES; ++InfluenceIndex)
+		{
+			Ar << OverrideEntry.InfluenceBones[InfluenceIndex];
+			Ar << OverrideEntry.InfluenceWeights[InfluenceIndex];
+		}
 	}
 
 	return Ar;
@@ -285,7 +331,7 @@ FSkinWeightVertexBuffer* FSkinWeightProfilesData::GetOverrideBuffer(const FName&
 				ApplyOverrideProfile(OverrideBuffer, ProfileName);
 
 				INC_DWORD_STAT_BY(STAT_SkeletalMeshVertexMemory, OverrideBuffer->GetVertexDataSize());
-				BeginInitResource(OverrideBuffer);
+				OverrideBuffer->BeginInitResources();
 
 				return OverrideBuffer;
 			}
@@ -307,20 +353,15 @@ FRuntimeSkinWeightProfileData& FSkinWeightProfilesData::AddOverrideData(const FN
 
 void FSkinWeightProfilesData::ApplyOverrideProfile(FSkinWeightVertexBuffer* OverrideBuffer, const FName& ProfileName)
 {
-	const bool bExtraWeights = BaseBuffer->HasExtraBoneInfluences();
-	OverrideBuffer->SetHasExtraBoneInfluences(bExtraWeights);
+	const uint32 NumInfluences = BaseBuffer->GetMaxBoneInfluences();
+	OverrideBuffer->SetMaxBoneInfluences(NumInfluences);
+	const bool bUse16BitBoneIndex = BaseBuffer->Use16BitBoneIndex();
+	OverrideBuffer->SetUse16BitBoneIndex(bUse16BitBoneIndex);
 
 	const FRuntimeSkinWeightProfileData* ProfilePtr = OverrideData.Find(ProfileName);
 	if (ProfilePtr)
 	{
-		if (bExtraWeights)
-		{
-			ProfilePtr->ApplyOverrides<true>(OverrideBuffer, BaseBuffer);
-		}
-		else
-		{
-			ProfilePtr->ApplyOverrides<false>(OverrideBuffer, BaseBuffer);
-		}
+		ProfilePtr->ApplyOverrides(OverrideBuffer, BaseBuffer);
 	}
 }
 
@@ -337,7 +378,7 @@ void FSkinWeightProfilesData::ReleaseBuffer(const FName& ProfileName)
 			ENQUEUE_RENDER_COMMAND(ReleaseSkinSkinWeightProfilesDataBufferCommand)(
 				[Buffer](FRHICommandListImmediate& RHICmdList)
 			{			
-				Buffer->ReleaseResource();
+				Buffer->ReleaseResources();
 				delete Buffer;		
 			});
 		}
@@ -366,7 +407,7 @@ void FSkinWeightProfilesData::ReleaseResources()
 		{
 			for (FSkinWeightVertexBuffer* Buffer : Buffers)
 			{	
-				Buffer->ReleaseResource();
+				Buffer->ReleaseResources();
 				delete Buffer;
 			}
 		});
@@ -412,7 +453,7 @@ void FSkinWeightProfilesData::ReleaseCPUResources()
 }
 
 template <bool bRenderThread>
-void FSkinWeightProfilesData::CreateRHIBuffers_Internal(TArray<TPair<FName, FVertexBufferRHIRef>>& OutBuffers)
+void FSkinWeightProfilesData::CreateRHIBuffers_Internal(TArray<TPair<FName, FSkinWeightRHIInfo>>& OutBuffers)
 {
 	const int32 NumActiveProfiles = ProfileNameToBuffer.Num();
 	check(BaseBuffer || !NumActiveProfiles);
@@ -433,12 +474,12 @@ void FSkinWeightProfilesData::CreateRHIBuffers_Internal(TArray<TPair<FName, FVer
 	}
 }
 
-void FSkinWeightProfilesData::CreateRHIBuffers_RenderThread(TArray<TPair<FName, FVertexBufferRHIRef>>& OutBuffers)
+void FSkinWeightProfilesData::CreateRHIBuffers_RenderThread(TArray<TPair<FName, FSkinWeightRHIInfo>>& OutBuffers)
 {
 	CreateRHIBuffers_Internal<true>(OutBuffers);
 }
 
-void FSkinWeightProfilesData::CreateRHIBuffers_Async(TArray<TPair<FName, FVertexBufferRHIRef>>& OutBuffers)
+void FSkinWeightProfilesData::CreateRHIBuffers_Async(TArray<TPair<FName, FSkinWeightRHIInfo>>& OutBuffers)
 {
 	CreateRHIBuffers_Internal<false>(OutBuffers);
 }
diff --git a/Engine/Source/Runtime/Engine/Private/Collision/KAggregateGeom.cpp b/Engine/Source/Runtime/Engine/Private/Collision/KAggregateGeom.cpp
--- a/Engine/Source/Runtime/Engine/Private/Collision/KAggregateGeom.cpp
+++ b/Engine/Source/Runtime/Engine/Private/Collision/KAggregateGeom.cpp
@@ -15,6 +15,9 @@
 #include "Engine/Polys.h"
 #include "PhysXIncludes.h"
 #include "Chaos/Convex.h"
+#if INTEL_ISPC
+#include "KAggregateGeom.ispc.generated.h"
+#endif
 
 #if WITH_CHAOS
 #include "Chaos/ImplicitObject.h"
@@ -309,14 +312,34 @@ EAggCollisionShape::Type FKBoxElem::StaticShapeType = EAggCollisionShape::Box;
 
 FBox FKBoxElem::CalcAABB(const FTransform& BoneTM, float Scale) const
 {
-	FTransform ElemTM = GetTransform();
-	ElemTM.ScaleTranslation( FVector(Scale) );
-	ElemTM *= BoneTM;
+	if (INTEL_ISPC)
+	{
+#if INTEL_ISPC
+		FBox LocalBox(ForceInit);
+		ispc::BoxCalcAABB(
+			(ispc::FBox&)LocalBox,
+			(ispc::FTransform&)BoneTM,
+			Scale,
+			(ispc::FRotator&)Rotation,
+			(ispc::FVector&)Center,
+			X,
+			Y,
+			Z);
+
+		return LocalBox;
+#endif
+	}
+	else
+	{
+		FTransform ElemTM = GetTransform();
+		ElemTM.ScaleTranslation(FVector(Scale));
+		ElemTM *= BoneTM;
 
-	FVector Extent(0.5f * Scale * X, 0.5f * Scale * Y, 0.5f * Scale * Z);
-	FBox LocalBox(-Extent, Extent);
+		FVector Extent(0.5f * Scale * X, 0.5f * Scale * Y, 0.5f * Scale * Z);
+		FBox LocalBox(-Extent, Extent);
 
-	return LocalBox.TransformBy(ElemTM);
+		return LocalBox.TransformBy(ElemTM);
+	}
 }
 
 
@@ -328,26 +351,45 @@ EAggCollisionShape::Type FKSphylElem::StaticShapeType = EAggCollisionShape::Sphy
 
 FBox FKSphylElem::CalcAABB(const FTransform& BoneTM, float Scale) const
 {
-	FTransform ElemTM = GetTransform();
-	ElemTM.ScaleTranslation( FVector(Scale) );
-	ElemTM *= BoneTM;
+	if (INTEL_ISPC)
+	{
+#if INTEL_ISPC
+		FBox Result(ForceInit);
+		ispc::SPhylCalcAABB(
+			(ispc::FBox&)Result,
+			(ispc::FTransform&)BoneTM,
+			Scale,
+			(ispc::FRotator&)Rotation,
+			(ispc::FVector&)Center,
+			Radius,
+			Length);
+
+		return Result;
+#endif
+	}
+	else
+	{
+		FTransform ElemTM = GetTransform();
+		ElemTM.ScaleTranslation( FVector(Scale) );
+		ElemTM *= BoneTM;
 
-	const FVector SphylCenter = ElemTM.GetLocation();
+		const FVector SphylCenter = ElemTM.GetLocation();
 
-	// Get sphyl axis direction
-	const FVector Axis = ElemTM.GetScaledAxis( EAxis::Z );
-	// Get abs of that vector
-	const FVector AbsAxis(FMath::Abs(Axis.X), FMath::Abs(Axis.Y), FMath::Abs(Axis.Z));
-	// Scale by length of sphyl
-	const FVector AbsDist = (Scale * 0.5f * Length) * AbsAxis;
+		// Get sphyl axis direction
+		const FVector Axis = ElemTM.GetScaledAxis( EAxis::Z );
+		// Get abs of that vector
+		const FVector AbsAxis(FMath::Abs(Axis.X), FMath::Abs(Axis.Y), FMath::Abs(Axis.Z));
+		// Scale by length of sphyl
+		const FVector AbsDist = (Scale * 0.5f * Length) * AbsAxis;
 
-	const FVector MaxPos = SphylCenter + AbsDist;
-	const FVector MinPos = SphylCenter - AbsDist;
-	const FVector Extent(Scale * Radius);
+		const FVector MaxPos = SphylCenter + AbsDist;
+		const FVector MinPos = SphylCenter - AbsDist;
+		const FVector Extent(Scale * Radius);
 
-	FBox Result(MinPos - Extent, MaxPos + Extent);
+		FBox Result(MinPos - Extent, MaxPos + Extent);
 
-	return Result;
+		return Result;
+	}
 }
 
 
diff --git a/Engine/Source/Runtime/Engine/Private/Components/SkeletalMeshComponent.cpp b/Engine/Source/Runtime/Engine/Private/Components/SkeletalMeshComponent.cpp
--- a/Engine/Source/Runtime/Engine/Private/Components/SkeletalMeshComponent.cpp
+++ b/Engine/Source/Runtime/Engine/Private/Components/SkeletalMeshComponent.cpp
@@ -40,6 +40,9 @@
 #include "SkeletalRenderPublic.h"
 #include "ContentStreaming.h"
 #include "Animation/AnimTrace.h"
+#if INTEL_ISPC
+#include "SkeletalMeshComponent.ispc.generated.h"
+#endif
 
 #define LOCTEXT_NAMESPACE "SkeletalMeshComponent"
 
@@ -586,6 +589,8 @@ void USkeletalMeshComponent::OnUnregister()
 		}
 	}
 
+	RequiredBones.Reset();
+
 	Super::OnUnregister();
 }
 
@@ -1388,32 +1393,48 @@ void USkeletalMeshComponent::FillComponentSpaceTransforms(const USkeletalMesh* I
 #endif
 	}
 
-	for (int32 i = 1; i<FillComponentSpaceTransformsRequiredBones.Num(); i++)
+	if (INTEL_ISPC)
 	{
-		const int32 BoneIndex = FillComponentSpaceTransformsRequiredBones[i];
-		FTransform* SpaceBase = ComponentSpaceData + BoneIndex;
+#if INTEL_ISPC
+		ispc::FillComponentSpaceTransforms(
+			(ispc::FTransform*)&ComponentSpaceData[0],
+			(ispc::FTransform*)&LocalTransformsData[0],
+			FillComponentSpaceTransformsRequiredBones.GetData(),
+			(const uint8*)InSkeletalMesh->RefSkeleton.GetRefBoneInfo().GetData(),
+			sizeof(FMeshBoneInfo),
+			offsetof(FMeshBoneInfo, ParentIndex),
+			FillComponentSpaceTransformsRequiredBones.Num());
+#endif
+	}
+	else
+	{
+		for (int32 i = 1; i < FillComponentSpaceTransformsRequiredBones.Num(); i++)
+		{
+			const int32 BoneIndex = FillComponentSpaceTransformsRequiredBones[i];
+			FTransform* SpaceBase = ComponentSpaceData + BoneIndex;
 
-		FPlatformMisc::Prefetch(SpaceBase);
+			FPlatformMisc::Prefetch(SpaceBase);
 
 #if DO_GUARD_SLOW
-		// Mark bone as processed
-		BoneProcessed[BoneIndex] = 1;
+			// Mark bone as processed
+			BoneProcessed[BoneIndex] = 1;
 #endif
-		// For all bones below the root, final component-space transform is relative transform * component-space transform of parent.
-		const int32 ParentIndex = InSkeletalMesh->RefSkeleton.GetParentIndex(BoneIndex);
-		FTransform* ParentSpaceBase = ComponentSpaceData + ParentIndex;
-		FPlatformMisc::Prefetch(ParentSpaceBase);
+			// For all bones below the root, final component-space transform is relative transform * component-space transform of parent.
+			const int32 ParentIndex = InSkeletalMesh->RefSkeleton.GetParentIndex(BoneIndex);
+			FTransform* ParentSpaceBase = ComponentSpaceData + ParentIndex;
+			FPlatformMisc::Prefetch(ParentSpaceBase);
 
 #if DO_GUARD_SLOW
-		// Check the precondition that Parents occur before Children in the RequiredBones array.
-		checkSlow(BoneProcessed[ParentIndex] == 1);
+			// Check the precondition that Parents occur before Children in the RequiredBones array.
+			checkSlow(BoneProcessed[ParentIndex] == 1);
 #endif
-		FTransform::Multiply(SpaceBase, LocalTransformsData + BoneIndex, ParentSpaceBase);
+			FTransform::Multiply(SpaceBase, LocalTransformsData + BoneIndex, ParentSpaceBase);
 
-		SpaceBase->NormalizeRotation();
+			SpaceBase->NormalizeRotation();
 
-		checkSlow(SpaceBase->IsRotationNormalized());
-		checkSlow(!SpaceBase->ContainsNaN());
+			checkSlow(SpaceBase->IsRotationNormalized());
+			checkSlow(!SpaceBase->ContainsNaN());
+		}
 	}
 }
 
diff --git a/Engine/Source/Runtime/Engine/Private/Components/SkeletalMeshComponent.ispc b/Engine/Source/Runtime/Engine/Private/Components/SkeletalMeshComponent.ispc
new file mode 100644
--- /dev/null
+++ b/Engine/Source/Runtime/Engine/Private/Components/SkeletalMeshComponent.ispc
@@ -0,0 +1,37 @@
+// Copyright Epic Games, Inc. All Rights Reserved.
+
+#define EXPLICIT_VECTOR4 1
+#define EXPLICIT_VECTOR 1
+#define EXPLICIT_MATRIX 1
+
+#include "Math/Vector.isph"
+#include "Math/Matrix.isph"
+#include "Math/Quat.isph"
+#include "Math/Transform.isph"
+
+export void FillComponentSpaceTransforms(uniform FTransform ComponentSpaceData[],
+										const uniform FTransform LocalTransformsData[],
+										const uniform unsigned int16 FillComponentSpaceTransformsRequiredBones[],
+										const uniform unsigned int8 * uniform ParentBones,
+										const uniform int ParentBoneSize,
+										const uniform int ParentIndexOffset,
+										const uniform int NumBones)
+{
+	for (uniform int i = 1; i < NumBones; i++)
+	{
+		const uniform int BoneIndex = FillComponentSpaceTransformsRequiredBones[i];
+		uniform FTransform SpaceBase = ComponentSpaceData[BoneIndex];
+
+		// For all bones below the root, final component-space transform is relative transform * component-space transform of parent.
+		const uniform unsigned int8 * uniform ParentIndexPtr = ParentBones + ((ParentBoneSize * BoneIndex) + ParentIndexOffset);
+		const uniform int ParentIndex = *((uniform int * uniform)ParentIndexPtr);
+		const uniform FTransform ParentSpaceBase = ComponentSpaceData[ParentIndex];
+
+		const uniform FTransform LocalTransform = LocalTransformsData[BoneIndex];
+
+		TransformMultiply(SpaceBase, LocalTransform, ParentSpaceBase);
+		SpaceBase.Rotation = VectorNormalizeQuaternion(SpaceBase.Rotation);
+
+		ComponentSpaceData[BoneIndex] = SpaceBase;
+	}
+}
