COPY from //Dev-Anim to //Dev-Main

diff --git a/Engine/Source/Runtime/Engine/Private/Components/SkinnedMeshComponent.cpp b/Engine/Source/Runtime/Engine/Private/Components/SkinnedMeshComponent.cpp
--- a/Engine/Source/Runtime/Engine/Private/Components/SkinnedMeshComponent.cpp
+++ b/Engine/Source/Runtime/Engine/Private/Components/SkinnedMeshComponent.cpp
@@ -560,11 +560,11 @@ void USkinnedMeshComponent::CreateRenderState_Concurrent()
 			{
 				int32 MaxBonesPerChunk = SkelMeshRenderData->GetMaxBonesPerSection();
 				int32 MaxSupportedGPUSkinBones = FMath::Min(GetFeatureLevelMaxNumberOfBones(SceneFeatureLevel), FGPUBaseSkinVertexFactory::GetMaxGPUSkinBones());
-				bool bHasExtraBoneInfluences = SkelMeshRenderData->HasExtraBoneInfluences();
+				int32 NumBoneInfluences = SkelMeshRenderData->GetNumBoneInfluences();
 				FString FeatureLevelName; GetFeatureLevelName(SceneFeatureLevel, FeatureLevelName);
 
-				UE_LOG(LogSkinnedMeshComp, Warning, TEXT("SkeletalMesh %s, is not supported for current feature level (%s) and will not be rendered. NumBones %d (supported %d), HasExtraBoneInfluences: %s"), 
-					*GetNameSafe(SkeletalMesh), *FeatureLevelName, MaxBonesPerChunk, MaxSupportedGPUSkinBones, bHasExtraBoneInfluences ? TEXT("true"):TEXT("false"));
+				UE_LOG(LogSkinnedMeshComp, Warning, TEXT("SkeletalMesh %s, is not supported for current feature level (%s) and will not be rendered. NumBones %d (supported %d), NumBoneInfluences: %d"), 
+					*GetNameSafe(SkeletalMesh), *FeatureLevelName, MaxBonesPerChunk, MaxSupportedGPUSkinBones, NumBoneInfluences);
 			}
 
 			//Allow the editor a chance to manipulate it before its added to the scene
@@ -2520,9 +2520,7 @@ FVector USkinnedMeshComponent::GetSkinnedVertexPosition(USkinnedMeshComponent* C
 	check(SectionIndex < LODData.RenderSections.Num());
 	const FSkelMeshRenderSection& Section = LODData.RenderSections[SectionIndex];
 
-	return  SkinWeightBuffer.HasExtraBoneInfluences()
-		? GetTypedSkinnedVertexPosition<true, false>(Component, Section, LODData.StaticVertexBuffers.PositionVertexBuffer, SkinWeightBuffer, VertIndex)
-		: GetTypedSkinnedVertexPosition<false, false>(Component, Section, LODData.StaticVertexBuffers.PositionVertexBuffer, SkinWeightBuffer, VertIndex);
+	return GetTypedSkinnedVertexPosition<false>(Component, Section, LODData.StaticVertexBuffers.PositionVertexBuffer, SkinWeightBuffer, VertIndex);
 }
 
 FVector USkinnedMeshComponent::GetSkinnedVertexPosition(USkinnedMeshComponent* Component, int32 VertexIndex, const FSkeletalMeshLODRenderData& LODData, FSkinWeightVertexBuffer& SkinWeightBuffer, TArray<FMatrix>& CachedRefToLocals) 
@@ -2536,9 +2534,7 @@ FVector USkinnedMeshComponent::GetSkinnedVertexPosition(USkinnedMeshComponent* C
 	check(SectionIndex < LODData.RenderSections.Num());
 	const FSkelMeshRenderSection& Section = LODData.RenderSections[SectionIndex];
 
-	return SkinWeightBuffer.HasExtraBoneInfluences()
-		? GetTypedSkinnedVertexPosition<true, false>(Component, Section, LODData.StaticVertexBuffers.PositionVertexBuffer, SkinWeightBuffer, VertIndex,CachedRefToLocals)
-		: GetTypedSkinnedVertexPosition<false, false>(Component, Section, LODData.StaticVertexBuffers.PositionVertexBuffer, SkinWeightBuffer, VertIndex, CachedRefToLocals);
+	return GetTypedSkinnedVertexPosition<false>(Component, Section, LODData.StaticVertexBuffers.PositionVertexBuffer, SkinWeightBuffer, VertIndex, CachedRefToLocals);
 }
 
 void USkinnedMeshComponent::SetRefPoseOverride(const TArray<FTransform>& NewRefPoseTransforms)
@@ -2661,15 +2657,13 @@ void USkinnedMeshComponent::ComputeSkinnedPositions(USkinnedMeshComponent* Compo
 	for (int32 SectionIdx = 0; SectionIdx < LODData.RenderSections.Num(); ++SectionIdx)
 	{
 		const FSkelMeshRenderSection& Section = LODData.RenderSections[SectionIdx];
-		bool bHasExtraBoneInfluences = SkinWeightBuffer.HasExtraBoneInfluences();
 		{
 			//soft
 			const uint32 SoftOffset = Section.BaseVertexIndex;
 			const uint32 NumSoftVerts = Section.NumVertices;
 			for (uint32 SoftIdx = 0; SoftIdx < NumSoftVerts; ++SoftIdx)
 			{
-				FVector SkinnedPosition = bHasExtraBoneInfluences ? GetTypedSkinnedVertexPosition<true, true>(Component, Section, LODData.StaticVertexBuffers.PositionVertexBuffer, SkinWeightBuffer, SoftIdx, CachedRefToLocals) :
-																	GetTypedSkinnedVertexPosition<false,true>(Component, Section, LODData.StaticVertexBuffers.PositionVertexBuffer, SkinWeightBuffer, SoftIdx, CachedRefToLocals);
+				FVector SkinnedPosition = GetTypedSkinnedVertexPosition<true>(Component, Section, LODData.StaticVertexBuffers.PositionVertexBuffer, SkinWeightBuffer, SoftIdx, CachedRefToLocals);
 				OutPositions[SoftOffset + SoftIdx] = SkinnedPosition;
 			}
 		}
@@ -3220,7 +3214,7 @@ void FSkelMeshComponentLODInfo::ReleaseOverrideSkinWeightsAndBlock()
 	if (OverrideSkinWeights)
 	{
 		// enqueue a rendering command to release
-		BeginReleaseResource(OverrideSkinWeights);
+		OverrideSkinWeights->BeginReleaseResources();
 		// Ensure the RT no longer accessed the data, might slow down
 		FlushRenderingCommands();
 		// The RT thread has no access to it any more so it's safe to delete it.
@@ -3233,7 +3227,7 @@ void FSkelMeshComponentLODInfo::BeginReleaseOverrideSkinWeights()
 	if (OverrideSkinWeights)
 	{
 		// enqueue a rendering command to release
-		BeginReleaseResource(OverrideSkinWeights);
+		OverrideSkinWeights->BeginReleaseResources();
 	}
 }
 
@@ -3346,16 +3340,15 @@ void USkinnedMeshComponent::ClearVertexColorOverride(int32 LODIndex)
 }
 
 /** 
- *	Templated util for converting from API skin weight description to templated GPU format. 
+ *	Util for converting from API skin weight description to GPU format. 
  * 	This includes remapping from skeleton bone index to section bone index.
  */
-template <bool bExtraBoneInfluencesT>
 void CreateSectionSkinWeightsArray(
 	const TArray<FSkelMeshSkinWeightInfo>& InSourceWeights,
 	int32 StartIndex,
 	int32 NumVerts,
 	const TMap<int32, int32>& SkelToSectionBoneMap,
-	TArray<TSkinWeightInfo<bExtraBoneInfluencesT>>& OutGPUWeights,
+	TArray<FSkinWeightInfo>& OutGPUWeights,
 	TArray<int32>& OutInvalidBones)
 {
 	OutGPUWeights.AddUninitialized(NumVerts);
@@ -3366,14 +3359,14 @@ void CreateSectionSkinWeightsArray(
 	// Iterate over new output buffer
 	for(int VertIndex = StartIndex; VertIndex < StartIndex + NumVerts; VertIndex++)
 	{
-		TSkinWeightInfo<bExtraBoneInfluencesT>& TargetWeight = OutGPUWeights[VertIndex];
+		FSkinWeightInfo& TargetWeight = OutGPUWeights[VertIndex];
 		// while we have valid entries in input buffer
 		if (VertIndex < InSourceWeights.Num())
 		{
 			const FSkelMeshSkinWeightInfo& SrcWeight = InSourceWeights[VertIndex];
 
 			// Iterate over influences
-			for (int32 InfIndex = 0; InfIndex < TargetWeight.NumInfluences; InfIndex++)
+			for (int32 InfIndex = 0; InfIndex < MAX_TOTAL_INFLUENCES; InfIndex++)
 			{
 				// init to zero
 				TargetWeight.InfluenceBones[InfIndex] = 0;
@@ -3408,7 +3401,7 @@ void CreateSectionSkinWeightsArray(
 			TargetWeight.InfluenceBones[0] = 0;
 			TargetWeight.InfluenceWeights[0] = 255;
 
-			for (int32 InfIndex = 1; InfIndex < TargetWeight.NumInfluences; InfIndex++)
+			for (int32 InfIndex = 1; InfIndex < MAX_TOTAL_INFLUENCES; InfIndex++)
 			{
 				TargetWeight.InfluenceBones[InfIndex] = 0;
 				TargetWeight.InfluenceWeights[InfIndex] = 0;
@@ -3422,11 +3415,10 @@ void CreateSectionSkinWeightsArray(
 	}
 }
 
-template <bool bExtraBoneInfluencesT>
 void CreateSkinWeightsArray(
 	const TArray<FSkelMeshSkinWeightInfo>& InSourceWeights,
 	FSkeletalMeshLODRenderData& LODData,
-	TArray<TSkinWeightInfo<bExtraBoneInfluencesT>>& OutGPUWeights,
+	TArray<FSkinWeightInfo>& OutGPUWeights,
 	const FReferenceSkeleton& RefSkel)
 {
 	// Index of first vertex in current section, in the big overall buffer
@@ -3443,9 +3435,9 @@ void CreateSkinWeightsArray(
 			SkelToSectionBoneMap.Add(Section.BoneMap[i], i);
 		}
 
-		// Convert skin weight struct format and assign to new vertex buffer (templated by num weights)
+		// Convert skin weight struct format and assign to new vertex buffer
 		TArray<int32> InvalidBones;
-		CreateSectionSkinWeightsArray<bExtraBoneInfluencesT>(InSourceWeights, BaseVertIndex, NumVertsInSection, SkelToSectionBoneMap, OutGPUWeights, InvalidBones);
+		CreateSectionSkinWeightsArray(InSourceWeights, BaseVertIndex, NumVertsInSection, SkelToSectionBoneMap, OutGPUWeights, InvalidBones);
 
 		// Log info for invalid bones
 		if (InvalidBones.Num() > 0)
@@ -3492,29 +3484,20 @@ void USkinnedMeshComponent::SetSkinWeightOverride(int32 LODIndex, const TArray<F
 				UE_LOG(LogSkinnedMeshComp, Warning, TEXT("SetSkinWeightOverride: Too many weights - expected %d, got %d - truncating"), ExpectedNumVerts, SkinWeights.Num());
 			}
 
-			bool bExtraWeights = LODData.DoesVertexBufferHaveExtraBoneInfluences();
+			uint32 NumBoneInfluences = LODData.GetVertexBufferMaxBoneInfluences();
+			bool bUse16BitBoneIndex = LODData.DoesVertexBufferUse16BitBoneIndex();
 
 			// Allocate skin weight override buffer
 			Info.OverrideSkinWeights = new FSkinWeightVertexBuffer;
 			Info.OverrideSkinWeights->SetNeedsCPUAccess(true);
-			Info.OverrideSkinWeights->SetHasExtraBoneInfluences(bExtraWeights);
+			Info.OverrideSkinWeights->SetMaxBoneInfluences(NumBoneInfluences);
+			Info.OverrideSkinWeights->SetUse16BitBoneIndex(bUse16BitBoneIndex);
 
 			const FReferenceSkeleton& RefSkel = SkeletalMesh->RefSkeleton;
-
-			if (bExtraWeights)
-			{
-				TArray<TSkinWeightInfo<true>> GPUWeights;
-				CreateSkinWeightsArray<true>(SkinWeights, LODData, GPUWeights, RefSkel);
-				*(Info.OverrideSkinWeights) = GPUWeights;
-			}
-			else
-			{
-				TArray<TSkinWeightInfo<false>> GPUWeights;
-				CreateSkinWeightsArray<false>(SkinWeights, LODData, GPUWeights, RefSkel);
-				*(Info.OverrideSkinWeights) = GPUWeights;
-			}
-
-			BeginInitResource(Info.OverrideSkinWeights);
+			TArray<FSkinWeightInfo> GPUWeights;
+			CreateSkinWeightsArray(SkinWeights, LODData, GPUWeights, RefSkel);
+			*(Info.OverrideSkinWeights) = GPUWeights;
+			Info.OverrideSkinWeights->BeginInitResources();
 
 			MarkRenderStateDirty();
 		}
@@ -3815,7 +3798,7 @@ float FAnimUpdateRateParameters::GetRootMotionInterp() const
 	return 1.f;
 }
 /** Simple, CPU evaluation of a vertex's skinned position helper function */
-template <bool bExtraBoneInfluencesT, bool bCachedMatrices>
+template <bool bCachedMatrices>
 FVector GetTypedSkinnedVertexPosition(
 	const USkinnedMeshComponent* SkinnedComp,
 	const FSkelMeshRenderSection& Section,
@@ -3832,8 +3815,7 @@ FVector GetTypedSkinnedVertexPosition(
 
 	// Do soft skinning for this vertex.
 	int32 BufferVertIndex = Section.GetVertexBufferIndex() + VertIndex;
-	const TSkinWeightInfo<bExtraBoneInfluencesT>* SrcSkinWeights = SkinWeightVertexBuffer.GetSkinWeightPtr<bExtraBoneInfluencesT>(BufferVertIndex);
-	int32 MaxBoneInfluences = bExtraBoneInfluencesT ? MAX_TOTAL_INFLUENCES : MAX_INFLUENCES_PER_STREAM;
+	int32 MaxBoneInfluences = SkinWeightVertexBuffer.GetMaxBoneInfluences();
 
 #if !PLATFORM_LITTLE_ENDIAN
 	// uint8[] elements in LOD.VertexBufferGPUSkin have been swapped for VET_UBYTE4 vertex stream use
@@ -3842,7 +3824,7 @@ FVector GetTypedSkinnedVertexPosition(
 	for (int32 InfluenceIndex = 0; InfluenceIndex < MaxBoneInfluences; InfluenceIndex++)
 #endif
 	{
-		const int32 MeshBoneIndex = Section.BoneMap[SrcSkinWeights->InfluenceBones[InfluenceIndex]];
+		const int32 MeshBoneIndex = Section.BoneMap[SkinWeightVertexBuffer.GetBoneIndex(BufferVertIndex, InfluenceIndex)];
 		int32 TransformBoneIndex = MeshBoneIndex;
 
 		if (MasterPoseComponentInst)
@@ -3852,7 +3834,7 @@ FVector GetTypedSkinnedVertexPosition(
 			TransformBoneIndex = MasterBoneMap[MeshBoneIndex];
 		}
 
-		const float	Weight = (float)SrcSkinWeights->InfluenceWeights[InfluenceIndex] / 255.0f;
+		const float	Weight = (float)SkinWeightVertexBuffer.GetBoneWeight(BufferVertIndex, InfluenceIndex) / 255.0f;
 		{
 			if (bCachedMatrices)
 			{
@@ -3873,14 +3855,8 @@ FVector GetTypedSkinnedVertexPosition(
 
 
 
-template FVector GetTypedSkinnedVertexPosition<false, true>(const USkinnedMeshComponent* SkinnedComp, const FSkelMeshRenderSection& Section, const FPositionVertexBuffer& PositionVertexBuffer,
-	const FSkinWeightVertexBuffer& SkinWeightVertexBuffer, const int32 VertIndex, const TArray<FMatrix> & RefToLocals);
-
-template FVector GetTypedSkinnedVertexPosition<true, true>(const USkinnedMeshComponent* SkinnedComp, const FSkelMeshRenderSection& Section, const FPositionVertexBuffer& PositionVertexBuffer,
-	const FSkinWeightVertexBuffer& SkinWeightVertexBuffer, const int32 VertIndex, const TArray<FMatrix> & RefToLocals);
-
-template FVector GetTypedSkinnedVertexPosition<true, false>(const USkinnedMeshComponent* SkinnedComp, const FSkelMeshRenderSection& Section, const FPositionVertexBuffer& PositionVertexBuffer,
+template FVector GetTypedSkinnedVertexPosition<true>(const USkinnedMeshComponent* SkinnedComp, const FSkelMeshRenderSection& Section, const FPositionVertexBuffer& PositionVertexBuffer,
 	const FSkinWeightVertexBuffer& SkinWeightVertexBuffer, const int32 VertIndex, const TArray<FMatrix> & RefToLocals);
 
-template FVector GetTypedSkinnedVertexPosition<false, false>(const USkinnedMeshComponent* SkinnedComp, const FSkelMeshRenderSection& Section, const FPositionVertexBuffer& PositionVertexBuffer,
+template FVector GetTypedSkinnedVertexPosition<false>(const USkinnedMeshComponent* SkinnedComp, const FSkelMeshRenderSection& Section, const FPositionVertexBuffer& PositionVertexBuffer,
 	const FSkinWeightVertexBuffer& SkinWeightVertexBuffer, const int32 VertIndex, const TArray<FMatrix> & RefToLocals);
diff --git a/Engine/Source/Runtime/Engine/Private/Curves/RichCurve.cpp b/Engine/Source/Runtime/Engine/Private/Curves/RichCurve.cpp
--- a/Engine/Source/Runtime/Engine/Private/Curves/RichCurve.cpp
+++ b/Engine/Source/Runtime/Engine/Private/Curves/RichCurve.cpp
@@ -161,6 +161,11 @@ FRichCurveKey FRichCurve::GetKey(FKeyHandle KeyHandle) const
 	return Keys[GetIndex(KeyHandle)];
 }
 
+const FRichCurveKey& FRichCurve::GetKeyRef(FKeyHandle KeyHandle) const
+{
+	EnsureAllIndicesHaveHandles();
+	return Keys[GetIndex(KeyHandle)];
+}
 
 FRichCurveKey FRichCurve::GetFirstKey() const
 {
@@ -1130,6 +1135,7 @@ static ERichCurveCompressionFormat FindRichCurveCompressionFormat(const FRichCur
 	switch (RefKey.InterpMode)
 	{
 	case RCIM_Constant:
+		return RCCF_Mixed;
 	case RCIM_None:
 	default:
 		return RCCF_Constant;
diff --git a/Engine/Source/Runtime/Engine/Private/GPUSkinCache.cpp b/Engine/Source/Runtime/Engine/Private/GPUSkinCache.cpp
--- a/Engine/Source/Runtime/Engine/Private/GPUSkinCache.cpp
+++ b/Engine/Source/Runtime/Engine/Private/GPUSkinCache.cpp
@@ -159,8 +159,12 @@ public:
 		}
 
 		FSkinWeightVertexBuffer* WeightBuffer = GPUSkin->GetSkinWeightVertexBuffer(LOD);
-		InputWeightStride = WeightBuffer->GetStride();
-		InputWeightStreamSRV = WeightBuffer->GetSRV();
+		BoneInfluenceType = WeightBuffer->GetBoneInfluenceType();
+		bUse16BitBoneIndex = WeightBuffer->Use16BitBoneIndex();
+		InputWeightIndexSize = WeightBuffer->GetBoneIndexByteSize();
+		InputWeightStride = WeightBuffer->GetConstantInfluencesVertexStride();
+		InputWeightStreamSRV = WeightBuffer->GetDataVertexBuffer()->GetSRV();
+		InputWeightLookupStreamSRV = WeightBuffer->GetLookupVertexBuffer()->GetSRV();
 	}
 
 	~FGPUSkinCacheEntry()
@@ -186,7 +190,7 @@ public:
 		// 0:normal, 1:with morph target, 2:with APEX cloth (not yet implemented)
 		uint32 SkinType = 0;
 		//
-		bool bExtraBoneInfluences = false;
+		uint32 NumBoneInfluences = 0;
 
 		// in floats (4 bytes)
 		uint32 OutputStreamStart = 0;
@@ -335,7 +339,7 @@ public:
 		Data.NumTexCoords = InSourceVertexFactory->GetNumTexCoords();
 		Data.PositionBufferSRV = InSourceVertexFactory->GetPositionsSRV();
 
-		Data.bExtraBoneInfluences = InSourceVertexFactory->UsesExtraBoneInfluences();
+		Data.NumBoneInfluences = InSourceVertexFactory->GetNumBoneInfluences();
 		check(Data.TangentBufferSRV && Data.PositionBufferSRV);
 
 		// weight buffer
@@ -385,8 +389,12 @@ protected:
 	TArray<FGPUSkinBatchElementUserData> BatchElementsUserData;
 	TArray<FSectionDispatchData> DispatchData;
 	FSkeletalMeshObjectGPUSkin* GPUSkin;
+	int BoneInfluenceType;
+	bool bUse16BitBoneIndex;
+	uint32 InputWeightIndexSize;
 	uint32 InputWeightStride;
 	FShaderResourceViewRHIRef InputWeightStreamSRV;
+	FShaderResourceViewRHIRef InputWeightLookupStreamSRV;
 	FRHIShaderResourceView* MorphBuffer;
 	FShaderResourceViewRHIRef ClothBuffer;
 	FShaderResourceViewRHIRef ClothPositionsAndNormalsBuffer;
@@ -415,9 +423,12 @@ public:
 
 		InputStreamStart.Bind(Initializer.ParameterMap, TEXT("InputStreamStart"));
 
+		NumBoneInfluences.Bind(Initializer.ParameterMap, TEXT("NumBoneInfluences"));
+		InputWeightIndexSize.Bind(Initializer.ParameterMap, TEXT("InputWeightIndexSize"));
 		InputWeightStart.Bind(Initializer.ParameterMap, TEXT("InputWeightStart"));
 		InputWeightStride.Bind(Initializer.ParameterMap, TEXT("InputWeightStride"));
 		InputWeightStream.Bind(Initializer.ParameterMap, TEXT("InputWeightStream"));
+		InputWeightLookupStream.Bind(Initializer.ParameterMap, TEXT("InputWeightLookupStream"));
 
 		PositionBufferUAV.Bind(Initializer.ParameterMap, TEXT("PositionBufferUAV"));
 		TangentBufferUAV.Bind(Initializer.ParameterMap, TEXT("TangentBufferUAV"));
@@ -450,9 +461,12 @@ public:
 		SetSRVParameter(RHICmdList, ShaderRHI, TangentInputBuffer, DispatchData.TangentBufferSRV);
 		SetSRVParameter(RHICmdList, ShaderRHI, PositionInputBuffer, DispatchData.PositionBufferSRV);
 
+		SetShaderValue(RHICmdList, ShaderRHI, NumBoneInfluences, DispatchData.NumBoneInfluences);
+		SetShaderValue(RHICmdList, ShaderRHI, InputWeightIndexSize, Entry->InputWeightIndexSize);
 		SetShaderValue(RHICmdList, ShaderRHI, InputWeightStart, DispatchData.InputWeightStart);
 		SetShaderValue(RHICmdList, ShaderRHI, InputWeightStride, Entry->InputWeightStride);
 		SetSRVParameter(RHICmdList, ShaderRHI, InputWeightStream, Entry->InputWeightStreamSRV);
+		SetSRVParameter(RHICmdList, ShaderRHI, InputWeightLookupStream, Entry->InputWeightLookupStreamSRV);
 
 		// output UAV
 		SetUAVParameter(RHICmdList, ShaderRHI, PositionBufferUAV, PositionUAV);
@@ -496,8 +510,8 @@ public:
 			<< TangentInputBuffer << PositionInputBuffer << PositionBufferUAV << TangentBufferUAV << BoneMatrices << MorphBuffer << MorphBufferOffset
 			<< SkinCacheDebug;
 
-		Ar << InputWeightStart << InputWeightStride << InputWeightStream;
-
+		Ar << NumBoneInfluences << InputWeightIndexSize << InputWeightStart << InputWeightStride << InputWeightStream << InputWeightLookupStream;
+		
         Ar << ClothBuffer << ClothPositionsAndNormalsBuffer << ClothBufferOffset << ClothBlendWeight << ClothLocalToWorld << ClothWorldToLocal;
 
 		//Ar << DebugParameter;
@@ -521,9 +535,12 @@ private:
 	FShaderResourceParameter PositionBufferUAV;
 	FShaderResourceParameter TangentBufferUAV;
 
+	FShaderParameter NumBoneInfluences;
+	FShaderParameter InputWeightIndexSize;
 	FShaderParameter InputWeightStart;
 	FShaderParameter InputWeightStride;
 	FShaderResourceParameter InputWeightStream;
+	FShaderResourceParameter InputWeightLookupStream;
 
 	FShaderResourceParameter MorphBuffer;
 	FShaderParameter MorphBufferOffset;
@@ -541,7 +558,9 @@ private:
 template <int Permutation>
 class TGPUSkinCacheCS : public FBaseGPUSkinCacheCS
 {
-    constexpr static bool bUseExtraBoneInfluencesT = (4 == (Permutation & 4));
+	constexpr static bool bBoneIndex16 = (16 == (Permutation & 16));
+	constexpr static bool bUnlimitedBoneInfluence = (8 == (Permutation & 12));
+	constexpr static bool bUseExtraBoneInfluencesT = (4 == (Permutation & 12));
     constexpr static bool bMorphBlend = (1 == (Permutation & 3));
 	constexpr static bool bApexCloth = (2 == (Permutation & 3));
 
@@ -556,14 +575,18 @@ public:
 	static void ModifyCompilationEnvironment(const FGlobalShaderPermutationParameters& Parameters, FShaderCompilerEnvironment& OutEnvironment)
 	{
 		FGlobalShader::ModifyCompilationEnvironment(Parameters, OutEnvironment);
+		const uint32 BoneIndex16 = bBoneIndex16;
+		const uint32 UnlimitedBoneInfluence = bUnlimitedBoneInfluence;
 		const uint32 UseExtraBoneInfluences = bUseExtraBoneInfluencesT;
 		const uint32 MorphBlend = bMorphBlend;
 		const uint32 ApexCloth = bApexCloth;
+		OutEnvironment.SetDefine(TEXT("GPUSKIN_UNLIMITED_BONE_INFLUENCE"), UnlimitedBoneInfluence);
 		OutEnvironment.SetDefine(TEXT("GPUSKIN_USE_EXTRA_INFLUENCES"), UseExtraBoneInfluences);
 		OutEnvironment.SetDefine(TEXT("GPUSKIN_MORPH_BLEND"), MorphBlend);
 		OutEnvironment.SetDefine(TEXT("GPUSKIN_APEX_CLOTH"), ApexCloth);
 		OutEnvironment.SetDefine(TEXT("GPUSKIN_RWBUFFER_OFFSET_TANGENT_X"), FGPUSkinCache::RWTangentXOffsetInFloats);
 		OutEnvironment.SetDefine(TEXT("GPUSKIN_RWBUFFER_OFFSET_TANGENT_Z"), FGPUSkinCache::RWTangentZOffsetInFloats);
+		OutEnvironment.SetDefine(TEXT("GPUSKIN_BONE_INDEX_UINT16"), BoneIndex16);
 	}
 
 	TGPUSkinCacheCS(const ShaderMetaType::CompiledShaderInitializerType& Initializer)
@@ -576,12 +599,21 @@ public:
 	}
 };
 
-IMPLEMENT_SHADER_TYPE(template<>, TGPUSkinCacheCS<0>, TEXT("/Engine/Private/GpuSkinCacheComputeShader.usf"), TEXT("SkinCacheUpdateBatchCS"), SF_Compute);
-IMPLEMENT_SHADER_TYPE(template<>, TGPUSkinCacheCS<1>, TEXT("/Engine/Private/GpuSkinCacheComputeShader.usf"), TEXT("SkinCacheUpdateBatchCS"), SF_Compute);
-IMPLEMENT_SHADER_TYPE(template<>, TGPUSkinCacheCS<2>, TEXT("/Engine/Private/GpuSkinCacheComputeShader.usf"), TEXT("SkinCacheUpdateBatchCS"), SF_Compute);
-IMPLEMENT_SHADER_TYPE(template<>, TGPUSkinCacheCS<4>, TEXT("/Engine/Private/GpuSkinCacheComputeShader.usf"), TEXT("SkinCacheUpdateBatchCS"), SF_Compute);
-IMPLEMENT_SHADER_TYPE(template<>, TGPUSkinCacheCS<5>, TEXT("/Engine/Private/GpuSkinCacheComputeShader.usf"), TEXT("SkinCacheUpdateBatchCS"), SF_Compute);
-IMPLEMENT_SHADER_TYPE(template<>, TGPUSkinCacheCS<6>, TEXT("/Engine/Private/GpuSkinCacheComputeShader.usf"), TEXT("SkinCacheUpdateBatchCS"), SF_Compute);
+IMPLEMENT_SHADER_TYPE(template<>, TGPUSkinCacheCS<0>, TEXT("/Engine/Private/GpuSkinCacheComputeShader.usf"), TEXT("SkinCacheUpdateBatchCS"), SF_Compute);	// 16bit_0, BoneInfluenceType_0, SkinType_0 
+IMPLEMENT_SHADER_TYPE(template<>, TGPUSkinCacheCS<1>, TEXT("/Engine/Private/GpuSkinCacheComputeShader.usf"), TEXT("SkinCacheUpdateBatchCS"), SF_Compute);	// 16bit_0, BoneInfluenceType_0, SkinType_1 
+IMPLEMENT_SHADER_TYPE(template<>, TGPUSkinCacheCS<2>, TEXT("/Engine/Private/GpuSkinCacheComputeShader.usf"), TEXT("SkinCacheUpdateBatchCS"), SF_Compute);	// 16bit_0, BoneInfluenceType_0, SkinType_2 
+IMPLEMENT_SHADER_TYPE(template<>, TGPUSkinCacheCS<4>, TEXT("/Engine/Private/GpuSkinCacheComputeShader.usf"), TEXT("SkinCacheUpdateBatchCS"), SF_Compute);	// 16bit_0, BoneInfluenceType_1, SkinType_0 
+IMPLEMENT_SHADER_TYPE(template<>, TGPUSkinCacheCS<5>, TEXT("/Engine/Private/GpuSkinCacheComputeShader.usf"), TEXT("SkinCacheUpdateBatchCS"), SF_Compute);	// 16bit_0, BoneInfluenceType_1, SkinType_1 
+IMPLEMENT_SHADER_TYPE(template<>, TGPUSkinCacheCS<6>, TEXT("/Engine/Private/GpuSkinCacheComputeShader.usf"), TEXT("SkinCacheUpdateBatchCS"), SF_Compute);	// 16bit_0, BoneInfluenceType_1, SkinType_2 
+IMPLEMENT_SHADER_TYPE(template<>, TGPUSkinCacheCS<8>, TEXT("/Engine/Private/GpuSkinCacheComputeShader.usf"), TEXT("SkinCacheUpdateBatchCS"), SF_Compute);	// 16bit_0, BoneInfluenceType_2, SkinType_0 
+IMPLEMENT_SHADER_TYPE(template<>, TGPUSkinCacheCS<9>, TEXT("/Engine/Private/GpuSkinCacheComputeShader.usf"), TEXT("SkinCacheUpdateBatchCS"), SF_Compute);	// 16bit_0, BoneInfluenceType_2, SkinType_1 
+IMPLEMENT_SHADER_TYPE(template<>, TGPUSkinCacheCS<10>, TEXT("/Engine/Private/GpuSkinCacheComputeShader.usf"), TEXT("SkinCacheUpdateBatchCS"), SF_Compute);	// 16bit_0, BoneInfluenceType_2, SkinType_2 
+IMPLEMENT_SHADER_TYPE(template<>, TGPUSkinCacheCS<16>, TEXT("/Engine/Private/GpuSkinCacheComputeShader.usf"), TEXT("SkinCacheUpdateBatchCS"), SF_Compute);  // 16bit_1, BoneInfluenceType_0, SkinType_0 
+IMPLEMENT_SHADER_TYPE(template<>, TGPUSkinCacheCS<17>, TEXT("/Engine/Private/GpuSkinCacheComputeShader.usf"), TEXT("SkinCacheUpdateBatchCS"), SF_Compute);  // 16bit_1, BoneInfluenceType_0, SkinType_1 
+IMPLEMENT_SHADER_TYPE(template<>, TGPUSkinCacheCS<18>, TEXT("/Engine/Private/GpuSkinCacheComputeShader.usf"), TEXT("SkinCacheUpdateBatchCS"), SF_Compute);  // 16bit_1, BoneInfluenceType_0, SkinType_2 
+IMPLEMENT_SHADER_TYPE(template<>, TGPUSkinCacheCS<20>, TEXT("/Engine/Private/GpuSkinCacheComputeShader.usf"), TEXT("SkinCacheUpdateBatchCS"), SF_Compute);  // 16bit_1, BoneInfluenceType_1, SkinType_0 
+IMPLEMENT_SHADER_TYPE(template<>, TGPUSkinCacheCS<21>, TEXT("/Engine/Private/GpuSkinCacheComputeShader.usf"), TEXT("SkinCacheUpdateBatchCS"), SF_Compute);  // 16bit_1, BoneInfluenceType_1, SkinType_1 
+IMPLEMENT_SHADER_TYPE(template<>, TGPUSkinCacheCS<22>, TEXT("/Engine/Private/GpuSkinCacheComputeShader.usf"), TEXT("SkinCacheUpdateBatchCS"), SF_Compute);  // 16bit_1, BoneInfluenceType_1, SkinType_2 
 
 FGPUSkinCache::FGPUSkinCache(bool bInRequiresMemoryLimit)
 	: UsedMemoryInBytes(0)
@@ -775,8 +807,7 @@ public:
 template <int Permutation>
 class FRecomputeTangentsPerTrianglePassCS : public FBaseRecomputeTangents
 {
-    constexpr static bool bMergeDuplicatedVerts = (4 == (Permutation & 4));
-	constexpr static bool bUseExtraBoneInfluencesT = (2 == (Permutation & 2));
+    constexpr static bool bMergeDuplicatedVerts = (2 == (Permutation & 2));
 	constexpr static bool bFullPrecisionUV = (1 == (Permutation & 1));
 
 	DECLARE_SHADER_TYPE(FRecomputeTangentsPerTrianglePassCS, Global);
@@ -784,9 +815,7 @@ class FRecomputeTangentsPerTrianglePassCS : public FBaseRecomputeTangents
 	static void ModifyCompilationEnvironment(const FGlobalShaderPermutationParameters& Parameters, FShaderCompilerEnvironment& OutEnvironment)
 	{
 		FGlobalShader::ModifyCompilationEnvironment(Parameters, OutEnvironment);
-		const uint32 UseExtraBoneInfluences = bUseExtraBoneInfluencesT;
 		OutEnvironment.SetDefine(TEXT("MERGE_DUPLICATED_VERTICES"), bMergeDuplicatedVerts);
-		OutEnvironment.SetDefine(TEXT("GPUSKIN_USE_EXTRA_INFLUENCES"), UseExtraBoneInfluences);
 		OutEnvironment.SetDefine(TEXT("THREADGROUP_SIZEX"), ThreadGroupSizeX);
 		OutEnvironment.SetDefine(TEXT("INTERMEDIATE_ACCUM_BUFFER_NUM_INTS"), FGPUSkinCache::IntermediateAccumBufferNumInts);
 		OutEnvironment.SetDefine(TEXT("FULL_PRECISION_UV"), bFullPrecisionUV);
@@ -805,10 +834,6 @@ IMPLEMENT_SHADER_TYPE(template<>, FRecomputeTangentsPerTrianglePassCS<0>, TEXT("
 IMPLEMENT_SHADER_TYPE(template<>, FRecomputeTangentsPerTrianglePassCS<1>, TEXT("/Engine/Private/RecomputeTangentsPerTrianglePass.usf"), TEXT("MainCS"), SF_Compute);
 IMPLEMENT_SHADER_TYPE(template<>, FRecomputeTangentsPerTrianglePassCS<2>, TEXT("/Engine/Private/RecomputeTangentsPerTrianglePass.usf"), TEXT("MainCS"), SF_Compute);
 IMPLEMENT_SHADER_TYPE(template<>, FRecomputeTangentsPerTrianglePassCS<3>, TEXT("/Engine/Private/RecomputeTangentsPerTrianglePass.usf"), TEXT("MainCS"), SF_Compute);
-IMPLEMENT_SHADER_TYPE(template<>, FRecomputeTangentsPerTrianglePassCS<4>, TEXT("/Engine/Private/RecomputeTangentsPerTrianglePass.usf"), TEXT("MainCS"), SF_Compute);
-IMPLEMENT_SHADER_TYPE(template<>, FRecomputeTangentsPerTrianglePassCS<5>, TEXT("/Engine/Private/RecomputeTangentsPerTrianglePass.usf"), TEXT("MainCS"), SF_Compute);
-IMPLEMENT_SHADER_TYPE(template<>, FRecomputeTangentsPerTrianglePassCS<6>, TEXT("/Engine/Private/RecomputeTangentsPerTrianglePass.usf"), TEXT("MainCS"), SF_Compute);
-IMPLEMENT_SHADER_TYPE(template<>, FRecomputeTangentsPerTrianglePassCS<7>, TEXT("/Engine/Private/RecomputeTangentsPerTrianglePass.usf"), TEXT("MainCS"), SF_Compute);
 
 /** Encapsulates the RecomputeSkinTangentsResolve compute shader. */
 class FRecomputeTangentsPerVertexPassCS : public FGlobalShader
@@ -825,7 +850,6 @@ class FRecomputeTangentsPerVertexPassCS : public FGlobalShader
 	{
 		FGlobalShader::ModifyCompilationEnvironment(Parameters, OutEnvironment);
 		// this pass cannot read the input as it doesn't have the permutation
-		OutEnvironment.SetDefine(TEXT("GPUSKIN_USE_EXTRA_INFLUENCES"), (uint32)0);
 		OutEnvironment.SetDefine(TEXT("THREADGROUP_SIZEX"), ThreadGroupSizeX);
 		OutEnvironment.SetDefine(TEXT("GPUSKIN_RWBUFFER_OFFSET_TANGENT_X"), FGPUSkinCache::RWTangentXOffsetInFloats);
 		OutEnvironment.SetDefine(TEXT("GPUSKIN_RWBUFFER_OFFSET_TANGENT_Z"), FGPUSkinCache::RWTangentZOffsetInFloats);
@@ -937,14 +961,10 @@ void FGPUSkinCache::DispatchUpdateSkinTangents(FRHICommandListImmediate& RHICmdL
 		// This code can be optimized by batched up and doing it with less Dispatch calls (costs more memory)
 		{
 			auto* GlobalShaderMap = GetGlobalShaderMap(ERHIFeatureLevel::SM5);
-			TShaderMapRef<FRecomputeTangentsPerTrianglePassCS<0>> ComputeShader000(GlobalShaderMap);
-			TShaderMapRef<FRecomputeTangentsPerTrianglePassCS<1>> ComputeShader010(GlobalShaderMap);
-			TShaderMapRef<FRecomputeTangentsPerTrianglePassCS<2>> ComputeShader100(GlobalShaderMap);
-			TShaderMapRef<FRecomputeTangentsPerTrianglePassCS<3>> ComputeShader110(GlobalShaderMap);
-			TShaderMapRef<FRecomputeTangentsPerTrianglePassCS<4>> ComputeShader001(GlobalShaderMap);
-			TShaderMapRef<FRecomputeTangentsPerTrianglePassCS<5>> ComputeShader011(GlobalShaderMap);
-			TShaderMapRef<FRecomputeTangentsPerTrianglePassCS<6>> ComputeShader101(GlobalShaderMap);
-			TShaderMapRef<FRecomputeTangentsPerTrianglePassCS<7>> ComputeShader111(GlobalShaderMap);
+			TShaderMapRef<FRecomputeTangentsPerTrianglePassCS<0>> ComputeShader00(GlobalShaderMap);
+			TShaderMapRef<FRecomputeTangentsPerTrianglePassCS<1>> ComputeShader01(GlobalShaderMap);
+			TShaderMapRef<FRecomputeTangentsPerTrianglePassCS<2>> ComputeShader10(GlobalShaderMap);
+			TShaderMapRef<FRecomputeTangentsPerTrianglePassCS<3>> ComputeShader11(GlobalShaderMap);
 
 			bool bFullPrecisionUV = LodData.StaticVertexBuffers.StaticMeshVertexBuffer.GetUseFullPrecisionUVs();
 
@@ -952,25 +972,11 @@ void FGPUSkinCache::DispatchUpdateSkinTangents(FRHICommandListImmediate& RHICmdL
 
 			if (bFullPrecisionUV)
 			{
-				if (DispatchData.bExtraBoneInfluences)
-				{
-					Shader = GAllowDupedVertsForRecomputeTangents ? (FBaseRecomputeTangents*)*ComputeShader100 : (FBaseRecomputeTangents*)*ComputeShader101;
-				}
-				else
-				{
-					Shader = GAllowDupedVertsForRecomputeTangents ? (FBaseRecomputeTangents*)*ComputeShader000 : (FBaseRecomputeTangents*)*ComputeShader001;
-				}
+				Shader = GAllowDupedVertsForRecomputeTangents ? (FBaseRecomputeTangents*)*ComputeShader11 : (FBaseRecomputeTangents*)*ComputeShader01;
 			}
 			else
 			{
-				if (DispatchData.bExtraBoneInfluences)
-				{
-					Shader = GAllowDupedVertsForRecomputeTangents ? (FBaseRecomputeTangents*)*ComputeShader110 : (FBaseRecomputeTangents*)*ComputeShader111;
-				}
-				else
-				{
-					Shader = GAllowDupedVertsForRecomputeTangents ? (FBaseRecomputeTangents*)*ComputeShader010 : (FBaseRecomputeTangents*)*ComputeShader011;
-				}
+				Shader = GAllowDupedVertsForRecomputeTangents ? (FBaseRecomputeTangents*)*ComputeShader10 : (FBaseRecomputeTangents*)*ComputeShader00;
 			}
 
 			check(Shader);
@@ -978,8 +984,8 @@ void FGPUSkinCache::DispatchUpdateSkinTangents(FRHICommandListImmediate& RHICmdL
 			uint32 NumTriangles = DispatchData.NumTriangles;
 			uint32 ThreadGroupCountValue = FMath::DivideAndRoundUp(NumTriangles, FBaseRecomputeTangents::ThreadGroupSizeX);
 
-			SCOPED_DRAW_EVENTF(RHICmdList, SkinTangents_PerTrianglePass, TEXT("TangentsTri IndexStart=%d Tri=%d ExtraBoneInfluences=%d UVPrecision=%d"),
-				DispatchData.IndexBufferOffsetValue, DispatchData.NumTriangles, DispatchData.bExtraBoneInfluences, bFullPrecisionUV);
+			SCOPED_DRAW_EVENTF(RHICmdList, SkinTangents_PerTrianglePass, TEXT("TangentsTri IndexStart=%d Tri=%d BoneInfluenceType=%d UVPrecision=%d"),
+				DispatchData.IndexBufferOffsetValue, DispatchData.NumTriangles, Entry->BoneInfluenceType, bFullPrecisionUV);
 
 			FRHIComputeShader* ShaderRHI = Shader->GetComputeShader();
 			RHICmdList.SetComputeShader(ShaderRHI);
@@ -1131,10 +1137,10 @@ void FGPUSkinCache::ProcessEntry(FRHICommandListImmediate& RHICmdList, FGPUBaseS
 
 		// weight buffer
 		FSkinWeightVertexBuffer* WeightBuffer = Skin->GetSkinWeightVertexBuffer(LODIndex);
-		uint32 WeightStride = WeightBuffer->GetStride();
+		uint32 WeightStride = WeightBuffer->GetConstantInfluencesVertexStride();
 		InOutEntry->DispatchData[Section].InputWeightStart = (WeightStride * BatchElement.BaseVertexIndex) / sizeof(float);
 		InOutEntry->InputWeightStride = WeightStride;
-		InOutEntry->InputWeightStreamSRV = WeightBuffer->GetSRV();
+		InOutEntry->InputWeightStreamSRV = WeightBuffer->GetDataVertexBuffer()->GetSRV();
 	}
 
     FVertexBufferAndSRV ClothPositionAndNormalsBuffer;
@@ -1256,26 +1262,41 @@ void FGPUSkinCache::DispatchUpdateSkinning(FRHICommandListImmediate& RHICmdList,
 	FGPUBaseSkinVertexFactory::FShaderDataType& ShaderData = DispatchData.SourceVertexFactory->GetShaderData();
 
 	SCOPED_DRAW_EVENTF(RHICmdList, SkinCacheDispatch,
-		TEXT("Skinning%d%d Chunk=%d InStreamStart=%d OutStart=%d Vert=%d Morph=%d/%d"),
-		(int32)DispatchData.bExtraBoneInfluences, DispatchData.SkinType,
+		TEXT("Skinning%d%d%d Chunk=%d InStreamStart=%d OutStart=%d Vert=%d Morph=%d/%d"),
+		(int32)Entry->bUse16BitBoneIndex, (int32)Entry->BoneInfluenceType, DispatchData.SkinType,
 		DispatchData.SectionIndex, DispatchData.InputStreamStart, DispatchData.OutputStreamStart, DispatchData.NumVertices, Entry->MorphBuffer != 0, DispatchData.MorphBufferOffset);
 	auto* GlobalShaderMap = GetGlobalShaderMap(GMaxRHIFeatureLevel);
-	TShaderMapRef<TGPUSkinCacheCS<0>> SkinCacheCS00(GlobalShaderMap);
-	TShaderMapRef<TGPUSkinCacheCS<1>> SkinCacheCS01(GlobalShaderMap);
-    TShaderMapRef<TGPUSkinCacheCS<2>> SkinCacheCS02(GlobalShaderMap);
-    TShaderMapRef<TGPUSkinCacheCS<4>> SkinCacheCS10(GlobalShaderMap);
-    TShaderMapRef<TGPUSkinCacheCS<5>> SkinCacheCS11(GlobalShaderMap);
-    TShaderMapRef<TGPUSkinCacheCS<6>> SkinCacheCS12(GlobalShaderMap);
+	TShaderMapRef<TGPUSkinCacheCS<0>> SkinCacheCS000(GlobalShaderMap);		// 16bit_0, BoneInfluenceType_0, SkinType_0
+	TShaderMapRef<TGPUSkinCacheCS<1>> SkinCacheCS001(GlobalShaderMap);		// 16bit_0, BoneInfluenceType_0, SkinType_1
+	TShaderMapRef<TGPUSkinCacheCS<2>> SkinCacheCS002(GlobalShaderMap);		// 16bit_0, BoneInfluenceType_0, SkinType_2
+	TShaderMapRef<TGPUSkinCacheCS<4>> SkinCacheCS010(GlobalShaderMap);		// 16bit_0, BoneInfluenceType_1, SkinType_0
+	TShaderMapRef<TGPUSkinCacheCS<5>> SkinCacheCS011(GlobalShaderMap);		// 16bit_0, BoneInfluenceType_1, SkinType_1
+	TShaderMapRef<TGPUSkinCacheCS<6>> SkinCacheCS012(GlobalShaderMap);		// 16bit_0, BoneInfluenceType_1, SkinType_2
+	TShaderMapRef<TGPUSkinCacheCS<8>> SkinCacheCS020(GlobalShaderMap);		// 16bit_0, BoneInfluenceType_2, SkinType_0
+	TShaderMapRef<TGPUSkinCacheCS<9>> SkinCacheCS021(GlobalShaderMap);		// 16bit_0, BoneInfluenceType_2, SkinType_1
+	TShaderMapRef<TGPUSkinCacheCS<10>> SkinCacheCS022(GlobalShaderMap);		// 16bit_0, BoneInfluenceType_2, SkinType_2
+	TShaderMapRef<TGPUSkinCacheCS<16>>  SkinCacheCS100(GlobalShaderMap);	// 16bit_1, BoneInfluenceType_0, SkinType_0
+	TShaderMapRef<TGPUSkinCacheCS<17>>  SkinCacheCS101(GlobalShaderMap);	// 16bit_1, BoneInfluenceType_0, SkinType_1
+	TShaderMapRef<TGPUSkinCacheCS<18>>  SkinCacheCS102(GlobalShaderMap);	// 16bit_1, BoneInfluenceType_0, SkinType_2
+	TShaderMapRef<TGPUSkinCacheCS<20>>  SkinCacheCS110(GlobalShaderMap);	// 16bit_1, BoneInfluenceType_1, SkinType_0
+	TShaderMapRef<TGPUSkinCacheCS<21>>  SkinCacheCS111(GlobalShaderMap);	// 16bit_1, BoneInfluenceType_1, SkinType_1
+	TShaderMapRef<TGPUSkinCacheCS<22>>  SkinCacheCS112(GlobalShaderMap);	// 16bit_1, BoneInfluenceType_1, SkinType_2
 
 	FBaseGPUSkinCacheCS* Shader = 0;
 
 	switch (DispatchData.SkinType)
 	{
-	case 0: Shader = DispatchData.bExtraBoneInfluences ? (FBaseGPUSkinCacheCS*)*SkinCacheCS10 : (FBaseGPUSkinCacheCS*)*SkinCacheCS00;
+	case 0: Shader = Entry->BoneInfluenceType == 0 ? (Entry->bUse16BitBoneIndex ? (FBaseGPUSkinCacheCS*)*SkinCacheCS100 : (FBaseGPUSkinCacheCS*)*SkinCacheCS000) :
+					(Entry->BoneInfluenceType == 1 ? (Entry->bUse16BitBoneIndex ? (FBaseGPUSkinCacheCS*)*SkinCacheCS110 : (FBaseGPUSkinCacheCS*)*SkinCacheCS010) :
+													 (FBaseGPUSkinCacheCS*)*SkinCacheCS020);
 		break;
-	case 1: Shader = DispatchData.bExtraBoneInfluences ? (FBaseGPUSkinCacheCS*)*SkinCacheCS11 : (FBaseGPUSkinCacheCS*)*SkinCacheCS01;
+	case 1: Shader = Entry->BoneInfluenceType == 0 ? (Entry->bUse16BitBoneIndex ? (FBaseGPUSkinCacheCS*)*SkinCacheCS101 : (FBaseGPUSkinCacheCS*)*SkinCacheCS001) :
+					(Entry->BoneInfluenceType == 1 ? (Entry->bUse16BitBoneIndex ? (FBaseGPUSkinCacheCS*)*SkinCacheCS111 : (FBaseGPUSkinCacheCS*)*SkinCacheCS011) :
+													 (FBaseGPUSkinCacheCS*)*SkinCacheCS021);
 		break;
-	case 2: Shader = DispatchData.bExtraBoneInfluences ? (FBaseGPUSkinCacheCS*)*SkinCacheCS12 : (FBaseGPUSkinCacheCS*)*SkinCacheCS02;
+	case 2: Shader = Entry->BoneInfluenceType == 0 ? (Entry->bUse16BitBoneIndex ? (FBaseGPUSkinCacheCS*)*SkinCacheCS102 : (FBaseGPUSkinCacheCS*)*SkinCacheCS002) :
+					(Entry->BoneInfluenceType == 1 ? (Entry->bUse16BitBoneIndex ? (FBaseGPUSkinCacheCS*)*SkinCacheCS112 : (FBaseGPUSkinCacheCS*)*SkinCacheCS012) :
+													 (FBaseGPUSkinCacheCS*)*SkinCacheCS022);
 		break;
 	default:
 		check(0);
diff --git a/Engine/Source/Runtime/Engine/Private/GPUSkinVertexFactory.cpp b/Engine/Source/Runtime/Engine/Private/GPUSkinVertexFactory.cpp
--- a/Engine/Source/Runtime/Engine/Private/GPUSkinVertexFactory.cpp
+++ b/Engine/Source/Runtime/Engine/Private/GPUSkinVertexFactory.cpp
@@ -10,6 +10,9 @@
 #include "GPUSkinCache.h"
 #include "ShaderParameterUtils.h"
 #include "MeshMaterialShader.h"
+#if INTEL_ISPC
+#include "GPUSkinVertexFactory.ispc.generated.h"
+#endif
 
 // Changing this is currently unsupported after content has been chunked with the previous setting
 // Changing this causes a full shader recompile
@@ -20,6 +23,13 @@ static FAutoConsoleVariableRef CVarMaxGPUSkinBones(
 	TEXT("Max number of bones that can be skinned on the GPU in a single draw call. Cannot be changed at runtime."),
 	ECVF_ReadOnly);
 
+static int32 GCVarSupport16BitBoneIndex = 0;
+static FAutoConsoleVariableRef CVarSupport16BitBoneIndex(
+	TEXT("r.GPUSkin.Support16BitBoneIndex"),
+	GCVarSupport16BitBoneIndex,
+	TEXT("If enabled, a new mesh imported will use 8 bit (if <=256 bones) or 16 bit (if > 256 bones) bone indices for rendering."),
+	ECVF_ReadOnly);
+
 // Whether to use 2 bones influence instead of default 4 for GPU skinning
 // Changing this causes a full shader recompile
 static TAutoConsoleVariable<int32> CVarGPUSkinLimit2BoneInfluences(
@@ -28,6 +38,10 @@ static TAutoConsoleVariable<int32> CVarGPUSkinLimit2BoneInfluences(
 	TEXT("Whether to use 2 bones influence instead of default 4 for GPU skinning. Cannot be changed at runtime."),
 	ECVF_ReadOnly);
 
+// Changing the unlimited bone influences settings require DDC key SKELETALMESH_DERIVEDDATA_VER to update in order to regenerate bone weight vertex buffers
+static const bool GUnlimitedBoneInfluences = false;
+static const int32 GUnlimitedBoneInfluencesThreshold = EXTRA_BONE_INFLUENCES;
+
 IMPLEMENT_GLOBAL_SHADER_PARAMETER_STRUCT(FAPEXClothUniformShaderParameters,"APEXClothParam");
 
 IMPLEMENT_GLOBAL_SHADER_PARAMETER_STRUCT(FBoneMatricesUniformShaderParameters,"Bones");
@@ -35,8 +49,8 @@ IMPLEMENT_GLOBAL_SHADER_PARAMETER_STRUCT(FBoneMatricesUniformShaderParameters,"B
 static FBoneMatricesUniformShaderParameters GBoneUniformStruct;
 
 #define IMPLEMENT_GPUSKINNING_VERTEX_FACTORY_TYPE_INTERNAL(FactoryClass, ShaderFilename,bUsedWithMaterials,bSupportsStaticLighting,bSupportsDynamicLighting,bPrecisePrevWorldPos,bSupportsPositionOnly) \
-	template <bool bExtraBoneInfluencesT> FVertexFactoryType FactoryClass<bExtraBoneInfluencesT>::StaticType( \
-	bExtraBoneInfluencesT ? TEXT(#FactoryClass) TEXT("true") : TEXT(#FactoryClass) TEXT("false"), \
+	template <GPUSkinBoneInfluenceType BoneInfluenceType> FVertexFactoryType FactoryClass<BoneInfluenceType>::StaticType( \
+	BoneInfluenceType == DefaultBoneInfluence ? TEXT(#FactoryClass) TEXT("Default") : (BoneInfluenceType == ExtraBoneInfluence ? TEXT(#FactoryClass) TEXT("Extra") : TEXT(#FactoryClass) TEXT("Unlimited")), \
 	TEXT(ShaderFilename), \
 	bUsedWithMaterials, \
 	bSupportsStaticLighting, \
@@ -45,20 +59,21 @@ static FBoneMatricesUniformShaderParameters GBoneUniformStruct;
 	bSupportsPositionOnly, \
 	false, \
 	false, \
-	Construct##FactoryClass##ShaderParameters<bExtraBoneInfluencesT>, \
-	FactoryClass<bExtraBoneInfluencesT>::ShouldCompilePermutation, \
-	FactoryClass<bExtraBoneInfluencesT>::ModifyCompilationEnvironment, \
-	FactoryClass<bExtraBoneInfluencesT>::ValidateCompiledResult, \
-	FactoryClass<bExtraBoneInfluencesT>::SupportsTessellationShaders \
+	Construct##FactoryClass##ShaderParameters<BoneInfluenceType>, \
+	FactoryClass<BoneInfluenceType>::ShouldCompilePermutation, \
+	FactoryClass<BoneInfluenceType>::ModifyCompilationEnvironment, \
+	FactoryClass<BoneInfluenceType>::ValidateCompiledResult, \
+	FactoryClass<BoneInfluenceType>::SupportsTessellationShaders \
 	); \
-	template <bool bExtraBoneInfluencesT> inline FVertexFactoryType* FactoryClass<bExtraBoneInfluencesT>::GetType() const { return &StaticType; }
+	template <GPUSkinBoneInfluenceType BoneInfluenceType> inline FVertexFactoryType* FactoryClass<BoneInfluenceType>::GetType() const { return &StaticType; }
 
 
 #define IMPLEMENT_GPUSKINNING_VERTEX_FACTORY_TYPE(FactoryClass, ShaderFilename,bUsedWithMaterials,bSupportsStaticLighting,bSupportsDynamicLighting,bPrecisePrevWorldPos,bSupportsPositionOnly) \
-	template <bool bExtraBoneInfluencesT> FVertexFactoryShaderParameters* Construct##FactoryClass##ShaderParameters(EShaderFrequency ShaderFrequency) { return FactoryClass<bExtraBoneInfluencesT>::ConstructShaderParameters(ShaderFrequency); } \
+	template <GPUSkinBoneInfluenceType BoneInfluenceType> FVertexFactoryShaderParameters* Construct##FactoryClass##ShaderParameters(EShaderFrequency ShaderFrequency) { return FactoryClass<BoneInfluenceType>::ConstructShaderParameters(ShaderFrequency); } \
 	IMPLEMENT_GPUSKINNING_VERTEX_FACTORY_TYPE_INTERNAL(FactoryClass, ShaderFilename,bUsedWithMaterials,bSupportsStaticLighting,bSupportsDynamicLighting,bPrecisePrevWorldPos,bSupportsPositionOnly) \
-	template class FactoryClass<false>;	\
-	template class FactoryClass<true>;
+	template class FactoryClass<DefaultBoneInfluence>;	\
+	template class FactoryClass<ExtraBoneInfluence>;	\
+	template class FactoryClass<UnlimitedBoneInfluence>;
 
 #if !(UE_BUILD_SHIPPING || UE_BUILD_TEST)
 static TAutoConsoleVariable<int32> CVarVelocityTest(
@@ -283,16 +298,30 @@ bool FGPUBaseSkinVertexFactory::FShaderDataType::UpdateBoneData(FRHICommandListI
 		// PLATFORM_CACHE_LINE_SIZE (128) / 48 = 2.6
 		//  sizeof(FMatrix) == 64
 		// PLATFORM_CACHE_LINE_SIZE (128) / 64 = 2
-		const int32 PreFetchStride = 2; // FPlatformMisc::Prefetch stride
-		for (uint32 BoneIdx = 0; BoneIdx < NumBones; BoneIdx++)
+
+		if (INTEL_ISPC)
+		{
+#if INTEL_ISPC
+			ispc::UpdateBoneData_CopyBones(
+				(ispc::FMatrix3x4*)&ChunkMatrices[0],
+				(ispc::FMatrix*)&ReferenceToLocalMatrices[0],
+				BoneMap.GetData(),
+				NumBones);
+#endif
+		}
+		else
 		{
-			const FBoneIndexType RefToLocalIdx = BoneMap[BoneIdx];
-			FPlatformMisc::Prefetch( ReferenceToLocalMatrices.GetData() + RefToLocalIdx + PreFetchStride );
-			FPlatformMisc::Prefetch( ReferenceToLocalMatrices.GetData() + RefToLocalIdx + PreFetchStride, PLATFORM_CACHE_LINE_SIZE );
+			const int32 PreFetchStride = 2; // FPlatformMisc::Prefetch stride
+			for (uint32 BoneIdx = 0; BoneIdx < NumBones; BoneIdx++)
+			{
+				const FBoneIndexType RefToLocalIdx = BoneMap[BoneIdx];
+				FPlatformMisc::Prefetch(ReferenceToLocalMatrices.GetData() + RefToLocalIdx + PreFetchStride);
+				FPlatformMisc::Prefetch(ReferenceToLocalMatrices.GetData() + RefToLocalIdx + PreFetchStride, PLATFORM_CACHE_LINE_SIZE);
 
-			FMatrix3x4& BoneMat = ChunkMatrices[BoneIdx];
-			const FMatrix& RefToLocal = ReferenceToLocalMatrices[RefToLocalIdx];
-			RefToLocal.To3x4MatrixTranspose( (float*)BoneMat.M );
+				FMatrix3x4& BoneMat = ChunkMatrices[BoneIdx];
+				const FMatrix& RefToLocal = ReferenceToLocalMatrices[RefToLocalIdx];
+				RefToLocal.To3x4MatrixTranspose((float*)BoneMat.M);
+			}
 		}
 	}
 	if (SupportsBonesBufferSRV(InFeatureLevel))
@@ -312,7 +341,21 @@ bool FGPUBaseSkinVertexFactory::FShaderDataType::UpdateBoneData(FRHICommandListI
 
 int32 FGPUBaseSkinVertexFactory::GetMaxGPUSkinBones()
 {
-	return GCVarMaxGPUSkinBones;
+	if (GCVarSupport16BitBoneIndex > 0)
+	{
+		// 16-bit bone index is supported, use GCVarMaxGPUSkinBones
+		return GCVarMaxGPUSkinBones;
+	}
+	else
+	{
+		// 16-bit bone index is not supported, clamp the max bones to 8-bit
+		return FMath::Min(GCVarMaxGPUSkinBones, 256);
+	}
+}
+
+bool FGPUBaseSkinVertexFactory::UseUnlimitedBoneInfluences(uint32 MaxBoneInfluences)
+{
+	return GUnlimitedBoneInfluences && MaxBoneInfluences > GUnlimitedBoneInfluencesThreshold;
 }
 
 /*-----------------------------------------------------------------------------
@@ -321,39 +364,41 @@ TGPUSkinVertexFactory
 
 TGlobalResource<FBoneBufferPool> FGPUBaseSkinVertexFactory::BoneBufferPool;
 
-template <bool bExtraBoneInfluencesT>
-bool TGPUSkinVertexFactory<bExtraBoneInfluencesT>::ShouldCompilePermutation(EShaderPlatform Platform, const class FMaterial* Material, const FShaderType* ShaderType)
+template <GPUSkinBoneInfluenceType BoneInfluenceType>
+bool TGPUSkinVertexFactory<BoneInfluenceType>::ShouldCompilePermutation(EShaderPlatform Platform, const class FMaterial* Material, const FShaderType* ShaderType)
 {
 	bool bLimit2BoneInfluences = (CVarGPUSkinLimit2BoneInfluences.GetValueOnAnyThread() != 0);
 	
 	// Skip trying to use extra bone influences on < SM4 or when project uses 2 bones influence
-	if (bExtraBoneInfluencesT && (GetMaxSupportedFeatureLevel(Platform) < ERHIFeatureLevel::ES3_1 || bLimit2BoneInfluences))
+	if (BoneInfluenceType != DefaultBoneInfluence && (GetMaxSupportedFeatureLevel(Platform) < ERHIFeatureLevel::ES3_1 || bLimit2BoneInfluences))
 	{
 		return false;
 	}
 
-	return (Material->IsUsedWithSkeletalMesh() || Material->IsSpecialEngineMaterial());
+	bool bUnlimitedBoneInfluences = (BoneInfluenceType == UnlimitedBoneInfluence && GUnlimitedBoneInfluences);
+	return ((Material->IsUsedWithSkeletalMesh() && (BoneInfluenceType != UnlimitedBoneInfluence || bUnlimitedBoneInfluences)) || Material->IsSpecialEngineMaterial());
 }
 
 
-template <bool bExtraBoneInfluencesT>
-void TGPUSkinVertexFactory<bExtraBoneInfluencesT>::ModifyCompilationEnvironment(const FVertexFactoryType* Type, EShaderPlatform Platform, const FMaterial* Material, FShaderCompilerEnvironment& OutEnvironment )
+template <GPUSkinBoneInfluenceType BoneInfluenceType>
+void TGPUSkinVertexFactory<BoneInfluenceType>::ModifyCompilationEnvironment(const FVertexFactoryType* Type, EShaderPlatform Platform, const FMaterial* Material, FShaderCompilerEnvironment& OutEnvironment )
 {
 	FVertexFactory::ModifyCompilationEnvironment(Type, Platform, Material, OutEnvironment);
 	const int32 MaxGPUSkinBones = GetFeatureLevelMaxNumberOfBones(GetMaxSupportedFeatureLevel(Platform));
 	OutEnvironment.SetDefine(TEXT("MAX_SHADER_BONES"), MaxGPUSkinBones);
-	const uint32 UseExtraBoneInfluences = bExtraBoneInfluencesT;
+	const uint32 UseExtraBoneInfluences = (BoneInfluenceType == ExtraBoneInfluence);
 	OutEnvironment.SetDefine(TEXT("GPUSKIN_USE_EXTRA_INFLUENCES"), UseExtraBoneInfluences);
 	{
 		bool bLimit2BoneInfluences = (CVarGPUSkinLimit2BoneInfluences.GetValueOnAnyThread() != 0);
 		OutEnvironment.SetDefine(TEXT("GPUSKIN_LIMIT_2BONE_INFLUENCES"), (bLimit2BoneInfluences ? 1 : 0));
 	}
 	OutEnvironment.SetDefine(TEXT("GPUSKIN_USE_BONES_SRV_BUFFER"), SupportsBonesBufferSRV(Platform) ? 1 : 0);
+	OutEnvironment.SetDefine(TEXT("GPUSKIN_UNLIMITED_BONE_INFLUENCE"), BoneInfluenceType == UnlimitedBoneInfluence ? 1 : 0);
 }
 
 
-template<bool bExtraBoneInfluencesT>
-void TGPUSkinVertexFactory<bExtraBoneInfluencesT>::CopyDataTypeForPassthroughFactory(FGPUSkinPassthroughVertexFactory* PassthroughVertexFactory)
+template<GPUSkinBoneInfluenceType BoneInfluenceType>
+void TGPUSkinVertexFactory<BoneInfluenceType>::CopyDataTypeForPassthroughFactory(FGPUSkinPassthroughVertexFactory* PassthroughVertexFactory)
 {
 	FGPUSkinPassthroughVertexFactory::FDataType DestDataType;
 	DestDataType.PositionComponent = Data.PositionComponent;
@@ -376,8 +421,8 @@ void TGPUSkinVertexFactory<bExtraBoneInfluencesT>::CopyDataTypeForPassthroughFac
 * @param InData - type with stream components
 * @param OutElements - vertex decl list to modify
 */
-template <bool bExtraBoneInfluencesT>
-void TGPUSkinVertexFactory<bExtraBoneInfluencesT>::AddVertexElements(FDataType& InData, FVertexDeclarationElementList& OutElements)
+template <GPUSkinBoneInfluenceType BoneInfluenceType>
+void TGPUSkinVertexFactory<BoneInfluenceType>::AddVertexElements(FDataType& InData, FVertexDeclarationElementList& OutElements)
 {
 	// position decls
 	OutElements.Add(AccessStreamComponent(InData.PositionComponent,0));
@@ -426,17 +471,24 @@ void TGPUSkinVertexFactory<bExtraBoneInfluencesT>::AddVertexElements(FDataType&
 		OutElements.Add(AccessStreamComponent(NullColorComponent, 13));
 	}
 
-	// bone indices decls
-	OutElements.Add(AccessStreamComponent(InData.BoneIndices,3));
+	if (BoneInfluenceType == UnlimitedBoneInfluence)
+	{
+		OutElements.Add(AccessStreamComponent(InData.BlendOffsetCount, 3));
+	}
+	else
+	{
+		// bone indices decls
+		OutElements.Add(AccessStreamComponent(InData.BoneIndices,3));
 
-	// bone weights decls
-	OutElements.Add(AccessStreamComponent(InData.BoneWeights,4));
+		// bone weights decls
+		OutElements.Add(AccessStreamComponent(InData.BoneWeights,4));
 
-	if (bExtraBoneInfluencesT)
-	{
-		// Extra bone indices & weights decls
-		OutElements.Add(AccessStreamComponent(InData.ExtraBoneIndices, 14));
-		OutElements.Add(AccessStreamComponent(InData.ExtraBoneWeights, 15));
+		if (GetNumBoneInfluences() > MAX_INFLUENCES_PER_STREAM)
+		{
+			// Extra bone indices & weights decls
+			OutElements.Add(AccessStreamComponent(InData.ExtraBoneIndices, 14));
+			OutElements.Add(AccessStreamComponent(InData.ExtraBoneWeights, 15));
+		}
 	}
 }
 
@@ -444,8 +496,8 @@ void TGPUSkinVertexFactory<bExtraBoneInfluencesT>::AddVertexElements(FDataType&
 * Creates declarations for each of the vertex stream components and
 * initializes the device resource
 */
-template <bool bExtraBoneInfluencesT>
-void TGPUSkinVertexFactory<bExtraBoneInfluencesT>::InitRHI()
+template <GPUSkinBoneInfluenceType BoneInfluenceType>
+void TGPUSkinVertexFactory<BoneInfluenceType>::InitRHI()
 {
 	// list of declaration items
 	FVertexDeclarationElementList Elements;
@@ -455,15 +507,15 @@ void TGPUSkinVertexFactory<bExtraBoneInfluencesT>::InitRHI()
 	InitDeclaration(Elements);
 }
 
-template <bool bExtraBoneInfluencesT>
-void TGPUSkinVertexFactory<bExtraBoneInfluencesT>::InitDynamicRHI()
+template <GPUSkinBoneInfluenceType BoneInfluenceType>
+void TGPUSkinVertexFactory<BoneInfluenceType>::InitDynamicRHI()
 {
 	FVertexFactory::InitDynamicRHI();
 	//ShaderData.UpdateBoneData(GetFeatureLevel());
 }
 
-template <bool bExtraBoneInfluencesT>
-void TGPUSkinVertexFactory<bExtraBoneInfluencesT>::ReleaseDynamicRHI()
+template <GPUSkinBoneInfluenceType BoneInfluenceType>
+void TGPUSkinVertexFactory<BoneInfluenceType>::ReleaseDynamicRHI()
 {
 	FVertexFactory::ReleaseDynamicRHI();
 	ShaderData.ReleaseBoneData();
@@ -473,8 +525,8 @@ void TGPUSkinVertexFactory<bExtraBoneInfluencesT>::ReleaseDynamicRHI()
 TGPUSkinAPEXClothVertexFactory
 -----------------------------------------------------------------------------*/
 
-template <bool bExtraBoneInfluencesT>
-void TGPUSkinAPEXClothVertexFactory<bExtraBoneInfluencesT>::ReleaseDynamicRHI()
+template <GPUSkinBoneInfluenceType BoneInfluenceType>
+void TGPUSkinAPEXClothVertexFactory<BoneInfluenceType>::ReleaseDynamicRHI()
 {
 	Super::ReleaseDynamicRHI();
 	ClothShaderData.ReleaseClothSimulData();
@@ -497,6 +549,8 @@ public:
 		PerBoneMotionBlur.Bind(ParameterMap,TEXT("PerBoneMotionBlur"));
 		BoneMatrices.Bind(ParameterMap,TEXT("BoneMatrices"));
 		PreviousBoneMatrices.Bind(ParameterMap,TEXT("PreviousBoneMatrices"));
+		InputWeightIndexSize.Bind(ParameterMap, TEXT("InputWeightIndexSize"));
+		InputWeightStream.Bind(ParameterMap, TEXT("InputWeightStream"));
 	}
 	/**
 	* Serialize shader params to an archive
@@ -507,6 +561,8 @@ public:
 		Ar << PerBoneMotionBlur;
 		Ar << BoneMatrices;
 		Ar << PreviousBoneMatrices;
+		Ar << InputWeightIndexSize;
+		Ar << InputWeightStream;
 	}
 
 	virtual void GetElementShaderBindings(
@@ -547,6 +603,13 @@ public:
 		}
 
 		ShaderBindings.Add(PerBoneMotionBlur, (uint32)(bLocalPerBoneMotionBlur ? 1 : 0));
+
+		ShaderBindings.Add(InputWeightIndexSize, ShaderData.InputWeightIndexSize);
+		if (InputWeightStream.IsBound())
+		{
+			FRHIShaderResourceView* CurrentData = ShaderData.InputWeightStream;
+			ShaderBindings.Add(InputWeightStream, CurrentData);
+		}
 	}
 
 	virtual uint32 GetSize() const override { return sizeof(*this); }
@@ -555,10 +618,12 @@ private:
 	FShaderParameter PerBoneMotionBlur;
 	FShaderResourceParameter BoneMatrices;
 	FShaderResourceParameter PreviousBoneMatrices;
+	FShaderParameter InputWeightIndexSize;
+	FShaderResourceParameter InputWeightStream;
 };
 
-template <bool bExtraBoneInfluencesT>
-FVertexFactoryShaderParameters* TGPUSkinVertexFactory<bExtraBoneInfluencesT>::ConstructShaderParameters(EShaderFrequency ShaderFrequency)
+template <GPUSkinBoneInfluenceType BoneInfluenceType>
+FVertexFactoryShaderParameters* TGPUSkinVertexFactory<BoneInfluenceType>::ConstructShaderParameters(EShaderFrequency ShaderFrequency)
 {
 	return (ShaderFrequency == SF_Vertex) ? new FGPUSkinVertexFactoryShaderParameters() : NULL;
 }
@@ -735,15 +800,15 @@ TGPUSkinMorphVertexFactory
 * Modify compile environment to enable the morph blend codepath
 * @param OutEnvironment - shader compile environment to modify
 */
-template <bool bExtraBoneInfluencesT>
-void TGPUSkinMorphVertexFactory<bExtraBoneInfluencesT>::ModifyCompilationEnvironment( const FVertexFactoryType* Type, EShaderPlatform Platform, const FMaterial* Material, FShaderCompilerEnvironment& OutEnvironment )
+template <GPUSkinBoneInfluenceType BoneInfluenceType>
+void TGPUSkinMorphVertexFactory<BoneInfluenceType>::ModifyCompilationEnvironment( const FVertexFactoryType* Type, EShaderPlatform Platform, const FMaterial* Material, FShaderCompilerEnvironment& OutEnvironment )
 {
 	Super::ModifyCompilationEnvironment(Type, Platform, Material, OutEnvironment);
 	OutEnvironment.SetDefine(TEXT("GPUSKIN_MORPH_BLEND"),TEXT("1"));
 }
 
-template <bool bExtraBoneInfluencesT>
-bool TGPUSkinMorphVertexFactory<bExtraBoneInfluencesT>::ShouldCompilePermutation(EShaderPlatform Platform, const class FMaterial* Material, const class FShaderType* ShaderType)
+template <GPUSkinBoneInfluenceType BoneInfluenceType>
+bool TGPUSkinMorphVertexFactory<BoneInfluenceType>::ShouldCompilePermutation(EShaderPlatform Platform, const class FMaterial* Material, const class FShaderType* ShaderType)
 {
 	return (Material->IsUsedWithMorphTargets() || Material->IsSpecialEngineMaterial()) 
 		&& Super::ShouldCompilePermutation(Platform, Material, ShaderType);
@@ -754,11 +819,11 @@ bool TGPUSkinMorphVertexFactory<bExtraBoneInfluencesT>::ShouldCompilePermutation
 * @param InData - type with stream components
 * @param OutElements - vertex decl list to modify
 */
-template <bool bExtraBoneInfluencesT>
-void TGPUSkinMorphVertexFactory<bExtraBoneInfluencesT>::AddVertexElements(FDataType& InData, FVertexDeclarationElementList& OutElements)
+template <GPUSkinBoneInfluenceType BoneInfluenceType>
+void TGPUSkinMorphVertexFactory<BoneInfluenceType>::AddVertexElements(FDataType& InData, FVertexDeclarationElementList& OutElements)
 {
 	// add the base gpu skin elements
-	TGPUSkinVertexFactory<bExtraBoneInfluencesT>::AddVertexElements(InData,OutElements);
+	TGPUSkinVertexFactory<BoneInfluenceType>::AddVertexElements(InData,OutElements);
 	// add the morph delta elements
 	OutElements.Add(FVertexFactory::AccessStreamComponent(InData.DeltaPositionComponent,9));
 	OutElements.Add(FVertexFactory::AccessStreamComponent(InData.DeltaTangentZComponent,10));
@@ -768,8 +833,8 @@ void TGPUSkinMorphVertexFactory<bExtraBoneInfluencesT>::AddVertexElements(FDataT
 * Creates declarations for each of the vertex stream components and
 * initializes the device resource
 */
-template <bool bExtraBoneInfluencesT>
-void TGPUSkinMorphVertexFactory<bExtraBoneInfluencesT>::InitRHI()
+template <GPUSkinBoneInfluenceType BoneInfluenceType>
+void TGPUSkinMorphVertexFactory<BoneInfluenceType>::InitRHI()
 {
 	// list of declaration items
 	FVertexDeclarationElementList Elements;	
@@ -779,8 +844,8 @@ void TGPUSkinMorphVertexFactory<bExtraBoneInfluencesT>::InitRHI()
 	FVertexFactory::InitDeclaration(Elements);
 }
 
-template <bool bExtraBoneInfluencesT>
-FVertexFactoryShaderParameters* TGPUSkinMorphVertexFactory<bExtraBoneInfluencesT>::ConstructShaderParameters(EShaderFrequency ShaderFrequency)
+template <GPUSkinBoneInfluenceType BoneInfluenceType>
+FVertexFactoryShaderParameters* TGPUSkinMorphVertexFactory<BoneInfluenceType>::ConstructShaderParameters(EShaderFrequency ShaderFrequency)
 {
 	return ShaderFrequency == SF_Vertex ? new FGPUSkinVertexFactoryShaderParameters() : NULL;
 }
@@ -844,9 +909,10 @@ public:
 		const auto* GPUSkinVertexFactory = (const FGPUBaseSkinVertexFactory*)VertexFactory;
 		// A little hacky; problem is we can't upcast from FGPUBaseSkinVertexFactory to FGPUBaseSkinAPEXClothVertexFactory as they are unrelated; a nice solution would be
 		// to use virtual inheritance, but that requires RTTI and complicates things further...
-		const FGPUBaseSkinAPEXClothVertexFactory::ClothShaderType& ClothShaderData = GPUSkinVertexFactory->UsesExtraBoneInfluences()
-			? ((const TGPUSkinAPEXClothVertexFactory<true>*)GPUSkinVertexFactory)->GetClothShaderData()
-			: ((const TGPUSkinAPEXClothVertexFactory<false>*)GPUSkinVertexFactory)->GetClothShaderData();
+		const FGPUBaseSkinAPEXClothVertexFactory::ClothShaderType& ClothShaderData = 
+			 GPUSkinVertexFactory->GetBoneInfluenceType() == DefaultBoneInfluence ? ((const TGPUSkinAPEXClothVertexFactory<DefaultBoneInfluence>*)GPUSkinVertexFactory)->GetClothShaderData() :
+			(GPUSkinVertexFactory->GetBoneInfluenceType() == ExtraBoneInfluence   ? ((const TGPUSkinAPEXClothVertexFactory<ExtraBoneInfluence>*)GPUSkinVertexFactory)->GetClothShaderData() :
+																					((const TGPUSkinAPEXClothVertexFactory<UnlimitedBoneInfluence>*)GPUSkinVertexFactory)->GetClothShaderData());
 
 		ShaderBindings.Add(Shader->GetUniformBufferParameter<FAPEXClothUniformShaderParameters>(),ClothShaderData.GetClothUniformBuffer());
 
@@ -859,14 +925,14 @@ public:
 		ShaderBindings.Add(ClothBlendWeightParameter,ClothShaderData.ClothBlendWeight);
 
 		ShaderBindings.Add(GPUSkinApexClothParameter,
-			GPUSkinVertexFactory->UsesExtraBoneInfluences()
-			? ((const TGPUSkinAPEXClothVertexFactory<true>*)GPUSkinVertexFactory)->GetClothBuffer()
-			: ((const TGPUSkinAPEXClothVertexFactory<false>*)GPUSkinVertexFactory)->GetClothBuffer());
+			 GPUSkinVertexFactory->GetBoneInfluenceType() == DefaultBoneInfluence ? ((const TGPUSkinAPEXClothVertexFactory<DefaultBoneInfluence>*)GPUSkinVertexFactory)->GetClothBuffer() :
+			(GPUSkinVertexFactory->GetBoneInfluenceType() == ExtraBoneInfluence   ? ((const TGPUSkinAPEXClothVertexFactory<ExtraBoneInfluence>*)GPUSkinVertexFactory)->GetClothBuffer() :
+																					((const TGPUSkinAPEXClothVertexFactory<UnlimitedBoneInfluence>*)GPUSkinVertexFactory)->GetClothBuffer()));
 
 		int32 ClothIndexOffset =
-			GPUSkinVertexFactory->UsesExtraBoneInfluences()
-			? ((const TGPUSkinAPEXClothVertexFactory<true>*)GPUSkinVertexFactory)->GetClothIndexOffset(BatchElement.MinVertexIndex)
-			: ((const TGPUSkinAPEXClothVertexFactory<false>*)GPUSkinVertexFactory)->GetClothIndexOffset(BatchElement.MinVertexIndex);
+			 GPUSkinVertexFactory->GetBoneInfluenceType() == DefaultBoneInfluence ? ((const TGPUSkinAPEXClothVertexFactory<DefaultBoneInfluence>*)GPUSkinVertexFactory)->GetClothIndexOffset(BatchElement.MinVertexIndex) :
+			(GPUSkinVertexFactory->GetBoneInfluenceType() == ExtraBoneInfluence   ? ((const TGPUSkinAPEXClothVertexFactory<ExtraBoneInfluence>*)GPUSkinVertexFactory)->GetClothIndexOffset(BatchElement.MinVertexIndex) :
+																					((const TGPUSkinAPEXClothVertexFactory<UnlimitedBoneInfluence>*)GPUSkinVertexFactory)->GetClothIndexOffset(BatchElement.MinVertexIndex));
 
 		FIntPoint GPUSkinApexClothStartIndexOffset(BatchElement.MinVertexIndex, ClothIndexOffset);
 		ShaderBindings.Add(GPUSkinApexClothStartIndexOffsetParameter, GPUSkinApexClothStartIndexOffset);
@@ -978,15 +1044,15 @@ TGlobalResource<FClothBufferPool> FGPUBaseSkinAPEXClothVertexFactory::ClothSimul
 * Modify compile environment to enable the apex clothing path
 * @param OutEnvironment - shader compile environment to modify
 */
-template <bool bExtraBoneInfluencesT>
-void TGPUSkinAPEXClothVertexFactory<bExtraBoneInfluencesT>::ModifyCompilationEnvironment( const FVertexFactoryType* Type, EShaderPlatform Platform, const FMaterial* Material, FShaderCompilerEnvironment& OutEnvironment )
+template <GPUSkinBoneInfluenceType BoneInfluenceType>
+void TGPUSkinAPEXClothVertexFactory<BoneInfluenceType>::ModifyCompilationEnvironment( const FVertexFactoryType* Type, EShaderPlatform Platform, const FMaterial* Material, FShaderCompilerEnvironment& OutEnvironment )
 {
 	Super::ModifyCompilationEnvironment(Type, Platform, Material, OutEnvironment);
 	OutEnvironment.SetDefine(TEXT("GPUSKIN_APEX_CLOTH"),TEXT("1"));
 }
 
-template <bool bExtraBoneInfluencesT>
-bool TGPUSkinAPEXClothVertexFactory<bExtraBoneInfluencesT>::ShouldCompilePermutation(EShaderPlatform Platform, const class FMaterial* Material, const class FShaderType* ShaderType)
+template <GPUSkinBoneInfluenceType BoneInfluenceType>
+bool TGPUSkinAPEXClothVertexFactory<BoneInfluenceType>::ShouldCompilePermutation(EShaderPlatform Platform, const class FMaterial* Material, const class FShaderType* ShaderType)
 {
 	return GetMaxSupportedFeatureLevel(Platform) >= ERHIFeatureLevel::SM5
 		&& (Material->IsUsedWithAPEXCloth() || Material->IsSpecialEngineMaterial()) 
@@ -998,11 +1064,11 @@ bool TGPUSkinAPEXClothVertexFactory<bExtraBoneInfluencesT>::ShouldCompilePermuta
 * @param InData - type with stream components
 * @param OutElements - vertex decl list to modify
 */
-template <bool bExtraBoneInfluencesT>
-void TGPUSkinAPEXClothVertexFactory<bExtraBoneInfluencesT>::AddVertexElements(FDataType& InData, FVertexDeclarationElementList& OutElements)
+template <GPUSkinBoneInfluenceType BoneInfluenceType>
+void TGPUSkinAPEXClothVertexFactory<BoneInfluenceType>::AddVertexElements(FDataType& InData, FVertexDeclarationElementList& OutElements)
 {
 	// add the base gpu skin elements
-	TGPUSkinVertexFactory<bExtraBoneInfluencesT>::AddVertexElements(InData,OutElements);
+	TGPUSkinVertexFactory<BoneInfluenceType>::AddVertexElements(InData,OutElements);
 	// add the morph delta elements
 //	return;
 	if(InData.CoordNormalComponent.VertexBuffer)
@@ -1018,8 +1084,8 @@ void TGPUSkinAPEXClothVertexFactory<bExtraBoneInfluencesT>::AddVertexElements(FD
 * Creates declarations for each of the vertex stream components and
 * initializes the device resource
 */
-template <bool bExtraBoneInfluencesT>
-void TGPUSkinAPEXClothVertexFactory<bExtraBoneInfluencesT>::InitRHI()
+template <GPUSkinBoneInfluenceType BoneInfluenceType>
+void TGPUSkinAPEXClothVertexFactory<BoneInfluenceType>::InitRHI()
 {
 	// list of declaration items
 	FVertexDeclarationElementList Elements;	
@@ -1029,8 +1095,8 @@ void TGPUSkinAPEXClothVertexFactory<bExtraBoneInfluencesT>::InitRHI()
 	FVertexFactory::InitDeclaration(Elements);
 }
 
-template <bool bExtraBoneInfluencesT>
-FVertexFactoryShaderParameters* TGPUSkinAPEXClothVertexFactory<bExtraBoneInfluencesT>::ConstructShaderParameters(EShaderFrequency ShaderFrequency)
+template <GPUSkinBoneInfluenceType BoneInfluenceType>
+FVertexFactoryShaderParameters* TGPUSkinAPEXClothVertexFactory<BoneInfluenceType>::ConstructShaderParameters(EShaderFrequency ShaderFrequency)
 {
 	return ShaderFrequency == SF_Vertex ? new TGPUSkinAPEXClothVertexFactoryShaderParameters() : NULL;
 }
diff --git a/Engine/Source/Runtime/Engine/Private/Particles/ParticleModules_Location.cpp b/Engine/Source/Runtime/Engine/Private/Particles/ParticleModules_Location.cpp
--- a/Engine/Source/Runtime/Engine/Private/Particles/ParticleModules_Location.cpp
+++ b/Engine/Source/Runtime/Engine/Private/Particles/ParticleModules_Location.cpp
@@ -2931,14 +2931,11 @@ bool UParticleModuleLocationSkelVertSurface::VertInfluencedByActiveBone(FParticl
 			}
 		}
 
-		return LODData.GetSkinWeightVertexBuffer()->HasExtraBoneInfluences()
-			? VertInfluencedByActiveBoneTyped<true>(LODData, 0, Section, VertIndex, InSkelMeshComponent, InstancePayload, OutBoneIndex)
-			: VertInfluencedByActiveBoneTyped<false>(LODData, 0, Section, VertIndex, InSkelMeshComponent, InstancePayload, OutBoneIndex);
+		return VertInfluencedByActiveBoneTyped(LODData, 0, Section, VertIndex, InSkelMeshComponent, InstancePayload, OutBoneIndex);
 	}
 	return false;
 }
 
-template<bool bExtraBoneInfluencesT>
 bool UParticleModuleLocationSkelVertSurface::VertInfluencedByActiveBoneTyped(
 	FSkeletalMeshLODRenderData& LODData,
 	int32 LODIndex,
@@ -2953,8 +2950,6 @@ bool UParticleModuleLocationSkelVertSurface::VertInfluencedByActiveBoneTyped(
 	FSkinWeightVertexBuffer* WeightBuffer = InSkelMeshComponent->GetSkinWeightBuffer(LODIndex);
 	if (WeightBuffer)
 	{
-		const TSkinWeightInfo<bExtraBoneInfluencesT>* SrcSkinWeights = WeightBuffer->GetSkinWeightPtr<bExtraBoneInfluencesT>(Section.GetVertexBufferIndex() + VertIndex);
-
 #if !PLATFORM_LITTLE_ENDIAN
 		// uint8[] elements in LOD.VertexBufferGPUSkin have been swapped for VET_UBYTE4 vertex stream use
 		for (int32 InfluenceIndex = MAX_INFLUENCES - 1; InfluenceIndex >= MAX_INFLUENCES - Section.MaxBoneInfluences; InfluenceIndex--)
@@ -2962,7 +2957,7 @@ bool UParticleModuleLocationSkelVertSurface::VertInfluencedByActiveBoneTyped(
 		for (int32 InfluenceIndex = 0; InfluenceIndex < Section.MaxBoneInfluences; InfluenceIndex++)
 #endif
 		{
-			int32 BoneIndex = Section.BoneMap[SrcSkinWeights->InfluenceBones[InfluenceIndex]];
+			int32 BoneIndex = Section.BoneMap[WeightBuffer->GetBoneIndex(Section.GetVertexBufferIndex() + VertIndex, InfluenceIndex)];
 			if (InSkelMeshComponent->MasterPoseComponent.IsValid())
 			{
 				check(MasterBoneMap.Num() == InSkelMeshComponent->SkeletalMesh->RefSkeleton.GetNum());
diff --git a/Engine/Source/Runtime/Engine/Private/Rendering/SkinWeightVertexBuffer.cpp b/Engine/Source/Runtime/Engine/Private/Rendering/SkinWeightVertexBuffer.cpp
--- a/Engine/Source/Runtime/Engine/Private/Rendering/SkinWeightVertexBuffer.cpp
+++ b/Engine/Source/Runtime/Engine/Private/Rendering/SkinWeightVertexBuffer.cpp
@@ -5,98 +5,273 @@
 #include "Rendering/SkeletalMeshLODModel.h"
 #include "RenderUtils.h"
 #include "SkeletalMeshTypes.h"
+#include "UObject/AnimObjectVersion.h"
 
-///
+/*-----------------------------------------------------------------------------
+FSkinWeightLookupVertexBuffer
+-----------------------------------------------------------------------------*/
 
-FSkinWeightVertexBuffer::FSkinWeightVertexBuffer()
-:	bNeedsCPUAccess(false)
-,	bExtraBoneInfluences(false)
-,	WeightData(nullptr)
-,	Data(nullptr)
-,	Stride(0)
-,	NumVertices(0)
+FSkinWeightLookupVertexBuffer::FSkinWeightLookupVertexBuffer()
+	: bNeedsCPUAccess(false)
+	, LookupData(nullptr)
+	, Data(nullptr)
+	, NumVertices(0)
 {
 }
 
-FSkinWeightVertexBuffer::FSkinWeightVertexBuffer( const FSkinWeightVertexBuffer &Other )
+FSkinWeightLookupVertexBuffer::FSkinWeightLookupVertexBuffer(const FSkinWeightLookupVertexBuffer &Other)
 	: bNeedsCPUAccess(Other.bNeedsCPUAccess)
-	, bExtraBoneInfluences(Other.bExtraBoneInfluences)
-	, WeightData(nullptr)
+	, LookupData(nullptr)
 	, Data(nullptr)
-	, Stride(0)
 	, NumVertices(0)
 {
-	
+
 }
 
-FSkinWeightVertexBuffer::~FSkinWeightVertexBuffer()
+FSkinWeightLookupVertexBuffer::~FSkinWeightLookupVertexBuffer()
 {
 	CleanUp();
 }
 
-FSkinWeightVertexBuffer& FSkinWeightVertexBuffer::operator=(const FSkinWeightVertexBuffer& Other)
+FSkinWeightLookupVertexBuffer& FSkinWeightLookupVertexBuffer::operator=(const FSkinWeightLookupVertexBuffer& Other)
 {
 	CleanUp();
 	bNeedsCPUAccess = Other.bNeedsCPUAccess;
-	bExtraBoneInfluences = Other.bExtraBoneInfluences;
 	return *this;
 }
 
-void FSkinWeightVertexBuffer::CleanUp()
+void FSkinWeightLookupVertexBuffer::CleanUp()
 {
-	if (WeightData)
+	if (LookupData)
 	{
-		delete WeightData;
-		WeightData = NULL;
+		delete LookupData;
+		LookupData = NULL;
 	}
 }
 
-bool FSkinWeightVertexBuffer::IsWeightDataValid() const
+bool FSkinWeightLookupVertexBuffer::IsLookupDataValid() const
 {
-	return WeightData != NULL;
+	return LookupData != NULL;
 }
 
-#if WITH_EDITOR
-
-void FSkinWeightVertexBuffer::Init(const TArray<FSoftSkinVertex>& InVertices)
+void FSkinWeightLookupVertexBuffer::Init(uint32 InNumVertices)
 {
-	// Make sure if this is console, use compressed otherwise, use not compressed
 	AllocateData();
+	ResizeBuffer(InNumVertices);
+}
+
+void FSkinWeightLookupVertexBuffer::ResizeBuffer(uint32 InNumVertices)
+{
+	NumVertices = InNumVertices;
+	LookupData->ResizeBuffer(NumVertices);
+	Data = NumVertices ? LookupData->GetDataPointer() : nullptr;
+}
 
-	WeightData->ResizeBuffer(InVertices.Num());
+FArchive& operator<<(FArchive& Ar, FSkinWeightLookupVertexBuffer& VertexBuffer)
+{
+	FStripDataFlags StripFlags(Ar);
 
-	if (InVertices.Num() > 0)
+	Ar.UsingCustomVersion(FAnimObjectVersion::GUID);
+	VertexBuffer.SerializeMetaData(Ar);
+
+	if (Ar.IsLoading() || VertexBuffer.LookupData == NULL)
 	{
-		Data = WeightData->GetDataPointer();
-		Stride = WeightData->GetStride();
-		NumVertices = InVertices.Num();
+		// If we're loading, or we have no valid buffer, allocate container.
+		VertexBuffer.AllocateData();
 	}
 
-	if (bExtraBoneInfluences)
+	// if Ar is counting, it still should serialize. Need to count VertexData
+	if (!StripFlags.IsDataStrippedForServer() || Ar.IsCountingMemory())
 	{
-		for (int32 VertIdx = 0; VertIdx < InVertices.Num(); VertIdx++)
+		if (VertexBuffer.LookupData != NULL)
 		{
-			const FSoftSkinVertex& SrcVertex = InVertices[VertIdx];
-			SetWeightsForVertex<true>(VertIdx, SrcVertex);
+			if (Ar.CustomVer(FAnimObjectVersion::GUID) >= FAnimObjectVersion::UnlimitedBoneInfluences)
+			{
+				VertexBuffer.LookupData->Serialize(Ar);
+			}
+
+			if (!Ar.IsCountingMemory())
+			{
+				// update cached buffer info
+				VertexBuffer.Data = (VertexBuffer.NumVertices > 0 && VertexBuffer.LookupData->GetResourceArray()->GetResourceDataSize()) ? VertexBuffer.LookupData->GetDataPointer() : nullptr;
+			}
 		}
 	}
-	else
+
+	return Ar;
+}
+
+void FSkinWeightLookupVertexBuffer::SerializeMetaData(FArchive& Ar)
+{
+	Ar.UsingCustomVersion(FAnimObjectVersion::GUID);
+	if (Ar.CustomVer(FAnimObjectVersion::GUID) >= FAnimObjectVersion::UnlimitedBoneInfluences)
 	{
-		for (int32 VertIdx = 0; VertIdx < InVertices.Num(); VertIdx++)
+		Ar << NumVertices;
+	}
+}
+
+void FSkinWeightLookupVertexBuffer::CopyMetaData(const FSkinWeightLookupVertexBuffer& Other)
+{
+	NumVertices = Other.NumVertices;
+}
+
+void FSkinWeightLookupVertexBuffer::AllocateData()
+{
+	// Clear any old WeightData before allocating.
+	CleanUp();
+
+	LookupData = new TStaticMeshVertexData<uint32>(bNeedsCPUAccess);
+}
+
+void FSkinWeightLookupVertexBuffer::InitRHI()
+{
+	// BUF_ShaderResource is needed for support of the SkinCache (we could make is dependent on GEnableGPUSkinCacheShaders or are there other users?)
+	VertexBufferRHI = CreateRHIBuffer_RenderThread();
+
+	bool bSRV = VertexBufferRHI && GSupportsResourceView && GPixelFormats[PixelFormat].Supported;
+	// When bAllowCPUAccess is true, the meshes is likely going to be used for Niagara to spawn particles on mesh surface.
+	// And it can be the case for CPU *and* GPU access: no differenciation today. That is why we create a SRV in this case.
+	// This also avoid setting lots of states on all the members of all the different buffers used by meshes. Follow up: https://jira.it.epicgames.net/browse/UE-69376.
+	bSRV |= GetNeedsCPUAccess();
+
+	if (bSRV)
+	{
+		SRVValue = RHICreateShaderResourceView(LookupData ? VertexBufferRHI : nullptr, PixelFormatStride, PixelFormat);
+	}
+}
+
+void FSkinWeightLookupVertexBuffer::ReleaseRHI()
+{
+	SRVValue.SafeRelease();
+
+	FVertexBuffer::ReleaseRHI();
+}
+
+FVertexBufferRHIRef FSkinWeightLookupVertexBuffer::CreateRHIBuffer_RenderThread()
+{
+	return CreateRHIBuffer_Internal<true>();
+}
+
+FVertexBufferRHIRef FSkinWeightLookupVertexBuffer::CreateRHIBuffer_Async()
+{
+	return CreateRHIBuffer_Internal<false>();
+}
+
+void FSkinWeightLookupVertexBuffer::GetWeightOffsetAndInfluenceCount(uint32 VertexIndex, uint32& OutWeightOffset, uint32& OutInfluenceCount) const
+{
+	uint32 Offset = VertexIndex * 4;
+	uint32 DataUInt32 = *((uint32*)(&Data[Offset]));
+	OutWeightOffset = DataUInt32 >> 8;
+	OutInfluenceCount = DataUInt32 & 0xff;
+}
+
+void FSkinWeightLookupVertexBuffer::SetWeightOffsetAndInfluenceCount(uint32 VertexIndex, uint32 WeightOffset, uint32 InfluenceCount)
+{
+	uint32 Offset = VertexIndex * 4;
+	uint32* DataUInt32 = (uint32*)(&Data[Offset]);
+	*DataUInt32 = (WeightOffset << 8) | InfluenceCount;
+}
+
+template <bool bRenderThread>
+FVertexBufferRHIRef FSkinWeightLookupVertexBuffer::CreateRHIBuffer_Internal()
+{
+	if (NumVertices)
+	{
+		// Create the vertex buffer.
+		FResourceArrayInterface* ResourceArray = LookupData ? LookupData->GetResourceArray() : nullptr;
+		const uint32 SizeInBytes = ResourceArray ? ResourceArray->GetResourceDataSize() : 0;
+		const uint32 BuffFlags = BUF_Static | BUF_ShaderResource;
+		FRHIResourceCreateInfo CreateInfo(ResourceArray);
+		CreateInfo.bWithoutNativeResource = !LookupData;
+
+		if (bRenderThread)
 		{
-			const FSoftSkinVertex& SrcVertex = InVertices[VertIdx];
-			SetWeightsForVertex<false>(VertIdx, SrcVertex);
+			return RHICreateVertexBuffer(SizeInBytes, BuffFlags, CreateInfo);
+		}
+		else
+		{
+			return RHIAsyncCreateVertexBuffer(SizeInBytes, BuffFlags, CreateInfo);
 		}
 	}
+	return nullptr;
 }
 
-#endif // WITH_EDITOR
+/*-----------------------------------------------------------------------------
+FSkinWeightDataVertexBuffer
+-----------------------------------------------------------------------------*/
 
+FSkinWeightDataVertexBuffer::FSkinWeightDataVertexBuffer()
+:	bNeedsCPUAccess(false)
+,	bVariableBonesPerVertex(false)
+,	MaxBoneInfluences(MAX_INFLUENCES_PER_STREAM)
+,	bUse16BitBoneIndex(false)
+,	WeightData(nullptr)
+,	Data(nullptr)
+,	NumVertices(0)
+,	NumBones(0)
+{
+}
 
-FArchive& operator<<(FArchive& Ar, FSkinWeightVertexBuffer& VertexBuffer)
+FSkinWeightDataVertexBuffer::FSkinWeightDataVertexBuffer( const FSkinWeightDataVertexBuffer &Other )
+	: bNeedsCPUAccess(Other.bNeedsCPUAccess)
+	, bVariableBonesPerVertex(Other.bVariableBonesPerVertex)
+	, MaxBoneInfluences(Other.MaxBoneInfluences)
+	, bUse16BitBoneIndex(Other.bUse16BitBoneIndex)
+	, WeightData(nullptr)
+	, Data(nullptr)
+	, NumVertices(0)
+	, NumBones(0)
+{
+	
+}
+
+FSkinWeightDataVertexBuffer::~FSkinWeightDataVertexBuffer()
+{
+	CleanUp();
+}
+
+FSkinWeightDataVertexBuffer& FSkinWeightDataVertexBuffer::operator=(const FSkinWeightDataVertexBuffer& Other)
+{
+	CleanUp();
+	bNeedsCPUAccess = Other.bNeedsCPUAccess;
+	bVariableBonesPerVertex = Other.bVariableBonesPerVertex;
+	MaxBoneInfluences = Other.MaxBoneInfluences;
+	bUse16BitBoneIndex = Other.bUse16BitBoneIndex;
+	return *this;
+}
+
+void FSkinWeightDataVertexBuffer::CleanUp()
+{
+	if (WeightData)
+	{
+		delete WeightData;
+		WeightData = NULL;
+	}
+}
+
+void FSkinWeightDataVertexBuffer::Init(uint32 InNumBones, uint32 InNumVertices)
+{
+	AllocateData();
+	ResizeBuffer(InNumBones, InNumVertices);
+}
+
+void FSkinWeightDataVertexBuffer::ResizeBuffer(uint32 InNumBones, uint32 InNumVertices)
+{
+	NumBones = InNumBones;
+	NumVertices = InNumVertices;
+	WeightData->ResizeBuffer(GetVertexDataSize());
+
+	if (NumBones > 0)
+	{
+		Data = WeightData->GetDataPointer();
+	}
+}
+
+FArchive& operator<<(FArchive& Ar, FSkinWeightDataVertexBuffer& VertexBuffer)
 {
 	FStripDataFlags StripFlags(Ar);
 
+	Ar.UsingCustomVersion(FAnimObjectVersion::GUID);
 	VertexBuffer.SerializeMetaData(Ar);
 
 	if (Ar.IsLoading() || VertexBuffer.WeightData == NULL)
@@ -110,13 +285,55 @@ FArchive& operator<<(FArchive& Ar, FSkinWeightVertexBuffer& VertexBuffer)
 	{
 		if (VertexBuffer.WeightData != NULL)
 		{
-			VertexBuffer.WeightData->Serialize(Ar);
+			if (Ar.IsLoading() && Ar.CustomVer(FAnimObjectVersion::GUID) < FAnimObjectVersion::UnlimitedBoneInfluences)
+			{
+				FStaticMeshVertexDataInterface* LegacyWeightData = nullptr;
+				bool bExtraBoneInfluences = VertexBuffer.MaxBoneInfluences > MAX_INFLUENCES_PER_STREAM;
+				if (bExtraBoneInfluences)
+				{
+					if (VertexBuffer.bUse16BitBoneIndex)
+					{
+						LegacyWeightData = new FSkinWeightVertexData< TLegacySkinWeightInfo<true, FBoneIndex16> >(VertexBuffer.bNeedsCPUAccess);
+					}
+					else
+					{
+						LegacyWeightData = new FSkinWeightVertexData< TLegacySkinWeightInfo<true, FBoneIndex8> >(VertexBuffer.bNeedsCPUAccess);
+					}
+				}
+				else
+				{
+					if (VertexBuffer.bUse16BitBoneIndex)
+					{
+						LegacyWeightData = new FSkinWeightVertexData< TLegacySkinWeightInfo<false, FBoneIndex16> >(VertexBuffer.bNeedsCPUAccess);
+					}
+					else
+					{
+						LegacyWeightData = new FSkinWeightVertexData< TLegacySkinWeightInfo<false, FBoneIndex8> >(VertexBuffer.bNeedsCPUAccess);
+					}
+				}
+				bool bForcePerElementSerialization = Ar.IsLoading() && Ar.CustomVer(FAnimObjectVersion::GUID) < FAnimObjectVersion::IncreaseBoneIndexLimitPerChunk;
+				LegacyWeightData->Serialize(Ar, bForcePerElementSerialization);
+	
+				uint32 BufferSize = LegacyWeightData->GetResourceArray()->GetResourceDataSize();
+				check(VertexBuffer.GetVertexDataSize() == BufferSize);
+				VertexBuffer.WeightData->ResizeBuffer(BufferSize);
+				FMemory::Memcpy(VertexBuffer.WeightData->GetDataPointer(), LegacyWeightData->GetDataPointer(), BufferSize);
+				VertexBuffer.NumVertices = LegacyWeightData->Num();
+
+				if (LegacyWeightData)
+				{
+					delete LegacyWeightData;
+				}
+			}
+			else
+			{
+				VertexBuffer.WeightData->Serialize(Ar);
+			}
 
 			if (!Ar.IsCountingMemory())
 			{
 				// update cached buffer info
-				VertexBuffer.Data = (VertexBuffer.NumVertices > 0 && VertexBuffer.WeightData->GetResourceArray()->GetResourceDataSize()) ? VertexBuffer.WeightData->GetDataPointer() : nullptr;
-				VertexBuffer.Stride = VertexBuffer.WeightData->GetStride();
+				VertexBuffer.Data = (VertexBuffer.NumBones > 0 && VertexBuffer.WeightData->GetResourceArray()->GetResourceDataSize()) ? VertexBuffer.WeightData->GetDataPointer() : nullptr;
 			}
 		}
 	}
@@ -124,31 +341,51 @@ FArchive& operator<<(FArchive& Ar, FSkinWeightVertexBuffer& VertexBuffer)
 	return Ar;
 }
 
-void FSkinWeightVertexBuffer::SerializeMetaData(FArchive& Ar)
+void FSkinWeightDataVertexBuffer::SerializeMetaData(FArchive& Ar)
 {
 	Ar.UsingCustomVersion(FSkeletalMeshCustomVersion::GUID);
-	if (Ar.CustomVer(FSkeletalMeshCustomVersion::GUID) < FSkeletalMeshCustomVersion::SplitModelAndRenderData)
+	Ar.UsingCustomVersion(FAnimObjectVersion::GUID);
+
+	if (Ar.IsLoading() && Ar.CustomVer(FAnimObjectVersion::GUID) < FAnimObjectVersion::UnlimitedBoneInfluences)
 	{
-		check(Ar.IsLoading());
-		Ar << bExtraBoneInfluences << NumVertices;
+		bool bExtraBoneInfluences = false;
+		if (Ar.CustomVer(FSkeletalMeshCustomVersion::GUID) < FSkeletalMeshCustomVersion::SplitModelAndRenderData)
+		{
+			Ar << bExtraBoneInfluences << NumVertices;
+		}
+		else
+		{		
+			uint32 Stride = 0;
+			Ar << bExtraBoneInfluences << Stride << NumVertices;
+		}
+		MaxBoneInfluences = bExtraBoneInfluences ? EXTRA_BONE_INFLUENCES : MAX_INFLUENCES_PER_STREAM;
+		NumBones = MaxBoneInfluences * NumVertices;
+		bVariableBonesPerVertex = false;
 	}
 	else
 	{
-		Ar << bExtraBoneInfluences << Stride << NumVertices;
+		Ar << bVariableBonesPerVertex << MaxBoneInfluences << NumBones << NumVertices;
+	}
+	
+	// bUse16BitBoneIndex doesn't exist before version IncreaseBoneIndexLimitPerChunk
+	if (Ar.CustomVer(FAnimObjectVersion::GUID) >= FAnimObjectVersion::IncreaseBoneIndexLimitPerChunk)
+	{
+		Ar << bUse16BitBoneIndex;
 	}
 }
 
-void FSkinWeightVertexBuffer::CopyMetaData(const FSkinWeightVertexBuffer& Other)
+void FSkinWeightDataVertexBuffer::CopyMetaData(const FSkinWeightDataVertexBuffer& Other)
 {
-	bExtraBoneInfluences = Other.bExtraBoneInfluences;
-	Stride = Other.Stride;
-	NumVertices = Other.NumVertices;
+	bVariableBonesPerVertex = Other.bVariableBonesPerVertex;
+	MaxBoneInfluences = Other.MaxBoneInfluences;
+	bUse16BitBoneIndex = Other.bUse16BitBoneIndex;
+	NumBones = Other.NumBones;
 }
 
 template <bool bRenderThread>
-FVertexBufferRHIRef FSkinWeightVertexBuffer::CreateRHIBuffer_Internal()
+FVertexBufferRHIRef FSkinWeightDataVertexBuffer::CreateRHIBuffer_Internal()
 {
-	if (NumVertices)
+	if (NumBones)
 	{
 		// Create the vertex buffer.
 		FResourceArrayInterface* ResourceArray = WeightData ? WeightData->GetResourceArray() : nullptr;
@@ -170,22 +407,22 @@ FVertexBufferRHIRef FSkinWeightVertexBuffer::CreateRHIBuffer_Internal()
 	return nullptr;
 }
 
-FVertexBufferRHIRef FSkinWeightVertexBuffer::CreateRHIBuffer_RenderThread()
+FVertexBufferRHIRef FSkinWeightDataVertexBuffer::CreateRHIBuffer_RenderThread()
 {
 	return CreateRHIBuffer_Internal<true>();
 }
 
-FVertexBufferRHIRef FSkinWeightVertexBuffer::CreateRHIBuffer_Async()
+FVertexBufferRHIRef FSkinWeightDataVertexBuffer::CreateRHIBuffer_Async()
 {
 	return CreateRHIBuffer_Internal<false>();
 }
 
-void FSkinWeightVertexBuffer::InitRHI()
+void FSkinWeightDataVertexBuffer::InitRHI()
 {
 	// BUF_ShaderResource is needed for support of the SkinCache (we could make is dependent on GEnableGPUSkinCacheShaders or are there other users?)
 	VertexBufferRHI = CreateRHIBuffer_RenderThread();
 
-	bool bSRV = VertexBufferRHI && GSupportsResourceView && GPixelFormats[PF_R32_UINT].Supported;
+	bool bSRV = VertexBufferRHI && GSupportsResourceView && GPixelFormats[GetPixelFormat()].Supported;
 	// When bAllowCPUAccess is true, the meshes is likely going to be used for Niagara to spawn particles on mesh surface.
 	// And it can be the case for CPU *and* GPU access: no differenciation today. That is why we create a SRV in this case.
 	// This also avoid setting lots of states on all the members of all the different buffers used by meshes. Follow up: https://jira.it.epicgames.net/browse/UE-69376.
@@ -193,44 +430,471 @@ void FSkinWeightVertexBuffer::InitRHI()
 
 	if (bSRV)
 	{
-		SRVValue = RHICreateShaderResourceView(WeightData ? VertexBufferRHI : nullptr, 4, PF_R32_UINT);
+		SRVValue = RHICreateShaderResourceView(WeightData ? VertexBufferRHI : nullptr, GetPixelFormatStride(), GetPixelFormat());
 	}
 }
 
-void FSkinWeightVertexBuffer::ReleaseRHI()
+void FSkinWeightDataVertexBuffer::ReleaseRHI()
 {
 	SRVValue.SafeRelease();
 
 	FVertexBuffer::ReleaseRHI();
 }
 
+void FSkinWeightDataVertexBuffer::AllocateData()
+{
+	// Clear any old WeightData before allocating.
+	CleanUp();
+
+	WeightData = new FSkinWeightVertexData<uint8>(bNeedsCPUAccess);
+}
 
+void FSkinWeightDataVertexBuffer::SetMaxBoneInfluences(uint32 InMaxBoneInfluences)
+{
+	check(InMaxBoneInfluences <= MAX_TOTAL_INFLUENCES);
+	if (InMaxBoneInfluences <= MAX_INFLUENCES_PER_STREAM)
+		MaxBoneInfluences = MAX_INFLUENCES_PER_STREAM;
+	else if (InMaxBoneInfluences <= EXTRA_BONE_INFLUENCES)
+		MaxBoneInfluences = EXTRA_BONE_INFLUENCES;
+	else
+		MaxBoneInfluences = MAX_TOTAL_INFLUENCES;
 
+	bVariableBonesPerVertex = FGPUBaseSkinVertexFactory::UseUnlimitedBoneInfluences(MaxBoneInfluences);
+	if (MaxBoneInfluences == MAX_TOTAL_INFLUENCES && !bVariableBonesPerVertex)
+	{
+		MaxBoneInfluences = EXTRA_BONE_INFLUENCES;
+		UE_LOG(LogSkeletalMesh, Error, TEXT("Skeletal mesh of %d bone influences requires unlimited bone influence mode. Influence truncation will occur and will not render correctly."), InMaxBoneInfluences);
+	}
+}
 
-void FSkinWeightVertexBuffer::AllocateData()
+bool FSkinWeightDataVertexBuffer::GetRigidWeightBone(uint32 VertexWeightOffset, uint32 VertexInfluenceCount, int32& OutBoneIndex) const
+{
+	bool bIsRigid = false;
+
+	uint8* BoneData = Data + VertexWeightOffset;
+	uint32 BoneWeightOffset = GetBoneIndexByteSize() * VertexInfluenceCount;
+	for (uint32 InfluenceIndex = 0; InfluenceIndex < VertexInfluenceCount; InfluenceIndex++)
+	{
+		uint8 BoneWeight = BoneData[BoneWeightOffset + InfluenceIndex];
+		if (BoneWeight == 255)
+		{
+			bIsRigid = true;
+			if (Use16BitBoneIndex())
+			{
+				FBoneIndex16* BoneIndex16Ptr = (FBoneIndex16*)BoneData;
+				OutBoneIndex = BoneIndex16Ptr[InfluenceIndex];
+			}
+			else
+			{
+				OutBoneIndex = BoneData[InfluenceIndex];
+			}
+			break;
+		}
+	}
+
+	return bIsRigid;
+}
+
+uint32 FSkinWeightDataVertexBuffer::GetBoneIndex(uint32 VertexWeightOffset, uint32 VertexInfluenceCount, uint32 InfluenceIndex) const
+{	
+	if (InfluenceIndex < VertexInfluenceCount)
+	{
+		uint8* BoneData = Data + VertexWeightOffset;
+		if (Use16BitBoneIndex())
+		{
+			FBoneIndex16* BoneIndex16Ptr = (FBoneIndex16*)BoneData;
+			return BoneIndex16Ptr[InfluenceIndex];
+		}
+		else
+		{
+			return BoneData[InfluenceIndex];
+		}
+	}
+	else
+	{
+		return 0;
+	}
+}
+
+void FSkinWeightDataVertexBuffer::SetBoneIndex(uint32 VertexWeightOffset, uint32 VertexInfluenceCount, uint32 InfluenceIndex, uint32 BoneIndex)
+{
+	if (InfluenceIndex < VertexInfluenceCount)
+	{
+		uint8* BoneData = Data + VertexWeightOffset;
+		if (Use16BitBoneIndex())
+		{
+			FBoneIndex16* BoneIndex16Ptr = (FBoneIndex16*)BoneData;
+			BoneIndex16Ptr[InfluenceIndex] = BoneIndex;
+		}
+		else
+		{
+			BoneData[InfluenceIndex] = BoneIndex;
+		}
+	}
+	else
+	{
+		check(false);
+	}
+}
+
+uint8 FSkinWeightDataVertexBuffer::GetBoneWeight(uint32 VertexWeightOffset, uint32 VertexInfluenceCount, uint32 InfluenceIndex) const
+{
+	if (InfluenceIndex < VertexInfluenceCount)
+	{
+		uint8* BoneData = Data + VertexWeightOffset;
+		uint32 BoneWeightOffset = GetBoneIndexByteSize() * VertexInfluenceCount;
+		return BoneData[BoneWeightOffset + InfluenceIndex];
+	}
+	else
+	{
+		return 0;
+	}
+}
+
+void FSkinWeightDataVertexBuffer::SetBoneWeight(uint32 VertexWeightOffset, uint32 VertexInfluenceCount, uint32 InfluenceIndex, uint8 BoneWeight)
+{
+	if (InfluenceIndex < VertexInfluenceCount)
+	{
+		uint8* BoneData = Data + VertexWeightOffset;
+		uint32 BoneWeightOffset = GetBoneIndexByteSize() * VertexInfluenceCount;
+		BoneData[BoneWeightOffset + InfluenceIndex] = BoneWeight;
+	}
+	else
+	{
+		check(false);
+	}
+}
+
+void FSkinWeightDataVertexBuffer::ResetVertexBoneWeights(uint32 VertexWeightOffset, uint32 VertexInfluenceCount)
+{
+	if (VertexInfluenceCount > 0)
+	{
+		uint8* BoneData = Data + VertexWeightOffset;
+		if (BoneData)
+		{
+			if (Use16BitBoneIndex())
+			{
+				FMemory::Memzero(BoneData, (sizeof(FBoneIndex16) + sizeof(uint8)) * VertexInfluenceCount);
+			}
+			else
+			{
+				FMemory::Memzero(BoneData, (sizeof(FBoneIndex8) + sizeof(uint8)) * VertexInfluenceCount);
+			}
+		}
+	}
+}
+
+
+/*-----------------------------------------------------------------------------
+FSkinWeightVertexBuffer
+-----------------------------------------------------------------------------*/
+
+FSkinWeightVertexBuffer::FSkinWeightVertexBuffer()
+{
+}
+
+FSkinWeightVertexBuffer::FSkinWeightVertexBuffer(const FSkinWeightVertexBuffer& Other)
+{
+	bool bNeedsCPUAccess = Other.GetNeedsCPUAccess();
+	DataVertexBuffer.SetNeedsCPUAccess(bNeedsCPUAccess);
+	LookupVertexBuffer.SetNeedsCPUAccess(bNeedsCPUAccess);
+
+	SetMaxBoneInfluences(Other.GetMaxBoneInfluences());
+	SetUse16BitBoneIndex(Other.Use16BitBoneIndex());
+}
+
+FSkinWeightVertexBuffer::~FSkinWeightVertexBuffer()
+{
+	CleanUp();
+}
+
+FSkinWeightVertexBuffer& FSkinWeightVertexBuffer::operator=(const FSkinWeightVertexBuffer& Other)
 {
-	// Clear any old WeightData before allocating.
 	CleanUp();
 
-	if (bExtraBoneInfluences)
+	bool bNeedsCPUAccess = Other.GetNeedsCPUAccess();
+	DataVertexBuffer.SetNeedsCPUAccess(bNeedsCPUAccess);
+	LookupVertexBuffer.SetNeedsCPUAccess(bNeedsCPUAccess);
+
+	SetMaxBoneInfluences(Other.GetMaxBoneInfluences());
+	SetUse16BitBoneIndex(Other.Use16BitBoneIndex());
+
+	return *this;
+}
+
+void FSkinWeightVertexBuffer::CleanUp()
+{
+	DataVertexBuffer.CleanUp();
+	LookupVertexBuffer.CleanUp();
+}
+
+FArchive& operator<<(FArchive& Ar, FSkinWeightVertexBuffer& VertexBuffer)
+{
+	Ar.UsingCustomVersion(FAnimObjectVersion::GUID);	
+	Ar << VertexBuffer.DataVertexBuffer;
+
+	if (Ar.IsLoading() && Ar.CustomVer(FAnimObjectVersion::GUID) < FAnimObjectVersion::UnlimitedBoneInfluences)
 	{
-		WeightData = new FSkinWeightVertexData< TSkinWeightInfo<true> >(bNeedsCPUAccess);
+		// LookupVertexBuffer doesn't exist before this version, so construct its content from scratch
+		uint32 MaxBoneInfluences = VertexBuffer.DataVertexBuffer.GetMaxBoneInfluences();
+		uint32 NumVertices = VertexBuffer.DataVertexBuffer.GetNumBones() / VertexBuffer.DataVertexBuffer.GetMaxBoneInfluences();
+		VertexBuffer.LookupVertexBuffer.Init(NumVertices);
+
+		uint32 WeightOffset = 0;
+		for (uint32 VertIdx = 0; VertIdx < NumVertices; VertIdx++)
+		{
+			VertexBuffer.LookupVertexBuffer.SetWeightOffsetAndInfluenceCount(VertIdx, WeightOffset * (VertexBuffer.GetBoneIndexByteSize() + sizeof(uint8)), MaxBoneInfluences);
+			WeightOffset += MaxBoneInfluences;
+		}
 	}
 	else
 	{
-		WeightData = new FSkinWeightVertexData< TSkinWeightInfo<false> >(bNeedsCPUAccess);
+		Ar << VertexBuffer.LookupVertexBuffer;
 	}
+
+	return Ar;
 }
 
-#if WITH_EDITOR
+void FSkinWeightVertexBuffer::SerializeMetaData(FArchive& Ar)
+{
+	DataVertexBuffer.SerializeMetaData(Ar);
+	LookupVertexBuffer.SerializeMetaData(Ar);
+}
+
+void FSkinWeightVertexBuffer::CopyMetaData(const FSkinWeightVertexBuffer& Other)
+{
+	DataVertexBuffer.CopyMetaData(Other.DataVertexBuffer);
+	LookupVertexBuffer.CopyMetaData(Other.LookupVertexBuffer);
+}
+
+FSkinWeightRHIInfo FSkinWeightVertexBuffer::CreateRHIBuffer_RenderThread()
+{
+	FSkinWeightRHIInfo RHIInfo;
+	RHIInfo.DataVertexBufferRHI = DataVertexBuffer.CreateRHIBuffer_RenderThread();
+	RHIInfo.LookupVertexBufferRHI = LookupVertexBuffer.CreateRHIBuffer_RenderThread();
+	return RHIInfo;
+}
 
-template <bool bExtraBoneInfluencesT>
-void FSkinWeightVertexBuffer::SetWeightsForVertex(uint32 VertexIndex, const FSoftSkinVertex& SrcVertex)
+FSkinWeightRHIInfo FSkinWeightVertexBuffer::CreateRHIBuffer_Async()
 {
-	checkSlow(VertexIndex < GetNumVertices());
-	auto* VertBase = (TSkinWeightInfo<bExtraBoneInfluencesT>*)(Data + VertexIndex * Stride);
-	FMemory::Memcpy(VertBase->InfluenceBones, SrcVertex.InfluenceBones, TSkinWeightInfo<bExtraBoneInfluencesT>::NumInfluences);
-	FMemory::Memcpy(VertBase->InfluenceWeights, SrcVertex.InfluenceWeights, TSkinWeightInfo<bExtraBoneInfluencesT>::NumInfluences);
+	FSkinWeightRHIInfo RHIInfo;
+	RHIInfo.DataVertexBufferRHI = DataVertexBuffer.CreateRHIBuffer_Async();
+	RHIInfo.LookupVertexBufferRHI = LookupVertexBuffer.CreateRHIBuffer_Async();
+	return RHIInfo;
 }
 
+GPUSkinBoneInfluenceType FSkinWeightVertexBuffer::GetBoneInfluenceType() const
+{
+	if (GetVariableBonesPerVertex())
+	{
+		return GPUSkinBoneInfluenceType::UnlimitedBoneInfluence;
+	}
+	else if (GetMaxBoneInfluences() > MAX_INFLUENCES_PER_STREAM)
+	{
+		return GPUSkinBoneInfluenceType::ExtraBoneInfluence;
+	}
+	else
+	{
+		return GPUSkinBoneInfluenceType::DefaultBoneInfluence;
+	}
+}
+
+void FSkinWeightVertexBuffer::GetVertexInfluenceOffsetCount(uint32 VertexIndex, uint32& VertexWeightOffset, uint32& VertexInfluenceCount) const
+{
+	if (GetVariableBonesPerVertex())
+	{
+		LookupVertexBuffer.GetWeightOffsetAndInfluenceCount(VertexIndex, VertexWeightOffset, VertexInfluenceCount);
+	}
+	else
+	{
+		VertexWeightOffset = GetConstantInfluencesVertexStride() * VertexIndex;
+		VertexInfluenceCount = GetMaxBoneInfluences();
+	}
+}
+
+bool FSkinWeightVertexBuffer::GetRigidWeightBone(uint32 VertexIndex, int32& OutBoneIndex) const
+{
+	uint32 VertexWeightOffset = 0;
+	uint32 VertexInfluenceCount = 0;
+	GetVertexInfluenceOffsetCount(VertexIndex, VertexWeightOffset, VertexInfluenceCount);
+	return DataVertexBuffer.GetRigidWeightBone(VertexWeightOffset, VertexInfluenceCount, OutBoneIndex);
+}
+
+uint32 FSkinWeightVertexBuffer::GetBoneIndex(uint32 VertexIndex, uint32 InfluenceIndex) const
+{
+	uint32 VertexWeightOffset = 0;
+	uint32 VertexInfluenceCount = 0;
+	GetVertexInfluenceOffsetCount(VertexIndex, VertexWeightOffset, VertexInfluenceCount);
+	return DataVertexBuffer.GetBoneIndex(VertexWeightOffset, VertexInfluenceCount, InfluenceIndex);
+}
+
+void FSkinWeightVertexBuffer::SetBoneIndex(uint32 VertexIndex, uint32 InfluenceIndex, uint32 BoneIndex)
+{
+	uint32 VertexWeightOffset = 0;
+	uint32 VertexInfluenceCount = 0;
+	GetVertexInfluenceOffsetCount(VertexIndex, VertexWeightOffset, VertexInfluenceCount);
+	DataVertexBuffer.SetBoneIndex(VertexWeightOffset, VertexInfluenceCount, InfluenceIndex, BoneIndex);
+}
+
+uint8 FSkinWeightVertexBuffer::GetBoneWeight(uint32 VertexIndex, uint32 InfluenceIndex) const
+{
+	uint32 VertexWeightOffset = 0;
+	uint32 VertexInfluenceCount = 0;
+	GetVertexInfluenceOffsetCount(VertexIndex, VertexWeightOffset, VertexInfluenceCount);
+	return DataVertexBuffer.GetBoneWeight(VertexWeightOffset, VertexInfluenceCount, InfluenceIndex);
+}
+
+void FSkinWeightVertexBuffer::SetBoneWeight(uint32 VertexIndex, uint32 InfluenceIndex, uint8 BoneWeight)
+{
+	uint32 VertexWeightOffset = 0;
+	uint32 VertexInfluenceCount = 0;
+	GetVertexInfluenceOffsetCount(VertexIndex, VertexWeightOffset, VertexInfluenceCount);
+	DataVertexBuffer.SetBoneWeight(VertexWeightOffset, VertexInfluenceCount, InfluenceIndex, BoneWeight);
+}
+
+void FSkinWeightVertexBuffer::ResetVertexBoneWeights(uint32 VertexIndex)
+{
+	uint32 VertexWeightOffset = 0;
+	uint32 VertexInfluenceCount = 0;
+	GetVertexInfluenceOffsetCount(VertexIndex, VertexWeightOffset, VertexInfluenceCount);
+	DataVertexBuffer.ResetVertexBoneWeights(VertexWeightOffset, VertexInfluenceCount);
+}
+
+void FSkinWeightVertexBuffer::BeginInitResources()
+{
+	BeginInitResource(&LookupVertexBuffer);
+	BeginInitResource(&DataVertexBuffer);
+}
+
+void FSkinWeightVertexBuffer::BeginReleaseResources()
+{
+	BeginReleaseResource(&LookupVertexBuffer);
+	BeginReleaseResource(&DataVertexBuffer);
+}
+
+void FSkinWeightVertexBuffer::ReleaseResources()
+{
+	check(IsInRenderingThread());
+	LookupVertexBuffer.ReleaseResource();
+	DataVertexBuffer.ReleaseResource();
+}
+
+#if WITH_EDITOR
+void FSkinWeightVertexBuffer::Init(const TArray<FSoftSkinVertex>& InVertices)
+{
+	TArray<FSkinWeightInfo> VertexWeightInfos;
+	VertexWeightInfos.AddUninitialized(InVertices.Num());
+
+	for (int32 VertIdx = 0; VertIdx < InVertices.Num(); VertIdx++)
+	{
+		const FSoftSkinVertex& SrcVertex = InVertices[VertIdx];
+		FSkinWeightInfo& DstVertex = VertexWeightInfos[VertIdx];
+		FMemory::Memcpy(DstVertex.InfluenceBones, SrcVertex.InfluenceBones, MAX_TOTAL_INFLUENCES * sizeof(FBoneIndexType));
+		FMemory::Memcpy(DstVertex.InfluenceWeights, SrcVertex.InfluenceWeights, MAX_TOTAL_INFLUENCES * sizeof(uint8));
+
+	}
+
+	*this = VertexWeightInfos;
+}
 #endif //WITH_EDITOR
+
+FSkinWeightVertexBuffer& FSkinWeightVertexBuffer::operator=(const TArray<FSkinWeightInfo>& InVertices)
+{
+	uint32 TotalNumUsedBones = 0;
+	bool bVariableBonesPerVertex = GetVariableBonesPerVertex();
+
+	if (bVariableBonesPerVertex)
+	{
+		LookupVertexBuffer.Init(InVertices.Num());
+		for (int32 VertIdx = 0; VertIdx < InVertices.Num(); VertIdx++)
+		{
+			int32 NumUsedBones = 0;
+			for (int32 InfluenceIdx = 0; InfluenceIdx < MAX_TOTAL_INFLUENCES; InfluenceIdx++)
+			{
+				if (InVertices[VertIdx].InfluenceWeights[InfluenceIdx] > 0)
+				{
+					NumUsedBones++;
+				}
+			}
+
+			LookupVertexBuffer.SetWeightOffsetAndInfluenceCount(VertIdx, TotalNumUsedBones * (GetBoneIndexByteSize() + sizeof(uint8)), NumUsedBones);
+			TotalNumUsedBones += NumUsedBones;
+		}
+	}
+	else
+	{
+		TotalNumUsedBones = GetMaxBoneInfluences() * InVertices.Num();
+	}
+
+	DataVertexBuffer.Init(TotalNumUsedBones, InVertices.Num());
+	for (int32 VertIdx = 0; VertIdx < InVertices.Num(); VertIdx++)
+	{
+		if (bVariableBonesPerVertex)
+		{
+			// Copy the weights and sort in descending order
+			FSkinWeightInfo VertexWeights = InVertices[VertIdx];
+			for (uint32 i = 0; i < MAX_TOTAL_INFLUENCES; i++)
+			{
+				uint32 MaxWeightIdx = i;
+				for (uint32 j = i+1; j < MAX_TOTAL_INFLUENCES; j++)
+				{
+					if (VertexWeights.InfluenceWeights[j] > VertexWeights.InfluenceWeights[MaxWeightIdx])
+					{
+						MaxWeightIdx = j;
+					}
+				}
+				if (MaxWeightIdx != i)
+				{
+					Exchange(VertexWeights.InfluenceBones[i], VertexWeights.InfluenceBones[MaxWeightIdx]);
+					Exchange(VertexWeights.InfluenceWeights[i], VertexWeights.InfluenceWeights[MaxWeightIdx]);
+				}
+			}
+
+			int32 NumUsedBones = 0;
+			for (uint32 InfluenceIdx = 0; InfluenceIdx < GetMaxBoneInfluences(); InfluenceIdx++)
+			{
+				if (VertexWeights.InfluenceWeights[InfluenceIdx] > 0)
+				{
+					SetBoneIndex(VertIdx, NumUsedBones, VertexWeights.InfluenceBones[InfluenceIdx]);
+					SetBoneWeight(VertIdx, NumUsedBones, VertexWeights.InfluenceWeights[InfluenceIdx]);
+					NumUsedBones++;
+				}
+			}
+		}
+		else
+		{
+			const FSkinWeightInfo& VertexWeights = InVertices[VertIdx];
+			for (uint32 InfluenceIdx = 0; InfluenceIdx < GetMaxBoneInfluences(); InfluenceIdx++)
+			{
+				SetBoneIndex(VertIdx, InfluenceIdx, VertexWeights.InfluenceBones[InfluenceIdx]);
+				SetBoneWeight(VertIdx, InfluenceIdx, VertexWeights.InfluenceWeights[InfluenceIdx]);
+			}
+		}
+	}
+
+	return *this;
+}
+
+void FSkinWeightVertexBuffer::GetSkinWeights(TArray<FSkinWeightInfo>& OutVertices) const
+{
+	OutVertices.SetNum(GetNumVertices());
+	for (uint32 VertIdx = 0; VertIdx < GetNumVertices(); VertIdx++)
+	{
+		for (int32 InfluenceIdx = 0; InfluenceIdx < MAX_TOTAL_INFLUENCES; InfluenceIdx++)
+		{
+			OutVertices[VertIdx].InfluenceBones[InfluenceIdx] = GetBoneIndex(VertIdx, InfluenceIdx);
+			OutVertices[VertIdx].InfluenceWeights[InfluenceIdx] = GetBoneWeight(VertIdx, InfluenceIdx);
+		}
+	}
+}
+
+FSkinWeightInfo FSkinWeightVertexBuffer::GetVertexSkinWeights(uint32 VertexIndex) const
+{
+	FSkinWeightInfo OutVertex;
+	for (int32 InfluenceIdx = 0; InfluenceIdx < MAX_TOTAL_INFLUENCES; InfluenceIdx++)
+	{
+		OutVertex.InfluenceBones[InfluenceIdx] = GetBoneIndex(VertexIndex, InfluenceIdx);
+		OutVertex.InfluenceWeights[InfluenceIdx] = GetBoneWeight(VertexIndex, InfluenceIdx);
+	}
+	return OutVertex;
+}
\ No newline at end of file
