COPY from //Dev-Anim to //Dev-Main

diff --git a/Engine/Source/Runtime/Engine/Private/SkeletalMesh.cpp b/Engine/Source/Runtime/Engine/Private/SkeletalMesh.cpp
--- a/Engine/Source/Runtime/Engine/Private/SkeletalMesh.cpp
+++ b/Engine/Source/Runtime/Engine/Private/SkeletalMesh.cpp
@@ -4411,7 +4411,7 @@ FSkeletalMeshSceneProxy::FSkeletalMeshSceneProxy(const USkinnedMeshComponent* Co
 				MaterialRelevance |= Material->GetRelevance(FeatureLevel);
 			}
 
-			const bool bRequiresAdjacencyInformation = RequiresAdjacencyInformation( Material, &TGPUSkinVertexFactory<false>::StaticType, FeatureLevel );
+			const bool bRequiresAdjacencyInformation = RequiresAdjacencyInformation( Material, &TGPUSkinVertexFactory<GPUSkinBoneInfluenceType::DefaultBoneInfluence>::StaticType, FeatureLevel );
 			if ( bRequiresAdjacencyInformation && LODData.AdjacencyMultiSizeIndexContainer.IsIndexBufferValid() == false )
 			{
 				UE_LOG(LogSkeletalMesh, Warning, 
@@ -5483,7 +5483,6 @@ FSkinnedMeshComponentRecreateRenderStateContext::~FSkinnedMeshComponentRecreateR
 
 //////////////////////////////////////////////////////////////////////////
 
-template <bool bExtraBoneInfluencesT>
 FVector GetRefVertexLocationTyped(
 	const USkeletalMesh* Mesh,
 	const FSkelMeshRenderSection& Section,
@@ -5496,8 +5495,7 @@ FVector GetRefVertexLocationTyped(
 
 	// Do soft skinning for this vertex.
 	int32 BufferVertIndex = Section.GetVertexBufferIndex() + VertIndex;
-	const TSkinWeightInfo<bExtraBoneInfluencesT>* SrcSkinWeights = SkinWeightVertexBuffer.GetSkinWeightPtr<bExtraBoneInfluencesT>(BufferVertIndex);
-	int32 MaxBoneInfluences = bExtraBoneInfluencesT ? MAX_TOTAL_INFLUENCES : MAX_INFLUENCES_PER_STREAM;
+	int32 MaxBoneInfluences = SkinWeightVertexBuffer.GetMaxBoneInfluences();
 
 #if !PLATFORM_LITTLE_ENDIAN
 	// uint8[] elements in LOD.VertexBufferGPUSkin have been swapped for VET_UBYTE4 vertex stream use
@@ -5506,8 +5504,8 @@ FVector GetRefVertexLocationTyped(
 	for (int32 InfluenceIndex = 0; InfluenceIndex < MaxBoneInfluences; InfluenceIndex++)
 #endif
 	{
-		const int32 MeshBoneIndex = Section.BoneMap[SrcSkinWeights->InfluenceBones[InfluenceIndex]];
-		const float	Weight = (float)SrcSkinWeights->InfluenceWeights[InfluenceIndex] / 255.0f;
+		const int32 MeshBoneIndex = Section.BoneMap[SkinWeightVertexBuffer.GetBoneIndex(BufferVertIndex, InfluenceIndex)];
+		const float	Weight = (float)SkinWeightVertexBuffer.GetBoneWeight(BufferVertIndex, InfluenceIndex) / 255.0f;
 		{
 			const FMatrix BoneTransformMatrix = FMatrix::Identity;//Mesh->GetComposedRefPoseMatrix(MeshBoneIndex);
 			const FMatrix RefToLocal = Mesh->RefBasesInvMatrix[MeshBoneIndex] * BoneTransformMatrix;
@@ -5525,14 +5523,7 @@ FVector GetSkeletalMeshRefVertLocation(const USkeletalMesh* Mesh, const FSkeleta
 	int32 VertIndexInChunk;
 	LODData.GetSectionFromVertexIndex(VertIndex, SectionIndex, VertIndexInChunk);
 	const FSkelMeshRenderSection& Section = LODData.RenderSections[SectionIndex];
-	if (SkinWeightVertexBuffer.HasExtraBoneInfluences())
-	{
-		return GetRefVertexLocationTyped<true>(Mesh, Section, LODData.StaticVertexBuffers.PositionVertexBuffer, SkinWeightVertexBuffer, VertIndexInChunk);
-	}
-	else
-	{
-		return GetRefVertexLocationTyped<false>(Mesh, Section, LODData.StaticVertexBuffers.PositionVertexBuffer, SkinWeightVertexBuffer, VertIndexInChunk);
-	}
+	return GetRefVertexLocationTyped(Mesh, Section, LODData.StaticVertexBuffers.PositionVertexBuffer, SkinWeightVertexBuffer, VertIndexInChunk);
 }
 
 #undef LOCTEXT_NAMESPACE
diff --git a/Engine/Source/Runtime/Engine/Private/SkeletalMeshComponentPhysics.cpp b/Engine/Source/Runtime/Engine/Private/SkeletalMeshComponentPhysics.cpp
--- a/Engine/Source/Runtime/Engine/Private/SkeletalMeshComponentPhysics.cpp
+++ b/Engine/Source/Runtime/Engine/Private/SkeletalMeshComponentPhysics.cpp
@@ -1935,9 +1935,7 @@ FVector GetTypedSkinnedVertexPositionWithCloth(USkeletalMeshComponent* Component
 				if (Component->ClothBlendWeight < 1.0f)
 				{
 					// blend with a skinned position
-					FVector SkinnedPos = SkinWeightBuffer.HasExtraBoneInfluences() ?
-						GetTypedSkinnedVertexPosition<true, bCachedMatrices>(Component, Section, LODData.StaticVertexBuffers.PositionVertexBuffer, SkinWeightBuffer, VertIndexInChunk, CachedRefToLocals) :
-						GetTypedSkinnedVertexPosition<false, bCachedMatrices>(Component, Section, LODData.StaticVertexBuffers.PositionVertexBuffer, SkinWeightBuffer, VertIndexInChunk, CachedRefToLocals);
+					FVector SkinnedPos = GetTypedSkinnedVertexPosition<bCachedMatrices>(Component, Section, LODData.StaticVertexBuffers.PositionVertexBuffer, SkinWeightBuffer, VertIndexInChunk, CachedRefToLocals);
 
 					SimulatedPos = SimulatedPos*Component->ClothBlendWeight + SkinnedPos*(1.0f - Component->ClothBlendWeight);
 				}
@@ -1946,9 +1944,7 @@ FVector GetTypedSkinnedVertexPositionWithCloth(USkeletalMeshComponent* Component
 		}
 	}
 
-	return SkinWeightBuffer.HasExtraBoneInfluences() ?
-		GetTypedSkinnedVertexPosition<true, bCachedMatrices>(Component, Section, LODData.StaticVertexBuffers.PositionVertexBuffer, SkinWeightBuffer, VertIndexInChunk, CachedRefToLocals) :
-		GetTypedSkinnedVertexPosition<false, bCachedMatrices>(Component, Section, LODData.StaticVertexBuffers.PositionVertexBuffer, SkinWeightBuffer, VertIndexInChunk, CachedRefToLocals);
+	return GetTypedSkinnedVertexPosition<bCachedMatrices>(Component, Section, LODData.StaticVertexBuffers.PositionVertexBuffer, SkinWeightBuffer, VertIndexInChunk, CachedRefToLocals);
 }
 
 FVector USkeletalMeshComponent::GetSkinnedVertexPosition(USkeletalMeshComponent* Component, int32 VertexIndex, const FSkeletalMeshLODRenderData& LODData, FSkinWeightVertexBuffer& SkinWeightBuffer)
@@ -1972,7 +1968,6 @@ void USkeletalMeshComponent::ComputeSkinnedPositions(USkeletalMeshComponent* Com
 
 	OutPositions.Empty();
 	OutPositions.AddUninitialized(LODData.GetNumVertices());
-	bool bHasExtraBoneInfluences = SkinWeightBuffer.HasExtraBoneInfluences();
 	
 	if (Component->SkeletalMesh->MeshClothingAssets.Num() > 0 &&
 		!Component->bDisableClothSimulation &&
@@ -2016,9 +2011,7 @@ void USkeletalMeshComponent::ComputeSkinnedPositions(USkeletalMeshComponent* Com
 							{
 								FVector SimulatedPos = ActorData->Positions[SoftIdx];
 
-								FVector SkinnedPosition = bHasExtraBoneInfluences ? 
-									GetTypedSkinnedVertexPosition<true, true>(Component, Section, LODData.StaticVertexBuffers.PositionVertexBuffer, SkinWeightBuffer, SoftIdx, CachedRefToLocals) :
-									GetTypedSkinnedVertexPosition<false, true>(Component, Section, LODData.StaticVertexBuffers.PositionVertexBuffer, SkinWeightBuffer, SoftIdx, CachedRefToLocals);
+								FVector SkinnedPosition = GetTypedSkinnedVertexPosition<true>(Component, Section, LODData.StaticVertexBuffers.PositionVertexBuffer, SkinWeightBuffer, SoftIdx, CachedRefToLocals);
 
 								OutPositions[SoftOffset + SoftIdx] = SimulatedPos*Component->ClothBlendWeight + SkinnedPosition*(1.0f - Component->ClothBlendWeight);
 							}
@@ -2041,9 +2034,7 @@ void USkeletalMeshComponent::ComputeSkinnedPositions(USkeletalMeshComponent* Com
 			const uint32 NumSoftVerts = Section.GetNumVertices();
 			for (uint32 SoftIdx = 0; SoftIdx < NumSoftVerts; ++SoftIdx)
 			{
-				FVector SkinnedPosition = bHasExtraBoneInfluences ?
-					GetTypedSkinnedVertexPosition<true, true>(Component, Section, LODData.StaticVertexBuffers.PositionVertexBuffer, SkinWeightBuffer, SoftIdx, CachedRefToLocals) :
-					GetTypedSkinnedVertexPosition<false, true>(Component, Section, LODData.StaticVertexBuffers.PositionVertexBuffer, SkinWeightBuffer, SoftIdx, CachedRefToLocals);
+				FVector SkinnedPosition = GetTypedSkinnedVertexPosition<true>(Component, Section, LODData.StaticVertexBuffers.PositionVertexBuffer, SkinWeightBuffer, SoftIdx, CachedRefToLocals);
 				OutPositions[SoftOffset + SoftIdx] = SkinnedPosition;
 			}
 		}
@@ -2058,9 +2049,7 @@ void USkeletalMeshComponent::ComputeSkinnedPositions(USkeletalMeshComponent* Com
 			const uint32 NumSoftVerts = Section.GetNumVertices();
 			for (uint32 SoftIdx = 0; SoftIdx < NumSoftVerts; ++SoftIdx)
 			{
-				FVector SkinnedPosition = bHasExtraBoneInfluences ? 
-					GetTypedSkinnedVertexPosition<true, true>(Component, Section, LODData.StaticVertexBuffers.PositionVertexBuffer, SkinWeightBuffer, SoftIdx, CachedRefToLocals) :
-					GetTypedSkinnedVertexPosition<false, true>(Component, Section, LODData.StaticVertexBuffers.PositionVertexBuffer, SkinWeightBuffer, SoftIdx, CachedRefToLocals);
+				FVector SkinnedPosition = GetTypedSkinnedVertexPosition<true>(Component, Section, LODData.StaticVertexBuffers.PositionVertexBuffer, SkinWeightBuffer, SoftIdx, CachedRefToLocals);
 				OutPositions[SoftOffset + SoftIdx] = SkinnedPosition;
 			}
 		}
diff --git a/Engine/Source/Runtime/Engine/Private/SkeletalMeshLODModel.cpp b/Engine/Source/Runtime/Engine/Private/SkeletalMeshLODModel.cpp
--- a/Engine/Source/Runtime/Engine/Private/SkeletalMeshLODModel.cpp
+++ b/Engine/Source/Runtime/Engine/Private/SkeletalMeshLODModel.cpp
@@ -19,6 +19,7 @@
 #include "Rendering/SkeletalMeshLODImporterData.h"
 #include "UObject/FortniteMainBranchObjectVersion.h"
 #include "GPUSkinVertexFactory.h"
+#include "UObject/AnimObjectVersion.h"
 
 /*-----------------------------------------------------------------------------
 FSoftSkinVertex
@@ -57,28 +58,52 @@ FArchive& operator<<(FArchive& Ar, FSoftSkinVertex& V)
 
 	Ar << V.Color;
 
+	Ar.UsingCustomVersion(FAnimObjectVersion::GUID);
+
+	if (Ar.IsLoading())
+	{
+		FMemory::Memzero(V.InfluenceBones);
+		FMemory::Memzero(V.InfluenceWeights);
+	}
+
 	// serialize bone and weight uint8 arrays in order
 	// this is required when serializing as bulk data memory (see TArray::BulkSerialize notes)
 	for (uint32 InfluenceIndex = 0; InfluenceIndex < MAX_INFLUENCES_PER_STREAM; InfluenceIndex++)
 	{
-		Ar << V.InfluenceBones[InfluenceIndex];
+		if (Ar.IsLoading() && Ar.CustomVer(FAnimObjectVersion::GUID) < FAnimObjectVersion::IncreaseBoneIndexLimitPerChunk)
+		{
+			uint8 BoneIndex = 0;
+			Ar << BoneIndex;
+			V.InfluenceBones[InfluenceIndex] = BoneIndex;
+		}
+		else
+		{
+			Ar << V.InfluenceBones[InfluenceIndex];
+		}
 	}
 
 	if (Ar.UE4Ver() >= VER_UE4_SUPPORT_8_BONE_INFLUENCES_SKELETAL_MESHES)
 	{
-		for (uint32 InfluenceIndex = MAX_INFLUENCES_PER_STREAM; InfluenceIndex < MAX_TOTAL_INFLUENCES; InfluenceIndex++)
+		for (uint32 InfluenceIndex = MAX_INFLUENCES_PER_STREAM; InfluenceIndex < EXTRA_BONE_INFLUENCES; InfluenceIndex++)
 		{
-			Ar << V.InfluenceBones[InfluenceIndex];
+			if (Ar.IsLoading() && Ar.CustomVer(FAnimObjectVersion::GUID) < FAnimObjectVersion::IncreaseBoneIndexLimitPerChunk)
+			{
+				uint8 BoneIndex = 0;
+				Ar << BoneIndex;
+				V.InfluenceBones[InfluenceIndex] = BoneIndex;
+			}
+			else
+			{
+				Ar << V.InfluenceBones[InfluenceIndex];
+			}
 		}
 	}
-	else
+
+	if (Ar.CustomVer(FAnimObjectVersion::GUID) >= FAnimObjectVersion::UnlimitedBoneInfluences)
 	{
-		if (Ar.IsLoading())
+		for (uint32 InfluenceIndex = EXTRA_BONE_INFLUENCES; InfluenceIndex < MAX_TOTAL_INFLUENCES; InfluenceIndex++)
 		{
-			for (uint32 InfluenceIndex = MAX_INFLUENCES_PER_STREAM; InfluenceIndex < MAX_TOTAL_INFLUENCES; InfluenceIndex++)
-			{
-				V.InfluenceBones[InfluenceIndex] = 0;
-			}
+			Ar << V.InfluenceBones[InfluenceIndex];
 		}
 	}
 
@@ -89,26 +114,24 @@ FArchive& operator<<(FArchive& Ar, FSoftSkinVertex& V)
 
 	if (Ar.UE4Ver() >= VER_UE4_SUPPORT_8_BONE_INFLUENCES_SKELETAL_MESHES)
 	{
-		for (uint32 InfluenceIndex = MAX_INFLUENCES_PER_STREAM; InfluenceIndex < MAX_TOTAL_INFLUENCES; InfluenceIndex++)
+		for (uint32 InfluenceIndex = MAX_INFLUENCES_PER_STREAM; InfluenceIndex < EXTRA_BONE_INFLUENCES; InfluenceIndex++)
 		{
 			Ar << V.InfluenceWeights[InfluenceIndex];
 		}
 	}
-	else
+
+	if (Ar.CustomVer(FAnimObjectVersion::GUID) >= FAnimObjectVersion::UnlimitedBoneInfluences)
 	{
-		if (Ar.IsLoading())
+		for (uint32 InfluenceIndex = EXTRA_BONE_INFLUENCES; InfluenceIndex < MAX_TOTAL_INFLUENCES; InfluenceIndex++)
 		{
-			for (uint32 InfluenceIndex = MAX_INFLUENCES_PER_STREAM; InfluenceIndex < MAX_TOTAL_INFLUENCES; InfluenceIndex++)
-			{
-				V.InfluenceWeights[InfluenceIndex] = 0;
-			}
+			Ar << V.InfluenceWeights[InfluenceIndex];
 		}
 	}
 
 	return Ar;
 }
 
-bool FSoftSkinVertex::GetRigidWeightBone(uint8& OutBoneIndex) const
+bool FSoftSkinVertex::GetRigidWeightBone(FBoneIndexType& OutBoneIndex) const
 {
 	bool bIsRigid = false;
 
@@ -251,6 +274,28 @@ void FSkelMeshSection::CalcMaxBoneInfluences()
 	}
 }
 
+/**
+* Calculate if this skel mesh section needs 16-bit bone indices
+*/
+void FSkelMeshSection::CalcUse16BitBoneIndex()
+{
+	bUse16BitBoneIndex = false;
+	FBoneIndexType MaxBoneIndex = 0;
+	for (int32 VertIdx = 0; VertIdx < SoftVertices.Num(); VertIdx++)
+	{
+		FSoftSkinVertex& SoftVert = SoftVertices[VertIdx];
+		for (int32 InfluenceIdx = 0; InfluenceIdx < MAX_TOTAL_INFLUENCES; InfluenceIdx++)
+		{
+			MaxBoneIndex = FMath::Max(SoftVert.InfluenceBones[InfluenceIdx], MaxBoneIndex);
+			if (MaxBoneIndex > 255)
+			{
+				bUse16BitBoneIndex = true;
+				return;
+			}
+		}
+	}
+}
+
 // Serialization.
 FArchive& operator<<(FArchive& Ar, FSkelMeshSection& S)
 {
@@ -371,6 +416,18 @@ FArchive& operator<<(FArchive& Ar, FSkelMeshSection& S)
 			}
 		}
 
+		Ar.UsingCustomVersion(FAnimObjectVersion::GUID);
+		if (Ar.IsLoading() && Ar.CustomVer(FAnimObjectVersion::GUID) < FAnimObjectVersion::IncreaseBoneIndexLimitPerChunk)
+		{
+			// Previous versions only supported 8-bit bone indices and bUse16BitBoneIndex wasn't serialized 
+			S.CalcUse16BitBoneIndex();
+			check(!S.bUse16BitBoneIndex);
+		}
+		else
+		{
+			Ar << S.bUse16BitBoneIndex;
+		}
+
 		Ar << S.BoneMap;
 
 		if (Ar.CustomVer(FSkeletalMeshCustomVersion::GUID) >= FSkeletalMeshCustomVersion::SaveNumVertices)
@@ -952,11 +1009,22 @@ void FSkeletalMeshLODModel::GetNonClothVertices(TArray<FSoftSkinVertex>& OutVert
 	}
 }
 
-bool FSkeletalMeshLODModel::DoSectionsNeedExtraBoneInfluences() const
+int32 FSkeletalMeshLODModel::GetMaxBoneInfluences() const
 {
+	int32 NumBoneInfluences = 0;
 	for (int32 SectionIdx = 0; SectionIdx < Sections.Num(); ++SectionIdx)
 	{
-		if (Sections[SectionIdx].HasExtraBoneInfluences())
+		NumBoneInfluences = FMath::Max(NumBoneInfluences, Sections[SectionIdx].GetMaxBoneInfluences());
+	}
+
+	return NumBoneInfluences;
+}
+
+bool FSkeletalMeshLODModel::DoSectionsUse16BitBoneIndex() const
+{
+	for (int32 SectionIdx = 0; SectionIdx < Sections.Num(); ++SectionIdx)
+	{
+		if (Sections[SectionIdx].Use16BitBoneIndex())
 		{
 			return true;
 		}
@@ -1035,6 +1103,24 @@ FString FSkeletalMeshLODModel::GetLODModelDeriveDataKey() const
 	return KeySuffix;
 }
 
+bool FSkeletalMeshLODModel::CopyStructure(FSkeletalMeshLODModel* Destination, FSkeletalMeshLODModel* Source)
+{
+	if (Source->RawPointIndices.IsLocked() || Source->LegacyRawPointIndices.IsLocked() || Source->RawSkeletalMeshBulkData_DEPRECATED.GetBulkData().IsLocked())
+	{
+		return false;
+	}
+	// Bulk data arrays need to be locked before a copy can be made.
+	Source->RawPointIndices.Lock(LOCK_READ_ONLY);
+	Source->LegacyRawPointIndices.Lock(LOCK_READ_ONLY);
+	Source->RawSkeletalMeshBulkData_DEPRECATED.GetBulkData().Lock(LOCK_READ_ONLY);
+	*Destination = *Source;
+	Source->RawSkeletalMeshBulkData_DEPRECATED.GetBulkData().Unlock();
+	Source->RawPointIndices.Unlock();
+	Source->LegacyRawPointIndices.Unlock();
+
+	return true;
+}
+
 void FSkeletalMeshLODModel::UpdateChunkedSectionInfo(const FString& SkeletalMeshName, TArray<int32>& LODMaterialMap)
 {
 	int32 LODModelSectionNum = Sections.Num();
@@ -1090,22 +1176,4 @@ void FSkeletalMeshLODModel::UpdateChunkedSectionInfo(const FString& SkeletalMesh
 	}
 }
 
-bool FSkeletalMeshLODModel::CopyStructure(FSkeletalMeshLODModel* Destination, FSkeletalMeshLODModel* Source)
-{
-	if (Source->RawPointIndices.IsLocked() || Source->LegacyRawPointIndices.IsLocked() || Source->RawSkeletalMeshBulkData_DEPRECATED.GetBulkData().IsLocked())
-	{
-		return false;
-	}
-	// Bulk data arrays need to be locked before a copy can be made.
-	Source->RawPointIndices.Lock(LOCK_READ_ONLY);
-	Source->LegacyRawPointIndices.Lock(LOCK_READ_ONLY);
-	Source->RawSkeletalMeshBulkData_DEPRECATED.GetBulkData().Lock(LOCK_READ_ONLY);
-	*Destination = *Source;
-	Source->RawSkeletalMeshBulkData_DEPRECATED.GetBulkData().Unlock();
-	Source->RawPointIndices.Unlock();
-	Source->LegacyRawPointIndices.Unlock();
-
-	return true;
-}
-
 #endif // WITH_EDITOR
diff --git a/Engine/Source/Runtime/Engine/Private/SkeletalMeshLODRenderData.cpp b/Engine/Source/Runtime/Engine/Private/SkeletalMeshLODRenderData.cpp
--- a/Engine/Source/Runtime/Engine/Private/SkeletalMeshLODRenderData.cpp
+++ b/Engine/Source/Runtime/Engine/Private/SkeletalMeshLODRenderData.cpp
@@ -196,7 +196,7 @@ void FSkeletalMeshLODRenderData::InitResources(bool bNeedsVertexColors, int32 LO
 	{
 		SkinWeightProfilesData.SetDynamicDefaultSkinWeightProfile(Owner, LODIndex);
 	}
-	BeginInitResource(&SkinWeightVertexBuffer);
+	SkinWeightVertexBuffer.BeginInitResources();
 
 	if (bNeedsVertexColors)
 	{
@@ -474,7 +474,7 @@ void FSkeletalMeshLODRenderData::ReleaseResources()
 
 	BeginReleaseResource(&StaticVertexBuffers.PositionVertexBuffer);
 	BeginReleaseResource(&StaticVertexBuffers.StaticMeshVertexBuffer);
-	BeginReleaseResource(&SkinWeightVertexBuffer);
+	SkinWeightVertexBuffer.BeginReleaseResources();
 	BeginReleaseResource(&StaticVertexBuffers.ColorVertexBuffer);
 	BeginReleaseResource(&ClothVertexBuffer);
 	// DuplicatedVerticesBuffer is used only for SkinCache and Editor features which is SM5 only
@@ -584,7 +584,8 @@ void FSkeletalMeshLODRenderData::BuildFromLODModel(const FSkeletalMeshLODModel*
 
 	// Init skin weight buffer
 	SkinWeightVertexBuffer.SetNeedsCPUAccess(true);
-	SkinWeightVertexBuffer.SetHasExtraBoneInfluences(ImportedModel->DoSectionsNeedExtraBoneInfluences());
+	SkinWeightVertexBuffer.SetMaxBoneInfluences(ImportedModel->GetMaxBoneInfluences());
+	SkinWeightVertexBuffer.SetUse16BitBoneIndex(ImportedModel->DoSectionsUse16BitBoneIndex());
 	SkinWeightVertexBuffer.Init(Vertices);
 
 	// Init the color buffer if this mesh has vertex colors.
@@ -641,11 +642,8 @@ void FSkeletalMeshLODRenderData::ReleaseCPUResources(bool bForStreaming)
 		{
 			AdjacencyMultiSizeIndexContainer.GetIndexBuffer()->Empty();
 		}
-		if (SkinWeightVertexBuffer.IsWeightDataValid())
-		{
-			SkinWeightVertexBuffer.CleanUp();
-		}
 
+		SkinWeightVertexBuffer.CleanUp();
 		StaticVertexBuffers.PositionVertexBuffer.CleanUp();
 		StaticVertexBuffers.StaticMeshVertexBuffer.CleanUp();
 
diff --git a/Engine/Source/Runtime/Engine/Private/SkeletalMeshMerge.cpp b/Engine/Source/Runtime/Engine/Private/SkeletalMeshMerge.cpp
--- a/Engine/Source/Runtime/Engine/Private/SkeletalMeshMerge.cpp
+++ b/Engine/Source/Runtime/Engine/Private/SkeletalMeshMerge.cpp
@@ -40,21 +40,21 @@ FSkeletalMeshMerge::FSkeletalMeshMerge(USkeletalMesh* InMergeMesh,
 }
 
 /** Helper macro to call GenerateLODModel which requires compile time vertex type. */
-#define GENERATE_LOD_MODEL( VertexType, NumUVs, bHasExtraBoneInfluences ) \
+#define GENERATE_LOD_MODEL( VertexType, NumUVs ) \
 {\
 	switch( NumUVs )\
 	{\
 	case 1:\
-		GenerateLODModel< VertexType<1>, TSkinWeightInfo<bHasExtraBoneInfluences> >( LODIdx + StripTopLODs );\
+		GenerateLODModel< VertexType<1> >( LODIdx + StripTopLODs );\
 		break;\
 	case 2:\
-		GenerateLODModel< VertexType<2>, TSkinWeightInfo<bHasExtraBoneInfluences> >( LODIdx + StripTopLODs );\
+		GenerateLODModel< VertexType<2> >( LODIdx + StripTopLODs );\
 		break;\
 	case 3:\
-		GenerateLODModel< VertexType<3>, TSkinWeightInfo<bHasExtraBoneInfluences> >( LODIdx + StripTopLODs );\
+		GenerateLODModel< VertexType<3> >( LODIdx + StripTopLODs );\
 		break;\
 	case 4:\
-		GenerateLODModel< VertexType<4>, TSkinWeightInfo<bHasExtraBoneInfluences> >( LODIdx + StripTopLODs );\
+		GenerateLODModel< VertexType<4> >( LODIdx + StripTopLODs );\
 		break;\
 	default:\
 		checkf(false, TEXT("Invalid number of UV sets.  Must be between 0 and 4") );\
@@ -166,9 +166,11 @@ bool FSkeletalMeshMerge::FinalizeMesh()
 
 		// Array of per-lod number of UV sets
 		TArray<uint32> PerLODNumUVSets;
-		TArray<bool> PerLODExtraBoneInfluences;
+		TArray<uint32> PerLODMaxBoneInfluences;
+		TArray<bool> PerLODUse16BitBoneIndex;
 		PerLODNumUVSets.AddZeroed(MaxNumLODs);
-		PerLODExtraBoneInfluences.AddZeroed(MaxNumLODs);
+		PerLODMaxBoneInfluences.AddZeroed(MaxNumLODs);
+		PerLODUse16BitBoneIndex.AddZeroed(MaxNumLODs);
 
 		// Get the number of UV sets for each LOD.
 		for (int32 MeshIdx = 0; MeshIdx < SrcMeshList.Num(); MeshIdx++)
@@ -183,7 +185,8 @@ bool FSkeletalMeshMerge::FinalizeMesh()
 					uint32& NumUVSets = PerLODNumUVSets[LODIdx];
 					NumUVSets = FMath::Max(NumUVSets, SrcResource->LODRenderData[LODIdx].GetNumTexCoords());
 
-					PerLODExtraBoneInfluences[LODIdx] |= SrcResource->LODRenderData[LODIdx].DoesVertexBufferHaveExtraBoneInfluences();
+					PerLODMaxBoneInfluences[LODIdx] = FMath::Max(PerLODMaxBoneInfluences[LODIdx], SrcResource->LODRenderData[LODIdx].GetVertexBufferMaxBoneInfluences());
+					PerLODUse16BitBoneIndex[LODIdx] |= SrcResource->LODRenderData[LODIdx].DoesVertexBufferUse16BitBoneIndex();
 				}
 			}
 		}
@@ -196,25 +199,11 @@ bool FSkeletalMeshMerge::FinalizeMesh()
 			bool bUseFullPrecisionUVs = LODInfoPtr ? LODInfoPtr->BuildSettings.bUseFullPrecisionUVs : GVertexElementTypeSupport.IsSupported(VET_Half2) ? false : true;
 			if (!bUseFullPrecisionUVs)
 			{
-				if (PerLODExtraBoneInfluences[LODIdx])
-				{
-					GENERATE_LOD_MODEL(TGPUSkinVertexFloat16Uvs, PerLODNumUVSets[LODIdx], true);
-				}
-				else
-				{
-					GENERATE_LOD_MODEL(TGPUSkinVertexFloat16Uvs, PerLODNumUVSets[LODIdx], false);
-				}
+				GENERATE_LOD_MODEL(TGPUSkinVertexFloat16Uvs, PerLODNumUVSets[LODIdx]);
 			}
 			else
 			{
-				if (PerLODExtraBoneInfluences[LODIdx])
-				{
-					GENERATE_LOD_MODEL(TGPUSkinVertexFloat32Uvs, PerLODNumUVSets[LODIdx], true);
-				}
-				else
-				{
-					GENERATE_LOD_MODEL(TGPUSkinVertexFloat32Uvs, PerLODNumUVSets[LODIdx], false);
-				}
+				GENERATE_LOD_MODEL(TGPUSkinVertexFloat32Uvs, PerLODNumUVSets[LODIdx]);
 			}
 		}
 		// update the merge skel mesh entries
@@ -415,24 +404,11 @@ void FSkeletalMeshMerge::CopyVertexFromSource(VertexDataType& DestVert, const FS
 	}
 }
 
-template<typename SkinWeightType, bool bHasExtraBoneInfluences>
-void FSkeletalMeshMerge::CopyWeightFromSource(SkinWeightType& DestWeight, const FSkeletalMeshLODRenderData& SrcLODData, int32 SourceVertIdx, const FMergeSectionInfo& MergeSectionInfo)
-{
-	const TSkinWeightInfo<bHasExtraBoneInfluences>* SrcSkinWeights = SrcLODData.GetSkinWeightVertexBuffer()->GetSkinWeightPtr<bHasExtraBoneInfluences>(SourceVertIdx);
-
-	// if source doesn't have extra influence, we have to clear the buffer
-	FMemory::Memzero(DestWeight.InfluenceBones);
-	FMemory::Memzero(DestWeight.InfluenceWeights);
-
-	FMemory::Memcpy(DestWeight.InfluenceBones, SrcSkinWeights->InfluenceBones, sizeof(SrcSkinWeights->InfluenceBones));
-	FMemory::Memcpy(DestWeight.InfluenceWeights, SrcSkinWeights->InfluenceWeights, sizeof(SrcSkinWeights->InfluenceWeights));
-}
-
 /**
 * Creates a new LOD model and adds the new merged sections to it. Modifies the MergedMesh.
 * @param LODIdx - current LOD to process
 */
-template<typename VertexDataType, typename SkinWeightType>
+template<typename VertexDataType>
 void FSkeletalMeshMerge::GenerateLODModel( int32 LODIdx )
 {
 	// add the new LOD model entry
@@ -454,7 +430,7 @@ void FSkeletalMeshMerge::GenerateLODModel( int32 LODIdx )
 	// merged vertex buffer
 	TArray< VertexDataType > MergedVertexBuffer;
 	// merged skin weight buffer
-	TArray< SkinWeightType > MergedSkinWeightBuffer;
+	TArray< FSkinWeightInfo > MergedSkinWeightBuffer;
 	// merged vertex color buffer
 	TArray< FColor > MergedColorBuffer;
 	// merged index buffer
@@ -463,8 +439,8 @@ void FSkeletalMeshMerge::GenerateLODModel( int32 LODIdx )
 	// The total number of UV sets for this LOD model
 	uint32 TotalNumUVs = 0;
 
-	// true if any extra bone influence exists
-	bool bSourceHasExtraBoneInfluences = false;
+	uint32 SourceMaxBoneInfluences = 0;
+	bool bSourceUse16BitBoneIndex = false;
 
 	for( int32 CreateIdx=0; CreateIdx < NewSectionArray.Num(); CreateIdx++ )
 	{
@@ -600,24 +576,19 @@ void FSkeletalMeshMerge::GenerateLODModel( int32 LODIdx )
 			// keep track of the current base vertex index before adding any new vertices
 			// this will be needed to remap the index buffer values to the new range
 			int32 CurrentBaseVertexIndex = MergedVertexBuffer.Num();
-			const bool bSourceExtraBoneInfluence = SrcLODData.GetSkinWeightVertexBuffer()->HasExtraBoneInfluences();
+			const uint32 MaxBoneInfluences = SrcLODData.GetSkinWeightVertexBuffer()->GetMaxBoneInfluences();
+			const bool bUse16BitBoneIndex = SrcLODData.GetSkinWeightVertexBuffer()->Use16BitBoneIndex();
 			for( int32 VertIdx=MergeSectionInfo.Section->BaseVertexIndex; VertIdx < MaxVertIdx; VertIdx++ )
 			{
 				// add the new vertex
 				VertexDataType& DestVert = MergedVertexBuffer[MergedVertexBuffer.AddUninitialized()];
-				SkinWeightType& DestWeight = MergedSkinWeightBuffer[MergedSkinWeightBuffer.AddUninitialized()];
+				FSkinWeightInfo& DestWeight = MergedSkinWeightBuffer[MergedSkinWeightBuffer.AddUninitialized()];
 
 				CopyVertexFromSource<VertexDataType>(DestVert, SrcLODData, VertIdx, MergeSectionInfo);
 
-				bSourceHasExtraBoneInfluences |= bSourceExtraBoneInfluence;
-				if (bSourceExtraBoneInfluence)
-				{
-					CopyWeightFromSource<SkinWeightType, true>(DestWeight, SrcLODData, VertIdx, MergeSectionInfo);
-				}
-				else
-				{
-					CopyWeightFromSource<SkinWeightType, false>(DestWeight, SrcLODData, VertIdx, MergeSectionInfo);
-				}
+				SourceMaxBoneInfluences = FMath::Max(SourceMaxBoneInfluences, MaxBoneInfluences);
+				bSourceUse16BitBoneIndex |= bUse16BitBoneIndex;
+				DestWeight = SrcLODData.GetSkinWeightVertexBuffer()->GetVertexSkinWeights(VertIdx);
 
 				// if the mesh uses vertex colors, copy the source color if possible or default to white
 				if( MergeMesh->bHasVertexColors )
@@ -641,7 +612,7 @@ void FSkeletalMeshMerge::GenerateLODModel( int32 LODIdx )
 				}
 
 				// remap the bone index used by this vertex to match the mergedbonemap 
-				for( int32 Idx=0; Idx < SkinWeightType::NumInfluences; Idx++ )
+				for( uint32 Idx=0; Idx < MAX_TOTAL_INFLUENCES; Idx++ )
 				{
 					if (DestWeight.InfluenceWeights[Idx] > 0)
 					{
@@ -749,7 +720,8 @@ void FSkeletalMeshMerge::GenerateLODModel( int32 LODIdx )
 		}
 	}
 
-	MergeLODData.SkinWeightVertexBuffer.SetHasExtraBoneInfluences(bSourceHasExtraBoneInfluences);
+	MergeLODData.SkinWeightVertexBuffer.SetMaxBoneInfluences(SourceMaxBoneInfluences);
+	MergeLODData.SkinWeightVertexBuffer.SetUse16BitBoneIndex(bSourceUse16BitBoneIndex);
 	MergeLODData.SkinWeightVertexBuffer.SetNeedsCPUAccess(bNeedsCPUAccess);
 
 	// copy vertex resource arrays
diff --git a/Engine/Source/Runtime/Engine/Private/SkeletalMeshRenderData.cpp b/Engine/Source/Runtime/Engine/Private/SkeletalMeshRenderData.cpp
--- a/Engine/Source/Runtime/Engine/Private/SkeletalMeshRenderData.cpp
+++ b/Engine/Source/Runtime/Engine/Private/SkeletalMeshRenderData.cpp
@@ -59,7 +59,7 @@ static void SerializeLODInfoForDDC(USkeletalMesh* SkeletalMesh, FString& KeySuff
 // differences, etc.) replace the version GUID below with a new one.
 // In case of merge conflicts with DDC versions, you MUST generate a new GUID
 // and set this new GUID as the version.
-#define SKELETALMESH_DERIVEDDATA_VER TEXT("7208E3E18CB047589DD75CD8E66EA1D6")
+#define SKELETALMESH_DERIVEDDATA_VER TEXT("BAD8180EFB1543D79AA1D10F3F5945A2")
 
 static const FString& GetSkeletalMeshDerivedDataVersion()
 {
@@ -447,18 +447,16 @@ void FSkeletalMeshRenderData::ReleaseResources()
 	}
 }
 
-bool FSkeletalMeshRenderData::HasExtraBoneInfluences() const
+uint32 FSkeletalMeshRenderData::GetNumBoneInfluences() const
 {
+	uint32 NumBoneInfluences = 0;
 	for (int32 LODIndex = 0; LODIndex < LODRenderData.Num(); ++LODIndex)
 	{
 		const FSkeletalMeshLODRenderData& Data = LODRenderData[LODIndex];
-		if (Data.DoesVertexBufferHaveExtraBoneInfluences())
-		{
-			return true;
-		}
+		NumBoneInfluences = FMath::Max(NumBoneInfluences, Data.GetVertexBufferMaxBoneInfluences());
 	}
 
-	return false;
+	return NumBoneInfluences;
 }
 
 bool FSkeletalMeshRenderData::RequiresCPUSkinning(ERHIFeatureLevel::Type FeatureLevel) const
@@ -466,7 +464,7 @@ bool FSkeletalMeshRenderData::RequiresCPUSkinning(ERHIFeatureLevel::Type Feature
 	const int32 MaxGPUSkinBones = FMath::Min(GetFeatureLevelMaxNumberOfBones(FeatureLevel), FGPUBaseSkinVertexFactory::GetMaxGPUSkinBones());
 	const int32 MaxBonesPerChunk = GetMaxBonesPerSection();
 	// Do CPU skinning if we need too many bones per chunk, or if we have too many influences per vertex on lower end
-	return (MaxBonesPerChunk > MaxGPUSkinBones) || (HasExtraBoneInfluences() && FeatureLevel < ERHIFeatureLevel::ES3_1);
+	return (MaxBonesPerChunk > MaxGPUSkinBones) || (GetNumBoneInfluences() > MAX_INFLUENCES_PER_STREAM && FeatureLevel < ERHIFeatureLevel::ES3_1);
 }
 
 void FSkeletalMeshRenderData::GetResourceSizeEx(FResourceSizeEx& CumulativeResourceSize)
diff --git a/Engine/Source/Runtime/Engine/Private/SkeletalMeshSampling.cpp b/Engine/Source/Runtime/Engine/Private/SkeletalMeshSampling.cpp
--- a/Engine/Source/Runtime/Engine/Private/SkeletalMeshSampling.cpp
+++ b/Engine/Source/Runtime/Engine/Private/SkeletalMeshSampling.cpp
@@ -221,16 +221,13 @@ bool FSkeletalMeshSamplingRegion::IsMaterialAllowed(FName MaterialName)
 
 #if WITH_EDITORONLY_DATA
 
-template<bool bExtraBoneInfluencesT>
 FORCEINLINE_DEBUGGABLE void GetRegionValidData(USkeletalMesh* SkeletalMesh, FSkeletalMeshSamplingRegion& SamplingRegion,
 	FSkeletalMeshLODRenderData& LODData, FReferenceSkeleton& RefSkel, FSkinWeightVertexBuffer* SkinWeightBuffer,
 	TArray<int32>& IncludedBoneIndices, TArray<int32>& ExcludedBoneIndices, TArray<int32>& OutValidTris, TArray<int32>& OutValidVerts)
 {
-	int32 MaxBoneInfluences = bExtraBoneInfluencesT ? MAX_TOTAL_INFLUENCES : MAX_INFLUENCES_PER_STREAM;
+	int32 MaxBoneInfluences = SkinWeightBuffer->GetMaxBoneInfluences();
 	auto VertIsValid = [&](FSkelMeshRenderSection& Section, int32 VertexIdx)
 	{
-		const TSkinWeightInfo<bExtraBoneInfluencesT>* SrcSkinWeights = SkinWeightBuffer->GetSkinWeightPtr<bExtraBoneInfluencesT>(VertexIdx);
-
 #if !PLATFORM_LITTLE_ENDIAN
 		// uint8[] elements in LOD.VertexBufferGPUSkin have been swapped for VET_UBYTE4 vertex stream use
 		for (int32 InfluenceIndex = MAX_INFLUENCES - 1; InfluenceIndex >= MAX_INFLUENCES - Section.MaxBoneInfluences; InfluenceIndex--)
@@ -238,7 +235,7 @@ FORCEINLINE_DEBUGGABLE void GetRegionValidData(USkeletalMesh* SkeletalMesh, FSke
 		for (int32 InfluenceIndex = 0; InfluenceIndex < MaxBoneInfluences; InfluenceIndex++)
 #endif
 		{
-			int32 BoneIndex = Section.BoneMap[SrcSkinWeights->InfluenceBones[InfluenceIndex]];
+			int32 BoneIndex = Section.BoneMap[SkinWeightBuffer->GetBoneIndex(VertexIdx, InfluenceIndex)];
 
 			if (IncludedBoneIndices.Contains(BoneIndex) && !ExcludedBoneIndices.Contains(BoneIndex))
 			{
@@ -329,14 +326,7 @@ void FSkeletalMeshSamplingInfo::BuildRegions(USkeletalMesh* SkeletalMesh)
 
 		RegionBuiltData.BoneIndices = IncludedBoneIndices;
 
-		if (SkinWeightBuffer->HasExtraBoneInfluences())
-		{
-			GetRegionValidData<true>(SkeletalMesh, Region, LODData, RefSkel, SkinWeightBuffer, IncludedBoneIndices, ExcludedBoneIndices, RegionBuiltData.TriangleIndices, RegionBuiltData.Vertices);
-		}
-		else
-		{
-			GetRegionValidData<false>(SkeletalMesh, Region, LODData, RefSkel, SkinWeightBuffer, IncludedBoneIndices, ExcludedBoneIndices, RegionBuiltData.TriangleIndices, RegionBuiltData.Vertices);
-		}
+		GetRegionValidData(SkeletalMesh, Region, LODData, RefSkel, SkinWeightBuffer, IncludedBoneIndices, ExcludedBoneIndices, RegionBuiltData.TriangleIndices, RegionBuiltData.Vertices);
 
 		if (Region.bSupportUniformlyDistributedSampling)
 		{
diff --git a/Engine/Source/Runtime/Engine/Private/SkeletalRenderCPUSkin.cpp b/Engine/Source/Runtime/Engine/Private/SkeletalRenderCPUSkin.cpp
--- a/Engine/Source/Runtime/Engine/Private/SkeletalRenderCPUSkin.cpp
+++ b/Engine/Source/Runtime/Engine/Private/SkeletalRenderCPUSkin.cpp
@@ -11,6 +11,8 @@
 #include "SkeletalRender.h"
 #include "Animation/MorphTarget.h"
 #include "GPUSkinVertexFactory.h"
+#include "Rendering/SkinWeightVertexBuffer.h"
+#include "BoneIndices.h"
 
 struct FMorphTargetDelta;
 
@@ -25,6 +27,10 @@ static void SkinVertices(FFinalSkinVertex* DestVertex, FMatrix* ReferenceToLocal
 #define INFLUENCE_5		5
 #define INFLUENCE_6		6
 #define INFLUENCE_7		7
+#define INFLUENCE_8		8
+#define INFLUENCE_9		9
+#define INFLUENCE_10	10
+#define INFLUENCE_11	11
 
 /** 
  * Modify the vertex buffer to store bone weights in the UV coordinates for rendering 
@@ -302,13 +308,13 @@ void FSkeletalMeshObjectCPUSkin::FSkeletalMeshObjectLOD::InitResources(FSkelMesh
 		if (CompLODInfo->OverrideSkinWeights &&
 			CompLODInfo->OverrideSkinWeights->GetNumVertices() == LODData.StaticVertexBuffers.PositionVertexBuffer.GetNumVertices())
 		{
-			check(LODData.SkinWeightVertexBuffer.HasExtraBoneInfluences() == CompLODInfo->OverrideSkinWeights->HasExtraBoneInfluences());
+			check(LODData.SkinWeightVertexBuffer.GetMaxBoneInfluences() == CompLODInfo->OverrideSkinWeights->GetMaxBoneInfluences());
 			MeshObjectWeightBuffer = CompLODInfo->OverrideSkinWeights;
 		}
 		else if (CompLODInfo->OverrideProfileSkinWeights &&
 			CompLODInfo->OverrideProfileSkinWeights->GetNumVertices() == LODData.StaticVertexBuffers.PositionVertexBuffer.GetNumVertices())
 		{
-			check(LODData.SkinWeightVertexBuffer.HasExtraBoneInfluences() == CompLODInfo->OverrideProfileSkinWeights->HasExtraBoneInfluences());
+			check(LODData.SkinWeightVertexBuffer.GetMaxBoneInfluences() == CompLODInfo->OverrideProfileSkinWeights->GetMaxBoneInfluences());
 			MeshObjectWeightBuffer = CompLODInfo->OverrideProfileSkinWeights;
 		}
 		else
@@ -424,13 +430,13 @@ void FSkeletalMeshObjectCPUSkin::FSkeletalMeshObjectLOD::UpdateSkinWeights(FSkel
 		if (CompLODInfo->OverrideSkinWeights &&
 			CompLODInfo->OverrideSkinWeights->GetNumVertices() == LODData.StaticVertexBuffers.PositionVertexBuffer.GetNumVertices())
 		{
-			check(LODData.SkinWeightVertexBuffer.HasExtraBoneInfluences() == CompLODInfo->OverrideSkinWeights->HasExtraBoneInfluences());
+			check(LODData.SkinWeightVertexBuffer.GetMaxBoneInfluences() == CompLODInfo->OverrideSkinWeights->GetMaxBoneInfluences());
 			NewMeshObjectWeightBuffer = CompLODInfo->OverrideSkinWeights;
 		}
 		else if (CompLODInfo->OverrideProfileSkinWeights &&
 			CompLODInfo->OverrideProfileSkinWeights->GetNumVertices() == LODData.StaticVertexBuffers.PositionVertexBuffer.GetNumVertices())
 		{
-			check(LODData.SkinWeightVertexBuffer.HasExtraBoneInfluences() == CompLODInfo->OverrideProfileSkinWeights->HasExtraBoneInfluences());
+			check(LODData.SkinWeightVertexBuffer.GetMaxBoneInfluences() == CompLODInfo->OverrideProfileSkinWeights->GetMaxBoneInfluences());
 			NewMeshObjectWeightBuffer = CompLODInfo->OverrideProfileSkinWeights;
 		}
 		else
@@ -682,7 +688,7 @@ const VectorRegister		VECTOR_0001				= DECLARE_VECTOR_REGISTER(0.f, 0.f, 0.f, 1.
 
 #define FIXED_VERTEX_INDEX 0xFFFF
 
-template<bool bExtraBoneInfluences, int32 MaxSectionBoneInfluences, typename VertexType>
+template<typename VertexType>
 static void SkinVertexSection(
 	FFinalSkinVertex*& DestVertex,
 	TArray<FMorphTargetInfo>& MorphEvalInfos,
@@ -708,6 +714,7 @@ static void SkinVertexSection(
 	FPlatformMisc::Prefetch( BoneMap );
 	FPlatformMisc::Prefetch( BoneMap, PLATFORM_CACHE_LINE_SIZE );
 
+	const int32 MaxSectionBoneInfluences = WeightBuffer.GetMaxBoneInfluences();
 	const bool bLODUsesCloth = LOD.HasClothData() && ClothSimData != nullptr && ClothBlendWeight > 0.0f;
 	const int32 NumSoftVertices = Section.GetNumVertices();
 	if (NumSoftVertices > 0)
@@ -731,7 +738,7 @@ static void SkinVertexSection(
 
 			VertexType* MorphedVertex = &SrcSoftVertex;
 
-			const TSkinWeightInfo<bExtraBoneInfluences>* SrcWeights = WeightBuffer.GetSkinWeightPtr<bExtraBoneInfluences>(VertexBufferIndex);
+			FSkinWeightInfo SrcWeights = WeightBuffer.GetVertexSkinWeights(VertexBufferIndex);
 
 			if( NumValidMorphs ) 
 			{
@@ -746,8 +753,8 @@ static void SkinVertexSection(
 				FPlatformMisc::Prefetch(ClothVertData, PLATFORM_CACHE_LINE_SIZE);	// Prefetch next cloth vertex
 			}
 
-			const uint8* RESTRICT BoneIndices = SrcWeights->InfluenceBones;
-			const uint8* RESTRICT BoneWeights = SrcWeights->InfluenceWeights;
+			const FBoneIndexType* RESTRICT BoneIndices = SrcWeights.InfluenceBones;
+			const uint8* RESTRICT BoneWeights = SrcWeights.InfluenceWeights;
 
 			static VectorRegister	SrcNormals[3];
 			VectorRegister			DstNormals[3];
@@ -755,11 +762,16 @@ static void SkinVertexSection(
 			SrcNormals[1] = Unpack3( &MorphedVertex->TangentX.Vector.Packed );
 			SrcNormals[2] = Unpack4( &MorphedVertex->TangentZ.Vector.Packed );
 			VectorRegister Weights = VectorMultiply( VectorLoadByte4(BoneWeights), VECTOR_INV_255 );
-			VectorRegister ExtraWeights;
+			VectorRegister ExtraWeights = MakeVectorRegister(0.f, 0.f, 0.f, 0.f);
+			VectorRegister ExtraWeights2 = MakeVectorRegister(0.f, 0.f, 0.f, 0.f);
 			if (MaxSectionBoneInfluences > 4)
 			{
 				ExtraWeights = VectorMultiply( VectorLoadByte4(&BoneWeights[MAX_INFLUENCES_PER_STREAM]), VECTOR_INV_255 );
 			}
+			if (MaxSectionBoneInfluences > 8)
+			{
+				ExtraWeights2 = VectorMultiply(VectorLoadByte4(&BoneWeights[EXTRA_BONE_INFLUENCES]), VECTOR_INV_255);
+			}
 			VectorResetFloatRegisters(); // Need to call this to be able to use regular floating point registers again after Unpack and VectorLoadByte4.
 
 			const FMatrix BoneMatrix0 = ReferenceToLocal[BoneMap[BoneIndices[INFLUENCE_0]]];
@@ -832,6 +844,46 @@ static void SkinVertexSection(
 									M10	= VectorMultiplyAdd( VectorLoadAligned( &BoneMatrix7.M[1][0] ), Weight7, M10 );
 									M20	= VectorMultiplyAdd( VectorLoadAligned( &BoneMatrix7.M[2][0] ), Weight7, M20 );
 									M30	= VectorMultiplyAdd( VectorLoadAligned( &BoneMatrix7.M[3][0] ), Weight7, M30 );
+
+									if (MaxSectionBoneInfluences > 8)
+									{
+										const FMatrix BoneMatrix8 = ReferenceToLocal[BoneMap[BoneIndices[INFLUENCE_8]]];
+										VectorRegister Weight8 = VectorReplicate( ExtraWeights2, INFLUENCE_8 - INFLUENCE_8 );
+										M00 = VectorMultiplyAdd(VectorLoadAligned( &BoneMatrix8.M[0][0]), Weight8, M00 );
+										M10 = VectorMultiplyAdd(VectorLoadAligned( &BoneMatrix8.M[1][0]), Weight8, M10 );
+										M20 = VectorMultiplyAdd(VectorLoadAligned( &BoneMatrix8.M[2][0]), Weight8, M20 );
+										M30 = VectorMultiplyAdd(VectorLoadAligned( &BoneMatrix8.M[3][0]), Weight8, M30 );
+
+										if (MaxSectionBoneInfluences > 9)
+										{
+											const FMatrix BoneMatrix9 = ReferenceToLocal[BoneMap[BoneIndices[INFLUENCE_9]]];
+											VectorRegister Weight9 = VectorReplicate(ExtraWeights2, INFLUENCE_9 - INFLUENCE_8);
+											M00 = VectorMultiplyAdd(VectorLoadAligned(&BoneMatrix9.M[0][0]), Weight9, M00);
+											M10 = VectorMultiplyAdd(VectorLoadAligned(&BoneMatrix9.M[1][0]), Weight9, M10);
+											M20 = VectorMultiplyAdd(VectorLoadAligned(&BoneMatrix9.M[2][0]), Weight9, M20);
+											M30 = VectorMultiplyAdd(VectorLoadAligned(&BoneMatrix9.M[3][0]), Weight9, M30);
+
+											if (MaxSectionBoneInfluences > 10)
+											{
+												const FMatrix BoneMatrix10 = ReferenceToLocal[BoneMap[BoneIndices[INFLUENCE_10]]];
+												VectorRegister Weight10 = VectorReplicate(ExtraWeights2, INFLUENCE_10 - INFLUENCE_8);
+												M00 = VectorMultiplyAdd(VectorLoadAligned(&BoneMatrix10.M[0][0]), Weight10, M00);
+												M10 = VectorMultiplyAdd(VectorLoadAligned(&BoneMatrix10.M[1][0]), Weight10, M10);
+												M20 = VectorMultiplyAdd(VectorLoadAligned(&BoneMatrix10.M[2][0]), Weight10, M20);
+												M30 = VectorMultiplyAdd(VectorLoadAligned(&BoneMatrix10.M[3][0]), Weight10, M30);
+
+												if (MaxSectionBoneInfluences > 11)
+												{
+													const FMatrix BoneMatrix11 = ReferenceToLocal[BoneMap[BoneIndices[INFLUENCE_11]]];
+													VectorRegister Weight11 = VectorReplicate(ExtraWeights2, INFLUENCE_11 - INFLUENCE_8);
+													M00 = VectorMultiplyAdd(VectorLoadAligned(&BoneMatrix11.M[0][0]), Weight11, M00);
+													M10 = VectorMultiplyAdd(VectorLoadAligned(&BoneMatrix11.M[1][0]), Weight11, M10);
+													M20 = VectorMultiplyAdd(VectorLoadAligned(&BoneMatrix11.M[2][0]), Weight11, M20);
+													M30 = VectorMultiplyAdd(VectorLoadAligned(&BoneMatrix11.M[3][0]), Weight11, M30);
+												}
+											}
+										}
+									}
 								}
 							}
 						}
@@ -956,37 +1008,6 @@ static void SkinVertexSection(
 	}
 }
 
-template< bool bExtraBoneInfluences, typename VertexType>
-static void SkinVertexSection( 
-	FFinalSkinVertex*& DestVertex, 
-	TArray<FMorphTargetInfo>& MorphEvalInfos, 
-	const TArray<float>& MorphWeights, 
-	const FSkelMeshRenderSection& Section, 
-	const FSkeletalMeshLODRenderData &LOD,
-	FSkinWeightVertexBuffer& WeightBuffer,
-	int32 VertexBufferBaseIndex, 
-	uint32 NumValidMorphs, 
-	int32 &CurBaseVertIdx, 
-	int32 LODIndex, 
-	const FMatrix* RESTRICT ReferenceToLocal, 
-	const FClothSimulData* ClothSimData, 
-	float ClothBlendWeight, 
-	const FMatrix& WorldToLocal)
-{
-	switch (Section.MaxBoneInfluences)
-	{
-		case 1: SkinVertexSection<bExtraBoneInfluences, 1, VertexType>(DestVertex, MorphEvalInfos, MorphWeights, Section, LOD, WeightBuffer, VertexBufferBaseIndex, NumValidMorphs, CurBaseVertIdx, LODIndex, ReferenceToLocal, ClothSimData, ClothBlendWeight, WorldToLocal); break;
-		case 2: SkinVertexSection<bExtraBoneInfluences, 2, VertexType>(DestVertex, MorphEvalInfos, MorphWeights, Section, LOD, WeightBuffer, VertexBufferBaseIndex, NumValidMorphs, CurBaseVertIdx, LODIndex, ReferenceToLocal, ClothSimData, ClothBlendWeight, WorldToLocal); break;
-		case 3: SkinVertexSection<bExtraBoneInfluences, 3, VertexType>(DestVertex, MorphEvalInfos, MorphWeights, Section, LOD, WeightBuffer, VertexBufferBaseIndex, NumValidMorphs, CurBaseVertIdx, LODIndex, ReferenceToLocal, ClothSimData, ClothBlendWeight, WorldToLocal); break;
-		case 4: SkinVertexSection<bExtraBoneInfluences, 4, VertexType>(DestVertex, MorphEvalInfos, MorphWeights, Section, LOD, WeightBuffer, VertexBufferBaseIndex, NumValidMorphs, CurBaseVertIdx, LODIndex, ReferenceToLocal, ClothSimData, ClothBlendWeight, WorldToLocal); break;
-		case 5: SkinVertexSection<bExtraBoneInfluences, 5, VertexType>(DestVertex, MorphEvalInfos, MorphWeights, Section, LOD, WeightBuffer, VertexBufferBaseIndex, NumValidMorphs, CurBaseVertIdx, LODIndex, ReferenceToLocal, ClothSimData, ClothBlendWeight, WorldToLocal); break;
-		case 6: SkinVertexSection<bExtraBoneInfluences, 6, VertexType>(DestVertex, MorphEvalInfos, MorphWeights, Section, LOD, WeightBuffer, VertexBufferBaseIndex, NumValidMorphs, CurBaseVertIdx, LODIndex, ReferenceToLocal, ClothSimData, ClothBlendWeight, WorldToLocal); break;
-		case 7: SkinVertexSection<bExtraBoneInfluences, 7, VertexType>(DestVertex, MorphEvalInfos, MorphWeights, Section, LOD, WeightBuffer, VertexBufferBaseIndex, NumValidMorphs, CurBaseVertIdx, LODIndex, ReferenceToLocal, ClothSimData, ClothBlendWeight, WorldToLocal); break;
-		case 8: SkinVertexSection<bExtraBoneInfluences, 8, VertexType>(DestVertex, MorphEvalInfos, MorphWeights, Section, LOD, WeightBuffer, VertexBufferBaseIndex, NumValidMorphs, CurBaseVertIdx, LODIndex, ReferenceToLocal, ClothSimData, ClothBlendWeight, WorldToLocal); break;
-		default: check(0);
-	}
-}
-
 template<typename VertexType>
 static void SkinVertices(
 	FFinalSkinVertex* DestVertex, 
@@ -1026,14 +1047,7 @@ static void SkinVertices(
 
 		const FClothSimulData* ClothSimData = ClothSimulUpdateData.Find(Section.CorrespondClothAssetIndex);
 
-		if (LOD.DoesVertexBufferHaveExtraBoneInfluences())
-		{
-			SkinVertexSection<true, VertexType>(DestVertex, MorphEvalInfos, MorphTargetWeights, Section, LOD, WeightBuffer, VertexBufferBaseIndex, NumValidMorphs, CurBaseVertIdx, LODIndex, ReferenceToLocal, ClothSimData, ClothBlendWeight, WorldToLocal);
-		}
-		else
-		{
-			SkinVertexSection<false, VertexType>(DestVertex, MorphEvalInfos, MorphTargetWeights, Section, LOD, WeightBuffer, VertexBufferBaseIndex, NumValidMorphs, CurBaseVertIdx, LODIndex, ReferenceToLocal, ClothSimData, ClothBlendWeight, WorldToLocal);
-		}
+		SkinVertexSection<VertexType>(DestVertex, MorphEvalInfos, MorphTargetWeights, Section, LOD, WeightBuffer, VertexBufferBaseIndex, NumValidMorphs, CurBaseVertIdx, LODIndex, ReferenceToLocal, ClothSimData, ClothBlendWeight, WorldToLocal);
 	}
 
 	VectorSetControlRegister( StatusRegister );
@@ -1061,7 +1075,6 @@ FVector4 GetTangetToColor(FPackedNormal Tangent)
  * @param LOD - LOD model corresponding to DestVertex 
  * @param BonesOfInterest - array of bones we want to display
  */
-template <bool bExtraBoneInfluencesT>
 static FORCEINLINE void CalculateSectionBoneWeights(FFinalSkinVertex*& DestVertex, FSkinWeightVertexBuffer& SkinWeightVertexBuffer, FSkelMeshRenderSection& Section, const TArray<int32>& BonesOfInterest)
 {
 	const float INV255 = 1.f/255.f;
@@ -1074,16 +1087,16 @@ static FORCEINLINE void CalculateSectionBoneWeights(FFinalSkinVertex*& DestVerte
 	for(int32 VertexIndex = VertexBufferBaseIndex;VertexIndex < Section.GetNumVertices();VertexIndex++,DestVertex++)
 	{
 		const int32 VertexBufferIndex = Section.GetVertexBufferIndex() + VertexIndex;
-		auto* SrcWeight = SkinWeightVertexBuffer.GetSkinWeightPtr<bExtraBoneInfluencesT>(VertexBufferIndex);
+		FSkinWeightInfo SrcWeight = SkinWeightVertexBuffer.GetVertexSkinWeights(VertexBufferIndex);
 
 		//Zero out the UV coords
 		DestVertex->U = 0.0f;
 		DestVertex->V = 0.0f;
 
-		const uint8* RESTRICT BoneIndices = SrcWeight->InfluenceBones;
-		const uint8* RESTRICT BoneWeights = SrcWeight->InfluenceWeights;
+		const FBoneIndexType* RESTRICT BoneIndices = SrcWeight.InfluenceBones;
+		const uint8* RESTRICT BoneWeights = SrcWeight.InfluenceWeights;
 
-		for (int32 i = 0; i < TSkinWeightInfo<bExtraBoneInfluencesT>::NumInfluences; i++)
+		for (uint32 i = 0; i < SkinWeightVertexBuffer.GetMaxBoneInfluences(); i++)
 		{
 			if (BonesOfInterest.Contains(BoneMap[BoneIndices[i]]))
 			{
@@ -1110,15 +1123,7 @@ static void CalculateBoneWeights(FFinalSkinVertex* DestVertex, FSkeletalMeshLODR
 	for(int32 SectionIndex= 0;SectionIndex< LOD.RenderSections.Num();SectionIndex++)
 	{
 		FSkelMeshRenderSection& Section = LOD.RenderSections[SectionIndex];
-
-		if (WeightBuffer.HasExtraBoneInfluences())
-		{
-			CalculateSectionBoneWeights<true>(DestVertex, WeightBuffer, Section, InBonesOfInterest);
-		}
-		else
-		{
-			CalculateSectionBoneWeights<false>(DestVertex, WeightBuffer, Section, InBonesOfInterest);
-		}
+		CalculateSectionBoneWeights(DestVertex, WeightBuffer, Section, InBonesOfInterest);
 	}
 }
 
diff --git a/Engine/Source/Runtime/Engine/Private/SkeletalRenderGPUSkin.cpp b/Engine/Source/Runtime/Engine/Private/SkeletalRenderGPUSkin.cpp
--- a/Engine/Source/Runtime/Engine/Private/SkeletalRenderGPUSkin.cpp
+++ b/Engine/Source/Runtime/Engine/Private/SkeletalRenderGPUSkin.cpp
@@ -948,13 +948,13 @@ void FSkeletalMeshObjectGPUSkin::FSkeletalMeshObjectLOD::UpdateSkinWeights(FSkel
 		if (CompLODInfo->OverrideSkinWeights &&
 			CompLODInfo->OverrideSkinWeights->GetNumVertices() == LODData.StaticVertexBuffers.PositionVertexBuffer.GetNumVertices())
 		{
-			check(LODData.SkinWeightVertexBuffer.HasExtraBoneInfluences() == CompLODInfo->OverrideSkinWeights->HasExtraBoneInfluences());
+			check(LODData.SkinWeightVertexBuffer.GetMaxBoneInfluences() == CompLODInfo->OverrideSkinWeights->GetMaxBoneInfluences());
 			NewMeshObjectWeightBuffer = CompLODInfo->OverrideSkinWeights;
 		}
 		else if (CompLODInfo->OverrideProfileSkinWeights &&
 			CompLODInfo->OverrideProfileSkinWeights->GetNumVertices() == LODData.StaticVertexBuffers.PositionVertexBuffer.GetNumVertices())
 		{
-			check(LODData.SkinWeightVertexBuffer.HasExtraBoneInfluences() == CompLODInfo->OverrideProfileSkinWeights->HasExtraBoneInfluences());
+			check(LODData.SkinWeightVertexBuffer.GetMaxBoneInfluences() == CompLODInfo->OverrideProfileSkinWeights->GetMaxBoneInfluences());
 			NewMeshObjectWeightBuffer = CompLODInfo->OverrideProfileSkinWeights;
 		}
 		else
@@ -1186,10 +1186,9 @@ void FSkeletalMeshObjectGPUSkin::RefreshClothingTransforms(const FMatrix& InNewL
  */
 template<class VertexFactoryType>
 void InitGPUSkinVertexFactoryComponents(typename VertexFactoryType::FDataType* VertexFactoryData, 
-										const FSkeletalMeshObjectGPUSkin::FVertexFactoryBuffers& VertexBuffers, const VertexFactoryType* VertexFactory)
+										const FSkeletalMeshObjectGPUSkin::FVertexFactoryBuffers& VertexBuffers, VertexFactoryType* VertexFactory)
 {
 	typedef TGPUSkinVertexBase BaseVertexType;
-	typedef TSkinWeightInfo<VertexFactoryType::HasExtraBoneInfluences> WeightInfoType;
 
 	//position
 	VertexBuffers.StaticVertexBuffers->PositionVertexBuffer.BindPositionVertexBuffer(VertexFactory, *VertexFactoryData);
@@ -1198,20 +1197,37 @@ void InitGPUSkinVertexFactoryComponents(typename VertexFactoryType::FDataType* V
 	VertexBuffers.StaticVertexBuffers->StaticMeshVertexBuffer.BindTangentVertexBuffer(VertexFactory, *VertexFactoryData);
 	VertexBuffers.StaticVertexBuffers->StaticMeshVertexBuffer.BindTexCoordVertexBuffer(VertexFactory, *VertexFactoryData);
 
-	// bone indices
-	VertexFactoryData->BoneIndices = FVertexStreamComponent(
-		VertexBuffers.SkinWeightVertexBuffer,STRUCT_OFFSET(WeightInfoType,InfluenceBones),VertexBuffers.SkinWeightVertexBuffer->GetStride(),VET_UByte4);
-	// bone weights
-	VertexFactoryData->BoneWeights = FVertexStreamComponent(
-		VertexBuffers.SkinWeightVertexBuffer,STRUCT_OFFSET(WeightInfoType,InfluenceWeights),VertexBuffers.SkinWeightVertexBuffer->GetStride(),VET_UByte4N);
+	bool bUse16BitBoneIndex = VertexBuffers.SkinWeightVertexBuffer->Use16BitBoneIndex();
+	VertexFactoryData->bUse16BitBoneIndex = bUse16BitBoneIndex;
+	VertexFactoryData->NumBoneInfluences = VertexBuffers.SkinWeightVertexBuffer->GetMaxBoneInfluences();
 
-	if (VertexFactoryType::HasExtraBoneInfluences)
+	GPUSkinBoneInfluenceType BoneInfluenceType = VertexBuffers.SkinWeightVertexBuffer->GetBoneInfluenceType();
+	if (BoneInfluenceType == GPUSkinBoneInfluenceType::UnlimitedBoneInfluence)
 	{
-		// Extra streams for bone indices & weights
-		VertexFactoryData->ExtraBoneIndices = FVertexStreamComponent(
-			VertexBuffers.SkinWeightVertexBuffer,STRUCT_OFFSET(WeightInfoType,InfluenceBones) + 4,VertexBuffers.SkinWeightVertexBuffer->GetStride(),VET_UByte4);
-		VertexFactoryData->ExtraBoneWeights = FVertexStreamComponent(
-			VertexBuffers.SkinWeightVertexBuffer,STRUCT_OFFSET(WeightInfoType,InfluenceWeights) + 4,VertexBuffers.SkinWeightVertexBuffer->GetStride(),VET_UByte4N);
+		FGPUBaseSkinVertexFactory::FShaderDataType& ShaderData = VertexFactory->GetShaderData();
+		ShaderData.InputWeightIndexSize = VertexBuffers.SkinWeightVertexBuffer->GetBoneIndexByteSize();
+		ShaderData.InputWeightStream = VertexBuffers.SkinWeightVertexBuffer->GetDataVertexBuffer()->GetSRV();
+
+		const FSkinWeightLookupVertexBuffer* LookupVertexBuffer = VertexBuffers.SkinWeightVertexBuffer->GetLookupVertexBuffer();
+		VertexFactoryData->BlendOffsetCount = FVertexStreamComponent(LookupVertexBuffer, 0, LookupVertexBuffer->GetStride(), VET_UInt);
+	}
+	else
+	{
+		// bone indices & weights
+		const FSkinWeightDataVertexBuffer* WeightDataVertexBuffer = VertexBuffers.SkinWeightVertexBuffer->GetDataVertexBuffer();
+		const uint32 Stride = VertexBuffers.SkinWeightVertexBuffer->GetConstantInfluencesVertexStride();
+		const uint32 WeightsOffset = VertexBuffers.SkinWeightVertexBuffer->GetConstantInfluencesBoneWeightsOffset();
+		VertexFactoryData->BoneIndices = FVertexStreamComponent(WeightDataVertexBuffer, 0, Stride, bUse16BitBoneIndex ? VET_UShort4 : VET_UByte4);
+		VertexFactoryData->BoneWeights = FVertexStreamComponent(WeightDataVertexBuffer, WeightsOffset, Stride, VET_UByte4N);
+
+		if (BoneInfluenceType == GPUSkinBoneInfluenceType::ExtraBoneInfluence)
+		{
+			// Extra streams for bone indices & weights
+			VertexFactoryData->ExtraBoneIndices = FVertexStreamComponent(
+				WeightDataVertexBuffer, 4 * VertexBuffers.SkinWeightVertexBuffer->GetBoneIndexByteSize(), Stride, bUse16BitBoneIndex ? VET_UShort4 : VET_UByte4);
+			VertexFactoryData->ExtraBoneWeights = FVertexStreamComponent(
+				WeightDataVertexBuffer, WeightsOffset + 4, Stride, VET_UByte4N);
+		}
 	}
 
 	// Color data may be NULL
@@ -1522,15 +1538,24 @@ void FSkeletalMeshObjectGPUSkin::FVertexFactoryData::InitVertexFactories(
 	{
 		for(int32 FactoryIdx = 0; FactoryIdx < Sections.Num(); ++FactoryIdx)
 		{
-			if (VertexBuffers.SkinWeightVertexBuffer->HasExtraBoneInfluences())
+			GPUSkinBoneInfluenceType BoneInfluenceType = VertexBuffers.SkinWeightVertexBuffer->GetBoneInfluenceType();
+			if (BoneInfluenceType == GPUSkinBoneInfluenceType::DefaultBoneInfluence)
 			{
-				TGPUSkinVertexFactory<true>* VertexFactory = CreateVertexFactory< FGPUBaseSkinVertexFactory, TGPUSkinVertexFactory<true> >(VertexFactories, VertexBuffers, InFeatureLevel);
-				CreatePassthroughVertexFactory<TGPUSkinVertexFactory<true>>(InFeatureLevel, PassthroughVertexFactories, VertexFactory);
+				TGPUSkinVertexFactory<GPUSkinBoneInfluenceType::DefaultBoneInfluence>* VertexFactory = 
+					CreateVertexFactory< FGPUBaseSkinVertexFactory, TGPUSkinVertexFactory<GPUSkinBoneInfluenceType::DefaultBoneInfluence> >(VertexFactories, VertexBuffers, InFeatureLevel);
+				CreatePassthroughVertexFactory<TGPUSkinVertexFactory<GPUSkinBoneInfluenceType::DefaultBoneInfluence>>(InFeatureLevel, PassthroughVertexFactories, VertexFactory);
+			}
+			else if (BoneInfluenceType == GPUSkinBoneInfluenceType::ExtraBoneInfluence)
+			{
+				TGPUSkinVertexFactory<GPUSkinBoneInfluenceType::ExtraBoneInfluence>* VertexFactory = 
+					CreateVertexFactory< FGPUBaseSkinVertexFactory, TGPUSkinVertexFactory<GPUSkinBoneInfluenceType::ExtraBoneInfluence> >(VertexFactories, VertexBuffers, InFeatureLevel);
+				CreatePassthroughVertexFactory<TGPUSkinVertexFactory<GPUSkinBoneInfluenceType::ExtraBoneInfluence>>(InFeatureLevel, PassthroughVertexFactories, VertexFactory);
 			}
 			else
 			{
-				TGPUSkinVertexFactory<false>* VertexFactory = CreateVertexFactory< FGPUBaseSkinVertexFactory, TGPUSkinVertexFactory<false> >(VertexFactories, VertexBuffers, InFeatureLevel);
-				CreatePassthroughVertexFactory<TGPUSkinVertexFactory<false>>(InFeatureLevel, PassthroughVertexFactories, VertexFactory);
+				TGPUSkinVertexFactory<GPUSkinBoneInfluenceType::UnlimitedBoneInfluence>* VertexFactory = 
+					CreateVertexFactory< FGPUBaseSkinVertexFactory, TGPUSkinVertexFactory<GPUSkinBoneInfluenceType::UnlimitedBoneInfluence> >(VertexFactories, VertexBuffers, InFeatureLevel);
+				CreatePassthroughVertexFactory<TGPUSkinVertexFactory<GPUSkinBoneInfluenceType::UnlimitedBoneInfluence>>(InFeatureLevel, PassthroughVertexFactories, VertexFactory);
 			}
 		}
 	}
@@ -1563,13 +1588,18 @@ void FSkeletalMeshObjectGPUSkin::FVertexFactoryData::InitMorphVertexFactories(
 	MorphVertexFactories.Empty(Sections.Num());
 	for( int32 FactoryIdx=0; FactoryIdx < Sections.Num(); FactoryIdx++ )
 	{
-		if (VertexBuffers.SkinWeightVertexBuffer->HasExtraBoneInfluences())
+		GPUSkinBoneInfluenceType BoneInfluenceType = VertexBuffers.SkinWeightVertexBuffer->GetBoneInfluenceType();
+		if (BoneInfluenceType == GPUSkinBoneInfluenceType::DefaultBoneInfluence)
+		{
+			CreateVertexFactoryMorph<FGPUBaseSkinVertexFactory, TGPUSkinMorphVertexFactory< GPUSkinBoneInfluenceType::DefaultBoneInfluence > >(MorphVertexFactories,VertexBuffers,InFeatureLevel);
+		}
+		else if (BoneInfluenceType == GPUSkinBoneInfluenceType::ExtraBoneInfluence)
 		{
-			CreateVertexFactoryMorph<FGPUBaseSkinVertexFactory, TGPUSkinMorphVertexFactory<true> >(MorphVertexFactories,VertexBuffers,InFeatureLevel);
+			CreateVertexFactoryMorph<FGPUBaseSkinVertexFactory, TGPUSkinMorphVertexFactory<GPUSkinBoneInfluenceType::ExtraBoneInfluence> >(MorphVertexFactories, VertexBuffers,InFeatureLevel);
 		}
 		else
 		{
-			CreateVertexFactoryMorph<FGPUBaseSkinVertexFactory, TGPUSkinMorphVertexFactory<false> >(MorphVertexFactories, VertexBuffers,InFeatureLevel);
+			CreateVertexFactoryMorph<FGPUBaseSkinVertexFactory, TGPUSkinMorphVertexFactory<GPUSkinBoneInfluenceType::UnlimitedBoneInfluence> >(MorphVertexFactories, VertexBuffers, InFeatureLevel);
 		}
 	}
 }
@@ -1597,13 +1627,18 @@ void FSkeletalMeshObjectGPUSkin::FVertexFactoryData::InitAPEXClothVertexFactorie
 	{
 		if (Sections[FactoryIdx].HasClothingData() && InFeatureLevel >= ERHIFeatureLevel::SM5)
 		{
-			if (VertexBuffers.SkinWeightVertexBuffer->HasExtraBoneInfluences())
+			GPUSkinBoneInfluenceType BoneInfluenceType = VertexBuffers.SkinWeightVertexBuffer->GetBoneInfluenceType();
+			if (BoneInfluenceType == GPUSkinBoneInfluenceType::DefaultBoneInfluence)
 			{
-				CreateVertexFactoryCloth<FGPUBaseSkinAPEXClothVertexFactory, TGPUSkinAPEXClothVertexFactory<true> >(ClothVertexFactories, VertexBuffers, InFeatureLevel);
+				CreateVertexFactoryCloth<FGPUBaseSkinAPEXClothVertexFactory, TGPUSkinAPEXClothVertexFactory<GPUSkinBoneInfluenceType::DefaultBoneInfluence> >(ClothVertexFactories, VertexBuffers, InFeatureLevel);
+			}
+			else if (BoneInfluenceType == GPUSkinBoneInfluenceType::ExtraBoneInfluence)
+			{
+				CreateVertexFactoryCloth<FGPUBaseSkinAPEXClothVertexFactory, TGPUSkinAPEXClothVertexFactory<GPUSkinBoneInfluenceType::ExtraBoneInfluence> >(ClothVertexFactories, VertexBuffers, InFeatureLevel);
 			}
 			else
 			{
-				CreateVertexFactoryCloth<FGPUBaseSkinAPEXClothVertexFactory, TGPUSkinAPEXClothVertexFactory<false> >(ClothVertexFactories, VertexBuffers, InFeatureLevel);
+				CreateVertexFactoryCloth<FGPUBaseSkinAPEXClothVertexFactory, TGPUSkinAPEXClothVertexFactory<GPUSkinBoneInfluenceType::UnlimitedBoneInfluence> >(ClothVertexFactories, VertexBuffers, InFeatureLevel);
 			}
 		}
 		else
@@ -1631,17 +1666,24 @@ void FSkeletalMeshObjectGPUSkin::FVertexFactoryData::ReleaseAPEXClothVertexFacto
 
 void FSkeletalMeshObjectGPUSkin::FVertexFactoryData::UpdateVertexFactoryData(const FVertexFactoryBuffers& VertexBuffers)
 {
-	if (VertexBuffers.SkinWeightVertexBuffer->HasExtraBoneInfluences())
+	GPUSkinBoneInfluenceType BoneInfluenceType = VertexBuffers.SkinWeightVertexBuffer->GetBoneInfluenceType();
+	if (BoneInfluenceType == GPUSkinBoneInfluenceType::DefaultBoneInfluence)
+	{
+		UpdateVertexFactory<FGPUBaseSkinVertexFactory, TGPUSkinVertexFactory<GPUSkinBoneInfluenceType::DefaultBoneInfluence>>(VertexFactories, VertexBuffers);
+		UpdateVertexFactoryCloth<FGPUBaseSkinAPEXClothVertexFactory, TGPUSkinAPEXClothVertexFactory<GPUSkinBoneInfluenceType::DefaultBoneInfluence>>(ClothVertexFactories, VertexBuffers);
+		UpdateVertexFactoryMorph<FGPUBaseSkinVertexFactory, TGPUSkinMorphVertexFactory<GPUSkinBoneInfluenceType::DefaultBoneInfluence>>(MorphVertexFactories, VertexBuffers);
+	}
+	else if (BoneInfluenceType == GPUSkinBoneInfluenceType::ExtraBoneInfluence)
 	{
-		UpdateVertexFactory<FGPUBaseSkinVertexFactory, TGPUSkinVertexFactory<true>>(VertexFactories, VertexBuffers);
-		UpdateVertexFactoryCloth<FGPUBaseSkinAPEXClothVertexFactory, TGPUSkinAPEXClothVertexFactory<true>>(ClothVertexFactories, VertexBuffers);
-		UpdateVertexFactoryMorph<FGPUBaseSkinVertexFactory, TGPUSkinMorphVertexFactory<true>>(MorphVertexFactories, VertexBuffers);
+		UpdateVertexFactory<FGPUBaseSkinVertexFactory, TGPUSkinVertexFactory<GPUSkinBoneInfluenceType::ExtraBoneInfluence>>(VertexFactories, VertexBuffers);
+		UpdateVertexFactoryCloth<FGPUBaseSkinAPEXClothVertexFactory, TGPUSkinAPEXClothVertexFactory<GPUSkinBoneInfluenceType::ExtraBoneInfluence>>(ClothVertexFactories, VertexBuffers);
+		UpdateVertexFactoryMorph<FGPUBaseSkinVertexFactory, TGPUSkinMorphVertexFactory<GPUSkinBoneInfluenceType::ExtraBoneInfluence>>(MorphVertexFactories, VertexBuffers);
 	}
 	else
 	{
-		UpdateVertexFactory<FGPUBaseSkinVertexFactory, TGPUSkinVertexFactory<false>>(VertexFactories, VertexBuffers);
-		UpdateVertexFactoryCloth<FGPUBaseSkinAPEXClothVertexFactory, TGPUSkinAPEXClothVertexFactory<false>>(ClothVertexFactories, VertexBuffers);
-		UpdateVertexFactoryMorph<FGPUBaseSkinVertexFactory, TGPUSkinMorphVertexFactory<false>>(MorphVertexFactories, VertexBuffers);
+		UpdateVertexFactory<FGPUBaseSkinVertexFactory, TGPUSkinVertexFactory<GPUSkinBoneInfluenceType::UnlimitedBoneInfluence>>(VertexFactories, VertexBuffers);
+		UpdateVertexFactoryCloth<FGPUBaseSkinAPEXClothVertexFactory, TGPUSkinAPEXClothVertexFactory<GPUSkinBoneInfluenceType::UnlimitedBoneInfluence>>(ClothVertexFactories, VertexBuffers);
+		UpdateVertexFactoryMorph<FGPUBaseSkinVertexFactory, TGPUSkinMorphVertexFactory<GPUSkinBoneInfluenceType::UnlimitedBoneInfluence>>(MorphVertexFactories, VertexBuffers);
 	}
 }
 
@@ -1658,13 +1700,13 @@ void FSkeletalMeshObjectGPUSkin::FSkeletalMeshObjectLOD::InitResources(const FSk
 		CompLODInfo->OverrideSkinWeights &&
 		CompLODInfo->OverrideSkinWeights->GetNumVertices() == LODData.StaticVertexBuffers.PositionVertexBuffer.GetNumVertices())
 	{
-		check(LODData.SkinWeightVertexBuffer.HasExtraBoneInfluences() == CompLODInfo->OverrideSkinWeights->HasExtraBoneInfluences());
+		check(LODData.SkinWeightVertexBuffer.GetMaxBoneInfluences() == CompLODInfo->OverrideSkinWeights->GetMaxBoneInfluences());
 		MeshObjectWeightBuffer = CompLODInfo->OverrideSkinWeights;
 	}
 	else if (CompLODInfo && CompLODInfo->OverrideProfileSkinWeights &&
 			CompLODInfo->OverrideProfileSkinWeights->GetNumVertices() == LODData.StaticVertexBuffers.PositionVertexBuffer.GetNumVertices())
 	{
-		check(LODData.SkinWeightVertexBuffer.HasExtraBoneInfluences() == CompLODInfo->OverrideProfileSkinWeights->HasExtraBoneInfluences());
+		check(LODData.SkinWeightVertexBuffer.GetMaxBoneInfluences() == CompLODInfo->OverrideProfileSkinWeights->GetMaxBoneInfluences());
 		MeshObjectWeightBuffer = CompLODInfo->OverrideProfileSkinWeights;
 	}
 	else
diff --git a/Engine/Source/Runtime/Engine/Private/Streaming/SkeletalMeshUpdate.h b/Engine/Source/Runtime/Engine/Private/Streaming/SkeletalMeshUpdate.h
--- a/Engine/Source/Runtime/Engine/Private/Streaming/SkeletalMeshUpdate.h
+++ b/Engine/Source/Runtime/Engine/Private/Streaming/SkeletalMeshUpdate.h
@@ -85,11 +85,11 @@ protected:
 		FVertexBufferRHIRef TexCoordVertexBuffer;
 		FVertexBufferRHIRef PositionVertexBuffer;
 		FVertexBufferRHIRef ColorVertexBuffer;
-		FVertexBufferRHIRef SkinWeightVertexBuffer;
+		FSkinWeightRHIInfo SkinWeightVertexBuffer;
 		FVertexBufferRHIRef ClothVertexBuffer;
 		FIndexBufferRHIRef IndexBuffer;
 		FIndexBufferRHIRef AdjacencyIndexBuffer;
-		TArray<TPair<FName, FVertexBufferRHIRef>> AltSkinWeightVertexBuffers;
+		TArray<TPair<FName, FSkinWeightRHIInfo>> AltSkinWeightVertexBuffers;
 
 		void CreateFromCPUData_RenderThread(USkeletalMesh* Mesh, FSkeletalMeshLODRenderData& LODResource);
 		void CreateFromCPUData_Async(USkeletalMesh* Mesh, FSkeletalMeshLODRenderData& LODResource);
diff --git a/Engine/Source/Runtime/Engine/Public/AnimEncoding.h b/Engine/Source/Runtime/Engine/Public/AnimEncoding.h
--- a/Engine/Source/Runtime/Engine/Public/AnimEncoding.h
+++ b/Engine/Source/Runtime/Engine/Public/AnimEncoding.h
@@ -61,149 +61,6 @@ typedef TArray<BoneTrackPair> BoneTrackArray;
 /** Array of FTransform using the game memory stack */
 typedef TArray< FTransform, TMemStackAllocator<> > FTransformArray;
 
-/**
- * Structure to wrap trivial track flags for easier and safer handling.
- */
-struct FTrivialTrackFlags
-{
-	constexpr explicit FTrivialTrackFlags(uint8 InFlags) : Flags(InFlags) {}
-
-	constexpr bool IsTranslationTrivial() const { return (Flags & 0x4) != 0; }
-	constexpr bool IsRotationTrivial() const { return (Flags & 0x2) != 0; }
-	constexpr bool IsScaleTrivial() const { return (Flags & 0x1) != 0; }
-
-	uint8 Flags;
-};
-
-/**
- * Structure that represents a sorted key header.
- *
- * Our time delta can be positive or negative. To make packing easier, we offset it in order to always make it positive
- * A small header has 5 bits to store the time delta. 5 bits == 0x1F == 31 == [0 .. 31]
- * A large header has 5 + 8 = 13 bits to store the time delta. 13 bits == 0x1FFF == 8191 == [0 .. 8191]
- * A small header has an offset of 16. e.g.:
- *     -5 + 16 = 11
- *     15 + 16 = 31
- *    -16 + 16 = 0
- * A small header this has a signed range of [-16 .. 15] and maps to [0 .. 31]
- * A large header has an offset of 4096. The signed range is [-4096 .. 4095] and maps to [0 ... 8191]
- */
-struct FSortedKeyHeader
-{
-	static constexpr uint8 KeyTypeMask = 0x3;
-	static constexpr uint8 KeyTypeShift = 5;
-	static constexpr uint8 KeyTimeDeltaMask = (1 << 5) - 1;
-	static constexpr uint8 KeyTimeDeltaShift = 8;
-	static constexpr uint8 KeyHeaderSizeMask = 0x80;
-
-	static constexpr uint32 LargestSmallHeaderTimeDelta = (1 << 5) - 1;
-	static constexpr uint32 LargestLargeHeaderTimeDelta = (1 << 13) - 1;
-	static constexpr int32 SmallHeaderTimeDeltaOffset = (LargestSmallHeaderTimeDelta + 1) / 2;
-	static constexpr int32 LargeHeaderTimeDeltaOffset = (LargestLargeHeaderTimeDelta + 1) / 2;
-
-	FSortedKeyHeader()
-	{
-		TrackIndex = 0xFFFF;
-		PackedData[0] = 0;
-		PackedData[1] = 0;
-	}
-
-	explicit FSortedKeyHeader(const uint8* InData)
-	{
-		TrackIndex = AnimationCompressionUtils::UnalignedRead<uint16>(InData);
-		PackedData[0] = InData[sizeof(uint16) + 0];
-		PackedData[1] = IsLargeHeader() ? InData[sizeof(uint16) + 1] : 0;
-	}
-
-	FSortedKeyHeader(uint16 InTrackIndex, uint8 KeyType, int32 TimeDelta)
-	{
-		TrackIndex = InTrackIndex;
-		if (TimeDelta + SmallHeaderTimeDeltaOffset >= 0 && TimeDelta + SmallHeaderTimeDeltaOffset <= LargestSmallHeaderTimeDelta)
-		{
-			// Small header
-			const int32 OffsetTimeDelta = TimeDelta + SmallHeaderTimeDeltaOffset;
-			PackedData[0] = 0x00 | (static_cast<uint8>(KeyType) << 5) | static_cast<uint8>(OffsetTimeDelta);
-			PackedData[1] = 0;
-		}
-		else
-		{
-			// Large header
-			// Pack each byte separately to avoid issues with little endian
-			const int32 OffsetTimeDelta = TimeDelta + LargeHeaderTimeDeltaOffset;
-			check(OffsetTimeDelta >= 0 && OffsetTimeDelta <= LargestLargeHeaderTimeDelta);
-			PackedData[0] = 0x80 | (static_cast<uint16>(KeyType) << 5) | static_cast<uint16>(OffsetTimeDelta >> 8);
-			PackedData[1] = OffsetTimeDelta & 0xFF;
-		}
-	}
-
-	constexpr bool IsEndOfStream() const { return TrackIndex == 0xFFFF; }
-	constexpr bool IsLargeHeader() const { return (PackedData[0] & KeyHeaderSizeMask) != 0; }
-	constexpr uint8 GetKeyType() const { return (PackedData[0] >> KeyTypeShift) & KeyTypeMask; }
-
-	int32 GetTimeDelta() const
-	{
-		int32 TimeDelta = static_cast<uint32>(PackedData[0] & KeyTimeDeltaMask);
-		if (IsLargeHeader())
-		{
-			TimeDelta = static_cast<int32>((static_cast<uint32>(TimeDelta) << KeyTimeDeltaShift) | PackedData[1]);
-			TimeDelta -= LargeHeaderTimeDeltaOffset;
-		}
-		else
-		{
-			TimeDelta -= SmallHeaderTimeDeltaOffset;
-		}
-		return TimeDelta;
-	}
-
-	constexpr uint8 GetSize() const { return sizeof(uint16) + (IsLargeHeader() ? sizeof(uint16) : sizeof(uint8)); }
-
-	uint16 TrackIndex;
-	uint8 PackedData[2];
-};
-
-/**
-* Small header in the anim sequence compressed stream.
-*/
-struct FAnimSequenceCompressionHeader
-{
-	uint32 NumTracks;
-	uint32 NumFrames;
-	uint32 SequenceCRC;
-	uint8 bHasScale;
-	uint8 bIsSorted;						// For variable interpolation
-};
-
-
-/**
- * Extracts a single BoneAtom from an Animation Sequence.
- *
- * @param	OutAtom			The BoneAtom to fill with the extracted result.
- * @param	DecompContext	The decompression context to use.
- * @param	TrackIndex		The index of the track desired in the Animation Sequence.
- */
-void AnimationFormat_GetBoneAtom(	FTransform& OutAtom,
-									FAnimSequenceDecompressionContext& DecompContext,
-									int32 TrackIndex);
-
-#if USE_ANIMATION_CODEC_BATCH_SOLVER
-
-/**
- * Extracts an array of BoneAtoms from an Animation Sequence representing an entire pose of the skeleton.
- *
- * @param	Atoms				The BoneAtoms to fill with the extracted result.
- * @param	RotationTracks		A BoneTrackArray element for each bone requesting rotation data.
- * @param	TranslationTracks	A BoneTrackArray element for each bone requesting translation data.
- * @param	ScaleTracks			A BoneTrackArray element for each bone requesting scale data.
- * @param	DecompContext		The decompression context to use.
- */
-void AnimationFormat_GetAnimationPose(
-	FTransformArray& Atoms,
-	const BoneTrackArray& RotationTracks,
-	const BoneTrackArray& TranslationTracks,
-	const BoneTrackArray& ScaleTracks,
-	FAnimSequenceDecompressionContext& DecompContext);
-
-#endif
 
 /**
  * Extracts statistics about a given Animation Sequence
@@ -253,12 +110,12 @@ void AnimationFormat_SetInterfaceLinks(CompressedDataType& CompressedData);
 	(Data) += (Len);
 #endif // !WITH_EDITORONLY_DATA
 
-extern const int32 CompressedTranslationStrides[ACF_MAX];
-extern const int32 CompressedTranslationNum[ACF_MAX];
+extern ENGINE_API const int32 CompressedTranslationStrides[ACF_MAX];
+extern ENGINE_API const int32 CompressedTranslationNum[ACF_MAX];
 extern ENGINE_API const int32 CompressedRotationStrides[ACF_MAX];
-extern const int32 CompressedRotationNum[ACF_MAX];
-extern const int32 CompressedScaleStrides[ACF_MAX];
-extern const int32 CompressedScaleNum[ACF_MAX];
+extern ENGINE_API const int32 CompressedRotationNum[ACF_MAX];
+extern ENGINE_API const int32 CompressedScaleStrides[ACF_MAX];
+extern ENGINE_API const int32 CompressedScaleNum[ACF_MAX];
 extern ENGINE_API const uint8 PerTrackNumComponentTable[ACF_MAX*8];
 
 class FMemoryWriter;
@@ -293,23 +150,6 @@ public:
 		FUECompressedAnimData& CompressedData,
 		FMemoryWriter& MemoryWriter) PURE_VIRTUAL(AnimEncoding::ByteSwapOut, );
 
-	/**
-	 * Extracts a single BoneAtom from an Animation Sequence.
-	 *
-	 * @param	OutAtom			The BoneAtom to fill with the extracted result.
-	 * @param	DecompContext	The decompression context to use.
-	 * @param	TrackIndex		The index of the track desired in the Animation Sequence.
-	 */
-	virtual void GetBoneAtom(
-		FTransform& OutAtom,
-		FAnimSequenceDecompressionContext& DecompContext,
-		int32 TrackIndex) PURE_VIRTUAL(AnimEncoding::GetBoneAtom,);
-
-#if USE_SEGMENTING_CONTEXT
-	virtual void CreateEncodingContext(FAnimSequenceDecompressionContext& DecompContext) {}
-	virtual void ReleaseEncodingContext(FAnimSequenceDecompressionContext& DecompContext) {}
-#endif
-
 #if USE_ANIMATION_CODEC_BATCH_SOLVER
 
 	/**
@@ -320,7 +160,7 @@ public:
 	 * @param	DecompContext	The decompression context to use.
 	 */
 	virtual void GetPoseRotations(
-		FTransformArray& Atoms,
+		TArrayView<FTransform>& Atoms,
 		const BoneTrackArray& DesiredPairs,
 		FAnimSequenceDecompressionContext& DecompContext) PURE_VIRTUAL(AnimEncoding::GetPoseRotations,);
 
@@ -332,7 +172,7 @@ public:
 	 * @param	DecompContext	The decompression context to use.
 	 */
 	virtual void GetPoseTranslations(
-		FTransformArray& Atoms,
+		TArrayView<FTransform>& Atoms,
 		const BoneTrackArray& DesiredPairs,
 		FAnimSequenceDecompressionContext& DecompContext) PURE_VIRTUAL(AnimEncoding::GetPoseTranslations,);
 
@@ -344,23 +184,11 @@ public:
 	 * @param	DecompContext	The decompression context to use.
 	 */
 	virtual void GetPoseScales(
-		FTransformArray& Atoms,
+		TArrayView<FTransform>& Atoms,
 		const BoneTrackArray& DesiredPairs,
 		FAnimSequenceDecompressionContext& DecompContext) PURE_VIRTUAL(AnimEncoding::GetPoseScales,);
 #endif
 
-#if USE_SEGMENTING_CONTEXT
-	static FORCEINLINE float TimeToIndex(
-		const FAnimSequenceDecompressionContext& DecompContext,
-		const uint8* TimeMarkers,
-		uint32 NumKeys,
-		uint32 NumFrames,
-		uint8 TimeMarkerSize,
-		float SegmentRelativePos,
-		int32& FrameIndex0Out,
-		int32& FrameIndex1Out);
-#endif
-
 protected:
 
 	/**
@@ -402,8 +230,6 @@ protected:
 		int32 NumKeys,
 		int32 &PosIndex0Out,
 		int32 &PosIndex1Out);
-
-	friend FAnimSequenceDecompressionContext;	// For the TimeToIndex functions
 };
 
 
@@ -415,18 +241,6 @@ protected:
 class AnimEncodingLegacyBase : public AnimEncoding
 {
 public:
-	/**
-	 * Extracts a single BoneAtom from an Animation Sequence.
-	 *
-	 * @param	OutAtom			The BoneAtom to fill with the extracted result.
-	 * @param	DecompContext	The decompression context to use.
-	 * @param	TrackIndex		The index of the track desired in the Animation Sequence.
-	 */
-	virtual void GetBoneAtom(
-		FTransform& OutAtom,
-		FAnimSequenceDecompressionContext& DecompContext,
-		int32 TrackIndex) override;
-
 	/**
 	 * Decompress the Rotation component of a BoneAtom
 	 *
@@ -784,75 +598,3 @@ FORCEINLINE float AnimEncoding::TimeToIndex(
 	
 	return Alpha;
 }
-
-#if USE_SEGMENTING_CONTEXT
-float AnimEncoding::TimeToIndex(
-	const FAnimSequenceDecompressionContext& DecompContext,
-	const uint8* TimeMarkers,
-	uint32 NumKeys,
-	uint32 NumFrames,
-	uint8 TimeMarkerSize,
-	float SegmentRelativePos,
-	int32& FrameIndex0Out,
-	int32& FrameIndex1Out)
-{
-	float Alpha = 0.0f;
-
-	checkSlow(NumKeys != 0);
-
-	const int32 LastKey = NumKeys - 1;
-
-	if (NumKeys < 2 || DecompContext.RelativePos <= 0.f)
-	{
-		// return the first key
-		FrameIndex0Out = 0;
-		FrameIndex1Out = 0;
-		Alpha = 0.0f;
-	}
-	else if (DecompContext.RelativePos >= 1.0f)
-	{
-		// return the ending key
-		FrameIndex0Out = LastKey;
-		FrameIndex1Out = LastKey;
-		Alpha = 0.0f;
-	}
-	else
-	{
-		// find the proper key range to return
-		const int32 LastFrame = NumFrames - 1;
-		const float KeyPos = SegmentRelativePos * (float)LastKey;
-		const float FramePos = SegmentRelativePos * (float)LastFrame;
-		const int32 FramePosFloor = FMath::Clamp(FMath::TruncToInt(FramePos), 0, LastFrame);
-		const int32 KeyEstimate = FMath::Clamp(FMath::TruncToInt(KeyPos), 0, LastKey);
-
-		int32 LowFrame = 0;
-		int32 HighFrame = 0;
-
-		// find the pair of keys which surround our target frame index
-		if (TimeMarkerSize == sizeof(uint16))
-		{
-			const uint16* Frames = reinterpret_cast<const uint16*>(TimeMarkers);
-			FrameIndex0Out = FindLowKeyIndex<uint16>(Frames, NumKeys, FramePosFloor, KeyEstimate);
-			LowFrame = Frames[FrameIndex0Out];
-
-			FrameIndex1Out = FMath::Min(FrameIndex0Out + 1, LastKey);
-			HighFrame = Frames[FrameIndex1Out];
-		}
-		else
-		{
-			FrameIndex0Out = FindLowKeyIndex<uint8>(TimeMarkers, NumKeys, FramePosFloor, KeyEstimate);
-			LowFrame = TimeMarkers[FrameIndex0Out];
-
-			FrameIndex1Out = FMath::Min(FrameIndex0Out + 1, LastKey);
-			HighFrame = TimeMarkers[FrameIndex1Out];
-		}
-
-		// compute the blend parameters for the keys we have found
-		const int32 Delta = FMath::Max(HighFrame - LowFrame, 1);
-		const float Remainder = FramePos - (float)LowFrame;
-		Alpha = DecompContext.GetInterpolation() == EAnimInterpolationType::Step ? 0.f : (Remainder / (float)Delta);
-	}
-
-	return Alpha;
-}
-#endif
diff --git a/Engine/Source/Runtime/Engine/Public/AnimEncoding_ConstantKeyLerp.h b/Engine/Source/Runtime/Engine/Public/AnimEncoding_ConstantKeyLerp.h
--- a/Engine/Source/Runtime/Engine/Public/AnimEncoding_ConstantKeyLerp.h
+++ b/Engine/Source/Runtime/Engine/Public/AnimEncoding_ConstantKeyLerp.h
@@ -9,28 +9,9 @@
 #include "CoreMinimal.h"
 #include "AnimEncoding.h"
 #include "AnimationCompression.h"
-#include "Animation/AnimEncodingDecompressionContext.h"
 
 class FMemoryWriter;
 
-#if USE_SEGMENTING_CONTEXT
-class FAEConstantKeyLerpContext : public FAnimEncodingDecompressionContext
-{
-public:
-	FAEConstantKeyLerpContext(const FAnimSequenceDecompressionContext& DecompContext);
-	virtual void Seek(const FAnimSequenceDecompressionContext& DecompContext, float SampleAtTime) override;
-
-	template<int32 Format> inline FQuat GetUniformRotation(const FAnimSequenceDecompressionContext& DecompContext, int32 TrackIndex, uint8 SegmentIndex) const;
-	template<int32 Format> inline FVector GetUniformTranslation(const FAnimSequenceDecompressionContext& DecompContext, int32 TrackIndex, uint8 SegmentIndex) const;
-	template<int32 Format> inline FVector GetUniformScale(const FAnimSequenceDecompressionContext& DecompContext, int32 TrackIndex, uint8 SegmentIndex) const;
-
-	TArray<int32> UniformKeyOffsets;
-	int32 KeyFrameSize;
-
-	int32 FrameKeysOffset[2];
-};
-#endif
-
 /**
  * Base class for all Animation Encoding Formats using consistently-spaced key interpolation.
  */
@@ -121,11 +102,6 @@ public:
 		FMemoryWriter& MemoryWriter,
 		uint8*& ScaleTrackData,
 		int32 NumKeysScale) override;
-
-#if USE_SEGMENTING_CONTEXT
-	virtual void CreateEncodingContext(FAnimSequenceDecompressionContext& DecompContext) override;
-	virtual void ReleaseEncodingContext(FAnimSequenceDecompressionContext& DecompContext) override;
-#endif
 };
 
 template<int32 FORMAT>
@@ -169,7 +145,7 @@ public:
 	 * @param	DecompContext	The decompression context to use.
 	 */
 	void GetPoseRotations(
-		FTransformArray& Atoms,
+		TArrayView<FTransform>& Atoms,
 		const BoneTrackArray& DesiredPairs,
 		FAnimSequenceDecompressionContext& DecompContext);
 
@@ -181,7 +157,7 @@ public:
 	 * @param	DecompContext	The decompression context to use.
 	 */
 	void GetPoseTranslations(
-		FTransformArray& Atoms,
+		TArrayView<FTransform>& Atoms,
 		const BoneTrackArray& DesiredPairs,
 		FAnimSequenceDecompressionContext& DecompContext);
 
@@ -193,7 +169,7 @@ public:
 	 * @param	DecompContext	The decompression context to use.
 	 */
 	void GetPoseScales(
-		FTransformArray& Atoms,
+		TArrayView<FTransform>& Atoms,
 		const BoneTrackArray& DesiredPairs,
 		FAnimSequenceDecompressionContext& DecompContext);
 #endif
@@ -212,43 +188,12 @@ public:
 template<int32 FORMAT>
 FORCEINLINE void AEFConstantKeyLerp<FORMAT>::GetBoneAtomRotation(FTransform& OutAtom, const FAnimSequenceDecompressionContext& DecompContext, int32 TrackIndex)
 {
-#if USE_SEGMENTING_CONTEXT
-	if (DecompContext.AnimSeq->CompressedSegments.Num() != 0)
-	{
-		const FTrivialAnimKeyHandle TrivialKeyHandle = DecompContext.GetTrivialRotationKeyHandle(TrackIndex);
-		if (TrivialKeyHandle.IsValid())
-		{
-			DecompContext.GetTrivialRotation(OutAtom, TrivialKeyHandle);
-		}
-		else
-		{
-			const FAEConstantKeyLerpContext& EncodingContext = *static_cast<const FAEConstantKeyLerpContext*>(DecompContext.EncodingContext);
-
-			const FQuat R0 = EncodingContext.GetUniformRotation<FORMAT>(DecompContext, TrackIndex, 0);
-
-			if (DecompContext.NeedsInterpolation)
-			{
-				const FQuat R1 = EncodingContext.GetUniformRotation<FORMAT>(DecompContext, TrackIndex, 1);
-
-				// Fast linear quaternion interpolation.
-				FQuat BlendedQuat = FQuat::FastLerp(R0, R1, DecompContext.KeyAlpha);
-				BlendedQuat.Normalize();
-				OutAtom.SetRotation(BlendedQuat);
-			}
-			else // (Index0 == Index1)
-			{
-				OutAtom.SetRotation(R0);
-			}
-		}
-
-		return;
-	}
-#endif
+	const FUECompressedAnimData& AnimData = static_cast<const FUECompressedAnimData&>(DecompContext.CompressedAnimData);
 
-	const int32* RESTRICT TrackData = DecompContext.GetCompressedTrackOffsets() + (TrackIndex * 4);
+	const int32* RESTRICT TrackData = AnimData.CompressedTrackOffsets.GetData() + (TrackIndex * 4);
 	int32 RotKeysOffset = TrackData[2];
 	int32 NumRotKeys = TrackData[3];
-	const uint8* RESTRICT RotStream = DecompContext.GetCompressedByteStream() + RotKeysOffset;
+	const uint8* RESTRICT RotStream = AnimData.CompressedByteStream.GetData() + RotKeysOffset;
 
 	if (NumRotKeys == 1)
 	{
@@ -261,7 +206,7 @@ FORCEINLINE void AEFConstantKeyLerp<FORMAT>::GetBoneAtomRotation(FTransform& Out
 	{
 		int32 Index0;
 		int32 Index1;
-		float Alpha = TimeToIndex(DecompContext.GetSequenceLength(), DecompContext.RelativePos, NumRotKeys, DecompContext.GetInterpolation(), Index0, Index1);
+		float Alpha = TimeToIndex(DecompContext.SequenceLength, DecompContext.RelativePos, NumRotKeys, DecompContext.Interpolation, Index0, Index1);
 
 		const int32 RotationStreamOffset = (FORMAT == ACF_IntervalFixed32NoW) ? (sizeof(float)*6) : 0; // offset past Min and Range data
 
@@ -303,44 +248,16 @@ FORCEINLINE void AEFConstantKeyLerp<FORMAT>::GetBoneAtomRotation(FTransform& Out
 template<int32 FORMAT>
 FORCEINLINE void AEFConstantKeyLerp<FORMAT>::GetBoneAtomTranslation(FTransform& OutAtom, const FAnimSequenceDecompressionContext& DecompContext, int32 TrackIndex)
 {
-#if USE_SEGMENTING_CONTEXT
-	if (DecompContext.AnimSeq->CompressedSegments.Num() != 0)
-	{
-		const FTrivialAnimKeyHandle TrivialKeyHandle = DecompContext.GetTrivialTranslationKeyHandle(TrackIndex);
-		if (TrivialKeyHandle.IsValid())
-		{
-			DecompContext.GetTrivialTranslation(OutAtom, TrivialKeyHandle);
-		}
-		else
-		{
-			const FAEConstantKeyLerpContext& EncodingContext = *static_cast<const FAEConstantKeyLerpContext*>(DecompContext.EncodingContext);
-
-			const FVector P0 = EncodingContext.GetUniformTranslation<FORMAT>(DecompContext, TrackIndex, 0);
-
-			if (DecompContext.NeedsInterpolation)
-			{
-				const FVector P1 = EncodingContext.GetUniformTranslation<FORMAT>(DecompContext, TrackIndex, 1);
-
-				OutAtom.SetTranslation(FMath::Lerp(P0, P1, DecompContext.KeyAlpha));
-			}
-			else // (Index0 == Index1)
-			{
-				OutAtom.SetTranslation(P0);
-			}
-		}
-
-		return;
-	}
-#endif
+	const FUECompressedAnimData& AnimData = static_cast<const FUECompressedAnimData&>(DecompContext.CompressedAnimData);
 
-	const int32* RESTRICT TrackData = DecompContext.GetCompressedTrackOffsets() + (TrackIndex * 4);
+	const int32* RESTRICT TrackData = AnimData.CompressedTrackOffsets.GetData() + (TrackIndex * 4);
 	int32 TransKeysOffset = TrackData[0];
 	int32 NumTransKeys = TrackData[1];
-	const uint8* RESTRICT TransStream = DecompContext.GetCompressedByteStream() + TransKeysOffset;
+	const uint8* RESTRICT TransStream = AnimData.CompressedByteStream.GetData() + TransKeysOffset;
 
 	int32 Index0;
 	int32 Index1;
-	float Alpha = TimeToIndex(DecompContext.GetSequenceLength(), DecompContext.RelativePos, NumTransKeys, DecompContext.GetInterpolation(), Index0, Index1);
+	float Alpha = TimeToIndex(DecompContext.SequenceLength, DecompContext.RelativePos, NumTransKeys, DecompContext.Interpolation, Index0, Index1);
 
 	const int32 TransStreamOffset = ((FORMAT == ACF_IntervalFixed32NoW) && NumTransKeys > 1) ? (sizeof(float)*6) : 0; // offset past Min and Range data
 
@@ -374,43 +291,15 @@ FORCEINLINE void AEFConstantKeyLerp<FORMAT>::GetBoneAtomTranslation(FTransform&
 template<int32 FORMAT>
 FORCEINLINE void AEFConstantKeyLerp<FORMAT>::GetBoneAtomScale(FTransform& OutAtom, const FAnimSequenceDecompressionContext& DecompContext, int32 TrackIndex)
 {
-#if USE_SEGMENTING_CONTEXT
-	if (DecompContext.AnimSeq->CompressedSegments.Num() != 0)
-	{
-		const FTrivialAnimKeyHandle TrivialKeyHandle = DecompContext.GetTrivialScaleKeyHandle(TrackIndex);
-		if (TrivialKeyHandle.IsValid())
-		{
-			DecompContext.GetTrivialScale(OutAtom, TrivialKeyHandle);
-		}
-		else
-		{
-			const FAEConstantKeyLerpContext& EncodingContext = *static_cast<const FAEConstantKeyLerpContext*>(DecompContext.EncodingContext);
-
-			const FVector S0 = EncodingContext.GetUniformScale<FORMAT>(DecompContext, TrackIndex, 0);
+	const FUECompressedAnimData& AnimData = static_cast<const FUECompressedAnimData&>(DecompContext.CompressedAnimData);
 
-			if (DecompContext.NeedsInterpolation)
-			{
-				const FVector S1 = EncodingContext.GetUniformScale<FORMAT>(DecompContext, TrackIndex, 1);
-
-				OutAtom.SetScale3D(FMath::Lerp(S0, S1, DecompContext.KeyAlpha));
-			}
-			else // (Index0 == Index1)
-			{
-				OutAtom.SetScale3D(S0);
-			}
-		}
-
-		return;
-	}
-#endif
-
-	const int32 ScaleKeysOffset = DecompContext.GetCompressedScaleOffsets()->GetOffsetData(TrackIndex, 0);
-	const int32 NumScaleKeys = DecompContext.GetCompressedScaleOffsets()->GetOffsetData(TrackIndex, 1);
-	const uint8* RESTRICT ScaleStream = DecompContext.GetCompressedByteStream() + ScaleKeysOffset;
+	const int32 ScaleKeysOffset = AnimData.CompressedScaleOffsets.GetOffsetData(TrackIndex, 0);
+	const int32 NumScaleKeys = AnimData.CompressedScaleOffsets.GetOffsetData(TrackIndex, 1);
+	const uint8* RESTRICT ScaleStream = AnimData.CompressedByteStream.GetData() + ScaleKeysOffset;
 
 	int32 Index0;
 	int32 Index1;
-	float Alpha = TimeToIndex(DecompContext.GetSequenceLength(), DecompContext.RelativePos, NumScaleKeys, DecompContext.GetInterpolation(), Index0, Index1);
+	float Alpha = TimeToIndex(DecompContext.SequenceLength, DecompContext.RelativePos, NumScaleKeys, DecompContext.Interpolation, Index0, Index1);
 
 	const int32 ScaleStreamOffset = ((FORMAT == ACF_IntervalFixed32NoW) && NumScaleKeys > 1) ? (sizeof(float)*6) : 0; // offset past Min and Range data
 
@@ -432,163 +321,4 @@ FORCEINLINE void AEFConstantKeyLerp<FORMAT>::GetBoneAtomScale(FTransform& OutAto
 		DecompressScale<FORMAT>( P0, ScaleStream, KeyData);
 		OutAtom.SetScale3D(P0);
 	}
-}
-
-#if USE_ANIMATION_CODEC_BATCH_SOLVER
-
-/**
- * Decompress all requested rotation components from an Animation Sequence
- *
- * @param	Atoms			The FTransform array to fill in.
- * @param	DesiredPairs	Array of requested bone information
- * @param	DecompContext	The decompression context to use.
- */
-template<int32 FORMAT>
-inline void AEFConstantKeyLerp<FORMAT>::GetPoseRotations(
-	FTransformArray& Atoms,
-	const BoneTrackArray& DesiredPairs,
-	FAnimSequenceDecompressionContext& DecompContext)
-{
-	const int32 PairCount = DesiredPairs.Num();
-
-	for (int32 PairIndex=0; PairIndex<PairCount; ++PairIndex)
-	{
-		const BoneTrackPair& Pair = DesiredPairs[PairIndex];
-		const int32 TrackIndex = Pair.TrackIndex;
-		const int32 AtomIndex = Pair.AtomIndex;
-		FTransform& BoneAtom = Atoms[AtomIndex];
-
-		// call the decoder directly (not through the vtable)
-		AEFConstantKeyLerp<FORMAT>::GetBoneAtomRotation(BoneAtom, DecompContext, TrackIndex);
-	}
-}
-
-/**
- * Decompress all requested translation components from an Animation Sequence
- *
- * @param	Atoms			The FTransform array to fill in.
- * @param	DesiredPairs	Array of requested bone information
- * @param	DecompContext	The decompression context to use.
- */
-template<int32 FORMAT>
-inline void AEFConstantKeyLerp<FORMAT>::GetPoseTranslations(
-	FTransformArray& Atoms,
-	const BoneTrackArray& DesiredPairs,
-	FAnimSequenceDecompressionContext& DecompContext)
-{
-	const int32 PairCount = DesiredPairs.Num();
-
-	//@TODO: Verify that this prefetch is helping
-	// Prefetch the desired pairs array and 2 destination spots; the loop will prefetch one 2 out each iteration
-	FPlatformMisc::Prefetch(&(DesiredPairs[0]));
-	const int32 PrefetchCount = FMath::Min(PairCount, 1);
-	for (int32 PairIndex=0; PairIndex<PairCount; ++PairIndex)
-	{
-		const BoneTrackPair& Pair = DesiredPairs[PairIndex];
-		FPlatformMisc::Prefetch(Atoms.GetData() + Pair.AtomIndex);
-	}
-
-	for (int32 PairIndex=0; PairIndex<PairCount; ++PairIndex)
-	{
-		int32 PrefetchIndex = PairIndex + PrefetchCount;
-		if (PrefetchIndex < PairCount)
-		{
-			FPlatformMisc::Prefetch(Atoms.GetData() + DesiredPairs[PrefetchIndex].AtomIndex);
-		}
-
-		const BoneTrackPair& Pair = DesiredPairs[PairIndex];
-		const int32 TrackIndex = Pair.TrackIndex;
-		const int32 AtomIndex = Pair.AtomIndex;
-		FTransform& BoneAtom = Atoms[AtomIndex];
-
-		// call the decoder directly (not through the vtable)
-		AEFConstantKeyLerp<FORMAT>::GetBoneAtomTranslation(BoneAtom, DecompContext, TrackIndex);
-	}
-}
-
-/**
- * Decompress all requested Scale components from an Animation Sequence
- *
- * @param	Atoms			The FTransform array to fill in.
- * @param	DesiredPairs	Array of requested bone information
- * @param	DecompContext	The decompression context to use.
- */
-template<int32 FORMAT>
-inline void AEFConstantKeyLerp<FORMAT>::GetPoseScales(
-	FTransformArray& Atoms,
-	const BoneTrackArray& DesiredPairs,
-	FAnimSequenceDecompressionContext& DecompContext)
-{
-	checkSlow(DecompContext.bHasScale);
-
-	const int32 PairCount= DesiredPairs.Num();
-
-	//@TODO: Verify that this prefetch is helping
-	// Prefetch the desired pairs array and 2 destination spots; the loop will prefetch one 2 out each iteration
-	FPlatformMisc::Prefetch(&(DesiredPairs[0]));
-	const int32 PrefetchCount = FMath::Min(PairCount, 1);
-	for (int32 PairIndex=0; PairIndex<PairCount; ++PairIndex)
-	{
-		const BoneTrackPair& Pair = DesiredPairs[PairIndex];
-		FPlatformMisc::Prefetch(Atoms.GetData() + Pair.AtomIndex);
-	}
-
-	for (int32 PairIndex=0; PairIndex<PairCount; ++PairIndex)
-	{
-		int32 PrefetchIndex = PairIndex + PrefetchCount;
-		if (PrefetchIndex < PairCount)
-		{
-			FPlatformMisc::Prefetch(Atoms.GetData() + DesiredPairs[PrefetchIndex].AtomIndex);
-		}
-
-		const BoneTrackPair& Pair = DesiredPairs[PairIndex];
-		const int32 TrackIndex = Pair.TrackIndex;
-		const int32 AtomIndex = Pair.AtomIndex;
-		FTransform& BoneAtom = Atoms[AtomIndex];
-
-		// call the decoder directly (not through the vtable)
-		AEFConstantKeyLerp<FORMAT>::GetBoneAtomScale(BoneAtom, DecompContext, TrackIndex);
-	}
-}
-#endif
-
-#if USE_SEGMENTING_CONTEXT
-template<int32 Format>
-FQuat FAEConstantKeyLerpContext::GetUniformRotation(const FAnimSequenceDecompressionContext& DecompContext, int32 TrackIndex, uint8 SegmentIndex) const
-{
-	const int32 FrameKeyOffset = UniformKeyOffsets[DecompContext.GetRotationValueOffset(TrackIndex)];
-	const int32 SegmentKeyOffset = FrameKeysOffset[SegmentIndex] + FrameKeyOffset;
-	const uint8* KeyData = DecompContext.CompressedByteStream + SegmentKeyOffset;
-
-	FQuat Rotation;
-	DecompressRotation<Format>(Rotation, DecompContext.TrackRangeData[SegmentIndex], KeyData);
-
-	return Rotation;
-}
-
-template<int32 Format>
-FVector FAEConstantKeyLerpContext::GetUniformTranslation(const FAnimSequenceDecompressionContext& DecompContext, int32 TrackIndex, uint8 SegmentIndex) const
-{
-	const int32 FrameKeyOffset = UniformKeyOffsets[DecompContext.GetTranslationValueOffset(TrackIndex)];
-	const int32 SegmentKeyOffset = FrameKeysOffset[SegmentIndex] + FrameKeyOffset;
-	const uint8* KeyData = DecompContext.CompressedByteStream + SegmentKeyOffset;
-
-	FVector Translation;
-	DecompressTranslation<Format>(Translation, DecompContext.TrackRangeData[SegmentIndex], KeyData);
-
-	return Translation;
-}
-
-template<int32 Format>
-FVector FAEConstantKeyLerpContext::GetUniformScale(const FAnimSequenceDecompressionContext& DecompContext, int32 TrackIndex, uint8 SegmentIndex) const
-{
-	const int32 FrameKeyOffset = UniformKeyOffsets[DecompContext.GetScaleValueOffset(TrackIndex)];
-	const int32 SegmentKeyOffset = FrameKeysOffset[SegmentIndex] + FrameKeyOffset;
-	const uint8* KeyData = DecompContext.CompressedByteStream + SegmentKeyOffset;
-
-	FVector Scale;
-	DecompressScale<Format>(Scale, DecompContext.TrackRangeData[SegmentIndex], KeyData);
-
-	return Scale;
-}
-#endif
\ No newline at end of file
+}
\ No newline at end of file
diff --git a/Engine/Source/Runtime/Engine/Public/AnimEncoding_PerTrackCompression.h b/Engine/Source/Runtime/Engine/Public/AnimEncoding_PerTrackCompression.h
--- a/Engine/Source/Runtime/Engine/Public/AnimEncoding_PerTrackCompression.h
+++ b/Engine/Source/Runtime/Engine/Public/AnimEncoding_PerTrackCompression.h
@@ -37,18 +37,6 @@ public:
 		FUECompressedAnimData& CompressedData,
 		FMemoryWriter& MemoryWriter) override;
 
-	/**
-	 * Extracts a single BoneAtom from an Animation Sequence.
-	 *
-	 * @param	OutAtom			The BoneAtom to fill with the extracted result.
-	 * @param	DecompContext	The decompression context to use.
-	 * @param	TrackIndex		The index of the track desired in the Animation Sequence.
-	 */
-	virtual void GetBoneAtom(
-		FTransform& OutAtom,
-		FAnimSequenceDecompressionContext& DecompContext,
-		int32 TrackIndex) override;
-
 #if USE_ANIMATION_CODEC_BATCH_SOLVER
 
 	/**
@@ -59,7 +47,7 @@ public:
 	 * @param	DecompContext	The decompression context to use.
 	 */
 	virtual void GetPoseRotations(
-		FTransformArray& Atoms,
+		TArrayView<FTransform>& Atoms,
 		const BoneTrackArray& DesiredPairs,
 		FAnimSequenceDecompressionContext& DecompContext) override;
 
@@ -71,7 +59,7 @@ public:
 	 * @param	DecompContext	The decompression context to use.
 	 */
 	virtual void GetPoseTranslations(
-		FTransformArray& Atoms,
+		TArrayView<FTransform>& Atoms,
 		const BoneTrackArray& DesiredPairs,
 		FAnimSequenceDecompressionContext& DecompContext) override;
 
@@ -83,16 +71,11 @@ public:
 	 * @param	DecompContext	The decompression context to use.
 	 */
 	virtual void GetPoseScales(
-		FTransformArray& Atoms,
+		TArrayView<FTransform>& Atoms,
 		const BoneTrackArray& DesiredPairs,
 		FAnimSequenceDecompressionContext& DecompContext) override;
 #endif
 
-#if USE_SEGMENTING_CONTEXT
-	virtual void CreateEncodingContext(FAnimSequenceDecompressionContext& DecompContext) override;
-	virtual void ReleaseEncodingContext(FAnimSequenceDecompressionContext& DecompContext) override;
-#endif
-
 protected:
 	/**
 	 * Handles Byte-swapping a single track of animation data from a MemoryReader or to a MemoryWriter
@@ -147,24 +130,6 @@ protected:
 		FTransform& OutAtom,
 		FAnimSequenceDecompressionContext& DecompContext,
 		int32 TrackIndex);
-};
 
-/**
- * Structure to wrap per track flags.
- */
-struct FPerTrackFlags
-{
-	constexpr explicit FPerTrackFlags(uint8 InFlags) : Flags(InFlags) {}
-	FPerTrackFlags(bool bHasTimeMarkers, AnimationCompressionFormat Format, uint8 FormatFlags)
-		: Flags((bHasTimeMarkers ? 0x80 : 0) | (FormatFlags << 4) | (uint8)Format)
-	{
-		check((FormatFlags & ~0x7) == 0);
-		check(((uint8)Format & ~0xF) == 0);
-	}
-
-	constexpr bool IsUniform() const { return (Flags & 0x80) == 0; }
-	constexpr uint8 GetFormatFlags() const { return (Flags >> 4) & 0x7; }
-	constexpr uint8 GetFormat() const { return Flags & 0xF; }
-
-	uint8 Flags;
+	friend class UAnimCompress_PerTrackCompression;
 };
