COPY from //Dev-Anim to //Dev-Main

diff --git a/Engine/Source/Runtime/Engine/Public/AnimEncoding_VariableKeyLerp.h b/Engine/Source/Runtime/Engine/Public/AnimEncoding_VariableKeyLerp.h
--- a/Engine/Source/Runtime/Engine/Public/AnimEncoding_VariableKeyLerp.h
+++ b/Engine/Source/Runtime/Engine/Public/AnimEncoding_VariableKeyLerp.h
@@ -9,70 +9,10 @@
 #include "CoreMinimal.h"
 #include "AnimEncoding.h"
 #include "AnimationCompression.h"
-#include "Animation/AnimEncodingDecompressionContext.h"
 #include "AnimEncoding_ConstantKeyLerp.h"
 
 class FMemoryWriter;
 
-#if USE_SEGMENTING_CONTEXT
-class FAEVariableKeyLerpSortedContext : public FAnimEncodingDecompressionContext
-{
-public:
-	struct FCachedKey
-	{
-		// Linear interpolation only requires 2 keys
-		// Index 0 is always the oldest key, index 1 the newest
-		int32 RotOffsets[2];
-		int32 RotFrameIndices[2];
-		int32 TransOffsets[2];
-		int32 TransFrameIndices[2];
-		int32 ScaleOffsets[2];
-		int32 ScaleFrameIndices[2];
-	};
-
-	TArray<FCachedKey> CachedKeys;			// 1 Entry per track
-
-	int32 SegmentStartFrame0;
-	int32 SegmentStartFrame1;
-	float FramePos;
-
-	const uint8_t* PackedSampleData;		// The current pointer into our data stream
-	int32 PreviousFrameIndex;				// The previously read frame index
-	int32 CurrentFrameIndex;				// The current frame index
-
-	float PreviousSampleAtTime;
-	uint16 PreviousSegmentIndex;
-
-	FAEVariableKeyLerpSortedContext(const FAnimSequenceDecompressionContext& DecompContext);
-	virtual void Seek(const FAnimSequenceDecompressionContext& DecompContext, float SampleAtTime) override;
-
-	template<int32 Format> inline FQuat GetSortedRotation(const FAnimSequenceDecompressionContext& DecompContext, int32 TrackIndex) const;
-	template<int32 Format> inline FVector GetSortedTranslation(const FAnimSequenceDecompressionContext& DecompContext, int32 TrackIndex) const;
-	template<int32 Format> inline FVector GetSortedScale(const FAnimSequenceDecompressionContext& DecompContext, int32 TrackIndex) const;
-};
-
-class FAEVariableKeyLerpLinearContext : public FAnimEncodingDecompressionContext
-{
-public:
-	static constexpr int32 OffsetNumKeysPairSize = sizeof(uint32) + sizeof(uint16);
-
-	float SegmentRelativePos0;
-
-	uint8 TimeMarkerSize[2];	// sizeof(uint8) or sizeof(uint16)
-
-	const uint8* OffsetNumKeysPairs[2];
-
-	TArray<int32> NumAnimatedTrackStreams;
-
-	FAEVariableKeyLerpLinearContext(const FAnimSequenceDecompressionContext& DecompContext);
-	virtual void Seek(const FAnimSequenceDecompressionContext& DecompContext, float SampleAtTime) override;
-
-	template<int32 Format> inline FQuat GetLinearRotation(const FAnimSequenceDecompressionContext& DecompContext, int32 TrackIndex) const;
-	template<int32 Format> inline FVector GetLinearTranslation(const FAnimSequenceDecompressionContext& DecompContext, int32 TrackIndex) const;
-	template<int32 Format> inline FVector GetLinearScale(const FAnimSequenceDecompressionContext& DecompContext, int32 TrackIndex) const;
-};
-#endif
-
 /**
  * Base class for all Animation Encoding Formats using variably-spaced key interpolation.
  */
@@ -162,11 +102,6 @@ public:
 		FMemoryWriter& MemoryWriter,
 		uint8*& ScaleTrackData,
 		int32 NumKeysScale) override;
-
-#if USE_SEGMENTING_CONTEXT
-	virtual void CreateEncodingContext(FAnimSequenceDecompressionContext& DecompContext) override;
-	virtual void ReleaseEncodingContext(FAnimSequenceDecompressionContext& DecompContext) override;
-#endif
 };
 
 template<int32 FORMAT>
@@ -210,7 +145,7 @@ public:
 	 * @param	DecompContext	The decompression context to use.
 	 */
 	virtual void GetPoseRotations(
-		FTransformArray& Atoms,
+		TArrayView<FTransform>& Atoms,
 		const BoneTrackArray& DesiredPairs,
 		FAnimSequenceDecompressionContext& DecompContext) override;
 
@@ -222,7 +157,7 @@ public:
 	 * @param	DecompContext	The decompression context to use.
 	 */
 	virtual void GetPoseTranslations(
-		FTransformArray& Atoms,
+		TArrayView<FTransform>& Atoms,
 		const BoneTrackArray& DesiredPairs,
 		FAnimSequenceDecompressionContext& DecompContext) override;
 
@@ -234,7 +169,7 @@ public:
 	 * @param	DecompContext	The decompression context to use.
 	 */
 	virtual void GetPoseScales(
-		FTransformArray& Atoms,
+		TArrayView<FTransform>& Atoms,
 		const BoneTrackArray& DesiredPairs,
 		FAnimSequenceDecompressionContext& DecompContext) override;
 #endif
@@ -252,40 +187,12 @@ public:
 template<int32 FORMAT>
 FORCEINLINE_DEBUGGABLE void AEFVariableKeyLerp<FORMAT>::GetBoneAtomRotation(FTransform& OutAtom, const FAnimSequenceDecompressionContext& DecompContext, int32 TrackIndex)
 {
-#if USE_SEGMENTING_CONTEXT
-	if (DecompContext.AnimSeq->CompressedSegments.Num() != 0)
-	{
-		const FTrivialAnimKeyHandle TrivialKeyHandle = DecompContext.GetTrivialRotationKeyHandle(TrackIndex);
-		if (TrivialKeyHandle.IsValid())
-		{
-			DecompContext.GetTrivialRotation(OutAtom, TrivialKeyHandle);
-		}
-		else
-		{
-			FQuat Rotation;
-			if (DecompContext.bIsSorted)
-			{
-				const FAEVariableKeyLerpSortedContext& EncodingContext = *static_cast<const FAEVariableKeyLerpSortedContext*>(DecompContext.EncodingContext);
-				Rotation = EncodingContext.GetSortedRotation<FORMAT>(DecompContext, TrackIndex);
-			}
-			else
-			{
-				const FAEVariableKeyLerpLinearContext& EncodingContext = *static_cast<const FAEVariableKeyLerpLinearContext*>(DecompContext.EncodingContext);
-				Rotation = EncodingContext.GetLinearRotation<FORMAT>(DecompContext, TrackIndex);
-			}
-
-			OutAtom.SetRotation(Rotation);
-		}
-
-		return;
-	}
-
-#endif
+	const FUECompressedAnimData& AnimData = static_cast<const FUECompressedAnimData&>(DecompContext.CompressedAnimData);
 
-	const int32* RESTRICT TrackData = DecompContext.GetCompressedTrackOffsets() + (TrackIndex * 4);
+	const int32* RESTRICT TrackData = AnimData.CompressedTrackOffsets.GetData() + (TrackIndex * 4);
 	int32 RotKeysOffset = TrackData[2];
 	int32 NumRotKeys = TrackData[3];
-	const uint8* RESTRICT RotStream = DecompContext.GetCompressedByteStream() + RotKeysOffset;
+	const uint8* RESTRICT RotStream = AnimData.CompressedByteStream.GetData() + RotKeysOffset;
 
 	if (NumRotKeys == 1)
 	{
@@ -302,7 +209,7 @@ FORCEINLINE_DEBUGGABLE void AEFVariableKeyLerp<FORMAT>::GetBoneAtomRotation(FTra
 
 		int32 Index0;
 		int32 Index1;
-		float Alpha = TimeToIndex(DecompContext.GetInterpolation(), DecompContext.GetCompressedNumberOfFrames(), FrameTable, DecompContext.RelativePos, NumRotKeys, Index0, Index1);
+		float Alpha = TimeToIndex(DecompContext.Interpolation, AnimData.CompressedNumberOfFrames, FrameTable, DecompContext.RelativePos, NumRotKeys, Index0, Index1);
 
 
 		if (Index0 != Index1)
@@ -344,47 +251,19 @@ FORCEINLINE_DEBUGGABLE void AEFVariableKeyLerp<FORMAT>::GetBoneAtomRotation(FTra
 template<int32 FORMAT>
 FORCEINLINE_DEBUGGABLE void AEFVariableKeyLerp<FORMAT>::GetBoneAtomTranslation(FTransform& OutAtom, const FAnimSequenceDecompressionContext& DecompContext, int32 TrackIndex)
 {
-#if USE_SEGMENTING_CONTEXT
-	if (DecompContext.AnimSeq->CompressedSegments.Num() != 0)
-	{
-		const FTrivialAnimKeyHandle TrivialKeyHandle = DecompContext.GetTrivialTranslationKeyHandle(TrackIndex);
-		if (TrivialKeyHandle.IsValid())
-		{
-			DecompContext.GetTrivialTranslation(OutAtom, TrivialKeyHandle);
-		}
-		else
-		{
-			FVector Translation;
-			if (DecompContext.bIsSorted)
-			{
-				const FAEVariableKeyLerpSortedContext& EncodingContext = *static_cast<const FAEVariableKeyLerpSortedContext*>(DecompContext.EncodingContext);
-				Translation = EncodingContext.GetSortedTranslation<FORMAT>(DecompContext, TrackIndex);
-			}
-			else
-			{
-				const FAEVariableKeyLerpLinearContext& EncodingContext = *static_cast<const FAEVariableKeyLerpLinearContext*>(DecompContext.EncodingContext);
-				Translation = EncodingContext.GetLinearTranslation<FORMAT>(DecompContext, TrackIndex);
-			}
-
-			OutAtom.SetTranslation(Translation);
-		}
-
-		return;
-	}
-
-#endif
+	const FUECompressedAnimData& AnimData = static_cast<const FUECompressedAnimData&>(DecompContext.CompressedAnimData);
 
-	const int32* RESTRICT TrackData = DecompContext.GetCompressedTrackOffsets() + (TrackIndex * 4);
+	const int32* RESTRICT TrackData = AnimData.CompressedTrackOffsets.GetData() + (TrackIndex * 4);
 	int32 TransKeysOffset = TrackData[0];
 	int32 NumTransKeys = TrackData[1];
-	const uint8* RESTRICT TransStream = DecompContext.GetCompressedByteStream() + TransKeysOffset;
+	const uint8* RESTRICT TransStream = AnimData.CompressedByteStream.GetData() + TransKeysOffset;
 
 	const uint8* RESTRICT FrameTable= TransStream +(NumTransKeys*CompressedTranslationStrides[FORMAT]*CompressedTranslationNum[FORMAT]);
 	FrameTable= Align(FrameTable, 4);
 
 	int32 Index0;
 	int32 Index1;
-	float Alpha = TimeToIndex(DecompContext.GetInterpolation(), DecompContext.GetCompressedNumberOfFrames(), FrameTable, DecompContext.RelativePos, NumTransKeys, Index0, Index1);
+	float Alpha = TimeToIndex(DecompContext.Interpolation, AnimData.CompressedNumberOfFrames, FrameTable, DecompContext.RelativePos, NumTransKeys, Index0, Index1);
 	const int32 TransStreamOffset = ((FORMAT == ACF_IntervalFixed32NoW) && NumTransKeys > 1) ? (sizeof(float)*6) : 0; // offset past Min and Range data
 
 	if (Index0 != Index1)
@@ -417,45 +296,18 @@ FORCEINLINE_DEBUGGABLE void AEFVariableKeyLerp<FORMAT>::GetBoneAtomTranslation(F
 template<int32 FORMAT>
 FORCEINLINE_DEBUGGABLE void AEFVariableKeyLerp<FORMAT>::GetBoneAtomScale(FTransform& OutAtom, const FAnimSequenceDecompressionContext& DecompContext, int32 TrackIndex)
 {
-#if USE_SEGMENTING_CONTEXT
-	if (DecompContext.AnimSeq->CompressedSegments.Num() != 0)
-	{
-		const FTrivialAnimKeyHandle TrivialKeyHandle = DecompContext.GetTrivialScaleKeyHandle(TrackIndex);
-		if (TrivialKeyHandle.IsValid())
-		{
-			DecompContext.GetTrivialScale(OutAtom, TrivialKeyHandle);
-		}
-		else
-		{
-			FVector Scale;
-			if (DecompContext.bIsSorted)
-			{
-				const FAEVariableKeyLerpSortedContext& EncodingContext = *static_cast<const FAEVariableKeyLerpSortedContext*>(DecompContext.EncodingContext);
-				Scale = EncodingContext.GetSortedScale<FORMAT>(DecompContext, TrackIndex);
-			}
-			else
-			{
-				const FAEVariableKeyLerpLinearContext& EncodingContext = *static_cast<const FAEVariableKeyLerpLinearContext*>(DecompContext.EncodingContext);
-				Scale = EncodingContext.GetLinearScale<FORMAT>(DecompContext, TrackIndex);
-			}
-
-			OutAtom.SetScale3D(Scale);
-		}
-
-		return;
-	}
-#endif
+	const FUECompressedAnimData& AnimData = static_cast<const FUECompressedAnimData&>(DecompContext.CompressedAnimData);
 
-	const int32 ScaleKeysOffset = DecompContext.GetCompressedScaleOffsets()->GetOffsetData(TrackIndex, 0);
-	const int32 NumScaleKeys = DecompContext.GetCompressedScaleOffsets()->GetOffsetData(TrackIndex, 1);
-	const uint8* RESTRICT ScaleStream = DecompContext.GetCompressedByteStream() + ScaleKeysOffset;
+	const int32 ScaleKeysOffset = AnimData.CompressedScaleOffsets.GetOffsetData(TrackIndex, 0);
+	const int32 NumScaleKeys = AnimData.CompressedScaleOffsets.GetOffsetData(TrackIndex, 1);
+	const uint8* RESTRICT ScaleStream = AnimData.CompressedByteStream.GetData() + ScaleKeysOffset;
 
 	const uint8* RESTRICT FrameTable= ScaleStream +(NumScaleKeys*CompressedScaleStrides[FORMAT]*CompressedScaleNum[FORMAT]);
 	FrameTable= Align(FrameTable, 4);
 
 	int32 Index0;
 	int32 Index1;
-	float Alpha = TimeToIndex(DecompContext.GetInterpolation(), DecompContext.GetCompressedNumberOfFrames(), FrameTable, DecompContext.RelativePos, NumScaleKeys, Index0, Index1);
+	float Alpha = TimeToIndex(DecompContext.Interpolation, AnimData.CompressedNumberOfFrames, FrameTable, DecompContext.RelativePos, NumScaleKeys, Index0, Index1);
 	const int32 ScaleStreamOffset = ((FORMAT == ACF_IntervalFixed32NoW) && NumScaleKeys > 1) ? (sizeof(float)*6) : 0; // offset past Min and Range data
 
 	if (Index0 != Index1)
@@ -476,363 +328,4 @@ FORCEINLINE_DEBUGGABLE void AEFVariableKeyLerp<FORMAT>::GetBoneAtomScale(FTransf
 		DecompressScale<FORMAT>( P0, ScaleStream, KeyData);
 		OutAtom.SetScale3D( P0 );
 	}
-}
-
-#if USE_ANIMATION_CODEC_BATCH_SOLVER
-
-/**
- * Decompress all requested rotation components from an Animation Sequence
- *
- * @param	Atoms			The FTransform array to fill in.
- * @param	DesiredPairs	Array of requested bone information
- * @param	DecompContext	The decompression context to use.
- */
-template<int32 FORMAT>
-void AEFVariableKeyLerp<FORMAT>::GetPoseRotations(
-	FTransformArray& Atoms,
-	const BoneTrackArray& DesiredPairs,
-	FAnimSequenceDecompressionContext& DecompContext)
-{
-	const int32 PairCount = DesiredPairs.Num();
-
-	for (int32 PairIndex=0; PairIndex<PairCount; ++PairIndex)
-	{
-		const BoneTrackPair& Pair = DesiredPairs[PairIndex];
-		const int32 TrackIndex = Pair.TrackIndex;
-		const int32 AtomIndex = Pair.AtomIndex;
-		FTransform& BoneAtom = Atoms[AtomIndex];
-
-		// call the decoder directly (not through the vtable)
-		AEFVariableKeyLerp<FORMAT>::GetBoneAtomRotation(BoneAtom, DecompContext, TrackIndex);
-	}
-}
-
-/**
- * Decompress all requested translation components from an Animation Sequence
- *
- * @param	Atoms			The FTransform array to fill in.
- * @param	DesiredPairs	Array of requested bone information
- * @param	DecompContext	The decompression context to use.
- */
-template<int32 FORMAT>
-void AEFVariableKeyLerp<FORMAT>::GetPoseTranslations(
-	FTransformArray& Atoms,
-	const BoneTrackArray& DesiredPairs,
-	FAnimSequenceDecompressionContext& DecompContext)
-{
-	const int32 PairCount= DesiredPairs.Num();
-
-	for (int32 PairIndex=0; PairIndex<PairCount; ++PairIndex)
-	{
-		const BoneTrackPair& Pair = DesiredPairs[PairIndex];
-		const int32 TrackIndex = Pair.TrackIndex;
-		const int32 AtomIndex = Pair.AtomIndex;
-		FTransform& BoneAtom = Atoms[AtomIndex];
-
-		// call the decoder directly (not through the vtable)
-		AEFVariableKeyLerp<FORMAT>::GetBoneAtomTranslation(BoneAtom, DecompContext, TrackIndex);
-	}
-}
-
-/**
- * Decompress all requested Scale components from an Animation Sequence
- *
- * @param	Atoms			The FTransform array to fill in.
- * @param	DesiredPairs	Array of requested bone information
- * @param	DecompContext	The decompression context to use.
- */
-template<int32 FORMAT>
-void AEFVariableKeyLerp<FORMAT>::GetPoseScales(
-	FTransformArray& Atoms,
-	const BoneTrackArray& DesiredPairs,
-	FAnimSequenceDecompressionContext& DecompContext)
-{
-	checkSlow(DecompContext.bHasScale);
-
-	const int32 PairCount= DesiredPairs.Num();
-
-	for (int32 PairIndex=0; PairIndex<PairCount; ++PairIndex)
-	{
-		const BoneTrackPair& Pair = DesiredPairs[PairIndex];
-		const int32 TrackIndex = Pair.TrackIndex;
-		const int32 AtomIndex = Pair.AtomIndex;
-		FTransform& BoneAtom = Atoms[AtomIndex];
-
-		// call the decoder directly (not through the vtable)
-		AEFVariableKeyLerp<FORMAT>::GetBoneAtomScale(BoneAtom, DecompContext, TrackIndex);
-	}
-}
-#endif
-
-#if USE_SEGMENTING_CONTEXT
-template<int32 Format>
-FQuat FAEVariableKeyLerpSortedContext::GetSortedRotation(const FAnimSequenceDecompressionContext& DecompContext, int32 TrackIndex) const
-{
-	const FCachedKey& CachedKey = CachedKeys[TrackIndex];
-
-	float Alpha;
-	if (DecompContext.NeedsTwoSegments)
-	{
-		Alpha = DecompContext.KeyAlpha;
-	}
-	else
-	{
-		// compute the blend parameters for the keys we have found
-		const int32 FrameIndex0 = SegmentStartFrame0 + CachedKey.RotFrameIndices[0];
-		const int32 FrameIndex1 = SegmentStartFrame1 + CachedKey.RotFrameIndices[1];
-
-		const int32 Delta = FMath::Max(FrameIndex1 - FrameIndex0, 1);
-		const float Remainder = FramePos - float(FrameIndex0);
-		Alpha = DecompContext.AnimSeq->Interpolation == EAnimInterpolationType::Step ? 0.0f : (Remainder / float(Delta));
-	}
-
-	const uint8* KeyData0 = DecompContext.CompressedByteStream + CachedKey.RotOffsets[0];
-
-	FQuat Rotation0;
-	DecompressRotation<Format, false>(Rotation0, DecompContext.TrackRangeData[0], KeyData0);
-
-	const uint8* KeyData1 = DecompContext.CompressedByteStream + CachedKey.RotOffsets[1];
-
-	FQuat Rotation1;
-	DecompressRotation<Format, false>(Rotation1, DecompContext.TrackRangeData[1], KeyData1);
-
-	// Fast linear quaternion interpolation.
-	FQuat BlendedQuat = FQuat::FastLerp(Rotation0, Rotation1, Alpha);
-	BlendedQuat.Normalize();
-
-	return BlendedQuat;
-}
-
-template<int32 Format>
-FVector FAEVariableKeyLerpSortedContext::GetSortedTranslation(const FAnimSequenceDecompressionContext& DecompContext, int32 TrackIndex) const
-{
-	const FCachedKey& CachedKey = CachedKeys[TrackIndex];
-
-	float Alpha;
-	if (DecompContext.NeedsTwoSegments)
-	{
-		Alpha = DecompContext.KeyAlpha;
-	}
-	else
-	{
-		// compute the blend parameters for the keys we have found
-		const int32 FrameIndex0 = SegmentStartFrame0 + CachedKey.TransFrameIndices[0];
-		const int32 FrameIndex1 = SegmentStartFrame1 + CachedKey.TransFrameIndices[1];
-
-		const int32 Delta = FMath::Max(FrameIndex1 - FrameIndex0, 1);
-		const float Remainder = FramePos - float(FrameIndex0);
-		Alpha = DecompContext.AnimSeq->Interpolation == EAnimInterpolationType::Step ? 0.0f : (Remainder / float(Delta));
-	}
-
-	const uint8* KeyData0 = DecompContext.CompressedByteStream + CachedKey.TransOffsets[0];
-
-	FVector Translation0;
-	DecompressTranslation<Format, false>(Translation0, DecompContext.TrackRangeData[0], KeyData0);
-
-	const uint8* KeyData1 = DecompContext.CompressedByteStream + CachedKey.TransOffsets[1];
-
-	FVector Translation1;
-	DecompressTranslation<Format, false>(Translation1, DecompContext.TrackRangeData[1], KeyData1);
-
-	return FMath::Lerp(Translation0, Translation1, Alpha);
-}
-
-template<int32 Format>
-FVector FAEVariableKeyLerpSortedContext::GetSortedScale(const FAnimSequenceDecompressionContext& DecompContext, int32 TrackIndex) const
-{
-	const FCachedKey& CachedKey = CachedKeys[TrackIndex];
-
-	float Alpha;
-	if (DecompContext.NeedsTwoSegments)
-	{
-		Alpha = DecompContext.KeyAlpha;
-	}
-	else
-	{
-		// compute the blend parameters for the keys we have found
-		const int32 FrameIndex0 = SegmentStartFrame0 + CachedKey.ScaleFrameIndices[0];
-		const int32 FrameIndex1 = SegmentStartFrame1 + CachedKey.ScaleFrameIndices[1];
-
-		const int32 Delta = FMath::Max(FrameIndex1 - FrameIndex0, 1);
-		const float Remainder = FramePos - float(FrameIndex0);
-		Alpha = DecompContext.AnimSeq->Interpolation == EAnimInterpolationType::Step ? 0.0f : (Remainder / float(Delta));
-	}
-
-	const uint8* KeyData0 = DecompContext.CompressedByteStream + CachedKey.ScaleOffsets[0];
-
-	FVector Scale0;
-	DecompressScale<Format, false>(Scale0, DecompContext.TrackRangeData[0], KeyData0);
-
-	const uint8* KeyData1 = DecompContext.CompressedByteStream + CachedKey.ScaleOffsets[1];
-
-	FVector Scale1;
-	DecompressScale<Format, false>(Scale1, DecompContext.TrackRangeData[1], KeyData1);
-
-	return FMath::Lerp(Scale0, Scale1, Alpha);
-}
-
-template<int32 Format>
-FQuat FAEVariableKeyLerpLinearContext::GetLinearRotation(const FAnimSequenceDecompressionContext& DecompContext, int32 TrackIndex) const
-{
-	const int32 NumTrackStreams = NumAnimatedTrackStreams[DecompContext.GetRotationValueOffset(TrackIndex)];
-
-	const uint8* OffsetNumKeysPair0 = OffsetNumKeysPairs[0] + (OffsetNumKeysPairSize * NumTrackStreams);
-	const uint16 NumKeys0 = *reinterpret_cast<const uint16*>(OffsetNumKeysPair0 + sizeof(uint32));
-
-	const uint32 KeysOffset0 = AnimationCompressionUtils::UnalignedRead<uint32>(OffsetNumKeysPair0);
-	const int32 TimeMarkersOffset0 = DecompContext.Segment0->ByteStreamOffset + KeysOffset0;
-
-	int32 FrameIndex0;
-	int32 FrameIndex1;
-	float Alpha;
-	if (DecompContext.NeedsTwoSegments)
-	{
-		FrameIndex0 = NumKeys0 - 1;
-		FrameIndex1 = 0;
-		Alpha = DecompContext.KeyAlpha;
-	}
-	else
-	{
-		const uint8* TimeMarkers0 = DecompContext.CompressedByteStream + TimeMarkersOffset0;
-
-		Alpha = AnimEncoding::TimeToIndex(DecompContext, TimeMarkers0, NumKeys0, DecompContext.Segment0->NumFrames, TimeMarkerSize[0], SegmentRelativePos0, FrameIndex0, FrameIndex1);
-	}
-
-	const int32 TrackDataOffset0 = Align(TimeMarkersOffset0 + (NumKeys0 * TimeMarkerSize[0]), 4);
-	const int32 KeyDataOffset0 = TrackDataOffset0 + (FrameIndex0 * CompressedRotationStrides[Format] * CompressedRotationNum[Format]);
-	const uint8* KeyData0 = DecompContext.CompressedByteStream + KeyDataOffset0;
-
-	FQuat Rotation;
-	DecompressRotation<Format>(Rotation, DecompContext.TrackRangeData[0], KeyData0);
-
-	if (DecompContext.NeedsInterpolation)
-	{
-		const uint8* OffsetNumKeysPair1 = OffsetNumKeysPairs[1] + (OffsetNumKeysPairSize * NumTrackStreams);
-		const uint32 KeysOffset1 = AnimationCompressionUtils::UnalignedRead<uint32>(OffsetNumKeysPair1);
-		const uint16 NumKeys1 = *reinterpret_cast<const uint16*>(OffsetNumKeysPair1 + sizeof(uint32));
-
-		const int32 TimeMarkersOffset1 = DecompContext.Segment1->ByteStreamOffset + KeysOffset1;
-		const int32 TrackDataOffset1 = Align(TimeMarkersOffset1 + (NumKeys1 * TimeMarkerSize[1]), 4);
-		const int32 KeyDataOffset1 = TrackDataOffset1 + (FrameIndex1 * CompressedRotationStrides[Format] * CompressedRotationNum[Format]);
-		const uint8* KeyData1 = DecompContext.CompressedByteStream + KeyDataOffset1;
-
-		FQuat Rotation1;
-		DecompressRotation<Format>(Rotation1, DecompContext.TrackRangeData[1], KeyData1);
-
-		// Fast linear quaternion interpolation.
-		FQuat BlendedQuat = FQuat::FastLerp(Rotation, Rotation1, Alpha);
-		BlendedQuat.Normalize();
-		Rotation = BlendedQuat;
-	}
-
-	return Rotation;
-}
-
-template<int32 Format>
-FVector FAEVariableKeyLerpLinearContext::GetLinearTranslation(const FAnimSequenceDecompressionContext& DecompContext, int32 TrackIndex) const
-{
-	const int32 NumTrackStreams = NumAnimatedTrackStreams[DecompContext.GetTranslationValueOffset(TrackIndex)];
-
-	const uint8* OffsetNumKeysPair0 = OffsetNumKeysPairs[0] + (OffsetNumKeysPairSize * NumTrackStreams);
-	const uint16 NumKeys0 = *reinterpret_cast<const uint16*>(OffsetNumKeysPair0 + sizeof(uint32));
-
-	const uint32 KeysOffset0 = AnimationCompressionUtils::UnalignedRead<uint32>(OffsetNumKeysPair0);
-	const int32 TimeMarkersOffset0 = DecompContext.Segment0->ByteStreamOffset + KeysOffset0;
-
-	int32 FrameIndex0;
-	int32 FrameIndex1;
-	float Alpha;
-	if (DecompContext.NeedsTwoSegments)
-	{
-		FrameIndex0 = NumKeys0 - 1;
-		FrameIndex1 = 0;
-		Alpha = DecompContext.KeyAlpha;
-	}
-	else
-	{
-		const uint8* TimeMarkers0 = DecompContext.CompressedByteStream + TimeMarkersOffset0;
-
-		Alpha = AnimEncoding::TimeToIndex(DecompContext, TimeMarkers0, NumKeys0, DecompContext.Segment0->NumFrames, TimeMarkerSize[0], SegmentRelativePos0, FrameIndex0, FrameIndex1);
-	}
-
-	const int32 TrackDataOffset0 = Align(TimeMarkersOffset0 + (NumKeys0 * TimeMarkerSize[0]), 4);
-	const int32 KeyDataOffset0 = TrackDataOffset0 + (FrameIndex0 * CompressedTranslationStrides[Format] * CompressedTranslationNum[Format]);
-	const uint8* KeyData0 = DecompContext.CompressedByteStream + KeyDataOffset0;
-
-	FVector Translation;
-	DecompressTranslation<Format>(Translation, DecompContext.TrackRangeData[0], KeyData0);
-
-	if (DecompContext.NeedsInterpolation)
-	{
-		const uint8* OffsetNumKeysPair1 = OffsetNumKeysPairs[1] + (OffsetNumKeysPairSize * NumTrackStreams);
-		const uint32 KeysOffset1 = AnimationCompressionUtils::UnalignedRead<uint32>(OffsetNumKeysPair1);
-		const uint16 NumKeys1 = *reinterpret_cast<const uint16*>(OffsetNumKeysPair1 + sizeof(uint32));
-
-		const int32 TimeMarkersOffset1 = DecompContext.Segment1->ByteStreamOffset + KeysOffset1;
-		const int32 TrackDataOffset1 = Align(TimeMarkersOffset1 + (NumKeys1 * TimeMarkerSize[1]), 4);
-		const int32 KeyDataOffset1 = TrackDataOffset1 + (FrameIndex1 * CompressedTranslationStrides[Format] * CompressedTranslationNum[Format]);
-		const uint8* KeyData1 = DecompContext.CompressedByteStream + KeyDataOffset1;
-
-		FVector Translation1;
-		DecompressTranslation<Format>(Translation1, DecompContext.TrackRangeData[1], KeyData1);
-
-		Translation = FMath::Lerp(Translation, Translation1, Alpha);
-	}
-
-	return Translation;
-}
-
-template<int32 Format>
-FVector FAEVariableKeyLerpLinearContext::GetLinearScale(const FAnimSequenceDecompressionContext& DecompContext, int32 TrackIndex) const
-{
-	const int32 NumTrackStreams = NumAnimatedTrackStreams[DecompContext.GetScaleValueOffset(TrackIndex)];
-
-	const uint8* OffsetNumKeysPair0 = OffsetNumKeysPairs[0] + (OffsetNumKeysPairSize * NumTrackStreams);
-	const uint16 NumKeys0 = *reinterpret_cast<const uint16*>(OffsetNumKeysPair0 + sizeof(uint32));
-
-	const uint32 KeysOffset0 = AnimationCompressionUtils::UnalignedRead<uint32>(OffsetNumKeysPair0);
-	const int32 TimeMarkersOffset0 = DecompContext.Segment0->ByteStreamOffset + KeysOffset0;
-
-	int32 FrameIndex0;
-	int32 FrameIndex1;
-	float Alpha;
-	if (DecompContext.NeedsTwoSegments)
-	{
-		FrameIndex0 = NumKeys0 - 1;
-		FrameIndex1 = 0;
-		Alpha = DecompContext.KeyAlpha;
-	}
-	else
-	{
-		const uint8* TimeMarkers0 = DecompContext.CompressedByteStream + TimeMarkersOffset0;
-
-		Alpha = AnimEncoding::TimeToIndex(DecompContext, TimeMarkers0, NumKeys0, DecompContext.Segment0->NumFrames, TimeMarkerSize[0], SegmentRelativePos0, FrameIndex0, FrameIndex1);
-	}
-
-	const int32 TrackDataOffset0 = Align(TimeMarkersOffset0 + (NumKeys0 * TimeMarkerSize[0]), 4);
-	const int32 KeyDataOffset0 = TrackDataOffset0 + (FrameIndex0 * CompressedScaleStrides[Format] * CompressedScaleNum[Format]);
-	const uint8* KeyData0 = DecompContext.CompressedByteStream + KeyDataOffset0;
-
-	FVector Scale;
-	DecompressScale<Format>(Scale, DecompContext.TrackRangeData[0], KeyData0);
-
-	if (DecompContext.NeedsInterpolation)
-	{
-		const uint8* OffsetNumKeysPair1 = OffsetNumKeysPairs[1] + (OffsetNumKeysPairSize * NumTrackStreams);
-		const uint32 KeysOffset1 = AnimationCompressionUtils::UnalignedRead<uint32>(OffsetNumKeysPair1);
-		const uint16 NumKeys1 = *reinterpret_cast<const uint16*>(OffsetNumKeysPair1 + sizeof(uint32));
-
-		const int32 TimeMarkersOffset1 = DecompContext.Segment1->ByteStreamOffset + KeysOffset1;
-		const int32 TrackDataOffset1 = Align(TimeMarkersOffset1 + (NumKeys1 * TimeMarkerSize[1]), 4);
-		const int32 KeyDataOffset1 = TrackDataOffset1 + (FrameIndex1 * CompressedScaleStrides[Format] * CompressedScaleNum[Format]);
-		const uint8* KeyData1 = DecompContext.CompressedByteStream + KeyDataOffset1;
-
-		FVector Scale1;
-		DecompressScale<Format>(Scale1, DecompContext.TrackRangeData[1], KeyData1);
-
-		Scale = FMath::Lerp(Scale, Scale1, Alpha);
-	}
-
-	return Scale;
-}
-#endif
\ No newline at end of file
+}
\ No newline at end of file
diff --git a/Engine/Source/Runtime/Engine/Public/Animation/AnimCompressionDerivedDataPublic.h b/Engine/Source/Runtime/Engine/Public/Animation/AnimCompressionDerivedDataPublic.h
new file mode 100644
--- /dev/null
+++ b/Engine/Source/Runtime/Engine/Public/Animation/AnimCompressionDerivedDataPublic.h
@@ -0,0 +1,86 @@
+// Copyright Epic Games, Inc. All Rights Reserved.
+
+#pragma once
+
+#if WITH_EDITOR
+
+#include "TickableEditorObject.h"
+
+class UAnimSequence;
+class FDerivedDataAnimationCompression;
+
+extern ENGINE_API class FAsyncCompressedAnimationsManagement* GAsyncCompressedAnimationsTracker;
+
+// Animation data that is currently being compressed
+struct FActiveAsyncCompressionTask
+{
+	FActiveAsyncCompressionTask(UAnimSequence* InSequence, const FString& InCacheKey, const uint64 InTaskSize, const uint32 InAsyncHandle, bool bInPerformFrameStripping)
+		: Sequence(InSequence)
+		, TaskSize(InTaskSize)
+		, CacheKey(InCacheKey)
+		, AsyncHandle(InAsyncHandle)
+		, bPerformFrameStripping(bInPerformFrameStripping)
+	{}
+	UAnimSequence* Sequence;
+	uint64 TaskSize;
+	FString CacheKey;
+	uint32 AsyncHandle;
+	bool bPerformFrameStripping;
+};
+
+// An animation waiting to be compressed
+struct FQueuedAsyncCompressionWork
+{
+	FQueuedAsyncCompressionWork(FDerivedDataAnimationCompression& InCompressor, UAnimSequence* InAnim, const bool bInPerformFrameStripping)
+		: Compressor(InCompressor)
+		, Anim(InAnim)
+		, bPerformFrameStripping(bInPerformFrameStripping)
+	{}
+
+	FDerivedDataAnimationCompression& Compressor;
+	UAnimSequence* Anim;
+	const bool bPerformFrameStripping;
+};
+
+
+// Manager for Async anim compression.
+//   Maintains active compressions
+//   tracks memory usage of async compression
+//   Gives API for blocking on compression
+class ENGINE_API FAsyncCompressedAnimationsManagement : public FTickableEditorObject
+{
+public:
+	static FAsyncCompressedAnimationsManagement& Get();
+
+	// Request an async compression of an animation, may not actually run async if memory usage is already high
+	// Returns true if an async compression task was allowed.
+	bool RequestAsyncCompression(FDerivedDataAnimationCompression& Compressor, UAnimSequence* Anim, const bool bPerformFrameStripping, TArray<uint8>& OutData);
+
+	// Returns the number of remaining compression jobs (used for UI)
+	int32 GetNumRemainingJobs() const { return ActiveAsyncCompressionTasks.Num() + QueuedAsyncCompressionWork.Num(); }
+	
+	// Blocks on compression for the supplied animation. Returns false if there was no compression job to wait on. 
+	bool WaitOnExistingCompression(UAnimSequence* Anim, const bool bCancelIfNotStarted);
+private:
+
+	void StartAsyncWork(FDerivedDataAnimationCompression& Compressor, UAnimSequence* Anim, const uint64 NewTaskSize, const bool bPerformFrameStripping);
+
+	bool WaitOnActiveCompression(UAnimSequence* Anim);
+
+	void OnActiveCompressionFinished(int32 ActiveAnimIndex);
+
+	FAsyncCompressedAnimationsManagement() : ActiveMemoryUsage(0) {}
+
+	//Array of active and queued compression jobs
+	TArray<FActiveAsyncCompressionTask> ActiveAsyncCompressionTasks;
+	TArray<FQueuedAsyncCompressionWork> QueuedAsyncCompressionWork;
+	uint64 ActiveMemoryUsage;
+
+	/* Begin FTickableEditorObject */
+	virtual void Tick(float DeltaTime) override;
+	virtual ETickableTickType GetTickableTickType() const override { return ETickableTickType::Always; }
+	virtual TStatId GetStatId() const override;
+	/* End FTickableEditorObject */
+};
+
+#endif
\ No newline at end of file
diff --git a/Engine/Source/Runtime/Engine/Public/Animation/AnimCompressionTypes.h b/Engine/Source/Runtime/Engine/Public/Animation/AnimCompressionTypes.h
--- a/Engine/Source/Runtime/Engine/Public/Animation/AnimCompressionTypes.h
+++ b/Engine/Source/Runtime/Engine/Public/Animation/AnimCompressionTypes.h
@@ -41,6 +41,8 @@ class FMemoryWriter;
 
 class UAnimCompress;
 class UAnimCurveCompressionSettings;
+class UAnimBoneCompressionSettings;
+class UAnimBoneCompressionCodec;
 class USkeleton;
 
 extern FGuid GenerateGuidFromRawAnimData(const TArray<FRawAnimSequenceTrack>& RawAnimationData, const FRawCurveTracks& RawCurveData);
@@ -101,6 +103,33 @@ struct ENGINE_API FCompressedOffsetDataBase
 	}
 };
 
+/**
+* A set of error statistics for an animation, gathered by FAnimationUtils::ComputeCompressionError
+*/
+USTRUCT()
+struct FAnimationErrorStats
+{
+	GENERATED_USTRUCT_BODY()
+
+	/** Average world-space translation error across all end-effectors **/
+	float AverageError;
+	/** The worst error encountered across all end effectors **/
+	float MaxError;
+	/** Time at which the worst error occurred */
+	float MaxErrorTime;
+	/** Bone on which the worst error occurred */
+	int32 MaxErrorBone;
+
+	FAnimationErrorStats()
+		: AverageError(0.f)
+		, MaxError(0.f)
+		, MaxErrorTime(0.f)
+		, MaxErrorBone(0)
+	{}
+};
+
+FArchive& operator<<(FArchive& Ar, FAnimationErrorStats& ErrorStats);
+
 // Helper for buiilding DDC keys of settings
 struct FArcToHexString
 {
@@ -131,66 +160,23 @@ public:
 
 struct FCompressedOffsetData : public FCompressedOffsetDataBase<TArray<int32>>
 {
-
 };
 
-
-/**
- * Represents a segment of the anim sequence that is compressed.
- */
-USTRUCT()
-struct ENGINE_API FCompressedSegment
-{
-	GENERATED_USTRUCT_BODY()
-
-	// Frame where the segment begins in the anim sequence
-	int32 StartFrame;
-
-	// Num of frames contained in the segment
-	int32 NumFrames;
-
-	// Segment data offset in CompressedByteStream
-	int32 ByteStreamOffset;
-
-	/** The compression format that was used to compress translation tracks. */
-	TEnumAsByte<enum AnimationCompressionFormat> TranslationCompressionFormat;
-
-	/** The compression format that was used to compress rotation tracks. */
-	TEnumAsByte<enum AnimationCompressionFormat> RotationCompressionFormat;
-
-	/** The compression format that was used to compress rotation tracks. */
-	TEnumAsByte<enum AnimationCompressionFormat> ScaleCompressionFormat;
-
-	FCompressedSegment()
-		: StartFrame(0)
-		, NumFrames(0)
-		, ByteStreamOffset(0)
-		, TranslationCompressionFormat(ACF_None)
-		, RotationCompressionFormat(ACF_None)
-		, ScaleCompressionFormat(ACF_None)
-	{
-	}
-
-	friend FArchive& operator<<(FArchive& Ar, FCompressedSegment &Segment)
-	{
-		return Ar << Segment.StartFrame << Segment.NumFrames << Segment.ByteStreamOffset
-			<< Segment.TranslationCompressionFormat << Segment.RotationCompressionFormat << Segment.ScaleCompressionFormat;
-	}
-};
+FArchive& operator<<(FArchive& Ar, FCompressedOffsetData& D);
 
 struct ENGINE_API FCompressibleAnimData
 {
 public:
 	FCompressibleAnimData();
 
-	FCompressibleAnimData(UAnimCompress* InRequestedCompressionScheme, UAnimCurveCompressionSettings* InCurveCompressionSettings, USkeleton* InSkeleton, EAnimInterpolationType InInterpolation, float InSequenceLength, int32 InNumFrames, const float InAltCompressionErrorThreshold);
+	FCompressibleAnimData(UAnimBoneCompressionSettings* InBoneCompressionSettings, UAnimCurveCompressionSettings* InCurveCompressionSettings, USkeleton* InSkeleton, EAnimInterpolationType InInterpolation, float InSequenceLength, int32 InNumFrames);
 
-	FCompressibleAnimData(class UAnimSequence* InSeq, const bool bPerformStripping, const float InAltCompressionErrorThreshold);
-
-	UAnimCompress* RequestedCompressionScheme;
+	FCompressibleAnimData(class UAnimSequence* InSeq, const bool bPerformStripping);
 
 	UAnimCurveCompressionSettings* CurveCompressionSettings;
 
+	UAnimBoneCompressionSettings* BoneCompressionSettings;
+
 	USkeleton* Skeleton;
 
 	TArray<FTrackToSkeletonMap> TrackToSkeletonMapTable;
@@ -211,7 +197,7 @@ public:
 
 	bool bIsValidAdditive;
 
-	float AltCompressionErrorThreshold;
+	float ErrorThresholdScale;
 
 	//For Logging
 	FString Name;
@@ -300,10 +286,6 @@ struct FCompressedAnimDataBase
 	 */
 	FCompressedOffsetDataBase<typename ContainerTypeMakerTemplate<int32>::Type>  CompressedScaleOffsets;
 
-	/**
-	 * ByteStream for compressed animation data.
-	 * The memory layout is dependent on the algorithm used to compress the anim sequence.
-	 */
 	typename ContainerTypeMakerTemplate<uint8>::Type CompressedByteStream;
 
 	/**
@@ -320,8 +302,6 @@ struct FCompressedAnimDataBase
 	AnimationCompressionFormat TranslationCompressionFormat;
 	AnimationCompressionFormat RotationCompressionFormat;
 	AnimationCompressionFormat ScaleCompressionFormat;
-
-	int32 CompressedNumberOfFrames;
 	
 	template <template <typename> class OtherContainerMaker>
 	explicit FCompressedAnimDataBase(FCompressedAnimDataBase<OtherContainerMaker>& InCompressedData)
@@ -336,8 +316,6 @@ struct FCompressedAnimDataBase
 		, TranslationCompressionFormat(InCompressedData.TranslationCompressionFormat)
 		, RotationCompressionFormat(InCompressedData.RotationCompressionFormat)
 		, ScaleCompressionFormat(InCompressedData.ScaleCompressionFormat)
-
-		, CompressedNumberOfFrames(InCompressedData.CompressedNumberOfFrames)
 	{
 		CompressedScaleOffsets.OffsetData = InCompressedData.CompressedScaleOffsets.OffsetData;
 		CompressedScaleOffsets.StripSize = InCompressedData.CompressedScaleOffsets.StripSize;
@@ -352,82 +330,48 @@ struct FCompressedAnimDataBase
 		, TranslationCompressionFormat((AnimationCompressionFormat)0)
 		, RotationCompressionFormat((AnimationCompressionFormat)0)
 		, ScaleCompressionFormat((AnimationCompressionFormat)0)
-
-		, CompressedNumberOfFrames(0)
 	{
-
-	}
-
-	template <template <typename> class OtherContainerMaker>
-	void CopyFromSettings(const FCompressedAnimDataBase<OtherContainerMaker>& Other)
-	{
-		TranslationCodec = Other.TranslationCodec;
-		RotationCodec = Other.RotationCodec;
-		ScaleCodec = Other.ScaleCodec;
-
-		KeyEncodingFormat = Other.KeyEncodingFormat;
-		TranslationCompressionFormat = Other.TranslationCompressionFormat;
-		RotationCompressionFormat = Other.RotationCompressionFormat;
-		ScaleCompressionFormat = Other.ScaleCompressionFormat;
-
-		CompressedNumberOfFrames = Other.CompressedNumberOfFrames;
-	}
-
-	int64 GetApproxBoneCompressedSize() const
-	{
-		return (int64)CompressedTrackOffsets.GetTypeSize()*(int64)CompressedTrackOffsets.Num() + (int64)CompressedByteStream.Num() + (int64)CompressedScaleOffsets.GetMemorySize();
-	}
-
-	bool IsCompressedDataValid() const
-	{
-		return CompressedByteStream.Num() > 0 || (TranslationCompressionFormat == ACF_Identity && RotationCompressionFormat == ACF_Identity && ScaleCompressionFormat == ACF_Identity);
 	}
 };
 
-struct ENGINE_API FCompressibleAnimDataResult : public FCompressedAnimDataBase<TArrayMaker>
+struct ENGINE_API ICompressedAnimData
 {
-	class UAnimCompress* CompressionScheme;
-
-	FCompressibleAnimDataResult()
-		: CompressionScheme(nullptr)
-	{
-
-	}
-
-	template <template <typename> class OtherContainerMaker>
-	explicit FCompressibleAnimDataResult(FCompressedAnimDataBase<OtherContainerMaker>& InCompressedData)
-		: FCompressedAnimDataBase(InCompressedData)
-	{}
+	/* Common data */
+	int32 CompressedNumberOfFrames;
 
-	void CopyFrom(const FCompressibleAnimDataResult& Other)
-	{
-		CompressedTrackOffsets = Other.CompressedTrackOffsets;
-		CompressedByteStream = Other.CompressedByteStream;
-		CompressedScaleOffsets.OffsetData = Other.CompressedScaleOffsets.OffsetData;
-		CompressedScaleOffsets.StripSize = Other.CompressedScaleOffsets.StripSize;
+#if WITH_EDITORONLY_DATA
+	/** The error stats from the current bone compression codec. */
+	FAnimationErrorStats BoneCompressionErrorStats;
+#endif
 
-		CopyFromSettings(Other);
-	}
+	/* Virtual interface codecs must implement */
+	virtual ~ICompressedAnimData() {}
 
-	FCompressibleAnimDataResult& operator=(const FCompressibleAnimDataResult& Other)
-	{
-		CopyFrom(Other);
-		CompressionScheme = Other.CompressionScheme;
-		return *this;
-	}
+	virtual void SerializeCompressedData(class FArchive& Ar);
+	virtual void Bind(const TArrayView<uint8> BulkData) = 0;
 
-	void BuildFinalBuffer(TArray<uint8>& OutBuffer);
+	virtual int64 GetApproxCompressedSize() const = 0;
+	virtual FString GetDebugString() const { return FString(); }
+	virtual bool IsValid() const = 0;
 };
 
-struct ICompressedAnimData
+struct ENGINE_API FCompressibleAnimDataResult
 {
-public:
-	virtual ~ICompressedAnimData() {}
+	TArray<uint8> CompressedByteStream;
+	TUniquePtr<ICompressedAnimData> AnimData;
+
+	UAnimBoneCompressionCodec* Codec;
 
-	virtual void SerializeCompressedData(class FArchive& Ar) = 0;
+	// Give compression code a place to store some extra data during compression
+	// (Needed for PerTrackCompression)
+	void* CompressionUserData;
 
-	virtual void ByteSwapIn(TArrayView<uint8> CompressedData, FMemoryReader& MemoryStream) = 0;
-	virtual void ByteSwapOut(TArrayView<uint8> CompressedData, FMemoryWriter& MemoryStream) = 0;
+	FCompressibleAnimDataResult()
+		: CompressedByteStream()
+		, AnimData()
+		, Codec(nullptr)
+		, CompressionUserData(nullptr)
+	{}
 };
 
 template<typename T>
@@ -442,30 +386,42 @@ TArrayView<T> RebaseTArrayView(const TArrayView<T>& ArrayView, const uint8* Orig
 	return ArrayView;
 }
 
-struct FUECompressedAnimData : public ICompressedAnimData, public FCompressedAnimDataBase<TNonConstArrayViewMaker>
+struct ENGINE_API FUECompressedAnimDataMutable : public ICompressedAnimData, public FCompressedAnimDataBase<TArrayMaker>
 {
-	FUECompressedAnimData() = default;
+	FUECompressedAnimDataMutable() = default;
 
-	template <template <typename> class OtherContainerMaker>
-	explicit FUECompressedAnimData(FCompressedAnimDataBase<OtherContainerMaker>& InCompressedData)
-		: FCompressedAnimDataBase(InCompressedData)
-	{}
+	void BuildFinalBuffer(TArray<uint8>& OutCompressedByteStream);
 
-	void Reset();
+	// ICompressedAnimData implementation
+	virtual void Bind(const TArrayView<uint8> BulkData) {}
+	virtual int64 GetApproxCompressedSize() const override { return (int64)CompressedTrackOffsets.GetTypeSize() * (int64)CompressedTrackOffsets.Num() + (int64)CompressedByteStream.Num() + (int64)CompressedScaleOffsets.GetMemorySize(); }
+	virtual bool IsValid() const override { return CompressedByteStream.Num() > 0 || (TranslationCompressionFormat == ACF_Identity && RotationCompressionFormat == ACF_Identity && ScaleCompressionFormat == ACF_Identity); }
+};
 
-	void InitViewsFromBuffer(const TArrayView<uint8> BulkData);
+struct ENGINE_API FUECompressedAnimData : public ICompressedAnimData, public FCompressedAnimDataBase<TNonConstArrayViewMaker>
+{
+	FUECompressedAnimData() = default;
 
-#if WITH_EDITOR
-	void CopyFrom(const FCompressibleAnimDataResult& Other);
-#endif
+	explicit FUECompressedAnimData(FUECompressedAnimDataMutable& InCompressedData)
+		: ICompressedAnimData(InCompressedData)
+		, FCompressedAnimDataBase(InCompressedData)
+	{}
 
-	virtual void SerializeCompressedData(class FArchive& Ar);
+	void InitViewsFromBuffer(const TArrayView<uint8> BulkData);
 
 	template<typename TArchive>
 	void ByteSwapData(TArrayView<uint8> CompresedData, TArchive& MemoryStream);
 
-	virtual void ByteSwapIn(TArrayView<uint8> CompressedData, FMemoryReader& MemoryStream) override { ByteSwapData(CompressedData, MemoryStream); }
-	virtual void ByteSwapOut(TArrayView<uint8> CompressedData, FMemoryWriter& MemoryStream) override { ByteSwapData(CompressedData, MemoryStream); }
+	void ByteSwapIn(TArrayView<uint8> CompressedData, FMemoryReader& MemoryStream) { ByteSwapData(CompressedData, MemoryStream); }
+	void ByteSwapOut(TArrayView<uint8> CompressedData, FMemoryWriter& MemoryStream) { ByteSwapData(CompressedData, MemoryStream); }
+
+	// ICompressedAnimData implementation
+	virtual void SerializeCompressedData(class FArchive& Ar) override;
+	virtual void Bind(const TArrayView<uint8> BulkData) override { InitViewsFromBuffer(BulkData); }
+
+	virtual int64 GetApproxCompressedSize() const override { return (int64)CompressedTrackOffsets.GetTypeSize() * (int64)CompressedTrackOffsets.Num() + (int64)CompressedByteStream.Num() + (int64)CompressedScaleOffsets.GetMemorySize(); }
+	virtual FString GetDebugString() const override;
+	virtual bool IsValid() const override { return CompressedByteStream.Num() > 0 || (TranslationCompressionFormat == ACF_Identity && RotationCompressionFormat == ACF_Identity && ScaleCompressionFormat == ACF_Identity); }
 };
 
 template<uint32 Alignment = DEFAULT_ALIGNMENT>
@@ -694,15 +650,25 @@ public:
 	/* Compressed curve data stream used by AnimCurveCompressionCodec */
 	TArray<uint8> CompressedCurveByteStream;
 
-	FUECompressedAnimData CompressedDataStructure;
+	/* Compressed bone data stream used by AnimBoneCompressionCodec */
+	TUniquePtr<ICompressedAnimData> CompressedDataStructure;
+
+	/** The codec used by the compressed data as determined by the bone compression settings. */
+	class UAnimBoneCompressionCodec* BoneCompressionCodec;
 
-	/** The codec used by the compressed data as determined by the compression settings. */
+	/** The codec used by the compressed data as determined by the curve compression settings. */
 	class UAnimCurveCompressionCodec* CurveCompressionCodec;
 
 	// The size of the raw data used to create the compressed data
 	int32 CompressedRawDataSize;
 
-	void SerializeCompressedData(FArchive& Ar, bool bDDCData, UObject* DataOwner, USkeleton* Skeleton, class UAnimCurveCompressionSettings* CurveCompressionSettings, bool bCanUseBulkData=true);
+	FCompressedAnimSequence()
+		: BoneCompressionCodec(nullptr)
+		, CurveCompressionCodec(nullptr)
+		, CompressedRawDataSize(0)
+	{}
+
+	void SerializeCompressedData(FArchive& Ar, bool bDDCData, UObject* DataOwner, USkeleton* Skeleton, UAnimBoneCompressionSettings* BoneCompressionSettings, UAnimCurveCompressionSettings* CurveCompressionSettings, bool bCanUseBulkData=true);
 
 	int32 GetSkeletonIndexFromTrackIndex(const int32 TrackIndex) const
 	{
@@ -711,6 +677,9 @@ public:
 
 	// Return the number of bytes used
 	SIZE_T GetMemorySize() const;
+
+	void ClearCompressedBoneData();
+	void ClearCompressedCurveData();
 };
 
 struct FRootMotionReset
diff --git a/Engine/Source/Runtime/Engine/Public/Animation/AnimCurveTypes.h b/Engine/Source/Runtime/Engine/Public/Animation/AnimCurveTypes.h
--- a/Engine/Source/Runtime/Engine/Public/Animation/AnimCurveTypes.h
+++ b/Engine/Source/Runtime/Engine/Public/Animation/AnimCurveTypes.h
@@ -83,6 +83,11 @@ struct ENGINE_API FAnimCurveBase
 	UPROPERTY()
 	FSmartName	Name;
 
+#if WITH_EDITORONLY_DATA
+	UPROPERTY()
+	FLinearColor Color;
+#endif
+
 private:
 	// this flag is mostly used by editor only now
 	// however I can't remove this to editor only because 
@@ -93,12 +98,21 @@ private:
 	int32		CurveTypeFlags;
 
 public:
-	FAnimCurveBase() : CurveTypeFlags(0) {}
+	FAnimCurveBase() 
+		: CurveTypeFlags(0)
+	{
+#if WITH_EDITORONLY_DATA
+		Color = MakeColor();
+#endif
+	}
 
 	FAnimCurveBase(FSmartName InName, int32 InCurveTypeFlags)
 		: Name(InName)
 		, CurveTypeFlags(InCurveTypeFlags)
-	{	
+	{
+#if WITH_EDITORONLY_DATA
+		Color = MakeColor();
+#endif
 	}
 
 	// To be able to use typedef'd types we need to serialize manually
@@ -128,6 +142,15 @@ public:
 	 * returns CurveTypeFlags
 	 */
 	int32 GetCurveTypeFlags() const;
+
+#if WITH_EDITORONLY_DATA
+	/** Get the color used to display this curve in the editor */
+	FLinearColor GetColor() const { return Color; }
+
+private:
+	/** Make an initial color */
+	FLinearColor MakeColor();
+#endif
 };
 
 USTRUCT()
@@ -482,6 +505,68 @@ struct FBaseBlendedCurve
 	/**
 	 * This doesn't blend but combine MAX(current weight, curvetocombine weight)
 	 */
+	void UseMaxValue(const FBaseBlendedCurve& CurveToCombine)
+	{
+		check(bInitialized);
+		check(Num() == CurveToCombine.Num());
+
+		for (int32 CurveId = 0; CurveId < CurveToCombine.Elements.Num(); ++CurveId)
+		{
+			// if target value is valid, we accept target value
+			if (Elements[CurveId].bValid && CurveToCombine.Elements[CurveId].bValid)
+			{
+				Elements[CurveId].Value = FMath::Max(Elements[CurveId].Value, CurveToCombine.Elements[CurveId].Value);
+			}
+			else if (CurveToCombine.Elements[CurveId].bValid)
+			{
+				Elements[CurveId].SetValue(CurveToCombine.Elements[CurveId].Value);
+			}
+		}
+	}
+
+	/**
+	 * This doesn't blend but combine MIN(current weight, curvetocombine weight)
+	 */
+	void UseMinValue(const FBaseBlendedCurve& CurveToCombine)
+	{
+		check(bInitialized);
+		check(Num() == CurveToCombine.Num());
+
+		for (int32 CurveId = 0; CurveId < CurveToCombine.Elements.Num(); ++CurveId)
+		{
+			// if target value is valid, we accept target value
+			if (Elements[CurveId].bValid && CurveToCombine.Elements[CurveId].bValid)
+			{
+				Elements[CurveId].Value = FMath::Min(Elements[CurveId].Value, CurveToCombine.Elements[CurveId].Value);
+			}
+			else if (CurveToCombine.Elements[CurveId].bValid)
+			{
+				Elements[CurveId].SetValue(CurveToCombine.Elements[CurveId].Value);
+			}
+		}
+	}
+
+	/**
+	 * This combines IF the input does not contain valid curve
+	 */
+	void CombinePreserved(const FBaseBlendedCurve& CurveToCombine)
+	{
+		check(bInitialized);
+		check(Num() == CurveToCombine.Num());
+
+		for (int32 CurveId = 0; CurveId < CurveToCombine.Elements.Num(); ++CurveId)
+		{
+			// if target value is valid, we accept target value
+			if (!Elements[CurveId].bValid && CurveToCombine.Elements[CurveId].bValid)
+			{
+				Elements[CurveId].SetValue(CurveToCombine.Elements[CurveId].Value);
+			}
+
+		}
+	}
+	/**
+	 * This doesn't blend but later pose with valid curve value overrides
+	 */
 	void Combine(const FBaseBlendedCurve& CurveToCombine)
 	{
 		check(bInitialized);
@@ -492,8 +577,7 @@ struct FBaseBlendedCurve
 			// if target value is valid, we accept target value
 			if (CurveToCombine.Elements[CurveId].bValid)
 			{
-				Elements[CurveId].bValid = true;
-				Elements[CurveId].Value = CurveToCombine.Elements[CurveId].Value;
+				Elements[CurveId].SetValue(CurveToCombine.Elements[CurveId].Value);
 			}
 
 		}
diff --git a/Engine/Source/Runtime/Engine/Public/Animation/AnimInstanceProxy.h b/Engine/Source/Runtime/Engine/Public/Animation/AnimInstanceProxy.h
--- a/Engine/Source/Runtime/Engine/Public/Animation/AnimInstanceProxy.h
+++ b/Engine/Source/Runtime/Engine/Public/Animation/AnimInstanceProxy.h
@@ -772,7 +772,7 @@ protected:
 	void InitializeRootNode_WithRoot(FAnimNode_Base* InRootNode);
 
 	/** Manually add object references to GC */
-	void AddReferencedObjects(UAnimInstance* InAnimInstance, FReferenceCollector& Collector);
+	virtual void AddReferencedObjects(UAnimInstance* InAnimInstance, FReferenceCollector& Collector);
 
 	/** Allow nodes to register log messages to be processed on the game thread */
 	void LogMessage(FName InLogType, EMessageSeverity::Type InSeverity, const FText& InMessage) const;
@@ -796,6 +796,13 @@ protected:
 	/** Add a curve value */
 	void AddCurveValue(const FSmartNameMapping& Mapping, const FName& CurveName, float Value);
 
+	/** Custom proxy Init/Cache/Update/Evaluate functions */
+	static void InitializeInputProxy(FAnimInstanceProxy* InputProxy, UAnimInstance* InAnimInstance);
+	static void GatherInputProxyDebugData(FAnimInstanceProxy* InputProxy, FNodeDebugData& DebugData);
+	static void CacheBonesInputProxy(FAnimInstanceProxy* InputProxy);
+	static void UpdateInputProxy(FAnimInstanceProxy* InputProxy, const FAnimationUpdateContext& Context);
+	static void EvaluateInputProxy(FAnimInstanceProxy* InputProxy, FPoseContext& Output);
+
 private:
 	/** The component to world transform of the component we are running on */
 	FTransform ComponentTransform;
diff --git a/Engine/Source/Runtime/Engine/Public/Animation/AnimSequenceDecompressionContext.h b/Engine/Source/Runtime/Engine/Public/Animation/AnimSequenceDecompressionContext.h
--- a/Engine/Source/Runtime/Engine/Public/Animation/AnimSequenceDecompressionContext.h
+++ b/Engine/Source/Runtime/Engine/Public/Animation/AnimSequenceDecompressionContext.h
@@ -3,229 +3,29 @@
 #pragma once
 
 #include "CoreMinimal.h"
-#include "Animation/AnimEncodingHeapAllocator.h"
-#include "AnimationCompression.h"
-#include "Containers/StaticArray.h"
-#include "Animation/AnimSequence.h"
+#include "Animation/AnimTypes.h"
 
-struct FCompressedOffsetData;
-struct FCompressedSegment;
-class FAnimEncodingDecompressionContext;
-class AnimEncoding;
+struct ICompressedAnimData;
 
-/**
- * Structure to wrap a trivial key handle. Used by the decompression context and encoders.
- */
-struct FTrivialAnimKeyHandle
+/* Encapsulates decompression related data used by bone compression codecs. */
+struct FAnimSequenceDecompressionContext
 {
-	constexpr FTrivialAnimKeyHandle() : Value(-1) {}
-	constexpr explicit FTrivialAnimKeyHandle(int32 InValue) : Value(InValue) {}
-
-	constexpr bool IsValid() const { return Value >= 0; }
-
-	int32 Value;
-};
-
-#if !USE_SEGMENTING_CONTEXT
-// Lightweight decompression context to match API of when USE_SEGMENT_CONTEXT is 1
-struct ENGINE_API FAnimSequenceDecompressionContext
-{
-public:
-#if WITH_EDITOR
-	FAnimSequenceDecompressionContext(const FCompressibleAnimData& CompressibleAnimData_, const FUECompressedAnimData& CompressedAnimData_) : SequenceLength(CompressibleAnimData_.SequenceLength), Interpolation(CompressibleAnimData_.Interpolation), AnimName(CompressibleAnimData_.AnimFName), CompressedAnimData(CompressedAnimData_), Time(0.f) {}
-#endif
-	FAnimSequenceDecompressionContext(float SequenceLength_, EAnimInterpolationType Interpolation_, const FName& AnimName_, const FUECompressedAnimData& CompressedAnimData_) : SequenceLength(SequenceLength_), Interpolation(Interpolation_), AnimName(AnimName_), CompressedAnimData(CompressedAnimData_), Time(0.f) {}
-	/**
-	* We don't disable copying but we also don't retain any state. Copying is required by the UE serialization
-	* but retaining state is not necessary and can cause issues.
-	*/
-	FAnimSequenceDecompressionContext(const FAnimSequenceDecompressionContext& Other) : SequenceLength(Other.SequenceLength), Interpolation(Other.Interpolation), AnimName(Other.AnimName), CompressedAnimData(Other.CompressedAnimData), Time(Other.Time) {}
+	FAnimSequenceDecompressionContext(float SequenceLength_, EAnimInterpolationType Interpolation_, const FName& AnimName_, const ICompressedAnimData& CompressedAnimData_)
+		: SequenceLength(SequenceLength_), Interpolation(Interpolation_), AnimName(AnimName_), CompressedAnimData(CompressedAnimData_), Time(0.f)
+	{}
 
 	// Anim info
 	float SequenceLength;
 	EAnimInterpolationType Interpolation;
 	FName AnimName;
 
-	const FUECompressedAnimData& CompressedAnimData;
+	const ICompressedAnimData& CompressedAnimData;
 	float Time;
 	float RelativePos;
-	bool bHasScale;
 
 	void Seek(float SampleAtTime)
 	{
 		Time = SampleAtTime;
 		RelativePos = SampleAtTime / SequenceLength;
-		bHasScale = CompressedAnimData.CompressedScaleOffsets.IsValid();
 	}
-
-	float GetSequenceLength() const { return SequenceLength; }
-
-	int32 GetCompressedNumberOfFrames() const { return CompressedAnimData.CompressedNumberOfFrames; }
-
-	// Following just expose the properties of the anim sequence, 
-	AnimEncoding* GetRotationCodec() const { return CompressedAnimData.RotationCodec; }
-	AnimEncoding* GetTranslationCodec() const { return CompressedAnimData.TranslationCodec; }
-	AnimEncoding* GetScaleCodec() const { return CompressedAnimData.ScaleCodec; }
-
-	const int32* GetCompressedTrackOffsets() const { return CompressedAnimData.CompressedTrackOffsets.GetData(); }
-	const uint8* GetCompressedByteStream() const { return CompressedAnimData.CompressedByteStream.GetData(); }
-	const FCompressedOffsetDataBase<TArrayView<int32>>* GetCompressedScaleOffsets() const { return &CompressedAnimData.CompressedScaleOffsets; }
-
-	EAnimInterpolationType GetInterpolation() const { return Interpolation; }
-
-	FName GetAnimFName() const { return AnimName; }
 };
-
-#else
-
-/**
- * Structure to hold the data required for decompression.
- * It is created from an UAnimSequence and stores all the intermediary
- * information required. A context is bound to a single sequence. Multiple
- * context instances can be bound to the same sequence.
- * In order to re-use a context with a new sequence, it needs to be bound.
- */
-struct ENGINE_API FAnimSequenceDecompressionContext
-{
-	FAnimSequenceDecompressionContext();
-	FAnimSequenceDecompressionContext(const UAnimSequence* AnimSeq_);
-	~FAnimSequenceDecompressionContext();
-
-	/**
-	 * We don't disable copying but we also don't retain any state. Copying is required by the UE serialization
-	 * but retaining state is not necessary and can cause issues.
-	 */
-	FAnimSequenceDecompressionContext(const FAnimSequenceDecompressionContext& Other);
-	FAnimSequenceDecompressionContext& operator=(const FAnimSequenceDecompressionContext& Other);
-
-	/**
-	 * Seek into an animation sequence at a particular time.
-	 */
-	void Seek(float SampleAtTime);
-
-	/**
-	 * Returns whether this context is bound to the provided anim sequence.
-	 */
-	bool IsStale(const UAnimSequence* AnimSeq_) const;
-
-	/**
-	 * Binds the context to the provided anim sequence. The sequence can be NULL
-	 * in which case the context will be reset.
-	 */
-	void Bind(const UAnimSequence* AnimSeq_);
-
-	inline FTrivialAnimKeyHandle GetTrivialRotationKeyHandle(int32 TrackIndex) const;
-	inline FTrivialAnimKeyHandle GetTrivialTranslationKeyHandle(int32 TrackIndex) const;
-	inline FTrivialAnimKeyHandle GetTrivialScaleKeyHandle(int32 TrackIndex) const;
-	FORCEINLINE void GetTrivialRotation(FTransform& OutAtom, FTrivialAnimKeyHandle KeyHandle) const;
-	FORCEINLINE void GetTrivialTranslation(FTransform& OutAtom, FTrivialAnimKeyHandle KeyHandle) const;
-	FORCEINLINE void GetTrivialScale(FTransform& OutAtom, FTrivialAnimKeyHandle KeyHandle) const;
-
-	inline int32 GetRotationValueOffset(int32 TrackIndex) const { return (TrackIndex * NumStreamsPerTrack) + 1; }
-	inline int32 GetTranslationValueOffset(int32 TrackIndex) const { return (TrackIndex * NumStreamsPerTrack) + 0; }
-	inline int32 GetScaleValueOffset(int32 TrackIndex) const { return (TrackIndex * NumStreamsPerTrack) + 2; }
-
-	AnimEncoding* GetRotationCodec() const { return RotationCodec; }
-	AnimEncoding* GetTranslationCodec() const { return TranslationCodec; }
-	AnimEncoding* GetScaleCodec() const { return ScaleCodec; }
-
-	const int32* GetCompressedTrackOffsets() const { return CompressedTrackOffsets; }
-	const uint8* GetCompressedByteStream() const { return CompressedByteStream; }
-	const FCompressedOffsetData* GetCompressedScaleOffsets() const { return CompressedScaleOffsets; }
-
-	const UAnimSequence* AnimSeq;
-
-	const int32* CompressedTrackOffsets;
-	const uint8* CompressedByteStream;
-	const FCompressedOffsetData* CompressedScaleOffsets;
-
-	float Time;
-	float RelativePos;
-
-	bool bHasScale;
-	bool NeedsTwoSegments;
-	bool NeedsInterpolation;
-	bool HasSegments;
-	bool bIsSorted;
-
-	uint8 NumStreamsPerTrack;
-
-	int32 NumTracks;
-	uint16 SegmentIndex0;
-	uint16 SegmentIndex1;
-
-	int32 KeyIndex0;
-	int32 KeyIndex1;
-	int32 SegmentKeyIndex0;
-	int32 SegmentKeyIndex1;
-	float KeyAlpha;
-
-	int32 PackedTranslationSize0;
-	int32 PackedRotationSize0;
-	int32 PackedScaleSize0;
-
-	int32 RangeDataSize0;
-
-	const uint8* TrackFlags;
-	const uint8* TrivialTrackKeys;
-	const uint8* TrackRangeData[2];
-
-	const FCompressedSegment* Segment0;
-	const FCompressedSegment* Segment1;
-
-	FAnimEncodingDecompressionContext* EncodingContext;
-	TArray<int32, FAnimEncodingHeapAllocator> TrivialTrackStreamOffsets;
-	uint32 SequenceCRC;
-
-#if WITH_EDITOR
-	TStaticArray<double, 4> PreviousBindTimeStamps;
-#endif
-
-private:
-	AnimEncoding * RotationCodec;
-	AnimEncoding* TranslationCodec;
-	AnimEncoding* ScaleCodec;
-};
-
-FTrivialAnimKeyHandle FAnimSequenceDecompressionContext::GetTrivialRotationKeyHandle(int32 TrackIndex) const
-{
-	const int32 TrivialKeyOffset = TrivialTrackStreamOffsets[GetRotationValueOffset(TrackIndex)];
-	return FTrivialAnimKeyHandle(TrivialKeyOffset);
-}
-
-FTrivialAnimKeyHandle FAnimSequenceDecompressionContext::GetTrivialTranslationKeyHandle(int32 TrackIndex) const
-{
-	const int32 TrivialKeyOffset = TrivialTrackStreamOffsets[GetTranslationValueOffset(TrackIndex)];
-	return FTrivialAnimKeyHandle(TrivialKeyOffset);
-}
-
-FTrivialAnimKeyHandle FAnimSequenceDecompressionContext::GetTrivialScaleKeyHandle(int32 TrackIndex) const
-{
-	const int32 TrivialKeyOffset = TrivialTrackStreamOffsets[GetScaleValueOffset(TrackIndex)];
-	return FTrivialAnimKeyHandle(TrivialKeyOffset);
-}
-
-FORCEINLINE void FAnimSequenceDecompressionContext::GetTrivialRotation(FTransform& OutAtom, FTrivialAnimKeyHandle KeyHandle) const
-{
-	const uint8* TrivialKeyData = TrivialTrackKeys + KeyHandle.Value;
-	FQuat R0;
-	DecompressRotation<ACF_Float96NoW>(R0, TrivialKeyData, TrivialKeyData);
-	OutAtom.SetRotation(R0);
-}
-
-FORCEINLINE void FAnimSequenceDecompressionContext::GetTrivialTranslation(FTransform& OutAtom, FTrivialAnimKeyHandle KeyHandle) const
-{
-	const uint8* TrivialKeyData = TrivialTrackKeys + KeyHandle.Value;
-	FVector P0;
-	DecompressTranslation<ACF_None>(P0, TrivialKeyData, TrivialKeyData);
-	OutAtom.SetTranslation(P0);
-}
-
-FORCEINLINE void FAnimSequenceDecompressionContext::GetTrivialScale(FTransform& OutAtom, FTrivialAnimKeyHandle KeyHandle) const
-{
-	const uint8* TrivialKeyData = TrivialTrackKeys + KeyHandle.Value;
-	FVector P0;
-	DecompressScale<ACF_None>(P0, TrivialKeyData, TrivialKeyData);
-	OutAtom.SetScale3D(P0);
-}
-#endif
\ No newline at end of file
diff --git a/Engine/Source/Runtime/Engine/Public/Animation/AnimTypes.h b/Engine/Source/Runtime/Engine/Public/Animation/AnimTypes.h
--- a/Engine/Source/Runtime/Engine/Public/Animation/AnimTypes.h
+++ b/Engine/Source/Runtime/Engine/Public/Animation/AnimTypes.h
@@ -346,6 +346,10 @@ struct ENGINE_VTABLE FAnimNotifyEvent : public FAnimLinkableElement
 	/** Color of Notify in editor */
 	UPROPERTY()
 	FColor NotifyColor;
+
+	/** Guid for tracking notifies in editor */
+	UPROPERTY()
+	FGuid Guid;
 #endif // WITH_EDITORONLY_DATA
 
 	/** 'Track' that the notify exists on, used for visual placement in editor and sorting priority in runtime */
@@ -389,6 +393,11 @@ public:
 	{
 	}
 
+#if WITH_EDITORONLY_DATA
+	ENGINE_API bool Serialize(FArchive& Ar);
+	ENGINE_API void PostSerialize(const FArchive& Ar);
+#endif
+
 	/** Updates trigger offset based on a combination of predicted offset and current offset */
 	ENGINE_API void RefreshTriggerOffset(EAnimEventTriggerOffsets::Type PredictedOffsetType);
 
@@ -432,6 +441,17 @@ public:
 	ENGINE_API FName GetNotifyEventName() const;
 };
 
+#if WITH_EDITORONLY_DATA
+template<> struct TStructOpsTypeTraits<FAnimNotifyEvent> : public TStructOpsTypeTraitsBase2<FAnimNotifyEvent>
+{
+	enum 
+	{ 
+		WithSerializer = true,
+		WithPostSerialize = true
+	};
+};
+#endif
+
 // Used by UAnimSequenceBase::SortNotifies() to sort its Notifies array
 FORCEINLINE bool FAnimNotifyEvent::operator<(const FAnimNotifyEvent& Other) const
 {
@@ -455,7 +475,15 @@ FORCEINLINE bool FAnimNotifyEvent::operator<(const FAnimNotifyEvent& Other) cons
 USTRUCT(BlueprintType)
 struct FAnimSyncMarker
 {
-	GENERATED_USTRUCT_BODY()
+	GENERATED_BODY()
+
+	FAnimSyncMarker()
+		: MarkerName(NAME_None)
+		, Time(0.0f)
+#if WITH_EDITORONLY_DATA
+		, TrackIndex(0)
+#endif
+	{}
 
 	// The name of this marker
 	UPROPERTY(BlueprintReadOnly, Category=Animation)
@@ -469,12 +497,29 @@ struct FAnimSyncMarker
 	// The editor track this marker sits on
 	UPROPERTY()
 	int32 TrackIndex;
+
+	UPROPERTY()
+	FGuid Guid;
+#endif
+	
+#if WITH_EDITORONLY_DATA
+	ENGINE_API bool Serialize(FArchive& Ar);
 #endif
 
 	/** This can be used with the Sort() function on a TArray of FAnimSyncMarker to sort the notifies array by time, earliest first. */
 	ENGINE_API bool operator <(const FAnimSyncMarker& Other) const { return Time < Other.Time; }
 };
 
+#if WITH_EDITORONLY_DATA
+template<> struct TStructOpsTypeTraits<FAnimSyncMarker> : public TStructOpsTypeTraitsBase2<FAnimSyncMarker>
+{
+	enum 
+	{ 
+		WithSerializer = true
+	};
+};
+#endif
+
 /**
  * Keyframe position data for one track.  Pos(i) occurs at Time(i).  Pos.Num() always equals Time.Num().
  */
@@ -526,12 +571,21 @@ namespace ECurveBlendOption
 {
 	enum Type
 	{
-		/** Find Max Weight of curve and use that weight. */
-		MaxWeight,
+		/* Last pose that contains valid curve value override it. 
+		 * Unfortunately this has been default behavior for a long time, so we still keep it. */
+		Override, // redirect from MaxWeight old legacy behavior
+		/** Only set the value if the previous pose doesn't have the curve value. */
+		DoNotOverride, 
 		/** Normalize By Sum of Weight and use it to blend. */
 		NormalizeByWeight,
 		/** Blend By Weight without normalizing*/
-		BlendByWeight
+		BlendByWeight, 
+		/** Use Base Pose for all curve values. Do not blend */
+		UseBasePose,
+		/** Find the highest curve value from multiple poses and use that. */
+		UseMaxValue,
+		/** Find the lowest curve value from multiple poses and use that. */
+		UseMinValue,
 	};
 }
 
diff --git a/Engine/Source/Runtime/Engine/Public/Animation/SkinWeightProfile.h b/Engine/Source/Runtime/Engine/Public/Animation/SkinWeightProfile.h
--- a/Engine/Source/Runtime/Engine/Public/Animation/SkinWeightProfile.h
+++ b/Engine/Source/Runtime/Engine/Public/Animation/SkinWeightProfile.h
@@ -8,7 +8,7 @@
 #include "Misc/CoreStats.h"
 #include "RenderingThread.h"
 #include "HAL/UnrealMemory.h"
-
+#include "BoneIndices.h"
 #include "SkinWeightProfile.generated.h"
 
 class USkeletalMesh;
@@ -53,7 +53,7 @@ struct FSkinWeightProfileInfo
 struct FRawSkinWeight
 {
 	// MAX_TOTAL_INFLUENCES for now
-	uint8 InfluenceBones[MAX_TOTAL_INFLUENCES];
+	FBoneIndexType InfluenceBones[MAX_TOTAL_INFLUENCES];
 	uint8 InfluenceWeights[MAX_TOTAL_INFLUENCES];
 
 	friend FArchive& operator<<(FArchive& Ar, FRawSkinWeight& OverrideEntry);
@@ -87,24 +87,20 @@ struct FRuntimeSkinWeightProfileData
 		friend FArchive& operator<<(FArchive& Ar, FSkinWeightOverrideInfo& OverrideInfo);
 	};
 
-	template<bool bExtraBoneInfluences>
 	void ApplyOverrides(FSkinWeightVertexBuffer* OverrideBuffer, const FSkinWeightVertexBuffer* BaseBuffer) const
 	{
-		const TSkinWeightInfo<bExtraBoneInfluences>* SkinWeightInfoPtr = BaseBuffer->GetSkinWeightPtr<bExtraBoneInfluences>(0);
-		
-		if (SkinWeightInfoPtr)
+		const uint32 ExpectedNumVerts = BaseBuffer->GetNumVertices();
+		if (ExpectedNumVerts)
 		{
-			TArray<TSkinWeightInfo<bExtraBoneInfluences>> OverrideArray;
-			const int32 ExpectedNumVerts = BaseBuffer->GetNumVertices();
-			OverrideArray.SetNumUninitialized(ExpectedNumVerts);
-			FMemory::Memcpy(OverrideArray.GetData(), SkinWeightInfoPtr, sizeof(TSkinWeightInfo<bExtraBoneInfluences>) * ExpectedNumVerts);
+			TArray<FSkinWeightInfo> OverrideArray;
+			BaseBuffer->GetSkinWeights(OverrideArray);
 
 			// Apply overrides
 			{
 				for (auto VertexIndexOverridePair : VertexIndexOverrideIndex)
 				{
 					const uint32 VertexIndex = VertexIndexOverridePair.Key;
-					TSkinWeightInfo<bExtraBoneInfluences>& Entry = OverrideArray[VertexIndex];
+					FSkinWeightInfo& Entry = OverrideArray[VertexIndex];
 
 					const uint32 OverrideIndex = VertexIndexOverridePair.Value;
 					const FSkinWeightOverrideInfo& OverrideInfo = OverridesInfo[OverrideIndex];
@@ -114,9 +110,9 @@ struct FRuntimeSkinWeightProfileData
 
 					for (int32 Index = 0; Index < OverrideInfo.NumInfluences; ++Index)
 					{
-						const uint16 WeightData = Weights[OverrideInfo.InfluencesOffset + Index];
+						const uint32 WeightData = Weights[OverrideInfo.InfluencesOffset + Index];
 
-						Entry.InfluenceBones[Index] = (WeightData) >> 8;
+						Entry.InfluenceBones[Index] = (WeightData) >> 16;
 						Entry.InfluenceWeights[Index] = (WeightData & 0xFF);
 					}
 				}
@@ -132,8 +128,8 @@ struct FRuntimeSkinWeightProfileData
 
 	/** Per skin weight offset into Weights array and number of weights stored */
 	TArray<FSkinWeightOverrideInfo> OverridesInfo;
-	/** Bulk data containing all Weights, stored as bone id in upper and weight in lower (8) bits */
-	TArray<uint16> Weights;	
+	/** Bulk data containing all Weights, stored as bone id in upper and weight in lower (16) bits */
+	TArray<uint32> Weights;
 	/** Map between Vertex Indices and entries of OverridesInfo */
 	TMap<uint32, uint32> VertexIndexOverrideIndex;
 	
@@ -173,16 +169,16 @@ struct ENGINE_API FSkinWeightProfilesData
 
 	void ReleaseCPUResources();
 
-	void CreateRHIBuffers_RenderThread(TArray<TPair<FName, FVertexBufferRHIRef>>& OutBuffers);
-	void CreateRHIBuffers_Async(TArray<TPair<FName, FVertexBufferRHIRef>>& OutBuffers);
+	void CreateRHIBuffers_RenderThread(TArray<TPair<FName, FSkinWeightRHIInfo>>& OutBuffers);
+	void CreateRHIBuffers_Async(TArray<TPair<FName, FSkinWeightRHIInfo>>& OutBuffers);
 
 	template <uint32 MaxNumUpdates>
-	void InitRHIForStreaming(const TArray<TPair<FName, FVertexBufferRHIRef>>& IntermediateBuffers, TRHIResourceUpdateBatcher<MaxNumUpdates>& Batcher)
+	void InitRHIForStreaming(const TArray<TPair<FName, FSkinWeightRHIInfo>>& IntermediateBuffers, TRHIResourceUpdateBatcher<MaxNumUpdates>& Batcher)
 	{
 		for (int32 Idx = 0; Idx < IntermediateBuffers.Num(); ++Idx)
 		{
 			const FName& ProfileName = IntermediateBuffers[Idx].Key;
-			FRHIVertexBuffer* IntermediateBuffer = IntermediateBuffers[Idx].Value;
+			const FSkinWeightRHIInfo& IntermediateBuffer = IntermediateBuffers[Idx].Value;
 			ProfileNameToBuffer.FindChecked(ProfileName)->InitRHIForStreaming(IntermediateBuffer, Batcher);
 		}
 	}
@@ -198,7 +194,7 @@ struct ENGINE_API FSkinWeightProfilesData
 
 protected:
 	template <bool bRenderThread>
-	void CreateRHIBuffers_Internal(TArray<TPair<FName, FVertexBufferRHIRef>>& OutBuffers);
+	void CreateRHIBuffers_Internal(TArray<TPair<FName, FSkinWeightRHIInfo>>& OutBuffers);
 
 	FSkinWeightVertexBuffer* BaseBuffer;
 	FSkinWeightVertexBuffer* DefaultOverrideSkinWeightBuffer;
diff --git a/Engine/Source/Runtime/Engine/Public/AnimationCompression.h b/Engine/Source/Runtime/Engine/Public/AnimationCompression.h
--- a/Engine/Source/Runtime/Engine/Public/AnimationCompression.h
+++ b/Engine/Source/Runtime/Engine/Public/AnimationCompression.h
@@ -11,7 +11,7 @@
 #include "Animation/AnimEnums.h"
 #include "EngineLogs.h"
 
-DECLARE_LOG_CATEGORY_EXTERN(LogAnimationCompression, Log, All);
+ENGINE_API DECLARE_LOG_CATEGORY_EXTERN(LogAnimationCompression, Log, All);
 
 // Thresholds
 #define TRANSLATION_ZEROING_THRESHOLD (0.0001f)
diff --git a/Engine/Source/Runtime/Engine/Public/AnimationUtils.h b/Engine/Source/Runtime/Engine/Public/AnimationUtils.h
--- a/Engine/Source/Runtime/Engine/Public/AnimationUtils.h
+++ b/Engine/Source/Runtime/Engine/Public/AnimationUtils.h
@@ -9,17 +9,12 @@
 #include "CoreMinimal.h"
 #include "Animation/AnimSequence.h"
 
-class UAnimCompress;
-class UAnimSet;
-struct FAnimCompressContext;
-
 ENGINE_API DECLARE_LOG_CATEGORY_EXTERN(LogMatinee, Warning, All);
 ENGINE_API DECLARE_LOG_CATEGORY_EXTERN(LogAnimNotify, Warning, All);
 ENGINE_API DECLARE_LOG_CATEGORY_EXTERN(LogAnimMontage, Warning, All);
 ENGINE_API DECLARE_LOG_CATEGORY_EXTERN(LogAnimMarkerSync, Warning, All);
 
 // Forward declarations.
-class UAnimCompress;
 class UAnimSet;
 struct FTransform;
 
@@ -39,27 +34,7 @@ struct FAnimPerturbationError
 	float MaxErrorInScaleDueToScale;
 };
 
-/**
- * A set of error statistics for an animation, gathered by FAnimationUtils::ComputeCompressionError
- */
-struct AnimationErrorStats
-{
-	/** Average world-space translation error across all end-effectors **/
-	float AverageError;
-	/** The worst error encountered across all end effectors **/
-	float MaxError;
-	/** Time at which the worst error occurred */
-	float MaxErrorTime;
-	/** Bone on which the worst error occurred */
-	int32 MaxErrorBone;
-
-	AnimationErrorStats()
-		: AverageError(0.f)
-		, MaxError(0.f)
-		, MaxErrorTime(0.f)
-		, MaxErrorBone(0)
-	{}
-};
+using AnimationErrorStats = FAnimationErrorStats;
 
 /** Different ways to create an additive animation */
 enum EConvertToAdditive
@@ -89,7 +64,7 @@ public:
 												const TArray<FTransform>& BoneSpaceTransforms,
 												const TArray<FBoneData>& BoneData);
 
-	static void BuildSkeletonMetaData(USkeleton* Skeleton, TArray<FBoneData>& OutBoneData);
+	ENGINE_API static void BuildSkeletonMetaData(USkeleton* Skeleton, TArray<FBoneData>& OutBoneData);
 
 	static int32 GetAnimTrackIndexForSkeletonBone(const int32 InSkeletonBoneIndex, const TArray<FTrackToSkeletonMap>& TrackToSkelMap);
 
@@ -103,62 +78,8 @@ public:
 	 * @param	ErrorStats	Output structure containing the final compression error values
 	 * @return				None.
 	 */
-	static void ComputeCompressionError(const FCompressibleAnimData& CompressibleAnimData, FCompressibleAnimDataResult& CompressedData, AnimationErrorStats& ErrorStats);
+	ENGINE_API static void ComputeCompressionError(const FCompressibleAnimData& CompressibleAnimData, FCompressibleAnimDataResult& CompressedData, AnimationErrorStats& ErrorStats);
 #endif
-
-	/**
-	 * Utility function to compress an animation. If the animation is currently associated with a codec, it will be used to 
-	 * compress the animation. Otherwise, the default codec will be used. The CompressContext can specify whether an
-	 * alternative compression codec can be used. If the alternative codec produces better compression and the accuracy
-	 * of the compressed animation remains within tolerances, the alternative codec will be used. 
-	 * See GetAlternativeCompressionThreshold for information on the tolerance value used.
-	 *
-	 * @param	CompressibleAnimData	The animation data to compress.
-	 * @param	CompressContext			Context for applying compression, tracking memory saved etc
-	 * @return							The compressed animation data
-	 */
-	ENGINE_API static void CompressAnimSequence(const FCompressibleAnimData& CompressibleAnimData, FCompressibleAnimDataResult& OutCompressedData, FAnimCompressContext& CompressContext);
-
-	/**
-	 * Utility function to compress an animation. If the animation is currently associated with a codec, it will be used to
-	 * compress the animation. Otherwise, the default codec will be used. The CompressContext can specify whether an
-	 * alternative compression codec can be used. If the alternative codec produces better compression and the accuracy
-	 * of the compressed animation remains within tolerances, the alternative codec will be used.
-	 * See GetAlternativeCompressionThreshold for information on the tolerance value used.
-	 *
-	 * @param	CompressibleAnimData					The animation data to compress.
-	 * @param	OriginalCompressedData					The original compressed data (if it exists).
-	 * @param	CompressContext							Context for applying compression, tracking memory saved etc*
-	 * @param	MasterTolerance							The alternate error threshold (0.0 means don't try anything other than the current / default scheme)
-	 * @param	bFirstRecompressUsingCurrentOrDefault	If true, then the animation will be first recompressed with it's current compressor if non-NULL, or with the global default compressor (specified in the engine ini)
-	 * @param	bForceBelowThreshold					If true and the existing compression error is greater than MasterTolerance, then any compression technique (even one that increases the size) with a lower error will be used until it falls below the threshold
-	 * @param	bRaiseMaxErrorToExisting				If true and the existing compression error is greater than MasterTolerance, then MasterTolerance will be effectively raised to the existing error level
-	 * @param	bTryExhaustiveSearch					If true, then an exhaustive search is used otherwise only a short list of the best methods is tried
-	 *
-	 * @return											The compressed animation data
-	 */
-	static void CompressAnimSequenceExplicit(
-		const FCompressibleAnimData& CompressibleAnimData,
-		FCompressibleAnimDataResult& OutCompressedData,
-		FAnimCompressContext& CompressContext,
-		float MasterTolerance,
-		bool bFirstRecompressUsingCurrentOrDefault,
-		bool bForceBelowThreshold,
-		bool bRaiseMaxErrorToExisting,
-		bool bTryExhaustiveSearch,
-		bool bEnableSegmenting,
-		int32 IdealNumFramesPerSegment,
-		int32 MaxNumFramesPerSegment);
-
-	/**
-	 * Determines the current setting for world-space error tolerance in the animation compressor.
-	 * When requested, animation being compressed will also consider an alternative compression
-	 * method if the end result of that method produces less error than the AlternativeCompressionThreshold.
-	 * The default tolerance value is 0.0f (no alternatives allowed) but may be overridden using a field in the base engine INI file.
-	 *
-	 * @return				World-space error tolerance for considering an alternative compression method
-	 */
-	ENGINE_API static float GetAlternativeCompressionThreshold();
 	
 	/**
 	 * Determines the current setting for recompressing all animations upon load. The default value 
@@ -208,7 +129,7 @@ public:
 	 * Note: If there are as many times as frames, they are automatically assumed to be uniformly spaced.
 	 * Note: If there are two or fewer times, they are automatically assumed to be uniformly spaced.
 	 *
-	 * @param AnimSeq		The number of frames in the source animation data
+	 * @param NumFrames		The number of frames in the source animation data
 	 * @param Times			The array of key times
 	 *
 	 * @return				true if the keys are uniformly spaced (or one of the trivial conditions is detected).  false if any key spacing is greater than 1e-4 off.
@@ -226,18 +147,31 @@ public:
 		const FVector& ScaleNudge,
 		TArray<FAnimPerturbationError>& InducedErrors);
 
-	/**
-	 * @return		The default animation compression algorithm singleton, instantiating it if necessary.
-	 */
-	ENGINE_API static UAnimCompress* GetDefaultAnimationCompressionAlgorithm();
+	/** Returns the default animation bone compression settings, can never by null. */
+	ENGINE_API static UAnimBoneCompressionSettings* GetDefaultAnimationBoneCompressionSettings();
+
+	/** Returns the default animation bone compression settings for the FAnimRecorder, can never by null. */
+	ENGINE_API static UAnimBoneCompressionSettings* GetDefaultAnimationRecorderBoneCompressionSettings();
 
 	/** Returns the default animation curve compression settings, can never by null. */
 	ENGINE_API static UAnimCurveCompressionSettings* GetDefaultAnimationCurveCompressionSettings();
 
+	/** Makes sure that the animation sequence and all its dependencies required for compression are loaded. */
+	ENGINE_API static void EnsureAnimSequenceLoaded(UAnimSequence& AnimSeq);
+
 	static void ExtractTransformFromTrack(float Time, int32 NumFrames, float SequenceLength, const struct FRawAnimSequenceTrack& RawTrack, EAnimInterpolationType Interpolation, FTransform &OutAtom);
 
 #if WITH_EDITOR
 	static void ExtractTransformFromCompressionData(const FCompressibleAnimData& CompressibleAnimData, FCompressibleAnimDataResult& CompressedAnimData, float Time, int32 TrackIndex, bool bUseRawData, FTransform& OutBoneTransform);
+
+	/**
+	 * Compresses the animation bones within a sequence with the chosen settings.
+	 * Note: This modifies the sequence.
+	 *
+	 * @return		Returns true on success, false it we fail to compress
+	 */
+	ENGINE_API static bool CompressAnimBones(FCompressibleAnimData& AnimSeq, FCompressibleAnimDataResult& Target);
+
 	/**
 	 * Compresses the animation curves within a sequence with the chosen settings.
 	 * Note: This modifies the sequence.
diff --git a/Engine/Source/Runtime/Engine/Public/BoneContainer.h b/Engine/Source/Runtime/Engine/Public/BoneContainer.h
--- a/Engine/Source/Runtime/Engine/Public/BoneContainer.h
+++ b/Engine/Source/Runtime/Engine/Public/BoneContainer.h
@@ -259,6 +259,11 @@ public:
 		return RefSkeleton->GetRefBonePose();
 	}
 
+	const TArray<FCompactPoseBoneIndex>& GetCompactPoseParentBoneArray() const
+	{
+		return CompactPoseParentBones;
+	}
+
 	const FTransform& GetRefPoseTransform(const FCompactPoseBoneIndex& BoneIndex) const
 	{
 		return CompactPoseRefPoseBones[BoneIndex.GetInt()];
diff --git a/Engine/Source/Runtime/Engine/Public/BonePose.h b/Engine/Source/Runtime/Engine/Public/BonePose.h
--- a/Engine/Source/Runtime/Engine/Public/BonePose.h
+++ b/Engine/Source/Runtime/Engine/Public/BonePose.h
@@ -63,6 +63,11 @@ public:
 		return Bones[BoneIndex.GetInt()];
 	}
 
+	FORCEINLINE TArrayView<FTransform> GetMutableBones()
+	{
+		return TArrayView<FTransform>(Bones);
+	}
+
 	//Bone Index Iteration
 	template<typename PoseType, typename IterType>
 	struct FRangedForSupport
@@ -360,10 +365,20 @@ protected:
 
 struct ENGINE_API FCompactPose : public FBaseCompactPose<FAnimStackAllocator>
 {
+	// Sets every bone transform to Identity
+	void ResetToAdditiveIdentity();
+
+	// Normalizes all rotations in this pose
+	void NormalizeRotations();
 };
 
 struct ENGINE_API FCompactHeapPose : public FBaseCompactPose<FDefaultAllocator>
 {
+	// Sets every bone transform to Identity
+	void ResetToAdditiveIdentity();
+
+	// Normalizes all rotations in this pose
+	void NormalizeRotations();
 };
 
 struct FMeshPose : public FBasePose<FMeshPoseBoneIndex, FDefaultAllocator>
diff --git a/Engine/Source/Runtime/Engine/Public/GPUSkinPublicDefs.h b/Engine/Source/Runtime/Engine/Public/GPUSkinPublicDefs.h
--- a/Engine/Source/Runtime/Engine/Public/GPUSkinPublicDefs.h
+++ b/Engine/Source/Runtime/Engine/Public/GPUSkinPublicDefs.h
@@ -9,7 +9,10 @@
 #include "CoreMinimal.h"
 
 /** Max number of bone influences that a single skinned vert can have. */
-#define MAX_TOTAL_INFLUENCES		8
+#define MAX_TOTAL_INFLUENCES		12
+
+/** Number of bone influences when using extra bone influences. */
+#define EXTRA_BONE_INFLUENCES		8
 
 /** Max number of bone influences that a single skinned vert can have per vertex stream. */
 #define MAX_INFLUENCES_PER_STREAM	4
diff --git a/Engine/Source/Runtime/Engine/Public/GPUSkinVertexFactory.h b/Engine/Source/Runtime/Engine/Public/GPUSkinVertexFactory.h
--- a/Engine/Source/Runtime/Engine/Public/GPUSkinVertexFactory.h
+++ b/Engine/Source/Runtime/Engine/Public/GPUSkinVertexFactory.h
@@ -149,6 +149,12 @@ public: // From FTickableObjectRenderThread
 	virtual TStatId GetStatId() const override;
 };
 
+enum GPUSkinBoneInfluenceType
+{
+	DefaultBoneInfluence,	// 4 bones per vertex
+	ExtraBoneInfluence,		// 8 bones per vertex
+	UnlimitedBoneInfluence	// unlimited bones per vertex
+};
 
 /** Vertex factory with vertex stream components for GPU skinned vertices */
 class FGPUBaseSkinVertexFactory : public FVertexFactory
@@ -228,6 +234,9 @@ public:
 			return (bPrevious) ? PreviousRevisionNumber : CurrentRevisionNumber;
 		}
 
+		int32 InputWeightIndexSize = 0;
+		FShaderResourceViewRHIRef InputWeightStream;
+
 	private:
 		// double buffered bone positions+orientations to support normal rendering and velocity (new-old position) rendering
 		FVertexBufferAndSRV BoneBuffer[2];
@@ -292,7 +301,9 @@ public:
 		return ShaderData;
 	}
 
-	virtual bool UsesExtraBoneInfluences() const { return false; }
+	virtual GPUSkinBoneInfluenceType GetBoneInfluenceType() const { return DefaultBoneInfluence; }
+	virtual uint32 GetNumBoneInfluences() const { return 0; }
+	virtual bool Use16BitBoneIndex() const { return false; }
 
 	static bool SupportsTessellationShaders() { return true; }
 
@@ -303,8 +314,10 @@ public:
 
 	ENGINE_API static int32 GetMaxGPUSkinBones();
 
-	static const uint32 GHardwareMaxGPUSkinBones = 256;	
+	static const uint32 GHardwareMaxGPUSkinBones = 65536;
 	
+	ENGINE_API static bool UseUnlimitedBoneInfluences(uint32 MaxBoneInfluences);
+
 	virtual const FShaderResourceViewRHIRef GetPositionsSRV() const = 0;
 	virtual const FShaderResourceViewRHIRef GetTangentsSRV() const = 0;
 	virtual const FShaderResourceViewRHIRef GetTextureCoordinatesSRV() const = 0;
@@ -331,18 +344,12 @@ private:
 };
 
 /** Vertex factory with vertex stream components for GPU skinned vertices */
-template<bool bExtraBoneInfluencesT>
+template<GPUSkinBoneInfluenceType BoneInfluenceType>
 class ENGINE_API TGPUSkinVertexFactory : public FGPUBaseSkinVertexFactory
 {
-	DECLARE_VERTEX_FACTORY_TYPE(TGPUSkinVertexFactory<bExtraBoneInfluencesT>);
+	DECLARE_VERTEX_FACTORY_TYPE(TGPUSkinVertexFactory<BoneInfluenceType>);
 
 public:
-
-	enum
-	{
-		HasExtraBoneInfluences = bExtraBoneInfluencesT,
-	};
-
 	struct FDataType : public FStaticMeshDataType
 	{
 		/** The stream to read the bone indices from */
@@ -356,6 +363,15 @@ public:
 
 		/** The stream to read the extra bone weights from */
 		FVertexStreamComponent ExtraBoneWeights;
+
+		/** The stream to read the blend stream offset and num of influences from */
+		FVertexStreamComponent BlendOffsetCount;
+
+		/** Number of bone influences */
+		uint32 NumBoneInfluences = 0;
+
+		/** If the bone indices are 16 or 8-bit format */
+		bool bUse16BitBoneIndex = 0;
 	};
 
 	/**
@@ -367,9 +383,17 @@ public:
 		: FGPUBaseSkinVertexFactory(InFeatureLevel, InNumVertices)
 	{}
 
-	virtual bool UsesExtraBoneInfluences() const override
+	virtual GPUSkinBoneInfluenceType GetBoneInfluenceType() const override
+	{ return BoneInfluenceType; }
+
+	virtual uint32 GetNumBoneInfluences() const override
+	{
+		return Data.NumBoneInfluences;
+	}
+
+	virtual bool Use16BitBoneIndex() const override
 	{
-		return bExtraBoneInfluencesT;
+		return Data.bUse16BitBoneIndex;
 	}
 
 	static void ModifyCompilationEnvironment(const FVertexFactoryType* Type, EShaderPlatform Platform, const class ::FMaterial* Material, FShaderCompilerEnvironment& OutEnvironment);
@@ -515,15 +539,15 @@ protected:
 };
 
 /** Vertex factory with vertex stream components for GPU-skinned and morph target streams */
-template<bool bExtraBoneInfluencesT>
-class TGPUSkinMorphVertexFactory : public TGPUSkinVertexFactory<bExtraBoneInfluencesT>
+template<GPUSkinBoneInfluenceType BoneInfluenceType>
+class TGPUSkinMorphVertexFactory : public TGPUSkinVertexFactory<BoneInfluenceType>
 {
-	DECLARE_VERTEX_FACTORY_TYPE(TGPUSkinMorphVertexFactory<bExtraBoneInfluencesT>);
+	DECLARE_VERTEX_FACTORY_TYPE(TGPUSkinMorphVertexFactory<BoneInfluenceType>);
 
-	typedef TGPUSkinVertexFactory<bExtraBoneInfluencesT> Super;
+	typedef TGPUSkinVertexFactory<BoneInfluenceType> Super;
 public:
 
-	struct FDataType : TGPUSkinVertexFactory<bExtraBoneInfluencesT>::FDataType
+	struct FDataType : TGPUSkinVertexFactory<BoneInfluenceType>::FDataType
 	{
 		/** stream which has the position deltas to add to the vertex position */
 		FVertexStreamComponent DeltaPositionComponent;
@@ -537,7 +561,7 @@ public:
 	 * @param	InBoneMatrices	Reference to shared bone matrices array.
 	 */
 	TGPUSkinMorphVertexFactory(ERHIFeatureLevel::Type InFeatureLevel, uint32 InNumVertices)
-	: TGPUSkinVertexFactory<bExtraBoneInfluencesT>(InFeatureLevel, InNumVertices)
+	: TGPUSkinVertexFactory<BoneInfluenceType>(InFeatureLevel, InNumVertices)
 	{}
 
 	static void ModifyCompilationEnvironment(const FVertexFactoryType* Type, EShaderPlatform Platform, const class ::FMaterial* Material, FShaderCompilerEnvironment& OutEnvironment);
@@ -566,6 +590,16 @@ public:
 
 	static FVertexFactoryShaderParameters* ConstructShaderParameters(EShaderFrequency ShaderFrequency);
 
+	virtual uint32 GetNumBoneInfluences() const override
+	{
+		return MorphData.NumBoneInfluences;
+	}
+
+	virtual bool Use16BitBoneIndex() const override
+	{
+		return MorphData.bUse16BitBoneIndex;
+	}
+
 	const FShaderResourceViewRHIRef GetPositionsSRV() const override
 	{
 		return MorphData.PositionComponentSRV;
@@ -807,16 +841,16 @@ protected:
 };
 
 /** Vertex factory with vertex stream components for GPU-skinned and morph target streams */
-template<bool bExtraBoneInfluencesT>
-class TGPUSkinAPEXClothVertexFactory : public FGPUBaseSkinAPEXClothVertexFactory, public TGPUSkinVertexFactory<bExtraBoneInfluencesT>
+template<GPUSkinBoneInfluenceType BoneInfluenceType>
+class TGPUSkinAPEXClothVertexFactory : public FGPUBaseSkinAPEXClothVertexFactory, public TGPUSkinVertexFactory<BoneInfluenceType>
 {
-	DECLARE_VERTEX_FACTORY_TYPE(TGPUSkinAPEXClothVertexFactory<bExtraBoneInfluencesT>);
+	DECLARE_VERTEX_FACTORY_TYPE(TGPUSkinAPEXClothVertexFactory<BoneInfluenceType>);
 
-	typedef TGPUSkinVertexFactory<bExtraBoneInfluencesT> Super;
+	typedef TGPUSkinVertexFactory<BoneInfluenceType> Super;
 
 public:
 
-	struct FDataType : TGPUSkinVertexFactory<bExtraBoneInfluencesT>::FDataType
+	struct FDataType : TGPUSkinVertexFactory<BoneInfluenceType>::FDataType
 	{
 		/** stream which has the physical mesh position + height offset */
 		FVertexStreamComponent CoordPositionComponent;
@@ -863,7 +897,7 @@ public:
 	 * @param	InBoneMatrices	Reference to shared bone matrices array.
 	 */
 	TGPUSkinAPEXClothVertexFactory(ERHIFeatureLevel::Type InFeatureLevel, uint32 InNumVertices)
-		: TGPUSkinVertexFactory<bExtraBoneInfluencesT>(InFeatureLevel, InNumVertices)
+		: TGPUSkinVertexFactory<BoneInfluenceType>(InFeatureLevel, InNumVertices)
 	{}
 
 	static void ModifyCompilationEnvironment(const FVertexFactoryType* Type, EShaderPlatform Platform, const class ::FMaterial* Material, FShaderCompilerEnvironment& OutEnvironment);
@@ -893,6 +927,16 @@ public:
 		return this;
 	}
 
+	virtual uint32 GetNumBoneInfluences() const override
+	{
+		return MeshMappingData.NumBoneInfluences;
+	}
+
+	virtual bool Use16BitBoneIndex() const override
+	{
+		return MeshMappingData.bUse16BitBoneIndex;
+	}
+
 	const FShaderResourceViewRHIRef GetPositionsSRV() const override
 	{
 		return MeshMappingData.PositionComponentSRV;
diff --git a/Engine/Source/Runtime/Engine/Public/Rendering/SkeletalMeshLODModel.h b/Engine/Source/Runtime/Engine/Public/Rendering/SkeletalMeshLODModel.h
--- a/Engine/Source/Runtime/Engine/Public/Rendering/SkeletalMeshLODModel.h
+++ b/Engine/Source/Runtime/Engine/Public/Rendering/SkeletalMeshLODModel.h
@@ -34,11 +34,11 @@ struct FSoftSkinVertex
 	FVector2D		UVs[MAX_TEXCOORDS];
 	// VertexColor
 	FColor			Color;
-	uint8			InfluenceBones[MAX_TOTAL_INFLUENCES];
+	FBoneIndexType	InfluenceBones[MAX_TOTAL_INFLUENCES];
 	uint8			InfluenceWeights[MAX_TOTAL_INFLUENCES];
 
 	/** If this vert is rigidly weighted to a bone, return true and the bone index. Otherwise return false. */
-	ENGINE_API bool GetRigidWeightBone(uint8& OutBoneIndex) const;
+	ENGINE_API bool GetRigidWeightBone(FBoneIndexType& OutBoneIndex) const;
 
 	/** Returns the maximum weight of any bone that influences this vertex. */
 	ENGINE_API uint8 GetMaximumWeight() const;
@@ -109,6 +109,9 @@ struct FSkelMeshSection
 	/** max # of bones used to skin the vertices in this section */
 	int32 MaxBoneInfluences;
 
+	/** whether to store bone indices as 16 bit or 8 bit in vertex buffer for rendering. */
+	bool bUse16BitBoneIndex;
+
 	// INDEX_NONE if not set
 	int16 CorrespondClothAssetIndex;
 
@@ -157,6 +160,7 @@ struct FSkelMeshSection
 		, BaseVertexIndex(0)
 		, NumVertices(0)
 		, MaxBoneInfluences(4)
+		, bUse16BitBoneIndex(false)
 		, CorrespondClothAssetIndex(INDEX_NONE)
 		, bDisabled(false)
 		, GenerateUpToLodIndex(INDEX_NONE)
@@ -196,9 +200,19 @@ struct FSkelMeshSection
 	*/
 	ENGINE_API void CalcMaxBoneInfluences();
 
-	FORCEINLINE bool HasExtraBoneInfluences() const
+	FORCEINLINE int32 GetMaxBoneInfluences() const
+	{
+		return MaxBoneInfluences;
+	}
+
+	/**
+	* Calculate if this skel mesh section needs 16-bit bone indices
+	*/
+	ENGINE_API void CalcUse16BitBoneIndex();
+
+	FORCEINLINE bool Use16BitBoneIndex() const
 	{
-		return MaxBoneInfluences > MAX_INFLUENCES_PER_STREAM;
+		return bUse16BitBoneIndex;
 	}
 
 	// Serialization.
@@ -377,7 +391,8 @@ public:
 	*/
 	ENGINE_API void GetNonClothVertices(TArray<FSoftSkinVertex>& OutVertices) const;
 
-	ENGINE_API bool DoSectionsNeedExtraBoneInfluences() const;
+	ENGINE_API int32 GetMaxBoneInfluences() const;
+	bool DoSectionsUse16BitBoneIndex() const;
 
 	void GetResourceSizeEx(FResourceSizeEx& CumulativeResourceSize) const;
 
@@ -395,17 +410,17 @@ public:
 	* Build a derive data key with the user section data (UserSectionsData) and the original bulk data
 	*/
 	ENGINE_API FString GetLODModelDeriveDataKey() const;
+	
+	/**
+	* Copy one structure to the other, make sure all bulk data is unlock and the data can be read before copying.
+	*/
+	static ENGINE_API bool CopyStructure(FSkeletalMeshLODModel* Destination, FSkeletalMeshLODModel* Source);
 
 	/**
 	* This function will update the chunked information for each section. Only old data before the 
 	* skeletal mesh build refactor should need to call this function.
 	*/
 	ENGINE_API void UpdateChunkedSectionInfo(const FString& SkeletalMeshName, TArray<int32>& LODMaterialMap);
-
-	/**
-	* Copy one structure to the other, make sure all bulk data is unlock and the data can be read before copying.
-	*/
-	static ENGINE_API bool CopyStructure(FSkeletalMeshLODModel* Destination, FSkeletalMeshLODModel* Source);
 };
 
 #endif // WITH_EDITOR
\ No newline at end of file
diff --git a/Engine/Source/Runtime/Engine/Public/Rendering/SkeletalMeshLODRenderData.h b/Engine/Source/Runtime/Engine/Public/Rendering/SkeletalMeshLODRenderData.h
--- a/Engine/Source/Runtime/Engine/Public/Rendering/SkeletalMeshLODRenderData.h
+++ b/Engine/Source/Runtime/Engine/Public/Rendering/SkeletalMeshLODRenderData.h
@@ -199,9 +199,14 @@ public:
 		return StaticVertexBuffers.PositionVertexBuffer.GetNumVertices();
 	}
 
-	bool DoesVertexBufferHaveExtraBoneInfluences() const
+	uint32 GetVertexBufferMaxBoneInfluences() const
 	{
-		return SkinWeightVertexBuffer.HasExtraBoneInfluences();
+		return SkinWeightVertexBuffer.GetMaxBoneInfluences();
+	}
+
+	bool DoesVertexBufferUse16BitBoneIndex() const
+	{
+		return SkinWeightVertexBuffer.Use16BitBoneIndex();
 	}
 
 	uint32 GetNumTexCoords() const
diff --git a/Engine/Source/Runtime/Engine/Public/Rendering/SkeletalMeshRenderData.h b/Engine/Source/Runtime/Engine/Public/Rendering/SkeletalMeshRenderData.h
--- a/Engine/Source/Runtime/Engine/Public/Rendering/SkeletalMeshRenderData.h
+++ b/Engine/Source/Runtime/Engine/Public/Rendering/SkeletalMeshRenderData.h
@@ -61,8 +61,8 @@ public:
 	/** Returns true if this resource must be skinned on the CPU for the given feature level. */
 	ENGINE_API bool RequiresCPUSkinning(ERHIFeatureLevel::Type FeatureLevel) const;
 
-	/** Returns true if there are more than MAX_INFLUENCES_PER_STREAM influences per vertex. */
-	bool HasExtraBoneInfluences() const;
+	/** Returns the number of bone influences per vertex. */
+	uint32 GetNumBoneInfluences() const;
 
 	/**
 	* Computes the maximum number of bones per section used to render this mesh.
