Updated UActorComponent and Actor to fully support PushModel.

diff --git a/Engine/Source/Runtime/Engine/Private/Actor.cpp b/Engine/Source/Runtime/Engine/Private/Actor.cpp
--- a/Engine/Source/Runtime/Engine/Private/Actor.cpp
+++ b/Engine/Source/Runtime/Engine/Private/Actor.cpp
@@ -41,6 +41,7 @@
 #include "Interfaces/ITargetPlatform.h"
 #include "DeviceProfiles/DeviceProfile.h"
 #include "ObjectTrace.h"
+#include "Net/Core/PushModel/PushModel.h"
 
 DEFINE_LOG_CATEGORY(LogActor);
 
@@ -1061,6 +1062,11 @@ bool AActor::ShouldTickIfViewportsOnly() const
 
 void AActor::PreReplication(IRepChangedPropertyTracker & ChangedPropertyTracker)
 {
+#if WITH_PUSH_MODEL
+	const AActor* const OldAttachParent = AttachmentReplication.AttachParent;
+	const UActorComponent* const OldAttachComponent = AttachmentReplication.AttachComponent;
+#endif
+
 	// Attachment replication gets filled in by GatherCurrentMovement(), but in the case of a detached root we need to trigger remote detachment.
 	AttachmentReplication.AttachParent = nullptr;
 	AttachmentReplication.AttachComponent = nullptr;
@@ -1074,6 +1080,14 @@ PRAGMA_ENABLE_DEPRECATION_WARNINGS
 	// Don't need to replicate AttachmentReplication if the root component replicates, because it already handles it.
 	DOREPLIFETIME_ACTIVE_OVERRIDE(AActor, AttachmentReplication, RootComponent && !RootComponent->GetIsReplicated());
 
+
+#if WITH_PUSH_MODEL
+	if (UNLIKELY(OldAttachParent != AttachmentReplication.AttachParent || OldAttachComponent != AttachmentReplication.AttachComponent))
+	{
+		MARK_PROPERTY_DIRTY_FROM_NAME(AActor, AttachmentReplication, this);
+	}
+#endif
+
 	UBlueprintGeneratedClass* BPClass = Cast<UBlueprintGeneratedClass>(GetClass());
 	if (BPClass != nullptr)
 	{
@@ -1586,7 +1600,7 @@ void AActor::SetOwner(AActor* NewOwner)
 		}
 
 		Owner = NewOwner;
-		// MARK_PROPERTY_DIRTY_BY_NAME(AActor, Owner, this);
+		MARK_PROPERTY_DIRTY_FROM_NAME(AActor, Owner, this);
 
 		if (Owner != nullptr)
 		{
@@ -1764,7 +1778,7 @@ PRAGMA_ENABLE_DEPRECATION_WARNINGS
 
 		// Clear AttachmentReplication struct
 		AttachmentReplication = FRepAttachment();
-		// MARK_PROPERTY_DIRTY_FROM_NAME(AActor, AttachmentReplication, this);
+		MARK_PROPERTY_DIRTY_FROM_NAME(AActor, AttachmentReplication, this);
 	}
 }
 
@@ -2221,7 +2235,7 @@ void AActor::TearOff()
 	if (NetMode == NM_ListenServer || NetMode == NM_DedicatedServer)
 	{
 		bTearOff = true;
-		// MARK_PROPERTY_DIRTY_FROM_NAME(AActor, bTearOff, this);
+		MARK_PROPERTY_DIRTY_FROM_NAME(AActor, bTearOff, this);
 		
 		FWorldContext* const Context = GEngine->GetWorldContextFromWorld(GetWorld());
 		if (Context != nullptr)
@@ -3265,7 +3279,7 @@ void AActor::SetReplicates(bool bInReplicates)
 			}
 		}
 
-		// MARK_PROPERTY_DIRTY_FROM_NAME(AActor, RemoteRole, this);
+		MARK_PROPERTY_DIRTY_FROM_NAME(AActor, RemoteRole, this);
 	}
 	else
 	{
@@ -3288,7 +3302,12 @@ void AActor::SetAutonomousProxy(const bool bInAutonomousProxy, const bool bAllow
 
 		if (bAllowForcePropertyCompare && RemoteRole != OldRemoteRole)
 		{
-			// MARK_PROPERTY_DIRTY_FROM_NAME(AActor, RemoteRole, this);
+			// We intentionally only mark RemoteRole dirty after this check, because Networking Code
+			// for swapping roles will call SetAutonomousProxy to downgrade RemoteRole.
+			// However, that's the only code that should pass bAllowForcePropertyCompare=false.
+			// Everywhere else should use bAllowForcePropertyCompare=true.
+			// TODO: Maybe split these methods up to make it clearer.
+			MARK_PROPERTY_DIRTY_FROM_NAME(AActor, RemoteRole, this);
 
 			// We have to do this so the role change above will replicate (turn off shadow state sharing for a frame)
 			// This is because RemoteRole is special since it will change between connections, so we have to special case
@@ -3304,7 +3323,7 @@ void AActor::SetAutonomousProxy(const bool bInAutonomousProxy, const bool bAllow
 void AActor::CopyRemoteRoleFrom(const AActor* CopyFromActor)
 {
 	RemoteRole = CopyFromActor->GetRemoteRole();
-	// MARK_PROPERTY_DIRTY_FROM_NAME(AActor, RemoteRole, this);
+	MARK_PROPERTY_DIRTY_FROM_NAME(AActor, RemoteRole, this);
 
 	if (RemoteRole != ROLE_None)
 	{
@@ -3353,8 +3372,8 @@ PRAGMA_DISABLE_DEPRECATION_WARNINGS
 	Swap(Role, RemoteRole);
 PRAGMA_ENABLE_DEPRECATION_WARNINGS
 
-	// MARK_PROPERTY_DIRTY_FROM_NAME(AActor, RemoteRole, this);
-	// MARK_PROPERTY_DIRTY_FROM_NAME(AActor, Role, this);
+	MARK_PROPERTY_DIRTY_FROM_NAME(AActor, RemoteRole, this);
+	MARK_PROPERTY_DIRTY_FROM_NAME(AActor, Role, this);
 
 	ForcePropertyCompare();
 }
@@ -4985,7 +5004,7 @@ void AActor::SetHidden(bool bInHidden)
 {
 PRAGMA_DISABLE_DEPRECATION_WARNINGS
 	bHidden = bInHidden;
-	// MARK_PROPERTY_DIRTY_FROM_NAME(AActor, bHidden, this);
+	MARK_PROPERTY_DIRTY_FROM_NAME(AActor, bHidden, this);
 PRAGMA_ENABLE_DEPRECATION_WARNINGS
 }
 
@@ -4993,7 +5012,7 @@ void AActor::SetReplicatingMovement(bool bInReplicateMovement)
 {
 PRAGMA_DISABLE_DEPRECATION_WARNINGS
 	bReplicateMovement = bInReplicateMovement;
-	// MARK_PROPERTY_DIRTY_FROM_NAME(AActor, bReplicateMovement, this);
+	MARK_PROPERTY_DIRTY_FROM_NAME(AActor, bReplicateMovement, this);
 PRAGMA_ENABLE_DEPRECATION_WARNINGS
 }
 
@@ -5001,7 +5020,7 @@ void AActor::SetCanBeDamaged(bool bInCanBeDamaged)
 {
 PRAGMA_DISABLE_DEPRECATION_WARNINGS
 	bCanBeDamaged = bInCanBeDamaged;
-	// MARK_PROPERTY_DIRTY_FROM_NAME(AActor, bCanBeDamaged, this);
+	MARK_PROPERTY_DIRTY_FROM_NAME(AActor, bCanBeDamaged, this);
 PRAGMA_ENABLE_DEPRECATION_WARNINGS
 }
 
@@ -5009,14 +5028,14 @@ void AActor::SetRole(ENetRole InRole)
 {
 PRAGMA_DISABLE_DEPRECATION_WARNINGS
 	Role = InRole;
-	// MARK_PROPERTY_DIRTY_FROM_NAME(AActor, Role, this);
+	MARK_PROPERTY_DIRTY_FROM_NAME(AActor, Role, this);
 PRAGMA_ENABLE_DEPRECATION_WARNINGS
 }
 
 FRepMovement& AActor::GetReplicatedMovement_Mutable()
 {
 PRAGMA_DISABLE_DEPRECATION_WARNINGS
-	// MARK_PROPERTY_DIRTY_FROM_NAME(AActor, ReplicatedMovement, this);
+	MARK_PROPERTY_DIRTY_FROM_NAME(AActor, ReplicatedMovement, this);
 	return ReplicatedMovement;
 PRAGMA_ENABLE_DEPRECATION_WARNINGS
 }
@@ -5030,7 +5049,7 @@ void AActor::SetInstigator(APawn* InInstigator)
 {
 PRAGMA_DISABLE_DEPRECATION_WARNINGS
 	Instigator = InInstigator;
-	// MARK_PROPERTY_DIRTY_FROM_NAME(AActor, Instigator, this);
+	MARK_PROPERTY_DIRTY_FROM_NAME(AActor, Instigator, this);
 PRAGMA_ENABLE_DEPRECATION_WARNINGS
 }
 
diff --git a/Engine/Source/Runtime/Engine/Private/ActorReplication.cpp b/Engine/Source/Runtime/Engine/Private/ActorReplication.cpp
--- a/Engine/Source/Runtime/Engine/Private/ActorReplication.cpp
+++ b/Engine/Source/Runtime/Engine/Private/ActorReplication.cpp
@@ -18,6 +18,7 @@
 #include "PhysicsReplication.h"
 #include "PhysicsPublic.h"
 #include "DrawDebugHelpers.h"
+#include "Net/Core/PushModel/PushModel.h"
 
 /*-----------------------------------------------------------------------------
 	AActor networking implementation.
@@ -342,7 +343,7 @@ PRAGMA_DISABLE_DEPRECATION_WARNINGS
 		bool bWasRepMovementModified = false;
 
 		AActor* OldAttachParent = AttachmentReplication.AttachParent;
-		USceneComponent* OldAttachComponent = AttachmentReplication.AttachComponent; 
+		USceneComponent* OldAttachComponent = AttachmentReplication.AttachComponent;
 	
 		AttachmentReplication.AttachParent = nullptr;
 		AttachmentReplication.AttachComponent = nullptr;
@@ -358,6 +359,7 @@ PRAGMA_DISABLE_DEPRECATION_WARNINGS
 			// Their replication will affect our position indirectly since we are attached.
 			ReplicatedMovement.bRepPhysics = !RootPrimComp->IsWelded();
 			
+			// Technically, the values might have stayed the same, but we'll just assume they've changed.
 			bWasRepMovementModified = true;
 		}
 		else if (RootComponent != nullptr)
@@ -391,20 +393,20 @@ PRAGMA_DISABLE_DEPRECATION_WARNINGS
 				bWasRepMovementModified = true;
 			}
 
-			if (bWasRepMovementModified ||
-				ReplicatedMovement.bRepPhysics != false)
-			{
-				// MARK_PROPERTY_DIRTY_FROM_NAME(AActor, ReplicatedMovement, this);
-			}
+			bWasRepMovementModified = (bWasRepMovementModified || ReplicatedMovement.bRepPhysics);
+			ReplicatedMovement.bRepPhysics = false;
+		}
 
-			if (bWasAttachmentModified ||
-				OldAttachParent != AttachmentReplication.AttachParent ||
-				OldAttachComponent != AttachmentReplication.AttachComponent)
-			{
-				// MARK_PROPERTY_DIRTY_FROM_NAME(AActor, AttachmentReplication, this);
-			}
+		if (bWasRepMovementModified)
+		{
+			MARK_PROPERTY_DIRTY_FROM_NAME(AActor, ReplicatedMovement, this);
+		}
 
-			ReplicatedMovement.bRepPhysics = false;
+		if (bWasAttachmentModified ||
+			OldAttachParent != AttachmentReplication.AttachParent ||
+			OldAttachComponent != AttachmentReplication.AttachComponent)
+		{
+			MARK_PROPERTY_DIRTY_FROM_NAME(AActor, AttachmentReplication, this);
 		}
 	}
 PRAGMA_ENABLE_DEPRECATION_WARNINGS
@@ -418,29 +420,23 @@ void AActor::GetLifetimeReplicatedProps( TArray< FLifetimeProperty > & OutLifeti
 		BPClass->GetLifetimeBlueprintReplicationList(OutLifetimeProps);
 	}
 
-PRAGMA_DISABLE_DEPRECATION_WARNINGS
-	DOREPLIFETIME( AActor, bReplicateMovement );
-	DOREPLIFETIME( AActor, Role );
-PRAGMA_ENABLE_DEPRECATION_WARNINGS
-	DOREPLIFETIME( AActor, RemoteRole );
-	DOREPLIFETIME( AActor, Owner );
-
-PRAGMA_DISABLE_DEPRECATION_WARNINGS
-	DOREPLIFETIME( AActor, bHidden );
-PRAGMA_ENABLE_DEPRECATION_WARNINGS
-
-	DOREPLIFETIME( AActor, bTearOff );
-
-PRAGMA_DISABLE_DEPRECATION_WARNINGS
-	DOREPLIFETIME( AActor, bCanBeDamaged );
-PRAGMA_ENABLE_DEPRECATION_WARNINGS
-
-	DOREPLIFETIME_CONDITION_NOTIFY( AActor, AttachmentReplication, COND_Custom, REPNOTIFY_Always );
+	FDoRepLifetimeParams SharedParams;
+	SharedParams.bIsPushBased = true;
+
+	DOREPLIFETIME_WITH_PARAMS_FAST(AActor, bReplicateMovement, SharedParams);
+	DOREPLIFETIME_WITH_PARAMS_FAST(AActor, Role, SharedParams);
+	DOREPLIFETIME_WITH_PARAMS_FAST(AActor, RemoteRole, SharedParams);
+	DOREPLIFETIME_WITH_PARAMS_FAST(AActor, Owner, SharedParams);
+	DOREPLIFETIME_WITH_PARAMS_FAST(AActor, bHidden, SharedParams);
+	DOREPLIFETIME_WITH_PARAMS_FAST(AActor, bTearOff, SharedParams);
+	DOREPLIFETIME_WITH_PARAMS_FAST(AActor, bCanBeDamaged, SharedParams);
+	DOREPLIFETIME_WITH_PARAMS_FAST(AActor, Instigator, SharedParams);
+	
+	FDoRepLifetimeParams AttachmentReplicationParams{COND_Custom, REPNOTIFY_Always, /*bIsPushBased=*/true};
+	DOREPLIFETIME_WITH_PARAMS_FAST(AActor, AttachmentReplication, AttachmentReplicationParams);
 
-PRAGMA_DISABLE_DEPRECATION_WARNINGS
-	DOREPLIFETIME( AActor, Instigator );
-	DOREPLIFETIME_CONDITION_NOTIFY( AActor, ReplicatedMovement, COND_SimulatedOrPhysics, REPNOTIFY_Always );
-PRAGMA_ENABLE_DEPRECATION_WARNINGS
+	FDoRepLifetimeParams ReplicatedMovementParams{COND_SimulatedOrPhysics, REPNOTIFY_Always, /*bIsPushBased=*/true};
+	DOREPLIFETIME_WITH_PARAMS_FAST(AActor, ReplicatedMovement, ReplicatedMovementParams);
 }
 
 bool AActor::ReplicateSubobjects(UActorChannel *Channel, FOutBunch *Bunch, FReplicationFlags *RepFlags)
diff --git a/Engine/Source/Runtime/Engine/Private/Components/ActorComponent.cpp b/Engine/Source/Runtime/Engine/Private/Components/ActorComponent.cpp
--- a/Engine/Source/Runtime/Engine/Private/Components/ActorComponent.cpp
+++ b/Engine/Source/Runtime/Engine/Private/Components/ActorComponent.cpp
@@ -30,6 +30,7 @@
 #include "ComponentUtils.h"
 #include "Engine/Engine.h"
 #include "HAL/LowLevelMemTracker.h"
+#include "Net/Core/PushModel/PushModel.h"
 
 #if WITH_EDITOR
 #include "Kismet2/ComponentEditorUtils.h"
@@ -1754,7 +1755,7 @@ void UActorComponent::SetIsReplicated(bool bShouldReplicate)
 		{
 PRAGMA_DISABLE_DEPRECATION_WARNINGS
 			bReplicates = bShouldReplicate;
-			// MARK_PROPERTY_DIRTY_FROM_NAME(UActorComponent, bReplicates, this);
+			MARK_PROPERTY_DIRTY_FROM_NAME(UActorComponent, bReplicates, this);
 PRAGMA_ENABLE_DEPRECATION_WARNINGS
 
 			if (AActor* MyOwner = GetOwner())
@@ -1802,10 +1803,11 @@ void UActorComponent::GetLifetimeReplicatedProps( TArray< FLifetimeProperty > &
 		BPClass->GetLifetimeBlueprintReplicationList(OutLifetimeProps);
 	}
 
-PRAGMA_DISABLE_DEPRECATION_WARNINGS
-	DOREPLIFETIME( UActorComponent, bIsActive );
-	DOREPLIFETIME( UActorComponent, bReplicates );
-PRAGMA_ENABLE_DEPRECATION_WARNINGS
+	FDoRepLifetimeParams SharedParams;
+	SharedParams.bIsPushBased = true;
+
+	DOREPLIFETIME_WITH_PARAMS_FAST(UActorComponent, bIsActive, SharedParams);
+	DOREPLIFETIME_WITH_PARAMS_FAST(UActorComponent, bReplicates, SharedParams);
 }
 
 void UActorComponent::OnRep_IsActive()
@@ -1971,7 +1973,7 @@ void UActorComponent::SetIsReplicatedByDefault(const bool bNewReplicates)
 	{
 PRAGMA_DISABLE_DEPRECATION_WARNINGS
 		bReplicates = bNewReplicates;
-		// MARK_PROPERTY_DIRTY_FROM_NAME(UActorComponent, bReplicates, this);
+		MARK_PROPERTY_DIRTY_FROM_NAME(UActorComponent, bReplicates, this);
 PRAGMA_ENABLE_DEPRECATION_WARNINGS
 	}
 	else
@@ -1985,7 +1987,7 @@ void UActorComponent::SetActiveFlag(const bool bNewIsActive)
 {
 PRAGMA_DISABLE_DEPRECATION_WARNINGS
 	bIsActive = bNewIsActive;
-	// MARK_PROPERTY_DIRTY_FROM_NAME(UActorComponent, bIsActive, this);
+	MARK_PROPERTY_DIRTY_FROM_NAME(UActorComponent, bIsActive, this);
 PRAGMA_ENABLE_DEPRECATION_WARNINGS
 }
 
diff --git a/Engine/Source/Runtime/Engine/Private/Components/SceneComponent.cpp b/Engine/Source/Runtime/Engine/Private/Components/SceneComponent.cpp
--- a/Engine/Source/Runtime/Engine/Private/Components/SceneComponent.cpp
+++ b/Engine/Source/Runtime/Engine/Private/Components/SceneComponent.cpp
@@ -33,6 +33,7 @@
 #include "DeviceProfiles/DeviceProfileManager.h"
 #include "Interfaces/ITargetPlatform.h"
 #include "DeviceProfiles/DeviceProfile.h"
+#include "Net/Core/PushModel/PushModel.h"
 
 #define LOCTEXT_NAMESPACE "SceneComponent"
 
@@ -3286,25 +3287,22 @@ void USceneComponent::GetLifetimeReplicatedProps(TArray< FLifetimeProperty > & O
 {
 	Super::GetLifetimeReplicatedProps(OutLifetimeProps);
 	
-PRAGMA_DISABLE_DEPRECATION_WARNINGS
-	DOREPLIFETIME(USceneComponent, bAbsoluteLocation);
-	DOREPLIFETIME(USceneComponent, bAbsoluteRotation);
-	DOREPLIFETIME(USceneComponent, bAbsoluteScale);
-	DOREPLIFETIME(USceneComponent, bVisible);
-PRAGMA_ENABLE_DEPRECATION_WARNINGS
-
-	DOREPLIFETIME(USceneComponent, bShouldBeAttached);
-	DOREPLIFETIME(USceneComponent, bShouldSnapLocationWhenAttached);
-	DOREPLIFETIME(USceneComponent, bShouldSnapRotationWhenAttached);
-	DOREPLIFETIME(USceneComponent, AttachParent);
-	DOREPLIFETIME(USceneComponent, AttachChildren);
-	DOREPLIFETIME(USceneComponent, AttachSocketName);
-
-PRAGMA_DISABLE_DEPRECATION_WARNINGS
-	DOREPLIFETIME(USceneComponent, RelativeLocation);
-	DOREPLIFETIME(USceneComponent, RelativeRotation);
-	DOREPLIFETIME(USceneComponent, RelativeScale3D);
-PRAGMA_ENABLE_DEPRECATION_WARNINGS
+	FDoRepLifetimeParams SharedParams;
+	SharedParams.bIsPushBased = true;
+
+	DOREPLIFETIME_WITH_PARAMS_FAST(USceneComponent, bAbsoluteLocation, SharedParams);
+	DOREPLIFETIME_WITH_PARAMS_FAST(USceneComponent, bAbsoluteRotation, SharedParams);
+	DOREPLIFETIME_WITH_PARAMS_FAST(USceneComponent, bAbsoluteScale, SharedParams);
+	DOREPLIFETIME_WITH_PARAMS_FAST(USceneComponent, bVisible, SharedParams);
+	DOREPLIFETIME_WITH_PARAMS_FAST(USceneComponent, bShouldBeAttached, SharedParams);
+	DOREPLIFETIME_WITH_PARAMS_FAST(USceneComponent, bShouldSnapLocationWhenAttached, SharedParams);
+	DOREPLIFETIME_WITH_PARAMS_FAST(USceneComponent, bShouldSnapRotationWhenAttached, SharedParams);
+	DOREPLIFETIME_WITH_PARAMS_FAST(USceneComponent, AttachParent, SharedParams);
+	DOREPLIFETIME_WITH_PARAMS_FAST(USceneComponent, AttachChildren, SharedParams);
+	DOREPLIFETIME_WITH_PARAMS_FAST(USceneComponent, AttachSocketName, SharedParams);
+	DOREPLIFETIME_WITH_PARAMS_FAST(USceneComponent, RelativeLocation, SharedParams);
+	DOREPLIFETIME_WITH_PARAMS_FAST(USceneComponent, RelativeRotation, SharedParams);
+	DOREPLIFETIME_WITH_PARAMS_FAST(USceneComponent, RelativeScale3D, SharedParams);
 }
 
 #if WITH_EDITOR
@@ -3838,14 +3836,14 @@ void USceneComponent::SetVisibleFlag(const bool bNewVisible)
 {
 PRAGMA_DISABLE_DEPRECATION_WARNINGS
 	bVisible = bNewVisible;
-	// MARK_PROPERTY_DIRTY_FROM_NAME(USceneComponent, bVisible, this);
+	MARK_PROPERTY_DIRTY_FROM_NAME(USceneComponent, bVisible, this);
 PRAGMA_ENABLE_DEPRECATION_WARNINGS
 }
 
 FVector& USceneComponent::GetRelativeLocation_DirectMutable()
 {
 PRAGMA_DISABLE_DEPRECATION_WARNINGS
-	// MARK_PROPERTY_DIRTY_FROM_NAME(USceneComponent, RelativeLocation, this);
+	MARK_PROPERTY_DIRTY_FROM_NAME(USceneComponent, RelativeLocation, this);
 	return RelativeLocation;
 PRAGMA_ENABLE_DEPRECATION_WARNINGS
 }
@@ -3858,7 +3856,7 @@ void USceneComponent::SetRelativeLocation_Direct(const FVector NewRelativeLocati
 FRotator& USceneComponent::GetRelativeRotation_DirectMutable()
 {
 PRAGMA_DISABLE_DEPRECATION_WARNINGS
-	// MARK_PROPERTY_DIRTY_FROM_NAME(USceneComponent, RelativeRotation, this);
+	MARK_PROPERTY_DIRTY_FROM_NAME(USceneComponent, RelativeRotation, this);
 	return RelativeRotation;
 PRAGMA_ENABLE_DEPRECATION_WARNINGS
 }
@@ -3871,7 +3869,7 @@ void USceneComponent::SetRelativeRotation_Direct(const FRotator NewRelativeRotat
 FVector& USceneComponent::GetRelativeScale3D_DirectMutable()
 {
 PRAGMA_DISABLE_DEPRECATION_WARNINGS
-	// MARK_PROPERTY_DIRTY_FROM_NAME(USceneComponent, RelativeScale3D, this);
+	MARK_PROPERTY_DIRTY_FROM_NAME(USceneComponent, RelativeScale3D, this);
 	return RelativeScale3D;
 PRAGMA_ENABLE_DEPRECATION_WARNINGS
 }
@@ -3885,7 +3883,7 @@ void USceneComponent::SetUsingAbsoluteLocation(bool bInAbsoluteLocation)
 {
 PRAGMA_DISABLE_DEPRECATION_WARNINGS
 	bAbsoluteLocation = bInAbsoluteLocation;
-	// MARK_PROPERTY_DIRTY_FROM_NAME(USceneComponent, bAbsoluteLocation, this);
+	MARK_PROPERTY_DIRTY_FROM_NAME(USceneComponent, bAbsoluteLocation, this);
 PRAGMA_ENABLE_DEPRECATION_WARNINGS
 }
 
@@ -3893,7 +3891,7 @@ void USceneComponent::SetUsingAbsoluteRotation(bool bInAbsoluteRotation)
 {
 PRAGMA_DISABLE_DEPRECATION_WARNINGS
 	bAbsoluteRotation = bInAbsoluteRotation;
-	// MARK_PROPERTY_DIRTY_FROM_NAME(USceneComponent, bAbsoluteRotation, this);
+	MARK_PROPERTY_DIRTY_FROM_NAME(USceneComponent, bAbsoluteRotation, this);
 PRAGMA_ENABLE_DEPRECATION_WARNINGS
 }
 
@@ -3901,7 +3899,7 @@ void USceneComponent::SetUsingAbsoluteScale(bool bInAbsoluteScale)
 {
 PRAGMA_DISABLE_DEPRECATION_WARNINGS
 	bAbsoluteScale = bInAbsoluteScale;
-	// MARK_PROPERTY_DIRTY_FROM_NAME(USceneComponent, bAbsoluteScale, this);
+	MARK_PROPERTY_DIRTY_FROM_NAME(USceneComponent, bAbsoluteScale, this);
 PRAGMA_ENABLE_DEPRECATION_WARNINGS
 }
 
@@ -3909,7 +3907,7 @@ void USceneComponent::SetAttachParent(USceneComponent* NewAttachParent)
 {
 PRAGMA_DISABLE_DEPRECATION_WARNINGS
 	AttachParent = NewAttachParent;
-	// MARK_PROPERTY_DIRTY_FROM_NAME(USceneComponent, AttachParent, this);
+	MARK_PROPERTY_DIRTY_FROM_NAME(USceneComponent, AttachParent, this);
 PRAGMA_ENABLE_DEPRECATION_WARNINGS
 }
 
@@ -3917,13 +3915,13 @@ void USceneComponent::SetAttachSocketName(FName NewSocketName)
 {
 PRAGMA_DISABLE_DEPRECATION_WARNINGS
 	AttachSocketName = NewSocketName;
-	// MARK_PROPERTY_DIRTY_FROM_NAME(USceneComponent, AttachSocketName, this);
+	MARK_PROPERTY_DIRTY_FROM_NAME(USceneComponent, AttachSocketName, this);
 PRAGMA_ENABLE_DEPRECATION_WARNINGS
 }
 
 void USceneComponent::ModifiedAttachChildren()
 {
-	// MARK_PROPERTY_DIRTY_FROM_NAME(USceneComponent, AttachChildren, this);
+	MARK_PROPERTY_DIRTY_FROM_NAME(USceneComponent, AttachChildren, this);
 }
 
 #undef LOCTEXT_NAMESPACE
