Dump PendingNetGuidResolves when a client's NetDriver is shutdown.

diff --git a/Engine/Source/Runtime/Engine/Classes/Engine/PackageMapClient.h b/Engine/Source/Runtime/Engine/Classes/Engine/PackageMapClient.h
--- a/Engine/Source/Runtime/Engine/Classes/Engine/PackageMapClient.h
+++ b/Engine/Source/Runtime/Engine/Classes/Engine/PackageMapClient.h
@@ -448,6 +448,18 @@ public:
 
 	virtual void Serialize(FArchive& Ar) override;
 
+	FString GetFullNetGUIDPath(const FNetworkGUID& NetGUID) const
+	{
+		FString FullGuidCachePath;
+
+		if (const FNetGUIDCache * const GuidCacheLocal = GuidCache.Get())
+		{
+			FullGuidCachePath = GuidCacheLocal->FullNetGUIDPath(NetGUID);
+		}
+
+		return FullGuidCachePath;
+	}
+
 protected:
 
 	/** Functions to help with exporting/importing net field export info */
diff --git a/Engine/Source/Runtime/Engine/Private/NetworkDriver.cpp b/Engine/Source/Runtime/Engine/Private/NetworkDriver.cpp
--- a/Engine/Source/Runtime/Engine/Private/NetworkDriver.cpp
+++ b/Engine/Source/Runtime/Engine/Private/NetworkDriver.cpp
@@ -1554,20 +1554,55 @@ void UNetDriver::UnregisterTickEvents(class UWorld* InWorld)
 	}
 }
 
+static bool bCVarLogPendingGuidsOnShutdown = false;
+static FAutoConsoleVariableRef CVarLogPendingGuidsOnShutdown(
+	TEXT("Net.LogPendingGuidsOnShutdown"),
+	bCVarLogPendingGuidsOnShutdown,
+	TEXT("")
+);
+
 /** Shutdown all connections managed by this net driver */
 void UNetDriver::Shutdown()
 {
 	// Client closing connection to server
 	if (ServerConnection)
 	{
+		const UPackageMapClient* const ClientPackageMap = Cast<UPackageMapClient>(ServerConnection->PackageMap);
+		const bool bLogGuids = bCVarLogPendingGuidsOnShutdown && (ClientPackageMap != nullptr);
+		TSet<FNetworkGUID> GuidsToLog;
+
 		for (UChannel* Channel : ServerConnection->OpenChannels)
-		 {
-			 UActorChannel* ActorChannel = Cast<UActorChannel>(Channel);
-			 if (ActorChannel)
+		{
+			 if (UActorChannel * ActorChannel = Cast<UActorChannel>(Channel))
 			 {
+				 if (bLogGuids)
+				 {
+					 GuidsToLog.Append(ActorChannel->PendingGuidResolves);
+				 }
 				 ActorChannel->CleanupReplicators();
 			 }
-		 }
+		}
+
+		if (GuidsToLog.Num() > 0)
+		{
+			TArray<FString> GuidStrings;
+			GuidStrings.Reserve(GuidsToLog.Num());
+			for (const FNetworkGUID GuidToLog : GuidsToLog)
+			{
+				FString FullNetGUIDPath = ClientPackageMap->GetFullNetGUIDPath(GuidToLog);
+				if (!FullNetGUIDPath.IsEmpty())
+				{
+					GuidStrings.Emplace(MoveTemp(FullNetGUIDPath));
+				}
+				else
+				{
+					GuidStrings.Emplace(GuidToLog.ToString());
+				}
+			}
+
+			UE_LOG(LogNet, Warning, TEXT("NetDriver::Shutdown: Pending Guids: \n%s"), *FString::Join(GuidStrings, TEXT("\n")));
+		}
+		
 
 		// Calls Channel[0]->Close to send a close bunch to server
 		ServerConnection->Close();
