Removing stereo bleed members. Not used in the new audio engine.

diff --git a/Engine/Source/Runtime/ALAudio/Private/ALAudioSource.cpp b/Engine/Source/Runtime/ALAudio/Private/ALAudioSource.cpp
--- a/Engine/Source/Runtime/ALAudio/Private/ALAudioSource.cpp
+++ b/Engine/Source/Runtime/ALAudio/Private/ALAudioSource.cpp
@@ -104,11 +104,6 @@ void FALSoundSource::Update( void )
 	else
 	{
 		Volume = WaveInstance->GetActualVolume();
-		if (SetStereoBleed())
-		{
-			// Emulate the bleed to rear speakers followed by stereo fold down
-			Volume *= 1.25f;
-		}
 		Volume *= AudioDevice->GetPlatformAudioHeadroom();
 		Volume = FMath::Clamp(Volume, 0.0f, MAX_VOLUME);
 	}
diff --git a/Engine/Source/Runtime/Android/AndroidAudio/Private/AndroidAudioSource.cpp b/Engine/Source/Runtime/Android/AndroidAudio/Private/AndroidAudioSource.cpp
--- a/Engine/Source/Runtime/Android/AndroidAudio/Private/AndroidAudioSource.cpp
+++ b/Engine/Source/Runtime/Android/AndroidAudio/Private/AndroidAudioSource.cpp
@@ -492,11 +492,7 @@ void FSLESSoundSource::Update( void )
 	FSoundSource::UpdateCommon();
 	
 	float Volume = WaveInstance->GetActualVolume();
-	if (SetStereoBleed())
-	{
-		// Emulate the bleed to rear speakers followed by stereo fold down
-		Volume *= 1.25f;
-	}
+
 	Volume *= AudioDevice->GetPlatformAudioHeadroom();
 	Volume = FMath::Clamp(Volume, 0.0f, MAX_VOLUME);
 	
diff --git a/Engine/Source/Runtime/Engine/Classes/Sound/SoundClass.h b/Engine/Source/Runtime/Engine/Classes/Sound/SoundClass.h
--- a/Engine/Source/Runtime/Engine/Classes/Sound/SoundClass.h
+++ b/Engine/Source/Runtime/Engine/Classes/Sound/SoundClass.h
@@ -69,10 +69,6 @@ struct FSoundClassProperties
 	UPROPERTY(EditAnywhere, BlueprintReadOnly, Category = General)
 	float AttenuationDistanceScale;
 
-	/** The amount of stereo sounds to bleed to the rear speakers */
-	UPROPERTY(EditAnywhere, BlueprintReadOnly, Category = Routing)
-	float StereoBleed;
-
 	/** The amount of a sound to bleed to the LFE channel */
 	UPROPERTY(EditAnywhere, BlueprintReadOnly, Category = Routing, meta = (DisplayName = "LFE Bleed"))
 	float LFEBleed;
diff --git a/Engine/Source/Runtime/Engine/Private/Audio.cpp b/Engine/Source/Runtime/Engine/Private/Audio.cpp
--- a/Engine/Source/Runtime/Engine/Private/Audio.cpp
+++ b/Engine/Source/Runtime/Engine/Private/Audio.cpp
@@ -348,24 +348,6 @@ bool FSoundSource::SetReverbApplied(bool bHardwareAvailable)
 	return(bReverbApplied);
 }
 
-float FSoundSource::SetStereoBleed()
-{
-	StereoBleed = 0.0f;
-
-	// All stereo sounds bleed by default
-	if (WaveInstance->WaveData->NumChannels == 2)
-	{
-		StereoBleed = WaveInstance->StereoBleed;
-
-		if (AudioDevice->GetMixDebugState() == DEBUGSTATE_TestStereoBleed)
-		{
-			StereoBleed = 1.0f;
-		}
-	}
-
-	return StereoBleed;
-}
-
 float FSoundSource::SetLFEBleed()
 {
 	LFEBleed = WaveInstance->LFEBleed;
@@ -801,7 +783,6 @@ FWaveInstance::FWaveInstance(const UPTRINT InWaveInstanceHash, FActiveSound& InA
 	, VoiceCenterChannelVolume(0.0f)
 	, RadioFilterVolume(0.0f)
 	, RadioFilterVolumeThreshold(0.0f)
-	, StereoBleed(0.0f)
 	, LFEBleed(0.0f)
 	, LoopingMode(LOOP_Never)
 	, StartTime(-1.f)
diff --git a/Engine/Source/Runtime/Engine/Private/AudioDevice.cpp b/Engine/Source/Runtime/Engine/Private/AudioDevice.cpp
--- a/Engine/Source/Runtime/Engine/Private/AudioDevice.cpp
+++ b/Engine/Source/Runtime/Engine/Private/AudioDevice.cpp
@@ -1261,13 +1261,6 @@ bool FAudioDevice::HandleTestHPFCommand(const TCHAR* Cmd, FOutputDevice& Ar)
 	return true;
 }
 
-bool FAudioDevice::HandleTestStereoBleedCommand(const TCHAR* Cmd, FOutputDevice& Ar)
-{
-	Ar.Logf(TEXT("StereoBleed set to max for all sources"));
-	SetMixDebugState(DEBUGSTATE_TestStereoBleed);
-	return true;
-}
-
 bool FAudioDevice::HandleTestLFEBleedCommand(const TCHAR* Cmd, FOutputDevice& Ar)
 {
 	Ar.Logf(TEXT("LFEBleed set to max for all sources"));
@@ -1957,10 +1950,6 @@ bool FAudioDevice::Exec(UWorld* InWorld, const TCHAR* Cmd, FOutputDevice& Ar)
 	{
 		return HandleTestLPFCommand(Cmd, Ar);
 	}
-	else if (FParse::Command(&Cmd, TEXT("TestStereoBleed")))
-	{
-		return HandleTestStereoBleedCommand(Cmd, Ar);
-	}
 	else if (FParse::Command(&Cmd, TEXT("TestLFEBleed")))
 	{
 		return HandleTestLPFCommand(Cmd, Ar);
diff --git a/Engine/Source/Runtime/Engine/Private/SoundClass.cpp b/Engine/Source/Runtime/Engine/Private/SoundClass.cpp
--- a/Engine/Source/Runtime/Engine/Private/SoundClass.cpp
+++ b/Engine/Source/Runtime/Engine/Private/SoundClass.cpp
@@ -23,7 +23,6 @@ FSoundClassProperties::FSoundClassProperties()
 	, Pitch(1.0f)
 	, LowPassFilterFrequency(MAX_FILTER_FREQUENCY)
 	, AttenuationDistanceScale(1.0f)
-	, StereoBleed(0.25f)
 	, LFEBleed(0.5f)
 	, VoiceCenterChannelVolume(0.0f)
 	, RadioFilterVolume(0.0f)
diff --git a/Engine/Source/Runtime/Engine/Private/SoundWave.cpp b/Engine/Source/Runtime/Engine/Private/SoundWave.cpp
--- a/Engine/Source/Runtime/Engine/Private/SoundWave.cpp
+++ b/Engine/Source/Runtime/Engine/Private/SoundWave.cpp
@@ -1746,7 +1746,6 @@ void USoundWave::Parse(FAudioDevice* AudioDevice, const UPTRINT NodeWaveInstance
 		WaveInstance->VoiceCenterChannelVolume = SoundClassProperties->VoiceCenterChannelVolume;
 		WaveInstance->RadioFilterVolume = SoundClassProperties->RadioFilterVolume * ParseParams.VolumeMultiplier;
 		WaveInstance->RadioFilterVolumeThreshold = SoundClassProperties->RadioFilterVolumeThreshold * ParseParams.VolumeMultiplier;
-		WaveInstance->StereoBleed = SoundClassProperties->StereoBleed;
 		WaveInstance->LFEBleed = SoundClassProperties->LFEBleed;
 
 		WaveInstance->bIsUISound = ActiveSound.bIsUISound || SoundClassProperties->bIsUISound;
@@ -1780,7 +1779,6 @@ void USoundWave::Parse(FAudioDevice* AudioDevice, const UPTRINT NodeWaveInstance
 		WaveInstance->VoiceCenterChannelVolume = 0.f;
 		WaveInstance->RadioFilterVolume = 0.f;
 		WaveInstance->RadioFilterVolumeThreshold = 0.f;
-		WaveInstance->StereoBleed = 0.f;
 		WaveInstance->LFEBleed = 0.f;
 		WaveInstance->bIsUISound = ActiveSound.bIsUISound;
 		WaveInstance->bIsMusic = ActiveSound.bIsMusic;
diff --git a/Engine/Source/Runtime/Engine/Public/Audio.h b/Engine/Source/Runtime/Engine/Public/Audio.h
--- a/Engine/Source/Runtime/Engine/Public/Audio.h
+++ b/Engine/Source/Runtime/Engine/Public/Audio.h
@@ -224,9 +224,6 @@ public:
 	/** The volume at which the radio filter kicks in */
 	float RadioFilterVolumeThreshold;
 
-	/** The amount of stereo sounds to bleed to the rear speakers */
-	float StereoBleed;
-
 	/** The amount of a sound to bleed to the LFE channel */
 	float LFEBleed;
 
@@ -540,7 +537,6 @@ public:
 		: AudioDevice(InAudioDevice)
 		, WaveInstance(nullptr)
 		, Buffer(nullptr)
-		, StereoBleed(0.0f)
 		, LFEBleed(0.5f)
 		, LPFFrequency(MAX_FILTER_FREQUENCY)
 		, HPFFrequency(MIN_FILTER_FREQUENCY)
@@ -632,9 +628,6 @@ public:
 	/** Set the bReverbApplied variable. */
 	ENGINE_API bool SetReverbApplied(bool bHardwareAvailable);
 
-	/** Set the StereoBleed variable. */
-	ENGINE_API float SetStereoBleed();
-
 	/** Updates and sets the LFEBleed variable. */
 	ENGINE_API float SetLFEBleed();
 
@@ -697,9 +690,6 @@ protected:
 	/** Cached sound buffer associated with currently bound wave instance. */
 	FSoundBuffer* Buffer;
 
-	/** The amount of stereo sounds to bleed to the rear speakers */
-	float StereoBleed;
-
 	/** The amount of a sound to bleed to the LFE speaker */
 	float LFEBleed;
 
diff --git a/Engine/Source/Runtime/Engine/Public/AudioDevice.h b/Engine/Source/Runtime/Engine/Public/AudioDevice.h
--- a/Engine/Source/Runtime/Engine/Public/AudioDevice.h
+++ b/Engine/Source/Runtime/Engine/Public/AudioDevice.h
@@ -72,8 +72,6 @@ enum EDebugState
 	DEBUGSTATE_TestLPF,
 	// Force LPF on all sources
 	DEBUGSTATE_TestHPF,
-	// Bleed stereo sounds fully to the rear speakers
-	DEBUGSTATE_TestStereoBleed,
 	// Bleed all sounds to the LFE speaker
 	DEBUGSTATE_TestLFEBleed,
 	// Disable any LPF filter effects
@@ -442,7 +440,6 @@ private:
 	bool HandleIsolateReverbCommand(const TCHAR* Cmd, FOutputDevice& Ar);
 	bool HandleTestLPFCommand(const TCHAR* Cmd, FOutputDevice& Ar);
 	bool HandleTestHPFCommand(const TCHAR* Cmd, FOutputDevice& Ar);
-	bool HandleTestStereoBleedCommand(const TCHAR* Cmd, FOutputDevice& Ar);
 	bool HandleTestLFEBleedCommand(const TCHAR* Cmd, FOutputDevice& Ar);
 	bool HandleDisableLPFCommand(const TCHAR* Cmd, FOutputDevice& Ar);
 	bool HandleDisableHPFCommand(const TCHAR* Cmd, FOutputDevice& Ar);
diff --git a/Engine/Source/Runtime/IOS/IOSAudio/Private/IOSAudioSource.cpp b/Engine/Source/Runtime/IOS/IOSAudio/Private/IOSAudioSource.cpp
--- a/Engine/Source/Runtime/IOS/IOSAudio/Private/IOSAudioSource.cpp
+++ b/Engine/Source/Runtime/IOS/IOSAudio/Private/IOSAudioSource.cpp
@@ -225,12 +225,6 @@ void FIOSAudioSoundSource::Update(void)
 		Volume = WaveInstance->GetActualVolume();
 	}
 
-	if (SetStereoBleed())
-	{
-		// Emulate the bleed to rear speakers followed by stereo fold down
-		Volume *= 1.25f;
-	}
-
     SetFilterFrequency();
     
     SourceLPFFrequency = LPFFrequency / (((float) SampleRate) * 0.5f);
diff --git a/Engine/Source/Runtime/Mac/CoreAudio/Private/CoreAudioSource.cpp b/Engine/Source/Runtime/Mac/CoreAudio/Private/CoreAudioSource.cpp
--- a/Engine/Source/Runtime/Mac/CoreAudio/Private/CoreAudioSource.cpp
+++ b/Engine/Source/Runtime/Mac/CoreAudio/Private/CoreAudioSource.cpp
@@ -338,12 +338,6 @@ void FCoreAudioSoundSource::Update( void )
 		float Azimuth = 0.0f;
 		float Elevation = 0.0f;
 
-		if( SetStereoBleed() )
-		{
-			// Emulate the bleed to rear speakers followed by stereo fold down
-			Volume *= 1.25f;
-		}
-		
 		// apply global multiplier (ie to disable sound when not the foreground app)
 		Volume = FMath::Clamp<float>( Volume, 0.0f, MAX_VOLUME );
 		
diff --git a/Engine/Source/Runtime/Windows/XAudio2/Private/XAudio2Source.cpp b/Engine/Source/Runtime/Windows/XAudio2/Private/XAudio2Source.cpp
--- a/Engine/Source/Runtime/Windows/XAudio2/Private/XAudio2Source.cpp
+++ b/Engine/Source/Runtime/Windows/XAudio2/Private/XAudio2Source.cpp
@@ -870,8 +870,8 @@ void FXAudio2SoundSource::GetStereoChannelVolumes(float ChannelVolumes[CHANNEL_M
 		// but only if this is not an HRTF-spatialized mono sound
 		if (!IsUsingHrtfSpatializer() && FXAudioDeviceProperties::NumSpeakers == 6)
 		{
-			ChannelVolumes[CHANNELOUT_LEFTSURROUND] = AttenuatedVolume * StereoBleed;
-			ChannelVolumes[CHANNELOUT_RIGHTSURROUND] = AttenuatedVolume * StereoBleed;
+			ChannelVolumes[CHANNELOUT_LEFTSURROUND] = AttenuatedVolume;
+			ChannelVolumes[CHANNELOUT_RIGHTSURROUND] = AttenuatedVolume;
 
 			ChannelVolumes[CHANNELOUT_LOWFREQUENCY] = AttenuatedVolume * LFEBleed * 0.5f;
 		}
@@ -1376,9 +1376,6 @@ void FXAudio2SoundSource::Update()
 		AudioDevice->ValidateAPICall( TEXT( "SetFrequencyRatio" ),
 			Source->SetFrequencyRatio( Pitch) );
 
-		// Set whether to bleed to the rear speakers
-		SetStereoBleed();
-
 		// Set the amount to bleed to the LFE speaker
 		SetLFEBleed();
 
