PhAt
- expose ConstraintInstance positions
- fix potential crash loading some physics assets

diff --git a/Engine/Source/Editor/PhysicsAssetEditor/Private/PhysicsAssetEditorSharedData.cpp b/Engine/Source/Editor/PhysicsAssetEditor/Private/PhysicsAssetEditorSharedData.cpp
--- a/Engine/Source/Editor/PhysicsAssetEditor/Private/PhysicsAssetEditorSharedData.cpp
+++ b/Engine/Source/Editor/PhysicsAssetEditor/Private/PhysicsAssetEditorSharedData.cpp
@@ -1921,7 +1921,10 @@ void FPhysicsAssetEditorSharedData::AddReferencedObjects(FReferenceCollector& Co
 	Collector.AddReferencedObject(CopiedBodySetup);
 	Collector.AddReferencedObject(CopiedConstraintTemplate);
 
-	PreviewScene.Pin()->AddReferencedObjects(Collector);
+	if (PreviewScene != nullptr)
+	{
+		PreviewScene.Pin()->AddReferencedObjects(Collector);
+	}
 }
 
 void FPhysicsAssetEditorSharedData::ForceDisableSimulation()
diff --git a/Engine/Source/Runtime/Engine/Classes/PhysicsEngine/ConstraintInstance.h b/Engine/Source/Runtime/Engine/Classes/PhysicsEngine/ConstraintInstance.h
--- a/Engine/Source/Runtime/Engine/Classes/PhysicsEngine/ConstraintInstance.h
+++ b/Engine/Source/Runtime/Engine/Classes/PhysicsEngine/ConstraintInstance.h
@@ -135,7 +135,7 @@ struct ENGINE_API FConstraintInstance
 	///////////////////////////// Body1 ref frame
 	
 	/** Location of constraint in Body1 reference frame. */
-	UPROPERTY()
+	UPROPERTY(EditInstanceOnly, Category=Constraint)
 	FVector Pos1;
 
 	/** Primary (twist) axis in Body1 reference frame. */
@@ -149,7 +149,7 @@ struct ENGINE_API FConstraintInstance
 	///////////////////////////// Body2 ref frame
 	
 	/** Location of constraint in Body2 reference frame. */
-	UPROPERTY()
+	UPROPERTY(EditInstanceOnly, Category=Constraint)
 	FVector Pos2;
 
 	/** Primary (twist) axis in Body2 reference frame. */
