#UE4 Added more info to an assert

diff --git a/Engine/Source/Runtime/Engine/Private/AudioDevice.cpp b/Engine/Source/Runtime/Engine/Private/AudioDevice.cpp
--- a/Engine/Source/Runtime/Engine/Private/AudioDevice.cpp
+++ b/Engine/Source/Runtime/Engine/Private/AudioDevice.cpp
@@ -5850,7 +5850,7 @@ void FAudioDevice::Precache(USoundWave* SoundWave, bool bSynchronous, bool bTrac
 			else
 			{
 				// This is the one case where precaching will not be done when this function exits
-				check(SoundWave->GetPrecacheState() == ESoundWavePrecacheState::InProgress);
+				checkf(SoundWave->GetPrecacheState() == ESoundWavePrecacheState::InProgress, TEXT("Bad PrecacheState %d on SoundWave %s"), static_cast<uint8>(SoundWave->GetPrecacheState()), *GetPathNameSafe(SoundWave));
 				SoundWave->AudioDecompressor = new FAsyncAudioDecompress(SoundWave, GetNumPrecacheFrames());
 				SoundWave->AudioDecompressor->StartBackgroundTask();
 				PrecachingSoundWaves.Add(SoundWave);
