Chaos: Fix Visitor crash. Fix UpdateBodyScale downcasting trianglemesh implicit to wrong type and creating corrupt shape pointer.
Add support for instanced geometry in UpdateBodyScale.

diff --git a/Engine/Source/Runtime/Engine/Private/PhysicsEngine/BodyInstance.cpp b/Engine/Source/Runtime/Engine/Private/PhysicsEngine/BodyInstance.cpp
--- a/Engine/Source/Runtime/Engine/Private/PhysicsEngine/BodyInstance.cpp
+++ b/Engine/Source/Runtime/Engine/Private/PhysicsEngine/BodyInstance.cpp
@@ -1707,9 +1707,6 @@ bool FBodyInstance::UpdateBodyScale(const FVector& InScale3D, bool bForceUpdate)
 
 			const FTransform& RelativeTM = GetRelativeBodyTransform(ShapeHandle);
 
-			// TODO: support instanced
-			CHAOS_ENSURE(!IsInstanced(ImplicitType));
-
 			bool bIsTransformed = false;
 			if(ImplicitType == ImplicitObjectType::Transformed)
 			{
@@ -1733,9 +1730,9 @@ bool FBodyInstance::UpdateBodyScale(const FVector& InScale3D, bool bForceUpdate)
 					FReal Radius = FMath::Max(SphereElem->Radius * AdjustedScale3DAbs.X, FCollisionShape::MinSphereRadius());
 
 
-					if (!CHAOS_ENSURE(!IsScaled(ImplicitType) || !CHAOS_ENSURE(!bIsTransformed)))
+					if (!CHAOS_ENSURE(!IsScaled(ImplicitType) && !bIsTransformed && !IsInstanced(ImplicitType)))
 					{
-						// No support for Scaled sphere or transformed.
+						// No support for Scaled, Instanced,or transformed.
 						break;
 					}
 
@@ -1749,9 +1746,9 @@ bool FBodyInstance::UpdateBodyScale(const FVector& InScale3D, bool bForceUpdate)
 				}
 				case ImplicitObjectType::Box:
 				{
-					if (!CHAOS_ENSURE(!IsScaled(ImplicitType)))
+					if (!CHAOS_ENSURE(!IsScaled(ImplicitType) && !IsInstanced(ImplicitType)))
 					{
-						// No support for ScaledImplicit Box yet
+						// No support for ScaledImplicit Box yet or instanced
 						break;
 					}
 
@@ -1803,9 +1800,9 @@ bool FBodyInstance::UpdateBodyScale(const FVector& InScale3D, bool bForceUpdate)
 
 					FKSphylElem* SphylElem = ShapeElem->GetShapeCheck<FKSphylElem>();
 
-					if (!CHAOS_ENSURE(!IsScaled(ImplicitType) || !CHAOS_ENSURE(!bIsTransformed)))
+					if (!CHAOS_ENSURE(!IsScaled(ImplicitType) && !bIsTransformed && !IsInstanced(ImplicitType)))
 					{
-						// No support for Scaled or Transformed Capsule yet
+						// No support for Scaled, Instanced, or Transformed Capsule yet
 						break;
 					}
 
@@ -1838,16 +1835,12 @@ bool FBodyInstance::UpdateBodyScale(const FVector& InScale3D, bool bForceUpdate)
 				}
 				case ImplicitObjectType::Convex:
 				{
-
-					if (!CHAOS_ENSURE(IsScaled(ImplicitType) || IsInstanced(ImplicitType)))
+					if(!CHAOS_ENSURE(IsInstanced(ImplicitType) || IsScaled(ImplicitType)))
 					{
-						// Currently assuming all convexes are scaled or instanced (if scale == 1).
+						CHAOS_ENSURE(false); // Expecting instanced or scaled.
 						break;
 					}
 
-					// Get Convex
-					const TImplicitObjectScaled<FConvex>* ConvexGeometry = static_cast<const TImplicitObjectScaled<FConvex>*>(&ImplicitObject);
-
 					FKConvexElem* ConvexElem = ShapeElem->GetShapeCheck<FKConvexElem>();
 					const auto& ConvexImplicit = ConvexElem->GetChaosConvexMesh();
 
@@ -1855,9 +1848,16 @@ bool FBodyInstance::UpdateBodyScale(const FVector& InScale3D, bool bForceUpdate)
 					CHAOS_ENSURE(RelativeTM.GetRotation() == FQuat::Identity);
 					CHAOS_ENSURE(RelativeTM.GetTranslation() == FVector(0, 0, 0));
 
-
-					TUniquePtr<TImplicitObjectScaled<FConvex>> NewConvex = MakeUnique<TImplicitObjectScaled<FConvex>>(MakeSerializable(ConvexImplicit), AdjustedScale3D);
-					NewGeometry.Emplace(MoveTemp(NewConvex));
+					if (AdjustedScale3D == FVector(1.0f, 1.0f, 1.0f))
+					{
+						TUniquePtr<TImplicitObjectInstanced<FConvex>> NewConvex = MakeUnique<TImplicitObjectInstanced<FConvex>>(ConvexImplicit);
+						NewGeometry.Emplace(MoveTemp(NewConvex));
+					}
+					else
+					{
+						TUniquePtr<TImplicitObjectScaled<FConvex>> NewConvex = MakeUnique<TImplicitObjectScaled<FConvex>>(ConvexImplicit, AdjustedScale3D);
+						NewGeometry.Emplace(MoveTemp(NewConvex));
+					}
 
 					bSuccess = true;
 
@@ -1871,15 +1871,37 @@ bool FBodyInstance::UpdateBodyScale(const FVector& InScale3D, bool bForceUpdate)
 						break;
 					}
 
-
 					// PhysX supports translation, we currently do not.
 					CHAOS_ENSURE(RelativeTM.GetTranslation() == FVector(0, 0, 0));
 
-					const TImplicitObjectScaled<FTriangleMeshImplicitObject>* ScaledTriangleMesh = (static_cast<const TImplicitObjectScaled<FTriangleMeshImplicitObject>*>(&ImplicitObject));
-					auto UnscaledTriangleMesh = ScaledTriangleMesh->Object();
+					TSharedPtr<FTriangleMeshImplicitObject, ESPMode::ThreadSafe> InnerTriangleMesh = nullptr;
+					if (IsScaled(ImplicitType))
+					{
+						const TImplicitObjectScaled<FTriangleMeshImplicitObject>* ScaledTriangleMesh = (static_cast<const TImplicitObjectScaled<FTriangleMeshImplicitObject>*>(&ImplicitObject));
+						InnerTriangleMesh = ScaledTriangleMesh->GetSharedObject();
+					}
+					else if (IsInstanced(ImplicitType))
+					{
+						const TImplicitObjectInstanced<FTriangleMeshImplicitObject>* InstancedTriangleMesh = (static_cast<const TImplicitObjectInstanced<FTriangleMeshImplicitObject>*>(&ImplicitObject));
+						InnerTriangleMesh = InstancedTriangleMesh->Object();
+					}
+					else
+					{
+						CHAOS_ENSURE(false);
+						break;
+					}
+
+					if (AdjustedScale3D == FVec3(1.0f, 1.0f, 1.0f))
+					{
+						TUniquePtr<TImplicitObjectInstanced<FTriangleMeshImplicitObject>> NewTriangleMesh = MakeUnique<TImplicitObjectInstanced<FTriangleMeshImplicitObject>>(InnerTriangleMesh);
+						NewGeometry.Emplace(MoveTemp(NewTriangleMesh));
+					}
+					else
+					{
+						TUniquePtr<TImplicitObjectScaled<FTriangleMeshImplicitObject>> NewTriangleMesh = MakeUnique<TImplicitObjectScaled<FTriangleMeshImplicitObject>>(MoveTemp(InnerTriangleMesh), AdjustedScale3D);
+						NewGeometry.Emplace(MoveTemp(NewTriangleMesh));
+					}
 
-					TUniquePtr<TImplicitObjectScaled<FTriangleMeshImplicitObject>> NewTriangleMesh = MakeUnique<TImplicitObjectScaled<FTriangleMeshImplicitObject>>(UnscaledTriangleMesh, AdjustedScale3D);
-					NewGeometry.Emplace(MoveTemp(NewTriangleMesh));
 
 					bSuccess = true;
 
diff --git a/Engine/Source/Runtime/Experimental/Chaos/Public/Chaos/ImplicitObjectScaled.h b/Engine/Source/Runtime/Experimental/Chaos/Public/Chaos/ImplicitObjectScaled.h
--- a/Engine/Source/Runtime/Experimental/Chaos/Public/Chaos/ImplicitObjectScaled.h
+++ b/Engine/Source/Runtime/Experimental/Chaos/Public/Chaos/ImplicitObjectScaled.h
@@ -541,6 +541,9 @@ public:
 	}
 
 	const ObjectType Object() const { return MObject; }
+
+	// Only should be retrieved for copy purposes. Do not modify or access.
+	TSharedPtr<TConcrete, ESPMode::ThreadSafe> GetSharedObject() const { return MSharedPtrForRefCount; }
 	
 	virtual void Serialize(FChaosArchive& Ar) override
 	{
