Fix shadow warning.

diff --git a/Engine/Source/Runtime/Engine/Private/PhysicsEngine/BodyInstance.cpp b/Engine/Source/Runtime/Engine/Private/PhysicsEngine/BodyInstance.cpp
--- a/Engine/Source/Runtime/Engine/Private/PhysicsEngine/BodyInstance.cpp
+++ b/Engine/Source/Runtime/Engine/Private/PhysicsEngine/BodyInstance.cpp
@@ -1282,8 +1282,8 @@ struct FInitBodiesHelper
 						{
 							if (UPrimitiveComponent* PrimComp = BI->OwnerComponent.Get())
 							{
-								FPhysScene_ChaosInterface* PhysScene = PrimComp->GetWorld()->GetPhysicsScene();
-								FPhysScene_Chaos& Scene = PhysScene->GetScene();
+								FPhysScene_ChaosInterface* LocalPhysScene = PrimComp->GetWorld()->GetPhysicsScene();
+								FPhysScene_Chaos& Scene = LocalPhysScene->GetScene();
 								Scene.RegisterForCollisionEvents(PrimComp);
 							}
 						}
