Don't access solver if it doesn't exist.

diff --git a/Engine/Source/Runtime/Engine/Private/PhysicsEngine/Experimental/PhysScene_Chaos.cpp b/Engine/Source/Runtime/Engine/Private/PhysicsEngine/Experimental/PhysScene_Chaos.cpp
--- a/Engine/Source/Runtime/Engine/Private/PhysicsEngine/Experimental/PhysScene_Chaos.cpp
+++ b/Engine/Source/Runtime/Engine/Private/PhysicsEngine/Experimental/PhysScene_Chaos.cpp
@@ -421,8 +421,11 @@ FPhysScene_Chaos::~FPhysScene_Chaos()
 	}
 #endif
 
-	Chaos::FEventManager* EventManager = SceneSolver->GetEventManager();
-	EventManager->UnregisterHandler(Chaos::EEventType::Collision, this);
+	if (SceneSolver)
+	{
+		Chaos::FEventManager* EventManager = SceneSolver->GetEventManager();
+		EventManager->UnregisterHandler(Chaos::EEventType::Collision, this);
+	}
 
 	Shutdown();
 	
