Merging CL#11103414

Add missing include.

diff --git a/Engine/Source/Runtime/Engine/Private/PhysicsEngine/Experimental/PhysScene_Chaos.cpp b/Engine/Source/Runtime/Engine/Private/PhysicsEngine/Experimental/PhysScene_Chaos.cpp
--- a/Engine/Source/Runtime/Engine/Private/PhysicsEngine/Experimental/PhysScene_Chaos.cpp
+++ b/Engine/Source/Runtime/Engine/Private/PhysicsEngine/Experimental/PhysScene_Chaos.cpp
@@ -19,6 +19,7 @@
 #include "Framework/PhysicsTickTask.h"
 #include "Misc/CoreDelegates.h"
 #include "Misc/ScopeLock.h"
+#include "PhysicsEngine/BodyInstance.h"
 #include "PhysicsEngine/PhysicsSettings.h"
 #include "Components/PrimitiveComponent.h"
 
