#UE4 Provide a referencer name for FTextRenderComponentMIDCache to help report object referencers

diff --git a/Engine/Source/Runtime/Engine/Private/Components/TextRenderComponent.cpp b/Engine/Source/Runtime/Engine/Private/Components/TextRenderComponent.cpp
--- a/Engine/Source/Runtime/Engine/Private/Components/TextRenderComponent.cpp
+++ b/Engine/Source/Runtime/Engine/Private/Components/TextRenderComponent.cpp
@@ -465,6 +465,11 @@ public:
 		}
 	}
 
+	virtual FString GetReferencerName() const override
+	{
+		return "FTextRenderComponentMIDCache";
+	}
+
 private:
 	/** Key identifying an array of MIDs */
 	struct FKey
