#UE4 Fix crash when loading maps in the editor due to references preventing world GC

diff --git a/Engine/Source/Runtime/Engine/Private/PhysicsEngine/Experimental/PhysScene_Chaos.cpp b/Engine/Source/Runtime/Engine/Private/PhysicsEngine/Experimental/PhysScene_Chaos.cpp
--- a/Engine/Source/Runtime/Engine/Private/PhysicsEngine/Experimental/PhysScene_Chaos.cpp
+++ b/Engine/Source/Runtime/Engine/Private/PhysicsEngine/Experimental/PhysScene_Chaos.cpp
@@ -397,8 +397,6 @@ FPhysScene_Chaos::FPhysScene_Chaos(AActor* InSolverActor
 	ComponentToPhysicsProxyMap.Reset();
 
 #if WITH_EDITOR
-	FGameDelegates::Get().GetEndPlayMapDelegate().AddRaw(this, &FPhysScene_Chaos::OnWorldEndPlay);
-
 	if(!PhysScene_ChaosPauseHandler)
 	{
 		PhysScene_ChaosPauseHandler = MakeUnique<FPhysScene_ChaosPauseHandler>(ChaosModule);
@@ -432,10 +430,6 @@ FPhysScene_Chaos::~FPhysScene_Chaos()
 	
 	FCoreDelegates::OnPreExit.RemoveAll(this);
 
-#if WITH_EDITOR
-	FGameDelegates::Get().GetEndPlayMapDelegate().RemoveAll(this);
-#endif
-
 #if CHAOS_WITH_PAUSABLE_SOLVER
 	if (SyncCaller)
 	{
@@ -1078,9 +1072,9 @@ void FPhysScene_Chaos::OnUpdateWorldPause()
 }
 #endif  // #if CHAOS_WITH_PAUSABLE_SOLVER
 
-#if WITH_EDITOR
 void FPhysScene_Chaos::OnWorldEndPlay()
 {
+#if WITH_EDITOR
 	// Mark PIE modified objects dirty - couldn't do this during the run because
 	// it's silently ignored
 	for(UObject* Obj : PieModifiedObjects)
@@ -1089,8 +1083,13 @@ void FPhysScene_Chaos::OnWorldEndPlay()
 	}
 
 	PieModifiedObjects.Reset();
+#endif
+
+	PhysicsProxyToComponentMap.Reset();
+	ComponentToPhysicsProxyMap.Reset();
 }
 
+#if WITH_EDITOR
 void FPhysScene_Chaos::AddPieModifiedObject(UObject* InObj)
 {
 	if(GIsPlayInEditorWorld)
@@ -1213,6 +1212,7 @@ void FPhysScene_ChaosInterface::OnWorldEndPlay()
 	}
 #endif
 
+	Scene.OnWorldEndPlay();
 }
 
 void FPhysScene_ChaosInterface::AddActorsToScene_AssumesLocked(TArray<FPhysicsActorHandle>& InHandles, const bool bImmediate)
diff --git a/Engine/Source/Runtime/Engine/Public/Physics/Experimental/PhysScene_Chaos.h b/Engine/Source/Runtime/Engine/Public/Physics/Experimental/PhysScene_Chaos.h
--- a/Engine/Source/Runtime/Engine/Public/Physics/Experimental/PhysScene_Chaos.h
+++ b/Engine/Source/Runtime/Engine/Public/Physics/Experimental/PhysScene_Chaos.h
@@ -158,6 +158,10 @@ public:
 
 	// FGCObject Interface ///////////////////////////////////////////////////
 	virtual void AddReferencedObjects(FReferenceCollector& Collector) override;
+	virtual FString GetReferencerName() const
+	{
+		return "FPhysScene_Chaos";
+	}
 	//////////////////////////////////////////////////////////////////////////
 	
 	/** Given a solver object, returns its associated component. */
@@ -184,6 +188,12 @@ public:
 	 */
 	void CopySolverAccelerationStructure();
 
+	/**
+	 * Callback when a world ends, to mark updated packages dirty. This can't be done in final
+	 * sync as the editor will ignore packages being dirtied in PIE. Also used to clean up any other references
+	 */
+	void OnWorldEndPlay();
+
 private:
 	UPROPERTY()
 	TArray<UPrimitiveComponent*> CollisionEventRegistrations;
@@ -229,13 +239,6 @@ private:
 #endif
 
 #if WITH_EDITOR
-	/**
-	 * Callback when a world ends, to mark updated packages dirty. This can't be done in final
-	 * sync as the editor will ignore packages being dirtied in PIE
-	 */
-	void OnWorldEndPlay();
-
-
 	// List of objects that we modified during a PIE run for physics simulation caching.
 	TArray<UObject*> PieModifiedObjects;
 #endif
