Editor - Adding a 'Studio Analytics' class to the engine.  The purpose is to allow studios to gather analytics about their internal development - for example, time to load the editor, or time to open different types of assets.  Licensees must create their own IAnalyticsProvider and assign it to the studio analytic static, otherwise it does nothing by default.

Added timing information to record when PIE starts.

diff --git a/Engine/Source/Editor/UnrealEd/Private/Analytics/EditorAnalytics.cpp b/Engine/Source/Editor/UnrealEd/Private/Analytics/EditorAnalytics.cpp
--- a/Engine/Source/Editor/UnrealEd/Private/Analytics/EditorAnalytics.cpp
+++ b/Engine/Source/Editor/UnrealEd/Private/Analytics/EditorAnalytics.cpp
@@ -10,56 +10,16 @@
 #include "AnalyticsEventAttribute.h"
 #include "Interfaces/IAnalyticsProvider.h"
 
-
 #define LOCTEXT_NAMESPACE "EditorAnalytics"
 
 
-void FEditorAnalytics::ReportBuildRequirementsFailure(FString EventName, FString PlatformName, bool bHasCode, int32 Requirements)
-{
-	TArray<FAnalyticsEventAttribute> ParamArray;
-	ParamArray.Add(FAnalyticsEventAttribute(TEXT("Time"), 0.0));
-	if (Requirements & ETargetPlatformReadyStatus::SDKNotFound)
-	{
-		ReportEvent(EventName, PlatformName, bHasCode, EAnalyticsErrorCodes::SDKNotFound, ParamArray);
-	}
-	if (Requirements & ETargetPlatformReadyStatus::LicenseNotAccepted)
-	{
-		ReportEvent(EventName, PlatformName, bHasCode, EAnalyticsErrorCodes::LicenseNotAccepted, ParamArray);
-	}
-	if (Requirements & ETargetPlatformReadyStatus::ProvisionNotFound)
-	{
-		ReportEvent(EventName, PlatformName, bHasCode, EAnalyticsErrorCodes::ProvisionNotFound, ParamArray);
-	}
-	if (Requirements & ETargetPlatformReadyStatus::SigningKeyNotFound)
-	{
-		ReportEvent(EventName, PlatformName, bHasCode, EAnalyticsErrorCodes::CertificateNotFound, ParamArray);
-	}
-	if (Requirements & ETargetPlatformReadyStatus::CodeUnsupported)
-	{
-		ReportEvent(EventName, PlatformName, bHasCode, EAnalyticsErrorCodes::CodeUnsupported, ParamArray);
-	}
-	if (Requirements & ETargetPlatformReadyStatus::PluginsUnsupported)
-	{
-		ReportEvent(EventName, PlatformName, bHasCode, EAnalyticsErrorCodes::PluginsUnsupported, ParamArray);
-	}
-}
-
 void FEditorAnalytics::ReportEvent(FString EventName, FString PlatformName, bool bHasCode)
 {
-	if( FEngineAnalytics::IsAvailable() )
-	{
-		const UGeneralProjectSettings& ProjectSettings = *GetDefault<UGeneralProjectSettings>();
-		TArray<FAnalyticsEventAttribute> ParamArray;
-		ParamArray.Add(FAnalyticsEventAttribute(TEXT("ProjectID"), ProjectSettings.ProjectID.ToString()));
-		ParamArray.Add(FAnalyticsEventAttribute(TEXT("Platform"), PlatformName));
-		ParamArray.Add(FAnalyticsEventAttribute(TEXT("ProjectType"), bHasCode ? TEXT("C++ Code") : TEXT("Content Only")));
-		ParamArray.Add(FAnalyticsEventAttribute(TEXT("VanillaEditor"), (GEngine && GEngine->IsVanillaProduct()) ? TEXT("Yes") : TEXT("No")));
-
-		FEngineAnalytics::GetProvider().RecordEvent( EventName, ParamArray );
-	}
+	TArray<FAnalyticsEventAttribute> Empty;
+	ReportEvent(EventName, PlatformName, bHasCode, Empty);
 }
 
-void FEditorAnalytics::ReportEvent(FString EventName, FString PlatformName, bool bHasCode, TArray<FAnalyticsEventAttribute>& ExtraParams)
+void FEditorAnalytics::ReportEvent(FString EventName, FString PlatformName, bool bHasCode, const TArray<FAnalyticsEventAttribute>& ExtraParams)
 {
 	if( FEngineAnalytics::IsAvailable() )
 	{
@@ -199,4 +159,34 @@ bool FEditorAnalytics::ShouldElevateMessageThroughDialog(const int32 ErrorCode)
 	return (EAnalyticsErrorCodes::Type)ErrorCode == EAnalyticsErrorCodes::MissingExecutable;
 }
 
+void FEditorAnalytics::ReportBuildRequirementsFailure(FString EventName, FString PlatformName, bool bHasCode, int32 Requirements)
+{
+	TArray<FAnalyticsEventAttribute> ParamArray;
+	ParamArray.Add(FAnalyticsEventAttribute(TEXT("Time"), 0.0));
+	if (Requirements & ETargetPlatformReadyStatus::SDKNotFound)
+	{
+		ReportEvent(EventName, PlatformName, bHasCode, EAnalyticsErrorCodes::SDKNotFound, ParamArray);
+	}
+	if (Requirements & ETargetPlatformReadyStatus::LicenseNotAccepted)
+	{
+		ReportEvent(EventName, PlatformName, bHasCode, EAnalyticsErrorCodes::LicenseNotAccepted, ParamArray);
+	}
+	if (Requirements & ETargetPlatformReadyStatus::ProvisionNotFound)
+	{
+		ReportEvent(EventName, PlatformName, bHasCode, EAnalyticsErrorCodes::ProvisionNotFound, ParamArray);
+	}
+	if (Requirements & ETargetPlatformReadyStatus::SigningKeyNotFound)
+	{
+		ReportEvent(EventName, PlatformName, bHasCode, EAnalyticsErrorCodes::CertificateNotFound, ParamArray);
+	}
+	if (Requirements & ETargetPlatformReadyStatus::CodeUnsupported)
+	{
+		ReportEvent(EventName, PlatformName, bHasCode, EAnalyticsErrorCodes::CodeUnsupported, ParamArray);
+	}
+	if (Requirements & ETargetPlatformReadyStatus::PluginsUnsupported)
+	{
+		ReportEvent(EventName, PlatformName, bHasCode, EAnalyticsErrorCodes::PluginsUnsupported, ParamArray);
+	}
+}
+
 #undef LOCTEXT_NAMESPACE
diff --git a/Engine/Source/Editor/UnrealEd/Private/EditorEngine.cpp b/Engine/Source/Editor/UnrealEd/Private/EditorEngine.cpp
--- a/Engine/Source/Editor/UnrealEd/Private/EditorEngine.cpp
+++ b/Engine/Source/Editor/UnrealEd/Private/EditorEngine.cpp
@@ -1340,6 +1340,7 @@ void UEditorEngine::Tick( float DeltaSeconds, bool bIdleMode )
 	}
 
 	FEngineAnalytics::Tick(DeltaSeconds);
+	FStudioAnalytics::Tick(DeltaSeconds);
 
 	// Look for realtime flags.
 	bool IsRealtime = false;
diff --git a/Engine/Source/Editor/UnrealEd/Private/FeedbackContextEditor.cpp b/Engine/Source/Editor/UnrealEd/Private/FeedbackContextEditor.cpp
--- a/Engine/Source/Editor/UnrealEd/Private/FeedbackContextEditor.cpp
+++ b/Engine/Source/Editor/UnrealEd/Private/FeedbackContextEditor.cpp
@@ -18,6 +18,8 @@
 #include "Interfaces/IMainFrameModule.h"
 #include "HAL/PlatformApplicationMisc.h"
 #include "Engine/Engine.h"
+#include "StudioAnalytics.h"
+#include "AnalyticsEventAttribute.h"
 
 /** Called to cancel the slow task activity */
 DECLARE_DELEGATE( FOnCancelClickedDelegate );
@@ -394,7 +396,6 @@ void FFeedbackContextEditor::StartSlowTask( const FText& Task, bool bShowCancelB
 		{
 			FSlatePlayInEditorInfo* SlatePlayInEditorSession = GEditor->SlatePlayInEditorMap.Find(PieWorldContext->ContextHandle);
 	
-
 			if (SlatePlayInEditorSession && SlatePlayInEditorSession->SlatePlayInEditorWindow.IsValid())
 			{
 				if (FSlateApplication::Get().GetActiveTopLevelWindow() == SlatePlayInEditorSession->SlatePlayInEditorWindow)
@@ -458,6 +459,8 @@ void FFeedbackContextEditor::StartSlowTask( const FText& Task, bool bShowCancelB
 
 				SlowTaskWindowRef->ShowWindow();
 
+				SlowTaskStartTime = FStudioAnalytics::GetAnalyticSeconds();
+
 				TickSlate(SlowTaskWindow.Pin());
 			}
 
@@ -471,6 +474,11 @@ void FFeedbackContextEditor::FinalizeSlowTask()
 	auto Window = SlowTaskWindow.Pin();
 	if (Window.IsValid())
 	{
+		const double SlowTaskDialogTime = FStudioAnalytics::GetAnalyticSeconds() - SlowTaskStartTime;
+		
+		const FText TaskName = ScopeStack.Num() > 0 ? ScopeStack[0]->DefaultMessage : FText::GetEmpty();
+		FStudioAnalytics::FireEvent_Loading(TEXT("SlowTaskDialog"), SlowTaskDialogTime, { FAnalyticsEventAttribute(TEXT("Task"), TaskName.ToString()) });
+
 		Window->SetContent(SNullWidget::NullWidget);
 		Window->RequestDestroyWindow();
 		SlowTaskWindow.Reset();
diff --git a/Engine/Source/Editor/UnrealEd/Private/FileHelpers.cpp b/Engine/Source/Editor/UnrealEd/Private/FileHelpers.cpp
--- a/Engine/Source/Editor/UnrealEd/Private/FileHelpers.cpp
+++ b/Engine/Source/Editor/UnrealEd/Private/FileHelpers.cpp
@@ -65,6 +65,9 @@
 #include "AutoSaveUtils.h"
 #include "AssetRegistryModule.h"
 #include "Misc/BlacklistNames.h"
+#include "EngineAnalytics.h"
+#include "AnalyticsEventAttribute.h"
+
 
 DEFINE_LOG_CATEGORY_STATIC(LogFileHelpers, Log, All);
 
@@ -2427,7 +2430,7 @@ static void NotifyBSPNeedsRebuild(const FString& PackageName)
  */
 bool FEditorFileUtils::LoadMap(const FString& InFilename, bool LoadAsTemplate, bool bShowProgress)
 {
-	double LoadStartTime = FPlatformTime::Seconds();
+	double LoadStartTime = FStudioAnalytics::GetAnalyticSeconds();
 	
 	if (GEditor->WarnIfLightingBuildIsCurrentlyRunning())
 	{
@@ -2536,7 +2539,10 @@ bool FEditorFileUtils::LoadMap(const FString& InFilename, bool LoadAsTemplate, b
 	}
 
 	// Track time spent loading map.
-	UE_LOG(LogFileHelpers, Log, TEXT("Loading map '%s' took %.3f"), *FPaths::GetBaseFilename(Filename), FPlatformTime::Seconds() - LoadStartTime );
+	const double MapLoadTime = FStudioAnalytics::GetAnalyticSeconds() - LoadStartTime;
+	UE_LOG(LogFileHelpers, Log, TEXT("Loading map '%s' took %.3f"), *FPaths::GetBaseFilename(Filename), MapLoadTime);
+
+	FStudioAnalytics::FireEvent_Loading(TEXT("LoadMap"), MapLoadTime, { FAnalyticsEventAttribute(TEXT("MapName"), FPaths::GetBaseFilename(Filename)) });
 
 	if (GUnrealEd)
 	{
diff --git a/Engine/Source/Editor/UnrealEd/Private/Subsystems/AssetEditorSubsystem.cpp b/Engine/Source/Editor/UnrealEd/Private/Subsystems/AssetEditorSubsystem.cpp
--- a/Engine/Source/Editor/UnrealEd/Private/Subsystems/AssetEditorSubsystem.cpp
+++ b/Engine/Source/Editor/UnrealEd/Private/Subsystems/AssetEditorSubsystem.cpp
@@ -265,6 +265,8 @@ bool UAssetEditorSubsystem::CloseAllAssetEditors()
 
 bool UAssetEditorSubsystem::OpenEditorForAsset(UObject* Asset, const EToolkitMode::Type ToolkitMode, TSharedPtr< IToolkitHost > OpenedFromLevelEditor, const bool bShowProgressWindow)
 {
+	const double OpenAssetStartTime = FStudioAnalytics::GetAnalyticSeconds();
+
 	if (!Asset)
 	{
 		UE_LOG(LogAssetEditorSubsystem, Error, TEXT("Opening Asset editor failed because asset is null"));
@@ -380,6 +382,13 @@ bool UAssetEditorSubsystem::OpenEditorForAsset(UObject* Asset, const EToolkitMod
 			RecentlyOpenedAssets->AddMRUItem(AssetPath);
 		}
 	}
+
+	const double OpenTime = FStudioAnalytics::GetAnalyticSeconds() - OpenAssetStartTime;
+	FStudioAnalytics::FireEvent_Loading(TEXT("OpenAssetEditor"), OpenTime, {
+		FAnalyticsEventAttribute(TEXT("AssetPath"), Asset->GetFullName()),
+		FAnalyticsEventAttribute(TEXT("AssetType"), Asset->GetClass()->GetName())
+	});
+
 	return true;
 }
 
diff --git a/Engine/Source/Editor/UnrealEd/Private/UnrealEdMisc.cpp b/Engine/Source/Editor/UnrealEd/Private/UnrealEdMisc.cpp
--- a/Engine/Source/Editor/UnrealEd/Private/UnrealEdMisc.cpp
+++ b/Engine/Source/Editor/UnrealEd/Private/UnrealEdMisc.cpp
@@ -83,6 +83,7 @@
 #include "ILauncherPlatform.h"
 #include "LauncherPlatformModule.h"
 #include "ILauncherServicesModule.h"
+#include "HAL/PlatformTime.h"
 
 #define USE_UNIT_TESTS 0
 
@@ -310,6 +311,11 @@ void FUnrealEdMisc::OnInit()
 		FPlatformSplash::Show();
 	}
 
+	const double InitialEditorStartupTime = (FStudioAnalytics::GetAnalyticSeconds() - GStartTime);
+	UE_LOG(LogUnrealEdMisc, Log, TEXT("Loading editor; pre map load, took %.3f"), InitialEditorStartupTime);
+
+	FStudioAnalytics::FireEvent_Loading(TEXT("InitializeEditor"), InitialEditorStartupTime);
+
 	// Check for automated build/submit option
 	const bool bDoAutomatedMapBuild = FParse::Param( ParsedCmdLine, TEXT("AutomatedMapBuild") );
 
@@ -517,6 +523,11 @@ void FUnrealEdMisc::OnInit()
 
 	// Handles "Enable World Composition" option in WorldSettings
 	UWorldComposition::EnableWorldCompositionEvent.BindRaw(this, &FUnrealEdMisc::EnableWorldComposition);
+
+	const double TotalEditorStartupTime = (FStudioAnalytics::GetAnalyticSeconds() - GStartTime);
+	UE_LOG(LogUnrealEdMisc, Log, TEXT("Total Editor Startup Time, took %.3f"), TotalEditorStartupTime);
+
+	FStudioAnalytics::FireEvent_Loading(TEXT("TotalEditorStartup"), TotalEditorStartupTime);
 }
 
 void FUnrealEdMisc::InitEngineAnalytics()
diff --git a/Engine/Source/Editor/UnrealEd/Public/EditorAnalytics.h b/Engine/Source/Editor/UnrealEd/Public/EditorAnalytics.h
--- a/Engine/Source/Editor/UnrealEd/Public/EditorAnalytics.h
+++ b/Engine/Source/Editor/UnrealEd/Public/EditorAnalytics.h
@@ -88,7 +88,7 @@ public:
 	* @param PlatformName - name of the platform being used
 	* @param ExtraParams - extra data needed by the event
 	*/
-	UNREALED_API static void ReportEvent(FString EventName, FString PlatformName, bool bHasCode, TArray<FAnalyticsEventAttribute>& ExtraParams);
+	UNREALED_API static void ReportEvent(FString EventName, FString PlatformName, bool bHasCode, const TArray<FAnalyticsEventAttribute>& ExtraParams);
 
 	/**
 	* Reports an event to the analytics system if it is enabled with some extra parameters with an error code
diff --git a/Engine/Source/Editor/UnrealEd/Public/FeedbackContextEditor.h b/Engine/Source/Editor/UnrealEd/Public/FeedbackContextEditor.h
--- a/Engine/Source/Editor/UnrealEd/Public/FeedbackContextEditor.h
+++ b/Engine/Source/Editor/UnrealEd/Public/FeedbackContextEditor.h
@@ -29,6 +29,8 @@ class UNREALED_API FFeedbackContextEditor : public FFeedbackContext
 
 	bool HasTaskBeenCancelled;
 
+	double SlowTaskStartTime = 0;
+
 public:
 
 	FFeedbackContextEditor();
diff --git a/Engine/Source/Runtime/Engine/Classes/Engine/GameInstance.h b/Engine/Source/Runtime/Engine/Classes/Engine/GameInstance.h
--- a/Engine/Source/Runtime/Engine/Classes/Engine/GameInstance.h
+++ b/Engine/Source/Runtime/Engine/Classes/Engine/GameInstance.h
@@ -117,6 +117,9 @@ struct FGameInstancePIEParameters
 	// Is this a dedicated server instance for PIE?
 	bool bRunAsDedicated;
 
+	// What time did we start PIE in the editor?
+	double PIEStartTime = 0;
+
 	// The feature level that PIE world should use
 	ERHIFeatureLevel::Type WorldFeatureLevel;
 
@@ -183,6 +186,9 @@ protected:
 public:
 
 	FString PIEMapName;
+#if WITH_EDITOR
+	double PIEStartTime = 0;
+#endif
 
 	//~ Begin FExec Interface
 	virtual bool Exec(UWorld* InWorld, const TCHAR* Cmd, FOutputDevice& Out = *GLog) override;
diff --git a/Engine/Source/Runtime/Engine/Engine.Build.cs b/Engine/Source/Runtime/Engine/Engine.Build.cs
--- a/Engine/Source/Runtime/Engine/Engine.Build.cs
+++ b/Engine/Source/Runtime/Engine/Engine.Build.cs
@@ -255,6 +255,11 @@ public class Engine : ModuleRules
 				}
 			);
 
+			if (Target.bWithLiveCoding)
+			{
+				PrivateDependencyModuleNames.Add("LiveCoding");
+			}
+
 			PrivateDependencyModuleNames.Add("CollisionAnalyzer");
 			CircularlyReferencedDependentModules.Add("CollisionAnalyzer");
 
@@ -314,11 +319,6 @@ public class Engine : ModuleRules
 			}
 		);
 
-		if(Target.bWithLiveCoding)
-		{
-			DynamicallyLoadedModuleNames.Add("LiveCoding");
-		}
-
 		if (Target.Type != TargetType.Server)
 		{
 			DynamicallyLoadedModuleNames.AddRange(
diff --git a/Engine/Source/Runtime/Engine/Private/Analytics/StudioAnalytics.cpp b/Engine/Source/Runtime/Engine/Private/Analytics/StudioAnalytics.cpp
new file mode 100644
--- /dev/null
+++ b/Engine/Source/Runtime/Engine/Private/Analytics/StudioAnalytics.cpp
@@ -0,0 +1,119 @@
+// Copyright Epic Games, Inc. All Rights Reserved.
+
+#include "StudioAnalytics.h"
+#include "Misc/Guid.h"
+#include "Stats/Stats.h"
+#include "Misc/ConfigCacheIni.h"
+#include "EngineGlobals.h"
+#include "Engine/Engine.h"
+#include "Misc/EngineBuildSettings.h"
+#include "AnalyticsBuildType.h"
+#include "AnalyticsEventAttribute.h"
+#include "IAnalyticsProviderET.h"
+#include "AnalyticsET.h"
+#include "GeneralProjectSettings.h"
+#include "Misc/EngineVersion.h"
+#include "RHI.h"
+#include "GenericPlatform/GenericPlatformCrashContext.h"
+#include "Interfaces/IAnalyticsProvider.h"
+#include "Templates/SharedPointer.h"
+#include "HAL/PlatformProcess.h"
+
+bool FStudioAnalytics::bInitialized = false;
+volatile double FStudioAnalytics::TimeEstimation = 0;
+FThread FStudioAnalytics::TimerThread;
+TSharedPtr<IAnalyticsProvider> FStudioAnalytics::Analytics;
+
+void FStudioAnalytics::SetProvider(TSharedRef<IAnalyticsProvider> InAnalytics)
+{
+	checkf(!Analytics.IsValid(), TEXT("FStudioAnalytics::SetProvider called more than once."));
+
+	bInitialized = true;
+
+	Analytics = InAnalytics;
+
+	TimeEstimation = FPlatformTime::Seconds();
+
+	TimerThread = FThread(TEXT("Studio Analytics Timer Thread"), []() { RunTimer_Concurrent(); });
+}
+
+IAnalyticsProvider& FStudioAnalytics::GetProvider()
+{
+	checkf(IsAvailable(), TEXT("FStudioAnalytics::GetProvider called outside of Initialize/Shutdown."));
+
+	return *Analytics.Get();
+}
+
+void FStudioAnalytics::RunTimer_Concurrent()
+{
+	TimeEstimation = FPlatformTime::Seconds();
+
+	while (bInitialized)
+	{
+		const double FixedInterval = 0.0333333333334;
+
+		const double StartTime = FPlatformTime::Seconds();
+		FPlatformProcess::Sleep((float)FixedInterval);
+		const double EndTime = FPlatformTime::Seconds();
+		const double DeltaTime = EndTime - StartTime;
+
+		if (DeltaTime > 1.0)
+		{
+			TimeEstimation += FixedInterval;
+		}
+		else
+		{
+			TimeEstimation += DeltaTime;
+		}
+	}
+}
+
+void FStudioAnalytics::Tick(float DeltaSeconds)
+{
+
+}
+
+void FStudioAnalytics::Shutdown()
+{
+	ensure(!Analytics.IsValid() || Analytics.IsUnique());
+	Analytics.Reset();
+
+	bInitialized = false;
+
+	if (TimerThread.IsJoinable())
+	{
+		TimerThread.Join();
+	}
+}
+
+double FStudioAnalytics::GetAnalyticSeconds()
+{
+	return bInitialized ? TimeEstimation : FPlatformTime::Seconds();
+}
+
+void FStudioAnalytics::ReportEvent(const FString& EventName, const TArray<FAnalyticsEventAttribute>& Attributes)
+{
+	if (FStudioAnalytics::IsAvailable())
+	{
+		FStudioAnalytics::GetProvider().RecordEvent(EventName, Attributes);
+	}
+}
+
+void FStudioAnalytics::FireEvent_Loading(const FString& LoadingName, double SecondsSpentLoading, const TArray<FAnalyticsEventAttribute>& InAttributes)
+{
+	// Ignore anything less than a 1/4th a second.
+	if (SecondsSpentLoading < 0.250)
+	{
+		return;
+	}
+
+	if (FStudioAnalytics::IsAvailable())
+	{
+		TArray<FAnalyticsEventAttribute> Attributes;
+		Attributes.Emplace(TEXT("LoadingName"), LoadingName);
+		Attributes.Emplace(TEXT("LoadingSeconds"), SecondsSpentLoading);
+		Attributes.Append(InAttributes);
+
+		FStudioAnalytics::GetProvider().RecordEvent(TEXT("Performance.Loading"), Attributes);
+	}
+}
\ No newline at end of file
diff --git a/Engine/Source/Runtime/Engine/Private/EngineAnalytics.cpp b/Engine/Source/Runtime/Engine/Private/EngineAnalytics.cpp
--- a/Engine/Source/Runtime/Engine/Private/EngineAnalytics.cpp
+++ b/Engine/Source/Runtime/Engine/Private/EngineAnalytics.cpp
@@ -15,6 +15,12 @@
 #include "EngineSessionManager.h"
 #include "Misc/EngineVersion.h"
 #include "RHI.h"
+#include "GenericPlatform/GenericPlatformCrashContext.h"
+#include "StudioAnalytics.h"
+
+#if WITH_LIVE_CODING
+#include "ILiveCodingModule.h"
+#endif
 
 #if WITH_EDITOR
 #include "EditorAnalyticsSession.h"
@@ -22,8 +28,6 @@
 #include "Analytics/EditorSessionSummaryWriter.h"
 #endif
 
-#include "GenericPlatform/GenericPlatformCrashContext.h"
-
 bool FEngineAnalytics::bIsInitialized;
 TSharedPtr<IAnalyticsProviderET> FEngineAnalytics::Analytics;
 TSharedPtr<FEngineSessionManager> FEngineAnalytics::SessionManager;
@@ -52,9 +56,43 @@ TFunction<FAnalyticsET::Config()>& GetEngineAnalyticsConfigFunc()
 	return Config;
 }
 
-/**
- * Get analytics pointer
- */
+static TSharedPtr<IAnalyticsProviderET> CreateEpicAnalyticsProvider()
+{
+	// Get the default config.
+	FAnalyticsET::Config Config = GetEngineAnalyticsConfigFunc()();
+	// Set any fields that weren't set by default.
+	if (Config.APIKeyET.IsEmpty())
+	{
+		// We always use the "Release" analytics account unless we're running in analytics test mode (usually with
+		// a command-line parameter), or we're an internal Epic build
+		const EAnalyticsBuildType AnalyticsBuildType = GetAnalyticsBuildType();
+		const bool bUseReleaseAccount =
+			(AnalyticsBuildType == EAnalyticsBuildType::Development || AnalyticsBuildType == EAnalyticsBuildType::Release) &&
+			!FEngineBuildSettings::IsInternalBuild();	// Internal Epic build
+		const TCHAR* BuildTypeStr = bUseReleaseAccount ? TEXT("Release") : TEXT("Dev");
+
+		FString UE4TypeOverride;
+		bool bHasOverride = GConfig->GetString(TEXT("Analytics"), TEXT("UE4TypeOverride"), UE4TypeOverride, GEngineIni);
+		const TCHAR* UE4TypeStr = bHasOverride ? *UE4TypeOverride : FEngineBuildSettings::IsPerforceBuild() ? TEXT("Perforce") : TEXT("UnrealEngine");
+		Config.APIKeyET = FString::Printf(TEXT("UEEditor.%s.%s"), UE4TypeStr, BuildTypeStr);
+	}
+	if (Config.APIServerET.IsEmpty())
+	{
+		Config.APIServerET = TEXT("https://datarouter.ol.epicgames.com/");
+	}
+	if (Config.AppEnvironment.IsEmpty())
+	{
+		Config.AppEnvironment = TEXT("datacollector-source");
+	}
+	if (Config.AppVersionET.IsEmpty())
+	{
+		Config.AppVersionET = FEngineVersion::Current().ToString();
+	}
+
+	// Connect the engine analytics provider (if there is a configuration delegate installed)
+	return FAnalyticsET::Get().CreateAnalyticsProvider(Config);
+}
+
 IAnalyticsProvider& FEngineAnalytics::GetProvider()
 {
 	checkf(bIsInitialized && IsAvailable(), TEXT("FEngineAnalytics::GetProvider called outside of Initialize/Shutdown."));
@@ -85,49 +123,17 @@ void FEngineAnalytics::Initialize()
 
 	if (bShouldInitAnalytics)
 	{
-		// Get the default config.
-		FAnalyticsET::Config Config = GetEngineAnalyticsConfigFunc()();
-		// Set any fields that weren't set by default.
-		if (Config.APIKeyET.IsEmpty())
-		{
-			// We always use the "Release" analytics account unless we're running in analytics test mode (usually with
-			// a command-line parameter), or we're an internal Epic build
-			const EAnalyticsBuildType AnalyticsBuildType = GetAnalyticsBuildType();
-			const bool bUseReleaseAccount =
-				(AnalyticsBuildType == EAnalyticsBuildType::Development || AnalyticsBuildType == EAnalyticsBuildType::Release) &&
-				!FEngineBuildSettings::IsInternalBuild();	// Internal Epic build
-			const TCHAR* BuildTypeStr = bUseReleaseAccount ? TEXT("Release") : TEXT("Dev");
-
-			FString UE4TypeOverride;
-			bool bHasOverride = GConfig->GetString(TEXT("Analytics"), TEXT("UE4TypeOverride"), UE4TypeOverride, GEngineIni);
-			const TCHAR* UE4TypeStr = bHasOverride ? *UE4TypeOverride : FEngineBuildSettings::IsPerforceBuild() ? TEXT("Perforce") : TEXT("UnrealEngine");
-			Config.APIKeyET = FString::Printf(TEXT("UEEditor.%s.%s"), UE4TypeStr, BuildTypeStr);
-		}
-		if (Config.APIServerET.IsEmpty())
-		{
-			Config.APIServerET = TEXT("https://datarouter.ol.epicgames.com/");
-		}
-		if (Config.AppEnvironment.IsEmpty())
-		{
-			Config.AppEnvironment = TEXT("datacollector-source");
-		}
-		if (Config.AppVersionET.IsEmpty())
-		{
-			Config.AppVersionET = FEngineVersion::Current().ToString();
-		}
-
-
-		// Connect the engine analytics provider (if there is a configuration delegate installed)
-		Analytics = FAnalyticsET::Get().CreateAnalyticsProvider(Config);
+		Analytics = CreateEpicAnalyticsProvider();
 
 		if (Analytics.IsValid())
 		{
 			Analytics->SetUserID(FString::Printf(TEXT("%s|%s|%s"), *FPlatformMisc::GetLoginId(), *FPlatformMisc::GetEpicAccountId(), *FPlatformMisc::GetOperatingSystemId()));
 
+			const UGeneralProjectSettings& ProjectSettings = *GetDefault<UGeneralProjectSettings>();
+
 			TArray<FAnalyticsEventAttribute> StartSessionAttributes;
 			GEngine->CreateStartupAnalyticsAttributes( StartSessionAttributes );
 			// Add project info whether we are in editor or game.
-			const UGeneralProjectSettings& ProjectSettings = *GetDefault<UGeneralProjectSettings>();
 			FString OSMajor;
 			FString OSMinor;
 			FPlatformMemoryStats Stats = FPlatformMemory::GetStats();
@@ -234,6 +240,16 @@ void FEngineAnalytics::Tick(float DeltaTime)
 #endif
 }
 
+void FEngineAnalytics::ReportEvent(const FString& EventName, const TArray<FAnalyticsEventAttribute>& Attributes)
+{
+	if (FEngineAnalytics::IsAvailable())
+	{
+		FEngineAnalytics::GetProvider().RecordEvent(EventName, Attributes);
+	}
+
+	FStudioAnalytics::ReportEvent(EventName, Attributes);
+}
+
 void FEngineAnalytics::LowDriveSpaceDetected()
 {
 #if WITH_EDITOR
diff --git a/Engine/Source/Runtime/Engine/Private/GameEngine.cpp b/Engine/Source/Runtime/Engine/Private/GameEngine.cpp
--- a/Engine/Source/Runtime/Engine/Private/GameEngine.cpp
+++ b/Engine/Source/Runtime/Engine/Private/GameEngine.cpp
@@ -1621,6 +1621,11 @@ void UGameEngine::Tick( float DeltaSeconds, bool bIdleMode )
 		FEngineAnalytics::Tick(DeltaSeconds);
 	}
 
+	{
+		SCOPE_TIME_GUARD(TEXT("UGameEngine::Tick - Studio Analytics"));
+		FStudioAnalytics::Tick(DeltaSeconds);
+	}
+
 	// -----------------------------------------------------
 	// Begin ticking worlds
 	// -----------------------------------------------------
diff --git a/Engine/Source/Runtime/Engine/Private/GameInstance.cpp b/Engine/Source/Runtime/Engine/Private/GameInstance.cpp
--- a/Engine/Source/Runtime/Engine/Private/GameInstance.cpp
+++ b/Engine/Source/Runtime/Engine/Private/GameInstance.cpp
@@ -234,6 +234,8 @@ static ENetMode GetNetModeFromPlayNetMode(const EPlayNetMode InPlayNetMode, cons
 
 FGameInstancePIEResult UGameInstance::InitializeForPlayInEditor(int32 PIEInstanceIndex, const FGameInstancePIEParameters& Params)
 {
+	PIEStartTime = Params.PIEStartTime;
+
 	UEditorEngine* const EditorEngine = CastChecked<UEditorEngine>(GetEngine());
 
 	// Look for an existing pie world context, may have been created before
diff --git a/Engine/Source/Runtime/Engine/Private/UnrealEngine.cpp b/Engine/Source/Runtime/Engine/Private/UnrealEngine.cpp
--- a/Engine/Source/Runtime/Engine/Private/UnrealEngine.cpp
+++ b/Engine/Source/Runtime/Engine/Private/UnrealEngine.cpp
@@ -1774,6 +1774,7 @@ void UEngine::PreExit()
 	ShutdownRenderingCVarsCaching();
 	const bool bIsEngineShutdown = true;
 	FEngineAnalytics::Shutdown(bIsEngineShutdown);
+	FStudioAnalytics::Shutdown();
 	if (ScreenSaverInhibitor)
 	{
 		// Resume the thread to avoid a deadlock while waiting for finish.
diff --git a/Engine/Source/Runtime/Engine/Public/EngineAnalytics.h b/Engine/Source/Runtime/Engine/Public/EngineAnalytics.h
--- a/Engine/Source/Runtime/Engine/Public/EngineAnalytics.h
+++ b/Engine/Source/Runtime/Engine/Public/EngineAnalytics.h
@@ -3,6 +3,7 @@
 #pragma once
 
 #include "CoreMinimal.h"
+#include "StudioAnalytics.h"
 
 class FEngineSessionManager;
 class IAnalyticsProvider;
@@ -26,13 +27,19 @@ public:
 	 * Note: must check IsAvailable() first else this code will assert if the provider is not valid.
 	 */
 	static ENGINE_API IAnalyticsProvider& GetProvider();
+
 	/** Helper function to determine if the provider is valid. */
 	static ENGINE_API bool IsAvailable() { return Analytics.IsValid(); }
+
 	/** Called to initialize the singleton. */
 	static ENGINE_API void Initialize();
+
 	/** Called to shut down the singleton */
 	static ENGINE_API void Shutdown(bool bIsEngineShutdown = false);
 
+	/** General report event function. */
+	static ENGINE_API void ReportEvent(const FString& EventName, const TArray<FAnalyticsEventAttribute>& Attributes);
+
 	static ENGINE_API void Tick(float DeltaTime);
 
 	static ENGINE_API void LowDriveSpaceDetected();
diff --git a/Engine/Source/Runtime/Engine/Public/StudioAnalytics.h b/Engine/Source/Runtime/Engine/Public/StudioAnalytics.h
new file mode 100644
--- /dev/null
+++ b/Engine/Source/Runtime/Engine/Public/StudioAnalytics.h
@@ -0,0 +1,48 @@
+// Copyright Epic Games, Inc. All Rights Reserved.
+
+#pragma once
+
+#include "CoreMinimal.h"
+#include "HAL/Thread.h"
+
+class IAnalyticsProvider;
+
+/**
+ * The public interface for the game studio to gather information about internal development metrics.
+ */
+class FStudioAnalytics : FNoncopyable
+{
+public:
+	static ENGINE_API void SetProvider(TSharedRef<IAnalyticsProvider> InAnalytics);
+
+	/**
+	 * Return the provider instance. Not valid outside of Initialize/Shutdown calls.
+	 * Note: must check IsAvailable() first else this code will assert if the provider is not valid.
+	 */
+	static ENGINE_API IAnalyticsProvider& GetProvider();
+
+	/** Helper function to determine if the provider is valid. */
+	static ENGINE_API bool IsAvailable() { return Analytics.IsValid(); }
+
+	static ENGINE_API double GetAnalyticSeconds();
+
+	static ENGINE_API void Tick(float DeltaSeconds);
+
+	static ENGINE_API void Shutdown();
+
+	/** General report event function. */
+	static ENGINE_API void ReportEvent(const FString& EventName, const TArray<FAnalyticsEventAttribute>& Attributes);
+
+	/** An event for reporting load time that blocks the editor. */
+	static ENGINE_API void FireEvent_Loading(const FString& LoadingName, double SecondsSpentLoading, const TArray<FAnalyticsEventAttribute>& Attributes = TArray<FAnalyticsEventAttribute>());
+
+private:
+	static void RunTimer_Concurrent();
+
+private:
+	static bool bInitialized;
+	static TSharedPtr<IAnalyticsProvider> Analytics;
+	static FThread TimerThread;
+	static volatile double TimeEstimation;
+};
+
