Checking in a win fix.

diff --git a/Engine/Source/Runtime/Engine/Private/EngineAnalytics.cpp b/Engine/Source/Runtime/Engine/Private/EngineAnalytics.cpp
--- a/Engine/Source/Runtime/Engine/Private/EngineAnalytics.cpp
+++ b/Engine/Source/Runtime/Engine/Private/EngineAnalytics.cpp
@@ -18,10 +18,6 @@
 #include "GenericPlatform/GenericPlatformCrashContext.h"
 #include "StudioAnalytics.h"
 
-#if WITH_LIVE_CODING
-#include "ILiveCodingModule.h"
-#endif
-
 #if WITH_EDITOR
 #include "EditorAnalyticsSession.h"
 #include "EditorSessionSummarySender.h"
