Fixing monolithic build includes.

diff --git a/Engine/Source/Runtime/Engine/Public/EngineAnalytics.h b/Engine/Source/Runtime/Engine/Public/EngineAnalytics.h
--- a/Engine/Source/Runtime/Engine/Public/EngineAnalytics.h
+++ b/Engine/Source/Runtime/Engine/Public/EngineAnalytics.h
@@ -4,6 +4,7 @@
 
 #include "CoreMinimal.h"
 #include "StudioAnalytics.h"
+#include "AnalyticsEventAttribute.h"
 
 class FEngineSessionManager;
 class IAnalyticsProvider;
diff --git a/Engine/Source/Runtime/Engine/Public/StudioAnalytics.h b/Engine/Source/Runtime/Engine/Public/StudioAnalytics.h
--- a/Engine/Source/Runtime/Engine/Public/StudioAnalytics.h
+++ b/Engine/Source/Runtime/Engine/Public/StudioAnalytics.h
@@ -4,6 +4,7 @@
 
 #include "CoreMinimal.h"
 #include "HAL/Thread.h"
+#include "AnalyticsEventAttribute.h"
 
 class IAnalyticsProvider;
 
