Fixing some IWYU issues.

diff --git a/Engine/Source/Editor/UnrealEd/Private/EditorEngine.cpp b/Engine/Source/Editor/UnrealEd/Private/EditorEngine.cpp
--- a/Engine/Source/Editor/UnrealEd/Private/EditorEngine.cpp
+++ b/Engine/Source/Editor/UnrealEd/Private/EditorEngine.cpp
@@ -217,6 +217,7 @@
 #include "ToolMenus.h"
 #include "IToolMenusEditorModule.h"
 #include "Subsystems/AssetEditorSubsystem.h"
+#include "StudioAnalytics.h"
 
 
 DEFINE_LOG_CATEGORY_STATIC(LogEditor, Log, All);
diff --git a/Engine/Source/Editor/UnrealEd/Private/Subsystems/AssetEditorSubsystem.cpp b/Engine/Source/Editor/UnrealEd/Private/Subsystems/AssetEditorSubsystem.cpp
--- a/Engine/Source/Editor/UnrealEd/Private/Subsystems/AssetEditorSubsystem.cpp
+++ b/Engine/Source/Editor/UnrealEd/Private/Subsystems/AssetEditorSubsystem.cpp
@@ -23,6 +23,7 @@
 #include "Misc/FeedbackContext.h"
 #include "Misc/ConfigCacheIni.h"
 #include "Misc/BlacklistNames.h"
+#include "StudioAnalytics.h"
 
 
 #define LOCTEXT_NAMESPACE "AssetEditorSubsystem"
diff --git a/Engine/Source/Editor/UnrealEd/Private/UnrealEdMisc.cpp b/Engine/Source/Editor/UnrealEd/Private/UnrealEdMisc.cpp
--- a/Engine/Source/Editor/UnrealEd/Private/UnrealEdMisc.cpp
+++ b/Engine/Source/Editor/UnrealEd/Private/UnrealEdMisc.cpp
@@ -84,6 +84,7 @@
 #include "LauncherPlatformModule.h"
 #include "ILauncherServicesModule.h"
 #include "HAL/PlatformTime.h"
+#include "StudioAnalytics.h"
 
 #define USE_UNIT_TESTS 0
 
diff --git a/Engine/Source/Runtime/Engine/Private/Analytics/StudioAnalytics.cpp b/Engine/Source/Runtime/Engine/Private/Analytics/StudioAnalytics.cpp
--- a/Engine/Source/Runtime/Engine/Private/Analytics/StudioAnalytics.cpp
+++ b/Engine/Source/Runtime/Engine/Private/Analytics/StudioAnalytics.cpp
@@ -91,6 +91,11 @@ double FStudioAnalytics::GetAnalyticSeconds()
 	return bInitialized ? TimeEstimation : FPlatformTime::Seconds();
 }
 
+void FStudioAnalytics::ReportEvent(const FString& EventName)
+{
+	ReportEvent(EventName, TArray<FAnalyticsEventAttribute>());
+}
+
 void FStudioAnalytics::ReportEvent(const FString& EventName, const TArray<FAnalyticsEventAttribute>& Attributes)
 {
 	if (FStudioAnalytics::IsAvailable())
diff --git a/Engine/Source/Runtime/Engine/Private/UnrealEngine.cpp b/Engine/Source/Runtime/Engine/Private/UnrealEngine.cpp
--- a/Engine/Source/Runtime/Engine/Private/UnrealEngine.cpp
+++ b/Engine/Source/Runtime/Engine/Private/UnrealEngine.cpp
@@ -230,6 +230,7 @@ UnrealEngine.cpp: Implements the UEngine class and helpers.
 #include "Components/SkinnedMeshComponent.h"
 #include "ObjectTrace.h"
 #include "Animation/AnimTrace.h"
+#include "StudioAnalytics.h"
 
 DEFINE_LOG_CATEGORY(LogEngine);
 IMPLEMENT_MODULE( FEngineModule, Engine );
diff --git a/Engine/Source/Runtime/Engine/Public/EngineAnalytics.h b/Engine/Source/Runtime/Engine/Public/EngineAnalytics.h
--- a/Engine/Source/Runtime/Engine/Public/EngineAnalytics.h
+++ b/Engine/Source/Runtime/Engine/Public/EngineAnalytics.h
@@ -3,12 +3,11 @@
 #pragma once
 
 #include "CoreMinimal.h"
-#include "StudioAnalytics.h"
-#include "AnalyticsEventAttribute.h"
 
 class FEngineSessionManager;
 class IAnalyticsProvider;
 class IAnalyticsProviderET;
+struct FAnalyticsEventAttribute;
 
 /**
  * The public interface for the editor's analytics provider singleton.
diff --git a/Engine/Source/Runtime/Engine/Public/StudioAnalytics.h b/Engine/Source/Runtime/Engine/Public/StudioAnalytics.h
--- a/Engine/Source/Runtime/Engine/Public/StudioAnalytics.h
+++ b/Engine/Source/Runtime/Engine/Public/StudioAnalytics.h
@@ -4,9 +4,9 @@
 
 #include "CoreMinimal.h"
 #include "HAL/Thread.h"
-#include "AnalyticsEventAttribute.h"
 
 class IAnalyticsProvider;
+struct FAnalyticsEventAttribute;
 
 /**
  * The public interface for the game studio to gather information about internal development metrics.
@@ -32,6 +32,7 @@ public:
 	static ENGINE_API void Shutdown();
 
 	/** General report event function. */
+	static ENGINE_API void ReportEvent(const FString& EventName);
 	static ENGINE_API void ReportEvent(const FString& EventName, const TArray<FAnalyticsEventAttribute>& Attributes);
 
 	/** An event for reporting load time that blocks the editor. */
