Change FConstraintInstance Pos1,Pos2 UProperty to EditAnywhere (from EditInstanceOnly)

diff --git a/Engine/Source/Runtime/Engine/Classes/PhysicsEngine/ConstraintInstance.h b/Engine/Source/Runtime/Engine/Classes/PhysicsEngine/ConstraintInstance.h
--- a/Engine/Source/Runtime/Engine/Classes/PhysicsEngine/ConstraintInstance.h
+++ b/Engine/Source/Runtime/Engine/Classes/PhysicsEngine/ConstraintInstance.h
@@ -134,8 +134,8 @@ struct ENGINE_API FConstraintInstance
 
 	///////////////////////////// Body1 ref frame
 	
-	/** Location of constraint in Body1 reference frame. */
-	UPROPERTY(EditInstanceOnly, Category=Constraint)
+	/** Location of constraint in Body1 reference frame (usually the "child" body for skeletal meshes). */
+	UPROPERTY(EditAnywhere, Category=Constraint)
 	FVector Pos1;
 
 	/** Primary (twist) axis in Body1 reference frame. */
@@ -148,8 +148,8 @@ struct ENGINE_API FConstraintInstance
 
 	///////////////////////////// Body2 ref frame
 	
-	/** Location of constraint in Body2 reference frame. */
-	UPROPERTY(EditInstanceOnly, Category=Constraint)
+	/** Location of constraint in Body2 reference frame (usually the "parent" body for skeletal meshes). */
+	UPROPERTY(EditAnywhere, Category=Constraint)
 	FVector Pos2;
 
 	/** Primary (twist) axis in Body2 reference frame. */
