Fix non-WITH-CHAOS build

diff --git a/Engine/Source/Runtime/Engine/Private/PhysicsEngine/Experimental/PhysScene_Chaos.cpp b/Engine/Source/Runtime/Engine/Private/PhysicsEngine/Experimental/PhysScene_Chaos.cpp
--- a/Engine/Source/Runtime/Engine/Private/PhysicsEngine/Experimental/PhysScene_Chaos.cpp
+++ b/Engine/Source/Runtime/Engine/Private/PhysicsEngine/Experimental/PhysScene_Chaos.cpp
@@ -62,7 +62,9 @@ TAutoConsoleVariable<int32> CVar_ChaosSimulationEnable(TEXT("P.Chaos.Simulation.
 
 DEFINE_LOG_CATEGORY_STATIC(LogFPhysScene_ChaosSolver, Log, All);
 
+#if WITH_CHAOS
 Chaos::TCollisionModifierCallback<float, 3> FPhysScene_ChaosInterface::CollisionModifierCallback;
+#endif // WITH_CHAOS
 
 void DumpHierarchyStats(const TArray<FString>& Args)
 {
