Fixing some non-unity build errors.

diff --git a/Engine/Source/Editor/UnrealEd/Private/FileHelpers.cpp b/Engine/Source/Editor/UnrealEd/Private/FileHelpers.cpp
--- a/Engine/Source/Editor/UnrealEd/Private/FileHelpers.cpp
+++ b/Engine/Source/Editor/UnrealEd/Private/FileHelpers.cpp
@@ -66,6 +66,7 @@
 #include "AssetRegistryModule.h"
 #include "Misc/BlacklistNames.h"
 #include "EngineAnalytics.h"
+#include "StudioAnalytics.h"
 #include "AnalyticsEventAttribute.h"
 
 
diff --git a/Engine/Source/Runtime/Engine/Private/GameEngine.cpp b/Engine/Source/Runtime/Engine/Private/GameEngine.cpp
--- a/Engine/Source/Runtime/Engine/Private/GameEngine.cpp
+++ b/Engine/Source/Runtime/Engine/Private/GameEngine.cpp
@@ -52,6 +52,7 @@
 #include "Misc/EmbeddedCommunication.h"
 #include "Engine/CoreSettings.h"
 #include "EngineAnalytics.h"
+#include "StudioAnalytics.h"
 #include "Engine/DemoNetDriver.h"
 
 #include "Tickable.h"
