fix for chunks loaded via request not showing up in the cache miss log.

diff --git a/Engine/Source/Runtime/Engine/Private/AudioStreamingCache.cpp b/Engine/Source/Runtime/Engine/Private/AudioStreamingCache.cpp
--- a/Engine/Source/Runtime/Engine/Private/AudioStreamingCache.cpp
+++ b/Engine/Source/Runtime/Engine/Private/AudioStreamingCache.cpp
@@ -536,6 +536,16 @@ bool FAudioChunkCache::AddOrTouchChunk(const FChunkKey& InKey, TFunction<void(EA
 
 		KickOffAsyncLoad(CacheElement, InKey, OnLoadCompleted, CallbackThread, bNeededForPlayback);
 
+		if (bNeededForPlayback && (bLogCacheMisses || AlwaysLogCacheMissesCVar))
+		{
+			// We missed 
+			const uint32 TotalNumChunksInWave = InKey.SoundWave->GetNumChunks();
+
+			FCacheMissInfo CacheMissInfo = { InKey.SoundWaveName, InKey.ChunkIndex, TotalNumChunksInWave, false };
+			CacheMissQueue.Enqueue(MoveTemp(CacheMissInfo));
+		}
+
+
 		if (TrimCacheWhenOverBudgetCVar != 0 && MemoryCounterBytes > MemoryLimitBytes)
 		{
 			TrimMemory(MemoryCounterBytes - MemoryLimitBytes);
