# Fix for montages not previewing properly in editor mode when using animation asset for skeletal mesh.
#rnx

diff --git a/Engine/Source/Runtime/Engine/Private/Animation/AnimSingleNodeInstance.cpp b/Engine/Source/Runtime/Engine/Private/Animation/AnimSingleNodeInstance.cpp
--- a/Engine/Source/Runtime/Engine/Private/Animation/AnimSingleNodeInstance.cpp
+++ b/Engine/Source/Runtime/Engine/Private/Animation/AnimSingleNodeInstance.cpp
@@ -192,6 +192,7 @@ void UAnimSingleNodeInstance::RestartMontage(UAnimMontage* Montage, FName FromSe
 			Montage_JumpToSection(FromSection);
 		}
 		SetMontageLoop(Montage, Proxy.IsLooping(), FromSection);
+		UpdateMontageWeightForTimeSkip(Montage->BlendIn.GetBlendTime());
 	}
 }
 
