Moved instanced render data destruction into DestroyRenderThreadResources. The destructor is now deferred past the lifetime of the parent UComponent and the UStaticMesh pointer can be GC'd.

diff --git a/Engine/Source/Runtime/Engine/Private/InstancedStaticMesh.cpp b/Engine/Source/Runtime/Engine/Private/InstancedStaticMesh.cpp
--- a/Engine/Source/Runtime/Engine/Private/InstancedStaticMesh.cpp
+++ b/Engine/Source/Runtime/Engine/Private/InstancedStaticMesh.cpp
@@ -890,6 +890,12 @@ void FInstancedStaticMeshSceneProxy::SetupProxy(UInstancedStaticMeshComponent* I
 	UserData_DeselectedInstances.bRenderSelected = false;
 }
 
+void FInstancedStaticMeshSceneProxy::DestroyRenderThreadResources()
+{
+	InstancedRenderData.ReleaseResources(&GetScene(), StaticMesh);
+	FStaticMeshSceneProxy::DestroyRenderThreadResources();
+}
+
 void FInstancedStaticMeshSceneProxy::SetupInstancedMeshBatch(int32 LODIndex, int32 BatchIndex, FMeshBatch& OutMeshBatch) const
 {
 	OutMeshBatch.VertexFactory = &InstancedRenderData.VertexFactories[LODIndex];
diff --git a/Engine/Source/Runtime/Engine/Private/InstancedStaticMesh.h b/Engine/Source/Runtime/Engine/Private/InstancedStaticMesh.h
--- a/Engine/Source/Runtime/Engine/Private/InstancedStaticMesh.h
+++ b/Engine/Source/Runtime/Engine/Private/InstancedStaticMesh.h
@@ -496,8 +496,6 @@ public:
 
 	~FInstancedStaticMeshSceneProxy()
 	{
-		InstancedRenderData.ReleaseResources(&GetScene( ), StaticMesh);
-
 #if RHI_RAYTRACING
 		for (int32 i = 0; i < RayTracingCullClusterInstances.Num(); ++i)
 		{
@@ -507,7 +505,9 @@ public:
 	}
 
 	// FPrimitiveSceneProxy interface.
-	
+
+	virtual void DestroyRenderThreadResources() override;
+
 	virtual FPrimitiveViewRelevance GetViewRelevance(const FSceneView* View) const override
 	{
 		FPrimitiveViewRelevance Result;
diff --git a/Engine/Source/Runtime/Engine/Private/StaticMeshRender.cpp b/Engine/Source/Runtime/Engine/Private/StaticMeshRender.cpp
--- a/Engine/Source/Runtime/Engine/Private/StaticMeshRender.cpp
+++ b/Engine/Source/Runtime/Engine/Private/StaticMeshRender.cpp
@@ -578,6 +578,7 @@ void FStaticMeshSceneProxy::DestroyRenderThreadResources()
 
 	// Call here because it uses RenderData from the StaticMesh which is not guaranteed to still be valid after this DestroyRenderThreadResources call
 	RemoveSpeedTreeWind();
+	StaticMesh = nullptr;
 }
 
 /** Sets up a wireframe FMeshBatch for a specific LOD. */
