Deprecating stereo bleed members.

diff --git a/Engine/Source/Runtime/Engine/Classes/Sound/SoundClass.h b/Engine/Source/Runtime/Engine/Classes/Sound/SoundClass.h
--- a/Engine/Source/Runtime/Engine/Classes/Sound/SoundClass.h
+++ b/Engine/Source/Runtime/Engine/Classes/Sound/SoundClass.h
@@ -46,6 +46,7 @@ struct FSoundClassEditorData
 /**
  * Structure containing configurable properties of a sound class.
  */
+PRAGMA_DISABLE_DEPRECATION_WARNINGS
 USTRUCT(BlueprintType)
 struct FSoundClassProperties
 {
@@ -69,6 +70,11 @@ struct FSoundClassProperties
 	UPROPERTY(EditAnywhere, BlueprintReadOnly, Category = General)
 	float AttenuationDistanceScale;
 
+	/** The amount of stereo sounds to bleed to the rear speakers */
+	UE_DEPRECATED(4.25, "Stereo Bleed is no longer supported.")
+	UPROPERTY(EditAnywhere, BlueprintReadOnly, Category = Routing, meta = (DeprecatedProperty, DeprecationMessage = "Stereo Bleed no longer supported."))
+	float StereoBleed = 0.f;
+
 	/** The amount of a sound to bleed to the LFE channel */
 	UPROPERTY(EditAnywhere, BlueprintReadOnly, Category = Routing, meta = (DisplayName = "LFE Bleed"))
 	float LFEBleed;
@@ -148,6 +154,7 @@ private:
 	FDynamicParameter AttenuationScaleParam;
 	float ParentAttenuationScale;	
 };
+PRAGMA_ENABLE_DEPRECATION_WARNINGS
 
 /**
  * Structure containing information on a SoundMix to activate passively.
diff --git a/Engine/Source/Runtime/Engine/Private/Audio.cpp b/Engine/Source/Runtime/Engine/Private/Audio.cpp
--- a/Engine/Source/Runtime/Engine/Private/Audio.cpp
+++ b/Engine/Source/Runtime/Engine/Private/Audio.cpp
@@ -348,6 +348,11 @@ bool FSoundSource::SetReverbApplied(bool bHardwareAvailable)
 	return(bReverbApplied);
 }
 
+float FSoundSource::SetStereoBleed()
+{
+	return 0.f;
+}
+
 float FSoundSource::SetLFEBleed()
 {
 	LFEBleed = WaveInstance->LFEBleed;
diff --git a/Engine/Source/Runtime/Engine/Public/Audio.h b/Engine/Source/Runtime/Engine/Public/Audio.h
--- a/Engine/Source/Runtime/Engine/Public/Audio.h
+++ b/Engine/Source/Runtime/Engine/Public/Audio.h
@@ -224,6 +224,10 @@ public:
 	/** The volume at which the radio filter kicks in */
 	float RadioFilterVolumeThreshold;
 
+	/** The amount of stereo sounds to bleed to the rear speakers */
+	UE_DEPRECATED(4.25, "Stereo Bleed is no longer supported.")
+	float StereoBleed;
+
 	/** The amount of a sound to bleed to the LFE channel */
 	float LFEBleed;
 
@@ -628,6 +632,10 @@ public:
 	/** Set the bReverbApplied variable. */
 	ENGINE_API bool SetReverbApplied(bool bHardwareAvailable);
 
+	/** Set the StereoBleed variable. */
+	UE_DEPRECATED(4.25, "Stereo Bleed is no longer supported.")
+	ENGINE_API float SetStereoBleed();
+
 	/** Updates and sets the LFEBleed variable. */
 	ENGINE_API float SetLFEBleed();
 
