Chaos - Joint overhead reduction

diff --git a/Engine/Source/Runtime/Engine/Private/PhysicsEngine/ImmediatePhysics/ImmediatePhysicsChaos/ImmediatePhysicsSimulation_Chaos.cpp b/Engine/Source/Runtime/Engine/Private/PhysicsEngine/ImmediatePhysics/ImmediatePhysicsChaos/ImmediatePhysicsSimulation_Chaos.cpp
--- a/Engine/Source/Runtime/Engine/Private/PhysicsEngine/ImmediatePhysics/ImmediatePhysicsChaos/ImmediatePhysicsSimulation_Chaos.cpp
+++ b/Engine/Source/Runtime/Engine/Private/PhysicsEngine/ImmediatePhysics/ImmediatePhysicsChaos/ImmediatePhysicsSimulation_Chaos.cpp
@@ -160,14 +160,14 @@ namespace ImmediatePhysics_Chaos
 	//
 
 	FSimulation::FSimulation()
-		: PotentiallyCollidingPairs()
+		: ActivePotentiallyCollidingPairs()
 		, CollidedParticles()
 		, ParticleMaterials()
 		, PerParticleMaterials()
 		, Particles()
 		, Joints()
 		, Collisions(Particles, CollidedParticles, ParticleMaterials, 0, 0, ChaosImmediate_Collision_CullDistance, ChaosImmediate_Collision_ShapePadding)
-		, BroadPhase(PotentiallyCollidingPairs, ChaosImmediate_Collision_CullDistance)
+		, BroadPhase(ActivePotentiallyCollidingPairs, ChaosImmediate_Collision_CullDistance)
 		, CollisionDetector(BroadPhase, Collisions)
 		, JointsRule(0, Joints)
 		, CollisionsRule(1, Collisions)
@@ -177,6 +177,8 @@ namespace ImmediatePhysics_Chaos
 		, RollingAverageStepTime(ChaosImmediate_Evolution_InitialStepTime)
 		, NumRollingAverageStepTimes(1)
 		, MaxNumRollingAverageStepTimes(ChaosImmediate_Evolution_DeltaTimeCount)
+		, bActorsDirty(false)
+		, bJointsDirty(false)
 	{
 		using namespace Chaos;
 
@@ -286,12 +288,14 @@ namespace ImmediatePhysics_Chaos
 
 	FActorHandle* FSimulation::CreateActor(EActorType ActorType, FBodyInstance* BodyInstance, const FTransform& Transform)
 	{
+		// @todo(ccaulfield): Shared materials
+		// @todo(ccaulfield): Add colliding particle pairs
+
 		using namespace Chaos;
 
 		FActorHandle* ActorHandle = new FActorHandle(Particles, ActorType, BodyInstance, Transform);
 		int ActorIndex = ActorHandles.Add(ActorHandle);
 
-		// @todo(ccaulfield): shared materials
 		TUniquePtr<FChaosPhysicsMaterial> Material = MakeUnique<FChaosPhysicsMaterial>();
 		if (BodyInstance != nullptr)
 		{
@@ -307,24 +311,33 @@ namespace ImmediatePhysics_Chaos
 		PerParticleMaterials.Add(MoveTemp(Material));
 		CollidedParticles.Add(false);
 
+		bActorsDirty = true;
+
 		return ActorHandle;
 	}
 
 	void FSimulation::DestroyActor(FActorHandle* ActorHandle)
 	{
 		// @todo(ccaulfield): FActorHandle could remember its index to optimize this
+		// @todo(ccaulfield): Remove colliding particle pairs
+
 		int32 Index = ActorHandles.Remove(ActorHandle);
 		delete ActorHandle;
 
 		ParticleMaterials.RemoveAt(Index, 1);
 		PerParticleMaterials.RemoveAt(Index, 1);
 		CollidedParticles.RemoveAt(Index, 1);
+
+		bActorsDirty = true;
 	}
 
 	FJointHandle* FSimulation::CreateJoint(FConstraintInstance* ConstraintInstance, FActorHandle* Body1, FActorHandle* Body2)
 	{
 		FJointHandle* JointHandle = new FJointHandle(&Joints, ConstraintInstance, Body1, Body2);
 		JointHandles.Add(JointHandle);
+
+		bJointsDirty = true;
+
 		return JointHandle;
 	}
 
@@ -333,6 +346,8 @@ namespace ImmediatePhysics_Chaos
 		// @todo(ccaulfield): FJointHandle could remember its index to optimize this
 		JointHandles.Remove(JointHandle);
 		delete JointHandle;
+
+		bJointsDirty = true;
 	}
 
 	void FSimulation::SetNumActiveBodies(int32 InNumActiveActorHandles)
@@ -363,6 +378,8 @@ namespace ImmediatePhysics_Chaos
 				Handle->SetEnabled(false);
 			}
 		}
+
+		bActorsDirty = true;
 	}
 
 	void FSimulation::SetIgnoreCollisionPairTable(const TArray<FIgnorePair>& InIgnoreCollisionPairTable)
@@ -408,6 +425,7 @@ namespace ImmediatePhysics_Chaos
 			}
 		}
 
+		bActorsDirty = true;
 	}
 
 	void FSimulation::SetIgnoreCollisionActors(const TArray<FActorHandle*>& InIgnoreCollisionActors)
@@ -422,6 +440,24 @@ namespace ImmediatePhysics_Chaos
 				Particle->SetCollisionGroup(INDEX_NONE);
 			}
 		}
+
+		bActorsDirty = true;
+	}
+
+	void FSimulation::UpdateActivePotentiallyCollidingPairs()
+	{
+		using namespace Chaos;
+
+		ActivePotentiallyCollidingPairs.Reset();
+		for (const FParticlePair& ParticlePair : PotentiallyCollidingPairs)
+		{
+			bool bAnyDisabled = TGenericParticleHandle<FReal, 3>(ParticlePair[0])->Disabled() || TGenericParticleHandle<FReal, 3>(ParticlePair[1])->Disabled();
+			bool bAnyDynamic = TGenericParticleHandle<FReal, 3>(ParticlePair[0])->IsDynamic() || TGenericParticleHandle<FReal, 3>(ParticlePair[1])->IsDynamic();
+			if (bAnyDynamic && !bAnyDisabled)
+			{
+				ActivePotentiallyCollidingPairs.Add(ParticlePair);
+			}
+		}
 	}
 
 	void FSimulation::ConditionConstraints()
@@ -561,7 +597,17 @@ namespace ImmediatePhysics_Chaos
 		DebugDrawDynamicParticles(2, 2, FColor(192, 192, 0));
 		DebugDrawConstraints(2, 2, 0.7f);
 
-		ConditionConstraints();
+		if (bJointsDirty)
+		{
+			ConditionConstraints();
+			bJointsDirty = false;
+		}
+
+		if (bActorsDirty)
+		{
+			UpdateActivePotentiallyCollidingPairs();
+			bActorsDirty = false;
+		}
 
 		Evolution.SetGravity(InGravity);
 		
diff --git a/Engine/Source/Runtime/Engine/Public/Physics/ImmediatePhysics/ImmediatePhysicsChaos/ImmediatePhysicsSimulation_Chaos.h b/Engine/Source/Runtime/Engine/Public/Physics/ImmediatePhysics/ImmediatePhysicsChaos/ImmediatePhysicsSimulation_Chaos.h
--- a/Engine/Source/Runtime/Engine/Public/Physics/ImmediatePhysics/ImmediatePhysicsChaos/ImmediatePhysicsSimulation_Chaos.h
+++ b/Engine/Source/Runtime/Engine/Public/Physics/ImmediatePhysics/ImmediatePhysicsChaos/ImmediatePhysicsSimulation_Chaos.h
@@ -74,6 +74,7 @@ namespace ImmediatePhysics_Chaos
 
 	private:
 		void ConditionConstraints();
+		void UpdateActivePotentiallyCollidingPairs();
 		FReal UpdateStepTime(const FReal DeltaTime, const FReal MaxStepTime);
 
 		void DebugDrawKinematicParticles(const int32 MinDebugLevel, const int32 MaxDebugLevel, const FColor& Color);
@@ -87,7 +88,7 @@ namespace ImmediatePhysics_Chaos
 		using FParticlePair = Chaos::TVector<Chaos::TGeometryParticleHandle<Chaos::FReal, 3>*, 2>;
 
 		// @todo(ccaulfield): Look into these...
-		TArray<FParticlePair> PotentiallyCollidingPairs;
+		TArray<FParticlePair> ActivePotentiallyCollidingPairs;
 		Chaos::TArrayCollectionArray<bool> CollidedParticles;
 		Chaos::TArrayCollectionArray<Chaos::TSerializablePtr<Chaos::FChaosPhysicsMaterial>> ParticleMaterials;
 		Chaos::TArrayCollectionArray<TUniquePtr<Chaos::FChaosPhysicsMaterial>> PerParticleMaterials;
@@ -114,11 +115,16 @@ namespace ImmediatePhysics_Chaos
 		// @todo(ccaulfield): Optimize
 		TMap<const FParticleHandle*, TSet<const FParticleHandle*>> IgnoreCollisionParticlePairTable;
 
+		TArray<FParticlePair> PotentiallyCollidingPairs;
+
 		FTransform SimulationSpaceTransform;
 
 		FReal RollingAverageStepTime;
 		int32 NumRollingAverageStepTimes;
 		int32 MaxNumRollingAverageStepTimes;
+
+		bool bActorsDirty;
+		bool bJointsDirty;
 	};
 
 }
diff --git a/Engine/Source/Runtime/Experimental/Chaos/Private/Chaos/Collision/CollisionDetector.cpp b/Engine/Source/Runtime/Experimental/Chaos/Private/Chaos/Collision/CollisionDetector.cpp
--- a/Engine/Source/Runtime/Experimental/Chaos/Private/Chaos/Collision/CollisionDetector.cpp
+++ b/Engine/Source/Runtime/Experimental/Chaos/Private/Chaos/Collision/CollisionDetector.cpp
@@ -8,7 +8,9 @@ namespace Chaos
 {
 	DEFINE_STAT(STAT_Collisions_Detect);
 	DEFINE_STAT(STAT_Collisions_BroadPhase);
+#if CHAOS_ENABLE_STAT_NARROWPHASE
 	DEFINE_STAT(STAT_Collisions_NarrowPhase);
+#endif
 
 	template class TCollisionDetector<FParticlePairBroadPhase, FNarrowPhase, FSyncCollisionReceiver, TPBDCollisionConstraints<FReal, 3>>;
 	template class TCollisionDetector<FSpatialAccelerationBroadPhase, FNarrowPhase, FAsyncCollisionReceiver, TPBDCollisionConstraints<FReal, 3>>;
diff --git a/Engine/Source/Runtime/Experimental/Chaos/Private/Chaos/CollisionResolution.cpp b/Engine/Source/Runtime/Experimental/Chaos/Private/Chaos/CollisionResolution.cpp
--- a/Engine/Source/Runtime/Experimental/Chaos/Private/Chaos/CollisionResolution.cpp
+++ b/Engine/Source/Runtime/Experimental/Chaos/Private/Chaos/CollisionResolution.cpp
@@ -20,7 +20,12 @@
 #include "Chaos/TriangleMeshImplicitObject.h"
 #include "Chaos/GeometryQueries.h"
 
+#if 0
 DECLARE_CYCLE_STAT(TEXT("Collisions::GJK"), STAT_Collisions_GJK, STATGROUP_ChaosCollision);
+#define SCOPE_CYCLE_COUNTER_GJK() SCOPE_CYCLE_COUNTER(STAT_Collisions_GJK)
+#else
+#define SCOPE_CYCLE_COUNTER_GJK()
+#endif
 
 //#pragma optimize("", off)
 
@@ -110,7 +115,7 @@ namespace Chaos
 		template <typename T, int d, typename GeometryA, typename GeometryB>
 		TContactPoint<T> GJKContactPoint(const GeometryA& A, const TRigidTransform<T, d>& ATM, const GeometryB& B, const TRigidTransform<T, d>& BTM, const TVector<T, 3>& InitialDir)
 		{
-			SCOPE_CYCLE_COUNTER(STAT_Collisions_GJK);
+			SCOPE_CYCLE_COUNTER_GJK();
 
 			TContactPoint<T> Contact;
 			const TRigidTransform<T, d> BToATM = BTM.GetRelativeTransform(ATM);
@@ -887,7 +892,7 @@ namespace Chaos
 			T Penetration;
 			FVec3 CapsuleClosestBoxSpace, BoxClosestBoxSpace, NormalBoxSpace;
 			{
-				SCOPE_CYCLE_COUNTER(STAT_Collisions_GJK);
+				SCOPE_CYCLE_COUNTER_GJK();
 				if (!ensure(GJKPenetration<true>(Box, Capsule, CapsuleToBoxTM, Penetration, BoxClosestBoxSpace, CapsuleClosestBoxSpace, NormalBoxSpace, (T)0, InitialDir, (T)0)))
 				{
 					return;
diff --git a/Engine/Source/Runtime/Experimental/Chaos/Public/Chaos/Collision/NarrowPhase.h b/Engine/Source/Runtime/Experimental/Chaos/Public/Chaos/Collision/NarrowPhase.h
--- a/Engine/Source/Runtime/Experimental/Chaos/Public/Chaos/Collision/NarrowPhase.h
+++ b/Engine/Source/Runtime/Experimental/Chaos/Public/Chaos/Collision/NarrowPhase.h
@@ -11,7 +11,13 @@
 
 namespace Chaos
 {
+#define CHAOS_ENABLE_STAT_NARROWPHASE 0
+#if CHAOS_ENABLE_STAT_NARROWPHASE
 	DECLARE_CYCLE_STAT_EXTERN(TEXT("Collisions::NarrowPhase"), STAT_Collisions_NarrowPhase, STATGROUP_ChaosCollision, CHAOS_API);
+	#define SCOPE_CYCLE_COUNTER_NAROWPHASE() SCOPE_CYCLE_COUNTER(STAT_Collisions_GJK)
+#else
+	#define SCOPE_CYCLE_COUNTER_NAROWPHASE()
+#endif
 
 	/**
 	 * Generate contact manifolds for particle pairs.
@@ -27,7 +33,7 @@ namespace Chaos
 		 */
 		void GenerateCollisions(FCollisionConstraintsArray& NewConstraints, FReal Dt, TGeometryParticleHandle<FReal, 3>* Particle0, TGeometryParticleHandle<FReal, 3>* Particle1, const FReal CullDistance, CollisionStats::FStatData& StatData)
 		{
-			SCOPE_CYCLE_COUNTER(STAT_Collisions_NarrowPhase);
+			SCOPE_CYCLE_COUNTER_NAROWPHASE();
 			if (ensure(Particle0 && Particle1))
 			{
 				//
diff --git a/Engine/Source/Runtime/Experimental/Chaos/Public/Chaos/Collision/ParticlePairBroadPhase.h b/Engine/Source/Runtime/Experimental/Chaos/Public/Chaos/Collision/ParticlePairBroadPhase.h
--- a/Engine/Source/Runtime/Experimental/Chaos/Public/Chaos/Collision/ParticlePairBroadPhase.h
+++ b/Engine/Source/Runtime/Experimental/Chaos/Public/Chaos/Collision/ParticlePairBroadPhase.h
@@ -44,48 +44,28 @@ namespace Chaos
 		{
 			SCOPE_CYCLE_COUNTER(STAT_Collisions_BroadPhase);
 
-			OverlappingIndices.Reserve(ParticlePairs.Num());
-
-			const int32 NumPairs = ParticlePairs.Num();
+			int32 NumPairs = ParticlePairs.Num();
 			for (int32 PairIndex = 0; PairIndex < NumPairs; ++PairIndex)
 			{
 				const FParticlePair& ParticlePair = ParticlePairs[PairIndex];
 
-				// Array is const, things in it are not...
-				TGeometryParticleHandle<FReal, 3>* Particle0 = const_cast<TGeometryParticleHandle<FReal, 3>*>(ParticlePair[0]);
-				TGeometryParticleHandle<FReal, 3>* Particle1 = const_cast<TGeometryParticleHandle<FReal, 3>*>(ParticlePair[1]);
-
-				// Particles may have been disabled or made kinematic
-				bool bAnyDisabled = TGenericParticleHandle<FReal, 3>(Particle0)->Disabled() || TGenericParticleHandle<FReal, 3>(Particle1)->Disabled();
-				bool bAnyDynamic = TGenericParticleHandle<FReal, 3>(Particle0)->IsDynamic() || TGenericParticleHandle<FReal, 3>(Particle1)->IsDynamic();
-				if (bAnyDynamic && !bAnyDisabled)
+				const TAABB<FReal, 3>& Box0 = ParticlePair[0]->WorldSpaceInflatedBounds();
+				const TAABB<FReal, 3>& Box1 = ParticlePair[1]->WorldSpaceInflatedBounds();
+				if (Box0.Intersects(Box1))
 				{
-					const TAABB<FReal, 3>& Box0 = Particle0->WorldSpaceInflatedBounds();
-					const TAABB<FReal, 3>& Box1 = Particle1->WorldSpaceInflatedBounds();
-					if (Box0.Intersects(Box1))
-					{
-						OverlappingIndices.Add(PairIndex);
-					}
-				}
-			}
+					// Pair array is const, the particles are not
+					TGeometryParticleHandle<FReal, 3>* Particle0 = const_cast<TGeometryParticleHandle<FReal, 3>*>(ParticlePair[0]);
+					TGeometryParticleHandle<FReal, 3>* Particle1 = const_cast<TGeometryParticleHandle<FReal, 3>*>(ParticlePair[1]);
 
-			for (int32 PairIndex : OverlappingIndices)
-			{
-				const FParticlePair& ParticlePair = ParticlePairs[PairIndex];
-				TGeometryParticleHandle<FReal, 3>* Particle0 = const_cast<TGeometryParticleHandle<FReal, 3>*>(ParticlePair[0]);
-				TGeometryParticleHandle<FReal, 3>* Particle1 = const_cast<TGeometryParticleHandle<FReal, 3>*>(ParticlePair[1]);
-
-				NarrowPhase.GenerateCollisions(NewConstraints, Dt, Particle0, Particle1, CullDistance, StatData);
-				Receiver.ReceiveCollisions(NewConstraints);
-				NewConstraints.Empty();
+					NarrowPhase.GenerateCollisions(NewConstraints, Dt, Particle0, Particle1, CullDistance, StatData);
+					Receiver.ReceiveCollisions(NewConstraints);
+					NewConstraints.Empty();
+				}
 			}
-
-			OverlappingIndices.Reset();
 		}
 
 	private:
 		const TArray<FParticlePair>& ParticlePairs;
-		TArray<int32> OverlappingIndices;
 		FReal CullDistance;
 		FCollisionConstraintsArray NewConstraints;
 	};
