Add missing #include to fix compile error

#ue4

diff --git a/Engine/Source/Runtime/Engine/Public/BufferVisualizationData.h b/Engine/Source/Runtime/Engine/Public/BufferVisualizationData.h
--- a/Engine/Source/Runtime/Engine/Public/BufferVisualizationData.h
+++ b/Engine/Source/Runtime/Engine/Public/BufferVisualizationData.h
@@ -2,8 +2,7 @@
 #pragma once
 
 #include "CoreMinimal.h"
-
-class UMaterialInterface;
+#include "Materials/MaterialInterface.h"
 
 class FBufferVisualizationData
 {
